Move all IoStoreWriter code (writing of .utoc/.ucas containers) out from Core into IoStoreUtilities.
Break cyclic dependency on DerivedDataCache module to utilize DDC2 API to cache compressed chunks.
The declarations moved from the public header are considered private and are not really usable outside of the IoStoreUtilities/UnrealPak context.

#rnx
#tests BuildCookRun Lyra with pak/container files

diff --git a/Engine/Source/Developer/IoStoreUtilities/IoStoreUtilities.Build.cs b/Engine/Source/Developer/IoStoreUtilities/IoStoreUtilities.Build.cs
--- a/Engine/Source/Developer/IoStoreUtilities/IoStoreUtilities.Build.cs
+++ b/Engine/Source/Developer/IoStoreUtilities/IoStoreUtilities.Build.cs
@@ -13,6 +13,7 @@ public class IoStoreUtilities : ModuleRules
 		PrivateDependencyModuleNames.AddRange(new string[] {
             "Core",
             "CoreUObject",
+            "DerivedDataCache",
 			"AssetRegistry",
 			"CookMetadata",
 			"Projects",
diff --git a/Engine/Source/Developer/IoStoreUtilities/Private/IoStoreLooseFiles.cpp b/Engine/Source/Developer/IoStoreUtilities/Private/IoStoreLooseFiles.cpp
--- a/Engine/Source/Developer/IoStoreUtilities/Private/IoStoreLooseFiles.cpp
+++ b/Engine/Source/Developer/IoStoreUtilities/Private/IoStoreLooseFiles.cpp
@@ -1,6 +1,7 @@
 // Copyright Epic Games, Inc. All Rights Reserved.
 
 #include "IoStoreLooseFiles.h"
+#include "IoStoreWriter.h"
 
 #include "HAL/FileManager.h"
 #include "IO/IoStore.h"
diff --git a/Engine/Source/Developer/IoStoreUtilities/Private/IoStoreUtilities.cpp b/Engine/Source/Developer/IoStoreUtilities/Private/IoStoreUtilities.cpp
--- a/Engine/Source/Developer/IoStoreUtilities/Private/IoStoreUtilities.cpp
+++ b/Engine/Source/Developer/IoStoreUtilities/Private/IoStoreUtilities.cpp
@@ -1,6 +1,7 @@
 // Copyright Epic Games, Inc. All Rights Reserved.
 
 #include "IoStoreUtilities.h"
+#include "IoStoreWriter.h"
 
 #include "IoStoreLooseFiles.h"
 #include "Algo/TopologicalSort.h"
diff --git a/Engine/Source/Developer/IoStoreUtilities/Private/IoStoreWriter.cpp b/Engine/Source/Developer/IoStoreUtilities/Private/IoStoreWriter.cpp
new file mode 100644
--- /dev/null
+++ b/Engine/Source/Developer/IoStoreUtilities/Private/IoStoreWriter.cpp
@@ -0,0 +1,1994 @@
+// Copyright Epic Games, Inc. All Rights Reserved.
+
+#include "IoStoreWriter.h"
+#include "IO/IoStore.h"
+
+#include "Async/Async.h"
+#include "Async/AsyncFileHandle.h"
+#include "Async/Future.h"
+#include "Async/ParallelFor.h"
+#include "Containers/Map.h"
+#include "DerivedDataCacheInterface.h"
+#include "HAL/FileManager.h"
+#include "HAL/PlatformFileManager.h"
+#include "IO/IoDirectoryIndex.h"
+#include "Misc/Compression.h"
+#include "Misc/Paths.h"
+#include "Misc/StringBuilder.h"
+#include "ProfilingDebugging/CountersTrace.h"
+#include "Serialization/LargeMemoryReader.h"
+#include "Serialization/MemoryWriter.h"
+#include "Templates/UniquePtr.h"
+
+struct FChunkBlock
+{
+	const uint8* UncompressedData = nullptr;
+	FIoBuffer* IoBuffer = nullptr;
+
+	// This is the size of the actual block after encryption alignment, and is
+	// set in EncryptAndSign. This happens whether or not the container is encrypted.
+	uint64 Size = 0;
+	uint64 CompressedSize = 0;
+	uint64 UncompressedSize = 0;
+	FName CompressionMethod = NAME_None;
+	FSHAHash Signature;
+
+	/** Hash of the block data as it would be found on disk - this includes encryption alignment padding */
+	FIoHash DiskHash;
+};
+
+struct FIoStoreWriteQueueEntry
+{
+	FIoStoreWriteQueueEntry* Next = nullptr;
+	class FIoStoreWriter* Writer = nullptr;
+	IIoStoreWriteRequest* Request = nullptr;
+	FIoChunkId ChunkId;
+	FIoChunkHash ChunkHash;
+	/** Hash of the block data as it would be found on disk after compression and encryption */
+	FIoHash ChunkDiskHash;
+
+	uint64 Sequence = 0;
+	
+	// We make this optional because at the latest it might not be valid until FinishCompressionBarrior
+	// completes and we'd like to have a check() on that.
+	TOptional<uint64> UncompressedSize = 0;
+
+	// this is not filled out until after encryption completes and *includes the alignment padding for encryption*!
+	// think of this as "size on disk".
+	uint64 CompressedSize = 0; 
+
+	uint64 Padding = 0;
+	uint64 Offset = 0;
+	TArray<FChunkBlock> ChunkBlocks;
+	FIoWriteOptions Options;
+	FGraphEventRef HashBarrier;
+	FGraphEventRef HashTask;
+	FGraphEventRef BeginCompressionBarrier;
+	FGraphEventRef FinishCompressionBarrier;
+	FGraphEventRef FinishEncryptionAndSigningBarrier;
+	FGraphEventRef BeginWriteBarrier;
+	FGraphEventRef WriteFinishedEvent;
+	TAtomic<int32> CompressedBlocksCount{ 0 };
+	TAtomic<int32> FinishedBlocksCount{ 0 };
+	int32 PartitionIndex = -1;
+	FString DDCCacheKey;
+	bool bAdded = false;
+	bool bModified = false;
+	bool bStoreCompressedDataInDDC = false;
+	
+	bool bCouldBeFromReferenceDb = false; // Whether the chunk is a valid candidate for the reference db.
+	bool bLoadingFromReferenceDb = false;
+	// When we know we're loading from the reference chunk db, we don't read the source
+	// buffer but we still need to know the number of chunks which we get from the refdb.
+	uint32 NumChunkBlocksFromRefDb = 0; 
+};
+
+class FIoStoreWriteQueue
+{
+public:
+	FIoStoreWriteQueue()
+		: Event(FPlatformProcess::GetSynchEventFromPool(false))
+	{ }
+	
+	~FIoStoreWriteQueue()
+	{
+		check(Head == nullptr && Tail == nullptr);
+		FPlatformProcess::ReturnSynchEventToPool(Event);
+	}
+
+	void Enqueue(FIoStoreWriteQueueEntry* Entry)
+	{
+		check(!bIsDoneAdding);
+		{
+			FScopeLock _(&CriticalSection);
+
+			if (!Tail)
+			{
+				Head = Tail = Entry;
+			}
+			else
+			{
+				Tail->Next = Entry;
+				Tail = Entry;
+			}
+			Entry->Next = nullptr;
+		}
+
+		Event->Trigger();
+	}
+
+	FIoStoreWriteQueueEntry* DequeueOrWait()
+	{
+		for (;;)
+		{
+			{
+				FScopeLock _(&CriticalSection);
+				if (Head)
+				{
+					FIoStoreWriteQueueEntry* Entry = Head;
+					Head = Tail = nullptr;
+					return Entry;
+				}
+			}
+
+			if (bIsDoneAdding)
+			{
+				break;
+			}
+
+			Event->Wait();
+		}
+
+		return nullptr;
+	}
+
+	void CompleteAdding()
+	{
+		bIsDoneAdding = true;
+		Event->Trigger();
+	}
+
+	bool IsEmpty() const
+	{
+		FScopeLock _(&CriticalSection);
+		return Head == nullptr;
+	}
+
+private:
+	mutable FCriticalSection CriticalSection;
+	FEvent* Event = nullptr;
+	FIoStoreWriteQueueEntry* Head = nullptr;
+	FIoStoreWriteQueueEntry* Tail = nullptr;
+	TAtomic<bool> bIsDoneAdding { false };
+};
+
+class FIoStoreWriterContextImpl
+{
+	static constexpr uint64 DefaultMemoryLimit = 2ull << 30ull;
+public:
+	FIoStoreWriterContextImpl()
+	{
+	}
+
+	~FIoStoreWriterContextImpl()
+	{
+		BeginCompressionQueue.CompleteAdding();
+		BeginEncryptionAndSigningQueue.CompleteAdding();
+		FinishEncryptionAndSigningQueue.CompleteAdding();
+		WriterQueue.CompleteAdding();
+		BeginCompressionThread.Wait();
+		BeginEncryptionAndSigningThread.Wait();
+		FinishEncryptionAndSigningThread.Wait();
+		WriterThread.Wait();
+		if (CompressionBufferAvailableEvent)
+		{
+			FPlatformProcess::ReturnSynchEventToPool(CompressionBufferAvailableEvent);
+		}
+		check(AvailableCompressionBuffers.Num() == TotalCompressionBufferCount);
+		for (FIoBuffer* IoBuffer : AvailableCompressionBuffers)
+		{
+			delete IoBuffer;
+		}
+	}
+
+	[[nodiscard]] FIoStatus Initialize(const FIoStoreWriterSettings& InWriterSettings)
+	{
+		TRACE_CPUPROFILER_EVENT_SCOPE(FIoStoreWriterContext::Initialize);
+		WriterSettings = InWriterSettings;
+		CompressionBufferAvailableEvent = FPlatformProcess::GetSynchEventFromPool(false);
+
+		if (WriterSettings.bCompressionEnableDDC)
+		{
+			DDC = GetDerivedDataCache();
+		}
+
+		if (WriterSettings.CompressionMethod != NAME_None)
+		{
+			CompressionBufferSize = FCompression::CompressMemoryBound(WriterSettings.CompressionMethod, static_cast<int32>(WriterSettings.CompressionBlockSize));
+		}
+		CompressionBufferSize = FMath::Max(CompressionBufferSize, static_cast<int32>(WriterSettings.CompressionBlockSize));
+		CompressionBufferSize = Align(CompressionBufferSize, FAES::AESBlockSize);
+
+		TotalCompressionBufferCount = int32(DefaultMemoryLimit / CompressionBufferSize);
+		AvailableCompressionBuffers.Reserve(TotalCompressionBufferCount);
+		for (int32 BufferIndex = 0; BufferIndex < TotalCompressionBufferCount; ++BufferIndex)
+		{
+			AvailableCompressionBuffers.Add(new FIoBuffer(CompressionBufferSize));
+		}
+		//TRACE_COUNTER_SET(IoStoreAvailableCompressionBuffers, AvailableCompressionBuffers.Num());
+
+		BeginCompressionThread = Async(EAsyncExecution::Thread, [this]() { BeginCompressionThreadFunc(); });
+		BeginEncryptionAndSigningThread = Async(EAsyncExecution::Thread, [this]() { BeginEncryptionAndSigningThreadFunc(); });
+		FinishEncryptionAndSigningThread = Async(EAsyncExecution::Thread, [this]() { FinishEncryptionAndSigningThreadFunc(); });
+		WriterThread = Async(EAsyncExecution::Thread, [this]() { WriterThreadFunc(); });
+
+		return FIoStatus::Ok;
+	}
+
+	TSharedPtr<IIoStoreWriter> CreateContainer(const TCHAR* InContainerPathAndBaseFileName, const FIoContainerSettings& InContainerSettings);
+
+	void Flush();
+
+	FIoStoreWriterContext::FProgress GetProgress() const
+	{
+		FIoStoreWriterContext::FProgress Progress;
+		Progress.HashDbChunksCount = HashDbChunksCount.Load();
+		for (uint8 i=0; i<(uint8)EIoChunkType::MAX; i++)
+		{
+			Progress.HashDbChunksByType[i] = HashDbChunksByType[i].Load();
+		}
+
+		Progress.TotalChunksCount = TotalChunksCount.Load();
+		Progress.HashedChunksCount = HashedChunksCount.Load();
+		Progress.CompressedChunksCount = CompressedChunksCount.Load();
+		for (uint8 i = 0; i < (uint8)EIoChunkType::MAX; i++)
+		{
+			Progress.CompressedChunksByType[i] = CompressedChunksByType[i].Load();
+		}
+		for (uint8 i = 0; i < (uint8)EIoChunkType::MAX; i++)
+		{
+			Progress.BeginCompressChunksByType[i] = BeginCompressChunksByType[i].Load();
+		}
+		Progress.SerializedChunksCount = SerializedChunksCount.Load();
+		Progress.ScheduledCompressionTasksCount = ScheduledCompressionTasksCount.Load();
+		Progress.CompressionDDCHitCount = CompressionDDCHitCount.Load();
+		Progress.CompressionDDCMissCount = CompressionDDCMissCount.Load();
+		Progress.RefDbChunksCount = RefDbChunksCount.Load();
+		for (uint8 i = 0; i < (uint8)EIoChunkType::MAX; i++)
+		{
+			Progress.RefDbChunksByType[i] = RefDbChunksByType[i].Load();
+		}
+
+		return Progress;
+	}
+
+	const FIoStoreWriterSettings& GetSettings() const
+	{
+		return WriterSettings;
+	}
+
+	void ScheduleCompression(FIoStoreWriteQueueEntry* QueueEntry)
+	{
+		BeginCompressionQueue.Enqueue(QueueEntry);
+		if (QueueEntry->bLoadingFromReferenceDb == false)
+		{
+			QueueEntry->Request->PrepareSourceBufferAsync(QueueEntry->BeginCompressionBarrier);
+		}
+		else
+		{
+			// We don't need to wait on a read so we can kick directly.
+			QueueEntry->BeginCompressionBarrier->DispatchSubsequents();
+		}
+	}
+
+	//
+	// This must be called prior to the consumer being dispatched in order to prevent resource
+	// contention deadlock (even if using retracting waits)
+	//
+	FIoBuffer* AllocCompressionBuffer(int32 TotalEntryChunkBlocksCount)
+	{
+		if (TotalEntryChunkBlocksCount > TotalCompressionBufferCount)
+		{
+			return new FIoBuffer(CompressionBufferSize);
+		}
+		FIoBuffer* AllocatedBuffer = nullptr;
+		while (!AllocatedBuffer)
+		{
+			{
+				FScopeLock Lock(&AvailableCompressionBuffersCritical);
+				if (AvailableCompressionBuffers.Num())
+				{
+					AllocatedBuffer = AvailableCompressionBuffers.Pop();
+					//TRACE_COUNTER_DECREMENT(IoStoreAvailableCompressionBuffers);
+				}
+			}
+			if (!AllocatedBuffer)
+			{
+				TRACE_CPUPROFILER_EVENT_SCOPE(WaitForCompressionBuffer);
+				CompressionBufferAvailableEvent->Wait();
+			}
+		}
+		return AllocatedBuffer;
+	}
+
+	void FreeCompressionBuffer(FIoBuffer* Buffer, int32 TotalEntryChunkBlocksCount)
+	{
+		if (TotalEntryChunkBlocksCount > TotalCompressionBufferCount)
+		{
+			delete Buffer;
+			return;
+		}
+		bool bTriggerEvent;
+		{
+			FScopeLock Lock(&AvailableCompressionBuffersCritical);
+			bTriggerEvent = AvailableCompressionBuffers.Num() == 0;
+			AvailableCompressionBuffers.Push(Buffer);
+			//TRACE_COUNTER_INCREMENT(IoStoreAvailableCompressionBuffers);
+		}
+		if (bTriggerEvent)
+		{
+			CompressionBufferAvailableEvent->Trigger();
+		}
+	}
+
+private:
+	void BeginCompressionThreadFunc();
+	void BeginEncryptionAndSigningThreadFunc();
+	void FinishEncryptionAndSigningThreadFunc();
+	void WriterThreadFunc();
+
+	FIoStoreWriterSettings WriterSettings;
+	FDerivedDataCacheInterface* DDC = nullptr;
+	FEvent* CompressionBufferAvailableEvent = nullptr;
+	TFuture<void> BeginCompressionThread;
+	TFuture<void> BeginEncryptionAndSigningThread;
+	TFuture<void> FinishEncryptionAndSigningThread;
+	TFuture<void> WriterThread;
+	FIoStoreWriteQueue BeginCompressionQueue;
+	FIoStoreWriteQueue BeginEncryptionAndSigningQueue;
+	FIoStoreWriteQueue FinishEncryptionAndSigningQueue;
+	FIoStoreWriteQueue WriterQueue;
+	TAtomic<uint64> TotalChunksCount{ 0 };
+	TAtomic<uint64> HashedChunksCount{ 0 };
+	TAtomic<uint64> HashDbChunksCount{ 0 };
+	TAtomic<uint64> HashDbChunksByType[(int8)EIoChunkType::MAX] = {0};
+	TAtomic<uint64> RefDbChunksCount{ 0 };
+	TAtomic<uint64> RefDbChunksByType[(int8)EIoChunkType::MAX] = { 0 };
+	TAtomic<uint64> CompressedChunksCount{ 0 };
+	TAtomic<uint64> CompressedChunksByType[(int8)EIoChunkType::MAX] = { 0 };
+	TAtomic<uint64> BeginCompressChunksByType[(int8)EIoChunkType::MAX] = { 0 };
+	TAtomic<uint64> SerializedChunksCount{ 0 };
+	TAtomic<uint64> WriteCycleCount{ 0 };
+	TAtomic<uint64> WriteByteCount{ 0 };
+	TAtomic<uint64> ScheduledCompressionTasksCount{ 0 };
+	TAtomic<uint64> CompressionDDCHitCount{ 0 };
+	TAtomic<uint64> CompressionDDCMissCount{ 0 };
+	FCriticalSection AvailableCompressionBuffersCritical;
+	TArray<FIoBuffer*> AvailableCompressionBuffers;
+	int32 CompressionBufferSize = -1;
+	int32 TotalCompressionBufferCount = -1;
+	TArray<TSharedPtr<FIoStoreWriter>> IoStoreWriters;
+
+	friend class FIoStoreWriter;
+};
+
+FIoStoreWriterContext::FIoStoreWriterContext()
+	: Impl(new FIoStoreWriterContextImpl())
+{
+
+}
+
+FIoStoreWriterContext::~FIoStoreWriterContext()
+{
+	delete Impl;
+}
+
+[[nodiscard]] FIoStatus FIoStoreWriterContext::Initialize(const FIoStoreWriterSettings& InWriterSettings)
+{
+	return Impl->Initialize(InWriterSettings);
+}
+
+TSharedPtr<IIoStoreWriter> FIoStoreWriterContext::CreateContainer(const TCHAR* InContainerPath, const FIoContainerSettings& InContainerSettings)
+{
+	return Impl->CreateContainer(InContainerPath, InContainerSettings);
+}
+
+void FIoStoreWriterContext::Flush()
+{
+	Impl->Flush();
+}
+
+FIoStoreWriterContext::FProgress FIoStoreWriterContext::GetProgress() const
+{
+	return Impl->GetProgress();
+}
+
+class FIoStoreTocBuilder
+{
+public:
+	FIoStoreTocBuilder()
+	{
+		FMemory::Memzero(&Toc.Header, sizeof(FIoStoreTocHeader));
+	}
+
+	int32 AddChunkEntry(const FIoChunkId& ChunkId, const FIoOffsetAndLength& OffsetLength, const FIoStoreTocEntryMeta& Meta)
+	{
+		int32& Index = ChunkIdToIndex.FindOrAdd(ChunkId);
+
+		if (!Index)
+		{
+			Index = Toc.ChunkIds.Add(ChunkId);
+			Toc.ChunkOffsetLengths.Add(OffsetLength);
+			Toc.ChunkMetas.Add(Meta);
+
+			return Index;
+		}
+
+		return INDEX_NONE;
+	}
+
+	FIoStoreTocCompressedBlockEntry& AddCompressionBlockEntry()
+	{
+		return Toc.CompressionBlocks.AddDefaulted_GetRef();
+	}
+
+	FSHAHash& AddBlockSignatureEntry()
+	{
+		return Toc.ChunkBlockSignatures.AddDefaulted_GetRef();
+	}
+
+	uint8 AddCompressionMethodEntry(FName CompressionMethod)
+	{
+		if (CompressionMethod == NAME_None)
+		{
+			return 0;
+		}
+
+		uint8 Index = 1;
+		for (const FName& Name : Toc.CompressionMethods)
+		{
+			if (Name == CompressionMethod)
+			{
+				return Index;
+			}
+			++Index;
+		}
+
+		return 1 + uint8(Toc.CompressionMethods.Add(CompressionMethod));
+	}
+
+	void AddToFileIndex(const FIoChunkId& ChunkId, const FString& FileName)
+	{
+		ChunkIdToFileName.Emplace(ChunkId, FileName);
+	}
+
+	FIoStoreTocResource& GetTocResource()
+	{
+		return Toc;
+	}
+
+	const FIoStoreTocResource& GetTocResource() const
+	{
+		return Toc;
+	}
+
+	const int32* GetTocEntryIndex(const FIoChunkId& ChunkId) const
+	{
+		return ChunkIdToIndex.Find(ChunkId);
+	}
+
+	void GetFileNamesToIndex(TArray<FStringView>& OutFileNames) const
+	{
+		OutFileNames.Empty(ChunkIdToFileName.Num());
+		for (auto& ChinkIdAndFileName : ChunkIdToFileName)
+		{
+			OutFileNames.Emplace(ChinkIdAndFileName.Value);
+		}
+	}
+
+	const FString* GetFileName(const FIoChunkId& ChunkId) const
+	{
+		return ChunkIdToFileName.Find(ChunkId);
+	}
+
+	FIoStoreTocChunkInfo GetTocChunkInfo(int32 TocEntryIndex) const
+	{
+		FIoStoreTocChunkInfo ChunkInfo = Toc.GetTocChunkInfo(TocEntryIndex);
+
+		ChunkInfo.FileName = FString::Printf(TEXT("<%s>"), *LexToString(ChunkInfo.ChunkType));
+		ChunkInfo.bHasValidFileName = false;
+
+		return ChunkInfo;
+	}
+
+private:
+	FIoStoreTocResource Toc;
+	TMap<FIoChunkId, int32> ChunkIdToIndex;
+	TMap<FIoChunkId, FString> ChunkIdToFileName;
+};
+
+class FIoStoreWriter
+	: public IIoStoreWriter
+{
+public:
+	FIoStoreWriter(const TCHAR* InContainerPathAndBaseFileName)
+		: ContainerPathAndBaseFileName(InContainerPathAndBaseFileName)
+	{
+	}
+
+	void SetReferenceChunkDatabase(TSharedPtr<IIoStoreWriterReferenceChunkDatabase> InReferenceChunkDatabase)
+	{
+		if (InReferenceChunkDatabase.IsValid() == false)
+		{
+			ReferenceChunkDatabase = InReferenceChunkDatabase;
+			return;
+		}
+
+		if (InReferenceChunkDatabase->GetCompressionBlockSize() != WriterContext->GetSettings().CompressionBlockSize)
+		{
+			UE_LOG(LogIoStore, Warning, TEXT("Reference chunk database has a different compression block size than the current writer!"));
+			UE_LOG(LogIoStore, Warning, TEXT("No chunks will match, so ignoring. ReferenceChunkDb: %d, IoStoreWriter: %d"), InReferenceChunkDatabase->GetCompressionBlockSize(), WriterContext->GetSettings().CompressionBlockSize);
+			return;
+		}
+		ReferenceChunkDatabase = InReferenceChunkDatabase;
+		
+		// Add ourselves to the reference chunk db's list of possibles
+		ReferenceChunkDatabase->NotifyAddedToWriter(ContainerSettings.ContainerId, FPaths::GetBaseFilename(TocFilePath));
+	}
+	void SetHashDatabase(TSharedPtr<IIoStoreWriterHashDatabase> InHashDatabase, bool bInVerifyHashDatabase)
+	{
+		HashDatabase = InHashDatabase;
+		bVerifyHashDatabase = bInVerifyHashDatabase;
+	}
+
+	void EnumerateChunks(TFunction<bool(FIoStoreTocChunkInfo&&)>&& Callback) const
+	{
+		const FIoStoreTocResource& TocResource = TocBuilder.GetTocResource();
+
+		for (int32 ChunkIndex = 0; ChunkIndex < TocResource.ChunkIds.Num(); ++ChunkIndex)
+		{
+			FIoStoreTocChunkInfo ChunkInfo = TocBuilder.GetTocChunkInfo(ChunkIndex);
+			if (!Callback(MoveTemp(ChunkInfo)))
+			{
+				break;
+			}
+		}
+	}
+
+	[[nodiscard]] FIoStatus Initialize(FIoStoreWriterContextImpl& InContext, const FIoContainerSettings& InContainerSettings)
+	{
+		WriterContext = &InContext;
+		ContainerSettings = InContainerSettings;
+
+		TocFilePath = ContainerPathAndBaseFileName + TEXT(".utoc");
+		
+		IPlatformFile& Ipf = IPlatformFile::GetPlatformPhysical();
+		Ipf.CreateDirectoryTree(*FPaths::GetPath(TocFilePath));
+
+		FPartition& Partition = Partitions.AddDefaulted_GetRef();
+		Partition.Index = 0;
+
+		return FIoStatus::Ok;
+	}
+
+	virtual void EnableDiskLayoutOrdering(const TArray<TUniquePtr<FIoStoreReader>>& PatchSourceReaders) override
+	{
+		check(!LayoutEntriesHead);
+		check(!Entries.Num());
+		LayoutEntriesHead = new FLayoutEntry();
+		LayoutEntries.Add(LayoutEntriesHead);
+		FLayoutEntry* PrevEntryLink = LayoutEntriesHead;
+
+		for (const TUniquePtr<FIoStoreReader>& PatchSourceReader : PatchSourceReaders)
+		{
+			TArray<TPair<uint64, FLayoutEntry*>> LayoutEntriesWithOffsets;
+			PatchSourceReader->EnumerateChunks([this, &PrevEntryLink, &LayoutEntriesWithOffsets](const FIoStoreTocChunkInfo& ChunkInfo)
+				{
+					FLayoutEntry* PreviousBuildEntry = new FLayoutEntry();
+					PreviousBuildEntry->Hash = ChunkInfo.Hash;
+					PreviousBuildEntry->PartitionIndex = ChunkInfo.PartitionIndex;
+					PreviousBuildEntry->CompressedSize = ChunkInfo.CompressedSize;
+					LayoutEntriesWithOffsets.Emplace(ChunkInfo.Offset, PreviousBuildEntry);
+					PreviousBuildLayoutEntryByChunkId.Add(ChunkInfo.Id, PreviousBuildEntry);
+					return true;
+				});
+
+			// Sort entries by offset
+			Algo::Sort(LayoutEntriesWithOffsets, [](const TPair<uint64, FLayoutEntry*>& A, const TPair<uint64, FLayoutEntry*>& B)
+				{
+					return A.Get<0>() < B.Get<0>();
+				});
+
+			for (const TPair<uint64, FLayoutEntry*>& EntryWithOffset : LayoutEntriesWithOffsets)
+			{
+				FLayoutEntry* PreviousBuildEntry = EntryWithOffset.Get<1>();
+				LayoutEntries.Add(PreviousBuildEntry);
+				PrevEntryLink->Next = PreviousBuildEntry;
+				PreviousBuildEntry->Prev = PrevEntryLink;
+				PrevEntryLink = PreviousBuildEntry;
+			}
+			if (!ContainerSettings.bGenerateDiffPatch)
+			{
+				break;
+			}
+		}
+
+		LayoutEntriesTail = new FLayoutEntry();
+		LayoutEntries.Add(LayoutEntriesTail);
+		PrevEntryLink->Next = LayoutEntriesTail;
+		LayoutEntriesTail->Prev = PrevEntryLink;
+	}
+
+	virtual void Append(const FIoChunkId& ChunkId, IIoStoreWriteRequest* Request, const FIoWriteOptions& WriteOptions) override
+	{
+		//
+		// This function sets up the sequence of events that takes a chunk from source data on disc
+		// to written to a container. The first thing that happens is the source data is read in order
+		// to hash it to detect whether or not it's modified as well as look up in reference databases.
+		// Load the data -> PrepareSourceBufferAsync
+		// Hash the data -> HashTask lambda
+		//
+		// The hash task itself doesn't continue to the next steps - the Flush() call
+		// waits for all hashes to be complete before kicking the next steps.
+		//
+		TRACE_CPUPROFILER_EVENT_SCOPE(AppendWriteRequest);
+		check(!bHasFlushed);
+		checkf(ChunkId.IsValid(), TEXT("ChunkId is not valid!"));
+
+		FIoStoreWriteQueueEntry* Entry = new FIoStoreWriteQueueEntry();
+		Entry->Writer = this;
+		Entry->Sequence = Entries.Num();
+		WriterContext->TotalChunksCount.IncrementExchange();
+		Entries.Add(Entry);
+		Entry->ChunkId = ChunkId;
+		Entry->Options = WriteOptions;
+		Entry->Request = Request;		
+		Entry->BeginCompressionBarrier = FGraphEvent::CreateGraphEvent();
+		Entry->FinishCompressionBarrier = FGraphEvent::CreateGraphEvent();
+		Entry->FinishEncryptionAndSigningBarrier = FGraphEvent::CreateGraphEvent();
+		Entry->BeginWriteBarrier = FGraphEvent::CreateGraphEvent();
+		Entry->WriteFinishedEvent = FGraphEvent::CreateGraphEvent();
+		
+		bool bTestHashValid = false;
+		FIoChunkHash TestHash;
+		// If we can get the hash without reading the whole thing and hashing it, do so
+		// to avoid the IO.
+		if (HashDatabase.IsValid() &&
+			HashDatabase->FindHashForChunkId(ChunkId, Entry->ChunkHash))
+		{
+			if (bVerifyHashDatabase == false)
+			{
+				// If we aren't validating then we just use it and bail.
+				WriterContext->HashDbChunksCount.IncrementExchange();
+				WriterContext->HashDbChunksByType[(int8)Entry->ChunkId.GetChunkType()].IncrementExchange();
+				WriterContext->HashedChunksCount.IncrementExchange();
+
+				if (ReferenceChunkDatabase.IsValid() && CompressionMethodForEntry(Entry) != NAME_None)
+				{
+					Entry->bLoadingFromReferenceDb = ReferenceChunkDatabase->ChunkExists(ContainerSettings.ContainerId, Entry->ChunkHash, Entry->ChunkId, Entry->NumChunkBlocksFromRefDb);
+					Entry->bCouldBeFromReferenceDb = true;
+				}
+				return;
+			}
+
+			// If we are validating, copy the hash out and run the normal path
+			// to check against it.
+			TestHash = Entry->ChunkHash;
+			bTestHashValid = true;
+		}
+		// Otherwise, we have to do the load & hash
+
+		Entry->HashBarrier = FGraphEvent::CreateGraphEvent();
+
+		FGraphEventArray HashPrereqs;
+		HashPrereqs.Add(Entry->HashBarrier);
+		Entry->HashTask = FFunctionGraphTask::CreateAndDispatchWhenReady([this, Entry, TestHash, bTestHashValid]()
+		{
+			TRACE_CPUPROFILER_EVENT_SCOPE(HashChunk);
+			const FIoBuffer* SourceBuffer = Entry->Request->GetSourceBuffer();
+			Entry->ChunkHash = FIoChunkHash::HashBuffer(SourceBuffer->Data(), SourceBuffer->DataSize());
+			WriterContext->HashedChunksCount.IncrementExchange();
+
+			if (bVerifyHashDatabase && bTestHashValid)
+			{
+				if (TestHash != Entry->ChunkHash)
+				{
+					UE_LOG(LogIoStore, Warning, TEXT("HashDb Validation Failed: ChunkId %s has mismatching hash"), *LexToString(Entry->ChunkId));
+				}
+			}
+
+
+			if (ReferenceChunkDatabase.IsValid() && CompressionMethodForEntry(Entry) != NAME_None)
+			{
+				Entry->bLoadingFromReferenceDb = ReferenceChunkDatabase->ChunkExists(ContainerSettings.ContainerId, Entry->ChunkHash, Entry->ChunkId, Entry->NumChunkBlocksFromRefDb);
+				Entry->bCouldBeFromReferenceDb = true;
+			}
+
+			// Release the source data buffer, it will be reloaded later when we start compressing the chunk
+			Entry->Request->FreeSourceBuffer();
+		}, TStatId(), &HashPrereqs, ENamedThreads::AnyHiPriThreadHiPriTask);
+		Entry->Request->PrepareSourceBufferAsync(Entry->HashBarrier);
+	}
+
+	virtual void Append(const FIoChunkId& ChunkId, FIoBuffer Chunk, const FIoWriteOptions& WriteOptions, uint64 OrderHint) override
+	{
+		struct FWriteRequest
+			: IIoStoreWriteRequest
+		{
+			FWriteRequest(FIoBuffer InSourceBuffer, uint64 InOrderHint)
+				: OrderHint(InOrderHint)
+			{
+				SourceBuffer = InSourceBuffer;
+				SourceBuffer.MakeOwned();
+			}
+
+			virtual ~FWriteRequest() = default;
+
+			void PrepareSourceBufferAsync(FGraphEventRef CompletionEvent) override
+			{
+				CompletionEvent->DispatchSubsequents();
+			}
+
+			const FIoBuffer* GetSourceBuffer() override
+			{
+				return &SourceBuffer;
+			}
+
+			void FreeSourceBuffer() override
+			{
+			}
+
+			uint64 GetOrderHint() override
+			{
+				return OrderHint;
+			}
+
+			TArrayView<const FFileRegion> GetRegions()
+			{
+				return TArrayView<const FFileRegion>();
+			}
+
+			FIoBuffer SourceBuffer;
+			uint64 OrderHint;
+		};
+
+		Append(ChunkId, new FWriteRequest(Chunk, OrderHint), WriteOptions);
+	}
+
+	bool GeneratePerfectHashes(FIoStoreTocResource& TocResource, const TCHAR* ContainerDebugName)
+	{
+		// https://en.wikipedia.org/wiki/Perfect_hash_function
+		TRACE_CPUPROFILER_EVENT_SCOPE(TocGeneratePerfectHashes);
+		uint32 ChunkCount = TocResource.ChunkIds.Num();
+		uint32 SeedCount = FMath::Max(1, FMath::RoundToInt32(ChunkCount / 2.0));
+		check(TocResource.ChunkOffsetLengths.Num() == ChunkCount);
+		
+		TArray<FIoChunkId> OutTocChunkIds;
+		OutTocChunkIds.SetNum(ChunkCount);
+		TArray<FIoOffsetAndLength> OutTocOffsetAndLengths;
+		OutTocOffsetAndLengths.SetNum(ChunkCount);
+		TArray<FIoStoreTocEntryMeta> OutTocChunkMetas;
+		OutTocChunkMetas.SetNum(ChunkCount);
+		TArray<int32> OutTocChunkHashSeeds;
+		OutTocChunkHashSeeds.SetNumZeroed(SeedCount);
+		TArray<int32> OutTocChunkIndicesWithoutPerfectHash;
+
+		TArray<TArray<int32>> Buckets;
+		Buckets.SetNum(SeedCount);
+
+		TBitArray<> FreeSlots(true, ChunkCount);
+		// Put each chunk in a bucket, each bucket contains the chunk ids that have colliding hashes
+		for (uint32 ChunkIndex = 0; ChunkIndex < ChunkCount; ++ChunkIndex)
+		{
+			const FIoChunkId& ChunkId = TocResource.ChunkIds[ChunkIndex];
+			Buckets[FIoStoreTocResource::HashChunkIdWithSeed(0, ChunkId) % SeedCount].Add(ChunkIndex);
+		}
+
+		uint64 TotalIterationCount = 0;
+		uint64 TotalOverflowBucketsCount = 0;
+		
+		// For each bucket containing more than one chunk id find a seed that makes its chunk ids
+		// hash to unused slots in the output array
+		Algo::Sort(Buckets, [](const TArray<int32>& A, const TArray<int32>& B)
+			{
+				return A.Num() > B.Num();
+			});
+		for (uint32 BucketIndex = 0; BucketIndex < SeedCount; ++BucketIndex)
+		{
+			const TArray<int32>& Bucket = Buckets[BucketIndex];
+			if (Bucket.Num() <= 1)
+			{
+				break;
+			}
+			uint64 BucketHash = FIoStoreTocResource::HashChunkIdWithSeed(0, TocResource.ChunkIds[Bucket[0]]);
+
+			static constexpr uint32 Primes[] = {
+				2, 3, 5, 7, 11, 13, 17, 19, 23, 29, 31, 37, 41, 43, 47, 53, 59, 61, 67, 71, 73, 79,
+				83, 89, 97, 101, 103, 107, 109, 113, 127, 131, 137, 139, 149, 151, 157, 163, 167,
+				173, 179, 181, 191, 193, 197, 199, 211, 223, 227, 229, 233, 239, 241, 251, 257, 263,
+				269, 271, 277, 281, 283, 293, 307, 311, 313, 317, 331, 337, 347, 349, 353, 359, 367,
+				373, 379, 383, 389, 397, 401, 409, 419, 421, 431, 433, 439, 443, 449, 457, 461, 463,
+				467, 479, 487, 491, 499, 503, 509, 521, 523, 541, 547, 557, 563, 569, 571, 577, 587,
+				593, 599, 601, 607, 613, 617, 619, 631, 641, 643, 647, 653, 659, 661, 673, 677, 683,
+				691, 701, 709, 719, 727, 733, 739, 743, 751, 757, 761, 769, 773, 787, 797, 809, 811,
+				821, 823, 827, 829, 839, 853, 857, 859, 863, 877, 881, 883, 887, 907, 911, 919, 929,
+				937, 941, 947, 953, 967, 971, 977, 983, 991, 997, 1009, 1013, 1019, 1021, 1031, 1033,
+				1039, 1049, 1051, 1061, 1063, 1069, 1087, 1091, 1093, 1097, 1103, 1109, 1117, 1123,
+				1129, 1151, 1153, 1163, 1171, 1181, 1187, 1193, 1201, 1213, 1217, 1223, 1229, 1231,
+				1237, 1249, 1259, 1277, 1279, 1283, 1289, 1291, 1297, 1301, 1303, 1307, 1319, 1321,
+				1327, 1361, 1367, 1373, 1381, 1399, 1409, 1423, 1427, 1429, 1433, 1439, 1447, 1451,
+				1453, 1459, 1471, 1481, 1483, 1487, 1489, 1493, 1499, 1511, 1523, 1531, 1543, 1549,
+				1553, 1559, 1567, 1571, 1579, 1583, 1597, 1601, 1607, 1609, 1613, 1619, 1621, 1627,
+				1637, 1657, 1663, 1667, 1669, 1693, 1697, 1699, 1709, 1721, 1723, 1733, 1741, 1747,
+				1753, 1759, 1777, 1783, 1787, 1789, 1801, 1811, 1823, 1831, 1847, 1861, 1867, 1871,
+				1873, 1877, 1879, 1889, 1901, 1907, 1913, 1931, 1933, 1949, 1951, 1973, 1979, 1987,
+				1993, 1997, 1999, 2003, 2011, 2017, 2027, 2029, 2039, 2053, 2063, 2069, 2081, 2083,
+				2087, 2089, 2099, 2111, 2113, 2129, 2131, 2137, 2141, 2143, 2153, 2161, 2179, 2203,
+				2207, 2213, 2221, 2237, 2239, 2243, 2251, 2267, 2269, 2273, 2281, 2287, 2293, 2297,
+				2309, 2311, 2333, 2339, 2341, 2347, 2351, 2357, 2371, 2377, 2381, 2383, 2389, 2393,
+				2399, 2411, 2417, 2423, 2437, 2441, 2447, 2459, 2467, 2473, 2477, 2503, 2521, 2531,
+				2539, 2543, 2549, 2551, 2557, 2579, 2591, 2593, 2609, 2617, 2621, 2633, 2647, 2657,
+				2659, 2663, 2671, 2677, 2683, 2687, 2689, 2693, 2699, 2707, 2711, 2713, 2719, 2729,
+				2731, 2741, 2749, 2753, 2767, 2777, 2789, 2791, 2797, 2801, 2803, 2819, 2833, 2837,
+				2843, 2851, 2857, 2861, 2879, 2887, 2897, 2903, 2909, 2917, 2927, 2939, 2953, 2957,
+				2963, 2969, 2971, 2999, 3001, 3011, 3019, 3023, 3037, 3041, 3049, 3061, 3067, 3079,
+				3083, 3089, 3109, 3119, 3121, 3137, 3163, 3167, 3169, 3181, 3187, 3191, 3203, 3209,
+				3217, 3221, 3229, 3251, 3253, 3257, 3259, 3271, 3299, 3301, 3307, 3313, 3319, 3323,
+				3329, 3331, 3343, 3347, 3359, 3361, 3371, 3373, 3389, 3391, 3407, 3413, 3433, 3449,
+				3457, 3461, 3463, 3467, 3469, 3491, 3499, 3511, 3517, 3527, 3529, 3533, 3539, 3541,
+				3547, 3557, 3559, 3571, 3581, 3583, 3593, 3607, 3613, 3617, 3623, 3631, 3637, 3643,
+				3659, 3671, 3673, 3677, 3691, 3697, 3701, 3709, 3719, 3727, 3733, 3739, 3761, 3767,
+				3769, 3779, 3793, 3797, 3803, 3821, 3823, 3833, 3847, 3851, 3853, 3863, 3877, 3881,
+				3889, 3907, 3911, 3917, 3919, 3923, 3929, 3931, 3943, 3947, 3967, 3989, 4001, 4003,
+				4007, 4013, 4019, 4021, 4027, 4049, 4051, 4057, 4073, 4079, 4091, 4093, 4099, 4111,
+				4127, 4129, 4133, 4139, 4153, 4157, 4159, 4177, 4201, 4211, 4217, 4219, 4229, 4231,
+				4241, 4243, 4253, 4259, 4261, 4271, 4273, 4283, 4289, 4297, 4327, 4337, 4339, 4349,
+				4357, 4363, 4373, 4391, 4397, 4409, 4421, 4423, 4441, 4447, 4451, 4457, 4463, 4481,
+				4483, 4493, 4507, 4513, 4517, 4519, 4523, 4547, 4549, 4561, 4567, 4583, 4591, 4597,
+				4603, 4621, 4637, 4639, 4643, 4649, 4651, 4657, 4663, 4673, 4679, 4691, 4703, 4721,
+				4723, 4729, 4733, 4751, 4759, 4783, 4787, 4789, 4793, 4799, 4801, 4813, 4817, 4831,
+				4861, 4871, 4877, 4889, 4903, 4909, 4919, 4931, 4933, 4937, 4943, 4951, 4957, 4967,
+				4969, 4973, 4987, 4993, 4999, 5003, 5009, 5011, 5021, 5023, 5039, 5051, 5059, 5077,
+				5081, 5087, 5099, 5101, 5107, 5113, 5119, 5147, 5153, 5167, 5171, 5179, 5189, 5197,
+				5209, 5227, 5231, 5233, 5237, 5261, 5273, 5279, 5281, 5297, 5303, 5309, 5323, 5333,
+				5347, 5351, 5381, 5387, 5393, 5399, 5407, 5413, 5417, 5419, 5431, 5437, 5441, 5443,
+				5449, 5471, 5477, 5479, 5483, 5501, 5503, 5507, 5519, 5521, 5527, 5531, 5557, 5563,
+				5569, 5573, 5581, 5591, 5623, 5639, 5641, 5647, 5651, 5653, 5657, 5659, 5669, 5683,
+				5689, 5693, 5701, 5711, 5717, 5737, 5741, 5743, 5749, 5779, 5783, 5791, 5801, 5807,
+				5813, 5821, 5827, 5839, 5843, 5849, 5851, 5857, 5861, 5867, 5869, 5879, 5881, 5897,
+				5903, 5923, 5927, 5939, 5953, 5981, 5987, 6007, 6011, 6029, 6037, 6043, 6047, 6053,
+				6067, 6073, 6079, 6089, 6091, 6101, 6113, 6121, 6131, 6133, 6143, 6151, 6163, 6173,
+				6197, 6199, 6203, 6211, 6217, 6221, 6229, 6247, 6257, 6263, 6269, 6271, 6277, 6287,
+				6299, 6301, 6311, 6317, 6323, 6329, 6337, 6343, 6353, 6359, 6361, 6367, 6373, 6379,
+				6389, 6397, 6421, 6427, 6449, 6451, 6469, 6473, 6481, 6491, 6521, 6529, 6547, 6551,
+				6553, 6563, 6569, 6571, 6577, 6581, 6599, 6607, 6619, 6637, 6653, 6659, 6661, 6673,
+				6679, 6689, 6691, 6701, 6703, 6709, 6719, 6733, 6737, 6761, 6763, 6779, 6781, 6791,
+				6793, 6803, 6823, 6827, 6829, 6833, 6841, 6857, 6863, 6869, 6871, 6883, 6899, 6907,
+				6911, 6917, 6947, 6949, 6959, 6961, 6967, 6971, 6977, 6983, 6991, 6997, 7001, 7013,
+				7019, 7027, 7039, 7043, 7057, 7069, 7079, 7103, 7109, 7121, 7127, 7129, 7151, 7159,
+				7177, 7187, 7193, 7207, 7211, 7213, 7219, 7229, 7237, 7243, 7247, 7253, 7283, 7297,
+				7307, 7309, 7321, 7331, 7333, 7349, 7351, 7369, 7393, 7411, 7417, 7433, 7451, 7457,
+				7459, 7477, 7481, 7487, 7489, 7499, 7507, 7517, 7523, 7529, 7537, 7541, 7547, 7549,
+				7559, 7561, 7573, 7577, 7583, 7589, 7591, 7603, 7607, 7621, 7639, 7643, 7649, 7669,
+				7673, 7681, 7687, 7691, 7699, 7703, 7717, 7723, 7727, 7741, 7753, 7757, 7759, 7789,
+				7793, 7817, 7823, 7829, 7841, 7853, 7867, 7873, 7877, 7879, 7883, 7901, 7907, 7919
+			};
+			static constexpr uint32 MaxIterations = UE_ARRAY_COUNT(Primes);
+
+			uint32 PrimeIndex = 0;
+			TBitArray<> BucketUsedSlots(false, ChunkCount);
+			int32 IndexInBucket = 0;
+			bool bFoundSeedForBucket = true;
+			uint64 BucketIterationCount = 0;
+			while (IndexInBucket < Bucket.Num())
+			{
+				++BucketIterationCount;
+				const FIoChunkId& ChunkId = TocResource.ChunkIds[Bucket[IndexInBucket]];
+				uint32 Seed = Primes[PrimeIndex];
+				uint32 Slot = FIoStoreTocResource::HashChunkIdWithSeed(Seed, ChunkId) % ChunkCount;
+				if (!FreeSlots[Slot] || BucketUsedSlots[Slot])
+				{
+					++PrimeIndex;
+					if (PrimeIndex == MaxIterations)
+					{
+						// Unable to resolve collisions for this bucket, put items in the overflow list and
+						// save the negative index of the first item in the bucket as the seed
+						// (-ChunkCount - 1 to separate from the single item buckets below)
+						UE_LOG(LogIoStore, Verbose, TEXT("%s: Failed finding seed for bucket with %d items after %d iterations."), ContainerDebugName, Bucket.Num(), BucketIterationCount);
+						bFoundSeedForBucket = false;
+						OutTocChunkHashSeeds[BucketHash % SeedCount] = -OutTocChunkIndicesWithoutPerfectHash.Num() - ChunkCount - 1;
+						OutTocChunkIndicesWithoutPerfectHash.Append(Bucket);
+						++TotalOverflowBucketsCount;
+						break;
+
+					}
+					IndexInBucket = 0;
+					BucketUsedSlots.Init(false, ChunkCount);
+				}
+				else
+				{
+					BucketUsedSlots[Slot] = true;
+					++IndexInBucket;
+				}
+			}
+
+			TotalIterationCount += BucketIterationCount;
+
+			if (bFoundSeedForBucket)
+			{
+				uint32 Seed = Primes[PrimeIndex];
+				OutTocChunkHashSeeds[BucketHash % SeedCount] = Seed;
+				for (IndexInBucket = 0; IndexInBucket < Bucket.Num(); ++IndexInBucket)
+				{
+					int32 ChunkIndex = Bucket[IndexInBucket];
+					const FIoChunkId& ChunkId = TocResource.ChunkIds[ChunkIndex];
+					uint32 Slot = FIoStoreTocResource::HashChunkIdWithSeed(Seed, ChunkId) % ChunkCount;
+					check(FreeSlots[Slot]);
+					FreeSlots[Slot] = false;
+					OutTocChunkIds[Slot] = ChunkId;
+					OutTocOffsetAndLengths[Slot] = TocResource.ChunkOffsetLengths[ChunkIndex];
+					OutTocChunkMetas[Slot] = TocResource.ChunkMetas[ChunkIndex];
+				}
+			}
+		}
+
+		// For the remaining buckets with only one chunk id put that chunk id in the first empty position in
+		// the output array and store the index as a negative seed for the bucket (-1 to allow use of slot 0)
+		TConstSetBitIterator<> FreeSlotIt(FreeSlots);
+		for (uint32 BucketIndex = 0; BucketIndex < SeedCount; ++BucketIndex)
+		{
+			const TArray<int32>& Bucket = Buckets[BucketIndex];
+			if (Bucket.Num() == 1)
+			{
+				uint32 Slot = FreeSlotIt.GetIndex();
+				++FreeSlotIt;
+				int32 ChunkIndex = Bucket[0];
+				const FIoChunkId& ChunkId = TocResource.ChunkIds[ChunkIndex];
+				uint64 BucketHash = FIoStoreTocResource::HashChunkIdWithSeed(0, ChunkId);
+				OutTocChunkHashSeeds[BucketHash % SeedCount] = -static_cast<int32>(Slot) - 1;
+				OutTocChunkIds[Slot] = ChunkId;
+				OutTocOffsetAndLengths[Slot] = TocResource.ChunkOffsetLengths[ChunkIndex];
+				OutTocChunkMetas[Slot] = TocResource.ChunkMetas[ChunkIndex];
+			}
+		}
+
+		if (!OutTocChunkIndicesWithoutPerfectHash.IsEmpty())
+		{
+			// Put overflow items in the remaining free slots and update the index for each overflow entry
+			UE_LOG(LogIoStore, Display, TEXT("%s: Failed finding perfect hashmap for %d items. %d overflow buckets with %d items."), ContainerDebugName, ChunkCount, TotalOverflowBucketsCount, OutTocChunkIndicesWithoutPerfectHash.Num());
+			for (int32& OverflowEntryIndex : OutTocChunkIndicesWithoutPerfectHash)
+			{
+				uint32 Slot = FreeSlotIt.GetIndex();
+				++FreeSlotIt;
+				const FIoChunkId& ChunkId = TocResource.ChunkIds[OverflowEntryIndex];
+				OutTocChunkIds[Slot] = ChunkId;
+				OutTocOffsetAndLengths[Slot] = TocResource.ChunkOffsetLengths[OverflowEntryIndex];
+				OutTocChunkMetas[Slot] = TocResource.ChunkMetas[OverflowEntryIndex];
+				OverflowEntryIndex = Slot;
+			}
+		}
+		else
+		{
+			UE_LOG(LogIoStore, Display, TEXT("%s: Found perfect hashmap for %d items."), ContainerDebugName, ChunkCount);
+		}
+		double AverageIterationCount = ChunkCount > 0 ? static_cast<double>(TotalIterationCount) / ChunkCount : 0.0;
+		UE_LOG(LogIoStore, Verbose, TEXT("%s: %f iterations/chunk"), ContainerDebugName, AverageIterationCount);
+
+		TocResource.ChunkIds = MoveTemp(OutTocChunkIds);
+		TocResource.ChunkOffsetLengths = MoveTemp(OutTocOffsetAndLengths);
+		TocResource.ChunkMetas = MoveTemp(OutTocChunkMetas);
+		TocResource.ChunkPerfectHashSeeds = MoveTemp(OutTocChunkHashSeeds);
+		TocResource.ChunkIndicesWithoutPerfectHash = MoveTemp(OutTocChunkIndicesWithoutPerfectHash);
+
+		return true;
+	}
+
+	void Finalize()
+	{
+		check(bHasFlushed);
+
+		UncompressedContainerSize = TotalEntryUncompressedSize + TotalPaddingSize;
+		CompressedContainerSize = 0;
+		const FIoStoreWriterSettings& WriterSettings = WriterContext->GetSettings();
+		for (FPartition& Partition : Partitions)
+		{
+			CompressedContainerSize += Partition.Offset;
+
+			if (bHasMemoryMappedEntry)
+			{
+				uint64 ExtraPaddingBytes = Align(Partition.Offset, WriterSettings.MemoryMappingAlignment) - Partition.Offset;
+				if (ExtraPaddingBytes)
+				{
+					TArray<uint8> Padding;
+					Padding.SetNumZeroed(int32(ExtraPaddingBytes));
+					Partition.ContainerFileHandle->Serialize(Padding.GetData(), ExtraPaddingBytes);
+					CompressedContainerSize += ExtraPaddingBytes;
+					UncompressedContainerSize += ExtraPaddingBytes;
+					Partition.Offset += ExtraPaddingBytes;
+					TotalPaddingSize += ExtraPaddingBytes;
+				}
+			}
+			
+			if (Partition.ContainerFileHandle)
+			{
+				TRACE_CPUPROFILER_EVENT_SCOPE(FlushContainerFile);
+				Partition.ContainerFileHandle->Flush();
+				check(Partition.ContainerFileHandle->Tell() == Partition.Offset);
+			}
+
+			if (Partition.RegionsArchive)
+			{
+				FFileRegion::SerializeFileRegions(*Partition.RegionsArchive.Get(), Partition.AllFileRegions);
+				Partition.RegionsArchive->Flush();
+			}
+		}
+
+		FIoStoreTocResource& TocResource = TocBuilder.GetTocResource();
+
+		GeneratePerfectHashes(TocResource, *FPaths::GetBaseFilename(TocFilePath));
+
+		if (ContainerSettings.IsIndexed())
+		{
+			TRACE_CPUPROFILER_EVENT_SCOPE(BuildIndex);
+			TArray<FStringView> FilesToIndex;
+			TocBuilder.GetFileNamesToIndex(FilesToIndex);
+
+			FString MountPoint = IoDirectoryIndexUtils::GetCommonRootPath(FilesToIndex);
+			FIoDirectoryIndexWriter DirectoryIndexWriter;
+			DirectoryIndexWriter.SetMountPoint(MountPoint);
+
+			uint32 TocEntryIndex = 0;
+			for (const FIoChunkId& ChunkId : TocResource.ChunkIds)
+			{
+				const FString* ChunkFileName = TocBuilder.GetFileName(ChunkId);
+				if (ChunkFileName)
+				{
+					const uint32 FileEntryIndex = DirectoryIndexWriter.AddFile(*ChunkFileName);
+					check(FileEntryIndex != ~uint32(0));
+					DirectoryIndexWriter.SetFileUserData(FileEntryIndex, TocEntryIndex);
+				}
+				++TocEntryIndex;
+			}
+
+			DirectoryIndexWriter.Flush(
+				TocResource.DirectoryIndexBuffer,
+				ContainerSettings.IsEncrypted() ? ContainerSettings.EncryptionKey : FAES::FAESKey());
+		}
+
+		TIoStatusOr<uint64> TocSize = FIoStoreTocResource::Write(*TocFilePath, TocResource, static_cast<uint32>(WriterSettings.CompressionBlockSize), WriterSettings.MaxPartitionSize, ContainerSettings);
+		check(TocSize.IsOk());
+		
+		Result.ContainerId = ContainerSettings.ContainerId;
+		Result.ContainerName = FPaths::GetBaseFilename(TocFilePath);
+		Result.ContainerFlags = ContainerSettings.ContainerFlags;
+		Result.TocSize = TocSize.ConsumeValueOrDie();
+		Result.TocEntryCount = TocResource.Header.TocEntryCount;
+		Result.PaddingSize = TotalPaddingSize;
+		Result.UncompressedContainerSize = UncompressedContainerSize;
+		Result.CompressedContainerSize = CompressedContainerSize;
+		Result.TotalEntryCompressedSize = TotalEntryCompressedSize;
+		Result.ReferenceCacheMissBytes = ReferenceCacheMissBytes;
+		Result.DirectoryIndexSize = TocResource.Header.DirectoryIndexSize;
+		Result.CompressionMethod = EnumHasAnyFlags(ContainerSettings.ContainerFlags, EIoContainerFlags::Compressed)
+			? WriterSettings.CompressionMethod
+			: NAME_None;
+		Result.ModifiedChunksCount = 0;
+		Result.AddedChunksCount = 0;
+		Result.ModifiedChunksSize= 0;
+		Result.AddedChunksSize = 0;
+		{
+			TRACE_CPUPROFILER_EVENT_SCOPE(Cleanup);
+			for (FIoStoreWriteQueueEntry* Entry : Entries)
+			{
+				if (Entry->bModified)
+				{
+					++Result.ModifiedChunksCount;
+					Result.ModifiedChunksSize += Entry->CompressedSize;
+				}
+				else if (Entry->bAdded)
+				{
+					++Result.AddedChunksCount;
+					Result.AddedChunksSize += Entry->CompressedSize;
+				}
+				delete Entry;
+			}
+		}
+
+		Entries.Empty();
+		bHasResult = true;
+	}
+
+	TIoStatusOr<FIoStoreWriterResult> GetResult()
+	{
+		if (!bHasResult)
+		{
+			return FIoStatus::Invalid;
+		}
+		return Result;
+	}
+
+private:
+	struct FPartition
+	{
+		TUniquePtr<FArchive> ContainerFileHandle;
+		TUniquePtr<FArchive> RegionsArchive;
+		uint64 Offset = 0;
+		uint64 ReservedSpace = 0;
+		TArray<FFileRegion> AllFileRegions;
+		int32 Index = -1;
+	};
+
+	struct FLayoutEntry
+	{
+		FLayoutEntry* Prev = nullptr;
+		FLayoutEntry* Next = nullptr;
+		uint64 IdealOrder = 0;
+		uint64 CompressedSize = uint64(-1);
+		FIoChunkHash Hash;
+		FIoStoreWriteQueueEntry* QueueEntry = nullptr;
+		int32 PartitionIndex = -1;
+	};
+
+	void FinalizeLayout()
+	{
+		TRACE_CPUPROFILER_EVENT_SCOPE(FinalizeLayout);
+		
+		Algo::Sort(Entries, [](const FIoStoreWriteQueueEntry* A, const FIoStoreWriteQueueEntry* B)
+		{
+			uint64 AOrderHint = A->Request->GetOrderHint();
+			uint64 BOrderHint = B->Request->GetOrderHint();
+			if (AOrderHint != BOrderHint)
+			{
+				return AOrderHint < BOrderHint;
+			}
+			return A->Sequence < B->Sequence;
+		});
+
+		TMap<int64, FLayoutEntry*> LayoutEntriesByOrderMap;
+		int64 IdealOrder = 0;
+		TArray<FLayoutEntry*> UnassignedEntries;
+		for (FIoStoreWriteQueueEntry* WriteQueueEntry : Entries)
+		{
+			FLayoutEntry* FindPreviousEntry = PreviousBuildLayoutEntryByChunkId.FindRef(WriteQueueEntry->ChunkId);
+			if (FindPreviousEntry)
+			{
+				if (FindPreviousEntry->Hash != WriteQueueEntry->ChunkHash)
+				{
+					WriteQueueEntry->bModified = true;
+				}
+				else
+				{
+					FindPreviousEntry->QueueEntry = WriteQueueEntry;
+					FindPreviousEntry->IdealOrder = IdealOrder;
+					WriteQueueEntry->PartitionIndex = FindPreviousEntry->PartitionIndex;
+				}
+			}
+			else
+			{
+				WriteQueueEntry->bAdded = true;
+			}
+			if (WriteQueueEntry->bModified || WriteQueueEntry->bAdded)
+			{
+				FLayoutEntry* NewLayoutEntry = new FLayoutEntry();
+				NewLayoutEntry->QueueEntry = WriteQueueEntry;
+				NewLayoutEntry->IdealOrder = IdealOrder;
+				LayoutEntries.Add(NewLayoutEntry);
+				UnassignedEntries.Add(NewLayoutEntry);
+			}
+			++IdealOrder;
+		}
+			
+		if (ContainerSettings.bGenerateDiffPatch)
+		{
+			LayoutEntriesHead->Next = LayoutEntriesTail;
+			LayoutEntriesTail->Prev = LayoutEntriesHead;
+		}
+		else
+		{
+			for (FLayoutEntry* EntryIt = LayoutEntriesHead->Next; EntryIt != LayoutEntriesTail; EntryIt = EntryIt->Next)
+			{
+				if (!EntryIt->QueueEntry)
+				{
+					EntryIt->Prev->Next = EntryIt->Next;
+					EntryIt->Next->Prev = EntryIt->Prev;
+				}
+				else
+				{
+					LayoutEntriesByOrderMap.Add(EntryIt->IdealOrder, EntryIt);
+				}
+			}
+		}
+		FLayoutEntry* LastAddedEntry = LayoutEntriesHead;
+		for (FLayoutEntry* UnassignedEntry : UnassignedEntries)
+		{
+			check(UnassignedEntry->QueueEntry);
+			FLayoutEntry* PutAfterEntry = LayoutEntriesByOrderMap.FindRef(UnassignedEntry->IdealOrder - 1);
+			if (!PutAfterEntry)
+			{
+				PutAfterEntry = LastAddedEntry;
+			}
+
+			UnassignedEntry->Prev = PutAfterEntry;
+			UnassignedEntry->Next = PutAfterEntry->Next;
+			PutAfterEntry->Next->Prev = UnassignedEntry;
+			PutAfterEntry->Next = UnassignedEntry;
+			LayoutEntriesByOrderMap.Add(UnassignedEntry->IdealOrder, UnassignedEntry);
+			LastAddedEntry = UnassignedEntry;
+		}
+
+		TArray<FIoStoreWriteQueueEntry*> IncludedQueueEntries;
+		for (FLayoutEntry* EntryIt = LayoutEntriesHead->Next; EntryIt != LayoutEntriesTail; EntryIt = EntryIt->Next)
+		{
+			check(EntryIt->QueueEntry);
+			IncludedQueueEntries.Add(EntryIt->QueueEntry);
+			int32 ReserveInPartitionIndex = EntryIt->QueueEntry->PartitionIndex;
+			if (ReserveInPartitionIndex >= 0)
+			{
+				while (Partitions.Num() <= ReserveInPartitionIndex)
+				{
+					FPartition& NewPartition = Partitions.AddDefaulted_GetRef();
+					NewPartition.Index = Partitions.Num() - 1;
+				}
+				FPartition& ReserveInPartition = Partitions[ReserveInPartitionIndex];
+				check(EntryIt->CompressedSize != uint64(-1));
+				ReserveInPartition.ReservedSpace += EntryIt->CompressedSize;
+			}
+		}
+		Swap(Entries, IncludedQueueEntries);
+
+		LayoutEntriesHead = nullptr;
+		LayoutEntriesTail = nullptr;
+		PreviousBuildLayoutEntryByChunkId.Empty();
+		for (FLayoutEntry* Entry : LayoutEntries)
+		{
+			delete Entry;
+		}
+		LayoutEntries.Empty();
+	}
+
+	FIoStatus CreatePartitionContainerFile(FPartition& Partition)
+	{
+		check(!Partition.ContainerFileHandle);
+		FString ContainerFilePath = ContainerPathAndBaseFileName;
+		if (Partition.Index > 0)
+		{
+			ContainerFilePath += FString::Printf(TEXT("_s%d"), Partition.Index);
+		}
+		ContainerFilePath += TEXT(".ucas");
+		
+		Partition.ContainerFileHandle.Reset(IFileManager::Get().CreateFileWriter(*ContainerFilePath));
+		if (!Partition.ContainerFileHandle)
+		{
+			return FIoStatusBuilder(EIoErrorCode::FileOpenFailed) << TEXT("Failed to open IoStore container file '") << *ContainerFilePath << TEXT("'");
+		}
+		if (WriterContext->GetSettings().bEnableFileRegions)
+		{
+			FString RegionsFilePath = ContainerFilePath + FFileRegion::RegionsFileExtension;
+			Partition.RegionsArchive.Reset(IFileManager::Get().CreateFileWriter(*RegionsFilePath));
+			if (!Partition.RegionsArchive)
+			{
+				return FIoStatusBuilder(EIoErrorCode::FileOpenFailed) << TEXT("Failed to open IoStore regions file '") << *RegionsFilePath << TEXT("'");
+			}
+		}
+
+		return FIoStatus::Ok;
+	}
+
+	void CompressBlock(FChunkBlock* Block)
+	{
+		TRACE_CPUPROFILER_EVENT_SCOPE(CompressBlock);
+		check(Block->CompressionMethod != NAME_None);
+		uint64 CompressedBlockSize = Block->IoBuffer->DataSize();
+		bool bCompressed;
+		{
+			if (!FCompression::CompressMemoryIfWorthDecompressing(
+				Block->CompressionMethod,
+				bCompressed,
+				(int64)WriterContext->WriterSettings.CompressionMinBytesSaved,
+				WriterContext->WriterSettings.CompressionMinPercentSaved,
+				Block->IoBuffer->Data(),
+				(int64&)CompressedBlockSize,
+				Block->UncompressedData,
+				(int64)Block->UncompressedSize,
+				COMPRESS_ForPackaging))
+			{
+				UE_LOG(LogIoStore, Error, TEXT("Compression failed: Method=%s, CompressedSize=0x%llx, UncompressedSize=0x%llx"), 
+					*Block->CompressionMethod.ToString(), CompressedBlockSize, Block->UncompressedSize);
+				bCompressed = false;
+			}
+		}
+		if (!bCompressed)
+		{
+			Block->CompressionMethod = NAME_None;
+			Block->CompressedSize = Block->UncompressedSize;
+			FMemory::Memcpy(Block->IoBuffer->Data(), Block->UncompressedData, Block->UncompressedSize);
+		}
+		else
+		{
+			check(CompressedBlockSize > 0);
+			check(CompressedBlockSize < Block->UncompressedSize);
+			Block->CompressedSize = CompressedBlockSize;
+		}
+	}
+
+	bool SerializeCompressedDDCData(FIoStoreWriteQueueEntry* Entry, FArchive& Ar)
+	{
+		uint32 NumBlocks = Entry->ChunkBlocks.Num();
+		Ar << NumBlocks;
+		if (NumBlocks != Entry->ChunkBlocks.Num())
+		{
+			return false;
+		}
+		for (FChunkBlock& Block : Entry->ChunkBlocks)
+		{
+			Ar << Block.CompressedSize;
+			if (Block.CompressedSize > Block.UncompressedSize)
+			{
+				return false;
+			}
+			if (Ar.IsLoading() && Block.CompressedSize == Block.UncompressedSize)
+			{
+				Block.CompressionMethod = NAME_None;
+			}
+			if (Block.IoBuffer->DataSize() < Block.CompressedSize)
+			{
+				return false;
+			}
+			Ar.Serialize(Block.IoBuffer->Data(), Block.CompressedSize);
+		}
+		return !Ar.IsError();
+	}
+
+	FName CompressionMethodForEntry(FIoStoreWriteQueueEntry* Entry) const
+	{
+		FName CompressionMethod = NAME_None;
+		const FIoStoreWriterSettings& WriterSettings = WriterContext->WriterSettings;
+		if (ContainerSettings.IsCompressed() && !Entry->Options.bForceUncompressed && !Entry->Options.bIsMemoryMapped)
+		{
+			CompressionMethod = WriterSettings.CompressionMethod;
+		}
+		return CompressionMethod;
+	}
+
+
+
+	void BeginCompress(FIoStoreWriteQueueEntry* Entry)
+	{
+		WriterContext->BeginCompressChunksByType[(int8)Entry->ChunkId.GetChunkType()].IncrementExchange();
+		
+		if (Entry->bLoadingFromReferenceDb)
+		{
+			if (Entry->NumChunkBlocksFromRefDb == 0)
+			{
+				Entry->FinishCompressionBarrier->DispatchSubsequents();
+				return;
+			}
+
+			// Allocate resources before launching the read tasks to prevent deadlock. Note this will
+			// allocate iobuffers big enough for uncompressed size, when we only actually need it for
+			// compressed size.
+			int32 NumChunkBlocks32 = IntCastChecked<int32>(Entry->NumChunkBlocksFromRefDb);
+			Entry->ChunkBlocks.SetNum(NumChunkBlocks32);
+			for (int32 BlockIndex = 0; BlockIndex < NumChunkBlocks32; ++BlockIndex)
+			{
+				FChunkBlock& Block = Entry->ChunkBlocks[BlockIndex];
+				Block.IoBuffer = WriterContext->AllocCompressionBuffer(NumChunkBlocks32);
+				// Everything else in a block gets filled out from the refdb.
+			}
+
+			// Valid chunks must create the same decompressed bits, but can have different compressed bits.
+			// Since we are on a lightweight dispatch thread, the actual read is async, as is the processing
+			// of the results.
+			bool bChunkExists = ReferenceChunkDatabase->RetrieveChunk(ContainerSettings.ContainerId, Entry->ChunkHash, Entry->ChunkId, [this, Entry](TIoStatusOr<FIoStoreCompressedReadResult> InReadResult)
+			{
+
+				// If we fail here, in order to recover we effectively need to re-kick this chunk's
+				// BeginCompress() as well as source buffer read... however, this is just a direct read and should only fail
+				// in catastrophic scenarios (loss of connection on a network drive?).
+				FIoStoreCompressedReadResult ReadResult = InReadResult.ValueOrDie();
+
+				uint64 TotalUncompressedSize = 0;
+				uint8* ReferenceData = ReadResult.IoBuffer.GetData();
+				uint64 TotalAlignedSize = 0;
+				for (int32 BlockIndex = 0; BlockIndex < ReadResult.Blocks.Num(); ++BlockIndex)
+				{
+					FIoStoreCompressedBlockInfo& ReferenceBlock = ReadResult.Blocks[BlockIndex];
+					FChunkBlock& Block = Entry->ChunkBlocks[BlockIndex];
+					Block.CompressionMethod = ReferenceBlock.CompressionMethod;
+					Block.CompressedSize = ReferenceBlock.CompressedSize;
+					Block.UncompressedSize = ReferenceBlock.UncompressedSize;
+					TotalUncompressedSize += ReferenceBlock.UncompressedSize;
+
+					// Future optimization: ReadCompressed returns the memory ready to encrypt in one
+					// large contiguous buffer (i.e. padded). We could use the FIoBuffer functionality of referencing a 
+					// sub block from a parent buffer, however this would mean that we need to add support
+					// for tracking the memory usage in order to remain within our prescribed limits. To do this
+					// requires releasing the entire chunk's memory at once after WriteEntry.
+					// As it stands, we temporarily use untracked memory in the ReadCompressed call (in RetrieveChunk),
+					// then immediately copy it to tracked memory. There's some waste as tracked memory is mod CompressionBlockSize
+					// and we are post compression, so with the average 50% compression rate, we're using double the memory
+					// we "could".
+					FMemory::Memcpy(Block.IoBuffer->GetData(), ReferenceData, Block.CompressedSize);
+					ReferenceData += ReferenceBlock.AlignedSize;
+					TotalAlignedSize += ReferenceBlock.AlignedSize;
+				}
+
+				if (TotalAlignedSize != ReadResult.IoBuffer.GetSize())
+				{
+					// If we hit this, we might have read garbage memory above! This is very bad.
+					UE_LOG(LogIoStore, Error, TEXT("Block aligned size does not match iobuffer source size! Blocks: %s source size: %s"),
+						*FText::AsNumber(TotalAlignedSize).ToString(),
+						*FText::AsNumber(ReadResult.IoBuffer.GetSize()).ToString());
+				}
+
+				Entry->UncompressedSize.Emplace(TotalUncompressedSize);
+				Entry->FinishCompressionBarrier->DispatchSubsequents();
+			});
+
+			check(bChunkExists); // Sanity - should never return false as we can only get here if ChunkExists() returns true.
+
+			WriterContext->RefDbChunksCount.IncrementExchange();
+			WriterContext->RefDbChunksByType[(int8)Entry->ChunkId.GetChunkType()].IncrementExchange();
+			// Lambda handles dispatch subsequents
+			return;
+		}
+
+		const FIoStoreWriterSettings& WriterSettings = WriterContext->WriterSettings;
+
+		FName CompressionMethod = CompressionMethodForEntry(Entry);
+
+		const FIoBuffer* SourceBuffer = Entry->Request->GetSourceBuffer();
+		Entry->UncompressedSize.Emplace(SourceBuffer->DataSize());
+
+		check(WriterSettings.CompressionBlockSize > 0);
+		const uint64 NumChunkBlocks64 = Align(Entry->UncompressedSize.GetValue(), WriterSettings.CompressionBlockSize) / WriterSettings.CompressionBlockSize;
+		if (NumChunkBlocks64 == 0)
+		{
+			Entry->FinishCompressionBarrier->DispatchSubsequents();
+			return;
+		}
+
+		int32 NumChunkBlocks = IntCastChecked<int32>(NumChunkBlocks64);
+		
+		Entry->ChunkBlocks.SetNum(NumChunkBlocks);
+		{
+			// We must allocate resources for our tasks up front to prevent resource deadlock.
+			uint64 BytesToProcess = Entry->UncompressedSize.GetValue();
+			const uint8* UncompressedData = SourceBuffer->Data();
+			for (int32 BlockIndex = 0; BlockIndex < NumChunkBlocks; ++BlockIndex)
+			{
+				FChunkBlock& Block = Entry->ChunkBlocks[BlockIndex];
+				Block.IoBuffer = WriterContext->AllocCompressionBuffer(NumChunkBlocks);
+				Block.CompressionMethod = CompressionMethod;
+				Block.UncompressedSize = FMath::Min(BytesToProcess, WriterSettings.CompressionBlockSize);
+				Block.UncompressedData = UncompressedData;
+				BytesToProcess -= Block.UncompressedSize;
+				UncompressedData += Block.UncompressedSize;
+			}
+		}
+
+		if (CompressionMethod == NAME_None)
+		{
+			for (FChunkBlock& Block : Entry->ChunkBlocks)
+			{
+				Block.CompressionMethod = NAME_None;
+				Block.CompressedSize = Block.UncompressedSize;
+				FMemory::Memcpy(Block.IoBuffer->Data(), Block.UncompressedData, Block.UncompressedSize);
+			}
+			Entry->FinishCompressionBarrier->DispatchSubsequents();
+			return;
+		}
+
+		bool bUseDDCCompression =
+			WriterContext->DDC &&
+			WriterSettings.bCompressionEnableDDC &&
+			Entry->UncompressedSize.GetValue() > WriterSettings.CompressionMinSizeToConsiderDDC;
+		if (bUseDDCCompression)
+		{
+			TStringBuilder<256> CacheKeySuffix;
+			CacheKeySuffix.Append(Entry->ChunkHash.ToString());
+			WriterSettings.CompressionMethod.AppendString(CacheKeySuffix);
+			CacheKeySuffix.Append(FCompression::GetCompressorDDCSuffix(WriterSettings.CompressionMethod));
+			CacheKeySuffix.Appendf(TEXT("%d_%d"),
+				WriterSettings.CompressionMinBytesSaved,
+				WriterSettings.CompressionMinPercentSaved);
+			Entry->DDCCacheKey = FDerivedDataCacheInterface::BuildCacheKey(
+				TEXT("IoStoreCompression"),
+				TEXT("985D3FAD-71D0-4758-A777-A910B49CC4BF"),
+				CacheKeySuffix.ToString());
+
+			TArray<uint8> DDCData;
+			TRACE_CPUPROFILER_EVENT_SCOPE(ReadFromDDC);
+			bool bFoundInDDC = false;
+			if (WriterContext->DDC->GetSynchronous(*Entry->DDCCacheKey, DDCData, Entry->Options.FileName))
+			{
+				FLargeMemoryReader DDCDataReader(DDCData.GetData(), DDCData.Num());
+				bFoundInDDC = SerializeCompressedDDCData(Entry, DDCDataReader);
+			}
+			if (bFoundInDDC)
+			{
+				WriterContext->CompressionDDCHitCount.IncrementExchange();
+				Entry->FinishCompressionBarrier->DispatchSubsequents();
+				return;
+			}
+			else
+			{
+				WriterContext->CompressionDDCMissCount.IncrementExchange();
+				Entry->bStoreCompressedDataInDDC = true;
+			}
+		}
+		
+		for (FChunkBlock& Block : Entry->ChunkBlocks)
+		{
+			WriterContext->ScheduledCompressionTasksCount.IncrementExchange();
+			FFunctionGraphTask::CreateAndDispatchWhenReady([this, Entry, BlockPtr = &Block]()
+			{
+				CompressBlock(BlockPtr);
+				WriterContext->ScheduledCompressionTasksCount.DecrementExchange();
+				int32 CompressedBlocksCount = Entry->CompressedBlocksCount.IncrementExchange();
+				if (CompressedBlocksCount + 1 == Entry->ChunkBlocks.Num())
+				{
+					WriterContext->CompressedChunksByType[(int8)Entry->ChunkId.GetChunkType()].IncrementExchange();
+					WriterContext->CompressedChunksCount.IncrementExchange();
+					Entry->FinishCompressionBarrier->DispatchSubsequents();
+				}
+			}, TStatId(), nullptr, ENamedThreads::AnyHiPriThreadNormalTask);
+		}
+	}
+
+	void BeginEncryptAndSign(FIoStoreWriteQueueEntry* Entry)
+	{
+		if (ContainerSettings.IsEncrypted() || ContainerSettings.IsSigned())
+		{
+			FFunctionGraphTask::CreateAndDispatchWhenReady([this, Entry]()
+			{
+				EncryptAndSign(Entry);
+				Entry->FinishEncryptionAndSigningBarrier->DispatchSubsequents();
+			}, TStatId(), nullptr, ENamedThreads::AnyHiPriThreadHiPriTask);
+		}
+		else
+		{
+			EncryptAndSign(Entry);
+			Entry->FinishEncryptionAndSigningBarrier->DispatchSubsequents();
+		}
+	}
+
+	void EncryptAndSign(FIoStoreWriteQueueEntry* Entry)
+	{
+		TRACE_CPUPROFILER_EVENT_SCOPE(EncryptAndSignChunk);
+		for (FChunkBlock& Block : Entry->ChunkBlocks)
+		{
+			// Always align each compressed block to AES block size but store the compressed block size in the TOC
+			Block.Size = Block.CompressedSize;
+			if (!IsAligned(Block.Size, FAES::AESBlockSize))
+			{
+				uint64 AlignedCompressedBlockSize = Align(Block.Size, FAES::AESBlockSize);
+				uint8* CompressedData = Block.IoBuffer->Data();
+				for (uint64 FillIndex = Block.Size; FillIndex < AlignedCompressedBlockSize; ++FillIndex)
+				{
+					check(FillIndex < Block.IoBuffer->DataSize());
+					CompressedData[FillIndex] = CompressedData[(FillIndex - Block.Size) % Block.Size];
+				}
+				Block.Size = AlignedCompressedBlockSize;
+			}
+
+			if (ContainerSettings.IsEncrypted())
+			{
+				FAES::EncryptData(Block.IoBuffer->Data(), static_cast<uint32>(Block.Size), ContainerSettings.EncryptionKey);
+			}
+
+			if (ContainerSettings.IsSigned())
+			{
+				FSHA1::HashBuffer(Block.IoBuffer->Data(), Block.Size, Block.Signature.Hash);
+			}
+		}
+	}
+
+	void WriteEntry(FIoStoreWriteQueueEntry* Entry)
+	{
+		TRACE_CPUPROFILER_EVENT_SCOPE(WriteEntry);
+		const int32* FindExistingIndex = TocBuilder.GetTocEntryIndex(Entry->ChunkId);
+		if (FindExistingIndex)
+		{
+			// afaict this should never happen so add a warning. If there's a legit reason for it
+			// we can pull this back out. If would violate some assumptions in the reference chunk
+			// database if we DO hit this, however...
+			UE_LOG(LogIoStore, Warning, TEXT("ChunkId was added twice in container %s, %s, file %s hash %s vs %s"), 
+				*FPaths::GetBaseFilename(TocFilePath),
+				*LexToString(Entry->ChunkId), *Entry->Options.FileName,
+				*TocBuilder.GetTocResource().ChunkMetas[*FindExistingIndex].ChunkHash.ToString(),
+				*Entry->ChunkHash.ToString()
+				);
+
+			checkf(TocBuilder.GetTocResource().ChunkMetas[*FindExistingIndex].ChunkHash == Entry->ChunkHash, TEXT("Chunk id has already been added with different content"));
+			for (FChunkBlock& ChunkBlock : Entry->ChunkBlocks)
+			{
+				WriterContext->FreeCompressionBuffer(ChunkBlock.IoBuffer, Entry->ChunkBlocks.Num());
+				ChunkBlock.IoBuffer = nullptr;
+			}
+			return;
+		}
+
+		FPartition* TargetPartition = &Partitions[CurrentPartitionIndex];
+		int32 NextPartitionIndexToTry = CurrentPartitionIndex + 1;
+		if (Entry->PartitionIndex >= 0)
+		{
+			TargetPartition = &Partitions[Entry->PartitionIndex];
+			if (TargetPartition->ReservedSpace > Entry->CompressedSize)
+			{
+				TargetPartition->ReservedSpace -= Entry->CompressedSize;
+			}
+			else
+			{
+				TargetPartition->ReservedSpace = 0;
+			}
+			NextPartitionIndexToTry = CurrentPartitionIndex;
+		}
+
+		const FIoStoreWriterSettings& WriterSettings = WriterContext->WriterSettings;
+		bHasMemoryMappedEntry |= Entry->Options.bIsMemoryMapped;
+		const uint64 ChunkAlignment = Entry->Options.bIsMemoryMapped ? WriterSettings.MemoryMappingAlignment : 0;
+		const uint64 PartitionSizeLimit = WriterSettings.MaxPartitionSize > 0 ? WriterSettings.MaxPartitionSize : MAX_uint64;
+		checkf(Entry->CompressedSize <= PartitionSizeLimit, TEXT("Chunk is too large, increase max partition size!"));
+		for (;;)
+		{
+			uint64 OffsetBeforePadding = TargetPartition->Offset;
+			if (ChunkAlignment)
+			{
+				TargetPartition->Offset = Align(TargetPartition->Offset, ChunkAlignment);
+			}
+			if (WriterSettings.CompressionBlockAlignment)
+			{
+				// Try and prevent entries from crossing compression alignment blocks if possible. This is to avoid
+				// small entries from causing multiple file system block reads afaict. Large entries necesarily get
+				// aligned to prevent things like a blocksize + 2 entry being at alignment -1, causing 3 low level reads.
+				// ...I think.
+				bool bCrossesBlockBoundary = Align(TargetPartition->Offset, WriterSettings.CompressionBlockAlignment) != Align(TargetPartition->Offset + Entry->CompressedSize - 1, WriterSettings.CompressionBlockAlignment);
+				if (bCrossesBlockBoundary)
+				{
+					TargetPartition->Offset = Align(TargetPartition->Offset, WriterSettings.CompressionBlockAlignment);
+				}
+			}
+
+			if (TargetPartition->Offset + Entry->CompressedSize + TargetPartition->ReservedSpace > PartitionSizeLimit)
+			{
+				TargetPartition->Offset = OffsetBeforePadding;
+				while (Partitions.Num() <= NextPartitionIndexToTry)
+				{
+					FPartition& NewPartition = Partitions.AddDefaulted_GetRef();
+					NewPartition.Index = Partitions.Num() - 1;
+				}
+				CurrentPartitionIndex = NextPartitionIndexToTry;
+				TargetPartition = &Partitions[CurrentPartitionIndex];
+				++NextPartitionIndexToTry;
+			}
+			else
+			{
+				Entry->Padding = TargetPartition->Offset - OffsetBeforePadding;
+				TotalPaddingSize += Entry->Padding;
+				break;
+			}
+		}
+
+		if (!TargetPartition->ContainerFileHandle)
+		{
+			TRACE_CPUPROFILER_EVENT_SCOPE(CreatePartitionContainerFile);
+			CreatePartitionContainerFile(*TargetPartition);
+		}
+		Entry->Offset = TargetPartition->Offset;
+
+		FIoOffsetAndLength OffsetLength;
+		OffsetLength.SetOffset(UncompressedFileOffset);
+		OffsetLength.SetLength(Entry->UncompressedSize.GetValue());
+
+		FIoStoreTocEntryMeta ChunkMeta{ Entry->ChunkHash, FIoStoreTocEntryMetaFlags::None };
+		if (Entry->Options.bIsMemoryMapped)
+		{
+			ChunkMeta.Flags |= FIoStoreTocEntryMetaFlags::MemoryMapped;
+		}
+
+		uint64 OffsetInChunk = 0;
+		for (const FChunkBlock& ChunkBlock : Entry->ChunkBlocks)
+		{
+			FIoStoreTocCompressedBlockEntry& BlockEntry = TocBuilder.AddCompressionBlockEntry();
+			BlockEntry.SetOffset(TargetPartition->Index * WriterSettings.MaxPartitionSize + TargetPartition->Offset + OffsetInChunk);
+			OffsetInChunk += ChunkBlock.Size;
+			BlockEntry.SetCompressedSize(uint32(ChunkBlock.CompressedSize));
+			BlockEntry.SetUncompressedSize(uint32(ChunkBlock.UncompressedSize));
+			BlockEntry.SetCompressionMethodIndex(TocBuilder.AddCompressionMethodEntry(ChunkBlock.CompressionMethod));
+
+			// We do this here so that we get the total size of data excluding the encryption alignment
+			TotalEntryCompressedSize += ChunkBlock.CompressedSize;
+			if (Entry->bCouldBeFromReferenceDb && !Entry->bLoadingFromReferenceDb)
+			{
+				ReferenceCacheMissBytes += ChunkBlock.CompressedSize;
+			}
+
+			if (!ChunkBlock.CompressionMethod.IsNone())
+			{
+				ChunkMeta.Flags |= FIoStoreTocEntryMetaFlags::Compressed;
+			}
+
+			if (ContainerSettings.IsSigned())
+			{
+				FSHAHash& Signature = TocBuilder.AddBlockSignatureEntry();
+				Signature = ChunkBlock.Signature;
+			}
+
+		}
+
+		const int32 TocEntryIndex = TocBuilder.AddChunkEntry(Entry->ChunkId, OffsetLength, ChunkMeta);
+		check(TocEntryIndex != INDEX_NONE);
+
+		if (ContainerSettings.IsIndexed() && Entry->Options.FileName.Len() > 0)
+		{
+			TocBuilder.AddToFileIndex(Entry->ChunkId, Entry->Options.FileName);
+		}
+
+		const uint64 RegionStartOffset = TargetPartition->Offset;
+		TargetPartition->Offset += Entry->CompressedSize;
+		UncompressedFileOffset += Align(Entry->UncompressedSize.GetValue(), WriterSettings.CompressionBlockSize);
+		TotalEntryUncompressedSize += Entry->UncompressedSize.GetValue();
+
+		if (WriterSettings.bEnableFileRegions)
+		{
+			FFileRegion::AccumulateFileRegions(TargetPartition->AllFileRegions, RegionStartOffset, RegionStartOffset, TargetPartition->Offset, Entry->Request->GetRegions());
+		}
+		delete Entry->Request;
+		Entry->Request = nullptr;
+		uint64 WriteStartCycles = FPlatformTime::Cycles64();
+		uint64 WriteBytes = 0;
+		if (Entry->Padding > 0)
+		{
+			if (PaddingBuffer.Num() < Entry->Padding)
+			{
+				PaddingBuffer.SetNumZeroed(int32(Entry->Padding));
+			}
+			{
+				TRACE_CPUPROFILER_EVENT_SCOPE(WritePaddingToContainer);
+				TargetPartition->ContainerFileHandle->Serialize(PaddingBuffer.GetData(), Entry->Padding);
+				WriteBytes += Entry->Padding;
+			}
+		}
+		check(Entry->Offset == TargetPartition->ContainerFileHandle->Tell());
+		for (FChunkBlock& ChunkBlock : Entry->ChunkBlocks)
+		{
+			{
+				TRACE_CPUPROFILER_EVENT_SCOPE(WriteBlockToContainer);
+				TargetPartition->ContainerFileHandle->Serialize(ChunkBlock.IoBuffer->Data(), ChunkBlock.Size);
+				WriteBytes += ChunkBlock.Size;
+			}
+			WriterContext->FreeCompressionBuffer(ChunkBlock.IoBuffer, Entry->ChunkBlocks.Num());
+			ChunkBlock.IoBuffer = nullptr;
+		}
+		uint64 WriteEndCycles = FPlatformTime::Cycles64();
+		WriterContext->WriteCycleCount.AddExchange(WriteEndCycles - WriteStartCycles);
+		WriterContext->WriteByteCount.AddExchange(WriteBytes);
+		WriterContext->SerializedChunksCount.IncrementExchange();
+	}
+
+	const FString				ContainerPathAndBaseFileName;
+	FIoStoreWriterContextImpl*	WriterContext = nullptr;
+	FIoContainerSettings		ContainerSettings;
+	FString						TocFilePath;
+	FIoStoreTocBuilder			TocBuilder;
+	TArray<uint8>				PaddingBuffer;
+	TArray<FPartition>			Partitions;
+	TArray<FIoStoreWriteQueueEntry*> Entries;
+	TArray<FLayoutEntry*>		LayoutEntries;
+	FLayoutEntry*				LayoutEntriesHead = nullptr;
+	FLayoutEntry*				LayoutEntriesTail = nullptr;
+	TMap<FIoChunkId, FLayoutEntry*> PreviousBuildLayoutEntryByChunkId;
+	TUniquePtr<FArchive>		CsvArchive;
+	FIoStoreWriterResult		Result;
+	uint64						UncompressedFileOffset = 0;
+	uint64						TotalEntryUncompressedSize = 0; // sum of all entry source buffer sizes
+	uint64						TotalEntryCompressedSize = 0; // entry compressed size excluding encryption alignment
+	uint64						ReferenceCacheMissBytes = 0; // number of compressed bytes excluding alignment that could have been from refcache but weren't.
+	uint64						TotalPaddingSize = 0;
+	uint64						UncompressedContainerSize = 0; // this is the size the container would be if it were uncompressed.
+	uint64						CompressedContainerSize = 0; // this is the size of the container with the given compression (which may be none).
+	int32						CurrentPartitionIndex = 0;
+	bool						bHasMemoryMappedEntry = false;
+	bool						bHasFlushed = false;
+	bool						bHasResult = false;
+	TSharedPtr<IIoStoreWriterReferenceChunkDatabase> ReferenceChunkDatabase;
+	TSharedPtr<IIoStoreWriterHashDatabase> HashDatabase;
+	bool						bVerifyHashDatabase = false;
+
+
+	friend class FIoStoreWriterContextImpl;
+};
+
+// InContainerPathAndBaseFileName: the utoc file will just be this with .utoc appended.
+// The base filename ends up getting returned as the container name in the writer results.
+TSharedPtr<IIoStoreWriter> FIoStoreWriterContextImpl::CreateContainer(const TCHAR* InContainerPathAndBaseFileName, const FIoContainerSettings& InContainerSettings)
+{
+	TSharedPtr<FIoStoreWriter> IoStoreWriter = MakeShared<FIoStoreWriter>(InContainerPathAndBaseFileName);
+	FIoStatus IoStatus = IoStoreWriter->Initialize(*this, InContainerSettings);
+	check(IoStatus.IsOk());
+	IoStoreWriters.Add(IoStoreWriter);
+	return IoStoreWriter;
+}
+
+void FIoStoreWriterContextImpl::Flush()
+{
+	TRACE_CPUPROFILER_EVENT_SCOPE(FIoStoreWriterContext::Flush);
+	TArray<FIoStoreWriteQueueEntry*> AllEntries;
+	for (TSharedPtr<FIoStoreWriter> IoStoreWriter : IoStoreWriters)
+	{
+		IoStoreWriter->bHasFlushed = true;
+		AllEntries.Append(IoStoreWriter->Entries);
+	}
+	{
+		TRACE_CPUPROFILER_EVENT_SCOPE(WaitForChunkHashes);
+		for (int32 EntryIndex = AllEntries.Num() - 1; EntryIndex >= 0; --EntryIndex)
+		{
+			FTaskGraphInterface::Get().WaitUntilTaskCompletes(AllEntries[EntryIndex]->HashTask);
+		}
+	}
+	for (TSharedPtr<FIoStoreWriter> IoStoreWriter : IoStoreWriters)
+	{
+		if (IoStoreWriter->LayoutEntriesHead)
+		{
+			IoStoreWriter->FinalizeLayout();
+		}
+	}
+	// Update list of all entries after having the finilized layouts of each container
+	AllEntries.Reset();
+	for (TSharedPtr<FIoStoreWriter> IoStoreWriter : IoStoreWriters)
+	{
+		AllEntries.Append(IoStoreWriter->Entries);
+	}
+
+	double WritesStart = FPlatformTime::Seconds();
+
+	for (FIoStoreWriteQueueEntry* Entry : AllEntries)
+	{
+		ScheduleCompression(Entry);
+	}
+
+	{
+		TRACE_CPUPROFILER_EVENT_SCOPE(WaitForWritesToComplete);
+		for (int32 EntryIndex = AllEntries.Num() - 1; EntryIndex >= 0; --EntryIndex)
+		{
+			AllEntries[EntryIndex]->WriteFinishedEvent->Wait();
+		}
+	}
+
+	double WritesEnd = FPlatformTime::Seconds();
+	double WritesSeconds = FPlatformTime::ToSeconds64(WriteCycleCount.Load());
+	UE_LOG(LogIoStore, Display, TEXT("Writing and compressing took %.2f seconds, writes to disk took %.2f seconds for %s bytes @ %s bytes per second."), 
+		WritesEnd - WritesStart,
+		WritesSeconds,
+		*FText::AsNumber(WriteByteCount.Load()).ToString(),
+		*FText::AsNumber((int64)((double)WriteByteCount.Load() / FMath::Max(.0001f, WritesSeconds))).ToString()
+		);
+
+	AllEntries.Empty();
+
+	// Classically there were so few writers that this didn't need to be multi threaded, but it
+	// involves writing files, and with content on demand this ends up being thousands of iterations. 
+	double FinalizeStart = FPlatformTime::Seconds();	
+	ParallelFor(TEXT("IoStoreWriter::Finalize.PF"), IoStoreWriters.Num(), 1, [this](int Index)
+	{ 
+		IoStoreWriters[Index]->Finalize(); 
+	});
+	double FinalizeEnd = FPlatformTime::Seconds();
+	int64 TotalTocSize = 0;
+	for (TSharedPtr<IIoStoreWriter> Writer : IoStoreWriters)
+	{
+		if (Writer->GetResult().IsOk())
+		{
+			TotalTocSize += Writer->GetResult().ValueOrDie().TocSize;
+		}
+	}
+
+	UE_LOG(LogIoStore, Display, TEXT("Finalize took %.1f seconds for %d writers to write %s bytes, %s bytes per second"), 
+		FinalizeEnd - FinalizeStart, 
+		IoStoreWriters.Num(), 
+		*FText::AsNumber(TotalTocSize).ToString(), 
+		*FText::AsNumber((int64)((double)TotalTocSize / FMath::Max(.0001f, FinalizeEnd - FinalizeStart))).ToString()
+		);
+}
+
+void FIoStoreWriterContextImpl::BeginCompressionThreadFunc()
+{
+	for (;;)
+	{
+		FIoStoreWriteQueueEntry* Entry = BeginCompressionQueue.DequeueOrWait();
+		if (!Entry)
+		{
+			return;
+		}
+		while (Entry)
+		{
+			FIoStoreWriteQueueEntry* Next = Entry->Next;
+			Entry->BeginCompressionBarrier->Wait();
+			Entry->Writer->BeginCompress(Entry);
+			BeginEncryptionAndSigningQueue.Enqueue(Entry);
+			Entry = Next;
+		}
+	}
+}
+
+void FIoStoreWriterContextImpl::BeginEncryptionAndSigningThreadFunc()
+{
+	for (;;)
+	{
+		FIoStoreWriteQueueEntry* Entry = BeginEncryptionAndSigningQueue.DequeueOrWait();
+		if (!Entry)
+		{
+			return;
+		}
+		while (Entry)
+		{
+			FIoStoreWriteQueueEntry* Next = Entry->Next;
+			Entry->FinishCompressionBarrier->Wait();
+			if (Entry->bLoadingFromReferenceDb == false)
+			{
+				Entry->Request->FreeSourceBuffer();
+			}
+			if (Entry->bStoreCompressedDataInDDC)
+			{
+				TRACE_CPUPROFILER_EVENT_SCOPE(StoreInDDC);
+				TArray<uint8> DDCData;
+				FMemoryWriter DDCDataWriter(DDCData, true);
+				if (Entry->Writer->SerializeCompressedDDCData(Entry, DDCDataWriter))
+				{
+					DDC->Put(*Entry->DDCCacheKey, DDCData, Entry->Options.FileName);
+				}
+			}
+			FinishEncryptionAndSigningQueue.Enqueue(Entry);
+			Entry->Writer->BeginEncryptAndSign(Entry);
+			Entry = Next;
+		}
+	}
+}
+
+void FIoStoreWriterContextImpl::FinishEncryptionAndSigningThreadFunc()
+{
+	for (;;)
+	{
+		FIoStoreWriteQueueEntry* Entry = FinishEncryptionAndSigningQueue.DequeueOrWait();
+		if (!Entry)
+		{
+			return;
+		}
+		while (Entry)
+		{
+			FIoStoreWriteQueueEntry* Next = Entry->Next;
+			Entry->FinishEncryptionAndSigningBarrier->Wait();
+			WriterQueue.Enqueue(Entry);
+			Entry->CompressedSize = 0;
+			for (const FChunkBlock& ChunkBlock : Entry->ChunkBlocks)
+			{
+				Entry->CompressedSize += ChunkBlock.Size;
+			}
+			Entry->BeginWriteBarrier->DispatchSubsequents();
+
+			Entry = Next;
+		}
+	}
+}
+
+void FIoStoreWriterContextImpl::WriterThreadFunc()
+{
+	for (;;)
+	{
+		FIoStoreWriteQueueEntry* Entry = WriterQueue.DequeueOrWait();
+		if (!Entry)
+		{
+			return;
+		}
+		while (Entry)
+		{
+			FIoStoreWriteQueueEntry* Next = Entry->Next;
+			Entry->BeginWriteBarrier->Wait();
+			Entry->Writer->WriteEntry(Entry);
+			Entry->WriteFinishedEvent->DispatchSubsequents();
+			Entry = Next;
+		}
+	}
+}
+
diff --git a/Engine/Source/Developer/IoStoreUtilities/Private/IoStoreWriter.h b/Engine/Source/Developer/IoStoreUtilities/Private/IoStoreWriter.h
new file mode 100644
--- /dev/null
+++ b/Engine/Source/Developer/IoStoreUtilities/Private/IoStoreWriter.h
@@ -0,0 +1,187 @@
+// Copyright Epic Games, Inc. All Rights Reserved.
+
+#pragma once
+
+#include "Containers/Array.h"
+#include "IO/IoContainerId.h"
+#include "IO/IoDispatcher.h"
+#include "UObject/NameTypes.h"
+
+class FIoStoreWriterContextImpl;
+
+struct FIoStoreWriterSettings
+{
+	FName CompressionMethod = NAME_None;
+	uint64 CompressionBlockSize = 64 << 10;
+
+	// This does not align every entry - it tries to prevent excess crossings of this boundary by inserting padding.
+	// and happens whether or not the entry is compressed.
+	uint64 CompressionBlockAlignment = 0;
+	int32 CompressionMinBytesSaved = 0;
+	int32 CompressionMinPercentSaved = 0;
+	int32 CompressionMinSizeToConsiderDDC = 0;
+	uint64 MemoryMappingAlignment = 0;
+	uint64 MaxPartitionSize = 0;
+	bool bEnableFileRegions = false;
+	bool bCompressionEnableDDC = false;
+};
+
+struct FIoStoreWriterResult
+{
+	FIoContainerId ContainerId;
+	FString ContainerName; // This is the base filename of the utoc used for output.
+	int64 TocSize = 0;
+	int64 TocEntryCount = 0;
+	int64 PaddingSize = 0;
+	int64 UncompressedContainerSize = 0; // this is the size the container would be if it were uncompressed.
+	int64 CompressedContainerSize = 0; // this is the size of the container with the given compression (which may be none). Should be the sum of all partition file sizes.
+	int64 DirectoryIndexSize = 0;
+	uint64 TotalEntryCompressedSize = 0; // sum of the compressed size of entries excluding encryption alignment.
+	uint64 ReferenceCacheMissBytes = 0; // number of compressed bytes excluding alignment that could have been from refcache but weren't.
+	uint64 AddedChunksCount = 0;
+	uint64 AddedChunksSize = 0;
+	uint64 ModifiedChunksCount = 0;
+	uint64 ModifiedChunksSize = 0;
+	FName CompressionMethod = NAME_None;
+	EIoContainerFlags ContainerFlags = EIoContainerFlags::None;
+};
+
+struct FIoWriteOptions
+{
+	FString FileName;
+	const TCHAR* DebugName = nullptr;
+	bool bForceUncompressed = false;
+	bool bIsMemoryMapped = false;
+};
+
+class FIoStoreWriterContext
+{
+public:
+	struct FProgress
+	{
+		uint64 TotalChunksCount = 0;
+		uint64 HashedChunksCount = 0;
+		// Number of chunks where we avoided reading and hashing, and instead used the result from the hashdb, and their types
+		uint64 HashDbChunksCount = 0;
+		uint64 HashDbChunksByType[(int8)EIoChunkType::MAX] = { 0 };
+		// Number of chunks that were passed to the compressor (i.e. passed the various opt-outs), and their types
+		uint64 CompressedChunksCount = 0;
+		uint64 CompressedChunksByType[(int8)EIoChunkType::MAX] = { 0 };
+		uint64 SerializedChunksCount = 0;
+		uint64 ScheduledCompressionTasksCount = 0;
+		uint64 CompressionDDCHitCount = 0;
+		uint64 CompressionDDCMissCount = 0;
+
+		// The number of chunk retrieved from the reference cache database, and their types.
+		uint64 RefDbChunksCount{ 0 };
+		uint64 RefDbChunksByType[(int8)EIoChunkType::MAX] = { 0 };
+		
+		// The type of chunk that landed in BeginCompress before any opt-outs.
+		uint64 BeginCompressChunksByType[(int8)EIoChunkType::MAX] = { 0 };
+	};
+
+	FIoStoreWriterContext();
+	~FIoStoreWriterContext();
+
+	[[nodiscard]] FIoStatus Initialize(const FIoStoreWriterSettings& InWriterSettings);
+	TSharedPtr<class IIoStoreWriter> CreateContainer(const TCHAR* InContainerPath, const FIoContainerSettings& InContainerSettings);
+	void Flush();
+	FProgress GetProgress() const;
+
+private:
+	FIoStoreWriterContextImpl* Impl;
+};
+
+class IIoStoreWriteRequest
+{
+public:
+	virtual ~IIoStoreWriteRequest() = default;
+
+	// Launches any async operations necessary in order to access the buffer. CompletionEvent is set once it's ready, which may be immediate.
+	virtual void PrepareSourceBufferAsync(FGraphEventRef CompletionEvent) = 0;
+	virtual uint64 GetOrderHint() = 0;
+	virtual TArrayView<const FFileRegion> GetRegions() = 0;
+
+	// Only valid after the completion event passed to PrepareSourceBufferAsync has fired.
+	virtual const FIoBuffer* GetSourceBuffer() = 0;
+
+	// Can't be called between PrepareSourceBufferAsync and its completion!
+	virtual void FreeSourceBuffer() = 0;
+};
+
+class IIoStoreWriterReferenceChunkDatabase
+{
+public:
+	virtual ~IIoStoreWriterReferenceChunkDatabase() = default;
+
+	/*
+	* Used by IIoStoreWriter to check and see if there's a reference chunk that matches the data that
+	* IoStoreWriter wants to compress and write. Validity checks must be synchronous - if a chunk can't be
+	* used for some reason (no matching chunk exists or otherwise), this function must return false and not 
+	* call InCompletionCallback.
+	* 
+	* Once a matching chunk is found, it is read from the source iostore container asynchronously, and upon
+	* completion InCompletionCallback is called with the raw output from FIoStoreReader::ReadCompressed (i.e.
+	* FIoStoreCompressedReadResult). Failures once the async read process has started are currently fatal due to
+	* difficulties in rekicking a read.
+	* 
+	* For the moment, changes in compression method are allowed.
+	* 
+	* RetrieveChunk is not currently thread safe and must be called from a single thread.
+	* 
+	* Chunks provided *MUST* decompress to bits that hash to the exact value provided in InChunkKey (i.e. be exactly the same bits),
+	* and also be the same number of blocks (i.e. same CompressionBlockSize)
+	*/
+	virtual bool RetrieveChunk(const FIoContainerId& InContainerId, const FIoChunkHash& InChunkHash, const FIoChunkId& InChunkId, TUniqueFunction<void(TIoStatusOr<FIoStoreCompressedReadResult>)> InCompletionCallback) = 0;
+
+	/* 
+	* Quick synchronous existence check that returns the number of blocks for the chunk. This is used to set up
+	* the necessary structures without needing to read the source data for the chunk. This might be called from
+	* multiple threads as it has to happen after we have the source hash computed.
+	*/
+	virtual bool ChunkExists(const FIoContainerId& InContainerId, const FIoChunkHash& InChunkHash, const FIoChunkId& InChunkId, uint32& OutNumChunkBlocks) = 0;
+
+	/*
+	* Returns the compression block size that was used to break up the IoChunks in the source containers. If this is different than what we want, 
+	* then none of the chunks will ever match. Knowing this up front allows us to only match on hash
+	*/
+	virtual uint32 GetCompressionBlockSize() const = 0;
+
+	/*
+	* Called by an iostore writer implementation to notify the ref cache it's been added
+	*/
+	virtual void NotifyAddedToWriter(const FIoContainerId& InContainerId, const FString& InContainerName) = 0;
+};
+
+/**
+*	Allows the IIoStoreWriter to avoid loading and hashing chunks, saving pak/stage time, as the normal
+*	process involved loading the chunks, hashing them, freeing them, making some decisions, then loading
+*	them _again_ for compression/writting. It's completely fine for this to not have all available hashes,
+*	but they have to match when provided!
+*/
+class IIoStoreWriterHashDatabase
+{
+public:
+	virtual ~IIoStoreWriterHashDatabase() = default;
+	virtual bool FindHashForChunkId(const FIoChunkId& ChunkId, FIoChunkHash& OutHash) const = 0;
+};
+
+
+class IIoStoreWriter
+{
+public:
+	virtual ~IIoStoreWriter() = default;
+
+	/**
+	*	If a reference database is provided, the IoStoreWriter implementation may elect to reuse compressed blocks
+	*	from previous containers instead of recompressing input data. This must be set before any writes are appended.
+	*/
+	virtual void SetReferenceChunkDatabase(TSharedPtr<IIoStoreWriterReferenceChunkDatabase> ReferenceChunkDatabase) = 0;
+	virtual void SetHashDatabase(TSharedPtr<IIoStoreWriterHashDatabase> HashDatabase, bool bVerifyHashDatabase) = 0;
+	virtual void EnableDiskLayoutOrdering(const TArray<TUniquePtr<FIoStoreReader>>& PatchSourceReaders = TArray<TUniquePtr<FIoStoreReader>>()) = 0;
+	virtual void Append(const FIoChunkId& ChunkId, FIoBuffer Chunk, const FIoWriteOptions& WriteOptions, uint64 OrderHint = MAX_uint64) = 0;
+	virtual void Append(const FIoChunkId& ChunkId, IIoStoreWriteRequest* Request, const FIoWriteOptions& WriteOptions) = 0;
+	virtual TIoStatusOr<FIoStoreWriterResult> GetResult() = 0;
+	virtual void EnumerateChunks(TFunction<bool(FIoStoreTocChunkInfo&&)>&& Callback) const = 0;
+};
+
diff --git a/Engine/Source/Runtime/Core/Internal/IO/IoDirectoryIndex.h b/Engine/Source/Runtime/Core/Internal/IO/IoDirectoryIndex.h
--- a/Engine/Source/Runtime/Core/Internal/IO/IoDirectoryIndex.h
+++ b/Engine/Source/Runtime/Core/Internal/IO/IoDirectoryIndex.h
@@ -15,7 +15,7 @@ class FArchive;
 
 namespace IoDirectoryIndexUtils
 {
-	FString GetCommonRootPath(const TArray<FStringView>& Filenames);
+	CORE_API FString GetCommonRootPath(const TArray<FStringView>& Filenames);
 }
 
 struct FIoDirectoryIndexEntry
@@ -50,15 +50,15 @@ struct FIoDirectoryIndexResource
 class FIoDirectoryIndexWriter
 {
 public:
-	FIoDirectoryIndexWriter();
+	CORE_API FIoDirectoryIndexWriter();
 
-	void SetMountPoint(FString InMountPoint);
-	uint32 AddFile(const FStringView& InFileName);
-	void SetFileUserData(uint32 InFileEntryIndex, uint32 InUserData);
+	CORE_API void SetMountPoint(FString InMountPoint);
+	CORE_API uint32 AddFile(const FStringView& InFileName);
+	CORE_API void SetFileUserData(uint32 InFileEntryIndex, uint32 InUserData);
 
 	// Flush() can be called only once, because afterwards contents of Strings member
 	// will be empty as all strings have been moved away
-	void Flush(TArray<uint8>& OutBuffer, FAES::FAESKey InEncryptionKey);
+	CORE_API void Flush(TArray<uint8>& OutBuffer, FAES::FAESKey InEncryptionKey);
 
 private:
 	uint32 GetDirectory(uint32 DirectoryName, uint32 Parent);
diff --git a/Engine/Source/Runtime/Core/Private/IO/IoStore.cpp b/Engine/Source/Runtime/Core/Private/IO/IoStore.cpp
--- a/Engine/Source/Runtime/Core/Private/IO/IoStore.cpp
+++ b/Engine/Source/Runtime/Core/Private/IO/IoStore.cpp
@@ -2,12 +2,7 @@
 
 #include "IO/IoStore.h"
 
-#include "Async/Async.h"
-#include "Async/AsyncFileHandle.h"
-#include "Async/Future.h"
-#include "Async/ParallelFor.h"
 #include "Containers/Map.h"
-#include "DerivedDataCacheInterface.h"
 #include "Features/IModularFeatures.h"
 #include "HAL/FileManager.h"
 #include "HAL/PlatformFileManager.h"
@@ -19,7 +14,6 @@
 #include "ProfilingDebugging/CountersTrace.h"
 #include "Serialization/BufferWriter.h"
 #include "Serialization/LargeMemoryReader.h"
-#include "Serialization/MemoryWriter.h"
 #include "Tasks/Task.h"
 #include "Templates/UniquePtr.h"
 
@@ -154,493 +148,6 @@ static FIoStatus ValidateContainerSignature(
 	return FIoStatus::Ok;
 }
 
-struct FChunkBlock
-{
-	const uint8* UncompressedData = nullptr;
-	FIoBuffer* IoBuffer = nullptr;
-
-	// This is the size of the actual block after encryption alignment, and is
-	// set in EncryptAndSign. This happens whether or not the container is encrypted.
-	uint64 Size = 0;
-	uint64 CompressedSize = 0;
-	uint64 UncompressedSize = 0;
-	FName CompressionMethod = NAME_None;
-	FSHAHash Signature;
-
-	/** Hash of the block data as it would be found on disk - this includes encryption alignment padding */
-	FIoHash DiskHash;
-};
-
-struct FIoStoreWriteQueueEntry
-{
-	FIoStoreWriteQueueEntry* Next = nullptr;
-	class FIoStoreWriter* Writer = nullptr;
-	IIoStoreWriteRequest* Request = nullptr;
-	FIoChunkId ChunkId;
-	FIoChunkHash ChunkHash;
-	/** Hash of the block data as it would be found on disk after compression and encryption */
-	FIoHash ChunkDiskHash;
-
-	uint64 Sequence = 0;
-	
-	// We make this optional because at the latest it might not be valid until FinishCompressionBarrior
-	// completes and we'd like to have a check() on that.
-	TOptional<uint64> UncompressedSize = 0;
-
-	// this is not filled out until after encryption completes and *includes the alignment padding for encryption*!
-	// think of this as "size on disk".
-	uint64 CompressedSize = 0; 
-
-	uint64 Padding = 0;
-	uint64 Offset = 0;
-	TArray<FChunkBlock> ChunkBlocks;
-	FIoWriteOptions Options;
-	FGraphEventRef HashBarrier;
-	FGraphEventRef HashTask;
-	FGraphEventRef BeginCompressionBarrier;
-	FGraphEventRef FinishCompressionBarrier;
-	FGraphEventRef FinishEncryptionAndSigningBarrier;
-	FGraphEventRef BeginWriteBarrier;
-	FGraphEventRef WriteFinishedEvent;
-	TAtomic<int32> CompressedBlocksCount{ 0 };
-	TAtomic<int32> FinishedBlocksCount{ 0 };
-	int32 PartitionIndex = -1;
-	FString DDCCacheKey;
-	bool bAdded = false;
-	bool bModified = false;
-	bool bStoreCompressedDataInDDC = false;
-	
-	bool bCouldBeFromReferenceDb = false; // Whether the chunk is a valid candidate for the reference db.
-	bool bLoadingFromReferenceDb = false;
-	// When we know we're loading from the reference chunk db, we don't read the source
-	// buffer but we still need to know the number of chunks which we get from the refdb.
-	uint32 NumChunkBlocksFromRefDb = 0; 
-};
-
-class FIoStoreWriteQueue
-{
-public:
-	FIoStoreWriteQueue()
-		: Event(FPlatformProcess::GetSynchEventFromPool(false))
-	{ }
-	
-	~FIoStoreWriteQueue()
-	{
-		check(Head == nullptr && Tail == nullptr);
-		FPlatformProcess::ReturnSynchEventToPool(Event);
-	}
-
-	void Enqueue(FIoStoreWriteQueueEntry* Entry)
-	{
-		check(!bIsDoneAdding);
-		{
-			FScopeLock _(&CriticalSection);
-
-			if (!Tail)
-			{
-				Head = Tail = Entry;
-			}
-			else
-			{
-				Tail->Next = Entry;
-				Tail = Entry;
-			}
-			Entry->Next = nullptr;
-		}
-
-		Event->Trigger();
-	}
-
-	FIoStoreWriteQueueEntry* DequeueOrWait()
-	{
-		for (;;)
-		{
-			{
-				FScopeLock _(&CriticalSection);
-				if (Head)
-				{
-					FIoStoreWriteQueueEntry* Entry = Head;
-					Head = Tail = nullptr;
-					return Entry;
-				}
-			}
-
-			if (bIsDoneAdding)
-			{
-				break;
-			}
-
-			Event->Wait();
-		}
-
-		return nullptr;
-	}
-
-	void CompleteAdding()
-	{
-		bIsDoneAdding = true;
-		Event->Trigger();
-	}
-
-	bool IsEmpty() const
-	{
-		FScopeLock _(&CriticalSection);
-		return Head == nullptr;
-	}
-
-private:
-	mutable FCriticalSection CriticalSection;
-	FEvent* Event = nullptr;
-	FIoStoreWriteQueueEntry* Head = nullptr;
-	FIoStoreWriteQueueEntry* Tail = nullptr;
-	TAtomic<bool> bIsDoneAdding { false };
-};
-
-class FIoStoreWriterContextImpl
-{
-	static constexpr uint64 DefaultMemoryLimit = 2ull << 30ull;
-public:
-	FIoStoreWriterContextImpl()
-	{
-	}
-
-	~FIoStoreWriterContextImpl()
-	{
-		BeginCompressionQueue.CompleteAdding();
-		BeginEncryptionAndSigningQueue.CompleteAdding();
-		FinishEncryptionAndSigningQueue.CompleteAdding();
-		WriterQueue.CompleteAdding();
-		BeginCompressionThread.Wait();
-		BeginEncryptionAndSigningThread.Wait();
-		FinishEncryptionAndSigningThread.Wait();
-		WriterThread.Wait();
-		if (CompressionBufferAvailableEvent)
-		{
-			FPlatformProcess::ReturnSynchEventToPool(CompressionBufferAvailableEvent);
-		}
-		check(AvailableCompressionBuffers.Num() == TotalCompressionBufferCount);
-		for (FIoBuffer* IoBuffer : AvailableCompressionBuffers)
-		{
-			delete IoBuffer;
-		}
-	}
-
-	[[nodiscard]] FIoStatus Initialize(const FIoStoreWriterSettings& InWriterSettings)
-	{
-		TRACE_CPUPROFILER_EVENT_SCOPE(FIoStoreWriterContext::Initialize);
-		WriterSettings = InWriterSettings;
-		CompressionBufferAvailableEvent = FPlatformProcess::GetSynchEventFromPool(false);
-
-		if (WriterSettings.bCompressionEnableDDC)
-		{
-			DDC = GetDerivedDataCache();
-		}
-
-		if (WriterSettings.CompressionMethod != NAME_None)
-		{
-			CompressionBufferSize = FCompression::CompressMemoryBound(WriterSettings.CompressionMethod, static_cast<int32>(WriterSettings.CompressionBlockSize));
-		}
-		CompressionBufferSize = FMath::Max(CompressionBufferSize, static_cast<int32>(WriterSettings.CompressionBlockSize));
-		CompressionBufferSize = Align(CompressionBufferSize, FAES::AESBlockSize);
-
-		TotalCompressionBufferCount = int32(DefaultMemoryLimit / CompressionBufferSize);
-		AvailableCompressionBuffers.Reserve(TotalCompressionBufferCount);
-		for (int32 BufferIndex = 0; BufferIndex < TotalCompressionBufferCount; ++BufferIndex)
-		{
-			AvailableCompressionBuffers.Add(new FIoBuffer(CompressionBufferSize));
-		}
-		//TRACE_COUNTER_SET(IoStoreAvailableCompressionBuffers, AvailableCompressionBuffers.Num());
-
-		BeginCompressionThread = Async(EAsyncExecution::Thread, [this]() { BeginCompressionThreadFunc(); });
-		BeginEncryptionAndSigningThread = Async(EAsyncExecution::Thread, [this]() { BeginEncryptionAndSigningThreadFunc(); });
-		FinishEncryptionAndSigningThread = Async(EAsyncExecution::Thread, [this]() { FinishEncryptionAndSigningThreadFunc(); });
-		WriterThread = Async(EAsyncExecution::Thread, [this]() { WriterThreadFunc(); });
-
-		return FIoStatus::Ok;
-	}
-
-	TSharedPtr<IIoStoreWriter> CreateContainer(const TCHAR* InContainerPathAndBaseFileName, const FIoContainerSettings& InContainerSettings);
-
-	void Flush();
-
-	FIoStoreWriterContext::FProgress GetProgress() const
-	{
-		FIoStoreWriterContext::FProgress Progress;
-		Progress.HashDbChunksCount = HashDbChunksCount.Load();
-		for (uint8 i=0; i<(uint8)EIoChunkType::MAX; i++)
-		{
-			Progress.HashDbChunksByType[i] = HashDbChunksByType[i].Load();
-		}
-
-		Progress.TotalChunksCount = TotalChunksCount.Load();
-		Progress.HashedChunksCount = HashedChunksCount.Load();
-		Progress.CompressedChunksCount = CompressedChunksCount.Load();
-		for (uint8 i = 0; i < (uint8)EIoChunkType::MAX; i++)
-		{
-			Progress.CompressedChunksByType[i] = CompressedChunksByType[i].Load();
-		}
-		for (uint8 i = 0; i < (uint8)EIoChunkType::MAX; i++)
-		{
-			Progress.BeginCompressChunksByType[i] = BeginCompressChunksByType[i].Load();
-		}
-		Progress.SerializedChunksCount = SerializedChunksCount.Load();
-		Progress.ScheduledCompressionTasksCount = ScheduledCompressionTasksCount.Load();
-		Progress.CompressionDDCHitCount = CompressionDDCHitCount.Load();
-		Progress.CompressionDDCMissCount = CompressionDDCMissCount.Load();
-		Progress.RefDbChunksCount = RefDbChunksCount.Load();
-		for (uint8 i = 0; i < (uint8)EIoChunkType::MAX; i++)
-		{
-			Progress.RefDbChunksByType[i] = RefDbChunksByType[i].Load();
-		}
-
-		return Progress;
-	}
-
-	const FIoStoreWriterSettings& GetSettings() const
-	{
-		return WriterSettings;
-	}
-
-	void ScheduleCompression(FIoStoreWriteQueueEntry* QueueEntry)
-	{
-		BeginCompressionQueue.Enqueue(QueueEntry);
-		if (QueueEntry->bLoadingFromReferenceDb == false)
-		{
-			QueueEntry->Request->PrepareSourceBufferAsync(QueueEntry->BeginCompressionBarrier);
-		}
-		else
-		{
-			// We don't need to wait on a read so we can kick directly.
-			QueueEntry->BeginCompressionBarrier->DispatchSubsequents();
-		}
-	}
-
-	//
-	// This must be called prior to the consumer being dispatched in order to prevent resource
-	// contention deadlock (even if using retracting waits)
-	//
-	FIoBuffer* AllocCompressionBuffer(int32 TotalEntryChunkBlocksCount)
-	{
-		if (TotalEntryChunkBlocksCount > TotalCompressionBufferCount)
-		{
-			return new FIoBuffer(CompressionBufferSize);
-		}
-		FIoBuffer* AllocatedBuffer = nullptr;
-		while (!AllocatedBuffer)
-		{
-			{
-				FScopeLock Lock(&AvailableCompressionBuffersCritical);
-				if (AvailableCompressionBuffers.Num())
-				{
-					AllocatedBuffer = AvailableCompressionBuffers.Pop();
-					//TRACE_COUNTER_DECREMENT(IoStoreAvailableCompressionBuffers);
-				}
-			}
-			if (!AllocatedBuffer)
-			{
-				TRACE_CPUPROFILER_EVENT_SCOPE(WaitForCompressionBuffer);
-				CompressionBufferAvailableEvent->Wait();
-			}
-		}
-		return AllocatedBuffer;
-	}
-
-	void FreeCompressionBuffer(FIoBuffer* Buffer, int32 TotalEntryChunkBlocksCount)
-	{
-		if (TotalEntryChunkBlocksCount > TotalCompressionBufferCount)
-		{
-			delete Buffer;
-			return;
-		}
-		bool bTriggerEvent;
-		{
-			FScopeLock Lock(&AvailableCompressionBuffersCritical);
-			bTriggerEvent = AvailableCompressionBuffers.Num() == 0;
-			AvailableCompressionBuffers.Push(Buffer);
-			//TRACE_COUNTER_INCREMENT(IoStoreAvailableCompressionBuffers);
-		}
-		if (bTriggerEvent)
-		{
-			CompressionBufferAvailableEvent->Trigger();
-		}
-	}
-
-private:
-	void BeginCompressionThreadFunc();
-	void BeginEncryptionAndSigningThreadFunc();
-	void FinishEncryptionAndSigningThreadFunc();
-	void WriterThreadFunc();
-
-	FIoStoreWriterSettings WriterSettings;
-	FDerivedDataCacheInterface* DDC = nullptr;
-	FEvent* CompressionBufferAvailableEvent = nullptr;
-	TFuture<void> BeginCompressionThread;
-	TFuture<void> BeginEncryptionAndSigningThread;
-	TFuture<void> FinishEncryptionAndSigningThread;
-	TFuture<void> WriterThread;
-	FIoStoreWriteQueue BeginCompressionQueue;
-	FIoStoreWriteQueue BeginEncryptionAndSigningQueue;
-	FIoStoreWriteQueue FinishEncryptionAndSigningQueue;
-	FIoStoreWriteQueue WriterQueue;
-	TAtomic<uint64> TotalChunksCount{ 0 };
-	TAtomic<uint64> HashedChunksCount{ 0 };
-	TAtomic<uint64> HashDbChunksCount{ 0 };
-	TAtomic<uint64> HashDbChunksByType[(int8)EIoChunkType::MAX] = {0};
-	TAtomic<uint64> RefDbChunksCount{ 0 };
-	TAtomic<uint64> RefDbChunksByType[(int8)EIoChunkType::MAX] = { 0 };
-	TAtomic<uint64> CompressedChunksCount{ 0 };
-	TAtomic<uint64> CompressedChunksByType[(int8)EIoChunkType::MAX] = { 0 };
-	TAtomic<uint64> BeginCompressChunksByType[(int8)EIoChunkType::MAX] = { 0 };
-	TAtomic<uint64> SerializedChunksCount{ 0 };
-	TAtomic<uint64> WriteCycleCount{ 0 };
-	TAtomic<uint64> WriteByteCount{ 0 };
-	TAtomic<uint64> ScheduledCompressionTasksCount{ 0 };
-	TAtomic<uint64> CompressionDDCHitCount{ 0 };
-	TAtomic<uint64> CompressionDDCMissCount{ 0 };
-	FCriticalSection AvailableCompressionBuffersCritical;
-	TArray<FIoBuffer*> AvailableCompressionBuffers;
-	int32 CompressionBufferSize = -1;
-	int32 TotalCompressionBufferCount = -1;
-	TArray<TSharedPtr<FIoStoreWriter>> IoStoreWriters;
-
-	friend class FIoStoreWriter;
-};
-
-FIoStoreWriterContext::FIoStoreWriterContext()
-	: Impl(new FIoStoreWriterContextImpl())
-{
-
-}
-
-FIoStoreWriterContext::~FIoStoreWriterContext()
-{
-	delete Impl;
-}
-
-[[nodiscard]] FIoStatus FIoStoreWriterContext::Initialize(const FIoStoreWriterSettings& InWriterSettings)
-{
-	return Impl->Initialize(InWriterSettings);
-}
-
-TSharedPtr<IIoStoreWriter> FIoStoreWriterContext::CreateContainer(const TCHAR* InContainerPath, const FIoContainerSettings& InContainerSettings)
-{
-	return Impl->CreateContainer(InContainerPath, InContainerSettings);
-}
-
-void FIoStoreWriterContext::Flush()
-{
-	Impl->Flush();
-}
-
-FIoStoreWriterContext::FProgress FIoStoreWriterContext::GetProgress() const
-{
-	return Impl->GetProgress();
-}
-
-class FIoStoreTocBuilder
-{
-public:
-	FIoStoreTocBuilder()
-	{
-		FMemory::Memzero(&Toc.Header, sizeof(FIoStoreTocHeader));
-	}
-
-	int32 AddChunkEntry(const FIoChunkId& ChunkId, const FIoOffsetAndLength& OffsetLength, const FIoStoreTocEntryMeta& Meta)
-	{
-		int32& Index = ChunkIdToIndex.FindOrAdd(ChunkId);
-
-		if (!Index)
-		{
-			Index = Toc.ChunkIds.Add(ChunkId);
-			Toc.ChunkOffsetLengths.Add(OffsetLength);
-			Toc.ChunkMetas.Add(Meta);
-
-			return Index;
-		}
-
-		return INDEX_NONE;
-	}
-
-	FIoStoreTocCompressedBlockEntry& AddCompressionBlockEntry()
-	{
-		return Toc.CompressionBlocks.AddDefaulted_GetRef();
-	}
-
-	FSHAHash& AddBlockSignatureEntry()
-	{
-		return Toc.ChunkBlockSignatures.AddDefaulted_GetRef();
-	}
-
-	uint8 AddCompressionMethodEntry(FName CompressionMethod)
-	{
-		if (CompressionMethod == NAME_None)
-		{
-			return 0;
-		}
-
-		uint8 Index = 1;
-		for (const FName& Name : Toc.CompressionMethods)
-		{
-			if (Name == CompressionMethod)
-			{
-				return Index;
-			}
-			++Index;
-		}
-
-		return 1 + uint8(Toc.CompressionMethods.Add(CompressionMethod));
-	}
-
-	void AddToFileIndex(const FIoChunkId& ChunkId, const FString& FileName)
-	{
-		ChunkIdToFileName.Emplace(ChunkId, FileName);
-	}
-
-	FIoStoreTocResource& GetTocResource()
-	{
-		return Toc;
-	}
-
-	const FIoStoreTocResource& GetTocResource() const
-	{
-		return Toc;
-	}
-
-	const int32* GetTocEntryIndex(const FIoChunkId& ChunkId) const
-	{
-		return ChunkIdToIndex.Find(ChunkId);
-	}
-
-	void GetFileNamesToIndex(TArray<FStringView>& OutFileNames) const
-	{
-		OutFileNames.Empty(ChunkIdToFileName.Num());
-		for (auto& ChinkIdAndFileName : ChunkIdToFileName)
-		{
-			OutFileNames.Emplace(ChinkIdAndFileName.Value);
-		}
-	}
-
-	const FString* GetFileName(const FIoChunkId& ChunkId) const
-	{
-		return ChunkIdToFileName.Find(ChunkId);
-	}
-
-	FIoStoreTocChunkInfo GetTocChunkInfo(int32 TocEntryIndex) const
-	{
-		FIoStoreTocChunkInfo ChunkInfo = Toc.GetTocChunkInfo(TocEntryIndex);
-
-		ChunkInfo.FileName = FString::Printf(TEXT("<%s>"), *LexToString(ChunkInfo.ChunkType));
-		ChunkInfo.bHasValidFileName = false;
-
-		return ChunkInfo;
-	}
-
-private:
-	FIoStoreTocResource Toc;
-	TMap<FIoChunkId, int32> ChunkIdToIndex;
-	TMap<FIoChunkId, FString> ChunkIdToFileName;
-};
-
 class FIoStoreTocReader
 {
 public:
@@ -798,1491 +305,6 @@ FIoStoreTocChunkInfo FIoStoreTocResource::GetTocChunkInfo(int32 TocEntryIndex) c
 	return ChunkInfo;
 }
 
-class FIoStoreWriter
-	: public IIoStoreWriter
-{
-public:
-	FIoStoreWriter(const TCHAR* InContainerPathAndBaseFileName)
-		: ContainerPathAndBaseFileName(InContainerPathAndBaseFileName)
-	{
-	}
-
-	void SetReferenceChunkDatabase(TSharedPtr<IIoStoreWriterReferenceChunkDatabase> InReferenceChunkDatabase)
-	{
-		if (InReferenceChunkDatabase.IsValid() == false)
-		{
-			ReferenceChunkDatabase = InReferenceChunkDatabase;
-			return;
-		}
-
-		if (InReferenceChunkDatabase->GetCompressionBlockSize() != WriterContext->GetSettings().CompressionBlockSize)
-		{
-			UE_LOG(LogIoStore, Warning, TEXT("Reference chunk database has a different compression block size than the current writer!"));
-			UE_LOG(LogIoStore, Warning, TEXT("No chunks will match, so ignoring. ReferenceChunkDb: %d, IoStoreWriter: %d"), InReferenceChunkDatabase->GetCompressionBlockSize(), WriterContext->GetSettings().CompressionBlockSize);
-			return;
-		}
-		ReferenceChunkDatabase = InReferenceChunkDatabase;
-		
-		// Add ourselves to the reference chunk db's list of possibles
-		ReferenceChunkDatabase->NotifyAddedToWriter(ContainerSettings.ContainerId, FPaths::GetBaseFilename(TocFilePath));
-	}
-	void SetHashDatabase(TSharedPtr<IIoStoreWriterHashDatabase> InHashDatabase, bool bInVerifyHashDatabase)
-	{
-		HashDatabase = InHashDatabase;
-		bVerifyHashDatabase = bInVerifyHashDatabase;
-	}
-
-	void EnumerateChunks(TFunction<bool(FIoStoreTocChunkInfo&&)>&& Callback) const
-	{
-		const FIoStoreTocResource& TocResource = TocBuilder.GetTocResource();
-
-		for (int32 ChunkIndex = 0; ChunkIndex < TocResource.ChunkIds.Num(); ++ChunkIndex)
-		{
-			FIoStoreTocChunkInfo ChunkInfo = TocBuilder.GetTocChunkInfo(ChunkIndex);
-			if (!Callback(MoveTemp(ChunkInfo)))
-			{
-				break;
-			}
-		}
-	}
-
-	[[nodiscard]] FIoStatus Initialize(FIoStoreWriterContextImpl& InContext, const FIoContainerSettings& InContainerSettings)
-	{
-		WriterContext = &InContext;
-		ContainerSettings = InContainerSettings;
-
-		TocFilePath = ContainerPathAndBaseFileName + TEXT(".utoc");
-		
-		IPlatformFile& Ipf = IPlatformFile::GetPlatformPhysical();
-		Ipf.CreateDirectoryTree(*FPaths::GetPath(TocFilePath));
-
-		FPartition& Partition = Partitions.AddDefaulted_GetRef();
-		Partition.Index = 0;
-
-		return FIoStatus::Ok;
-	}
-
-	virtual void EnableDiskLayoutOrdering(const TArray<TUniquePtr<FIoStoreReader>>& PatchSourceReaders) override
-	{
-		check(!LayoutEntriesHead);
-		check(!Entries.Num());
-		LayoutEntriesHead = new FLayoutEntry();
-		LayoutEntries.Add(LayoutEntriesHead);
-		FLayoutEntry* PrevEntryLink = LayoutEntriesHead;
-
-		for (const TUniquePtr<FIoStoreReader>& PatchSourceReader : PatchSourceReaders)
-		{
-			TArray<TPair<uint64, FLayoutEntry*>> LayoutEntriesWithOffsets;
-			PatchSourceReader->EnumerateChunks([this, &PrevEntryLink, &LayoutEntriesWithOffsets](const FIoStoreTocChunkInfo& ChunkInfo)
-				{
-					FLayoutEntry* PreviousBuildEntry = new FLayoutEntry();
-					PreviousBuildEntry->Hash = ChunkInfo.Hash;
-					PreviousBuildEntry->PartitionIndex = ChunkInfo.PartitionIndex;
-					PreviousBuildEntry->CompressedSize = ChunkInfo.CompressedSize;
-					LayoutEntriesWithOffsets.Emplace(ChunkInfo.Offset, PreviousBuildEntry);
-					PreviousBuildLayoutEntryByChunkId.Add(ChunkInfo.Id, PreviousBuildEntry);
-					return true;
-				});
-
-			// Sort entries by offset
-			Algo::Sort(LayoutEntriesWithOffsets, [](const TPair<uint64, FLayoutEntry*>& A, const TPair<uint64, FLayoutEntry*>& B)
-				{
-					return A.Get<0>() < B.Get<0>();
-				});
-
-			for (const TPair<uint64, FLayoutEntry*>& EntryWithOffset : LayoutEntriesWithOffsets)
-			{
-				FLayoutEntry* PreviousBuildEntry = EntryWithOffset.Get<1>();
-				LayoutEntries.Add(PreviousBuildEntry);
-				PrevEntryLink->Next = PreviousBuildEntry;
-				PreviousBuildEntry->Prev = PrevEntryLink;
-				PrevEntryLink = PreviousBuildEntry;
-			}
-			if (!ContainerSettings.bGenerateDiffPatch)
-			{
-				break;
-			}
-		}
-
-		LayoutEntriesTail = new FLayoutEntry();
-		LayoutEntries.Add(LayoutEntriesTail);
-		PrevEntryLink->Next = LayoutEntriesTail;
-		LayoutEntriesTail->Prev = PrevEntryLink;
-	}
-
-	virtual void Append(const FIoChunkId& ChunkId, IIoStoreWriteRequest* Request, const FIoWriteOptions& WriteOptions) override
-	{
-		//
-		// This function sets up the sequence of events that takes a chunk from source data on disc
-		// to written to a container. The first thing that happens is the source data is read in order
-		// to hash it to detect whether or not it's modified as well as look up in reference databases.
-		// Load the data -> PrepareSourceBufferAsync
-		// Hash the data -> HashTask lambda
-		//
-		// The hash task itself doesn't continue to the next steps - the Flush() call
-		// waits for all hashes to be complete before kicking the next steps.
-		//
-		TRACE_CPUPROFILER_EVENT_SCOPE(AppendWriteRequest);
-		check(!bHasFlushed);
-		checkf(ChunkId.IsValid(), TEXT("ChunkId is not valid!"));
-
-		FIoStoreWriteQueueEntry* Entry = new FIoStoreWriteQueueEntry();
-		Entry->Writer = this;
-		Entry->Sequence = Entries.Num();
-		WriterContext->TotalChunksCount.IncrementExchange();
-		Entries.Add(Entry);
-		Entry->ChunkId = ChunkId;
-		Entry->Options = WriteOptions;
-		Entry->Request = Request;		
-		Entry->BeginCompressionBarrier = FGraphEvent::CreateGraphEvent();
-		Entry->FinishCompressionBarrier = FGraphEvent::CreateGraphEvent();
-		Entry->FinishEncryptionAndSigningBarrier = FGraphEvent::CreateGraphEvent();
-		Entry->BeginWriteBarrier = FGraphEvent::CreateGraphEvent();
-		Entry->WriteFinishedEvent = FGraphEvent::CreateGraphEvent();
-		
-		bool bTestHashValid = false;
-		FIoChunkHash TestHash;
-		// If we can get the hash without reading the whole thing and hashing it, do so
-		// to avoid the IO.
-		if (HashDatabase.IsValid() &&
-			HashDatabase->FindHashForChunkId(ChunkId, Entry->ChunkHash))
-		{
-			if (bVerifyHashDatabase == false)
-			{
-				// If we aren't validating then we just use it and bail.
-				WriterContext->HashDbChunksCount.IncrementExchange();
-				WriterContext->HashDbChunksByType[(int8)Entry->ChunkId.GetChunkType()].IncrementExchange();
-				WriterContext->HashedChunksCount.IncrementExchange();
-
-				if (ReferenceChunkDatabase.IsValid() && CompressionMethodForEntry(Entry) != NAME_None)
-				{
-					Entry->bLoadingFromReferenceDb = ReferenceChunkDatabase->ChunkExists(ContainerSettings.ContainerId, Entry->ChunkHash, Entry->ChunkId, Entry->NumChunkBlocksFromRefDb);
-					Entry->bCouldBeFromReferenceDb = true;
-				}
-				return;
-			}
-
-			// If we are validating, copy the hash out and run the normal path
-			// to check against it.
-			TestHash = Entry->ChunkHash;
-			bTestHashValid = true;
-		}
-		// Otherwise, we have to do the load & hash
-
-		Entry->HashBarrier = FGraphEvent::CreateGraphEvent();
-
-		FGraphEventArray HashPrereqs;
-		HashPrereqs.Add(Entry->HashBarrier);
-		Entry->HashTask = FFunctionGraphTask::CreateAndDispatchWhenReady([this, Entry, TestHash, bTestHashValid]()
-		{
-			TRACE_CPUPROFILER_EVENT_SCOPE(HashChunk);
-			const FIoBuffer* SourceBuffer = Entry->Request->GetSourceBuffer();
-			Entry->ChunkHash = FIoChunkHash::HashBuffer(SourceBuffer->Data(), SourceBuffer->DataSize());
-			WriterContext->HashedChunksCount.IncrementExchange();
-
-			if (bVerifyHashDatabase && bTestHashValid)
-			{
-				if (TestHash != Entry->ChunkHash)
-				{
-					UE_LOG(LogIoStore, Warning, TEXT("HashDb Validation Failed: ChunkId %s has mismatching hash"), *LexToString(Entry->ChunkId));
-				}
-			}
-
-
-			if (ReferenceChunkDatabase.IsValid() && CompressionMethodForEntry(Entry) != NAME_None)
-			{
-				Entry->bLoadingFromReferenceDb = ReferenceChunkDatabase->ChunkExists(ContainerSettings.ContainerId, Entry->ChunkHash, Entry->ChunkId, Entry->NumChunkBlocksFromRefDb);
-				Entry->bCouldBeFromReferenceDb = true;
-			}
-
-			// Release the source data buffer, it will be reloaded later when we start compressing the chunk
-			Entry->Request->FreeSourceBuffer();
-		}, TStatId(), &HashPrereqs, ENamedThreads::AnyHiPriThreadHiPriTask);
-		Entry->Request->PrepareSourceBufferAsync(Entry->HashBarrier);
-	}
-
-	virtual void Append(const FIoChunkId& ChunkId, FIoBuffer Chunk, const FIoWriteOptions& WriteOptions, uint64 OrderHint) override
-	{
-		struct FWriteRequest
-			: IIoStoreWriteRequest
-		{
-			FWriteRequest(FIoBuffer InSourceBuffer, uint64 InOrderHint)
-				: OrderHint(InOrderHint)
-			{
-				SourceBuffer = InSourceBuffer;
-				SourceBuffer.MakeOwned();
-			}
-
-			virtual ~FWriteRequest() = default;
-
-			void PrepareSourceBufferAsync(FGraphEventRef CompletionEvent) override
-			{
-				CompletionEvent->DispatchSubsequents();
-			}
-
-			const FIoBuffer* GetSourceBuffer() override
-			{
-				return &SourceBuffer;
-			}
-
-			void FreeSourceBuffer() override
-			{
-			}
-
-			uint64 GetOrderHint() override
-			{
-				return OrderHint;
-			}
-
-			TArrayView<const FFileRegion> GetRegions()
-			{
-				return TArrayView<const FFileRegion>();
-			}
-
-			FIoBuffer SourceBuffer;
-			uint64 OrderHint;
-		};
-
-		Append(ChunkId, new FWriteRequest(Chunk, OrderHint), WriteOptions);
-	}
-
-	bool GeneratePerfectHashes(FIoStoreTocResource& TocResource, const TCHAR* ContainerDebugName)
-	{
-		// https://en.wikipedia.org/wiki/Perfect_hash_function
-		TRACE_CPUPROFILER_EVENT_SCOPE(TocGeneratePerfectHashes);
-		uint32 ChunkCount = TocResource.ChunkIds.Num();
-		uint32 SeedCount = FMath::Max(1, FMath::RoundToInt32(ChunkCount / 2.0));
-		check(TocResource.ChunkOffsetLengths.Num() == ChunkCount);
-		
-		TArray<FIoChunkId> OutTocChunkIds;
-		OutTocChunkIds.SetNum(ChunkCount);
-		TArray<FIoOffsetAndLength> OutTocOffsetAndLengths;
-		OutTocOffsetAndLengths.SetNum(ChunkCount);
-		TArray<FIoStoreTocEntryMeta> OutTocChunkMetas;
-		OutTocChunkMetas.SetNum(ChunkCount);
-		TArray<int32> OutTocChunkHashSeeds;
-		OutTocChunkHashSeeds.SetNumZeroed(SeedCount);
-		TArray<int32> OutTocChunkIndicesWithoutPerfectHash;
-
-		TArray<TArray<int32>> Buckets;
-		Buckets.SetNum(SeedCount);
-
-		TBitArray<> FreeSlots(true, ChunkCount);
-		// Put each chunk in a bucket, each bucket contains the chunk ids that have colliding hashes
-		for (uint32 ChunkIndex = 0; ChunkIndex < ChunkCount; ++ChunkIndex)
-		{
-			const FIoChunkId& ChunkId = TocResource.ChunkIds[ChunkIndex];
-			Buckets[FIoStoreTocResource::HashChunkIdWithSeed(0, ChunkId) % SeedCount].Add(ChunkIndex);
-		}
-
-		uint64 TotalIterationCount = 0;
-		uint64 TotalOverflowBucketsCount = 0;
-		
-		// For each bucket containing more than one chunk id find a seed that makes its chunk ids
-		// hash to unused slots in the output array
-		Algo::Sort(Buckets, [](const TArray<int32>& A, const TArray<int32>& B)
-			{
-				return A.Num() > B.Num();
-			});
-		for (uint32 BucketIndex = 0; BucketIndex < SeedCount; ++BucketIndex)
-		{
-			const TArray<int32>& Bucket = Buckets[BucketIndex];
-			if (Bucket.Num() <= 1)
-			{
-				break;
-			}
-			uint64 BucketHash = FIoStoreTocResource::HashChunkIdWithSeed(0, TocResource.ChunkIds[Bucket[0]]);
-
-			static constexpr uint32 Primes[] = {
-				2, 3, 5, 7, 11, 13, 17, 19, 23, 29, 31, 37, 41, 43, 47, 53, 59, 61, 67, 71, 73, 79,
-				83, 89, 97, 101, 103, 107, 109, 113, 127, 131, 137, 139, 149, 151, 157, 163, 167,
-				173, 179, 181, 191, 193, 197, 199, 211, 223, 227, 229, 233, 239, 241, 251, 257, 263,
-				269, 271, 277, 281, 283, 293, 307, 311, 313, 317, 331, 337, 347, 349, 353, 359, 367,
-				373, 379, 383, 389, 397, 401, 409, 419, 421, 431, 433, 439, 443, 449, 457, 461, 463,
-				467, 479, 487, 491, 499, 503, 509, 521, 523, 541, 547, 557, 563, 569, 571, 577, 587,
-				593, 599, 601, 607, 613, 617, 619, 631, 641, 643, 647, 653, 659, 661, 673, 677, 683,
-				691, 701, 709, 719, 727, 733, 739, 743, 751, 757, 761, 769, 773, 787, 797, 809, 811,
-				821, 823, 827, 829, 839, 853, 857, 859, 863, 877, 881, 883, 887, 907, 911, 919, 929,
-				937, 941, 947, 953, 967, 971, 977, 983, 991, 997, 1009, 1013, 1019, 1021, 1031, 1033,
-				1039, 1049, 1051, 1061, 1063, 1069, 1087, 1091, 1093, 1097, 1103, 1109, 1117, 1123,
-				1129, 1151, 1153, 1163, 1171, 1181, 1187, 1193, 1201, 1213, 1217, 1223, 1229, 1231,
-				1237, 1249, 1259, 1277, 1279, 1283, 1289, 1291, 1297, 1301, 1303, 1307, 1319, 1321,
-				1327, 1361, 1367, 1373, 1381, 1399, 1409, 1423, 1427, 1429, 1433, 1439, 1447, 1451,
-				1453, 1459, 1471, 1481, 1483, 1487, 1489, 1493, 1499, 1511, 1523, 1531, 1543, 1549,
-				1553, 1559, 1567, 1571, 1579, 1583, 1597, 1601, 1607, 1609, 1613, 1619, 1621, 1627,
-				1637, 1657, 1663, 1667, 1669, 1693, 1697, 1699, 1709, 1721, 1723, 1733, 1741, 1747,
-				1753, 1759, 1777, 1783, 1787, 1789, 1801, 1811, 1823, 1831, 1847, 1861, 1867, 1871,
-				1873, 1877, 1879, 1889, 1901, 1907, 1913, 1931, 1933, 1949, 1951, 1973, 1979, 1987,
-				1993, 1997, 1999, 2003, 2011, 2017, 2027, 2029, 2039, 2053, 2063, 2069, 2081, 2083,
-				2087, 2089, 2099, 2111, 2113, 2129, 2131, 2137, 2141, 2143, 2153, 2161, 2179, 2203,
-				2207, 2213, 2221, 2237, 2239, 2243, 2251, 2267, 2269, 2273, 2281, 2287, 2293, 2297,
-				2309, 2311, 2333, 2339, 2341, 2347, 2351, 2357, 2371, 2377, 2381, 2383, 2389, 2393,
-				2399, 2411, 2417, 2423, 2437, 2441, 2447, 2459, 2467, 2473, 2477, 2503, 2521, 2531,
-				2539, 2543, 2549, 2551, 2557, 2579, 2591, 2593, 2609, 2617, 2621, 2633, 2647, 2657,
-				2659, 2663, 2671, 2677, 2683, 2687, 2689, 2693, 2699, 2707, 2711, 2713, 2719, 2729,
-				2731, 2741, 2749, 2753, 2767, 2777, 2789, 2791, 2797, 2801, 2803, 2819, 2833, 2837,
-				2843, 2851, 2857, 2861, 2879, 2887, 2897, 2903, 2909, 2917, 2927, 2939, 2953, 2957,
-				2963, 2969, 2971, 2999, 3001, 3011, 3019, 3023, 3037, 3041, 3049, 3061, 3067, 3079,
-				3083, 3089, 3109, 3119, 3121, 3137, 3163, 3167, 3169, 3181, 3187, 3191, 3203, 3209,
-				3217, 3221, 3229, 3251, 3253, 3257, 3259, 3271, 3299, 3301, 3307, 3313, 3319, 3323,
-				3329, 3331, 3343, 3347, 3359, 3361, 3371, 3373, 3389, 3391, 3407, 3413, 3433, 3449,
-				3457, 3461, 3463, 3467, 3469, 3491, 3499, 3511, 3517, 3527, 3529, 3533, 3539, 3541,
-				3547, 3557, 3559, 3571, 3581, 3583, 3593, 3607, 3613, 3617, 3623, 3631, 3637, 3643,
-				3659, 3671, 3673, 3677, 3691, 3697, 3701, 3709, 3719, 3727, 3733, 3739, 3761, 3767,
-				3769, 3779, 3793, 3797, 3803, 3821, 3823, 3833, 3847, 3851, 3853, 3863, 3877, 3881,
-				3889, 3907, 3911, 3917, 3919, 3923, 3929, 3931, 3943, 3947, 3967, 3989, 4001, 4003,
-				4007, 4013, 4019, 4021, 4027, 4049, 4051, 4057, 4073, 4079, 4091, 4093, 4099, 4111,
-				4127, 4129, 4133, 4139, 4153, 4157, 4159, 4177, 4201, 4211, 4217, 4219, 4229, 4231,
-				4241, 4243, 4253, 4259, 4261, 4271, 4273, 4283, 4289, 4297, 4327, 4337, 4339, 4349,
-				4357, 4363, 4373, 4391, 4397, 4409, 4421, 4423, 4441, 4447, 4451, 4457, 4463, 4481,
-				4483, 4493, 4507, 4513, 4517, 4519, 4523, 4547, 4549, 4561, 4567, 4583, 4591, 4597,
-				4603, 4621, 4637, 4639, 4643, 4649, 4651, 4657, 4663, 4673, 4679, 4691, 4703, 4721,
-				4723, 4729, 4733, 4751, 4759, 4783, 4787, 4789, 4793, 4799, 4801, 4813, 4817, 4831,
-				4861, 4871, 4877, 4889, 4903, 4909, 4919, 4931, 4933, 4937, 4943, 4951, 4957, 4967,
-				4969, 4973, 4987, 4993, 4999, 5003, 5009, 5011, 5021, 5023, 5039, 5051, 5059, 5077,
-				5081, 5087, 5099, 5101, 5107, 5113, 5119, 5147, 5153, 5167, 5171, 5179, 5189, 5197,
-				5209, 5227, 5231, 5233, 5237, 5261, 5273, 5279, 5281, 5297, 5303, 5309, 5323, 5333,
-				5347, 5351, 5381, 5387, 5393, 5399, 5407, 5413, 5417, 5419, 5431, 5437, 5441, 5443,
-				5449, 5471, 5477, 5479, 5483, 5501, 5503, 5507, 5519, 5521, 5527, 5531, 5557, 5563,
-				5569, 5573, 5581, 5591, 5623, 5639, 5641, 5647, 5651, 5653, 5657, 5659, 5669, 5683,
-				5689, 5693, 5701, 5711, 5717, 5737, 5741, 5743, 5749, 5779, 5783, 5791, 5801, 5807,
-				5813, 5821, 5827, 5839, 5843, 5849, 5851, 5857, 5861, 5867, 5869, 5879, 5881, 5897,
-				5903, 5923, 5927, 5939, 5953, 5981, 5987, 6007, 6011, 6029, 6037, 6043, 6047, 6053,
-				6067, 6073, 6079, 6089, 6091, 6101, 6113, 6121, 6131, 6133, 6143, 6151, 6163, 6173,
-				6197, 6199, 6203, 6211, 6217, 6221, 6229, 6247, 6257, 6263, 6269, 6271, 6277, 6287,
-				6299, 6301, 6311, 6317, 6323, 6329, 6337, 6343, 6353, 6359, 6361, 6367, 6373, 6379,
-				6389, 6397, 6421, 6427, 6449, 6451, 6469, 6473, 6481, 6491, 6521, 6529, 6547, 6551,
-				6553, 6563, 6569, 6571, 6577, 6581, 6599, 6607, 6619, 6637, 6653, 6659, 6661, 6673,
-				6679, 6689, 6691, 6701, 6703, 6709, 6719, 6733, 6737, 6761, 6763, 6779, 6781, 6791,
-				6793, 6803, 6823, 6827, 6829, 6833, 6841, 6857, 6863, 6869, 6871, 6883, 6899, 6907,
-				6911, 6917, 6947, 6949, 6959, 6961, 6967, 6971, 6977, 6983, 6991, 6997, 7001, 7013,
-				7019, 7027, 7039, 7043, 7057, 7069, 7079, 7103, 7109, 7121, 7127, 7129, 7151, 7159,
-				7177, 7187, 7193, 7207, 7211, 7213, 7219, 7229, 7237, 7243, 7247, 7253, 7283, 7297,
-				7307, 7309, 7321, 7331, 7333, 7349, 7351, 7369, 7393, 7411, 7417, 7433, 7451, 7457,
-				7459, 7477, 7481, 7487, 7489, 7499, 7507, 7517, 7523, 7529, 7537, 7541, 7547, 7549,
-				7559, 7561, 7573, 7577, 7583, 7589, 7591, 7603, 7607, 7621, 7639, 7643, 7649, 7669,
-				7673, 7681, 7687, 7691, 7699, 7703, 7717, 7723, 7727, 7741, 7753, 7757, 7759, 7789,
-				7793, 7817, 7823, 7829, 7841, 7853, 7867, 7873, 7877, 7879, 7883, 7901, 7907, 7919
-			};
-			static constexpr uint32 MaxIterations = UE_ARRAY_COUNT(Primes);
-
-			uint32 PrimeIndex = 0;
-			TBitArray<> BucketUsedSlots(false, ChunkCount);
-			int32 IndexInBucket = 0;
-			bool bFoundSeedForBucket = true;
-			uint64 BucketIterationCount = 0;
-			while (IndexInBucket < Bucket.Num())
-			{
-				++BucketIterationCount;
-				const FIoChunkId& ChunkId = TocResource.ChunkIds[Bucket[IndexInBucket]];
-				uint32 Seed = Primes[PrimeIndex];
-				uint32 Slot = FIoStoreTocResource::HashChunkIdWithSeed(Seed, ChunkId) % ChunkCount;
-				if (!FreeSlots[Slot] || BucketUsedSlots[Slot])
-				{
-					++PrimeIndex;
-					if (PrimeIndex == MaxIterations)
-					{
-						// Unable to resolve collisions for this bucket, put items in the overflow list and
-						// save the negative index of the first item in the bucket as the seed
-						// (-ChunkCount - 1 to separate from the single item buckets below)
-						UE_LOG(LogIoStore, Verbose, TEXT("%s: Failed finding seed for bucket with %d items after %d iterations."), ContainerDebugName, Bucket.Num(), BucketIterationCount);
-						bFoundSeedForBucket = false;
-						OutTocChunkHashSeeds[BucketHash % SeedCount] = -OutTocChunkIndicesWithoutPerfectHash.Num() - ChunkCount - 1;
-						OutTocChunkIndicesWithoutPerfectHash.Append(Bucket);
-						++TotalOverflowBucketsCount;
-						break;
-
-					}
-					IndexInBucket = 0;
-					BucketUsedSlots.Init(false, ChunkCount);
-				}
-				else
-				{
-					BucketUsedSlots[Slot] = true;
-					++IndexInBucket;
-				}
-			}
-
-			TotalIterationCount += BucketIterationCount;
-
-			if (bFoundSeedForBucket)
-			{
-				uint32 Seed = Primes[PrimeIndex];
-				OutTocChunkHashSeeds[BucketHash % SeedCount] = Seed;
-				for (IndexInBucket = 0; IndexInBucket < Bucket.Num(); ++IndexInBucket)
-				{
-					int32 ChunkIndex = Bucket[IndexInBucket];
-					const FIoChunkId& ChunkId = TocResource.ChunkIds[ChunkIndex];
-					uint32 Slot = FIoStoreTocResource::HashChunkIdWithSeed(Seed, ChunkId) % ChunkCount;
-					check(FreeSlots[Slot]);
-					FreeSlots[Slot] = false;
-					OutTocChunkIds[Slot] = ChunkId;
-					OutTocOffsetAndLengths[Slot] = TocResource.ChunkOffsetLengths[ChunkIndex];
-					OutTocChunkMetas[Slot] = TocResource.ChunkMetas[ChunkIndex];
-				}
-			}
-		}
-
-		// For the remaining buckets with only one chunk id put that chunk id in the first empty position in
-		// the output array and store the index as a negative seed for the bucket (-1 to allow use of slot 0)
-		TConstSetBitIterator<> FreeSlotIt(FreeSlots);
-		for (uint32 BucketIndex = 0; BucketIndex < SeedCount; ++BucketIndex)
-		{
-			const TArray<int32>& Bucket = Buckets[BucketIndex];
-			if (Bucket.Num() == 1)
-			{
-				uint32 Slot = FreeSlotIt.GetIndex();
-				++FreeSlotIt;
-				int32 ChunkIndex = Bucket[0];
-				const FIoChunkId& ChunkId = TocResource.ChunkIds[ChunkIndex];
-				uint64 BucketHash = FIoStoreTocResource::HashChunkIdWithSeed(0, ChunkId);
-				OutTocChunkHashSeeds[BucketHash % SeedCount] = -static_cast<int32>(Slot) - 1;
-				OutTocChunkIds[Slot] = ChunkId;
-				OutTocOffsetAndLengths[Slot] = TocResource.ChunkOffsetLengths[ChunkIndex];
-				OutTocChunkMetas[Slot] = TocResource.ChunkMetas[ChunkIndex];
-			}
-		}
-
-		if (!OutTocChunkIndicesWithoutPerfectHash.IsEmpty())
-		{
-			// Put overflow items in the remaining free slots and update the index for each overflow entry
-			UE_LOG(LogIoStore, Display, TEXT("%s: Failed finding perfect hashmap for %d items. %d overflow buckets with %d items."), ContainerDebugName, ChunkCount, TotalOverflowBucketsCount, OutTocChunkIndicesWithoutPerfectHash.Num());
-			for (int32& OverflowEntryIndex : OutTocChunkIndicesWithoutPerfectHash)
-			{
-				uint32 Slot = FreeSlotIt.GetIndex();
-				++FreeSlotIt;
-				const FIoChunkId& ChunkId = TocResource.ChunkIds[OverflowEntryIndex];
-				OutTocChunkIds[Slot] = ChunkId;
-				OutTocOffsetAndLengths[Slot] = TocResource.ChunkOffsetLengths[OverflowEntryIndex];
-				OutTocChunkMetas[Slot] = TocResource.ChunkMetas[OverflowEntryIndex];
-				OverflowEntryIndex = Slot;
-			}
-		}
-		else
-		{
-			UE_LOG(LogIoStore, Display, TEXT("%s: Found perfect hashmap for %d items."), ContainerDebugName, ChunkCount);
-		}
-		double AverageIterationCount = ChunkCount > 0 ? static_cast<double>(TotalIterationCount) / ChunkCount : 0.0;
-		UE_LOG(LogIoStore, Verbose, TEXT("%s: %f iterations/chunk"), ContainerDebugName, AverageIterationCount);
-
-		TocResource.ChunkIds = MoveTemp(OutTocChunkIds);
-		TocResource.ChunkOffsetLengths = MoveTemp(OutTocOffsetAndLengths);
-		TocResource.ChunkMetas = MoveTemp(OutTocChunkMetas);
-		TocResource.ChunkPerfectHashSeeds = MoveTemp(OutTocChunkHashSeeds);
-		TocResource.ChunkIndicesWithoutPerfectHash = MoveTemp(OutTocChunkIndicesWithoutPerfectHash);
-
-		return true;
-	}
-
-	void Finalize()
-	{
-		check(bHasFlushed);
-
-		UncompressedContainerSize = TotalEntryUncompressedSize + TotalPaddingSize;
-		CompressedContainerSize = 0;
-		const FIoStoreWriterSettings& WriterSettings = WriterContext->GetSettings();
-		for (FPartition& Partition : Partitions)
-		{
-			CompressedContainerSize += Partition.Offset;
-
-			if (bHasMemoryMappedEntry)
-			{
-				uint64 ExtraPaddingBytes = Align(Partition.Offset, WriterSettings.MemoryMappingAlignment) - Partition.Offset;
-				if (ExtraPaddingBytes)
-				{
-					TArray<uint8> Padding;
-					Padding.SetNumZeroed(int32(ExtraPaddingBytes));
-					Partition.ContainerFileHandle->Serialize(Padding.GetData(), ExtraPaddingBytes);
-					CompressedContainerSize += ExtraPaddingBytes;
-					UncompressedContainerSize += ExtraPaddingBytes;
-					Partition.Offset += ExtraPaddingBytes;
-					TotalPaddingSize += ExtraPaddingBytes;
-				}
-			}
-			
-			if (Partition.ContainerFileHandle)
-			{
-				TRACE_CPUPROFILER_EVENT_SCOPE(FlushContainerFile);
-				Partition.ContainerFileHandle->Flush();
-				check(Partition.ContainerFileHandle->Tell() == Partition.Offset);
-			}
-
-			if (Partition.RegionsArchive)
-			{
-				FFileRegion::SerializeFileRegions(*Partition.RegionsArchive.Get(), Partition.AllFileRegions);
-				Partition.RegionsArchive->Flush();
-			}
-		}
-
-		FIoStoreTocResource& TocResource = TocBuilder.GetTocResource();
-
-		GeneratePerfectHashes(TocResource, *FPaths::GetBaseFilename(TocFilePath));
-
-		if (ContainerSettings.IsIndexed())
-		{
-			TRACE_CPUPROFILER_EVENT_SCOPE(BuildIndex);
-			TArray<FStringView> FilesToIndex;
-			TocBuilder.GetFileNamesToIndex(FilesToIndex);
-
-			FString MountPoint = IoDirectoryIndexUtils::GetCommonRootPath(FilesToIndex);
-			FIoDirectoryIndexWriter DirectoryIndexWriter;
-			DirectoryIndexWriter.SetMountPoint(MountPoint);
-
-			uint32 TocEntryIndex = 0;
-			for (const FIoChunkId& ChunkId : TocResource.ChunkIds)
-			{
-				const FString* ChunkFileName = TocBuilder.GetFileName(ChunkId);
-				if (ChunkFileName)
-				{
-					const uint32 FileEntryIndex = DirectoryIndexWriter.AddFile(*ChunkFileName);
-					check(FileEntryIndex != ~uint32(0));
-					DirectoryIndexWriter.SetFileUserData(FileEntryIndex, TocEntryIndex);
-				}
-				++TocEntryIndex;
-			}
-
-			DirectoryIndexWriter.Flush(
-				TocResource.DirectoryIndexBuffer,
-				ContainerSettings.IsEncrypted() ? ContainerSettings.EncryptionKey : FAES::FAESKey());
-		}
-
-		TIoStatusOr<uint64> TocSize = FIoStoreTocResource::Write(*TocFilePath, TocResource, static_cast<uint32>(WriterSettings.CompressionBlockSize), WriterSettings.MaxPartitionSize, ContainerSettings);
-		check(TocSize.IsOk());
-		
-		Result.ContainerId = ContainerSettings.ContainerId;
-		Result.ContainerName = FPaths::GetBaseFilename(TocFilePath);
-		Result.ContainerFlags = ContainerSettings.ContainerFlags;
-		Result.TocSize = TocSize.ConsumeValueOrDie();
-		Result.TocEntryCount = TocResource.Header.TocEntryCount;
-		Result.PaddingSize = TotalPaddingSize;
-		Result.UncompressedContainerSize = UncompressedContainerSize;
-		Result.CompressedContainerSize = CompressedContainerSize;
-		Result.TotalEntryCompressedSize = TotalEntryCompressedSize;
-		Result.ReferenceCacheMissBytes = ReferenceCacheMissBytes;
-		Result.DirectoryIndexSize = TocResource.Header.DirectoryIndexSize;
-		Result.CompressionMethod = EnumHasAnyFlags(ContainerSettings.ContainerFlags, EIoContainerFlags::Compressed)
-			? WriterSettings.CompressionMethod
-			: NAME_None;
-		Result.ModifiedChunksCount = 0;
-		Result.AddedChunksCount = 0;
-		Result.ModifiedChunksSize= 0;
-		Result.AddedChunksSize = 0;
-		{
-			TRACE_CPUPROFILER_EVENT_SCOPE(Cleanup);
-			for (FIoStoreWriteQueueEntry* Entry : Entries)
-			{
-				if (Entry->bModified)
-				{
-					++Result.ModifiedChunksCount;
-					Result.ModifiedChunksSize += Entry->CompressedSize;
-				}
-				else if (Entry->bAdded)
-				{
-					++Result.AddedChunksCount;
-					Result.AddedChunksSize += Entry->CompressedSize;
-				}
-				delete Entry;
-			}
-		}
-
-		Entries.Empty();
-		bHasResult = true;
-	}
-
-	TIoStatusOr<FIoStoreWriterResult> GetResult()
-	{
-		if (!bHasResult)
-		{
-			return FIoStatus::Invalid;
-		}
-		return Result;
-	}
-
-private:
-	struct FPartition
-	{
-		TUniquePtr<FArchive> ContainerFileHandle;
-		TUniquePtr<FArchive> RegionsArchive;
-		uint64 Offset = 0;
-		uint64 ReservedSpace = 0;
-		TArray<FFileRegion> AllFileRegions;
-		int32 Index = -1;
-	};
-
-	struct FLayoutEntry
-	{
-		FLayoutEntry* Prev = nullptr;
-		FLayoutEntry* Next = nullptr;
-		uint64 IdealOrder = 0;
-		uint64 CompressedSize = uint64(-1);
-		FIoChunkHash Hash;
-		FIoStoreWriteQueueEntry* QueueEntry = nullptr;
-		int32 PartitionIndex = -1;
-	};
-
-	void FinalizeLayout()
-	{
-		TRACE_CPUPROFILER_EVENT_SCOPE(FinalizeLayout);
-		
-		Algo::Sort(Entries, [](const FIoStoreWriteQueueEntry* A, const FIoStoreWriteQueueEntry* B)
-		{
-			uint64 AOrderHint = A->Request->GetOrderHint();
-			uint64 BOrderHint = B->Request->GetOrderHint();
-			if (AOrderHint != BOrderHint)
-			{
-				return AOrderHint < BOrderHint;
-			}
-			return A->Sequence < B->Sequence;
-		});
-
-		TMap<int64, FLayoutEntry*> LayoutEntriesByOrderMap;
-		int64 IdealOrder = 0;
-		TArray<FLayoutEntry*> UnassignedEntries;
-		for (FIoStoreWriteQueueEntry* WriteQueueEntry : Entries)
-		{
-			FLayoutEntry* FindPreviousEntry = PreviousBuildLayoutEntryByChunkId.FindRef(WriteQueueEntry->ChunkId);
-			if (FindPreviousEntry)
-			{
-				if (FindPreviousEntry->Hash != WriteQueueEntry->ChunkHash)
-				{
-					WriteQueueEntry->bModified = true;
-				}
-				else
-				{
-					FindPreviousEntry->QueueEntry = WriteQueueEntry;
-					FindPreviousEntry->IdealOrder = IdealOrder;
-					WriteQueueEntry->PartitionIndex = FindPreviousEntry->PartitionIndex;
-				}
-			}
-			else
-			{
-				WriteQueueEntry->bAdded = true;
-			}
-			if (WriteQueueEntry->bModified || WriteQueueEntry->bAdded)
-			{
-				FLayoutEntry* NewLayoutEntry = new FLayoutEntry();
-				NewLayoutEntry->QueueEntry = WriteQueueEntry;
-				NewLayoutEntry->IdealOrder = IdealOrder;
-				LayoutEntries.Add(NewLayoutEntry);
-				UnassignedEntries.Add(NewLayoutEntry);
-			}
-			++IdealOrder;
-		}
-			
-		if (ContainerSettings.bGenerateDiffPatch)
-		{
-			LayoutEntriesHead->Next = LayoutEntriesTail;
-			LayoutEntriesTail->Prev = LayoutEntriesHead;
-		}
-		else
-		{
-			for (FLayoutEntry* EntryIt = LayoutEntriesHead->Next; EntryIt != LayoutEntriesTail; EntryIt = EntryIt->Next)
-			{
-				if (!EntryIt->QueueEntry)
-				{
-					EntryIt->Prev->Next = EntryIt->Next;
-					EntryIt->Next->Prev = EntryIt->Prev;
-				}
-				else
-				{
-					LayoutEntriesByOrderMap.Add(EntryIt->IdealOrder, EntryIt);
-				}
-			}
-		}
-		FLayoutEntry* LastAddedEntry = LayoutEntriesHead;
-		for (FLayoutEntry* UnassignedEntry : UnassignedEntries)
-		{
-			check(UnassignedEntry->QueueEntry);
-			FLayoutEntry* PutAfterEntry = LayoutEntriesByOrderMap.FindRef(UnassignedEntry->IdealOrder - 1);
-			if (!PutAfterEntry)
-			{
-				PutAfterEntry = LastAddedEntry;
-			}
-
-			UnassignedEntry->Prev = PutAfterEntry;
-			UnassignedEntry->Next = PutAfterEntry->Next;
-			PutAfterEntry->Next->Prev = UnassignedEntry;
-			PutAfterEntry->Next = UnassignedEntry;
-			LayoutEntriesByOrderMap.Add(UnassignedEntry->IdealOrder, UnassignedEntry);
-			LastAddedEntry = UnassignedEntry;
-		}
-
-		TArray<FIoStoreWriteQueueEntry*> IncludedQueueEntries;
-		for (FLayoutEntry* EntryIt = LayoutEntriesHead->Next; EntryIt != LayoutEntriesTail; EntryIt = EntryIt->Next)
-		{
-			check(EntryIt->QueueEntry);
-			IncludedQueueEntries.Add(EntryIt->QueueEntry);
-			int32 ReserveInPartitionIndex = EntryIt->QueueEntry->PartitionIndex;
-			if (ReserveInPartitionIndex >= 0)
-			{
-				while (Partitions.Num() <= ReserveInPartitionIndex)
-				{
-					FPartition& NewPartition = Partitions.AddDefaulted_GetRef();
-					NewPartition.Index = Partitions.Num() - 1;
-				}
-				FPartition& ReserveInPartition = Partitions[ReserveInPartitionIndex];
-				check(EntryIt->CompressedSize != uint64(-1));
-				ReserveInPartition.ReservedSpace += EntryIt->CompressedSize;
-			}
-		}
-		Swap(Entries, IncludedQueueEntries);
-
-		LayoutEntriesHead = nullptr;
-		LayoutEntriesTail = nullptr;
-		PreviousBuildLayoutEntryByChunkId.Empty();
-		for (FLayoutEntry* Entry : LayoutEntries)
-		{
-			delete Entry;
-		}
-		LayoutEntries.Empty();
-	}
-
-	FIoStatus CreatePartitionContainerFile(FPartition& Partition)
-	{
-		check(!Partition.ContainerFileHandle);
-		FString ContainerFilePath = ContainerPathAndBaseFileName;
-		if (Partition.Index > 0)
-		{
-			ContainerFilePath += FString::Printf(TEXT("_s%d"), Partition.Index);
-		}
-		ContainerFilePath += TEXT(".ucas");
-		
-		Partition.ContainerFileHandle.Reset(IFileManager::Get().CreateFileWriter(*ContainerFilePath));
-		if (!Partition.ContainerFileHandle)
-		{
-			return FIoStatusBuilder(EIoErrorCode::FileOpenFailed) << TEXT("Failed to open IoStore container file '") << *ContainerFilePath << TEXT("'");
-		}
-		if (WriterContext->GetSettings().bEnableFileRegions)
-		{
-			FString RegionsFilePath = ContainerFilePath + FFileRegion::RegionsFileExtension;
-			Partition.RegionsArchive.Reset(IFileManager::Get().CreateFileWriter(*RegionsFilePath));
-			if (!Partition.RegionsArchive)
-			{
-				return FIoStatusBuilder(EIoErrorCode::FileOpenFailed) << TEXT("Failed to open IoStore regions file '") << *RegionsFilePath << TEXT("'");
-			}
-		}
-
-		return FIoStatus::Ok;
-	}
-
-	void CompressBlock(FChunkBlock* Block)
-	{
-		TRACE_CPUPROFILER_EVENT_SCOPE(CompressBlock);
-		check(Block->CompressionMethod != NAME_None);
-		uint64 CompressedBlockSize = Block->IoBuffer->DataSize();
-		bool bCompressed;
-		{
-			if (!FCompression::CompressMemoryIfWorthDecompressing(
-				Block->CompressionMethod,
-				bCompressed,
-				(int64)WriterContext->WriterSettings.CompressionMinBytesSaved,
-				WriterContext->WriterSettings.CompressionMinPercentSaved,
-				Block->IoBuffer->Data(),
-				(int64&)CompressedBlockSize,
-				Block->UncompressedData,
-				(int64)Block->UncompressedSize,
-				COMPRESS_ForPackaging))
-			{
-				UE_LOG(LogIoStore, Error, TEXT("Compression failed: Method=%s, CompressedSize=0x%llx, UncompressedSize=0x%llx"), 
-					*Block->CompressionMethod.ToString(), CompressedBlockSize, Block->UncompressedSize);
-				bCompressed = false;
-			}
-		}
-		if (!bCompressed)
-		{
-			Block->CompressionMethod = NAME_None;
-			Block->CompressedSize = Block->UncompressedSize;
-			FMemory::Memcpy(Block->IoBuffer->Data(), Block->UncompressedData, Block->UncompressedSize);
-		}
-		else
-		{
-			check(CompressedBlockSize > 0);
-			check(CompressedBlockSize < Block->UncompressedSize);
-			Block->CompressedSize = CompressedBlockSize;
-		}
-	}
-
-	bool SerializeCompressedDDCData(FIoStoreWriteQueueEntry* Entry, FArchive& Ar)
-	{
-		uint32 NumBlocks = Entry->ChunkBlocks.Num();
-		Ar << NumBlocks;
-		if (NumBlocks != Entry->ChunkBlocks.Num())
-		{
-			return false;
-		}
-		for (FChunkBlock& Block : Entry->ChunkBlocks)
-		{
-			Ar << Block.CompressedSize;
-			if (Block.CompressedSize > Block.UncompressedSize)
-			{
-				return false;
-			}
-			if (Ar.IsLoading() && Block.CompressedSize == Block.UncompressedSize)
-			{
-				Block.CompressionMethod = NAME_None;
-			}
-			if (Block.IoBuffer->DataSize() < Block.CompressedSize)
-			{
-				return false;
-			}
-			Ar.Serialize(Block.IoBuffer->Data(), Block.CompressedSize);
-		}
-		return !Ar.IsError();
-	}
-
-	FName CompressionMethodForEntry(FIoStoreWriteQueueEntry* Entry) const
-	{
-		FName CompressionMethod = NAME_None;
-		const FIoStoreWriterSettings& WriterSettings = WriterContext->WriterSettings;
-		if (ContainerSettings.IsCompressed() && !Entry->Options.bForceUncompressed && !Entry->Options.bIsMemoryMapped)
-		{
-			CompressionMethod = WriterSettings.CompressionMethod;
-		}
-		return CompressionMethod;
-	}
-
-
-
-	void BeginCompress(FIoStoreWriteQueueEntry* Entry)
-	{
-		WriterContext->BeginCompressChunksByType[(int8)Entry->ChunkId.GetChunkType()].IncrementExchange();
-		
-		if (Entry->bLoadingFromReferenceDb)
-		{
-			if (Entry->NumChunkBlocksFromRefDb == 0)
-			{
-				Entry->FinishCompressionBarrier->DispatchSubsequents();
-				return;
-			}
-
-			// Allocate resources before launching the read tasks to prevent deadlock. Note this will
-			// allocate iobuffers big enough for uncompressed size, when we only actually need it for
-			// compressed size.
-			int32 NumChunkBlocks32 = IntCastChecked<int32>(Entry->NumChunkBlocksFromRefDb);
-			Entry->ChunkBlocks.SetNum(NumChunkBlocks32);
-			for (int32 BlockIndex = 0; BlockIndex < NumChunkBlocks32; ++BlockIndex)
-			{
-				FChunkBlock& Block = Entry->ChunkBlocks[BlockIndex];
-				Block.IoBuffer = WriterContext->AllocCompressionBuffer(NumChunkBlocks32);
-				// Everything else in a block gets filled out from the refdb.
-			}
-
-			// Valid chunks must create the same decompressed bits, but can have different compressed bits.
-			// Since we are on a lightweight dispatch thread, the actual read is async, as is the processing
-			// of the results.
-			bool bChunkExists = ReferenceChunkDatabase->RetrieveChunk(ContainerSettings.ContainerId, Entry->ChunkHash, Entry->ChunkId, [this, Entry](TIoStatusOr<FIoStoreCompressedReadResult> InReadResult)
-			{
-
-				// If we fail here, in order to recover we effectively need to re-kick this chunk's
-				// BeginCompress() as well as source buffer read... however, this is just a direct read and should only fail
-				// in catastrophic scenarios (loss of connection on a network drive?).
-				FIoStoreCompressedReadResult ReadResult = InReadResult.ValueOrDie();
-
-				uint64 TotalUncompressedSize = 0;
-				uint8* ReferenceData = ReadResult.IoBuffer.GetData();
-				uint64 TotalAlignedSize = 0;
-				for (int32 BlockIndex = 0; BlockIndex < ReadResult.Blocks.Num(); ++BlockIndex)
-				{
-					FIoStoreCompressedBlockInfo& ReferenceBlock = ReadResult.Blocks[BlockIndex];
-					FChunkBlock& Block = Entry->ChunkBlocks[BlockIndex];
-					Block.CompressionMethod = ReferenceBlock.CompressionMethod;
-					Block.CompressedSize = ReferenceBlock.CompressedSize;
-					Block.UncompressedSize = ReferenceBlock.UncompressedSize;
-					TotalUncompressedSize += ReferenceBlock.UncompressedSize;
-
-					// Future optimization: ReadCompressed returns the memory ready to encrypt in one
-					// large contiguous buffer (i.e. padded). We could use the FIoBuffer functionality of referencing a 
-					// sub block from a parent buffer, however this would mean that we need to add support
-					// for tracking the memory usage in order to remain within our prescribed limits. To do this
-					// requires releasing the entire chunk's memory at once after WriteEntry.
-					// As it stands, we temporarily use untracked memory in the ReadCompressed call (in RetrieveChunk),
-					// then immediately copy it to tracked memory. There's some waste as tracked memory is mod CompressionBlockSize
-					// and we are post compression, so with the average 50% compression rate, we're using double the memory
-					// we "could".
-					FMemory::Memcpy(Block.IoBuffer->GetData(), ReferenceData, Block.CompressedSize);
-					ReferenceData += ReferenceBlock.AlignedSize;
-					TotalAlignedSize += ReferenceBlock.AlignedSize;
-				}
-
-				if (TotalAlignedSize != ReadResult.IoBuffer.GetSize())
-				{
-					// If we hit this, we might have read garbage memory above! This is very bad.
-					UE_LOG(LogIoStore, Error, TEXT("Block aligned size does not match iobuffer source size! Blocks: %s source size: %s"),
-						*FText::AsNumber(TotalAlignedSize).ToString(),
-						*FText::AsNumber(ReadResult.IoBuffer.GetSize()).ToString());
-				}
-
-				Entry->UncompressedSize.Emplace(TotalUncompressedSize);
-				Entry->FinishCompressionBarrier->DispatchSubsequents();
-			});
-
-			check(bChunkExists); // Sanity - should never return false as we can only get here if ChunkExists() returns true.
-
-			WriterContext->RefDbChunksCount.IncrementExchange();
-			WriterContext->RefDbChunksByType[(int8)Entry->ChunkId.GetChunkType()].IncrementExchange();
-			// Lambda handles dispatch subsequents
-			return;
-		}
-
-		const FIoStoreWriterSettings& WriterSettings = WriterContext->WriterSettings;
-
-		FName CompressionMethod = CompressionMethodForEntry(Entry);
-
-		const FIoBuffer* SourceBuffer = Entry->Request->GetSourceBuffer();
-		Entry->UncompressedSize.Emplace(SourceBuffer->DataSize());
-
-		check(WriterSettings.CompressionBlockSize > 0);
-		const uint64 NumChunkBlocks64 = Align(Entry->UncompressedSize.GetValue(), WriterSettings.CompressionBlockSize) / WriterSettings.CompressionBlockSize;
-		if (NumChunkBlocks64 == 0)
-		{
-			Entry->FinishCompressionBarrier->DispatchSubsequents();
-			return;
-		}
-
-		int32 NumChunkBlocks = IntCastChecked<int32>(NumChunkBlocks64);
-		
-		Entry->ChunkBlocks.SetNum(NumChunkBlocks);
-		{
-			// We must allocate resources for our tasks up front to prevent resource deadlock.
-			uint64 BytesToProcess = Entry->UncompressedSize.GetValue();
-			const uint8* UncompressedData = SourceBuffer->Data();
-			for (int32 BlockIndex = 0; BlockIndex < NumChunkBlocks; ++BlockIndex)
-			{
-				FChunkBlock& Block = Entry->ChunkBlocks[BlockIndex];
-				Block.IoBuffer = WriterContext->AllocCompressionBuffer(NumChunkBlocks);
-				Block.CompressionMethod = CompressionMethod;
-				Block.UncompressedSize = FMath::Min(BytesToProcess, WriterSettings.CompressionBlockSize);
-				Block.UncompressedData = UncompressedData;
-				BytesToProcess -= Block.UncompressedSize;
-				UncompressedData += Block.UncompressedSize;
-			}
-		}
-
-		if (CompressionMethod == NAME_None)
-		{
-			for (FChunkBlock& Block : Entry->ChunkBlocks)
-			{
-				Block.CompressionMethod = NAME_None;
-				Block.CompressedSize = Block.UncompressedSize;
-				FMemory::Memcpy(Block.IoBuffer->Data(), Block.UncompressedData, Block.UncompressedSize);
-			}
-			Entry->FinishCompressionBarrier->DispatchSubsequents();
-			return;
-		}
-
-		bool bUseDDCCompression =
-			WriterContext->DDC &&
-			WriterSettings.bCompressionEnableDDC &&
-			Entry->UncompressedSize.GetValue() > WriterSettings.CompressionMinSizeToConsiderDDC;
-		if (bUseDDCCompression)
-		{
-			TStringBuilder<256> CacheKeySuffix;
-			CacheKeySuffix.Append(Entry->ChunkHash.ToString());
-			WriterSettings.CompressionMethod.AppendString(CacheKeySuffix);
-			CacheKeySuffix.Append(FCompression::GetCompressorDDCSuffix(WriterSettings.CompressionMethod));
-			CacheKeySuffix.Appendf(TEXT("%d_%d"),
-				WriterSettings.CompressionMinBytesSaved,
-				WriterSettings.CompressionMinPercentSaved);
-			Entry->DDCCacheKey = FDerivedDataCacheInterface::BuildCacheKey(
-				TEXT("IoStoreCompression"),
-				TEXT("985D3FAD-71D0-4758-A777-A910B49CC4BF"),
-				CacheKeySuffix.ToString());
-
-			TArray<uint8> DDCData;
-			TRACE_CPUPROFILER_EVENT_SCOPE(ReadFromDDC);
-			bool bFoundInDDC = false;
-			if (WriterContext->DDC->GetSynchronous(*Entry->DDCCacheKey, DDCData, Entry->Options.FileName))
-			{
-				FLargeMemoryReader DDCDataReader(DDCData.GetData(), DDCData.Num());
-				bFoundInDDC = SerializeCompressedDDCData(Entry, DDCDataReader);
-			}
-			if (bFoundInDDC)
-			{
-				WriterContext->CompressionDDCHitCount.IncrementExchange();
-				Entry->FinishCompressionBarrier->DispatchSubsequents();
-				return;
-			}
-			else
-			{
-				WriterContext->CompressionDDCMissCount.IncrementExchange();
-				Entry->bStoreCompressedDataInDDC = true;
-			}
-		}
-		
-		for (FChunkBlock& Block : Entry->ChunkBlocks)
-		{
-			WriterContext->ScheduledCompressionTasksCount.IncrementExchange();
-			FFunctionGraphTask::CreateAndDispatchWhenReady([this, Entry, BlockPtr = &Block]()
-			{
-				CompressBlock(BlockPtr);
-				WriterContext->ScheduledCompressionTasksCount.DecrementExchange();
-				int32 CompressedBlocksCount = Entry->CompressedBlocksCount.IncrementExchange();
-				if (CompressedBlocksCount + 1 == Entry->ChunkBlocks.Num())
-				{
-					WriterContext->CompressedChunksByType[(int8)Entry->ChunkId.GetChunkType()].IncrementExchange();
-					WriterContext->CompressedChunksCount.IncrementExchange();
-					Entry->FinishCompressionBarrier->DispatchSubsequents();
-				}
-			}, TStatId(), nullptr, ENamedThreads::AnyHiPriThreadNormalTask);
-		}
-	}
-
-	void BeginEncryptAndSign(FIoStoreWriteQueueEntry* Entry)
-	{
-		if (ContainerSettings.IsEncrypted() || ContainerSettings.IsSigned())
-		{
-			FFunctionGraphTask::CreateAndDispatchWhenReady([this, Entry]()
-			{
-				EncryptAndSign(Entry);
-				Entry->FinishEncryptionAndSigningBarrier->DispatchSubsequents();
-			}, TStatId(), nullptr, ENamedThreads::AnyHiPriThreadHiPriTask);
-		}
-		else
-		{
-			EncryptAndSign(Entry);
-			Entry->FinishEncryptionAndSigningBarrier->DispatchSubsequents();
-		}
-	}
-
-	void EncryptAndSign(FIoStoreWriteQueueEntry* Entry)
-	{
-		TRACE_CPUPROFILER_EVENT_SCOPE(EncryptAndSignChunk);
-		for (FChunkBlock& Block : Entry->ChunkBlocks)
-		{
-			// Always align each compressed block to AES block size but store the compressed block size in the TOC
-			Block.Size = Block.CompressedSize;
-			if (!IsAligned(Block.Size, FAES::AESBlockSize))
-			{
-				uint64 AlignedCompressedBlockSize = Align(Block.Size, FAES::AESBlockSize);
-				uint8* CompressedData = Block.IoBuffer->Data();
-				for (uint64 FillIndex = Block.Size; FillIndex < AlignedCompressedBlockSize; ++FillIndex)
-				{
-					check(FillIndex < Block.IoBuffer->DataSize());
-					CompressedData[FillIndex] = CompressedData[(FillIndex - Block.Size) % Block.Size];
-				}
-				Block.Size = AlignedCompressedBlockSize;
-			}
-
-			if (ContainerSettings.IsEncrypted())
-			{
-				FAES::EncryptData(Block.IoBuffer->Data(), static_cast<uint32>(Block.Size), ContainerSettings.EncryptionKey);
-			}
-
-			if (ContainerSettings.IsSigned())
-			{
-				FSHA1::HashBuffer(Block.IoBuffer->Data(), Block.Size, Block.Signature.Hash);
-			}
-		}
-	}
-
-	void WriteEntry(FIoStoreWriteQueueEntry* Entry)
-	{
-		TRACE_CPUPROFILER_EVENT_SCOPE(WriteEntry);
-		const int32* FindExistingIndex = TocBuilder.GetTocEntryIndex(Entry->ChunkId);
-		if (FindExistingIndex)
-		{
-			// afaict this should never happen so add a warning. If there's a legit reason for it
-			// we can pull this back out. If would violate some assumptions in the reference chunk
-			// database if we DO hit this, however...
-			UE_LOG(LogIoStore, Warning, TEXT("ChunkId was added twice in container %s, %s, file %s hash %s vs %s"), 
-				*FPaths::GetBaseFilename(TocFilePath),
-				*LexToString(Entry->ChunkId), *Entry->Options.FileName,
-				*TocBuilder.GetTocResource().ChunkMetas[*FindExistingIndex].ChunkHash.ToString(),
-				*Entry->ChunkHash.ToString()
-				);
-
-			checkf(TocBuilder.GetTocResource().ChunkMetas[*FindExistingIndex].ChunkHash == Entry->ChunkHash, TEXT("Chunk id has already been added with different content"));
-			for (FChunkBlock& ChunkBlock : Entry->ChunkBlocks)
-			{
-				WriterContext->FreeCompressionBuffer(ChunkBlock.IoBuffer, Entry->ChunkBlocks.Num());
-				ChunkBlock.IoBuffer = nullptr;
-			}
-			return;
-		}
-
-		FPartition* TargetPartition = &Partitions[CurrentPartitionIndex];
-		int32 NextPartitionIndexToTry = CurrentPartitionIndex + 1;
-		if (Entry->PartitionIndex >= 0)
-		{
-			TargetPartition = &Partitions[Entry->PartitionIndex];
-			if (TargetPartition->ReservedSpace > Entry->CompressedSize)
-			{
-				TargetPartition->ReservedSpace -= Entry->CompressedSize;
-			}
-			else
-			{
-				TargetPartition->ReservedSpace = 0;
-			}
-			NextPartitionIndexToTry = CurrentPartitionIndex;
-		}
-
-		const FIoStoreWriterSettings& WriterSettings = WriterContext->WriterSettings;
-		bHasMemoryMappedEntry |= Entry->Options.bIsMemoryMapped;
-		const uint64 ChunkAlignment = Entry->Options.bIsMemoryMapped ? WriterSettings.MemoryMappingAlignment : 0;
-		const uint64 PartitionSizeLimit = WriterSettings.MaxPartitionSize > 0 ? WriterSettings.MaxPartitionSize : MAX_uint64;
-		checkf(Entry->CompressedSize <= PartitionSizeLimit, TEXT("Chunk is too large, increase max partition size!"));
-		for (;;)
-		{
-			uint64 OffsetBeforePadding = TargetPartition->Offset;
-			if (ChunkAlignment)
-			{
-				TargetPartition->Offset = Align(TargetPartition->Offset, ChunkAlignment);
-			}
-			if (WriterSettings.CompressionBlockAlignment)
-			{
-				// Try and prevent entries from crossing compression alignment blocks if possible. This is to avoid
-				// small entries from causing multiple file system block reads afaict. Large entries necesarily get
-				// aligned to prevent things like a blocksize + 2 entry being at alignment -1, causing 3 low level reads.
-				// ...I think.
-				bool bCrossesBlockBoundary = Align(TargetPartition->Offset, WriterSettings.CompressionBlockAlignment) != Align(TargetPartition->Offset + Entry->CompressedSize - 1, WriterSettings.CompressionBlockAlignment);
-				if (bCrossesBlockBoundary)
-				{
-					TargetPartition->Offset = Align(TargetPartition->Offset, WriterSettings.CompressionBlockAlignment);
-				}
-			}
-
-			if (TargetPartition->Offset + Entry->CompressedSize + TargetPartition->ReservedSpace > PartitionSizeLimit)
-			{
-				TargetPartition->Offset = OffsetBeforePadding;
-				while (Partitions.Num() <= NextPartitionIndexToTry)
-				{
-					FPartition& NewPartition = Partitions.AddDefaulted_GetRef();
-					NewPartition.Index = Partitions.Num() - 1;
-				}
-				CurrentPartitionIndex = NextPartitionIndexToTry;
-				TargetPartition = &Partitions[CurrentPartitionIndex];
-				++NextPartitionIndexToTry;
-			}
-			else
-			{
-				Entry->Padding = TargetPartition->Offset - OffsetBeforePadding;
-				TotalPaddingSize += Entry->Padding;
-				break;
-			}
-		}
-
-		if (!TargetPartition->ContainerFileHandle)
-		{
-			TRACE_CPUPROFILER_EVENT_SCOPE(CreatePartitionContainerFile);
-			CreatePartitionContainerFile(*TargetPartition);
-		}
-		Entry->Offset = TargetPartition->Offset;
-
-		FIoOffsetAndLength OffsetLength;
-		OffsetLength.SetOffset(UncompressedFileOffset);
-		OffsetLength.SetLength(Entry->UncompressedSize.GetValue());
-
-		FIoStoreTocEntryMeta ChunkMeta{ Entry->ChunkHash, FIoStoreTocEntryMetaFlags::None };
-		if (Entry->Options.bIsMemoryMapped)
-		{
-			ChunkMeta.Flags |= FIoStoreTocEntryMetaFlags::MemoryMapped;
-		}
-
-		uint64 OffsetInChunk = 0;
-		for (const FChunkBlock& ChunkBlock : Entry->ChunkBlocks)
-		{
-			FIoStoreTocCompressedBlockEntry& BlockEntry = TocBuilder.AddCompressionBlockEntry();
-			BlockEntry.SetOffset(TargetPartition->Index * WriterSettings.MaxPartitionSize + TargetPartition->Offset + OffsetInChunk);
-			OffsetInChunk += ChunkBlock.Size;
-			BlockEntry.SetCompressedSize(uint32(ChunkBlock.CompressedSize));
-			BlockEntry.SetUncompressedSize(uint32(ChunkBlock.UncompressedSize));
-			BlockEntry.SetCompressionMethodIndex(TocBuilder.AddCompressionMethodEntry(ChunkBlock.CompressionMethod));
-
-			// We do this here so that we get the total size of data excluding the encryption alignment
-			TotalEntryCompressedSize += ChunkBlock.CompressedSize;
-			if (Entry->bCouldBeFromReferenceDb && !Entry->bLoadingFromReferenceDb)
-			{
-				ReferenceCacheMissBytes += ChunkBlock.CompressedSize;
-			}
-
-			if (!ChunkBlock.CompressionMethod.IsNone())
-			{
-				ChunkMeta.Flags |= FIoStoreTocEntryMetaFlags::Compressed;
-			}
-
-			if (ContainerSettings.IsSigned())
-			{
-				FSHAHash& Signature = TocBuilder.AddBlockSignatureEntry();
-				Signature = ChunkBlock.Signature;
-			}
-
-		}
-
-		const int32 TocEntryIndex = TocBuilder.AddChunkEntry(Entry->ChunkId, OffsetLength, ChunkMeta);
-		check(TocEntryIndex != INDEX_NONE);
-
-		if (ContainerSettings.IsIndexed() && Entry->Options.FileName.Len() > 0)
-		{
-			TocBuilder.AddToFileIndex(Entry->ChunkId, Entry->Options.FileName);
-		}
-
-		const uint64 RegionStartOffset = TargetPartition->Offset;
-		TargetPartition->Offset += Entry->CompressedSize;
-		UncompressedFileOffset += Align(Entry->UncompressedSize.GetValue(), WriterSettings.CompressionBlockSize);
-		TotalEntryUncompressedSize += Entry->UncompressedSize.GetValue();
-
-		if (WriterSettings.bEnableFileRegions)
-		{
-			FFileRegion::AccumulateFileRegions(TargetPartition->AllFileRegions, RegionStartOffset, RegionStartOffset, TargetPartition->Offset, Entry->Request->GetRegions());
-		}
-		delete Entry->Request;
-		Entry->Request = nullptr;
-		uint64 WriteStartCycles = FPlatformTime::Cycles64();
-		uint64 WriteBytes = 0;
-		if (Entry->Padding > 0)
-		{
-			if (PaddingBuffer.Num() < Entry->Padding)
-			{
-				PaddingBuffer.SetNumZeroed(int32(Entry->Padding));
-			}
-			{
-				TRACE_CPUPROFILER_EVENT_SCOPE(WritePaddingToContainer);
-				TargetPartition->ContainerFileHandle->Serialize(PaddingBuffer.GetData(), Entry->Padding);
-				WriteBytes += Entry->Padding;
-			}
-		}
-		check(Entry->Offset == TargetPartition->ContainerFileHandle->Tell());
-		for (FChunkBlock& ChunkBlock : Entry->ChunkBlocks)
-		{
-			{
-				TRACE_CPUPROFILER_EVENT_SCOPE(WriteBlockToContainer);
-				TargetPartition->ContainerFileHandle->Serialize(ChunkBlock.IoBuffer->Data(), ChunkBlock.Size);
-				WriteBytes += ChunkBlock.Size;
-			}
-			WriterContext->FreeCompressionBuffer(ChunkBlock.IoBuffer, Entry->ChunkBlocks.Num());
-			ChunkBlock.IoBuffer = nullptr;
-		}
-		uint64 WriteEndCycles = FPlatformTime::Cycles64();
-		WriterContext->WriteCycleCount.AddExchange(WriteEndCycles - WriteStartCycles);
-		WriterContext->WriteByteCount.AddExchange(WriteBytes);
-		WriterContext->SerializedChunksCount.IncrementExchange();
-	}
-
-	const FString				ContainerPathAndBaseFileName;
-	FIoStoreWriterContextImpl*	WriterContext = nullptr;
-	FIoContainerSettings		ContainerSettings;
-	FString						TocFilePath;
-	FIoStoreTocBuilder			TocBuilder;
-	TArray<uint8>				PaddingBuffer;
-	TArray<FPartition>			Partitions;
-	TArray<FIoStoreWriteQueueEntry*> Entries;
-	TArray<FLayoutEntry*>		LayoutEntries;
-	FLayoutEntry*				LayoutEntriesHead = nullptr;
-	FLayoutEntry*				LayoutEntriesTail = nullptr;
-	TMap<FIoChunkId, FLayoutEntry*> PreviousBuildLayoutEntryByChunkId;
-	TUniquePtr<FArchive>		CsvArchive;
-	FIoStoreWriterResult		Result;
-	uint64						UncompressedFileOffset = 0;
-	uint64						TotalEntryUncompressedSize = 0; // sum of all entry source buffer sizes
-	uint64						TotalEntryCompressedSize = 0; // entry compressed size excluding encryption alignment
-	uint64						ReferenceCacheMissBytes = 0; // number of compressed bytes excluding alignment that could have been from refcache but weren't.
-	uint64						TotalPaddingSize = 0;
-	uint64						UncompressedContainerSize = 0; // this is the size the container would be if it were uncompressed.
-	uint64						CompressedContainerSize = 0; // this is the size of the container with the given compression (which may be none).
-	int32						CurrentPartitionIndex = 0;
-	bool						bHasMemoryMappedEntry = false;
-	bool						bHasFlushed = false;
-	bool						bHasResult = false;
-	TSharedPtr<IIoStoreWriterReferenceChunkDatabase> ReferenceChunkDatabase;
-	TSharedPtr<IIoStoreWriterHashDatabase> HashDatabase;
-	bool						bVerifyHashDatabase = false;
-
-
-	friend class FIoStoreWriterContextImpl;
-};
-
-// InContainerPathAndBaseFileName: the utoc file will just be this with .utoc appended.
-// The base filename ends up getting returned as the container name in the writer results.
-TSharedPtr<IIoStoreWriter> FIoStoreWriterContextImpl::CreateContainer(const TCHAR* InContainerPathAndBaseFileName, const FIoContainerSettings& InContainerSettings)
-{
-	TSharedPtr<FIoStoreWriter> IoStoreWriter = MakeShared<FIoStoreWriter>(InContainerPathAndBaseFileName);
-	FIoStatus IoStatus = IoStoreWriter->Initialize(*this, InContainerSettings);
-	check(IoStatus.IsOk());
-	IoStoreWriters.Add(IoStoreWriter);
-	return IoStoreWriter;
-}
-
-void FIoStoreWriterContextImpl::Flush()
-{
-	TRACE_CPUPROFILER_EVENT_SCOPE(FIoStoreWriterContext::Flush);
-	TArray<FIoStoreWriteQueueEntry*> AllEntries;
-	for (TSharedPtr<FIoStoreWriter> IoStoreWriter : IoStoreWriters)
-	{
-		IoStoreWriter->bHasFlushed = true;
-		AllEntries.Append(IoStoreWriter->Entries);
-	}
-	{
-		TRACE_CPUPROFILER_EVENT_SCOPE(WaitForChunkHashes);
-		for (int32 EntryIndex = AllEntries.Num() - 1; EntryIndex >= 0; --EntryIndex)
-		{
-			FTaskGraphInterface::Get().WaitUntilTaskCompletes(AllEntries[EntryIndex]->HashTask);
-		}
-	}
-	for (TSharedPtr<FIoStoreWriter> IoStoreWriter : IoStoreWriters)
-	{
-		if (IoStoreWriter->LayoutEntriesHead)
-		{
-			IoStoreWriter->FinalizeLayout();
-		}
-	}
-	// Update list of all entries after having the finilized layouts of each container
-	AllEntries.Reset();
-	for (TSharedPtr<FIoStoreWriter> IoStoreWriter : IoStoreWriters)
-	{
-		AllEntries.Append(IoStoreWriter->Entries);
-	}
-
-	double WritesStart = FPlatformTime::Seconds();
-
-	for (FIoStoreWriteQueueEntry* Entry : AllEntries)
-	{
-		ScheduleCompression(Entry);
-	}
-
-	{
-		TRACE_CPUPROFILER_EVENT_SCOPE(WaitForWritesToComplete);
-		for (int32 EntryIndex = AllEntries.Num() - 1; EntryIndex >= 0; --EntryIndex)
-		{
-			AllEntries[EntryIndex]->WriteFinishedEvent->Wait();
-		}
-	}
-
-	double WritesEnd = FPlatformTime::Seconds();
-	double WritesSeconds = FPlatformTime::ToSeconds64(WriteCycleCount.Load());
-	UE_LOG(LogIoStore, Display, TEXT("Writing and compressing took %.2f seconds, writes to disk took %.2f seconds for %s bytes @ %s bytes per second."), 
-		WritesEnd - WritesStart,
-		WritesSeconds,
-		*FText::AsNumber(WriteByteCount.Load()).ToString(),
-		*FText::AsNumber((int64)((double)WriteByteCount.Load() / FMath::Max(.0001f, WritesSeconds))).ToString()
-		);
-
-	AllEntries.Empty();
-
-	// Classically there were so few writers that this didn't need to be multi threaded, but it
-	// involves writing files, and with content on demand this ends up being thousands of iterations. 
-	double FinalizeStart = FPlatformTime::Seconds();	
-	ParallelFor(TEXT("IoStoreWriter::Finalize.PF"), IoStoreWriters.Num(), 1, [this](int Index)
-	{ 
-		IoStoreWriters[Index]->Finalize(); 
-	});
-	double FinalizeEnd = FPlatformTime::Seconds();
-	int64 TotalTocSize = 0;
-	for (TSharedPtr<IIoStoreWriter> Writer : IoStoreWriters)
-	{
-		if (Writer->GetResult().IsOk())
-		{
-			TotalTocSize += Writer->GetResult().ValueOrDie().TocSize;
-		}
-	}
-
-	UE_LOG(LogIoStore, Display, TEXT("Finalize took %.1f seconds for %d writers to write %s bytes, %s bytes per second"), 
-		FinalizeEnd - FinalizeStart, 
-		IoStoreWriters.Num(), 
-		*FText::AsNumber(TotalTocSize).ToString(), 
-		*FText::AsNumber((int64)((double)TotalTocSize / FMath::Max(.0001f, FinalizeEnd - FinalizeStart))).ToString()
-		);
-}
-
-void FIoStoreWriterContextImpl::BeginCompressionThreadFunc()
-{
-	for (;;)
-	{
-		FIoStoreWriteQueueEntry* Entry = BeginCompressionQueue.DequeueOrWait();
-		if (!Entry)
-		{
-			return;
-		}
-		while (Entry)
-		{
-			FIoStoreWriteQueueEntry* Next = Entry->Next;
-			Entry->BeginCompressionBarrier->Wait();
-			Entry->Writer->BeginCompress(Entry);
-			BeginEncryptionAndSigningQueue.Enqueue(Entry);
-			Entry = Next;
-		}
-	}
-}
-
-void FIoStoreWriterContextImpl::BeginEncryptionAndSigningThreadFunc()
-{
-	for (;;)
-	{
-		FIoStoreWriteQueueEntry* Entry = BeginEncryptionAndSigningQueue.DequeueOrWait();
-		if (!Entry)
-		{
-			return;
-		}
-		while (Entry)
-		{
-			FIoStoreWriteQueueEntry* Next = Entry->Next;
-			Entry->FinishCompressionBarrier->Wait();
-			if (Entry->bLoadingFromReferenceDb == false)
-			{
-				Entry->Request->FreeSourceBuffer();
-			}
-			if (Entry->bStoreCompressedDataInDDC)
-			{
-				TRACE_CPUPROFILER_EVENT_SCOPE(StoreInDDC);
-				TArray<uint8> DDCData;
-				FMemoryWriter DDCDataWriter(DDCData, true);
-				if (Entry->Writer->SerializeCompressedDDCData(Entry, DDCDataWriter))
-				{
-					DDC->Put(*Entry->DDCCacheKey, DDCData, Entry->Options.FileName);
-				}
-			}
-			FinishEncryptionAndSigningQueue.Enqueue(Entry);
-			Entry->Writer->BeginEncryptAndSign(Entry);
-			Entry = Next;
-		}
-	}
-}
-
-void FIoStoreWriterContextImpl::FinishEncryptionAndSigningThreadFunc()
-{
-	for (;;)
-	{
-		FIoStoreWriteQueueEntry* Entry = FinishEncryptionAndSigningQueue.DequeueOrWait();
-		if (!Entry)
-		{
-			return;
-		}
-		while (Entry)
-		{
-			FIoStoreWriteQueueEntry* Next = Entry->Next;
-			Entry->FinishEncryptionAndSigningBarrier->Wait();
-			WriterQueue.Enqueue(Entry);
-			Entry->CompressedSize = 0;
-			for (const FChunkBlock& ChunkBlock : Entry->ChunkBlocks)
-			{
-				Entry->CompressedSize += ChunkBlock.Size;
-			}
-			Entry->BeginWriteBarrier->DispatchSubsequents();
-
-			Entry = Next;
-		}
-	}
-}
-
-void FIoStoreWriterContextImpl::WriterThreadFunc()
-{
-	for (;;)
-	{
-		FIoStoreWriteQueueEntry* Entry = WriterQueue.DequeueOrWait();
-		if (!Entry)
-		{
-			return;
-		}
-		while (Entry)
-		{
-			FIoStoreWriteQueueEntry* Next = Entry->Next;
-			Entry->BeginWriteBarrier->Wait();
-			Entry->Writer->WriteEntry(Entry);
-			Entry->WriteFinishedEvent->DispatchSubsequents();
-			Entry = Next;
-		}
-	}
-}
-
 class FIoStoreReaderImpl
 {
 public:
diff --git a/Engine/Source/Runtime/Core/Public/IO/IoDispatcher.h b/Engine/Source/Runtime/Core/Public/IO/IoDispatcher.h
--- a/Engine/Source/Runtime/Core/Public/IO/IoDispatcher.h
+++ b/Engine/Source/Runtime/Core/Public/IO/IoDispatcher.h
@@ -58,7 +58,6 @@ class FIoRequestImpl;
 class FIoStoreEnvironment;
 class FIoStoreReader;
 class FIoStoreReaderImpl;
-class FIoStoreWriterContextImpl;
 class FPackageId;
 class IMappedFileHandle;
 class IMappedFileRegion;
@@ -475,23 +474,6 @@ private:
 
 //////////////////////////////////////////////////////////////////////////
 
-struct FIoStoreWriterSettings
-{
-	FName CompressionMethod = NAME_None;
-	uint64 CompressionBlockSize = 64 << 10;
-
-	// This does not align every entry - it tries to prevent excess crossings of this boundary by inserting padding.
-	// and happens whether or not the entry is compressed.
-	uint64 CompressionBlockAlignment = 0;
-	int32 CompressionMinBytesSaved = 0;
-	int32 CompressionMinPercentSaved = 0;
-	int32 CompressionMinSizeToConsiderDDC = 0;
-	uint64 MemoryMappingAlignment = 0;
-	uint64 MaxPartitionSize = 0;
-	bool bEnableFileRegions = false;
-	bool bCompressionEnableDDC = false;
-};
-
 enum class EIoContainerFlags : uint8
 {
 	None,
@@ -538,90 +520,6 @@ struct FIoContainerSettings
 	}
 };
 
-struct FIoStoreWriterResult
-{
-	FIoContainerId ContainerId;
-	FString ContainerName; // This is the base filename of the utoc used for output.
-	int64 TocSize = 0;
-	int64 TocEntryCount = 0;
-	int64 PaddingSize = 0;
-	int64 UncompressedContainerSize = 0; // this is the size the container would be if it were uncompressed.
-	int64 CompressedContainerSize = 0; // this is the size of the container with the given compression (which may be none). Should be the sum of all partition file sizes.
-	int64 DirectoryIndexSize = 0;
-	uint64 TotalEntryCompressedSize = 0; // sum of the compressed size of entries excluding encryption alignment.
-	uint64 ReferenceCacheMissBytes = 0; // number of compressed bytes excluding alignment that could have been from refcache but weren't.
-	uint64 AddedChunksCount = 0;
-	uint64 AddedChunksSize = 0;
-	uint64 ModifiedChunksCount = 0;
-	uint64 ModifiedChunksSize = 0;
-	FName CompressionMethod = NAME_None;
-	EIoContainerFlags ContainerFlags = EIoContainerFlags::None;
-};
-
-struct FIoWriteOptions
-{
-	FString FileName;
-	const TCHAR* DebugName = nullptr;
-	bool bForceUncompressed = false;
-	bool bIsMemoryMapped = false;
-};
-
-class FIoStoreWriterContext
-{
-public:
-	struct FProgress
-	{
-		uint64 TotalChunksCount = 0;
-		uint64 HashedChunksCount = 0;
-		// Number of chunks where we avoided reading and hashing, and instead used the result from the hashdb, and their types
-		uint64 HashDbChunksCount = 0;
-		uint64 HashDbChunksByType[(int8)EIoChunkType::MAX] = { 0 };
-		// Number of chunks that were passed to the compressor (i.e. passed the various opt-outs), and their types
-		uint64 CompressedChunksCount = 0;
-		uint64 CompressedChunksByType[(int8)EIoChunkType::MAX] = { 0 };
-		uint64 SerializedChunksCount = 0;
-		uint64 ScheduledCompressionTasksCount = 0;
-		uint64 CompressionDDCHitCount = 0;
-		uint64 CompressionDDCMissCount = 0;
-
-		// The number of chunk retrieved from the reference cache database, and their types.
-		uint64 RefDbChunksCount{ 0 };
-		uint64 RefDbChunksByType[(int8)EIoChunkType::MAX] = { 0 };
-		
-		// The type of chunk that landed in BeginCompress before any opt-outs.
-		uint64 BeginCompressChunksByType[(int8)EIoChunkType::MAX] = { 0 };
-	};
-
-	CORE_API FIoStoreWriterContext();
-	CORE_API ~FIoStoreWriterContext();
-
-	[[nodiscard]] CORE_API FIoStatus Initialize(const FIoStoreWriterSettings& InWriterSettings);
-	CORE_API TSharedPtr<class IIoStoreWriter> CreateContainer(const TCHAR* InContainerPath, const FIoContainerSettings& InContainerSettings);
-	CORE_API void Flush();
-	CORE_API FProgress GetProgress() const;
-
-private:
-	FIoStoreWriterContextImpl* Impl;
-};
-
-class IIoStoreWriteRequest
-{
-public:
-	virtual ~IIoStoreWriteRequest() = default;
-
-	// Launches any async operations necessary in order to access the buffer. CompletionEvent is set once it's ready, which may be immediate.
-	virtual void PrepareSourceBufferAsync(FGraphEventRef CompletionEvent) = 0;
-	virtual uint64 GetOrderHint() = 0;
-	virtual TArrayView<const FFileRegion> GetRegions() = 0;
-
-	// Only valid after the completion event passed to PrepareSourceBufferAsync has fired.
-	virtual const FIoBuffer* GetSourceBuffer() = 0;
-
-	// Can't be called between PrepareSourceBufferAsync and its completion!
-	virtual void FreeSourceBuffer() = 0;
-};
-
-
 struct FIoStoreTocChunkInfo
 {
 	FIoChunkId Id;
@@ -678,82 +576,6 @@ struct FIoStoreCompressedReadResult
 	uint64 TotalCompressedSize = 0;
 };
 
-class IIoStoreWriterReferenceChunkDatabase
-{
-public:
-	virtual ~IIoStoreWriterReferenceChunkDatabase() = default;
-
-	/*
-	* Used by IIoStoreWriter to check and see if there's a reference chunk that matches the data that
-	* IoStoreWriter wants to compress and write. Validity checks must be synchronous - if a chunk can't be
-	* used for some reason (no matching chunk exists or otherwise), this function must return false and not 
-	* call InCompletionCallback.
-	* 
-	* Once a matching chunk is found, it is read from the source iostore container asynchronously, and upon
-	* completion InCompletionCallback is called with the raw output from FIoStoreReader::ReadCompressed (i.e.
-	* FIoStoreCompressedReadResult). Failures once the async read process has started are currently fatal due to
-	* difficulties in rekicking a read.
-	* 
-	* For the moment, changes in compression method are allowed.
-	* 
-	* RetrieveChunk is not currently thread safe and must be called from a single thread.
-	* 
-	* Chunks provided *MUST* decompress to bits that hash to the exact value provided in InChunkKey (i.e. be exactly the same bits),
-	* and also be the same number of blocks (i.e. same CompressionBlockSize)
-	*/
-	virtual bool RetrieveChunk(const FIoContainerId& InContainerId, const FIoChunkHash& InChunkHash, const FIoChunkId& InChunkId, TUniqueFunction<void(TIoStatusOr<FIoStoreCompressedReadResult>)> InCompletionCallback) = 0;
-
-	/* 
-	* Quick synchronous existence check that returns the number of blocks for the chunk. This is used to set up
-	* the necessary structures without needing to read the source data for the chunk. This might be called from
-	* multiple threads as it has to happen after we have the source hash computed.
-	*/
-	virtual bool ChunkExists(const FIoContainerId& InContainerId, const FIoChunkHash& InChunkHash, const FIoChunkId& InChunkId, uint32& OutNumChunkBlocks) = 0;
-
-	/*
-	* Returns the compression block size that was used to break up the IoChunks in the source containers. If this is different than what we want, 
-	* then none of the chunks will ever match. Knowing this up front allows us to only match on hash
-	*/
-	virtual uint32 GetCompressionBlockSize() const = 0;
-
-	/*
-	* Called by an iostore writer implementation to notify the ref cache it's been added
-	*/
-	virtual void NotifyAddedToWriter(const FIoContainerId& InContainerId, const FString& InContainerName) = 0;
-};
-
-/**
-*	Allows the IIoStoreWriter to avoid loading and hashing chunks, saving pak/stage time, as the normal
-*	process involved loading the chunks, hashing them, freeing them, making some decisions, then loading
-*	them _again_ for compression/writting. It's completely fine for this to not have all available hashes,
-*	but they have to match when provided!
-*/
-class IIoStoreWriterHashDatabase
-{
-public:
-	virtual ~IIoStoreWriterHashDatabase() = default;
-	virtual bool FindHashForChunkId(const FIoChunkId& ChunkId, FIoChunkHash& OutHash) const = 0;
-};
-
-
-class IIoStoreWriter
-{
-public:
-	virtual ~IIoStoreWriter() = default;
-
-	/**
-	*	If a reference database is provided, the IoStoreWriter implementation may elect to reuse compressed blocks
-	*	from previous containers instead of recompressing input data. This must be set before any writes are appended.
-	*/
-	virtual void SetReferenceChunkDatabase(TSharedPtr<IIoStoreWriterReferenceChunkDatabase> ReferenceChunkDatabase) = 0;
-	virtual void SetHashDatabase(TSharedPtr<IIoStoreWriterHashDatabase> HashDatabase, bool bVerifyHashDatabase) = 0;
-	virtual void EnableDiskLayoutOrdering(const TArray<TUniquePtr<FIoStoreReader>>& PatchSourceReaders = TArray<TUniquePtr<FIoStoreReader>>()) = 0;
-	virtual void Append(const FIoChunkId& ChunkId, FIoBuffer Chunk, const FIoWriteOptions& WriteOptions, uint64 OrderHint = MAX_uint64) = 0;
-	virtual void Append(const FIoChunkId& ChunkId, IIoStoreWriteRequest* Request, const FIoWriteOptions& WriteOptions) = 0;
-	virtual TIoStatusOr<FIoStoreWriterResult> GetResult() = 0;
-	virtual void EnumerateChunks(TFunction<bool(FIoStoreTocChunkInfo&&)>&& Callback) const = 0;
-};
-
 class FIoStoreReader
 {
 public:
