- Fix a race condition that can happen inside the FPersistentStorageManager ctor where the logging system is asking for a FPersistentStorageManager::Get after testing FPersistentStorageManager::IsReady true

diff --git a/Engine/Source/Runtime/Core/Private/HAL/IPlatformFileManagedStorageWrapper.cpp b/Engine/Source/Runtime/Core/Private/HAL/IPlatformFileManagedStorageWrapper.cpp
--- a/Engine/Source/Runtime/Core/Private/HAL/IPlatformFileManagedStorageWrapper.cpp
+++ b/Engine/Source/Runtime/Core/Private/HAL/IPlatformFileManagedStorageWrapper.cpp
@@ -197,10 +197,13 @@ void FManagedStorageScopeFileLock::Lock()
 	}
 }
 
+bool FPersistentStorageManager::bCreated = false;
+
 FPersistentStorageManager& FPersistentStorageManager::Get()
 {
-	check(IsReady());
+	check(IsReadyInternal());
 	static FPersistentStorageManager Singleton;
+	bCreated = true;
 	return Singleton;
 }
 
diff --git a/Engine/Source/Runtime/Core/Public/HAL/IPlatformFileManagedStorageWrapper.h b/Engine/Source/Runtime/Core/Public/HAL/IPlatformFileManagedStorageWrapper.h
--- a/Engine/Source/Runtime/Core/Public/HAL/IPlatformFileManagedStorageWrapper.h
+++ b/Engine/Source/Runtime/Core/Public/HAL/IPlatformFileManagedStorageWrapper.h
@@ -364,7 +364,7 @@ public:
 		// FPersistentStorageManager depends on GConfig.
 		// FPersistentStorageManager can't be constructed until its dependencies are ready
 		// FPersistentStorageManager will try and allocate memory during a crash but this could hang during log file flushing
-		return GConfig && GConfig->IsReadyForUse() && FCommandLine::IsInitialized() && !GIsCriticalError;
+		return bCreated && GConfig && GConfig->IsReadyForUse() && FCommandLine::IsInitialized() && !GIsCriticalError;
 	}
 
 	/** Singleton access **/
@@ -621,6 +621,16 @@ public:
 private:
 	friend class FManagedStorageScopeFileLock; // For access to Categories
 
+	static bool IsReadyInternal()
+	{
+		// FPersistentStorageManager depends on FPaths which depends on the command line being initialized.
+		// FPersistentStorageManager depends on GConfig.
+		// FPersistentStorageManager can't be constructed until its dependencies are ready
+		// FPersistentStorageManager will try and allocate memory during a crash but this could hang during log file flushing
+		return GConfig && GConfig->IsReadyForUse() && FCommandLine::IsInitialized() && !GIsCriticalError;
+	}
+
+	static bool bCreated;
 	bool bInitialized;
 
 	// Top level of all managed storage
