Fix for FString::ParseIntoArray when string to parse has a length of 0, but still has backing memory

diff --git a/Engine/Source/Runtime/Core/Private/Containers/String.cpp.inl b/Engine/Source/Runtime/Core/Private/Containers/String.cpp.inl
--- a/Engine/Source/Runtime/Core/Private/Containers/String.cpp.inl
+++ b/Engine/Source/Runtime/Core/Private/Containers/String.cpp.inl
@@ -1449,7 +1449,7 @@ int32 UE_STRING_CLASS::ParseIntoArray(TArray<UE_STRING_CLASS>& OutArray, const E
 		if(!InCullEmpty || SubstringLength != 0)
 		{
 			// ... add new string from substring beginning up to the beginning of this delimiter.
-			OutArray.Emplace(Start + SubstringBeginIndex);
+			OutArray.Emplace(TStringView<ElementType>(Start + SubstringBeginIndex, SubstringLength));
 		}
 	}
 
