Fixed inconsistency between EmptyString.ParseIntoArray(Array, DelimiterArray, 1, true) and EmptyString.ParseIntoArray(Array, DelimiterArray[0], true).
Added some tests for this, which also covers the case which was fixed in CL# 31993660.

diff --git a/Engine/Source/Runtime/Core/Private/Containers/String.cpp.inl b/Engine/Source/Runtime/Core/Private/Containers/String.cpp.inl
--- a/Engine/Source/Runtime/Core/Private/Containers/String.cpp.inl
+++ b/Engine/Source/Runtime/Core/Private/Containers/String.cpp.inl
@@ -1400,59 +1400,56 @@ int32 UE_STRING_CLASS::ParseIntoArray(TArray<UE_STRING_CLASS>& OutArray, const E
 	OutArray.Reset();
 	const ElementType* Start = Data.GetData();
 	const int32 Length = Len();
-	if (Start)
-	{
-		int32 SubstringBeginIndex = 0;
+	int32 SubstringBeginIndex = 0;
 
 		// Iterate through string.
-		for(int32 i = 0; i < Len();)
+	for(int32 i = 0; i < Length;)
+	{
+		int32 SubstringEndIndex = INDEX_NONE;
+		int32 DelimiterLength = 0;
+
+		// Attempt each delimiter.
+		for(int32 DelimIndex = 0; DelimIndex < NumDelims; ++DelimIndex)
 		{
-			int32 SubstringEndIndex = INDEX_NONE;
-			int32 DelimiterLength = 0;
+			DelimiterLength = TCString<ElementType>::Strlen(DelimArray[DelimIndex]);
 
-			// Attempt each delimiter.
-			for(int32 DelimIndex = 0; DelimIndex < NumDelims; ++DelimIndex)
+			// If we found a delimiter...
+			if (TCString<ElementType>::Strncmp(Start + i, DelimArray[DelimIndex], DelimiterLength) == 0)
 			{
-				DelimiterLength = TCString<ElementType>::Strlen(DelimArray[DelimIndex]);
-
-				// If we found a delimiter...
-				if (TCString<ElementType>::Strncmp(Start + i, DelimArray[DelimIndex], DelimiterLength) == 0)
-				{
-					// Mark the end of the substring.
-					SubstringEndIndex = i;
-					break;
-				}
+				// Mark the end of the substring.
+				SubstringEndIndex = i;
+				break;
 			}
+		}
 
-			if (SubstringEndIndex != INDEX_NONE)
-			{
-				const int32 SubstringLength = SubstringEndIndex - SubstringBeginIndex;
-				// If we're not culling empty strings or if we are but the string isn't empty anyways...
-				if(!InCullEmpty || SubstringLength != 0)
-				{
-					// ... add new string from substring beginning up to the beginning of this delimiter.
-					OutArray.Add(UE_STRING_CLASS::ConstructFromPtrSize(Start + SubstringBeginIndex, SubstringEndIndex - SubstringBeginIndex));
-				}
-				// Next substring begins at the end of the discovered delimiter.
-				SubstringBeginIndex = SubstringEndIndex + DelimiterLength;
-				i = SubstringBeginIndex;
-			}
-			else
+		if (SubstringEndIndex != INDEX_NONE)
+		{
+			const int32 SubstringLength = SubstringEndIndex - SubstringBeginIndex;
+			// If we're not culling empty strings or if we are but the string isn't empty anyways...
+			if(!InCullEmpty || SubstringLength != 0)
 			{
-				++i;
+				// ... add new string from substring beginning up to the beginning of this delimiter.
+				OutArray.Add(UE_STRING_CLASS::ConstructFromPtrSize(Start + SubstringBeginIndex, SubstringEndIndex - SubstringBeginIndex));
 			}
+			// Next substring begins at the end of the discovered delimiter.
+			SubstringBeginIndex = SubstringEndIndex + DelimiterLength;
+			i = SubstringBeginIndex;
 		}
-
-		// Add any remaining characters after the last delimiter.
-		const int32 SubstringLength = Length - SubstringBeginIndex;
-		// If we're not culling empty strings or if we are but the string isn't empty anyways...
-		if(!InCullEmpty || SubstringLength != 0)
+		else
 		{
-			// ... add new string from substring beginning up to the beginning of this delimiter.
-			OutArray.Emplace(TStringView<ElementType>(Start + SubstringBeginIndex, SubstringLength));
+			++i;
 		}
 	}
 
+	// Add any remaining characters after the last delimiter.
+	const int32 SubstringLength = Length - SubstringBeginIndex;
+	// If we're not culling empty strings or if we are but the string isn't empty anyways...
+	if(!InCullEmpty || SubstringLength != 0)
+	{
+		// ... add new string from substring beginning up to the beginning of this delimiter.
+		OutArray.Emplace(TStringView<ElementType>(Start + SubstringBeginIndex, SubstringLength));
+	}
+
 	return OutArray.Num();
 }
 
diff --git a/Engine/Source/Runtime/Core/Private/Containers/StringTest.cpp b/Engine/Source/Runtime/Core/Private/Containers/StringTest.cpp
--- a/Engine/Source/Runtime/Core/Private/Containers/StringTest.cpp
+++ b/Engine/Source/Runtime/Core/Private/Containers/StringTest.cpp
@@ -240,6 +240,49 @@ bool FLexTryParseStringTest::RunTest(const FString& Parameters)
 	return true;
 }
 
+IMPLEMENT_SIMPLE_AUTOMATION_TEST(FStringParseIntoArrayTest, "System.Core.String.ParseIntoArray", EAutomationTestFlags::ApplicationContextMask | EAutomationTestFlags::SmokeFilter)
+
+bool FStringParseIntoArrayTest::RunTest(const FString& Parameters)
+{
+	FString Empty;
+	FString EmptyButAllocated = TEXT("This|is|a|long|string");
+	EmptyButAllocated.GetCharArray().Reset();
+
+	const TCHAR* Delim[] = { TEXT("|") };
+
+	// Test that an array of 1 delimeter and single delimeters work the same way
+	{
+		TArray<FString> Parsed;
+
+		Parsed.Reset(); FString(TEXT("|"))  .ParseIntoArray(Parsed, Delim[0], /*cull = */true);  TestTrue(TEXT("parse a single delimiter character with a single delimiter, culling empties"),            Parsed == TArray<FString>{});
+		Parsed.Reset(); FString(TEXT("|"))  .ParseIntoArray(Parsed, Delim, 1, /*cull = */true);  TestTrue(TEXT("parse a single delimiter character with an array of 1 delimiter, culling empties"),       Parsed == TArray<FString>{});
+		Parsed.Reset(); FString{}           .ParseIntoArray(Parsed, Delim[0], /*cull = */true);  TestTrue(TEXT("parse an unallocated string with a single delimiter, culling empties"),                   Parsed == TArray<FString>{});
+		Parsed.Reset(); FString{}           .ParseIntoArray(Parsed, Delim, 1, /*cull = */true);  TestTrue(TEXT("parse an unallocated string with an array of 1 delimiter, culling empties"),              Parsed == TArray<FString>{});
+		Parsed.Reset(); EmptyButAllocated   .ParseIntoArray(Parsed, Delim[0], /*cull = */true);  TestTrue(TEXT("parse an allocated empty string with a single delimiter, culling empties"),               Parsed == TArray<FString>{});
+		Parsed.Reset(); EmptyButAllocated   .ParseIntoArray(Parsed, Delim, 1, /*cull = */true);  TestTrue(TEXT("parse an allocated empty string with an array of 1 delimiter, culling empties"),          Parsed == TArray<FString>{});
+		Parsed.Reset(); FString(TEXT("a|")) .ParseIntoArray(Parsed, Delim[0], /*cull = */true);  TestTrue(TEXT("parse a string with a trailing delimiter with a single delimiter, culling empties"),      Parsed == TArray<FString>{TEXT("a")});
+		Parsed.Reset(); FString(TEXT("a|")) .ParseIntoArray(Parsed, Delim, 1, /*cull = */true);  TestTrue(TEXT("parse a string with a trailing delimiter with an array of 1 delimiter, culling empties"), Parsed == TArray<FString>{TEXT("a")});
+		Parsed.Reset(); FString(TEXT("|b")) .ParseIntoArray(Parsed, Delim[0], /*cull = */true);  TestTrue(TEXT("parse a string with a leading delimiter with a single delimiter, culling empties"),       Parsed == TArray<FString>{TEXT("b")});
+		Parsed.Reset(); FString(TEXT("|b")) .ParseIntoArray(Parsed, Delim, 1, /*cull = */true);  TestTrue(TEXT("parse a string with a leading delimiter with an array of 1 delimiter, culling empties"),  Parsed == TArray<FString>{TEXT("b")});
+		Parsed.Reset(); FString(TEXT("a|b")).ParseIntoArray(Parsed, Delim[0], /*cull = */true);  TestTrue(TEXT("parse a string with a middle delimiter with a single delimiter, culling empties"),        Parsed == TArray<FString>{TEXT("a"), TEXT("b")});
+		Parsed.Reset(); FString(TEXT("a|b")).ParseIntoArray(Parsed, Delim, 1, /*cull = */true);  TestTrue(TEXT("parse a string with a middle delimiter with an array of 1 delimiter, culling empties"),   Parsed == TArray<FString>{TEXT("a"), TEXT("b")});
+		Parsed.Reset(); FString(TEXT("|"))  .ParseIntoArray(Parsed, Delim[0], /*cull = */false); TestTrue(TEXT("parse a single delimiter character with a single delimiter, keeping empties"),            Parsed == TArray<FString>{TEXT(""), TEXT("")});
+		Parsed.Reset(); FString(TEXT("|"))  .ParseIntoArray(Parsed, Delim, 1, /*cull = */false); TestTrue(TEXT("parse a single delimiter character with an array of 1 delimiter, keeping empties"),       Parsed == TArray<FString>{TEXT(""), TEXT("")});
+		Parsed.Reset(); FString{}           .ParseIntoArray(Parsed, Delim[0], /*cull = */false); TestTrue(TEXT("parse an unallocated string with a single delimiter, keeping empties"),                   Parsed == TArray<FString>{TEXT("")});
+		Parsed.Reset(); FString{}           .ParseIntoArray(Parsed, Delim, 1, /*cull = */false); TestTrue(TEXT("parse an unallocated string with an array of 1 delimiter, keeping empties"),              Parsed == TArray<FString>{TEXT("")});
+		Parsed.Reset(); EmptyButAllocated   .ParseIntoArray(Parsed, Delim[0], /*cull = */false); TestTrue(TEXT("parse an allocated empty string with a single delimiter, keeping empties"),               Parsed == TArray<FString>{TEXT("")});
+		Parsed.Reset(); EmptyButAllocated   .ParseIntoArray(Parsed, Delim, 1, /*cull = */false); TestTrue(TEXT("parse an allocated empty string with an array of 1 delimiter, keeping empties"),          Parsed == TArray<FString>{TEXT("")});
+		Parsed.Reset(); FString(TEXT("a|")) .ParseIntoArray(Parsed, Delim[0], /*cull = */false); TestTrue(TEXT("parse a string with a trailing delimiter with a single delimiter, keeping empties"),      Parsed == TArray<FString>{TEXT("a"), TEXT("")});
+		Parsed.Reset(); FString(TEXT("a|")) .ParseIntoArray(Parsed, Delim, 1, /*cull = */false); TestTrue(TEXT("parse a string with a trailing delimiter with an array of 1 delimiter, keeping empties"), Parsed == TArray<FString>{TEXT("a"), TEXT("")});
+		Parsed.Reset(); FString(TEXT("|b")) .ParseIntoArray(Parsed, Delim[0], /*cull = */false); TestTrue(TEXT("parse a string with a leading delimiter with a single delimiter, keeping empties"),       Parsed == TArray<FString>{TEXT(""), TEXT("b")});
+		Parsed.Reset(); FString(TEXT("|b")) .ParseIntoArray(Parsed, Delim, 1, /*cull = */false); TestTrue(TEXT("parse a string with a leading delimiter with an array of 1 delimiter, keeping empties"),  Parsed == TArray<FString>{TEXT(""), TEXT("b")});
+		Parsed.Reset(); FString(TEXT("a|b")).ParseIntoArray(Parsed, Delim[0], /*cull = */false); TestTrue(TEXT("parse a string with a middle delimiter with a single delimiter, keeping empties"),        Parsed == TArray<FString>{TEXT("a"), TEXT("b")});
+		Parsed.Reset(); FString(TEXT("a|b")).ParseIntoArray(Parsed, Delim, 1, /*cull = */false); TestTrue(TEXT("parse a string with a middle delimiter with an array of 1 delimiter, keeping empties"),   Parsed == TArray<FString>{TEXT("a"), TEXT("b")});
+	}
+
+	return true;
+}
+
 IMPLEMENT_SIMPLE_AUTOMATION_TEST(FStringSubstringTest, "System.Core.String.Substring", EAutomationTestFlags::ApplicationContextMask | EAutomationTestFlags::SmokeFilter)
 bool FStringSubstringTest::RunTest(const FString& Parameters)
 {
