- Added StagedBuild.ini file to staged builds so runtime code can determine if the build is staged or not
- Using it now in SubmitTool and in ProjectDir() logic

diff --git a/Engine/Source/Programs/AutomationTool/Scripts/CopyBuildToStagingDirectory.Automation.cs b/Engine/Source/Programs/AutomationTool/Scripts/CopyBuildToStagingDirectory.Automation.cs
--- a/Engine/Source/Programs/AutomationTool/Scripts/CopyBuildToStagingDirectory.Automation.cs
+++ b/Engine/Source/Programs/AutomationTool/Scripts/CopyBuildToStagingDirectory.Automation.cs
@@ -1239,6 +1239,13 @@ namespace AutomationScripts
 						}
 					}
 
+					// write out a StagedBuild.ini file that's only in Staged builds so runtime knows the build is staged
+					DirectoryReference IntermediateDir = DirectoryReference.Combine(SC.EngineRoot, "Intermediate", "Build", "BuildCookRun");
+					FileReference StagedBuildConfigFile = FileReference.Combine(IntermediateDir, "StagedBuild.ini");
+					DirectoryReference.CreateDirectory(IntermediateDir);
+					File.WriteAllText(StagedBuildConfigFile.FullName, "");
+					SC.StageFile(StagedFileType.UFS, StagedBuildConfigFile, new StagedFileReference($"Engine/Config/{StagedBuildConfigFile.GetFileName()}"));
+
 					// Stage ICU internationalization data from Engine.
 					var ICUDataVersion = SC.StageTargetPlatform.ICUDataVersion;
 					SC.StageFiles(StagedFileType.UFS, DirectoryReference.Combine(SC.LocalRoot, "Engine", "Content", "Internationalization", InternationalizationPreset, ICUDataVersion), StageFilesSearch.AllDirectories, new StagedDirectoryReference(String.Format("Engine/Content/Internationalization/{0}", ICUDataVersion)));
diff --git a/Engine/Source/Runtime/Core/Private/GenericPlatform/GenericPlatformMisc.cpp b/Engine/Source/Runtime/Core/Private/GenericPlatform/GenericPlatformMisc.cpp
--- a/Engine/Source/Runtime/Core/Private/GenericPlatform/GenericPlatformMisc.cpp
+++ b/Engine/Source/Runtime/Core/Private/GenericPlatform/GenericPlatformMisc.cpp
@@ -1277,14 +1277,16 @@ const TCHAR* FGenericPlatformMisc::ProjectDir()
 		ProjectDir.Reserve(FPlatformMisc::GetMaxPathLength());
 		if (FPlatformProperties::IsProgram())
 		{
-			// monolithic, game-agnostic executables, the ini is in Engine/Config/Platform
-			ProjectDir = FString::Printf(TEXT("../../../Engine/Programs/%s/"), FApp::GetProjectName());
-
-			// however, if it was staged, that directory won't exist, so look in the normal staged location
-			if (!FPlatformFileManager::Get().GetPlatformFile().DirectoryExists(*ProjectDir))
+			if (FPaths::IsStaged())
 			{
+				// if staged, use the remapped location
 				ProjectDir = FString::Printf(TEXT("../../../%s/"), FApp::GetProjectName());
 			}
+			else
+			{
+				// monolithic, game-agnostic executables
+				ProjectDir = FString::Printf(TEXT("../../../Engine/Programs/%s/"), FApp::GetProjectName());
+			}
 		}
 		else
 		{
@@ -1312,10 +1314,6 @@ const TCHAR* FGenericPlatformMisc::ProjectDir()
 							// We found a project folder for the game
 							FPaths::SetProjectFilePath(GameProjectFile);
 							ProjectDir = FPaths::GetPath(GameProjectFile);
-							if (ProjectDir.EndsWith(TEXT("/")) == false)
-							{
-								ProjectDir += TEXT("/");
-							}
 						}
 					}
 				}
@@ -1347,11 +1345,6 @@ const TCHAR* FGenericPlatformMisc::ProjectDir()
 					{
 						ProjectDir = LocalProjectDir;
 					}
-
-					if (ProjectDir.EndsWith(TEXT("/")) == false)
-					{
-						ProjectDir += TEXT("/");
-					}
 #endif
 				}
 			}
@@ -1361,9 +1354,12 @@ const TCHAR* FGenericPlatformMisc::ProjectDir()
 				ProjectDir = FPaths::EngineUserDir();
 				FPaths::NormalizeFilename(ProjectDir);
 				ProjectDir = FFileManagerGeneric::DefaultConvertToRelativePath(*ProjectDir);
-				if(!ProjectDir.EndsWith(TEXT("/"))) ProjectDir += TEXT("/");
 			}
 		}
+		if (!ProjectDir.EndsWith(TEXT("/")))
+		{
+			ProjectDir += TEXT("/");
+		}
 	}
 
 	return *ProjectDir;
diff --git a/Engine/Source/Runtime/Core/Private/Misc/Paths.cpp b/Engine/Source/Runtime/Core/Private/Misc/Paths.cpp
--- a/Engine/Source/Runtime/Core/Private/Misc/Paths.cpp
+++ b/Engine/Source/Runtime/Core/Private/Misc/Paths.cpp
@@ -144,6 +144,12 @@ namespace UE4Paths_Private
 	}
 }
 
+bool FPaths::IsStaged()
+{
+	static bool bIsStaged = FileExists(Combine(EngineConfigDir(), TEXT("StagedBuild.ini")));
+	return bIsStaged;
+}
+
 bool FPaths::ShouldSaveToUserDir()
 {
 	static bool bShouldSaveToUserDir =
diff --git a/Engine/Source/Runtime/Core/Public/Misc/Paths.h b/Engine/Source/Runtime/Core/Public/Misc/Paths.h
--- a/Engine/Source/Runtime/Core/Public/Misc/Paths.h
+++ b/Engine/Source/Runtime/Core/Public/Misc/Paths.h
@@ -36,6 +36,11 @@ class FPaths
 {
 public:
 
+	/**
+	  * @return true if the build has been staged and uses the simplified remapped directory structure made when staging
+	 */
+	static CORE_API bool IsStaged();
+	
 	/**
 	 * Should the "saved" directory structures be rooted in the user dir or relative to the "engine/game" 
 	 */
