Moved UE_REQUIRES into its own header.

diff --git a/Engine/Source/Runtime/Core/Public/Containers/Array.h b/Engine/Source/Runtime/Core/Public/Containers/Array.h
--- a/Engine/Source/Runtime/Core/Public/Containers/Array.h
+++ b/Engine/Source/Runtime/Core/Public/Containers/Array.h
@@ -23,6 +23,7 @@
 #include "Templates/IdentityFunctor.h"
 #include "Templates/Invoke.h"
 #include "Templates/Less.h"
+#include "Templates/Requires.h"
 #include "Templates/Sorting.h"
 #include "Templates/AlignmentTemplates.h"
 #include "Traits/ElementType.h"
diff --git a/Engine/Source/Runtime/Core/Public/Containers/StringView.h b/Engine/Source/Runtime/Core/Public/Containers/StringView.h
--- a/Engine/Source/Runtime/Core/Public/Containers/StringView.h
+++ b/Engine/Source/Runtime/Core/Public/Containers/StringView.h
@@ -11,6 +11,7 @@
 #include "Misc/CString.h"
 #include "Misc/ReverseIterate.h"
 #include "String/Find.h"
+#include "Templates/Requires.h"
 #include "Templates/UnrealTemplate.h"
 #include "Traits/ElementType.h"
 #include "Traits/IsCharEncodingCompatibleWith.h"
diff --git a/Engine/Source/Runtime/Core/Public/GenericPlatform/GenericPlatformMath.h b/Engine/Source/Runtime/Core/Public/GenericPlatform/GenericPlatformMath.h
--- a/Engine/Source/Runtime/Core/Public/GenericPlatform/GenericPlatformMath.h
+++ b/Engine/Source/Runtime/Core/Public/GenericPlatform/GenericPlatformMath.h
@@ -14,6 +14,7 @@
 #include "Templates/Decay.h"
 #include "Templates/IsFloatingPoint.h"
 #include "Templates/UnrealTypeTraits.h"
+#include "Templates/Requires.h"
 #include "Templates/ResolveTypeAmbiguity.h"
 #include "Templates/TypeCompatibleBytes.h"
 #include <limits>
diff --git a/Engine/Source/Runtime/Core/Public/Logging/StructuredLog.h b/Engine/Source/Runtime/Core/Public/Logging/StructuredLog.h
--- a/Engine/Source/Runtime/Core/Public/Logging/StructuredLog.h
+++ b/Engine/Source/Runtime/Core/Public/Logging/StructuredLog.h
@@ -12,6 +12,7 @@
 #include "Serialization/CompactBinaryWriter.h"
 #include "Templates/IsArrayOrRefOfType.h"
 #include "Templates/Models.h"
+#include "Templates/Requires.h"
 #include "Templates/UnrealTypeTraits.h"
 
 #include <atomic>
diff --git a/Engine/Source/Runtime/Core/Public/Math/GuardedInt.h b/Engine/Source/Runtime/Core/Public/Math/GuardedInt.h
--- a/Engine/Source/Runtime/Core/Public/Math/GuardedInt.h
+++ b/Engine/Source/Runtime/Core/Public/Math/GuardedInt.h
@@ -5,6 +5,7 @@
 #include "CoreTypes.h"
 #include "HAL/PlatformMath.h"
 #include "Misc/AssertionMacros.h"
+#include "Templates/Requires.h"
 #include <limits>
 #include <type_traits>
 
diff --git a/Engine/Source/Runtime/Core/Public/Math/UnrealMathUtility.h b/Engine/Source/Runtime/Core/Public/Math/UnrealMathUtility.h
--- a/Engine/Source/Runtime/Core/Public/Math/UnrealMathUtility.h
+++ b/Engine/Source/Runtime/Core/Public/Math/UnrealMathUtility.h
@@ -7,6 +7,7 @@
 #include "HAL/PlatformMath.h"
 #include "Math/MathFwd.h"
 #include "Templates/Identity.h"
+#include "Templates/Requires.h"
 
 // Assert on non finite numbers. Used to track NaNs.
 #ifndef ENABLE_NAN_DIAGNOSTIC
diff --git a/Engine/Source/Runtime/Core/Public/Misc/StringBuilder.h b/Engine/Source/Runtime/Core/Public/Misc/StringBuilder.h
--- a/Engine/Source/Runtime/Core/Public/Misc/StringBuilder.h
+++ b/Engine/Source/Runtime/Core/Public/Misc/StringBuilder.h
@@ -12,6 +12,7 @@
 #include "Templates/EnableIf.h"
 #include "Templates/IsArrayOrRefOfTypeByPredicate.h"
 #include "Templates/IsValidVariadicFunctionArg.h"
+#include "Templates/Requires.h"
 #include "Templates/UnrealTemplate.h"
 #include "Templates/UnrealTypeTraits.h"
 #include "Traits/IsCharEncodingCompatibleWith.h"
diff --git a/Engine/Source/Runtime/Core/Public/Misc/TVariant.h b/Engine/Source/Runtime/Core/Public/Misc/TVariant.h
--- a/Engine/Source/Runtime/Core/Public/Misc/TVariant.h
+++ b/Engine/Source/Runtime/Core/Public/Misc/TVariant.h
@@ -2,6 +2,7 @@
 
 #pragma once
 
+#include "Templates/Requires.h"
 #include "TVariantMeta.h"
 #include <type_traits>
 
diff --git a/Engine/Source/Runtime/Core/Public/Templates/Function.h b/Engine/Source/Runtime/Core/Public/Templates/Function.h
--- a/Engine/Source/Runtime/Core/Public/Templates/Function.h
+++ b/Engine/Source/Runtime/Core/Public/Templates/Function.h
@@ -10,6 +10,7 @@
 #include "Templates/UnrealTypeTraits.h"
 #include "Templates/Invoke.h"
 #include "Templates/UnrealTemplate.h"
+#include "Templates/Requires.h"
 #include "Math/UnrealMathUtility.h"
 #include <new> // IWYU pragma: export
 #include <type_traits>
diff --git a/Engine/Source/Runtime/Core/Public/Templates/RefCounting.h b/Engine/Source/Runtime/Core/Public/Templates/RefCounting.h
--- a/Engine/Source/Runtime/Core/Public/Templates/RefCounting.h
+++ b/Engine/Source/Runtime/Core/Public/Templates/RefCounting.h
@@ -11,6 +11,7 @@
 #include "Serialization/Archive.h"
 #include "Serialization/MemoryLayout.h"
 #include "Templates/TypeHash.h"
+#include "Templates/Requires.h"
 
 /** A virtual interface for ref counted objects to implement. */
 class IRefCountedObject
diff --git a/Engine/Source/Runtime/Core/Public/Templates/Requires.h b/Engine/Source/Runtime/Core/Public/Templates/Requires.h
new file mode 100644
--- /dev/null
+++ b/Engine/Source/Runtime/Core/Public/Templates/Requires.h
@@ -0,0 +1,33 @@
+// Copyright Epic Games, Inc. All Rights Reserved.
+
+#pragma once
+
+#include "Templates/EnableIf.h"
+#include <type_traits>
+
+/*-----------------------------------------------------------------------------
+	Readability macro for a constraint in template definitions, future-proofed
+	for C++ 20 concepts. Usage:
+
+	template <
+		typename T,
+		typename U  // note - no trailing comma before the constraint
+		UE_REQUIRES(std::is_integral_v<T> && sizeof(U) <= 4)
+	>
+	void IntegralUpTo32Bit(T Lhs, U Rhs) {}
+ -----------------------------------------------------------------------------*/
+#if __cplusplus < 202000
+	#define UE_REQUIRES(...) , std::enable_if_t<(__VA_ARGS__), int> = 0
+#else
+	namespace UE::Core::Private
+	{
+		// Only needed for the UE_REQUIRES macro to work, to allow for a trailing > token after the macro
+		template <bool B>
+		concept BoolIdentityConcept = B;
+	}
+
+	#define UE_REQUIRES(...) > requires (!!(__VA_ARGS__)) && UE::Core::Private::BoolIdentityConcept<true
+#endif
+
+// This should be regarded as deprecated - please use UE_REQUIRES instead.
+#define TEMPLATE_REQUIRES(...) typename TEnableIf<__VA_ARGS__, int>::type = 0
diff --git a/Engine/Source/Runtime/Core/Public/Templates/ResolveTypeAmbiguity.h b/Engine/Source/Runtime/Core/Public/Templates/ResolveTypeAmbiguity.h
--- a/Engine/Source/Runtime/Core/Public/Templates/ResolveTypeAmbiguity.h
+++ b/Engine/Source/Runtime/Core/Public/Templates/ResolveTypeAmbiguity.h
@@ -7,6 +7,7 @@
 #include "Templates/UnrealTypeTraits.h"
 #include "Templates/IsSigned.h"
 #include "Templates/IsIntegral.h"
+#include "Templates/Requires.h"
 
 /**
  * This file defines a few macros which can be used to restrict (via deprecation) ambiguous calls to functions expecting float params/return type.
diff --git a/Engine/Source/Runtime/Core/Public/Templates/Tuple.h b/Engine/Source/Runtime/Core/Public/Templates/Tuple.h
--- a/Engine/Source/Runtime/Core/Public/Templates/Tuple.h
+++ b/Engine/Source/Runtime/Core/Public/Templates/Tuple.h
@@ -8,6 +8,7 @@
 #include "Templates/Invoke.h"
 #include "Serialization/StructuredArchive.h"
 #include "Serialization/MemoryLayout.h"
+#include "Templates/Requires.h"
 #include "Templates/TypeHash.h"
 #include <tuple>
 
diff --git a/Engine/Source/Runtime/Core/Public/Templates/UniquePtr.h b/Engine/Source/Runtime/Core/Public/Templates/UniquePtr.h
--- a/Engine/Source/Runtime/Core/Public/Templates/UniquePtr.h
+++ b/Engine/Source/Runtime/Core/Public/Templates/UniquePtr.h
@@ -6,6 +6,7 @@
 #include "Templates/UnrealTemplate.h"
 #include "Templates/IsArray.h"
 #include "Templates/RemoveExtent.h"
+#include "Templates/Requires.h"
 #include "Serialization/MemoryLayout.h"
 
 // Single-ownership smart pointer in the vein of std::unique_ptr.
diff --git a/Engine/Source/Runtime/Core/Public/Templates/UnrealTemplate.h b/Engine/Source/Runtime/Core/Public/Templates/UnrealTemplate.h
--- a/Engine/Source/Runtime/Core/Public/Templates/UnrealTemplate.h
+++ b/Engine/Source/Runtime/Core/Public/Templates/UnrealTemplate.h
@@ -8,6 +8,7 @@
 #include "Templates/CopyQualifiersAndRefsFromTo.h"
 #include "Templates/UnrealTypeTraits.h"
 #include "Templates/RemoveReference.h"
+#include "Templates/Requires.h"
 #include "Templates/TypeCompatibleBytes.h"
 #include "Templates/Identity.h"
 #include "Traits/IsContiguousContainer.h"
diff --git a/Engine/Source/Runtime/Core/Public/Templates/UnrealTypeTraits.h b/Engine/Source/Runtime/Core/Public/Templates/UnrealTypeTraits.h
--- a/Engine/Source/Runtime/Core/Public/Templates/UnrealTypeTraits.h
+++ b/Engine/Source/Runtime/Core/Public/Templates/UnrealTypeTraits.h
@@ -21,32 +21,6 @@
 #include "Templates/IsUECoreType.h"
 #include "Templates/IsTriviallyCopyConstructible.h"
 
-/*-----------------------------------------------------------------------------
-	Readability macro for a constraint in template definitions, future-proofed
-	for C++ 20 concepts. Usage:
-
-	template <
-		typename T,
-		typename U  // note - no trailing comma before the constraint
-		UE_REQUIRES(std::is_integral_v<T> && sizeof(U) <= 4)
-	>
-	void IntegralUpTo32Bit(T Lhs, U Rhs) {}
- -----------------------------------------------------------------------------*/
-#if __cplusplus < 202000
-	#define UE_REQUIRES(...) , std::enable_if_t<(__VA_ARGS__), int> = 0
-#else
-	namespace UE::Core::Private
-	{
-		// Only needed for the UE_REQUIRES macro to work, to allow for a trailing > token after the macro
-		template <bool B>
-		concept BoolIdentityConcept = B;
-	}
-
-	#define UE_REQUIRES(...) > requires (!!(__VA_ARGS__)) && UE::Core::Private::BoolIdentityConcept<true
-#endif
-
-#define TEMPLATE_REQUIRES(...) typename TEnableIf<__VA_ARGS__, int>::type = 0
-
 /*-----------------------------------------------------------------------------
  * Macros to abstract the presence of certain compiler intrinsic type traits 
  -----------------------------------------------------------------------------*/
@@ -538,3 +512,7 @@ public:																											\
 #undef IS_EMPTY
 #undef IS_POD
 #undef HAS_TRIVIAL_CONSTRUCTOR
+
+#if UE_ENABLE_INCLUDE_ORDER_DEPRECATED_IN_5_5
+#include "Templates/Requires.h"
+#endif
diff --git a/Engine/Source/Runtime/Core/Public/UObject/WeakObjectPtrTemplates.h b/Engine/Source/Runtime/Core/Public/UObject/WeakObjectPtrTemplates.h
--- a/Engine/Source/Runtime/Core/Public/UObject/WeakObjectPtrTemplates.h
+++ b/Engine/Source/Runtime/Core/Public/UObject/WeakObjectPtrTemplates.h
@@ -4,6 +4,7 @@
 
 #include "CoreTypes.h"
 #include "Templates/LosesQualifiersFromTo.h"
+#include "Templates/Requires.h"
 #include "Containers/Map.h"
 #include "UObject/WeakObjectPtrTemplatesFwd.h"
 
