Temporarily release the lock before calling the OutOfMemory handler, which will allow the crash reporter thread to run.

Mark out-of-memory fatal asserts as OOM for better handling of the reports.

#rnx
#tests frontend using debug commands to generate OOM

diff --git a/Engine/Source/Runtime/Core/Private/HAL/MallocBinned2.cpp b/Engine/Source/Runtime/Core/Private/HAL/MallocBinned2.cpp
--- a/Engine/Source/Runtime/Core/Private/HAL/MallocBinned2.cpp
+++ b/Engine/Source/Runtime/Core/Private/HAL/MallocBinned2.cpp
@@ -310,7 +310,7 @@ struct FMallocBinned2::Private
 		/** 
 		 * Creates an array of FPoolInfo structures for tracking allocations.
 		 */
-		auto CreatePoolArray = [](uint64 NumPools)
+		auto CreatePoolArray = [&Allocator](uint64 NumPools)
 		{
 			uint64 PoolArraySize = NumPools * sizeof(FPoolInfo);
 
@@ -325,6 +325,7 @@ struct FMallocBinned2::Private
 
 			if (!Result)
 			{
+				FScopeUnlock TempUnlock(&Allocator.Mutex);
 				OutOfMemory(PoolArraySize);
 			}
 
@@ -646,6 +647,7 @@ FMallocBinned2::FPoolInfo& FMallocBinned2::FPoolList::PushNewPoolToFront(FMalloc
 	void* FreePtr = Allocator.CachedOSPageAllocator.Allocate(LocalPageSize, FMemory::AllocationHints::SmallPool);
 	if (!FreePtr)
 	{
+		FScopeUnlock TempUnlock(&Allocator.Mutex);
 		Private::OutOfMemory(LocalPageSize);
 	}
 #if !UE_USE_VERYLARGEPAGEALLOCATOR || !BINNED2_BOOKKEEPING_AT_THE_END_OF_LARGEBLOCK
@@ -932,6 +934,7 @@ void* FMallocBinned2::MallocExternalLarge(SIZE_T Size, uint32 Alignment)
 		Result = CachedOSPageAllocator.Allocate(AlignedSize, 0, &Mutex);
 		if (!Result)
 		{
+			FScopeUnlock TempUnlock(&Mutex);
 			Private::OutOfMemory(AlignedSize);
 		}
 
diff --git a/Engine/Source/Runtime/Core/Private/Windows/WindowsPlatformCrashContext.cpp b/Engine/Source/Runtime/Core/Private/Windows/WindowsPlatformCrashContext.cpp
--- a/Engine/Source/Runtime/Core/Private/Windows/WindowsPlatformCrashContext.cpp
+++ b/Engine/Source/Runtime/Core/Private/Windows/WindowsPlatformCrashContext.cpp
@@ -92,6 +92,7 @@ enum EConstants
 const uint32 EnsureExceptionCode = ECrashExitCodes::UnhandledEnsure; // Use a rather unique exception code in case SEH doesn't handle it as expected.
 const uint32 AssertExceptionCode = 0x4000;
 const uint32 GPUCrashExceptionCode = 0x8000;
+const uint32 OutOfMemoryExceptionCode = 0xc000;
 constexpr double DefaultCrashHandlingTimeoutSecs = 60.0;
 
 namespace {
@@ -1446,7 +1447,6 @@ private:
 			ErrorMessage = Info.ErrorMessage;
 			ErrorProgramCounter = Info.ProgramCounter;
 		}
-		// Generic exception description is stored in GErrorExceptionDescription
 		else if (ExceptionInfo->ExceptionRecord->ExceptionCode == EnsureExceptionCode)
 		{
 			const FAssertInfo& Info = *(const FAssertInfo*)ExceptionInfo->ExceptionRecord->ExceptionInformation[0];
@@ -1454,6 +1454,13 @@ private:
 			ErrorMessage = Info.ErrorMessage;
 			ErrorProgramCounter = Info.ProgramCounter;
 		}
+		else if (ExceptionInfo->ExceptionRecord->ExceptionCode == OutOfMemoryExceptionCode)
+		{
+			const FAssertInfo& Info = *(const FAssertInfo*)ExceptionInfo->ExceptionRecord->ExceptionInformation[0];
+			Type = ECrashContextType::OutOfMemory;
+			ErrorMessage = Info.ErrorMessage;
+			ErrorProgramCounter = Info.ProgramCounter;
+		}
 		// Generic exception description is stored in GErrorExceptionDescription
 		else
 		{
@@ -1819,7 +1826,14 @@ void ReportAssert(const TCHAR* ErrorMessage, void* ProgramCounter)
 	FAssertInfo Info(ErrorMessage, ProgramCounter);
 
 	ULONG_PTR Arguments[] = { (ULONG_PTR)&Info };
-	::RaiseException(AssertExceptionCode, 0, UE_ARRAY_COUNT(Arguments), Arguments);
+	if (FGenericPlatformMemory::bIsOOM)
+	{
+		::RaiseException(OutOfMemoryExceptionCode, 0, UE_ARRAY_COUNT(Arguments), Arguments);
+	}
+	else
+	{
+		::RaiseException(AssertExceptionCode, 0, UE_ARRAY_COUNT(Arguments), Arguments);
+	}
 }
 
 FORCENOINLINE void ReportGPUCrash(const TCHAR* ErrorMessage, void* ProgramCounter)
