Revert correct TLS Index check in MB2\3 & CSV Profiler to old behavior in an attemt to mitigate rate Linux server crash

#localization none
#tests locally tested that both MB2 and MB3 allocators work on Win Client builds

diff --git a/Engine/Source/Runtime/Core/Private/HAL/MallocBinned2.cpp b/Engine/Source/Runtime/Core/Private/HAL/MallocBinned2.cpp
--- a/Engine/Source/Runtime/Core/Private/HAL/MallocBinned2.cpp
+++ b/Engine/Source/Runtime/Core/Private/HAL/MallocBinned2.cpp
@@ -1286,11 +1286,11 @@ void FMallocBinned2::SetupTLSCachesOnCurrentThread()
 	{
 		return;
 	}
-	if (!FPlatformTLS::IsValidTlsSlot(FMallocBinned2::BinnedTlsSlot))
+	if (!FMallocBinned2::BinnedTlsSlot)
 	{
 		FMallocBinned2::BinnedTlsSlot = FPlatformTLS::AllocTlsSlot();
 	}
-	check(FPlatformTLS::IsValidTlsSlot(FMallocBinned2::BinnedTlsSlot));
+	check(FMallocBinned2::BinnedTlsSlot);
 	FPerThreadFreeBlockLists::SetTLS();
 }
 
diff --git a/Engine/Source/Runtime/Core/Private/HAL/MallocBinned3.cpp b/Engine/Source/Runtime/Core/Private/HAL/MallocBinned3.cpp
--- a/Engine/Source/Runtime/Core/Private/HAL/MallocBinned3.cpp
+++ b/Engine/Source/Runtime/Core/Private/HAL/MallocBinned3.cpp
@@ -1352,11 +1352,11 @@ void FMallocBinned3::SetupTLSCachesOnCurrentThread()
 	{
 		return;
 	}
-	if (!FPlatformTLS::IsValidTlsSlot(FMallocBinned3::BinnedTlsSlot))
+	if (!FMallocBinned3::BinnedTlsSlot)
 	{
 		FMallocBinned3::BinnedTlsSlot = FPlatformTLS::AllocTlsSlot();
 	}
-	check(FPlatformTLS::IsValidTlsSlot(FMallocBinned3::BinnedTlsSlot));
+	check(FMallocBinned3::BinnedTlsSlot);
 	FPerThreadFreeBlockLists::SetTLS();
 }
 
diff --git a/Engine/Source/Runtime/Core/Private/HAL/MallocBinnedCommon.cpp b/Engine/Source/Runtime/Core/Private/HAL/MallocBinnedCommon.cpp
--- a/Engine/Source/Runtime/Core/Private/HAL/MallocBinnedCommon.cpp
+++ b/Engine/Source/Runtime/Core/Private/HAL/MallocBinnedCommon.cpp
@@ -492,7 +492,7 @@ static FAutoConsoleVariableRef GMallocBinned3BundleCountCVar(
 
 #endif
 
-uint32 FMallocBinnedCommonBase::BinnedTlsSlot = FPlatformTLS::InvalidTlsSlot;
+uint32 FMallocBinnedCommonBase::BinnedTlsSlot = 0;
 #if UE_BINNEDCOMMON_ALLOCATOR_STATS
 std::atomic<int64> FMallocBinnedCommonBase::TLSMemory(0);
 std::atomic<int64> FMallocBinnedCommonBase::ConsolidatedMemory(0);
diff --git a/Engine/Source/Runtime/Core/Private/ProfilingDebugging/CsvProfiler.cpp b/Engine/Source/Runtime/Core/Private/ProfilingDebugging/CsvProfiler.cpp
--- a/Engine/Source/Runtime/Core/Private/ProfilingDebugging/CsvProfiler.cpp
+++ b/Engine/Source/Runtime/Core/Private/ProfilingDebugging/CsvProfiler.cpp
@@ -1982,7 +1982,7 @@ private:
 public:
 	static void InitTls()
 	{
-		if (!FPlatformTLS::IsValidTlsSlot(TlsSlot))
+		if (TlsSlot == 0)
 		{
 			TlsSlot = FPlatformTLS::AllocTlsSlot();
 			FPlatformMisc::MemoryBarrier();
@@ -2226,7 +2226,7 @@ public:
 	TArray<const char*> WaitStatNameStack;
 };
 
-uint32 FCsvProfilerThreadData::TlsSlot = FPlatformTLS::InvalidTlsSlot;
+uint32 FCsvProfilerThreadData::TlsSlot = 0;
 FCriticalSection FCsvProfilerThreadData::TlsCS;
 TArray<FCsvProfilerThreadData::FWeakPtr> FCsvProfilerThreadData::TlsInstances;
 
diff --git a/Engine/Source/Runtime/Core/Public/HAL/MallocBinnedCommon.h b/Engine/Source/Runtime/Core/Public/HAL/MallocBinnedCommon.h
--- a/Engine/Source/Runtime/Core/Public/HAL/MallocBinnedCommon.h
+++ b/Engine/Source/Runtime/Core/Public/HAL/MallocBinnedCommon.h
@@ -403,12 +403,12 @@ protected:
 	{
 		FORCEINLINE static FPerThreadFreeBlockLists* Get()
 		{
-			return FPlatformTLS::IsValidTlsSlot(BinnedTlsSlot) ? (FPerThreadFreeBlockLists*)FPlatformTLS::GetTlsValue(BinnedTlsSlot) : nullptr;
+			return BinnedTlsSlot ? (FPerThreadFreeBlockLists*)FPlatformTLS::GetTlsValue(BinnedTlsSlot) : nullptr;
 		}
 
 		static void SetTLS()
 		{
-			check(FPlatformTLS::IsValidTlsSlot(BinnedTlsSlot));
+			check(BinnedTlsSlot);
 			FPerThreadFreeBlockLists* ThreadSingleton = (FPerThreadFreeBlockLists*)FPlatformTLS::GetTlsValue(BinnedTlsSlot);
 			if (!ThreadSingleton)
 			{
@@ -425,7 +425,7 @@ protected:
 
 		static void ClearTLS()
 		{
-			check(FPlatformTLS::IsValidTlsSlot(BinnedTlsSlot));
+			check(BinnedTlsSlot);
 			FPerThreadFreeBlockLists* ThreadSingleton = (FPerThreadFreeBlockLists*)FPlatformTLS::GetTlsValue(BinnedTlsSlot);
 			if (ThreadSingleton)
 			{
