Fix FBuildVersion to be able to load newer versions containing the "BuildUrl" key/pair.

#rnx

- The C# code that writes out version info added a new field "BuildUrl" last year but the cpp code was not updated.
- Unfortunatly the cpp loading code is set to fail reading the info if it encounters an unknown field so "BuildUrl" not being supported prevented it from reading the newer info.
- Added logging when unknown fields are encountered so this might be noticed quicker next time.
- Added ::TryReadFromString that allows the version info to be read from a string rather than a .version file. This would allow higher level code to open a .target file, then find the version info json object and pass it to FBuildVersion in a way that it can be read.

diff --git a/Engine/Source/Runtime/Core/Private/Modules/BuildVersion.cpp b/Engine/Source/Runtime/Core/Private/Modules/BuildVersion.cpp
--- a/Engine/Source/Runtime/Core/Private/Modules/BuildVersion.cpp
+++ b/Engine/Source/Runtime/Core/Private/Modules/BuildVersion.cpp
@@ -85,6 +85,11 @@ bool FBuildVersion::TryRead(const FString& FileName, FBuildVersion& OutVersion)
 		return false;
 	}
 
+	return TryReadFromString(Text, OutVersion);
+}
+
+bool FBuildVersion::TryReadFromString(const FString& Text, FBuildVersion& OutVersion)
+{
 	const TCHAR* Ptr = *Text;
 	if (!FSimpleParse::MatchZeroOrMoreWhitespace(Ptr) || !FSimpleParse::MatchChar(Ptr, TEXT('{')) || !FSimpleParse::MatchZeroOrMoreWhitespace(Ptr) || FSimpleParse::MatchChar(Ptr, TEXT('}')))
 	{
@@ -183,8 +188,16 @@ bool FBuildVersion::TryRead(const FString& FileName, FBuildVersion& OutVersion)
 				return false;
 			}
 		}
+		else if (Field == TEXT("BuildURL"))
+		{
+			if (!FSimpleParse::ParseString(Ptr, OutVersion.BuildUrl))
+			{
+				return false;
+			}
+		}
 		else
 		{
+			UE_LOG(LogCore, Warning, TEXT("Unknown field '%s' found when serializing FBuildVersion"), *Field);
 			return false;
 		}
 
diff --git a/Engine/Source/Runtime/Core/Public/Modules/BuildVersion.h b/Engine/Source/Runtime/Core/Public/Modules/BuildVersion.h
--- a/Engine/Source/Runtime/Core/Public/Modules/BuildVersion.h
+++ b/Engine/Source/Runtime/Core/Public/Modules/BuildVersion.h
@@ -64,6 +64,11 @@ public:
 	 */
 	FString BuildVersion;
 
+	/**
+	 * [Optional] A url to the web page for the build (if it was created by an automated build systme for example)
+	 */
+	FString BuildUrl;
+
 	/**
 	 * Default constructor. Initializes the structure to empty.
 	 */
@@ -102,9 +107,18 @@ public:
 	/**
 	 * Try to read a version file from disk
 	 *
-	 * @param FileName Path to the version file
-	 * @param OutVersion The version information
-	 * @return True if the version was read successfully, false otherwise
+	 * @param FileName		Path to the version file
+	 * @param OutVersion	The version information
+	 * @return				True if the version was read successfully, false otherwise
 	 */
 	static CORE_API bool TryRead(const FString& FileName, FBuildVersion& OutVersion);
+
+	/**
+	 * Try to read version info from data stored in a string.
+	 * 
+	 * @param Text			The version info stored in string form
+	 * @param OutVersion	The version information
+	 * @return				True if the version was read successfully, false otherwise
+	 */
+	static CORE_API bool TryReadFromString(const FString& Text, FBuildVersion& OutVersion);
 };
