LLM - Changed bIsDisabled into a multi-state variable that includes the state NotYetKnown. During the NotYetKnown state, the multithreaded synchronization strategy is different, because some values that are normally immutable have not yet been initialized.
#rnx

diff --git a/Engine/Source/Runtime/Core/Private/HAL/LowLevelMemTracker.cpp b/Engine/Source/Runtime/Core/Private/HAL/LowLevelMemTracker.cpp
--- a/Engine/Source/Runtime/Core/Private/HAL/LowLevelMemTracker.cpp
+++ b/Engine/Source/Runtime/Core/Private/HAL/LowLevelMemTracker.cpp
@@ -1004,6 +1004,55 @@ FAutoConsoleCommand DumpLLM(
 		FLowLevelMemTracker::Get().DumpToLog(DumpFormat, &Ar, SizeParams, TagSet);
 	}));
 
+#if !PLATFORM_HAS_MULTITHREADED_PREMAIN
+struct FLowLevelMemTracker::FEnableStateScopeLock
+{
+};
+bool FLowLevelMemTracker::TryEnterEnabled(FLowLevelMemTracker::FEnableStateScopeLock& ScopeLock)
+{
+	return IsEnabled();
+}
+#else // !PLATFORM_HAS_MULTITHREADED_PREMAIN
+struct FLowLevelMemTracker::FEnableStateScopeLock
+{
+	// Undefine copy/move constructor since FReadScopeLock does not support it.
+	FEnableStateScopeLock() = default;
+	FEnableStateScopeLock(const FEnableStateScopeLock&) = delete;
+	FEnableStateScopeLock(FEnableStateScopeLock&&) = delete;
+
+	TOptional<FReadScopeLock> Inner;
+};
+
+namespace UE::LLMPrivate
+{
+
+FRWLock& GetEnableStateLock()
+{
+	static FRWLock Lock;
+	return Lock;
+}
+
+}
+
+bool FLowLevelMemTracker::TryEnterEnabled(FLowLevelMemTracker::FEnableStateScopeLock& ScopeLock)
+{
+	switch (EnabledState)
+	{
+	case EEnabled::NotYetKnown:
+		ScopeLock.Inner.Emplace(UE::LLMPrivate::GetEnableStateLock());
+		// Evaluate EnabledState again since it may have changed under ProcessCommandLine's WriteLock.
+		return IsEnabled();
+	case EEnabled::Disabled:
+		return false;
+	case EEnabled::Enabled:
+		return true;
+	default:
+		checkNoEntry();
+		return false;
+	}
+}
+#endif // !PLATFORM_HAS_MULTITHREADED_PREMAIN
+
 void FLowLevelMemTracker::DumpToLog(EDumpFormat DumpFormat, FOutputDevice* OutputDevice, UE::LLM::ESizeParams SizeParams, ELLMTagSet TagSet)
 {
 	if (!IsEnabled())
@@ -1707,14 +1756,10 @@ FLowLevelMemTracker& FLowLevelMemTracker::Construct()
 	return Tracker;
 }
 
-bool FLowLevelMemTracker::IsEnabled()
-{
-	return !bIsDisabled;
-}
-
 FLowLevelMemTracker* FLowLevelMemTracker::TrackerInstance = nullptr;
 // LLM must start off enabled because allocations happen before the command line enables/disables us
 bool FLowLevelMemTracker::bIsDisabled = false;
+FLowLevelMemTracker::EEnabled FLowLevelMemTracker::EnabledState = EEnabled::NotYetKnown;
 
 static const TCHAR* InvalidLLMTagName = TEXT("?");
 
@@ -1746,7 +1791,10 @@ FLowLevelMemTracker::FLowLevelMemTracker()
 	int32 Alignment = 0;
 	if (!FPlatformMemory::GetLLMAllocFunctions(PlatformLLMAlloc, PlatformLLMFree, Alignment))
 	{
+		EnabledState = EEnabled::Disabled;
+		PRAGMA_DISABLE_DEPRECATION_WARNINGS;
 		bIsDisabled = true;
+		PRAGMA_ENABLE_DEPRECATION_WARNINGS;
 		bCanEnable = false;
 		bConfigurationComplete = true;
 		return;
@@ -1767,7 +1815,10 @@ FLowLevelMemTracker::~FLowLevelMemTracker()
 {
 	using namespace UE::LLMPrivate;
 
+	EnabledState = EEnabled::Disabled;
+	PRAGMA_DISABLE_DEPRECATION_WARNINGS;
 	bIsDisabled = true;
+	PRAGMA_ENABLE_DEPRECATION_WARNINGS;
 	Clear();
 	FLLMAllocator::Get() = nullptr;
 }
@@ -1812,7 +1863,7 @@ void FLowLevelMemTracker::Clear()
 		return;
 	}
 
-	LLMCheck(bIsDisabled); // tracking must be stopped at this point or it will crash while tracking its own destruction
+	LLMCheck(!IsEnabled()); // tracking must be stopped at this point or it will crash while tracking its own destruction
 	for (int32 TrackerIndex = 0; TrackerIndex < static_cast<int32>(ELLMTracker::Max); TrackerIndex++)
 	{
 		GetTracker((ELLMTracker)TrackerIndex)->Clear();
@@ -1829,8 +1880,7 @@ void FLowLevelMemTracker::Clear()
 void FLowLevelMemTracker::OnPreFork()
 {
 	using namespace UE::LLMPrivate;
-
-	if (!bIsDisabled)
+	if (IsEnabled())
 	{
 		FLLMTracker& DefaultTracker = *GetTracker(ELLMTracker::Default);
 		FLLMTracker& PlatformTracker = *GetTracker(ELLMTracker::Platform);
@@ -1847,7 +1897,7 @@ void FLowLevelMemTracker::UpdateStatsPerFrame(const TCHAR* LogName)
 	LlmTrackArrayTick();
 #endif
 
-	if (bIsDisabled)
+	if (!IsEnabled())
 	{
 		if (bFirstTimeUpdating)
 		{
@@ -1888,7 +1938,7 @@ void FLowLevelMemTracker::UpdateStatsPerFrame(const TCHAR* LogName)
 
 void FLowLevelMemTracker::Tick()
 {
-	if (bIsDisabled)
+	if (!IsEnabled())
 	{
 		return;
 	}
@@ -2071,7 +2121,7 @@ void FLowLevelMemTracker::PublishDataPerFrame(const TCHAR* LogName)
 
 	// set overhead stats
 	SET_MEMORY_STAT(STAT_LLMOverheadTotal, MemoryUsageCurrentOverhead);
-	if (!bIsDisabled)
+	if (IsEnabled())
 	{
 		FLLMTracker& DefaultTracker = *GetTracker(ELLMTracker::Default);
 		FLLMTracker& PlatformTracker = *GetTracker(ELLMTracker::Platform);
@@ -2147,7 +2197,7 @@ void FLowLevelMemTracker::InitialiseProgramSize()
 
 void FLowLevelMemTracker::SetProgramSize(uint64 InProgramSize)
 {
-	if (bIsDisabled)
+	if (!IsEnabled())
 	{
 		return;
 	}
@@ -2192,7 +2242,10 @@ void FLowLevelMemTracker::ProcessCommandLine(const TCHAR* CmdLine)
 
 	if (!bCanEnable)
 	{
+		LLMCheck(EnabledState == EEnabled::Disabled);
+		PRAGMA_DISABLE_DEPRECATION_WARNINGS;
 		LLMCheck(bIsDisabled);
+		PRAGMA_ENABLE_DEPRECATION_WARNINGS;
 		if (!bShouldDisable)
 		{
 			UE_LOG(LogInit, Log,
@@ -2206,11 +2259,26 @@ void FLowLevelMemTracker::ProcessCommandLine(const TCHAR* CmdLine)
 	{
 		// Before we shutdown, update once so we can publish the overhead-when-disabled later during the first
 		// call to UpdateStatsPerFrame.
-		if (!bIsDisabled)
+		if (IsEnabled())
 		{
 			Tick();
 		}
-		bIsDisabled = true;
+
+		{
+#if PLATFORM_HAS_MULTITHREADED_PREMAIN
+			// The EnableStateLock must be limited in scope because other code in the function
+			// allocates memory and would block on the readlock. Trying to take it around the write of the
+			// EnabledState is sufficient; this will cause us to wait until threads already in OnLowLevelAlloc exit the
+			// function and clear their readlock, and will cause new calls to those functions to block while we're waiting
+			// and call IsEnabled again after we release the lock.
+			FWriteScopeLock EnableStateLock(UE::LLMPrivate::GetEnableStateLock());
+#endif
+
+			EnabledState = EEnabled::Disabled;
+			PRAGMA_DISABLE_DEPRECATION_WARNINGS;
+			bIsDisabled = true;
+			PRAGMA_ENABLE_DEPRECATION_WARNINGS;
+		}
 		bCsvWriterEnabled = false;
 		bTraceWriterEnabled = false;
 		bCanEnable = false; // Reenabling after a clear is not implemented
@@ -2218,7 +2286,13 @@ void FLowLevelMemTracker::ProcessCommandLine(const TCHAR* CmdLine)
 		return;
 	}
 	CSV_METADATA(TEXT("LLM"), TEXT("1"));
+
+	// PLATFORM_HAS_MULTITHREADED_PREMAIN: No need for a Write lock because we are not changing state to disabled.
+	// The other data we modify in this function is synchronized using other methods.
+	EnabledState = EEnabled::Enabled;
+	PRAGMA_DISABLE_DEPRECATION_WARNINGS;
 	bIsDisabled = false;
+	PRAGMA_ENABLE_DEPRECATION_WARNINGS;
 	bCsvWriterEnabled = bLocalCsvWriterEnabled;
 	bTraceWriterEnabled = bLocalTraceWriterEnabled;
 	FinishInitialise();
@@ -2275,7 +2349,7 @@ void FLowLevelMemTracker::ProcessCommandLine(const TCHAR* CmdLine)
 // Return the total amount of memory being tracked
 uint64 FLowLevelMemTracker::GetTotalTrackedMemory(ELLMTracker Tracker)
 {
-	if (bIsDisabled)
+	if (!IsEnabled())
 	{
 		return 0;
 	}
@@ -2287,10 +2361,12 @@ uint64 FLowLevelMemTracker::GetTotalTrackedMemory(ELLMTracker Tracker)
 void FLowLevelMemTracker::OnLowLevelAlloc(ELLMTracker Tracker, const void* Ptr, uint64 Size, ELLMTag DefaultTag,
 	ELLMAllocType AllocType, bool bTrackInMemPro)
 {
-	if (bIsDisabled)
+	FEnableStateScopeLock EnableScopeLock;
+	if (!TryEnterEnabled(EnableScopeLock))
 	{
 		return;
 	}
+
 	BootstrapInitialise();
 
 	GetTracker(Tracker)->TrackAllocation(Ptr, static_cast<int64>(Size), DefaultTag, AllocType, bTrackInMemPro);
@@ -2299,10 +2375,12 @@ void FLowLevelMemTracker::OnLowLevelAlloc(ELLMTracker Tracker, const void* Ptr,
 void FLowLevelMemTracker::OnLowLevelAlloc(ELLMTracker Tracker, const void* Ptr, uint64 Size, FName DefaultTag,
 	ELLMAllocType AllocType, bool bTrackInMemPro)
 {
-	if (bIsDisabled)
+	FEnableStateScopeLock EnableScopeLock;
+	if (!TryEnterEnabled(EnableScopeLock))
 	{
 		return;
 	}
+
 	BootstrapInitialise();
 
 	GetTracker(Tracker)->TrackAllocation(Ptr, static_cast<int64>(Size), DefaultTag, AllocType, bTrackInMemPro);
@@ -2311,10 +2389,12 @@ void FLowLevelMemTracker::OnLowLevelAlloc(ELLMTracker Tracker, const void* Ptr,
 void FLowLevelMemTracker::OnLowLevelFree(ELLMTracker Tracker, const void* Ptr,
 	ELLMAllocType AllocType, bool bTrackInMemPro)
 {
-	if (bIsDisabled)
+	FEnableStateScopeLock EnableScopeLock;
+	if (!TryEnterEnabled(EnableScopeLock))
 	{
 		return;
 	}
+
 	BootstrapInitialise();
 
 	if (Ptr != nullptr)
@@ -2326,10 +2406,12 @@ void FLowLevelMemTracker::OnLowLevelFree(ELLMTracker Tracker, const void* Ptr,
 void FLowLevelMemTracker::OnLowLevelChangeInMemoryUse(ELLMTracker Tracker, int64 DeltaMemory, ELLMTag DefaultTag,
 	ELLMAllocType AllocType)
 {
-	if (bIsDisabled)
+	FEnableStateScopeLock EnableScopeLock;
+	if (!TryEnterEnabled(EnableScopeLock))
 	{
 		return;
 	}
+
 	BootstrapInitialise();
 	GetTracker(Tracker)->TrackMemoryOfActiveTag(DeltaMemory, DefaultTag, AllocType);
 }
@@ -2337,10 +2419,12 @@ void FLowLevelMemTracker::OnLowLevelChangeInMemoryUse(ELLMTracker Tracker, int64
 void FLowLevelMemTracker::OnLowLevelChangeInMemoryUse(ELLMTracker Tracker, int64 DeltaMemory, FName DefaultTag,
 	ELLMAllocType AllocType)
 {
-	if (bIsDisabled)
+	FEnableStateScopeLock EnableScopeLock;
+	if (!TryEnterEnabled(EnableScopeLock))
 	{
 		return;
 	}
+
 	BootstrapInitialise();
 	GetTracker(Tracker)->TrackMemoryOfActiveTag(DeltaMemory, DefaultTag, AllocType);
 }
@@ -2348,10 +2432,12 @@ void FLowLevelMemTracker::OnLowLevelChangeInMemoryUse(ELLMTracker Tracker, int64
 void FLowLevelMemTracker::OnLowLevelAllocMoved(ELLMTracker Tracker, const void* Dest, const void* Source,
 	ELLMAllocType AllocType)
 {
-	if (bIsDisabled)
+	FEnableStateScopeLock EnableScopeLock;
+	if (!TryEnterEnabled(EnableScopeLock))
 	{
 		return;
 	}
+
 	BootstrapInitialise();
 
 	//update the allocation map
@@ -2370,7 +2456,7 @@ const UE::LLMPrivate::FLLMTracker* FLowLevelMemTracker::GetTracker(ELLMTracker T
 
 bool FLowLevelMemTracker::Exec(const TCHAR* Cmd, FOutputDevice& Ar)
 {
-	if (bIsDisabled)
+	if (!IsEnabled())
 	{
 		return false;
 	}
@@ -2454,7 +2540,7 @@ bool FLowLevelMemTracker::Exec(const TCHAR* Cmd, FOutputDevice& Ar)
 
 bool FLowLevelMemTracker::IsTagSetActive(ELLMTagSet Set)
 {
-	if (bIsDisabled)
+	if (!IsEnabled())
 	{
 		return false;
 	}
@@ -2466,7 +2552,7 @@ bool FLowLevelMemTracker::IsTagSetActive(ELLMTagSet Set)
 bool FLowLevelMemTracker::ShouldReduceThreads()
 {
 #if LLM_ENABLED_REDUCE_THREADS
-	if (bIsDisabled)
+	if (!IsEnabled())
 	{
 		return false;
 	}
@@ -2507,7 +2593,7 @@ void FLowLevelMemTracker::RegisterPlatformTag(int32 Tag, const TCHAR* Name, FNam
 {
 	MemoryTrace_AnnounceCustomTag(Tag, ParentTag, Name);
 
-	if (bIsDisabled)
+	if (!IsEnabled())
 	{
 		return;
 	}
@@ -2523,7 +2609,7 @@ void FLowLevelMemTracker::RegisterProjectTag(int32 Tag, const TCHAR* Name, FName
 {
 	MemoryTrace_AnnounceCustomTag(Tag, ParentTag, Name);
 
-	if (bIsDisabled)
+	if (!IsEnabled())
 	{
 		return;
 	}
@@ -2535,7 +2621,7 @@ void FLowLevelMemTracker::RegisterProjectTag(int32 Tag, const TCHAR* Name, FName
 
 void GlobalRegisterTagDeclaration(FLLMTagDeclaration& TagDeclaration)
 {
-	if (FLowLevelMemTracker::bIsDisabled)
+	if (!FLowLevelMemTracker::IsEnabled())
 	{
 		return;
 	}
@@ -2985,7 +3071,7 @@ TArray<const UE::LLMPrivate::FTagData*> FLowLevelMemTracker::GetTrackedTags(ELLM
 {
 	using namespace UE::LLMPrivate;
 
-	if (bIsDisabled)
+	if (!IsEnabled())
 	{
 		return TArray<const FTagData*>();
 	}
@@ -3005,7 +3091,7 @@ TArray<const UE::LLMPrivate::FTagData*> FLowLevelMemTracker::GetTrackedTags(ELLM
 {
 	using namespace UE::LLMPrivate;
 
-	if (bIsDisabled)
+	if (!IsEnabled())
 	{
 		return TArray<const FTagData*>();
 	}
@@ -3020,7 +3106,7 @@ void FLowLevelMemTracker::GetTrackedTagsNamesWithAmount(TMap<FName, uint64>& Tag
 {
 	using namespace UE::LLMPrivate;
 
-	if (bIsDisabled)
+	if (!IsEnabled())
 	{
 		return;
 	}
@@ -3035,7 +3121,7 @@ void FLowLevelMemTracker::GetTrackedTagsNamesWithAmountFiltered(TMap<FName, uint
 {
 	using namespace UE::LLMPrivate;
 
-	if (bIsDisabled)
+	if (!IsEnabled())
 	{
 		return;
 	}
@@ -3050,7 +3136,7 @@ bool FLowLevelMemTracker::FindTagByName( const TCHAR* Name, uint64& OutTag, ELLM
 {
 	using namespace UE::LLMPrivate;
 
-	if (bIsDisabled)
+	if (!IsEnabled())
 	{
 		return false;
 	}
@@ -3092,7 +3178,7 @@ bool FLowLevelMemTracker::FindTagByName( const TCHAR* Name, uint64& OutTag, ELLM
 
 const TCHAR* FLowLevelMemTracker::FindTagName(uint64 Tag) const
 {
-	if (bIsDisabled)
+	if (!IsEnabled())
 	{
 		return nullptr;
 	}
@@ -3121,7 +3207,7 @@ FName FLowLevelMemTracker::FindTagDisplayName(uint64 Tag) const
 {
 	using namespace UE::LLMPrivate;
 
-	if (bIsDisabled)
+	if (!IsEnabled())
 	{
 		return NAME_None;
 	}
@@ -3144,7 +3230,7 @@ FName FLowLevelMemTracker::FindPtrDisplayName(void* Ptr) const
 {
 	using namespace UE::LLMPrivate;
 
-	if (bIsDisabled || !bFullyInitialised)
+	if (!IsEnabled() || !bFullyInitialised)
 	{
 		return NAME_None;
 	}
@@ -3197,7 +3283,7 @@ ELLMTag FLowLevelMemTracker::GetTagClosestEnumTag(const UE::LLMPrivate::FTagData
 // Deprecated in 5.3
 int64 FLowLevelMemTracker::GetTagAmountForTracker(ELLMTracker Tracker, ELLMTag Tag, bool bPeakAmount)
 {
-	if (bIsDisabled)
+	if (!IsEnabled())
 	{
 		return 0;
 	}
@@ -3213,7 +3299,7 @@ int64 FLowLevelMemTracker::GetTagAmountForTracker(ELLMTracker Tracker, const UE:
 {
 	using namespace UE::LLMPrivate;
 
-	if (bIsDisabled)
+	if (!IsEnabled())
 	{
 		return 0;
 	}
@@ -3233,7 +3319,7 @@ int64 FLowLevelMemTracker::GetTagAmountForTracker(ELLMTracker Tracker, const UE:
 int64 FLowLevelMemTracker::GetTagAmountForTracker(ELLMTracker Tracker, FName Tag, ELLMTagSet TagSet,
 	bool bPeakAmount)
 {
-	if (bIsDisabled)
+	if (!IsEnabled())
 	{
 		return 0;
 	}
@@ -3254,7 +3340,7 @@ int64 FLowLevelMemTracker::GetTagAmountForTracker(ELLMTracker Tracker, const UE:
 {
 	using namespace UE::LLMPrivate;
 
-	if (bIsDisabled)
+	if (!IsEnabled())
 	{
 		return 0;
 	}
@@ -3273,7 +3359,7 @@ int64 FLowLevelMemTracker::GetTagAmountForTracker(ELLMTracker Tracker, const UE:
 int64 FLowLevelMemTracker::GetTagAmountForTracker(ELLMTracker Tracker, FName Tag, ELLMTagSet TagSet,
 	UE::LLM::ESizeParams SizeParams)
 {
-	if (bIsDisabled)
+	if (!IsEnabled())
 	{
 		return 0;
 	}
@@ -3285,7 +3371,7 @@ int64 FLowLevelMemTracker::GetTagAmountForTracker(ELLMTracker Tracker, FName Tag
 
 int64 FLowLevelMemTracker::GetTagAmountForTracker(ELLMTracker Tracker, ELLMTag Tag, UE::LLM::ESizeParams SizeParams)
 {
-	if (bIsDisabled)
+	if (!IsEnabled())
 	{
 		return 0;
 	}
@@ -3299,7 +3385,7 @@ void FLowLevelMemTracker::SetTagAmountForTracker(ELLMTracker Tracker, ELLMTag Ta
 {
 	using namespace UE::LLMPrivate;
 
-	if (bIsDisabled)
+	if (!IsEnabled())
 	{
 		return;
 	}
@@ -3314,7 +3400,7 @@ void FLowLevelMemTracker::SetTagAmountForTracker(ELLMTracker Tracker, FName Tag,
 {
 	using namespace UE::LLMPrivate;
 
-	if (bIsDisabled)
+	if (!IsEnabled())
 	{
 		return;
 	}
@@ -3329,7 +3415,7 @@ int64 FLowLevelMemTracker::GetActiveTag(ELLMTracker Tracker)
 {
 	using namespace UE::LLMPrivate;
 
-	if (bIsDisabled)
+	if (!IsEnabled())
 	{
 		return static_cast<int64>(ELLMTag::Untagged);
 	}
@@ -3348,7 +3434,7 @@ int64 FLowLevelMemTracker::GetActiveTag(ELLMTracker Tracker)
 
 const UE::LLMPrivate::FTagData* FLowLevelMemTracker::GetActiveTagData(ELLMTracker Tracker, ELLMTagSet TagSet /*= ELLMTagSet::None*/)
 {
-	if (bIsDisabled)
+	if (!IsEnabled())
 	{
 		return nullptr;
 	}
@@ -3360,7 +3446,7 @@ const UE::LLMPrivate::FTagData* FLowLevelMemTracker::GetActiveTagData(ELLMTracke
 uint64 FLowLevelMemTracker::DumpTag( ELLMTracker Tracker, const char* FileName, int LineNumber )
 {
 	using namespace UE::LLMPrivate;
-	if (bIsDisabled)
+	if (!IsEnabled())
 	{
 		return static_cast<int64>(ELLMTag::Untagged);
 	}
@@ -3523,9 +3609,9 @@ void FLLMScope::Init(ELLMTag TagEnum, bool bInIsStatTag, ELLMTagSet InTagSet, EL
 	LLMCheck((TagEnum != ELLMTag::FMalloc) || (ELLMTracker::Platform == InTracker));
 
 	FLowLevelMemTracker& LLMRef = FLowLevelMemTracker::Get();
-	// We have to check bIsDisabled again after calling Get, because the constructor is called from Get, and will set 
-	// bIsDisabled=false if the platform doesn't support it.
-	if (FLowLevelMemTracker::bIsDisabled)
+	// We have to check IsEnabled() again after calling Get, because the constructor is called
+	// from Get, and will set EnabledState=Disabled if the platform doesn't support it.
+	if (!FLowLevelMemTracker::IsEnabled())
 	{
 		bEnabled = false;
 		return;
@@ -3547,9 +3633,9 @@ void FLLMScope::Init(ELLMTag TagEnum, bool bInIsStatTag, ELLMTagSet InTagSet, EL
 void FLLMScope::Init(FName TagName, bool bInIsStatTag, ELLMTagSet InTagSet, ELLMTracker InTracker, bool bOverride)
 {
 	FLowLevelMemTracker& LLMRef = FLowLevelMemTracker::Get();
-	// We have to check bIsDisabled again after calling Get, because the constructor is called from Get, and will set
-	// bIsDisabled=false if the platform doesn't support it.
-	if (FLowLevelMemTracker::bIsDisabled)
+	// We have to check IsEnabled() again after calling Get, because the constructor is called
+	// from Get, and will set EnabledStae=Disabled if the platform doesn't support it.
+	if (!FLowLevelMemTracker::IsEnabled())
 	{
 		bEnabled = false;
 		return;
@@ -3578,9 +3664,9 @@ void FLLMScope::Init(const UE::LLMPrivate::FTagData* TagData, bool bInIsStatTag,
 	ELLMTracker InTracker, bool bOverride)
 {
 	FLowLevelMemTracker& LLMRef = FLowLevelMemTracker::Get();
-	// We have to check bIsDisabled again after calling Get, because the constructor is called from Get, and will set
-	// bIsDisabled=true if the platform doesn't support it.
-	if (FLowLevelMemTracker::bIsDisabled)
+	// We have to check IsEnabled() again after calling Get, because the constructor is called
+	// from Get, and will set EnabledState=Disabled if the platform doesn't support it.
+	if (!FLowLevelMemTracker::IsEnabled())
 	{
 		bEnabled = false;
 		return;
@@ -3608,9 +3694,9 @@ void FLLMScope::Destruct()
 void FLLMScopeDynamic::Init(ELLMTracker InTracker, ELLMTagSet InTagSet)
 {
 	FLowLevelMemTracker& LLMRef = FLowLevelMemTracker::Get();
-	// We have to check bIsDisabled again after calling Get, because the constructor is called from Get, and will set
-	// bIsDisabled=true if the platform doesn't support it.
-	if (FLowLevelMemTracker::bIsDisabled)
+	// We have to check IsEnabled() again after calling Get, because the constructor is called
+	// from Get, and will set EnabledState=Disabled if the platform doesn't support it.
+	if (!FLowLevelMemTracker::IsEnabled())
 	{
 		bEnabled = false;
 		return;
@@ -3693,7 +3779,7 @@ void FLLMScopeDynamic::Destruct()
 FLLMPauseScope::FLLMPauseScope(FName TagName, bool bIsStatTag, uint64 Amount, ELLMTracker TrackerToPause,
 	ELLMAllocType InAllocType)
 {
-	if (FLowLevelMemTracker::bIsDisabled)
+	if (!FLowLevelMemTracker::IsEnabled())
 	{
 		bEnabled = false;
 		return;
@@ -3704,7 +3790,7 @@ FLLMPauseScope::FLLMPauseScope(FName TagName, bool bIsStatTag, uint64 Amount, EL
 FLLMPauseScope::FLLMPauseScope(ELLMTag TagEnum, bool bIsStatTag, uint64 Amount, ELLMTracker TrackerToPause,
 	ELLMAllocType InAllocType)
 {
-	if (FLowLevelMemTracker::bIsDisabled)
+	if (!FLowLevelMemTracker::IsEnabled())
 	{
 		bEnabled = false;
 		return;
@@ -3717,9 +3803,9 @@ void FLLMPauseScope::Init(FName TagName, ELLMTag EnumTag, bool bIsEnumTag, bool
 	ELLMTracker TrackerToPause, ELLMAllocType InAllocType)
 {
 	FLowLevelMemTracker& LLMRef = FLowLevelMemTracker::Get();
-	// We have to check bIsDisabled again after calling Get, because the constructor is called from Get, and will set 
-	// bIsDisabled=false if the platform doesn't support it.
-	if (FLowLevelMemTracker::bIsDisabled)
+	// We have to check IsEnabled() again after calling Get, because the constructor is called
+	// from Get, and will set EnabledState=Disabled if the platform doesn't support it.
+	if (!FLowLevelMemTracker::IsEnabled())
 	{
 		bEnabled = false;
 		return;
@@ -3784,21 +3870,21 @@ FLLMPauseScope::~FLLMPauseScope()
 FLLMScopeFromPtr::FLLMScopeFromPtr(void* Ptr, ELLMTracker InTracker)
 {
 	using namespace UE::LLMPrivate;
-	if (FLowLevelMemTracker::bIsDisabled)
+	if (!FLowLevelMemTracker::IsEnabled())
 	{
 		DisableAll();
 		return;
 	}
-	if(Ptr == nullptr)
+	if (Ptr == nullptr)
 	{
 		DisableAll();
 		return;
 	}
 
 	FLowLevelMemTracker& LLMRef = FLowLevelMemTracker::Get();
-	// We have to check bIsDisabled again after calling Get, because the constructor is called from Get, and will set
-	// bIsDisabled=false if the platform doesn't support it.
-	if (FLowLevelMemTracker::bIsDisabled)
+	// We have to check IsEnabled() again after calling Get, because the constructor is called
+	// from Get, and will set EnabledState=Disabled if the platform doesn't support it.
+	if (!FLowLevelMemTracker::IsEnabled())
 	{
 		DisableAll();
 		return;
diff --git a/Engine/Source/Runtime/Core/Public/Async/InheritedContext.h b/Engine/Source/Runtime/Core/Public/Async/InheritedContext.h
--- a/Engine/Source/Runtime/Core/Public/Async/InheritedContext.h
+++ b/Engine/Source/Runtime/Core/Public/Async/InheritedContext.h
@@ -19,7 +19,9 @@ namespace UE
 		{
 			for (int32 TagSetIndex = 0; TagSetIndex < static_cast<int32>(ELLMTagSet::Max); ++TagSetIndex)
 			{
-				LLMTags[TagSetIndex] = FLowLevelMemTracker::bIsDisabled ? nullptr : FLowLevelMemTracker::Get().GetActiveTagData(ELLMTracker::Default, static_cast<ELLMTagSet>(TagSetIndex));
+				LLMTags[TagSetIndex] = FLowLevelMemTracker::IsEnabled() ?
+					FLowLevelMemTracker::Get().GetActiveTagData(ELLMTracker::Default, static_cast<ELLMTagSet>(TagSetIndex))
+					: nullptr;
 			}
 		}
 	};
diff --git a/Engine/Source/Runtime/Core/Public/HAL/LowLevelMemTracker.h b/Engine/Source/Runtime/Core/Public/HAL/LowLevelMemTracker.h
--- a/Engine/Source/Runtime/Core/Public/HAL/LowLevelMemTracker.h
+++ b/Engine/Source/Runtime/Core/Public/HAL/LowLevelMemTracker.h
@@ -323,7 +323,7 @@ extern FName LLMGetTagStat(ELLMTag Tag);
  * LLM utility macros
  */
 #define LLM(x) x
-#define LLM_IF_ENABLED(x) if (!FLowLevelMemTracker::bIsDisabled) { x; }
+#define LLM_IF_ENABLED(x) if (FLowLevelMemTracker::IsEnabled()) { x; }
 #define SCOPE_NAME PREPROCESSOR_JOIN(LLMScope,__LINE__)
 
 ///////////////////////////////////////////////////////////////////////////////////////
@@ -592,7 +592,11 @@ public:
 
 	static CORE_API FLowLevelMemTracker& Construct();
 
-	static CORE_API bool IsEnabled();
+	/**
+	 * Return whether LLM is enabled - activated by the commandline and is tracking allocations.
+	 * Also returns true early during startup before commandline has been processed.
+	 */
+	inline static bool IsEnabled();
 
 	/**
 	 * We always start up running, but if the commandline disables us, we will do it later after main
@@ -779,45 +783,64 @@ public:
 	CORE_API void FinishInitialise();
 
 private:
-	CORE_API FLowLevelMemTracker();
+	enum class EEnabled : uint8
+	{
+		NotYetKnown = 0,
+		Disabled,
+		Enabled,
+	};
+	struct FEnableStateScopeLock;
+
+private:
+	FLowLevelMemTracker();
 
-	CORE_API ~FLowLevelMemTracker();
+	~FLowLevelMemTracker();
 
 	bool IsBootstrapping() const { return bIsBootstrapping; }
 
 	/** Free all memory. This will put the tracker into a permanently disabled state. */
-	CORE_API void Clear();
-	CORE_API void InitialiseProgramSize();
+	void Clear();
+	void InitialiseProgramSize();
 
-	CORE_API class UE::LLMPrivate::FLLMTracker* GetTracker(ELLMTracker Tracker);
-	CORE_API const class UE::LLMPrivate::FLLMTracker* GetTracker(ELLMTracker Tracker) const;
+	class UE::LLMPrivate::FLLMTracker* GetTracker(ELLMTracker Tracker);
+	const class UE::LLMPrivate::FLLMTracker* GetTracker(ELLMTracker Tracker) const;
 
-	CORE_API void TickInternal();
-	CORE_API void UpdateTags();
-	CORE_API void SortTags(UE::LLMPrivate::FTagDataArray*& OutOldTagDatas);
-	CORE_API void PublishDataPerFrame(const TCHAR* LogName);
+	void TickInternal();
+	void UpdateTags();
+	void SortTags(UE::LLMPrivate::FTagDataArray*& OutOldTagDatas);
+	void PublishDataPerFrame(const TCHAR* LogName);
 
-	CORE_API void RegisterCustomTagInternal(int32 Tag, ELLMTagSet TagSet, const TCHAR* Name, FName StatName, FName SummaryStatName, int32 ParentTag = -1);
+	void RegisterCustomTagInternal(int32 Tag, ELLMTagSet TagSet, const TCHAR* Name, FName StatName, FName SummaryStatName, int32 ParentTag = -1);
 	/**
 	* Called during C++ global static initialization when GMalloc is available (and hence FNames are unavailable)
 	* Creates the subset of tags necessary to record allocations during GMalloc and FName construction
 	*/
-	CORE_API void BootstrapTagDatas();
-	CORE_API void InitialiseTagDatas_SetLLMTagNames();
-	CORE_API void InitialiseTagDatas_FinishRegister();
-	CORE_API void InitialiseTagDatas();
-	CORE_API void ClearTagDatas();
-	CORE_API void RegisterTagDeclaration(FLLMTagDeclaration& TagDeclaration);
-	CORE_API UE::LLMPrivate::FTagData& RegisterTagData(FName Name, FName DisplayName, FName ParentName, FName StatName, FName SummaryStatName, bool bHasEnumTag, ELLMTag EnumTag, bool bIsStatTag, UE::LLMPrivate::ETagReferenceSource ReferenceSource, ELLMTagSet TagSet = ELLMTagSet::None);
+	void BootstrapTagDatas();
+	void InitialiseTagDatas_SetLLMTagNames();
+	void InitialiseTagDatas_FinishRegister();
+	void InitialiseTagDatas();
+	void ClearTagDatas();
+	void RegisterTagDeclaration(FLLMTagDeclaration& TagDeclaration);
+	UE::LLMPrivate::FTagData& RegisterTagData(FName Name, FName DisplayName, FName ParentName, FName StatName, FName SummaryStatName, bool bHasEnumTag, ELLMTag EnumTag, bool bIsStatTag, UE::LLMPrivate::ETagReferenceSource ReferenceSource, ELLMTagSet TagSet = ELLMTagSet::None);
 	/** Construct if not yet done the data on the given FTagData that relies on the presence of other TagDatas */
-	CORE_API void FinishConstruct(UE::LLMPrivate::FTagData* TagData, UE::LLMPrivate::ETagReferenceSource ReferenceSource);
-	CORE_API void ReportDuplicateTagName(UE::LLMPrivate::FTagData* TagData, UE::LLMPrivate::ETagReferenceSource ReferenceSource);
+	void FinishConstruct(UE::LLMPrivate::FTagData* TagData, UE::LLMPrivate::ETagReferenceSource ReferenceSource);
+	void ReportDuplicateTagName(UE::LLMPrivate::FTagData* TagData, UE::LLMPrivate::ETagReferenceSource ReferenceSource);
 
-	CORE_API const UE::LLMPrivate::FTagData* FindOrAddTagData(ELLMTag EnumTag, UE::LLMPrivate::ETagReferenceSource ReferenceSource = UE::LLMPrivate::ETagReferenceSource::FunctionAPI);
-	CORE_API const UE::LLMPrivate::FTagData* FindOrAddTagData(FName Name, ELLMTagSet TagSet, bool bIsStatData=false, UE::LLMPrivate::ETagReferenceSource ReferenceSource = UE::LLMPrivate::ETagReferenceSource::FunctionAPI);
-	CORE_API const UE::LLMPrivate::FTagData* FindOrAddTagData(FName Name, ELLMTagSet TagSet, FName StatName, UE::LLMPrivate::ETagReferenceSource ReferenceSource = UE::LLMPrivate::ETagReferenceSource::FunctionAPI);
-	CORE_API const UE::LLMPrivate::FTagData* FindTagData(ELLMTag EnumTag, UE::LLMPrivate::ETagReferenceSource ReferenceSource = UE::LLMPrivate::ETagReferenceSource::FunctionAPI);
-	CORE_API const UE::LLMPrivate::FTagData* FindTagData(FName Name, ELLMTagSet TagSet, UE::LLMPrivate::ETagReferenceSource ReferenceSource = UE::LLMPrivate::ETagReferenceSource::FunctionAPI);
+	const UE::LLMPrivate::FTagData* FindOrAddTagData(ELLMTag EnumTag, UE::LLMPrivate::ETagReferenceSource ReferenceSource = UE::LLMPrivate::ETagReferenceSource::FunctionAPI);
+	const UE::LLMPrivate::FTagData* FindOrAddTagData(FName Name, ELLMTagSet TagSet, bool bIsStatData=false, UE::LLMPrivate::ETagReferenceSource ReferenceSource = UE::LLMPrivate::ETagReferenceSource::FunctionAPI);
+	const UE::LLMPrivate::FTagData* FindOrAddTagData(FName Name, ELLMTagSet TagSet, FName StatName, UE::LLMPrivate::ETagReferenceSource ReferenceSource = UE::LLMPrivate::ETagReferenceSource::FunctionAPI);
+	const UE::LLMPrivate::FTagData* FindTagData(ELLMTag EnumTag, UE::LLMPrivate::ETagReferenceSource ReferenceSource = UE::LLMPrivate::ETagReferenceSource::FunctionAPI);
+	const UE::LLMPrivate::FTagData* FindTagData(FName Name, ELLMTagSet TagSet, UE::LLMPrivate::ETagReferenceSource ReferenceSource = UE::LLMPrivate::ETagReferenceSource::FunctionAPI);
+
+	/**
+	 * A function that returns IsEnabled(), but also enters a lock if current state is NotYetKnown and therefore might
+	 * change if the main thread calls ProcessCommandLine and commandline specifies that LLM should disable itself.
+	 * The lock is only needed on platforms with PLATFORM_HAS_MULTITHREADED_PREMAIN; other platforms are not
+	 * multithreaded until after ProcessCommandLine is called. This function is only necessary on handlers for new
+	 * and delete (OnLowLevelAlloc, etc), since those are the only functions that can be called from other threads
+	 * before ProcessCommandLine.
+	 */
+	inline static bool TryEnterEnabled(FEnableStateScopeLock& ScopeLock);
 
 	friend class FLLMPauseScope;
 	friend class FLLMScope;
@@ -830,6 +853,7 @@ private:
 	friend class UE::LLMPrivate::FLLMCsvProfilerWriter;
 	friend void GlobalRegisterTagDeclaration(FLLMTagDeclaration& TagDeclaration);
 
+private:
 	UE::LLMPrivate::FLLMAllocator Allocator;
 	/** All TagDatas that have been constructed, in an array sorted by TagData->GetIndex() */
 	UE::LLMPrivate::FTagDataArray* TagDatas;
@@ -863,7 +887,9 @@ private:
 	bool bCapturedSizeSnapshot;
 
 	static CORE_API FLowLevelMemTracker* TrackerInstance;
-public: // really internal but needs to be visible for LLM_IF_ENABLED macro
+	static CORE_API EEnabled EnabledState;
+public:
+	UE_DEPRECATED(5.5, "Use IsEnabled instead.")
 	static CORE_API bool bIsDisabled;
 };
 
@@ -873,14 +899,14 @@ class FLLMScope
 public:
 	FLLMScope(FName TagName, bool bIsStatTag, ELLMTagSet InTagSet, ELLMTracker InTracker, bool bOverride = true)
 	{
-		if (!FLowLevelMemTracker::bIsDisabled)
+		if (FLowLevelMemTracker::IsEnabled())
 		{
 			Init(TagName, bIsStatTag, InTagSet, InTracker, bOverride);
 		}
 	}
 	FLLMScope(ELLMTag TagEnum, bool bIsStatTag, ELLMTagSet InTagSet, ELLMTracker InTracker, bool bOverride = true)
 	{
-		if (!FLowLevelMemTracker::bIsDisabled)
+		if (FLowLevelMemTracker::IsEnabled())
 		{
 			Init(TagEnum, bIsStatTag, InTagSet, InTracker, bOverride);
 		}
@@ -889,7 +915,7 @@ public:
 
 	FLLMScope(const UE::LLMPrivate::FTagData* TagData, bool bIsStatTag, ELLMTagSet Set, ELLMTracker Tracker, bool bOverride = true)
 	{
-		if (!FLowLevelMemTracker::bIsDisabled)
+		if (FLowLevelMemTracker::IsEnabled())
 		{
 			Init(TagData, bIsStatTag, Set, Tracker, bOverride);
 		}
@@ -953,7 +979,7 @@ class FLLMScopeDynamic
 public:
 	FLLMScopeDynamic(ELLMTracker InTracker, ELLMTagSet InTagSet)
 	{
-		if (!FLowLevelMemTracker::bIsDisabled)
+		if (FLowLevelMemTracker::IsEnabled())
 		{
 			Init(InTracker, InTagSet);
 		}
@@ -1055,6 +1081,11 @@ struct FLLMTagSetAllocationFilter
 	ELLMTagSet TagSet;
 };
 
+inline bool FLowLevelMemTracker::IsEnabled()
+{
+	return EnabledState != EEnabled::Disabled;
+}
+
 #else
 
 #define LLM(...)
diff --git a/Engine/Source/Runtime/Core/Public/HAL/Platform.h b/Engine/Source/Runtime/Core/Public/HAL/Platform.h
--- a/Engine/Source/Runtime/Core/Public/HAL/Platform.h
+++ b/Engine/Source/Runtime/Core/Public/HAL/Platform.h
@@ -285,6 +285,13 @@
 #ifndef PLATFORM_USE_PTHREADS
 	#define PLATFORM_USE_PTHREADS				1
 #endif
+#ifndef PLATFORM_HAS_MULTITHREADED_PREMAIN
+	/**
+	 * If true, the platform uses multiple threads in c++ static initialization before main is called, and systems
+	 * accessible from multiple threads during premain static initialization must handle multithreaded synchronization.
+	 */
+	#define PLATFORM_HAS_MULTITHREADED_PREMAIN 0
+#endif
 #ifndef PLATFORM_MAX_FILEPATH_LENGTH_DEPRECATED
 	#define PLATFORM_MAX_FILEPATH_LENGTH_DEPRECATED		128			// Deprecated - prefer FPlatformMisc::GetMaxPathLength() instead.
 #endif
