LLM: Set PLATFORM_HAS_MULTITHREADED_PREMAIN for IOS platforms; the load of executables on IOS does include OS threads that can call new and delete and that run before Main.
#rnx

diff --git a/Engine/Source/Runtime/Core/Public/IOS/IOSPlatform.h b/Engine/Source/Runtime/Core/Public/IOS/IOSPlatform.h
--- a/Engine/Source/Runtime/Core/Public/IOS/IOSPlatform.h
+++ b/Engine/Source/Runtime/Core/Public/IOS/IOSPlatform.h
@@ -45,6 +45,7 @@ typedef FIOSPlatformTypes FPlatformTypes;
 #define PLATFORM_HAS_BSD_TIME							1
 #define PLATFORM_HAS_BSD_IPV6_SOCKETS					1
 #define PLATFORM_HAS_BSD_SOCKET_FEATURE_MSG_DONTWAIT	1
+#define PLATFORM_HAS_MULTITHREADED_PREMAIN				1
 #define PLATFORM_MAX_FILEPATH_LENGTH_DEPRECATED			IOS_MAX_PATH
 #define PLATFORM_SUPPORTS_TEXTURE_STREAMING				1
 #define PLATFORM_BUILTIN_VERTEX_HALF_FLOAT				0
diff --git a/Engine/Source/Runtime/Core/Public/Mac/MacPlatform.h b/Engine/Source/Runtime/Core/Public/Mac/MacPlatform.h
--- a/Engine/Source/Runtime/Core/Public/Mac/MacPlatform.h
+++ b/Engine/Source/Runtime/Core/Public/Mac/MacPlatform.h
@@ -70,6 +70,7 @@ typedef FMacPlatformTypes FPlatformTypes;
 #define PLATFORM_HAS_BSD_SOCKET_FEATURE_IOCTL			1
 #define PLATFORM_HAS_BSD_SOCKET_FEATURE_POLL			1
 //#define PLATFORM_USE_PTHREADS							1
+#define PLATFORM_HAS_MULTITHREADED_PREMAIN				1
 #define PLATFORM_MAX_FILEPATH_LENGTH_DEPRECATED			MAC_MAX_PATH
 #define PLATFORM_SUPPORTS_TBB							1
 #define PLATFORM_SUPPORTS_MIMALLOC						PLATFORM_64BITS
