Refactored FBitTree::GetMemoryRequirements to be a constexpr to allow computing memory requirement in compile time

diff --git a/Engine/Source/Runtime/Core/Private/HAL/MallocBinnedCommon.cpp b/Engine/Source/Runtime/Core/Private/HAL/MallocBinnedCommon.cpp
--- a/Engine/Source/Runtime/Core/Private/HAL/MallocBinnedCommon.cpp
+++ b/Engine/Source/Runtime/Core/Private/HAL/MallocBinnedCommon.cpp
@@ -157,27 +157,6 @@ uint8 FSizeTableEntry::FillSizeTable(uint64 PlatformPageSize, FSizeTableEntry* S
 	return (uint8)Index;
 }
 
-uint32 FBitTree::GetMemoryRequirements(uint32 DesiredCapacity)
-{
-	uint32 AllocationSize = 8;
-	uint32 RowsUint64s = 1;
-	uint32 Capacity = 64;
-	uint32 OffsetOfLastRow = 0;
-
-	while (Capacity < DesiredCapacity)
-	{
-		Capacity *= 64;
-		RowsUint64s *= 64;
-		OffsetOfLastRow = AllocationSize / 8;
-		AllocationSize += 8 * RowsUint64s;
-	}
-
-	uint32 LastRowTotal = (AllocationSize - OffsetOfLastRow * 8) * 8;
-	uint32 ExtraBits = LastRowTotal - DesiredCapacity;
-	AllocationSize -= (ExtraBits / 64) * 8;
-	return AllocationSize;
-}
-
 void FBitTree::FBitTreeInit(uint32 InDesiredCapacity, void * Memory, uint32 MemorySize, bool InitialValue)
 {
 	Bits = (uint64*)Memory;
diff --git a/Engine/Source/Runtime/Core/Public/HAL/MallocBinnedCommon.h b/Engine/Source/Runtime/Core/Public/HAL/MallocBinnedCommon.h
--- a/Engine/Source/Runtime/Core/Public/HAL/MallocBinnedCommon.h
+++ b/Engine/Source/Runtime/Core/Public/HAL/MallocBinnedCommon.h
@@ -41,7 +41,27 @@ public:
 	{
 	}
 
-	static uint32 GetMemoryRequirements(uint32 DesiredCapacity);
+	static constexpr uint32 GetMemoryRequirements(uint32 NumPages)
+	{
+		uint32 AllocationSize = 8;
+		uint32 RowsUint64s = 1;
+		uint32 Capacity = 64;
+		uint32 OffsetOfLastRow = 0;
+
+		while (Capacity < NumPages)
+		{
+			Capacity *= 64;
+			RowsUint64s *= 64;
+			OffsetOfLastRow = AllocationSize / 8;
+			AllocationSize += 8 * RowsUint64s;
+		}
+
+		uint32 LastRowTotal = (AllocationSize - OffsetOfLastRow * 8) * 8;
+		uint32 ExtraBits = LastRowTotal - NumPages;
+		AllocationSize -= (ExtraBits / 64) * 8;
+		return AllocationSize;
+	}
+
 	void FBitTreeInit(uint32 InDesiredCapacity, void * Memory, uint32 MemorySize, bool InitialValue);
 	uint32 AllocBit();
 	bool IsAllocated(uint32 Index) const;
