Make Apple file descriptors unconditionally non-inheritable as they are on Unix.  This avoids issues with file descriptors accumulating as we respawn the editor when creating or switching projects.

diff --git a/Engine/Source/Runtime/Core/Private/Apple/ApplePlatformFile.cpp b/Engine/Source/Runtime/Core/Private/Apple/ApplePlatformFile.cpp
--- a/Engine/Source/Runtime/Core/Private/Apple/ApplePlatformFile.cpp
+++ b/Engine/Source/Runtime/Core/Private/Apple/ApplePlatformFile.cpp
@@ -283,7 +283,7 @@ private:
 				ReserveSlot();
 
 				TRACE_PLATFORMFILE_BEGIN_OPEN(*Filename);
-				FileHandle = open(TCHAR_TO_UTF8(*Filename), O_RDONLY | O_SHLOCK);
+				FileHandle = open(TCHAR_TO_UTF8(*Filename), O_RDONLY | O_SHLOCK | O_CLOEXEC);
 				if( FileHandle != -1 )
 				{
 					TRACE_PLATFORMFILE_END_OPEN(FileHandle);
@@ -574,7 +574,7 @@ ESymlinkResult FApplePlatformFile::IsSymlink(const TCHAR* Filename)
 IFileHandle* FApplePlatformFile::OpenRead(const TCHAR* Filename, bool bAllowWrite)
 {
 	TRACE_PLATFORMFILE_BEGIN_OPEN(Filename);
-	int32 Handle = open(TCHAR_TO_UTF8(*NormalizeFilename(Filename)), O_RDONLY);
+	int32 Handle = open(TCHAR_TO_UTF8(*NormalizeFilename(Filename)), O_RDONLY | O_CLOEXEC);
 	if (Handle != -1)
 	{
 		TRACE_PLATFORMFILE_END_OPEN(Handle);
@@ -615,7 +615,7 @@ IFileHandle* FApplePlatformFile::OpenRead(const TCHAR* Filename, bool bAllowWrit
 
 IFileHandle* FApplePlatformFile::OpenWrite(const TCHAR* Filename, bool bAppend, bool bAllowRead)
 {
-	int Flags = O_CREAT;
+	int Flags = O_CREAT | O_CLOEXEC;
 	
 	if (bAllowRead)
 	{
