Updated FCompressedBufferReader to support a RawHash of zero

This situation arises with the Zen project store. This is a temporary fix to avoid unintentional reuse of decompressed data between from different compressed buffers.

#rnx

diff --git a/Engine/Source/Runtime/Core/Private/Compression/CompressedBuffer.cpp b/Engine/Source/Runtime/Core/Private/Compression/CompressedBuffer.cpp
--- a/Engine/Source/Runtime/Core/Private/Compression/CompressedBuffer.cpp
+++ b/Engine/Source/Runtime/Core/Private/Compression/CompressedBuffer.cpp
@@ -893,7 +893,7 @@ static bool TryReadHeader(FDecoderContext& Context, const FCompositeBuffer& Buff
 		{
 			Context.HeaderOffset = 0;
 			Context.HeaderSize = HeaderSize;
-			if (Context.HeaderCrc32 != Header.Crc32)
+			if (Context.HeaderCrc32 != Header.Crc32 || Header.RawHash.IsZero())
 			{
 				Context.HeaderCrc32 = Header.Crc32;
 				Context.RawBlockIndex = MAX_uint32;
