FBoxSphereBounds::Builder - Make sure to retain a tight SphereRadius rather than performing a BoxSphereBounds -> Box -> BoxSphereBounds conversion which is causing the SphereRadius to be computed from the box corners

diff --git a/Engine/Source/Runtime/Core/Private/UObject/DevObjectVersion.cpp b/Engine/Source/Runtime/Core/Private/UObject/DevObjectVersion.cpp
--- a/Engine/Source/Runtime/Core/Private/UObject/DevObjectVersion.cpp
+++ b/Engine/Source/Runtime/Core/Private/UObject/DevObjectVersion.cpp
@@ -128,7 +128,7 @@ FDevSystemGuids::FDevSystemGuids()
 	, Niagara_LatestScriptCompileVersion(0x6D32B8EE, 0x909FCA7E, 0xA5CE4F17, 0x066A5F25)
 	, POSESEARCHDB_DERIVEDDATA_VER(0x389117E4, 0x807A4CC0, 0x9F37C2E6, 0xD808A78D)
 	, SkeletalMeshDerivedDataVersion(0x9B5F4544, 0x76D7481C, 0x9AD3F614, 0xA6C07904)
-	, STATICMESH_DERIVEDDATA_VER(0x5B7A05A7, 0x7371440C, 0xA314621A, 0x9B1E08BC)
+	, STATICMESH_DERIVEDDATA_VER(0x2C6C400C, 0x2EDF47B3, 0x9BD36689, 0x7B77B208)
 	, MaterialTranslationDDCVersion(0x59E1F296, 0x27064757, 0x8CA7EB9B, 0xFC366BD2)
 {
 }
diff --git a/Engine/Source/Runtime/Core/Public/Math/BoxSphereBounds.h b/Engine/Source/Runtime/Core/Public/Math/BoxSphereBounds.h
--- a/Engine/Source/Runtime/Core/Public/Math/BoxSphereBounds.h
+++ b/Engine/Source/Runtime/Core/Public/Math/BoxSphereBounds.h
@@ -313,48 +313,51 @@ public:
 	 */
 	struct Builder
 	{
-		FORCEINLINE Builder()
-			: Box(ForceInitToZero)
+		FORCEINLINE Builder& operator+=(const TBoxSphereBounds<T, TExtent>& Other)
 		{
+			return Append(Other);
 		}
 
-		FORCEINLINE Builder& operator+=(const TBoxSphereBounds<T, TExtent>& OtherBounds)
+		FORCEINLINE Builder& operator+=(const TBox<T>& Box)
 		{
-			Box += OtherBounds.GetBox();
-			return *this;
-		}
-
-		FORCEINLINE Builder& operator+=(const TBox<T>& OtherBox)
-		{
-			Box += OtherBox;
-			return *this;
+			return Append(Box);
 		}
 
 		FORCEINLINE Builder& operator+=(const TSphere<T>& Sphere)
 		{
-			Box += Sphere.Center + TVector<T>(Sphere.W);
-			Box += Sphere.Center - TVector<T>(Sphere.W);
-			return *this;
+			return Append(Sphere);
 		}
 
 		FORCEINLINE Builder& operator+=(const TVector<T>& Point)
 		{
-			Box += Point;
-			return *this;
+			return Append(TBoxSphereBounds<T, TExtent>(Point, TVector<T>::ZeroVector, T(0)));
 		}
 
 		FORCEINLINE bool IsValid() const
 		{
-			return !!Box.IsValid;
+			return BoxSphereBounds.IsSet();
 		}
 
 		FORCEINLINE operator TBoxSphereBounds<T, TExtent>() const
 		{
-			return TBoxSphereBounds<T, TExtent>(Box);
+			return BoxSphereBounds.Get(TBoxSphereBounds<T, TExtent>());
 		}
 
 	private:
-		TBox<T> Box;
+		FORCEINLINE Builder& Append(const TBoxSphereBounds<T, TExtent>& Other)
+		{
+			if (IsValid())
+			{
+				*BoxSphereBounds = *BoxSphereBounds + Other;
+			}
+			else
+			{
+				BoxSphereBounds.Emplace(Other);
+			}
+			return *this;
+		}
+
+		TOptional<TBoxSphereBounds<T, TExtent>> BoxSphereBounds;
 	};
 };
 
