Enable AutoRTFM by default.

VkEdit and VkPlay will use AutoRTFM transactions in Verse by default from this change. Locally you can disable it for now with the -dpcvars=AutoRTFMRuntimeEnabled=0, but please contact #verse-autortfm-dev-ext if you do so we can look into any potential issues.

diff --git a/Engine/Source/Runtime/Core/Private/AutoRTFM/API.cpp b/Engine/Source/Runtime/Core/Private/AutoRTFM/API.cpp
--- a/Engine/Source/Runtime/Core/Private/AutoRTFM/API.cpp
+++ b/Engine/Source/Runtime/Core/Private/AutoRTFM/API.cpp
@@ -8,7 +8,7 @@
 namespace
 {
 	// Move this to a local only and use functions to access this
-	int GAutoRTFMRuntimeEnabled = AutoRTFM::EAutoRTFMEnabledState::AutoRTFM_Disabled;
+	int GAutoRTFMRuntimeEnabled = AutoRTFM::EAutoRTFMEnabledState::AutoRTFM_Enabled;
 
 	void UpdateAutoRTFMRuntimeCrashData()
 	{
