- In FPaths::IsStaged(), changed the logic to assume the build is staged when FPlatformProperties::RequiresCookedData() is true - this way we can avoid re-entrancy on platforms that attempt to get the ProjectDir() on the first file operation

diff --git a/Engine/Source/Runtime/Core/Private/Misc/Paths.cpp b/Engine/Source/Runtime/Core/Private/Misc/Paths.cpp
--- a/Engine/Source/Runtime/Core/Private/Misc/Paths.cpp
+++ b/Engine/Source/Runtime/Core/Private/Misc/Paths.cpp
@@ -146,7 +146,9 @@ namespace UE4Paths_Private
 
 bool FPaths::IsStaged()
 {
-	static bool bIsStaged = FileExists(Combine(EngineConfigDir(), TEXT("StagedBuild.ini")));
+	// if the platform requires cooked data, then we can assume staged, but if we don't require cooked data, then
+	// check if it went through the staging process
+	static bool bIsStaged = FPlatformProperties::RequiresCookedData() || FileExists(Combine(EngineConfigDir(), TEXT("StagedBuild.ini")));
 	return bIsStaged;
 }
 
