[Scheduler]
  - Fix race where multiple threads trying to steal from the same local queue could miss a task

diff --git a/Engine/Source/Runtime/Core/Public/Async/Fundamental/LocalQueue.h b/Engine/Source/Runtime/Core/Public/Async/Fundamental/LocalQueue.h
--- a/Engine/Source/Runtime/Core/Public/Async/Fundamental/LocalQueue.h
+++ b/Engine/Source/Runtime/Core/Public/Async/Fundamental/LocalQueue.h
@@ -37,7 +37,7 @@ protected:
 			ItemSlots[Idx].Value.store(Item, std::memory_order_release);
 			Head++;
 			checkSlow(Head % NumItems == Idx);
-			return true;		
+			return true;
 		}
 		return false;
 	}
@@ -69,6 +69,14 @@ protected:
 
 			if (Slot == uintptr_t(ESlotState::Free))
 			{
+				// Once we find a free slot, we need to verify if it's been freed by another steal
+				// so check back the Tail value to make sure it wasn't incremented since we first read the value.
+				// If we don't do this, some threads might not see that other threads
+				// have already stolen the slot, and will wrongly return that no more tasks are available to steal.
+				if (IdxVer != Tail.load(std::memory_order_acquire))
+				{
+					continue; // Loop again since tail has changed
+				}
 				return false;
 			}
 			else if (Slot != uintptr_t(ESlotState::Taken) && ItemSlots[Idx].Value.compare_exchange_weak(Slot, uintptr_t(ESlotState::Taken), std::memory_order_acq_rel))
@@ -725,4 +733,4 @@ private:
 	PRAGMA_ENABLE_DEPRECATION_WARNINGS
 };
 
-}
\ No newline at end of file
+}
