[Scheduler]
 - Fix race where detection for waking up a worker could see a stale value and wrongly think no more workers are available

diff --git a/Engine/Source/Runtime/Core/Private/Async/Fundamental/WaitingQueue.cpp b/Engine/Source/Runtime/Core/Private/Async/Fundamental/WaitingQueue.cpp
--- a/Engine/Source/Runtime/Core/Private/Async/Fundamental/WaitingQueue.cpp
+++ b/Engine/Source/Runtime/Core/Private/Async/Fundamental/WaitingQueue.cpp
@@ -221,14 +221,16 @@ int32 FWaitingQueue::NotifyInternal(int32 Count)
 			const uint64_t Waiters = (LocalState & WaiterMask) >> WaiterShift;
 			const uint64_t Signals = (LocalState & SignalMask) >> SignalShift;
 			const bool bNotifyAll = Count >= NodesArray.Num();
-			// Easy case: no waiters.
+
+			uint64_t NewState;
 			if ((LocalState & StackMask) == StackMask && Waiters == Signals)
 			{
-				WAITINGQUEUE_EVENT_SCOPE(NoMoreWaiter1);
-				return Notifications;
+				// No more waiters, go through the CAS to provide proper ordering
+				// with other threads entering PrepareWait.
+				WAITINGQUEUE_EVENT_SCOPE(TryNoMoreWaiter);
+				NewState = LocalState;
 			}
-			uint64_t NewState;
-			if (bNotifyAll)
+			else if (bNotifyAll)
 			{
 				WAITINGQUEUE_EVENT_SCOPE(TryUnblockAll);
 				// Empty wait stack and set signal to number of pre-wait threads.
@@ -260,7 +262,7 @@ int32 FWaitingQueue::NotifyInternal(int32 Count)
 
 				if ((LocalState & StackMask) == StackMask)
 				{
-					WAITINGQUEUE_EVENT_SCOPE(NoMoreWaiter2);
+					WAITINGQUEUE_EVENT_SCOPE(NoMoreWaiter);
 					return Notifications;
 				}
 
