Fix FunixPlatformProcess::CreateProc to handle NULL being passed in for parameters argument

diff --git a/Engine/Source/Runtime/Core/Private/Unix/UnixPlatformProcess.cpp b/Engine/Source/Runtime/Core/Private/Unix/UnixPlatformProcess.cpp
--- a/Engine/Source/Runtime/Core/Private/Unix/UnixPlatformProcess.cpp
+++ b/Engine/Source/Runtime/Core/Private/Unix/UnixPlatformProcess.cpp
@@ -904,6 +904,11 @@ FProcHandle FUnixPlatformProcess::CreateProc(const TCHAR* URL, const TCHAR* Parm
 		}
 	}
 
+	if (Parms == nullptr)
+	{
+		Parms = TEXT("");
+	}
+
 	const FString Commandline = FString::Printf(TEXT("\"%s\" %s"), *ProcessPath, Parms);
 	UE_LOG(LogHAL, Verbose, TEXT("FUnixPlatformProcess::CreateProc: '%s'"), *Commandline);
 
