- Fixing a bug where two branches for SourceControlSettings (project vs engine-only) would re-use one of the branches, instead of loading it a second time (so FIndBranch() will now check the FIlename if appropriate)
- Added support in FConfigContext for contexts that empty out the PRojectConfigDir to only load Engine ini layers

diff --git a/Engine/Source/Developer/SourceControl/Private/SourceControlHelpers.cpp b/Engine/Source/Developer/SourceControl/Private/SourceControlHelpers.cpp
--- a/Engine/Source/Developer/SourceControl/Private/SourceControlHelpers.cpp
+++ b/Engine/Source/Developer/SourceControl/Private/SourceControlHelpers.cpp
@@ -2048,6 +2048,7 @@ const FString& USourceControlHelpers::GetGlobalSettingsIni()
 	{
 		FConfigContext Context = FConfigContext::ReadIntoGConfig();
 		Context.GeneratedConfigDir = FPaths::EngineSavedDir() + TEXT("Config/");
+		Context.ProjectConfigDir = (""); // don't load anything from project configs
 		Context.Load(TEXT("SourceControlSettings"), SourceControlGlobalSettingsIni);
 	}
 	return SourceControlGlobalSettingsIni;
diff --git a/Engine/Source/Runtime/Core/Private/Misc/ConfigCacheIni.cpp b/Engine/Source/Runtime/Core/Private/Misc/ConfigCacheIni.cpp
--- a/Engine/Source/Runtime/Core/Private/Misc/ConfigCacheIni.cpp
+++ b/Engine/Source/Runtime/Core/Private/Misc/ConfigCacheIni.cpp
@@ -3175,6 +3175,13 @@ FConfigBranch* FConfigCacheIni::FindBranch(FName BaseIniName, const FString& Fil
 			}
 		}
 	}
+
+	// if Filename is a .ini and it doesn't match what the KnownFile has (if it has one yet), then we can't use it
+	if (Branch && Branch->IniPath.Len() > 0 && Filename.Len() > 0 && Filename.EndsWith(".ini") && Branch->IniPath != Filename)
+	{
+		Branch = nullptr;
+	}
+	
 	return Branch;
 }
 
diff --git a/Engine/Source/Runtime/Core/Private/Misc/ConfigContext.cpp b/Engine/Source/Runtime/Core/Private/Misc/ConfigContext.cpp
--- a/Engine/Source/Runtime/Core/Private/Misc/ConfigContext.cpp
+++ b/Engine/Source/Runtime/Core/Private/Misc/ConfigContext.cpp
@@ -94,10 +94,11 @@ void FConfigContext::CachePaths()
 		
 		// for the hierarchy replacements, we need to have a directory called Config - or we will have to do extra processing for these non-standard cases
 		check(EngineConfigDir.EndsWith(TEXT("Config/")));
-		check(ProjectConfigDir.EndsWith(TEXT("Config/")));
+		// allow for an empty projcet config dir, which means (below) to not load any of the {PROJECT} layers
+		check(ProjectConfigDir.Len() == 0 || ProjectConfigDir.EndsWith(TEXT("Config/")));
 
 		EngineRootDir = FPaths::GetPath(FPaths::GetPath(EngineConfigDir));
-		ProjectRootDir = FPaths::GetPath(FPaths::GetPath(ProjectConfigDir));
+		ProjectRootDir = (ProjectConfigDir.Len() > 0) ? FPaths::GetPath(FPaths::GetPath(ProjectConfigDir)) : FString();
 
 		if (FPaths::IsUnderDirectory(ProjectRootDir, EngineRootDir))
 		{
@@ -111,6 +112,18 @@ void FConfigContext::CachePaths()
 			ProjectNotForLicenseesDir = FPaths::Combine(ProjectRootDir, TEXT("Restricted/NotForLicensees"));
 			ProjectNoRedistDir = FPaths::Combine(ProjectRootDir, TEXT("Restricted/NoRedist"));
 		}
+		
+		// if we explicitly don't want project configs, then make a limited layer set without any {PROJECT} paths
+		if (ProjectConfigDir.Len() == 0 && OverrideLayers.Num() == 0)
+		{
+			for (const FConfigLayer& Layer : GConfigLayers)
+			{
+				if (FCString::Strstr(Layer.Path, TEXT("{PROJECT}")) == nullptr)
+				{
+					OverrideLayers.Add(Layer);
+				}
+			}
+		}
 	}
 }
 
