- Fixed issue with FPaths::IsStaged() with -pak generated builds (the StagedBuild.ini file can't be in the .pak since it needs ProjectDir() to find .pak files!)

diff --git a/Engine/Source/Programs/AutomationTool/Scripts/CopyBuildToStagingDirectory.Automation.cs b/Engine/Source/Programs/AutomationTool/Scripts/CopyBuildToStagingDirectory.Automation.cs
--- a/Engine/Source/Programs/AutomationTool/Scripts/CopyBuildToStagingDirectory.Automation.cs
+++ b/Engine/Source/Programs/AutomationTool/Scripts/CopyBuildToStagingDirectory.Automation.cs
@@ -1265,7 +1265,8 @@ namespace AutomationScripts
 					FileReference StagedBuildConfigFile = FileReference.Combine(IntermediateDir, "StagedBuild.ini");
 					DirectoryReference.CreateDirectory(IntermediateDir);
 					File.WriteAllText(StagedBuildConfigFile.FullName, "");
-					SC.StageFile(StagedFileType.UFS, StagedBuildConfigFile, new StagedFileReference($"Engine/Config/{StagedBuildConfigFile.GetFileName()}"));
+					// has to be outside of the .pak files since we use to to set the ProjectDir, which is needed to load .pak files!
+					SC.StageFile(StagedFileType.NonUFS, StagedBuildConfigFile, new StagedFileReference($"Engine/Config/{StagedBuildConfigFile.GetFileName()}"));
 
 					// Stage ICU internationalization data from Engine.
 					var ICUDataVersion = SC.StageTargetPlatform.ICUDataVersion;
diff --git a/Engine/Source/Runtime/Core/Private/GenericPlatform/GenericPlatformMisc.cpp b/Engine/Source/Runtime/Core/Private/GenericPlatform/GenericPlatformMisc.cpp
--- a/Engine/Source/Runtime/Core/Private/GenericPlatform/GenericPlatformMisc.cpp
+++ b/Engine/Source/Runtime/Core/Private/GenericPlatform/GenericPlatformMisc.cpp
@@ -1248,6 +1248,10 @@ const TCHAR* FGenericPlatformMisc::ProjectDir()
 {
 	FString& ProjectDir = TLazySingleton<FStaticData>::Get().ProjectDir;
 
+	// always call this at ProjectDir() time and cache it (internally) because there could be pak file interaction that can disable
+	// ini file loading
+	bool bIsStaged = FPaths::IsStaged();
+
 	// track if last time we called this function the .ini was ready and had fixed the GameName case
 	static bool bWasIniReady = false;
 	bool bIsIniReady = GConfig && GConfig->IsReadyForUse();
@@ -1277,7 +1281,7 @@ const TCHAR* FGenericPlatformMisc::ProjectDir()
 		ProjectDir.Reserve(FPlatformMisc::GetMaxPathLength());
 		if (FPlatformProperties::IsProgram())
 		{
-			if (FPaths::IsStaged())
+			if (bIsStaged)
 			{
 				// if staged, use the remapped location
 				ProjectDir = FString::Printf(TEXT("../../../%s/"), FApp::GetProjectName());
