- Moved FConsoleManager to be actually private, just by moving the class decl into the .cpp. There were a lot of uses of #include "HAL/ConsoleManager.h" instead of IConsoleManager.h (and unknown number of licensee uses), so I left that file there and it just includes IConsoleManager.h now.
- Fixed up a couple of uses of FConsoleManager when IConsoleManager is the correct usage

diff --git a/Engine/Plugins/Runtime/nDisplay/Source/DisplayClusterConfigurator/Private/DisplayClusterConfiguratorBlueprintEditor.cpp b/Engine/Plugins/Runtime/nDisplay/Source/DisplayClusterConfigurator/Private/DisplayClusterConfiguratorBlueprintEditor.cpp
--- a/Engine/Plugins/Runtime/nDisplay/Source/DisplayClusterConfigurator/Private/DisplayClusterConfiguratorBlueprintEditor.cpp
+++ b/Engine/Plugins/Runtime/nDisplay/Source/DisplayClusterConfigurator/Private/DisplayClusterConfiguratorBlueprintEditor.cpp
@@ -919,7 +919,7 @@ void FDisplayClusterConfiguratorBlueprintEditor::CreateWidgets()
 
 	// Set the visibility
 	{
-		bool bReadOnly = false;// FConsoleManager::Get().FindConsoleVariable(TEXT("nDisplay.configurator.ReadOnly"))->GetBool();
+		bool bReadOnly = false;// IConsoleManager::Get().FindConsoleVariable(TEXT("nDisplay.configurator.ReadOnly"))->GetBool();
 
 		ViewOutputMapping->SetEnabled(!bReadOnly);
 		ViewCluster->SetEnabled(!bReadOnly);
diff --git a/Engine/Plugins/VirtualProduction/RemoteControl/Source/WebRemoteControl/Private/WebRemoteControlEditorRoutes.cpp b/Engine/Plugins/VirtualProduction/RemoteControl/Source/WebRemoteControl/Private/WebRemoteControlEditorRoutes.cpp
--- a/Engine/Plugins/VirtualProduction/RemoteControl/Source/WebRemoteControl/Private/WebRemoteControlEditorRoutes.cpp
+++ b/Engine/Plugins/VirtualProduction/RemoteControl/Source/WebRemoteControl/Private/WebRemoteControlEditorRoutes.cpp
@@ -35,7 +35,7 @@
 
 void FWebRemoteControlEditorRoutes::RegisterRoutes(FWebRemoteControlModule* WebRemoteControl)
 {
-	if (FConsoleManager::Get().FindConsoleVariable(TEXT("WebControl.EnableExperimentalRoutes"))->GetBool())
+	if (IConsoleManager::Get().FindConsoleVariable(TEXT("WebControl.EnableExperimentalRoutes"))->GetBool())
 	{
 		static const FName ModuleName = "WebRemoteControl";
 		// Events
diff --git a/Engine/Source/Runtime/Core/Private/HAL/ConsoleManager.cpp b/Engine/Source/Runtime/Core/Private/HAL/ConsoleManager.cpp
--- a/Engine/Source/Runtime/Core/Private/HAL/ConsoleManager.cpp
+++ b/Engine/Source/Runtime/Core/Private/HAL/ConsoleManager.cpp
@@ -5,6 +5,11 @@ ConsoleManager.cpp: console command handling
 =============================================================================*/
 
 #include "HAL/ConsoleManager.h"
+#include "Containers/Array.h"
+#include "Containers/Map.h"
+#include "Containers/UnrealString.h"
+#include "CoreTypes.h"
+#include "HAL/CriticalSection.h"
 #include "Misc/ScopeLock.h"
 #include "Misc/ScopeRWLock.h"
 #include "Misc/Paths.h"
@@ -26,6 +31,174 @@ ConsoleManager.cpp: console command handling
 DEFINE_LOG_CATEGORY(LogConsoleResponse);
 DEFINE_LOG_CATEGORY_STATIC(LogConsoleManager, Log, All);
 
+class FOutputDevice;
+class UWorld;
+
+class FConsoleManager :public IConsoleManager
+{
+public:
+	/** constructor */
+	FConsoleManager()
+		: bHistoryWasLoaded(false)
+		, ThreadPropagationCallback(0)
+		, bCallAllConsoleVariableSinks(true)
+	{
+	}
+
+	/** destructor */
+	~FConsoleManager()
+	{
+		for (TMap<FString, IConsoleObject*>::TConstIterator PairIt(ConsoleObjects); PairIt; ++PairIt)
+		{
+			IConsoleObject* Var = PairIt.Value();
+
+			delete Var;
+		}
+	}
+
+
+
+	// internally needed or ECVF_RenderThreadSafe
+	IConsoleThreadPropagation* GetThreadPropagationCallback();
+	// internally needed or ECVF_RenderThreadSafe
+	bool IsThreadPropagationThread();
+
+	/** @param InVar must not be 0 */
+	virtual FString FindConsoleObjectName(const IConsoleObject* Obj) const override;
+
+	/** Can be moved out into some automated testing system */
+	void Test();
+
+	void OnCVarChanged();
+
+	virtual FConsoleVariableMulticastDelegate& OnCVarUnregistered()override;
+	virtual FConsoleObjectWithNameMulticastDelegate& OnConsoleObjectUnregistered() override;
+
+	// interface IConsoleManager -----------------------------------
+
+	virtual IConsoleVariable* RegisterConsoleVariable(const TCHAR* Name, bool DefaultValue, const TCHAR* Help, uint32 Flags) override;
+	virtual IConsoleVariable* RegisterConsoleVariable(const TCHAR* Name, int32 DefaultValue, const TCHAR* Help, uint32 Flags) override;
+	virtual IConsoleVariable* RegisterConsoleVariable(const TCHAR* Name, float DefaultValue, const TCHAR* Help, uint32 Flags) override;
+	virtual IConsoleVariable* RegisterConsoleVariable(const TCHAR* Name, const TCHAR* DefaultValue, const TCHAR* Help, uint32 Flags) override;
+	virtual IConsoleVariable* RegisterConsoleVariable(const TCHAR* Name, const FString& DefaultValue, const TCHAR* Help, uint32 Flags) override;
+
+	virtual IConsoleVariable* RegisterConsoleVariableRef(const TCHAR* Name, bool& RefValue, const TCHAR* Help, uint32 Flags) override;
+	virtual IConsoleVariable* RegisterConsoleVariableRef(const TCHAR* Name, int32& RefValue, const TCHAR* Help, uint32 Flags) override;
+	virtual IConsoleVariable* RegisterConsoleVariableRef(const TCHAR* Name, float& RefValue, const TCHAR* Help, uint32 Flags) override;
+	virtual IConsoleVariable* RegisterConsoleVariableRef(const TCHAR* Name, FString& RefValue, const TCHAR* Help, uint32 Flags) override;
+	virtual IConsoleVariable* RegisterConsoleVariableBitRef(const TCHAR* CVarName, const TCHAR* FlagName, uint32 BitNumber, uint8* Force0MaskPtr, uint8* Force1MaskPtr, const TCHAR* Help, uint32 Flags) override;
+
+
+	virtual void CallAllConsoleVariableSinks() override;
+
+	virtual FConsoleVariableSinkHandle RegisterConsoleVariableSink_Handle(const FConsoleCommandDelegate& Command) override;
+	virtual void UnregisterConsoleVariableSink_Handle(FConsoleVariableSinkHandle Handle) override;
+
+	virtual IConsoleCommand* RegisterConsoleCommand(const TCHAR* Name, const TCHAR* Help, const FConsoleCommandDelegate& Command, uint32 Flags) override;
+	virtual IConsoleCommand* RegisterConsoleCommand(const TCHAR* Name, const TCHAR* Help, const FConsoleCommandWithArgsDelegate& Command, uint32 Flags) override;
+	virtual IConsoleCommand* RegisterConsoleCommand(const TCHAR* Name, const TCHAR* Help, const FConsoleCommandWithWorldDelegate& Command, uint32 Flags) override;
+	virtual IConsoleCommand* RegisterConsoleCommand(const TCHAR* Name, const TCHAR* Help, const FConsoleCommandWithWorldAndArgsDelegate& Command, uint32 Flags) override;
+	virtual IConsoleCommand* RegisterConsoleCommand(const TCHAR* Name, const TCHAR* Help, const FConsoleCommandWithArgsAndOutputDeviceDelegate& Command, uint32 Flags) override;
+	virtual IConsoleCommand* RegisterConsoleCommand(const TCHAR* Name, const TCHAR* Help, const FConsoleCommandWithWorldArgsAndOutputDeviceDelegate& Command, uint32 Flags) override;
+	virtual IConsoleCommand* RegisterConsoleCommand(const TCHAR* Name, const TCHAR* Help, const FConsoleCommandWithOutputDeviceDelegate& Command, uint32 Flags) override;
+	virtual IConsoleCommand* RegisterConsoleCommand(const TCHAR* Name, const TCHAR* Help, uint32 Flags) override;
+	virtual IConsoleObject* FindConsoleObject(const TCHAR* Name, bool bTrackFrequentCalls = true) const override;
+	virtual IConsoleVariable* FindConsoleVariable(const TCHAR* Name, bool bTrackFrequentCalls = true) const override;
+	virtual void ForEachConsoleObjectThatStartsWith(const FConsoleObjectVisitor& Visitor, const TCHAR* ThatStartsWith) const override;
+	virtual void ForEachConsoleObjectThatContains(const FConsoleObjectVisitor& Visitor, const TCHAR* ThatContains) const override;
+	virtual bool ProcessUserConsoleInput(const TCHAR* InInput, FOutputDevice& Ar, UWorld* InWorld) override;
+	virtual void AddConsoleHistoryEntry(const TCHAR* Key, const TCHAR* Input) override;
+	virtual void GetConsoleHistory(const TCHAR* Key, TArray<FString>& Out) override;
+	virtual bool IsNameRegistered(const TCHAR* Name) const override;
+	virtual void RegisterThreadPropagation(uint32 ThreadId, IConsoleThreadPropagation* InCallback) override;
+	virtual void UnregisterConsoleObject(IConsoleObject* Object, bool bKeepState) override;
+	virtual void UnsetAllConsoleVariablesWithTag(FName Tag, EConsoleVariableFlags Priority) override;
+#if ALLOW_OTHER_PLATFORM_CONFIG
+	virtual void LoadAllPlatformCVars(FName PlatformName, const FString& DeviceProfileName = FString()) override;
+	virtual void ClearAllPlatformCVars(FName PlatformName = NAME_None, const FString& DeviceProfileName = FString()) override;
+	virtual void PreviewPlatformCVars(FName PlatformName, const FString& DeviceProfileName, FName PreviewModeTag) override;
+#endif
+
+	// Internal functionality to this file -----------------------------------
+
+	/** Map of console variables and commands, indexed by the name of that command or variable */
+	// [name] = pointer (pointer must not be 0)
+	TMap<FString, IConsoleObject*> ConsoleObjects;
+
+	bool bHistoryWasLoaded;
+	TMap<FString, TArray<FString>>	HistoryEntriesMap;
+	TArray<FConsoleCommandDelegate>	ConsoleVariableChangeSinks;
+
+	FConsoleVariableMulticastDelegate ConsoleVariableUnregisteredDelegate;
+	FConsoleObjectWithNameMulticastDelegate ConsoleObjectUnregisteredDelegate;
+
+	IConsoleThreadPropagation* ThreadPropagationCallback;
+
+	FCriticalSection CachedPlatformsAndDeviceProfilesLock;
+	TSet<FName> CachedPlatformsAndDeviceProfiles;
+
+	// if true the next call to CallAllConsoleVariableSinks() we will call all registered sinks
+	bool bCallAllConsoleVariableSinks;
+
+	/**
+		* Used to prevent concurrent access to ConsoleObjects.
+		* We don't aim to solve all concurrency problems (for example registering and unregistering a cvar on different threads, or reading a cvar from one thread while writing it from a different thread).
+		* Rather we just ensure that operations on a cvar from one thread will not conflict with operations on another cvar from another thread.
+	**/
+	mutable FCriticalSection ConsoleObjectsSynchronizationObject;
+
+	/**
+	 * @param Name must not be 0, must not be empty
+	 * @param Obj must not be 0
+	 * @return 0 if the name was already in use
+	 */
+	IConsoleObject* AddConsoleObject(const TCHAR* Name, IConsoleObject* Obj);
+	
+	/**
+	 * Similar to AddConsoleObject, but it just adds it without any flag checking or preexisting var checking
+	 */
+	void AddShadowConsoleObject(const TCHAR* Name, IConsoleObject* Obj);
+
+	/**
+	 * @param Stream must not be 0
+	 * @param Pattern must not be 0
+	 */
+	static bool MatchPartialName(const TCHAR* Stream, const TCHAR* Pattern);
+
+	/** Returns true if Pattern is found in Stream, case insensitive. */
+	static bool MatchSubstring(const TCHAR* Stream, const TCHAR* Pattern);
+
+	/**
+	 * Get string till whitespace, jump over whitespace
+	 * inefficient but this code is not performance critical
+	 */
+	static FString GetTextSection(const TCHAR*& It);
+
+	/** same as FindConsoleObject() but ECVF_CreatedFromIni are not filtered out (for internal use) */
+	IConsoleObject* FindConsoleObjectUnfiltered(const TCHAR* Name) const;
+
+	/**
+	 * Unregisters a console variable or command, if that object was registered.  For console variables, this will
+	 * actually only "deactivate" the variable so if it becomes registered again the state may persist
+	 * (unless bKeepState is false).
+	 *
+	 * @param	Name	Name of the console object to remove (not case sensitive)
+	 * @param	bKeepState	if the current state is kept in memory until a cvar with the same name is registered
+	 */
+	void UnregisterConsoleObject(const TCHAR* Name, bool bKeepState);
+
+	// reads HistoryEntriesMap from the .ini file (if not already loaded)
+	void LoadHistoryIfNeeded();
+
+	// writes HistoryEntriesMap to the .ini file
+	void SaveHistory();
+};
+
+static FConsoleManager& GetManager()
+{
+	return (FConsoleManager&)IConsoleManager::Get();
+}
+
 namespace UE::ConsoleManager::Private
 {
 	// this tracks the cvars that were  added dynamically with a tag (via plugin or similar)
@@ -221,8 +394,7 @@ public:
 
 		if(!bRet)
 		{
-			FConsoleManager& ConsoleManager = (FConsoleManager&)IConsoleManager::Get();
-			FString CVarName = ConsoleManager.FindConsoleObjectName(this);
+			FString CVarName = IConsoleManager::Get().FindConsoleObjectName(this);
 
 			const FString Message = FString::Printf(TEXT("Setting the console variable '%s' with 'SetBy%s' was ignored as it is lower priority than the previous 'SetBy%s'. Value remains '%s'"),
 				CVarName.IsEmpty() ? TEXT("unknown?") : *CVarName,
@@ -280,14 +452,6 @@ public:
 	// ------
 	// Helper accessors to get to FConsoleVariableExtendedData, when we don't have a Type
 
-	/**
-	 * Print the history to a log
-	 */
-	virtual void LogHistory(FOutputDevice& Ar) = 0;
-	/**
-	 * Track memory used by history data
-	 */
-	virtual SIZE_T GetHistorySize() = 0;
 
 #if ALLOW_OTHER_PLATFORM_CONFIG
 	/**
@@ -319,7 +483,7 @@ protected: // -----------------------------------------
 		}
 		else
 		{
-			FConsoleManager& ConsoleManager = (FConsoleManager&)IConsoleManager::Get();
+			FConsoleManager& ConsoleManager = GetManager();
 			if(ConsoleManager.IsThreadPropagationThread() && FPlatformProcess::SupportsMultithreading())
 			{
 				if(!bWarnedAboutThreadSafety)
@@ -410,7 +574,7 @@ void OnCVarChange(T& Dst, const T& Src, EConsoleVariableFlags Flags, EConsoleVar
         return;
     }
     
-	FConsoleManager& ConsoleManager = (FConsoleManager&)IConsoleManager::Get();
+	FConsoleManager& ConsoleManager = GetManager();
 
 #if WITH_RELOAD
 	// Unlike HotReload, Live Coding does global initialization outside of the main thread.  During global initialization,
@@ -3397,7 +3561,7 @@ void CreateConsoleVariables()
 
 	// testing code
 	{
-		FConsoleManager& ConsoleManager = (FConsoleManager&)IConsoleManager::Get();
+		FConsoleManager& ConsoleManager = GetManager();
 
 		ConsoleManager.Test();
 	}
diff --git a/Engine/Source/Runtime/Core/Public/HAL/ConsoleManager.h b/Engine/Source/Runtime/Core/Public/HAL/ConsoleManager.h
--- a/Engine/Source/Runtime/Core/Public/HAL/ConsoleManager.h
+++ b/Engine/Source/Runtime/Core/Public/HAL/ConsoleManager.h
@@ -2,164 +2,6 @@
 
 #pragma once
 
-#include "Containers/Array.h"
-#include "Containers/Map.h"
-#include "Containers/UnrealString.h"
-#include "CoreTypes.h"
-#include "HAL/CriticalSection.h"
+// Simply include the interface header - this file could be deleted if all #includes of it were changed to IConsoleManager
+// or we could move this into Private, and put the FConsoleManager class back in
 #include "HAL/IConsoleManager.h"
-
-class FOutputDevice;
-class UWorld;
-
-class FConsoleManager :public IConsoleManager
-{
-public:
-	/** constructor */
-	FConsoleManager()
-		: bHistoryWasLoaded(false)
-		, ThreadPropagationCallback(0)
-		, bCallAllConsoleVariableSinks(true)
-	{
-	}
-
-	/** destructor */
-	~FConsoleManager()
-	{
-		for(TMap<FString, IConsoleObject*>::TConstIterator PairIt(ConsoleObjects); PairIt; ++PairIt)
-		{
-			IConsoleObject* Var = PairIt.Value();
-
-			delete Var;
-		}
-	}
-	
-	// internally needed or ECVF_RenderThreadSafe
-	CORE_API IConsoleThreadPropagation* GetThreadPropagationCallback();
-	// internally needed or ECVF_RenderThreadSafe
-	CORE_API bool IsThreadPropagationThread();
-
-	/** @param InVar must not be 0 */
-	CORE_API virtual FString FindConsoleObjectName(const IConsoleObject* Obj) const override;
-
-	/** Can be moved out into some automated testing system */
-	CORE_API void Test();
-
-	CORE_API void OnCVarChanged();
-
-	CORE_API virtual FConsoleVariableMulticastDelegate& OnCVarUnregistered()override;
-	CORE_API virtual FConsoleObjectWithNameMulticastDelegate& OnConsoleObjectUnregistered() override;
-
-	// interface IConsoleManager -----------------------------------
-
-	CORE_API virtual IConsoleVariable* RegisterConsoleVariable(const TCHAR* Name, bool DefaultValue, const TCHAR* Help, uint32 Flags) override;
-	CORE_API virtual IConsoleVariable* RegisterConsoleVariable(const TCHAR* Name, int32 DefaultValue, const TCHAR* Help, uint32 Flags) override;
-	CORE_API virtual IConsoleVariable* RegisterConsoleVariable(const TCHAR* Name, float DefaultValue, const TCHAR* Help, uint32 Flags) override;
-	CORE_API virtual IConsoleVariable* RegisterConsoleVariable(const TCHAR* Name, const TCHAR* DefaultValue, const TCHAR* Help, uint32 Flags) override;
-	CORE_API virtual IConsoleVariable* RegisterConsoleVariable(const TCHAR* Name, const FString& DefaultValue, const TCHAR* Help, uint32 Flags) override;
-
-	CORE_API virtual IConsoleVariable* RegisterConsoleVariableRef(const TCHAR* Name, bool& RefValue, const TCHAR* Help, uint32 Flags) override;
-	CORE_API virtual IConsoleVariable* RegisterConsoleVariableRef(const TCHAR* Name, int32& RefValue, const TCHAR* Help, uint32 Flags) override;
-	CORE_API virtual IConsoleVariable* RegisterConsoleVariableRef(const TCHAR* Name, float& RefValue, const TCHAR* Help, uint32 Flags) override;
-	CORE_API virtual IConsoleVariable* RegisterConsoleVariableRef(const TCHAR* Name, FString& RefValue, const TCHAR* Help, uint32 Flags) override;
-	CORE_API virtual IConsoleVariable* RegisterConsoleVariableBitRef(const TCHAR* CVarName, const TCHAR* FlagName, uint32 BitNumber, uint8* Force0MaskPtr, uint8* Force1MaskPtr, const TCHAR* Help, uint32 Flags) override;
-	
-	CORE_API virtual void CallAllConsoleVariableSinks() override;
-
-	CORE_API virtual FConsoleVariableSinkHandle RegisterConsoleVariableSink_Handle(const FConsoleCommandDelegate& Command) override;
-	CORE_API virtual void UnregisterConsoleVariableSink_Handle(FConsoleVariableSinkHandle Handle) override;
-
-	CORE_API virtual IConsoleCommand* RegisterConsoleCommand(const TCHAR* Name, const TCHAR* Help, const FConsoleCommandDelegate& Command, uint32 Flags) override;
-	CORE_API virtual IConsoleCommand* RegisterConsoleCommand(const TCHAR* Name, const TCHAR* Help, const FConsoleCommandWithArgsDelegate& Command, uint32 Flags) override;
-	CORE_API virtual IConsoleCommand* RegisterConsoleCommand(const TCHAR* Name, const TCHAR* Help, const FConsoleCommandWithWorldDelegate& Command, uint32 Flags) override;
-	CORE_API virtual IConsoleCommand* RegisterConsoleCommand(const TCHAR* Name, const TCHAR* Help, const FConsoleCommandWithWorldAndArgsDelegate& Command, uint32 Flags) override;
-	CORE_API virtual IConsoleCommand* RegisterConsoleCommand(const TCHAR* Name, const TCHAR* Help, const FConsoleCommandWithArgsAndOutputDeviceDelegate& Command, uint32 Flags) override;
-	CORE_API virtual IConsoleCommand* RegisterConsoleCommand(const TCHAR* Name, const TCHAR* Help, const FConsoleCommandWithWorldArgsAndOutputDeviceDelegate& Command, uint32 Flags) override;
-	CORE_API virtual IConsoleCommand* RegisterConsoleCommand(const TCHAR* Name, const TCHAR* Help, const FConsoleCommandWithOutputDeviceDelegate& Command, uint32 Flags) override;
-	CORE_API virtual IConsoleCommand* RegisterConsoleCommand(const TCHAR* Name, const TCHAR* Help, uint32 Flags) override;
-	CORE_API virtual IConsoleObject* FindConsoleObject(const TCHAR* Name, bool bTrackFrequentCalls = true) const override;
-	CORE_API virtual IConsoleVariable* FindConsoleVariable(const TCHAR* Name, bool bTrackFrequentCalls = true) const override;
-	CORE_API virtual void ForEachConsoleObjectThatStartsWith(const FConsoleObjectVisitor& Visitor, const TCHAR* ThatStartsWith) const override;
-	CORE_API virtual void ForEachConsoleObjectThatContains(const FConsoleObjectVisitor& Visitor, const TCHAR* ThatContains) const override;
-	CORE_API virtual bool ProcessUserConsoleInput(const TCHAR* InInput, FOutputDevice& Ar, UWorld* InWorld) override;
-	CORE_API virtual void AddConsoleHistoryEntry(const TCHAR* Key, const TCHAR* Input) override;
-	CORE_API virtual void GetConsoleHistory(const TCHAR* Key, TArray<FString>& Out) override;
-	CORE_API virtual bool IsNameRegistered(const TCHAR* Name) const override;	
-	CORE_API virtual void RegisterThreadPropagation(uint32 ThreadId, IConsoleThreadPropagation* InCallback) override;
-	CORE_API virtual void UnregisterConsoleObject( IConsoleObject* Object, bool bKeepState) override;
-	CORE_API virtual void UnsetAllConsoleVariablesWithTag(FName Tag, EConsoleVariableFlags Priority) override;
-#if ALLOW_OTHER_PLATFORM_CONFIG
-	CORE_API virtual void LoadAllPlatformCVars(FName PlatformName, const FString& DeviceProfileName=FString()) override;
-	CORE_API virtual void ClearAllPlatformCVars(FName PlatformName=NAME_None, const FString& DeviceProfileName=FString()) override;
-	CORE_API virtual void PreviewPlatformCVars(FName PlatformName, const FString& DeviceProfileName, FName PreviewModeTag) override;
-#endif
-	
-private: // ----------------------------------------------------
-
-	/** Map of console variables and commands, indexed by the name of that command or variable */
-	// [name] = pointer (pointer must not be 0)
-	TMap<FString, IConsoleObject*> ConsoleObjects;
-
-	bool bHistoryWasLoaded;
-	TMap<FString, TArray<FString>>	HistoryEntriesMap;
-	TArray<FConsoleCommandDelegate>	ConsoleVariableChangeSinks;
-
-	FConsoleVariableMulticastDelegate ConsoleVariableUnregisteredDelegate;
-	FConsoleObjectWithNameMulticastDelegate ConsoleObjectUnregisteredDelegate;
-
-	IConsoleThreadPropagation* ThreadPropagationCallback;
-	
-	FCriticalSection CachedPlatformsAndDeviceProfilesLock;
-	TSet<FName> CachedPlatformsAndDeviceProfiles;
-
-	// if true the next call to CallAllConsoleVariableSinks() we will call all registered sinks
-	bool bCallAllConsoleVariableSinks;
-
-	/** 
-		* Used to prevent concurrent access to ConsoleObjects.
-		* We don't aim to solve all concurrency problems (for example registering and unregistering a cvar on different threads, or reading a cvar from one thread while writing it from a different thread).
-		* Rather we just ensure that operations on a cvar from one thread will not conflict with operations on another cvar from another thread.
-	**/
-	mutable FCriticalSection ConsoleObjectsSynchronizationObject;
-
-	/** 
-	 * @param Name must not be 0, must not be empty
-	 * @param Obj must not be 0
-	 * @return 0 if the name was already in use
-	 */
-	CORE_API IConsoleObject* AddConsoleObject(const TCHAR* Name, IConsoleObject* Obj);
-
-	/**
-	 * @param Stream must not be 0
-	 * @param Pattern must not be 0
-	 */
-	static CORE_API bool MatchPartialName(const TCHAR* Stream, const TCHAR* Pattern);
-
-	/** Returns true if Pattern is found in Stream, case insensitive. */
-	static CORE_API bool MatchSubstring(const TCHAR* Stream, const TCHAR* Pattern);
-
-	/**
-	 * Get string till whitespace, jump over whitespace
-	 * inefficient but this code is not performance critical
-	 */
-	static CORE_API FString GetTextSection(const TCHAR* &It);
-
-	/** same as FindConsoleObject() but ECVF_CreatedFromIni are not filtered out (for internal use) */
-	CORE_API IConsoleObject* FindConsoleObjectUnfiltered(const TCHAR* Name) const;
-
-	/**
-	 * Unregisters a console variable or command, if that object was registered.  For console variables, this will
-	 * actually only "deactivate" the variable so if it becomes registered again the state may persist
-	 * (unless bKeepState is false).
-	 *
-	 * @param	Name	Name of the console object to remove (not case sensitive)
-	 * @param	bKeepState	if the current state is kept in memory until a cvar with the same name is registered
-	 */
-	CORE_API void UnregisterConsoleObject(const TCHAR* Name, bool bKeepState);
-
-	// reads HistoryEntriesMap from the .ini file (if not already loaded)
-	CORE_API void LoadHistoryIfNeeded();
-
-	// writes HistoryEntriesMap to the .ini file
-	CORE_API void SaveHistory();
-};
diff --git a/Engine/Source/Runtime/Core/Public/HAL/IConsoleManager.h b/Engine/Source/Runtime/Core/Public/HAL/IConsoleManager.h
--- a/Engine/Source/Runtime/Core/Public/HAL/IConsoleManager.h
+++ b/Engine/Source/Runtime/Core/Public/HAL/IConsoleManager.h
@@ -583,6 +583,16 @@ public:
 	
 #endif
 
+	/**
+	 * Print the history to a log
+	*/
+	virtual void LogHistory(FOutputDevice& Ar) = 0;
+	
+	/**
+	 * Track memory used by history data
+	 */
+	virtual SIZE_T GetHistorySize() = 0;
+
 	// convenience methods
 
 	/** Set the internal value from the specified bool. */
diff --git a/Engine/Source/Runtime/Engine/Private/DataDrivenCVars/DataDrivenCVars.cpp b/Engine/Source/Runtime/Engine/Private/DataDrivenCVars/DataDrivenCVars.cpp
--- a/Engine/Source/Runtime/Engine/Private/DataDrivenCVars/DataDrivenCVars.cpp
+++ b/Engine/Source/Runtime/Engine/Private/DataDrivenCVars/DataDrivenCVars.cpp
@@ -114,7 +114,7 @@ void UDataDrivenConsoleVariableSettings::OnDataDrivenChange(IConsoleVariable* CV
 		UDataDrivenCVarEngineSubsystem* Subsystem = GEngine->GetEngineSubsystem<UDataDrivenCVarEngineSubsystem>();
 		if (Subsystem)
 		{
-			FConsoleManager& ConsoleManager = (FConsoleManager&)IConsoleManager::Get();
+			const IConsoleManager& ConsoleManager = IConsoleManager::Get();
 			Subsystem->OnDataDrivenCVarDelegate.Broadcast(ConsoleManager.FindConsoleObjectName(CVar));
 		}
 	}
