Cache PluginBase.ini to ensure we don't break platform restrictions on multiple read access during boot

#rnx
#tests Certification issue confirmed fixed by QA

diff --git a/Engine/Source/Runtime/Core/Private/Misc/ConfigCacheIni.cpp b/Engine/Source/Runtime/Core/Private/Misc/ConfigCacheIni.cpp
--- a/Engine/Source/Runtime/Core/Private/Misc/ConfigCacheIni.cpp
+++ b/Engine/Source/Runtime/Core/Private/Misc/ConfigCacheIni.cpp
@@ -519,14 +519,16 @@ void FixupArrayOfStructKeysForSection(SectionType* Section, const FString& Secti
  */
 static bool LoadConfigFileWrapper(const TCHAR* IniFile, FString& Contents, bool bIsOverride = false)
 {
-	// We read the Base.ini file many many times, so cache it
+	// We read the Base.ini and PluginBase.ini files many many times, so cache them
 	static FString BaseIniContents;
+	static FString PluginBaseIniContents;
 
 	const TCHAR* LastSlash = FCString::Strrchr(IniFile, '/');
 	if (LastSlash == nullptr)
 	{
 		LastSlash = FCString::Strrchr(IniFile, '\\');
 	}
+
 	bool bIsBaseIni = LastSlash != nullptr && FCString::Stricmp(LastSlash + 1, TEXT("Base.ini")) == 0;
 	if (bIsBaseIni && BaseIniContents.Len() > 0)
 	{
@@ -534,6 +536,13 @@ static bool LoadConfigFileWrapper(const TCHAR* IniFile, FString& Contents, bool
 		return true;
 	}
 
+	bool bIsPluginBaseIni = LastSlash != nullptr && FCString::Stricmp(LastSlash + 1, TEXT("PluginBase.ini")) == 0;
+	if (bIsPluginBaseIni && PluginBaseIniContents.Len() > 0)
+	{
+		Contents = PluginBaseIniContents;
+		return true;
+	}
+
 	// let other systems load the file instead of the standard load below
 	FCoreDelegates::TSPreLoadConfigFileDelegate().Broadcast(IniFile, Contents);
 
@@ -557,9 +566,16 @@ static bool LoadConfigFileWrapper(const TCHAR* IniFile, FString& Contents, bool
 	// needs file ops), and the other caller of this is already checking for disabled file ops
 	// and don't read from the file, if the delegate got anything loaded
 	bool bResult = FFileHelper::LoadFileToString(Contents, IniFile);
-	if (bResult && bIsBaseIni)
+	if (bResult)
 	{
-		BaseIniContents = Contents;
+		if (bIsBaseIni)
+		{
+			BaseIniContents = Contents;
+		}
+		else if (bIsPluginBaseIni)
+		{
+			PluginBaseIniContents = Contents;
+		}
 	}
 	return bResult;
 }
