Implemented Weak and Strong Ptr Variants

diff --git a/Engine/Plugins/Experimental/ControlFlows/Source/ControlFlows/Public/CastedTo.h b/Engine/Plugins/Experimental/ControlFlows/Source/ControlFlows/Public/CastedTo.h
deleted file mode 100644
diff --git a/Engine/Plugins/Experimental/ControlFlows/Source/ControlFlows/Public/ControlFlowContainer.h b/Engine/Plugins/Experimental/ControlFlows/Source/ControlFlows/Public/ControlFlowContainer.h
--- a/Engine/Plugins/Experimental/ControlFlows/Source/ControlFlows/Public/ControlFlowContainer.h
+++ b/Engine/Plugins/Experimental/ControlFlows/Source/ControlFlows/Public/ControlFlowContainer.h
@@ -5,7 +5,7 @@
 // HEADER_UNIT_SKIP - Not included directly
 
 #include "Templates/SharedPointer.h"
-#include "CastedTo.h"
+#include "Templates/PointerVariants.h"
 
 class FControlFlowContainerBase : public TSharedFromThis<FControlFlowContainerBase>
 {
@@ -50,9 +50,9 @@ public:
 private:
 	virtual const void* const GetOwningObject() const override final
 	{
-		return OwningObject.IsValid() ? OwningObject.Cast() : nullptr;
+		return OwningObject.IsValid() ? OwningObject.Pin().Get() : nullptr;
 	}
 
 private:
-	TWeakContainer<T> OwningObject;
+	TWeakPtrVariant<T> OwningObject;
 };
\ No newline at end of file
diff --git a/Engine/Plugins/Experimental/ControlFlows/Source/ControlFlows/Public/SmartPtr.h b/Engine/Plugins/Experimental/ControlFlows/Source/ControlFlows/Public/SmartPtr.h
deleted file mode 100644
diff --git a/Engine/Source/Runtime/Core/Public/Templates/PointerVariants.h b/Engine/Source/Runtime/Core/Public/Templates/PointerVariants.h
new file mode 100644
--- /dev/null
+++ b/Engine/Source/Runtime/Core/Public/Templates/PointerVariants.h
@@ -0,0 +1,244 @@
+// Copyright Epic Games, Inc. All Rights Reserved.
+
+#pragma once
+
+#include "Misc/TVariant.h"
+#include "Templates/SharedPointer.h"
+#include "UObject/StrongObjectPtr.h"
+#include "UObject/WeakObjectPtr.h"
+#include "UnrealTypeTraits.h"
+#include "UnrealTemplate.h"
+
+/*
+* `TWeakPtrVariant` and `TStrongPtrVariant` are particularly useful for "interfaces" in areas where
+* UObjects and non-UObjects exist to hold onto a "smart" pointer to the Interface, rather than needing to keep track of the implementing class to check validity and manually cast
+*/
+
+/* `TWeakPtrVariant` will be either a `TWeakObjectPtr` or `TWeakPtr` */
+template<typename BaseType>
+struct TWeakPtrVariant;
+
+/* `TStrongPtrVariant` will be either a `TStrongObjectPtr` or `TSharedPtr` */
+template<typename BaseType>
+struct TStrongPtrVariant;
+
+namespace UE::Core::Private
+{
+	template<typename BaseType, bool bIsStrong>
+	struct TPtrVariantBase
+	{
+		friend struct TWeakPtrVariant<BaseType>;
+		friend struct TStrongPtrVariant<BaseType>;
+
+		using SharedType = std::conditional_t<bIsStrong, TSharedPtr<BaseType>, TWeakPtr<BaseType>>;
+
+		struct FObjectPtrWrapper
+		{
+			using ObjectPtrType = std::conditional_t<bIsStrong, TStrongObjectPtr<const UObject>, TWeakObjectPtr<const UObject>>;
+
+			FObjectPtrWrapper(const UObject* InObjectPtr, BaseType* const InCastedPtr)
+				: ObjectPtr(InObjectPtr)
+				, CastedPtr(InCastedPtr)
+			{
+			}
+
+			ObjectPtrType ObjectPtr;
+			BaseType* CastedPtr = nullptr;
+		};
+
+	public:
+		constexpr TPtrVariantBase() = default;
+		constexpr TPtrVariantBase(TYPE_OF_NULLPTR)
+		{
+		}
+
+		template<typename DerivedType,
+			decltype(ImplicitConv<BaseType*>((DerivedType*)nullptr))* = nullptr
+			UE_REQUIRES(std::is_base_of_v<UObject, DerivedType> || IsDerivedFromSharedFromThis<DerivedType>())
+			>
+		constexpr TPtrVariantBase(DerivedType* InDerived)
+		{   
+			// TODO:	Look for a way to move this "smart pointer selection" logic outside of the constructor; then we can remove "TVariant" as a dependency
+			//			The problem is that "DerivedType" is defined in the constructor; We don't want TWeakPtrVariant or TStrongPtrVariant to have "DerivedType" as part of their template args
+			if constexpr (std::is_base_of_v<UObject, DerivedType>)
+			{
+				PtrVariant.template Emplace<FObjectPtrWrapper>(Cast<const UObject>(InDerived), Cast<BaseType>(InDerived));
+			}
+			else
+			{
+				if constexpr (bIsStrong)
+				{
+					PtrVariant.template Emplace<SharedType>(StaticCastSharedRef<BaseType>(StaticCastSharedRef<DerivedType>(InDerived->AsShared())).ToSharedPtr());
+				}
+				else
+				{
+					PtrVariant.template Emplace<SharedType>(StaticCastSharedRef<BaseType>(StaticCastSharedRef<DerivedType>(InDerived->AsShared())).ToWeakPtr());
+				}
+			}
+		}
+
+		TPtrVariantBase& operator=(TYPE_OF_NULLPTR)
+		{
+			Reset();
+			return *this;
+		}
+
+		bool IsValid() const
+		{
+			return ::Visit([](auto& StoredValue)
+				{
+					using StoredValueType = std::decay_t<decltype(StoredValue)>;
+					if constexpr (std::is_same_v<StoredValueType, SharedType>)
+					{
+						return StoredValue.IsValid();
+					}
+					else
+					{
+						return StoredValue.ObjectPtr.IsValid();
+					}
+				}, PtrVariant);
+		}
+
+		void Reset()
+		{
+			::Visit([](auto& StoredValue)
+				{
+					using StoredValueType = std::decay_t<decltype(StoredValue)>;
+					if constexpr (std::is_same_v<StoredValueType, SharedType>)
+					{
+						StoredValue.Reset();
+					}
+					else
+					{
+						StoredValue.ObjectPtr.Reset();
+					}
+				}, PtrVariant);
+		}
+
+	private:
+		TVariant<SharedType, FObjectPtrWrapper> PtrVariant;
+	};
+}
+
+/* Global Comparisons */
+
+template <typename LhsType, bool bLhsStrength, typename RhsType, bool bRhsStrength, 
+	decltype((LhsType*)nullptr == (RhsType*)nullptr) = false /* this will prevent the template instantiating if the pointer types are not comparable*/>
+bool operator==(const UE::Core::Private::TPtrVariantBase<LhsType, bLhsStrength>& Lhs, const UE::Core::Private::TPtrVariantBase<RhsType, bRhsStrength>& Rhs)
+{
+	return Lhs.Get() == Rhs.Get();
+}
+
+template <typename LhsType, bool bLhsStrength, typename RhsType, bool bRhsStrength, decltype((LhsType*)nullptr == (RhsType*)nullptr) = false>
+bool operator!=(const UE::Core::Private::TPtrVariantBase<LhsType, bLhsStrength>& Lhs, const UE::Core::Private::TPtrVariantBase<RhsType, bRhsStrength>& Rhs)
+{
+	return !(Lhs == Rhs);
+}
+
+template<typename LhsType, typename RhsType, bool bRhsStrength, decltype((LhsType*)nullptr == (RhsType*)nullptr) = false>
+bool operator==(LhsType* Lhs, const UE::Core::Private::TPtrVariantBase<RhsType, bRhsStrength>& Rhs)
+{
+	return Lhs == Rhs.Get();
+}
+
+template<typename LhsType, typename RhsType, bool bRhsStrength, decltype((LhsType*)nullptr == (RhsType*)nullptr) = false>
+bool operator!=(LhsType* Lhs, const UE::Core::Private::TPtrVariantBase<RhsType, bRhsStrength>& Rhs)
+{
+	return !(Lhs == Rhs);
+}
+
+template<typename LhsType, bool bLhsStrength, typename RhsType, decltype((LhsType*)nullptr == (RhsType*)nullptr) = false>
+bool operator==(const UE::Core::Private::TPtrVariantBase<LhsType, bLhsStrength>& Lhs, RhsType* Rhs)
+{
+	return Lhs.Get() == Rhs;
+}
+
+template<typename LhsType, bool bLhsStrength, typename RhsType, decltype((LhsType*)nullptr == (RhsType*)nullptr) = false>
+bool operator!=(const UE::Core::Private::TPtrVariantBase<LhsType, bLhsStrength>& Lhs, RhsType* Rhs)
+{
+	return !(Lhs == Rhs);
+}
+
+template<typename BaseType>
+struct TStrongPtrVariant : public UE::Core::Private::TPtrVariantBase<BaseType, true>
+{
+	using Super = UE::Core::Private::TPtrVariantBase<BaseType, true>;
+	using Super::Super;
+
+	BaseType* Get() const
+	{
+		return ::Visit([](const auto& StoredValue)
+			{
+				using StoredValueType = std::decay_t<decltype(StoredValue)>;
+				if constexpr (TIsTSharedPtr_V<StoredValueType>)
+				{
+					return StoredValue.Get();
+				}
+				else
+				{
+					return StoredValue.CastedPtr;
+				}
+			}, Super::PtrVariant);
+	}
+
+	TWeakPtrVariant<BaseType> ToWeakVariant() const
+	{
+		if (Super::IsValid())
+		{
+			return ::Visit([](auto& StoredValue)
+				{
+					TWeakPtrVariant<BaseType> WeakPtrVariant;
+
+					using StoredValueType = std::decay_t<decltype(StoredValue)>;
+					if constexpr (TIsTSharedPtr_V<StoredValueType>)
+					{
+						WeakPtrVariant.PtrVariant.template Emplace<typename TWeakPtrVariant<BaseType>::SharedType>(StoredValue.ToWeakPtr());
+					}
+					else
+					{
+						WeakPtrVariant.PtrVariant.template Emplace<typename TWeakPtrVariant<BaseType>::FObjectPtrWrapper>(StoredValue.ObjectPtr.Get(), StoredValue.CastedPtr);
+					}
+
+					return WeakPtrVariant;
+				}, Super::PtrVariant);
+		}
+		else
+		{
+			return TWeakPtrVariant<BaseType>();
+		}
+	}
+};
+
+template<typename BaseType>
+struct TWeakPtrVariant : public UE::Core::Private::TPtrVariantBase<BaseType, false>
+{
+	using Super = UE::Core::Private::TPtrVariantBase<BaseType, false>;
+	using Super::Super;
+
+	TStrongPtrVariant<BaseType> Pin() const
+	{
+		if (Super::IsValid())
+		{
+			return ::Visit([](auto& StoredValue)
+				{
+					TStrongPtrVariant<BaseType> StrongPtrVariant;
+
+					using StoredValueType = std::decay_t<decltype(StoredValue)>;
+					if constexpr (TIsTWeakPtr_V<StoredValueType>)
+					{
+						StrongPtrVariant.PtrVariant.template Emplace<typename TStrongPtrVariant<BaseType>::SharedType>(StoredValue.Pin());
+					}
+					else
+					{
+						StrongPtrVariant.PtrVariant.template Emplace<typename TStrongPtrVariant<BaseType>::FObjectPtrWrapper>(StoredValue.ObjectPtr.Get(), StoredValue.CastedPtr);
+					}
+
+					return StrongPtrVariant;
+				}, Super::PtrVariant);
+		}
+		else
+		{
+			return TStrongPtrVariant<BaseType>();
+		}
+	}
+};
\ No newline at end of file
