Bump _MSC_VER check with updated Visual Studio 2022 versions

#rnx

diff --git a/Engine/Source/Runtime/Core/Public/Windows/WindowsPlatformCompilerSetup.h b/Engine/Source/Runtime/Core/Public/Windows/WindowsPlatformCompilerSetup.h
--- a/Engine/Source/Runtime/Core/Public/Windows/WindowsPlatformCompilerSetup.h
+++ b/Engine/Source/Runtime/Core/Public/Windows/WindowsPlatformCompilerSetup.h
@@ -11,14 +11,14 @@
 //
 // Future-proofing the min version check so we keep bumping it whenever we upgrade.
 //
-#if defined(_MSC_VER) && _MSC_VER > 1939 
+#if defined(_MSC_VER) && _MSC_VER > 1949 
 	#pragma message("Detected compiler newer than Visual Studio 2022, please update min version checking in WindowsPlatformCompilerSetup.h")
 #endif
 
 //
-// We require at least Visual Studio 2019 to compile
+// We require at least Visual Studio 2022 v17.4 to compile
 //
-static_assert(_MSC_VER >= 1929, "Visual Studio 2019 v16.11 or greater is required to build Unreal Engine.");
+static_assert(_MSC_VER >= 1934, "Visual Studio 2022 v17.4 or greater is required to build Unreal Engine.");
 
 //
 // Manually enable all warnings as errors, except ones that are explicitly skipped.
