- Added FAutoConsoleVariableDeprecated (and FAutoConsoleVariableShadow), which can be used to change the name of a cvar but leave a shadow cvar behind with the old name, so old ini files are not broken, and can control if the user gets warnings, asserts, etc when used
- Added support for skipping over Shadow cvars when all CVars are iterated over programmatically, so they are only used directly (like when one is in a .ini file)

diff --git a/Engine/Source/Runtime/Core/Private/HAL/ConsoleManager.cpp b/Engine/Source/Runtime/Core/Private/HAL/ConsoleManager.cpp
--- a/Engine/Source/Runtime/Core/Private/HAL/ConsoleManager.cpp
+++ b/Engine/Source/Runtime/Core/Private/HAL/ConsoleManager.cpp
@@ -25,12 +25,15 @@ ConsoleManager.cpp: console command handling
 #include "ProfilingDebugging/CsvProfiler.h"
 #include "HAL/FileManager.h"
 #include "Serialization/ArchiveCountMem.h"
+#include "Logging/MessageLog.h"
 
 #include <clocale>
 
 DEFINE_LOG_CATEGORY(LogConsoleResponse);
 DEFINE_LOG_CATEGORY_STATIC(LogConsoleManager, Log, All);
 
+#define LOCTEXT_NAMESPACE "ConsoleManager"
+
 class FOutputDevice;
 class UWorld;
 
@@ -2394,7 +2397,7 @@ void FConsoleManager::ForEachConsoleObjectThatStartsWith(const FConsoleObjectVis
 		const FString& Name = PairIt.Key();
 		IConsoleObject* CVar = PairIt.Value();
 
-		if(MatchPartialName(*Name, ThatStartsWith))
+		if(!CVar->IsShadowObject() && MatchPartialName(*Name, ThatStartsWith))
 		{
 			Visitor.Execute(*Name, CVar);
 		}
@@ -2417,6 +2420,10 @@ void FConsoleManager::ForEachConsoleObjectThatContains(const FConsoleObjectVisit
 		const FString& Name = PairIt.Key();
 		IConsoleObject* CVar = PairIt.Value();
 
+		if (CVar->IsShadowObject())
+		{
+			continue;
+		}
 		if (ContainsStringLength == 1)
 		{
 			if (MatchPartialName(*Name, ThatContains))
@@ -2541,13 +2548,18 @@ static void DumpObjects(const TMap<FString, IConsoleObject*>& ConsoleObjects, co
 				{
 					continue;
 				}
+				FString ShadowString;
+				if (Obj->IsShadowObject())
+				{
+					ShadowString = TEXT(" [SHADOW]");
+				}
 				if (bWriteToCSV)
 				{
-					MultiLogf(Log, CSV, TEXT("%s,%s,%s%s"), *Key, *CVar->GetString(), GetConsoleVariableSetByName(CVar->GetFlags()), *Help);
+					MultiLogf(Log, CSV, TEXT("%s%s,%s,%s%s"), *Key, *ShadowString, *CVar->GetString(), GetConsoleVariableSetByName(CVar->GetFlags()), *Help);
 				}
 				else
 				{
-					MultiLogf(Log, CSV, TEXT("%s = \"%s\"      LastSetBy: %s%s"), *Key, *CVar->GetString(), GetConsoleVariableSetByName(CVar->GetFlags()), *Help);
+					MultiLogf(Log, CSV, TEXT("%s%s = \"%s\"      LastSetBy: %s%s"), *Key, *ShadowString, *CVar->GetString(), GetConsoleVariableSetByName(CVar->GetFlags()), *Help);
 				}
 			}
 		}
@@ -2866,6 +2878,11 @@ bool FConsoleManager::ProcessUserConsoleInput(const TCHAR* InInput, FOutputDevic
 	return true;
 }
 
+void FConsoleManager::AddShadowConsoleObject(const TCHAR* Name, IConsoleObject* Obj)
+{
+	ConsoleObjects.Add(Name, Obj);
+}
+
 IConsoleObject* FConsoleManager::AddConsoleObject(const TCHAR* Name, IConsoleObject* Obj)
 {
 	check(Name);
@@ -3256,6 +3273,10 @@ void FConsoleManager::PreviewPlatformCVars(FName PlatformName, const FString& De
 
 	for (auto Pair : ConsoleObjects)
 	{
+		if (Pair.Value->IsShadowObject())
+		{
+			continue;
+		}
 		if (IConsoleVariable* CVar = Pair.Value->AsVariable())
 		{
 			// we want Preview but not Cheat
@@ -3296,6 +3317,11 @@ void FConsoleManager::ClearAllPlatformCVars(FName PlatformName, const FString& D
 	
 	for (auto Pair : ConsoleObjects)
 	{
+		if (Pair.Value->IsShadowObject())
+		{
+			continue;
+		}
+
 		if (IConsoleVariable* CVar = Pair.Value->AsVariable())
 		{
 			// clear any cached values for this key
@@ -4400,3 +4426,423 @@ static TAutoConsoleVariable<int32> CVarMobileSupportsGen4TAA(
 		 "0: Fallback to FXAA"
 		 "1: Support Desktop Gen4 TAA (default)"),
 	ECVF_ReadOnly | ECVF_RenderThreadSafe);
+
+
+
+class FConsoleVariableShadow : public IConsoleVariable
+{
+public:
+	FConsoleVariableShadow(const TCHAR* CVarToShadow, const TCHAR* InDeprecatedVersion, EShadowCVarBehavior InLookupBehavior, EShadowCVarBehavior InUsageBehavior)
+		: ShadowName(CVarToShadow)
+		, DeprecatedVersion(InDeprecatedVersion)
+		, LookupBehavior(InLookupBehavior)
+		, UsageBehavior(InUsageBehavior)
+		, bHasLooked(false)
+		, bHasMessagedForUsage(false)
+		, bHasMessagedEditorForUsage(false)
+		, RealVariable(nullptr)
+	{
+
+	}
+
+	virtual class IConsoleVariable* AsVariable() override
+	{
+		return this;
+	}
+
+	virtual bool IsShadowObject() const override
+	{
+		return true;
+	}
+
+	virtual bool IsVariableBool() const override
+	{
+		if (Bind())
+		{
+			return RealVariable->IsVariableBool();
+		}
+		return false;
+	}
+	virtual bool IsVariableInt() const override
+	{
+		if (Bind())
+		{
+			return RealVariable->IsVariableInt();
+		}
+		return false;
+	}
+	virtual bool IsVariableFloat() const override
+	{
+		if (Bind())
+		{
+			return RealVariable->IsVariableFloat();
+		}
+		return false;
+	}
+	virtual bool IsVariableString() const override
+	{
+		if (Bind())
+		{
+			return RealVariable->IsVariableString();
+		}
+		return false;
+	}
+
+	virtual class TConsoleVariableData<bool>* AsVariableBool() override
+	{
+		if (BindForUsage())
+		{
+			return RealVariable->AsVariableBool();
+		}
+		return nullptr;
+	}
+	virtual class TConsoleVariableData<int32>* AsVariableInt() override
+	{
+		if (BindForUsage())
+		{
+			return RealVariable->AsVariableInt();
+		}
+		return nullptr;
+	}
+	virtual class TConsoleVariableData<float>* AsVariableFloat() override
+	{
+		if (BindForUsage())
+		{
+			return RealVariable->AsVariableFloat();
+		}
+		return nullptr;
+	}
+	virtual class TConsoleVariableData<FString>* AsVariableString() override
+	{
+		if (BindForUsage())
+		{
+			return RealVariable->AsVariableString();
+		}
+		return nullptr;
+	}
+
+
+	virtual const TCHAR* GetHelp() const override
+	{
+		if (Bind())
+		{
+			return RealVariable->GetHelp();
+		}
+		return TEXT("");
+	}
+
+	virtual void SetHelp(const TCHAR* Value) override
+	{
+		if (Bind())
+		{
+			return RealVariable->SetHelp(Value);
+		}
+	}
+
+	virtual EConsoleVariableFlags GetFlags() const override
+	{
+		if (Bind())
+		{
+			return RealVariable->GetFlags();
+		}
+		return (EConsoleVariableFlags)0;
+
+	}
+
+	virtual void SetFlags(const EConsoleVariableFlags Value) override
+	{
+		if (BindForUsage())
+		{
+			RealVariable->SetFlags(Value);
+		}
+	}
+
+
+	virtual void Set(const TCHAR* InValue, EConsoleVariableFlags SetBy, FName Tag) override
+	{
+		if (BindForUsage())
+		{
+			RealVariable->Set(InValue, SetBy, Tag);
+		}
+	}
+
+	virtual void Unset(EConsoleVariableFlags SetBy, FName Tag) override
+	{
+		if (BindForUsage())
+		{
+			RealVariable->Unset(SetBy, Tag);
+		}
+	}
+
+	virtual bool GetBool() const override
+	{
+		if (BindForUsage())
+		{
+			return RealVariable->GetBool();
+		}
+		return false;
+	}
+
+	virtual int32 GetInt() const override
+	{
+		if (BindForUsage())
+		{
+			return RealVariable->GetInt();
+		}
+		return 0;
+	}
+
+	virtual float GetFloat() const override
+	{
+		if (BindForUsage())
+		{
+			return RealVariable->GetFloat();
+		}
+		return 0;
+	}
+
+	virtual FString GetString() const override
+	{
+		if (BindForUsage())
+		{
+			return RealVariable->GetString();
+		}
+		return FString();
+	}
+
+	virtual void SetOnChangedCallback(const FConsoleVariableDelegate& Callback) override
+	{
+		if (Bind())
+		{
+			RealVariable->SetOnChangedCallback(Callback);
+		}
+	}
+
+	virtual FConsoleVariableMulticastDelegate& OnChangedDelegate() override
+	{
+		if (Bind())
+		{
+			return RealVariable->OnChangedDelegate();
+		}
+		static FConsoleVariableMulticastDelegate Dummy;
+		return Dummy;
+	}
+
+	virtual FString GetDefaultValue() override
+	{
+		if (BindForUsage())
+		{
+			return RealVariable->GetDefaultValue();
+		}
+		return FString();
+	}
+
+	virtual void LogHistory(FOutputDevice& Ar) override
+	{
+		if (Bind())
+		{
+			RealVariable->LogHistory(Ar);
+		}
+	}
+
+	/**
+	 * Track memory used by history data
+	 */
+	virtual SIZE_T GetHistorySize() override
+	{
+		if (Bind())
+		{
+			return RealVariable->GetHistorySize();
+		}
+		return 0;
+	}
+
+	virtual void Release() override
+	{
+		// the real var should do it's own Release when it's time
+		RealVariable = nullptr;
+	}
+
+#if ALLOW_OTHER_PLATFORM_CONFIG
+
+	virtual TSharedPtr<IConsoleVariable> GetPlatformValueVariable(FName PlatformName, const FString& DeviceProfileName = FString())
+	{
+		if (BindForUsage())
+		{
+			return RealVariable->GetPlatformValueVariable(PlatformName, DeviceProfileName);
+		}
+		return nullptr;
+	}
+
+	virtual bool HasPlatformValueVariable(FName PlatformName, const FString& DeviceProfileName = FString()) override
+	{
+		if (Bind())
+		{
+			return RealVariable->HasPlatformValueVariable(PlatformName, DeviceProfileName);
+		}
+		return false;
+	}
+
+	virtual void ClearPlatformVariables(FName PlatformName = NAME_None)
+	{
+		if (Bind())
+		{
+			RealVariable->ClearPlatformVariables(PlatformName);
+		}
+	}
+
+#endif
+
+private:
+
+	void LogOrEditorMessage(const FText& Msg, bool bIsError) const
+	{
+		if (GIsEditor && !bHasMessagedEditorForUsage)
+		{
+			bHasMessagedEditorForUsage = true;
+
+			FMessageLog EditorErrors("EditorErrors");
+			TSharedRef<FTokenizedMessage> Message = EditorErrors.Message(bIsError ? EMessageSeverity::Error : EMessageSeverity::Warning);
+			Message->AddToken(FTextToken::Create(Msg));
+			EditorErrors.Notify();
+		}
+
+		// always spit to log
+		if (bIsError)
+		{
+			UE_LOG(LogConsoleManager, Error, TEXT("%s"), *Msg.ToString());
+		}
+		else
+		{
+			UE_LOG(LogConsoleManager, Warning, TEXT("%s"), *Msg.ToString());
+		}
+	}
+
+	bool Bind() const 
+	{
+		if (RealVariable == nullptr)
+		{
+			// if we looked but it wasn't found, then just return false
+			if (bHasLooked)
+			{
+				return false;
+			}
+			bHasLooked = true;
+
+			RealVariable = IConsoleManager::Get().FindConsoleVariable(*ShadowName, false);
+			if (RealVariable == nullptr)
+			{
+				if (LookupBehavior != EShadowCVarBehavior::NoMessaging)
+				{
+					FFormatNamedArguments Arguments;
+					Arguments.Add(TEXT("ThisName"), FText::FromString(IConsoleManager::Get().FindConsoleObjectName(this)));
+					Arguments.Add(TEXT("ShadowName"), FText::FromString(ShadowName));
+					
+					FText Message;
+					if (LookupBehavior == EShadowCVarBehavior::Assert)
+					{
+						Message = FText::Format(LOCTEXT("FailedShadowCVarLookup_Assert", "Attempted to delay-load real CVar '{ThisName}' for shadowed CVar '{ShadowName}' failed."), Arguments);
+					}
+					else
+					{
+						Message = FText::Format(LOCTEXT("FailedShadowCVarLookup", "Attempted to delay-load real CVar '{ThisName}' for shadowed CVar '{ShadowName}' failed. Uses of '{ThisName}' will do nothing."), Arguments);
+					}
+
+					switch (LookupBehavior)
+					{
+						case EShadowCVarBehavior::Warn:
+							LogOrEditorMessage(Message, false);
+							break;
+						case EShadowCVarBehavior::Error:
+							LogOrEditorMessage(Message, true);
+							break;
+						case EShadowCVarBehavior::Ensure:
+							ensureMsgf(false, TEXT("%s"), *Message.ToString());
+							break;
+						case EShadowCVarBehavior::Assert:
+							UE_LOG(LogConsoleManager, Fatal, TEXT("%s"), *Message.ToString());
+							break;
+					}
+				}
+				return false;
+			}
+		}
+
+		return true;
+	}
+
+	bool BindForUsage() const
+	{
+		if (!Bind())
+		{
+			return false;
+		}
+
+		if (!bHasMessagedForUsage)
+		{
+			if (UsageBehavior != EShadowCVarBehavior::NoMessaging)
+			{
+				FFormatNamedArguments Arguments;
+				Arguments.Add(TEXT("ThisName"), FText::FromString(IConsoleManager::Get().FindConsoleObjectName(this)));
+				Arguments.Add(TEXT("ShadowName"), FText::FromString(ShadowName));
+				Arguments.Add(TEXT("DeprecatedVersion"), FText::FromString(DeprecatedVersion));
+
+				FText Message;
+				if (DeprecatedVersion != TEXT(""))
+				{
+					Message = FText::Format(LOCTEXT("ShadowCVarUsage_Deprecated", "Using a deprecated (as of UE {DeprecatedVersion}) CVar: '{ThisName}'. It will be removed in the future. Change all uses to '{ShadowName}' instead."), Arguments);
+				}
+				else
+				{
+					Message = FText::Format(LOCTEXT("ShadowCVarUsage_Deprecated", "Using a shadowed CVar '{ThisName}'. It is recommended to change all uses to '{ShadowName}' instead."), Arguments);
+				}
+				switch (UsageBehavior)
+				{
+				case EShadowCVarBehavior::Warn:
+					LogOrEditorMessage(Message, false);
+					break;
+				case EShadowCVarBehavior::Error:
+					LogOrEditorMessage(Message, true);
+					break;
+				case EShadowCVarBehavior::Ensure:
+					// don't keep re-ensuring for this cvar
+					bHasMessagedForUsage = true;
+					ensureAlwaysMsgf(false, TEXT("%s"), *Message.ToString());
+					break;
+				case EShadowCVarBehavior::Assert:
+					UE_LOG(LogConsoleManager, Fatal, TEXT("%s"), *Message.ToString());
+					break;
+				}
+			}
+		}
+
+		return true;
+	}
+
+	FString ShadowName;
+	FString DeprecatedVersion;
+	EShadowCVarBehavior LookupBehavior;
+	EShadowCVarBehavior UsageBehavior;
+	mutable uint8 bHasLooked:1;
+	mutable uint8 bHasMessagedForUsage:1;
+	mutable uint8 bHasMessagedEditorForUsage:1;
+	mutable IConsoleVariable* RealVariable;
+};
+
+
+
+FAutoConsoleVariableShadow::FAutoConsoleVariableShadow(const TCHAR* Name, const TCHAR* CVarToShadow, EShadowCVarBehavior LookupFailureBehavior)
+{
+#if !NO_CVARS
+	GetManager().AddShadowConsoleObject(Name, new FConsoleVariableShadow(CVarToShadow, nullptr, LookupFailureBehavior, EShadowCVarBehavior::NoMessaging));
+#endif
+}
+
+FAutoConsoleVariableDeprecated::FAutoConsoleVariableDeprecated(const TCHAR* Name, const TCHAR* CVarToShadow, const TCHAR* DeprecatedAtVersion, EShadowCVarBehavior UsageBehavior,  EShadowCVarBehavior LookupFailureBehavior)
+{
+#if !NO_CVARS
+	GetManager().AddShadowConsoleObject(Name, new FConsoleVariableShadow(CVarToShadow, DeprecatedAtVersion, LookupFailureBehavior, UsageBehavior));
+#endif
+}
+
+#undef LOCTEXT_NAMESPACE
\ No newline at end of file
diff --git a/Engine/Source/Runtime/Core/Public/HAL/IConsoleManager.h b/Engine/Source/Runtime/Core/Public/HAL/IConsoleManager.h
--- a/Engine/Source/Runtime/Core/Public/HAL/IConsoleManager.h
+++ b/Engine/Source/Runtime/Core/Public/HAL/IConsoleManager.h
@@ -462,6 +462,15 @@ public:
 		return 0; 
 	}
 
+	/**
+	 * A shadow ConsoleObject is one that exists and can be used, but shouldn't be iterated over in anything that walks over every 
+	 * object and acts on it
+	 */
+	virtual bool IsShadowObject() const
+	{
+		return false;
+	}
+
 private: // -----------------------------------------
 
 #if TRACK_CONSOLE_FIND_COUNT
@@ -2038,5 +2047,59 @@ public:
 
 #endif
 
+
+enum class EShadowCVarBehavior : uint8
+{
+	// skip all messaging
+	NoMessaging,
+	// log a warning on first usage
+	Warn,
+	// log an error on first usage
+	Error,
+	// throw an ensure on first usage
+	Ensure,
+	// throw an assertion on first usage
+	Assert,
+};
+
+/**
+ * A class that can be used to shadow a CVar, where both are valid to use. To deprecate a CVar, it's suggested to use the FAutoConsoleVariableDeprecated class below
+ * 
+ * FAutoConsoleVariableShadow ShadowCVar(TEXT("r.newname"), TEXT("r.oldname"));
+ * 
+ */
+class FAutoConsoleVariableShadow
+{
+public:
+
+		/**
+		 * @param Name The name of this cvar
+		 * @param CVarToShadow The name of the existing cvar that this one will shadow
+		 * @param LookupFailureBehavior How to handle failure of looking up the CVarToShadow name - defaults to ensure
+		 */
+		CORE_API FAutoConsoleVariableShadow(const TCHAR* Name, const TCHAR* CVarToShadow, EShadowCVarBehavior LookupFailureBehavior=EShadowCVarBehavior::Ensure);
+};
+
+/**
+ * A class that can be used to deprecate a CVar - rename your existing CVar to the new name, then
+ * create a shadow of it with the old name with something like this:
+ *
+ * FAutoConsoleVariableDeprecated ShadowCVar(TEXT("r.newname"), TEXT("r.oldname"), TEXT("5.5"));
+ * FAutoConsoleVariableDeprecated ShadowCVar(TEXT("r.newname"), TEXT("r.oldname"), TEXT("5.5"), EShadowCVarBehavior::Assert);
+ *
+ */
+class FAutoConsoleVariableDeprecated
+{
+public:
+	/**
+	 * @param Name The name of this cvar
+	 * @param CVarToShadow The name of the existing cvar that this one will shadow
+	 * @param DeprecatedAtVersion THe first UE version when this was deprecated
+	 * @param UsageBehavior How to handle uses of the shadowed cvar by the deprecated name - defaults to ensure
+	 * @param LookupFailureBehavior How to handle failure of looking up the CVarToShadow name - defaults to ensure
+	 */
+	CORE_API FAutoConsoleVariableDeprecated(const TCHAR* Name, const TCHAR* CVarToShadow, const TCHAR* DeprecatedAtVersion, EShadowCVarBehavior UsageBehavior=EShadowCVarBehavior::Ensure, EShadowCVarBehavior LookupFailureBehavior = EShadowCVarBehavior::Ensure);
+};
+
 CORE_API DECLARE_LOG_CATEGORY_EXTERN(LogConsoleResponse, Log, All);
 
