Remove problematic log

This log can run during static destruction in applications that don't shutdown the localization system prior to main exiting, which can lead to a crash when GLog is in an invalid state

#rnx

diff --git a/Engine/Source/Runtime/Core/Private/Internationalization/TextKey.cpp b/Engine/Source/Runtime/Core/Private/Internationalization/TextKey.cpp
--- a/Engine/Source/Runtime/Core/Private/Internationalization/TextKey.cpp
+++ b/Engine/Source/Runtime/Core/Private/Internationalization/TextKey.cpp
@@ -173,7 +173,7 @@ private:
 				TotalNumElementsWasted += (SlabSizeInElements - Slab.NumElementsUsed);
 				FMemory::Free(Slab.Allocation);
 			}
-			UE_LOG(LogLocalization, Log, TEXT("FTextKey slab allocator allocated %d slabs (of %d elements) and wasted %d elements (%d bytes)"), Slabs.Num(), SlabSizeInElements, TotalNumElementsWasted, TotalNumElementsWasted * sizeof(TCHAR));
+			//UE_LOG(LogLocalization, Log, TEXT("FTextKey slab allocator allocated %d slabs (of %d elements) and wasted %d elements (%d bytes)"), Slabs.Num(), SlabSizeInElements, TotalNumElementsWasted, TotalNumElementsWasted * sizeof(TCHAR));
 		}
 
 		const TCHAR* Add(FStringView InStr)
