Added sub-tags for more granular localization LLM tracking

#rnx

diff --git a/Engine/Source/Runtime/Core/Private/Internationalization/ICUInternationalization.cpp b/Engine/Source/Runtime/Core/Private/Internationalization/ICUInternationalization.cpp
--- a/Engine/Source/Runtime/Core/Private/Internationalization/ICUInternationalization.cpp
+++ b/Engine/Source/Runtime/Core/Private/Internationalization/ICUInternationalization.cpp
@@ -54,7 +54,7 @@ namespace
 
 		static void* U_CALLCONV Malloc(const void* context, size_t size)
 		{
-			LLM_SCOPE(ELLMTag::Localization);
+			LLM_SCOPE_BYNAME(TEXT("Localization/ICU"));
 			void* Result = FMemory::Malloc(size);
 #if STATS
 			BytesInUseCount += FMemory::GetAllocSize(Result);
@@ -75,7 +75,7 @@ namespace
 
 		static void* U_CALLCONV Realloc(const void* context, void* mem, size_t size)
 		{
-			LLM_SCOPE(ELLMTag::Localization);
+			LLM_SCOPE_BYNAME(TEXT("Localization/ICU"));
 			return FMemory::Realloc(mem, size);
 		}
 
@@ -768,7 +768,7 @@ UDate FICUInternationalization::UEDateTimeToICUDate(const FDateTime& DateTime)
 
 UBool FICUInternationalization::OpenDataFile(const void* InContext, void** OutFileContext, void** OutContents, const char* InPath)
 {
-	LLM_SCOPE(ELLMTag::Localization);
+	LLM_SCOPE_BYNAME(TEXT("Localization/ICU"));
 
 	FICUInternationalization* This = (FICUInternationalization*)InContext;
 	check(This);
diff --git a/Engine/Source/Runtime/Core/Private/Internationalization/ICUInternationalization.h b/Engine/Source/Runtime/Core/Private/Internationalization/ICUInternationalization.h
--- a/Engine/Source/Runtime/Core/Private/Internationalization/ICUInternationalization.h
+++ b/Engine/Source/Runtime/Core/Private/Internationalization/ICUInternationalization.h
@@ -117,7 +117,6 @@ private:
 	struct FICUCachedFileData
 	{
 		FICUCachedFileData(const int64 FileSize);
-		FICUCachedFileData(void* ExistingBuffer);
 		FICUCachedFileData(FICUCachedFileData&& Source);
 		~FICUCachedFileData();
 
diff --git a/Engine/Source/Runtime/Core/Private/Internationalization/Internationalization.cpp b/Engine/Source/Runtime/Core/Private/Internationalization/Internationalization.cpp
--- a/Engine/Source/Runtime/Core/Private/Internationalization/Internationalization.cpp
+++ b/Engine/Source/Runtime/Core/Private/Internationalization/Internationalization.cpp
@@ -37,7 +37,6 @@ void FInternationalization::TearDown()
 
 FText FInternationalization::ForUseOnlyByLocMacroAndGraphNodeTextLiterals_CreateText(const TCHAR* InTextLiteral, const TCHAR* InNamespace, const TCHAR* InKey)
 {
-	LLM_SCOPE(ELLMTag::Localization);
 	return FTextCache::Get().FindOrCache(InTextLiteral, InNamespace, InKey);
 }
 
diff --git a/Engine/Source/Runtime/Core/Private/Internationalization/TextKey.cpp b/Engine/Source/Runtime/Core/Private/Internationalization/TextKey.cpp
--- a/Engine/Source/Runtime/Core/Private/Internationalization/TextKey.cpp
+++ b/Engine/Source/Runtime/Core/Private/Internationalization/TextKey.cpp
@@ -36,7 +36,7 @@ public:
 	void Shrink()
 	{
 		FScopeLock ScopeLock(&SynchronizationObject);
-		LLM_SCOPE(ELLMTag::Localization);
+		LLM_SCOPE_BYNAME(TEXT("Localization/TextKeys"));
 		KeysTable.Shrink();
 	}
 
@@ -131,7 +131,7 @@ private:
 		const TCHAR* StrPtr = KeysTable.FindRef(KeyData);
 		if (!StrPtr)
 		{
-			LLM_SCOPE(ELLMTag::Localization);
+			LLM_SCOPE_BYNAME(TEXT("Localization/TextKeys"));
 			StrPtr = StringAllocations.Add(KeyData.ToView());
 			KeysTable.Add(FKeyData(StrPtr, KeyData), StrPtr);
 		}
@@ -140,7 +140,7 @@ private:
 		const FString* StrPtr = KeysTable.Find(KeyData);
 		if (!StrPtr)
 		{
-			LLM_SCOPE(ELLMTag::Localization);
+			LLM_SCOPE_BYNAME(TEXT("Localization/TextKeys"));
 
 			// Need to copy the string here so we can reference its internal allocation as the key
 			FString StrCopy;
diff --git a/Engine/Source/Runtime/Core/Private/Internationalization/TextLocalizationManager.cpp b/Engine/Source/Runtime/Core/Private/Internationalization/TextLocalizationManager.cpp
--- a/Engine/Source/Runtime/Core/Private/Internationalization/TextLocalizationManager.cpp
+++ b/Engine/Source/Runtime/Core/Private/Internationalization/TextLocalizationManager.cpp
@@ -594,7 +594,7 @@ void FTextLocalizationManager::DumpMemoryInfo() const
 void FTextLocalizationManager::CompactDataStructures()
 {
 	TRACE_CPUPROFILER_EVENT_SCOPE(FTextLocalizationManager::CompactDataStructures);
-	LLM_SCOPE(ELLMTag::Localization);
+	LLM_SCOPE_BYNAME(TEXT("Localization/DisplayStrings"));
 
 	double StartTime = FPlatformTime::Seconds();
 	{
@@ -859,7 +859,7 @@ FTextConstDisplayStringPtr FTextLocalizationManager::FindDisplayString(const FTe
 FTextConstDisplayStringRef FTextLocalizationManager::GetDisplayString(const FTextKey& Namespace, const FTextKey& Key, const FString* const SourceStringPtr)
 {
 	TRACE_CPUPROFILER_EVENT_SCOPE(FTextLocalizationManager::GetDisplayString);
-	LLM_SCOPE(ELLMTag::Localization);
+	LLM_SCOPE_BYNAME(TEXT("Localization/DisplayStrings"));
 
 	auto GetEmptyDisplayString = []()
 	{
@@ -1069,7 +1069,7 @@ void FTextLocalizationManager::GetTextRevisions(const FTextId& InTextId, uint16&
 bool FTextLocalizationManager::AddDisplayString(const FTextDisplayStringRef& DisplayString, const FTextKey& Namespace, const FTextKey& Key)
 {
 	TRACE_CPUPROFILER_EVENT_SCOPE(FTextLocalizationManager::AddDisplayString);
-	LLM_SCOPE(ELLMTag::Localization);
+	LLM_SCOPE_BYNAME(TEXT("Localization/DisplayStrings"));
 
 	FScopeLock ScopeLock(&DisplayStringTableCS);
 
@@ -1254,6 +1254,7 @@ void FTextLocalizationManager::LoadLocalizationResourcesForPrioritizedCultures_S
 	FTextLocalizationResource LocalizedResource;
 	for (const TSharedPtr<ILocalizedTextSource>& LocalizedTextSource : AvailableTextSources)
 	{
+		LLM_SCOPE_BYNAME(TEXT("Localization/DisplayStrings"));
 		LocalizedTextSource->LoadLocalizedResources(FinalLocLoadFlags, PrioritizedCultureNames, NativeResource, LocalizedResource);
 	}
 
@@ -1314,7 +1315,10 @@ void FTextLocalizationManager::LoadLocalizationTargetsForPrioritizedCultures_Syn
 	{
 		UE_LOG(LogTextLocalizationManager, Verbose, TEXT("Loading LocRes data from '%s'"), *LocalizationTargetPath);
 	}
-	LocResTextSource->LoadLocalizedResourcesFromPaths(TArrayView<FString>(), LocalizationTargetPaths, TArrayView<FString>(), LocLoadFlags, PrioritizedCultureNames, UnusedNativeResource, LocalizedResource);
+	{
+		LLM_SCOPE_BYNAME(TEXT("Localization/DisplayStrings"));
+		LocResTextSource->LoadLocalizedResourcesFromPaths(TArrayView<FString>(), LocalizationTargetPaths, TArrayView<FString>(), LocLoadFlags, PrioritizedCultureNames, UnusedNativeResource, LocalizedResource);
+	}
 
 	// Allow any higher priority text sources to override the additional text loaded (eg, to allow polyglot hot-fixes to take priority)
 	// Note: If any text sources don't support dynamic queries, then we must do a much slower full refresh instead :(
@@ -1331,6 +1335,7 @@ void FTextLocalizationManager::LoadLocalizationTargetsForPrioritizedCultures_Syn
 				continue;
 			}
 
+			LLM_SCOPE_BYNAME(TEXT("Localization/DisplayStrings"));
 			for (const FTextId& NewTextId : NewTextIds)
 			{
 				if (LocalizedTextSource->QueryLocalizedResource(LocLoadFlags, PrioritizedCultureNames, NewTextId, UnusedNativeResource, LocalizedResource) == EQueryLocalizedResourceResult::NotImplemented)
@@ -1505,7 +1510,7 @@ void FTextLocalizationManager::QueueAsyncTask(TUniqueFunction<void()>&& Task)
 void FTextLocalizationManager::UpdateFromNative(FTextLocalizationResource&& TextLocalizationResource, const bool bDirtyTextRevision)
 {
 	TRACE_CPUPROFILER_EVENT_SCOPE(FTextLocalizationManager::UpdateFromNative);
-	LLM_SCOPE(ELLMTag::Localization);
+	LLM_SCOPE_BYNAME(TEXT("Localization/DisplayStrings"));
 
 	// Nothing to do?
 	if (!FTextLocalizationManager::IsDisplayStringSupportEnabled())
@@ -1619,7 +1624,7 @@ void FTextLocalizationManager::UpdateFromNative(FTextLocalizationResource&& Text
 void FTextLocalizationManager::UpdateFromLocalizations(FTextLocalizationResource&& TextLocalizationResource, const bool bDirtyTextRevision)
 {
 	TRACE_CPUPROFILER_EVENT_SCOPE(FTextLocalizationManager::UpdateFromLocalizations);
-	LLM_SCOPE(ELLMTag::Localization);
+	LLM_SCOPE_BYNAME(TEXT("Localization/DisplayStrings"));
 
 	// Nothing to do?
 	if (!FTextLocalizationManager::IsDisplayStringSupportEnabled())
@@ -1741,7 +1746,7 @@ void FTextLocalizationManager::UpdateFromLocalizations(FTextLocalizationResource
 void FTextLocalizationManager::DirtyLocalRevisionForTextId(const FTextId& InTextId)
 {
 	TRACE_CPUPROFILER_EVENT_SCOPE(FTextLocalizationManager::DirtyLocalRevisionForTextId);
-	LLM_SCOPE(ELLMTag::Localization);
+	LLM_SCOPE_BYNAME(TEXT("Localization/DisplayStrings"));
 
 	FWriteScopeLock ScopeLock(TextRevisionRW);
 
@@ -1759,7 +1764,7 @@ void FTextLocalizationManager::DirtyLocalRevisionForTextId(const FTextId& InText
 void FTextLocalizationManager::DirtyTextRevision()
 {
 	TRACE_CPUPROFILER_EVENT_SCOPE(FTextLocalizationManager::DirtyTextRevision);
-	LLM_SCOPE(ELLMTag::Localization);
+	LLM_SCOPE_BYNAME(TEXT("Localization/DisplayStrings"));
 
 	// Lock while updating the data
 	{
