Updated text display strings to use TRefCountPtr

This is wrapped in a minimal TSharedRef/TSharedPtr interface for backwards compatibility, and can be controlled by UE_TEXT_DISPLAYSTRING_USE_REFCOUNTPTR

This saves 8-bytes per-display string.

diff --git a/Engine/Source/Runtime/Core/Public/Internationalization/LocalizedTextSourceTypes.h b/Engine/Source/Runtime/Core/Public/Internationalization/LocalizedTextSourceTypes.h
--- a/Engine/Source/Runtime/Core/Public/Internationalization/LocalizedTextSourceTypes.h
+++ b/Engine/Source/Runtime/Core/Public/Internationalization/LocalizedTextSourceTypes.h
@@ -4,6 +4,7 @@
 
 #include "CoreTypes.h"
 #include "Misc/EnumClassFlags.h"
+#include "Templates/RefCounting.h"
 #include "Templates/SharedPointer.h"
 
 /**
@@ -73,12 +74,302 @@ struct ELocalizedTextSourcePriority
 	};
 };
 
-typedef TSharedRef<FString, ESPMode::ThreadSafe> FTextDisplayStringRef;
-typedef TSharedPtr<FString, ESPMode::ThreadSafe> FTextDisplayStringPtr;
-typedef TSharedRef<const FString, ESPMode::ThreadSafe> FTextConstDisplayStringRef;
-typedef TSharedPtr<const FString, ESPMode::ThreadSafe> FTextConstDisplayStringPtr;
+#ifndef UE_TEXT_DISPLAYSTRING_USE_REFCOUNTPTR
+	#define UE_TEXT_DISPLAYSTRING_USE_REFCOUNTPTR (1)
+#endif
 
-inline FTextDisplayStringRef MakeTextDisplayString(FString&& InDisplayString)
+#if UE_TEXT_DISPLAYSTRING_USE_REFCOUNTPTR
+
+namespace UE::Text::Private
+{
+
+class FRefCountedDisplayString
+{
+public:
+	explicit FRefCountedDisplayString(FString&& InDisplayString)
+		: DisplayString(MoveTemp(InDisplayString))
+		, NumRefs(0)
+	{
+	}
+
+	~FRefCountedDisplayString()
+	{
+		check(NumRefs.GetValue() == 0);
+	}
+
+	FRefCountedDisplayString(const FRefCountedDisplayString&) = delete;
+	FRefCountedDisplayString& operator=(const FRefCountedDisplayString&) = delete;
+
+	uint32 AddRef() const
+	{
+		return uint32(NumRefs.Increment());
+	}
+
+	uint32 Release() const
+	{
+		uint32 Refs = uint32(NumRefs.Decrement());
+		if (Refs == 0)
+		{
+			delete this;
+		}
+		return Refs;
+	}
+
+	uint32 GetRefCount() const
+	{
+		return uint32(NumRefs.GetValue());
+	}
+
+	[[nodiscard]] FString& Private_GetDisplayString()
+	{
+		return DisplayString;
+	}
+
+private:
+	FString DisplayString;
+	mutable FThreadSafeCounter NumRefs;
+};
+
+/**
+ * Wrapper to give TRefCountPtr a minimal TSharedRef/TSharedPtr interface for backwards compatibility with code that was already using FTextDisplayStringRef/FTextDisplayStringPtr
+ */
+template <typename ObjectType>
+class TDisplayStringPtrBase
+{
+	static_assert(TPointerIsConvertibleFromTo<ObjectType, const FString>::Value, "TDisplayStringPtrBase can only be constructed with FString types");
+
+public:
+	TDisplayStringPtrBase() = default;
+
+	explicit TDisplayStringPtrBase(const TRefCountPtr<FRefCountedDisplayString>& InDisplayStringPtr)
+		: DisplayStringPtr(InDisplayStringPtr)
+	{
+	}
+
+	[[nodiscard]] explicit operator bool() const
+	{
+		return IsValid();
+	}
+
+	[[nodiscard]] bool IsValid() const
+	{
+		return DisplayStringPtr.IsValid();
+	}
+
+	[[nodiscard]] ObjectType& operator*() const
+	{
+		return GetDisplayString();
+	}
+
+	[[nodiscard]] ObjectType* operator->() const
+	{
+		return &GetDisplayString();
+	}
+
+	[[nodiscard]] const TRefCountPtr<FRefCountedDisplayString>& Private_GetDisplayStringPtr() const
+	{
+		return DisplayStringPtr;
+	}
+
+protected:
+	[[nodiscard]] FString& GetDisplayString() const
+	{
+		check(IsValid());
+		return DisplayStringPtr.GetReference()->Private_GetDisplayString();
+	}
+
+	TRefCountPtr<FRefCountedDisplayString> DisplayStringPtr;
+};
+
+/**
+ * Wrapper to give TRefCountPtr a minimal TSharedRef interface for backwards compatibility with code that was already using FTextDisplayStringRef
+ */
+template <typename ObjectType>
+class TDisplayStringRef : public TDisplayStringPtrBase<ObjectType>
+{
+public:
+	TDisplayStringRef() = default;
+
+	explicit TDisplayStringRef(const TRefCountPtr<FRefCountedDisplayString>& InDisplayStringPtr)
+		: TDisplayStringPtrBase<ObjectType>(InDisplayStringPtr)
+	{
+		check(this->IsValid());
+	}
+
+	template <typename OtherType, typename = decltype(ImplicitConv<ObjectType*>((OtherType*)nullptr))>
+	TDisplayStringRef(const TDisplayStringRef<OtherType>& InOther)
+		: TDisplayStringPtrBase<ObjectType>(InOther.Private_GetDisplayStringPtr())
+	{
+	}
+
+	template <typename OtherType, typename = decltype(ImplicitConv<ObjectType*>((OtherType*)nullptr))>
+	TDisplayStringRef& operator=(const TDisplayStringRef<OtherType>& InOther)
+	{
+		if (this->DisplayStringPtr != InOther.Private_GetDisplayStringPtr())
+		{
+			this->DisplayStringPtr = InOther.Private_GetDisplayStringPtr();
+		}
+		return *this;
+	}
+
+	[[nodiscard]] ObjectType& Get() const
+	{
+		return this->GetDisplayString();
+	}
+};
+
+/**
+ * Wrapper to give TRefCountPtr a minimal TSharedPtr interface for backwards compatibility with code that was already using FTextDisplayStringPtr
+ */
+template <typename ObjectType>
+class TDisplayStringPtr : public TDisplayStringPtrBase<ObjectType>
+{
+public:
+	TDisplayStringPtr() = default;
+
+	TDisplayStringPtr(TYPE_OF_NULLPTR)
+		: TDisplayStringPtr()
+	{
+	}
+
+	explicit TDisplayStringPtr(const TRefCountPtr<FRefCountedDisplayString>& InDisplayStringPtr)
+		: TDisplayStringPtrBase<ObjectType>(InDisplayStringPtr)
+	{
+	}
+
+	template <typename OtherType, typename = decltype(ImplicitConv<ObjectType*>((OtherType*)nullptr))>
+	TDisplayStringPtr(const TDisplayStringPtr<OtherType>& InOther)
+		: TDisplayStringPtrBase<ObjectType>(InOther.Private_GetDisplayStringPtr())
+	{
+	}
+
+	template <typename OtherType, typename = decltype(ImplicitConv<ObjectType*>((OtherType*)nullptr))>
+	TDisplayStringPtr(const TDisplayStringRef<OtherType>& InOther)
+		: TDisplayStringPtrBase<ObjectType>(InOther.Private_GetDisplayStringPtr())
+	{
+	}
+
+	template <typename OtherType, typename = decltype(ImplicitConv<ObjectType*>((OtherType*)nullptr))>
+	TDisplayStringPtr& operator=(const TDisplayStringPtr<OtherType>& InOther)
+	{
+		if (this->DisplayStringPtr != InOther.Private_GetDisplayStringPtr())
+		{
+			this->DisplayStringPtr = InOther.Private_GetDisplayStringPtr();
+		}
+		return *this;
+	}
+
+	template <typename OtherType, typename = decltype(ImplicitConv<ObjectType*>((OtherType*)nullptr))>
+	TDisplayStringPtr& operator=(const TDisplayStringRef<OtherType>& InOther)
+	{
+		if (this->DisplayStringPtr != InOther.Private_GetDisplayStringPtr())
+		{
+			this->DisplayStringPtr = InOther.Private_GetDisplayStringPtr();
+		}
+		return *this;
+	}
+
+	[[nodiscard]] ObjectType* Get() const
+	{
+		return this->IsValid()
+			? &this->GetDisplayString()
+			: nullptr;
+	}
+
+	void Reset()
+	{
+		this->DisplayStringPtr = TRefCountPtr<FRefCountedDisplayString>();
+	}
+
+	[[nodiscard]] TDisplayStringRef<ObjectType> ToSharedRef() const
+	{
+		check(this->IsValid());
+		return TDisplayStringRef<ObjectType>(this->DisplayStringPtr);
+	}
+};
+
+} // namespace UE::Text::Private
+
+template <typename ObjectTypeA, typename ObjectTypeB>
+[[nodiscard]] bool operator==(const UE::Text::Private::TDisplayStringRef<ObjectTypeA>& A, const UE::Text::Private::TDisplayStringRef<ObjectTypeB>& B)
+{
+	return &A.Get() == &B.Get();
+}
+
+template <typename ObjectTypeA, typename ObjectTypeB>
+[[nodiscard]] bool operator==(const UE::Text::Private::TDisplayStringRef<ObjectTypeA>& A, const UE::Text::Private::TDisplayStringPtr<ObjectTypeB>& B)
+{
+	return &A.Get() == B.Get();
+}
+
+template <typename ObjectTypeA, typename ObjectTypeB>
+[[nodiscard]] bool operator==(const UE::Text::Private::TDisplayStringPtr<ObjectTypeA>& A, const UE::Text::Private::TDisplayStringRef<ObjectTypeB>& B)
+{
+	return A.Get() == &B.Get();
+}
+
+template <typename ObjectTypeA, typename ObjectTypeB>
+[[nodiscard]] bool operator==(const UE::Text::Private::TDisplayStringPtr<ObjectTypeA>& A, const UE::Text::Private::TDisplayStringPtr<ObjectTypeB>& B)
+{
+	return A.Get() == B.Get();
+}
+
+template <typename ObjectTypeA, typename ObjectTypeB>
+[[nodiscard]] bool operator!=(const UE::Text::Private::TDisplayStringRef<ObjectTypeA>& A, const UE::Text::Private::TDisplayStringRef<ObjectTypeB>& B)
+{
+	return &A.Get() != &B.Get();
+}
+
+template <typename ObjectTypeA, typename ObjectTypeB>
+[[nodiscard]] bool operator!=(const UE::Text::Private::TDisplayStringRef<ObjectTypeA>& A, const UE::Text::Private::TDisplayStringPtr<ObjectTypeB>& B)
+{
+	return &A.Get() != B.Get();
+}
+
+template <typename ObjectTypeA, typename ObjectTypeB>
+[[nodiscard]] bool operator!=(const UE::Text::Private::TDisplayStringPtr<ObjectTypeA>& A, const UE::Text::Private::TDisplayStringRef<ObjectTypeB>& B)
+{
+	return A.Get() != &B.Get();
+}
+
+template <typename ObjectTypeA, typename ObjectTypeB>
+[[nodiscard]] bool operator!=(const UE::Text::Private::TDisplayStringPtr<ObjectTypeA>& A, const UE::Text::Private::TDisplayStringPtr<ObjectTypeB>& B)
+{
+	return A.Get() != B.Get();
+}
+
+template <typename ObjectType>
+[[nodiscard]] uint32 GetTypeHash(const UE::Text::Private::TDisplayStringRef<ObjectType>& A)
+{
+	return ::PointerHash(&A.Get());
+}
+
+template <typename ObjectType>
+[[nodiscard]] uint32 GetTypeHash(const UE::Text::Private::TDisplayStringPtr<ObjectType>& A)
+{
+	return ::PointerHash(A.Get());
+}
+
+using FTextDisplayStringRef = UE::Text::Private::TDisplayStringRef<FString>;
+using FTextDisplayStringPtr = UE::Text::Private::TDisplayStringPtr<FString>;
+using FTextConstDisplayStringRef = UE::Text::Private::TDisplayStringRef<const FString>;
+using FTextConstDisplayStringPtr = UE::Text::Private::TDisplayStringPtr<const FString>;
+
+[[nodiscard]] inline FTextDisplayStringRef MakeTextDisplayString(FString&& InDisplayString)
+{
+	return FTextDisplayStringRef(MakeRefCount<UE::Text::Private::FRefCountedDisplayString>(MoveTemp(InDisplayString)));
+}
+
+#else	// UE_TEXT_DISPLAYSTRING_USE_REFCOUNTPTR
+
+using FTextDisplayStringRef = TSharedRef<FString, ESPMode::ThreadSafe>;
+using FTextDisplayStringPtr = TSharedPtr<FString, ESPMode::ThreadSafe>;
+using FTextConstDisplayStringRef = TSharedRef<const FString, ESPMode::ThreadSafe>;
+using FTextConstDisplayStringPtr = TSharedPtr<const FString, ESPMode::ThreadSafe>;
+
+[[nodiscard]] inline FTextDisplayStringRef MakeTextDisplayString(FString&& InDisplayString)
 {
 	return MakeShared<FString, ESPMode::ThreadSafe>(MoveTemp(InDisplayString));
 }
+
+#endif	// UE_TEXT_DISPLAYSTRING_USE_REFCOUNTPTR
