Disable TRefCountPtr display strings again until a rollback issue can be fixed

#rnx

diff --git a/Engine/Source/Runtime/Core/Public/Internationalization/LocalizedTextSourceTypes.h b/Engine/Source/Runtime/Core/Public/Internationalization/LocalizedTextSourceTypes.h
--- a/Engine/Source/Runtime/Core/Public/Internationalization/LocalizedTextSourceTypes.h
+++ b/Engine/Source/Runtime/Core/Public/Internationalization/LocalizedTextSourceTypes.h
@@ -75,7 +75,7 @@ struct ELocalizedTextSourcePriority
 };
 
 #ifndef UE_TEXT_DISPLAYSTRING_USE_REFCOUNTPTR
-	#define UE_TEXT_DISPLAYSTRING_USE_REFCOUNTPTR (1)
+	#define UE_TEXT_DISPLAYSTRING_USE_REFCOUNTPTR (0)
 #endif
 
 #if UE_TEXT_DISPLAYSTRING_USE_REFCOUNTPTR
