remove the deprecated FCustomChunk functionality.

 - it has been superceded by named chunks.

#rnx

diff --git a/Engine/Source/Runtime/Core/Private/GenericPlatform/GenericPlatformMisc.cpp b/Engine/Source/Runtime/Core/Private/GenericPlatform/GenericPlatformMisc.cpp
--- a/Engine/Source/Runtime/Core/Private/GenericPlatform/GenericPlatformMisc.cpp
+++ b/Engine/Source/Runtime/Core/Private/GenericPlatform/GenericPlatformMisc.cpp
@@ -1909,35 +1909,6 @@ bool FGenericPlatformMisc::RequestDeviceCheckToken(TFunction<void(const TArray<u
 	return false;
 }
 
-TArray<FCustomChunk> FGenericPlatformMisc::GetOnDemandChunksForPakchunkIndices(const TArray<int32>& PakchunkIndices)
-{
-	return TArray<FCustomChunk>();
-}
-
-TArray<FCustomChunk> FGenericPlatformMisc::GetAllOnDemandChunks()
-{
-	return TArray<FCustomChunk>();
-}
-
-TArray<FCustomChunk> FGenericPlatformMisc::GetAllLanguageChunks()
-{
-	return TArray<FCustomChunk>();
-}
-
-TArray<FCustomChunk> FGenericPlatformMisc::GetCustomChunksByType(ECustomChunkType DesiredChunkType)
-{
-	PRAGMA_DISABLE_DEPRECATION_WARNINGS
-	if (DesiredChunkType == ECustomChunkType::OnDemandChunk)
-	{
-		return GetAllOnDemandChunks();
-	}
-	else
-	{
-		return GetAllLanguageChunks();
-	}
-	PRAGMA_ENABLE_DEPRECATION_WARNINGS
-}
-
 FString FGenericPlatformMisc::LoadTextFileFromPlatformPackage(const FString& RelativePath)
 {
 	FString Path = RootDir() / RelativePath;
diff --git a/Engine/Source/Runtime/Core/Public/GenericPlatform/GenericPlatformChunkInstall.h b/Engine/Source/Runtime/Core/Public/GenericPlatform/GenericPlatformChunkInstall.h
--- a/Engine/Source/Runtime/Core/Public/GenericPlatform/GenericPlatformChunkInstall.h
+++ b/Engine/Source/Runtime/Core/Public/GenericPlatform/GenericPlatformChunkInstall.h
@@ -111,24 +111,6 @@ DECLARE_MULTICAST_DELEGATE_TwoParams(FPlatformNamedChunkInstallMultiDelegate, FN
 DECLARE_DELEGATE_OneParam(FPlatformNamedChunkCompleteDelegate, const FNamedChunkCompleteCallbackParam&);
 DECLARE_MULTICAST_DELEGATE_OneParam(FPlatformNamedChunkCompleteMultiDelegate, const FNamedChunkCompleteCallbackParam&);
 
-enum class ECustomChunkType : uint8
-{
-	OnDemandChunk,
-	LanguageChunk
-};
-
-struct FCustomChunk
-{
-	FString ChunkTag;
-	FString ChunkTag2;
-	uint32	ChunkID;
-	ECustomChunkType ChunkType;
-
-	FCustomChunk(FString InTag, uint32 InID, ECustomChunkType InChunkType, FString InTag2 = TEXT("")) :
-		ChunkTag(InTag), ChunkTag2(InTag2), ChunkID(InID), ChunkType(InChunkType)
-	{}
-};
-
 struct FCustomChunkMapping
 {
 	enum class CustomChunkMappingType : uint8
@@ -242,10 +224,6 @@ public:
 	virtual void RemoveChunkInstallDelegate( FDelegateHandle Delegate ) = 0;
 
 
-	UE_DEPRECATED(5.2, "Call GetNamedChunksByType instead")
-	virtual TArray<FCustomChunk> GetCustomChunksByType(ECustomChunkType DesiredChunkType) = 0;
-
-
 	/**
 	 * Check whether current platform supports chunk installation by name
 	 * @return				whether Intelligent Install is supported
@@ -480,11 +458,6 @@ public:
 		InstallDelegate.Remove(Delegate);
 	}
 
-	virtual TArray<FCustomChunk> GetCustomChunksByType(ECustomChunkType DesiredChunkType) override
-	{
-		return TArray<FCustomChunk>();
-	}
-
 	virtual bool SupportsNamedChunkInstall() const override
 	{
 		return false;
diff --git a/Engine/Source/Runtime/Core/Public/GenericPlatform/GenericPlatformMisc.h b/Engine/Source/Runtime/Core/Public/GenericPlatform/GenericPlatformMisc.h
--- a/Engine/Source/Runtime/Core/Public/GenericPlatform/GenericPlatformMisc.h
+++ b/Engine/Source/Runtime/Core/Public/GenericPlatform/GenericPlatformMisc.h
@@ -26,14 +26,11 @@ class IInstallBundleManager;
 class IPlatformChunkInstall;
 class IPlatformCompression;
 class IPlatformHostCommunication;
-struct FCustomChunk;
 struct FDefaultDelegateUserPolicy;
 struct FGenericCrashContext;
 struct FGenericMemoryWarningContext;
 struct FGuid;
 
-enum class ECustomChunkType : uint8;
-
 template <typename FuncType>
 class TFunction;
 
@@ -1774,18 +1771,6 @@ public:
 
 	static CORE_API bool RequestDeviceCheckToken(TFunction<void(const TArray<uint8>&)> QuerySucceededFunc, TFunction<void(const FString&, const FString&)> QueryFailedFunc);
 
-	UE_DEPRECATED(5.1, "Use named chunks instead")
-	static CORE_API TArray<FCustomChunk> GetOnDemandChunksForPakchunkIndices(const TArray<int32>& PakchunkIndices);
-
-	UE_DEPRECATED(5.1, "Use IPlatformChunkInstall::GetNamedChunksByType instead")
-	static CORE_API TArray<FCustomChunk> GetAllOnDemandChunks();
-
-	UE_DEPRECATED(5.1, "Use IPlatformChunkInstall::GetNamedChunksByType instead")
-	static CORE_API TArray<FCustomChunk> GetAllLanguageChunks();
-
-	UE_DEPRECATED(5.1, "Use IPlatformChunkInstall::GetNamedChunksByType instead")
-	static CORE_API TArray<FCustomChunk> GetCustomChunksByType(ECustomChunkType DesiredChunkType);
-
 	/**
 	 * Loads a text file relative to the package root on platforms that distribute apps in package formats.
 	 * For other platforms, the path is relative to the root directory.
