Pass FDevSystemGuidRegistry singleton via a function to force its initialization before it is used by global static variables and prevent static initialization order fiasco

#rnx

diff --git a/Engine/Source/Runtime/Core/Private/UObject/DevObjectVersion.cpp b/Engine/Source/Runtime/Core/Private/UObject/DevObjectVersion.cpp
--- a/Engine/Source/Runtime/Core/Private/UObject/DevObjectVersion.cpp
+++ b/Engine/Source/Runtime/Core/Private/UObject/DevObjectVersion.cpp
@@ -68,11 +68,17 @@ private:
 	TMap<FGuid, FRegisteredGuid> RegisteredSystemGuids;
 };
 
-static FDevSystemGuidRegistry GSystemGuidRegistry;
+static FDevSystemGuidRegistry& GetSystemGuidRegistry()
+{
+	// Pass FDevSystemGuidRegistry singleton via a function to force its initialization
+	// before it is used by global static variables and prevent static initialization order fiasco
+	static FDevSystemGuidRegistry GSystemGuidRegistry;
+	return GSystemGuidRegistry;
+}
 
 FDevSystemGuidRegistration::FDevSystemGuidRegistration(const TMap<FGuid, FGuid>& SystemGuids)
 {
-	GSystemGuidRegistry.RegisterSystemGuids(SystemGuids);
+	GetSystemGuidRegistry().RegisterSystemGuids(SystemGuids);
 }
 
 void FDevSystemGuidRegistry::RegisterSystemGuid(FGuid System, FGuid Guid)
@@ -141,7 +147,7 @@ const FDevSystemGuids& FDevSystemGuids::Get()
 
 FGuid FDevSystemGuids::GetSystemGuid(FGuid System)
 {
-	return GSystemGuidRegistry.GetSystemGuid(System);
+	return GetSystemGuidRegistry().GetSystemGuid(System);
 }
 
 void FDevVersionRegistration::RecordDevVersion(FGuid Key)
