Immediately resolve string table reference display strings once the referenced table is available

This allows the natvis to show the correct result once the string table entry is available, without needing something else (like FText::ToString) to trigger the display string resolution.

#rnx

diff --git a/Engine/Source/Runtime/Core/Private/Internationalization/TextHistory.cpp b/Engine/Source/Runtime/Core/Private/Internationalization/TextHistory.cpp
--- a/Engine/Source/Runtime/Core/Private/Internationalization/TextHistory.cpp
+++ b/Engine/Source/Runtime/Core/Private/Internationalization/TextHistory.cpp
@@ -2552,14 +2552,14 @@ void FTextHistory_StringTableEntry::FStringTableReferenceData::Initialize(FName
 	{
 		// No loading attempt
 		LoadingPhase = EStringTableLoadingPhase::Loaded;
-		ResolveStringTableEntry();
+		ResolveDisplayString();
 	}
 	else if (InLoadingPolicy == EStringTableLoadingPolicy::FindOrFullyLoad && IStringTableEngineBridge::CanFindOrLoadStringTableAsset())
 	{
 		// Forced synchronous load
 		LoadingPhase = EStringTableLoadingPhase::Loaded;
 		IStringTableEngineBridge::FullyLoadStringTableAsset(TableId);
-		ResolveStringTableEntry();
+		ResolveDisplayString();
 	}
 	else
 	{
@@ -2721,7 +2721,7 @@ void FTextHistory_StringTableEntry::FStringTableReferenceData::ConditionalBeginA
 		}
 		This->LoadingPhase = EStringTableLoadingPhase::Loaded;
 
-		This->ResolveStringTableEntry();
+		This->ResolveDisplayString();
 	});
 }
 
