Fix the check ABI mechanism in AutoRTFM to use major.minor.patch versioning instead of the defunct constant checks.

diff --git a/Engine/Source/Runtime/Core/Private/AutoRTFM/API.cpp b/Engine/Source/Runtime/Core/Private/AutoRTFM/API.cpp
--- a/Engine/Source/Runtime/Core/Private/AutoRTFM/API.cpp
+++ b/Engine/Source/Runtime/Core/Private/AutoRTFM/API.cpp
@@ -318,19 +318,9 @@ extern "C" UE_AUTORTFM_NOAUTORTFM void autortfm_check_abi(void* const Ptr, const
 {
     struct FConstants final
     {
-        const size_t LogLineBytes = Constants::LogLineBytes;
-        const size_t LineBytes = Constants::LineBytes;
-        const size_t LineTableSize = Constants::LineTableSize;
-        const size_t Offset_Context_CurrentTransaction = Constants::Offset_Context_CurrentTransaction;
-        const size_t Offset_Context_LineTable = Constants::Offset_Context_LineTable;
-        const size_t Offset_Context_Status = Constants::Offset_Context_Status;
-        const size_t LogSize_LineEntry = Constants::LogSize_LineEntry;
-        const size_t Size_LineEntry = Constants::Size_LineEntry;
-        const size_t Offset_LineEntry_LogicalLine = Constants::Offset_LineEntry_LogicalLine;
-        const size_t Offset_LineEntry_ActiveLine = Constants::Offset_LineEntry_ActiveLine;
-        const size_t Offset_LineEntry_LoggingTransaction = Constants::Offset_LineEntry_LoggingTransaction;
-        const size_t Offset_LineEntry_AccessMask = Constants::Offset_LineEntry_AccessMask;
-        const uint32_t Context_Status_OnTrack = Constants::Context_Status_OnTrack;
+		const uint32_t Major = AutoRTFM::Constants::Major;
+		const uint32_t Minor = AutoRTFM::Constants::Minor;
+		const uint32_t Patch = AutoRTFM::Constants::Patch;
 
 		// This is messy - but we want to do comparisons but without comparing any padding bytes.
 		// Before C++20 we cannot use a default created operator== and operator!=, so we use this
@@ -338,7 +328,7 @@ extern "C" UE_AUTORTFM_NOAUTORTFM void autortfm_check_abi(void* const Ptr, const
 	private:
 		auto Tied() const
 		{
-			return Tie(LogLineBytes, LineBytes, LineTableSize, Offset_Context_CurrentTransaction, Offset_Context_LineTable, Offset_Context_Status, LogSize_LineEntry, Size_LineEntry, Offset_LineEntry_LogicalLine, Offset_LineEntry_ActiveLine, Offset_LineEntry_LoggingTransaction, Offset_LineEntry_AccessMask, Context_Status_OnTrack);
+			return Tie(Major, Minor, Patch);
 		}
 
 	public:
diff --git a/Engine/Source/Runtime/Core/Private/AutoRTFM/ContextStatus.h b/Engine/Source/Runtime/Core/Private/AutoRTFM/ContextStatus.h
--- a/Engine/Source/Runtime/Core/Private/AutoRTFM/ContextStatus.h
+++ b/Engine/Source/Runtime/Core/Private/AutoRTFM/ContextStatus.h
@@ -7,8 +7,6 @@
 
 namespace AutoRTFM
 {
-static_assert(static_cast<unsigned>(EContextStatus::OnTrack) == Constants::Context_Status_OnTrack, "Not equal");
-
 inline const char* GetContextStatusName(EContextStatus Status)
 {
     switch (Status)
diff --git a/Engine/Source/Runtime/Core/Public/AutoRTFM/AutoRTFMConstants.h b/Engine/Source/Runtime/Core/Public/AutoRTFM/AutoRTFMConstants.h
--- a/Engine/Source/Runtime/Core/Public/AutoRTFM/AutoRTFMConstants.h
+++ b/Engine/Source/Runtime/Core/Public/AutoRTFM/AutoRTFMConstants.h
@@ -5,9 +5,6 @@
 #include <cstddef>
 #include <cstdint>
 
-// FIXME: This is going to need support for cross-compilation. So, we will possibly have different offsets for
-// different platforms, and both the compiler and the runtime will have to make choices about which to use.
-
 // WARNING: Any change in these constants will require a re-patch and re-build of LLVM!
 
 namespace AutoRTFM
@@ -15,22 +12,9 @@ namespace AutoRTFM
 namespace Constants
 {
 
-inline constexpr size_t LogLineBytes = 3;
-inline constexpr size_t LineBytes = static_cast<size_t>(1) << LogLineBytes;
-inline constexpr size_t LineTableSize = 33554432;
-
-inline constexpr size_t Offset_Context_CurrentTransaction = 0;
-inline constexpr size_t Offset_Context_LineTable = 192;
-inline constexpr size_t Offset_Context_Status = 152;
-
-inline constexpr size_t LogSize_LineEntry = 5;
-inline constexpr size_t Size_LineEntry = static_cast<size_t>(1) << LogSize_LineEntry;
-inline constexpr size_t Offset_LineEntry_LogicalLine = 0;
-inline constexpr size_t Offset_LineEntry_ActiveLine = 16;
-inline constexpr size_t Offset_LineEntry_LoggingTransaction = 24;
-inline constexpr size_t Offset_LineEntry_AccessMask = 14;
-
-inline constexpr uint32_t Context_Status_OnTrack = 1;
+	inline constexpr uint32_t Major = 0;
+	inline constexpr uint32_t Minor = 1;
+	inline constexpr uint32_t Patch = 0;
 
 } // namespace Constants
 } // namespace AutoRTFM
