Fixed placement new calls in many Core containers to ensure that they still work with types which haven't overloaded operator new.

diff --git a/Engine/Source/Runtime/Core/Private/Containers/LockFreeList.cpp b/Engine/Source/Runtime/Core/Private/Containers/LockFreeList.cpp
--- a/Engine/Source/Runtime/Core/Private/Containers/LockFreeList.cpp
+++ b/Engine/Source/Runtime/Core/Private/Containers/LockFreeList.cpp
@@ -212,7 +212,7 @@ static LockFreeLinkAllocator_TLSCache& GetLockFreeAllocator()
 	static bool bIsInitialized = false;
 	if (!bIsInitialized)
 	{
-		new(Data)LockFreeLinkAllocator_TLSCache();
+		::new((void*)Data)LockFreeLinkAllocator_TLSCache();
 		bIsInitialized = true;
 	}
 	return *(LockFreeLinkAllocator_TLSCache*)Data;
diff --git a/Engine/Source/Runtime/Core/Public/Containers/Array.h b/Engine/Source/Runtime/Core/Public/Containers/Array.h
--- a/Engine/Source/Runtime/Core/Public/Containers/Array.h
+++ b/Engine/Source/Runtime/Core/Public/Containers/Array.h
@@ -1647,7 +1647,7 @@ public:
 		// construct a copy in place at Index (this new operator will insert at 
 		// Index, then construct that memory with Item)
 		InsertUninitializedImpl(Index);
-		new(GetData() + Index) ElementType(MoveTempIfPossible(Item));
+		::new((void*)(GetData() + Index)) ElementType(MoveTempIfPossible(Item));
 		return Index;
 	}
 
@@ -1666,7 +1666,7 @@ public:
 		// construct a copy in place at Index (this new operator will insert at 
 		// Index, then construct that memory with Item)
 		InsertUninitializedImpl(Index);
-		new(GetData() + Index) ElementType(Item);
+		::new((void*)(GetData() + Index)) ElementType(Item);
 		return Index;
 	}
 
@@ -1687,7 +1687,7 @@ public:
 		// Index, then construct that memory with Item)
 		InsertUninitializedImpl(Index);
 		ElementType* Ptr = GetData() + Index;
-		new(Ptr) ElementType(MoveTempIfPossible(Item));
+		::new((void*)Ptr) ElementType(MoveTempIfPossible(Item));
 		return *Ptr;
 	}
 
@@ -1707,7 +1707,7 @@ public:
 		// Index, then construct that memory with Item)
 		InsertUninitializedImpl(Index);
 		ElementType* Ptr = GetData() + Index;
-		new(Ptr) ElementType(Item);
+		::new((void*)Ptr) ElementType(Item);
 		return *Ptr;
 	}
 
@@ -2230,7 +2230,7 @@ public:
 		//     // Won't compile if the caller doesn't have access to FMyType::FPrivateToken
 		//     Arr.Emplace(FMyType::FPrivateToken{}, 5, 3.14f, TEXT("Banana"));
 		//
-		new(GetData() + Index) ElementType(Forward<ArgsType>(Args)...);
+		::new((void*)(GetData() + Index)) ElementType(Forward<ArgsType>(Args)...);
 		return Index;
 	}
 
@@ -2245,7 +2245,7 @@ public:
 	{
 		const SizeType Index = AddUninitialized();
 		ElementType* Ptr = GetData() + Index;
-		new(Ptr) ElementType(Forward<ArgsType>(Args)...);
+		::new((void*)Ptr) ElementType(Forward<ArgsType>(Args)...);
 		return *Ptr;
 	}
 
@@ -2259,7 +2259,7 @@ public:
 	FORCEINLINE void EmplaceAt(SizeType Index, ArgsType&&... Args)
 	{
 		InsertUninitializedImpl(Index, 1);
-		new(GetData() + Index) ElementType(Forward<ArgsType>(Args)...);
+		::new((void*)(GetData() + Index)) ElementType(Forward<ArgsType>(Args)...);
 	}
 
 	/**
@@ -2274,7 +2274,7 @@ public:
 	{
 		InsertUninitializedImpl(Index, 1);
 		ElementType* Ptr = GetData() + Index;
-		new(Ptr) ElementType(Forward<ArgsType>(Args)...);
+		::new((void*)Ptr) ElementType(Forward<ArgsType>(Args)...);
 		return *Ptr;
 	}
 
@@ -3255,13 +3255,13 @@ public:
 	{
 		if constexpr (TAllocatorTraits<AllocatorType>::SupportsFreezeMemoryImage && THasTypeLayout<ElementType>::Value)
 		{
-			TArray* DstArray = new(Dst) TArray();
+			TArray* DstArray = ::new(Dst) TArray();
 			DstArray->SetNumZeroed(this->ArrayNum);
 			this->AllocatorInstance.CopyUnfrozen(Context, StaticGetTypeLayoutDesc<ElementType>(), this->ArrayNum, DstArray->GetData());
 		}
 		else
 		{
-			new(Dst) TArray();
+			::new(Dst) TArray();
 		}
 	}
 
diff --git a/Engine/Source/Runtime/Core/Public/Containers/BasicArray.h b/Engine/Source/Runtime/Core/Public/Containers/BasicArray.h
--- a/Engine/Source/Runtime/Core/Public/Containers/BasicArray.h
+++ b/Engine/Source/Runtime/Core/Public/Containers/BasicArray.h
@@ -62,7 +62,7 @@ public:
 		int32 Result = Num();
 
 		void* LocationToAdd = InsertUninitialized(Result);
-		new (LocationToAdd) T(Forward<ArgTypes>(Args)...);
+		::new (LocationToAdd) T(Forward<ArgTypes>(Args)...);
 
 		return Result;
 	}
@@ -71,7 +71,7 @@ public:
 	void EmplaceAt(int32 Index, ArgTypes&&... Args)
 	{
 		void* LocationToAdd = InsertUninitialized(Index);
-		new (LocationToAdd) T(Forward<ArgTypes>(Args)...);
+		::new (LocationToAdd) T(Forward<ArgTypes>(Args)...);
 	}
 
 	void RemoveAt(int32 Index, int32 NumToRemove = 1)
diff --git a/Engine/Source/Runtime/Core/Public/Containers/ClosableMpscQueue.h b/Engine/Source/Runtime/Core/Public/Containers/ClosableMpscQueue.h
--- a/Engine/Source/Runtime/Core/Public/Containers/ClosableMpscQueue.h
+++ b/Engine/Source/Runtime/Core/Public/Containers/ClosableMpscQueue.h
@@ -50,7 +50,7 @@ public:
 		}
 
 		FNode* New = new FNode;
-		new (&New->Value) T(Forward<ArgTypes>(Args)...);
+		::new ((void*)&New->Value) T(Forward<ArgTypes>(Args)...);
 
 		while (!Head.compare_exchange_weak(Prev, New, std::memory_order_release) && Prev != nullptr) // linearisation point
 		{
diff --git a/Engine/Source/Runtime/Core/Public/Containers/ConsumeAllMpmcQueue.h b/Engine/Source/Runtime/Core/Public/Containers/ConsumeAllMpmcQueue.h
--- a/Engine/Source/Runtime/Core/Public/Containers/ConsumeAllMpmcQueue.h
+++ b/Engine/Source/Runtime/Core/Public/Containers/ConsumeAllMpmcQueue.h
@@ -50,8 +50,8 @@ namespace UE
 		template <typename... ArgTypes>
 		EConsumeAllMpmcQueueResult ProduceItem(ArgTypes&&... Args)
 		{
-			FNode* New = new(AllocatorType::Malloc(sizeof(FNode), alignof(FNode))) FNode;
-			new (&New->Item) T(Forward<ArgTypes>(Args)...);
+			FNode* New = ::new(AllocatorType::Malloc(sizeof(FNode), alignof(FNode))) FNode;
+			::new ((void*)&New->Item) T(Forward<ArgTypes>(Args)...);
 
 			//Atomically append to the top of the Queue
 			FNode* Prev = Head.load(std::memory_order_relaxed);
diff --git a/Engine/Source/Runtime/Core/Public/Containers/DepletableMpmcQueue.h b/Engine/Source/Runtime/Core/Public/Containers/DepletableMpmcQueue.h
--- a/Engine/Source/Runtime/Core/Public/Containers/DepletableMpmcQueue.h
+++ b/Engine/Source/Runtime/Core/Public/Containers/DepletableMpmcQueue.h
@@ -63,8 +63,8 @@ namespace UE
 		template <typename... ArgTypes>
 		bool EnqueueAndReturnWasEmpty(ArgTypes&&... Args)
 		{
-			FNode* New = new(AllocatorType::Malloc(sizeof(FNode), alignof(FNode))) FNode;
-			new (&New->Value) T(Forward<ArgTypes>(Args)...);
+			FNode* New = ::new(AllocatorType::Malloc(sizeof(FNode), alignof(FNode))) FNode;
+			::new ((void*)&New->Value) T(Forward<ArgTypes>(Args)...);
 
 			// switch `Tail` to the new node and only then link the old tail to the new one. The list is not fully linked between these ops,
 			// this is explicitly handled by the consumer by waiting for the link
diff --git a/Engine/Source/Runtime/Core/Public/Containers/Deque.h b/Engine/Source/Runtime/Core/Public/Containers/Deque.h
--- a/Engine/Source/Runtime/Core/Public/Containers/Deque.h
+++ b/Engine/Source/Runtime/Core/Public/Containers/Deque.h
@@ -225,7 +225,7 @@ public:
 	{
 		GrowIfRequired();
 		ElementType* Target = GetData() + Tail;
-		new (Target) ElementType(Forward<ArgsType>(Args)...);
+		::new ((void*)Target) ElementType(Forward<ArgsType>(Args)...);
 		Tail = UE::Deque::Private::WrapAround(Tail + 1, Capacity);
 		Count++;
 		return *Target;
@@ -241,7 +241,7 @@ public:
 		GrowIfRequired();
 		Head = UE::Deque::Private::WrapAround(Head + Capacity - 1, Capacity);
 		ElementType* Target = GetData() + Head;
-		new (Target) ElementType(Forward<ArgsType>(Args)...);
+		::new ((void*)Target) ElementType(Forward<ArgsType>(Args)...);
 		Count++;
 		return *Target;
 	}
diff --git a/Engine/Source/Runtime/Core/Public/Containers/HashTable.h b/Engine/Source/Runtime/Core/Public/Containers/HashTable.h
--- a/Engine/Source/Runtime/Core/Public/Containers/HashTable.h
+++ b/Engine/Source/Runtime/Core/Public/Containers/HashTable.h
@@ -470,13 +470,13 @@ public:
 	{
 		if constexpr (TAllocatorTraits<Allocator>::SupportsFreezeMemoryImage)
 		{
-			THashTable* DstTable = new(Dst) THashTable(this->HashMask + 1u, this->IndexSize);
+			THashTable* DstTable = ::new(Dst) THashTable(this->HashMask + 1u, this->IndexSize);
 			FMemory::Memcpy(DstTable->Hash.GetAllocation(), this->Hash.GetAllocation(), (this->HashMask + 1u) * 4);
 			FMemory::Memcpy(DstTable->NextIndex.GetAllocation(), this->NextIndex.GetAllocation(), this->IndexSize * 4);
 		}
 		else
 		{
-			new(Dst) THashTable();
+			::new(Dst) THashTable();
 		}
 	}
 };
diff --git a/Engine/Source/Runtime/Core/Public/Containers/LockFreeFixedSizeAllocator.h b/Engine/Source/Runtime/Core/Public/Containers/LockFreeFixedSizeAllocator.h
--- a/Engine/Source/Runtime/Core/Public/Containers/LockFreeFixedSizeAllocator.h
+++ b/Engine/Source/Runtime/Core/Public/Containers/LockFreeFixedSizeAllocator.h
@@ -327,7 +327,7 @@ public:
 	 */
 	T* New()
 	{
-		return new (Allocate()) T();
+		return ::new (Allocate()) T();
 	}
 
 	/**
@@ -371,7 +371,7 @@ public:
 	 */
 	T* New()
 	{
-		return new (Allocate()) T();
+		return ::new (Allocate()) T();
 	}
 
 	/**
diff --git a/Engine/Source/Runtime/Core/Public/Containers/LockFreeList.h b/Engine/Source/Runtime/Core/Public/Containers/LockFreeList.h
--- a/Engine/Source/Runtime/Core/Public/Containers/LockFreeList.h
+++ b/Engine/Source/Runtime/Core/Public/Containers/LockFreeList.h
@@ -82,7 +82,7 @@ public:
 		}
 		for (uint32 CurrentItem = FirstItem; CurrentItem < FirstItem + Count; CurrentItem++)
 		{
-			new (GetRawItem(CurrentItem)) T();
+			::new (GetRawItem(CurrentItem)) T();
 		}
 		return FirstItem;
 	}
diff --git a/Engine/Source/Runtime/Core/Public/Containers/MpscQueue.h b/Engine/Source/Runtime/Core/Public/Containers/MpscQueue.h
--- a/Engine/Source/Runtime/Core/Public/Containers/MpscQueue.h
+++ b/Engine/Source/Runtime/Core/Public/Containers/MpscQueue.h
@@ -23,7 +23,7 @@ public:
 
 	TMpscQueue()
 	{
-		FNode* Sentinel = new(AllocatorType::Malloc(sizeof(FNode), alignof(FNode))) FNode;
+		FNode* Sentinel = ::new(AllocatorType::Malloc(sizeof(FNode), alignof(FNode))) FNode;
 		Head.store(Sentinel, std::memory_order_relaxed);
 		Tail = Sentinel;
 	}
@@ -48,8 +48,8 @@ public:
 	template <typename... ArgTypes>
 	void Enqueue(ArgTypes&&... Args)
 	{
-		FNode* New = new(AllocatorType::Malloc(sizeof(FNode), alignof(FNode))) FNode;
-		new (&New->Value) ElementType(Forward<ArgTypes>(Args)...);
+		FNode* New = ::new(AllocatorType::Malloc(sizeof(FNode), alignof(FNode))) FNode;
+		::new ((void*)&New->Value) ElementType(Forward<ArgTypes>(Args)...);
 
 		FNode* Prev = Head.exchange(New, std::memory_order_acq_rel);
 		Prev->Next.store(New, std::memory_order_release);
diff --git a/Engine/Source/Runtime/Core/Public/Containers/RingBuffer.h b/Engine/Source/Runtime/Core/Public/Containers/RingBuffer.h
--- a/Engine/Source/Runtime/Core/Public/Containers/RingBuffer.h
+++ b/Engine/Source/Runtime/Core/Public/Containers/RingBuffer.h
@@ -274,7 +274,7 @@ public:
 	{
 		IndexType ResultIndex = AddUninitialized();
 		ElementType& Result = GetAtIndexNoCheck(ResultIndex);
-		new (&Result) ElementType(MoveTempIfPossible(Element));
+		::new ((void*)&Result) ElementType(MoveTempIfPossible(Element));
 		return ResultIndex;
 	}
 
@@ -282,7 +282,7 @@ public:
 	ElementType& Add_GetRef(ElementType&& Element)
 	{
 		ElementType& Result = AddUninitialized_GetRef();
-		new (&Result) ElementType(MoveTempIfPossible(Element));
+		::new ((void*)&Result) ElementType(MoveTempIfPossible(Element));
 		return Result;
 	}
 
@@ -291,7 +291,7 @@ public:
 	{
 		IndexType ResultIndex = AddUninitialized();
 		ElementType& Result = GetAtIndexNoCheck(ResultIndex);
-		new (&Result) ElementType(Element);
+		::new ((void*)&Result) ElementType(Element);
 		return ResultIndex;
 	}
 
@@ -300,7 +300,7 @@ public:
 	{
 		IndexType ResultIndex = AddUninitialized();
 		ElementType& Result = GetAtIndexNoCheck(ResultIndex);
-		new (&Result) ElementType(Element);
+		::new ((void*)&Result) ElementType(Element);
 		return Result;
 	}
 
@@ -310,7 +310,7 @@ public:
 	{
 		IndexType ResultIndex = AddUninitialized();
 		ElementType& Result = GetAtIndexNoCheck(ResultIndex);
-		new (&Result) ElementType(Forward<ArgsType>(Args)...);
+		::new ((void*)&Result) ElementType(Forward<ArgsType>(Args)...);
 		return ResultIndex;
 	}
 
@@ -319,7 +319,7 @@ public:
 	ElementType& Emplace_GetRef(ArgsType&&... Args)
 	{
 		ElementType& Result = AddUninitialized_GetRef();
-		new (&Result) ElementType(Forward<ArgsType>(Args)...);
+		::new ((void*)&Result) ElementType(Forward<ArgsType>(Args)...);
 		return Result;
 	}
 
@@ -380,7 +380,7 @@ public:
 	{
 		IndexType ResultIndex = AddFrontUninitialized();
 		ElementType& Result = GetAtIndexNoCheck(ResultIndex);
-		new (&Result) ElementType(MoveTempIfPossible(Element));
+		::new ((void*)&Result) ElementType(MoveTempIfPossible(Element));
 		return ResultIndex;
 	}
 
@@ -388,7 +388,7 @@ public:
 	ElementType& AddFront_GetRef(ElementType&& Element)
 	{
 		ElementType& Result = AddFrontUninitialized_GetRef();
-		new (&Result) ElementType(MoveTempIfPossible(Element));
+		::new ((void*)&Result) ElementType(MoveTempIfPossible(Element));
 		return Result;
 	}
 
@@ -397,7 +397,7 @@ public:
 	{
 		IndexType ResultIndex = AddFrontUninitialized();
 		ElementType& Result = GetAtIndexNoCheck(ResultIndex);
-		new (&Result) ElementType(Element);
+		::new ((void*)&Result) ElementType(Element);
 		return ResultIndex;
 	}
 
@@ -406,7 +406,7 @@ public:
 	{
 		IndexType ResultIndex = AddFrontUninitialized();
 		ElementType& Result = GetAtIndexNoCheck(ResultIndex);
-		new (&Result) ElementType(Element);
+		::new ((void*)&Result) ElementType(Element);
 		return Result;
 	}
 
@@ -416,7 +416,7 @@ public:
 	{
 		IndexType ResultIndex = AddFrontUninitialized();
 		ElementType& Result = GetAtIndexNoCheck(ResultIndex);
-		new (&Result) ElementType(Forward<ArgsType>(Args)...);
+		::new ((void*)&Result) ElementType(Forward<ArgsType>(Args)...);
 		return ResultIndex;
 	}
 
@@ -425,7 +425,7 @@ public:
 	ElementType& EmplaceFront_GetRef(ArgsType&&... Args)
 	{
 		ElementType& Result = AddFrontUninitialized_GetRef();
-		new (&Result) ElementType(Forward<ArgsType>(Args)...);
+		::new ((void*)&Result) ElementType(Forward<ArgsType>(Args)...);
 		return Result;
 	}
 
@@ -695,7 +695,7 @@ public:
 				{
 					DestructItem(&Data[WriteIndex & IndexMask]);
 				}
-				new (&Data[WriteIndex & IndexMask]) ElementType(MoveTemp(Data[ReadIndex & IndexMask]));
+				::new ((void*)&Data[WriteIndex & IndexMask]) ElementType(MoveTemp(Data[ReadIndex & IndexMask]));
 				++WriteIndex;
 			}
 			++ReadIndex;
@@ -807,12 +807,12 @@ private:
 				StorageModuloType WriteIndex = 0;
 				for (StorageModuloType ReadIndex = MaskedFront; ReadIndex < SrcCapacity; ++ReadIndex)
 				{
-					new (&NewData[WriteIndex++]) ElementType(MoveTemp(SrcData[ReadIndex]));
+					::new ((void*)&NewData[WriteIndex++]) ElementType(MoveTemp(SrcData[ReadIndex]));
 				}
 				DestructRange(MaskedFront, SrcCapacity);
 				for (StorageModuloType ReadIndex = 0; ReadIndex < MaskedAfterBack; ++ReadIndex)
 				{
-					new (&NewData[WriteIndex++]) ElementType(MoveTemp(SrcData[ReadIndex]));
+					::new ((void*)&NewData[WriteIndex++]) ElementType(MoveTemp(SrcData[ReadIndex]));
 				}
 				DestructRange(0, MaskedAfterBack);
 			}
@@ -821,7 +821,7 @@ private:
 				StorageModuloType WriteIndex = 0;
 				for (StorageModuloType ReadIndex = MaskedFront; ReadIndex < MaskedAfterBack; ++ReadIndex)
 				{
-					new(&NewData[WriteIndex++]) ElementType(MoveTemp(SrcData[ReadIndex]));
+					::new((void*)&NewData[WriteIndex++]) ElementType(MoveTemp(SrcData[ReadIndex]));
 				}
 				DestructRange(MaskedFront, MaskedAfterBack);
 			}
@@ -912,11 +912,11 @@ private:
 				StorageModuloType OldValueIndex = Index & IndexMask;
 				StorageModuloType NewValueIndex = (Index + RangeDirection) & IndexMask;
 				DestructItems(Data + OldValueIndex, 1);
-				new (Data + OldValueIndex) ElementType(MoveTemp(Data[NewValueIndex]));
+				::new ((void*)(Data + OldValueIndex)) ElementType(MoveTemp(Data[NewValueIndex]));
 			}
 			StorageModuloType OldValueIndex = RangeFirst & IndexMask;
 			DestructItems(Data + OldValueIndex, 1);
-			new (Data + OldValueIndex) ElementType(MoveTemp(Copy));
+			::new ((void*)(Data + OldValueIndex)) ElementType(MoveTemp(Copy));
 		}
 		else
 		{
@@ -924,10 +924,10 @@ private:
 			{
 				StorageModuloType OldValueIndex = Index & IndexMask;
 				StorageModuloType NewValueIndex = (Index + RangeDirection) & IndexMask;
-				new (Data + OldValueIndex) ElementType(MoveTemp(Data[NewValueIndex]));
+				::new ((void*)(Data + OldValueIndex)) ElementType(MoveTemp(Data[NewValueIndex]));
 			}
 			StorageModuloType OldValueIndex = RangeFirst & IndexMask;
-			new (Data + OldValueIndex) ElementType(MoveTemp(Copy));
+			::new ((void*)(Data + OldValueIndex)) ElementType(MoveTemp(Copy));
 		}
 	}
 
diff --git a/Engine/Source/Runtime/Core/Public/Containers/Set.h b/Engine/Source/Runtime/Core/Public/Containers/Set.h
--- a/Engine/Source/Runtime/Core/Public/Containers/Set.h
+++ b/Engine/Source/Runtime/Core/Public/Containers/Set.h
@@ -1416,14 +1416,14 @@ public:
 			TSet* DstObject = static_cast<TSet*>(Dst);
 			this->Elements.CopyUnfrozen(Context, &DstObject->Elements);
 
-			new(&DstObject->Hash) HashType();
+			::new((void*)&DstObject->Hash) HashType();
 			DstObject->Hash.ResizeAllocation(0, this->HashSize, sizeof(FSetElementId));
 			FMemory::Memcpy(DstObject->Hash.GetAllocation(), this->Hash.GetAllocation(), sizeof(FSetElementId) * this->HashSize);
 			DstObject->HashSize = this->HashSize;
 		}
 		else
 		{
-			new(Dst) TSet();
+			::new(Dst) TSet();
 		}
 	}
 
diff --git a/Engine/Source/Runtime/Core/Public/Containers/SortedMap.h b/Engine/Source/Runtime/Core/Public/Containers/SortedMap.h
--- a/Engine/Source/Runtime/Core/Public/Containers/SortedMap.h
+++ b/Engine/Source/Runtime/Core/Public/Containers/SortedMap.h
@@ -193,7 +193,7 @@ public:
 	{
 		ElementType* DataPtr = AllocateMemoryForEmplace(InKey);
 	
-		new(DataPtr) ElementType(TPairInitializer<InitKeyType&&, InitValueType&&>(Forward<InitKeyType>(InKey), Forward<InitValueType>(InValue)));
+		::new((void*)DataPtr) ElementType(TPairInitializer<InitKeyType&&, InitValueType&&>(Forward<InitKeyType>(InKey), Forward<InitValueType>(InValue)));
 
 		return DataPtr->Value;
 	}
@@ -209,7 +209,7 @@ public:
 	{
 		ElementType* DataPtr = AllocateMemoryForEmplace(InKey);
 
-		new(DataPtr) ElementType(TKeyInitializer<InitKeyType&&>(Forward<InitKeyType>(InKey)));
+		::new((void*)DataPtr) ElementType(TKeyInitializer<InitKeyType&&>(Forward<InitKeyType>(InKey)));
 
 		return DataPtr->Value;
 	}
diff --git a/Engine/Source/Runtime/Core/Public/Containers/SparseArray.h b/Engine/Source/Runtime/Core/Public/Containers/SparseArray.h
--- a/Engine/Source/Runtime/Core/Public/Containers/SparseArray.h
+++ b/Engine/Source/Runtime/Core/Public/Containers/SparseArray.h
@@ -1178,7 +1178,7 @@ public:
 			const FTypeLayoutDesc& ElementTypeDesc = StaticGetTypeLayoutDesc<ElementType>();
 			TSparseArray* DstObject = (TSparseArray*)Dst;
 			{
-				new(&DstObject->Data) DataType();
+				::new((void*)&DstObject->Data) DataType();
 				DstObject->Data.SetNumUninitialized(this->Data.Num());
 				for (int32 i = 0; i < this->Data.Num(); ++i)
 				{
@@ -1196,13 +1196,13 @@ public:
 				}
 			}
 
-			new(&DstObject->AllocationFlags) AllocationBitArrayType(this->AllocationFlags);
+			::new((void*)&DstObject->AllocationFlags) AllocationBitArrayType(this->AllocationFlags);
 			DstObject->FirstFreeIndex = this->FirstFreeIndex;
 			DstObject->NumFreeIndices = this->NumFreeIndices;
 		}
 		else
 		{
-			new(Dst) TSparseArray();
+			::new((void*)Dst) TSparseArray();
 		}
 	}
 
diff --git a/Engine/Source/Runtime/Core/Public/Containers/SpscQueue.h b/Engine/Source/Runtime/Core/Public/Containers/SpscQueue.h
--- a/Engine/Source/Runtime/Core/Public/Containers/SpscQueue.h
+++ b/Engine/Source/Runtime/Core/Public/Containers/SpscQueue.h
@@ -23,7 +23,7 @@ public:
 
 	TSpscQueue()
 	{
-		FNode* Node = new(AllocatorType::Malloc(sizeof(FNode), alignof(FNode))) FNode;
+		FNode* Node = ::new(AllocatorType::Malloc(sizeof(FNode), alignof(FNode))) FNode;
 		Tail.store(Node, std::memory_order_relaxed);
 		Head = First = TailCopy = Node;
 	}
@@ -57,7 +57,7 @@ public:
 	void Enqueue(ArgTypes&&... Args)
 	{
 		FNode* Node = AllocNode();
-		new(&Node->Value) ElementType(Forward<ArgTypes>(Args)...);
+		::new((void*)&Node->Value) ElementType(Forward<ArgTypes>(Args)...);
 
 		Head->Next.store(Node, std::memory_order_release);
 		Head = Node;
@@ -148,7 +148,7 @@ private:
 			return AllocFromCache();
 		}
 
-		return new(AllocatorType::Malloc(sizeof(FNode), alignof(FNode))) FNode();
+		return ::new(AllocatorType::Malloc(sizeof(FNode), alignof(FNode))) FNode();
 	}
 
 private:
diff --git a/Engine/Source/Runtime/Core/Public/Containers/TransArray.h b/Engine/Source/Runtime/Core/Public/Containers/TransArray.h
--- a/Engine/Source/Runtime/Core/Public/Containers/TransArray.h
+++ b/Engine/Source/Runtime/Core/Public/Containers/TransArray.h
@@ -173,7 +173,7 @@ protected:
 
 	static void DefaultConstructItem( void* TPtr )
 	{
-		new (TPtr) T;
+		::new (TPtr) T;
 	}
 	static void SerializeItem( FArchive& Ar, void* TPtr )
 	{
diff --git a/Engine/Source/Runtime/Core/Public/Containers/Union.h b/Engine/Source/Runtime/Core/Public/Containers/Union.h
--- a/Engine/Source/Runtime/Core/Public/Containers/Union.h
+++ b/Engine/Source/Runtime/Core/Public/Containers/Union.h
@@ -186,7 +186,7 @@ public:
 
 		Reset();
 
-		new(SubtypeValuePointer) Subtype();
+		::new((void*)SubtypeValuePointer) Subtype();
 
 		CurrentSubtypeIndex = SubtypeIndex;
 		return SubtypeValuePointer;
@@ -202,7 +202,7 @@ public:
 
 		Reset();
 
-		new(SubtypeValuePointer) Subtype(NewValue);
+		::new((void*)SubtypeValuePointer) Subtype(NewValue);
 
 		CurrentSubtypeIndex = SubtypeIndex;
 		return SubtypeValuePointer;
@@ -342,7 +342,7 @@ private:
 	Subtype& InitSubtype()
 	{
 		Subtype* NewSubtype = &GetSubtype<Subtype>();
-		return *new(NewSubtype) Subtype;
+		return *::new((void*)NewSubtype) Subtype;
 	}
 
 	/** Determines the index and reference to the potential value for the given union subtype. */
diff --git a/Engine/Source/Runtime/Core/Public/Experimental/Containers/RobinHoodHashTable.h b/Engine/Source/Runtime/Core/Public/Experimental/Containers/RobinHoodHashTable.h
--- a/Engine/Source/Runtime/Core/Public/Experimental/Containers/RobinHoodHashTable.h
+++ b/Engine/Source/Runtime/Core/Public/Experimental/Containers/RobinHoodHashTable.h
@@ -351,7 +351,7 @@ namespace RobinHoodHashTable_Private
 				if (FreeList.Num() > 0)
 				{
 					Index = FreeList.Pop();
-					new (&KeyVals[Index]) KeyValueType{ Forward<DeducedKeyType>(Key), Forward<DeducedValueType>(Val) };
+					::new ((void*)&KeyVals[Index]) KeyValueType{ Forward<DeducedKeyType>(Key), Forward<DeducedValueType>(Val) };
 					Hashes[Index] = Hash;
 
 				}
diff --git a/Engine/Source/Runtime/Core/Public/Misc/ExpressionParserTypes.inl b/Engine/Source/Runtime/Core/Public/Misc/ExpressionParserTypes.inl
--- a/Engine/Source/Runtime/Core/Public/Misc/ExpressionParserTypes.inl
+++ b/Engine/Source/Runtime/Core/Public/Misc/ExpressionParserTypes.inl
@@ -40,7 +40,7 @@ namespace Impl
 		
 		const T* Access() const { return &Value; }
 
-		virtual void Reseat(uint8* Dst) override 			{ new(Dst) FInlineDataStorage(MoveTemp(Value)); }
+		virtual void Reseat(uint8* Dst) override 			{ ::new((void*)Dst) FInlineDataStorage(MoveTemp(Value)); }
 		virtual void MoveAssign(uint8* Dst) override 		{ reinterpret_cast<FInlineDataStorage*>(Dst)->Value = MoveTemp(Value); }
 		virtual FExpressionNode Copy() const override		{ return Value; }
 	};
@@ -59,7 +59,7 @@ namespace Impl
 		
 		const T* Access() const { return Value.Get(); }
 
-		virtual void Reseat(uint8* Dst) override 			{ new(Dst) FHeapDataStorage(MoveTemp(*this)); }
+		virtual void Reseat(uint8* Dst) override 			{ ::new((void*)Dst) FHeapDataStorage(MoveTemp(*this)); }
 		virtual void MoveAssign(uint8* Dst) override 		{ reinterpret_cast<FHeapDataStorage*>(Dst)->Value = MoveTemp(Value); }
 		virtual FExpressionNode Copy() const override		{ return *Value; }
 	};
@@ -168,7 +168,7 @@ FExpressionNode::FExpressionNode(T In, typename TEnableIf<!TPointerIsConvertible
 	: TypeId(TGetExpressionNodeTypeId<T>::GetTypeId())
 {
 	// Choose the relevant allocation strategy based on the size of the type
-	new(InlineBytes) typename Impl::TStorageTypeDeduction<T, MaxStackAllocationSize>::Type(MoveTemp(In));
+	::new((void*)InlineBytes) typename Impl::TStorageTypeDeduction<T, MaxStackAllocationSize>::Type(MoveTemp(In));
 }
 
 template<typename T>
diff --git a/Engine/Source/Runtime/Core/Public/Misc/InlineValue.h b/Engine/Source/Runtime/Core/Public/Misc/InlineValue.h
--- a/Engine/Source/Runtime/Core/Public/Misc/InlineValue.h
+++ b/Engine/Source/Runtime/Core/Public/Misc/InlineValue.h
@@ -201,7 +201,7 @@ private:
 		bIsValid = true;
 
 		// Placement new our value into the structure
-		new(&GetValue()) T(Forward<ArgsType>(Args)...);
+		::new((void*)&GetValue()) T(Forward<ArgsType>(Args)...);
 
 		checkf((void*)&GetValue() == (BaseType*)((T*)&GetValue()), TEXT("TInlineValue cannot operate with multiple inheritance objects."));
 	}
diff --git a/Engine/Source/Runtime/Core/Public/Misc/LazySingleton.h b/Engine/Source/Runtime/Core/Public/Misc/LazySingleton.h
--- a/Engine/Source/Runtime/Core/Public/Misc/LazySingleton.h
+++ b/Engine/Source/Runtime/Core/Public/Misc/LazySingleton.h
@@ -8,7 +8,7 @@
 class FLazySingleton
 {
 protected:
-	template<class T> static void Construct(void* Place)	{ new (Place) T; }
+	template<class T> static void Construct(void* Place)	{ ::new (Place) T; }
 	template<class T> static void Destruct(T* Instance)		{ Instance->~T(); }
 };
 
diff --git a/Engine/Source/Runtime/Core/Public/Misc/Optional.h b/Engine/Source/Runtime/Core/Public/Misc/Optional.h
--- a/Engine/Source/Runtime/Core/Public/Misc/Optional.h
+++ b/Engine/Source/Runtime/Core/Public/Misc/Optional.h
@@ -92,7 +92,7 @@ public:
 		//     // Won't compile if the caller doesn't have access to FMyType::FPrivateToken
 		//     TOptional<FMyType> Opt(InPlace, FMyType::FPrivateToken{}, 5, 3.14f, TEXT("Banana"));
 		//
-		new(&Value) OptionalType(Forward<ArgTypes>(Args)...);
+		::new((void*)&Value) OptionalType(Forward<ArgTypes>(Args)...);
 
 		if constexpr (!bUsingIntrusiveUnsetState)
 		{
@@ -116,7 +116,7 @@ public:
 	{
 		if constexpr (bUsingIntrusiveUnsetState)
 		{
-			new (&Value) OptionalType(FIntrusiveUnsetOptionalState{});
+			::new ((void*)&Value) OptionalType(FIntrusiveUnsetOptionalState{});
 		}
 		else
 		{
@@ -142,7 +142,7 @@ public:
 			}
 		}
 
-		new(&Value) OptionalType(*(const OptionalType*)&Other.Value);
+		::new((void*)&Value) OptionalType(*(const OptionalType*)&Other.Value);
 	}
 	TOptional(TOptional&& Other)
 	{
@@ -156,7 +156,7 @@ public:
 			}
 		}
 
-		new(&Value) OptionalType(MoveTempIfPossible(*(OptionalType*)&Other.Value));
+		::new((void*)&Value) OptionalType(MoveTempIfPossible(*(OptionalType*)&Other.Value));
 	}
 
 	TOptional& operator=(const TOptional& Other)
@@ -172,7 +172,7 @@ public:
 				Reset();
 				if (Other.Value.bIsSet)
 				{
-					new(&Value) OptionalType(*(const OptionalType*)&Other.Value);
+					::new((void*)&Value) OptionalType(*(const OptionalType*)&Other.Value);
 					Value.bIsSet = true;
 				}
 			}
@@ -192,7 +192,7 @@ public:
 				Reset();
 				if (Other.Value.bIsSet)
 				{
-					new(&Value) OptionalType(MoveTempIfPossible(*(OptionalType*)&Other.Value));
+					::new((void*)&Value) OptionalType(MoveTempIfPossible(*(OptionalType*)&Other.Value));
 					Value.bIsSet = true;
 				}
 			}
@@ -274,7 +274,7 @@ public:
 		//     // Won't compile if the caller doesn't have access to FMyType::FPrivateToken
 		//     Opt.Emplace(FMyType::FPrivateToken{}, 5, 3.14f, TEXT("Banana"));
 		//
-		OptionalType* Result = new(&Value) OptionalType(Forward<ArgsType>(Args)...);
+		OptionalType* Result = ::new((void*)&Value) OptionalType(Forward<ArgsType>(Args)...);
 
 		if constexpr (!bUsingIntrusiveUnsetState)
 		{
diff --git a/Engine/Source/Runtime/Core/Public/Misc/RangeAllocator.h b/Engine/Source/Runtime/Core/Public/Misc/RangeAllocator.h
--- a/Engine/Source/Runtime/Core/Public/Misc/RangeAllocator.h
+++ b/Engine/Source/Runtime/Core/Public/Misc/RangeAllocator.h
@@ -381,7 +381,7 @@ public:
 			ChunkIndex = Chunks.AddUninitialized();
 		}
 		FChunk& Chunk = Chunks[ChunkIndex];
-		new (&Chunk) FChunk(*this);
+		::new ((void*)&Chunk) FChunk(*this);
 		SIZE_T Ptr = Chunk.Alloc(Size, Alignment, MaxFreeRangeSize, OutAllocInfo.RangeOffset, OutAllocInfo.RangeSize);
 		Check(Ptr != 0);
 #if RANGE_ALLOCATOR_RECORD_STATS
diff --git a/Engine/Source/Runtime/Core/Public/Misc/TVariant.h b/Engine/Source/Runtime/Core/Public/Misc/TVariant.h
--- a/Engine/Source/Runtime/Core/Public/Misc/TVariant.h
+++ b/Engine/Source/Runtime/Core/Public/Misc/TVariant.h
@@ -42,7 +42,7 @@ public:
 	TVariant()
 	{
 		static_assert(std::is_constructible_v<T>, "To default-initialize a TVariant, the first type in the parameter pack must be default constructible. Use FEmptyVariantState as the first type if none of the other types can be listed first.");
-		new(&UE::Core::Private::CastToStorage(*this).Storage) T();
+		::new((void*)&UE::Core::Private::CastToStorage(*this).Storage) T();
 		TypeIndex = 0;
 	}
 
@@ -53,7 +53,7 @@ public:
 		constexpr SIZE_T Index = UE::Core::Private::TParameterPackTypeIndex<U, T, Ts...>::Value;
 		static_assert(Index != (SIZE_T)-1, "The TVariant is not declared to hold the type being constructed");
 
-		new(&UE::Core::Private::CastToStorage(*this).Storage) U(Forward<TArgs>(Args)...);
+		::new((void*)&UE::Core::Private::CastToStorage(*this).Storage) U(Forward<TArgs>(Args)...);
 		TypeIndex = (uint8)Index;
 	}
 
@@ -169,7 +169,7 @@ public:
 		static_assert(Index != (SIZE_T)-1, "The TVariant is not declared to hold the type passed to Emplace<>");
 
 		UE::Core::Private::TDestructorLookup<T, Ts...>::Destruct(TypeIndex, &UE::Core::Private::CastToStorage(*this).Storage);
-		new(&UE::Core::Private::CastToStorage(*this).Storage) U(Forward<TArgs>(Args)...);
+		::new((void*)&UE::Core::Private::CastToStorage(*this).Storage) U(Forward<TArgs>(Args)...);
 		TypeIndex = (uint8)Index;
 	}
 
diff --git a/Engine/Source/Runtime/Core/Public/Misc/TVariantMeta.h b/Engine/Source/Runtime/Core/Public/Misc/TVariantMeta.h
--- a/Engine/Source/Runtime/Core/Public/Misc/TVariantMeta.h
+++ b/Engine/Source/Runtime/Core/Public/Misc/TVariantMeta.h
@@ -216,7 +216,7 @@ namespace Private
 		/** Call the copy constructor of a type with the provided memory location and value */
 		static void Construct(void* Storage, const void* Value)
 		{
-			new(Storage) T(*static_cast<const T*>(Value));
+			::new(Storage) T(*static_cast<const T*>(Value));
 		}
 	};
 
@@ -241,7 +241,7 @@ namespace Private
 		/** Call the move constructor of a type with the provided memory location and value */
 		static void Construct(void* Storage, void* Value)
 		{
-			new(Storage) T(MoveTemp(*static_cast<T*>(Value)));
+			::new(Storage) T(MoveTemp(*static_cast<T*>(Value)));
 		}
 	};
 
diff --git a/Engine/Source/Runtime/Core/Public/Templates/Function.h b/Engine/Source/Runtime/Core/Public/Templates/Function.h
--- a/Engine/Source/Runtime/Core/Public/Templates/Function.h
+++ b/Engine/Source/Runtime/Core/Public/Templates/Function.h
@@ -334,7 +334,7 @@ namespace UE::Core::Private::Function
 			}
 
 			CA_ASSUME(NewAlloc);
-			auto* NewOwned = new (NewAlloc) OwnedType(Forward<FunctorType>(InFunc));
+			auto* NewOwned = ::new (NewAlloc) OwnedType(Forward<FunctorType>(InFunc));
 			return &NewOwned->Obj;
 		}
 	};
@@ -358,7 +358,7 @@ namespace UE::Core::Private::Function
 			CA_ASSUME(NewAlloc);
 		}
 
-		auto* NewOwned = new (NewAlloc) TFunction_CopyableOwnedObject(this->Obj);
+		auto* NewOwned = ::new (NewAlloc) TFunction_CopyableOwnedObject(this->Obj);
 
 		return &NewOwned->Obj;
 	}
@@ -531,7 +531,7 @@ namespace UE::Core::Private::Function
 				// but it ensures that the vptr is set correctly for the bound type, and so is
 				// visualizable.  We never depend on the state of this object at runtime, so it's
 				// ok.
-				new ((void*)&DebugPtrStorage) TDebugHelper<DecayedFunctorType>;
+				::new ((void*)&DebugPtrStorage) TDebugHelper<DecayedFunctorType>;
 				DebugPtrStorage.Ptr = (void*)Binding;
 			#endif
 		}
diff --git a/Engine/Source/Runtime/Core/Public/Templates/MemoryOps.h b/Engine/Source/Runtime/Core/Public/Templates/MemoryOps.h
--- a/Engine/Source/Runtime/Core/Public/Templates/MemoryOps.h
+++ b/Engine/Source/Runtime/Core/Public/Templates/MemoryOps.h
@@ -44,7 +44,7 @@ FORCEINLINE void DefaultConstructItems(void* Address, SizeType Count)
 		ElementType* Element = (ElementType*)Address;
 		while (Count)
 		{
-			new (Element) ElementType;
+			::new ((void*)Element) ElementType;
 			++Element;
 			--Count;
 		}
@@ -128,7 +128,7 @@ FORCEINLINE void ConstructItems(void* Dest, const SourceElementType* Source, Siz
 	{
 		while (Count)
 		{
-			new (Dest) DestinationElementType(*Source);
+			::new ((void*)Dest) DestinationElementType(*Source);
 			++(DestinationElementType*&)Dest;
 			++Source;
 			--Count;
@@ -198,7 +198,7 @@ FORCEINLINE void RelocateConstructItem(void* Dest, const SourceElementType* Sour
 		// We need a typedef here because VC won't compile the destructor call below if SourceElementType itself has a member called SourceElementType
 		typedef SourceElementType RelocateConstructItemsElementTypeTypedef;
 
-		new (Dest) DestinationElementType(*Source);
+		::new ((void*)Dest) DestinationElementType(*Source);
 		Source->RelocateConstructItemsElementTypeTypedef::~RelocateConstructItemsElementTypeTypedef();
 	}
 }
@@ -237,7 +237,7 @@ FORCEINLINE void RelocateConstructItems(void* Dest, const SourceElementType* Sou
 			// We need a typedef here because VC won't compile the destructor call below if SourceElementType itself has a member called SourceElementType
 			typedef SourceElementType RelocateConstructItemsElementTypeTypedef;
 
-			new (Dest) DestinationElementType(*Source);
+			::new ((void*)Dest) DestinationElementType(*Source);
 			++(DestinationElementType*&)Dest;
 			(Source++)->RelocateConstructItemsElementTypeTypedef::~RelocateConstructItemsElementTypeTypedef();
 			--Count;
@@ -267,7 +267,7 @@ FORCEINLINE void MoveConstructItems(void* Dest, const ElementType* Source, SizeT
 	{
 		while (Count)
 		{
-			new (Dest) ElementType((ElementType&&)*Source);
+			::new ((void*)Dest) ElementType((ElementType&&)*Source);
 			++(ElementType*&)Dest;
 			++Source;
 			--Count;
diff --git a/Engine/Source/Runtime/Core/Public/Templates/SharedPointerInternals.h b/Engine/Source/Runtime/Core/Public/Templates/SharedPointerInternals.h
--- a/Engine/Source/Runtime/Core/Public/Templates/SharedPointerInternals.h
+++ b/Engine/Source/Runtime/Core/Public/Templates/SharedPointerInternals.h
@@ -411,7 +411,7 @@ namespace SharedPointerInternals
 			//     // Won't compile if the caller doesn't have access to FMyType::FPrivateToken
 			//     TSharedPtr<FMyType> Val = MakeShared<FMyType>(FMyType::FPrivateToken{}, 5, 3.14f, TEXT("Banana"));
 			//
-			new ((void*)&ObjectStorage) ObjectType(Forward<ArgTypes>(Args)...);
+			::new ((void*)&ObjectStorage) ObjectType(Forward<ArgTypes>(Args)...);
 		}
 
 		ObjectType* GetObjectPtr() const
diff --git a/Engine/Source/Runtime/Core/Public/Templates/UnrealTemplate.h b/Engine/Source/Runtime/Core/Public/Templates/UnrealTemplate.h
--- a/Engine/Source/Runtime/Core/Public/Templates/UnrealTemplate.h
+++ b/Engine/Source/Runtime/Core/Public/Templates/UnrealTemplate.h
@@ -55,7 +55,7 @@ FORCEINLINE void Move(T& A,typename TMoveSupportTraits<T>::Copy B)
 	A.~T();
 
 	// Use placement new and a copy constructor so types with const members will work.
-	new(&A) T(B);
+	::new((void*)&A) T(B);
 }
 
 /** This is used to provide type specific behavior for a move which may change the value of B. */
@@ -66,7 +66,7 @@ FORCEINLINE void Move(T& A,typename TMoveSupportTraits<T>::Move B)
 	A.~T();
 
 	// Use placement new and a copy constructor so types with const members will work.
-	new(&A) T(MoveTemp(B));
+	::new((void*)&A) T(MoveTemp(B));
 }
 
 /**
