[Insights]
  - Removed function FChannel::EnumerateChannels deprecated since UE 5.2.
  - Fixed / updated comments for UE::Trace::FChannel class.
  - TraceAux: Made the warning message emitted when disabling a read only channel to be an error. Also update the error messages for enabling/disabling read only channels.
  - TraceAux: Fixed comments in TraceAuxiliary API.
  - TraceAux: Removed unused private function in TraceAuxiliary.
  - TraceAux: Minor refactoring / code cleanup.

diff --git a/Engine/Source/Runtime/Core/Private/ProfilingDebugging/TraceAuxiliary.cpp b/Engine/Source/Runtime/Core/Private/ProfilingDebugging/TraceAuxiliary.cpp
--- a/Engine/Source/Runtime/Core/Private/ProfilingDebugging/TraceAuxiliary.cpp
+++ b/Engine/Source/Runtime/Core/Private/ProfilingDebugging/TraceAuxiliary.cpp
@@ -136,12 +136,11 @@ private:
 
 	struct FChannelEntry
 	{
-		FString	Name;
+		FString Name;
 		bool bActive = false;
 	};
 
-	void AddChannel(const TCHAR* Name, const FTraceAuxiliary::FLogCategoryAlias& LogCategory);
-	void RemoveChannel(const TCHAR* Name, const FTraceAuxiliary::FLogCategoryAlias& LogCategory);
+	void AddCommandlineChannel(const TCHAR* Name, const FTraceAuxiliary::FLogCategoryAlias& LogCategory);
 	template <class T> void ForEachChannel(const TCHAR* ChannelList, bool bResolvePresets, const FTraceAuxiliary::FLogCategoryAlias& LogCategory, T Callable);
 	static uint32 HashChannelName(const TCHAR* Name);
 	bool EnableChannel(const TCHAR* Channel, const FTraceAuxiliary::FLogCategoryAlias& LogCategory);
@@ -180,10 +179,29 @@ bool FTraceAuxiliaryImpl::IsParentProcessAndPreFork()
 	return FForkProcessHelper::IsForkRequested() && !FForkProcessHelper::IsForkedChildProcess();
 }
 
+////////////////////////////////////////////////////////////////////////////////
+void FTraceAuxiliaryImpl::AddCommandlineChannel(const TCHAR* Name, const FTraceAuxiliary::FLogCategoryAlias& LogCategory)
+{
+	uint32 Hash = HashChannelName(Name);
+
+	if (CommandlineChannels.Find(Hash) != nullptr)
+	{
+		return;
+	}
+
+	FChannelEntry& Value = CommandlineChannels.Add(Hash, {});
+	Value.Name = Name;
+
+	if (IsConnected() && !Value.bActive)
+	{
+		Value.bActive = EnableChannel(*Value.Name, LogCategory);
+	}
+}
+
 ////////////////////////////////////////////////////////////////////////////////
 void FTraceAuxiliaryImpl::AddCommandlineChannels(const TCHAR* ChannelList)
 {
-	ForEachChannel(ChannelList, true, LogTrace, &FTraceAuxiliaryImpl::AddChannel);
+	ForEachChannel(ChannelList, true, LogTrace, &FTraceAuxiliaryImpl::AddCommandlineChannel);
 }
 
 ////////////////////////////////////////////////////////////////////////////////
@@ -232,11 +250,11 @@ void FTraceAuxiliaryImpl::ForEachChannel(const TCHAR* ChannelList, bool bResolve
 		{
 			FString Value;
 			// Check against hard coded presets
-			if(FCString::Stricmp(Name, TEXT("default")) == 0)
+			if (FCString::Stricmp(Name, TEXT("default")) == 0)
 			{
 				ForEachChannel(GDefaultChannels, false, LogCategory, Callable);
 			}
-			else if(FCString::Stricmp(Name, TEXT("memory"))== 0)
+			else if (FCString::Stricmp(Name, TEXT("memory"))== 0)
 			{
 				ForEachChannel(GMemoryChannels, false, LogCategory, Callable);
 			}
@@ -264,43 +282,6 @@ uint32 FTraceAuxiliaryImpl::HashChannelName(const TCHAR* Name)
 	return Hash;
 }
 
-////////////////////////////////////////////////////////////////////////////////
-void FTraceAuxiliaryImpl::AddChannel(const TCHAR* Name, const FTraceAuxiliary::FLogCategoryAlias& LogCategory)
-{
-	uint32 Hash = HashChannelName(Name);
-
-	if (CommandlineChannels.Find(Hash) != nullptr)
-	{
-		return;
-	}
-
-	FChannelEntry& Value = CommandlineChannels.Add(Hash, {});
-	Value.Name = Name;
-
-	if (IsConnected() && !Value.bActive)
-	{
-		Value.bActive = EnableChannel(*Value.Name, LogCategory);
-	}
-}
-
-////////////////////////////////////////////////////////////////////////////////
-void FTraceAuxiliaryImpl::RemoveChannel(const TCHAR* Name, const FTraceAuxiliary::FLogCategoryAlias& LogCategory)
-{
-	uint32 Hash = HashChannelName(Name);
-
-	FChannelEntry Channel;
-	if (!CommandlineChannels.RemoveAndCopyValue(Hash, Channel))
-	{
-		return;
-	}
-
-	if (IsConnected() && Channel.bActive)
-	{
-		DisableChannel(*Channel.Name, LogCategory);
-		Channel.bActive = false;
-	}
-}
-
 ////////////////////////////////////////////////////////////////////////////////
 bool FTraceAuxiliaryImpl::Connect(FTraceAuxiliary::EConnectionType Type, const TCHAR* Parameter, const FTraceAuxiliary::FLogCategoryAlias& LogCategory, uint16 SendFlags)
 {
@@ -403,7 +384,7 @@ bool FTraceAuxiliaryImpl::IsReadOnlyChannel(const TCHAR* InChannel)
 		FAnsiStringView Channel;
 		bool bIsReadOnlyChannel = false;
 	} UserData = { FAnsiStringView(ChannelA.Get(), ChannelA.Length()), false};
-	
+
 	UE::Trace::EnumerateChannels([](const UE::Trace::FChannelInfo& Info, void* User)
 	{
 		FAnsiStringView NameView = FAnsiStringView(Info.Name).LeftChop(7); // Remove "Channel" suffix
@@ -430,22 +411,22 @@ bool FTraceAuxiliaryImpl::EnableChannel(const TCHAR* Channel, const FTraceAuxili
 		return false;
 	}
 
-	// It is not possible to change read only channels once trace is
-	// initialized.
-	if(bReadOnlyChannelsFrozen && IsReadOnlyChannel(Channel))
+	// It is not possible to change read only channels once trace is initialized.
+	if (bReadOnlyChannelsFrozen && IsReadOnlyChannel(Channel))
 	{
 		UE_LOG_REF(
 			LogCategory,
 			Error,
-			TEXT("Channel '%s' is a read only channel. It is only possible to enable on startup using command line parameters."),
+			TEXT("Channel '%s' is read only. It is not allowed to manually enable this channel."),
 			Channel
 		);
 		return false;
 	}
 
-	FPlatformEventsTrace::OnTraceChannelUpdated(Channel, true);
+	bool bIsEnabled = UE::Trace::ToggleChannel(Channel, true);
+	FPlatformEventsTrace::OnTraceChannelUpdated(Channel, bIsEnabled);
 
-	return UE::Trace::ToggleChannel(Channel, true);
+	return bIsEnabled;
 }
 
 ////////////////////////////////////////////////////////////////////////////////
@@ -458,22 +439,21 @@ void FTraceAuxiliaryImpl::DisableChannel(const TCHAR* Channel, const FTraceAuxil
 	{
 		return;
 	}
-	
-	// Warn when disabling a read-only channel after init, it will not be possible
-	// to enable again.
-	if(bReadOnlyChannelsFrozen && IsReadOnlyChannel(Channel))
+
+	// It is not possible to change read only channels once trace is initialized.
+	if (bReadOnlyChannelsFrozen && IsReadOnlyChannel(Channel))
 	{
 		UE_LOG_REF(
 			LogCategory,
-			Warning,
-			TEXT("Channel '%s' is a read only channel. It has been disabled, but cannot be enabled again."),
+			Error,
+			TEXT("Channel '%s' is read only. It is not allowed to manually disable this channel."),
 			Channel
 		);
+		return;
 	}
 
-	FPlatformEventsTrace::OnTraceChannelUpdated(Channel, false);
-
-	UE::Trace::ToggleChannel(Channel, false);
+	bool bIsEnabled = UE::Trace::ToggleChannel(Channel, false);
+	FPlatformEventsTrace::OnTraceChannelUpdated(Channel, bIsEnabled);
 }
 
 ////////////////////////////////////////////////////////////////////////////////
@@ -853,7 +833,7 @@ static void SetupInitFromConfig(UE::Trace::FInitializeDesc& OutDesc)
 
 	// Note that these options can only be used when tracing from a forked process (e.g. server).
 	// For a regular process use command line argument -TraceThreadSleepTime and -TraceTailSizeMb
-	
+
 	int32 SleepTimeConfig = 0;
 	if (GConfig->GetInt(GTraceConfigSection, TEXT("SleepTimeInMS"), SleepTimeConfig, GEngineIni))
 	{
@@ -1321,7 +1301,7 @@ static bool StartFromCommandlineArguments(const TCHAR* CommandLine, bool& bOutSt
 
 	// Trace's worker thread should really only be started by Trace itself as
 	// order is important. At the very least it must be done after Trace is
-	// initialised. It isn't yet here so we defer it.
+	// initialized. It isn't yet here so we defer it.
 	bOutStartWorkerThread = !Opts.bNoWorkerThread;
 	Opts.bNoWorkerThread = true;
 
@@ -1516,8 +1496,8 @@ void FTraceAuxiliary::Initialize(const TCHAR* CommandLine)
 	}
 #endif
 
-	// Trace out information about this session. This is done before initialisation,
-	// so that it is always sent (all channels are enabled prior to initialisation).
+	// Trace out information about this session. This is done before initialization,
+	// so that it is always sent (all channels are enabled prior to initialization).
 	const TCHAR* BranchName = BuildSettings::GetBranchName();
 	const TCHAR* BuildVersion = BuildSettings::GetBuildVersion();
 	constexpr uint32 PlatformLen = UE_ARRAY_COUNT(PREPROCESSOR_TO_STRING(UBT_COMPILED_PLATFORM)) - 1;
@@ -1545,7 +1525,7 @@ void FTraceAuxiliary::Initialize(const TCHAR* CommandLine)
 		<< Session2.TargetType(uint8(FApp::GetBuildTargetType()));
 
 	// Attempt to send trace data somewhere from the command line. It perhaps
-	// seems odd to do this before initialising Trace, but it is done this way
+	// seems odd to do this before initializing Trace, but it is done this way
 	// to support disabling the "important" cache without losing any events.
 	bool bShouldStartWorkerThread = false;
 	StartFromCommandlineArguments(CommandLine, bShouldStartWorkerThread);
diff --git a/Engine/Source/Runtime/Core/Public/ProfilingDebugging/TraceAuxiliary.h b/Engine/Source/Runtime/Core/Public/ProfilingDebugging/TraceAuxiliary.h
--- a/Engine/Source/Runtime/Core/Public/ProfilingDebugging/TraceAuxiliary.h
+++ b/Engine/Source/Runtime/Core/Public/ProfilingDebugging/TraceAuxiliary.h
@@ -52,17 +52,17 @@ public:
 	/** Callback whenever a trace is started */
 	DECLARE_TS_MULTICAST_DELEGATE_TwoParams(FOnTraceStarted, FTraceAuxiliary::EConnectionType TraceType, const FString& TraceDestination);
 
-	/** 
-	* Callback whenever a trace recording is stopped. 
-	* TraceType tells what kind of trace it is.
-	* TraceDestination will be either the the filename and path for a file trace or the network connection for a network trace
-	*/
+	/**
+	 * Callback whenever a trace recording is stopped.
+	 * TraceType tells what kind of trace it is.
+	 * TraceDestination will be either the filename and path for a file trace or the network connection for a network trace
+	 */
 	DECLARE_TS_MULTICAST_DELEGATE_TwoParams(FOnTraceStopped, FTraceAuxiliary::EConnectionType TraceType, const FString& TraceDestination);
 
 	/**
-	* Callback whenever a trace snapshot is saved.
-	* Path is the file system path of the snapshot file.
-	*/
+	 * Callback whenever a trace snapshot is saved.
+	 * Path is the file system path of the snapshot file.
+	 */
 	DECLARE_TS_MULTICAST_DELEGATE_TwoParams(FOnSnapshotSaved, FTraceAuxiliary::EConnectionType TraceType, const FString& TraceDestination);
 
 	struct FOptions
@@ -99,8 +99,8 @@ public:
 	static CORE_API bool Pause();
 
 	/**
-	* @return True if trace was paused and the list of channels to resume exists.
-	*/
+	 * @return True if trace was paused and the list of channels to resume exists.
+	 */
 	static CORE_API bool IsPaused();
 
 	/**
@@ -113,7 +113,7 @@ public:
 	 * @param FilePath Path to the file to write the snapshot to. If it is null or empty a file path will be generated.
 	 */
 	static CORE_API bool WriteSnapshot(const TCHAR* FilePath);
-	
+
 	/**
 	 * Write tailing memory state to a trace server.
 	 * @param FilePath Path to the file to write the snapshot to. If it is null or empty a file path will be generated.
@@ -150,9 +150,9 @@ public:
 	static CORE_API void EnableChannels();
 
 	/**
-	* Disable channels to stop recording traces with them.
-	* @param Channels List of channels (or a preset) to disable. If null it will disable all active channels.
-	*/
+	 * Disable channels to stop recording traces with them.
+	 * @param Channels List of channels (or a preset) to disable. If null it will disable all active channels.
+	 */
 	static CORE_API void DisableChannels(const TCHAR* Channels = nullptr);
 
 	/**
@@ -160,19 +160,18 @@ public:
 	 *  Contains either a file path or network address. Points to an empty string if tracing is disabled.
 	 */
 	UE_DEPRECATED(5.3, "Use GetTraceDestinationString instead.")
-	static CORE_API const TCHAR*	GetTraceDestination();
-	
+	static CORE_API const TCHAR* GetTraceDestination();
+
 	/**
 	 *  Returns the destination string that is currently being traced to.
 	 *  Contains either a file path or network address. Empty if tracing is disabled.
 	 */
 	static CORE_API FString GetTraceDestinationString();
-	
+
 	/**
 	 *  Returns whether the trace system is currently connected to a trace sink (file or network)
 	 */
-	static CORE_API bool	IsConnected();
-
+	static CORE_API bool IsConnected();
 
 	/**
 	 * Returns whether the trace system is currently connected. If connected writes the session/trace identifiers.
@@ -183,14 +182,14 @@ public:
 	static CORE_API bool IsConnected(FGuid& OutSessionGuid, FGuid& OutTraceGuid);
 
 	/**
-	*  Returns the current connection type.
-	*/
+	 * Returns the current connection type.
+	 */
 	static CORE_API EConnectionType GetConnectionType();
 
 	/**
-	 *  Adds a comma separated list of currently active channels to the passed in StringBuilder
+	 * Adds a comma separated list of currently active channels to the passed in StringBuilder
 	 */
-	static CORE_API void	GetActiveChannelsString(FStringBuilderBase& String);
+	static CORE_API void GetActiveChannelsString(FStringBuilderBase& String);
 
 	/**
 	 * Used when process is panicking. Stops all tracing immediately to avoid further allocations. Process is not
@@ -204,10 +203,9 @@ public:
 	 *  * Only NoSync event types can be emitted.
 	 *  * Important events should not be emitted. They will appear after the events in the tail.
 	 *  * Callback is issued from a worker thread. User is responsible to synchronize shared resources.
-	 * 
-	 * @note This is an advanced feature to avoid using important events in cases where event data can be
-	 *		 recalled easily.
-	 *		 
+	 *
+	 * @note This is an advanced feature to avoid using important events in cases where event data can be recalled easily.
+	 *
 	 * @param Callback Delegate to call on new connections.
 	 */
 	static CORE_API FOnConnection OnConnection;
@@ -229,7 +227,6 @@ public:
 	 * The path to the snapshot file is passed to the delegate.
 	 */
 	static CORE_API FOnSnapshotSaved OnSnapshotSaved;
-	
 };
 
 #if UE_TRACE_SERVER_CONTROLS_ENABLED
@@ -241,10 +238,11 @@ class FTraceServerControls
 {
 public:
 	/**
-	 * Launch the server using the "fork" command. This spins off an separate running process.	 * 
+	 * Launch the server using the "fork" command. This spins off a separate running process.
 	 * @return True if the server was successfully started or already running.
 	 */
 	static CORE_API bool Start();
+
 	/**
 	 * Stop any running instance of the server.
 	 * @return True if the stop command was successful. False otherwise.
@@ -252,4 +250,4 @@ public:
 	static CORE_API bool Stop();
 };
 
-#endif // UE_TRACE_SERVER_CONTROLS_ENABLED
\ No newline at end of file
+#endif // UE_TRACE_SERVER_CONTROLS_ENABLED
diff --git a/Engine/Source/Runtime/TraceLog/Private/Trace/Channel.cpp b/Engine/Source/Runtime/TraceLog/Private/Trace/Channel.cpp
--- a/Engine/Source/Runtime/TraceLog/Private/Trace/Channel.cpp
+++ b/Engine/Source/Runtime/TraceLog/Private/Trace/Channel.cpp
@@ -156,8 +156,8 @@ void FChannel::Setup(const ANSICHAR* InChannelName, const InitArgs& InArgs)
 		}
 	}
 
-	// If channel is initialized after the all channels are disabled (post static init)
-	// this channel needs to be disabled.
+	// If channel is initialized after all channels are disabled (post static init)
+	// this channel needs to be disabled too.
 	if (GChannelsInitialized)
 	{
 		Enabled = -1;
@@ -177,8 +177,11 @@ void FChannel::Announce() const
 ///////////////////////////////////////////////////////////////////////////////
 void FChannel::Initialize()
 {
-	// All channels are initialized as enabled (zero), and act like so during
-	// from process start until this method is called (i.e. when Trace is initalized).
+	// During static initialization, all channels are created as enabled (zero),
+	// and act like so from the process start until this method is called (i.e. when Trace is initialized).
+	// Now we can disable all channels.
+	// Channels specified on the command line (using -trace=<channels> argument)
+	// will be further re-enabled after this call.
 	ToggleAll(false);
 	GChannelsInitialized = true;
 }
@@ -248,24 +251,6 @@ FChannel* FChannel::FindChannel(const ANSICHAR* ChannelName)
 	return nullptr;
 }
 
-///////////////////////////////////////////////////////////////////////////////
-void FChannel::EnumerateChannels(ChannelIterFunc Func, void* User) 
-{
-	using namespace Private;
-	FChannel* ChannelLists[] =
-	{
-		AtomicLoadAcquire(&GNewChannelList),
-		AtomicLoadAcquire(&GHeadChannel),
-	};
-	for (FChannel* Channel : ChannelLists)
-	{
-		for (; Channel != nullptr; Channel = Channel->Next)
-		{
-			Func(Channel->Name.Ptr, Channel->IsEnabled(), User);
-		}
-	}
-}
-
 ///////////////////////////////////////////////////////////////////////////////
 void FChannel::EnumerateChannels(ChannelIterCallback Func, void* User)
 {
diff --git a/Engine/Source/Runtime/TraceLog/Public/Trace/Detail/Channel.h b/Engine/Source/Runtime/TraceLog/Public/Trace/Detail/Channel.h
--- a/Engine/Source/Runtime/TraceLog/Public/Trace/Detail/Channel.h
+++ b/Engine/Source/Runtime/TraceLog/Public/Trace/Detail/Channel.h
@@ -19,7 +19,6 @@ struct FChannelInfo
 	bool bIsReadOnly;
 };
 
-typedef void ChannelIterFunc(const ANSICHAR*, bool, void*);	
 typedef bool ChannelIterCallback(const FChannelInfo& OutChannelInfo, void*);
 
 /*
@@ -32,10 +31,9 @@ typedef bool ChannelIterCallback(const FChannelInfo& OutChannelInfo, void*);
 
 	Note that this works as an AND operator, similar to how a bitmask is constructed.
 
-	Channels are by default enabled until this method is called. This is to allow
+	Channels are by default enabled until FChannel::Initialize() is called. This is to allow
 	events to be emitted during static initialization. In fact all events during
-	this phase are always emitted. In this method we disable all channels except
-	those specified on the command line using -tracechannels argument.
+	this phase are always emitted.
 */
 class FChannel
 {
@@ -49,9 +47,8 @@ public:
 
 	struct InitArgs
 	{
-		const ANSICHAR* 	Desc;		// User facing description string
+		const ANSICHAR*		Desc;		// User facing description string
 		bool				bReadOnly;	// If set, channel cannot be changed during a run, only set through command line.
-
 	};
 
 	TRACELOG_API void	Setup(const ANSICHAR* InChannelName, const InitArgs& Args);
@@ -62,8 +59,6 @@ public:
 	static void			ToggleAll(bool bEnabled);
 	static void			PanicDisableAll(); // Disabled channels wont be logged with UE_TRACE_LOG
 	static FChannel*	FindChannel(const ANSICHAR* ChannelName);
-	UE_DEPRECATED(5.2, "Please use the ChannelIterCallback overload for enumerating channels.")
-	static void			EnumerateChannels(ChannelIterFunc Func, void* User);
 	static void			EnumerateChannels(ChannelIterCallback Func, void* User);
 	bool				Toggle(bool bEnabled);
 	bool				IsEnabled() const;
