- Fixed a threading crash with ConfigCacheIni::RegisteredPlugins

diff --git a/Engine/Source/Runtime/Core/Private/Misc/ConfigCacheIni.cpp b/Engine/Source/Runtime/Core/Private/Misc/ConfigCacheIni.cpp
--- a/Engine/Source/Runtime/Core/Private/Misc/ConfigCacheIni.cpp
+++ b/Engine/Source/Runtime/Core/Private/Misc/ConfigCacheIni.cpp
@@ -5630,6 +5630,7 @@ FCriticalSection FConfigCacheIni::ConfigForPlatformLock;
 #endif
 
 TMap<FName, FConfigCacheIni::FPluginInfo*> FConfigCacheIni::RegisteredPlugins;
+FCriticalSection FConfigCacheIni::RegisteredPluginsLock;
 
 
 void FConfigCacheIni::AddPluginToAllBranches(FName PluginName, FConfigModificationTracker* ModificationTracker)
@@ -5664,11 +5665,16 @@ void FConfigCacheIni::AddPluginToBranches(FName PluginName, FConfigModificationT
 {
 	// @todo make sure we are still pending
 	
-	FPluginInfo* PluginInfo = RegisteredPlugins.FindRef(PluginName);
-	if (PluginInfo == nullptr)
+	FPluginInfo* PluginInfo = nullptr;
 	{
-		UE_LOG(LogConfig, Warning, TEXT("Attempting to load a dynamic plugin (%s) that was not registered ahead of time!"), *PluginName.ToString());
-		return;
+		FScopeLock Lock(&RegisteredPluginsLock);
+		
+		PluginInfo = RegisteredPlugins.FindRef(PluginName);
+		if (PluginInfo == nullptr)
+		{
+			UE_LOG(LogConfig, Warning, TEXT("Attempting to load a dynamic plugin (%s) that was not registered ahead of time!"), *PluginName.ToString());
+			return;
+		}
 	}
 	
 	TArray<FString> PluginConfigs;
@@ -5835,11 +5841,13 @@ void FConfigCacheIni::RegisterPlugin(FName PluginName, const FString& PluginDir,
 {
 	FPluginInfo* Info = new FPluginInfo();
 
-	RegisteredPlugins.Add(PluginName, Info);
 	Info->PluginDir = PluginDir;
 	Info->ChildPluginDirs = ChildPluginDirs;
 	Info->Priority = Priority;
 	Info->bIncludePluginNameInBranchName = bIncludePluginNameInBranchName;
+
+	FScopeLock Lock(&RegisteredPluginsLock);
+	RegisteredPlugins.Add(PluginName, Info);
 }
 
 
diff --git a/Engine/Source/Runtime/Core/Private/Misc/ConfigContext.cpp b/Engine/Source/Runtime/Core/Private/Misc/ConfigContext.cpp
--- a/Engine/Source/Runtime/Core/Private/Misc/ConfigContext.cpp
+++ b/Engine/Source/Runtime/Core/Private/Misc/ConfigContext.cpp
@@ -83,6 +83,8 @@ void FConfigContext::CachePaths()
 		// are we loading a plugin?
 		if (ConfigSystem != nullptr)
 		{
+			FScopeLock Lock(&FConfigCacheIni::RegisteredPluginsLock);
+			
 			FConfigCacheIni::FPluginInfo* PluginInfo = ConfigSystem->RegisteredPlugins.FindRef(ConfigFileTag == NAME_None ? FName(*BaseIniName) : ConfigFileTag);
 			if (PluginInfo != nullptr)
 			{
diff --git a/Engine/Source/Runtime/Core/Public/Misc/ConfigCacheIni.h b/Engine/Source/Runtime/Core/Public/Misc/ConfigCacheIni.h
--- a/Engine/Source/Runtime/Core/Public/Misc/ConfigCacheIni.h
+++ b/Engine/Source/Runtime/Core/Public/Misc/ConfigCacheIni.h
@@ -1766,7 +1766,7 @@ private:
 	};
 	static TMap<FName, FPluginInfo*> RegisteredPlugins;
 	TArray<FName> PendingModificationPlugins;
-	FCriticalSection RegisteredPluginsLock;
+	static FCriticalSection RegisteredPluginsLock;
 	
 #if ALLOW_OTHER_PLATFORM_CONFIG
 	static TMap<FName, FConfigCacheIni*> ConfigForPlatform;
