Fix failure check on Macs

#rnx

diff --git a/Engine/Source/Runtime/Core/Private/Apple/ApplePlatformFile.cpp b/Engine/Source/Runtime/Core/Private/Apple/ApplePlatformFile.cpp
--- a/Engine/Source/Runtime/Core/Private/Apple/ApplePlatformFile.cpp
+++ b/Engine/Source/Runtime/Core/Private/Apple/ApplePlatformFile.cpp
@@ -515,7 +515,7 @@ FDateTime FApplePlatformFile::GetTimeStamp(const TCHAR* Filename)
 {
 	// get file times
 	struct stat FileInfo;
-	if(Stat(Filename, &FileInfo) == -1)
+	if(Stat(Filename, &FileInfo) != 0)
 	{
 		return FDateTime::MinValue();
 	}
@@ -530,7 +530,7 @@ void FApplePlatformFile::SetTimeStamp(const TCHAR* Filename, const FDateTime Dat
 {
 	// get file times
 	struct stat FileInfo;
-	if (Stat(Filename, &FileInfo) == 0)
+	if (Stat(Filename, &FileInfo) != 0)
 	{
 		return;
 	}
