Don't reset the LocRes ID when restoring the pre-leet state

This is restoring the pre-leet display string, so the LocRes ID (which isn't cleared when leetifying the text) is still correct for the original display string

#rnx

diff --git a/Engine/Source/Runtime/Core/Private/Internationalization/TextLocalizationManager.cpp b/Engine/Source/Runtime/Core/Private/Internationalization/TextLocalizationManager.cpp
--- a/Engine/Source/Runtime/Core/Private/Internationalization/TextLocalizationManager.cpp
+++ b/Engine/Source/Runtime/Core/Private/Internationalization/TextLocalizationManager.cpp
@@ -1606,9 +1606,6 @@ void FTextLocalizationManager::UpdateFromNative(FTextLocalizationResource&& Text
 				DisplayStringBackupTable.RemoveAndCopyValue(TextId, DisplayStringBackup))
 			{
 				LiveEntry.DisplayString = DisplayStringBackup.ToSharedRef();
-#if WITH_EDITORONLY_DATA
-				LiveEntry.LocResID = FTextKey();
-#endif	// WITH_EDITORONLY_DATA
 			}
 #endif	// ENABLE_LOC_TESTING
 		}
