Added TRefCountingMixin to provide AutoRTFM aware ref-counting for use with TRefCountPtr

Note: This is essentially a copy of the TReferenceControllerBase implementation for TSharedPtr

diff --git a/Engine/Source/Runtime/Core/Public/Templates/RefCounting.h b/Engine/Source/Runtime/Core/Public/Templates/RefCounting.h
--- a/Engine/Source/Runtime/Core/Public/Templates/RefCounting.h
+++ b/Engine/Source/Runtime/Core/Public/Templates/RefCounting.h
@@ -3,6 +3,7 @@
 #pragma once
 
 #include "CoreTypes.h"
+#include "AutoRTFM/AutoRTFM.h"
 #include "HAL/PlatformAtomics.h"
 #include "HAL/PreprocessorHelpers.h"
 #include "HAL/ThreadSafeCounter.h"
@@ -12,6 +13,8 @@
 #include "Serialization/MemoryLayout.h"
 #include "Templates/TypeHash.h"
 #include "Templates/Requires.h"
+#include <atomic>
+#include <type_traits>
 
 /** A virtual interface for ref counted objects to implement. */
 class IRefCountedObject
@@ -135,7 +138,152 @@ private:
 	mutable FThreadSafeCounter NumRefs;
 };
 
+/**
+ * ERefCountingMode is used select between either 'fast' or 'thread safe' ref-counting types.
+ * This is only used by templates at compile time to generate one code path or another.
+ */
+enum class ERefCountingMode : uint8
+{
+	/** Forced to be not thread-safe. */
+	NotThreadSafe = 0,
+
+	/** Thread-safe, never spin locks, but slower */
+	ThreadSafe = 1
+};
 
+/**
+ * Ref-counting mixin, designed to add ref-counting to an object without requiring a virtual destructor.
+ * Implements support for AutoRTFM, is thread-safe by default, and can support custom deleters via T::StaticDestroyObject.
+ * 
+ * @note AutoRTFM means that the return value of AddRef/Release is nonsense, but this is fine for use with TRefCountPtr.
+ * 
+ * Basic Example:
+ *  struct FMyRefCountedObject : public TRefCountingMixin<FMyRefCountedObject>
+ *  {
+ *      // ...
+ *  };
+ * 
+ * Deleter Example:
+ *  struct FMyRefCountedPooledObject : public TRefCountingMixin<FMyRefCountedPooledObject>
+ *  {
+ *      static void StaticDestroyObject(const FMyRefCountedPooledObject* Obj)
+ *      {
+ *          GPool->ReturnToPool(Obj);
+ *      }
+ *  };
+ */
+template <typename T, ERefCountingMode Mode = ERefCountingMode::ThreadSafe>
+class TRefCountingMixin
+{
+	using RefCountType = std::conditional_t<Mode == ERefCountingMode::ThreadSafe, std::atomic<uint32>, uint32>;
+
+public:
+	TRefCountingMixin() = default;
+
+	TRefCountingMixin(const TRefCountingMixin&) = delete;
+	TRefCountingMixin& operator=(const TRefCountingMixin&) = delete;
+
+	uint32 AddRef() const
+	{
+		if constexpr (Mode == ERefCountingMode::ThreadSafe)
+		{
+			// Incrementing a reference count with relaxed ordering is always safe because no other action is taken
+			// in response to the increment, so there's nothing to order with.
+
+#if defined(_MSC_VER) && (defined(_M_X64) || defined(_M_IX86))
+			UE_AUTORTFM_OPEN(
+			{
+				// We do a regular SC increment here because it maps to an _InterlockedIncrement (lock inc).
+				// The codegen for a relaxed fetch_add is actually much worse under MSVC (lock xadd).
+				++RefCount;
+			});
+#else
+			UE_AUTORTFM_OPEN(
+			{
+				RefCount.fetch_add(1, std::memory_order_relaxed);
+			});
+#endif
+
+			// If the transaction would abort, we need to undo adding the reference.
+			AutoRTFM::OnAbort([this]
+			{
+				Release();
+			});
+		}
+		else
+		{
+			++RefCount;
+		}
+
+		// Note: TRefCountPtr doesn't use the return value
+		return 0;
+	}
+
+	uint32 Release() const
+	{
+		if constexpr (Mode == ERefCountingMode::ThreadSafe)
+		{
+			AutoRTFM::OnCommit([this]
+			{
+				// std::memory_order_acq_rel is used here so that, if we do end up executing the destructor, it's not possible
+				// for side effects from executing the destructor end up being visible before we've determined that the 
+				// reference count is actually zero.
+
+				uint32 OldCount = RefCount.fetch_sub(1, std::memory_order_acq_rel);
+				checkSlow(OldCount > 0);
+				if (OldCount == 1)
+				{
+					T::StaticDestroyObject(static_cast<const T*>(this));
+				}
+			});
+		}
+		else
+		{
+			checkSlow(RefCount > 0);
+
+			if (--RefCount == 0)
+			{
+				T::StaticDestroyObject(static_cast<const T*>(this));
+			}
+		}
+
+		// Note: TRefCountPtr doesn't use the return value
+		return 0;
+	}
+
+	uint32 GetRefCount() const
+	{
+		if constexpr (Mode == ERefCountingMode::ThreadSafe)
+		{
+			// A 'live' reference count is unstable by nature and so there's no benefit
+			// to try and enforce memory ordering around the reading of it.
+			//
+			// This is equivalent to https://en.cppreference.com/w/cpp/memory/shared_ptr/use_count
+
+			uint32 Count = 0;
+
+			UE_AUTORTFM_OPEN(
+			{
+				// This reference count may be accessed by multiple threads
+				Count = RefCount.load(std::memory_order_relaxed);
+			});
+
+			return Count;
+		}
+		else
+		{
+			return RefCount;
+		}
+	}
+
+	static void StaticDestroyObject(const T* Obj)
+	{
+		delete Obj;
+	}
+
+private:
+	mutable RefCountType RefCount{0};
+};
 
 /**
  * A smart pointer to an object which implements AddRef/Release.
