Updated FRefCountedDisplayString to use TRefCountingMixin so that it can support AutoRTFM

#rnx

diff --git a/Engine/Source/Runtime/Core/Public/Internationalization/LocalizedTextSourceTypes.h b/Engine/Source/Runtime/Core/Public/Internationalization/LocalizedTextSourceTypes.h
--- a/Engine/Source/Runtime/Core/Public/Internationalization/LocalizedTextSourceTypes.h
+++ b/Engine/Source/Runtime/Core/Public/Internationalization/LocalizedTextSourceTypes.h
@@ -75,7 +75,7 @@ struct ELocalizedTextSourcePriority
 };
 
 #ifndef UE_TEXT_DISPLAYSTRING_USE_REFCOUNTPTR
-	#define UE_TEXT_DISPLAYSTRING_USE_REFCOUNTPTR (0)
+	#define UE_TEXT_DISPLAYSTRING_USE_REFCOUNTPTR (1)
 #endif
 
 #if UE_TEXT_DISPLAYSTRING_USE_REFCOUNTPTR
@@ -83,43 +83,14 @@ struct ELocalizedTextSourcePriority
 namespace UE::Text::Private
 {
 
-class FRefCountedDisplayString
+class FRefCountedDisplayString : public TRefCountingMixin<FRefCountedDisplayString>
 {
 public:
 	explicit FRefCountedDisplayString(FString&& InDisplayString)
 		: DisplayString(MoveTemp(InDisplayString))
-		, NumRefs(0)
 	{
 	}
 
-	~FRefCountedDisplayString()
-	{
-		check(NumRefs.GetValue() == 0);
-	}
-
-	FRefCountedDisplayString(const FRefCountedDisplayString&) = delete;
-	FRefCountedDisplayString& operator=(const FRefCountedDisplayString&) = delete;
-
-	uint32 AddRef() const
-	{
-		return uint32(NumRefs.Increment());
-	}
-
-	uint32 Release() const
-	{
-		uint32 Refs = uint32(NumRefs.Decrement());
-		if (Refs == 0)
-		{
-			delete this;
-		}
-		return Refs;
-	}
-
-	uint32 GetRefCount() const
-	{
-		return uint32(NumRefs.GetValue());
-	}
-
 	[[nodiscard]] FString& Private_GetDisplayString()
 	{
 		return DisplayString;
@@ -127,7 +98,6 @@ public:
 
 private:
 	FString DisplayString;
-	mutable FThreadSafeCounter NumRefs;
 };
 
 /**
