Workaround for task crash on behalf of [REDACTED_HANDLE], optionally adds a reference based on a CVar to the task before we decrement the lock count (which was enabling the task to be retracted, and potentially destroyed, by another thread).

#rnx
#tests Preflight with replay run and verified assembly.

diff --git a/Engine/Source/Runtime/Core/Private/Tasks/TaskPrivate.cpp b/Engine/Source/Runtime/Core/Private/Tasks/TaskPrivate.cpp
--- a/Engine/Source/Runtime/Core/Private/Tasks/TaskPrivate.cpp
+++ b/Engine/Source/Runtime/Core/Private/Tasks/TaskPrivate.cpp
@@ -4,6 +4,7 @@
 #include "Tasks/Pipe.h"
 
 #include "Async/TaskGraphInterfaces.h"
+#include "HAL/IConsoleManager.h"
 
 namespace UE::Tasks
 {
@@ -535,5 +536,12 @@ namespace UE::Tasks
 		}
 
 #endif // !TASKGRAPH_NEW_FRONTEND
+
+		bool GAddReferenceInTryUnlock = true;
+		static FAutoConsoleVariableRef AddReferenceInTryUnlockCVar(
+			TEXT("task.AddReferenceInTryUnlock"),
+			GAddReferenceInTryUnlock,
+			TEXT("Whether to enable the workaround where we add a reference in TryUnlock before decrementing lock count"),
+			ECVF_Default);
 	}
 }
diff --git a/Engine/Source/Runtime/Core/Public/Tasks/TaskPrivate.h b/Engine/Source/Runtime/Core/Public/Tasks/TaskPrivate.h
--- a/Engine/Source/Runtime/Core/Public/Tasks/TaskPrivate.h
+++ b/Engine/Source/Runtime/Core/Public/Tasks/TaskPrivate.h
@@ -112,6 +112,9 @@ namespace UE::Tasks
 		UE_DEPRECATED(5.1, "You should not use this function as it exists only to patch another system and can be removed any time.")
 		CORE_API bool IsThreadRetractingTask();
 
+		// For workaround: see Private::FTaskBase::TryUnlock
+		extern CORE_API bool GAddReferenceInTryUnlock;
+
 		// An abstract base class for task implementation. 
 		// Implements internal logic of task prerequisites, nested tasks and deep task retraction.
 		// Implements intrusive ref-counting and so can be used with TRefCountPtr.
@@ -564,6 +567,21 @@ public:
 
 				FPipe* LocalPipe = GetPipe(); // cache data locally so we won't need to touch the member (read below)
 
+				// TEMP: Patch to prevent the task getting destroyed. If NumLocks reaches 0 and retraction wins the execution, it could end
+				// up releasing the last reference before we continue and exit this function.
+				const bool bHasAddedReference = GAddReferenceInTryUnlock;
+				if (bHasAddedReference)
+				{
+					AddRef();
+				}
+				ON_SCOPE_EXIT
+				{
+					if (bHasAddedReference)
+					{
+						Release();
+					}
+				};
+
 				uint32 PrevNumLocks = NumLocks.fetch_sub(1, std::memory_order_acq_rel); // `acq_rel` to make it happen after task 
 				// preparation and before launching it
 				// the task can be dead already as the prev line can remove the lock hold for this execution path, another thread(s) can unlock
