Fixed static analysis warning about dereferencing null.

diff --git a/Engine/Source/Runtime/Core/Public/Containers/RingBuffer.h b/Engine/Source/Runtime/Core/Public/Containers/RingBuffer.h
--- a/Engine/Source/Runtime/Core/Public/Containers/RingBuffer.h
+++ b/Engine/Source/Runtime/Core/Public/Containers/RingBuffer.h
@@ -789,43 +789,43 @@ private:
 		{
 			NewData = reinterpret_cast<ElementType*>(FMemory::Malloc(sizeof(ElementType) * NewCapacity, alignof(ElementType)));
 			NewIndexMask = NewCapacity - 1;
-		}
-		else
-		{
-			NewData = nullptr;
-			NewIndexMask = static_cast<StorageModuloType>(-1);
-		}
-		if (SrcNum > 0)
-		{
-			// move data to new storage
-			const StorageModuloType MaskedFront = Front & IndexMask;
-			const StorageModuloType MaskedAfterBack = AfterBack & IndexMask;
-
-			// MaskedFront equal to MaskedAfterBack will occur if the queue's Num equals Capacity or if the Num == 0.  We checked Num == 0 above, so here it means Num == Capacity.
-			if (MaskedFront >= MaskedAfterBack)
+			if (SrcNum > 0)
 			{
-				StorageModuloType WriteIndex = 0;
-				for (StorageModuloType ReadIndex = MaskedFront; ReadIndex < SrcCapacity; ++ReadIndex)
-				{
-					::new ((void*)&NewData[WriteIndex++]) ElementType(MoveTemp(SrcData[ReadIndex]));
-				}
-				DestructRange(MaskedFront, SrcCapacity);
-				for (StorageModuloType ReadIndex = 0; ReadIndex < MaskedAfterBack; ++ReadIndex)
+				// move data to new storage
+				const StorageModuloType MaskedFront = Front & IndexMask;
+				const StorageModuloType MaskedAfterBack = AfterBack & IndexMask;
+
+				// MaskedFront equal to MaskedAfterBack will occur if the queue's Num equals Capacity or if the Num == 0.  We checked Num == 0 above, so here it means Num == Capacity.
+				if (MaskedFront >= MaskedAfterBack)
 				{
-					::new ((void*)&NewData[WriteIndex++]) ElementType(MoveTemp(SrcData[ReadIndex]));
+					StorageModuloType WriteIndex = 0;
+					for (StorageModuloType ReadIndex = MaskedFront; ReadIndex < SrcCapacity; ++ReadIndex)
+					{
+						::new ((void*)&NewData[WriteIndex++]) ElementType(MoveTemp(SrcData[ReadIndex]));
+					}
+					DestructRange(MaskedFront, SrcCapacity);
+					for (StorageModuloType ReadIndex = 0; ReadIndex < MaskedAfterBack; ++ReadIndex)
+					{
+						::new ((void*)&NewData[WriteIndex++]) ElementType(MoveTemp(SrcData[ReadIndex]));
+					}
+					DestructRange(0, MaskedAfterBack);
 				}
-				DestructRange(0, MaskedAfterBack);
-			}
-			else
-			{
-				StorageModuloType WriteIndex = 0;
-				for (StorageModuloType ReadIndex = MaskedFront; ReadIndex < MaskedAfterBack; ++ReadIndex)
+				else
 				{
-					::new((void*)&NewData[WriteIndex++]) ElementType(MoveTemp(SrcData[ReadIndex]));
+					StorageModuloType WriteIndex = 0;
+					for (StorageModuloType ReadIndex = MaskedFront; ReadIndex < MaskedAfterBack; ++ReadIndex)
+					{
+						::new((void*)&NewData[WriteIndex++]) ElementType(MoveTemp(SrcData[ReadIndex]));
+					}
+					DestructRange(MaskedFront, MaskedAfterBack);
 				}
-				DestructRange(MaskedFront, MaskedAfterBack);
 			}
 		}
+		else
+		{
+			NewData = nullptr;
+			NewIndexMask = static_cast<StorageModuloType>(-1);
+		}
 		if (AllocationData)
 		{
 			FMemory::Free(AllocationData);
