Removed deprecated function.

diff --git a/Engine/Source/Runtime/Core/Public/Containers/StaticArray.h b/Engine/Source/Runtime/Core/Public/Containers/StaticArray.h
--- a/Engine/Source/Runtime/Core/Public/Containers/StaticArray.h
+++ b/Engine/Source/Runtime/Core/Public/Containers/StaticArray.h
@@ -22,12 +22,6 @@ public:
 	{
 	}
 
-	UE_DEPRECATED(5.0, "Please call TStaticArray(InPlace, DefaultElement) instead.")
-	explicit TStaticArray(const InElementType& DefaultElement)
-		: Storage(InPlace, TMakeIntegerSequence<uint32, NumElements>(), DefaultElement)
-	{
-	}
-
 	template <typename... ArgTypes>
 	explicit TStaticArray(EInPlace, ArgTypes&&... Args)
 		: Storage(InPlace, TMakeIntegerSequence<uint32, NumElements>(), Forward<ArgTypes>(Args)...)
