- Fixed a re-used localization tag

diff --git a/Engine/Source/Runtime/Core/Private/HAL/ConsoleManager.cpp b/Engine/Source/Runtime/Core/Private/HAL/ConsoleManager.cpp
--- a/Engine/Source/Runtime/Core/Private/HAL/ConsoleManager.cpp
+++ b/Engine/Source/Runtime/Core/Private/HAL/ConsoleManager.cpp
@@ -4794,7 +4794,7 @@ private:
 				}
 				else
 				{
-					Message = FText::Format(LOCTEXT("ShadowCVarUsage_Deprecated", "Using a shadowed CVar '{ThisName}'. It is recommended to change all uses to '{ShadowName}' instead."), Arguments);
+					Message = FText::Format(LOCTEXT("ShadowCVarUsage", "Using a shadowed CVar '{ThisName}'. It is recommended to change all uses to '{ShadowName}' instead."), Arguments);
 				}
 				switch (UsageBehavior)
 				{
