Box Sphere Bounds builder to return zero-initialized bounds if not set

diff --git a/Engine/Source/Runtime/Core/Public/Math/BoxSphereBounds.h b/Engine/Source/Runtime/Core/Public/Math/BoxSphereBounds.h
--- a/Engine/Source/Runtime/Core/Public/Math/BoxSphereBounds.h
+++ b/Engine/Source/Runtime/Core/Public/Math/BoxSphereBounds.h
@@ -340,7 +340,7 @@ public:
 
 		FORCEINLINE operator TBoxSphereBounds<T, TExtent>() const
 		{
-			return BoxSphereBounds.Get(TBoxSphereBounds<T, TExtent>());
+			return BoxSphereBounds.Get(TBoxSphereBounds<T, TExtent>(ForceInitToZero));
 		}
 
 	private:
