Changed TStaticArray's default constructor so that it can be default constructed or value constructed by choice (until now, they've always been default constructed):

// Before:
TStaticArray<int32, N> Arr1;      // ints are uninitialized
TStaticArray<int32, N> Arr2 = {}; // ints are uninitialized

// After:
TStaticArray<int32, N> Arr1;      // ints are uninitialized
TStaticArray<int32, N> Arr2 = {}; // ints are zeroed

Fixed a couple of bugs which the compiler highlighted by this change.

diff --git a/Engine/Source/Runtime/Core/Public/Containers/StaticArray.h b/Engine/Source/Runtime/Core/Public/Containers/StaticArray.h
--- a/Engine/Source/Runtime/Core/Public/Containers/StaticArray.h
+++ b/Engine/Source/Runtime/Core/Public/Containers/StaticArray.h
@@ -17,10 +17,7 @@ class alignas(Alignment) TStaticArray
 public:
 	using ElementType = InElementType;
 
-	TStaticArray() 
-		: Storage()
-	{
-	}
+	TStaticArray() = default;
 
 	template <typename... ArgTypes>
 	explicit TStaticArray(EInPlace, ArgTypes&&... Args)
@@ -93,7 +90,7 @@ private:
 
 	struct alignas(Alignment) TArrayStorageElementAligned
 	{
-		TArrayStorageElementAligned() {}
+		TArrayStorageElementAligned() = default;
 
 		// Index is used to achieve pack expansion in TArrayStorage, but is unused here
 		template <typename... ArgTypes>
@@ -107,10 +104,7 @@ private:
 
 	struct TArrayStorage
 	{
-		TArrayStorage()
-			: Elements()
-		{
-		}
+		TArrayStorage() = default;
 
 		template<uint32... Indices, typename... ArgTypes>
 		explicit TArrayStorage(EInPlace, TIntegerSequence<uint32, Indices...>, ArgTypes&&... Args)
diff --git a/Engine/Source/Runtime/Renderer/Private/SceneRendering.cpp b/Engine/Source/Runtime/Renderer/Private/SceneRendering.cpp
--- a/Engine/Source/Runtime/Renderer/Private/SceneRendering.cpp
+++ b/Engine/Source/Runtime/Renderer/Private/SceneRendering.cpp
@@ -1158,7 +1158,7 @@ void SetupPhysicsFieldUniformBufferParameters(const FScene* Scene, FEngineShowFl
 	}
 	else
 	{
-		TStaticArray<FIntVector4, MAX_PHYSICS_FIELD_TARGETS, 16> EmptyTargets;
+		TStaticArray<FIntVector4, MAX_PHYSICS_FIELD_TARGETS, 16> EmptyTargets = {};
 		ViewUniformShaderParameters.PhysicsFieldClipmapBuffer = GWhiteVertexBufferWithSRV->ShaderResourceViewRHI;
 		ViewUniformShaderParameters.PhysicsFieldClipmapCenter = FVector3f::ZeroVector;
 		ViewUniformShaderParameters.PhysicsFieldClipmapDistance = 1.0;
diff --git a/Engine/Source/Runtime/Renderer/Private/ScreenSpaceDenoise.cpp b/Engine/Source/Runtime/Renderer/Private/ScreenSpaceDenoise.cpp
--- a/Engine/Source/Runtime/Renderer/Private/ScreenSpaceDenoise.cpp
+++ b/Engine/Source/Runtime/Renderer/Private/ScreenSpaceDenoise.cpp
@@ -2287,8 +2287,8 @@ public:
 			Settings.MaxInputSPP = FMath::Max(Settings.MaxInputSPP, InputParameters[BatchedSignalId].RayTracingConfig.RayCountPerPixel);
 		}
 
-		TStaticArray<FScreenSpaceDenoiserHistory*, IScreenSpaceDenoiser::kMaxBatchSize> PrevHistories;
-		TStaticArray<FScreenSpaceDenoiserHistory*, IScreenSpaceDenoiser::kMaxBatchSize> NewHistories;
+		TStaticArray<FScreenSpaceDenoiserHistory*, IScreenSpaceDenoiser::kMaxBatchSize> PrevHistories = {};
+		TStaticArray<FScreenSpaceDenoiserHistory*, IScreenSpaceDenoiser::kMaxBatchSize> NewHistories  = {};
 		for (int32 BatchedSignalId = 0; BatchedSignalId < InputParameterCount; BatchedSignalId++)
 		{
 			const FShadowVisibilityParameters& Parameters = InputParameters[BatchedSignalId];
