Allow TStaticArray to be per-element constructible from a braced argument list.

diff --git a/Engine/Source/Runtime/Core/Public/Containers/StaticArray.h b/Engine/Source/Runtime/Core/Public/Containers/StaticArray.h
--- a/Engine/Source/Runtime/Core/Public/Containers/StaticArray.h
+++ b/Engine/Source/Runtime/Core/Public/Containers/StaticArray.h
@@ -19,12 +19,23 @@ public:
 
 	TStaticArray() = default;
 
+	// Constructs each element with Args
 	template <typename... ArgTypes>
 	explicit TStaticArray(EInPlace, ArgTypes&&... Args)
 		: Storage(InPlace, TMakeIntegerSequence<uint32, NumElements>(), Forward<ArgTypes>(Args)...)
 	{
 	}
 
+	// Directly initializes the array with the provided values.
+	template <
+		typename... ArgTypes
+		UE_REQUIRES((sizeof...(ArgTypes) > 0 && sizeof...(ArgTypes) <= NumElements) && (std::is_convertible_v<ArgTypes, InElementType> && ...))
+	>
+	TStaticArray(ArgTypes&&... Args)
+		: Storage(PerElement, Forward<ArgTypes>(Args)...)
+	{
+	}
+
 	TStaticArray(TStaticArray&& Other) = default;
 	TStaticArray(const TStaticArray& Other) = default;
 	TStaticArray& operator=(TStaticArray&& Other) = default;
@@ -92,7 +103,7 @@ private:
 	{
 		TArrayStorageElementAligned() = default;
 
-		// Index is used to achieve pack expansion in TArrayStorage, but is unused here
+		// Index is used to achieve pack expansion in TArrayStorage's first constructor, but is unused here
 		template <typename... ArgTypes>
 		explicit TArrayStorageElementAligned(EInPlace, uint32 /*Index*/, ArgTypes&&... Args)
 			: Element(Forward<ArgTypes>(Args)...)
@@ -119,6 +130,12 @@ private:
 			// a single element.
 		}
 
+		template<typename... ArgTypes>
+		explicit TArrayStorage(EPerElement, ArgTypes&&... Args)
+			: Elements{ TArrayStorageElementAligned(InPlace, 0 /* dummy index */, Forward<ArgTypes>(Args))... }
+		{
+		}
+
 		TArrayStorageElementAligned Elements[NumElements];
 	};
 
diff --git a/Engine/Source/Runtime/Core/Public/Misc/CoreMiscDefines.h b/Engine/Source/Runtime/Core/Public/Misc/CoreMiscDefines.h
--- a/Engine/Source/Runtime/Core/Public/Misc/CoreMiscDefines.h
+++ b/Engine/Source/Runtime/Core/Public/Misc/CoreMiscDefines.h
@@ -140,6 +140,7 @@ enum EForceInit
 };
 enum ENoInit {NoInit};
 enum EInPlace {InPlace};
+enum EPerElement {PerElement};
 
 #endif // RC_INVOKED
 
