Fixed double delete in TRefCountingMixin when AutoRTFM aborted a transaction

We now defer both AddRef and Release until OnCommit. This avoids needing to use OnAbort at all, as the ref-count never changes if the transaction is aborted.

This does mean that the observed ref-count won't be "accurate" during a transaction, though nothing should be trying to observe a thread-safe ref-count anyway.

#rnx

diff --git a/Engine/Source/Runtime/Core/Public/Templates/RefCounting.h b/Engine/Source/Runtime/Core/Public/Templates/RefCounting.h
--- a/Engine/Source/Runtime/Core/Public/Templates/RefCounting.h
+++ b/Engine/Source/Runtime/Core/Public/Templates/RefCounting.h
@@ -155,7 +155,8 @@ enum class ERefCountingMode : uint8
  * Ref-counting mixin, designed to add ref-counting to an object without requiring a virtual destructor.
  * Implements support for AutoRTFM, is thread-safe by default, and can support custom deleters via T::StaticDestroyObject.
  * 
- * @note AutoRTFM means that the return value of AddRef/Release is nonsense, but this is fine for use with TRefCountPtr.
+ * @note AutoRTFM means that the return value of AddRef/Release is nonsense (as the ref-count doesn't change until the
+ *       transaction is committed), but this is fine for use with TRefCountPtr (as it doesn't use those return values).
  * 
  * Basic Example:
  *  struct FMyRefCountedObject : public TRefCountingMixin<FMyRefCountedObject>
@@ -190,24 +191,15 @@ public:
 			// Incrementing a reference count with relaxed ordering is always safe because no other action is taken
 			// in response to the increment, so there's nothing to order with.
 
-#if defined(_MSC_VER) && (defined(_M_X64) || defined(_M_IX86))
-			UE_AUTORTFM_OPEN(
+			AutoRTFM::OnCommit([this]
 			{
+#if defined(_MSC_VER) && (defined(_M_X64) || defined(_M_IX86))
 				// We do a regular SC increment here because it maps to an _InterlockedIncrement (lock inc).
 				// The codegen for a relaxed fetch_add is actually much worse under MSVC (lock xadd).
 				++RefCount;
-			});
 #else
-			UE_AUTORTFM_OPEN(
-			{
 				RefCount.fetch_add(1, std::memory_order_relaxed);
-			});
 #endif
-
-			// If the transaction would abort, we need to undo adding the reference.
-			AutoRTFM::OnAbort([this]
-			{
-				Release();
 			});
 		}
 		else
