Fix overridable serialization
Initialize new archieve flag

diff --git a/Engine/Source/Runtime/Core/Private/Serialization/Archive.cpp b/Engine/Source/Runtime/Core/Private/Serialization/Archive.cpp
--- a/Engine/Source/Runtime/Core/Private/Serialization/Archive.cpp
+++ b/Engine/Source/Runtime/Core/Private/Serialization/Archive.cpp
@@ -191,6 +191,7 @@ void FArchiveState::Reset()
 	ArIsNetArchive						= false;
 	ArCustomPropertyList				= nullptr;
 	ArUseCustomPropertyList				= false;
+	ArMergeOverrides					= false;
 	ArShouldSkipUpdateCustomVersion		= false;
 	CookData							= nullptr;
 	SerializedProperty					= nullptr;
@@ -250,6 +251,7 @@ void FArchiveState::CopyTrivialFArchiveStatusMembers(const FArchiveState& Archiv
 	ArIsNetArchive                       = ArchiveToCopy.ArIsNetArchive;
 	ArCustomPropertyList                 = ArchiveToCopy.ArCustomPropertyList;
 	ArUseCustomPropertyList              = ArchiveToCopy.ArUseCustomPropertyList;
+	ArMergeOverrides					 = ArchiveToCopy.ArMergeOverrides;
 	ArShouldSkipUpdateCustomVersion		 = ArchiveToCopy.ArShouldSkipUpdateCustomVersion;
 	CookData							 = ArchiveToCopy.CookData;
 	SerializedProperty					 = ArchiveToCopy.SerializedProperty;
