- Straightforward IOS compile fixes

diff --git a/Engine/Source/Runtime/Core/Private/FramePro/FramePro.cpp b/Engine/Source/Runtime/Core/Private/FramePro/FramePro.cpp
--- a/Engine/Source/Runtime/Core/Private/FramePro/FramePro.cpp
+++ b/Engine/Source/Runtime/Core/Private/FramePro/FramePro.cpp
@@ -9582,13 +9582,13 @@ namespace FramePro
 			//------------------------------------------------------------------------
 			void VSPrintf(char* p_buffer, size_t const buffer_size, const char* p_format, va_list arg_list)
 			{
-				vsprintf(p_buffer, p_format, arg_list);
+				vsnprintf(p_buffer, buffer_size, p_format, arg_list);
 			}
 
 			//------------------------------------------------------------------------
-			void ToString(int value, char* p_dest, int)
+			void ToString(int value, char* p_dest, int dest_size)
 			{
-				sprintf(p_dest, "%d", value);
+				snprintf(p_dest, dest_size, "%d", value);
 			}
 
 			//------------------------------------------------------------------------
diff --git a/Engine/Source/Runtime/CrashReportCore/Private/IOS/IOSErrorReport.cpp b/Engine/Source/Runtime/CrashReportCore/Private/IOS/IOSErrorReport.cpp
--- a/Engine/Source/Runtime/CrashReportCore/Private/IOS/IOSErrorReport.cpp
+++ b/Engine/Source/Runtime/CrashReportCore/Private/IOS/IOSErrorReport.cpp
@@ -50,11 +50,11 @@ FString FIOSErrorReport::FindCrashedAppPath() const
 			int32 LineEnd = -1;
 			if(PathData.FindChar( TCHAR('\r'), LineEnd ))
 			{
-				PathData.LeftInline(LineEnd, false);
+				PathData.LeftInline(LineEnd, EAllowShrinking::No);
 			}
 			if(PathData.FindChar( TCHAR('\n'), LineEnd ))
 			{
-				PathData.LeftInline(LineEnd, false);
+				PathData.LeftInline(LineEnd, EAllowShrinking::No);
 			}
 			return PathData;
 		}
