Added FGenericPlatformMemory::PrettyMemory that displays memory amounts in a pretty way, similar to FGenericPlatformTime::PrettyTime.

diff --git a/Engine/Source/Runtime/Core/Private/GenericPlatform/GenericPlatformMemory.cpp b/Engine/Source/Runtime/Core/Private/GenericPlatform/GenericPlatformMemory.cpp
--- a/Engine/Source/Runtime/Core/Private/GenericPlatform/GenericPlatformMemory.cpp
+++ b/Engine/Source/Runtime/Core/Private/GenericPlatform/GenericPlatformMemory.cpp
@@ -573,3 +573,32 @@ uint64 FGenericPlatformMemoryStats::GetAvailablePhysical(bool bExcludeExtraDevMe
 
 	return BytesAvailable;
 }
+
+FString FGenericPlatformMemory::PrettyMemory(uint64 Memory)
+{
+	static const TCHAR* Units[] = { TEXT("B"), TEXT("KB"), TEXT("MB"), TEXT("GB"), TEXT("TB"), TEXT("PB"), TEXT("EB") };
+			
+	uint64 UnitIndex = 0;
+	uint64 Remainder = 0;
+
+	while ((Memory > 1024) && (UnitIndex++ < (UE_ARRAY_COUNT(Units) - 1)))
+	{
+		Remainder = Memory & 1023;
+		Memory >>= 10llu;
+	}
+
+	// Convert remainder to percentage
+	const uint64 RemainderPerc = (Remainder * 100) >> 10;
+
+	if (RemainderPerc)
+	{
+		if (RemainderPerc % 10)
+		{
+			return FString::Printf(TEXT("%llu.%02llu%s"), Memory, RemainderPerc, Units[UnitIndex]);
+		}
+
+		return FString::Printf(TEXT("%llu.%llu%s"), Memory, RemainderPerc / 10, Units[UnitIndex]);
+	}
+
+	return FString::Printf(TEXT("%llu%s"), Memory, Units[UnitIndex]);
+}
\ No newline at end of file
diff --git a/Engine/Source/Runtime/Core/Public/GenericPlatform/GenericPlatformMemory.h b/Engine/Source/Runtime/Core/Public/GenericPlatform/GenericPlatformMemory.h
--- a/Engine/Source/Runtime/Core/Public/GenericPlatform/GenericPlatformMemory.h
+++ b/Engine/Source/Runtime/Core/Public/GenericPlatform/GenericPlatformMemory.h
@@ -808,6 +808,13 @@ public:
 		return false; // Most platform do not implement this.
 	}
 
+	/**
+	 * Returns a pretty-string for an amount of memory given in bytes.
+	 *
+	 * @param Memory amount in bytes
+	 * @return Memory in a pretty formatted string
+	 */
+	static CORE_API FString PrettyMemory( uint64 Memory );
 
 protected:
 	friend struct FGenericStatsUpdater;
diff --git a/Engine/Source/Runtime/Engine/Private/Tests/WorldPartition/WorldPartitionStaticSpatialIndexTests.cpp b/Engine/Source/Runtime/Engine/Private/Tests/WorldPartition/WorldPartitionStaticSpatialIndexTests.cpp
--- a/Engine/Source/Runtime/Engine/Private/Tests/WorldPartition/WorldPartitionStaticSpatialIndexTests.cpp
+++ b/Engine/Source/Runtime/Engine/Private/Tests/WorldPartition/WorldPartitionStaticSpatialIndexTests.cpp
@@ -32,7 +32,7 @@ namespace WorldPartitionTests
 		}
 
 		const double RunTime = FPlatformTime::Seconds() - StartTime;
-		Test->AddInfo(FString::Printf(TEXT("%s: %d tests in %s (%.2f/s)"), Name, Tests.Num(), *FPlatformTime::PrettyTime(RunTime), Tests.Num() / RunTime), 0);
+		Test->AddInfo(FString::Printf(TEXT("%s: %d tests in %s (%.2f/s, %s)"), Name, Tests.Num(), *FPlatformTime::PrettyTime(RunTime), Tests.Num() / RunTime, *FGenericPlatformMemory::PrettyMemory(SpatialIndex.GetAllocatedSize())));
 	}
 #endif
 
