Added UE_DEPRECATED_HEADER which acts like UE_DEPRECATED and should be placed in headers that should no longer be included.

diff --git a/Engine/Plugins/AudioGameplayVolume/Source/AudioGameplayVolume/Public/AudioGameplayVolumeProxyMutator.h b/Engine/Plugins/AudioGameplayVolume/Source/AudioGameplayVolume/Public/AudioGameplayVolumeProxyMutator.h
--- a/Engine/Plugins/AudioGameplayVolume/Source/AudioGameplayVolume/Public/AudioGameplayVolumeProxyMutator.h
+++ b/Engine/Plugins/AudioGameplayVolume/Source/AudioGameplayVolume/Public/AudioGameplayVolumeProxyMutator.h
@@ -13,11 +13,7 @@
 
 // HEADER_UNIT_SKIP - Deprecated
 
-#ifdef _MSC_VER
-#pragma message(__FILE__"(15): warning: use AudioGameplayVolumeSubsystem.h and AudioGameplayVolumeMutator.h instead of AudioGameplayVolumeProxyMutator.h")
-#else
-#pragma message("#include AudioGameplayVolumeSubsystem.h and AudioGameplayVolumeMutator.h instead of AudioGameplayVolumeProxyMutator.h")
-#endif
+UE_DEPRECATED_HEADER(5.1, "Use AudioGameplayVolumeSubsystem.h and AudioGameplayVolumeMutator.h instead of AudioGameplayVolumeProxyMutator.h.")
 
  // Include the new file so that the project still compiles, but has warnings
 #include "AudioGameplayVolumeSubsystem.h"
diff --git a/Engine/Plugins/MovieScene/MovieRenderPipeline/Source/MovieRenderPipelineCore/Public/MoviePipelineMasterConfig.h b/Engine/Plugins/MovieScene/MovieRenderPipeline/Source/MovieRenderPipelineCore/Public/MoviePipelineMasterConfig.h
--- a/Engine/Plugins/MovieScene/MovieRenderPipeline/Source/MovieRenderPipelineCore/Public/MoviePipelineMasterConfig.h
+++ b/Engine/Plugins/MovieScene/MovieRenderPipeline/Source/MovieRenderPipelineCore/Public/MoviePipelineMasterConfig.h
@@ -7,11 +7,7 @@
 
 #pragma once
 
-#ifdef _MSC_VER
-	#pragma message(__FILE__"(9): warning: use MoviePipelinePrimaryConfig.h instead of MoviePipelineMasterConfig.h")
-#else
-	#pragma message("#include MoviePipelinePrimaryConfig.h instead of MoviePipelineMasterConfig")
-#endif
+UE_DEPRECATED_HEADER(5.2, "Use MoviePipelinePrimaryConfig.h instead of MoviePipelineMasterConfig.h.")
 
 #include "MoviePipelinePrimaryConfig.h"
 
diff --git a/Engine/Plugins/MovieScene/SequencerScripting/Source/SequencerScripting/Public/SequencerBindingProxy.h b/Engine/Plugins/MovieScene/SequencerScripting/Source/SequencerScripting/Public/SequencerBindingProxy.h
--- a/Engine/Plugins/MovieScene/SequencerScripting/Source/SequencerScripting/Public/SequencerBindingProxy.h
+++ b/Engine/Plugins/MovieScene/SequencerScripting/Source/SequencerScripting/Public/SequencerBindingProxy.h
@@ -7,10 +7,6 @@
  
 #pragma once
 
-#ifdef _MSC_VER
-#pragma message(__FILE__"(9): warning: use MovieSceneBindingProxy.h instead of SequencerBindingProxy.h")
-#else
-#pragma message("#include MovieSceneBindingProxy.h instead of SequencerBindingProxy.h")
-#endif
+UE_DEPRECATED_HEADER(5.1, "Use MovieSceneBindingProxy.h instead of SequencerBindingProxy.h.")
 
 #include "MovieSceneBindingProxy.h"
diff --git a/Engine/Plugins/Runtime/CommonUI/Source/CommonUI/Public/Input/CommonInputMode.h b/Engine/Plugins/Runtime/CommonUI/Source/CommonUI/Public/Input/CommonInputMode.h
--- a/Engine/Plugins/Runtime/CommonUI/Source/CommonUI/Public/Input/CommonInputMode.h
+++ b/Engine/Plugins/Runtime/CommonUI/Source/CommonUI/Public/Input/CommonInputMode.h
@@ -4,10 +4,6 @@
 
 // HEADER_UNIT_SKIP - Deprecated
 
-#ifdef _MSC_VER
-#pragma message(__FILE__"(8): warning: use CommonInputModeTypes.h from CommonInput module instead of Input/CommonInputMode.h from CommonUI module")
-#else
-#pragma message("#include CommonInputModeTypes.h from CommonInput module instead of Input/CommonInputMode.h from CommonUI module")
-#endif
+UE_DEPRECATED_HEADER(5.1, "Use CommonInputModeTypes.h from CommonInput module instead of Input/CommonInputMode.h from CommonUI module.")
 
 #include "CommonInputModeTypes.h"
\ No newline at end of file
diff --git a/Engine/Plugins/Runtime/Metasound/Source/MetasoundFrontend/Public/MetasoundFrontendArchetypeRegistry.h b/Engine/Plugins/Runtime/Metasound/Source/MetasoundFrontend/Public/MetasoundFrontendArchetypeRegistry.h
--- a/Engine/Plugins/Runtime/Metasound/Source/MetasoundFrontend/Public/MetasoundFrontendArchetypeRegistry.h
+++ b/Engine/Plugins/Runtime/Metasound/Source/MetasoundFrontend/Public/MetasoundFrontendArchetypeRegistry.h
@@ -4,10 +4,6 @@
 
 // HEADER_UNIT_SKIP - Deprecated
 
-#ifdef _MSC_VER
-#pragma message(__FILE__"(9): warning: use Interfaces/MetasoundFrontendInterfaceRegistry.h instead instead of MetasoundFrontendArchetypeRegistry.h")
-#else
-#pragma message("#include Interfaces/MetasoundFrontendInterfaceRegistry.h instead of MetasoundFrontendArchetypeRegistry.h")
-#endif
+UE_DEPRECATED_HEADER(5.3, "Use Interfaces/MetasoundFrontendInterfaceRegistry.h instead instead of MetasoundFrontendArchetypeRegistry.h.")
 
 #include "Interfaces/MetasoundFrontendInterfaceRegistry.h"
diff --git a/Engine/Source/Editor/MovieSceneTools/Public/CommonMovieSceneTools.h b/Engine/Source/Editor/MovieSceneTools/Public/CommonMovieSceneTools.h
--- a/Engine/Source/Editor/MovieSceneTools/Public/CommonMovieSceneTools.h
+++ b/Engine/Source/Editor/MovieSceneTools/Public/CommonMovieSceneTools.h
@@ -4,10 +4,6 @@
  * DEPRECATED CommonMovieSceneTools.h. This is no longer used.
  */
  
- #pragma once
-  
-#ifdef _MSC_VER
-    #pragma message(__FILE__"(9): warning: CommonMovieSceneTools.h is no longer used")
-#else
-    #pragma message("CommonMovieSceneTools.h is no longer used")
-#endif
\ No newline at end of file
+#pragma once
+
+UE_DEPRECATED_HEADER(5.4, "CommonMovieSceneTools.h is no longer used.")
diff --git a/Engine/Source/Editor/PropertyEditor/Public/PropertyEditorPermissionList.h b/Engine/Source/Editor/PropertyEditor/Public/PropertyEditorPermissionList.h
--- a/Engine/Source/Editor/PropertyEditor/Public/PropertyEditorPermissionList.h
+++ b/Engine/Source/Editor/PropertyEditor/Public/PropertyEditorPermissionList.h
@@ -3,6 +3,6 @@
 #pragma once
 
 // Deprecated in UE 5.3
-#pragma message(__FILE__"(6): warning: Header has been moved. #include PropertyPermissionList.h instead of PropertyEditorPermissionList.h")
+UE_DEPRECATED_HEADER(5.3, "Header has been moved. #include PropertyPermissionList.h instead of PropertyEditorPermissionList.h.")
 
 #include "PropertyPermissionList.h"
\ No newline at end of file
diff --git a/Engine/Source/Runtime/AnimGraphRuntime/Public/AnimNodes/AnimNode_StateResult.h b/Engine/Source/Runtime/AnimGraphRuntime/Public/AnimNodes/AnimNode_StateResult.h
--- a/Engine/Source/Runtime/AnimGraphRuntime/Public/AnimNodes/AnimNode_StateResult.h
+++ b/Engine/Source/Runtime/AnimGraphRuntime/Public/AnimNodes/AnimNode_StateResult.h
@@ -1,13 +1,9 @@
-﻿// Copyright Epic Games, Inc. All Rights Reserved.
+// Copyright Epic Games, Inc. All Rights Reserved.
 
 #pragma once
 
 // HEADER_UNIT_SKIP - Deprecated
 
-#ifdef _MSC_VER
-#pragma message(__FILE__"(8): warning: Including AnimGraphRuntime/AnimNode_StateResult.h is deprecated. The file location changed. Use AnimNode_StateResult.h from the Engine module (Engine/Animation/AnimNode_StateResult.h) instead of the AnimGraphRuntime module.")
-#else
-#pragma message("Including AnimGraphRuntime/AnimNode_StateResult.h is deprecated. The file location changed. Use AnimNode_StateResult.h from the Engine module (Engine/Animation/AnimNode_StateResult.h) instead of the AnimGraphRuntime module.")
-#endif
+UE_DEPRECATED_HEADER(5.3, "Including AnimGraphRuntime/AnimNode_StateResult.h is deprecated. The file location changed. Use AnimNode_StateResult.h from the Engine module (Engine/Animation/AnimNode_StateResult.h) instead of the AnimGraphRuntime module.")
 
 #include "Animation/AnimNode_StateResult.h"
\ No newline at end of file
diff --git a/Engine/Source/Runtime/ApplicationCore/Public/GenericPlatform/IForceFeedbackSystem.h b/Engine/Source/Runtime/ApplicationCore/Public/GenericPlatform/IForceFeedbackSystem.h
--- a/Engine/Source/Runtime/ApplicationCore/Public/GenericPlatform/IForceFeedbackSystem.h
+++ b/Engine/Source/Runtime/ApplicationCore/Public/GenericPlatform/IForceFeedbackSystem.h
@@ -9,11 +9,7 @@
 
 // HEADER_UNIT_SKIP - Deprecated
 
-#ifdef _MSC_VER
-	#pragma message(__FILE__"(9): warning: use GenericPlatform/IInputInterface.h instead of GenericPlatform/IForceFeedbackSystem.h")
-#else
-	#pragma message("#include GenericPlatform/IInputInterface.h instead of GenericPlatform/IForceFeedbackSystem.h")
-#endif
+UE_DEPRECATED_HEADER(5.1, "Use GenericPlatform / IInputInterface.h instead of GenericPlatform / IForceFeedbackSystem.h.")
 
 #include "GenericPlatform/IInputInterface.h"
 
diff --git a/Engine/Source/Runtime/ClothingSystemRuntimeNv/Public/NvClothIncludes.h b/Engine/Source/Runtime/ClothingSystemRuntimeNv/Public/NvClothIncludes.h
--- a/Engine/Source/Runtime/ClothingSystemRuntimeNv/Public/NvClothIncludes.h
+++ b/Engine/Source/Runtime/ClothingSystemRuntimeNv/Public/NvClothIncludes.h
@@ -4,7 +4,7 @@
 
 // HEADER_UNIT_SKIP - Deprecated
 
-#pragma message (__FILE__"(5): warning: NvClothIncludes.h is no longer useful")
+UE_DEPRECATED_HEADER(5.4, "NvClothIncludes.h is no longer useful.")
 
 #include "CoreMinimal.h"
 #include "EngineDefines.h"
diff --git a/Engine/Source/Runtime/Core/Public/Containers/DepletableMpscQueue.h b/Engine/Source/Runtime/Core/Public/Containers/DepletableMpscQueue.h
--- a/Engine/Source/Runtime/Core/Public/Containers/DepletableMpscQueue.h
+++ b/Engine/Source/Runtime/Core/Public/Containers/DepletableMpscQueue.h
@@ -3,10 +3,6 @@
 #pragma once
 
 // HEADER_UNIT_SKIP - Deprecated
-#ifdef _MSC_VER
-#pragma message(__FILE__"(6): warning : #include Containers/DepletableMpmcQueue.h instead of Containers/DepletableMpscQueue.h (/s/Mpsc/Mpmc)")
-#else
-#pragma message("#include Containers/DepletableMpmcQueue.h instead of Containers/DepletableMpscQueue.h (/s/Mpsc/Mpmc)")
-#endif
+UE_DEPRECATED_HEADER(5.2, "#include Containers/DepletableMpmcQueue.h instead of Containers/DepletableMpscQueue.h (/s/Mpsc/Mpmc).")
 
 #include "Containers/DepletableMpmcQueue.h"
diff --git a/Engine/Source/Runtime/Core/Public/Misc/CoreMiscDefines.h b/Engine/Source/Runtime/Core/Public/Misc/CoreMiscDefines.h
--- a/Engine/Source/Runtime/Core/Public/Misc/CoreMiscDefines.h
+++ b/Engine/Source/Runtime/Core/Public/Misc/CoreMiscDefines.h
@@ -152,8 +152,6 @@ enum EPerElement {PerElement};
 	#define UE_PUSH_MACRO(name) __pragma(push_macro(name))
 	#define UE_POP_MACRO(name) __pragma(pop_macro(name))
 #endif
-#define PUSH_MACRO(name) UE_DEPRECATED_MACRO(5.0, "PUSH_MACRO is deprecated. Use UE_PUSH_MACRO and pass the macro name as a string.") UE_PUSH_MACRO(PREPROCESSOR_TO_STRING(name))
-#define POP_MACRO(name) UE_DEPRECATED_MACRO(5.0, "POP_MACRO is deprecated. Use UE_POP_MACRO and pass the macro name as a string.") UE_POP_MACRO(PREPROCESSOR_TO_STRING(name))
 
 #ifdef __COUNTER__
 	// Created a variable with a unique name
@@ -289,6 +287,21 @@ enum EPerElement {PerElement};
 	#define UE_INTERNAL
 #endif
 
+
+/**
+ * Macro which can be placed in a header to throw a deprecation warning when it is included.
+ */
+#ifdef _MSC_VER
+    #if UE_WARNINGS_AS_ERRORS
+    	#define UE_DEPRECATED_HEADER(Version, Message) __pragma(message(__FILE__ "(" PREPROCESSOR_TO_STRING(__LINE__) "): error C4996: " Message " Please update your code to the new API before upgrading to the next release, otherwise your project will no longer compile."))
+    #else
+    	#define UE_DEPRECATED_HEADER(Version, Message) __pragma(message(__FILE__ "(" PREPROCESSOR_TO_STRING(__LINE__) "): warning C4996: " Message " Please update your code to the new API before upgrading to the next release, otherwise your project will no longer compile."))
+    #endif
+#else
+	#define UE_DEPRECATED_HEADER(Version, Message) _Pragma(PREPROCESSOR_TO_STRING(message(Message " Please update your code to the new API before upgrading to the next release, otherwise your project will no longer compile.")))
+#endif
+
+
 /*
  * Macro that can be defined in the target file to strip deprecated properties in objects across the engine that check against this define.
  * Can be used by project that have migrated away from using deprecated functions and data members to potentially gain back some memory and perf.
diff --git a/Engine/Source/Runtime/Core/Public/Windows/COMPointer.h b/Engine/Source/Runtime/Core/Public/Windows/COMPointer.h
--- a/Engine/Source/Runtime/Core/Public/Windows/COMPointer.h
+++ b/Engine/Source/Runtime/Core/Public/Windows/COMPointer.h
@@ -4,7 +4,6 @@
 
 // HEADER_UNIT_SKIP - Should not be included
 
-// Producing a warning that is not treated as an error here calls non-deterministic incremental build warnings
-//#pragma message(__FILE__"(5): warning : #include Microsoft/COMPointer.h instead of Windows/COMPointer.h")
+UE_DEPRECATED_HEADER(4.27, "#include Microsoft/COMPointer.h instead of Windows/COMPointer.h.")
 
 #include "Microsoft/COMPointer.h"
diff --git a/Engine/Source/Runtime/CoreUObject/Public/UObject/ErrorException.h b/Engine/Source/Runtime/CoreUObject/Public/UObject/ErrorException.h
--- a/Engine/Source/Runtime/CoreUObject/Public/UObject/ErrorException.h
+++ b/Engine/Source/Runtime/CoreUObject/Public/UObject/ErrorException.h
@@ -4,4 +4,4 @@
 
 // HEADER_UNIT_SKIP - Deprecated
 
-#pragma message (__FILE__"(5): warning: ErrorException.h has been deprecated")
+UE_DEPRECATED_HEADER(5.2, "ErrorException.h has been deprecated.")
diff --git a/Engine/Source/Runtime/Engine/Public/ADPCMAudioInfo.h b/Engine/Source/Runtime/Engine/Public/ADPCMAudioInfo.h
--- a/Engine/Source/Runtime/Engine/Public/ADPCMAudioInfo.h
+++ b/Engine/Source/Runtime/Engine/Public/ADPCMAudioInfo.h
@@ -8,11 +8,7 @@
 
 #pragma once
 
-#ifdef _MSC_VER
-	#pragma message(__FILE__"(9): has moved to it's own module ADPCMAudioDecoder, please update your code.")
-#else
-	#pragma message("ADPCMAudioInfo.h has moved to it's own module ADPCMAudioDecoder, please update your code.")
-#endif
+UE_DEPRECATED_HEADER(5.4, "ADPCMAudioInfo.h has moved to its own module ADPCMAudioDecoder.")
 
  // Include the new file so that the project still compiles, but has warnings
 #include "Decoders/ADPCMAudioInfo.h"
diff --git a/Engine/Source/Runtime/Engine/Public/Engine/ICookInfo.h b/Engine/Source/Runtime/Engine/Public/Engine/ICookInfo.h
--- a/Engine/Source/Runtime/Engine/Public/Engine/ICookInfo.h
+++ b/Engine/Source/Runtime/Engine/Public/Engine/ICookInfo.h
@@ -4,5 +4,4 @@
 
 // HEADER_UNIT_SKIP - Deprecated
 
-// Deprecated in UE 5.3.
-#pragma message (__FILE__"(8): warning: #include UObject/ICookInfo.h instead of Engine/ICookInfo.h")
+UE_DEPRECATED_HEADER(5.3, "#include UObject/ICookInfo.h instead of Engine/ICookInfo.h.")
diff --git a/Engine/Source/Runtime/Engine/Public/OpusAudioInfo.h b/Engine/Source/Runtime/Engine/Public/OpusAudioInfo.h
--- a/Engine/Source/Runtime/Engine/Public/OpusAudioInfo.h
+++ b/Engine/Source/Runtime/Engine/Public/OpusAudioInfo.h
@@ -8,11 +8,7 @@
 
 #pragma once
 
-#ifdef _MSC_VER
-	#pragma message(__FILE__"(9): has moved to it's own module OpusAudioDecoder, please update your code.")
-#else
-	#pragma message("OpusAudioInfo.h has moved to it's own module OpusAudioDecoder, please update your code.")
-#endif
+UE_DEPRECATED_HEADER(5.4, "OpusAudioInfo.h has moved to its own module OpusAudioDecoder.")
 
  // Include the new file so that the project still compiles, but has warnings
 #include "Decoders/OpusAudioInfo.h"
diff --git a/Engine/Source/Runtime/Engine/Public/VorbisAudioInfo.h b/Engine/Source/Runtime/Engine/Public/VorbisAudioInfo.h
--- a/Engine/Source/Runtime/Engine/Public/VorbisAudioInfo.h
+++ b/Engine/Source/Runtime/Engine/Public/VorbisAudioInfo.h
@@ -8,11 +8,7 @@
 
 #pragma once
 
-#ifdef _MSC_VER
-	#pragma message(__FILE__"(9): has moved to it's own module VorbisAudioDecoder, please update your code.")
-#else
-	#pragma message("VorbisAudioInfo.h has moved to it's own module VorbisAudioDecoder, please update your code.")
-#endif
+UE_DEPRECATED_HEADER(5.4, "VorbisAudioInfo.h has moved to its own module VorbisAudioDecoder.")
 
  // Include the new file so that the project still compiles, but has warnings
 #include "Decoders/VorbisAudioInfo.h"
diff --git a/Engine/Source/Runtime/PhysicsCore/Public/PhysXIncludes.h b/Engine/Source/Runtime/PhysicsCore/Public/PhysXIncludes.h
--- a/Engine/Source/Runtime/PhysicsCore/Public/PhysXIncludes.h
+++ b/Engine/Source/Runtime/PhysicsCore/Public/PhysXIncludes.h
@@ -4,6 +4,6 @@
 
 // HEADER_UNIT_SKIP - Deprecated
 
-#pragma message (__FILE__"(5): warning: PhysXIncludes.h is no longer useful")
+UE_DEPRECATED_HEADER(5.4, "PhysXIncludes.h is no longer useful.")
 
 #include "CoreMinimal.h"
\ No newline at end of file
diff --git a/Engine/Source/Runtime/Windows/AudioMixerXAudio2/Private/AudioMixerPlatformWindows.cpp b/Engine/Source/Runtime/Windows/AudioMixerXAudio2/Private/AudioMixerPlatformWindows.cpp
--- a/Engine/Source/Runtime/Windows/AudioMixerXAudio2/Private/AudioMixerPlatformWindows.cpp
+++ b/Engine/Source/Runtime/Windows/AudioMixerXAudio2/Private/AudioMixerPlatformWindows.cpp
@@ -15,7 +15,7 @@
 
 #if PLATFORM_WINDOWS
 
-#include "Windows/COMPointer.h"
+#include "Microsoft/COMPointer.h"
 #include "ScopedCom.h"					// FScopedComString
 
 #include "Windows/AllowWindowsPlatformTypes.h"
diff --git a/Engine/Source/Runtime/Windows/AudioMixerXAudio2/Private/WindowsMMDeviceInfoCache.h b/Engine/Source/Runtime/Windows/AudioMixerXAudio2/Private/WindowsMMDeviceInfoCache.h
--- a/Engine/Source/Runtime/Windows/AudioMixerXAudio2/Private/WindowsMMDeviceInfoCache.h
+++ b/Engine/Source/Runtime/Windows/AudioMixerXAudio2/Private/WindowsMMDeviceInfoCache.h
@@ -14,7 +14,7 @@
 #include "Windows/AllowWindowsPlatformTypes.h"
 #include "Windows/AllowWindowsPlatformAtomics.h"
 
-#include "Windows/COMPointer.h"
+#include "Microsoft/COMPointer.h"
 
 THIRD_PARTY_INCLUDES_START
 #include <mmreg.h>				// WAVEFORMATEX
diff --git a/Engine/Source/Runtime/Windows/AudioMixerXAudio2/Private/WindowsMMNotificationClient.h b/Engine/Source/Runtime/Windows/AudioMixerXAudio2/Private/WindowsMMNotificationClient.h
--- a/Engine/Source/Runtime/Windows/AudioMixerXAudio2/Private/WindowsMMNotificationClient.h
+++ b/Engine/Source/Runtime/Windows/AudioMixerXAudio2/Private/WindowsMMNotificationClient.h
@@ -13,7 +13,7 @@
 #include "Windows/AllowWindowsPlatformTypes.h"
 #include "Windows/AllowWindowsPlatformAtomics.h"
 
-#include "Windows/COMPointer.h"
+#include "Microsoft/COMPointer.h"
 
 THIRD_PARTY_INCLUDES_START
 #include <mmdeviceapi.h>			// IMMNotificationClient
