Changed the custom dissuasion code to use UE_DEPRECATED_HEADER instead.

diff --git a/Engine/Source/Runtime/Core/Public/UObject/FortniteNCBranchObjectVersion.h b/Engine/Source/Runtime/Core/Public/UObject/FortniteNCBranchObjectVersion.h
--- a/Engine/Source/Runtime/Core/Public/UObject/FortniteNCBranchObjectVersion.h
+++ b/Engine/Source/Runtime/Core/Public/UObject/FortniteNCBranchObjectVersion.h
@@ -1,11 +1,7 @@
 // Copyright Epic Games, Inc. All Rights Reserved.
 #pragma once
 
-#ifdef _MSC_VER
-#pragma message(__FILE__"(9): warning: use FortniteSeasonBranchObjectVersion.h instead of FortniteNCBranchObjectVersion.h")
-#else
-#pragma message("#include FortniteSeasonBranchObjectVersion.h instead of FortniteNCBranchObjectVersion.h")
-#endif
+UE_DEPRECATED_HEADER(5.3, "Use FortniteSeasonBranchObjectVersion.h instead of FortniteNCBranchObjectVersion.h.")
 
 #include "UObject/FortniteSeasonBranchObjectVersion.h"
 
