Fix const range-based for loops over TBitArray<>

diff --git a/Engine/Source/Runtime/Core/Public/Containers/BitArray.h b/Engine/Source/Runtime/Core/Public/Containers/BitArray.h
--- a/Engine/Source/Runtime/Core/Public/Containers/BitArray.h
+++ b/Engine/Source/Runtime/Core/Public/Containers/BitArray.h
@@ -1506,7 +1506,7 @@ public:
 			return *this;
 		}
 
-		FORCEINLINE FBitReference operator*() const
+		FORCEINLINE FConstBitReference operator*() const
 		{
 			return GetValue();
 		}
