[TaskGraph]
 - Fix race between TryUnlock and TryRetract for inline tasks and task events leading to use-after-free when retraction wins the race to execution

diff --git a/Engine/Source/Runtime/Core/Private/Tasks/TaskPrivate.cpp b/Engine/Source/Runtime/Core/Private/Tasks/TaskPrivate.cpp
--- a/Engine/Source/Runtime/Core/Private/Tasks/TaskPrivate.cpp
+++ b/Engine/Source/Runtime/Core/Private/Tasks/TaskPrivate.cpp
@@ -4,7 +4,6 @@
 #include "Tasks/Pipe.h"
 
 #include "Async/TaskGraphInterfaces.h"
-#include "HAL/IConsoleManager.h"
 
 namespace UE::Tasks
 {
@@ -75,7 +74,7 @@ namespace UE::Tasks
 
 			if (!IsAwaitable())
 			{
-				UE_LOG(LogTemp, Fatal, TEXT("Deadlock detected! A task can't be waited here, e.g. because it's being executed by the currect thread"));
+				UE_LOG(LogTemp, Fatal, TEXT("Deadlock detected! A task can't be waited here, e.g. because it's being executed by the current thread"));
 				return false;
 			}
 
@@ -120,28 +119,26 @@ namespace UE::Tasks
 				}
 			}
 
+			// If we don't have any more prerequisites, let TryUnlock
+			// execute these to avoid any race condition where we could clear
+			// the last reference before TryUnlock finishes and cause a use-after-free.
+			// These are super fast to process anyway so we can just consider them done
+			// for retraction purpose.
+			if (ExtendedPriority == EExtendedTaskPriority::TaskEvent ||
+				ExtendedPriority == EExtendedTaskPriority::Inline)
 			{
-				FThreadLocalRetractionScope ThreadLocalRetractionScope;
-
-				// next we try to execute the task, despite we haven't verified that the task is unlocked. trying to obtain execution permission will fail in this case
+				return true;
+			}
 
-				if (ExtendedPriority == EExtendedTaskPriority::TaskEvent)
-				{
-					if (!TrySetExecutionFlag())
-					{
-						return false;
-					}
+			if (Timeout)
+			{
+				return IsCompleted();
+			}
 
-					// task events have nothing to execute, and so can't have nested task, just close it
-					Close();
-					ReleaseInternalReference();
-					return true;
-				}
+			{
+				FThreadLocalRetractionScope ThreadLocalRetractionScope;
 
-				if (Timeout)
-				{
-					return IsCompleted();
-				}
+				// next we try to execute the task, despite we haven't verified that the task is unlocked. trying to obtain execution permission will fail in this case
 
 				if (!TryExecuteTask())
 				{
@@ -536,12 +533,5 @@ namespace UE::Tasks
 		}
 
 #endif // !TASKGRAPH_NEW_FRONTEND
-
-		bool GAddReferenceInTryUnlock = true;
-		static FAutoConsoleVariableRef AddReferenceInTryUnlockCVar(
-			TEXT("task.AddReferenceInTryUnlock"),
-			GAddReferenceInTryUnlock,
-			TEXT("Whether to enable the workaround where we add a reference in TryUnlock before decrementing lock count"),
-			ECVF_Default);
 	}
 }
diff --git a/Engine/Source/Runtime/Core/Public/Tasks/TaskPrivate.h b/Engine/Source/Runtime/Core/Public/Tasks/TaskPrivate.h
--- a/Engine/Source/Runtime/Core/Public/Tasks/TaskPrivate.h
+++ b/Engine/Source/Runtime/Core/Public/Tasks/TaskPrivate.h
@@ -112,9 +112,6 @@ namespace UE::Tasks
 		UE_DEPRECATED(5.1, "You should not use this function as it exists only to patch another system and can be removed any time.")
 		CORE_API bool IsThreadRetractingTask();
 
-		// For workaround: see Private::FTaskBase::TryUnlock
-		extern CORE_API bool GAddReferenceInTryUnlock;
-
 		// An abstract base class for task implementation. 
 		// Implements internal logic of task prerequisites, nested tasks and deep task retraction.
 		// Implements intrusive ref-counting and so can be used with TRefCountPtr.
@@ -567,21 +564,6 @@ public:
 
 				FPipe* LocalPipe = GetPipe(); // cache data locally so we won't need to touch the member (read below)
 
-				// TEMP: Patch to prevent the task getting destroyed. If NumLocks reaches 0 and retraction wins the execution, it could end
-				// up releasing the last reference before we continue and exit this function.
-				const bool bHasAddedReference = GAddReferenceInTryUnlock;
-				if (bHasAddedReference)
-				{
-					AddRef();
-				}
-				ON_SCOPE_EXIT
-				{
-					if (bHasAddedReference)
-					{
-						Release();
-					}
-				};
-
 				uint32 PrevNumLocks = NumLocks.fetch_sub(1, std::memory_order_acq_rel); // `acq_rel` to make it happen after task 
 				// preparation and before launching it
 				// the task can be dead already as the prev line can remove the lock hold for this execution path, another thread(s) can unlock
