[Scheduler/Taskgraph]
  - Replace busy wait logic by oversubscription using a pool of standby thread.
  - This fixes a whole class of deadlocks because of dependency inversion between tasks on the same callstack.
  - Using a simple oversubscription scope, we can now improve concurrency on any long IO or wait operations.
  - Oversubscription has already been applied to many wait API so it just works without having developers know it's there.
  - Standby threads go back to sleep as soon as the oversubscription period finishes and they finish their current task.
  - Standby threads never busy yield, they go to sleep when no more work is available, giving time slice back for real work.
  - Standby threads are only active when all normal threads are busy and some oversubscription scopes are active.
  - Add dynamic thread creation support so we only create what's needed (especially good on high core count machines).
  - For platforms where static thread creation is preferred. TaskGraph.UseDynamicThreadCreation = 0 can be used.
  - The maximum number of standby threads is controlled by TaskGraph.OversubscriptionRatio. (Current default 2x).
  - Deprecate reserve workers as they are now superseded by this feature.
  - Busywait API has been reimplemented using oversubscription and will be deprecated in another CL to keep this one focused.

diff --git a/Engine/Source/Runtime/Core/Private/Async/Fundamental/ReserveScheduler.cpp b/Engine/Source/Runtime/Core/Private/Async/Fundamental/ReserveScheduler.cpp
--- a/Engine/Source/Runtime/Core/Private/Async/Fundamental/ReserveScheduler.cpp
+++ b/Engine/Source/Runtime/Core/Private/Async/Fundamental/ReserveScheduler.cpp
@@ -13,6 +13,8 @@
 
 extern CORE_API bool GTaskGraphUseDynamicPrioritization;
 
+PRAGMA_DISABLE_DEPRECATION_WARNINGS
+
 namespace LowLevelTasks
 {
 
@@ -121,4 +123,6 @@ void FReserveScheduler::StopWorkers()
 	}
 }
 
-}
\ No newline at end of file
+}
+
+PRAGMA_ENABLE_DEPRECATION_WARNINGS
\ No newline at end of file
diff --git a/Engine/Source/Runtime/Core/Private/Async/Fundamental/Scheduler.cpp b/Engine/Source/Runtime/Core/Private/Async/Fundamental/Scheduler.cpp
--- a/Engine/Source/Runtime/Core/Private/Async/Fundamental/Scheduler.cpp
+++ b/Engine/Source/Runtime/Core/Private/Async/Fundamental/Scheduler.cpp
@@ -4,6 +4,7 @@
 #include "Async/Fundamental/Task.h"
 #include "Async/TaskTrace.h"
 #include "Logging/LogMacros.h"
+#include "ProfilingDebugging/CsvProfiler.h"
 #include "HAL/RunnableThread.h"
 #include "Misc/ScopeLock.h"
 #include "Misc/CommandLine.h"
@@ -11,6 +12,10 @@
 #include "CoreGlobals.h"
 
 extern CORE_API bool GTaskGraphUseDynamicPrioritization;
+extern CORE_API float GTaskGraphOversubscriptionRatio;
+extern CORE_API bool GTaskGraphUseDynamicThreadCreation;
+
+CSV_DEFINE_CATEGORY(Scheduler, true);
 
 namespace LowLevelTasks
 {
@@ -20,13 +25,12 @@ namespace LowLevelTasks
 	thread_local FTask* FTask::ActiveTask = nullptr;
 	thread_local FSchedulerTls* FSchedulerTls::ActiveScheduler = nullptr;
 	thread_local FSchedulerTls::EWorkerType FSchedulerTls::WorkerType = FSchedulerTls::EWorkerType::None;
-	thread_local uint32 FSchedulerTls::BusyWaitingDepth = 0;
+	thread_local bool Private::FOversubscriptionTls::bIsOversubscriptionAllowed = false;
 
 	FScheduler FScheduler::Singleton;
 
-	TUniquePtr<FThread> FScheduler::CreateWorker(bool bPermitBackgroundWork, FThread::EForkable IsForkable, Private::FWaitEvent* ExternalWorkerEvent, FSchedulerTls::FLocalQueueType* ExternalWorkerLocalQueue, EThreadPriority Priority, uint64 InAffinity)
+	TUniquePtr<FThread> FScheduler::CreateWorker(uint32 WorkerId, const TCHAR* Name, bool bPermitBackgroundWork, FThread::EForkable IsForkable, Private::FWaitEvent* ExternalWorkerEvent, FSchedulerTls::FLocalQueueType* ExternalWorkerLocalQueue, EThreadPriority Priority, uint64 InAffinity)
 	{
-		uint32 WorkerId = NextWorkerId++;
 		const uint32 WaitTimes[8] = { 719, 991, 1361, 1237, 1597, 953, 587, 1439 };
 		uint32 WaitTime = WaitTimes[WorkerId % 8];
 		uint64 ThreadAffinityMask = FPlatformAffinity::GetTaskGraphThreadMask();
@@ -64,7 +68,7 @@ namespace LowLevelTasks
 		
 		return MakeUnique<FThread>
 		(
-			bPermitBackgroundWork ? *FString::Printf(TEXT("Background Worker #%d"), WorkerId) : *FString::Printf(TEXT("Foreground Worker #%d"), WorkerId),
+			Name,
 			[this, ExternalWorkerEvent, ExternalWorkerLocalQueue, WaitTime, bPermitBackgroundWork]
 			{ 
 				WorkerMain(ExternalWorkerEvent, ExternalWorkerLocalQueue, WaitTime, bPermitBackgroundWork);
@@ -80,6 +84,11 @@ namespace LowLevelTasks
 			GTaskGraphUseDynamicPrioritization = Value != 0;
 		}
 
+		if (FParse::Value(FCommandLine::Get(), TEXT("TaskGraphUseDynamicThreadCreation="), Value))
+		{
+			GTaskGraphUseDynamicThreadCreation = Value != 0;
+		}
+
 		if (NumForegroundWorkers == 0 && NumBackgroundWorkers == 0)
 		{
 			NumForegroundWorkers = FMath::Max<int32>(1, FMath::Min<int32>(2, FPlatformMisc::NumberOfWorkerThreadsToSpawn() - 1));
@@ -101,45 +110,95 @@ namespace LowLevelTasks
 		const bool bSupportsMultithreading = FPlatformProcess::SupportsMultithreading() || FForkProcessHelper::IsForkedMultithreadInstance();
 
 		uint32 OldActiveWorkers = ActiveWorkers.load(std::memory_order_relaxed);
-		if(OldActiveWorkers == 0 && bSupportsMultithreading && ActiveWorkers.compare_exchange_strong(OldActiveWorkers, NumForegroundWorkers + NumBackgroundWorkers, std::memory_order_relaxed))
+		if (OldActiveWorkers == 0 && bSupportsMultithreading && ActiveWorkers.compare_exchange_strong(OldActiveWorkers, NumForegroundWorkers + NumBackgroundWorkers, std::memory_order_relaxed))
 		{
 			FScopeLock Lock(&WorkerThreadsCS);
 			check(!WorkerThreads.Num());
 			check(!WorkerLocalQueues.Num());
-			check(!WorkerEvents.Num());		
+			check(!WorkerEvents.Num());
 			check(NextWorkerId == 0);
+			ForegroundCreationIndex = 0;
+			BackgroundCreationIndex = 0;
 
-			WorkerEvents.Reserve(NumForegroundWorkers + NumBackgroundWorkers);
-			WorkerLocalQueues.Reserve(NumForegroundWorkers + NumBackgroundWorkers);
-			for (uint32 WorkerId = 0; WorkerId < NumForegroundWorkers; ++WorkerId)
-			{
-				WorkerEvents.Emplace();
-				WorkerLocalQueues.Emplace(QueueRegistry, Private::ELocalQueueType::EForeground);
-			}
+			const float OversubscriptionRatio = FMath::Max(1.0f, GTaskGraphOversubscriptionRatio);
+			const int32 MaxForegroundWorkers = FMath::CeilToInt(float(NumForegroundWorkers) * OversubscriptionRatio);
+			const int32 MaxBackgroundWorkers = FMath::CeilToInt(float(NumBackgroundWorkers) * OversubscriptionRatio);
+			const int32 MaxWorkers = MaxForegroundWorkers + MaxBackgroundWorkers;
+			const EThreadPriority ActualBackgroundPriority = GTaskGraphUseDynamicPrioritization ? WorkerPriority : BackgroundPriority;
 
-			for (uint32 WorkerId = 0; WorkerId < NumBackgroundWorkers; ++WorkerId)
-			{
-				WorkerEvents.Emplace();
-				WorkerLocalQueues.Emplace(QueueRegistry, Private::ELocalQueueType::EBackground);
-			}
+			WorkerEvents.SetNum(MaxWorkers);
+			WorkerLocalQueues.Reserve(MaxWorkers);
+			WorkerThreads.Reserve(MaxWorkers);
 
-			// WorkerEvents are now ready, time to init.
-			WaitingQueue[0].Init();
-			WaitingQueue[1].Init();
+			auto CreateThread =
+				[this, IsForkable](Private::ELocalQueueType LocalQueueType, const TCHAR* ThreadGroup, const TCHAR* Prefix, std::atomic<int32>& CreationIndex, int32 NumWorkers, int32 NumMaxWorkers, EThreadPriority Priority, uint64 Affinity)
+				{
+					LLM_SCOPE_BYNAME(TEXT("EngineMisc/WorkerThreads"));
 
-			WorkerThreads.Reserve(NumForegroundWorkers + NumBackgroundWorkers);
-			UE::Trace::ThreadGroupBegin(TEXT("Foreground Workers"));
-			for (uint32 WorkerId = 0; WorkerId < NumForegroundWorkers; ++WorkerId)
-			{
-				WorkerThreads.Add(CreateWorker(false, IsForkable, &WorkerEvents[WorkerId], &WorkerLocalQueues[WorkerId], WorkerPriority, WorkerAffinity));
-			}
-			UE::Trace::ThreadGroupEnd();
-			UE::Trace::ThreadGroupBegin(TEXT("Background Workers"));
-			for (uint32 WorkerId = 0; WorkerId < NumBackgroundWorkers; ++WorkerId)
+					// Thread creation can end up waiting, we don't want to recursively oversubscribe if that happens.
+					Private::FOversubscriptionAllowedScope _(false); 
+
+					const int32 LocalCreationIndex = CreationIndex++;
+					check(LocalCreationIndex < NumMaxWorkers);
+					const bool bIsStandbyWorker = LocalCreationIndex >= NumWorkers;
+					FString WorkerName;
+					if (bIsStandbyWorker)
+					{
+						WorkerName = FString::Printf(TEXT("%s Worker (Standby #%d)"), Prefix, LocalCreationIndex - NumWorkers);
+					}
+					else
+					{
+						WorkerName = FString::Printf(TEXT("%s Worker #%d"), Prefix, LocalCreationIndex);
+					}
+
+					uint32 WorkerId = NextWorkerId++;
+					UE::Trace::ThreadGroupBegin(ThreadGroup);
+					WorkerLocalQueues.Emplace(QueueRegistry, LocalQueueType);
+					WorkerEvents[WorkerId].bIsStandby = bIsStandbyWorker;
+					WorkerThreads.Add(
+						CreateWorker(
+							WorkerId,
+							*WorkerName,
+							LocalQueueType == Private::ELocalQueueType::EBackground, /* bPermitBackgroundWork */
+							IsForkable,
+							&WorkerEvents[WorkerId],
+							&WorkerLocalQueues[WorkerId],
+							Priority,
+							Affinity)
+					);
+					UE::Trace::ThreadGroupEnd();
+				};
+
+			TFunction<void()> ForegroundCreateThread =
+				[this, CreateThread, NumWorkers = NumForegroundWorkers, MaxWorkers = MaxForegroundWorkers]()
+				{
+					FScopeLock Lock(&WorkerThreadsCS);
+					CreateThread(Private::ELocalQueueType::EForeground, TEXT("Foreground Workers"), TEXT("Foreground"), ForegroundCreationIndex, NumWorkers, MaxWorkers, WorkerPriority, WorkerAffinity);
+				};
+
+			TFunction<void()> BackgroundCreateThread =
+				[this, CreateThread, NumWorkers = NumBackgroundWorkers, MaxWorkers = MaxBackgroundWorkers, ActualBackgroundPriority]()
+				{
+					FScopeLock Lock(&WorkerThreadsCS);
+					CreateThread(Private::ELocalQueueType::EBackground, TEXT("Background Workers"), TEXT("Background"), BackgroundCreationIndex, NumWorkers, MaxWorkers, ActualBackgroundPriority, BackgroundAffinity);
+				};
+
+			WaitingQueue[0].Init(NumForegroundWorkers, MaxForegroundWorkers, ForegroundCreateThread, GTaskGraphUseDynamicThreadCreation ? 0 : MaxForegroundWorkers /* ActiveThreadCount */);
+			WaitingQueue[1].Init(NumBackgroundWorkers, MaxBackgroundWorkers, BackgroundCreateThread, GTaskGraphUseDynamicThreadCreation ? 0 : MaxBackgroundWorkers /* ActiveThreadCount */);
+
+			// Precreate all the threads if dynamic thread creation is not activated.
+			if (!GTaskGraphUseDynamicThreadCreation)
 			{
-				WorkerThreads.Add(CreateWorker(true, IsForkable, &WorkerEvents[NumForegroundWorkers + WorkerId], &WorkerLocalQueues[NumForegroundWorkers + WorkerId], GTaskGraphUseDynamicPrioritization ? WorkerPriority : BackgroundPriority, BackgroundAffinity));
+				for (int32 Index = 0; Index < MaxForegroundWorkers; Index++)
+				{
+					ForegroundCreateThread();
+				}
+
+				for (int32 Index = 0; Index < MaxBackgroundWorkers; Index++)
+				{
+					BackgroundCreateThread();
+				}
 			}
-			UE::Trace::ThreadGroupEnd();
 		}
 	}
 
@@ -191,20 +250,23 @@ namespace LowLevelTasks
 	void FScheduler::StopWorkers(bool DrainGlobalQueue)
 	{
 		uint32 OldActiveWorkers = ActiveWorkers.load(std::memory_order_relaxed);
-		if(OldActiveWorkers != 0 && ActiveWorkers.compare_exchange_strong(OldActiveWorkers, 0, std::memory_order_relaxed))
+		if (OldActiveWorkers != 0 && ActiveWorkers.compare_exchange_strong(OldActiveWorkers, 0, std::memory_order_relaxed))
 		{
 			FScopeLock Lock(&WorkerThreadsCS);
 
-			WaitingQueue[0].NotifyAll();
-			WaitingQueue[1].NotifyAll();
+			WaitingQueue[0].StartShutdown();
+			WaitingQueue[1].StartShutdown();
 
 			for (TUniquePtr<FThread>& Thread : WorkerThreads)
 			{
-				Thread->Join();
+				if (Thread.IsValid())
+				{
+					Thread->Join();
+				}
 			}
 
-			WaitingQueue[0].Shutdown();
-			WaitingQueue[1].Shutdown();
+			WaitingQueue[0].FinishShutdown();
+			WaitingQueue[1].FinishShutdown();
 
 			NextWorkerId = 0;
 			WorkerThreads.Reset();
@@ -224,6 +286,8 @@ namespace LowLevelTasks
 					}
 				}
 			}
+
+			QueueRegistry.Reset();
 		}
 	}
 
@@ -279,6 +343,32 @@ namespace LowLevelTasks
 		}
 	}
 
+	void FScheduler::IncrementOversubscription()
+	{
+		FSchedulerTls::EWorkerType LocalWorkerType = WorkerType;
+
+		if (LocalWorkerType != EWorkerType::None)
+		{
+			// The goal is to minimize the amount of wait in the worker tasks, this will help drive the
+			// total number of oversubscription down and show any regressions.
+			CSV_CUSTOM_STAT(Scheduler, Oversubscription, 1, ECsvCustomStatOp::Accumulate);
+
+			const bool bPermitBackgroundWork = LocalWorkerType == FSchedulerTls::EWorkerType::Background;
+			WaitingQueue[bPermitBackgroundWork].IncrementOversubscription();
+		}
+	}
+
+	void FScheduler::DecrementOversubscription()
+	{
+		FSchedulerTls::EWorkerType LocalWorkerType = WorkerType;
+
+		if (LocalWorkerType != EWorkerType::None)
+		{
+			const bool bPermitBackgroundWork = LocalWorkerType == FSchedulerTls::EWorkerType::Background;
+			WaitingQueue[bPermitBackgroundWork].DecrementOversubscription();
+		}
+	}
+
 #if PLATFORM_DESKTOP || !IS_MONOLITHIC
 	const FTask* FTask::GetActiveTask()
 	{
@@ -293,19 +383,10 @@ namespace LowLevelTasks
 
 	bool FSchedulerTls::IsBusyWaiting()
 	{
-		return BusyWaitingDepth != 0;
-	}
-
-	uint32 FSchedulerTls::GetAffinityIndex()
-	{
-		if (LocalQueue)
-		{
-			return LocalQueue->GetAffinityIndex();
-		}
-		return ~0;
+		return false;
 	}
 
-	template<typename QueueType, FTask* (QueueType::*DequeueFunction)(bool), bool bIsBusyWaiting>
+	template<typename QueueType, FTask* (QueueType::*DequeueFunction)(bool), bool bIsStandbyWorker>
 	bool FScheduler::TryExecuteTaskFrom(Private::FWaitEvent* WaitEvent, QueueType* Queue, Private::FOutOfWork& OutOfWork, bool bPermitBackgroundWork)
 	{
 		bool AnyExecuted = false;
@@ -313,45 +394,21 @@ namespace LowLevelTasks
 		FTask* Task = (Queue->*DequeueFunction)(bPermitBackgroundWork);
 		while (Task)
 		{
-			if constexpr (bIsBusyWaiting)
-			{
-				FTask::FInitData InitData = Task->GetInitData();
-				const bool bAllowBusyWaiting = EnumHasAnyFlags(InitData.Flags, ETaskFlags::AllowBusyWaiting) || Task->WasCanceledOrIsExpediting();
-
-				if (!bAllowBusyWaiting)
-				{
-					// The task is not allowed during busy waiting so requeue in a different
-					// queue to make sure we won't pick it up again.
-					constexpr static bool bAllowInsideBusyWaiting = false;
-					QueueRegistry.Enqueue<bAllowInsideBusyWaiting>(Task, uint32(InitData.Priority));
-
-					// Fetch another task we could potentially run.
-					Task = (Queue->*DequeueFunction)(bPermitBackgroundWork);
-
-					// Make sure we run the filtering logic again for the new task.
-					continue;
-				}
-			}
-			else
-			{
-				checkSlow(FTask::ActiveTask == nullptr);
-			}
+			checkSlow(FTask::ActiveTask == nullptr);
 
 			if (OutOfWork.Stop())
 			{
-				bool bWakeUpWanted = true;
-				if (WaitEvent)
+				// Standby workers don't need cancellation, this logic doesn't apply to them.
+				if constexpr (bIsStandbyWorker == false)
 				{
 					// CancelWait will tell us if we need to start a new worker to replace
 					// a potential wakeup we might have consumed during the cancellation.
-					bWakeUpWanted = WaitingQueue[bPermitBackgroundWork].CancelWait(WaitEvent);
-				}
-
-				if (bWakeUpWanted)
-				{
-					if (!WakeUpWorker(bPermitBackgroundWork) && !FSchedulerTls::IsBackgroundWorker())
+					if (WaitingQueue[bPermitBackgroundWork].CancelWait(WaitEvent))
 					{
-						WakeUpWorker(!bPermitBackgroundWork);
+						if (!WakeUpWorker(bPermitBackgroundWork) && !FSchedulerTls::IsBackgroundWorker())
+						{
+							WakeUpWorker(!bPermitBackgroundWork);
+						}
 					}
 				}
 			}
@@ -362,44 +419,67 @@ namespace LowLevelTasks
 			if ((Task = ExecuteTask(Task)) != nullptr)
 			{
 				verifySlow(Task->TryPrepareLaunch());
-
-				// When busy waiting, we exit every time a task is run
-				// so queue this new task for later and bail out.
-				if constexpr (bIsBusyWaiting)
-				{
-					QueueRegistry.Enqueue(Task, uint32(Task->GetPriority()));
-					return AnyExecuted;
-				}
 			}
 		}
 		return AnyExecuted;
 	}
 
-	void FScheduler::WorkerMain(Private::FWaitEvent* WorkerEvent, FSchedulerTls::FLocalQueueType* WorkerLocalQueue, uint32 WaitCycles, bool bPermitBackgroundWork)
+	void FScheduler::StandbyLoop(Private::FWaitEvent* WorkerEvent, FSchedulerTls::FLocalQueueType* WorkerLocalQueue, uint32 WaitCycles, bool bPermitBackgroundWork)
 	{
-		checkSlow(FSchedulerTls::LocalQueue == nullptr);
-		checkSlow(WorkerLocalQueue != nullptr);
-		checkSlow(WorkerEvent != nullptr);
+		bool bPreparingStandby = false;
+		Private::FOutOfWork OutOfWork;
+		while (true)
+		{
+			bool bExecutedSomething = false;
+			while (TryExecuteTaskFrom<FSchedulerTls::FLocalQueueType, &FSchedulerTls::FLocalQueueType::Dequeue, true>(WorkerEvent, WorkerLocalQueue, OutOfWork, bPermitBackgroundWork)
+				|| TryExecuteTaskFrom<FSchedulerTls::FLocalQueueType, &FSchedulerTls::FLocalQueueType::DequeueSteal, true>(WorkerEvent, WorkerLocalQueue, OutOfWork, bPermitBackgroundWork))
+			{
+				bPreparingStandby = false;
+				bExecutedSomething = true;
 
-		FTaskTagScope WorkerScope(ETaskTag::EWorkerThread);
-		FSchedulerTls::ActiveScheduler = this;
+				// If we're currently oversubscribed... we might be selected for standby even when there is work left.
+				WaitingQueue[bPermitBackgroundWork].ConditionalStandby(WorkerEvent);
+			}
 
-		FMemory::SetupTLSCachesOnCurrentThread();
-		FSchedulerTls::WorkerType = bPermitBackgroundWork ? FSchedulerTls::EWorkerType::Background : FSchedulerTls::EWorkerType::Foreground;
-		FSchedulerTls::LocalQueue = WorkerLocalQueue;
+			// Check if we're shutting down
+			if (ActiveWorkers.load(std::memory_order_relaxed) == 0)
+			{
+				OutOfWork.Stop();
+				break;
+			}
+
+			if (bExecutedSomething == false)
+			{
+				if (!bPreparingStandby)
+				{
+					OutOfWork.Start();
+					WaitingQueue[bPermitBackgroundWork].PrepareStandby(WorkerEvent);
+					bPreparingStandby = true;
+				}
+				else if (WaitingQueue[bPermitBackgroundWork].CommitStandby(WorkerEvent, OutOfWork))
+				{
+					// Only reset this when the commit succeeded, otherwise we're backing off the commit and looking at the queue again
+					bPreparingStandby = false;
+				}
+			}
+		}
+	}
 
+	void FScheduler::WorkerLoop(Private::FWaitEvent* WorkerEvent, FSchedulerTls::FLocalQueueType* WorkerLocalQueue, uint32 WaitCycles, bool bPermitBackgroundWork)
+	{
 		bool bPreparingWait = false;
 		Private::FOutOfWork OutOfWork;
 		while (true)
 		{
 			bool bExecutedSomething = false;
-			while(TryExecuteTaskFrom<FSchedulerTls::FLocalQueueType, &FSchedulerTls::FLocalQueueType::Dequeue, false>(WorkerEvent, WorkerLocalQueue, OutOfWork, bPermitBackgroundWork)
-			   || TryExecuteTaskFrom<FSchedulerTls::FLocalQueueType, &FSchedulerTls::FLocalQueueType::DequeueSteal, false>(WorkerEvent, WorkerLocalQueue, OutOfWork, bPermitBackgroundWork))
+			while (TryExecuteTaskFrom<FSchedulerTls::FLocalQueueType, &FSchedulerTls::FLocalQueueType::Dequeue, false>(WorkerEvent, WorkerLocalQueue, OutOfWork, bPermitBackgroundWork)
+				|| TryExecuteTaskFrom<FSchedulerTls::FLocalQueueType, &FSchedulerTls::FLocalQueueType::DequeueSteal, false>(WorkerEvent, WorkerLocalQueue, OutOfWork, bPermitBackgroundWork))
 			{
 				bPreparingWait = false;
 				bExecutedSomething = true;
 			}
 
+			// Check if we're shutting down
 			if (ActiveWorkers.load(std::memory_order_relaxed) == 0)
 			{
 				// Don't leave the waiting queue in a bad state
@@ -425,84 +505,49 @@ namespace LowLevelTasks
 				}
 			}
 		}
-
-		WaitingQueue[bPermitBackgroundWork].NotifyAll();
-
-		FSchedulerTls::LocalQueue = nullptr;
-		FSchedulerTls::ActiveScheduler = nullptr;
-		FSchedulerTls::WorkerType = FSchedulerTls::EWorkerType::None;
-		FMemory::ClearAndDisableTLSCachesOnCurrentThread();
 	}
 
-	void FScheduler::BusyWaitInternal(const FConditional& Conditional, bool ForceAllowBackgroundWork)
+	void FScheduler::WorkerMain(Private::FWaitEvent* WorkerEvent, FSchedulerTls::FLocalQueueType* WorkerLocalQueue, uint32 WaitCycles, bool bPermitBackgroundWork)
 	{
-		TRACE_CPUPROFILER_EVENT_SCOPE(FScheduler::BusyWaitInternal);
-		FTaskTagScope WorkerScope(ETaskTag::EWorkerThread);
+		checkSlow(FSchedulerTls::LocalQueue == nullptr);
+		checkSlow(WorkerLocalQueue != nullptr);
+		checkSlow(WorkerEvent != nullptr);
 
-		++FSchedulerTls::BusyWaitingDepth;
-		ON_SCOPE_EXIT{ --FSchedulerTls::BusyWaitingDepth; };
+		FTaskTagScope WorkerScope(ETaskTag::EWorkerThread);
+		FSchedulerTls::ActiveScheduler = this;
 
-		constexpr static bool bIsInsideBusyWait = true;
-		check(ActiveWorkers.load(std::memory_order_relaxed));
-		FSchedulerTls::FLocalQueueType* WorkerLocalQueue = FSchedulerTls::LocalQueue;
+		FMemory::SetupTLSCachesOnCurrentThread();
+		FSchedulerTls::WorkerType = bPermitBackgroundWork ? FSchedulerTls::EWorkerType::Background : FSchedulerTls::EWorkerType::Foreground;
+		FSchedulerTls::LocalQueue = WorkerLocalQueue;
 
-		uint32 WaitCount = 0;
-		bool HasWokenEmergencyWorker = false;
-		const bool bIsBackgroundWorker = FSchedulerTls::IsBackgroundWorker();
-		bool bPermitBackgroundWork = FTask::PermitBackgroundWork() || ForceAllowBackgroundWork;
-		Private::FOutOfWork OutOfWork;
-		while (true)
 		{
-			if (WorkerLocalQueue)
+			Private::FOversubscriptionAllowedScope _(true);
+
+			if (WorkerEvent->bIsStandby)
 			{
-				while(TryExecuteTaskFrom<FSchedulerTls::FLocalQueueType, &FSchedulerTls::FLocalQueueType::Dequeue<bIsInsideBusyWait>, bIsInsideBusyWait>(nullptr, WorkerLocalQueue, OutOfWork, bPermitBackgroundWork)
- 									|| TryExecuteTaskFrom<FSchedulerTls::FLocalQueueType, &FSchedulerTls::FLocalQueueType::DequeueSteal, bIsInsideBusyWait>(nullptr, WorkerLocalQueue, OutOfWork, bPermitBackgroundWork))
-				{
-					if (Conditional())
-					{
-						return;
-					}
-					WaitCount = 0;
-				}
+				StandbyLoop(WorkerEvent, WorkerLocalQueue, WaitCycles, bPermitBackgroundWork);
 			}
 			else
 			{
-				while(TryExecuteTaskFrom<FSchedulerTls::FQueueRegistry, &FSchedulerTls::FQueueRegistry::DequeueGlobal<bIsInsideBusyWait>, bIsInsideBusyWait>(nullptr, &QueueRegistry, OutOfWork, bPermitBackgroundWork)
-									|| TryExecuteTaskFrom<FSchedulerTls::FQueueRegistry, &FSchedulerTls::FQueueRegistry::DequeueSteal, bIsInsideBusyWait>(nullptr, &QueueRegistry, OutOfWork, bPermitBackgroundWork))
-				{
-					if (Conditional())
-					{
-						return;
-					}
-					WaitCount = 0;
-				}
+				WorkerLoop(WorkerEvent, WorkerLocalQueue, WaitCycles, bPermitBackgroundWork);
 			}
+		}
 
-			if (Conditional())
-			{
-				return;
-			}
+		FSchedulerTls::LocalQueue = nullptr;
+		FSchedulerTls::ActiveScheduler = nullptr;
+		FSchedulerTls::WorkerType = FSchedulerTls::EWorkerType::None;
+		FMemory::ClearAndDisableTLSCachesOnCurrentThread();
+	}
 
-			if (WaitCount < WorkerSpinCycles)
-			{
-				OutOfWork.Start();
-				FPlatformProcess::Yield();
-				FPlatformProcess::Yield();
-				WaitCount++;
-			}
-			else if (!bPermitBackgroundWork && bIsBackgroundWorker)
-			{
-				bPermitBackgroundWork = true;
-			}
-			else
-			{
-				if(!HasWokenEmergencyWorker)
-				{
-					WakeUpWorker(true);
-					HasWokenEmergencyWorker = true;
-				}
-				WaitCount = 0;
-			}
+	void FScheduler::BusyWaitInternal(const FConditional& Conditional, bool ForceAllowBackgroundWork)
+	{
+		TRACE_CPUPROFILER_EVENT_SCOPE(FScheduler::BusyWaitInternal);
+		CSV_SCOPED_TIMING_STAT(Scheduler, BusyWait);
+
+		FOversubscriptionScope _;
+		while (!Conditional())
+		{
+			FPlatformProcess::YieldThread();
 		}
 	}
 }
\ No newline at end of file
diff --git a/Engine/Source/Runtime/Core/Private/Async/Fundamental/WaitingQueue.cpp b/Engine/Source/Runtime/Core/Private/Async/Fundamental/WaitingQueue.cpp
--- a/Engine/Source/Runtime/Core/Private/Async/Fundamental/WaitingQueue.cpp
+++ b/Engine/Source/Runtime/Core/Private/Async/Fundamental/WaitingQueue.cpp
@@ -9,6 +9,7 @@
 // rewritten.
 
 #include "Async/Fundamental/WaitingQueue.h"
+#include "Async/Fundamental/Scheduler.h"
 #include "Async/TaskTrace.h"
 #include "Logging/LogMacros.h"
 #include "HAL/RunnableThread.h"
@@ -18,6 +19,7 @@
 #include "HAL/PlatformMisc.h"
 #include "HAL/PlatformProcess.h"
 #include "ProfilingDebugging/CpuProfilerTrace.h"
+#include "ProfilingDebugging/CsvProfiler.h"
 #include "Trace/Trace.h"
 #include "Trace/Trace.inl"
 #include "Math/Color.h"
@@ -68,33 +70,84 @@ namespace LowLevelTasks::Impl
 UE_DISABLE_OPTIMIZATION
 #endif
 
+CSV_DECLARE_CATEGORY_EXTERN(Scheduler);
+
 namespace LowLevelTasks::Private
 {
+	namespace WaitingQueueImpl
+	{
+		// State_ layout:
+		// - low kWaiterBits is a stack of waiters committed wait
+		//   (indexes in NodesArray are used as stack elements,
+		//   kStackMask means empty stack).
+		// - next kWaiterBits is count of waiters in prewait state.
+		// - next kWaiterBits is count of pending signals.
+		// - remaining bits are ABA counter for the stack.
+		//   (stored in Waiter node and incremented on push).
+		static constexpr uint64 WaiterBits = 14;
+		static constexpr uint64 StackMask = (1ull << WaiterBits) - 1;
+		static constexpr uint64 WaiterShift = WaiterBits;
+		static constexpr uint64 WaiterMask = ((1ull << WaiterBits) - 1) << WaiterShift;
+		static constexpr uint64 WaiterInc = 1ull << WaiterShift;
+		static constexpr uint64 SignalShift = 2 * WaiterBits;
+		static constexpr uint64 SignalMask = ((1ull << WaiterBits) - 1) << SignalShift;
+		static constexpr uint64 SignalInc = 1ull << SignalShift;
+		static constexpr uint64 EpochShift = 3 * WaiterBits;
+		static constexpr uint64 EpochBits = 64 - EpochShift;
+		static constexpr uint64 EpochMask = ((1ull << EpochBits) - 1) << EpochShift;
+		static constexpr uint64 EpochInc = 1ull << EpochShift;
+
+		// Get the active thread count out of the standby state.
+		uint64 GetActiveThreadCount(uint64 StandbyState)
+		{
+			// The StandbyState stores the active thread count in the waiter bits.
+			return (StandbyState & WaiterMask) >> WaiterShift;
+		}
+	}
 
-void FWaitingQueue::Init()
+void FWaitingQueue::Init(uint32 InThreadCount, uint32 InMaxThreadCount, TFunction<void()> InCreateThread, uint32 InActiveThreadCount)
 {
-	check(NodesArray.Num() < (1 << WaiterBits) - 1);
-	check(State == StackMask);
+	using namespace WaitingQueueImpl;
+
+	ThreadCount = InThreadCount;
+	MaxThreadCount = InMaxThreadCount;
+	CreateThread = InCreateThread;
+	Oversubscription = 0;
+	bIsShuttingDown = false;
+	State = StackMask;
+
+	// Store the external thread creations in the waiter bits which
+	// represent the number of currently active threads.
+	StandbyState = StackMask | ((uint64(InActiveThreadCount) << WaiterBits) & WaiterMask);
+
+	check(NodesArray.Num() < (1ull << WaiterBits) - 1);
 }
 
-void FWaitingQueue::Shutdown()
+void FWaitingQueue::FinishShutdown()
 {
+	using namespace WaitingQueueImpl;
+
 	check((State & (StackMask | WaiterMask)) == StackMask);
+	check((StandbyState & (StackMask | WaiterMask)) == StackMask);
 }
 
 void FWaitingQueue::PrepareWait(FWaitEvent* Node)
 {
+	using namespace WaitingQueueImpl;
+
 	WAITINGQUEUE_EVENT_SCOPE(FWaitingQueue_PrepareWait);
 
 	State.fetch_add(WaiterInc, std::memory_order_relaxed);
 }
 
-void FWaitingQueue::CheckState(uint64_t InState, bool bInIsWaiter)
+void FWaitingQueue::CheckState(uint64 InState, bool bInIsWaiter)
 {
+	using namespace WaitingQueueImpl;
+
 	static_assert(EpochBits >= 20, "Not enough bits to prevent ABA problem");
 #if WITH_WAITINGQUEUE_CHECK
-	const uint64_t Waiters = (InState & WaiterMask) >> WaiterShift;
-	const uint64_t Signals = (InState & SignalMask) >> SignalShift;
+	const uint64 Waiters = (InState & WaiterMask) >> WaiterShift;
+	const uint64 Signals = (InState & SignalMask) >> SignalShift;
 	check(Waiters >= Signals);
 	check(Waiters < (1 << WaiterBits) - 1);
 	check(!bInIsWaiter || Waiters > 0);
@@ -103,31 +156,46 @@ void FWaitingQueue::CheckState(uint64_t InState, bool bInIsWaiter)
 #endif
 }
 
+void FWaitingQueue::CheckStandbyState(uint64 InState)
+{
+	using namespace WaitingQueueImpl;
+
+#if WITH_WAITINGQUEUE_CHECK
+	const uint64 Index = (InState & StackMask);
+	const uint64 ActiveThreadCount = (InState & WaiterMask) >> WaiterShift;
+	const uint64 Signals = (InState & SignalMask) >> SignalShift;
+	check(Signals == 0); // Unused in this mode
+	check(ActiveThreadCount <= NodesArray.Num());
+	check(Index == StackMask || Index < NodesArray.Num())
+#endif
+}
+
 bool FWaitingQueue::CommitWait(FWaitEvent* Node, FOutOfWork& OutOfWork, int32 SpinCycles, int32 WaitCycles)
 {
+	using namespace WaitingQueueImpl;
+
 	{
 		WAITINGQUEUE_EVENT_SCOPE(FWaitingQueue_CommitWait);
 
 		check((Node->Epoch & ~EpochMask) == 0);
 		Node->State.store(EWaitState::NotSignaled, std::memory_order_relaxed);
 
-		const uint64_t Myself = (Node - &NodesArray[0]) | Node->Epoch;
-		uint64_t LocalState = State.load(std::memory_order_relaxed);
+		uint64 LocalState = State.load(std::memory_order_relaxed);
 
 		CheckState(LocalState, true);
-		uint64_t NewState;
+		uint64 NewState;
 		if ((LocalState & SignalMask) != 0)
 		{
 			WAITINGQUEUE_EVENT_SCOPE(CommitWait_TryConsume);
 			// Consume the signal and return immediately.
-			NewState = LocalState - WaiterInc - SignalInc;
+			NewState = LocalState - WaiterInc - SignalInc + EpochInc;
 		}
 		else
 		{
 			WAITINGQUEUE_EVENT_SCOPE(CommitWait_TryCommit);
 			// Remove this thread from pre-wait counter and add to the waiter stack.
-			NewState = ((LocalState & WaiterMask) - WaiterInc) | Myself;
-			Node->Next.store(LocalState & (StackMask | EpochMask), std::memory_order_relaxed);
+			NewState = ((LocalState & (WaiterMask | EpochMask)) - WaiterInc + EpochInc) | (Node - &NodesArray[0]);
+			Node->Next.store(LocalState & StackMask, std::memory_order_relaxed);
 		}
 		CheckState(NewState);
 		if (State.compare_exchange_weak(LocalState, NewState, std::memory_order_acq_rel, std::memory_order_relaxed))
@@ -135,8 +203,6 @@ bool FWaitingQueue::CommitWait(FWaitEvent* Node, FOutOfWork& OutOfWork, int32 Sp
 			if ((LocalState & SignalMask) == 0)
 			{
 				WAITINGQUEUE_EVENT_SCOPE(CommitWait_Success);
-				Node->Epoch += EpochInc;
-
 				// Fallthrough to park but we want to get out of the CommitWait scope first so it doesn't stick
 			}
 			else
@@ -162,9 +228,10 @@ bool FWaitingQueue::CommitWait(FWaitEvent* Node, FOutOfWork& OutOfWork, int32 Sp
 
 bool FWaitingQueue::CancelWait(FWaitEvent* Node)
 {
-	WAITINGQUEUE_EVENT_SCOPE(FWaitingQueue_CancelWait);
+	using namespace WaitingQueueImpl;
 
-	uint64_t LocalState = State.load(std::memory_order_relaxed);
+	WAITINGQUEUE_EVENT_SCOPE(FWaitingQueue_CancelWait);
+	uint64 LocalState = State.load(std::memory_order_relaxed);
 	for (;;)
 	{
 		bool bConsumedSignal = false;
@@ -207,48 +274,247 @@ bool FWaitingQueue::CancelWait(FWaitEvent* Node)
 	}
 }
 
+void FWaitingQueue::StartShutdown()
+{
+	using namespace WaitingQueueImpl;
+
+	bIsShuttingDown = true;
+
+	// Wake up all workers.
+	NotifyInternal(NodesArray.Num());
+
+	// Notification above doesn't trigger standby threads
+	// during shutdown so trigger them here.
+	uint64 LocalState = StandbyState;
+	while ((LocalState & StackMask) != StackMask)
+	{
+		FWaitEvent* Node = &NodesArray[LocalState & StackMask];
+		Node->Event->Trigger();
+		LocalState = Node->Next;
+	}
+	StandbyState = StackMask;
+}
+
+void FWaitingQueue::PrepareStandby(FWaitEvent* Node)
+{
+	// We store the whole state before going back checking the queue so that we can't possibly
+	// miss an event in-between PrepareStandby and CommitStandby.
+
+	WAITINGQUEUE_EVENT_SCOPE(FWaitingQueue_PrepareStandby);
+
+	Node->Epoch = StandbyState;
+}
+
+void FWaitingQueue::ConditionalStandby(FWaitEvent* Node)
+{
+	using namespace WaitingQueueImpl;
+
+	WAITINGQUEUE_EVENT_SCOPE(FWaitingQueue_ConditionalStandby);
+
+	if (bIsShuttingDown.load(std::memory_order_relaxed))
+	{
+		return;
+	}
+
+	uint64 LocalState = StandbyState;
+	while (GetActiveThreadCount(LocalState) > ThreadCount + Oversubscription.load(std::memory_order_relaxed))
+	{
+		WAITINGQUEUE_EVENT_SCOPE(FWaitingQueue_ConditionalStandby_Iteration);
+
+		CheckStandbyState(LocalState);
+		// We store the active thread count in the waiters slot, so decrement it by 1.
+		const uint64 Waiters  = (LocalState & WaiterMask) - WaiterInc;
+		const uint64 NewEpoch = (LocalState & EpochMask) + EpochInc;
+		const uint64 NewState = (Node - &NodesArray[0]) | NewEpoch | Waiters;
+
+		Node->Next.store(LocalState & StackMask);
+		Node->Event->Reset();
+
+		CheckStandbyState(NewState);
+		if (StandbyState.compare_exchange_weak(LocalState, NewState))
+		{
+			WAITINGQUEUE_EVENT_SCOPE(Standby);
+			Private::FOversubscriptionAllowedScope _(false /* Disallow oversubscription for this wait */);
+			Node->Event->Wait();
+		}
+		else
+		{
+			WAITINGQUEUE_EVENT_SCOPE(Standby_Fail);
+		}
+	}
+}
+
+bool FWaitingQueue::CommitStandby(FWaitEvent* Node, FOutOfWork& OutOfWork)
+{
+	using namespace WaitingQueueImpl;
+
+	{
+		WAITINGQUEUE_EVENT_SCOPE(FWaitingQueue_CommitStandby);
+
+		uint64 LocalState = Node->Epoch;
+		CheckStandbyState(LocalState);
+		// We store the active thread count in the waiters slot, so decrement it by 1.
+		const uint64 Waiters = (LocalState & WaiterMask) - WaiterInc;
+		const uint64 Epoch = (LocalState & EpochMask) + EpochInc;
+		const uint64 NewState = (Node - &NodesArray[0]) | Epoch | Waiters;
+
+		Node->Next.store(LocalState & StackMask);
+		Node->Event->Reset();
+
+		CheckStandbyState(NewState);
+		if (StandbyState.compare_exchange_strong(LocalState, NewState))
+		{
+			// fallthrough to the end of the function where we wait
+		}
+		else
+		{
+			WAITINGQUEUE_EVENT_SCOPE(CommitStandby_Abort);
+			// Update the value before we go back checking if new tasks have been queued.
+			Node->Epoch = LocalState;
+			return false;
+		}
+	}
+
+	OutOfWork.Stop();
+	// Flush any open scope before going to sleep so that anything that happened
+	// before appears in UnrealInsights right away. If we don't do this,
+	// the thread buffer will be held to this thread until we wake up and fill it
+	// so it might cause events to appear as missing in UnrealInsights, especially
+	// in case we never wake up again (i.e. deadlock / crash).
+	TRACE_CPUPROFILER_EVENT_FLUSH();
+	Private::FOversubscriptionAllowedScope _(false /* Disallow oversubscription for this wait */);
+	Node->Event->Wait();
+	return true;
+}
+
+void FWaitingQueue::IncrementOversubscription()
+{
+	using namespace WaitingQueueImpl;
+
+	++Oversubscription;
+
+	// This is important that StandbyState is invalidated after Oversubscription is increased so we
+	// can detect stale decisions and reevaluate oversubscription.
+	// Notify -> TryStartNewThread takes care of updating StandbyState for us, but only
+	// when standby threads are actually needed.
+
+	Notify();
+}
+
+void FWaitingQueue::DecrementOversubscription()
+{
+	--Oversubscription;
+}
+
+bool FWaitingQueue::TryStartNewThread()
+{
+	WAITINGQUEUE_EVENT_SCOPE(FWaitingQueue_TryStartNewThread);
+	using namespace WaitingQueueImpl;
+
+	// Invalidate the current state by adding an Epoch right away so compare-exchange for other threads can detect
+	// oversubscription has changed which happens in IncrementOversubscription before calling this function.
+	//
+	// Important to always read the StandbyState before the Oversubscription value so that we capture the current epoch to validate
+	// Oversubscription didn't change while we were doing the CAS.
+	uint64 LocalState = StandbyState.fetch_add(EpochInc, std::memory_order_seq_cst) + EpochInc;
+	while (GetActiveThreadCount(LocalState) < MaxThreadCount && GetActiveThreadCount(LocalState) < ThreadCount + Oversubscription.load(std::memory_order_relaxed))
+	{
+		WAITINGQUEUE_EVENT_SCOPE(FWaitingQueue_TryStartNewThread_Iteration);
+
+		CheckStandbyState(LocalState);
+
+		// We store the active thread count in the waiters slot, so increment it by 1.
+		const uint64 NewEpoch = (LocalState & EpochMask) + EpochInc;
+		uint64 NewState = NewEpoch | (LocalState & WaiterMask) + WaiterInc;
+		if ((LocalState & StackMask) != StackMask)
+		{
+			WAITINGQUEUE_EVENT_SCOPE(FWaitingQueue_TryStartNewThread_FoundNode);
+			FWaitEvent* Node = &NodesArray[LocalState & StackMask];
+			uint64 Next = Node->Next.load(std::memory_order_relaxed);
+			NewState |= Next & StackMask;
+		}
+		else
+		{
+			WAITINGQUEUE_EVENT_SCOPE(FWaitingQueue_TryStartNewThread_Empty);
+			NewState |= LocalState & StackMask;
+		}
+
+		CheckStandbyState(NewState);
+		if (StandbyState.compare_exchange_weak(LocalState, NewState, std::memory_order_acq_rel, std::memory_order_relaxed))
+		{
+			if ((LocalState & StackMask) != StackMask)
+			{
+				// We got an existing node, wake it from standby
+				WAITINGQUEUE_EVENT_SCOPE_ALWAYS(FWaitingQueue_SignalStandbyThread);
+				CSV_SCOPED_TIMING_STAT(Scheduler, SignalStandbyThread);
+				FWaitEvent* Node = &NodesArray[LocalState & StackMask];
+				Node->Event->Trigger();
+				return true;
+			}
+			else if (bIsShuttingDown.load(std::memory_order_relaxed) == false)
+			{
+				CSV_SCOPED_TIMING_STAT(Scheduler, CreateThread);
+				WAITINGQUEUE_EVENT_SCOPE_ALWAYS(FWaitingQueue_CreateThread);
+				CreateThread();
+				return true;
+			}
+			else
+			{
+				WAITINGQUEUE_EVENT_SCOPE(FWaitingQueue_TryStartNewThread_Backoff);
+				StandbyState -= WaiterInc;
+				return false;
+			}
+		}
+	}
+
+	return false;
+}
+
 int32 FWaitingQueue::NotifyInternal(int32 Count)
 {
+	using namespace WaitingQueueImpl;
+
 	WAITINGQUEUE_EVENT_SCOPE(FWaitingQueue_Notify);
 
 	int32 Notifications = 0;
 	while (Count > Notifications)
 	{
-		uint64_t LocalState = State.load(std::memory_order_relaxed);
+		uint64 LocalState = State.load(std::memory_order_relaxed);
 		for (;;)
 		{
 			CheckState(LocalState);
-			const uint64_t Waiters = (LocalState & WaiterMask) >> WaiterShift;
-			const uint64_t Signals = (LocalState & SignalMask) >> SignalShift;
+			const uint64 Waiters  = (LocalState & WaiterMask) >> WaiterShift;
+			const uint64 Signals  = (LocalState & SignalMask) >> SignalShift;
+			const uint64 NewEpoch = (LocalState & EpochMask) + EpochInc;
 			const bool bNotifyAll = Count >= NodesArray.Num();
 
-			uint64_t NewState;
+			uint64 NewState;
 			if ((LocalState & StackMask) == StackMask && Waiters == Signals)
 			{
 				// No more waiters, go through the CAS to provide proper ordering
 				// with other threads entering PrepareWait.
 				WAITINGQUEUE_EVENT_SCOPE(TryNoMoreWaiter);
-				NewState = LocalState;
+				NewState = LocalState + EpochInc;
 			}
 			else if (bNotifyAll)
 			{
 				WAITINGQUEUE_EVENT_SCOPE(TryUnblockAll);
 				// Empty wait stack and set signal to number of pre-wait threads.
-				NewState = (LocalState & WaiterMask) | (Waiters << SignalShift) | StackMask;
+				NewState = (LocalState & WaiterMask) | (Waiters << SignalShift) | StackMask | NewEpoch;
 			}
 			else if (Signals < Waiters)
 			{
 				WAITINGQUEUE_EVENT_SCOPE(TryAbortOnePreWait);
 				// There is a thread in pre-wait state, unblock it.
-				NewState = LocalState + SignalInc;
+				NewState = LocalState + SignalInc + EpochInc;
 			}
 			else
 			{
 				WAITINGQUEUE_EVENT_SCOPE(TryUnparkOne);
 				// Pop a waiter from list and unpark it.
 				FWaitEvent* Node = &NodesArray[LocalState & StackMask];
-				uint64_t Next = Node->Next.load(std::memory_order_relaxed);
-				NewState = (LocalState & (WaiterMask | SignalMask)) | Next;
+				uint64 Next = Node->Next.load(std::memory_order_relaxed);
+				NewState = (LocalState & (WaiterMask | SignalMask)) | (Next & StackMask) | NewEpoch;
 			}
 			CheckState(NewState);
 			if (State.compare_exchange_weak(LocalState, NewState, std::memory_order_acq_rel, std::memory_order_relaxed))
@@ -263,6 +529,11 @@ int32 FWaitingQueue::NotifyInternal(int32 Count)
 				if ((LocalState & StackMask) == StackMask)
 				{
 					WAITINGQUEUE_EVENT_SCOPE(NoMoreWaiter);
+					if (TryStartNewThread())
+					{
+						Notifications++;
+						break;
+					}
 					return Notifications;
 				}
 
@@ -289,6 +560,8 @@ int32 FWaitingQueue::NotifyInternal(int32 Count)
 
 void FWaitingQueue::Park(FWaitEvent* Node, FOutOfWork& OutOfWork, int32 SpinCycles, int32 WaitCycles)
 {
+	using namespace WaitingQueueImpl;
+
 	{
 		ON_SCOPE_EXIT { OutOfWork.Stop(); };
 		WAITINGQUEUE_EVENT_SCOPE(FWaitingQueue_Park);
@@ -331,20 +604,21 @@ void FWaitingQueue::Park(FWaitEvent* Node, FOutOfWork& OutOfWork, int32 SpinCycl
 	// the thread buffer will be held to this thread until we wake up and fill it
 	// so it might cause events to appear as missing in UnrealInsights, especially
 	// in case we never wake up again (i.e. deadlock / crash).
-	#ifdef TRACE_CPUPROFILER_EVENT_FLUSH
 	TRACE_CPUPROFILER_EVENT_FLUSH();
-	#endif
+	Private::FOversubscriptionAllowedScope _(false /* Disallow oversubscription for this wait */);
 	Node->Event->Wait();
 }
 
 int32 FWaitingQueue::Unpark(FWaitEvent* Node)
 {
+	using namespace WaitingQueueImpl;
+
 	WAITINGQUEUE_EVENT_SCOPE(FWaitingQueue_Unpark);
 
 	int32 UnparkedCount = 0;
 	for (FWaitEvent* Next; Node; Node = Next)
 	{
-		uint64_t NextNode = Node->Next.load(std::memory_order_relaxed) & StackMask;
+		uint64 NextNode = Node->Next.load(std::memory_order_relaxed) & StackMask;
 		Next = NextNode == StackMask ? nullptr : &NodesArray[(int)NextNode];
 
 		UnparkedCount++;
@@ -353,7 +627,7 @@ int32 FWaitingQueue::Unpark(FWaitEvent* Node)
 		// the event if the other thread was in the waiting state.
 		if (Node->State.exchange(EWaitState::Signaled, std::memory_order_relaxed) == EWaitState::Waiting)
 		{
-			// This one  we actually care about since signaling cost is very expensive.
+			// Always trace this one since signaling cost can be very expensive.
 			WAITINGQUEUE_EVENT_SCOPE_ALWAYS(FWaitingQueue_Unpark_SignalWaitingThread);
 			Node->Event->Trigger();
 		}
diff --git a/Engine/Source/Runtime/Core/Private/Async/TaskGraph.cpp b/Engine/Source/Runtime/Core/Private/Async/TaskGraph.cpp
--- a/Engine/Source/Runtime/Core/Private/Async/TaskGraph.cpp
+++ b/Engine/Source/Runtime/Core/Private/Async/TaskGraph.cpp
@@ -36,7 +36,6 @@
 #include "Misc/CoreDelegates.h"
 
 #include "Async/Fundamental/Scheduler.h"
-#include "Async/Fundamental/ReserveScheduler.h"
 #include "Tasks/Pipe.h"
 
 DEFINE_LOG_CATEGORY_STATIC(LogTaskGraph, Log, All);
@@ -110,6 +109,25 @@ static FAutoConsoleVariableRef CVarTaskDynamicPrioritization(
 	ECVF_ReadOnly
 );
 
+CORE_API float GTaskGraphOversubscriptionRatio = 2.0f;
+static FAutoConsoleVariableRef CVarTaskOversubscriptionRatio(
+	TEXT("TaskGraph.OversubscriptionRatio"),
+	GTaskGraphOversubscriptionRatio,
+	TEXT("Ratio used to compute the maximum numbers of workers allowed during oversubscription.\n")
+	TEXT("You might need to increase that value depending on how many recursive waits the scheduled tasks may contain.\n")
+	TEXT("The optimal scenario to strive for is using prerequisites to setup dependencies instead of waiting.\n")
+	TEXT("Once none of the scheduled tasks contains waiting logic anymore, this can be set to 1.0f, which effectively deactivate the feature."),
+	ECVF_ReadOnly
+);
+
+CORE_API bool GTaskGraphUseDynamicThreadCreation = true;
+static FAutoConsoleVariableRef CVarTaskDynamicThreadCreation(
+	TEXT("TaskGraph.UseDynamicThreadCreation"),
+	GTaskGraphUseDynamicThreadCreation,
+	TEXT("Allow threads to be created only when needed instead of at engine initialization."),
+	ECVF_ReadOnly
+);
+
 CORE_API int32 GUseNewTaskBackend = 1;
 CORE_API int32 GNumForegroundWorkers = 2;
 static FAutoConsoleVariableRef CVarNumForegroundWorkers(
@@ -1849,17 +1867,6 @@ public:
 
 			LowLevelTasks::FScheduler::Get().StartWorkers(NumForegroundWorkers, NumBackgroundWorkers, FForkProcessHelper::IsForkedMultithreadInstance() ? FThread::Forkable : FThread::NonForkable, FPlatformAffinity::GetTaskThreadPriority(), FPlatformAffinity::GetTaskBPThreadPriority());
 
-			check(IsInGameThread()); // otherwise we can have a race on starting reserve workers below
-			if (GConfig == nullptr)
-			{
-				// postpone starting reserve workers until GConfig is initialized, to know if reserve workers are disabled
-				FCoreDelegates::TSConfigReadyForUse().AddRaw(this, &FTaskGraphCompatibilityImplementation::StartReserveWorkers);
-			}
-			else
-			{
-				StartReserveWorkers();
-			}
-
 			NumNamedThreads = ENamedThreads::ActualRenderingThread + 1;
 			ENamedThreads::bHasBackgroundThreads = 1;
 			ENamedThreads::bHasHighPriorityThreads = 1;
@@ -1867,10 +1874,6 @@ public:
 		else
 		{
 			LowLevelTasks::FScheduler::Get().StopWorkers();
-			if (bReserveWorkersEnabled)
-			{
-				LowLevelTasks::FReserveScheduler::Get().StopWorkers();
-			}
 			NumNamedThreads = ENamedThreads::ActualRenderingThread;
 			ENamedThreads::bHasBackgroundThreads = 0;
 			ENamedThreads::bHasHighPriorityThreads = 0;
@@ -1903,10 +1906,6 @@ public:
 			NamedThreads[ThreadIndex].bAttached = false;
 		}
 		LowLevelTasks::FScheduler::Get().StopWorkers();
-		if (bReserveWorkersEnabled)
-		{
-			LowLevelTasks::FReserveScheduler::Get().StopWorkers();
-		}
 		FPlatformTLS::FreeTlsSlot(PerThreadIDTLSSlot);
 	}
 
@@ -1940,12 +1939,6 @@ public:
 
 			LowLevelTasks::FScheduler::Get().StopWorkers();
 			LowLevelTasks::FScheduler::Get().StartWorkers(NumForegroundWorkers, NumBackgroundWorkers, FForkProcessHelper::IsForkedMultithreadInstance() ? FThread::Forkable : FThread::NonForkable, Pri, FPlatformAffinity::GetTaskBPThreadPriority());
-
-			if (bReserveWorkersEnabled)
-			{
-				LowLevelTasks::FReserveScheduler::Get().StopWorkers();
-				LowLevelTasks::FReserveScheduler::Get().StartWorkers(NumForegroundWorkers + NumBackgroundWorkers, FForkProcessHelper::IsForkedMultithreadInstance() ? FThread::Forkable : FThread::NonForkable, FPlatformAffinity::GetTaskBPThreadPriority());
-			}
 		}
 	}
 
@@ -2146,75 +2139,28 @@ private:
 		}
 #endif
 
-		if (LowLevelTasks::FScheduler::Get().IsWorkerThread() || LowLevelTasks::FReserveScheduler::Get().IsWorkerThread())
+		if (!FTaskGraphInterface::IsMultithread())
 		{
-			// a worker thread gets blocked, involve a reserve worker to utilise an idle core
-			bool bSuccess = false; // has a reserve worker helped?
-			if (bReserveWorkersEnabled.load(std::memory_order_relaxed))
+			bool bAnyPending = false;
+			for (int32 Index = 0; Index < Tasks.Num(); Index++)
 			{
-				if (LowLevelTasks::DoReserveWorkUntil([Index(0), Tasks]() mutable
-				{
-					while (Index < Tasks.Num())
-					{
-						if (!Tasks[Index]->IsComplete())
-						{
-							return false;
-						}
-						Index++;
-					}
-					return true;
-				}))
+				FGraphEvent* Task = Tasks[Index].GetReference();
+				if (Task && !Task->IsComplete())
 				{
-					// a reserve worker was woken up to help with the tasks.
-					// stall this thread on an event while we wait
-					FScopedEvent Event;
-					TriggerEventWhenTasksComplete(Event.Get(), Tasks, CurrentThreadIfKnown);
-					bSuccess = true;
+					bAnyPending = true;
+					break;
 				}
 			}
-			
-			// for latency reason, worker thread aren't tackling other tasks while waiting if we have reserve workers, 
-			// otherwise worker threads help with task execution while they wait
-			if(!bSuccess)
+			if (!bAnyPending)
 			{
-				LowLevelTasks::BusyWaitUntil([Index(0), &Tasks]() mutable
-				{
-					while (Index < Tasks.Num())
-					{
-						if (!Tasks[Index]->IsComplete())
-						{
-							return false;
-						}
-						Index++;
-					}
-					return true;
-				});
-			}
-		}
-		else
-		{
-			if (!FTaskGraphInterface::IsMultithread())
-			{
-				bool bAnyPending = false;
-				for (int32 Index = 0; Index < Tasks.Num(); Index++)
-				{
-					FGraphEvent* Task = Tasks[Index].GetReference();
-					if (Task && !Task->IsComplete())
-					{
-						bAnyPending = true;
-						break;
-					}
-				}
-				if (!bAnyPending)
-				{
-					return;
-				}
-				UE_LOG(LogTaskGraph, Fatal, TEXT("Recursive waits are not allowed in single threaded mode."));
+				return;
 			}
-			// We will just stall this thread on an event while we wait
-			FScopedEvent Event;
-			TriggerEventWhenTasksComplete(Event.Get(), Tasks, CurrentThreadIfKnown);
+			UE_LOG(LogTaskGraph, Fatal, TEXT("Recursive waits are not allowed in single threaded mode."));
 		}
+
+		// We will just stall this thread on an event while we wait
+		FScopedEvent Event;
+		TriggerEventWhenTasksComplete(Event.Get(), Tasks, CurrentThreadIfKnown);
 	}
 
 	void TriggerEventWhenTasksComplete(FEvent* InEvent, const FGraphEventArray& Tasks, ENamedThreads::Type CurrentThreadIfKnown = ENamedThreads::AnyThread, ENamedThreads::Type TriggerThread = ENamedThreads::AnyHiPriThreadHiPriTask) final override
@@ -2308,29 +2254,6 @@ private:
 		checkThreadGraph(NamedThreads[Index].TaskGraphWorker->GetThreadId() == Index);
 		return *NamedThreads[Index].TaskGraphWorker;
 	}
-
-	void StartReserveWorkers()
-	{
-		TRACE_CPUPROFILER_EVENT_SCOPE(StartReserveWorkers);
-		if (bReserveWorkersEnabled)
-		{
-			return; // once enabled, reserve workers can't be disabled
-		}
-
-		NumBackgroundWorkers = FMath::Max(1, NumWorkerThreads - FMath::Min<int>(GNumForegroundWorkers, NumWorkerThreads));
-		NumForegroundWorkers = FMath::Max(1, NumWorkerThreads - NumBackgroundWorkers);
-
-		check(GConfig);
-		
-		bool bEnableReserveWorkers = false; // by default
-		GConfig->GetBool(TEXT("TaskGraph"), TEXT("EnableReserveWorkers"), bEnableReserveWorkers, GEngineIni);
-
-		if (bEnableReserveWorkers)
-		{
-			bReserveWorkersEnabled = true;
-			LowLevelTasks::FReserveScheduler::Get().StartWorkers(NumForegroundWorkers + NumBackgroundWorkers, FForkProcessHelper::IsForkedMultithreadInstance() ? FThread::Forkable : FThread::NonForkable, FPlatformAffinity::GetTaskBPThreadPriority());
-		}
-	}
 };
 
 thread_local TArray<FBaseGraphTask*> FTaskGraphCompatibilityImplementation::NewTasks;
diff --git a/Engine/Source/Runtime/Core/Private/GenericPlatform/GenericPlatformManualResetEvent.cpp b/Engine/Source/Runtime/Core/Private/GenericPlatform/GenericPlatformManualResetEvent.cpp
new file mode 100644
--- /dev/null
+++ b/Engine/Source/Runtime/Core/Private/GenericPlatform/GenericPlatformManualResetEvent.cpp
@@ -0,0 +1,27 @@
+// Copyright Epic Games, Inc. All Rights Reserved.
+
+#include "GenericPlatform/GenericPlatformManualResetEvent.h"
+
+#include "Async/Fundamental/Scheduler.h"
+
+namespace UE::HAL::Private
+{
+	bool FGenericPlatformManualResetEvent::WaitUntil(FMonotonicTimePoint WaitTime)
+	{
+		std::unique_lock SelfLock(Lock);
+		if (WaitTime.IsInfinity())
+		{
+			LowLevelTasks::FOversubscriptionScope _;
+			Condition.wait(SelfLock, [this] { return !bWait.load(std::memory_order_acquire); });
+			return true;
+		}
+		if (FMonotonicTimeSpan WaitSpan = WaitTime - FMonotonicTimePoint::Now(); WaitSpan > FMonotonicTimeSpan::Zero())
+		{
+			LowLevelTasks::FOversubscriptionScope _;
+			const int64 WaitMs = FPlatformMath::CeilToInt64(WaitSpan.ToMilliseconds());
+			return Condition.wait_for(SelfLock, std::chrono::milliseconds(WaitMs), [this] { return !bWait.load(std::memory_order_acquire); });
+		}
+		return !bWait.load(std::memory_order_acquire);
+	}
+
+} // UE::HAL::Private
diff --git a/Engine/Source/Runtime/Core/Private/GenericPlatform/GenericPlatformManualResetEvent.h b/Engine/Source/Runtime/Core/Private/GenericPlatform/GenericPlatformManualResetEvent.h
--- a/Engine/Source/Runtime/Core/Private/GenericPlatform/GenericPlatformManualResetEvent.h
+++ b/Engine/Source/Runtime/Core/Private/GenericPlatform/GenericPlatformManualResetEvent.h
@@ -55,21 +55,7 @@ public:
 	 * @param WaitTime   Absolute time after which waiting is canceled and the thread wakes.
 	 * @return True if Notify() was called before the wait time elapsed, otherwise false.
 	 */
-	bool WaitUntil(FMonotonicTimePoint WaitTime)
-	{
-		std::unique_lock SelfLock(Lock);
-		if (WaitTime.IsInfinity())
-		{
-			Condition.wait(SelfLock, [this] { return !bWait.load(std::memory_order_acquire); });
-			return true;
-		}
-		if (FMonotonicTimeSpan WaitSpan = WaitTime - FMonotonicTimePoint::Now(); WaitSpan > FMonotonicTimeSpan::Zero())
-		{
-			const int64 WaitMs = FPlatformMath::CeilToInt64(WaitSpan.ToMilliseconds());
-			return Condition.wait_for(SelfLock, std::chrono::milliseconds(WaitMs), [this] { return !bWait.load(std::memory_order_acquire); });
-		}
-		return !bWait.load(std::memory_order_acquire);
-	}
+	bool WaitUntil(FMonotonicTimePoint WaitTime);
 
 	/**
 	 * Notifies the waiting thread.
diff --git a/Engine/Source/Runtime/Core/Private/GenericPlatform/GenericPlatformProcess.cpp b/Engine/Source/Runtime/Core/Private/GenericPlatform/GenericPlatformProcess.cpp
--- a/Engine/Source/Runtime/Core/Private/GenericPlatform/GenericPlatformProcess.cpp
+++ b/Engine/Source/Runtime/Core/Private/GenericPlatform/GenericPlatformProcess.cpp
@@ -472,6 +472,7 @@ bool FPThreadEvent::Wait(uint32 WaitTime, const bool bIgnoreThreadIdleStats /*=
 			WaitingThreads++;
 			if (WaitTime == ((uint32)-1)) // infinite wait?
 			{
+				LowLevelTasks::FOversubscriptionScope _; // Let the scheduler know one of its thread might be waiting.
 				int rc = pthread_cond_wait(&Condition, &Mutex);  // unlocks Mutex while blocking...
 				check(rc == 0);
 			}
@@ -481,6 +482,7 @@ bool FPThreadEvent::Wait(uint32 WaitTime, const bool bIgnoreThreadIdleStats /*=
 				const uint32 ms = (StartTime.tv_usec / 1000) + WaitTime;
 				TimeOut.tv_sec = StartTime.tv_sec + (ms / 1000);
 				TimeOut.tv_nsec = (ms % 1000) * 1000000;  // remainder of milliseconds converted to nanoseconds.
+				LowLevelTasks::FOversubscriptionScope _; // Let the scheduler know one of its thread might be waiting.
 				int rc = pthread_cond_timedwait(&Condition, &Mutex, &TimeOut);    // unlocks Mutex while blocking...
 				check((rc == 0) || (rc == ETIMEDOUT));
 
diff --git a/Engine/Source/Runtime/Core/Private/Microsoft/MicrosoftPlatformManualResetEvent.cpp b/Engine/Source/Runtime/Core/Private/Microsoft/MicrosoftPlatformManualResetEvent.cpp
--- a/Engine/Source/Runtime/Core/Private/Microsoft/MicrosoftPlatformManualResetEvent.cpp
+++ b/Engine/Source/Runtime/Core/Private/Microsoft/MicrosoftPlatformManualResetEvent.cpp
@@ -4,7 +4,7 @@
 
 #include "HAL/PlatformMath.h"
 #include "Misc/MonotonicTime.h"
-
+#include "Async/Fundamental/Scheduler.h"
 #include "Microsoft/WindowsHWrapper.h"
 
 namespace UE::HAL::Private
@@ -20,6 +20,9 @@ bool FMicrosoftPlatformManualResetEvent::WaitUntil(FMonotonicTimePoint WaitTime)
 	bool bLocalWait = true;
 	if (WaitTime.IsInfinity())
 	{
+		// Let the scheduler know one of its thread might be waiting.
+		LowLevelTasks::FOversubscriptionScope _;
+
 		for (;;)
 		{
 			if (WaitOnAddress(&bWait, &bLocalWait, sizeof(bool), INFINITE) && !bWait)
@@ -30,6 +33,9 @@ bool FMicrosoftPlatformManualResetEvent::WaitUntil(FMonotonicTimePoint WaitTime)
 	}
 	else
 	{
+		// Let the scheduler know one of its thread might be waiting.
+		LowLevelTasks::FOversubscriptionScope _(WaitTime - FMonotonicTimePoint::Now() > FMonotonicTimeSpan::Zero());
+
 		for (;;)
 		{
 			FMonotonicTimeSpan WaitSpan = WaitTime - FMonotonicTimePoint::Now();
diff --git a/Engine/Source/Runtime/Core/Private/Tasks/TaskPrivate.cpp b/Engine/Source/Runtime/Core/Private/Tasks/TaskPrivate.cpp
--- a/Engine/Source/Runtime/Core/Private/Tasks/TaskPrivate.cpp
+++ b/Engine/Source/Runtime/Core/Private/Tasks/TaskPrivate.cpp
@@ -42,7 +42,6 @@ namespace UE::Tasks
 			}
 #endif
 
-			bWakeUpWorker |= LowLevelTasks::FSchedulerTls::IsBusyWaiting();
 			bWakeUpWorker |= LowLevelTasks::FScheduler::Get().TryLaunch(LowLevelTask, bWakeUpWorker ? LowLevelTasks::EQueuePreference::GlobalQueuePreference : LowLevelTasks::EQueuePreference::LocalQueuePreference, bWakeUpWorker);
 		}
 
diff --git a/Engine/Source/Runtime/Core/Private/Windows/WindowsPlatformProcess.cpp b/Engine/Source/Runtime/Core/Private/Windows/WindowsPlatformProcess.cpp
--- a/Engine/Source/Runtime/Core/Private/Windows/WindowsPlatformProcess.cpp
+++ b/Engine/Source/Runtime/Core/Private/Windows/WindowsPlatformProcess.cpp
@@ -5,6 +5,7 @@
 #include "Containers/Set.h"
 #include "Containers/StringConv.h"
 #include "Containers/UnrealString.h"
+#include "Async/Fundamental/Scheduler.h"
 #include "CoreGlobals.h"
 #include "GenericPlatform/GenericPlatformFile.h"
 #include "HAL/PlatformAffinity.h"
@@ -1633,6 +1634,7 @@ bool FEventWin::Wait(uint32 WaitTime, const bool bIgnoreThreadIdleStats /*= fals
 	CSV_SCOPED_WAIT(WaitTime);
 	check(Event);
 
+	LowLevelTasks::FOversubscriptionScope _(WaitTime != 0); // Let the scheduler know one of its thread might be waiting.
 	FThreadIdleStats::FScopeIdle Scope( bIgnoreThreadIdleStats );
 	return (WaitForSingleObject( Event, WaitTime ) == WAIT_OBJECT_0);
 }
diff --git a/Engine/Source/Runtime/Core/Public/Async/AsyncWork.h b/Engine/Source/Runtime/Core/Public/Async/AsyncWork.h
--- a/Engine/Source/Runtime/Core/Public/Async/AsyncWork.h
+++ b/Engine/Source/Runtime/Core/Public/Async/AsyncWork.h
@@ -345,11 +345,6 @@ class FAsyncTaskBase
 			FScopeCycleCounter Scope(StatId);
 			DECLARE_SCOPE_CYCLE_COUNTER(TEXT("FAsyncTask::SyncCompletion"), STAT_FAsyncTask_SyncCompletion, STATGROUP_ThreadPoolAsyncTasks);
 
-			if (LowLevelTasks::FScheduler::Get().IsWorkerThread() && !bIsLatencySensitive)
-			{
-				LowLevelTasks::BusyWaitUntil([this]() { return IsWorkDone(); });
-			}
-
 			check(DoneEvent); // if it is not done yet, we must have an event
 			DoneEvent->Wait();
 			QueuedPool = 0;
diff --git a/Engine/Source/Runtime/Core/Public/Async/Fundamental/LocalQueue.h b/Engine/Source/Runtime/Core/Public/Async/Fundamental/LocalQueue.h
--- a/Engine/Source/Runtime/Core/Public/Async/Fundamental/LocalQueue.h
+++ b/Engine/Source/Runtime/Core/Public/Async/Fundamental/LocalQueue.h
@@ -144,7 +144,7 @@ enum class ELocalQueueType
  * A Dequeue Operation can only be done starting from a LocalQueue, than the GlobalQueue will be checked.                                   *
  * Finally Items might get Stolen from other LocalQueues that are registered with the LocalQueueRegistry.                                   *
  ********************************************************************************************************************************************/
-template<uint32 NumLocalItems = 1024>
+template<uint32 NumLocalItems = 1024, uint32 MaxLocalQueues = 1024>
 class TLocalQueueRegistry
 {
 	static uint32 Rand()
@@ -163,32 +163,20 @@ private:
 	using FOverflowQueueType = FAAArrayQueue<FTask>;
 	using DequeueHazard		 = typename FOverflowQueueType::DequeueHazard;
 
-	enum class EOverflowType
-	{
-		// Tasks that can't be run inside busy waiting will be moved in that queue during busy wait loops.
-		// This queue is always looked at first by non busy waiting workers to keep ordering as much as possible.
-		DoNotRunInsideBusyWait,
-		// All tasks including those that can't run inside busy wait are queued here by default to keep FIFO ordering as much as possible.
-		Default,
-		Count,
-	};
-
 public:
 	class TLocalQueue
 	{
-		template<uint32>
+		template<uint32, uint32>
 		friend class TLocalQueueRegistry;
 
 	public:
 		TLocalQueue(TLocalQueueRegistry& InRegistry, ELocalQueueType InQueueType) : Registry(&InRegistry), QueueType(InQueueType)
 		{
-			AffinityIndex = Registry->AddLocalQueue(this, QueueType);
+			// Local queues are never unregistered, everything is shutdown at once.
+			Registry->AddLocalQueue(this);
 			for (int32 PriorityIndex = 0; PriorityIndex < int32(ETaskPriority::Count); ++PriorityIndex)
 			{
-				for (int32 OverflowIndex = 0; OverflowIndex < int32(EOverflowType::Count); ++OverflowIndex)
-				{
-					DequeueHazards[PriorityIndex][OverflowIndex] = Registry->OverflowQueues[PriorityIndex][OverflowIndex].getHeadHazard();
-				}
+				DequeueHazards[PriorityIndex] = Registry->OverflowQueues[PriorityIndex].getHeadHazard();
 			}
 		}
 
@@ -203,10 +191,9 @@ public:
 					{
 						break;
 					}
-					Registry->OverflowQueues[PriorityIndex][(int32)EOverflowType::Default].enqueue(Item);
+					Registry->OverflowQueues[PriorityIndex].enqueue(Item);
 				}
 			}
-			Registry->DeleteLocalQueue(this, QueueType);
 		}
 
 		// add an item to the local queue and overflow into the global queue if full
@@ -219,16 +206,14 @@ public:
 
 			if (!LocalQueues[PriorityIndex].Put(Item))
 			{
-				Registry->OverflowQueues[PriorityIndex][(int32)EOverflowType::Default].enqueue(Item);
+				Registry->OverflowQueues[PriorityIndex].enqueue(Item);
 			}
 		}
 
 		// Check both the local and global queue in priority order
-		template <bool bIsInsideBusyWait = false>
 		inline FTask* Dequeue(bool GetBackGroundTasks)
 		{
 			const int32 MaxPriority = GetBackGroundTasks ? int32(ETaskPriority::Count)   : int32(ETaskPriority::ForegroundCount);
-			constexpr int32 MinOverflow = bIsInsideBusyWait  ? int32(EOverflowType::Default) : int32(EOverflowType::DoNotRunInsideBusyWait);
 
 			for (int32 PriorityIndex = 0; PriorityIndex < MaxPriority; ++PriorityIndex)
 			{
@@ -238,13 +223,10 @@ public:
 					return Item;
 				}
 
-				for (int32 OverflowIndex = MinOverflow; OverflowIndex < int32(EOverflowType::Count); ++OverflowIndex)
+				Item = Registry->OverflowQueues[PriorityIndex].dequeue(DequeueHazards[PriorityIndex]);
+				if (Item)
 				{
-					Item = Registry->OverflowQueues[PriorityIndex][OverflowIndex].dequeue(DequeueHazards[PriorityIndex][OverflowIndex]);
-					if (Item)
-					{
-						return Item;
-					}
+					return Item;
 				}
 			}
 			return nullptr;
@@ -265,60 +247,55 @@ public:
 			return nullptr;
 		}
 
-		uint32 GetAffinityIndex() const
-		{
-			return AffinityIndex;
-		}
-
 	private:
 		static constexpr uint32    InvalidIndex = ~0u;
 		FLocalQueueType            LocalQueues[uint32(ETaskPriority::Count)];
-		DequeueHazard              DequeueHazards[uint32(ETaskPriority::Count)][int32(EOverflowType::Count)];
+		DequeueHazard              DequeueHazards[uint32(ETaskPriority::Count)];
 		TLocalQueueRegistry*       Registry;
 		uint32                     CachedRandomIndex = InvalidIndex;
 		uint32                     CachedPriorityIndex = 0;
-		uint32                     AffinityIndex = ~0u;
 		ELocalQueueType            QueueType;
 	};
 
-	TLocalQueueRegistry() = default;
-
-private:
-	// add a queue to the Registry
-	uint32 AddLocalQueue(TLocalQueue* QueueToAdd, ELocalQueueType QueueType)
+	TLocalQueueRegistry()
 	{
-		NumActiveWorkers[QueueType == ELocalQueueType::EBackground]++;
-		LocalQueues.Add(QueueToAdd);
-		return LocalQueues.Num() - 1;
 	}
 
-	// remove a queue from the Registry
-	void DeleteLocalQueue(TLocalQueue* QueueToRemove, ELocalQueueType QueueType)
+private:
+	// Add a queue to the Registry. Thread-safe.
+	void AddLocalQueue(TLocalQueue* QueueToAdd)
 	{
-		NumActiveWorkers[QueueType == ELocalQueueType::EBackground]--;
-		verify(LocalQueues.Remove(QueueToRemove) == 1);
+		uint32 Index = NumLocalQueues.fetch_add(1, std::memory_order_relaxed);
+		check(Index < MaxLocalQueues);
+
+		// std::memory_order_release to make sure values are all written to the TLocalQueue before publishing.
+		LocalQueues[Index].store(QueueToAdd, std::memory_order_release);
 	}
 
 	// StealItem tries to steal an Item from a Registered LocalQueue
+	// Thread-safe with AddLocalQueue
 	FTask* StealItem(uint32& CachedRandomIndex, uint32& CachedPriorityIndex, bool GetBackGroundTasks)
 	{
-		uint32 NumQueues = LocalQueues.Num();
+		uint32 NumQueues   = NumLocalQueues.load(std::memory_order_relaxed);
 		uint32 MaxPriority = GetBackGroundTasks ? int32(ETaskPriority::Count) : int32(ETaskPriority::ForegroundCount);
-		CachedRandomIndex = CachedRandomIndex % NumQueues;
+		CachedRandomIndex  = CachedRandomIndex % NumQueues;
 
-		for(uint32 i = 0; i < NumQueues; i++)
+		for (uint32 Index = 0; Index < NumLocalQueues; Index++)
 		{
-			TLocalQueue* LocalQueue = LocalQueues[CachedRandomIndex];
-			for(uint32 PriorityIndex = 0; PriorityIndex < MaxPriority; PriorityIndex++)
+			// Test for null in case we race on reading NumLocalQueues reserved index before the pointer is set
+			if (TLocalQueue* LocalQueue = LocalQueues[Index].load(std::memory_order_acquire))
 			{
-				FTask* Item;
-				if (LocalQueue->LocalQueues[CachedPriorityIndex].Steal(Item))
+				for(uint32 PriorityIndex = 0; PriorityIndex < MaxPriority; PriorityIndex++)
 				{
-					return Item;
+					FTask* Item;
+					if (LocalQueue->LocalQueues[PriorityIndex].Steal(Item))
+					{
+						return Item;
+					}
+					CachedPriorityIndex = ++CachedPriorityIndex < MaxPriority ? CachedPriorityIndex : 0;
 				}
-				CachedPriorityIndex = ++CachedPriorityIndex < MaxPriority ? CachedPriorityIndex : 0;
+				CachedRandomIndex = ++CachedRandomIndex < NumQueues ? CachedRandomIndex : 0;
 			}
-			CachedRandomIndex = ++CachedRandomIndex < NumQueues ? CachedRandomIndex : 0;
 		}
 		CachedPriorityIndex = 0;
 		CachedRandomIndex = TLocalQueue::InvalidIndex;
@@ -327,31 +304,24 @@ private:
 
 public:
 	// enqueue an Item directy into the Global OverflowQueue
-	template <bool bAllowInsideBusyWaiting = true>
 	void Enqueue(FTask* Item, uint32 PriorityIndex)
 	{
 		check(PriorityIndex < int32(ETaskPriority::Count));
 		check(Item != nullptr);
 
-		constexpr int32 OverflowIndex = bAllowInsideBusyWaiting ? (int32)EOverflowType::Default : (int32)EOverflowType::DoNotRunInsideBusyWait;
-		OverflowQueues[PriorityIndex][OverflowIndex].enqueue(Item);
+		OverflowQueues[PriorityIndex].enqueue(Item);
 	}
 
 	// grab an Item directy from the Global OverflowQueue
-	template <bool bIsInsideBusyWait = false>
 	FTask* DequeueGlobal(bool GetBackGroundTasks = true)
 	{
 		const int32 MaxPriority = GetBackGroundTasks ? int32(ETaskPriority::Count) : int32(ETaskPriority::ForegroundCount);
-		constexpr int32 MinOverflow = bIsInsideBusyWait ? int32(EOverflowType::Default) : int32(EOverflowType::DoNotRunInsideBusyWait);
 
 		for (int32 PriorityIndex = 0; PriorityIndex < MaxPriority; ++PriorityIndex)
 		{
-			for (int32 OverflowIndex = MinOverflow; OverflowIndex < int32(EOverflowType::Count); ++OverflowIndex)
+			if (FTask* Item = OverflowQueues[PriorityIndex].dequeue())
 			{
-				if (FTask* Item = OverflowQueues[PriorityIndex][OverflowIndex].dequeue())
-				{
-					return Item;
-				}
+				return Item;
 			}
 		}
 		return nullptr;
@@ -369,368 +339,24 @@ public:
 		return nullptr;
 	}
 
-private:
-	FOverflowQueueType     OverflowQueues[uint32(ETaskPriority::Count)][uint32(EOverflowType::Count)];
-	TArray<TLocalQueue*>   LocalQueues;
-	int                    NumActiveWorkers[2] = { 0, 0 };
-};
-
-} // namespace Private
-
-enum class UE_DEPRECATED(5.4, "This was meant for internal use only and will be removed") ELocalQueueType
-{
-	EBackground,
-	EForeground,
-};
-
-
-template<uint32 NumLocalItems = 1024>
-class UE_DEPRECATED(5.4, "This was meant for internal use only and will be removed") TLocalQueueRegistry
-{
-PRAGMA_DISABLE_DEPRECATION_WARNINGS
-	static uint32 Rand()
-	{
-		uint32 State = FPlatformTime::Cycles();
-		State = State * 747796405u + 2891336453u;
-		State = ((State >> ((State >> 28u) + 4u)) ^ State) * 277803737u;
-		return (State >> 22u) ^ State;
-	}
-
-public:
-	class TLocalQueue;
-
-	// FOutOfWork is used to track the time while a worker is waiting for work
-	// this happens after a worker was unable to aquire any task from the queues and until it finds work again or it goes into drowsing state.
-	class FOutOfWork
-	{
-	private:
-		template<uint32>
-		friend class TLocalQueueRegistry;
-
-		std::atomic_int& NumWorkersLookingForWork;
-		bool ActivelyLookingForWork = false;
-
-#if CPUPROFILERTRACE_ENABLED
-		bool bCpuBeginEventEmitted = false;
-#endif
-
-		inline FOutOfWork(std::atomic_int& InNumWorkersLookingForWork) : NumWorkersLookingForWork(InNumWorkersLookingForWork)
-		{
-		}
-
-	public:
-		inline ~FOutOfWork()
-		{
-			Stop();
-		}
-
-		inline bool Start()
-		{
-			if (!ActivelyLookingForWork)
-			{
-#if CPUPROFILERTRACE_ENABLED
-				if (CpuChannel)
-				{
-					static uint32 WorkerLookingForWorkTraceId = FCpuProfilerTrace::OutputEventType("TaskWorkerIsLookingForWork");
-					FCpuProfilerTrace::OutputBeginEvent(WorkerLookingForWorkTraceId);
-					bCpuBeginEventEmitted = true;
-				}
-#endif
-				NumWorkersLookingForWork.fetch_add(1, std::memory_order_relaxed);
-				ActivelyLookingForWork = true;
-				return true;
-			}
-			return false;
-		}
-
-		inline bool Stop()
-		{
-			if (ActivelyLookingForWork)
-			{
-#if CPUPROFILERTRACE_ENABLED
-				if (bCpuBeginEventEmitted)
-				{
-					FCpuProfilerTrace::OutputEndEvent();
-					bCpuBeginEventEmitted = false;
-				}
-#endif
-				NumWorkersLookingForWork.fetch_sub(1, std::memory_order_release);
-				ActivelyLookingForWork = false;
-				return true;
-			}
-			return false;
-		}
-	};
-
-private:
-	using FLocalQueueType	 = LocalQueue_Impl::TWorkStealingQueue2<FTask, NumLocalItems>;
-	using FOverflowQueueType = FAAArrayQueue<FTask>;
-	using DequeueHazard		 = typename FOverflowQueueType::DequeueHazard;
-
-public:
-	class TLocalQueue
-	{
-		template<uint32>
-		friend class TLocalQueueRegistry;
-
-	public:
-		TLocalQueue(TLocalQueueRegistry& InRegistry, ELocalQueueType InQueueType, FSleepEvent* InSleepEvent) : Registry(&InRegistry), SleepEvent(InSleepEvent), QueueType(InQueueType)
-		{
-			AffinityIndex = Registry->AddLocalQueue(this, QueueType);
-			for (int32 PriorityIndex = 0; PriorityIndex < int32(ETaskPriority::Count); PriorityIndex++)
-			{
-				DequeueHazards[PriorityIndex] = Registry->OverflowQueues[PriorityIndex].getHeadHazard();
-			}
-			AffinityHazard = AffinityQueue.getHeadHazard();
-		}
-
-		~TLocalQueue()
-		{
-			for (int32 PriorityIndex = 0; PriorityIndex < int32(ETaskPriority::Count); PriorityIndex++)
-			{
-				while (true)
-				{
-					FTask* Item;
-					if (!LocalQueues[PriorityIndex].Get(Item))
-					{
-						break;
-					}
-					Registry->OverflowQueues[PriorityIndex].enqueue(Item);
-				}
-			}
-			Registry->DeleteLocalQueue(this, QueueType);
-		}
-
-		//add an item to the local queue and overflow into the global queue if full
-		// returns true if we should wake a worker
-		inline bool Enqueue(FTask* Item, uint32 PriorityIndex)
-		{
-			checkSlow(Registry);
-			checkSlow(PriorityIndex < int32(ETaskPriority::Count));
-			checkSlow(Item != nullptr);
-
-			bool bBackgroundTask = Item->IsBackgroundTask();
-			if (!LocalQueues[PriorityIndex].Put(Item))
-			{
-				Registry->OverflowQueues[PriorityIndex].enqueue(Item);
-			}
-			return Registry->LessThanHalfWorkersLookingForWork(bBackgroundTask);
-		}
-
-		inline FTask* DequeueLocal(bool GetBackGroundTasks, bool bDisableThrottleStealing)
-		{
-			int32 MaxPriority = GetBackGroundTasks ? int32(ETaskPriority::Count) : int32(ETaskPriority::ForegroundCount);
-			for (int32 PriorityIndex = 0; PriorityIndex < MaxPriority; PriorityIndex++)
-			{
-				FTask* Item;
-				if (LocalQueues[PriorityIndex].Get(Item))
-				{
-					return Item;
-				}
-			}
-			return nullptr;
-		}
-
-		inline FTask* DequeueGlobal(bool GetBackGroundTasks, bool bDisableThrottleStealing)
-		{
-			if (bDisableThrottleStealing || (Registry->NumActiveWorkers[GetBackGroundTasks] >= (2 * Registry->NumWorkersLookingForWork[GetBackGroundTasks].load(std::memory_order_relaxed) - 1)) || ((Rand() % 4) == 0))
-			{
-				int32 MaxPriority = GetBackGroundTasks ? int32(ETaskPriority::Count) : int32(ETaskPriority::ForegroundCount);
-				for (int32 PriorityIndex = 0; PriorityIndex < MaxPriority; PriorityIndex++)
-				{
-					FTask* Item = Registry->OverflowQueues[PriorityIndex].dequeue(DequeueHazards[PriorityIndex]);
-					if (Item)
-					{
-						return Item;
-					}
-				}
-			}
-			return nullptr;
-		}
-
-		inline FTask* DequeueSteal(bool GetBackGroundTasks, bool bDisableThrottleStealing)
-		{
-			if (bDisableThrottleStealing || (Registry->NumActiveWorkers[GetBackGroundTasks] >= (2 * Registry->NumWorkersLookingForWork[GetBackGroundTasks].load(std::memory_order_relaxed) - 1)) || ((Rand() % 4) == 0))
-			{
-				if (CachedRandomIndex == InvalidIndex)
-				{
-					CachedRandomIndex = Rand();
-				}
-
-				FTask* Result = Registry->StealItem(CachedRandomIndex, CachedPriorityIndex, GetBackGroundTasks);
-				if (Result)
-				{
-					return Result;
-				}
-			}
-			return nullptr;
-		}
-
-		inline FTask* DequeueAffinity(bool GetBackGroundTasks, bool bDisableThrottleStealing)
-		{
-			return AffinityQueue.dequeue(AffinityHazard);
-		}
-
-		uint32 GetAffinityIndex() const
-		{
-			return AffinityIndex;
-		}
-
-	private:
-		inline void EnqueueAffinity(FTask* Item)
-		{
-			check(SleepEvent != nullptr);
-			AffinityQueue.enqueue(Item);
-			
-			ESleepState SleepState = ESleepState::Sleeping;
-			ESleepState DrowsingState = ESleepState::Drowsing;
-			if (SleepEvent->SleepState.compare_exchange_strong(DrowsingState, ESleepState::Affinity, std::memory_order_acquire))
-			{
-			}
-			else if (SleepEvent->SleepState.compare_exchange_strong(SleepState, ESleepState::Drowsing, std::memory_order_acquire))
-			{
-				SleepEvent->SleepEvent->Trigger();
-			}
-		}
-
-	private:
-		static constexpr uint32	InvalidIndex = ~0u;
-		FLocalQueueType			LocalQueues[uint32(ETaskPriority::Count)];
-		DequeueHazard			DequeueHazards[uint32(ETaskPriority::Count)];
-		FOverflowQueueType		AffinityQueue;
-		DequeueHazard			AffinityHazard;
-		TLocalQueueRegistry*	Registry;
-		FSleepEvent*			SleepEvent;
-		uint32					CachedRandomIndex = InvalidIndex;
-		uint32					CachedPriorityIndex = 0;
-		uint32					AffinityIndex = ~0u;
-		ELocalQueueType			QueueType;
-	};
-
-	TLocalQueueRegistry() = default;
-
-private:
-	// add a queue to the Registry
-	uint32 AddLocalQueue(TLocalQueue* QueueToAdd, ELocalQueueType QueueType)
-	{
-		NumActiveWorkers[QueueType == ELocalQueueType::EBackground]++;
-		LocalQueues.Add(QueueToAdd);
-		return LocalQueues.Num() - 1;
-	}
-
-	// remove a queue from the Registry
-	void DeleteLocalQueue(TLocalQueue* QueueToRemove, ELocalQueueType QueueType)
-	{
-		NumActiveWorkers[QueueType == ELocalQueueType::EBackground]--;
-		verify(LocalQueues.Remove(QueueToRemove) == 1);
-	}
-
-	// StealItem tries to steal an Item from a Registered LocalQueue
-	FTask* StealItem(uint32& CachedRandomIndex, uint32& CachedPriorityIndex, bool GetBackGroundTasks)
-	{
-		uint32 NumQueues = LocalQueues.Num();
-		uint32 MaxPriority = GetBackGroundTasks ? int32(ETaskPriority::Count) : int32(ETaskPriority::ForegroundCount);
-		CachedRandomIndex = CachedRandomIndex % NumQueues;
-
-		for(uint32 i = 0; i < NumQueues; i++)
-		{
-			TLocalQueue* LocalQueue = LocalQueues[CachedRandomIndex];
-			for(uint32 PriorityIndex = 0; PriorityIndex < MaxPriority; PriorityIndex++)
-			{	
-				FTask* Item;
-				if (LocalQueue->LocalQueues[CachedPriorityIndex].Steal(Item))
-				{
-					return Item;
-				}
-				CachedPriorityIndex = ++CachedPriorityIndex < MaxPriority ? CachedPriorityIndex : 0;
-			}
-			CachedRandomIndex = ++CachedRandomIndex < NumQueues ? CachedRandomIndex : 0;
-		}
-		CachedPriorityIndex = 0;
-		CachedRandomIndex = TLocalQueue::InvalidIndex;
-		return nullptr;
-	}
-
-public:
-	// enqueue an Item directy into the Global OverflowQueue
-	// returns true if we should wake a worker for stealing
-	bool Enqueue(FTask* Item, uint32 PriorityIndex)
-	{
-		check(PriorityIndex < int32(ETaskPriority::Count));
-		check(Item != nullptr);
-
-		bool bBackgroundTask = Item->IsBackgroundTask();
-		OverflowQueues[PriorityIndex].enqueue(Item);
-
-		return LessThanHalfWorkersLookingForWork(bBackgroundTask);
-	}
-
-	inline bool EnqueueAffinity(FTask* Item, uint32 AffinityIndex)
-	{
-		if (AffinityIndex < uint32(LocalQueues.Num()) && LocalQueues[AffinityIndex]->SleepEvent)
-		{
-			LocalQueues[AffinityIndex]->EnqueueAffinity(Item);
-			return true;
-		}
-		return false;
-	}
-
-	// grab an Item directy from the Global OverflowQueue
-	FTask* DequeueGlobal(bool GetBackGroundTasks = true, bool bDisableThrottleStealing = true)
-	{
-		if (bDisableThrottleStealing || (NumActiveWorkers[GetBackGroundTasks] >= (2 * NumWorkersLookingForWork[GetBackGroundTasks].load(std::memory_order_relaxed) - 1)) || ((Rand() % 4) == 0))
-		{
-			int32 MaxPriority = GetBackGroundTasks ? int32(ETaskPriority::Count) : int32(ETaskPriority::ForegroundCount);
-			for (int32 PriorityIndex = 0; PriorityIndex < MaxPriority; PriorityIndex++)
-			{
-				FTask* Item = OverflowQueues[PriorityIndex].dequeue();
-				if (Item)
-				{
-					return Item;
-				}
-			}
-		}
-		return nullptr;
-	}
-
-	inline FTask* DequeueSteal(bool GetBackGroundTasks = true, bool bDisableThrottleStealing = true)
+	// Not thread-safe.
+	void Reset()
 	{
-		if (bDisableThrottleStealing || (NumActiveWorkers[GetBackGroundTasks] >= (2 * NumWorkersLookingForWork[GetBackGroundTasks].load(std::memory_order_relaxed) - 1)) || ((Rand() % 4) == 0))
+		uint32 NumQueues = NumLocalQueues.load(std::memory_order_relaxed);
+		for (uint32 Index = 0; Index < NumQueues; Index++)
 		{
-			uint32 CachedRandomIndex = Rand();
-			uint32 CachedPriorityIndex = 0;
-			FTask* Result = StealItem(CachedRandomIndex, CachedPriorityIndex, GetBackGroundTasks);
-			if (Result)
-			{
-				return Result;
-			}
+			LocalQueues[Index].store(0, std::memory_order_relaxed);
 		}
-		return nullptr;
-	}
 
-	inline FOutOfWork GetOutOfWorkScope(ELocalQueueType QueueType)
-	{
-		return FOutOfWork(NumWorkersLookingForWork[QueueType == ELocalQueueType::EBackground]);
+		NumLocalQueues.store(0, std::memory_order_release);
 	}
 
 private:
-	inline bool LessThanHalfWorkersLookingForWork(bool bBackgroundTask) const
-	{
-		if (bBackgroundTask)
-		{
-			return NumWorkersLookingForWork[true].load(std::memory_order_acquire) * 2 <= NumActiveWorkers[true];
-		}
-		else
-		{
-			return (NumWorkersLookingForWork[false].load(std::memory_order_acquire) + NumWorkersLookingForWork[true].load(std::memory_order_acquire)) * 2 <= (NumActiveWorkers[false] + NumActiveWorkers[true]);
-		}
-	}
-
-	FOverflowQueueType		OverflowQueues[uint32(ETaskPriority::Count)];
-	TArray<TLocalQueue*>	LocalQueues;
-	std::atomic_int			NumWorkersLookingForWork[2] = { {0}, {0} };
-	int						NumActiveWorkers[2] = { 0, 0 };
-	PRAGMA_ENABLE_DEPRECATION_WARNINGS
+	FOverflowQueueType        OverflowQueues[uint32(ETaskPriority::Count)];
+	std::atomic<TLocalQueue*> LocalQueues[MaxLocalQueues] { nullptr };
+	std::atomic<uint32>       NumLocalQueues {0};
 };
 
+} // namespace Private
+
 }
diff --git a/Engine/Source/Runtime/Core/Public/Async/Fundamental/ReserveScheduler.h b/Engine/Source/Runtime/Core/Public/Async/Fundamental/ReserveScheduler.h
--- a/Engine/Source/Runtime/Core/Public/Async/Fundamental/ReserveScheduler.h
+++ b/Engine/Source/Runtime/Core/Public/Async/Fundamental/ReserveScheduler.h
@@ -17,7 +17,7 @@
 
 namespace LowLevelTasks
 {
-	class FReserveScheduler final : public FSchedulerTls
+	class UE_DEPRECATED(5.5, "This class will be removed.") FReserveScheduler final : public FSchedulerTls
 	{
 		UE_NONCOPYABLE(FReserveScheduler);
 		static CORE_API FReserveScheduler Singleton;
@@ -47,8 +47,9 @@ namespace LowLevelTasks
 		};
 
 		TUniquePtr<FThread> CreateWorker(FThread::EForkable IsForkable = FThread::NonForkable, FYieldedWork* ReserveEvent = nullptr, EThreadPriority Priority = EThreadPriority::TPri_Normal);
-
+	PRAGMA_DISABLE_DEPRECATION_WARNINGS
 		TEventStack<FYieldedWork> 						EventStack = { ReserveEvents };
+	PRAGMA_ENABLE_DEPRECATION_WARNINGS
 		TAlignedArray<FYieldedWork>						ReserveEvents;
 		FCriticalSection 								WorkerThreadsCS;
 		TArray<TUniquePtr<FThread>>						WorkerThreads;
@@ -56,6 +57,7 @@ namespace LowLevelTasks
 		std::atomic_uint								NextWorkerId { 0 };
 	};
 
+	PRAGMA_DISABLE_DEPRECATION_WARNINGS
 	inline FReserveScheduler& FReserveScheduler::Get()
 	{
 		return Singleton;
@@ -70,4 +72,5 @@ namespace LowLevelTasks
 	{
 		StopWorkers();
 	}
+	PRAGMA_ENABLE_DEPRECATION_WARNINGS
 }
\ No newline at end of file
diff --git a/Engine/Source/Runtime/Core/Public/Async/Fundamental/Scheduler.h b/Engine/Source/Runtime/Core/Public/Async/Fundamental/Scheduler.h
--- a/Engine/Source/Runtime/Core/Public/Async/Fundamental/Scheduler.h
+++ b/Engine/Source/Runtime/Core/Public/Async/Fundamental/Scheduler.h
@@ -35,7 +35,7 @@ namespace LowLevelTasks
 	//implementation of a treiber stack
 	//(https://en.wikipedia.org/wiki/Treiber_stack)
 	template<typename NodeType>
-	class TEventStack
+	class UE_DEPRECATED(5.5, "This class will be removed.") TEventStack
 	{
 		static constexpr uint32 EVENT_INDEX_NONE = ~0u;
 
@@ -124,8 +124,6 @@ namespace LowLevelTasks
 		static thread_local FSchedulerTls* ActiveScheduler;
 		static thread_local FLocalQueueType* LocalQueue;
 		static thread_local EWorkerType WorkerType;
-		// number of busy-waiting calls in the call-stack
-		static thread_local uint32 BusyWaitingDepth;
 
 	public:
 		CORE_API bool IsWorkerThread() const;
@@ -195,15 +193,19 @@ namespace LowLevelTasks
 
 	private: 
 		[[nodiscard]] FTask* ExecuteTask(FTask* InTask);
-		TUniquePtr<FThread> CreateWorker(bool bPermitBackgroundWork = false, FThread::EForkable IsForkable = FThread::NonForkable, Private::FWaitEvent* ExternalWorkerEvent = nullptr, FSchedulerTls::FLocalQueueType* ExternalWorkerLocalQueue = nullptr, EThreadPriority Priority = EThreadPriority::TPri_Normal, uint64 InAffinity = 0);
+		TUniquePtr<FThread> CreateWorker(uint32 WorkerId, const TCHAR* Name, bool bPermitBackgroundWork = false, FThread::EForkable IsForkable = FThread::NonForkable, Private::FWaitEvent* ExternalWorkerEvent = nullptr, FSchedulerTls::FLocalQueueType* ExternalWorkerLocalQueue = nullptr, EThreadPriority Priority = EThreadPriority::TPri_Normal, uint64 InAffinity = 0);
 		void WorkerMain(Private::FWaitEvent* WorkerEvent, FSchedulerTls::FLocalQueueType* ExternalWorkerLocalQueue, uint32 WaitCycles, bool bPermitBackgroundWork);
+		void StandbyLoop(Private::FWaitEvent* WorkerEvent, FSchedulerTls::FLocalQueueType* ExternalWorkerLocalQueue, uint32 WaitCycles, bool bPermitBackgroundWork);
+		void WorkerLoop(Private::FWaitEvent* WorkerEvent, FSchedulerTls::FLocalQueueType* ExternalWorkerLocalQueue, uint32 WaitCycles, bool bPermitBackgroundWork);
 		CORE_API void LaunchInternal(FTask& Task, EQueuePreference QueuePreference, bool bWakeUpWorker);
 		CORE_API void BusyWaitInternal(const FConditional& Conditional, bool ForceAllowBackgroundWork);
 		inline bool WakeUpWorker(bool bBackgroundWorker);
-
-		template<typename QueueType, FTask* (QueueType::*DequeueFunction)(bool), bool bIsBusyWaiting>
+		CORE_API void IncrementOversubscription();
+		CORE_API void DecrementOversubscription();
+		template<typename QueueType, FTask* (QueueType::*DequeueFunction)(bool), bool bIsStandbyWorker>
 		bool TryExecuteTaskFrom(Private::FWaitEvent* WaitEvent, QueueType* Queue, Private::FOutOfWork& OutOfWork, bool bPermitBackgroundWork);
 
+		friend class FOversubscriptionScope;
 	private:
 		Private::FWaitingQueue                         WaitingQueue[2] = { WorkerEvents, WorkerEvents };
 		FSchedulerTls::FQueueRegistry                  QueueRegistry;
@@ -213,6 +215,8 @@ namespace LowLevelTasks
 		TAlignedArray<Private::FWaitEvent>             WorkerEvents;
 		std::atomic_uint                               ActiveWorkers { 0 };
 		std::atomic_uint                               NextWorkerId { 0 };
+		std::atomic<int32>                             ForegroundCreationIndex{ 0 };
+		std::atomic<int32>                             BackgroundCreationIndex{ 0 };
 		uint64                                         WorkerAffinity = 0;
 		uint64                                         BackgroundAffinity = 0;
 		EThreadPriority                                WorkerPriority = EThreadPriority::TPri_Normal;
@@ -220,6 +224,80 @@ namespace LowLevelTasks
 		std::atomic_bool                               TemporaryShutdown{ false };
 	};
 
+	namespace Private
+	{
+		class FOversubscriptionTls
+		{
+			static thread_local bool bIsOversubscriptionAllowed;
+
+			friend class FOversubscriptionAllowedScope;
+		public:
+			static bool IsOversubscriptionAllowed() { return bIsOversubscriptionAllowed; }
+		};
+
+		class FOversubscriptionAllowedScope
+		{
+			UE_NONCOPYABLE(FOversubscriptionAllowedScope);
+
+		public:
+			FOversubscriptionAllowedScope(bool bIsOversubscriptionAllowed)
+			{
+				bPreviousValue = FOversubscriptionTls::bIsOversubscriptionAllowed;
+				FOversubscriptionTls::bIsOversubscriptionAllowed = bIsOversubscriptionAllowed;
+			}
+
+			~FOversubscriptionAllowedScope()
+			{
+				FOversubscriptionTls::bIsOversubscriptionAllowed = bPreviousValue;
+			}
+		private:
+			bool bPreviousValue;
+		};
+	}
+
+	class FOversubscriptionScope
+	{
+		UE_NONCOPYABLE(FOversubscriptionScope);
+
+	public:
+		FOversubscriptionScope(bool bCondition = true)
+		{
+			if (bCondition && Private::FOversubscriptionTls::IsOversubscriptionAllowed())
+			{
+				bIncrementOversubscriptionEmitted = true;
+
+#if CPUPROFILERTRACE_ENABLED
+				if (CpuChannel)
+				{
+					static uint32 OversubscriptionTraceId = FCpuProfilerTrace::OutputEventType("Oversubscription");
+					FCpuProfilerTrace::OutputBeginEvent(OversubscriptionTraceId);
+					bCpuBeginEventEmitted = true;
+				}
+#endif
+				FScheduler::Get().IncrementOversubscription();
+			}
+		}
+
+		~FOversubscriptionScope()
+		{
+			if (bIncrementOversubscriptionEmitted)
+			{
+				FScheduler::Get().DecrementOversubscription();
+
+#if CPUPROFILERTRACE_ENABLED
+				if (bCpuBeginEventEmitted)
+				{
+					FCpuProfilerTrace::OutputEndEvent();
+					bCpuBeginEventEmitted = false;
+				}
+#endif
+			}
+		}
+	private:
+		bool bIncrementOversubscriptionEmitted = false;
+		bool bCpuBeginEventEmitted = false;
+	};
+
 	FORCEINLINE_DEBUGGABLE bool TryLaunch(FTask& Task, EQueuePreference QueuePreference = EQueuePreference::DefaultPreference, bool bWakeUpWorker = true)
 	{
 		return FScheduler::Get().TryLaunch(Task, QueuePreference, bWakeUpWorker);
diff --git a/Engine/Source/Runtime/Core/Public/Async/Fundamental/WaitingQueue.h b/Engine/Source/Runtime/Core/Public/Async/Fundamental/WaitingQueue.h
--- a/Engine/Source/Runtime/Core/Public/Async/Fundamental/WaitingQueue.h
+++ b/Engine/Source/Runtime/Core/Public/Async/Fundamental/WaitingQueue.h
@@ -32,67 +32,73 @@ namespace LowLevelTasks::Private
 	*/
 	struct alignas(64) FWaitEvent
 	{
-		std::atomic<uint64_t>     Next{ 0 };
-		uint64_t                  Epoch{ 0 };
-		std::atomic<EWaitState>   State{ EWaitState::NotSignaled };
-		FEventRef                 Event{ EEventMode::ManualReset };
+		std::atomic<uint64>     Next{ 0 };
+		uint64                  Epoch{ 0 };
+		std::atomic<EWaitState> State{ EWaitState::NotSignaled };
+		FEventRef               Event{ EEventMode::ManualReset };
+		bool                    bIsStandby{ false };
 	};
 
 	class FWaitingQueue
 	{
-		// State_ layout:
-		// - low kWaiterBits is a stack of waiters committed wait
-		//   (indexes in NodesArray are used as stack elements,
-		//   kStackMask means empty stack).
-		// - next kWaiterBits is count of waiters in prewait state.
-		// - next kWaiterBits is count of pending signals.
-		// - remaining bits are ABA counter for the stack.
-		//   (stored in Waiter node and incremented on push).
-		static constexpr uint64_t WaiterBits = 14;
-		static constexpr uint64_t StackMask = (1ull << WaiterBits) - 1;
-		static constexpr uint64_t WaiterShift = WaiterBits;
-		static constexpr uint64_t WaiterMask = ((1ull << WaiterBits) - 1) << WaiterShift;
-		static constexpr uint64_t WaiterInc = 1ull << WaiterShift;
-		static constexpr uint64_t SignalShift = 2 * WaiterBits;
-		static constexpr uint64_t SignalMask = ((1ull << WaiterBits) - 1) << SignalShift;
-		static constexpr uint64_t SignalInc = 1ull << SignalShift;
-		static constexpr uint64_t EpochShift = 3 * WaiterBits;
-		static constexpr uint64_t EpochBits = 64 - EpochShift;
-		static constexpr uint64_t EpochMask = ((1ull << EpochBits) - 1) << EpochShift;
-		static constexpr uint64_t EpochInc = 1ull << EpochShift;
-
-		std::atomic<uint64_t>      State{ StackMask };
+		uint32                     ThreadCount{ 0 };    // Normal amount of threads when there is no oversubscription.
+		uint32                     MaxThreadCount{ 0 }; // Max limit that can be reached during oversubscription period.
+		TFunction<void()>          CreateThread;
+		std::atomic<uint32>        Oversubscription{ 0 };
+		std::atomic<uint64>        State;
+		std::atomic<uint64>        StandbyState;
 		TAlignedArray<FWaitEvent>& NodesArray;
+		std::atomic<bool>          bIsShuttingDown{ false };
 
 	public:
-		FWaitingQueue(TAlignedArray<FWaitEvent>& NodesArray)
-			: NodesArray(NodesArray)
+		FWaitingQueue(TAlignedArray<FWaitEvent>& InNodesArray)
+			: NodesArray(InNodesArray)
 		{
 		}
 
-		CORE_API void Init();
-		CORE_API void Shutdown();
+		CORE_API void Init(uint32 InThreadCount, uint32 InMaxThreadCount, TFunction<void()> InCreateThread, uint32 InActiveThreadCount);
+		CORE_API void StartShutdown();
+		CORE_API void FinishShutdown();
+
+		// First step to execute when no more work is found in the queues.
+		CORE_API void PrepareStandby(FWaitEvent* Node);
+		// Second step to execute when no more work is found in the queues.
+		CORE_API bool CommitStandby(FWaitEvent* Node, FOutOfWork& OutOfWork);
+
+		// Immediately goes to sleep if oversubscription period is finished and we're over the allowed thread count.
+		CORE_API void ConditionalStandby(FWaitEvent* Node);
+
+		// First step run by normal workers when no more work is found in the queues.
 		CORE_API void PrepareWait(FWaitEvent* Node);
+		// Second step run by normal workers when no more work is found in the queues.
 		CORE_API bool CommitWait(FWaitEvent* Node, FOutOfWork& OutOfWork, int32 SpinCycles, int32 WaitCycles);
 
-		// returns true if we need to wake up a new worker
+		// Step to run by normal workers if they detect new work after they called prepare wait.
+		// Returns true if we need to wake up a new worker.
 		CORE_API bool CancelWait(FWaitEvent* Node);
 
+		// Increment oversubscription and notify a thread if we're under the allowed thread count.
+		// If dynamic thread creation is allowed, this could spawn a new thread if needed.
+		CORE_API void IncrementOversubscription();
+
+		// Decrement oversubscription only, any active threads will finish their current task and will
+		// go to sleep if conditional standby determines we're now over the active thread count.
+		CORE_API void DecrementOversubscription();
+
+		// Try to wake up the amount of workers passed in the parameters.
+		// Return the number that were woken up.
 		int32 Notify(int32 Count = 1)
 		{
 			return NotifyInternal(Count);
 		}
 
-		int32 NotifyAll()
-		{
-			return NotifyInternal(MAX_int32);
-		}
-
 	private:
+		CORE_API bool  TryStartNewThread();
 		CORE_API int32 NotifyInternal(int32 Count);
 		CORE_API void  Park(FWaitEvent* Node, FOutOfWork& OutOfWork, int32 SpinCycles, int32 WaitCycles);
 		CORE_API int32 Unpark(FWaitEvent* InNode);
-		CORE_API void  CheckState(uint64_t state, bool waiter = false);
+		CORE_API void  CheckState(uint64 State, bool bIsWaiter = false);
+		CORE_API void  CheckStandbyState(uint64 State);
 	};
 
 } // namespace LowLevelTasks::Private
\ No newline at end of file
diff --git a/Engine/Source/Runtime/Core/Public/Experimental/Async/AwaitableTask.h b/Engine/Source/Runtime/Core/Public/Experimental/Async/AwaitableTask.h
--- a/Engine/Source/Runtime/Core/Public/Experimental/Async/AwaitableTask.h
+++ b/Engine/Source/Runtime/Core/Public/Experimental/Async/AwaitableTask.h
@@ -3,6 +3,7 @@
 #pragma once
 #include "Async/Fundamental/Task.h"
 #include "Async/Fundamental/Scheduler.h"
+#include "Async/ManualResetEvent.h"
 #include "Misc/Launder.h"
 
 template<typename>
@@ -211,15 +212,15 @@ namespace AwaitableTask_Detail
 
 	private:
 		CallableType Callable;
-		LowLevelTasks::FTask Task;	
+		LowLevelTasks::FTask Task;
 		ReturnType ReturnValue;
 		std::atomic_int ReferenceCounter{ 2 }; //starts at 2 because at construction the AwaitableTask will hold a Reference and the LowLevelTasks::FTask does as well.
-		std::atomic_bool Completed { false };
+		UE::FManualResetEvent Completed;
 
 		inline void Execute()
 		{
 			ReturnValue = Invoke(Callable);
-			Completed.store(true, std::memory_order_release);
+			Completed.Notify();
 		}
 
 		TPromise(const TPromise&) = delete;
@@ -252,7 +253,7 @@ namespace AwaitableTask_Detail
 
 		inline bool IsCompleted() const
 		{
-			return Completed.load(std::memory_order_acquire);
+			return Completed.IsNotified();
 		}
 
 		inline void Finish()
@@ -273,7 +274,7 @@ namespace AwaitableTask_Detail
 			}
 			else
 			{
-				LowLevelTasks::BusyWaitUntil([this](){ return IsCompleted(); });
+				Completed.Wait();
 				return ReturnValue;
 			}
 		}
@@ -292,12 +293,12 @@ namespace AwaitableTask_Detail
 		CallableType Callable;
 		LowLevelTasks::FTask Task;
 		std::atomic_int ReferenceCounter{ 2 }; //starts at 2 because at construction the AwaitableTask will hold a Reference and the LowLevelTasks::FTask does as well.
-		std::atomic_bool Completed { false };
+		UE::FManualResetEvent Completed;
 
 		inline void Execute()
 		{
 			Invoke(Callable);
-			Completed.store(true, std::memory_order_release);
+			Completed.Notify();
 		}
 
 		TPromise(const TPromise&) = delete;
@@ -330,7 +331,7 @@ namespace AwaitableTask_Detail
 
 		inline bool IsCompleted() const
 		{
-			return Completed.load(std::memory_order_acquire);
+			return Completed.IsNotified();
 		}
 
 		inline void Finish()
@@ -355,7 +356,7 @@ namespace AwaitableTask_Detail
 			}
 			else
 			{
-				LowLevelTasks::BusyWaitUntil([this](){ return IsCompleted(); });
+				Completed.Wait();
 			}
 		}
 	};
diff --git a/Engine/Source/Runtime/Core/Public/Tasks/TaskPrivate.h b/Engine/Source/Runtime/Core/Public/Tasks/TaskPrivate.h
--- a/Engine/Source/Runtime/Core/Public/Tasks/TaskPrivate.h
+++ b/Engine/Source/Runtime/Core/Public/Tasks/TaskPrivate.h
@@ -2,6 +2,7 @@
 
 #pragma once
 
+#include "Async/EventCount.h"
 #include "Async/Fundamental/Scheduler.h"
 #include "Async/Fundamental/Task.h"
 #include "Async/Mutex.h"
@@ -436,14 +437,7 @@ public:
 			// waits until the task is completed or waiting timed out, while executing other tasks
 			bool BusyWait(FTimeout Timeout)
 			{
-				TaskTrace::FWaitingScope WaitingScope(GetTraceId());
-				TRACE_CPUPROFILER_EVENT_SCOPE(Tasks::BusyWait);
-
-				// ignore the result as we still have to make sure the task is completed upon returning from this function call
-				TryRetractAndExecute(Timeout);
-
-				LowLevelTasks::BusyWaitUntil([this, Timeout] { return IsCompleted() || Timeout; });
-				return IsCompleted();
+				return Wait(Timeout);
 			}
 
 			// waits until the task is completed or the condition returns true, while executing other tasks
diff --git a/Engine/Source/Runtime/Core/Tests/Async/TaskGraphTest.cpp b/Engine/Source/Runtime/Core/Tests/Async/TaskGraphTest.cpp
--- a/Engine/Source/Runtime/Core/Tests/Async/TaskGraphTest.cpp
+++ b/Engine/Source/Runtime/Core/Tests/Async/TaskGraphTest.cpp
@@ -1422,6 +1422,56 @@ PRAGMA_ENABLE_DEPRECATION_WARNINGS
 		}
 	}
 
+	// This test proves that retraction alone is not sufficient to avoid deadlocks and oversubscription is also required as a last resort.
+	// This test is a repro for a virtual texture system scheduled on foreground thread, sending and waiting on an IO thread request that itself needs to 
+	// do some processing on the foreground thread before releasing the IO event, causing a deadlock when only a single worker is present.
+	TEST_CASE_NAMED(FTaskGraphConstrainedRetraction, "System::Core::Async::TaskGraph::ConstrainedRetraction", "[.][ApplicationContextMask][EngineFilter]")
+	{
+		TRACE_CPUPROFILER_EVENT_SCOPE(FTaskGraphConstrainedRetraction);
+
+		// A lot easier to prove when we have a minimal amount of workers
+		LowLevelTasks::FScheduler::Get().RestartWorkers(1, 0);
+
+		TFunction<void(FGraphEventRef Task)> WaitMethods[2] =
+		{
+			[](FGraphEventRef Task) { TRACE_CPUPROFILER_EVENT_SCOPE(Task_Wait); TRACE_CPUPROFILER_EVENT_FLUSH(); Task->Wait(); },
+			[](FGraphEventRef Task) { TRACE_CPUPROFILER_EVENT_SCOPE(Task_WaitUntilTaskCompletes); TRACE_CPUPROFILER_EVENT_FLUSH(); FTaskGraphInterface::Get().WaitUntilTaskCompletes(Task); }
+		};
+
+		for (int32 WaitMethodIndex = 0; WaitMethodIndex < UE_ARRAY_COUNT(WaitMethods); ++WaitMethodIndex)
+		{
+			// This is going to be a dependency on something outside the taskgraph system
+			// Could be another system entirely like the IO thread pool, or even a lowlevel task without prerequisites/subsequents knowledge.
+			FGraphEventRef ExternalEvent = FGraphEvent::CreateGraphEvent();
+
+			// We need the worker thread to be busy with our root task while we schedule the external one
+			FGraphEventRef RootTask = 
+				FFunctionGraphTask::CreateAndDispatchWhenReady(
+					[&]
+					{
+						FFunctionGraphTask::CreateAndDispatchWhenReady(
+							[&]
+							{
+								TRACE_CPUPROFILER_EVENT_SCOPE(DispatchSubsequents);
+								TRACE_CPUPROFILER_EVENT_FLUSH();
+								ExternalEvent->DispatchSubsequents();
+
+							}, TStatId{}
+						);
+						WaitMethods[WaitMethodIndex](ExternalEvent);
+
+					}, TStatId{}
+				);
+
+			// This wait should try to perform retraction, but it can't see past the external task prerequisites
+			// Since we launch a task on a single worker system and then wait, retraction alone without oversubscription
+			// will hold the worker forever and will cause a deadlock.
+			WaitMethods[WaitMethodIndex](RootTask);
+		}
+
+		LowLevelTasks::FScheduler::Get().RestartWorkers();
+	}
+
 	TEST_CASE_NAMED(FTaskGraphAnyTask, "System::Core::Async::TaskGraph::AnyTask", "[.][ApplicationContextMask][EngineFilter]")
 	{
 		{	// blocks if none of tasks is completed
@@ -1457,4 +1507,4 @@ PRAGMA_ENABLE_DEPRECATION_WARNINGS
 	}
 }
 
-#endif //WITH_TESTS
+#endif //WITH_TESTS
\ No newline at end of file
diff --git a/Engine/Source/Runtime/Core/Tests/Tasks/TasksTest.cpp b/Engine/Source/Runtime/Core/Tests/Tasks/TasksTest.cpp
--- a/Engine/Source/Runtime/Core/Tests/Tasks/TasksTest.cpp
+++ b/Engine/Source/Runtime/Core/Tests/Tasks/TasksTest.cpp
@@ -1665,160 +1665,6 @@ namespace UE { namespace TasksTests
 			Benchmark<5>(*Name, [&Index]() { BenchmarkBlockUnblockWorkers(Index); });
 		}
 	}
-
-	TEST_CASE_NAMED(FTasksDoNotRunInsideBusyWait, "System::Core::Async::Tasks::DoNotRunInsideBusyWait", "[.][ApplicationContextMask][EngineFilter]")
-	{
-		using namespace LowLevelTasks;
-
-		FPlatformProcess::Sleep(0.1f); // give workers time to fall asleep, to avoid any reserve worker messing around
-
-		uint32 NumWorkers = LowLevelTasks::FScheduler::Get().GetNumWorkers();
-
-		// Block all workers to make sure the tasks we're going to queue are not executed before we enter our busy wait loop.
-		FTaskEvent ResumeEvent{ UE_SOURCE_LOCATION };
-		TArray<LowLevelTasks::FTask> WorkerBlockers = BlockWorkers(ResumeEvent, NumWorkers);
-
-		std::atomic<bool>     CanExecute { false };
-		std::atomic<int32>    NumExecuted { 0 };
-		UE::FManualResetEvent Done;
-
-		// Queue enough busy wait excluded tasks to verify that the scheduler
-		// will do the right thing and not end up actually executing one of them
-		// from inside busy wait.
-		for (int32 Index = 0; Index < 100; ++Index)
-		{
-			Launch(UE_SOURCE_LOCATION,
-				[&CanExecute, &NumExecuted, &Done]()
-				{
-					verify(CanExecute.load());
-					if (++NumExecuted == 100)
-					{
-						Done.Notify();
-					}
-				},
-				UE::Tasks::ETaskPriority::Default,
-				UE::Tasks::EExtendedTaskPriority::None,
-				UE::Tasks::ETaskFlags::DoNotRunInsideBusyWait // this should not be picked up by busy waiting
-			);
-		}
-
-		// Busy wait for a second, making sure no excluded tasks are executed.
-		double StartTime = FPlatformTime::Seconds() + 1;
-		BusyWaitUntil([&StartTime]() { return FPlatformTime::Seconds() > StartTime; });
-
-		// Allow execution now.
-		CanExecute = true;
-
-		ResumeEvent.Trigger();
-
-		// Now busy waiting will run since we just unblocked all workers.
-		verify(Done.WaitFor(UE::FMonotonicTimeSpan::FromSeconds(1)));
-
-		// Do not exit the test before all blocked workers' tasks are done.
-		LowLevelTasks::BusyWaitForTasks<LowLevelTasks::FTask>(WorkerBlockers);
-	}
-
-	// Test that validates each tasks have their own execution resource stack and that busy wait cannot interfere with other tasks' resources.
-	TEST_CASE_NAMED(FTasksCaptureExecutionResourceInsideBusyWait, "System::Core::Async::Tasks::CaptureExecutionResourceInsideBusyWait", "[.][ApplicationContextMask][EngineFilter]")
-	{
-		using namespace LowLevelTasks;
-		
-		FPlatformProcess::Sleep(0.1f); // give workers time to fall asleep, to avoid any reserve worker messing around
-
-		uint32 NumWorkers = LowLevelTasks::FScheduler::Get().GetNumWorkers();
-
-		// Block all workers but one so we can prevent the busywait task from running outside the busywait
-		FTaskEvent ResumeEvent{ UE_SOURCE_LOCATION };
-		TArray<LowLevelTasks::FTask> WorkerBlockers = BlockWorkers(ResumeEvent, NumWorkers - 1);
-
-		std::atomic<bool>     NormalResourceHeld{ false };
-		std::atomic<bool>     BusyWaitResourceHeld{ false };
-
-		class FResource : public FThreadSafeRefCountedObject, public IExecutionResource
-		{
-			std::atomic<bool>& ResourceHeld;
-		public:
-			FResource(std::atomic<bool>& ResourceToHold)
-				: ResourceHeld(ResourceToHold)
-			{
-				ResourceHeld = true;
-			}
-			virtual ~FResource() override
-			{
-				ResourceHeld = false;
-			}
-			uint32 AddRef() const override
-			{
-				return FThreadSafeRefCountedObject::AddRef();
-			}
-
-			uint32 Release() const override
-			{
-				return FThreadSafeRefCountedObject::Release();
-			}
-
-			uint32 GetRefCount() const override
-			{
-				return FThreadSafeRefCountedObject::GetRefCount();
-			}
-		};
-
-		TRefCountPtr<IExecutionResource> NormalResource;
-		TRefCountPtr<IExecutionResource> BusyWaitResource;
-
-		FTask NormalTask =
-			Launch(TEXT("NormalTask"),
-				[&]()
-				{
-					verify(!LowLevelTasks::FScheduler::Get().IsBusyWaiting());
-
-					FExecutionResourceContextScope ExecutionResourceScope(new FResource(NormalResourceHeld));
-
-					NormalResource = FExecutionResourceContext::Get();
-
-					BusyWaitUntil([&]() { return BusyWaitResourceHeld.load(); });
-				}
-			);
-
-		FTask BusyWaitTask =
-			Launch(TEXT("BusyWaitTask"),
-				[&]()
-				{
-					verify(LowLevelTasks::FScheduler::Get().IsBusyWaiting());
-
-					FExecutionResourceContextScope ExecutionResourceScope(new FResource(BusyWaitResourceHeld));
-
-					BusyWaitResource = FExecutionResourceContext::Get();
-				}
-		);
-
-		NormalTask.Wait();
-		
-		// This should have been completed as part of the normal task
-		verify(BusyWaitTask.IsCompleted());
-
-		// Verify now that both resource are held
-		verify(NormalResource.IsValid());
-		verify(BusyWaitResource.IsValid());
-
-		verify(NormalResourceHeld.load());
-		verify(BusyWaitResourceHeld.load());
-
-		// Now confirm that we can release the normal resource and its not also being held by the busy wait resource
-		NormalResource.SafeRelease();
-		verify(!NormalResourceHeld.load());
-		verify(BusyWaitResourceHeld.load());
-
-		// Now verify that we can release the busy wait resource too.
-		BusyWaitResource.SafeRelease();
-		verify(!BusyWaitResourceHeld.load());
-
-		// Unblock all workers
-		ResumeEvent.Trigger();
-
-		// Do not exit the test before all blocked workers' tasks are done.
-		LowLevelTasks::BusyWaitForTasks<LowLevelTasks::FTask>(WorkerBlockers);
-	}
 }}
 
 #endif // WITH_TESTS
