Use RemoveAt instead of RemoveAtSwap to keep order in GetDestructionSentinelStackTls() contents

diff --git a/Engine/Source/Runtime/Core/Public/Misc/MTAccessDetector.h b/Engine/Source/Runtime/Core/Public/Misc/MTAccessDetector.h
--- a/Engine/Source/Runtime/Core/Public/Misc/MTAccessDetector.h
+++ b/Engine/Source/Runtime/Core/Public/Misc/MTAccessDetector.h
@@ -494,7 +494,7 @@ public:
 					ExpectedState.SetWriterThreadId(FPlatformTLS::GetCurrentThreadId());
 				}
 
-				GetDestructionSentinelStackTls().RemoveAtSwap(DestructionSentinelIndex);
+				GetDestructionSentinelStackTls().RemoveAt(DestructionSentinelIndex, 1, EAllowShrinking::No);
 				--DestructionSentinelIndex;
 			}
 		}
@@ -554,7 +554,7 @@ public:
 		checkfSlow(GetDestructionSentinelStackTls().Num() != 0, TEXT("An attempt to remove a not registered destruction sentinel"));
 		checkfSlow(GetDestructionSentinelStackTls().Last() == &DestructionSentinel, TEXT("Mismatched destruction sentinel: %p != %p"), GetDestructionSentinelStackTls().Last(), &DestructionSentinel);
 
-		GetDestructionSentinelStackTls().RemoveAtSwap(GetDestructionSentinelStackTls().Num() - 1);
+		GetDestructionSentinelStackTls().RemoveAt(GetDestructionSentinelStackTls().Num() - 1, 1, EAllowShrinking::No);
 	}
 
 	FORCEINLINE void AcquireWriteAccess()
@@ -641,7 +641,7 @@ public:
 		checkfSlow(GetDestructionSentinelStackTls().Num() != 0, TEXT("An attempt to remove a not registered destruction sentinel"));
 		checkfSlow(GetDestructionSentinelStackTls().Last() == &DestructionSentinel, TEXT("Mismatched destruction sentinel: %p != %p"), GetDestructionSentinelStackTls().Last(), &DestructionSentinel);
 
-		GetDestructionSentinelStackTls().RemoveAtSwap(GetDestructionSentinelStackTls().Num() - 1);
+		GetDestructionSentinelStackTls().RemoveAt(GetDestructionSentinelStackTls().Num() - 1, 1, EAllowShrinking::No);
 	}
 };
 
