Bump MATERIALSHADERMAP_DERIVEDDATA_VER to fix stale ddc asset causing crash due to struct layout change in FMaterialCompilationOutput

diff --git a/Engine/Source/Runtime/Core/Private/UObject/DevObjectVersion.cpp b/Engine/Source/Runtime/Core/Private/UObject/DevObjectVersion.cpp
--- a/Engine/Source/Runtime/Core/Private/UObject/DevObjectVersion.cpp
+++ b/Engine/Source/Runtime/Core/Private/UObject/DevObjectVersion.cpp
@@ -128,7 +128,7 @@ FDevSystemGuids::FDevSystemGuids()
 	, GROOM_BINDING_DERIVED_DATA_VERSION(0x30769E53, 0x0C574C7B, 0xA15C56F2, 0x24A64E32)
 	, GROOM_DERIVED_DATA_VERSION(0x05A37379, 0xF8A049D5, 0x986824E9, 0xAAA83F41)
 	, LANDSCAPE_MOBILE_COOK_VERSION(0x0E9ADF72, 0xD6B64E0D, 0x81C4A92B, 0x081A37AB)
-	, MATERIALSHADERMAP_DERIVEDDATA_VER(0x9131A169, 0x3E004B0D, 0x864390EF, 0xA1F12934)
+	, MATERIALSHADERMAP_DERIVEDDATA_VER(0xD31091F1, 0x46F0456F, 0x9DA512EF, 0xA47EFB19)
 	, NANITE_DERIVEDDATA_VER(0xBEB0226A, 0x070E4ECA, 0x972C1E7D, 0xD8599E68)
 	, NIAGARASHADERMAP_DERIVEDDATA_VER(0x7BBD9913, 0xC1554D20, 0xADAE9F17, 0xB006299E)
 	, Niagara_LatestScriptCompileVersion(0x6D32B8EE, 0x909FCA7E, 0xA5CE4F17, 0x066A5F25)
diff --git a/Engine/Source/Runtime/Core/Private/UObject/FortniteMainBranchObjectVersion.cpp b/Engine/Source/Runtime/Core/Private/UObject/FortniteMainBranchObjectVersion.cpp
--- a/Engine/Source/Runtime/Core/Private/UObject/FortniteMainBranchObjectVersion.cpp
+++ b/Engine/Source/Runtime/Core/Private/UObject/FortniteMainBranchObjectVersion.cpp
@@ -8,7 +8,7 @@ TMap<FGuid, FGuid> FFortniteMainBranchObjectVersion::GetSystemGuids()
 
 	SystemGuids.Add(DevGuids.GLOBALSHADERMAP_DERIVEDDATA_VER, FGuid("8D3A12924CDF4D9F84C0A900E9445CAD"));
 	SystemGuids.Add(DevGuids.LANDSCAPE_MOBILE_COOK_VERSION, FGuid("32D02EF867C74B71A0D4E0FA41392732"));
-	SystemGuids.Add(DevGuids.MATERIALSHADERMAP_DERIVEDDATA_VER, FGuid("CE7776E63E9E43FBB8C103EE1AB40B5D"));
+	SystemGuids.Add(DevGuids.MATERIALSHADERMAP_DERIVEDDATA_VER, FGuid("E416BA6E4D48491994BB91E4B06E43AC"));
 	SystemGuids.Add(DevGuids.NANITE_DERIVEDDATA_VER, FGuid("D28C003F825F4C178F698938FB0E47F7"));
 	SystemGuids.Add(DevGuids.NIAGARASHADERMAP_DERIVEDDATA_VER, FGuid("C417A50271B3427D98785388266E3B66"));
 	SystemGuids.Add(DevGuids.Niagara_LatestScriptCompileVersion, FGuid("48A35E7D14134FBE9DF89B3E485E5762"));
