Added support for TUniquePtr<void> and (restored) support for TSharedPtr<void>.

diff --git a/Engine/Source/Runtime/Core/Public/Templates/SharedPointer.h b/Engine/Source/Runtime/Core/Public/Templates/SharedPointer.h
--- a/Engine/Source/Runtime/Core/Public/Templates/SharedPointer.h
+++ b/Engine/Source/Runtime/Core/Public/Templates/SharedPointer.h
@@ -161,6 +161,8 @@ template< class ObjectType, ESPMode InMode >
 class TSharedRef
 {
 public:
+	static_assert(!std::is_void_v<ObjectType>, "TSharedRef<void> is not supported - use TSharedPtr instead");
+
 	using ElementType = ObjectType;
 	static constexpr ESPMode Mode = InMode;
 
@@ -1097,7 +1099,11 @@ public:
 	 *
 	 * @return  Reference to the object
 	 */
-	[[nodiscard]] FORCEINLINE decltype(auto) operator*() const
+	template <
+		typename DummyObjectType = ObjectType,
+		decltype(*(DummyObjectType*)nullptr, 0) = 0 // this construct means that operator* is only considered for overload resolution if T is dereferenceable
+	>
+	[[nodiscard]] FORCEINLINE DummyObjectType& operator*() const
 	{
 		check( IsValid() );
 		return *Object;
diff --git a/Engine/Source/Runtime/Core/Public/Templates/UniquePtr.h b/Engine/Source/Runtime/Core/Public/Templates/UniquePtr.h
--- a/Engine/Source/Runtime/Core/Public/Templates/UniquePtr.h
+++ b/Engine/Source/Runtime/Core/Public/Templates/UniquePtr.h
@@ -8,6 +8,7 @@
 #include "Templates/RemoveExtent.h"
 #include "Templates/Requires.h"
 #include "Serialization/MemoryLayout.h"
+#include <type_traits>
 
 // Single-ownership smart pointer in the vein of std::unique_ptr.
 // Use this when you need an object's lifetime to be strictly bound to the lifetime of a single smart pointer.
@@ -302,7 +303,11 @@ public:
 	 *
 	 * @return A reference to the object owned by the TUniquePtr.
 	 */
-	[[nodiscard]] FORCEINLINE T& operator*() const
+	template <
+		typename DummyT = T,
+		decltype(*(DummyT*)nullptr, 0) = 0 // this construct means that operator* is only considered for overload resolution if T is dereferenceable
+	>
+	[[nodiscard]] FORCEINLINE DummyT& operator*() const
 	{
 		return *Ptr;
 	}
