TaskGraphTest - Store return value of FCrc::MemCrc32() to silence error from the soon to be added [[nodiscard]] attribute on that function

diff --git a/Engine/Source/Runtime/Core/Tests/Async/TaskGraphTest.cpp b/Engine/Source/Runtime/Core/Tests/Async/TaskGraphTest.cpp
--- a/Engine/Source/Runtime/Core/Tests/Async/TaskGraphTest.cpp
+++ b/Engine/Source/Runtime/Core/Tests/Async/TaskGraphTest.cpp
@@ -664,7 +664,7 @@ namespace OldTaskGraphTests
 				uint32 RunningCrc = 0;
 				for (int32 Index = 0; Index < 1000000; Index++)
 				{
-					FCrc::MemCrc32(CompletionEvent.GetReference(), sizeof(FGraphEvent), RunningCrc);
+					RunningCrc = FCrc::MemCrc32(CompletionEvent.GetReference(), sizeof(FGraphEvent), RunningCrc);
 				}
 				uint64 StartTime = FPlatformTime::Cycles64();
 				FFunctionGraphTask::CreateAndDispatchWhenReady([StartTime, &ForegroundTask] { ForegroundTask(StartTime); }, TStatId{}, nullptr, ENamedThreads::AnyHiPriThreadHiPriTask);
