Store ITextData as a TRefCountPtr

This shrinks FText by 8-bytes, and allows the ref-count to pack directly inside FTextHistory without increasing its size (saving another 8-bytes vs TReferenceControllerBase)

diff --git a/Engine/Source/Runtime/Core/Private/Internationalization/Text.cpp b/Engine/Source/Runtime/Core/Private/Internationalization/Text.cpp
--- a/Engine/Source/Runtime/Core/Private/Internationalization/Text.cpp
+++ b/Engine/Source/Runtime/Core/Private/Internationalization/Text.cpp
@@ -116,7 +116,7 @@ bool FTextInspector::GetHistoricNumericData(const FText& Text, FHistoricTextNume
 
 const void* FTextInspector::GetSharedDataId(const FText& Text)
 {
-	return &Text.TextData.Get();
+	return Text.TextData.GetReference();
 }
 
 // These default values have been duplicated to the KismetTextLibrary functions for Blueprints. Please replicate any changes there!
@@ -244,42 +244,53 @@ FText::FText()
 {
 }
 
-FText::FText( EInitToEmptyString )
-	: TextData(MakeShared<FTextHistory_Base, ESPMode::ThreadSafe>())
-	, Flags(0)
-{
-}
-
 const FText& FText::GetEmpty()
 {
-	static const FText StaticEmptyText = FText(FText::EInitToEmptyString::Value);
+	static const FText StaticEmptyText = FText(MakeRefCount<FTextHistory_Base>());
 	return StaticEmptyText;
 }
 
-FText::FText( TSharedRef<ITextData, ESPMode::ThreadSafe> InTextData )
-	: TextData(MoveTemp(InTextData))
-	, Flags(0)
-{
-}
-
 FText::FText( FString&& InSourceString )
-	: TextData(MakeShared<FTextHistory_Base, ESPMode::ThreadSafe>(FTextId(), MoveTemp(InSourceString)))
+	: TextData(MakeRefCount<FTextHistory_Base>(FTextId(), MoveTemp(InSourceString)))
 	, Flags(0)
 {
 }
 
 FText::FText( FName InTableId, FString InKey, const EStringTableLoadingPolicy InLoadingPolicy )
-	: TextData(MakeShared<FTextHistory_StringTableEntry, ESPMode::ThreadSafe>(InTableId, MoveTemp(InKey), InLoadingPolicy))
+	: TextData(MakeRefCount<FTextHistory_StringTableEntry>(InTableId, MoveTemp(InKey), InLoadingPolicy))
 	, Flags(0)
 {
 }
 
 FText::FText( FString&& InSourceString, const FTextKey& InNamespace, const FTextKey& InKey, uint32 InFlags )
-	: TextData(MakeShared<FTextHistory_Base, ESPMode::ThreadSafe>(FTextId(InNamespace, InKey), MoveTemp(InSourceString)))
+	: TextData(MakeRefCount<FTextHistory_Base>(FTextId(InNamespace, InKey), MoveTemp(InSourceString)))
 	, Flags(InFlags)
 {
 }
 
+FText::FText(FText&& Other)
+	: TextData(MoveTemp(Other.TextData))
+	, Flags(Other.Flags)
+{
+	// TextData must always point to something valid
+	Other.TextData = FText::GetEmpty().TextData;
+	Other.Flags = 0;
+}
+
+FText& FText::operator=(FText&& Other)
+{
+	if (this != &Other)
+	{
+		TextData = MoveTemp(Other.TextData);
+		Flags = Other.Flags;
+
+		// TextData must always point to something valid
+		Other.TextData = FText::GetEmpty().TextData;
+		Other.Flags = 0;
+	}
+	return *this;
+}
+
 bool FText::IsEmpty() const
 {
 	return ToString().IsEmpty();
@@ -308,7 +319,7 @@ FText FText::ToLower() const
 {
 	FString ResultString = FTextTransformer::ToLower(ToString());
 
-	FText Result = FText(MakeShared<FTextHistory_Transform, ESPMode::ThreadSafe>(MoveTemp(ResultString), *this, FTextHistory_Transform::ETransformType::ToLower));
+	FText Result = FText(MakeRefCount<FTextHistory_Transform>(MoveTemp(ResultString), *this, FTextHistory_Transform::ETransformType::ToLower));
 	if (!GIsEditor)
 	{
 		Result.Flags |= ETextFlag::Transient;
@@ -320,7 +331,7 @@ FText FText::ToUpper() const
 {
 	FString ResultString = FTextTransformer::ToUpper(ToString());
 
-	FText Result = FText(MakeShared<FTextHistory_Transform, ESPMode::ThreadSafe>(MoveTemp(ResultString), *this, FTextHistory_Transform::ETransformType::ToUpper));
+	FText Result = FText(MakeRefCount<FTextHistory_Transform>(MoveTemp(ResultString), *this, FTextHistory_Transform::ETransformType::ToUpper));
 	if (!GIsEditor)
 	{
 		Result.Flags |= ETextFlag::Transient;
@@ -508,7 +519,7 @@ FText FText::FromTextGenerator(const TSharedRef<ITextGenerator>& TextGenerator)
 {
 	FString ResultString = TextGenerator->BuildLocalizedDisplayString();
 
-	FText Result = FText(MakeShared<FTextHistory_TextGenerator, ESPMode::ThreadSafe>(MoveTemp(ResultString), TextGenerator));
+	FText Result = FText(MakeRefCount<FTextHistory_TextGenerator>(MoveTemp(ResultString), TextGenerator));
 	if (!GIsEditor)
 	{
 		Result.Flags |= ETextFlag::Transient;
@@ -564,7 +575,7 @@ FText FText::AsNumberTemplate(T1 Val, const FNumberFormattingOptions* const Opti
 	const FNumberFormattingOptions& FormattingOptions = (Options) ? *Options : FormattingRules.CultureDefaultFormattingOptions;
 	FString NativeString = FastDecimalFormat::NumberToString(Val, FormattingRules, FormattingOptions);
 
-	FText Result = FText(MakeShared<FTextHistory_AsNumber, ESPMode::ThreadSafe>(MoveTemp(NativeString), Val, Options, TargetCulture));
+	FText Result = FText(MakeRefCount<FTextHistory_AsNumber>(MoveTemp(NativeString), Val, Options, TargetCulture));
 	if (!GIsEditor)
 	{
 		Result.Flags |= ETextFlag::Transient;
@@ -601,7 +612,7 @@ FText FText::AsCurrencyTemplate(T1 Val, const FString& CurrencyCode, const FNumb
 	const FNumberFormattingOptions& FormattingOptions = (Options) ? *Options : FormattingRules.CultureDefaultFormattingOptions;
 	FString NativeString = FastDecimalFormat::NumberToString(Val, FormattingRules, FormattingOptions);
 
-	FText Result = FText(MakeShared<FTextHistory_AsCurrency, ESPMode::ThreadSafe>(MoveTemp(NativeString), Val, CurrencyCode, Options, TargetCulture));
+	FText Result = FText(MakeRefCount<FTextHistory_AsCurrency>(MoveTemp(NativeString), Val, CurrencyCode, Options, TargetCulture));
 	if (!GIsEditor)
 	{
 		Result.Flags |= ETextFlag::Transient;
@@ -621,7 +632,7 @@ FText FText::AsCurrencyBase(int64 BaseVal, const FString& CurrencyCode, const FC
 	double Val = static_cast<double>(BaseVal) / static_cast<double>(FastDecimalFormat::Pow10(DecimalPlaces));
 	FString NativeString = FastDecimalFormat::NumberToString(Val, FormattingRules, FormattingOptions);
 
-	FText Result = FText(MakeShared<FTextHistory_AsCurrency, ESPMode::ThreadSafe>(MoveTemp(NativeString), Val, CurrencyCode, nullptr, TargetCulture));
+	FText Result = FText(MakeRefCount<FTextHistory_AsCurrency>(MoveTemp(NativeString), Val, CurrencyCode, nullptr, TargetCulture));
 	if (!GIsEditor)
 	{
 		Result.Flags |= ETextFlag::Transient;
@@ -651,7 +662,7 @@ FText FText::AsPercentTemplate(T1 Val, const FNumberFormattingOptions* const Opt
 	const FNumberFormattingOptions& FormattingOptions = (Options) ? *Options : FormattingRules.CultureDefaultFormattingOptions;
 	FString NativeString = FastDecimalFormat::NumberToString(Val * static_cast<T1>(100), FormattingRules, FormattingOptions);
 
-	FText Result = FText(MakeShared<FTextHistory_AsPercent, ESPMode::ThreadSafe>(MoveTemp(NativeString), Val, Options, TargetCulture));
+	FText Result = FText(MakeRefCount<FTextHistory_AsPercent>(MoveTemp(NativeString), Val, Options, TargetCulture));
 	if (!GIsEditor)
 	{
 		Result.Flags |= ETextFlag::Transient;
@@ -666,7 +677,7 @@ FText FText::AsDate(const FDateTime& DateTime, const EDateTimeStyle::Type DateSt
 	const FCulture& Culture = TargetCulture.IsValid() ? *TargetCulture : *I18N.GetCurrentLocale();
 
 	FString ChronoSting = FTextChronoFormatter::AsDate(DateTime, DateStyle, TimeZone, Culture);
-	FText Result = FText(MakeShared<FTextHistory_AsDate, ESPMode::ThreadSafe>(MoveTemp(ChronoSting), DateTime, DateStyle, TimeZone, TargetCulture));
+	FText Result = FText(MakeRefCount<FTextHistory_AsDate>(MoveTemp(ChronoSting), DateTime, DateStyle, TimeZone, TargetCulture));
 	if (!GIsEditor)
 	{
 		Result.Flags |= ETextFlag::Transient;
@@ -681,7 +692,7 @@ FText FText::AsDateTime(const FDateTime& DateTime, const EDateTimeStyle::Type Da
 	const FCulture& Culture = TargetCulture.IsValid() ? *TargetCulture : *I18N.GetCurrentLocale();
 
 	FString ChronoString = FTextChronoFormatter::AsDateTime(DateTime, DateStyle, TimeStyle, TimeZone, Culture);
-	FText Result = FText(MakeShared<FTextHistory_AsDateTime, ESPMode::ThreadSafe>(MoveTemp(ChronoString), DateTime, DateStyle, TimeStyle, TimeZone, TargetCulture));
+	FText Result = FText(MakeRefCount<FTextHistory_AsDateTime>(MoveTemp(ChronoString), DateTime, DateStyle, TimeStyle, TimeZone, TargetCulture));
 	if (!GIsEditor)
 	{
 		Result.Flags |= ETextFlag::Transient;
@@ -696,7 +707,7 @@ FText FText::AsDateTime(const FDateTime& DateTime, const FString& CustomPattern,
 	const FCulture& Culture = TargetCulture.IsValid() ? *TargetCulture : *I18N.GetCurrentLocale();
 
 	FString ChronoString = FTextChronoFormatter::AsDateTime(DateTime, CustomPattern, TimeZone, Culture);
-	FText Result = FText(MakeShared<FTextHistory_AsDateTime, ESPMode::ThreadSafe>(MoveTemp(ChronoString), DateTime, CustomPattern, TimeZone, TargetCulture));
+	FText Result = FText(MakeRefCount<FTextHistory_AsDateTime>(MoveTemp(ChronoString), DateTime, CustomPattern, TimeZone, TargetCulture));
 	if (!GIsEditor)
 	{
 		Result.Flags |= ETextFlag::Transient;
@@ -711,7 +722,7 @@ FText FText::AsTime(const FDateTime& DateTime, const EDateTimeStyle::Type TimeSt
 	const FCulture& Culture = TargetCulture.IsValid() ? *TargetCulture : *I18N.GetCurrentLocale();
 
 	FString ChronoString = FTextChronoFormatter::AsTime(DateTime, TimeStyle, TimeZone, Culture);
-	FText Result = FText(MakeShared<FTextHistory_AsTime, ESPMode::ThreadSafe>(MoveTemp(ChronoString), DateTime, TimeStyle, TimeZone, TargetCulture));
+	FText Result = FText(MakeRefCount<FTextHistory_AsTime>(MoveTemp(ChronoString), DateTime, TimeStyle, TimeZone, TargetCulture));
 	if (!GIsEditor)
 	{
 		Result.Flags |= ETextFlag::Transient;
@@ -804,7 +815,7 @@ bool FText::FindText(const FTextKey& Namespace, const FTextKey& Key, FText& OutT
 
 	if ( FoundString.IsValid() )
 	{
-		OutText = FText(MakeShared<FTextHistory_Base, ESPMode::ThreadSafe>(FTextId(Namespace, Key), SourceString ? FString(*SourceString) : FString(), FoundString.ToSharedRef()));
+		OutText = FText(MakeRefCount<FTextHistory_Base>(FTextId(Namespace, Key), SourceString ? FString(*SourceString) : FString(), FoundString.ToSharedRef()));
 		return true;
 	}
 
@@ -847,7 +858,7 @@ void FText::SerializeText(FStructuredArchive::FSlot Slot, FText& Value)
 			TextId = FTextId(Namespace, Key);
 		}
 
-		Value.TextData = MakeShared<FTextHistory_Base, ESPMode::ThreadSafe>(TextId, MoveTemp(SourceStringToImplantIntoHistory));
+		Value.TextData = MakeRefCount<FTextHistory_Base>(TextId, MoveTemp(SourceStringToImplantIntoHistory));
 	}
 
 #if WITH_EDITOR
@@ -923,67 +934,67 @@ void FText::SerializeText(FStructuredArchive::FSlot Slot, FText& Value)
 			{
 			case ETextHistoryType::Base:
 				{
-					Value.TextData = MakeShared<FTextHistory_Base, ESPMode::ThreadSafe>();
+					Value.TextData = MakeRefCount<FTextHistory_Base>();
 					break;
 				}
 			case ETextHistoryType::NamedFormat:
 				{
-					Value.TextData = MakeShared<FTextHistory_NamedFormat, ESPMode::ThreadSafe>();
+					Value.TextData = MakeRefCount<FTextHistory_NamedFormat>();
 					break;
 				}
 			case ETextHistoryType::OrderedFormat:
 				{
-					Value.TextData = MakeShared<FTextHistory_OrderedFormat, ESPMode::ThreadSafe>();
+					Value.TextData = MakeRefCount<FTextHistory_OrderedFormat>();
 					break;
 				}
 			case ETextHistoryType::ArgumentFormat:
 				{
-					Value.TextData = MakeShared<FTextHistory_ArgumentDataFormat, ESPMode::ThreadSafe>();
+					Value.TextData = MakeRefCount<FTextHistory_ArgumentDataFormat>();
 					break;
 				}
 			case ETextHistoryType::AsNumber:
 				{
-					Value.TextData = MakeShared<FTextHistory_AsNumber, ESPMode::ThreadSafe>();
+					Value.TextData = MakeRefCount<FTextHistory_AsNumber>();
 					break;
 				}
 			case ETextHistoryType::AsPercent:
 				{
-					Value.TextData = MakeShared<FTextHistory_AsPercent, ESPMode::ThreadSafe>();
+					Value.TextData = MakeRefCount<FTextHistory_AsPercent>();
 					break;
 				}
 			case ETextHistoryType::AsCurrency:
 				{
-					Value.TextData = MakeShared<FTextHistory_AsCurrency, ESPMode::ThreadSafe>();
+					Value.TextData = MakeRefCount<FTextHistory_AsCurrency>();
 					break;
 				}
 			case ETextHistoryType::AsDate:
 				{
-					Value.TextData = MakeShared<FTextHistory_AsDate, ESPMode::ThreadSafe>();
+					Value.TextData = MakeRefCount<FTextHistory_AsDate>();
 					break;
 				}
 			case ETextHistoryType::AsTime:
 				{
-					Value.TextData = MakeShared<FTextHistory_AsTime, ESPMode::ThreadSafe>();
+					Value.TextData = MakeRefCount<FTextHistory_AsTime>();
 					break;
 				}
 			case ETextHistoryType::AsDateTime:
 				{
-					Value.TextData = MakeShared<FTextHistory_AsDateTime, ESPMode::ThreadSafe>();
+					Value.TextData = MakeRefCount<FTextHistory_AsDateTime>();
 					break;
 				}
 			case ETextHistoryType::Transform:
 				{
-					Value.TextData = MakeShared<FTextHistory_Transform, ESPMode::ThreadSafe>();
+					Value.TextData = MakeRefCount<FTextHistory_Transform>();
 					break;
 				}
 			case ETextHistoryType::StringTableEntry:
 				{
-					Value.TextData = MakeShared<FTextHistory_StringTableEntry, ESPMode::ThreadSafe>();
+					Value.TextData = MakeRefCount<FTextHistory_StringTableEntry>();
 					break;
 				}
 			case ETextHistoryType::TextGenerator:
 				{
-					Value.TextData = MakeShared<FTextHistory_TextGenerator, ESPMode::ThreadSafe>();
+					Value.TextData = MakeRefCount<FTextHistory_TextGenerator>();
 					break;
 				}
 			default:
@@ -1708,11 +1719,11 @@ const TCHAR* FTextStringHelper::ReadFromBuffer_ComplexText(const TCHAR* Buffer,
 	{
 		auto CreateTextHistory = [&](FText& InOutTmpText) -> bool
 		{
-			#define CONDITIONAL_CREATE_TEXT_HISTORY(HistoryClass)								\
-				if (HistoryClass::StaticShouldReadFromBuffer(Buffer))							\
-				{																				\
-					InOutTmpText.TextData = MakeShared<HistoryClass, ESPMode::ThreadSafe>();	\
-					return true;																\
+			#define CONDITIONAL_CREATE_TEXT_HISTORY(HistoryClass)			\
+				if (HistoryClass::StaticShouldReadFromBuffer(Buffer))		\
+				{															\
+					InOutTmpText.TextData = MakeRefCount<HistoryClass>();	\
+					return true;											\
 				}
 			CONDITIONAL_CREATE_TEXT_HISTORY(FTextHistory_Base);
 			CONDITIONAL_CREATE_TEXT_HISTORY(FTextHistory_NamedFormat);
diff --git a/Engine/Source/Runtime/Core/Private/Internationalization/TextFormatter.cpp b/Engine/Source/Runtime/Core/Private/Internationalization/TextFormatter.cpp
--- a/Engine/Source/Runtime/Core/Private/Internationalization/TextFormatter.cpp
+++ b/Engine/Source/Runtime/Core/Private/Internationalization/TextFormatter.cpp
@@ -1032,7 +1032,7 @@ FText FTextFormatter::Format(FTextFormat&& InFmt, FFormatNamedArguments&& InArgu
 {
 	FString ResultString = FormatStr(InFmt, InArguments, bInRebuildText, bInRebuildAsSource);
 
-	FText Result = FText(MakeShared<FTextHistory_NamedFormat, ESPMode::ThreadSafe>(MoveTemp(ResultString), MoveTemp(InFmt), MoveTemp(InArguments)));
+	FText Result = FText(MakeRefCount<FTextHistory_NamedFormat>(MoveTemp(ResultString), MoveTemp(InFmt), MoveTemp(InArguments)));
 	if (!GIsEditor)
 	{
 		Result.Flags |= ETextFlag::Transient;
@@ -1044,7 +1044,7 @@ FText FTextFormatter::Format(FTextFormat&& InFmt, FFormatOrderedArguments&& InAr
 {
 	FString ResultString = FormatStr(InFmt, InArguments, bInRebuildText, bInRebuildAsSource);
 
-	FText Result = FText(MakeShared<FTextHistory_OrderedFormat, ESPMode::ThreadSafe>(MoveTemp(ResultString), MoveTemp(InFmt), MoveTemp(InArguments)));
+	FText Result = FText(MakeRefCount<FTextHistory_OrderedFormat>(MoveTemp(ResultString), MoveTemp(InFmt), MoveTemp(InArguments)));
 	if (!GIsEditor)
 	{
 		Result.Flags |= ETextFlag::Transient;
@@ -1056,7 +1056,7 @@ FText FTextFormatter::Format(FTextFormat&& InFmt, TArray<FFormatArgumentData>&&
 {
 	FString ResultString = FormatStr(InFmt, InArguments, bInRebuildText, bInRebuildAsSource);
 
-	FText Result = FText(MakeShared<FTextHistory_ArgumentDataFormat, ESPMode::ThreadSafe>(MoveTemp(ResultString), MoveTemp(InFmt), MoveTemp(InArguments)));
+	FText Result = FText(MakeRefCount<FTextHistory_ArgumentDataFormat>(MoveTemp(ResultString), MoveTemp(InFmt), MoveTemp(InArguments)));
 	if (!GIsEditor)
 	{
 		Result.Flags |= ETextFlag::Transient;
diff --git a/Engine/Source/Runtime/Core/Private/Internationalization/TextHistory.cpp b/Engine/Source/Runtime/Core/Private/Internationalization/TextHistory.cpp
--- a/Engine/Source/Runtime/Core/Private/Internationalization/TextHistory.cpp
+++ b/Engine/Source/Runtime/Core/Private/Internationalization/TextHistory.cpp
@@ -737,29 +737,6 @@ FTextHistory_Base::FTextHistory_Base(const FTextId& InTextId, FString&& InSource
 	MarkDisplayStringUpToDate();
 }
 
-FTextHistory_Base::FTextHistory_Base(FTextHistory_Base&& Other)
-	: FTextHistory(MoveTemp(Other))
-	, TextId(Other.TextId)
-	, SourceString(MoveTemp(Other.SourceString))
-	, LocalizedString(MoveTemp(Other.LocalizedString))
-{
-	Other.TextId.Reset();
-}
-
-FTextHistory_Base& FTextHistory_Base::operator=(FTextHistory_Base&& Other)
-{
-	FTextHistory::operator=(MoveTemp(Other));
-	if (this != &Other)
-	{
-		TextId = Other.TextId;
-		SourceString = MoveTemp(Other.SourceString);
-		LocalizedString = MoveTemp(Other.LocalizedString);
-
-		Other.TextId.Reset();
-	}
-	return *this;
-}
-
 FTextId FTextHistory_Base::GetTextId() const
 {
 	return TextId;
diff --git a/Engine/Source/Runtime/Core/Private/Internationalization/TextHistory.h b/Engine/Source/Runtime/Core/Private/Internationalization/TextHistory.h
--- a/Engine/Source/Runtime/Core/Private/Internationalization/TextHistory.h
+++ b/Engine/Source/Runtime/Core/Private/Internationalization/TextHistory.h
@@ -136,20 +136,21 @@ void WriteScopedEnumToBuffer(FString& Buffer, const TCHAR* Scope, const T Value)
 }	// namespace TextStringificationUtil
 
 /** Base interface class for all FText history types */
-class FTextHistory : public ITextData
+class FTextHistory : public ITextData, public TRefCountingMixin<FTextHistory>
 {
 public:
 	FTextHistory() = default;
 	virtual ~FTextHistory() = default;
 
-	/** Allow moving */
-	FTextHistory(FTextHistory&& Other) = default;
-	FTextHistory& operator=(FTextHistory&& Other) = default;
-
 	/** Disallow copying */
 	FTextHistory(const FTextHistory&) = delete;
 	FTextHistory& operator=(FTextHistory&) = delete;
 
+	//~ IRefCountedObject
+	virtual uint32 AddRef() const override final { return TRefCountingMixin<FTextHistory>::AddRef(); }
+	virtual uint32 Release() const override final { return TRefCountingMixin<FTextHistory>::Release(); }
+	virtual uint32 GetRefCount() const override final { return TRefCountingMixin<FTextHistory>::GetRefCount(); }
+
 	//~ ITextData
 	virtual const FString& GetSourceString() const override { return GetDisplayString(); }
 	virtual FTextConstDisplayStringPtr GetLocalizedString() const override { return nullptr; }
@@ -248,10 +249,6 @@ public:
 	FTextHistory_Base(const FTextId& InTextId, FString&& InSourceString);
 	FTextHistory_Base(const FTextId& InTextId, FString&& InSourceString, FTextConstDisplayStringPtr&& InLocalizedString);
 
-	/** Allow moving */
-	FTextHistory_Base(FTextHistory_Base&& Other);
-	FTextHistory_Base& operator=(FTextHistory_Base&& Other);
-
 	//~ Begin FTextHistory Interface
 	OVERRIDE_TEXT_HISTORY_STRINGIFICATION;
 	virtual ETextHistoryType GetType() const override { return ETextHistoryType::Base; }
@@ -282,10 +279,6 @@ public:
 	FTextHistory_Generated() = default;
 	explicit FTextHistory_Generated(FString&& InDisplayString);
 
-	/** Allow moving */
-	FTextHistory_Generated(FTextHistory_Generated&& Other) = default;
-	FTextHistory_Generated& operator=(FTextHistory_Generated&& Other) = default;
-
 	//~ Begin FTextHistory Interface
 	virtual FTextId GetTextId() const override final { return FTextId(); }
 	virtual const FString& GetDisplayString() const override;
@@ -308,20 +301,16 @@ class FTextHistory_NamedFormat : public FTextHistory_Generated
 {
 public:
 	FTextHistory_NamedFormat() = default;
-	CORE_API FTextHistory_NamedFormat(FString&& InDisplayString, FTextFormat&& InSourceFmt, FFormatNamedArguments&& InArguments);
-
-	/** Allow moving */
-	FTextHistory_NamedFormat(FTextHistory_NamedFormat&& Other) = default;
-	FTextHistory_NamedFormat& operator=(FTextHistory_NamedFormat&& Other) = default;
+	FTextHistory_NamedFormat(FString&& InDisplayString, FTextFormat&& InSourceFmt, FFormatNamedArguments&& InArguments);
 
 	//~ Begin FTextHistory Interface
 	OVERRIDE_TEXT_HISTORY_STRINGIFICATION;
 	virtual ETextHistoryType GetType() const override { return ETextHistoryType::NamedFormat; }
-	CORE_API virtual bool IdenticalTo(const FTextHistory& Other, const ETextIdenticalModeFlags CompareModeFlags) const override;
-	CORE_API virtual FString BuildLocalizedDisplayString() const override;
-	CORE_API virtual FString BuildInvariantDisplayString() const override;
-	CORE_API virtual void Serialize(FStructuredArchive::FRecord Record) override;
-	CORE_API virtual void GetHistoricFormatData(const FText& InText, TArray<FHistoricTextFormatData>& OutHistoricFormatData) const override;
+	virtual bool IdenticalTo(const FTextHistory& Other, const ETextIdenticalModeFlags CompareModeFlags) const override;
+	virtual FString BuildLocalizedDisplayString() const override;
+	virtual FString BuildInvariantDisplayString() const override;
+	virtual void Serialize(FStructuredArchive::FRecord Record) override;
+	virtual void GetHistoricFormatData(const FText& InText, TArray<FHistoricTextFormatData>& OutHistoricFormatData) const override;
 	//~ End FTextHistory Interface
 
 private:
@@ -338,10 +327,6 @@ public:
 	FTextHistory_OrderedFormat() = default;
 	FTextHistory_OrderedFormat(FString&& InDisplayString, FTextFormat&& InSourceFmt, FFormatOrderedArguments&& InArguments);
 
-	/** Allow moving */
-	FTextHistory_OrderedFormat(FTextHistory_OrderedFormat&& Other) = default;
-	FTextHistory_OrderedFormat& operator=(FTextHistory_OrderedFormat&& Other) = default;
-
 	//~ Begin FTextHistory Interface
 	OVERRIDE_TEXT_HISTORY_STRINGIFICATION;
 	virtual ETextHistoryType GetType() const override { return ETextHistoryType::OrderedFormat; }
@@ -366,10 +351,6 @@ public:
 	FTextHistory_ArgumentDataFormat() = default;
 	FTextHistory_ArgumentDataFormat(FString&& InDisplayString, FTextFormat&& InSourceFmt, TArray<FFormatArgumentData>&& InArguments);
 
-	/** Allow moving */
-	FTextHistory_ArgumentDataFormat(FTextHistory_ArgumentDataFormat&& Other) = default;
-	FTextHistory_ArgumentDataFormat& operator=(FTextHistory_ArgumentDataFormat&& Other) = default;
-
 	//~ Begin FTextHistory Interface
 	OVERRIDE_TEXT_HISTORY_STRINGIFICATION;
 	virtual ETextHistoryType GetType() const override { return ETextHistoryType::ArgumentFormat; }
@@ -394,10 +375,6 @@ public:
 	FTextHistory_FormatNumber() = default;
 	FTextHistory_FormatNumber(FString&& InDisplayString, FFormatArgumentValue InSourceValue, const FNumberFormattingOptions* const InFormatOptions, FCulturePtr InTargetCulture);
 
-	/** Allow moving */
-	FTextHistory_FormatNumber(FTextHistory_FormatNumber&& Other) = default;
-	FTextHistory_FormatNumber& operator=(FTextHistory_FormatNumber&& Other) = default;
-
 	//~ Begin FTextHistory Interface
 	virtual bool IdenticalTo(const FTextHistory& Other, const ETextIdenticalModeFlags CompareModeFlags) const override;
 	virtual void Serialize(FStructuredArchive::FRecord Record) override;
@@ -422,10 +399,6 @@ public:
 	FTextHistory_AsNumber() = default;
 	FTextHistory_AsNumber(FString&& InDisplayString, FFormatArgumentValue InSourceValue, const FNumberFormattingOptions* const InFormatOptions, FCulturePtr InTargetCulture);
 
-	/** Allow moving */
-	FTextHistory_AsNumber(FTextHistory_AsNumber&& Other) = default;
-	FTextHistory_AsNumber& operator=(FTextHistory_AsNumber&& Other) = default;
-
 	//~ Begin FTextHistory Interface
 	OVERRIDE_TEXT_HISTORY_STRINGIFICATION;
 	virtual ETextHistoryType GetType() const override { return ETextHistoryType::AsNumber; }
@@ -443,10 +416,6 @@ public:
 	FTextHistory_AsPercent() = default;
 	FTextHistory_AsPercent(FString&& InDisplayString, FFormatArgumentValue InSourceValue, const FNumberFormattingOptions* const InFormatOptions, FCulturePtr InTargetCulture);
 
-	/** Allow moving */
-	FTextHistory_AsPercent(FTextHistory_AsPercent&& Other) = default;
-	FTextHistory_AsPercent& operator=(FTextHistory_AsPercent&& Other) = default;
-
 	//~ Begin FTextHistory Interface
 	OVERRIDE_TEXT_HISTORY_STRINGIFICATION;
 	virtual ETextHistoryType GetType() const override { return ETextHistoryType::AsPercent; }
@@ -464,10 +433,6 @@ public:
 	FTextHistory_AsCurrency() = default;
 	FTextHistory_AsCurrency(FString&& InDisplayString, FFormatArgumentValue InSourceValue, FString InCurrencyCode, const FNumberFormattingOptions* const InFormatOptions, FCulturePtr InTargetCulture);
 
-	/** Allow moving */
-	FTextHistory_AsCurrency(FTextHistory_AsCurrency&& Other) = default;
-	FTextHistory_AsCurrency& operator=(FTextHistory_AsCurrency&& Other) = default;
-
 	//~ Begin FTextHistory Interface
 	OVERRIDE_TEXT_HISTORY_STRINGIFICATION;
 	virtual ETextHistoryType GetType() const override { return ETextHistoryType::AsCurrency; }
@@ -488,10 +453,6 @@ public:
 	FTextHistory_AsDate() = default;
 	FTextHistory_AsDate(FString&& InDisplayString, FDateTime InSourceDateTime, const EDateTimeStyle::Type InDateStyle, FString InTimeZone, FCulturePtr InTargetCulture);
 
-	/** Allow moving */
-	FTextHistory_AsDate(FTextHistory_AsDate&& Other) = default;
-	FTextHistory_AsDate& operator=(FTextHistory_AsDate&& Other) = default;
-
 	//~ Begin FTextHistory Interface
 	OVERRIDE_TEXT_HISTORY_STRINGIFICATION;
 	virtual ETextHistoryType GetType() const override { return ETextHistoryType::AsDate; }
@@ -519,10 +480,6 @@ public:
 	FTextHistory_AsTime() = default;
 	FTextHistory_AsTime(FString&& InDisplayString, FDateTime InSourceDateTime, const EDateTimeStyle::Type InTimeStyle, FString InTimeZone, FCulturePtr InTargetCulture);
 
-	/** Allow moving */
-	FTextHistory_AsTime(FTextHistory_AsTime&& Other) = default;
-	FTextHistory_AsTime& operator=(FTextHistory_AsTime&& Other) = default;
-
 	//~ Begin FTextHistory Interface
 	OVERRIDE_TEXT_HISTORY_STRINGIFICATION;
 	virtual ETextHistoryType GetType() const override { return ETextHistoryType::AsTime; }
@@ -551,10 +508,6 @@ public:
 	FTextHistory_AsDateTime(FString&& InDisplayString, FDateTime InSourceDateTime, const EDateTimeStyle::Type InDateStyle, const EDateTimeStyle::Type InTimeStyle, FString InTimeZone, FCulturePtr InTargetCulture);
 	FTextHistory_AsDateTime(FString&& InDisplayString, FDateTime InSourceDateTime, FString InCustomPattern, FString InTimeZone, FCulturePtr InTargetCulture);
 
-	/** Allow moving */
-	FTextHistory_AsDateTime(FTextHistory_AsDateTime&& Other) = default;
-	FTextHistory_AsDateTime& operator=(FTextHistory_AsDateTime&& Other) = default;
-
 	//~ Begin FTextHistory Interface
 	OVERRIDE_TEXT_HISTORY_STRINGIFICATION;
 	virtual ETextHistoryType GetType() const override { return ETextHistoryType::AsDateTime; }
@@ -594,10 +547,6 @@ public:
 	FTextHistory_Transform() = default;
 	FTextHistory_Transform(FString&& InDisplayString, FText InSourceText, const ETransformType InTransformType);
 
-	/** Allow moving */
-	FTextHistory_Transform(FTextHistory_Transform&& Other) = default;
-	FTextHistory_Transform& operator=(FTextHistory_Transform&& Other) = default;
-
 	//~ Begin FTextHistory Interface
 	OVERRIDE_TEXT_HISTORY_STRINGIFICATION;
 	virtual ETextHistoryType GetType() const override { return ETextHistoryType::Transform; }
@@ -623,10 +572,6 @@ public:
 	FTextHistory_StringTableEntry() = default;
 	FTextHistory_StringTableEntry(FName InTableId, FString&& InKey, const EStringTableLoadingPolicy InLoadingPolicy);
 
-	/** Allow moving */
-	FTextHistory_StringTableEntry(FTextHistory_StringTableEntry&& Other) = default;
-	FTextHistory_StringTableEntry& operator=(FTextHistory_StringTableEntry&& Other) = default;
-
 	//~ Begin FTextHistory Interface
 	OVERRIDE_TEXT_HISTORY_STRINGIFICATION;
 	virtual ETextHistoryType GetType() const override { return ETextHistoryType::StringTableEntry; }
@@ -720,10 +665,6 @@ public:
 	FTextHistory_TextGenerator() = default;
 	FTextHistory_TextGenerator(FString&& InDisplayString, const TSharedRef<ITextGenerator>& InTextGenerator);
 
-	/** Allow moving */
-	FTextHistory_TextGenerator(FTextHistory_TextGenerator&& Other) = default;
-	FTextHistory_TextGenerator& operator=(FTextHistory_TextGenerator&& Other) = default;
-
 	//~ Begin FTextHistory Interface
 	virtual ETextHistoryType GetType() const override { return ETextHistoryType::TextGenerator; }
 	virtual bool IdenticalTo(const FTextHistory& Other, const ETextIdenticalModeFlags CompareModeFlags) const override;
diff --git a/Engine/Source/Runtime/Core/Public/Internationalization/ITextData.h b/Engine/Source/Runtime/Core/Public/Internationalization/ITextData.h
--- a/Engine/Source/Runtime/Core/Public/Internationalization/ITextData.h
+++ b/Engine/Source/Runtime/Core/Public/Internationalization/ITextData.h
@@ -4,13 +4,14 @@
 
 #include "CoreTypes.h"
 #include "Internationalization/TextLocalizationManager.h"
+#include "Templates/RefCounting.h"
 
 class FTextHistory;
 
 /** 
  * Interface to the internal data for an FText.
  */
-class ITextData
+class ITextData : public IRefCountedObject
 {
 public:
 	virtual ~ITextData() = default;
diff --git a/Engine/Source/Runtime/Core/Public/Internationalization/Text.h b/Engine/Source/Runtime/Core/Public/Internationalization/Text.h
--- a/Engine/Source/Runtime/Core/Public/Internationalization/Text.h
+++ b/Engine/Source/Runtime/Core/Public/Internationalization/Text.h
@@ -382,11 +382,12 @@ public:
 public:
 
 	CORE_API FText();
+	
 	FText(const FText&) = default;
-	FText(FText&&) = default;
-
 	FText& operator=(const FText&) = default;
-	FText& operator=(FText&&) = default;
+	
+	CORE_API FText(FText&& Other);
+	CORE_API FText& operator=(FText&& Other);
 
 	/**
 	 * Generate an FText that represents the passed number in the current culture
@@ -765,11 +766,12 @@ public:
 #endif
 
 private:
-	/** Special constructor used to create StaticEmptyText without also allocating a history object */
-	enum class EInitToEmptyString : uint8 { Value };
-	CORE_API explicit FText( EInitToEmptyString );
-
-	CORE_API explicit FText( TSharedRef<ITextData, ESPMode::ThreadSafe> InTextData );
+	template <typename HistoryType, typename = decltype(ImplicitConv<ITextData*>((HistoryType*)nullptr))>
+	explicit FText( TRefCountPtr<HistoryType>&& InTextData )
+		: TextData(MoveTemp(InTextData))
+		, Flags(0)
+	{
+	}
 
 	CORE_API explicit FText( FString&& InSourceString );
 
@@ -806,7 +808,7 @@ private:
 
 private:
 	/** The internal shared data for this FText */
-	TSharedRef<ITextData, ESPMode::ThreadSafe> TextData;
+	TRefCountPtr<ITextData> TextData;
 
 	/** Flags with various information on what sort of FText this is */
 	uint32 Flags;
@@ -1128,7 +1130,7 @@ private:
 	static uint16 GetLocalHistoryRevisionForText(const FText& InText);
 
 	/** A pointer to the text data for the FText that we took a snapshot of (used for an efficient pointer compare) */
-	TSharedPtr<ITextData, ESPMode::ThreadSafe> TextDataPtr;
+	TRefCountPtr<ITextData> TextDataPtr;
 
 	/** The localized string of the text when we took the snapshot (if any) */
 	FTextConstDisplayStringPtr LocalizedStringPtr;
