Core - Added the [[nodiscard]] attribute to common hash functions

diff --git a/Engine/Source/Runtime/Core/Public/Misc/Crc.h b/Engine/Source/Runtime/Core/Public/Misc/Crc.h
--- a/Engine/Source/Runtime/Core/Public/Misc/Crc.h
+++ b/Engine/Source/Runtime/Core/Public/Misc/Crc.h
@@ -29,21 +29,21 @@ struct FCrc
 	/** generates CRC hash of the memory area */
 	typedef uint32 (*MemCrc32Functor)( const void* Data, int32 Length, uint32 CRC );
 	static CORE_API MemCrc32Functor MemCrc32Func;
-	static FORCEINLINE uint32 MemCrc32(const void* Data, int32 Length, uint32 CRC = 0)
+	[[nodiscard]] static FORCEINLINE uint32 MemCrc32(const void* Data, int32 Length, uint32 CRC = 0)
 	{
 		return MemCrc32Func(Data, Length, CRC);
 	}
 
 	/** generates CRC hash of the element */
 	template <typename T>
-	static uint32 TypeCrc32( const T& Data, uint32 CRC=0 )
+	[[nodiscard]] static uint32 TypeCrc32( const T& Data, uint32 CRC=0 )
 	{
 		return MemCrc32(&Data, sizeof(T), CRC);
 	}
 
 	/** String CRC. */
 	template <typename CharType>
-	static typename TEnableIf<sizeof(CharType) != 1, uint32>::Type StrCrc32(const CharType* Data, uint32 CRC = 0)
+	[[nodiscard]] static typename TEnableIf<sizeof(CharType) != 1, uint32>::Type StrCrc32(const CharType* Data, uint32 CRC = 0)
 	{
 		// We ensure that we never try to do a StrCrc32 with a CharType of more than 4 bytes.  This is because
 		// we always want to treat every CRC as if it was based on 4 byte chars, even if it's less, because we
@@ -65,7 +65,7 @@ struct FCrc
 	}
 
 	template <typename CharType>
-	static typename TEnableIf<sizeof(CharType) == 1, uint32>::Type StrCrc32(const CharType* Data, uint32 CRC = 0)
+	[[nodiscard]] static typename TEnableIf<sizeof(CharType) == 1, uint32>::Type StrCrc32(const CharType* Data, uint32 CRC = 0)
 	{
 		/* Overload for when CharType is a byte, which causes warnings when right-shifting by 8 */
 		CRC = ~CRC;
diff --git a/Engine/Source/Runtime/Core/Public/Templates/TypeHash.h b/Engine/Source/Runtime/Core/Public/Templates/TypeHash.h
--- a/Engine/Source/Runtime/Core/Public/Templates/TypeHash.h
+++ b/Engine/Source/Runtime/Core/Public/Templates/TypeHash.h
@@ -33,7 +33,7 @@ namespace UE
  * This function cannot change for backward compatibility reasons.
  * You may want to choose HashCombineFast for a better in-memory hash combining function.
  */
-inline uint32 HashCombine(uint32 A, uint32 C)
+[[nodiscard]] inline uint32 HashCombine(uint32 A, uint32 C)
 {
 	uint32 B = 0x9e3779b9;
 	A += B;
@@ -59,12 +59,12 @@ inline uint32 HashCombine(uint32 A, uint32 C)
  *           combined hash values which don't leave the running process,
  *           e.g. GetTypeHash() overloads.
  */
-inline uint32 HashCombineFast(uint32 A, uint32 B)
+[[nodiscard]] inline uint32 HashCombineFast(uint32 A, uint32 B)
 {
 	return A ^ (B + 0x9e3779b9 + (A << 6) + (A >> 2));
 }
 
-inline uint32 PointerHash(const void* Key)
+[[nodiscard]] inline uint32 PointerHash(const void* Key)
 {
 	// Ignoring the lower 4 bits since they are likely zero anyway.
 	// Higher bits are more significant in 64 bit builds.
@@ -72,7 +72,7 @@ inline uint32 PointerHash(const void* Key)
 	return UE::Private::MurmurFinalize32((uint32)PtrInt);
 }
 
-inline uint32 PointerHash(const void* Key, uint32 C)
+[[nodiscard]] inline uint32 PointerHash(const void* Key, uint32 C)
 {
 	// we can use HashCombineFast here because pointers are non-persistent
 	return HashCombineFast(PointerHash(Key), C);
@@ -91,7 +91,7 @@ template <
 	typename ScalarType,
 	std::enable_if_t<std::is_scalar_v<ScalarType> && !std::is_same_v<ScalarType, TCHAR*> && !std::is_same_v<ScalarType, const TCHAR*>>* = nullptr
 >
-inline uint32 GetTypeHash(ScalarType Value)
+[[nodiscard]] inline uint32 GetTypeHash(ScalarType Value)
 {
 	if constexpr (std::is_integral_v<ScalarType>)
 	{
@@ -173,7 +173,7 @@ inline uint32 GetTypeHash(T* Value)
 }
 
 template <typename T>
-inline uint32 GetArrayHash(const T* Ptr, uint64 Size, uint32 PreviousHash = 0)
+[[nodiscard]] inline uint32 GetArrayHash(const T* Ptr, uint64 Size, uint32 PreviousHash = 0)
 {
 	uint32 Result = PreviousHash;
 	while (Size)
@@ -188,4 +188,4 @@ inline uint32 GetArrayHash(const T* Ptr, uint64 Size, uint32 PreviousHash = 0)
 
 // Use this when inside type that has GetTypeHash() (no in-parameters) implemented. It makes GetTypeHash dispatch in global namespace
 template <typename T>
-FORCEINLINE uint32 GetTypeHashHelper(const T& V) { return GetTypeHash(V); }
+[[nodiscard]] FORCEINLINE uint32 GetTypeHashHelper(const T& V) { return GetTypeHash(V); }
