Fixed constraints on TSharedPtr<>::operator*() and TUniquePtr<>::operator*(), which would fail to compile for types which overloaded the comma operator.
Reverted workaround.

diff --git a/Engine/Plugins/NetcodeUnitTest/NetcodeUnitTest/Source/NetcodeUnitTest/Private/NUTUtilReflectionParser.cpp b/Engine/Plugins/NetcodeUnitTest/NetcodeUnitTest/Source/NetcodeUnitTest/Private/NUTUtilReflectionParser.cpp
--- a/Engine/Plugins/NetcodeUnitTest/NetcodeUnitTest/Source/NetcodeUnitTest/Private/NUTUtilReflectionParser.cpp
+++ b/Engine/Plugins/NetcodeUnitTest/NetcodeUnitTest/Source/NetcodeUnitTest/Private/NUTUtilReflectionParser.cpp
@@ -196,7 +196,7 @@ FVMReflectionParser::FVMReflectionParser()
 
 		if (RealContext != nullptr)
 		{
-			FVMReflection& Refl = *(RealContext->Refl.Get());
+			FVMReflection& Refl = *(RealContext->Refl);
 
 			Refl->*(B.IdentifierName);
 
diff --git a/Engine/Source/Runtime/Core/Public/Templates/SharedPointer.h b/Engine/Source/Runtime/Core/Public/Templates/SharedPointer.h
--- a/Engine/Source/Runtime/Core/Public/Templates/SharedPointer.h
+++ b/Engine/Source/Runtime/Core/Public/Templates/SharedPointer.h
@@ -1101,7 +1101,7 @@ public:
 	 */
 	template <
 		typename DummyObjectType = ObjectType,
-		decltype(*(DummyObjectType*)nullptr, 0) = 0 // this construct means that operator* is only considered for overload resolution if T is dereferenceable
+		decltype((int&)(*(DummyObjectType*)nullptr), 0) = 0 // this construct means that operator* is only considered for overload resolution if T is dereferenceable
 	>
 	[[nodiscard]] FORCEINLINE DummyObjectType& operator*() const
 	{
diff --git a/Engine/Source/Runtime/Core/Public/Templates/UniquePtr.h b/Engine/Source/Runtime/Core/Public/Templates/UniquePtr.h
--- a/Engine/Source/Runtime/Core/Public/Templates/UniquePtr.h
+++ b/Engine/Source/Runtime/Core/Public/Templates/UniquePtr.h
@@ -305,7 +305,7 @@ public:
 	 */
 	template <
 		typename DummyT = T,
-		decltype(*(DummyT*)nullptr, 0) = 0 // this construct means that operator* is only considered for overload resolution if T is dereferenceable
+		decltype((int&)*(DummyT*)nullptr, 0) = 0 // this construct means that operator* is only considered for overload resolution if T is dereferenceable
 	>
 	[[nodiscard]] FORCEINLINE DummyT& operator*() const
 	{
