Implemented localization resource unloading for plugins

This leverages 4-bytes of padding in FDisplayStringEntry to store an ID (an int32 index generated via FTextLocalizationManager::GetLocalizationTargetPathId) that is used to tag each display string with the localization target that it belongs to, without increasing the size of DisplayStringLookupTable.

There is an additional LUT mapping each localization target ID to its text IDs to act as an accelerator for HandleLocalizationTargetsUnmounted, however this LUT only tracks text IDs associated with localization data loaded via HandleLocalizationTargetsMounted to minimize the memory impact of this additional tracking. Without this LUT HandleLocalizationTargetsUnmounted would need to spin the entire live table each time, which is not acceptable from a performance standpoint.

Now that HandleLocalizationTargetsUnmounted is implemented it's also been hooked into GFP deactivation, to mirror how HandleLocalizationTargetsMounted was hooked into GFP activation.

diff --git a/Engine/Plugins/Runtime/GameFeatures/Source/GameFeatures/Private/GameFeaturePluginStateMachine.cpp b/Engine/Plugins/Runtime/GameFeatures/Source/GameFeatures/Private/GameFeaturePluginStateMachine.cpp
--- a/Engine/Plugins/Runtime/GameFeatures/Source/GameFeatures/Private/GameFeaturePluginStateMachine.cpp
+++ b/Engine/Plugins/Runtime/GameFeatures/Source/GameFeatures/Private/GameFeaturePluginStateMachine.cpp
@@ -2883,6 +2883,8 @@ struct FGameFeaturePluginState_Deactivating : public FGameFeaturePluginState
 			{
 				StateStatus.SetTransition(EGameFeaturePluginState::Loaded);
 			}
+
+			IPluginManager::Get().UnmountExplicitlyLoadedPluginLocalizationData(StateProperties.PluginName);
 		}
 		else
 		{
diff --git a/Engine/Source/Runtime/Core/Private/Internationalization/TextCache.cpp b/Engine/Source/Runtime/Core/Private/Internationalization/TextCache.cpp
--- a/Engine/Source/Runtime/Core/Private/Internationalization/TextCache.cpp
+++ b/Engine/Source/Runtime/Core/Private/Internationalization/TextCache.cpp
@@ -20,15 +20,19 @@ void FTextCache::TearDown()
 }
 
 FText FTextCache::FindOrCache(const TCHAR* InTextLiteral, const TCHAR* InNamespace, const TCHAR* InKey)
+{
+	return FindOrCache(InTextLiteral, FTextId(InNamespace, InKey));
+}
+
+FText FTextCache::FindOrCache(const TCHAR* InTextLiteral, const FTextId& InTextId)
 {
 	LLM_SCOPE(ELLMTag::Localization);
-	const FTextId TextId(InNamespace, InKey);
 
 	// First try and find a cached instance
 	{
 		FScopeLock Lock(&CachedTextCS);
 	
-		const FText* FoundText = CachedText.Find(TextId);
+		const FText* FoundText = CachedText.Find(InTextId);
 		if (FoundText)
 		{
 			const FString* FoundTextLiteral = FTextInspector::GetSourceString(*FoundText);
@@ -40,14 +44,37 @@ FText FTextCache::FindOrCache(const TCHAR* InTextLiteral, const TCHAR* InNamespa
 	}
 
 	// Not currently cached, make a new instance...
-	FText NewText = FText(InTextLiteral, TextId.GetNamespace(), TextId.GetKey(), ETextFlag::Immutable);
+	FText NewText = FText(InTextLiteral, InTextId.GetNamespace(), InTextId.GetKey(), ETextFlag::Immutable);
 
 	// ... and add it to the cache
 	{
 		FScopeLock Lock(&CachedTextCS);
 
-		CachedText.Emplace(TextId, NewText);
+		CachedText.Emplace(InTextId, NewText);
 	}
 
 	return NewText;
-}
\ No newline at end of file
+}
+
+void FTextCache::RemoveCache(const FTextId& InTextId)
+{
+	return RemoveCache(MakeArrayView(&InTextId, 1));
+}
+
+void FTextCache::RemoveCache(TArrayView<const FTextId> InTextIds)
+{
+	FScopeLock Lock(&CachedTextCS);
+	for (const FTextId& TextId : InTextIds)
+	{
+		CachedText.Remove(TextId);
+	}
+}
+
+void FTextCache::RemoveCache(const TSet<FTextId>& InTextIds)
+{
+	FScopeLock Lock(&CachedTextCS);
+	for (const FTextId& TextId : InTextIds)
+	{
+		CachedText.Remove(TextId);
+	}
+}
diff --git a/Engine/Source/Runtime/Core/Private/Internationalization/TextCache.h b/Engine/Source/Runtime/Core/Private/Internationalization/TextCache.h
--- a/Engine/Source/Runtime/Core/Private/Internationalization/TextCache.h
+++ b/Engine/Source/Runtime/Core/Private/Internationalization/TextCache.h
@@ -19,11 +19,24 @@ public:
 	static FTextCache& Get();
 	static void TearDown();
 
+	FTextCache() = default;
+
+	/** Non-copyable */
+	FTextCache(const FTextCache&) = delete;
+	FTextCache& operator=(const FTextCache&) = delete;
+
 	/**
 	 * Try and find an existing cached entry for the given data, or construct and cache a new entry if one cannot be found.
 	 */
 	FText FindOrCache(const TCHAR* InTextLiteral, const TCHAR* InNamespace, const TCHAR* InKey);
+	FText FindOrCache(const TCHAR* InTextLiteral, const FTextId& InTextId);
 
+	/**
+	 * Remove any cached entries for the given text IDs.
+	 */
+	void RemoveCache(const FTextId& InTextId);
+	void RemoveCache(TArrayView<const FTextId> InTextIds);
+	void RemoveCache(const TSet<FTextId>& InTextIds);
 
 private:
 	TMap<FTextId, FText> CachedText;
diff --git a/Engine/Source/Runtime/Core/Private/Internationalization/TextLocalizationManager.cpp b/Engine/Source/Runtime/Core/Private/Internationalization/TextLocalizationManager.cpp
--- a/Engine/Source/Runtime/Core/Private/Internationalization/TextLocalizationManager.cpp
+++ b/Engine/Source/Runtime/Core/Private/Internationalization/TextLocalizationManager.cpp
@@ -529,6 +529,58 @@ void InitGameTextLocalization()
 	});
 }
 
+FTextLocalizationManager::FDisplayStringsForLocalizationTarget& FTextLocalizationManager::FDisplayStringsByLocalizationTargetId::FindOrAdd(FStringView InLocalizationTargetPath, int32* OutLocalizationTargetPathId)
+{
+	LLM_SCOPE_BYNAME(TEXT("Localization/DisplayStringsByTarget"));
+
+	check(!InLocalizationTargetPath.IsEmpty());
+
+	FString NormalizedLocalizationTargetPath = FPaths::ConvertRelativePathToFull(FString(InLocalizationTargetPath));
+	FPaths::NormalizeDirectoryName(NormalizedLocalizationTargetPath);
+
+	int32 LocalizationTargetPathId = LocalizationTargetPathsToIds.FindRef(NormalizedLocalizationTargetPath, INDEX_NONE);
+	if (LocalizationTargetPathId == INDEX_NONE)
+	{
+		LocalizationTargetPathId = LocalizationTargets.Emplace(MoveTemp(NormalizedLocalizationTargetPath));
+		LocalizationTargetPathsToIds.Add(LocalizationTargets[LocalizationTargetPathId].LocalizationTargetPath, LocalizationTargetPathId);
+	}
+
+	if (OutLocalizationTargetPathId)
+	{
+		*OutLocalizationTargetPathId = LocalizationTargetPathId;
+	}
+	return LocalizationTargets[LocalizationTargetPathId];
+}
+
+FTextLocalizationManager::FDisplayStringsForLocalizationTarget* FTextLocalizationManager::FDisplayStringsByLocalizationTargetId::Find(const int32 InLocalizationTargetPathId)
+{
+	return LocalizationTargets.IsValidIndex(InLocalizationTargetPathId)
+		? &LocalizationTargets[InLocalizationTargetPathId]
+		: nullptr;
+}
+
+void FTextLocalizationManager::FDisplayStringsByLocalizationTargetId::TrackTextId(const int32 InCurrentLocalizationPathId, const int32 InNewLocalizationPathId, const FTextId& InTextId)
+{
+	if (InCurrentLocalizationPathId == InNewLocalizationPathId)
+	{
+		return;
+	}
+
+	LLM_SCOPE_BYNAME(TEXT("Localization/DisplayStringsByTarget"));
+
+	if (FDisplayStringsForLocalizationTarget* DisplayStringsForCurrentLocalizationTarget = Find(InCurrentLocalizationPathId);
+		DisplayStringsForCurrentLocalizationTarget && DisplayStringsForCurrentLocalizationTarget->bIsMounted)
+	{
+		DisplayStringsForCurrentLocalizationTarget->TextIds.Remove(InTextId);
+	}
+
+	if (FDisplayStringsForLocalizationTarget* DisplayStringsForNewLocalizationTarget = Find(InNewLocalizationPathId);
+		DisplayStringsForNewLocalizationTarget && DisplayStringsForNewLocalizationTarget->bIsMounted)
+	{
+		DisplayStringsForNewLocalizationTarget->TextIds.Add(InTextId);
+	}
+}
+
 FTextLocalizationManager& FTextLocalizationManager::Get()
 {
 	return TLazySingleton<FTextLocalizationManager>::Get();
@@ -725,6 +777,15 @@ TArray<FString> FTextLocalizationManager::GetLocalizedCultureNames(const ELocali
 	return LocalizedCultureNames;
 }
 
+int32 FTextLocalizationManager::GetLocalizationTargetPathId(FStringView InLocalizationTargetPath)
+{
+	FScopeLock ScopeLock(&DisplayStringTableCS);
+
+	int32 LocalizationTargetPathId = INDEX_NONE;
+	DisplayStringsByLocalizationTargetId.FindOrAdd(InLocalizationTargetPath, &LocalizationTargetPathId);
+	return LocalizationTargetPathId;
+}
+
 void FTextLocalizationManager::RegisterTextSource(const TSharedRef<ILocalizedTextSource>& InLocalizedTextSource, const bool InRefreshResources)
 {
 	LocalizedTextSources.Add(InLocalizedTextSource);
@@ -1003,6 +1064,7 @@ FTextConstDisplayStringRef FTextLocalizationManager::GetDisplayString(const FTex
 		// Make entries so that they can be updated when system is initialized or a culture swap occurs.
 		FDisplayStringEntry NewEntry(
 			FTextKey(),					/*LocResID*/
+			INDEX_NONE,					/*LocalizationTargetPathId*/
 			SourceStringHash,			/*SourceStringHash*/
 			UnlocalizedString			/*String*/
 		);
@@ -1085,7 +1147,7 @@ bool FTextLocalizationManager::AddDisplayString(const FTextDisplayStringRef& Dis
 	}
 
 	// Add the necessary association.
-	DisplayStringLookupTable.Emplace(TextId, FDisplayStringEntry(FTextKey(), FTextLocalizationResource::HashString(*DisplayString), DisplayString));
+	DisplayStringLookupTable.Emplace(TextId, FDisplayStringEntry(FTextKey(), INDEX_NONE, FTextLocalizationResource::HashString(*DisplayString), DisplayString));
 
 	return true;
 }
@@ -1137,6 +1199,22 @@ void FTextLocalizationManager::HandleLocalizationTargetsMounted(TArrayView<const
 		return;
 	}
 
+	// Nothing to do?
+	if (!FTextLocalizationManager::IsDisplayStringSupportEnabled())
+	{
+		return;
+	}
+
+	// Mark the targets as mounted before loading any of their data
+	{
+		FScopeLock ScopeLock(&DisplayStringTableCS);
+		for (const FString& LocalizationTargetPath : LocalizationTargetPaths)
+		{
+			FDisplayStringsForLocalizationTarget& DisplayStringsForLocalizationTarget = DisplayStringsByLocalizationTargetId.FindOrAdd(LocalizationTargetPath);
+			DisplayStringsForLocalizationTarget.bIsMounted = true;
+		}
+	}
+
 	ELocalizationLoadFlags LocLoadFlags = ELocalizationLoadFlags::None;
 	LocLoadFlags |= (WITH_EDITOR ? ELocalizationLoadFlags::Editor : ELocalizationLoadFlags::None);
 	LocLoadFlags |= (FApp::IsGame() ? ELocalizationLoadFlags::Game : ELocalizationLoadFlags::None);
@@ -1149,15 +1227,47 @@ void FTextLocalizationManager::HandleLocalizationTargetsMounted(TArrayView<const
 
 void FTextLocalizationManager::HandleLocalizationTargetsUnmounted(TArrayView<const FString> LocalizationTargetPaths)
 {
-	if (!IsInitialized())
+	if (!IsInitialized() || LocalizationTargetPaths.IsEmpty())
 	{
 		// If we've not yet loaded localization data then there's nothing to do
 		return;
 	}
 	
-	// Note: We don't track which LocRes text comes from (as it uses too much memory), so unloading text would require a full refresh (including flushing DisplayStringLookupTable)
-	//       RefreshResources cannot currently do a flush, so for now we ignore any unmount notifications
-	//RefreshResources();
+	// Nothing to do?
+	if (!FTextLocalizationManager::IsDisplayStringSupportEnabled())
+	{
+		return;
+	}
+
+	// Async update the live table
+	QueueAsyncTask([LocalizationTargetPaths = TArray<FString>(LocalizationTargetPaths)]()
+	{
+		TRACE_CPUPROFILER_EVENT_SCOPE(FTextLocalizationManager::HandleLocalizationTargetsUnmounted);
+		LLM_SCOPE_BYNAME(TEXT("Localization/DisplayStrings"));
+
+		FTextLocalizationManager& TLM = FTextLocalizationManager::Get();
+		FTextCache& TextCache = FTextCache::Get();
+
+		// Lock while updating the tables
+		FScopeLock ScopeLock(&TLM.DisplayStringTableCS);
+
+		// Discard the data for each localization target that was unmounted, and mark the target as no longer mounted so that we no longer track its text IDs
+		for (const FString& LocalizationTargetPath : LocalizationTargetPaths)
+		{
+			FDisplayStringsForLocalizationTarget& DisplayStringsForLocalizationTarget = TLM.DisplayStringsByLocalizationTargetId.FindOrAdd(LocalizationTargetPath);
+			if (DisplayStringsForLocalizationTarget.bIsMounted)
+			{
+				for (const FTextId& TextId : DisplayStringsForLocalizationTarget.TextIds)
+				{
+					TLM.DisplayStringLookupTable.Remove(TextId);
+				}
+				TextCache.RemoveCache(DisplayStringsForLocalizationTarget.TextIds);
+
+				DisplayStringsForLocalizationTarget.TextIds.Empty();
+				DisplayStringsForLocalizationTarget.bIsMounted = false;
+			}
+		}
+	});
 }
 
 void FTextLocalizationManager::OnPakFileMounted(const IPakFile& PakFile)
@@ -1546,6 +1656,8 @@ void FTextLocalizationManager::UpdateFromNative(FTextLocalizationResource&& Text
 #if WITH_EDITORONLY_DATA
 					LiveEntry->LocResID = NewEntry.LocResID;
 #endif	// WITH_EDITORONLY_DATA
+					DisplayStringsByLocalizationTargetId.TrackTextId(LiveEntry->LocalizationTargetPathId, NewEntry.LocalizationTargetPathId, TextId);
+					LiveEntry->LocalizationTargetPathId = NewEntry.LocalizationTargetPathId;
 #if ENABLE_LOC_TESTING
 					DisplayStringBackupTable.Remove(TextId);
 #endif	// ENABLE_LOC_TESTING
@@ -1556,11 +1668,13 @@ void FTextLocalizationManager::UpdateFromNative(FTextLocalizationResource&& Text
 				// Add new entry
 				FDisplayStringEntry NewLiveEntry(
 					NewEntry.LocResID,						/*LocResID*/
+					NewEntry.LocalizationTargetPathId,		/*LocalizationTargetPathId*/
 					NewEntry.SourceStringHash,				/*SourceStringHash*/
 					NewEntry.LocalizedString.ToSharedRef()	/*String*/
 				);
 
 				DisplayStringLookupTable.Emplace(TextId, NewLiveEntry);
+				DisplayStringsByLocalizationTargetId.TrackTextId(INDEX_NONE, NewEntry.LocalizationTargetPathId, TextId);
 			}
 		}
 
@@ -1592,6 +1706,8 @@ void FTextLocalizationManager::UpdateFromNative(FTextLocalizationResource&& Text
 #if WITH_EDITORONLY_DATA
 						LiveEntry.LocResID = DisplayStringEntry->LocResID;
 #endif	// WITH_EDITORONLY_DATA
+						DisplayStringsByLocalizationTargetId.TrackTextId(LiveEntry.LocalizationTargetPathId, DisplayStringEntry->LocalizationTargetPathId, TextId);
+						LiveEntry.LocalizationTargetPathId = DisplayStringEntry->LocalizationTargetPathId;
 #if ENABLE_LOC_TESTING
 						DisplayStringBackupTable.Remove(TextId);
 #endif	// ENABLE_LOC_TESTING
@@ -1670,6 +1786,8 @@ void FTextLocalizationManager::UpdateFromLocalizations(FTextLocalizationResource
 #if WITH_EDITORONLY_DATA
 					LiveEntry->LocResID = FTextKey();
 #endif	// WITH_EDITORONLY_DATA
+					DisplayStringsByLocalizationTargetId.TrackTextId(LiveEntry->LocalizationTargetPathId, INDEX_NONE, TextId);
+					LiveEntry->LocalizationTargetPathId = INDEX_NONE;
 				}
 #endif	// ENABLE_LOC_TESTING
 			}
@@ -1678,11 +1796,13 @@ void FTextLocalizationManager::UpdateFromLocalizations(FTextLocalizationResource
 				// Add new entry
 				FDisplayStringEntry NewLiveEntry(
 					NewEntry.LocResID,						/*LocResID*/
+					NewEntry.LocalizationTargetPathId,		/*LocalizationTargetPathId*/
 					NewEntry.SourceStringHash,				/*SourceStringHash*/
 					NewEntry.LocalizedString.ToSharedRef()	/*String*/
 				);
 
 				DisplayStringLookupTable.Emplace(TextId, NewLiveEntry);
+				DisplayStringsByLocalizationTargetId.TrackTextId(INDEX_NONE, NewEntry.LocalizationTargetPathId, TextId);
 			}
 		}
 
@@ -1716,6 +1836,8 @@ void FTextLocalizationManager::UpdateFromLocalizations(FTextLocalizationResource
 #if WITH_EDITORONLY_DATA
 						LiveEntry.LocResID = DisplayStringEntry->LocResID;
 #endif	// WITH_EDITORONLY_DATA
+						DisplayStringsByLocalizationTargetId.TrackTextId(LiveEntry.LocalizationTargetPathId, DisplayStringEntry->LocalizationTargetPathId, TextId);
+						LiveEntry.LocalizationTargetPathId = DisplayStringEntry->LocalizationTargetPathId;
 					}
 #if ENABLE_LOC_TESTING
 					else if (bShouldLEETIFYUnlocalizedString && !LiveEntry.DisplayString->IsEmpty())
@@ -1726,6 +1848,8 @@ void FTextLocalizationManager::UpdateFromLocalizations(FTextLocalizationResource
 #if WITH_EDITORONLY_DATA
 						LiveEntry.LocResID = FTextKey();
 #endif	// WITH_EDITORONLY_DATA
+						DisplayStringsByLocalizationTargetId.TrackTextId(LiveEntry.LocalizationTargetPathId, INDEX_NONE, TextId);
+						LiveEntry.LocalizationTargetPathId = INDEX_NONE;
 					}
 #endif	// ENABLE_LOC_TESTING
 				}
diff --git a/Engine/Source/Runtime/Core/Private/Internationalization/TextLocalizationResource.cpp b/Engine/Source/Runtime/Core/Private/Internationalization/TextLocalizationResource.cpp
--- a/Engine/Source/Runtime/Core/Private/Internationalization/TextLocalizationResource.cpp
+++ b/Engine/Source/Runtime/Core/Private/Internationalization/TextLocalizationResource.cpp
@@ -7,6 +7,7 @@
 #include "Misc/App.h"
 #include "Misc/Parse.h"
 #include "Misc/Paths.h"
+#include "Misc/PathViews.h"
 #include "Misc/Optional.h"
 #include "Misc/ConfigCacheIni.h"
 #include "Templates/UniquePtr.h"
@@ -26,6 +27,25 @@ const FGuid FTextLocalizationResourceVersion::LocResMagic = FGuid(0x7574140E, 0x
 /** LocRes files can be quite large, so we won't pre-load those by default */
 #define PRELOAD_LOCRES_FILES (0)
 
+namespace TextLocalizationResourceUtil
+{
+
+int32 GetLocalizationTargetPathIdFromLocResId(const FTextKey& InLocResID)
+{
+	int32 LocalizationTargetPathId = INDEX_NONE;
+	if (!InLocResID.IsEmpty())
+	{
+		// LocResID would be "/Path/To/LocalizationTarget/Culture/LocalizationTarget.locres" so trim this back to "/Path/To/LocalizationTarget"
+		FStringView LocalizationTargetPath = InLocResID.GetChars();
+		LocalizationTargetPath = FPathViews::GetPath(LocalizationTargetPath); // Remove "LocalizationTarget.locres"
+		LocalizationTargetPath = FPathViews::GetPath(LocalizationTargetPath); // Remove "Culture"
+		LocalizationTargetPathId = FTextLocalizationManager::Get().GetLocalizationTargetPathId(LocalizationTargetPath);
+	}
+	return LocalizationTargetPathId;
+}
+
+}
+
 bool FTextLocalizationMetaDataResource::LoadFromFile(const FString& FilePath)
 {
 	TUniquePtr<FArchive> Reader;
@@ -170,6 +190,7 @@ void FTextLocalizationResource::AddEntry(const FTextKey& InNamespace, const FTex
 {
 	FEntry NewEntry;
 	NewEntry.LocResID = InLocResID;
+	NewEntry.LocalizationTargetPathId = TextLocalizationResourceUtil::GetLocalizationTargetPathIdFromLocResId(InLocResID);
 	NewEntry.SourceStringHash = InSourceStringHash;
 	NewEntry.LocalizedString = InLocalizedString;
 	NewEntry.Priority = InPriority;
@@ -318,6 +339,7 @@ bool FTextLocalizationResource::LoadFromArchive(FArchive& Archive, const FTextKe
 		}
 	};
 
+	const int32 LocalizationTargetPathId = TextLocalizationResourceUtil::GetLocalizationTargetPathIdFromLocResId(LocResID);
 	for (uint32 i = 0; i < NamespaceCount; ++i)
 	{
 		// Read namespace
@@ -336,6 +358,7 @@ bool FTextLocalizationResource::LoadFromArchive(FArchive& Archive, const FTextKe
 
 			FEntry NewEntry;
 			NewEntry.LocResID = LocResID;
+			NewEntry.LocalizationTargetPathId = LocalizationTargetPathId;
 			NewEntry.Priority = Priority;
 
 			Archive << NewEntry.SourceStringHash;
diff --git a/Engine/Source/Runtime/Core/Public/Internationalization/TextLocalizationManager.h b/Engine/Source/Runtime/Core/Public/Internationalization/TextLocalizationManager.h
--- a/Engine/Source/Runtime/Core/Public/Internationalization/TextLocalizationManager.h
+++ b/Engine/Source/Runtime/Core/Public/Internationalization/TextLocalizationManager.h
@@ -54,13 +54,15 @@ private:
 #if WITH_EDITORONLY_DATA
 		FTextKey LocResID;
 #endif
+		int32 LocalizationTargetPathId = INDEX_NONE;
 		uint32 SourceStringHash;
 
-		FDisplayStringEntry(const FTextKey& InLocResID, const uint32 InSourceStringHash, const FTextConstDisplayStringRef& InDisplayString)
+		FDisplayStringEntry(const FTextKey& InLocResID, const int32 InLocalizationTargetPathId, const uint32 InSourceStringHash, const FTextConstDisplayStringRef& InDisplayString)
 			: DisplayString(InDisplayString)
 #if WITH_EDITORONLY_DATA
 			, LocResID(InLocResID)
 #endif
+			, LocalizationTargetPathId(InLocalizationTargetPathId)
 			, SourceStringHash(InSourceStringHash)
 		{
 		}
@@ -77,6 +79,39 @@ private:
 	/** Manages the currently loaded or registered text localizations. */
 	typedef TMap<FTextId, FDisplayStringEntry> FDisplayStringLookupTable;
 
+	struct FDisplayStringsForLocalizationTarget
+	{
+		/**
+		 * The path of this localization target.
+		 */
+		FString LocalizationTargetPath;
+
+		/**
+		 * Text IDs currently associated with this localization target.
+		 * @note This information is also known via FDisplayStringEntry::LocalizationTargetPathId, but this serves as 
+		 *       an accelerator for HandleLocalizationTargetsUnmounted to avoid spinning the entire live table.
+		 */
+		TSet<FTextId> TextIds;
+
+		/**
+		 * True if this localization target has been mounted via HandleLocalizationTargetsMounted.
+		 * @note Only mounted localization targets track TextIds, as they're the only things that can be unloaded via HandleLocalizationTargetsUnmounted.
+		 */
+		bool bIsMounted = false;
+	};
+
+	struct FDisplayStringsByLocalizationTargetId
+	{
+	public:
+		FDisplayStringsForLocalizationTarget& FindOrAdd(FStringView InLocalizationTargetPath, int32* OutLocalizationTargetPathId = nullptr);
+		FDisplayStringsForLocalizationTarget* Find(const int32 InLocalizationTargetPathId);
+		void TrackTextId(const int32 InCurrentLocalizationPathId, const int32 InNewLocalizationPathId, const FTextId& InTextId);
+
+	private:
+		TArray<FDisplayStringsForLocalizationTarget> LocalizationTargets;
+		TMap<FStringView, int32> LocalizationTargetPathsToIds;
+	};
+
 private:
 	std::atomic<ETextLocalizationManagerInitializedFlags> InitializedFlags{ ETextLocalizationManagerInitializedFlags::None };
 	
@@ -90,6 +125,7 @@ private:
 #if ENABLE_LOC_TESTING
 	TMap<FTextId, FTextConstDisplayStringPtr> DisplayStringBackupTable;
 #endif
+	FDisplayStringsByLocalizationTargetId DisplayStringsByLocalizationTargetId;
 
 	mutable FRWLock TextRevisionRW;
 	TMap<FTextId, uint16> LocalTextRevisions;
@@ -151,6 +187,12 @@ public:
 	 */
 	CORE_API TArray<FString> GetLocalizedCultureNames(const ELocalizationLoadFlags InLoadFlags) const;
 
+	/**
+	 * Given a localization target path, get the ID associated with it.
+	 * @note This ID is unstable and should only be used for quick in-process comparison.
+	 */
+	CORE_API int32 GetLocalizationTargetPathId(FStringView InLocalizationTargetPath);
+
 	/**
 	 * Register a localized text source with the text localization manager.
 	 */
@@ -215,8 +257,8 @@ public:
 	CORE_API void HandleLocalizationTargetsMounted(TArrayView<const FString> LocalizationTargetPaths);
 
 	 /**
-	  * Called when paths containing additional localization target data (LocRes) are unmounted, to allow the display strings to dynamically update without waiting for a refresh.
-	  * @note The loading is async, see WaitUntilLoadingCompletes.
+	  * Called when paths containing additional localization target data (LocRes) are unmounted, to allow the display strings to be unloaded.
+	  * @note The unloading is async, see WaitUntilLoadingCompletes.
 	  * @see FCoreDelegates::GatherAdditionalLocResPathsCallback.
 	  */
 	CORE_API void HandleLocalizationTargetsUnmounted(TArrayView<const FString> LocalizationTargetPaths);
diff --git a/Engine/Source/Runtime/Core/Public/Internationalization/TextLocalizationResource.h b/Engine/Source/Runtime/Core/Public/Internationalization/TextLocalizationResource.h
--- a/Engine/Source/Runtime/Core/Public/Internationalization/TextLocalizationResource.h
+++ b/Engine/Source/Runtime/Core/Public/Internationalization/TextLocalizationResource.h
@@ -65,6 +65,7 @@ public:
 	{
 		FTextConstDisplayStringPtr LocalizedString;
 		FTextKey LocResID;
+		int32 LocalizationTargetPathId = INDEX_NONE;
 		uint32 SourceStringHash = 0;
 		int32 Priority = 0; // Smaller numbers are higher priority
 	};
diff --git a/Engine/Source/Runtime/Projects/Private/PluginManager.cpp b/Engine/Source/Runtime/Projects/Private/PluginManager.cpp
--- a/Engine/Source/Runtime/Projects/Private/PluginManager.cpp
+++ b/Engine/Source/Runtime/Projects/Private/PluginManager.cpp
@@ -3267,6 +3267,56 @@ bool FPluginManager::MountExplicitlyLoadedPluginLocalizationData(const FString&
 	return true;
 }
 
+bool FPluginManager::UnmountExplicitlyLoadedPluginLocalizationData(const FString& PluginName)
+{
+	TSharedPtr<FPlugin> Plugin = FindPluginInstance(PluginName);
+
+	TRACE_CPUPROFILER_EVENT_SCOPE(UnmountExplicitlyLoadedPluginLocalizationData);
+	if (!Plugin.IsValid() || !Plugin->bIsMounted)
+	{
+		// Does not exist or is not mounted
+		UE_LOG(LogPluginManager, Error, TEXT("Cannot unmount plugin localization for '%s' as the plugin is unknown or not mounted. Did you forget to call MountExplicitlyLoadedPlugin?"), *PluginName);
+		return false;
+	}
+
+	if (!Plugin->Descriptor.bExplicitlyLoaded)
+	{
+		// Not supported
+		UE_LOG(LogPluginManager, Warning, TEXT("Cannot unmount plugin localization for '%s' as the plugin isn't explicitly loaded."), *PluginName);
+		return false;
+	}
+
+	if (!Plugin->bIsExplicitlyLoadedLocalizationDataMounted)
+	{
+		// Already unloaded
+		UE_LOG(LogPluginManager, Verbose, TEXT("Ignoring request to unmount plugin localization for '%s' as the localization data was not mounted."), *PluginName);
+		return false;
+	}
+
+	if (Plugin->Descriptor.LocalizationTargets.Num() == 0)
+	{
+		// Nothing to load
+		UE_LOG(LogPluginManager, Verbose, TEXT("Ignoring request to unmount plugin localization for '%s' as the plugin has no localization targets defined."), *PluginName);
+		return false;
+	}
+
+	if (IsEngineExitRequested())
+	{
+		// Exiting; unloading will be handled by the engine shutdown
+		UE_LOG(LogPluginManager, Verbose, TEXT("Ignoring request to unmount plugin localization for '%s' as the engine is exiting."), *PluginName);
+		return false;
+	}
+
+	UE_LOG(LogPluginManager, Log, TEXT("Unmounting plugin localization for '%s'..."), *PluginName);
+	Plugin->bIsExplicitlyLoadedLocalizationDataMounted = false;
+
+	// Notify that additional localization data should be unloaded
+	TArray<FString> AdditionalLocResPaths;
+	PluginLocalizationUtils::GetLocalizationPathsForPlugin(*Plugin, AdditionalLocResPaths);
+	FTextLocalizationManager::Get().HandleLocalizationTargetsUnmounted(AdditionalLocResPaths);
+	return true;
+}
+
 bool FPluginManager::TryMountExplicitlyLoadedPluginVersion(TSharedRef<FPlugin>* AllPlugins_PluginPtr)
 {
 	bool bSuccess = false;
@@ -3449,13 +3499,14 @@ bool FPluginManager::UnmountExplicitlyLoadedPlugin(const FString& PluginName, FT
 	}
 
 	// Notify that additional localization data should be unloaded
-	if (Plugin->Descriptor.LocalizationTargets.Num() > 0 && !IsEngineExitRequested())
+	if (Plugin->bIsExplicitlyLoadedLocalizationDataMounted && Plugin->Descriptor.LocalizationTargets.Num() > 0 && !IsEngineExitRequested())
 	{
+		Plugin->bIsExplicitlyLoadedLocalizationDataMounted = false;
+
 		TArray<FString> AdditionalLocResPaths;
 		PluginLocalizationUtils::GetLocalizationPathsForPlugin(*Plugin, AdditionalLocResPaths);
 		FTextLocalizationManager::Get().HandleLocalizationTargetsUnmounted(AdditionalLocResPaths);
 	}
-	Plugin->bIsExplicitlyLoadedLocalizationDataMounted = false;
 
 	if ((Plugin->CanContainContent() || Plugin->CanContainVerse()) && ensure(UnRegisterMountPointDelegate.IsBound()))
 	{
diff --git a/Engine/Source/Runtime/Projects/Private/PluginManager.h b/Engine/Source/Runtime/Projects/Private/PluginManager.h
--- a/Engine/Source/Runtime/Projects/Private/PluginManager.h
+++ b/Engine/Source/Runtime/Projects/Private/PluginManager.h
@@ -200,6 +200,7 @@ public:
 	virtual bool MountExplicitlyLoadedPlugin_FromFileName(const FString& PluginFileName) override;
 	virtual bool MountExplicitlyLoadedPlugin_FromDescriptor(const FPluginReferenceDescriptor& PluginDescriptor) override;
 	virtual bool MountExplicitlyLoadedPluginLocalizationData(const FString& PluginName) override;
+	virtual bool UnmountExplicitlyLoadedPluginLocalizationData(const FString& PluginName) override;
 	virtual bool UnmountExplicitlyLoadedPlugin(const FString& PluginName, FText* OutReason) override { return UnmountExplicitlyLoadedPlugin(PluginName, OutReason, true); }
 	virtual bool UnmountExplicitlyLoadedPlugin(const FString& PluginName, FText* OutReason, bool bAllowUnloadCode) override;
 	virtual bool GetPluginDependencies(const FString& PluginName, TArray<FPluginReferenceDescriptor>& PluginDependencies) override;
diff --git a/Engine/Source/Runtime/Projects/Public/Interfaces/IPluginManager.h b/Engine/Source/Runtime/Projects/Public/Interfaces/IPluginManager.h
--- a/Engine/Source/Runtime/Projects/Public/Interfaces/IPluginManager.h
+++ b/Engine/Source/Runtime/Projects/Public/Interfaces/IPluginManager.h
@@ -476,6 +476,13 @@ public:
 	 */
 	virtual bool MountExplicitlyLoadedPluginLocalizationData(const FString& PluginName) = 0;
 
+	/**
+	 * Start unloading localization data for an explicitly loaded plugin that had its localization data mounted via MountExplicitlyLoadedPluginLocalizationData.
+	 * @note Localization data is also automatically unloaded when calling UnmountExplicitlyLoadedPlugin.
+	 * @return True if localization data started to unload, or false if the plugin was missing or had no localization data to unload.
+	 */
+	virtual bool UnmountExplicitlyLoadedPluginLocalizationData(const FString& PluginName) = 0;
+
 	/**
 	 * Marks an explicitly loaded plugin as disabled, unmounts its content (does not work on plugins with compiled modules).
 	 */
