Added FMath::FloorLog2NonZero to enable optimizing callers where the value is known to be non-zero

diff --git a/Engine/Source/Runtime/Core/Public/Clang/ClangPlatformMath.h b/Engine/Source/Runtime/Core/Public/Clang/ClangPlatformMath.h
--- a/Engine/Source/Runtime/Core/Public/Clang/ClangPlatformMath.h
+++ b/Engine/Source/Runtime/Core/Public/Clang/ClangPlatformMath.h
@@ -94,11 +94,21 @@ struct FClangPlatformMath : public FGenericPlatformMath
 		return 31 - __builtin_clz(Value | 1);
 	}
 
+	static FORCEINLINE uint32 FloorLog2NonZero(uint32 Value)
+	{
+		return 31 - __builtin_clz(Value);
+	}
+
 	static FORCEINLINE uint64 FloorLog2_64(uint64 Value)
 	{
 		return 63 - __builtin_clzll(Value | 1);
 	}
-	
+
+	static FORCEINLINE uint64 FloorLog2NonZero_64(uint64 Value)
+	{
+		return 63 - __builtin_clzll(Value);
+	}
+
 	/**
 	 * Adds two integers of any integer type, checking for overflow.
 	 * If there was overflow, it returns false, and OutResult may or may not be written.
diff --git a/Engine/Source/Runtime/Core/Public/GenericPlatform/GenericPlatformMath.h b/Engine/Source/Runtime/Core/Public/GenericPlatform/GenericPlatformMath.h
--- a/Engine/Source/Runtime/Core/Public/GenericPlatform/GenericPlatformMath.h
+++ b/Engine/Source/Runtime/Core/Public/GenericPlatform/GenericPlatformMath.h
@@ -657,6 +657,12 @@ struct FGenericPlatformMath
 		return pos;
 	}
 
+	/** FloorLog2 but the caller guarantees that Value is not 0. */
+	static constexpr FORCEINLINE uint32 FloorLog2NonZero(uint32 Value)
+	{
+		return FloorLog2(Value);
+	}
+
 	/**
 	 * Computes the base 2 logarithm for a 64-bit value.
 	 * The result is rounded down to the nearest integer.
@@ -676,6 +682,12 @@ struct FGenericPlatformMath
 		return pos;
 	}
 
+	/** FloorLog2_64 but the caller guarantees that Value is not 0. */
+	static constexpr FORCEINLINE uint64 FloorLog2NonZero_64(uint64 Value)
+	{
+		return FloorLog2_64(Value);
+	}
+
 	/**
 	 * Counts the number of leading zeros in the bit representation of the 8-bit value
 	 *
diff --git a/Engine/Source/Runtime/Core/Public/Microsoft/MicrosoftPlatformMath.h b/Engine/Source/Runtime/Core/Public/Microsoft/MicrosoftPlatformMath.h
--- a/Engine/Source/Runtime/Core/Public/Microsoft/MicrosoftPlatformMath.h
+++ b/Engine/Source/Runtime/Core/Public/Microsoft/MicrosoftPlatformMath.h
@@ -26,6 +26,14 @@ struct FMicrosoftPlatformMathBase : public TUnrealPlatformMathSSE4Base<FGenericP
 		unsigned long BitIndex;
 		return _BitScanReverse(&BitIndex, Value) ? BitIndex : 0;
 	}
+
+	static FORCEINLINE uint32 FloorLog2NonZero(uint32 Value)
+	{
+		unsigned long BitIndex = 0;
+		_BitScanReverse(&BitIndex, Value);
+		return BitIndex;
+	}
+
 	static FORCEINLINE uint8 CountLeadingZeros8(uint8 Value)
 	{
 		unsigned long BitIndex;
@@ -67,6 +75,13 @@ struct FMicrosoftPlatformMathBase : public TUnrealPlatformMathSSE4Base<FGenericP
 		return _BitScanReverse64(&BitIndex, Value) ? BitIndex : 0;
 	}
 
+	static FORCEINLINE uint64 FloorLog2NonZero_64(uint64 Value)
+	{
+		unsigned long BitIndex = 0;
+		_BitScanReverse64(&BitIndex, Value);
+		return BitIndex;
+	}
+
 	static FORCEINLINE uint64 CeilLogTwo64(uint64 Arg)
 	{
 		// if Arg is 0, change it to 1 so that we return 0
