Explicitly call aggregate initializer

#rnx

diff --git a/Engine/Source/Runtime/Core/Private/Internationalization/TextLocalizationManager.cpp b/Engine/Source/Runtime/Core/Private/Internationalization/TextLocalizationManager.cpp
--- a/Engine/Source/Runtime/Core/Private/Internationalization/TextLocalizationManager.cpp
+++ b/Engine/Source/Runtime/Core/Private/Internationalization/TextLocalizationManager.cpp
@@ -541,7 +541,7 @@ FTextLocalizationManager::FDisplayStringsForLocalizationTarget& FTextLocalizatio
 	int32 LocalizationTargetPathId = LocalizationTargetPathsToIds.FindRef(NormalizedLocalizationTargetPath, INDEX_NONE);
 	if (LocalizationTargetPathId == INDEX_NONE)
 	{
-		LocalizationTargetPathId = LocalizationTargets.Emplace(MoveTemp(NormalizedLocalizationTargetPath));
+		LocalizationTargetPathId = LocalizationTargets.Emplace(FDisplayStringsForLocalizationTarget{ MoveTemp(NormalizedLocalizationTargetPath) });
 		LocalizationTargetPathsToIds.Add(LocalizationTargets[LocalizationTargetPathId].LocalizationTargetPath, LocalizationTargetPathId);
 	}
 
