Add an ensure when calculating UnTrackedMemory. Ensure it never reports a value below 0.

diff --git a/Engine/Source/Runtime/Core/Private/HAL/LowLevelMemTracker.cpp b/Engine/Source/Runtime/Core/Private/HAL/LowLevelMemTracker.cpp
--- a/Engine/Source/Runtime/Core/Private/HAL/LowLevelMemTracker.cpp
+++ b/Engine/Source/Runtime/Core/Private/HAL/LowLevelMemTracker.cpp
@@ -2015,6 +2015,7 @@ void FLowLevelMemTracker::TickInternal()
 
 	// Compare memory the platform thinks we have allocated to what we have tracked, including the program memory
 	const int64 PlatformTrackedTotal = PlatformTracker.GetTrackedTotal();
+	ensureMsgf(PlatformProcessMemory >= PlatformTrackedTotal, TEXT("Untracked Memory is negative. This is wrong and should be investigated."));
 	MemoryUsagePlatformTotalUntracked = FMath::Max<int64>(0, PlatformProcessMemory - PlatformTrackedTotal);
 
 	PlatformTracker.SetTagAmountInUpdate(FindOrAddTagData(ELLMTag::PlatformTotal), PlatformProcessMemory, false);
