Fixed TStaticArray failing to build under /permissive- under MSVC.

diff --git a/Engine/Source/Runtime/Core/Public/Containers/StaticArray.h b/Engine/Source/Runtime/Core/Public/Containers/StaticArray.h
--- a/Engine/Source/Runtime/Core/Public/Containers/StaticArray.h
+++ b/Engine/Source/Runtime/Core/Public/Containers/StaticArray.h
@@ -10,6 +10,14 @@
 #include "Delegates/IntegerSequence.h"
 #include "Templates/TypeHash.h"
 
+namespace UE::Core::Private
+{
+	// This is a workaround for a parsing error in MSVC under /persmissive- builds, which would
+	// get confused by the fold expression in the constraint in the constructor.
+	template <typename InElementType, typename... ArgTypes>
+	constexpr bool TCanBeConvertedToFromAll_V = (std::is_convertible_v<ArgTypes, InElementType> && ...);
+}
+
 /** An array with a static number of elements. */
 template <typename InElementType, uint32 NumElements, uint32 Alignment = alignof(InElementType)>
 class alignas(Alignment) TStaticArray
@@ -29,7 +37,7 @@ public:
 	// Directly initializes the array with the provided values.
 	template <
 		typename... ArgTypes
-		UE_REQUIRES((sizeof...(ArgTypes) > 0 && sizeof...(ArgTypes) <= NumElements) && (std::is_convertible_v<ArgTypes, InElementType> && ...))
+		UE_REQUIRES((sizeof...(ArgTypes) > 0 && sizeof...(ArgTypes) <= NumElements) && UE::Core::Private::TCanBeConvertedToFromAll_V<InElementType, ArgTypes...>)
 	>
 	TStaticArray(ArgTypes&&... Args)
 		: Storage(PerElement, Forward<ArgTypes>(Args)...)
