Core - Added [[nodiscard]] to most GetTypeHash() overloads in Core

diff --git a/Engine/Source/Runtime/Core/Public/Misc/Guid.h b/Engine/Source/Runtime/Core/Public/Misc/Guid.h
--- a/Engine/Source/Runtime/Core/Public/Misc/Guid.h
+++ b/Engine/Source/Runtime/Core/Public/Misc/Guid.h
@@ -348,7 +348,7 @@ public:
 	 * @param Guid The GUID to calculate the hash for.
 	 * @return The hash.
 	 */
-	friend uint32 GetTypeHash(const FGuid& Guid)
+	[[nodiscard]] friend uint32 GetTypeHash(const FGuid& Guid)
 	{
 		return uint32(CityHash64((char*)&Guid, sizeof(FGuid)));
 	}
diff --git a/Engine/Source/Runtime/Core/Public/UObject/NameTypes.h b/Engine/Source/Runtime/Core/Public/UObject/NameTypes.h
--- a/Engine/Source/Runtime/Core/Public/UObject/NameTypes.h
+++ b/Engine/Source/Runtime/Core/Public/UObject/NameTypes.h
@@ -1744,7 +1744,7 @@ private:
 #endif
 	friend bool operator==(FNameEntryId A, FDisplayNameEntryId B) { return A == B.GetDisplayId(); }
 	friend bool operator==(FDisplayNameEntryId A, FNameEntryId B) { return A.GetDisplayId() == B; }
-	friend uint32 GetTypeHash(FDisplayNameEntryId InId) { return GetTypeHash(InId.GetDisplayId()); }
+	[[nodiscard]] friend uint32 GetTypeHash(FDisplayNameEntryId InId) { return GetTypeHash(InId.GetDisplayId()); }
 
 public: // Internal functions for batch serialization code - intentionally lacking CORE_API
 	static FDisplayNameEntryId FromComparisonId(FNameEntryId ComparisonId);
diff --git a/Engine/Source/Runtime/Core/Public/UObject/ScriptDelegates.h b/Engine/Source/Runtime/Core/Public/UObject/ScriptDelegates.h
--- a/Engine/Source/Runtime/Core/Public/UObject/ScriptDelegates.h
+++ b/Engine/Source/Runtime/Core/Public/UObject/ScriptDelegates.h
@@ -447,7 +447,7 @@ public:
 		ObjectPtr->ProcessEvent(Function, Parameters);
 	}
 
-	friend uint32 GetTypeHash(const TScriptDelegate& Delegate)
+	[[nodiscard]] friend uint32 GetTypeHash(const TScriptDelegate& Delegate)
 	{
 		FReadAccessScope ReadScope = Delegate.GetReadAccessScope();
 
diff --git a/Engine/Source/Runtime/CoreUObject/Private/UObject/PropertyPathName.cpp b/Engine/Source/Runtime/CoreUObject/Private/UObject/PropertyPathName.cpp
--- a/Engine/Source/Runtime/CoreUObject/Private/UObject/PropertyPathName.cpp
+++ b/Engine/Source/Runtime/CoreUObject/Private/UObject/PropertyPathName.cpp
@@ -92,7 +92,7 @@ void FPropertyPathName::ToString(FStringBuilderBase& Out, FStringView Separator)
 	}
 }
 
-uint32 GetTypeHash(const FPropertyPathName& Path)
+[[nodiscard]] uint32 GetTypeHash(const FPropertyPathName& Path)
 {
 	uint32 Hash = 0;
 	for (const FPropertyPathName::FSegment& Segment : Path.Segments)
diff --git a/Engine/Source/Runtime/CoreUObject/Private/UObject/PropertyTypeName.cpp b/Engine/Source/Runtime/CoreUObject/Private/UObject/PropertyTypeName.cpp
--- a/Engine/Source/Runtime/CoreUObject/Private/UObject/PropertyTypeName.cpp
+++ b/Engine/Source/Runtime/CoreUObject/Private/UObject/PropertyTypeName.cpp
@@ -275,7 +275,7 @@ bool FPropertyTypeName::IsEnum(FName EnumName) const
 	return First->InnerCount > 0 && (First->Name == NAME_EnumProperty || First->Name == NAME_ByteProperty) && First[1].Name == EnumName;
 }
 
-uint32 GetTypeHash(const FPropertyTypeName& TypeName)
+[[nodiscard]] uint32 GetTypeHash(const FPropertyTypeName& TypeName)
 {
 	const FPropertyTypeNameNode* First = GPropertyTypeNameTable.ResolveByIndex(TypeName.Index);
 	return GetTypeHash(FPropertyTypeNameNodeProxy{First});
diff --git a/Engine/Source/Runtime/CoreUObject/Public/Templates/SubclassOf.h b/Engine/Source/Runtime/CoreUObject/Public/Templates/SubclassOf.h
--- a/Engine/Source/Runtime/CoreUObject/Public/Templates/SubclassOf.h
+++ b/Engine/Source/Runtime/CoreUObject/Public/Templates/SubclassOf.h
@@ -159,7 +159,7 @@ public:
 		Slot << Class;
 	}
 
-	friend uint32 GetTypeHash(const TSubclassOf& SubclassOf)
+	[[nodiscard]] friend uint32 GetTypeHash(const TSubclassOf& SubclassOf)
 	{
 		return GetTypeHash(SubclassOf.Class);
 	}
diff --git a/Engine/Source/Runtime/CoreUObject/Public/UObject/DynamicallyTypedValue.h b/Engine/Source/Runtime/CoreUObject/Public/UObject/DynamicallyTypedValue.h
--- a/Engine/Source/Runtime/CoreUObject/Public/UObject/DynamicallyTypedValue.h
+++ b/Engine/Source/Runtime/CoreUObject/Public/UObject/DynamicallyTypedValue.h
@@ -113,7 +113,7 @@ namespace UE
 		}
 
 		// Returns hash of the underlying FDynamicallyTypedValue's value. Added to allow for FDynamicallyTypedValue to be used as TMap keys.
-		friend uint32 GetTypeHash(const FDynamicallyTypedValue& DynamicallyTypedValue)
+		[[nodiscard]] friend uint32 GetTypeHash(const FDynamicallyTypedValue& DynamicallyTypedValue)
 		{
 			return DynamicallyTypedValue.GetType().GetValueHash(DynamicallyTypedValue.GetDataPointer());
 		}
diff --git a/Engine/Source/Runtime/CoreUObject/Public/UObject/Field.h b/Engine/Source/Runtime/CoreUObject/Public/UObject/Field.h
--- a/Engine/Source/Runtime/CoreUObject/Public/UObject/Field.h
+++ b/Engine/Source/Runtime/CoreUObject/Public/UObject/Field.h
@@ -435,7 +435,7 @@ public:
 #endif
 
 	/** Support comparison functions that make this usable as a KeyValue for a TSet<> */
-	friend uint32 GetTypeHash(const FFieldVariant& InFieldVariant)
+	[[nodiscard]] friend uint32 GetTypeHash(const FFieldVariant& InFieldVariant)
 	{
 		return GetTypeHash(InFieldVariant.GetRawPointer());
 	}
diff --git a/Engine/Source/Runtime/CoreUObject/Public/UObject/FieldPath.h b/Engine/Source/Runtime/CoreUObject/Public/UObject/FieldPath.h
--- a/Engine/Source/Runtime/CoreUObject/Public/UObject/FieldPath.h
+++ b/Engine/Source/Runtime/CoreUObject/Public/UObject/FieldPath.h
@@ -267,7 +267,7 @@ public:
 	COREUOBJECT_API friend FArchive& operator<<(FArchive& Ar, FFieldPath& InOutPropertyPath);
 
 	/** Hash function. */
-	FORCEINLINE friend uint32 GetTypeHash(const FFieldPath& InPropertyPath)
+	[[nodiscard]] FORCEINLINE friend uint32 GetTypeHash(const FFieldPath& InPropertyPath)
 	{
 		uint32 HashValue = 0;
 		for (const FName& PathSegment : InPropertyPath.Path)
diff --git a/Engine/Source/Runtime/CoreUObject/Public/UObject/ObjectKey.h b/Engine/Source/Runtime/CoreUObject/Public/UObject/ObjectKey.h
--- a/Engine/Source/Runtime/CoreUObject/Public/UObject/ObjectKey.h
+++ b/Engine/Source/Runtime/CoreUObject/Public/UObject/ObjectKey.h
@@ -134,7 +134,7 @@ public:
 	}
 
 	/** Hash function */
-	friend uint32 GetTypeHash(const FObjectKey& Key)
+	[[nodiscard]] friend uint32 GetTypeHash(const FObjectKey& Key)
 	{
 		return HashCombine(Key.ObjectIndex, Key.ObjectSerialNumber);
 	}
@@ -208,7 +208,7 @@ public:
 	}
 
 	//** Hash function */
-	friend uint32 GetTypeHash(const TObjectKey& Key)
+	[[nodiscard]] friend uint32 GetTypeHash(const TObjectKey& Key)
 	{
 		return GetTypeHash(Key.ObjectKey);
 	}
diff --git a/Engine/Source/Runtime/CoreUObject/Public/UObject/ObjectResource.h b/Engine/Source/Runtime/CoreUObject/Public/UObject/ObjectResource.h
--- a/Engine/Source/Runtime/CoreUObject/Public/UObject/ObjectResource.h
+++ b/Engine/Source/Runtime/CoreUObject/Public/UObject/ObjectResource.h
@@ -150,7 +150,7 @@ public:
 		Slot << Value.Index;
 	}
 
-	FORCEINLINE friend uint32 GetTypeHash(const FPackageIndex& In)
+	[[nodiscard]] FORCEINLINE friend uint32 GetTypeHash(const FPackageIndex& In)
 	{
 		return uint32(In.Index);
 	}
diff --git a/Engine/Source/Runtime/CoreUObject/Public/UObject/SoftObjectPath.h b/Engine/Source/Runtime/CoreUObject/Public/UObject/SoftObjectPath.h
--- a/Engine/Source/Runtime/CoreUObject/Public/UObject/SoftObjectPath.h
+++ b/Engine/Source/Runtime/CoreUObject/Public/UObject/SoftObjectPath.h
@@ -323,7 +323,7 @@ PRAGMA_ENABLE_DEPRECATION_WARNINGS;
 	/** Fixes soft object path for CoreRedirects to handle renamed native objects, returns true if it was modified */
 	COREUOBJECT_API bool FixupCoreRedirects();
 
-	FORCEINLINE friend uint32 GetTypeHash(FSoftObjectPath const& This)
+	[[nodiscard]] FORCEINLINE friend uint32 GetTypeHash(FSoftObjectPath const& This)
 	{
 		uint32 Hash = 0;
 
diff --git a/Engine/Source/Runtime/CoreUObject/Public/UObject/StrongObjectPtr.h b/Engine/Source/Runtime/CoreUObject/Public/UObject/StrongObjectPtr.h
--- a/Engine/Source/Runtime/CoreUObject/Public/UObject/StrongObjectPtr.h
+++ b/Engine/Source/Runtime/CoreUObject/Public/UObject/StrongObjectPtr.h
@@ -173,7 +173,7 @@ public:
 		}
 	}
 
-	FORCEINLINE_DEBUGGABLE friend uint32 GetTypeHash(const TStrongObjectPtr& InStrongObjectPtr)
+	[[nodiscard]] FORCEINLINE_DEBUGGABLE friend uint32 GetTypeHash(const TStrongObjectPtr& InStrongObjectPtr)
 	{
 		return GetTypeHash(InStrongObjectPtr.Get());
 	}
diff --git a/Engine/Source/Runtime/CoreUObject/Public/UObject/TopLevelAssetPath.h b/Engine/Source/Runtime/CoreUObject/Public/UObject/TopLevelAssetPath.h
--- a/Engine/Source/Runtime/CoreUObject/Public/UObject/TopLevelAssetPath.h
+++ b/Engine/Source/Runtime/CoreUObject/Public/UObject/TopLevelAssetPath.h
@@ -162,7 +162,7 @@ struct FTopLevelAssetPath
 		return AssetName.CompareIndexes(Other.AssetName);
 	}
 
-	friend uint32 GetTypeHash(FTopLevelAssetPath const& This)
+	[[nodiscard]] friend uint32 GetTypeHash(FTopLevelAssetPath const& This)
 	{
 		return HashCombineFast(GetTypeHash(This.PackageName), GetTypeHash(This.AssetName));
 	}
diff --git a/Engine/Source/Runtime/CoreUObject/Public/UObject/WeakFieldPtr.h b/Engine/Source/Runtime/CoreUObject/Public/UObject/WeakFieldPtr.h
--- a/Engine/Source/Runtime/CoreUObject/Public/UObject/WeakFieldPtr.h
+++ b/Engine/Source/Runtime/CoreUObject/Public/UObject/WeakFieldPtr.h
@@ -209,7 +209,7 @@ public:
 	}
 
 	/** Hash function. */
-	FORCEINLINE friend uint32 GetTypeHash(const TWeakFieldPtr& WeakObjectPtr)
+	[[nodiscard]] FORCEINLINE friend uint32 GetTypeHash(const TWeakFieldPtr& WeakObjectPtr)
 	{
 		return GetTypeHash(WeakObjectPtr.Field);
 	}
