[TaskGraph]
 - Add missing CORE_API export for WaitForAnyTaskCompleted and AnyTaskCompleted

diff --git a/Engine/Source/Runtime/Core/Public/Async/TaskGraphInterfaces.h b/Engine/Source/Runtime/Core/Public/Async/TaskGraphInterfaces.h
--- a/Engine/Source/Runtime/Core/Public/Async/TaskGraphInterfaces.h
+++ b/Engine/Source/Runtime/Core/Public/Async/TaskGraphInterfaces.h
@@ -1381,10 +1381,10 @@ private:
 // Blocks the current thread until any of the given tasks is completed.
 // Is slightly more efficient than `AnyTaskCompleted()->Wait()` and supports timeout while `FGraphEvent::Wait()` doesn't.
 // Returns the index of the first completed task, or `INDEX_NONE` on timeout.
-int32 WaitForAnyTaskCompleted(const FGraphEventArray& GraphEvents, FTimespan Timeout = FTimespan::MaxValue());
+CORE_API int32 WaitForAnyTaskCompleted(const FGraphEventArray& GraphEvents, FTimespan Timeout = FTimespan::MaxValue());
 
 // Returns a graph event that gets completed as soon as any of the given tasks gets completed
-FGraphEventRef AnyTaskCompleted(const FGraphEventArray& GraphEvents);
+CORE_API FGraphEventRef AnyTaskCompleted(const FGraphEventArray& GraphEvents);
 
 /** 
  *	FReturnGraphTask is a task used to return flow control from a named thread back to the original caller of ProcessThreadUntilRequestReturn
