[Scheduler]
 - Remove thread pinning on cores in secondary thread groups for hardware with more than 64 cores

diff --git a/Engine/Source/Runtime/Core/Private/Async/Fundamental/Scheduler.cpp b/Engine/Source/Runtime/Core/Private/Async/Fundamental/Scheduler.cpp
--- a/Engine/Source/Runtime/Core/Private/Async/Fundamental/Scheduler.cpp
+++ b/Engine/Source/Runtime/Core/Private/Async/Fundamental/Scheduler.cpp
@@ -57,9 +57,9 @@ namespace LowLevelTasks
 			uint32 CpusInGroup = FMath::CountBits(ProcessorGroups.ThreadAffinities[GroupIndex]);
 			if (GroupWorkerId < CpusInGroup)
 			{
-				if (CpuGroup != 0) //pin larger groups workers to a core and leave first group as is for legacy reasons
+				if (CpuGroup != 0) // don't pin larger groups workers to a core and leave first group as is for legacy reasons
 				{
-					ThreadAffinityMask = 1ull << GroupWorkerId;
+					ThreadAffinityMask = MAX_uint64;
 				}
 				break;
 			}
