Added UE_REQUIRES_EXPR for wrapping up C++20 requires expressions.
Fixed up TSharedPtr and TUniquePtr to use this.

diff --git a/Engine/Source/Runtime/Core/Public/Templates/Requires.h b/Engine/Source/Runtime/Core/Public/Templates/Requires.h
--- a/Engine/Source/Runtime/Core/Public/Templates/Requires.h
+++ b/Engine/Source/Runtime/Core/Public/Templates/Requires.h
@@ -15,9 +15,38 @@
 		UE_REQUIRES(std::is_integral_v<T> && sizeof(U) <= 4)
 	>
 	void IntegralUpTo32Bit(T Lhs, U Rhs) {}
+
+	UE_REQUIRES_EXPR() wraps the effects of a requires expression, used to test the
+	compilability of an expression based on a deduced template parameter.  Usage:
+
+	template <typename Type>
+	struct TSmartPtr
+	{
+		// Only enable this constructor if the incoming pointer type is a UObject
+		// convertible to the smart pointer type:
+		template <
+			typename OtherType
+			UE_REQUIRES(
+				std::is_base_of_v<UObject, OtherType> &&
+				UE_REQUIRES_EXPR(ImplicitConv<Type*>((OtherType*)nullptr)))
+			)
+		>
+		explicit TSmartPtr(OtherType* OtherPtr)
+			: Ptr(Cast<Type>(OtherPtr))
+		{
+		}
+
+	private:
+		BaseType* Ptr;
+	};
+
+	Unlike a C++20 requires expression, UE_REQUIRES_EXPR() can only be used in
+	the body of a UE_REQUIRES() macro - standalone concept checks must still be
+	done via something like a TModels-type concept.
  -----------------------------------------------------------------------------*/
 #if __cplusplus < 202000
 	#define UE_REQUIRES(...) , std::enable_if_t<(__VA_ARGS__), int> = 0
+    #define UE_REQUIRES_EXPR(...) ((decltype((void)(__VA_ARGS__), 0))0 == 0)
 #else
 	namespace UE::Core::Private
 	{
@@ -27,6 +56,7 @@
 	}
 
 	#define UE_REQUIRES(...) > requires (!!(__VA_ARGS__)) && UE::Core::Private::BoolIdentityConcept<true
+    #define UE_REQUIRES_EXPR(...) requires { (__VA_ARGS__); }
 #endif
 
 // This should be regarded as deprecated - please use UE_REQUIRES instead.
diff --git a/Engine/Source/Runtime/Core/Public/Templates/SharedPointer.h b/Engine/Source/Runtime/Core/Public/Templates/SharedPointer.h
--- a/Engine/Source/Runtime/Core/Public/Templates/SharedPointer.h
+++ b/Engine/Source/Runtime/Core/Public/Templates/SharedPointer.h
@@ -175,8 +175,8 @@ public:
 	 * @param  InObject  Object this shared reference to retain a reference to
 	 */
 	template <
-		typename OtherType,
-		typename = decltype(ImplicitConv<ObjectType*>((OtherType*)nullptr))
+		typename OtherType
+		UE_REQUIRES(UE_REQUIRES_EXPR(ImplicitConv<ObjectType*>((OtherType*)nullptr)))
 	>
 	FORCEINLINE explicit TSharedRef( OtherType* InObject )
 		: Object( InObject )
@@ -195,8 +195,8 @@ public:
 	 */
 	template <
 		typename OtherType,
-		typename DeleterType,
-		typename = decltype(ImplicitConv<ObjectType*>((OtherType*)nullptr))
+		typename DeleterType
+		UE_REQUIRES(UE_REQUIRES_EXPR(ImplicitConv<ObjectType*>((OtherType*)nullptr)))
 	>
 	FORCEINLINE TSharedRef( OtherType* InObject, DeleterType&& InDeleter )
 		: Object( InObject )
@@ -228,8 +228,8 @@ public:
 	 */
 	// NOTE: The following is an Unreal extension to standard shared_ptr behavior
 	template <
-		typename OtherType,
-		typename = decltype(ImplicitConv<ObjectType*>((OtherType*)nullptr))
+		typename OtherType
+		UE_REQUIRES(UE_REQUIRES_EXPR(ImplicitConv<ObjectType*>((OtherType*)nullptr)))
 	>
 	FORCEINLINE TSharedRef( SharedPointerInternals::TRawPtrProxy< OtherType > const& InRawPtrProxy )
 		: Object( InRawPtrProxy.Object )
@@ -255,8 +255,8 @@ public:
 	// NOTE: The following is an Unreal extension to standard shared_ptr behavior
 	template <
 		typename OtherType,
-		typename DeleterType,
-		typename = decltype(ImplicitConv<ObjectType*>((OtherType*)nullptr))
+		typename DeleterType
+		UE_REQUIRES(UE_REQUIRES_EXPR(ImplicitConv<ObjectType*>((OtherType*)nullptr)))
 	>
 	FORCEINLINE TSharedRef( SharedPointerInternals::TRawPtrProxyWithDeleter< OtherType, DeleterType > const& InRawPtrProxy )
 		: Object( InRawPtrProxy.Object )
@@ -282,8 +282,8 @@ public:
 	// NOTE: The following is an Unreal extension to standard shared_ptr behavior
 	template <
 		typename OtherType,
-		typename DeleterType,
-		typename = decltype(ImplicitConv<ObjectType*>((OtherType*)nullptr))
+		typename DeleterType
+		UE_REQUIRES(UE_REQUIRES_EXPR(ImplicitConv<ObjectType*>((OtherType*)nullptr)))
 	>
 	FORCEINLINE TSharedRef( SharedPointerInternals::TRawPtrProxyWithDeleter< OtherType, DeleterType >&& InRawPtrProxy )
 		: Object( InRawPtrProxy.Object )
@@ -307,8 +307,8 @@ public:
 	 * @param  InSharedRef  The shared reference whose object we should create an additional reference to
 	 */
 	template <
-		typename OtherType,
-		typename = decltype(ImplicitConv<ObjectType*>((OtherType*)nullptr))
+		typename OtherType
+		UE_REQUIRES(UE_REQUIRES_EXPR(ImplicitConv<ObjectType*>((OtherType*)nullptr)))
 	>
 	FORCEINLINE TSharedRef( TSharedRef< OtherType, Mode > const& InSharedRef )
 		: Object( InSharedRef.Object )
@@ -406,8 +406,8 @@ public:
 	 */
 	// NOTE: The following is an Unreal extension to standard shared_ptr behavior
 	template <
-		typename OtherType,
-		typename = decltype(ImplicitConv<ObjectType*>((OtherType*)nullptr))
+		typename OtherType
+		UE_REQUIRES(UE_REQUIRES_EXPR(ImplicitConv<ObjectType*>((OtherType*)nullptr)))
 	>
 	FORCEINLINE TSharedRef& operator=( SharedPointerInternals::TRawPtrProxy< OtherType > const& InRawPtrProxy )
 	{
@@ -429,8 +429,8 @@ public:
 	// NOTE: The following is an Unreal extension to standard shared_ptr behavior
 	template <
 		typename OtherType,
-		typename DeleterType,
-		typename = decltype(ImplicitConv<ObjectType*>((OtherType*)nullptr))
+		typename DeleterType
+		UE_REQUIRES(UE_REQUIRES_EXPR(ImplicitConv<ObjectType*>((OtherType*)nullptr)))
 	>
 	FORCEINLINE TSharedRef& operator=( SharedPointerInternals::TRawPtrProxyWithDeleter< OtherType, DeleterType > const& InRawPtrProxy )
 	{
@@ -452,8 +452,8 @@ public:
 	// NOTE: The following is an Unreal extension to standard shared_ptr behavior
 	template <
 		typename OtherType,
-		typename DeleterType,
-		typename = decltype(ImplicitConv<ObjectType*>((OtherType*)nullptr))
+		typename DeleterType
+		UE_REQUIRES(UE_REQUIRES_EXPR(ImplicitConv<ObjectType*>((OtherType*)nullptr)))
 	>
 	FORCEINLINE TSharedRef& operator=( SharedPointerInternals::TRawPtrProxyWithDeleter< OtherType, DeleterType >&& InRawPtrProxy )
 	{
@@ -591,8 +591,8 @@ private:
 	 * @return  Reference to the object
 	 */
 	template <
-		typename OtherType,
-		typename = decltype(ImplicitConv<ObjectType*>((OtherType*)nullptr))
+		typename OtherType
+		UE_REQUIRES(UE_REQUIRES_EXPR(ImplicitConv<ObjectType*>((OtherType*)nullptr)))
 	>
 	FORCEINLINE explicit TSharedRef( TSharedPtr< OtherType, Mode > const& InSharedPtr )
 		: Object( InSharedPtr.Object )
@@ -604,8 +604,8 @@ private:
 	}
 
 	template <
-		typename OtherType,
-		typename = decltype(ImplicitConv<ObjectType*>((OtherType*)nullptr))
+		typename OtherType
+		UE_REQUIRES(UE_REQUIRES_EXPR(ImplicitConv<ObjectType*>((OtherType*)nullptr)))
 	>
 	FORCEINLINE explicit TSharedRef( TSharedPtr< OtherType, Mode >&& InSharedPtr )
 		: Object( InSharedPtr.Object )
@@ -711,8 +711,8 @@ public:
 	 * @param  InObject  Object this shared pointer to retain a reference to
 	 */
 	template <
-		typename OtherType,
-		typename = decltype(ImplicitConv<ObjectType*>((OtherType*)nullptr))
+		typename OtherType
+		UE_REQUIRES(UE_REQUIRES_EXPR(ImplicitConv<ObjectType*>((OtherType*)nullptr)))
 	>
 	FORCEINLINE explicit TSharedPtr( OtherType* InObject )
 		: Object( InObject )
@@ -734,8 +734,8 @@ public:
 	 */
 	template <
 		typename OtherType,
-		typename DeleterType,
-		typename = decltype(ImplicitConv<ObjectType*>((OtherType*)nullptr))
+		typename DeleterType
+		UE_REQUIRES(UE_REQUIRES_EXPR(ImplicitConv<ObjectType*>((OtherType*)nullptr)))
 	>
 	FORCEINLINE TSharedPtr( OtherType* InObject, DeleterType&& InDeleter )
 		: Object( InObject )
@@ -755,8 +755,8 @@ public:
 	 */
 	// NOTE: The following is an Unreal extension to standard shared_ptr behavior
 	template <
-		typename OtherType,
-		typename = decltype(ImplicitConv<ObjectType*>((OtherType*)nullptr))
+		typename OtherType
+		UE_REQUIRES(UE_REQUIRES_EXPR(ImplicitConv<ObjectType*>((OtherType*)nullptr)))
 	>
 	FORCEINLINE TSharedPtr( SharedPointerInternals::TRawPtrProxy< OtherType > const& InRawPtrProxy )
 		: Object( InRawPtrProxy.Object )
@@ -777,8 +777,8 @@ public:
 	// NOTE: The following is an Unreal extension to standard shared_ptr behavior
 	template <
 		typename OtherType,
-		typename DeleterType,
-		typename = decltype(ImplicitConv<ObjectType*>((OtherType*)nullptr))
+		typename DeleterType
+		UE_REQUIRES(UE_REQUIRES_EXPR(ImplicitConv<ObjectType*>((OtherType*)nullptr)))
 	>
 	FORCEINLINE TSharedPtr( SharedPointerInternals::TRawPtrProxyWithDeleter< OtherType, DeleterType > const& InRawPtrProxy )
 		: Object( InRawPtrProxy.Object )
@@ -799,8 +799,8 @@ public:
 	// NOTE: The following is an Unreal extension to standard shared_ptr behavior
 	template <
 		typename OtherType,
-		typename DeleterType,
-		typename = decltype(ImplicitConv<ObjectType*>((OtherType*)nullptr))
+		typename DeleterType
+		UE_REQUIRES(UE_REQUIRES_EXPR(ImplicitConv<ObjectType*>((OtherType*)nullptr)))
 	>
 	FORCEINLINE TSharedPtr( SharedPointerInternals::TRawPtrProxyWithDeleter< OtherType, DeleterType >&& InRawPtrProxy )
 		: Object( InRawPtrProxy.Object )
@@ -820,8 +820,8 @@ public:
 	 * @param  InSharedPtr  The shared pointer whose object we should create an additional reference to
 	 */
 	template <
-		typename OtherType,
-		typename = decltype(ImplicitConv<ObjectType*>((OtherType*)nullptr))
+		typename OtherType
+		UE_REQUIRES(UE_REQUIRES_EXPR(ImplicitConv<ObjectType*>((OtherType*)nullptr)))
 	>
 	FORCEINLINE TSharedPtr( TSharedPtr< OtherType, Mode > const& InSharedPtr )
 		: Object( InSharedPtr.Object )
@@ -850,8 +850,8 @@ public:
 	 */
 	// NOTE: The following is an Unreal extension to standard shared_ptr behavior
 	template <
-		typename OtherType,
-		typename = decltype(ImplicitConv<ObjectType*>((OtherType*)nullptr))
+		typename OtherType
+		UE_REQUIRES(UE_REQUIRES_EXPR(ImplicitConv<ObjectType*>((OtherType*)nullptr)))
 	>
 	FORCEINLINE TSharedPtr( TSharedRef< OtherType, Mode > const& InSharedRef )
 		: Object( InSharedRef.Object )
@@ -979,8 +979,8 @@ public:
 	 */
 	// NOTE: The following is an Unreal extension to standard shared_ptr behavior
 	template <
-		typename OtherType,
-		typename = decltype(ImplicitConv<ObjectType*>((OtherType*)nullptr))
+		typename OtherType
+		UE_REQUIRES(UE_REQUIRES_EXPR(ImplicitConv<ObjectType*>((OtherType*)nullptr)))
 	>
 	FORCEINLINE TSharedPtr& operator=( SharedPointerInternals::TRawPtrProxy< OtherType > const& InRawPtrProxy )
 	{
@@ -998,8 +998,8 @@ public:
 	// NOTE: The following is an Unreal extension to standard shared_ptr behavior
 	template <
 		typename OtherType,
-		typename DeleterType,
-		typename = decltype(ImplicitConv<ObjectType*>((OtherType*)nullptr))
+		typename DeleterType
+		UE_REQUIRES(UE_REQUIRES_EXPR(ImplicitConv<ObjectType*>((OtherType*)nullptr)))
 	>
 	FORCEINLINE TSharedPtr& operator=( SharedPointerInternals::TRawPtrProxyWithDeleter< OtherType, DeleterType > const& InRawPtrProxy )
 	{
@@ -1017,8 +1017,8 @@ public:
 	// NOTE: The following is an Unreal extension to standard shared_ptr behavior
 	template <
 		typename OtherType,
-		typename DeleterType,
-		typename = decltype(ImplicitConv<ObjectType*>((OtherType*)nullptr))
+		typename DeleterType
+		UE_REQUIRES(UE_REQUIRES_EXPR(ImplicitConv<ObjectType*>((OtherType*)nullptr)))
 	>
 	FORCEINLINE TSharedPtr& operator=( SharedPointerInternals::TRawPtrProxyWithDeleter< OtherType, DeleterType >&& InRawPtrProxy )
 	{
@@ -1100,8 +1100,8 @@ public:
 	 * @return  Reference to the object
 	 */
 	template <
-		typename DummyObjectType = ObjectType,
-		decltype((int&)(*(DummyObjectType*)nullptr), 0) = 0 // this construct means that operator* is only considered for overload resolution if T is dereferenceable
+		typename DummyObjectType = ObjectType
+		UE_REQUIRES(UE_REQUIRES_EXPR(*(DummyObjectType*)nullptr)) // this construct means that operator* is only considered for overload resolution if T is dereferenceable
 	>
 	[[nodiscard]] FORCEINLINE DummyObjectType& operator*() const
 	{
@@ -1164,8 +1164,8 @@ private:
 	 *       pointer to a shared pointer.  Use the weak pointer's Pin() method instead!
 	 */
 	template <
-		typename OtherType,
-		typename = decltype(ImplicitConv<ObjectType*>((OtherType*)nullptr))
+		typename OtherType
+		UE_REQUIRES(UE_REQUIRES_EXPR(ImplicitConv<ObjectType*>((OtherType*)nullptr)))
 	>
 	FORCEINLINE explicit TSharedPtr( TWeakPtr< OtherType, Mode > const& InWeakPtr )
 		: Object( nullptr )
@@ -1279,8 +1279,8 @@ public:
 	 */
 	// NOTE: The following is an Unreal extension to standard shared_ptr behavior
 	template <
-		typename OtherType,
-		typename = decltype(ImplicitConv<ObjectType*>((OtherType*)nullptr))
+		typename OtherType
+		UE_REQUIRES(UE_REQUIRES_EXPR(ImplicitConv<ObjectType*>((OtherType*)nullptr)))
 	>
 	FORCEINLINE TWeakPtr( TSharedRef< OtherType, Mode > const& InSharedRef )
 		: Object( InSharedRef.Object )
@@ -1294,8 +1294,8 @@ public:
 	 * @param  InSharedPtr  The shared pointer to create a weak pointer from
 	 */
 	template <
-		typename OtherType,
-		typename = decltype(ImplicitConv<ObjectType*>((OtherType*)nullptr))
+		typename OtherType
+		UE_REQUIRES(UE_REQUIRES_EXPR(ImplicitConv<ObjectType*>((OtherType*)nullptr)))
 	>
 	FORCEINLINE TWeakPtr( TSharedPtr< OtherType, Mode > const& InSharedPtr )
 		: Object( InSharedPtr.Object )
@@ -1340,8 +1340,8 @@ public:
 	 * @param  InWeakPtr  The weak pointer to create a weak pointer from
 	 */
 	template <
-		typename OtherType,
-		typename = decltype(ImplicitConv<ObjectType*>((OtherType*)nullptr))
+		typename OtherType
+		UE_REQUIRES(UE_REQUIRES_EXPR(ImplicitConv<ObjectType*>((OtherType*)nullptr)))
 	>
 	FORCEINLINE TWeakPtr( TWeakPtr< OtherType, Mode > const& InWeakPtr )
 		: Object( InWeakPtr.Object )
@@ -1350,8 +1350,8 @@ public:
 	}
 
 	template <
-		typename OtherType,
-		typename = decltype(ImplicitConv<ObjectType*>((OtherType*)nullptr))
+		typename OtherType
+		UE_REQUIRES(UE_REQUIRES_EXPR(ImplicitConv<ObjectType*>((OtherType*)nullptr)))
 	>
 	FORCEINLINE TWeakPtr( TWeakPtr< OtherType, Mode >&& InWeakPtr )
 		: Object( InWeakPtr.Object )
@@ -1413,8 +1413,8 @@ public:
 	 * @param  InWeakPtr  The weak pointer for the object to assign
 	 */
 	template <
-		typename OtherType,
-		typename = decltype(ImplicitConv<ObjectType*>((OtherType*)nullptr))
+		typename OtherType
+		UE_REQUIRES(UE_REQUIRES_EXPR(ImplicitConv<ObjectType*>((OtherType*)nullptr)))
 	>
 	FORCEINLINE TWeakPtr& operator=( TWeakPtr<OtherType, Mode> const& InWeakPtr )
 	{
@@ -1424,8 +1424,8 @@ public:
 	}
 
 	template <
-		typename OtherType,
-		typename = decltype(ImplicitConv<ObjectType*>((OtherType*)nullptr))
+		typename OtherType
+		UE_REQUIRES(UE_REQUIRES_EXPR(ImplicitConv<ObjectType*>((OtherType*)nullptr)))
 	>
 	FORCEINLINE TWeakPtr& operator=( TWeakPtr<OtherType, Mode>&& InWeakPtr )
 	{
@@ -1442,8 +1442,8 @@ public:
 	 */
 	// NOTE: The following is an Unreal extension to standard shared_ptr behavior
 	template <
-		typename OtherType,
-		typename = decltype(ImplicitConv<ObjectType*>((OtherType*)nullptr))
+		typename OtherType
+		UE_REQUIRES(UE_REQUIRES_EXPR(ImplicitConv<ObjectType*>((OtherType*)nullptr)))
 	>
 	FORCEINLINE TWeakPtr& operator=( TSharedRef< OtherType, Mode > const& InSharedRef )
 	{
@@ -1458,8 +1458,8 @@ public:
 	 * @param  InSharedPtr  The shared pointer used to assign to this weak pointer
 	 */
 	template <
-		typename OtherType,
-		typename = decltype(ImplicitConv<ObjectType*>((OtherType*)nullptr))
+		typename OtherType
+		UE_REQUIRES(UE_REQUIRES_EXPR(ImplicitConv<ObjectType*>((OtherType*)nullptr)))
 	>
 	FORCEINLINE TWeakPtr& operator=( TSharedPtr< OtherType, Mode > const& InSharedPtr )
 	{
diff --git a/Engine/Source/Runtime/Core/Public/Templates/UniquePtr.h b/Engine/Source/Runtime/Core/Public/Templates/UniquePtr.h
--- a/Engine/Source/Runtime/Core/Public/Templates/UniquePtr.h
+++ b/Engine/Source/Runtime/Core/Public/Templates/UniquePtr.h
@@ -34,16 +34,16 @@ struct TDefaultDelete
 	~TDefaultDelete() = default;
 
 	template <
-		typename U,
-		typename = decltype(ImplicitConv<T*>((U*)nullptr))
+		typename U
+		UE_REQUIRES(UE_REQUIRES_EXPR(ImplicitConv<T*>((U*)nullptr)))
 	>
 	TDefaultDelete(const TDefaultDelete<U>&)
 	{
 	}
 
 	template <
-		typename U,
-		typename = decltype(ImplicitConv<T*>((U*)nullptr))
+		typename U
+		UE_REQUIRES(UE_REQUIRES_EXPR(ImplicitConv<T*>((U*)nullptr)))
 	>
 	TDefaultDelete& operator=(const TDefaultDelete<U>&)
 	{
@@ -76,16 +76,16 @@ struct TDefaultDelete<T[]>
 	~TDefaultDelete() = default;
 
 	template <
-		typename U,
-		typename = decltype(ImplicitConv<T(*)[]>((U(*)[])nullptr))
+		typename U
+		UE_REQUIRES(UE_REQUIRES_EXPR(ImplicitConv<T(*)[]>((U(*)[])nullptr)))
 	>
 	TDefaultDelete(const TDefaultDelete<U[]>&)
 	{
 	}
 
 	template <
-		typename U,
-		typename = decltype(ImplicitConv<T(*)[]>((U(*)[])nullptr))
+		typename U
+		UE_REQUIRES(UE_REQUIRES_EXPR(ImplicitConv<T(*)[]>((U(*)[])nullptr)))
 	>
 	TDefaultDelete& operator=(const TDefaultDelete<U[]>&)
 	{
@@ -93,8 +93,8 @@ struct TDefaultDelete<T[]>
 	}
 
 	template <
-		typename U,
-		typename = decltype(ImplicitConv<T(*)[]>((U(*)[])nullptr))
+		typename U
+		UE_REQUIRES(UE_REQUIRES_EXPR(ImplicitConv<T(*)[]>((U(*)[])nullptr)))
 	>
 	void operator()(U* Ptr) const
 	{
@@ -132,8 +132,8 @@ public:
 	 * @param InPtr The pointed-to object to take ownership of.
 	 */
 	template <
-		typename U,
-		typename = decltype(ImplicitConv<T*>((U*)nullptr))
+		typename U
+		UE_REQUIRES(UE_REQUIRES_EXPR(ImplicitConv<T*>((U*)nullptr)))
 	>
 	explicit FORCEINLINE TUniquePtr(U* InPtr)
 		: Deleter()
@@ -147,8 +147,8 @@ public:
 	 * @param InPtr The pointed-to object to take ownership of.
 	 */
 	template <
-		typename U,
-		typename = decltype(ImplicitConv<T*>((U*)nullptr))
+		typename U
+		UE_REQUIRES(UE_REQUIRES_EXPR(ImplicitConv<T*>((U*)nullptr)))
 	>
 	explicit FORCEINLINE TUniquePtr(U* InPtr, Deleter&& InDeleter)
 		: Deleter(MoveTemp(InDeleter))
@@ -162,8 +162,8 @@ public:
 	 * @param InPtr The pointed-to object to take ownership of.
 	 */
 	template <
-		typename U,
-		typename = decltype(ImplicitConv<T*>((U*)nullptr))
+		typename U
+		UE_REQUIRES(UE_REQUIRES_EXPR(ImplicitConv<T*>((U*)nullptr)))
 	>
 	explicit FORCEINLINE TUniquePtr(U* InPtr, const Deleter& InDeleter)
 		: Deleter(InDeleter)
@@ -195,9 +195,11 @@ public:
 	 */
 	template <
 		typename OtherT,
-		typename OtherDeleter,
-		typename = decltype(ImplicitConv<T*>((OtherT*)nullptr))
-		UE_REQUIRES(!std::is_array_v<OtherT>)
+		typename OtherDeleter
+		UE_REQUIRES(
+			!std::is_array_v<OtherT> &&
+			UE_REQUIRES_EXPR(ImplicitConv<T*>((OtherT*)nullptr))
+		)
 	>
 	FORCEINLINE TUniquePtr(TUniquePtr<OtherT, OtherDeleter>&& Other)
 		: Deleter(MoveTemp(Other.GetDeleter()))
@@ -230,9 +232,11 @@ public:
 	 */
 	template <
 		typename OtherT,
-		typename OtherDeleter,
-		typename = decltype(ImplicitConv<T*>((OtherT*)nullptr))
-		UE_REQUIRES(!std::is_array_v<OtherT>)
+		typename OtherDeleter
+		UE_REQUIRES(
+			!std::is_array_v<OtherT> &&
+			UE_REQUIRES_EXPR(ImplicitConv<T*>((OtherT*)nullptr))
+		)
 	>
 	FORCEINLINE TUniquePtr& operator=(TUniquePtr<OtherT, OtherDeleter>&& Other)
 	{
@@ -304,8 +308,8 @@ public:
 	 * @return A reference to the object owned by the TUniquePtr.
 	 */
 	template <
-		typename DummyT = T,
-		decltype((int&)*(DummyT*)nullptr, 0) = 0 // this construct means that operator* is only considered for overload resolution if T is dereferenceable
+		typename DummyT = T
+		UE_REQUIRES(UE_REQUIRES_EXPR(*(DummyT*)nullptr)) // this construct means that operator* is only considered for overload resolution if T is dereferenceable
 	>
 	[[nodiscard]] FORCEINLINE DummyT& operator*() const
 	{
@@ -458,8 +462,8 @@ public:
 	 * @param InPtr The pointed-to array to take ownership of.
 	 */
 	template <
-		typename U,
-		typename = decltype(ImplicitConv<T(*)[]>((U(*)[])nullptr))
+		typename U
+		UE_REQUIRES(UE_REQUIRES_EXPR(ImplicitConv<T(*)[]>((U(*)[])nullptr)))
 	>
 	explicit FORCEINLINE TUniquePtr(U* InPtr)
 		: Deleter()
@@ -473,8 +477,8 @@ public:
 	 * @param InPtr The pointed-to array to take ownership of.
 	 */
 	template <
-		typename U,
-		typename = decltype(ImplicitConv<T(*)[]>((U(*)[])nullptr))
+		typename U
+		UE_REQUIRES(UE_REQUIRES_EXPR(ImplicitConv<T(*)[]>((U(*)[])nullptr)))
 	>
 	explicit FORCEINLINE TUniquePtr(U* InPtr, Deleter&& InDeleter)
 		: Deleter(MoveTemp(InDeleter))
@@ -488,8 +492,8 @@ public:
 	 * @param InPtr The pointed-to array to take ownership of.
 	 */
 	template <
-		typename U,
-		typename = decltype(ImplicitConv<T(*)[]>((U(*)[])nullptr))
+		typename U
+		UE_REQUIRES(UE_REQUIRES_EXPR(ImplicitConv<T(*)[]>((U(*)[])nullptr)))
 	>
 	explicit FORCEINLINE TUniquePtr(U* InPtr, const Deleter& InDeleter)
 		: Deleter(InDeleter)
@@ -521,8 +525,8 @@ public:
 	 */
 	template <
 		typename OtherT,
-		typename OtherDeleter,
-		typename = decltype(ImplicitConv<T(*)[]>((OtherT(*)[])nullptr))
+		typename OtherDeleter
+		UE_REQUIRES(UE_REQUIRES_EXPR(ImplicitConv<T(*)[]>((OtherT(*)[])nullptr)))
 	>
 	FORCEINLINE TUniquePtr(TUniquePtr<OtherT, OtherDeleter>&& Other)
 		: Deleter(MoveTemp(Other.GetDeleter()))
@@ -555,8 +559,8 @@ public:
 	 */
 	template <
 		typename OtherT,
-		typename OtherDeleter,
-		typename = decltype(ImplicitConv<T(*)[]>((OtherT(*)[])nullptr))
+		typename OtherDeleter
+		UE_REQUIRES(UE_REQUIRES_EXPR(ImplicitConv<T(*)[]>((OtherT(*)[])nullptr)))
 	>
 	FORCEINLINE TUniquePtr& operator=(TUniquePtr<OtherT, OtherDeleter>&& Other)
 	{
@@ -650,8 +654,8 @@ public:
 	 * @param InPtr A pointer to the array to take ownership of.
 	 */
 	template <
-		typename U,
-		typename = decltype(ImplicitConv<T(*)[]>((U(*)[])nullptr))
+		typename U
+		UE_REQUIRES(UE_REQUIRES_EXPR(ImplicitConv<T(*)[]>((U(*)[])nullptr)))
 	>
 	FORCEINLINE void Reset(U* InPtr)
 	{
