Exposes the `LogFilePathTag` XML crash property.

diff --git a/Engine/Source/Runtime/Core/Private/GenericPlatform/GenericPlatformCrashContext.cpp b/Engine/Source/Runtime/Core/Private/GenericPlatform/GenericPlatformCrashContext.cpp
--- a/Engine/Source/Runtime/Core/Private/GenericPlatform/GenericPlatformCrashContext.cpp
+++ b/Engine/Source/Runtime/Core/Private/GenericPlatform/GenericPlatformCrashContext.cpp
@@ -208,6 +208,7 @@ const TCHAR* const FGenericCrashContext::CallStackTag = TEXT("CallStack");
 const TCHAR* const FGenericCrashContext::PortableCallStackTag = TEXT("PCallStack");
 const TCHAR* const FGenericCrashContext::PortableCallStackHashTag = TEXT("PCallStackHash");
 const TCHAR* const FGenericCrashContext::IsRequestingExitTag = TEXT("IsRequestingExit");
+const TCHAR* const FGenericCrashContext::LogFilePathTag = TEXT("LogFilePath");
 const TCHAR* const FGenericCrashContext::UEMinidumpName = TEXT( "UEMinidump.dmp" );
 const TCHAR* const FGenericCrashContext::NewLineTag = TEXT( "&nl;" );
 
@@ -877,7 +878,7 @@ void FGenericCrashContext::SerializeUserSettings(FString& Buffer)
 	AddCrashPropertyInternal(Buffer, TEXT("NoDialog"), NCached::UserSettings.bNoDialog);
 	AddCrashPropertyInternal(Buffer, TEXT("SendUnattendedBugReports"), NCached::UserSettings.bSendUnattendedBugReports);
 	AddCrashPropertyInternal(Buffer, TEXT("SendUsageData"), NCached::UserSettings.bSendUsageData);
-	AddCrashPropertyInternal(Buffer, TEXT("LogFilePath"), FPlatformOutputDevices::GetAbsoluteLogFilename()); // Don't use the value cached, it may be out of date.
+	AddCrashPropertyInternal(Buffer, FGenericCrashContext::LogFilePathTag, FPlatformOutputDevices::GetAbsoluteLogFilename()); // Don't use the value cached, it may be out of date.
 }
 
 // This function may be called in the crashing executable or in an external crash reporter program. Take care with accessing global variables vs member variables or 
diff --git a/Engine/Source/Runtime/Core/Public/GenericPlatform/GenericPlatformCrashContext.h b/Engine/Source/Runtime/Core/Public/GenericPlatform/GenericPlatformCrashContext.h
--- a/Engine/Source/Runtime/Core/Public/GenericPlatform/GenericPlatformCrashContext.h
+++ b/Engine/Source/Runtime/Core/Public/GenericPlatform/GenericPlatformCrashContext.h
@@ -398,6 +398,7 @@ public:
 	CORE_API static const TCHAR* const PortableCallStackTag;
 	CORE_API static const TCHAR* const PortableCallStackHashTag;
 	CORE_API static const TCHAR* const IsRequestingExitTag;
+	CORE_API static const TCHAR* const LogFilePathTag;
 
 	static constexpr inline int32 CrashGUIDLength = 128;
 
