ColorManagement: Fix stress unity build error with cvar hiding global declaration.

#rnx

diff --git a/Engine/Source/Runtime/Core/Private/Math/Color.cpp b/Engine/Source/Runtime/Core/Private/Math/Color.cpp
--- a/Engine/Source/Runtime/Core/Private/Math/Color.cpp
+++ b/Engine/Source/Runtime/Core/Private/Math/Color.cpp
@@ -300,8 +300,8 @@ FLinearColor FLinearColor::Desaturate( float Desaturation ) const
 /** Computes the perceptually weighted luminance value of a color. */
 float FLinearColor::GetLuminance() const
 {
-	static const auto CVarLegacyLuminanceFactors = IConsoleManager::Get().FindConsoleVariable(TEXT("r.LegacyLuminanceFactors"));
-	static const FLinearColor LuminanceFactors = CVarLegacyLuminanceFactors->GetInt() != 0 ? FLinearColor(0.3f, 0.59f, 0.11f) : UE::Color::FColorSpace::GetWorking().GetLuminanceFactors();
+	static const IConsoleVariable* CVar = IConsoleManager::Get().FindConsoleVariable(TEXT("r.LegacyLuminanceFactors"));
+	static const FLinearColor LuminanceFactors = CVar->GetInt() != 0 ? FLinearColor(0.3f, 0.59f, 0.11f) : UE::Color::FColorSpace::GetWorking().GetLuminanceFactors();
 
 	return R * LuminanceFactors.R + G * LuminanceFactors.G + B * LuminanceFactors.B;
 }
