No longer create display strings for unlocalized text

In UE 5.0 we removed the assumption that every FText must have an associated display string. Prior to this change, FText didn't store its TextID locally, and instead resolved it via a reverse LUT against the display string pointer.

While the initial change for 5.0 cleaned up most of the code that was creating redundant display strings, it missed a few places that were still creating display strings for unlocalized text. Now that we can unload the localization data for plugins, it's important that any lingering FText instances don't re-add a redundant display string via FText::Rebuild.

To this end, this change removes the remaining vestiges of the old "every FText must have a display string" assumption and updates FTextLocalizationManager to no longer create display strings for unlocalized text.

This change breaks two pieces of API (which is why it wasn't done initially), however this API was only used within FText itself. The likelihood of this breaking user code is very low.

Breaking changes:
  * FTextLocalizationManager::GetDisplayString now returns a FTextConstDisplayStringPtr rather than a FTextConstDisplayStringRef, and will now return null for an unlocalized text.
  * FTextLocalizationManager::AddDisplayString has been deleted. It has served no purpose since the reverse LUT was removed.

diff --git a/Engine/Source/Runtime/Core/Private/Internationalization/TextHistory.cpp b/Engine/Source/Runtime/Core/Private/Internationalization/TextHistory.cpp
--- a/Engine/Source/Runtime/Core/Private/Internationalization/TextHistory.cpp
+++ b/Engine/Source/Runtime/Core/Private/Internationalization/TextHistory.cpp
@@ -846,22 +846,13 @@ void FTextHistory_Base::Serialize(FStructuredArchive::FRecord Record)
 			}
 #endif // USE_STABLE_LOCALIZATION_KEYS
 
+#if WITH_EDITOR
 			// If this has no key, give it a GUID for a key
-			if (GIsEditor && TextId.IsEmpty() && (BaseArchive.IsPersistent() && !BaseArchive.HasAnyPortFlags(PPF_Duplicate)))
+			if (GIsEditor && Key.IsEmpty() && (BaseArchive.IsPersistent() && !BaseArchive.HasAnyPortFlags(PPF_Duplicate)))
 			{
 				Key = FGuid::NewGuid().ToString();
-				if (FTextLocalizationManager::Get().AddDisplayString(MakeTextDisplayString(CopyTemp(SourceString)), Namespace, Key))
-				{
-					TextId = FTextId(Namespace, Key);
-					MarkDisplayStringOutOfDate();
-				}
-				else
-				{
-					// Could not add display string, reset namespace and key.
-					Namespace.Reset();
-					Key.Reset();
-				}
 			}
+#endif // WITH_EDITOR
 		}
 
 		// Serialize the Namespace
diff --git a/Engine/Source/Runtime/Core/Private/Internationalization/TextLocalizationManager.cpp b/Engine/Source/Runtime/Core/Private/Internationalization/TextLocalizationManager.cpp
--- a/Engine/Source/Runtime/Core/Private/Internationalization/TextLocalizationManager.cpp
+++ b/Engine/Source/Runtime/Core/Private/Internationalization/TextLocalizationManager.cpp
@@ -107,6 +107,13 @@ static FAutoConsoleCommand CmdDumpLiveTable(
 		}
 	}));
 #endif
+
+FString KeyifyTextId(const FTextId& TextId)
+{
+	// We want to show the identity in terms of key, namespace. This is to try and fit into the constraints of UI text blocks and at least let the key component be visible to easily identify a piece of text.
+	// If the key/namespace pair is too long, the Slate.LogPaintedText cvar can be used to see the entire thing.
+	return FString::Printf(TEXT("%s, %s"), TextId.GetKey().GetChars(), TextId.GetNamespace().GetChars());
+}
 }
 
 enum class ERequestedCultureOverrideLevel : uint8
@@ -885,7 +892,7 @@ void FTextLocalizationManager::RegisterPolyglotTextData(TArrayView<const FPolygl
 
 FTextConstDisplayStringPtr FTextLocalizationManager::FindDisplayString(const FTextKey& Namespace, const FTextKey& Key, const FString* const SourceStringPtr) const
 {
-	if (!FTextLocalizationManager::IsDisplayStringSupportEnabled())
+	if (Key.IsEmpty() || !FTextLocalizationManager::IsDisplayStringSupportEnabled())
 	{
 		return nullptr;
 	}
@@ -917,21 +924,14 @@ FTextConstDisplayStringPtr FTextLocalizationManager::FindDisplayString(const FTe
 	return nullptr;
 }
 
-FTextConstDisplayStringRef FTextLocalizationManager::GetDisplayString(const FTextKey& Namespace, const FTextKey& Key, const FString* const SourceStringPtr)
+FTextConstDisplayStringPtr FTextLocalizationManager::GetDisplayString(const FTextKey& Namespace, const FTextKey& Key, const FString* const SourceStringPtr)
 {
 	TRACE_CPUPROFILER_EVENT_SCOPE(FTextLocalizationManager::GetDisplayString);
 	LLM_SCOPE_BYNAME(TEXT("Localization/DisplayStrings"));
 
-	auto GetEmptyDisplayString = []()
-	{
-		static const FTextConstDisplayStringRef EmptyDisplayString = MakeTextDisplayString(FString());
-		return EmptyDisplayString;
-	};
-
-	if (!FTextLocalizationManager::IsDisplayStringSupportEnabled())
+	if (Key.IsEmpty() || !FTextLocalizationManager::IsDisplayStringSupportEnabled())
 	{
-		ensureAlwaysMsgf(false, TEXT("FTextLocalizationManager::GetDisplayString was called but display string support is disabled! Either update the calling code to respect FTextLocalizationManager::IsDisplayStringSupportEnabled, or re-enable display string support via 'Localization.DisplayStringSupport'."));
-		return GetEmptyDisplayString();
+		return nullptr;
 	}
 
 	auto GetSourceStringRef = [SourceStringPtr]() -> const FString&
@@ -949,14 +949,9 @@ FTextConstDisplayStringRef FTextLocalizationManager::GetDisplayString(const FTex
 
 	FScopeLock ScopeLock(&DisplayStringTableCS);
 
-	// Hack fix for old assets that don't have namespace/key info.
-	if (Namespace.IsEmpty() && Key.IsEmpty())
-	{
-		return MakeTextDisplayString(CopyTemp(SourceString));
-	}
-
 #if ENABLE_LOC_TESTING
 	const bool bShouldLEETIFYAll = IsInitialized() && FInternationalization::Get().GetCurrentLanguage()->GetName() == FLeetCulture::StaticGetName();
+	const bool bShouldKeyifyAll = IsInitialized() && FInternationalization::Get().GetCurrentLanguage()->GetName() == FKeysCulture::StaticGetName();
 
 	// Attempt to set bShouldLEETIFYUnlocalizedString appropriately, only once, after the commandline is initialized and parsed.
 	static bool bShouldLEETIFYUnlocalizedString = false;
@@ -1012,8 +1007,18 @@ FTextConstDisplayStringRef FTextLocalizationManager::GetDisplayString(const FTex
 			DirtyLocalRevisionForTextId(TextId);
 
 #if ENABLE_LOC_TESTING
-			if ((bShouldLEETIFYAll || bShouldLEETIFYUnlocalizedString) && !LiveEntry->DisplayString->IsEmpty())
+			if (bShouldKeyifyAll)
 			{
+				DisplayStringBackupTable.Add(TextId, LiveEntry->DisplayString);
+				LiveEntry->DisplayString = MakeTextDisplayString(TextLocalizationManager::KeyifyTextId(TextId));
+			}
+			else if ((bShouldLEETIFYAll || bShouldLEETIFYUnlocalizedString) && !LiveEntry->DisplayString->IsEmpty())
+			{
+				if (!bShouldLEETIFYUnlocalizedString)
+				{
+					DisplayStringBackupTable.Add(TextId, LiveEntry->DisplayString);
+				}
+
 				FTextDisplayStringRef TmpDisplayString = MakeTextDisplayString(CopyTemp(*LiveEntry->DisplayString));
 				FInternationalization::Leetify(*TmpDisplayString);
 				LiveEntry->DisplayString = TmpDisplayString;
@@ -1039,29 +1044,33 @@ FTextConstDisplayStringRef FTextLocalizationManager::GetDisplayString(const FTex
 
 		return NewEntry.DisplayString;
 	}
+#if ENABLE_LOC_TESTING
 	// Entry is absent.
-	else
+	else if (bShouldKeyifyAll || ((bShouldLEETIFYAll || bShouldLEETIFYUnlocalizedString) && !SourceString.IsEmpty()))
 	{
 		TRACE_CPUPROFILER_EVENT_SCOPE(FTextLocalizationManager::GetDisplayString_AddNewEntry);
 
-		// Don't log warnings about unlocalized strings if the system hasn't been initialized - we simply don't have localization data yet.
-		if (IsInitialized())
+		FTextConstDisplayStringRef UnlocalizedString = MakeTextDisplayString(CopyTemp(SourceString));
+
+		if (bShouldKeyifyAll)
 		{
-			UE_LOG(LogTextLocalizationManager, Verbose, TEXT("An attempt was made to get a localized string (Namespace:%s, Key:%s, Source:%s), but it did not exist."), TextId.GetNamespace().GetChars(), TextId.GetKey().GetChars(), *SourceString);
+			DisplayStringBackupTable.Add(TextId, UnlocalizedString);
+			UnlocalizedString = MakeTextDisplayString(TextLocalizationManager::KeyifyTextId(TextId));
 		}
-		
-		FTextConstDisplayStringRef UnlocalizedString = !SourceString.IsEmpty() ? MakeTextDisplayString(CopyTemp(SourceString)) : GetEmptyDisplayString();
-
-#if ENABLE_LOC_TESTING
-		if ((bShouldLEETIFYAll || bShouldLEETIFYUnlocalizedString) && !UnlocalizedString->IsEmpty())
+		else if (bShouldLEETIFYAll || bShouldLEETIFYUnlocalizedString)
 		{
+			check(!SourceString.IsEmpty());
+
+			if (!bShouldLEETIFYUnlocalizedString)
+			{
+				DisplayStringBackupTable.Add(TextId, UnlocalizedString);
+			}
+
 			FTextDisplayStringRef TmpDisplayString = MakeTextDisplayString(CopyTemp(*UnlocalizedString));
 			FInternationalization::Leetify(*TmpDisplayString);
 			UnlocalizedString = TmpDisplayString;
 		}
-#endif
 
-		// Make entries so that they can be updated when system is initialized or a culture swap occurs.
 		FDisplayStringEntry NewEntry(
 			FTextKey(),					/*LocResID*/
 			INDEX_NONE,					/*LocalizationTargetPathId*/
@@ -1073,6 +1082,9 @@ FTextConstDisplayStringRef FTextLocalizationManager::GetDisplayString(const FTex
 
 		return UnlocalizedString;
 	}
+#endif // ENABLE_LOC_TESTING
+
+	return nullptr;
 }
 
 #if WITH_EDITORONLY_DATA
@@ -1128,30 +1140,6 @@ void FTextLocalizationManager::GetTextRevisions(const FTextId& InTextId, uint16&
 	}
 }
 
-bool FTextLocalizationManager::AddDisplayString(const FTextDisplayStringRef& DisplayString, const FTextKey& Namespace, const FTextKey& Key)
-{
-	TRACE_CPUPROFILER_EVENT_SCOPE(FTextLocalizationManager::AddDisplayString);
-	LLM_SCOPE_BYNAME(TEXT("Localization/DisplayStrings"));
-
-	FScopeLock ScopeLock(&DisplayStringTableCS);
-
-	const FTextId TextId(Namespace, Key);
-
-	// Try to find existing entry.
-	FDisplayStringEntry* ExistingDisplayStringEntry = DisplayStringLookupTable.Find(TextId);
-
-	// If there are any existing entry, they may cause a conflict, unless they're exactly the same as what we would be adding.
-	if (ExistingDisplayStringEntry && ExistingDisplayStringEntry->DisplayString != DisplayString) // Namespace and key mustn't be associated with a different display string.
-	{
-		return false;
-	}
-
-	// Add the necessary association.
-	DisplayStringLookupTable.Emplace(TextId, FDisplayStringEntry(FTextKey(), INDEX_NONE, FTextLocalizationResource::HashString(*DisplayString), DisplayString));
-
-	return true;
-}
-
 void FTextLocalizationManager::UpdateFromLocalizationResource(const FString& LocalizationResourceFilePath)
 {
 	FTextLocalizationResource TextLocalizationResource;
@@ -1267,6 +1255,9 @@ void FTextLocalizationManager::HandleLocalizationTargetsUnmounted(TArrayView<con
 				DisplayStringsForLocalizationTarget.bIsMounted = false;
 			}
 		}
+
+		// Allow any lingering texts that were referencing the unloaded display strings to release their references
+		TLM.DirtyTextRevision();
 	});
 }
 
@@ -2024,11 +2015,7 @@ void FTextLocalizationManager::KeyifyAllDisplayStrings()
 		FDisplayStringEntry& LiveEntry = DisplayStringPair.Value;
 		DisplayStringBackupTable.Add(DisplayStringPair.Key, LiveEntry.DisplayString);
 
-		// We want to show the identity in terms of key, namespace. This is to try and fit into the constraints of UI text blocks and at least let the key component be visible to easily identify a piece of text.
-		// If the key/namespace pair is too long, the Slate.LogPaintedText cvar can be used to see the entire thing.
-		FString KeyNamespaceDisplay = FString::Printf(TEXT("%s, %s"), DisplayStringPair.Key.GetKey().GetChars(), DisplayStringPair.Key.GetNamespace().GetChars());
-		FTextDisplayStringRef TmpDisplayString = MakeTextDisplayString(MoveTemp(KeyNamespaceDisplay));
-		LiveEntry.DisplayString = TmpDisplayString;
+		LiveEntry.DisplayString = MakeTextDisplayString(TextLocalizationManager::KeyifyTextId(DisplayStringPair.Key));
 	}
 }
 #endif 
diff --git a/Engine/Source/Runtime/Core/Public/Internationalization/TextLocalizationManager.h b/Engine/Source/Runtime/Core/Public/Internationalization/TextLocalizationManager.h
--- a/Engine/Source/Runtime/Core/Public/Internationalization/TextLocalizationManager.h
+++ b/Engine/Source/Runtime/Core/Public/Internationalization/TextLocalizationManager.h
@@ -214,18 +214,12 @@ public:
 	 *		... but it was not localized from the specified source string, the display string will be set to the specified source and returned.
 	 *		... and it was localized from the specified source string (or the source string was null or empty), the display string will be returned.
 	*/
-	CORE_API FTextConstDisplayStringRef GetDisplayString(const FTextKey& Namespace, const FTextKey& Key, const FString* const SourceStringPtr);
+	CORE_API FTextConstDisplayStringPtr GetDisplayString(const FTextKey& Namespace, const FTextKey& Key, const FString* const SourceStringPtr);
 
 #if WITH_EDITORONLY_DATA
 	/** If an entry exists for the specified namespace and key, returns true and provides the localization resource identifier from which it was loaded. Otherwise, returns false. */
 	CORE_API bool GetLocResID(const FTextKey& Namespace, const FTextKey& Key, FString& OutLocResId) const;
 #endif
-	
-	/**	Attempts to register the specified display string, associating it with the specified namespace and key.
-	 *	Returns true if the display string has been or was already associated with the namespace and key.
-	 *	Returns false if the display string was already associated with another namespace and key or the namespace and key are already in use by another display string.
-	 */
-	CORE_API bool AddDisplayString(const FTextDisplayStringRef& DisplayString, const FTextKey& Namespace, const FTextKey& Key);
 
 	/** Updates display string entries and adds new display string entries based on localizations found in a specified localization resource. */
 	CORE_API void UpdateFromLocalizationResource(const FString& LocalizationResourceFilePath);
