autortfm:
- add PushOnAbortHandler/PopOnAbortHandler duo so you can register a handler that can later be unregistered. Useful for push/pop semantics where we don't need to do anything OnAbort after the pop
- deploy usage to FHashTableLock

diff --git a/Engine/Source/Programs/AutoRTFMTests/Private/Abort.cpp b/Engine/Source/Programs/AutoRTFMTests/Private/Abort.cpp
--- a/Engine/Source/Programs/AutoRTFMTests/Private/Abort.cpp
+++ b/Engine/Source/Programs/AutoRTFMTests/Private/Abort.cpp
@@ -198,3 +198,137 @@ TEST_CASE("Abort.CascadeThroughManualTransaction")
 	REQUIRE(AutoRTFM::ETransactionResult::AbortedByCascade == Result);
 	REQUIRE(false == bTouched);
 }
+
+inline void* UIntToPointer(uint64 Value)
+{
+	union
+	{
+		uint64 Int;
+		void* Ptr;
+	};
+
+	Int = Value;
+	return Ptr;
+}
+
+TEST_CASE("Abort.PushOnAbortHandler_NoAbort")
+{
+	int Value = 55;
+
+	AutoRTFM::Commit([&]
+	{
+		Value = 66;
+		AutoRTFM::PushOnAbortHandler(UIntToPointer(747), [&Value](){ Value = 77; });
+	});
+
+	REQUIRE(Value == 66);
+}
+
+TEST_CASE("Abort.PushOnAbortHandler_WithAbort")
+{
+	int Value = 55;
+
+	const AutoRTFM::ETransactionResult Result = AutoRTFM::Transact([&]
+	{
+		Value = 66;
+		AutoRTFM::PushOnAbortHandler(UIntToPointer(747), [&Value](){ Value = 77; });
+
+		AutoRTFM::AbortTransaction();
+	});
+
+	REQUIRE(AutoRTFM::ETransactionResult::AbortedByRequest == Result);
+	REQUIRE(Value == 77);
+}
+
+TEST_CASE("Abort.PushOnAbortHandler_WithPop_NoAbort")
+{
+	int Value = 55;
+
+	AutoRTFM::Commit([&]
+	{
+		Value = 66;
+		AutoRTFM::PushOnAbortHandler(UIntToPointer(747), [&Value](){ Value = 77; });
+		Value = 88;
+
+		AutoRTFM::PopOnAbortHandler(UIntToPointer(747));
+	});
+
+	REQUIRE(Value == 88);
+}
+
+TEST_CASE("Abort.PushOnAbortHandler_WithPop_WithAbort")
+{
+	int Value = 55;
+
+	const AutoRTFM::ETransactionResult Result = AutoRTFM::Transact([&]
+	{
+		Value = 66;
+		AutoRTFM::PushOnAbortHandler(UIntToPointer(747), [&Value](){ Value = 77; });
+		Value = 88;
+
+		AutoRTFM::PopOnAbortHandler(UIntToPointer(747));
+
+		AutoRTFM::AbortTransaction();
+	});
+
+	REQUIRE(AutoRTFM::ETransactionResult::AbortedByRequest == Result);
+	REQUIRE(Value == 55);
+}
+
+TEST_CASE("Abort.PushOnAbortHandler_Duplicates1")
+{
+	int Value = 55;
+
+	const AutoRTFM::ETransactionResult Result = AutoRTFM::Transact([&]
+	{
+		Value = 66;
+		AutoRTFM::PushOnAbortHandler(UIntToPointer(747), [&Value](){ Value = 77; });
+		AutoRTFM::PushOnAbortHandler(UIntToPointer(747), [&Value](){ Value = 88; });
+		Value = 99;
+
+		AutoRTFM::PopOnAbortHandler(UIntToPointer(747));
+
+		AutoRTFM::AbortTransaction();
+	});
+
+	REQUIRE(AutoRTFM::ETransactionResult::AbortedByRequest == Result);
+	REQUIRE(Value == 55);
+}
+
+TEST_CASE("Abort.PushOnAbortHandler_Duplicates2")
+{
+	int Value = 55;
+
+	const AutoRTFM::ETransactionResult Result = AutoRTFM::Transact([&]
+	{
+		Value = 66;
+		AutoRTFM::PushOnAbortHandler(UIntToPointer(747), [&Value](){ Value = 77; });
+		AutoRTFM::PushOnAbortHandler(UIntToPointer(747), [&Value](){ Value = 88; });
+		Value = 99;
+
+		AutoRTFM::AbortTransaction();
+	});
+
+	REQUIRE(AutoRTFM::ETransactionResult::AbortedByRequest == Result);
+	REQUIRE(Value == 77);
+}
+
+TEST_CASE("Abort.PushOnAbortHandler_MultiplePops")
+{
+	int Value = 55;
+
+	const AutoRTFM::ETransactionResult Result = AutoRTFM::Transact([&]
+	{
+		Value = 66;
+		AutoRTFM::PopOnAbortHandler(UIntToPointer(747));
+		AutoRTFM::PopOnAbortHandler(UIntToPointer(747));
+		AutoRTFM::PopOnAbortHandler(UIntToPointer(747));
+		AutoRTFM::PopOnAbortHandler(UIntToPointer(747));
+		Value = 99;
+
+		AutoRTFM::AbortTransaction();
+	});
+
+	REQUIRE(AutoRTFM::ETransactionResult::AbortedByRequest == Result);
+	REQUIRE(Value == 55);
+}
diff --git a/Engine/Source/Runtime/Core/Private/AutoRTFM/API.cpp b/Engine/Source/Runtime/Core/Private/AutoRTFM/API.cpp
--- a/Engine/Source/Runtime/Core/Private/AutoRTFM/API.cpp
+++ b/Engine/Source/Runtime/Core/Private/AutoRTFM/API.cpp
@@ -258,6 +258,14 @@ UE_AUTORTFM_AUTORTFM("RTFM_OnAbort") void OnAbort(TFunction<void()> && Work)
 {
 }
 
+UE_AUTORTFM_AUTORTFM("RTFM_PushOnAbortHandler") void PushOnAbortHandler(const void* Key, TFunction<void()> && Work)
+{
+}
+
+UE_AUTORTFM_AUTORTFM("RTFM_PopOnAbortHandler") void PopOnAbortHandler(const void* Key)
+{
+}
+
 void OpenCommit(TFunction<void()>&& Work)
 {
 	OnCommit(MoveTemp(Work));
@@ -268,12 +276,20 @@ void OpenAbort(TFunction<void()>&& Work)
 	OnAbort(MoveTemp(Work));
 }
 
-extern "C" UE_AUTORTFM_AUTORTFM("RTFM_autortfm_on_commit") void autortfm_on_commit(void (*Work)(void* Arg), void* Arg)
+extern "C" UE_AUTORTFM_AUTORTFM("RTFM_autortfm_on_commit") void autortfm_on_commit(void (*Work)(void*), void* Arg)
 {
     Work(Arg);
 }
 
-extern "C" UE_AUTORTFM_AUTORTFM("RTFM_autortfm_on_abort") void autortfm_on_abort(void (*Work)(void* arg), void* Arg)
+extern "C" UE_AUTORTFM_AUTORTFM("RTFM_autortfm_on_abort") void autortfm_on_abort(void (*Work)(void*), void* Arg)
+{
+}
+
+extern "C" UE_AUTORTFM_AUTORTFM("RTFM_autortfm_push_on_abort_handler") void autortfm_push_on_abort_handler(const void* Key, void (*Work)(void*), void* Arg)
+{
+}
+
+extern "C" UE_AUTORTFM_AUTORTFM("RTFM_autortfm_pop_on_abort_handler") void autortfm_pop_on_abort_handler(const void* Key)
 {
 }
 
@@ -469,16 +485,40 @@ extern "C" UE_AUTORTFM_NOAUTORTFM void RTFM_OnAbort(TFunction<void()>&& Work)
     Context->GetCurrentTransaction()->DeferUntilAbort(MoveTemp(Work));
 }
 
-extern "C" UE_AUTORTFM_NOAUTORTFM void RTFM_autortfm_on_commit(void (*Work)(void* Arg), void* Arg)
+extern "C" UE_AUTORTFM_NOAUTORTFM void RTFM_PushOnAbortHandler(const void* Key, TFunction<void()>&& Work)
+{
+	FContext* Context = FContext::Get();
+    ASSERT(Context->GetStatus() == EContextStatus::OnTrack);
+    Context->GetCurrentTransaction()->PushDeferUntilAbortHandler(Key, MoveTemp(Work));
+}
+
+extern "C" UE_AUTORTFM_NOAUTORTFM void RTFM_PopOnAbortHandler(const void* Key)
+{
+	FContext* Context = FContext::Get();
+    ASSERT(Context->GetStatus() == EContextStatus::OnTrack);
+    Context->GetCurrentTransaction()->PopDeferUntilAbortHandler(Key);
+}
+
+extern "C" UE_AUTORTFM_NOAUTORTFM void RTFM_autortfm_on_commit(void (*Work)(void*), void* Arg)
 {
     RTFM_OnCommit([Work, Arg] { Work(Arg); });
 }
 
-extern "C" UE_AUTORTFM_NOAUTORTFM void RTFM_autortfm_on_abort(void (*Work)(void* arg), void* Arg)
+extern "C" UE_AUTORTFM_NOAUTORTFM void RTFM_autortfm_on_abort(void (*Work)(void*), void* Arg)
 {
     RTFM_OnAbort([Work, Arg] { Work(Arg); });
 }
 
+extern "C" UE_AUTORTFM_NOAUTORTFM void RTFM_autortfm_push_on_abort_handler(const void* Key, void (*Work)(void*), void* Arg)
+{
+    RTFM_PushOnAbortHandler(Key, [Work, Arg] { Work(Arg); });
+}
+
+extern "C" UE_AUTORTFM_NOAUTORTFM void RTFM_autortfm_pop_on_abort_handler(const void* Key)
+{
+    RTFM_PopOnAbortHandler(Key);
+}
+
 extern "C" UE_AUTORTFM_NOAUTORTFM void* RTFM_autortfm_did_allocate(void* Ptr, size_t Size)
 {
 	FContext* Context = FContext::Get();
diff --git a/Engine/Source/Runtime/Core/Private/AutoRTFM/TaskArray.h b/Engine/Source/Runtime/Core/Private/AutoRTFM/TaskArray.h
--- a/Engine/Source/Runtime/Core/Private/AutoRTFM/TaskArray.h
+++ b/Engine/Source/Runtime/Core/Private/AutoRTFM/TaskArray.h
@@ -4,7 +4,6 @@
 
 #include "Containers/Array.h"
 #include "Templates/Function.h"
-#include "Templates/SharedPointer.h"
 
 namespace AutoRTFM
 {
@@ -18,77 +17,119 @@ public:
 	auto end() { return Obj.rend(); }
 };
 
-template<typename T>
-class TTaskArray
+template<typename InKeyType, typename InValType>
+struct TTaskArrayKeyValuePair
 {
-	template<typename U> using USharedPtr = TSharedPtr<U, ESPMode::NotThreadSafe>;
+	InKeyType Key;
+	InValType Val;
+
+	TTaskArrayKeyValuePair() = default;
+	TTaskArrayKeyValuePair(const InKeyType& Key, const InValType& Val) : Key(Key), Val(Val){}
+	TTaskArrayKeyValuePair(InKeyType&& Key, InValType&& Val) : Key(MoveTemp(Key)), Val(MoveTemp(Val)){}
+};
 
+template<typename InElementType>
+class TTaskArray
+{
+	using SKey = const void*;
+	using SKeyValuePair = TTaskArrayKeyValuePair<SKey, InElementType>;
+	using FInternalArray = TArray<SKeyValuePair>;
 public:
-	TTaskArray() : Latest(new TArray<T>()) {}
+
+	TTaskArray() = default;
 	TTaskArray(const TTaskArray&) = delete;
 	void operator=(const TTaskArray&) = delete;
 
-    bool IsEmpty() const { return Latest->IsEmpty() && Stash.IsEmpty(); }
+    bool IsEmpty() const { return Latest.IsEmpty() && Stash.IsEmpty(); }
+
+    void Add(InElementType&& Value)
+    {
+		Latest.Push(TTaskArrayKeyValuePair(SKey{}, MoveTemp(Value)));
+    }
+
+	void AddKeyed(SKey Key, InElementType&& Value)
+    {
+		ASSERT(Key != SKey{});
+		Latest.Push(TTaskArrayKeyValuePair(Key, MoveTemp(Value)));
+    }
 
-    void Add(T&& value)
+    void Add(const InElementType& Value)
     {
-        Latest->Push(MoveTemp(value));
+		Latest.Push(TTaskArrayKeyValuePair(SKey{}, Value));
     }
 
-    void Add(const T& value)
+	void AddKeyed(SKey Key, const InElementType& Value)
     {
-        Latest->Push(value);
+		ASSERT(Key != SKey{});
+		Latest.Push(TTaskArrayKeyValuePair(Key, Value));
     }
 
-    void AddAll(TTaskArray<T>&& Other)
+	void DeleteKey(SKey Key)
+	{
+		ASSERT(Key != SKey{});
+
+		auto EraseKeyFromArray = [Key](FInternalArray& Array) -> void
+		{
+			uint64 WriteIdx = 0;
+			for(uint64 ReadIdx = 0; ReadIdx < Array.Num(); ReadIdx++)
+			{
+				if(Array[ReadIdx].Key != Key)
+				{
+					// keep it
+					if(WriteIdx != ReadIdx) // the common case will be WriteIdx==ReadIdx so don't invoke an unnecessary copy-to-self
+					{
+						Array[WriteIdx] = MoveTemp(Array[ReadIdx]);
+					}
+					WriteIdx++;
+				}
+			}
+
+			Array.SetNum(WriteIdx, EAllowShrinking::No);
+		};
+
+		for(FInternalArray& Array : Stash)
+		{
+			EraseKeyFromArray(Array);
+		}
+
+		EraseKeyFromArray(Latest);
+	}
+
+    void AddAll(TTaskArray&& Other)
     {
         Canonicalize();
 
-        for (USharedPtr<TArray<T>>& StashedVectorBox : Other.Stash)
+        for (FInternalArray& StashedVectorBox : Other.Stash)
         {
             Stash.Push(MoveTemp(StashedVectorBox));
         }
 
         Other.Stash.Empty();
 
-        if (!Other.Latest->IsEmpty())
+        if (!Other.Latest.IsEmpty())
         {
             Stash.Push(MoveTemp(Other.Latest));
-			USharedPtr<TArray<T>> NewLatest(new TArray<T>());
-			Other.Latest = MoveTemp(NewLatest);
-        }
-    }
-
-    void AddAll(const TTaskArray<T>& Other)
-    {
-        Canonicalize();
-        Other.Canonicalize();
-        for (const USharedPtr<TArray<T>>& StashedVectorBox : Other.Stash)
-        {
-            Stash.Push(StashedVectorBox);
+			Other.Latest = FInternalArray{};
         }
-
-
-		ASSERT(Latest != nullptr);
     }
 
     template<typename TFunc>
     bool ForEachForward(const TFunc& Func) const
     {
-        for (USharedPtr<TArray<T>>& StashedVectorBox : Stash)
+        for (FInternalArray& StashedVectorBox : Stash)
         {
-            for (const T& Entry : *StashedVectorBox)
+            for (const SKeyValuePair& EntryKVP : StashedVectorBox)
             {
-                if (!Func(Entry))
+                if (!Func(EntryKVP.Val))
                 {
                     return false;
                 }
             }
         }
 
-        for (const T& Entry : *Latest)
+        for (const SKeyValuePair& EntryKVP : Latest)
         {
-            if (!Func(Entry))
+            if (!Func(EntryKVP.Val))
             {
                 return false;
             }
@@ -99,19 +140,19 @@ public:
     template<typename TFunc>
     bool ForEachBackward(const TFunc& Func) const
     {
-		for (const T& Entry : TBackwards(*Latest))
+		for (const SKeyValuePair& EntryKVP : TBackwards(Latest))
 		{
-			if (!Func(Entry))
+			if (!Func(EntryKVP.Val))
 			{
 				return false;
 			}
 		}
 
-		for (USharedPtr<TArray<T>>& StashedVectorBox : TBackwards(Stash))
+		for (const FInternalArray& StashedVectorBox : TBackwards(Stash))
 		{
-			for (const T& Entry : TBackwards(*StashedVectorBox))
+			for (const SKeyValuePair& EntryKVP : TBackwards(StashedVectorBox))
 			{
-				if (!Func(Entry))
+				if (!Func(EntryKVP.Val))
 				{
 					return false;
 				}
@@ -123,17 +164,17 @@ public:
 
     void Reset()
     {
-        Latest->Empty();
+        Latest.Empty();
         Stash.Empty();
     }
 
     size_t Num() const
     {
-        size_t Result = Latest->Num();
+        size_t Result = Latest.Num();
 
         for (size_t Index = 0; Index < Stash.Num(); Index++)
         {
-            const TArray<T>& StashedVector = *Stash[Index];
+            const FInternalArray& StashedVector = Stash[Index];
             Result += StashedVector.Num();
         }
 
@@ -146,19 +187,15 @@ private:
     // method is `const`.
     void Canonicalize() const
     {
-        if (!Latest->IsEmpty())
+        if (!Latest.IsEmpty())
         {
 			Stash.Push(MoveTemp(Latest));
-			USharedPtr<TArray<T>> NewLatest(new TArray<T>());
-			Latest = MoveTemp(NewLatest);
-        }
-
-		ASSERT(Latest != nullptr);
+			Latest = FInternalArray{};
+		}
     }
     
-    mutable USharedPtr<TArray<T>> Latest;
-    
-    mutable TArray<USharedPtr<TArray<T>>> Stash;
+    mutable FInternalArray Latest;
+    mutable TArray<FInternalArray> Stash;
 };
 
 } // namespace AutoRTFM
diff --git a/Engine/Source/Runtime/Core/Private/AutoRTFM/Transaction.h b/Engine/Source/Runtime/Core/Private/AutoRTFM/Transaction.h
--- a/Engine/Source/Runtime/Core/Private/AutoRTFM/Transaction.h
+++ b/Engine/Source/Runtime/Core/Private/AutoRTFM/Transaction.h
@@ -63,6 +63,8 @@ public:
 
     void DeferUntilCommit(TFunction<void()>&&);
     void DeferUntilAbort(TFunction<void()>&&);
+    void PushDeferUntilAbortHandler(const void* Key, TFunction<void()>&&);
+    void PopDeferUntilAbortHandler(const void* Key);
 
     void AbortAndThrow();
     void AbortWithoutThrowing();
diff --git a/Engine/Source/Runtime/Core/Private/AutoRTFM/TransactionInlines.h b/Engine/Source/Runtime/Core/Private/AutoRTFM/TransactionInlines.h
--- a/Engine/Source/Runtime/Core/Private/AutoRTFM/TransactionInlines.h
+++ b/Engine/Source/Runtime/Core/Private/AutoRTFM/TransactionInlines.h
@@ -158,6 +158,20 @@ UE_AUTORTFM_FORCEINLINE void FTransaction::DeferUntilAbort(TFunction<void()>&& C
     AbortTasks.Add(MoveTemp(Copy));
 }
 
+UE_AUTORTFM_FORCEINLINE void FTransaction::PushDeferUntilAbortHandler(const void* Key, TFunction<void()>&& Callback)
+{
+	// We explicitly must copy the function here because the original was allocated
+	// within a transactional context, and thus the memory is allocating under
+	// transactionalized conditions. By copying, we create an open copy of the callback.
+	TFunction<void()> Copy(Callback);
+    AbortTasks.AddKeyed(Key, MoveTemp(Copy));
+}
+
+UE_AUTORTFM_FORCEINLINE void FTransaction::PopDeferUntilAbortHandler(const void* Key)
+{
+    AbortTasks.DeleteKey(Key);
+}
+
 UE_AUTORTFM_FORCEINLINE void FTransaction::CollectStats() const
 {
     Stats.Collect<EStatsKind::AverageWriteLogEntries>(WriteLog.Num());
diff --git a/Engine/Source/Runtime/Core/Public/AutoRTFM/AutoRTFM.h b/Engine/Source/Runtime/Core/Public/AutoRTFM/AutoRTFM.h
--- a/Engine/Source/Runtime/Core/Public/AutoRTFM/AutoRTFM.h
+++ b/Engine/Source/Runtime/Core/Public/AutoRTFM/AutoRTFM.h
@@ -238,6 +238,23 @@ UE_AUTORTFM_FORCEINLINE void autortfm_on_abort(void (*work)(void* arg), void* ar
 }
 #endif
 
+#if UE_AUTORTFM
+UE_AUTORTFM_API void autortfm_push_on_abort_handler(const void* key, void (*work)(void* arg), void* arg);
+UE_AUTORTFM_API void autortfm_pop_on_abort_handler(const void* key);
+#else
+UE_AUTORTFM_FORCEINLINE void autortfm_push_on_abort_handler(const void* key, void (*work)(void* arg), void* arg)
+{
+	UE_AUTORTFM_UNUSED(key);
+	UE_AUTORTFM_UNUSED(work);
+	UE_AUTORTFM_UNUSED(arg);
+}
+
+UE_AUTORTFM_FORCEINLINE void autortfm_pop_on_abort_handler(const void* key)
+{
+	UE_AUTORTFM_UNUSED(key);
+}
+#endif
+
 [[deprecated("Use autortfm_on_abort instead.")]]
 UE_AUTORTFM_FORCEINLINE void autortfm_open_abort(void (*work)(void* arg), void* arg)
 {
@@ -468,10 +485,30 @@ template<typename TFunctor> UE_AUTORTFM_FORCEINLINE void OnCommit(const TFunctor
 // Have some work happen when this transaction aborts. If this is called
 // outside a transaction or from an open nest then the work is ignored.
 UE_AUTORTFM_API void OnAbort(TFunction<void()>&& Work);
+
+// Register a handler for transaction abort. Takes a key parameter so that
+// the handler can be unregistered (see PopOnAbortHandler). This is useful
+// for scoped mutations that need an abort handler present unless execution
+// reaches the end of the relevant scope.
+UE_AUTORTFM_API void PushOnAbortHandler(const void* Key, TFunction<void()>&& Work);
+
+// Unregister all handlers for transaction abort that were previously pushed
+// via PushOnAbortHandler with the given key
+UE_AUTORTFM_API void PopOnAbortHandler(const void* Key);
 #else
 // Have some work happen when this transaction aborts. If this is called
 // outside a transaction or from an open nest then the work is ignored.
 template<typename TFunctor> UE_AUTORTFM_FORCEINLINE void OnAbort(const TFunctor&) {}
+
+// Register a handler for transaction abort. Takes a key parameter so that
+// the handler can be unregistered (see PopOnAbortHandler). This is useful
+// for scoped mutations that need an abort handler present unless execution
+// reaches the end of the relevant scope.
+template<typename TFunctor> UE_AUTORTFM_FORCEINLINE void PushOnAbortHandler(const void* Key, const TFunctor&) {}
+
+// Unregister all handlers for transaction abort that were previously pushed
+// via PushOnAbortHandler with the given key
+UE_AUTORTFM_FORCEINLINE void PopOnAbortHandler(const void* Key) {}
 #endif
 
 #if UE_AUTORTFM
@@ -758,7 +795,7 @@ UE_AUTORTFM_FORCEINLINE void CheckConsistencyAssumingNoRaces() { return ForTheRu
 #define UE_AUTORTFM_CONCAT(A, B) UE_AUTORTFM_CONCAT_IMPL(A, B)
 
 #if UE_AUTORTFM
-#define UE_AUTORTFM_REGISTER_OPEN_FUNCTION_EXPLICIT_IMPL(OriginalFunction, NewFunction) static const ForTheRuntime::FRegisterOpenFunction UE_AUTORTFM_CONCAT(AutoRTFMFunctionRegistration, __COUNTER__)(reinterpret_cast<void*>(OriginalFunction), reinterpret_cast<void*>(NewFunction))
+#define UE_AUTORTFM_REGISTER_OPEN_FUNCTION_EXPLICIT_IMPL(OriginalFunction, NewFunction) static const AutoRTFM::ForTheRuntime::FRegisterOpenFunction UE_AUTORTFM_CONCAT(AutoRTFMFunctionRegistration, __COUNTER__)(reinterpret_cast<void*>(OriginalFunction), reinterpret_cast<void*>(NewFunction))
 #define UE_AUTORTFM_REGISTER_OPEN_FUNCTION_IMPL(OriginalFunction) UE_AUTORTFM_REGISTER_OPEN_FUNCTION_EXPLICIT(OriginalFunction, RTFM_ ## OriginalFunction)
 #define UE_AUTORTFM_REGISTER_SELF_FUNCTION_IMPL(OriginalFunction) UE_AUTORTFM_REGISTER_OPEN_FUNCTION_EXPLICIT(OriginalFunction, OriginalFunction)
 #else
diff --git a/Engine/Source/Runtime/CoreUObject/Private/UObject/UObjectHash.cpp b/Engine/Source/Runtime/CoreUObject/Private/UObject/UObjectHash.cpp
--- a/Engine/Source/Runtime/CoreUObject/Private/UObject/UObjectHash.cpp
+++ b/Engine/Source/Runtime/CoreUObject/Private/UObject/UObjectHash.cpp
@@ -522,7 +522,8 @@ public:
 		if (!(IsGarbageCollectingAndLockingUObjectHashTables() && IsInGameThread()))
 		{
 			Tables = &InTables;
-			InTables.Lock();
+			UE_AUTORTFM_OPEN({ InTables.Lock(); });
+			AutoRTFM::PushOnAbortHandler(this, [this](){ if (this->Tables){ this->Tables->Unlock(); }});
 		}
 		else
 		{
@@ -537,7 +538,8 @@ public:
 #if THREADSAFE_UOBJECTS
 		if (Tables)
 		{
-			Tables->Unlock();
+			UE_AUTORTFM_OPEN({ Tables->Unlock(); });
+			AutoRTFM::PopOnAbortHandler(this);
 		}
 #endif
 	}
