[Scheduler]
 - Deactivate dynamic thread creation during investigation of the presence of tasks with no threads to process them

diff --git a/Engine/Source/Runtime/Core/Private/Async/TaskGraph.cpp b/Engine/Source/Runtime/Core/Private/Async/TaskGraph.cpp
--- a/Engine/Source/Runtime/Core/Private/Async/TaskGraph.cpp
+++ b/Engine/Source/Runtime/Core/Private/Async/TaskGraph.cpp
@@ -120,7 +120,7 @@ static FAutoConsoleVariableRef CVarTaskOversubscriptionRatio(
 	ECVF_ReadOnly
 );
 
-CORE_API bool GTaskGraphUseDynamicThreadCreation = true;
+CORE_API bool GTaskGraphUseDynamicThreadCreation = (PLATFORM_DESKTOP == 1);
 static FAutoConsoleVariableRef CVarTaskDynamicThreadCreation(
 	TEXT("TaskGraph.UseDynamicThreadCreation"),
 	GTaskGraphUseDynamicThreadCreation,
