Added intrusive unset optional states to TSharedPtr and TWeakPtr.

diff --git a/Engine/Source/Runtime/Core/Public/Templates/SharedPointer.h b/Engine/Source/Runtime/Core/Public/Templates/SharedPointer.h
--- a/Engine/Source/Runtime/Core/Public/Templates/SharedPointer.h
+++ b/Engine/Source/Runtime/Core/Public/Templates/SharedPointer.h
@@ -1152,6 +1152,28 @@ public:
 		return SharedReferenceCount.IsUnique();
 	}
 
+	///////////////////////////////////////////////////
+	// Start - intrusive TOptional<TSharedPtr> state //
+	///////////////////////////////////////////////////
+	constexpr static bool bHasIntrusiveUnsetOptionalState = true;
+	using IntrusiveUnsetOptionalStateType = TSharedPtr;
+
+	explicit TSharedPtr(FIntrusiveUnsetOptionalState)
+		: Object((ObjectType*)-1)
+	{
+	}
+	void operator=(FIntrusiveUnsetOptionalState)
+	{
+		Object = (ObjectType*)-1;
+		SharedReferenceCount = {};
+	}
+	bool operator==(FIntrusiveUnsetOptionalState) const
+	{
+		return Object == (ObjectType*)-1;
+	}
+	/////////////////////////////////////////////////
+	// End - intrusive TOptional<TSharedPtr> state //
+	/////////////////////////////////////////////////
 private:
 
 	/**
@@ -1526,6 +1548,28 @@ public:
 		return ::PointerHash( Object );
 	}
 
+	/////////////////////////////////////////////////
+	// Start - intrusive TOptional<TWeakPtr> state //
+	/////////////////////////////////////////////////
+	constexpr static bool bHasIntrusiveUnsetOptionalState = true;
+	using IntrusiveUnsetOptionalStateType = TWeakPtr;
+
+	explicit TWeakPtr(FIntrusiveUnsetOptionalState)
+		: Object((ObjectType*)-1)
+	{
+	}
+	void operator=(FIntrusiveUnsetOptionalState)
+	{
+		Object = (ObjectType*)-1;
+		WeakReferenceCount = SharedPointerInternals::FWeakReferencer< Mode >{};
+	}
+	bool operator==(FIntrusiveUnsetOptionalState) const
+	{
+		return Object == (ObjectType*)-1;
+	}
+	///////////////////////////////////////////////
+	// End - intrusive TOptional<TWeakPtr> state //
+	///////////////////////////////////////////////
 private:
 	
 	// We declare ourselves as a friend (templated using OtherType) so we can access members as needed
