[Scheduler]
 - Fix deadlock when dynamic thread creation is active and there are tasks left in the queue during call to RestartWorkers

diff --git a/Engine/Source/Runtime/Core/Private/Async/Fundamental/Scheduler.cpp b/Engine/Source/Runtime/Core/Private/Async/Fundamental/Scheduler.cpp
--- a/Engine/Source/Runtime/Core/Private/Async/Fundamental/Scheduler.cpp
+++ b/Engine/Source/Runtime/Core/Private/Async/Fundamental/Scheduler.cpp
@@ -199,6 +199,13 @@ namespace LowLevelTasks
 					BackgroundCreateThread();
 				}
 			}
+			else if (TemporaryShutdown)
+			{
+				// Since the global queue is not drained during temporary shutdown, kick threads
+				// here so we can continue work if there was any tasks left when we stopped the workers.
+				WaitingQueue[0].Notify();
+				WaitingQueue[1].Notify();
+			}
 		}
 	}
 
@@ -247,7 +254,7 @@ namespace LowLevelTasks
 		return OutTask;
 	}
 
-	void FScheduler::StopWorkers(bool DrainGlobalQueue)
+	void FScheduler::StopWorkers(bool bDrainGlobalQueue)
 	{
 		uint32 OldActiveWorkers = ActiveWorkers.load(std::memory_order_relaxed);
 		if (OldActiveWorkers != 0 && ActiveWorkers.compare_exchange_strong(OldActiveWorkers, 0, std::memory_order_relaxed))
@@ -273,7 +280,7 @@ namespace LowLevelTasks
 			WorkerLocalQueues.Reset();
 			WorkerEvents.Reset();
 
-			if (DrainGlobalQueue)
+			if (bDrainGlobalQueue)
 			{
 				for (FTask* Task = QueueRegistry.DequeueGlobal(); Task != nullptr; Task = QueueRegistry.DequeueGlobal())
 				{
@@ -295,7 +302,8 @@ namespace LowLevelTasks
 	{
 		FScopeLock Lock(&WorkerThreadsCS);
 		TemporaryShutdown.store(true, std::memory_order_release);
-		StopWorkers(false);
+		const bool bDrainGlobalQueue = false;
+		StopWorkers(bDrainGlobalQueue);
 		StartWorkers(NumForegroundWorkers, NumBackgroundWorkers, IsForkable, InWorkerPriority, InBackgroundPriority, InWorkerAffinity, InBackgroundAffinity);
 		TemporaryShutdown.store(false, std::memory_order_release);
 	}
