Fixed Clang complaining about deprecated enums when used with a TEnumRange.

diff --git a/Engine/Source/Runtime/Core/Public/Misc/EnumRange.h b/Engine/Source/Runtime/Core/Public/Misc/EnumRange.h
--- a/Engine/Source/Runtime/Core/Public/Misc/EnumRange.h
+++ b/Engine/Source/Runtime/Core/Public/Misc/EnumRange.h
@@ -3,6 +3,7 @@
 #pragma once
 
 #include "CoreTypes.h"
+#include <type_traits>
 
 /**
  * Defines a contiguous enum range containing Count values, starting from zero:
@@ -21,7 +22,7 @@
  * // Defines iteration over ECountedThing to be: First, Second, Third
  * ENUM_RANGE_BY_COUNT(ECountedThing, ECountedThing::Count)
  */
-#define ENUM_RANGE_BY_COUNT(EnumType, Count) ENUM_RANGE_BY_FIRST_AND_LAST(EnumType, 0, (__underlying_type(EnumType))(Count) - 1)
+#define ENUM_RANGE_BY_COUNT(EnumType, Count) ENUM_RANGE_BY_FIRST_AND_LAST(EnumType, 0, (std::underlying_type_t<EnumType>)(Count) - 1)
 
 
 /**
@@ -48,8 +49,8 @@
 	struct NEnumRangePrivate::TEnumRangeTraits<EnumType> \
 	{ \
 		enum { RangeType = 0 }; \
-		static constexpr __underlying_type(EnumType) Begin = (__underlying_type(EnumType))(First); \
-		static constexpr __underlying_type(EnumType) End   = (__underlying_type(EnumType))(Last) + 1; \
+		static constexpr std::underlying_type_t<EnumType> Begin = (std::underlying_type_t<EnumType>)(First); \
+		static constexpr std::underlying_type_t<EnumType> End   = (std::underlying_type_t<EnumType>)(Last) + 1; \
 	};
 
 
@@ -101,7 +102,7 @@ namespace NEnumRangePrivate
 	template <typename EnumType>
 	struct TEnumContiguousIterator
 	{
-		typedef __underlying_type(EnumType) IntType;
+		using IntType = std::underlying_type_t<EnumType>;
 
 		FORCEINLINE explicit TEnumContiguousIterator(IntType InValue)
 			: Value(InValue)
@@ -176,7 +177,7 @@ namespace UE::EnumFlags::Private
 	template <typename EnumType>
 	struct TIterator
 	{
-		typedef __underlying_type(EnumType) IntType;
+		using IntType = std::underlying_type_t<EnumType>;
 
 		FORCEINLINE explicit TIterator(EnumType InFlags)
 			: Flags(IntType(InFlags))
