Exposed missing delegate variants on TAttribute (SPLambda, WeakLambda)
#tests editor

diff --git a/Engine/Source/Runtime/Core/Public/Misc/Attribute.h b/Engine/Source/Runtime/Core/Public/Misc/Attribute.h
--- a/Engine/Source/Runtime/Core/Public/Misc/Attribute.h
+++ b/Engine/Source/Runtime/Core/Public/Misc/Attribute.h
@@ -177,7 +177,30 @@ public:
 	template<typename LambdaType, typename... PayloadTypes>
 	[[nodiscard]] FORCEINLINE static TAttribute CreateLambda(LambdaType&& InCallable, PayloadTypes&&... InputPayload)
 	{
-		return Create(FGetter::CreateLambda(InCallable, Forward<PayloadTypes>(InputPayload)...));
+		return Create(FGetter::CreateLambda(Forward<LambdaType>(InCallable), Forward<PayloadTypes>(InputPayload)...));
+	}
+
+	/**
+	 * Helper function for creating TAttributes from a weak shared pointer C++ lambda delegate
+	 */
+	template<typename UserClass, ESPMode Mode, typename LambdaType, typename... PayloadTypes>
+	[[nodiscard]] FORCEINLINE static TAttribute CreateSPLambda(const TSharedRef<UserClass, Mode>& InUserObjectRef, LambdaType&& InCallable, PayloadTypes&&... InputPayload)
+	{
+		return Create(FGetter::CreateSPLambda(InUserObjectRef, Forward<LambdaType>(InCallable), Forward<PayloadTypes>(InputPayload)...));
+	}
+	template <typename UserClass, typename LambdaType, typename... PayloadTypes>
+	[[nodiscard]] FORCEINLINE static TAttribute CreateSPLambda(UserClass* InUserObject, LambdaType&& InCallable, PayloadTypes&&... InputPayload)
+	{
+		return Create(FGetter::CreateSPLambda(InUserObject, Forward<LambdaType>(InCallable), Forward<PayloadTypes>(InputPayload)...));
+	}
+
+	/**
+	 * Helper function for creating TAttributes from a weak object pointer C++ lambda delegate
+	 */
+	template <typename UserClass, typename LambdaType, typename... PayloadTypes>
+	[[nodiscard]] FORCEINLINE static TAttribute CreateWeakLambda(UserClass* InUserObject, LambdaType&& InCallable, PayloadTypes&&... InputPayload)
+	{
+		return Create(FGetter::CreateWeakLambda(InUserObject, Forward<LambdaType>(InCallable), Forward<PayloadTypes>(InputPayload)...));
 	}
 
 	/**
@@ -380,6 +403,45 @@ public:
 		Getter.BindUFunction(InUserObject, InFunctionName);
 	}
 
+	/**
+	 * Static: Binds a C++ lambda delegate
+	 * technically this works for any functor types, but lambdas are the primary use case
+	 */
+	template<typename LambdaType, typename... PayloadTypes>
+	inline void BindLambda(LambdaType&& InCallable, PayloadTypes&&... InputPayload)
+	{
+		bIsSet = true;
+		Getter.BindLambda(Forward<LambdaType>(InCallable), Forward<PayloadTypes>(InputPayload)...);
+	}
+
+	/**
+	 * Static: Binds a weak shared pointer C++ lambda delegate
+	 * technically this works for any functor types, but lambdas are the primary use case
+	 */
+	template<typename UserClass, ESPMode Mode, typename LambdaType, typename... PayloadTypes>
+	inline void BindSPLambda(const TSharedRef<UserClass, Mode>& InUserObjectRef, LambdaType&& InCallable, PayloadTypes&&... InputPayload)
+	{
+		bIsSet = true;
+		Getter.BindSPLambda(InUserObjectRef, Forward<LambdaType>(InCallable), Forward<PayloadTypes>(InputPayload)...);
+	}
+	template <typename UserClass, typename LambdaType, typename... PayloadTypes>
+	inline void BindSPLambda(const UserClass* InUserObject, LambdaType&& InCallable, PayloadTypes&&... InputPayload)
+	{
+		bIsSet = true;
+		Getter.BindSPLambda(InUserObject, Forward<LambdaType>(InCallable), Forward<PayloadTypes>(InputPayload)...);
+	}
+
+	/**
+	 * Static: Binds a weak object C++ lambda delegate
+	 * technically this works for any functor types, but lambdas are the primary use case
+	 */
+	template<typename UserClass, typename LambdaType, typename... PayloadTypes>
+	inline void BindWeakLambda(UserClass* InUserObject, LambdaType&& InCallable, PayloadTypes&&... InputPayload)
+	{
+		bIsSet = true;
+		Getter.BindWeakLambda(InUserObject, Forward<LambdaType>(InCallable), Forward<PayloadTypes>(InputPayload)...);
+	}
+
 	/**
 	 * Checks to see if this attribute has a 'getter' function bound
 	 *
@@ -519,5 +581,34 @@ template<typename LambdaType, typename... PayloadTypes>
 {
 	typedef decltype(InCallable(DeclVal<PayloadTypes>()...)) T;
 
-	return TAttribute<T>::Create(TAttribute<T>::FGetter::CreateLambda(InCallable, Forward<PayloadTypes>(InputPayload)...));
+	return TAttribute<T>::Create(TAttribute<T>::FGetter::CreateLambda(Forward<LambdaType>(InCallable), Forward<PayloadTypes>(InputPayload)...));
 }
+
+/**
+ * Helper function for creating TAttributes from a weak shared pointer C++ lambda delegate
+ */
+template<typename UserClass, ESPMode Mode, typename LambdaType, typename... PayloadTypes>
+[[nodiscard]] decltype(auto) MakeAttributeSPLambda(const TSharedRef<UserClass, Mode>& InUserObjectRef, LambdaType&& InCallable, PayloadTypes&&... InputPayload)
+{
+	typedef decltype(InCallable(DeclVal<PayloadTypes>()...)) T;
+
+	return TAttribute<T>::Create(TAttribute<T>::FGetter::CreateSPLambda(InUserObjectRef, Forward<LambdaType>(InCallable), Forward<PayloadTypes>(InputPayload)...));
+}
+template <typename UserClass, typename LambdaType, typename... PayloadTypes>
+[[nodiscard]] decltype(auto) MakeAttributeSPLambda(UserClass* InUserObject, LambdaType&& InCallable, PayloadTypes&&... InputPayload)
+{
+	typedef decltype(InCallable(DeclVal<PayloadTypes>()...)) T;
+
+	return TAttribute<T>::Create(TAttribute<T>::FGetter::CreateSPLambda(InUserObject, Forward<LambdaType>(InCallable), Forward<PayloadTypes>(InputPayload)...));
+}
+
+/**
+ * Helper function for creating TAttributes from a weak object pointer C++ lambda delegate
+ */
+template <typename UserClass, typename LambdaType, typename... PayloadTypes>
+[[nodiscard]] decltype(auto) MakeAttributeWeakLambda(UserClass* InUserObject, LambdaType&& InCallable, PayloadTypes&&... InputPayload)
+{
+	typedef decltype(InCallable(DeclVal<PayloadTypes>()...)) T;
+
+	return TAttribute<T>::Create(TAttribute<T>::FGetter::CreateWeakLambda(InUserObject, Forward<LambdaType>(InCallable), Forward<PayloadTypes>(InputPayload)...));
+}
\ No newline at end of file
