Add compiler flag PLATFORM_WRITES_ARE_SLOW and set true for Switch. Use GIOThreadPool instead of normal Async threads for slow writes.

diff --git a/Engine/Source/Runtime/Core/Private/Misc/ConfigUtilities.cpp b/Engine/Source/Runtime/Core/Private/Misc/ConfigUtilities.cpp
--- a/Engine/Source/Runtime/Core/Private/Misc/ConfigUtilities.cpp
+++ b/Engine/Source/Runtime/Core/Private/Misc/ConfigUtilities.cpp
@@ -101,6 +101,27 @@ void SaveCVarForNextBoot(const TCHAR* Key, const TCHAR* Value)
 	FString StrKey(Key);
 	FString StrValue(Value);
 
+#if PLATFORM_WRITES_ARE_SLOW
+	AsyncPool(*GIOThreadPool, [StrKey, StrValue] {
+		TMap<FString, FString> CVarsToSave;
+
+		// Read from file, in case there are more than one cvar hotfix event in same run
+		LoadCVarsFromFileForNextBoot(CVarsToSave);
+
+		CVarsToSave.FindOrAdd(StrKey) = StrValue;
+
+		FString ContentToSave;
+		for (const TPair<FString, FString>& CVarPair : CVarsToSave)
+		{
+			ContentToSave.Append(FString::Format(TEXT("{0}={1}\r\n"), { CVarPair.Key, CVarPair.Value }));
+		}
+
+		const FString FullPath = FPaths::ProjectPersistentDownloadDir() / UE_HOTFIX_FOR_NEXT_BOOT_FILENAME;
+		FFileHelper::SaveStringToFile(ContentToSave, *FullPath);
+
+		UE_LOG(LogConfig, Log, TEXT("Local boot hotfix file [%s] saved with hotfixed CVar: %s=%s"), *FullPath, *StrKey, *StrValue);
+		});
+#else
 	AsyncTaskPipe.Launch(UE_SOURCE_LOCATION, [StrKey, StrValue] {
 		TMap<FString, FString> CVarsToSave;
 
@@ -120,6 +141,7 @@ void SaveCVarForNextBoot(const TCHAR* Key, const TCHAR* Value)
 
 		UE_LOG(LogConfig, Log, TEXT("Local boot hotfix file [%s] saved with hotfixed CVar: %s=%s"), *FullPath, *StrKey, *StrValue);
 	}, LowLevelTasks::ETaskPriority::BackgroundLow);
+#endif // PLATFORM_WRITES_ARE_SLOW
 #endif // !UE_SERVER
 }
 
diff --git a/Engine/Source/Runtime/Core/Public/HAL/Platform.h b/Engine/Source/Runtime/Core/Public/HAL/Platform.h
--- a/Engine/Source/Runtime/Core/Public/HAL/Platform.h
+++ b/Engine/Source/Runtime/Core/Public/HAL/Platform.h
@@ -607,6 +607,11 @@
 #ifndef PLATFORM_IMPLEMENTS_BATCH_FILE_DELETE
 	#define PLATFORM_IMPLEMENTS_BATCH_FILE_DELETE 0
 #endif
+
+#ifndef PLATFORM_WRITES_ARE_SLOW
+	#define PLATFORM_WRITES_ARE_SLOW 0
+#endif
+
 // deprecated, do not use
 #define PLATFORM_HAS_THREADSAFE_RHIGetRenderQueryResult	#
 #define PLATFORM_SUPPORTS_RHI_THREAD #
