Fix non-Unity compile issue when PLATFORM_HAS_SLOW_WRITES is enabled.

diff --git a/Engine/Source/Runtime/Core/Private/Misc/ConfigUtilities.cpp b/Engine/Source/Runtime/Core/Private/Misc/ConfigUtilities.cpp
--- a/Engine/Source/Runtime/Core/Private/Misc/ConfigUtilities.cpp
+++ b/Engine/Source/Runtime/Core/Private/Misc/ConfigUtilities.cpp
@@ -8,6 +8,10 @@
 #include "Misc/FileHelper.h"
 #include "Tasks/Pipe.h"
 
+#if PLATFORM_WRITES_ARE_SLOW
+#include "Misc/QueuedThreadPool.h"
+#endif
+
 #define UE_HOTFIX_FOR_NEXT_BOOT_FILENAME TEXT("HotfixForNextBoot.txt")
 
 namespace UE::ConfigUtilities
