Expose more of the Crash Reporter tags.

diff --git a/Engine/Source/Runtime/Core/Private/GenericPlatform/GenericPlatformCrashContext.cpp b/Engine/Source/Runtime/Core/Private/GenericPlatform/GenericPlatformCrashContext.cpp
--- a/Engine/Source/Runtime/Core/Private/GenericPlatform/GenericPlatformCrashContext.cpp
+++ b/Engine/Source/Runtime/Core/Private/GenericPlatform/GenericPlatformCrashContext.cpp
@@ -202,6 +202,7 @@ const TCHAR* const FGenericCrashContext::CrashTypeTag = TEXT("CrashType");
 const TCHAR* const FGenericCrashContext::ErrorMessageTag = TEXT("ErrorMessage");
 const TCHAR* const FGenericCrashContext::CrashReporterMessageTag = TEXT("CrashReporterMessage");
 const TCHAR* const FGenericCrashContext::AttendedStatusTag = TEXT("CrashReporterMessage");
+const TCHAR* const FGenericCrashContext::ProcessIdTag = TEXT("ProcessId");
 const TCHAR* const FGenericCrashContext::SecondsSinceStartTag = TEXT("SecondsSinceStart");
 const TCHAR* const FGenericCrashContext::BuildVersionTag = TEXT("BuildVersion");
 const TCHAR* const FGenericCrashContext::CallStackTag = TEXT("CallStack");
@@ -209,6 +210,11 @@ const TCHAR* const FGenericCrashContext::PortableCallStackTag = TEXT("PCallStack
 const TCHAR* const FGenericCrashContext::PortableCallStackHashTag = TEXT("PCallStackHash");
 const TCHAR* const FGenericCrashContext::IsRequestingExitTag = TEXT("IsRequestingExit");
 const TCHAR* const FGenericCrashContext::LogFilePathTag = TEXT("LogFilePath");
+const TCHAR* const FGenericCrashContext::IsInternalBuildTag = TEXT("IsInternalBuild");
+const TCHAR* const FGenericCrashContext::IsPerforceBuildTag = TEXT("IsPerforceBuild");
+const TCHAR* const FGenericCrashContext::IsWithDebugInfoTag = TEXT("IsWithDebugInfo");
+const TCHAR* const FGenericCrashContext::IsSourceDistributionTag = TEXT("IsSourceDistribution");
+
 const TCHAR* const FGenericCrashContext::UEMinidumpName = TEXT( "UEMinidump.dmp" );
 const TCHAR* const FGenericCrashContext::NewLineTag = TEXT( "&nl;" );
 
@@ -773,13 +779,13 @@ void FGenericCrashContext::SerializeTempCrashContextToFile()
 // fields in NCached!
 void FGenericCrashContext::SerializeSessionContext(FString& Buffer)
 {
-	AddCrashPropertyInternal(Buffer, TEXT("ProcessId"), NCached::Session.ProcessId);
+	AddCrashPropertyInternal(Buffer, FGenericCrashContext::ProcessIdTag, NCached::Session.ProcessId);
 	AddCrashPropertyInternal(Buffer, FGenericCrashContext::SecondsSinceStartTag, NCached::Session.SecondsSinceStart);
 
-	AddCrashPropertyInternal(Buffer, TEXT("IsInternalBuild"), NCached::Session.bIsInternalBuild);
-	AddCrashPropertyInternal(Buffer, TEXT("IsPerforceBuild"), NCached::Session.bIsPerforceBuild);
-	AddCrashPropertyInternal(Buffer, TEXT("IsWithDebugInfo"), NCached::Session.bWithDebugInfo);
-	AddCrashPropertyInternal(Buffer, TEXT("IsSourceDistribution"), NCached::Session.bIsSourceDistribution);
+	AddCrashPropertyInternal(Buffer, FGenericCrashContext::IsInternalBuildTag, NCached::Session.bIsInternalBuild);
+	AddCrashPropertyInternal(Buffer, FGenericCrashContext::IsPerforceBuildTag, NCached::Session.bIsPerforceBuild);
+	AddCrashPropertyInternal(Buffer, FGenericCrashContext::IsWithDebugInfoTag, NCached::Session.bWithDebugInfo);
+	AddCrashPropertyInternal(Buffer, FGenericCrashContext::IsSourceDistributionTag, NCached::Session.bIsSourceDistribution);
 
 	if (FCString::Strlen(NCached::Session.GameName) > 0)
 	{
diff --git a/Engine/Source/Runtime/Core/Public/GenericPlatform/GenericPlatformCrashContext.h b/Engine/Source/Runtime/Core/Public/GenericPlatform/GenericPlatformCrashContext.h
--- a/Engine/Source/Runtime/Core/Public/GenericPlatform/GenericPlatformCrashContext.h
+++ b/Engine/Source/Runtime/Core/Public/GenericPlatform/GenericPlatformCrashContext.h
@@ -399,6 +399,11 @@ public:
 	CORE_API static const TCHAR* const PortableCallStackHashTag;
 	CORE_API static const TCHAR* const IsRequestingExitTag;
 	CORE_API static const TCHAR* const LogFilePathTag;
+	CORE_API static const TCHAR* const ProcessIdTag;
+	CORE_API static const TCHAR* const IsInternalBuildTag;
+	CORE_API static const TCHAR* const IsPerforceBuildTag;
+	CORE_API static const TCHAR* const IsWithDebugInfoTag;
+	CORE_API static const TCHAR* const IsSourceDistributionTag;
 
 	static constexpr inline int32 CrashGUIDLength = 128;
 
