Removed deprecated TAttribute functions

diff --git a/Engine/Source/Runtime/Core/Public/Misc/Attribute.h b/Engine/Source/Runtime/Core/Public/Misc/Attribute.h
--- a/Engine/Source/Runtime/Core/Public/Misc/Attribute.h
+++ b/Engine/Source/Runtime/Core/Public/Misc/Attribute.h
@@ -132,16 +132,6 @@ public:
 		return TAttribute(FGetter::CreateStatic(InFuncPtr, MoveTemp(Vars)...), bExplicitConstructor);
 	}
 
-	/**
-	 * Helper function for creating TAttributes from a function pointer, accessed through a raw pointer
-	 */
-	template<typename SourceType, typename SourceTypeOrBase, typename... PayloadTypes>
-	UE_DEPRECATED(5.0, "Attribute's Getter should be const.")
-	[[nodiscard]] FORCEINLINE static TAttribute CreateRaw(SourceType* InObject, ObjectType (SourceTypeOrBase::*InMethod)(PayloadTypes...), typename TDecay<PayloadTypes>::Type... InputPayload)
-	{
-		return Create(FGetter::CreateRaw(InObject, InMethod, MoveTemp(InputPayload)...));
-	}
-
 	/**
 	 * Helper function for creating TAttributes from a const member function pointer, accessed through a raw pointer
 	 */
@@ -151,16 +141,6 @@ public:
 		return Create(FGetter::CreateRaw(InObject, InMethod, MoveTemp(InputPayload)...));
 	}
 
-	/**
-	 * Helper function for creating TAttributes from a non-const member function pointer, accessed through a weak pointer to the shared object
-	 */
-	template<typename SourceType, typename SourceTypeOrBase, typename... PayloadTypes>
-	UE_DEPRECATED(5.0, "Attribute's Getter should be const.")
-	[[nodiscard]] FORCEINLINE static TAttribute CreateSP(SourceType* InObject, ObjectType (SourceTypeOrBase::*InMethod)(PayloadTypes...), typename TDecay<PayloadTypes>::Type... InputPayload)
-	{
-		return Create(FGetter::CreateSP(InObject, InMethod, MoveTemp(InputPayload)...));
-	}
-
 	/**
 	 * Helper function for creating TAttributes from a const member function pointer, accessed through a weak pointer to the shared object
 	 */
@@ -523,15 +503,6 @@ private:
 	FGetter Getter;
 };
 
-/**
- * Helper function for creating TAttributes from a non-const member function pointer, accessed through a raw pointer
- */
-template<typename T, typename SourceType, typename SourceTypeOrBase, typename... PayloadTypes>
-UE_DEPRECATED(5.0, "Attribute's Getter should be const.")
-[[nodiscard]] FORCEINLINE TAttribute<T> MakeAttributeRaw(SourceType* InObject, T (SourceTypeOrBase::*InMethod)(PayloadTypes...), typename TDecay<PayloadTypes>::Type... InputPayload)
-{
-	return TAttribute<T>::Create(TAttribute<T>::FGetter::CreateRaw(InObject, InMethod, MoveTemp(InputPayload)...));
-}
 
 /**
  * Helper function for creating TAttributes from a const member function pointer, accessed through a raw pointer
@@ -542,16 +513,6 @@ template<typename T, typename SourceType, typename SourceTypeOrBase, typename...
 	return TAttribute<T>::Create(TAttribute<T>::FGetter::CreateRaw(InObject, InMethod, MoveTemp(InputPayload)...));
 }
 
-/**
- * Helper function for creating TAttributes from a non-const member function pointer, accessed through a weak pointer to the shared object
- */
-template<typename T, typename SourceType, typename SourceTypeOrBase, typename... PayloadTypes>
-UE_DEPRECATED(5.0, "Attribute's Getter should be const.")
-[[nodiscard]] FORCEINLINE TAttribute<T> MakeAttributeSP(SourceType* InObject, T (SourceTypeOrBase::*InMethod)(PayloadTypes...), typename TDecay<PayloadTypes>::Type...  InputPayload)
-{
-	return TAttribute<T>::Create(TAttribute<T>::FGetter::CreateSP(InObject, InMethod, MoveTemp(InputPayload)...));
-}
-
 /**
  * Helper function for creating TAttributes from a const member function pointer, accessed through a weak pointer to the shared object
  */
