Properly getting Mac memory pressure status

diff --git a/Engine/Source/Runtime/Core/Private/Mac/MacPlatformMemory.cpp b/Engine/Source/Runtime/Core/Private/Mac/MacPlatformMemory.cpp
--- a/Engine/Source/Runtime/Core/Private/Mac/MacPlatformMemory.cpp
+++ b/Engine/Source/Runtime/Core/Private/Mac/MacPlatformMemory.cpp
@@ -43,6 +43,8 @@ void* CFNetwork_CFAllocatorOperatorNew_Replacement(unsigned long Size, CFAllocat
 }
 #endif // PLATFORM_MAC_X86
 
+FGenericPlatformMemoryStats::EMemoryPressureStatus FMacPlatformMemory::MemoryPressureStatus = FGenericPlatformMemoryStats::EMemoryPressureStatus::Unknown;
+
 static bool HasArg(const char *Arg)
 {
 	if (_NSGetArgc() && _NSGetArgv())
@@ -64,6 +66,19 @@ static bool HasArg(const char *Arg)
 
 FMalloc* FMacPlatformMemory::BaseAllocator()
 {
+	auto dispatch_source = dispatch_source_create(DISPATCH_SOURCE_TYPE_MEMORYPRESSURE, 0, DISPATCH_MEMORYPRESSURE_NORMAL, dispatch_get_main_queue());
+	dispatch_source_set_event_handler(dispatch_source, ^{FMacPlatformMemory::MemoryPressureStatus = FGenericPlatformMemoryStats::EMemoryPressureStatus::Nominal;});
+	dispatch_activate(dispatch_source);
+	dispatch_retain(dispatch_source);
+	dispatch_source = dispatch_source_create(DISPATCH_SOURCE_TYPE_MEMORYPRESSURE, 0, DISPATCH_MEMORYPRESSURE_WARN, dispatch_get_main_queue());
+	dispatch_source_set_event_handler(dispatch_source, ^{FMacPlatformMemory::MemoryPressureStatus = FGenericPlatformMemoryStats::EMemoryPressureStatus::Warning;});
+	dispatch_activate(dispatch_source);
+	dispatch_retain(dispatch_source);
+	dispatch_source = dispatch_source_create(DISPATCH_SOURCE_TYPE_MEMORYPRESSURE, 0, DISPATCH_MEMORYPRESSURE_CRITICAL, dispatch_get_main_queue());
+	dispatch_source_set_event_handler(dispatch_source, ^{FMacPlatformMemory::MemoryPressureStatus = FGenericPlatformMemoryStats::EMemoryPressureStatus::Critical;});
+	dispatch_activate(dispatch_source);
+	dispatch_retain(dispatch_source);
+	
 	static FMalloc* Instance = nullptr;
 	if (Instance != nullptr)
 	{
@@ -179,7 +194,7 @@ FPlatformMemoryStats FMacPlatformMemory::GetStats()
 	{
 		MemoryStats.PeakUsedVirtual = MemoryStats.UsedVirtual;
 	}
-	
+	MemoryStats.MemoryPressureStatus = MemoryPressureStatus;
 
 	return MemoryStats;
 }
diff --git a/Engine/Source/Runtime/Core/Public/Apple/ApplePlatformMemory.h b/Engine/Source/Runtime/Core/Public/Apple/ApplePlatformMemory.h
--- a/Engine/Source/Runtime/Core/Public/Apple/ApplePlatformMemory.h
+++ b/Engine/Source/Runtime/Core/Public/Apple/ApplePlatformMemory.h
@@ -55,7 +55,27 @@ OBJC_EXPORT @interface FApplePlatformObject : NSObject
  *	Max implementation of the FGenericPlatformMemoryStats.
  */
 struct FPlatformMemoryStats : public FGenericPlatformMemoryStats
-{};
+{
+	FPlatformMemoryStats()
+		: FGenericPlatformMemoryStats()
+		, MemoryPressureStatus(EMemoryPressureStatus::Unknown)
+	{}
+	
+	EMemoryPressureStatus GetMemoryPressureStatus() const
+	{
+		if (MemoryPressureStatus == EMemoryPressureStatus::Unknown)
+		{
+			// if platform doesn't make use of MemoryPressureStatus, use default implementation
+			return FGenericPlatformMemoryStats::GetMemoryPressureStatus();
+		}
+		else
+		{
+			return MemoryPressureStatus;
+		}
+	}
+	
+	EMemoryPressureStatus MemoryPressureStatus;
+};
 
 /**
  * Common Apple platform memory functions.
diff --git a/Engine/Source/Runtime/Core/Public/Mac/MacPlatformMemory.h b/Engine/Source/Runtime/Core/Public/Mac/MacPlatformMemory.h
--- a/Engine/Source/Runtime/Core/Public/Mac/MacPlatformMemory.h
+++ b/Engine/Source/Runtime/Core/Public/Mac/MacPlatformMemory.h
@@ -18,6 +18,8 @@ struct CORE_API FMacPlatformMemory : public FApplePlatformMemory
 	static FPlatformMemoryStats GetStats();
 	static const FPlatformMemoryConstants& GetConstants();
 	static FMalloc* BaseAllocator();
+	
+	static FGenericPlatformMemoryStats::EMemoryPressureStatus MemoryPressureStatus;
 	//~ End FGenericPlatformMemory Interface
 };
 
