Optimized TConstSetBitIterator by adding a constructor without StartIndex

#rnx

diff --git a/Engine/Source/Runtime/Core/Public/Containers/BitArray.h b/Engine/Source/Runtime/Core/Public/Containers/BitArray.h
--- a/Engine/Source/Runtime/Core/Public/Containers/BitArray.h
+++ b/Engine/Source/Runtime/Core/Public/Containers/BitArray.h
@@ -1805,8 +1805,20 @@ class TConstSetBitIterator : public FRelativeBitReference
 {
 public:
 
-	/** Constructor. */
-	explicit TConstSetBitIterator(const TBitArray<Allocator>& InArray UE_LIFETIMEBOUND,int32 StartIndex = 0)
+	explicit TConstSetBitIterator(const TBitArray<Allocator>& InArray UE_LIFETIMEBOUND)
+		: FRelativeBitReference(0)
+		, Array                (InArray)
+		, UnvisitedBitMask     (~0U)
+		, CurrentBitIndex      (0)
+		, BaseBitIndex         (0)
+	{
+		if (Array.Num())
+		{
+			FindFirstSetBit();
+		}
+	}
+
+	explicit TConstSetBitIterator(const TBitArray<Allocator>& InArray UE_LIFETIMEBOUND, int32 StartIndex)
 		: FRelativeBitReference(StartIndex)
 		, Array                (InArray)
 		, UnvisitedBitMask     ((~0U) << (StartIndex & (NumBitsPerDWORD - 1)))
