Temp workaround for incorrect deprecation warning, reported fixed in VS2022 17.10

#rnx

diff --git a/Engine/Source/Runtime/Core/Public/Delegates/DelegateInstanceInterface.h b/Engine/Source/Runtime/Core/Public/Delegates/DelegateInstanceInterface.h
--- a/Engine/Source/Runtime/Core/Public/Delegates/DelegateInstanceInterface.h
+++ b/Engine/Source/Runtime/Core/Public/Delegates/DelegateInstanceInterface.h
@@ -42,6 +42,9 @@ struct IBaseDelegateInstance<RetType(ArgTypes...), UserPolicy> : public UserPoli
 	virtual bool ExecuteIfSafe(ArgTypes...) const = 0;
 };
 
+// Temp workaround for deprecation warnings being emitted by this template with VS2022 17.9. Reported fixed in 17.10
+PRAGMA_DISABLE_DEPRECATION_WARNINGS
+
 template <bool Const, typename Class, typename FuncType>
 struct TMemFunPtrType;
 
@@ -57,6 +60,9 @@ struct TMemFunPtrType<true, Class, RetType(ArgTypes...)>
 	typedef RetType (Class::* Type)(ArgTypes...) const;
 };
 
+PRAGMA_ENABLE_DEPRECATION_WARNINGS
+// End temp workaround
+
 template <typename FuncType>
 struct TPayload;
 
