Fix non-Unity compile issue when PLATFORM_HAS_SLOW_WRITES is enabled.

diff --git a/Engine/Source/Runtime/Core/Private/Misc/ConfigUtilities.cpp b/Engine/Source/Runtime/Core/Private/Misc/ConfigUtilities.cpp
--- a/Engine/Source/Runtime/Core/Private/Misc/ConfigUtilities.cpp
+++ b/Engine/Source/Runtime/Core/Private/Misc/ConfigUtilities.cpp
@@ -9,6 +9,7 @@
 #include "Tasks/Pipe.h"
 
 #if PLATFORM_WRITES_ARE_SLOW
+#include "Async/Async.h"
 #include "Misc/QueuedThreadPool.h"
 #endif
 
