[Scheduler]
 - Fix potential deadlock by always enqueueing to the global queue and perform wakeup for standby workers in case they go to sleep unconditionally.

diff --git a/Engine/Source/Runtime/Core/Private/Async/Fundamental/Scheduler.cpp b/Engine/Source/Runtime/Core/Private/Async/Fundamental/Scheduler.cpp
--- a/Engine/Source/Runtime/Core/Private/Async/Fundamental/Scheduler.cpp
+++ b/Engine/Source/Runtime/Core/Private/Async/Fundamental/Scheduler.cpp
@@ -25,6 +25,7 @@ namespace LowLevelTasks
 	thread_local FTask* FTask::ActiveTask = nullptr;
 	thread_local FSchedulerTls* FSchedulerTls::ActiveScheduler = nullptr;
 	thread_local FSchedulerTls::EWorkerType FSchedulerTls::WorkerType = FSchedulerTls::EWorkerType::None;
+	thread_local bool FSchedulerTls::Impl::bIsStandbyWorker = false;
 	thread_local bool Private::FOversubscriptionTls::bIsOversubscriptionAllowed = false;
 
 	FScheduler FScheduler::Singleton;
@@ -314,12 +315,18 @@ namespace LowLevelTasks
 		{
 			const bool bIsBackgroundTask = Task.IsBackgroundTask();
 			const bool bIsBackgroundWorker = FSchedulerTls::IsBackgroundWorker();
-			if (bIsBackgroundTask && !bIsBackgroundWorker)
+			const bool bIsStandbyWorker = FSchedulerTls::IsStandbyWorker();
+
+			// Standby workers always enqueue to the global queue and perform wakeup
+			// as they can go to sleep whenever the oversubscription period is done
+			// and we don't want that to happen without another thread picking up
+			// this task.
+			if ((bIsBackgroundTask && !bIsBackgroundWorker) || bIsStandbyWorker)
 			{
 				QueuePreference = EQueuePreference::GlobalQueuePreference;
 			}
 
-			bWakeUpWorker |= FSchedulerTls::LocalQueue == nullptr;
+			bWakeUpWorker |= bIsStandbyWorker || FSchedulerTls::LocalQueue == nullptr;
 
 			if (FSchedulerTls::LocalQueue && QueuePreference != EQueuePreference::GlobalQueuePreference)
 			{
@@ -526,6 +533,7 @@ namespace LowLevelTasks
 
 		FMemory::SetupTLSCachesOnCurrentThread();
 		FSchedulerTls::WorkerType = bPermitBackgroundWork ? FSchedulerTls::EWorkerType::Background : FSchedulerTls::EWorkerType::Foreground;
+		FSchedulerTls::SetStandbyWorker(WorkerEvent->bIsStandby);
 		FSchedulerTls::LocalQueue = WorkerLocalQueue;
 
 		{
@@ -543,6 +551,7 @@ namespace LowLevelTasks
 
 		FSchedulerTls::LocalQueue = nullptr;
 		FSchedulerTls::ActiveScheduler = nullptr;
+		FSchedulerTls::SetStandbyWorker(false);
 		FSchedulerTls::WorkerType = FSchedulerTls::EWorkerType::None;
 		FMemory::ClearAndDisableTLSCachesOnCurrentThread();
 	}
diff --git a/Engine/Source/Runtime/Core/Public/Async/Fundamental/Scheduler.h b/Engine/Source/Runtime/Core/Public/Async/Fundamental/Scheduler.h
--- a/Engine/Source/Runtime/Core/Public/Async/Fundamental/Scheduler.h
+++ b/Engine/Source/Runtime/Core/Public/Async/Fundamental/Scheduler.h
@@ -124,6 +124,14 @@ namespace LowLevelTasks
 		static thread_local FSchedulerTls* ActiveScheduler;
 		static thread_local FLocalQueueType* LocalQueue;
 		static thread_local EWorkerType WorkerType;
+	
+	private:
+		// Hide the member to avoid clashing with local variables
+		// already named bIsStandbyWorker.
+		struct Impl
+		{
+			static thread_local bool bIsStandbyWorker;
+		};
 
 	public:
 		CORE_API bool IsWorkerThread() const;
@@ -140,6 +148,16 @@ namespace LowLevelTasks
 		{
 			return WorkerType == EWorkerType::Background;
 		}
+
+		inline static bool IsStandbyWorker()
+		{
+			return Impl::bIsStandbyWorker;
+		}
+
+		inline static void SetStandbyWorker(bool bInIsStandbyWorker)
+		{
+			Impl::bIsStandbyWorker = bInIsStandbyWorker;
+		}
 	};
 
 	class FScheduler final : public FSchedulerTls
