Fixed thread unsafe static vars in CSV_SCOPED_TIMING_STAT_RECURSIVE

diff --git a/Engine/Source/Runtime/Core/Public/ProfilingDebugging/CsvProfiler.h b/Engine/Source/Runtime/Core/Public/ProfilingDebugging/CsvProfiler.h
--- a/Engine/Source/Runtime/Core/Public/ProfilingDebugging/CsvProfiler.h
+++ b/Engine/Source/Runtime/Core/Public/ProfilingDebugging/CsvProfiler.h
@@ -67,11 +67,11 @@ struct FCsvDeclaredStat;
 	FScopedCsvStatExclusiveConditional _ScopedCsvStatExclusive_ ## StatName (#StatName,Condition, "CSV_"#StatName);
 #define CSV_SCOPED_TIMING_STAT_RECURSIVE(Category,StatName) \
 	TRACE_CSV_PROFILER_INLINE_STAT(#StatName, CSV_CATEGORY_INDEX(Category)); \
-	static int32 _ScopedCsvStatRecursive_EntryCount_ ## StatName = 0; \
+	static thread_local int32 _ScopedCsvStatRecursive_EntryCount_ ## StatName = 0; \
 	FScopedCsvStatRecursive _ScopedCsvStatRecursive_ ## StatName (_ScopedCsvStatRecursive_EntryCount_ ## StatName, #StatName, CSV_CATEGORY_INDEX(Category), "CSV_"#StatName);
 #define CSV_SCOPED_TIMING_STAT_RECURSIVE_CONDITIONAL(Category,StatName,Condition) \
 	TRACE_CSV_PROFILER_INLINE_STAT(#StatName, CSV_CATEGORY_INDEX(Category)); \
-	static int32 _ScopedCsvStatRecursive_EntryCount_ ## StatName = 0; \
+	static thread_local int32 _ScopedCsvStatRecursive_EntryCount_ ## StatName = 0; \
 	FScopedCsvStatRecursiveConditional _ScopedCsvStatRecursive_ ## StatName (_ScopedCsvStatRecursive_EntryCount_ ## StatName, #StatName, CSV_CATEGORY_INDEX(Category), Condition, "CSV_"#StatName);
 
 #define CSV_SCOPED_WAIT(WaitTime)							FScopedCsvWaitConditional _ScopedCsvWait(WaitTime>0 && FCsvProfiler::IsWaitTrackingEnabledOnCurrentThread());
