- Fix missing CORE_API on some Unix memory methods

diff --git a/Engine/Source/Runtime/Core/Public/Unix/UnixPlatformMemory.h b/Engine/Source/Runtime/Core/Public/Unix/UnixPlatformMemory.h
--- a/Engine/Source/Runtime/Core/Public/Unix/UnixPlatformMemory.h
+++ b/Engine/Source/Runtime/Core/Public/Unix/UnixPlatformMemory.h
@@ -100,9 +100,9 @@ struct FUnixPlatformMemory : public FGenericPlatformMemory
 		FPlatformVirtualMemoryBlock(const FPlatformVirtualMemoryBlock& Other) = default;
 		FPlatformVirtualMemoryBlock& operator=(const FPlatformVirtualMemoryBlock& Other) = default;
 
-		void Commit(size_t InOffset, size_t InSize);
-		void Decommit(size_t InOffset, size_t InSize);
-		void FreeVirtual();
+		CORE_API void Commit(size_t InOffset, size_t InSize);
+		CORE_API void Decommit(size_t InOffset, size_t InSize);
+		CORE_API void FreeVirtual();
 
 		FORCEINLINE void CommitByPtr(void *InPtr, size_t InSize)
 		{
@@ -129,9 +129,9 @@ struct FUnixPlatformMemory : public FGenericPlatformMemory
 			return VMSizeDivVirtualSizeAlignment * GetVirtualSizeAlignment();
 		}
 
-		static FPlatformVirtualMemoryBlock AllocateVirtual(size_t Size, size_t InAlignment = FPlatformVirtualMemoryBlock::GetVirtualSizeAlignment());
-		static size_t GetCommitAlignment();
-		static size_t GetVirtualSizeAlignment();
+		static CORE_API FPlatformVirtualMemoryBlock AllocateVirtual(size_t Size, size_t InAlignment = FPlatformVirtualMemoryBlock::GetVirtualSizeAlignment());
+		static CORE_API size_t GetCommitAlignment();
+		static CORE_API size_t GetVirtualSizeAlignment();
 	};
 
 	static CORE_API FSharedMemoryRegion * MapNamedSharedMemoryRegion(const FString& InName, bool bCreate, uint32 AccessMode, SIZE_T Size);
