Remove UE_ENABLE_PLATFORM_MEMORY_CSV_STATS. The implementation is not safe. We should re-write this in a less intrusive way if we need to profile the memory profiler (e.g accumulating our own timers in the allocator and simply using CSV_CUSTOM_STAT to emit to the Csv profiler once per frame)

diff --git a/Engine/Source/Runtime/Core/Private/Unix/UnixPlatformMemory.cpp b/Engine/Source/Runtime/Core/Private/Unix/UnixPlatformMemory.cpp
--- a/Engine/Source/Runtime/Core/Private/Unix/UnixPlatformMemory.cpp
+++ b/Engine/Source/Runtime/Core/Private/Unix/UnixPlatformMemory.cpp
@@ -431,11 +431,6 @@ struct FOSAllocationDescriptor
 
 void* FUnixPlatformMemory::BinnedAllocFromOS(SIZE_T Size)
 {
-#if UE_ENABLE_PLATFORM_MEMORY_CSV_STATS
-	const bool bIsInGameThread = IsInGameThread();
-	CSV_SCOPED_TIMING_STAT_RECURSIVE_CONDITIONAL(FMemory, BinnedAllocFromOSTime, bIsInGameThread);
-#endif // UE_ENABLE_PLATFORM_MEMORY_CSV_STATS
-
 #if UE_CHECK_LARGE_ALLOCATIONS
 	if (UE::Memory::Private::GEnableLargeAllocationChecks)
 	{
@@ -626,11 +621,6 @@ size_t FUnixPlatformMemory::FPlatformVirtualMemoryBlock::GetCommitAlignment()
 
 FUnixPlatformMemory::FPlatformVirtualMemoryBlock FUnixPlatformMemory::FPlatformVirtualMemoryBlock::AllocateVirtual(size_t InSize, size_t InAlignment)
 {
-#if UE_ENABLE_PLATFORM_MEMORY_CSV_STATS
-	const bool bIsInGameThread = IsInGameThread();
-	CSV_SCOPED_TIMING_STAT_RECURSIVE_CONDITIONAL(FMemory, AllocateVirtualTime, bIsInGameThread);
-#endif // UE_ENABLE_PLATFORM_MEMORY_CSV_STATS
-
 	FPlatformVirtualMemoryBlock Result;
 	InSize = Align(InSize, GetVirtualSizeAlignment());
 	Result.VMSizeDivVirtualSizeAlignment = InSize / GetVirtualSizeAlignment();
diff --git a/Engine/Source/Runtime/Core/Private/Windows/WindowsPlatformMemory.cpp b/Engine/Source/Runtime/Core/Private/Windows/WindowsPlatformMemory.cpp
--- a/Engine/Source/Runtime/Core/Private/Windows/WindowsPlatformMemory.cpp
+++ b/Engine/Source/Runtime/Core/Private/Windows/WindowsPlatformMemory.cpp
@@ -580,11 +580,6 @@ bool FWindowsPlatformMemory::PageProtect(void* const Ptr, const SIZE_T Size, con
 }
 void* FWindowsPlatformMemory::BinnedAllocFromOS( SIZE_T Size )
 {
-#if UE_ENABLE_PLATFORM_MEMORY_CSV_STATS
-	const bool bIsInGameThread = IsInGameThread();
-	CSV_SCOPED_TIMING_STAT_RECURSIVE_CONDITIONAL(FMemory, BinnedAllocFromOSTime, bIsInGameThread);
-#endif // UE_ENABLE_PLATFORM_MEMORY_CSV_STATS
-
 #if UE_CHECK_LARGE_ALLOCATIONS
 	if (UE::Memory::Private::GEnableLargeAllocationChecks)
 	{
@@ -622,11 +617,6 @@ size_t FWindowsPlatformMemory::FPlatformVirtualMemoryBlock::GetCommitAlignment()
 
 FWindowsPlatformMemory::FPlatformVirtualMemoryBlock FWindowsPlatformMemory::FPlatformVirtualMemoryBlock::AllocateVirtual(size_t InSize, size_t InAlignment)
 {
-#if UE_ENABLE_PLATFORM_MEMORY_CSV_STATS
-	const bool bIsInGameThread = IsInGameThread();
-	CSV_SCOPED_TIMING_STAT_RECURSIVE_CONDITIONAL(FMemory, AllocateVirtualTime, bIsInGameThread);
-#endif // UE_ENABLE_PLATFORM_MEMORY_CSV_STATS
-
 	FPlatformVirtualMemoryBlock Result;
 	InSize = Align(InSize, GetVirtualSizeAlignment());
 	Result.VMSizeDivVirtualSizeAlignment = InSize / GetVirtualSizeAlignment();
diff --git a/Engine/Source/Runtime/Core/Public/GenericPlatform/GenericPlatformMemory.h b/Engine/Source/Runtime/Core/Public/GenericPlatform/GenericPlatformMemory.h
--- a/Engine/Source/Runtime/Core/Public/GenericPlatform/GenericPlatformMemory.h
+++ b/Engine/Source/Runtime/Core/Public/GenericPlatform/GenericPlatformMemory.h
@@ -35,10 +35,6 @@ namespace UE::Memory::Private
 }
 #endif
 
-#ifndef UE_ENABLE_PLATFORM_MEMORY_CSV_STATS
-#define UE_ENABLE_PLATFORM_MEMORY_CSV_STATS 0
-#endif
-
 /**
 * Platform-dependent "bucket" for memory size, where Default is the normal, or possibly the largest.
 * This is generally used for texture LOD settings for how to fit in smaller memory devices
