Bring back correct handling of TLS index in MB2\3

diff --git a/Engine/Source/Runtime/Core/Private/HAL/MallocBinned2.cpp b/Engine/Source/Runtime/Core/Private/HAL/MallocBinned2.cpp
--- a/Engine/Source/Runtime/Core/Private/HAL/MallocBinned2.cpp
+++ b/Engine/Source/Runtime/Core/Private/HAL/MallocBinned2.cpp
@@ -1286,11 +1286,11 @@ void FMallocBinned2::SetupTLSCachesOnCurrentThread()
 	{
 		return;
 	}
-	if (!FMallocBinned2::BinnedTlsSlot)
+	if (!FPlatformTLS::IsValidTlsSlot(FMallocBinned2::BinnedTlsSlot))
 	{
 		FMallocBinned2::BinnedTlsSlot = FPlatformTLS::AllocTlsSlot();
 	}
-	check(FMallocBinned2::BinnedTlsSlot);
+	check(FPlatformTLS::IsValidTlsSlot(FMallocBinned2::BinnedTlsSlot));
 	FPerThreadFreeBlockLists::SetTLS();
 }
 
diff --git a/Engine/Source/Runtime/Core/Private/HAL/MallocBinned3.cpp b/Engine/Source/Runtime/Core/Private/HAL/MallocBinned3.cpp
--- a/Engine/Source/Runtime/Core/Private/HAL/MallocBinned3.cpp
+++ b/Engine/Source/Runtime/Core/Private/HAL/MallocBinned3.cpp
@@ -1352,11 +1352,11 @@ void FMallocBinned3::SetupTLSCachesOnCurrentThread()
 	{
 		return;
 	}
-	if (!FMallocBinned3::BinnedTlsSlot)
+	if (!FPlatformTLS::IsValidTlsSlot(FMallocBinned3::BinnedTlsSlot))
 	{
 		FMallocBinned3::BinnedTlsSlot = FPlatformTLS::AllocTlsSlot();
 	}
-	check(FMallocBinned3::BinnedTlsSlot);
+	check(FPlatformTLS::IsValidTlsSlot(FMallocBinned3::BinnedTlsSlot));
 	FPerThreadFreeBlockLists::SetTLS();
 }
 
diff --git a/Engine/Source/Runtime/Core/Private/HAL/MallocBinnedCommon.cpp b/Engine/Source/Runtime/Core/Private/HAL/MallocBinnedCommon.cpp
--- a/Engine/Source/Runtime/Core/Private/HAL/MallocBinnedCommon.cpp
+++ b/Engine/Source/Runtime/Core/Private/HAL/MallocBinnedCommon.cpp
@@ -471,7 +471,7 @@ static FAutoConsoleVariableRef GMallocBinned3BundleCountCVar(
 
 #endif
 
-uint32 FMallocBinnedCommonBase::BinnedTlsSlot = 0;
+uint32 FMallocBinnedCommonBase::BinnedTlsSlot = FPlatformTLS::InvalidTlsSlot;
 #if UE_BINNEDCOMMON_ALLOCATOR_STATS
 std::atomic<int64> FMallocBinnedCommonBase::TLSMemory(0);
 std::atomic<int64> FMallocBinnedCommonBase::ConsolidatedMemory(0);
diff --git a/Engine/Source/Runtime/Core/Private/ProfilingDebugging/CsvProfiler.cpp b/Engine/Source/Runtime/Core/Private/ProfilingDebugging/CsvProfiler.cpp
--- a/Engine/Source/Runtime/Core/Private/ProfilingDebugging/CsvProfiler.cpp
+++ b/Engine/Source/Runtime/Core/Private/ProfilingDebugging/CsvProfiler.cpp
@@ -1982,7 +1982,7 @@ private:
 public:
 	static void InitTls()
 	{
-		if (TlsSlot == 0)
+		if (!FPlatformTLS::IsValidTlsSlot(TlsSlot))
 		{
 			TlsSlot = FPlatformTLS::AllocTlsSlot();
 			FPlatformMisc::MemoryBarrier();
@@ -2226,7 +2226,7 @@ public:
 	TArray<const char*> WaitStatNameStack;
 };
 
-uint32 FCsvProfilerThreadData::TlsSlot = 0;
+uint32 FCsvProfilerThreadData::TlsSlot = FPlatformTLS::InvalidTlsSlot;
 FCriticalSection FCsvProfilerThreadData::TlsCS;
 TArray<FCsvProfilerThreadData::FWeakPtr> FCsvProfilerThreadData::TlsInstances;
 
diff --git a/Engine/Source/Runtime/Core/Public/HAL/MallocBinnedCommon.h b/Engine/Source/Runtime/Core/Public/HAL/MallocBinnedCommon.h
--- a/Engine/Source/Runtime/Core/Public/HAL/MallocBinnedCommon.h
+++ b/Engine/Source/Runtime/Core/Public/HAL/MallocBinnedCommon.h
@@ -423,12 +423,12 @@ protected:
 	{
 		FORCEINLINE static FPerThreadFreeBlockLists* Get()
 		{
-			return BinnedTlsSlot ? (FPerThreadFreeBlockLists*)FPlatformTLS::GetTlsValue(BinnedTlsSlot) : nullptr;
+			return FPlatformTLS::IsValidTlsSlot(BinnedTlsSlot) ? (FPerThreadFreeBlockLists*)FPlatformTLS::GetTlsValue(BinnedTlsSlot) : nullptr;
 		}
 
 		static void SetTLS()
 		{
-			check(BinnedTlsSlot);
+			check(FPlatformTLS::IsValidTlsSlot(BinnedTlsSlot));
 			FPerThreadFreeBlockLists* ThreadSingleton = (FPerThreadFreeBlockLists*)FPlatformTLS::GetTlsValue(BinnedTlsSlot);
 			if (!ThreadSingleton)
 			{
@@ -445,7 +445,7 @@ protected:
 
 		static void ClearTLS()
 		{
-			check(BinnedTlsSlot);
+			check(FPlatformTLS::IsValidTlsSlot(BinnedTlsSlot));
 			FPerThreadFreeBlockLists* ThreadSingleton = (FPerThreadFreeBlockLists*)FPlatformTLS::GetTlsValue(BinnedTlsSlot);
 			if (ThreadSingleton)
 			{
