Added a IsStuck field in the crash context. This flag gets set when a thread exceeds the Stuck threshold (usually one second, normally much lower than the Hang threshold). The flag is cleared when the thread resumes.

IsStuck can be used to categorize abnormal terminations. This can catch reports where the player terminates the program before the hang detection kicks in (normally quite high, 45 to 60 seconds).

#tests multiple crash reports using debug commands
#rnx

diff --git a/Engine/Source/Programs/CrashReportClient/Private/CrashReportClientApp.cpp b/Engine/Source/Programs/CrashReportClient/Private/CrashReportClientApp.cpp
--- a/Engine/Source/Programs/CrashReportClient/Private/CrashReportClientApp.cpp
+++ b/Engine/Source/Programs/CrashReportClient/Private/CrashReportClientApp.cpp
@@ -929,6 +929,7 @@ static bool LoadTempCrashContextFromFile(FSharedCrashContext& CrashContext, uint
 	FindAndCopyValue(ContextProperties, TEXT("Misc.OSVersionMajor"), SessionContext.OsVersion);
 	FindAndCopyValue(ContextProperties, TEXT("Misc.OSVersionMinor"), SessionContext.OsSubVersion);
 	FindAndCopyValue(ContextProperties, TEXT("Misc.AnticheatProvider"), SessionContext.AnticheatProvider);
+	FindAndParseValue(ContextProperties, TEXT("Misc.IsStuck"), SessionContext.bIsStuck);
 	FindAndParseValue(ContextProperties, TEXT("MemoryStats.AvailablePhysical"), SessionContext.MemoryStats.AvailablePhysical);
 	FindAndParseValue(ContextProperties, TEXT("MemoryStats.AvailableVirtual"), SessionContext.MemoryStats.AvailableVirtual);
 	FindAndParseValue(ContextProperties, TEXT("MemoryStats.UsedPhysical"), SessionContext.MemoryStats.UsedPhysical);
diff --git a/Engine/Source/Runtime/Core/Private/GenericPlatform/GenericPlatformCrashContext.cpp b/Engine/Source/Runtime/Core/Private/GenericPlatform/GenericPlatformCrashContext.cpp
--- a/Engine/Source/Runtime/Core/Private/GenericPlatform/GenericPlatformCrashContext.cpp
+++ b/Engine/Source/Runtime/Core/Private/GenericPlatform/GenericPlatformCrashContext.cpp
@@ -672,6 +672,28 @@ void FGenericCrashContext::SetAnticheatProvider(const FString& AnticheatProvider
 	SerializeTempCrashContextToFile();
 }
 
+void FGenericCrashContext::OnThreadStuck(uint32 ThreadId)
+{
+	if (!NCached::Session.bIsStuck || NCached::Session.StuckThreadId != ThreadId)
+	{
+		NCached::Session.bIsStuck = true;
+		NCached::Session.StuckThreadId = ThreadId;
+
+		SerializeTempCrashContextToFile();
+	}
+}
+
+void FGenericCrashContext::OnThreadUnstuck(uint32 ThreadId)
+{
+	if (NCached::Session.bIsStuck)
+	{
+		NCached::Session.bIsStuck = false;
+		NCached::Session.StuckThreadId = 0;
+
+		SerializeTempCrashContextToFile();
+	}
+}
+
 FGenericCrashContext::FGenericCrashContext(ECrashContextType InType, const TCHAR* InErrorMessage)
 	: Type(InType)
 	, CrashedThreadId(~uint32(0))
@@ -856,6 +878,11 @@ void FGenericCrashContext::SerializeSessionContext(FString& Buffer)
 	AddCrashPropertyInternal(Buffer, TEXT("Misc.OSVersionMajor"), NCached::Session.OsVersion);
 	AddCrashPropertyInternal(Buffer, TEXT("Misc.OSVersionMinor"), NCached::Session.OsSubVersion);
 	AddCrashPropertyInternal(Buffer, TEXT("Misc.AnticheatProvider"), NCached::Session.AnticheatProvider);
+	if (NCached::Session.bIsStuck)
+	{
+		AddCrashPropertyInternal(Buffer, TEXT("Misc.IsStuck"), NCached::Session.bIsStuck);
+		AddCrashPropertyInternal(Buffer, TEXT("Misc.StuckThreadId"), NCached::Session.StuckThreadId);
+	}
 
 	// FPlatformMemory::GetConstants is called in the GCreateMalloc, so we can assume it is always valid.
 	{
diff --git a/Engine/Source/Runtime/Core/Private/HAL/ThreadHeartBeat.cpp b/Engine/Source/Runtime/Core/Private/HAL/ThreadHeartBeat.cpp
--- a/Engine/Source/Runtime/Core/Private/HAL/ThreadHeartBeat.cpp
+++ b/Engine/Source/Runtime/Core/Private/HAL/ThreadHeartBeat.cpp
@@ -28,17 +28,6 @@
 	#define UE_ASSERT_ON_HANG 0
 #endif
 
-// Enabling AttempStuckThreadResuscitation will add a check for early hung thread detection and pass the ThreadId through the OnStuck
-// delegate, allowing the platform to boost it's priority or other action to get the thread scheduled again.
-// Core.System StuckDuration can be changed to alter the time that the OnStuck delegate is triggered. Currently defaults to 1.0 second
-static bool AttemptStuckThreadResuscitation = false;
-
-static FAutoConsoleVariableRef CVarAttemptStuckThreadResuscitation(
-	TEXT("AttemptStuckThreadResuscitation"),
-	AttemptStuckThreadResuscitation,
-	TEXT("Attempt to resusicate stuck thread by boosting priority. Enabled by default\n"),
-	ECVF_Default);
-
 // The maximum clock time steps for the hang and hitch detectors.
 // These are the amounts the clocks are allowed to advance by before another tick is required.
 const double HangDetectorClock_MaxTimeStep_MS = 2000.0;
@@ -375,7 +364,7 @@ void FThreadHeartBeat::InitSettings()
 		const double MinStuckDuration = 1.0;
 		if (NewStuckDuration > 0.0 && NewStuckDuration < MinStuckDuration)
 		{
-			UE_LOG(LogCore, Warning, TEXT("HangDuration is set to %.4fs which is a very short time for hang detection. Changing to %.2fs."), NewStuckDuration, MinStuckDuration);
+			UE_LOG(LogCore, Warning, TEXT("StuckDuration is set to %.4fs which is a very short time for stuck detection. Changing to %.2fs."), NewStuckDuration, MinStuckDuration);
 			NewStuckDuration = MinStuckDuration;
 		}
 
@@ -584,7 +573,7 @@ uint32 FThreadHeartBeat::CheckHeartBeat(double& OutHangDuration)
 				}
 			}
 
-			if (AttemptStuckThreadResuscitation && (LongestStuckThreadId != InvalidThreadId))
+			if (LongestStuckThreadId != InvalidThreadId)
 			{
 				// Is there a currently stuck thread. Replace it.
 				if (LastStuckThreadId != LongestStuckThreadId)
diff --git a/Engine/Source/Runtime/Core/Private/Windows/WindowsPlatformMisc.cpp b/Engine/Source/Runtime/Core/Private/Windows/WindowsPlatformMisc.cpp
--- a/Engine/Source/Runtime/Core/Private/Windows/WindowsPlatformMisc.cpp
+++ b/Engine/Source/Runtime/Core/Private/Windows/WindowsPlatformMisc.cpp
@@ -871,6 +871,9 @@ void FWindowsPlatformMisc::PlatformPreInit()
 
 	FGenericPlatformMisc::PlatformPreInit();
 
+	FThreadHeartBeat::Get().GetOnThreadStuck().BindStatic(&FGenericCrashContext::OnThreadStuck);
+	FThreadHeartBeat::Get().GetOnThreadUnstuck().BindStatic(&FGenericCrashContext::OnThreadUnstuck);
+
 	// Load the bundled version of dbghelp.dll if necessary
 #if USE_BUNDLED_DBGHELP
 	// Loading newer versions of DbgHelp fails on Windows 7 since it is no longer supported.
diff --git a/Engine/Source/Runtime/Core/Public/GenericPlatform/GenericPlatformCrashContext.h b/Engine/Source/Runtime/Core/Public/GenericPlatform/GenericPlatformCrashContext.h
--- a/Engine/Source/Runtime/Core/Public/GenericPlatform/GenericPlatformCrashContext.h
+++ b/Engine/Source/Runtime/Core/Public/GenericPlatform/GenericPlatformCrashContext.h
@@ -166,6 +166,7 @@ struct FSessionContext
 	bool 					bIsUERelease;
 	bool					bIsOOM;
 	bool					bIsExitRequested;
+	bool					bIsStuck;
 	uint32					ProcessId;
 	int32 					LanguageLCID;
 	int32 					NumberOfCores;
@@ -173,6 +174,7 @@ struct FSessionContext
 	int32 					SecondsSinceStart;
 	int32 					CrashDumpMode;
 	int32					CrashTrigger;
+	uint32					StuckThreadId;
 	int32					OOMAllocationAlignment;
 	uint64					OOMAllocationSize;
 	TCHAR 					EngineVersion[CR_MAX_GENERIC_FIELD_CHARS];
@@ -620,6 +622,14 @@ public:
 	/** Sets the Anticheat client provider. */
 	CORE_API static void SetAnticheatProvider(const FString& AnticheatProvider);
 
+	/** Sets a flag that one of the threads is stuck.
+	 *  This is meant to be bound to the ThreadHeartBeat::OnThreadStuck delegate. Not all platforms register to save this flag. */
+	CORE_API static void OnThreadStuck(uint32 ThreadId);
+
+	/** Clears the stuck flag.
+	 *  This is meant to be bound to the ThreadHeartBeat::OnThreadUnstuck delegate. Not all platforms register to save this flag. */
+	CORE_API static void OnThreadUnstuck(uint32 ThreadId);
+
 	/** Attempts to create the output report directory. */
 	CORE_API static bool CreateCrashReportDirectory(const TCHAR* CrashGUIDRoot, int32 CrashIndex, FString& OutCrashDirectoryAbsolute);
 
