Fixed TScriptBitArray breaking the ClearPartialSlackBits invariant

diff --git a/Engine/Source/Runtime/Core/Public/Containers/BitArray.h b/Engine/Source/Runtime/Core/Public/Containers/BitArray.h
--- a/Engine/Source/Runtime/Core/Public/Containers/BitArray.h
+++ b/Engine/Source/Runtime/Core/Public/Containers/BitArray.h
@@ -2136,6 +2136,11 @@ public:
 		{
 			ReallocGrow(NumBits - 1);
 		}
+		else if ((Index % NumBitsPerDWORD) == 0)
+		{
+			// Clear the new word to maintain the ClearPartialSlackBits invariant.
+			GetData()[Index / NumBitsPerDWORD] = 0;
+		}
 		(*this)[Index] = Value;
 		return Index;
 	}
diff --git a/Engine/Source/Runtime/Core/Tests/Containers/BitArrayTests.cpp b/Engine/Source/Runtime/Core/Tests/Containers/BitArrayTests.cpp
--- a/Engine/Source/Runtime/Core/Tests/Containers/BitArrayTests.cpp
+++ b/Engine/Source/Runtime/Core/Tests/Containers/BitArrayTests.cpp
@@ -2186,6 +2186,16 @@ TEST_CASE_NAMED(FBitArrayTestFindAndSetLastZeroBitTest, "System::Core::Container
 	}
 }
 
+TEST_CASE_NAMED(FScriptBitArrayTestClearPartialSlackBits, "System::Core::Containers::BitArray::ScriptClearPartialSlackBits", "[ApplicationContextMask][SmokeFilter]")
+{
+	TTypeCompatibleBytes<TBitArray<>> Memory;
+	// Inline allocation flags are not initialized and the test requires they are non-zero.
+	FMemory::Memset(Memory, 0xff);
+	TBitArray<>& Array = *new(&Memory) TBitArray<>();
+	((FScriptBitArray&)Array).Add(true);
+	Array.CheckInvariants();
+}
+
 class FBitArrayMemoryTest
 {
 public:
