Fixed an uninitialized invocation list lock count when move-constructing a TMulticastDelegate.

GitHub #11642 : Fix uninitialized member InvocationListLockCount in TMulticastDelegateBase when an instance is move constructed
GitHub #11650 : Fix move constructor of TMulticastDelegateBase to initialize InvocationListLockCount

diff --git a/Engine/Source/Runtime/Core/Public/Delegates/MulticastDelegateBase.h b/Engine/Source/Runtime/Core/Public/Delegates/MulticastDelegateBase.h
--- a/Engine/Source/Runtime/Core/Public/Delegates/MulticastDelegateBase.h
+++ b/Engine/Source/Runtime/Core/Public/Delegates/MulticastDelegateBase.h
@@ -59,6 +59,7 @@ public:
 			LocalInvocationList = MoveTemp(Other.InvocationList);
 			LocalCompactionThreshold = Other.CompactionThreshold;
 			Other.CompactionThreshold = UE_MULTICAST_DELEGATE_DEFAULT_COMPACTION_THRESHOLD;
+			checkf(Other.InvocationListLockCount == 0, TEXT("Moving from a multicast delegate while it is mid broadcast"));
 		}
 
 		{
@@ -67,6 +68,8 @@ public:
 			ClearUnchecked();
 			InvocationList = MoveTemp(LocalInvocationList);
 			CompactionThreshold = LocalCompactionThreshold;
+			InvocationListLockCount = 0;
+			checkf(InvocationListLockCount == 0, TEXT("Moving to a multicast delegate while it is mid broadcast"));
 		}
 
 		return *this;
@@ -202,7 +205,8 @@ protected:
 	inline TMulticastDelegateBase( )
 		: CompactionThreshold(UE_MULTICAST_DELEGATE_DEFAULT_COMPACTION_THRESHOLD)
 		, InvocationListLockCount(0)
-	{ }
+	{
+	}
 
 protected:
 	template<typename DelegateInstanceInterfaceType>
@@ -408,8 +412,8 @@ private:
 	InvocationListType InvocationList;
 
 	/** Used to determine when a compaction should happen. */
-	int32 CompactionThreshold;
+	int32 CompactionThreshold = UE_MULTICAST_DELEGATE_DEFAULT_COMPACTION_THRESHOLD;
 
 	/** Holds a lock counter for the invocation list. */
-	mutable int32 InvocationListLockCount;
+	mutable int32 InvocationListLockCount = 0;
 };
