Minor MB2 refactoring to reduce lock time under #if !UE_USE_VERYLARGEPAGEALLOCATOR

diff --git a/Engine/Source/Runtime/Core/Private/HAL/MallocBinned2.cpp b/Engine/Source/Runtime/Core/Private/HAL/MallocBinned2.cpp
--- a/Engine/Source/Runtime/Core/Private/HAL/MallocBinned2.cpp
+++ b/Engine/Source/Runtime/Core/Private/HAL/MallocBinned2.cpp
@@ -159,12 +159,12 @@ struct FMallocBinned2::FPoolInfo
 		FirstFreeBlockIsPtr = 0xf317
 	};
 
- public:	uint16      Taken;          // Number of allocated elements in this pool, when counts down to zero can free the entire pool	
-public:	ECanary		Canary;	// See ECanary
-private:	uint32      AllocSize;      // Number of bytes allocated
- public:	FFreeBlock* FirstFreeBlock; // Pointer to first free memory in this pool or the OS Allocation Size in bytes if this allocation is not binned
- public:	FPoolInfo*  Next;           // Pointer to next pool
- public:	FPoolInfo** PtrToPrevNext;  // Pointer to whichever pointer points to this pool
+	uint16      Taken;          // Number of allocated elements in this pool, when counts down to zero can free the entire pool	
+	ECanary     Canary;         // See ECanary
+	uint32      AllocSize;      // Number of bytes allocated
+	FFreeBlock* FirstFreeBlock; // Pointer to first free memory in this pool or the OS Allocation Size in bytes if this allocation is not binned
+	FPoolInfo*  Next;           // Pointer to next pool
+	FPoolInfo** PtrToPrevNext;  // Pointer to whichever pointer points to this pool
 
 #if PLATFORM_32BITS
 /** Explicit padding for 32 bit builds */
@@ -792,8 +792,8 @@ void FMallocBinned2::OnPreFork()
 		}
 	}
 
-	FScopeLock Lock(&Mutex);
 #if !UE_USE_VERYLARGEPAGEALLOCATOR
+	FScopeLock Lock(&Mutex);
 	CachedOSPageAllocator.FreeAll(&Mutex);
 #endif
 
@@ -1208,7 +1208,6 @@ const TCHAR* FMallocBinned2::GetDescriptiveName()
 
 void FMallocBinned2::FlushCurrentThreadCache()
 {
-	double StartTimeInner = FPlatformTime::Seconds();
 	TRACE_CPUPROFILER_EVENT_SCOPE(FMallocBinned2::FlushCurrentThreadCache);
 	NOALLOC_SCOPE_CYCLE_COUNTER(STAT_FMallocBinned2_FlushCurrentThreadCache);
 	FPerThreadFreeBlockLists* Lists = FPerThreadFreeBlockLists::Get();
@@ -1218,6 +1217,7 @@ void FMallocBinned2::FlushCurrentThreadCache()
 
 	if (Lists)
 	{
+		const double StartTimeInner = FPlatformTime::Seconds();
 		FScopeLock Lock(&Mutex);
 		WaitForMutexTime = FPlatformTime::Seconds() - StartTimeInner;
 		for (int32 PoolIndex = 0; PoolIndex != BINNED2_SMALL_POOL_COUNT; ++PoolIndex)
@@ -1269,8 +1269,8 @@ void FMallocBinned2::Trim(bool bTrimThreadCaches)
 	}
 	{
 		//double StartTime = FPlatformTime::Seconds();
-		FScopeLock Lock(&Mutex);
 #if !UE_USE_VERYLARGEPAGEALLOCATOR
+		FScopeLock Lock(&Mutex);
 		// this cache is recycled anyway, if you need to trim it based on being OOM, it's already too late.
 		CachedOSPageAllocator.FreeAll(&Mutex);
 #endif
