Optimized TConstSetBitIterator

Measurements here are from a benchmark that sums the set bit indices of a TBitArray 50,000 times.
The setup adds 16,384 set bits to a TBitArray<> then clears 4,096 bits at random.

Total: 1.401902 s, Min: 0.000027 s, Avg: 0.000028 s, Test: Initial
Total: 1.347723 s, Min: 0.000026 s, Avg: 0.000027 s, Test: Constructor w/o StartIndex <- CL 32380294
Total: 1.050828 s, Min: 0.000020 s, Avg: 0.000021 s, Test: Optimized with FloorLog2
Total: 0.929872 s, Min: 0.000018 s, Avg: 0.000019 s, Test: Optimized with FloorLog2NonZero

Instruction counts from the benchmark function when compiled for Win64 Development:
126 -> 54 instructions for the whole function
32 -> 21 instructions for the loop body
22 -> 12 instructions executed on the hot path of the inner loop

diff --git a/Engine/Source/Runtime/Core/Public/Containers/BitArray.h b/Engine/Source/Runtime/Core/Public/Containers/BitArray.h
--- a/Engine/Source/Runtime/Core/Public/Containers/BitArray.h
+++ b/Engine/Source/Runtime/Core/Public/Containers/BitArray.h
@@ -1807,68 +1807,56 @@ public:
 
 	explicit TConstSetBitIterator(const TBitArray<Allocator>& InArray UE_LIFETIMEBOUND)
 		: FRelativeBitReference(0)
-		, Array                (InArray)
-		, UnvisitedBitMask     (~0U)
-		, CurrentBitIndex      (0)
-		, BaseBitIndex         (0)
+		, Array(InArray)
+		, WordIt(InArray.GetData())
+		, BitCount(InArray.Num())
 	{
-		if (Array.Num())
+		if (BitCount)
 		{
+			RemainingWord = *WordIt++;
 			FindFirstSetBit();
 		}
 	}
 
 	explicit TConstSetBitIterator(const TBitArray<Allocator>& InArray UE_LIFETIMEBOUND, int32 StartIndex)
 		: FRelativeBitReference(StartIndex)
-		, Array                (InArray)
-		, UnvisitedBitMask     ((~0U) << (StartIndex & (NumBitsPerDWORD - 1)))
-		, CurrentBitIndex      (StartIndex)
-		, BaseBitIndex         (StartIndex & ~(NumBitsPerDWORD - 1))
+		, Array(InArray)
+		, WordIt(InArray.GetData() + WordIndex)
+		, BitCount(InArray.Num())
+		, CurrentBitIndex(StartIndex)
 	{
-		check(StartIndex >= 0 && StartIndex <= Array.Num());
-		if (StartIndex != Array.Num())
+		check(CurrentBitIndex >= 0 && CurrentBitIndex <= BitCount);
+		if (CurrentBitIndex != BitCount && BitCount)
 		{
+			RemainingWord = *WordIt++ & ~(Mask - 1);
 			FindFirstSetBit();
 		}
 	}
 
-	/** Forwards iteration operator. */
 	FORCEINLINE TConstSetBitIterator& operator++()
 	{
-		// Mark the current bit as visited.
-		UnvisitedBitMask &= ~this->Mask;
-
-		// Find the first set bit that hasn't been visited yet.
 		FindFirstSetBit();
-
 		return *this;
 	}
 
 	FORCEINLINE bool operator==(const TConstSetBitIterator& Rhs) const
 	{
-		// We only need to compare the bit index and the array... all the rest of the state is unobservable.
+		// We only need to compare the bit index and the array... the rest of the state is unobservable.
 		return CurrentBitIndex == Rhs.CurrentBitIndex && &Array == &Rhs.Array;
 	}
 
 #if !PLATFORM_COMPILER_HAS_GENERATED_COMPARISON_OPERATORS
 	FORCEINLINE bool operator!=(const TConstSetBitIterator& Rhs) const
-	{ 
+	{
 		return !(*this == Rhs);
 	}
 #endif
 
-	/** conversion to "bool" returning true if the iterator is valid. */
 	FORCEINLINE explicit operator bool() const
-	{ 
-		return CurrentBitIndex < Array.Num(); 
-	}
-	/** inverse of the "bool" operator */
-	FORCEINLINE bool operator !() const 
 	{
-		return !(bool)*this;
+		return CurrentBitIndex < BitCount;
 	}
 
-	/** Index accessor. */
 	FORCEINLINE int32 GetIndex() const
 	{
 		return CurrentBitIndex;
@@ -1878,49 +1866,41 @@ private:
 
 	const TBitArray<Allocator>& Array;
 
-	uint32 UnvisitedBitMask;
-	int32 CurrentBitIndex;
-	int32 BaseBitIndex;
+	const uint32* WordIt;
+	const int32 BitCount;
+	int32 CurrentBitIndex = 0;
+	uint32 RemainingWord = 0;
 
 	/** Find the first set bit starting with the current bit, inclusive. */
 	void FindFirstSetBit()
 	{
-		const uint32* ArrayData      = Array.GetData();
-		const int32   ArrayNum       = Array.Num();
-		const int32   LastWordIndex = (ArrayNum - 1) / NumBitsPerDWORD;
-
-		// Advance to the next non-zero uint32.
-		uint32 RemainingBitMask = ArrayData[this->WordIndex] & UnvisitedBitMask;
-		while (!RemainingBitMask)
+		for (;;)
 		{
-			++this->WordIndex;
-			BaseBitIndex += NumBitsPerDWORD;
-			if (this->WordIndex > LastWordIndex)
+			// Find the index of the first remaining set bit in the word.
+			if (RemainingWord)
 			{
-				// We've advanced past the end of the array.
-				CurrentBitIndex = ArrayNum;
-				return;
-			}
-
-			RemainingBitMask = ArrayData[this->WordIndex];
-			UnvisitedBitMask = ~0u;
-		}
+				// Clear the first remaining set bit in the word.
+				const uint32 NewRemainingWord = RemainingWord & (RemainingWord - 1);
 
-		// This operation has the effect of unsetting the lowest set bit of BitMask
-		const uint32 NewRemainingBitMask = RemainingBitMask & (RemainingBitMask - 1);
+				// Calculate the index of the bit that was cleared.
+				Mask = RemainingWord ^ NewRemainingWord;
+				WordIndex = CurrentBitIndex >> NumBitsPerDWORDLogTwo;
+				CurrentBitIndex = (WordIndex << NumBitsPerDWORDLogTwo) | FMath::FloorLog2NonZero(Mask);
 
-		// This operation XORs the above mask with the original mask, which has the effect
-		// of returning only the bits which differ; specifically, the lowest bit
-		this->Mask = NewRemainingBitMask ^ RemainingBitMask;
+				RemainingWord = NewRemainingWord;
+				return;
+			}
 
-		// If the Nth bit was the lowest set bit of BitMask, then this gives us N
-		CurrentBitIndex = BaseBitIndex + NumBitsPerDWORD - 1 - FMath::CountLeadingZeros(this->Mask);
+			// Advance the bit index to the next word.
+			CurrentBitIndex = (CurrentBitIndex & ~(NumBitsPerDWORD - 1)) + NumBitsPerDWORD;
+			if (CurrentBitIndex >= BitCount)
+			{
+				CurrentBitIndex = BitCount;
+				return;
+			}
 
-		// If we've accidentally iterated off the end of an array but still within the same Word
-		// then set the index to the last index of the array
-		if (CurrentBitIndex > ArrayNum)
-		{
-			CurrentBitIndex = ArrayNum;
+			// Read the next word into the iterator.
+			RemainingWord = *WordIt++;
 		}
 	}
 };
