CIS fix.

diff --git a/Engine/Source/Runtime/Core/Public/Delegates/MulticastDelegateBase.h b/Engine/Source/Runtime/Core/Public/Delegates/MulticastDelegateBase.h
--- a/Engine/Source/Runtime/Core/Public/Delegates/MulticastDelegateBase.h
+++ b/Engine/Source/Runtime/Core/Public/Delegates/MulticastDelegateBase.h
@@ -68,7 +68,6 @@ public:
 			ClearUnchecked();
 			InvocationList = MoveTemp(LocalInvocationList);
 			CompactionThreshold = LocalCompactionThreshold;
-			InvocationListLockCount = 0;
 			checkf(InvocationListLockCount == 0, TEXT("Moving to a multicast delegate while it is mid broadcast"));
 		}
 
