Add support of partial zen DDC GetChunks requests

diff --git a/Engine/Source/Developer/DerivedDataCache/Private/ZenCacheStore.cpp b/Engine/Source/Developer/DerivedDataCache/Private/ZenCacheStore.cpp
--- a/Engine/Source/Developer/DerivedDataCache/Private/ZenCacheStore.cpp
+++ b/Engine/Source/Developer/DerivedDataCache/Private/ZenCacheStore.cpp
@@ -1131,11 +1131,13 @@ public:
 			{
 				BatchRequest << ANSITEXTVIEW("Method") << "GetCacheChunks";
 				BatchRequest.AddInteger(ANSITEXTVIEW("Accept"), Zen::Http::kCbPkgMagic);
+				uint32_t AcceptFlags = static_cast<uint32_t>(Zen::Http::RpcAcceptOptions::kAllowPartialCacheChunks);
 				if (CacheStore.bIsLocalConnection)
 				{
-					BatchRequest.AddInteger(ANSITEXTVIEW("AcceptFlags"), static_cast<uint32_t>(Zen::Http::RpcAcceptOptions::kAllowLocalReferences));
+					AcceptFlags |= static_cast<uint32_t>(Zen::Http::RpcAcceptOptions::kAllowLocalReferences);
 					BatchRequest.AddInteger(ANSITEXTVIEW("Pid"), FPlatformProcess::GetCurrentProcessId());
 				}
+				BatchRequest.AddInteger(ANSITEXTVIEW("AcceptFlags"), AcceptFlags);
 
 				BatchRequest.BeginObject(ANSITEXTVIEW("Params"));
 				{
@@ -1218,20 +1220,9 @@ public:
 							RawHash = HashView.AsHash();
 							if (!HashView.HasError())
 							{
-								const FCbAttachment* Attachment = EnumHasAnyFlags(RequestWithStats.Request.Policy, ECachePolicy::SkipData) ? nullptr : Response.FindAttachment(HashView.AsHash());
-
-								if (Attachment)
-								{
-									FCompressedBuffer CompressedBuffer = Attachment->AsCompressedBinary();
-									if (CompressedBuffer)
-									{
-										TRACE_COUNTER_ADD(ZenDDC_BytesReceived, CompressedBuffer.GetCompressedSize());
-										RequestedBytes = FCompressedBufferReader(CompressedBuffer).Decompress(Request.RawOffset, Request.RawSize);
-										RawSize = RequestedBytes.GetSize();
-										Succeeded = true;
-									}
-								}
-								else
+								FIoHash AttachmentHash = RawHash;
+								const FCbAttachment* Attachment = nullptr;
+								if (EnumHasAnyFlags(RequestWithStats.Request.Policy, ECachePolicy::SkipData))
 								{
 									FCbFieldView RawSizeField = ResultObject[ANSITEXTVIEW("RawSize")];
 									uint64 TotalSize = RawSizeField.AsUInt64();
@@ -1241,6 +1232,28 @@ public:
 										RawSize = FMath::Min(Request.RawSize, TotalSize - FMath::Min(Request.RawOffset, TotalSize));
 									}
 								}
+								else
+								{
+									FCbFieldView FragmentOffsetField = ResultObject[ANSITEXTVIEW("FragmentOffset")];
+									uint64 FragmentOffset = FragmentOffsetField.AsUInt64();
+
+									FCbFieldView FragmentHashView = ResultObject[ANSITEXTVIEW("FragmentHash")];
+									if (FragmentHashView.IsHash())
+									{
+										FIoHash FragmentHash = FragmentHashView.AsHash();
+										AttachmentHash = FragmentHash;
+									}
+
+									if (Attachment = Response.FindAttachment(AttachmentHash); Attachment != nullptr)
+									{
+										if (FCompressedBuffer CompressedBuffer = Attachment->AsCompressedBinary(); CompressedBuffer)
+										{
+											RequestedBytes = FCompressedBufferReader(CompressedBuffer).Decompress(Request.RawOffset - FragmentOffset, Request.RawSize);
+											RawSize = RequestedBytes.GetSize();
+											Succeeded = true;
+										}
+									}
+								}
 							}
 						}
 						Succeeded ? OnHit(RequestWithStats, MoveTemp(RawHash), RawSize, MoveTemp(RequestedBytes)) : OnMiss(RequestWithStats);
diff --git a/Engine/Source/Developer/Zen/Internal/ZenSerialization.h b/Engine/Source/Developer/Zen/Internal/ZenSerialization.h
--- a/Engine/Source/Developer/Zen/Internal/ZenSerialization.h
+++ b/Engine/Source/Developer/Zen/Internal/ZenSerialization.h
@@ -32,6 +32,7 @@ enum class RpcAcceptOptions : uint32_t
 {
 	kNone = 0,
 	kAllowLocalReferences = (1u << 0),
+	kAllowPartialCacheChunks = (1u << 2)
 };
 
 ZEN_API void SaveCbPackage(const FCbPackage& Package, FArchive& Ar);
diff --git a/Engine/Source/Developer/Zen/Private/ZenSerialization.cpp b/Engine/Source/Developer/Zen/Private/ZenSerialization.cpp
--- a/Engine/Source/Developer/Zen/Private/ZenSerialization.cpp
+++ b/Engine/Source/Developer/Zen/Private/ZenSerialization.cpp
@@ -393,7 +393,7 @@ bool TryLoadCbPackage(FCbPackage& Package, FArchive& Ar, FCbBufferAllocator Allo
 				Ar.SetError();
 				return false;
 			}
-			FCbAttachment Attachment(MoveTemp(CompBuf));
+			FCbAttachment Attachment(MoveTemp(CompBuf), Entry.AttachmentHash);
 			Package.AddAttachment(Attachment);
 		}
 		else if (Entry.Flags & CbAttachmentEntry::IsCompressed)
@@ -408,7 +408,7 @@ bool TryLoadCbPackage(FCbPackage& Package, FArchive& Ar, FCbBufferAllocator Allo
 			}
 			else
 			{
-				FCbAttachment Attachment(MoveTemp(CompBuf));
+				FCbAttachment Attachment(MoveTemp(CompBuf), Entry.AttachmentHash);
 				Package.AddAttachment(Attachment);
 			}
 		}
@@ -422,7 +422,7 @@ bool TryLoadCbPackage(FCbPackage& Package, FArchive& Ar, FCbBufferAllocator Allo
 			}
 			else
 			{
-				FCbAttachment Attachment(AttachmentData.MoveToShared());
+				FCbAttachment Attachment(AttachmentData.MoveToShared(), Entry.AttachmentHash);
 				Package.AddAttachment(Attachment);
 			}
 		}
diff --git a/Engine/Source/Runtime/Core/Private/Serialization/CompactBinaryPackage.cpp b/Engine/Source/Runtime/Core/Private/Serialization/CompactBinaryPackage.cpp
--- a/Engine/Source/Runtime/Core/Private/Serialization/CompactBinaryPackage.cpp
+++ b/Engine/Source/Runtime/Core/Private/Serialization/CompactBinaryPackage.cpp
@@ -88,9 +88,9 @@ FCbAttachment::FObjectValue::FObjectValue(FCbObject&& InObject, const FIoHash* c
 
 FIoHash FCbAttachment::GetHash() const
 {
-	if (const FCompressedBuffer* Buffer = Value.TryGet<FCompressedBuffer>())
+	if (const FCompressedBinaryValue* CompressedValue = Value.TryGet<FCompressedBinaryValue>())
 	{
-		return Buffer->GetRawHash();
+		return CompressedValue->Hash;
 	}
 	else if (const FBinaryValue* BinaryValue = Value.TryGet<FBinaryValue>())
 	{
@@ -153,7 +153,7 @@ bool FCbAttachment::TryLoad(FCbFieldIterator& Fields)
 		else if (FCompressedBuffer CompressedBuffer{FCompressedBuffer::FromCompressed(MoveTemp(Buffer))})
 		{
 			// Compressed Binary
-			Value.Emplace<FCompressedBuffer>(MoveTemp(CompressedBuffer));
+			Value.Emplace<FCompressedBinaryValue>(MoveTemp(CompressedBuffer));
 			++Fields;
 			return true;
 		}
@@ -185,9 +185,9 @@ void FCbAttachment::Save(FCbWriter& Writer) const
 		}
 		Writer.AddBinary(BinaryValue->Buffer);
 	}
-	else if (const FCompressedBuffer* BufferValue = Value.TryGet<FCompressedBuffer>())
+	else if (const FCompressedBinaryValue* CompressedValue = Value.TryGet<FCompressedBinaryValue>())
 	{
-		Writer.AddBinary(BufferValue->GetCompressed());
+		Writer.AddBinary(CompressedValue->Buffer.GetCompressed());
 	}
 	else
 	{
diff --git a/Engine/Source/Runtime/Core/Public/Serialization/CompactBinaryPackage.h b/Engine/Source/Runtime/Core/Public/Serialization/CompactBinaryPackage.h
--- a/Engine/Source/Runtime/Core/Public/Serialization/CompactBinaryPackage.h
+++ b/Engine/Source/Runtime/Core/Public/Serialization/CompactBinaryPackage.h
@@ -65,6 +65,26 @@ class FCbAttachment
 		}
 	};
 
+	struct FCompressedBinaryValue
+	{
+		FCompressedBuffer Buffer;
+		FIoHash Hash;
+
+		template <typename BufferType, decltype(FCompressedBuffer(DeclVal<BufferType>().MakeOwned()))* = nullptr>
+		inline explicit FCompressedBinaryValue(BufferType&& InBuffer)
+			: Buffer(Forward<BufferType>(InBuffer).MakeOwned())
+			, Hash(Buffer.GetRawHash())
+		{
+		}
+
+		template <typename BufferType, decltype(FCompressedBuffer(DeclVal<BufferType>().MakeOwned()))* = nullptr>
+		inline explicit FCompressedBinaryValue(BufferType&& InBuffer, const FIoHash& InHash)
+			: Buffer(Forward<BufferType>(InBuffer).MakeOwned())
+			, Hash(InHash)
+		{
+		}
+	};
+
 	template <typename... ArgTypes, decltype(FBinaryValue(DeclVal<ArgTypes>()...))* = nullptr>
 	inline FCbAttachment(TInPlaceType<FBinaryValue>, ArgTypes&&... Args)
 		: Value(TInPlaceType<FBinaryValue>(), Forward<ArgTypes>(Args)...)
@@ -99,19 +119,36 @@ public:
 
 	/** Construct a compressed binary attachment. Value is cloned if not owned. */
 	inline explicit FCbAttachment(const FCompressedBuffer& InValue)
-		: Value(TInPlaceType<FCompressedBuffer>(), InValue)
+		: Value(TInPlaceType<FCompressedBinaryValue>(), InValue)
 	{
-		if (Value.Get<FCompressedBuffer>().IsNull())
+		if (Value.Get<FCompressedBinaryValue>().Buffer.IsNull())
+		{
+			Value.Emplace<TYPE_OF_NULLPTR>();
+		}
+	}
+
+	UE_INTERNAL inline explicit FCbAttachment(const FCompressedBuffer& InValue, const FIoHash& Hash)
+		: Value(TInPlaceType<FCompressedBinaryValue>(), InValue, Hash)
+	{
+		if (Value.Get<FCompressedBinaryValue>().Buffer.IsNull())
 		{
 			Value.Emplace<TYPE_OF_NULLPTR>();
 		}
 	}
 
-	/** Construct a compressed binary attachment. Value is cloned if not owned. */
 	inline explicit FCbAttachment(FCompressedBuffer&& InValue)
-		: Value(TInPlaceType<FCompressedBuffer>(), InValue)
+		: Value(TInPlaceType<FCompressedBinaryValue>(), InValue)
+	{
+		if (Value.Get<FCompressedBinaryValue>().Buffer.IsNull())
+		{
+			Value.Emplace<TYPE_OF_NULLPTR>();
+		}
+	}
+
+	UE_INTERNAL inline explicit FCbAttachment(FCompressedBuffer&& InValue, const FIoHash& Hash)
+		: Value(TInPlaceType<FCompressedBinaryValue>(), InValue, Hash)
 	{
-		if (Value.Get<FCompressedBuffer>().IsNull())
+		if (Value.Get<FCompressedBinaryValue>().Buffer.IsNull())
 		{
 			Value.Emplace<TYPE_OF_NULLPTR>();
 		}
@@ -145,7 +182,7 @@ public:
 	inline bool IsBinary() const { return Value.IsType<FBinaryValue>(); }
 
 	/** Returns whether the attachment is compressed binary. */
-	inline bool IsCompressedBinary() const { return Value.IsType<FCompressedBuffer>(); }
+	inline bool IsCompressedBinary() const { return Value.IsType<FCompressedBinaryValue>(); }
 
 	/** Returns the hash of the attachment value. */
 	CORE_API FIoHash GetHash() const;
@@ -182,7 +219,7 @@ public:
 	CORE_API void Save(FArchive& Ar) const;
 
 private:
-	TVariant<TYPE_OF_NULLPTR, FObjectValue, FBinaryValue, FCompressedBuffer> Value;
+	TVariant<TYPE_OF_NULLPTR, FObjectValue, FBinaryValue, FCompressedBinaryValue> Value;
 };
 
 /** Hashes attachments by their hash. Any discrepancy in type must be handled externally. */
@@ -443,9 +480,9 @@ const FCompositeBuffer& FCbAttachment::AsCompositeBinary() const
 
 const FCompressedBuffer& FCbAttachment::AsCompressedBinary() const
 {
-	if (const FCompressedBuffer* CompressedBuffer = Value.TryGet<FCompressedBuffer>())
+	if (const FCompressedBinaryValue* CompressedBuffer = Value.TryGet<FCompressedBinaryValue>())
 	{
-		return *CompressedBuffer;
+		return CompressedBuffer->Buffer;
 	}
 	return FCompressedBuffer::Null;
 }
