Added TDelegateRegistration and TMulticastDelegateRegistration, which can be returned from OnX() functions to allow users to bind or unbind from delegates but not invoke them, like C# events.

diff --git a/Engine/Source/Runtime/Core/Public/Delegates/DelegateSignatureImpl.inl b/Engine/Source/Runtime/Core/Public/Delegates/DelegateSignatureImpl.inl
--- a/Engine/Source/Runtime/Core/Public/Delegates/DelegateSignatureImpl.inl
+++ b/Engine/Source/Runtime/Core/Public/Delegates/DelegateSignatureImpl.inl
@@ -36,286 +36,307 @@
 template<typename UserPolicy> class TMulticastDelegateBase;
 
 /**
- * Unicast delegate template class.
+ * Delegate registration class.  A function can expose this base class to users to let them bind
+ * against, but without letting them execute.  This gives us the benefit of C# events.
+ *
+ * Example:
+ *
+ * struct FInterestingType
+ * {
+ * public:
+ *     TDelegateRegistration<void()>& OnInterestingThing() const
+ *     {
+ *         return OnInterestingThingDelegate;
+ *     }
+ *
+ * private:
+ *     mutable TDelegate<void()> OnInterestingThingDelegate;
+ * };
+ *
+ * void Func(const FInterestingThing& Thing)
+ * {
+ *     // Binding and unbinding are allowed
+ *     Thing.OnInterestingThing().BindLambda([](){ RespondToInterestingThing(); });
+ *     Thing.OnInterestingThing().Unbind();
+ *
+ *     // Execute is not
+ *     Thing.OnInterestingThing().Execute(); // error: Execute is deleted
+ * }
  */
+
 template <typename DelegateSignature, typename UserPolicy = FDefaultDelegateUserPolicy>
-class TDelegate
-{
-	static_assert(sizeof(UserPolicy) == 0, "Expected a function signature for the delegate template parameter");
-};
+class TDelegateRegistration;
 
 template <typename InRetValType, typename... ParamTypes, typename UserPolicy>
-class TDelegate<InRetValType(ParamTypes...), UserPolicy> : public UserPolicy::FDelegateExtras
+class TDelegateRegistration<InRetValType(ParamTypes...), UserPolicy> : public UserPolicy::FDelegateExtras
 {
-	using Super                         = typename UserPolicy::FDelegateExtras;
+private:
+	using Super = typename UserPolicy::FDelegateExtras;
+
+protected:
 	using FuncType                      = InRetValType (ParamTypes...);
 	using DelegateInstanceInterfaceType = IBaseDelegateInstance<FuncType, UserPolicy>;
 
+private:
 	static_assert(std::is_convertible_v<typename UserPolicy::FDelegateInstanceExtras*, IDelegateInstance*>, "UserPolicy::FDelegateInstanceExtras should publicly inherit IDelegateInstance");
 	static_assert(std::is_convertible_v<typename UserPolicy::FMulticastDelegateExtras*, TMulticastDelegateBase<UserPolicy>*>, "UserPolicy::FMulticastDelegateExtras should publicly inherit TMulticastDelegateBase<UserPolicy>");
 
 	template <typename, typename>
 	friend class TDelegate;
 
+	template <typename, typename>
+	friend class TDelegateRegistration;
+
 	template <typename>
 	friend class TMulticastDelegateBase;
 
 	template <typename, typename>
 	friend class TMulticastDelegate;
 
-private:
-	// Make sure FDelegateBase's protected functions are not accidentally exposed through the TDelegate API
-	using typename Super::FReadAccessScope;
-	using Super::GetReadAccessScope;
-	using typename Super::FWriteAccessScope;
-	using Super::GetWriteAccessScope;
+	template <typename, typename>
+	friend class TMulticastDelegateRegistration;
 
 public:
 	/** Type definition for return value type. */
-	typedef InRetValType RetValType;
-	typedef InRetValType TFuncType(ParamTypes...);
+	using RetValType = InRetValType;
+	using TFuncType  = InRetValType(ParamTypes...);
 
-	/* Helper typedefs for getting a member function pointer type for the delegate with a given payload */
-	template <typename... VarTypes>                     using TFuncPtr        = RetValType(*)(ParamTypes..., VarTypes...);
-	template <typename UserClass, typename... VarTypes> using TMethodPtr      = typename TMemFunPtrType<false, UserClass, RetValType(ParamTypes..., VarTypes...)>::Type;
-	template <typename UserClass, typename... VarTypes> using TConstMethodPtr = typename TMemFunPtrType<true,  UserClass, RetValType(ParamTypes..., VarTypes...)>::Type;
+protected:
+	TDelegateRegistration() = default;
+	TDelegateRegistration(TDelegateRegistration&&) = default;
+	TDelegateRegistration(const TDelegateRegistration&) = default;
+	TDelegateRegistration& operator=(TDelegateRegistration&&) = default;
+	TDelegateRegistration& operator=(const TDelegateRegistration&) = default;
+	~TDelegateRegistration() = default;
 
 public:
-
 	/**
-	 * Static: Creates a raw C++ pointer global function delegate
+	 * Binds a raw C++ pointer global function delegate
 	 */
 	template <typename... VarTypes>
-	[[nodiscard]] inline static TDelegate<RetValType(ParamTypes...), UserPolicy> CreateStatic(typename TIdentity<RetValType (*)(ParamTypes..., std::decay_t<VarTypes>...)>::Type InFunc, VarTypes&&... Vars)
+	inline void BindStatic(typename TBaseStaticDelegateInstance<FuncType, UserPolicy, std::decay_t<VarTypes>...>::FFuncPtr InFunc, VarTypes&&... Vars)
 	{
-		TDelegate<RetValType(ParamTypes...), UserPolicy> Result;
-		Result.template CreateDelegateInstance<TBaseStaticDelegateInstance<FuncType, UserPolicy, std::decay_t<VarTypes>...>>(InFunc, Forward<VarTypes>(Vars)...);
-		return Result;
+		Super::template CreateDelegateInstance<TBaseStaticDelegateInstance<FuncType, UserPolicy, std::decay_t<VarTypes>...>>(InFunc, Forward<VarTypes>(Vars)...);
 	}
-
+	
 	/**
-	 * Static: Creates a C++ lambda delegate
+	 * Static: Binds a C++ lambda delegate
 	 * technically this works for any functor types, but lambdas are the primary use case
 	 */
 	template<typename FunctorType, typename... VarTypes>
-	[[nodiscard]] inline static TDelegate<RetValType(ParamTypes...), UserPolicy> CreateLambda(FunctorType&& InFunctor, VarTypes&&... Vars)
+	inline void BindLambda(FunctorType&& InFunctor, VarTypes&&... Vars)
 	{
-		TDelegate<RetValType(ParamTypes...), UserPolicy> Result;
-		Result.template CreateDelegateInstance<TBaseFunctorDelegateInstance<FuncType, UserPolicy, typename TRemoveReference<FunctorType>::Type, std::decay_t<VarTypes>...>>(Forward<FunctorType>(InFunctor), Forward<VarTypes>(Vars)...);
-		return Result;
+		Super::template CreateDelegateInstance<TBaseFunctorDelegateInstance<FuncType, UserPolicy, typename TRemoveReference<FunctorType>::Type, std::decay_t<VarTypes>...>>(Forward<FunctorType>(InFunctor), Forward<VarTypes>(Vars)...);
 	}
 
 	/**
-	 * Static: Creates a weak shared pointer C++ lambda delegate
+	 * Static: Binds a weak shared pointer C++ lambda delegate
 	 * technically this works for any functor types, but lambdas are the primary use case
 	 */
 	template<typename UserClass, ESPMode Mode, typename FunctorType, typename... VarTypes>
-	[[nodiscard]] inline static TDelegate<RetValType(ParamTypes...), UserPolicy> CreateSPLambda(const TSharedRef<UserClass, Mode>& InUserObjectRef, FunctorType&& InFunctor, VarTypes&&... Vars)
+	inline void BindSPLambda(const TSharedRef<UserClass, Mode>& InUserObjectRef, FunctorType&& InFunctor, VarTypes&&... Vars)
 	{
-		TDelegate<RetValType(ParamTypes...), UserPolicy> Result;
-		Result.template CreateDelegateInstance<TBaseSPLambdaDelegateInstance<UserClass, Mode, FuncType, UserPolicy, typename TRemoveReference<FunctorType>::Type, std::decay_t<VarTypes>...>>(InUserObjectRef, Forward<FunctorType>(InFunctor), Forward<VarTypes>(Vars)...);
-		return Result;
+		Super::template CreateDelegateInstance<TBaseSPLambdaDelegateInstance<const UserClass, Mode, FuncType, UserPolicy, typename TRemoveReference<FunctorType>::Type, std::decay_t<VarTypes>...>>(InUserObjectRef, Forward<FunctorType>(InFunctor), Forward<VarTypes>(Vars)...);
 	}
 	template <typename UserClass, typename FunctorType, typename... VarTypes>
-	[[nodiscard]] inline static TDelegate<RetValType(ParamTypes...), UserPolicy> CreateSPLambda(UserClass* InUserObject, FunctorType&& InFunctor, VarTypes&&... Vars)
+	inline void BindSPLambda(const UserClass* InUserObject, FunctorType&& InFunctor, VarTypes&&... Vars)
 	{
-		TDelegate<RetValType(ParamTypes...), UserPolicy> Result;
-		Result.template CreateDelegateInstance<TBaseSPLambdaDelegateInstance<UserClass, decltype(InUserObject->AsShared())::Mode, FuncType, UserPolicy, typename TRemoveReference<FunctorType>::Type, std::decay_t<VarTypes>...>>(StaticCastSharedRef<UserClass>(InUserObject->AsShared()), Forward<FunctorType>(InFunctor), Forward<VarTypes>(Vars)...);
-		return Result;
+		Super::template CreateDelegateInstance<TBaseSPLambdaDelegateInstance<const UserClass, decltype(InUserObject->AsShared())::Mode, FuncType, UserPolicy, typename TRemoveReference<FunctorType>::Type, std::decay_t<VarTypes>...>>(StaticCastSharedRef<const UserClass>(InUserObject->AsShared()), Forward<FunctorType>(InFunctor), Forward<VarTypes>(Vars)...);
 	}
 
 	/**
-	 * Static: Creates a weak object C++ lambda delegate
+	 * Static: Binds a weak object C++ lambda delegate
 	 * technically this works for any functor types, but lambdas are the primary use case
 	 */
 	template<typename UserClass, typename FunctorType, typename... VarTypes>
-	[[nodiscard]] inline static TDelegate<RetValType(ParamTypes...), UserPolicy> CreateWeakLambda(UserClass* InUserObject, FunctorType&& InFunctor, VarTypes&&... Vars)
+	inline void BindWeakLambda(UserClass* InUserObject, FunctorType&& InFunctor, VarTypes&&... Vars)
 	{
-		TDelegate<RetValType(ParamTypes...), UserPolicy> Result;
-		Result.template CreateDelegateInstance<TWeakBaseFunctorDelegateInstance<UserClass, FuncType, UserPolicy, typename TRemoveReference<FunctorType>::Type, std::decay_t<VarTypes>...>>(InUserObject, Forward<FunctorType>(InFunctor), Forward<VarTypes>(Vars)...);
-		return Result;
+		Super::template CreateDelegateInstance<TWeakBaseFunctorDelegateInstance<UserClass, FuncType, UserPolicy, typename TRemoveReference<FunctorType>::Type, std::decay_t<VarTypes>...>>(InUserObject, Forward<FunctorType>(InFunctor), Forward<VarTypes>(Vars)...);
 	}
 
 	/**
-	 * Static: Creates a raw C++ pointer member function delegate.
+	 * Binds a raw C++ pointer delegate.
 	 *
 	 * Raw pointer doesn't use any sort of reference, so may be unsafe to call if the object was
 	 * deleted out from underneath your delegate. Be careful when calling Execute()!
 	 */
 	template <typename UserClass, typename... VarTypes>
-	[[nodiscard]] inline static TDelegate<RetValType(ParamTypes...), UserPolicy> CreateRaw(UserClass* InUserObject, typename TMemFunPtrType<false, UserClass, RetValType (ParamTypes..., std::decay_t<VarTypes>...)>::Type InFunc, VarTypes&&... Vars)
+	inline void BindRaw(UserClass* InUserObject, typename TMemFunPtrType<false, UserClass, RetValType (ParamTypes..., std::decay_t<VarTypes>...)>::Type InFunc, VarTypes&&... Vars)
 	{
 		static_assert(!TIsConst<UserClass>::Value, "Attempting to bind a delegate with a const object pointer and non-const member function.");
 
-		TDelegate<RetValType(ParamTypes...), UserPolicy> Result;
-		Result.template CreateDelegateInstance<TBaseRawMethodDelegateInstance<false, UserClass, FuncType, UserPolicy, std::decay_t<VarTypes>...>>(InUserObject, InFunc, Forward<VarTypes>(Vars)...);
-		return Result;
+		Super::template CreateDelegateInstance<TBaseRawMethodDelegateInstance<false, UserClass, FuncType, UserPolicy, std::decay_t<VarTypes>...>>(InUserObject, InFunc, Forward<VarTypes>(Vars)...);
 	}
 	template <typename UserClass, typename... VarTypes>
-	[[nodiscard]] inline static TDelegate<RetValType(ParamTypes...), UserPolicy> CreateRaw(const UserClass* InUserObject, typename TMemFunPtrType<true, UserClass, RetValType (ParamTypes..., std::decay_t<VarTypes>...)>::Type InFunc, VarTypes&&... Vars)
+	inline void BindRaw(const UserClass* InUserObject, typename TMemFunPtrType<true, UserClass, RetValType (ParamTypes..., std::decay_t<VarTypes>...)>::Type InFunc, VarTypes&&... Vars)
 	{
-		TDelegate<RetValType(ParamTypes...), UserPolicy> Result;
-		Result.template CreateDelegateInstance<TBaseRawMethodDelegateInstance<true, const UserClass, FuncType, UserPolicy, std::decay_t<VarTypes>...>>(InUserObject, InFunc, Forward<VarTypes>(Vars)...);
-		return Result;
+		Super::template CreateDelegateInstance<TBaseRawMethodDelegateInstance<true, const UserClass, FuncType, UserPolicy, std::decay_t<VarTypes>...>>(InUserObject, InFunc, Forward<VarTypes>(Vars)...);
 	}
 
 	/**
-	 * Static: Creates a shared pointer-based member function delegate.
-	 *
-	 * Shared pointer delegates keep a weak reference to your object.
-	 * You can use ExecuteIfBound() to call them.
+	 * Binds a shared pointer-based member function delegate.  Shared pointer delegates keep a weak reference to your object.  You can use ExecuteIfBound() to call them.
 	 */
 	template <typename UserClass, ESPMode Mode, typename... VarTypes>
-	[[nodiscard]] inline static TDelegate<RetValType(ParamTypes...), UserPolicy> CreateSP(const TSharedRef<UserClass, Mode>& InUserObjectRef, typename TMemFunPtrType<false, UserClass, RetValType (ParamTypes..., std::decay_t<VarTypes>...)>::Type InFunc, VarTypes&&... Vars)
+	inline void BindSP(const TSharedRef<UserClass, Mode>& InUserObjectRef, typename TMemFunPtrType<false, UserClass, RetValType (ParamTypes..., std::decay_t<VarTypes>...)>::Type InFunc, VarTypes&&... Vars)
 	{
 		static_assert(!TIsConst<UserClass>::Value, "Attempting to bind a delegate with a const object pointer and non-const member function.");
 
-		TDelegate<RetValType(ParamTypes...), UserPolicy> Result;
-		Result.template CreateDelegateInstance<TBaseSPMethodDelegateInstance<false, UserClass, Mode, FuncType, UserPolicy, std::decay_t<VarTypes>...>>(InUserObjectRef, InFunc, Forward<VarTypes>(Vars)...);
-		return Result;
+		Super::template CreateDelegateInstance<TBaseSPMethodDelegateInstance<false, UserClass, Mode, FuncType, UserPolicy, std::decay_t<VarTypes>...>>(InUserObjectRef, InFunc, Forward<VarTypes>(Vars)...);
 	}
 	template <typename UserClass, ESPMode Mode, typename... VarTypes>
-	[[nodiscard]] inline static TDelegate<RetValType(ParamTypes...), UserPolicy> CreateSP(const TSharedRef<UserClass, Mode>& InUserObjectRef, typename TMemFunPtrType<true, UserClass, RetValType (ParamTypes..., std::decay_t<VarTypes>...)>::Type InFunc, VarTypes&&... Vars)
+	inline void BindSP(const TSharedRef<UserClass, Mode>& InUserObjectRef, typename TMemFunPtrType<true, UserClass, RetValType (ParamTypes..., std::decay_t<VarTypes>...)>::Type InFunc, VarTypes&&... Vars)
 	{
-		TDelegate<RetValType(ParamTypes...), UserPolicy> Result;
-		Result.template CreateDelegateInstance<TBaseSPMethodDelegateInstance<true, const UserClass, Mode, FuncType, UserPolicy, std::decay_t<VarTypes>...>>(InUserObjectRef, InFunc, Forward<VarTypes>(Vars)...);
-		return Result;
+		Super::template CreateDelegateInstance<TBaseSPMethodDelegateInstance<true, const UserClass, Mode, FuncType, UserPolicy, std::decay_t<VarTypes>...>>(InUserObjectRef, InFunc, Forward<VarTypes>(Vars)...);
 	}
 
 	/**
-	 * Static: Creates a shared pointer-based member function delegate.
+	 * Binds a shared pointer-based member function delegate.
 	 *
 	 * Shared pointer delegates keep a weak reference to your object.
 	 * You can use ExecuteIfBound() to call them.
 	 */
 	template <typename UserClass, typename... VarTypes>
-	[[nodiscard]] inline static TDelegate<RetValType(ParamTypes...), UserPolicy> CreateSP(UserClass* InUserObject, typename TMemFunPtrType<false, UserClass, RetValType (ParamTypes..., std::decay_t<VarTypes>...)>::Type InFunc, VarTypes&&... Vars)
+	inline void BindSP(UserClass* InUserObject, typename TMemFunPtrType<false, UserClass, RetValType (ParamTypes..., std::decay_t<VarTypes>...)>::Type InFunc, VarTypes&&... Vars)
 	{
 		static_assert(!TIsConst<UserClass>::Value, "Attempting to bind a delegate with a const object pointer and non-const member function.");
 
-		TDelegate<RetValType(ParamTypes...), UserPolicy> Result;
-		Result.template CreateDelegateInstance<TBaseSPMethodDelegateInstance<false, UserClass, decltype(InUserObject->AsShared())::Mode, FuncType, UserPolicy, std::decay_t<VarTypes>...>>(StaticCastSharedRef<UserClass>(InUserObject->AsShared()), InFunc, Forward<VarTypes>(Vars)...);
-		return Result;
+		Super::template CreateDelegateInstance<TBaseSPMethodDelegateInstance<false, UserClass, decltype(InUserObject->AsShared())::Mode, FuncType, UserPolicy, std::decay_t<VarTypes>...>>(StaticCastSharedRef<UserClass>(InUserObject->AsShared()), InFunc, Forward<VarTypes>(Vars)...);
 	}
 	template <typename UserClass, typename... VarTypes>
-	[[nodiscard]] inline static TDelegate<RetValType(ParamTypes...), UserPolicy> CreateSP(const UserClass* InUserObject, typename TMemFunPtrType<true, UserClass, RetValType (ParamTypes..., std::decay_t<VarTypes>...)>::Type InFunc, VarTypes&&... Vars)
+	inline void BindSP(const UserClass* InUserObject, typename TMemFunPtrType<true, UserClass, RetValType (ParamTypes..., std::decay_t<VarTypes>...)>::Type InFunc, VarTypes&&... Vars)
 	{
-		TDelegate<RetValType(ParamTypes...), UserPolicy> Result;
-		Result.template CreateDelegateInstance<TBaseSPMethodDelegateInstance<true, const UserClass, decltype(InUserObject->AsShared())::Mode, FuncType, UserPolicy, std::decay_t<VarTypes>...>>(StaticCastSharedRef<const UserClass>(InUserObject->AsShared()), InFunc, Forward<VarTypes>(Vars)...);
-		return Result;
+		Super::template CreateDelegateInstance<TBaseSPMethodDelegateInstance<true, const UserClass, decltype(InUserObject->AsShared())::Mode, FuncType, UserPolicy, std::decay_t<VarTypes>...>>(StaticCastSharedRef<const UserClass>(InUserObject->AsShared()), InFunc, Forward<VarTypes>(Vars)...);
 	}
 
 	/**
-	 * Static: Creates a shared pointer-based (thread-safe) member function delegate.
+	 * Binds a shared pointer-based (thread-safe) member function delegate.
 	 *
 	 * Shared pointer delegates keep a weak reference to your object.
 	 * You can use ExecuteIfBound() to call them.
 	 *
-	 * Note: This function is redundant, but is retained for backwards compatibility.  CreateSP() works in both thread-safe and not-thread-safe modes and should be preferred.
+	 * Note: This function is redundant, but is retained for backwards compatibility.  BindSP() works in both thread-safe and not-thread-safe modes and should be preferred.
 	 */
 	template <typename UserClass, typename... VarTypes>
-	[[nodiscard]] inline static TDelegate<RetValType(ParamTypes...), UserPolicy> CreateThreadSafeSP(const TSharedRef<UserClass, ESPMode::ThreadSafe>& InUserObjectRef, typename TMemFunPtrType<false, UserClass, RetValType (ParamTypes..., std::decay_t<VarTypes>...)>::Type InFunc, VarTypes&&... Vars)
+	inline void BindThreadSafeSP(const TSharedRef<UserClass, ESPMode::ThreadSafe>& InUserObjectRef, typename TMemFunPtrType<false, UserClass, RetValType (ParamTypes..., std::decay_t<VarTypes>...)>::Type InFunc, VarTypes&&... Vars)
 	{
 		static_assert(!TIsConst<UserClass>::Value, "Attempting to bind a delegate with a const object pointer and non-const member function.");
 
-		TDelegate<RetValType(ParamTypes...), UserPolicy> Result;
-		Result.template CreateDelegateInstance<TBaseSPMethodDelegateInstance<false, UserClass, ESPMode::ThreadSafe, FuncType, UserPolicy, std::decay_t<VarTypes>...>>(InUserObjectRef, InFunc, Forward<VarTypes>(Vars)...);
-		return Result;
+		Super::template CreateDelegateInstance<TBaseSPMethodDelegateInstance<false, UserClass, ESPMode::ThreadSafe, FuncType, UserPolicy, std::decay_t<VarTypes>...>>(InUserObjectRef, InFunc, Forward<VarTypes>(Vars)...);
 	}
 	template <typename UserClass, typename... VarTypes>
-	[[nodiscard]] inline static TDelegate<RetValType(ParamTypes...), UserPolicy> CreateThreadSafeSP(const TSharedRef<UserClass, ESPMode::ThreadSafe>& InUserObjectRef, typename TMemFunPtrType<true, UserClass, RetValType (ParamTypes..., std::decay_t<VarTypes>...)>::Type InFunc, VarTypes&&... Vars)
+	inline void BindThreadSafeSP(const TSharedRef<UserClass, ESPMode::ThreadSafe>& InUserObjectRef, typename TMemFunPtrType<true, UserClass, RetValType (ParamTypes..., std::decay_t<VarTypes>...)>::Type InFunc, VarTypes&&... Vars)
 	{
-		TDelegate<RetValType(ParamTypes...), UserPolicy> Result;
-		Result.template CreateDelegateInstance<TBaseSPMethodDelegateInstance<true, const UserClass, ESPMode::ThreadSafe, FuncType, UserPolicy, std::decay_t<VarTypes>...>>(InUserObjectRef, InFunc, Forward<VarTypes>(Vars)...);
-		return Result;
+		Super::template CreateDelegateInstance<TBaseSPMethodDelegateInstance<true, const UserClass, ESPMode::ThreadSafe, FuncType, UserPolicy, std::decay_t<VarTypes>...>>(InUserObjectRef, InFunc, Forward<VarTypes>(Vars)...);
 	}
 
 	/**
-	 * Static: Creates a shared pointer-based (thread-safe) member function delegate.
+	 * Binds a shared pointer-based (thread-safe) member function delegate.
 	 *
 	 * Shared pointer delegates keep a weak reference to your object.
 	 * You can use ExecuteIfBound() to call them.
 	 *
-	 * Note: This function is redundant, but is retained for backwards compatibility.  CreateSP() works in both thread-safe and not-thread-safe modes and should be preferred.
+	 * Note: This function is redundant, but is retained for backwards compatibility.  BindSP() works in both thread-safe and not-thread-safe modes and should be preferred.
 	 */
 	template <typename UserClass, typename... VarTypes>
-	[[nodiscard]] inline static TDelegate<RetValType(ParamTypes...), UserPolicy> CreateThreadSafeSP(UserClass* InUserObject, typename TMemFunPtrType<false, UserClass, RetValType (ParamTypes..., std::decay_t<VarTypes>...)>::Type InFunc, VarTypes&&... Vars)
+	inline void BindThreadSafeSP(UserClass* InUserObject, typename TMemFunPtrType<false, UserClass, RetValType (ParamTypes..., std::decay_t<VarTypes>...)>::Type InFunc, VarTypes&&... Vars)
 	{
 		static_assert(!TIsConst<UserClass>::Value, "Attempting to bind a delegate with a const object pointer and non-const member function.");
 
-		TDelegate<RetValType(ParamTypes...), UserPolicy> Result;
-		Result.template CreateDelegateInstance<TBaseSPMethodDelegateInstance<false, UserClass, ESPMode::ThreadSafe, FuncType, UserPolicy, std::decay_t<VarTypes>...>>(StaticCastSharedRef<UserClass>(InUserObject->AsShared()), InFunc, Forward<VarTypes>(Vars)...);
-		return Result;
+		Super::template CreateDelegateInstance<TBaseSPMethodDelegateInstance<false, UserClass, ESPMode::ThreadSafe, FuncType, UserPolicy, std::decay_t<VarTypes>...>>(StaticCastSharedRef<UserClass>(InUserObject->AsShared()), InFunc, Forward<VarTypes>(Vars)...);
 	}
 	template <typename UserClass, typename... VarTypes>
-	[[nodiscard]] inline static TDelegate<RetValType(ParamTypes...), UserPolicy> CreateThreadSafeSP(const UserClass* InUserObject, typename TMemFunPtrType<true, UserClass, RetValType (ParamTypes..., std::decay_t<VarTypes>...)>::Type InFunc, VarTypes&&... Vars)
+	inline void BindThreadSafeSP(const UserClass* InUserObject, typename TMemFunPtrType<true, UserClass, RetValType (ParamTypes..., std::decay_t<VarTypes>...)>::Type InFunc, VarTypes&&... Vars)
 	{
-		TDelegate<RetValType(ParamTypes...), UserPolicy> Result;
-		Result.template CreateDelegateInstance<TBaseSPMethodDelegateInstance<true, const UserClass, ESPMode::ThreadSafe, FuncType, UserPolicy, std::decay_t<VarTypes>...>>(StaticCastSharedRef<const UserClass>(InUserObject->AsShared()), InFunc, Forward<VarTypes>(Vars)...);
-		return Result;
+		Super::template CreateDelegateInstance<TBaseSPMethodDelegateInstance<true, const UserClass, ESPMode::ThreadSafe, FuncType, UserPolicy, std::decay_t<VarTypes>...>>(StaticCastSharedRef<const UserClass>(InUserObject->AsShared()), InFunc, Forward<VarTypes>(Vars)...);
 	}
 
 	/**
-	 * Static: Creates a UFunction-based member function delegate.
+	 * Binds a UFunction-based member function delegate.
 	 *
 	 * UFunction delegates keep a weak reference to your object.
 	 * You can use ExecuteIfBound() to call them.
 	 */
 	template <typename UObjectTemplate, typename... VarTypes>
-	[[nodiscard]] inline static TDelegate<RetValType(ParamTypes...), UserPolicy> CreateUFunction(UObjectTemplate* InUserObject, const FName& InFunctionName, VarTypes&&... Vars)
+	inline void BindUFunction(UObjectTemplate* InUserObject, const FName& InFunctionName, VarTypes&&... Vars)
 	{
-		TDelegate<RetValType(ParamTypes...), UserPolicy> Result;
-		Result.template CreateDelegateInstance<TBaseUFunctionDelegateInstance<UObjectTemplate, FuncType, UserPolicy, std::decay_t<VarTypes>...>>(InUserObject, InFunctionName, Forward<VarTypes>(Vars)...);
-		return Result;
+		Super::template CreateDelegateInstance<TBaseUFunctionDelegateInstance<UObjectTemplate, FuncType, UserPolicy, std::decay_t<VarTypes>...>>(InUserObject, InFunctionName, Forward<VarTypes>(Vars)...);
 	}
 	template <typename UObjectTemplate, typename... VarTypes>
-	[[nodiscard]] inline static TDelegate<RetValType(ParamTypes...), UserPolicy> CreateUFunction(TObjectPtr<UObjectTemplate> InUserObject, const FName& InFunctionName, VarTypes&&... Vars)
+	inline void BindUFunction(TObjectPtr<UObjectTemplate> InUserObject, const FName& InFunctionName, VarTypes&&... Vars)
 	{
-		TDelegate<RetValType(ParamTypes...), UserPolicy> Result;
-		Result.template CreateDelegateInstance<TBaseUFunctionDelegateInstance<UObjectTemplate, FuncType, UserPolicy, std::decay_t<VarTypes>...>>(ToRawPtr(InUserObject), InFunctionName, Forward<VarTypes>(Vars)...);
-		return Result;
+		Super::template CreateDelegateInstance<TBaseUFunctionDelegateInstance<UObjectTemplate, FuncType, UserPolicy, std::decay_t<VarTypes>...>>(ToRawPtr(InUserObject), InFunctionName, Forward<VarTypes>(Vars)...);
 	}
 
 	/**
-	 * Static: Creates a UObject-based member function delegate.
+	 * Binds a UObject-based member function delegate.
 	 *
 	 * UObject delegates keep a weak reference to your object.
 	 * You can use ExecuteIfBound() to call them.
 	 */
 	template <typename UserClass, typename... VarTypes>
-	[[nodiscard]] inline static TDelegate<RetValType(ParamTypes...), UserPolicy> CreateUObject(UserClass* InUserObject, typename TMemFunPtrType<false, UserClass, RetValType (ParamTypes..., std::decay_t<VarTypes>...)>::Type InFunc, VarTypes&&... Vars)
+	inline void BindUObject(UserClass* InUserObject, typename TMemFunPtrType<false, UserClass, RetValType (ParamTypes..., std::decay_t<VarTypes>...)>::Type InFunc, VarTypes&&... Vars)
 	{
 		static_assert(!TIsConst<UserClass>::Value, "Attempting to bind a delegate with a const object pointer and non-const member function.");
 
-		TDelegate<RetValType(ParamTypes...), UserPolicy> Result;
-		Result.template CreateDelegateInstance<TBaseUObjectMethodDelegateInstance<false, UserClass, FuncType, UserPolicy, std::decay_t<VarTypes>...>>(InUserObject, InFunc, Forward<VarTypes>(Vars)...);
-		return Result;
+		Super::template CreateDelegateInstance<TBaseUObjectMethodDelegateInstance<false, UserClass, FuncType, UserPolicy, std::decay_t<VarTypes>...>>(InUserObject, InFunc, Forward<VarTypes>(Vars)...);
 	}
 	template <typename UserClass, typename... VarTypes>
-	[[nodiscard]] inline static TDelegate<RetValType(ParamTypes...), UserPolicy> CreateUObject(const UserClass* InUserObject, typename TMemFunPtrType<true, UserClass, RetValType (ParamTypes..., std::decay_t<VarTypes>...)>::Type InFunc, VarTypes&&... Vars)
+	inline void BindUObject(const UserClass* InUserObject, typename TMemFunPtrType<true, UserClass, RetValType (ParamTypes..., std::decay_t<VarTypes>...)>::Type InFunc, VarTypes&&... Vars)
 	{
-		TDelegate<RetValType(ParamTypes...), UserPolicy> Result;
-		Result.template CreateDelegateInstance<TBaseUObjectMethodDelegateInstance<true, const UserClass, FuncType, UserPolicy, std::decay_t<VarTypes>...>>(InUserObject, InFunc, Forward<VarTypes>(Vars)...);
-		return Result;
+		Super::template CreateDelegateInstance<TBaseUObjectMethodDelegateInstance<true, const UserClass, FuncType, UserPolicy, std::decay_t<VarTypes>...>>(InUserObject, InFunc, Forward<VarTypes>(Vars)...);
 	}
 	template <typename UserClass, typename... VarTypes>
-	[[nodiscard]] inline static TDelegate<RetValType(ParamTypes...), UserPolicy> CreateUObject(TObjectPtr<UserClass> InUserObject, typename TMemFunPtrType<false, UserClass, RetValType (ParamTypes..., std::decay_t<VarTypes>...)>::Type InFunc, VarTypes&&... Vars)
+	inline void BindUObject(TObjectPtr<UserClass> InUserObject, typename TMemFunPtrType<false, UserClass, RetValType (ParamTypes..., std::decay_t<VarTypes>...)>::Type InFunc, VarTypes&&... Vars)
 	{
 		static_assert(!TIsConst<UserClass>::Value, "Attempting to bind a delegate with a const object pointer and non-const member function.");
 
-		TDelegate<RetValType(ParamTypes...), UserPolicy> Result;
-		Result.template CreateDelegateInstance<TBaseUObjectMethodDelegateInstance<false, UserClass, FuncType, UserPolicy, std::decay_t<VarTypes>...>>(ToRawPtr(InUserObject), InFunc, Forward<VarTypes>(Vars)...);
-		return Result;
+		Super::template CreateDelegateInstance<TBaseUObjectMethodDelegateInstance<false, UserClass, FuncType, UserPolicy, std::decay_t<VarTypes>...>>(ToRawPtr(InUserObject), InFunc, Forward<VarTypes>(Vars)...);
 	}
 	template <typename UserClass, typename... VarTypes>
-	[[nodiscard]] inline static TDelegate<RetValType(ParamTypes...), UserPolicy> CreateUObject(TObjectPtr<UserClass> InUserObject, typename TMemFunPtrType<true, UserClass, RetValType (ParamTypes..., std::decay_t<VarTypes>...)>::Type InFunc, VarTypes&&... Vars)
+	inline void BindUObject(TObjectPtr<UserClass> InUserObject, typename TMemFunPtrType<true, UserClass, RetValType (ParamTypes..., std::decay_t<VarTypes>...)>::Type InFunc, VarTypes&&... Vars)
 	{
-		TDelegate<RetValType(ParamTypes...), UserPolicy> Result;
-		Result.template CreateDelegateInstance<TBaseUObjectMethodDelegateInstance<true, const UserClass, FuncType, UserPolicy, std::decay_t<VarTypes>...>>(ToRawPtr(InUserObject), InFunc, Forward<VarTypes>(Vars)...);
-		return Result;
+		Super::template CreateDelegateInstance<TBaseUObjectMethodDelegateInstance<true, const UserClass, FuncType, UserPolicy, std::decay_t<VarTypes>...>>(ToRawPtr(InUserObject), InFunc, Forward<VarTypes>(Vars)...);
 	}
 
+	// Executing via a delegate registration reference is not allowed
+	RetValType Execute       (ParamTypes... Params) const = delete;
+	bool       ExecuteIfBound(ParamTypes... Params) const = delete;
+};
+
+/**
+ * Unicast delegate template class.
+ */
+template <typename DelegateSignature, typename UserPolicy = FDefaultDelegateUserPolicy>
+class TDelegate
+{
+	static_assert(sizeof(UserPolicy) == 0, "Expected a function signature for the delegate template parameter");
+};
+
+template <typename InRetValType, typename... ParamTypes, typename UserPolicy>
+class TDelegate<InRetValType(ParamTypes...), UserPolicy> : public TDelegateRegistration<InRetValType(ParamTypes...), UserPolicy>
+{
+private:
+	using Super                         = TDelegateRegistration<InRetValType(ParamTypes...), UserPolicy>;
+	using DelegateInstanceInterfaceType = typename Super::DelegateInstanceInterfaceType;
+	using FuncType                      = typename Super::FuncType;
+
+	// Make sure FDelegateBase's protected functions are not accidentally exposed through the TDelegate API
+	using typename Super::FReadAccessScope;
+	using Super::GetReadAccessScope;
+	using typename Super::FWriteAccessScope;
+	using Super::GetWriteAccessScope;
+
 public:
+	using RegistrationType = Super;
 
+	/** Type definition for return value type. */
+	using RetValType = InRetValType;
+	using TFuncType  = InRetValType(ParamTypes...);
+
+	/* Helper typedefs for getting a member function pointer type for the delegate with a given payload */
+	template <typename... VarTypes>                     using TFuncPtr        = RetValType(*)(ParamTypes..., VarTypes...);
+	template <typename UserClass, typename... VarTypes> using TMethodPtr      = typename TMemFunPtrType<false, UserClass, RetValType(ParamTypes..., VarTypes...)>::Type;
+	template <typename UserClass, typename... VarTypes> using TConstMethodPtr = typename TMemFunPtrType<true,  UserClass, RetValType(ParamTypes..., VarTypes...)>::Type;
+
 	TDelegate() = default;
 
@@ -336,195 +357,239 @@
 
 	TDelegate(TDelegate&& Other) = default;
 	TDelegate& operator=(TDelegate&& Other) = default;
-
-public:
+	~TDelegate() = default;
 
 	/**
-	 * Binds a raw C++ pointer global function delegate
+	 * Static: Creates a raw C++ pointer global function delegate
 	 */
 	template <typename... VarTypes>
-	inline void BindStatic(typename TBaseStaticDelegateInstance<FuncType, UserPolicy, std::decay_t<VarTypes>...>::FFuncPtr InFunc, VarTypes&&... Vars)
+	[[nodiscard]] inline static TDelegate<RetValType(ParamTypes...), UserPolicy> CreateStatic(typename TIdentity<RetValType (*)(ParamTypes..., std::decay_t<VarTypes>...)>::Type InFunc, VarTypes&&... Vars)
 	{
-		Super::template CreateDelegateInstance<TBaseStaticDelegateInstance<FuncType, UserPolicy, std::decay_t<VarTypes>...>>(InFunc, Forward<VarTypes>(Vars)...);
+		TDelegate<RetValType(ParamTypes...), UserPolicy> Result;
+		Result.template CreateDelegateInstance<TBaseStaticDelegateInstance<FuncType, UserPolicy, std::decay_t<VarTypes>...>>(InFunc, Forward<VarTypes>(Vars)...);
+		return Result;
 	}
-	
+
 	/**
-	 * Static: Binds a C++ lambda delegate
+	 * Static: Creates a C++ lambda delegate
 	 * technically this works for any functor types, but lambdas are the primary use case
 	 */
 	template<typename FunctorType, typename... VarTypes>
-	inline void BindLambda(FunctorType&& InFunctor, VarTypes&&... Vars)
+	[[nodiscard]] inline static TDelegate<RetValType(ParamTypes...), UserPolicy> CreateLambda(FunctorType&& InFunctor, VarTypes&&... Vars)
 	{
-		Super::template CreateDelegateInstance<TBaseFunctorDelegateInstance<FuncType, UserPolicy, typename TRemoveReference<FunctorType>::Type, std::decay_t<VarTypes>...>>(Forward<FunctorType>(InFunctor), Forward<VarTypes>(Vars)...);
+		TDelegate<RetValType(ParamTypes...), UserPolicy> Result;
+		Result.template CreateDelegateInstance<TBaseFunctorDelegateInstance<FuncType, UserPolicy, typename TRemoveReference<FunctorType>::Type, std::decay_t<VarTypes>...>>(Forward<FunctorType>(InFunctor), Forward<VarTypes>(Vars)...);
+		return Result;
 	}
 
 	/**
-	 * Static: Binds a weak shared pointer C++ lambda delegate
+	 * Static: Creates a weak shared pointer C++ lambda delegate
 	 * technically this works for any functor types, but lambdas are the primary use case
 	 */
 	template<typename UserClass, ESPMode Mode, typename FunctorType, typename... VarTypes>
-	inline void BindSPLambda(const TSharedRef<UserClass, Mode>& InUserObjectRef, FunctorType&& InFunctor, VarTypes&&... Vars)
+	[[nodiscard]] inline static TDelegate<RetValType(ParamTypes...), UserPolicy> CreateSPLambda(const TSharedRef<UserClass, Mode>& InUserObjectRef, FunctorType&& InFunctor, VarTypes&&... Vars)
 	{
-		Super::template CreateDelegateInstance<TBaseSPLambdaDelegateInstance<const UserClass, Mode, FuncType, UserPolicy, typename TRemoveReference<FunctorType>::Type, std::decay_t<VarTypes>...>>(InUserObjectRef, Forward<FunctorType>(InFunctor), Forward<VarTypes>(Vars)...);
+		TDelegate<RetValType(ParamTypes...), UserPolicy> Result;
+		Result.template CreateDelegateInstance<TBaseSPLambdaDelegateInstance<UserClass, Mode, FuncType, UserPolicy, typename TRemoveReference<FunctorType>::Type, std::decay_t<VarTypes>...>>(InUserObjectRef, Forward<FunctorType>(InFunctor), Forward<VarTypes>(Vars)...);
+		return Result;
 	}
 	template <typename UserClass, typename FunctorType, typename... VarTypes>
-	inline void BindSPLambda(const UserClass* InUserObject, FunctorType&& InFunctor, VarTypes&&... Vars)
+	[[nodiscard]] inline static TDelegate<RetValType(ParamTypes...), UserPolicy> CreateSPLambda(UserClass* InUserObject, FunctorType&& InFunctor, VarTypes&&... Vars)
 	{
-		Super::template CreateDelegateInstance<TBaseSPLambdaDelegateInstance<const UserClass, decltype(InUserObject->AsShared())::Mode, FuncType, UserPolicy, typename TRemoveReference<FunctorType>::Type, std::decay_t<VarTypes>...>>(StaticCastSharedRef<const UserClass>(InUserObject->AsShared()), Forward<FunctorType>(InFunctor), Forward<VarTypes>(Vars)...);
+		TDelegate<RetValType(ParamTypes...), UserPolicy> Result;
+		Result.template CreateDelegateInstance<TBaseSPLambdaDelegateInstance<UserClass, decltype(InUserObject->AsShared())::Mode, FuncType, UserPolicy, typename TRemoveReference<FunctorType>::Type, std::decay_t<VarTypes>...>>(StaticCastSharedRef<UserClass>(InUserObject->AsShared()), Forward<FunctorType>(InFunctor), Forward<VarTypes>(Vars)...);
+		return Result;
 	}
 
 	/**
-	 * Static: Binds a weak object C++ lambda delegate
+	 * Static: Creates a weak object C++ lambda delegate
 	 * technically this works for any functor types, but lambdas are the primary use case
 	 */
 	template<typename UserClass, typename FunctorType, typename... VarTypes>
-	inline void BindWeakLambda(UserClass* InUserObject, FunctorType&& InFunctor, VarTypes&&... Vars)
+	[[nodiscard]] inline static TDelegate<RetValType(ParamTypes...), UserPolicy> CreateWeakLambda(UserClass* InUserObject, FunctorType&& InFunctor, VarTypes&&... Vars)
 	{
-		Super::template CreateDelegateInstance<TWeakBaseFunctorDelegateInstance<UserClass, FuncType, UserPolicy, typename TRemoveReference<FunctorType>::Type, std::decay_t<VarTypes>...>>(InUserObject, Forward<FunctorType>(InFunctor), Forward<VarTypes>(Vars)...);
+		TDelegate<RetValType(ParamTypes...), UserPolicy> Result;
+		Result.template CreateDelegateInstance<TWeakBaseFunctorDelegateInstance<UserClass, FuncType, UserPolicy, typename TRemoveReference<FunctorType>::Type, std::decay_t<VarTypes>...>>(InUserObject, Forward<FunctorType>(InFunctor), Forward<VarTypes>(Vars)...);
+		return Result;
 	}
 
 	/**
-	 * Binds a raw C++ pointer delegate.
+	 * Static: Creates a raw C++ pointer member function delegate.
 	 *
 	 * Raw pointer doesn't use any sort of reference, so may be unsafe to call if the object was
 	 * deleted out from underneath your delegate. Be careful when calling Execute()!
 	 */
 	template <typename UserClass, typename... VarTypes>
-	inline void BindRaw(UserClass* InUserObject, typename TMemFunPtrType<false, UserClass, RetValType (ParamTypes..., std::decay_t<VarTypes>...)>::Type InFunc, VarTypes&&... Vars)
+	[[nodiscard]] inline static TDelegate<RetValType(ParamTypes...), UserPolicy> CreateRaw(UserClass* InUserObject, typename TMemFunPtrType<false, UserClass, RetValType (ParamTypes..., std::decay_t<VarTypes>...)>::Type InFunc, VarTypes&&... Vars)
 	{
 		static_assert(!TIsConst<UserClass>::Value, "Attempting to bind a delegate with a const object pointer and non-const member function.");
 
-		Super::template CreateDelegateInstance<TBaseRawMethodDelegateInstance<false, UserClass, FuncType, UserPolicy, std::decay_t<VarTypes>...>>(InUserObject, InFunc, Forward<VarTypes>(Vars)...);
+		TDelegate<RetValType(ParamTypes...), UserPolicy> Result;
+		Result.template CreateDelegateInstance<TBaseRawMethodDelegateInstance<false, UserClass, FuncType, UserPolicy, std::decay_t<VarTypes>...>>(InUserObject, InFunc, Forward<VarTypes>(Vars)...);
+		return Result;
 	}
 	template <typename UserClass, typename... VarTypes>
-	inline void BindRaw(const UserClass* InUserObject, typename TMemFunPtrType<true, UserClass, RetValType (ParamTypes..., std::decay_t<VarTypes>...)>::Type InFunc, VarTypes&&... Vars)
+	[[nodiscard]] inline static TDelegate<RetValType(ParamTypes...), UserPolicy> CreateRaw(const UserClass* InUserObject, typename TMemFunPtrType<true, UserClass, RetValType (ParamTypes..., std::decay_t<VarTypes>...)>::Type InFunc, VarTypes&&... Vars)
 	{
-		Super::template CreateDelegateInstance<TBaseRawMethodDelegateInstance<true, const UserClass, FuncType, UserPolicy, std::decay_t<VarTypes>...>>(InUserObject, InFunc, Forward<VarTypes>(Vars)...);
+		TDelegate<RetValType(ParamTypes...), UserPolicy> Result;
+		Result.template CreateDelegateInstance<TBaseRawMethodDelegateInstance<true, const UserClass, FuncType, UserPolicy, std::decay_t<VarTypes>...>>(InUserObject, InFunc, Forward<VarTypes>(Vars)...);
+		return Result;
 	}
 
 	/**
-	 * Binds a shared pointer-based member function delegate.  Shared pointer delegates keep a weak reference to your object.  You can use ExecuteIfBound() to call them.
+	 * Static: Creates a shared pointer-based member function delegate.
+	 *
+	 * Shared pointer delegates keep a weak reference to your object.
+	 * You can use ExecuteIfBound() to call them.
 	 */
 	template <typename UserClass, ESPMode Mode, typename... VarTypes>
-	inline void BindSP(const TSharedRef<UserClass, Mode>& InUserObjectRef, typename TMemFunPtrType<false, UserClass, RetValType (ParamTypes..., std::decay_t<VarTypes>...)>::Type InFunc, VarTypes&&... Vars)
+	[[nodiscard]] inline static TDelegate<RetValType(ParamTypes...), UserPolicy> CreateSP(const TSharedRef<UserClass, Mode>& InUserObjectRef, typename TMemFunPtrType<false, UserClass, RetValType (ParamTypes..., std::decay_t<VarTypes>...)>::Type InFunc, VarTypes&&... Vars)
 	{
 		static_assert(!TIsConst<UserClass>::Value, "Attempting to bind a delegate with a const object pointer and non-const member function.");
 
-		Super::template CreateDelegateInstance<TBaseSPMethodDelegateInstance<false, UserClass, Mode, FuncType, UserPolicy, std::decay_t<VarTypes>...>>(InUserObjectRef, InFunc, Forward<VarTypes>(Vars)...);
+		TDelegate<RetValType(ParamTypes...), UserPolicy> Result;
+		Result.template CreateDelegateInstance<TBaseSPMethodDelegateInstance<false, UserClass, Mode, FuncType, UserPolicy, std::decay_t<VarTypes>...>>(InUserObjectRef, InFunc, Forward<VarTypes>(Vars)...);
+		return Result;
 	}
 	template <typename UserClass, ESPMode Mode, typename... VarTypes>
-	inline void BindSP(const TSharedRef<UserClass, Mode>& InUserObjectRef, typename TMemFunPtrType<true, UserClass, RetValType (ParamTypes..., std::decay_t<VarTypes>...)>::Type InFunc, VarTypes&&... Vars)
+	[[nodiscard]] inline static TDelegate<RetValType(ParamTypes...), UserPolicy> CreateSP(const TSharedRef<UserClass, Mode>& InUserObjectRef, typename TMemFunPtrType<true, UserClass, RetValType (ParamTypes..., std::decay_t<VarTypes>...)>::Type InFunc, VarTypes&&... Vars)
 	{
-		Super::template CreateDelegateInstance<TBaseSPMethodDelegateInstance<true, const UserClass, Mode, FuncType, UserPolicy, std::decay_t<VarTypes>...>>(InUserObjectRef, InFunc, Forward<VarTypes>(Vars)...);
+		TDelegate<RetValType(ParamTypes...), UserPolicy> Result;
+		Result.template CreateDelegateInstance<TBaseSPMethodDelegateInstance<true, const UserClass, Mode, FuncType, UserPolicy, std::decay_t<VarTypes>...>>(InUserObjectRef, InFunc, Forward<VarTypes>(Vars)...);
+		return Result;
 	}
 
 	/**
-	 * Binds a shared pointer-based member function delegate.
+	 * Static: Creates a shared pointer-based member function delegate.
 	 *
 	 * Shared pointer delegates keep a weak reference to your object.
 	 * You can use ExecuteIfBound() to call them.
 	 */
 	template <typename UserClass, typename... VarTypes>
-	inline void BindSP(UserClass* InUserObject, typename TMemFunPtrType<false, UserClass, RetValType (ParamTypes..., std::decay_t<VarTypes>...)>::Type InFunc, VarTypes&&... Vars)
+	[[nodiscard]] inline static TDelegate<RetValType(ParamTypes...), UserPolicy> CreateSP(UserClass* InUserObject, typename TMemFunPtrType<false, UserClass, RetValType (ParamTypes..., std::decay_t<VarTypes>...)>::Type InFunc, VarTypes&&... Vars)
 	{
 		static_assert(!TIsConst<UserClass>::Value, "Attempting to bind a delegate with a const object pointer and non-const member function.");
 
-		Super::template CreateDelegateInstance<TBaseSPMethodDelegateInstance<false, UserClass, decltype(InUserObject->AsShared())::Mode, FuncType, UserPolicy, std::decay_t<VarTypes>...>>(StaticCastSharedRef<UserClass>(InUserObject->AsShared()), InFunc, Forward<VarTypes>(Vars)...);
+		TDelegate<RetValType(ParamTypes...), UserPolicy> Result;
+		Result.template CreateDelegateInstance<TBaseSPMethodDelegateInstance<false, UserClass, decltype(InUserObject->AsShared())::Mode, FuncType, UserPolicy, std::decay_t<VarTypes>...>>(StaticCastSharedRef<UserClass>(InUserObject->AsShared()), InFunc, Forward<VarTypes>(Vars)...);
+		return Result;
 	}
 	template <typename UserClass, typename... VarTypes>
-	inline void BindSP(const UserClass* InUserObject, typename TMemFunPtrType<true, UserClass, RetValType (ParamTypes..., std::decay_t<VarTypes>...)>::Type InFunc, VarTypes&&... Vars)
+	[[nodiscard]] inline static TDelegate<RetValType(ParamTypes...), UserPolicy> CreateSP(const UserClass* InUserObject, typename TMemFunPtrType<true, UserClass, RetValType (ParamTypes..., std::decay_t<VarTypes>...)>::Type InFunc, VarTypes&&... Vars)
 	{
-		Super::template CreateDelegateInstance<TBaseSPMethodDelegateInstance<true, const UserClass, decltype(InUserObject->AsShared())::Mode, FuncType, UserPolicy, std::decay_t<VarTypes>...>>(StaticCastSharedRef<const UserClass>(InUserObject->AsShared()), InFunc, Forward<VarTypes>(Vars)...);
+		TDelegate<RetValType(ParamTypes...), UserPolicy> Result;
+		Result.template CreateDelegateInstance<TBaseSPMethodDelegateInstance<true, const UserClass, decltype(InUserObject->AsShared())::Mode, FuncType, UserPolicy, std::decay_t<VarTypes>...>>(StaticCastSharedRef<const UserClass>(InUserObject->AsShared()), InFunc, Forward<VarTypes>(Vars)...);
+		return Result;
 	}
 
 	/**
-	 * Binds a shared pointer-based (thread-safe) member function delegate.
+	 * Static: Creates a shared pointer-based (thread-safe) member function delegate.
 	 *
 	 * Shared pointer delegates keep a weak reference to your object.
 	 * You can use ExecuteIfBound() to call them.
 	 *
-	 * Note: This function is redundant, but is retained for backwards compatibility.  BindSP() works in both thread-safe and not-thread-safe modes and should be preferred.
+	 * Note: This function is redundant, but is retained for backwards compatibility.  CreateSP() works in both thread-safe and not-thread-safe modes and should be preferred.
 	 */
 	template <typename UserClass, typename... VarTypes>
-	inline void BindThreadSafeSP(const TSharedRef<UserClass, ESPMode::ThreadSafe>& InUserObjectRef, typename TMemFunPtrType<false, UserClass, RetValType (ParamTypes..., std::decay_t<VarTypes>...)>::Type InFunc, VarTypes&&... Vars)
+	[[nodiscard]] inline static TDelegate<RetValType(ParamTypes...), UserPolicy> CreateThreadSafeSP(const TSharedRef<UserClass, ESPMode::ThreadSafe>& InUserObjectRef, typename TMemFunPtrType<false, UserClass, RetValType (ParamTypes..., std::decay_t<VarTypes>...)>::Type InFunc, VarTypes&&... Vars)
 	{
 		static_assert(!TIsConst<UserClass>::Value, "Attempting to bind a delegate with a const object pointer and non-const member function.");
 
-		Super::template CreateDelegateInstance<TBaseSPMethodDelegateInstance<false, UserClass, ESPMode::ThreadSafe, FuncType, UserPolicy, std::decay_t<VarTypes>...>>(InUserObjectRef, InFunc, Forward<VarTypes>(Vars)...);
+		TDelegate<RetValType(ParamTypes...), UserPolicy> Result;
+		Result.template CreateDelegateInstance<TBaseSPMethodDelegateInstance<false, UserClass, ESPMode::ThreadSafe, FuncType, UserPolicy, std::decay_t<VarTypes>...>>(InUserObjectRef, InFunc, Forward<VarTypes>(Vars)...);
+		return Result;
 	}
 	template <typename UserClass, typename... VarTypes>
-	inline void BindThreadSafeSP(const TSharedRef<UserClass, ESPMode::ThreadSafe>& InUserObjectRef, typename TMemFunPtrType<true, UserClass, RetValType (ParamTypes..., std::decay_t<VarTypes>...)>::Type InFunc, VarTypes&&... Vars)
+	[[nodiscard]] inline static TDelegate<RetValType(ParamTypes...), UserPolicy> CreateThreadSafeSP(const TSharedRef<UserClass, ESPMode::ThreadSafe>& InUserObjectRef, typename TMemFunPtrType<true, UserClass, RetValType (ParamTypes..., std::decay_t<VarTypes>...)>::Type InFunc, VarTypes&&... Vars)
 	{
-		Super::template CreateDelegateInstance<TBaseSPMethodDelegateInstance<true, const UserClass, ESPMode::ThreadSafe, FuncType, UserPolicy, std::decay_t<VarTypes>...>>(InUserObjectRef, InFunc, Forward<VarTypes>(Vars)...);
+		TDelegate<RetValType(ParamTypes...), UserPolicy> Result;
+		Result.template CreateDelegateInstance<TBaseSPMethodDelegateInstance<true, const UserClass, ESPMode::ThreadSafe, FuncType, UserPolicy, std::decay_t<VarTypes>...>>(InUserObjectRef, InFunc, Forward<VarTypes>(Vars)...);
+		return Result;
 	}
 
 	/**
-	 * Binds a shared pointer-based (thread-safe) member function delegate.
+	 * Static: Creates a shared pointer-based (thread-safe) member function delegate.
 	 *
 	 * Shared pointer delegates keep a weak reference to your object.
 	 * You can use ExecuteIfBound() to call them.
 	 *
-	 * Note: This function is redundant, but is retained for backwards compatibility.  BindSP() works in both thread-safe and not-thread-safe modes and should be preferred.
+	 * Note: This function is redundant, but is retained for backwards compatibility.  CreateSP() works in both thread-safe and not-thread-safe modes and should be preferred.
 	 */
 	template <typename UserClass, typename... VarTypes>
-	inline void BindThreadSafeSP(UserClass* InUserObject, typename TMemFunPtrType<false, UserClass, RetValType (ParamTypes..., std::decay_t<VarTypes>...)>::Type InFunc, VarTypes&&... Vars)
+	[[nodiscard]] inline static TDelegate<RetValType(ParamTypes...), UserPolicy> CreateThreadSafeSP(UserClass* InUserObject, typename TMemFunPtrType<false, UserClass, RetValType (ParamTypes..., std::decay_t<VarTypes>...)>::Type InFunc, VarTypes&&... Vars)
 	{
 		static_assert(!TIsConst<UserClass>::Value, "Attempting to bind a delegate with a const object pointer and non-const member function.");
 
-		Super::template CreateDelegateInstance<TBaseSPMethodDelegateInstance<false, UserClass, ESPMode::ThreadSafe, FuncType, UserPolicy, std::decay_t<VarTypes>...>>(StaticCastSharedRef<UserClass>(InUserObject->AsShared()), InFunc, Forward<VarTypes>(Vars)...);
+		TDelegate<RetValType(ParamTypes...), UserPolicy> Result;
+		Result.template CreateDelegateInstance<TBaseSPMethodDelegateInstance<false, UserClass, ESPMode::ThreadSafe, FuncType, UserPolicy, std::decay_t<VarTypes>...>>(StaticCastSharedRef<UserClass>(InUserObject->AsShared()), InFunc, Forward<VarTypes>(Vars)...);
+		return Result;
 	}
 	template <typename UserClass, typename... VarTypes>
-	inline void BindThreadSafeSP(const UserClass* InUserObject, typename TMemFunPtrType<true, UserClass, RetValType (ParamTypes..., std::decay_t<VarTypes>...)>::Type InFunc, VarTypes&&... Vars)
+	[[nodiscard]] inline static TDelegate<RetValType(ParamTypes...), UserPolicy> CreateThreadSafeSP(const UserClass* InUserObject, typename TMemFunPtrType<true, UserClass, RetValType (ParamTypes..., std::decay_t<VarTypes>...)>::Type InFunc, VarTypes&&... Vars)
 	{
-		Super::template CreateDelegateInstance<TBaseSPMethodDelegateInstance<true, const UserClass, ESPMode::ThreadSafe, FuncType, UserPolicy, std::decay_t<VarTypes>...>>(StaticCastSharedRef<const UserClass>(InUserObject->AsShared()), InFunc, Forward<VarTypes>(Vars)...);
+		TDelegate<RetValType(ParamTypes...), UserPolicy> Result;
+		Result.template CreateDelegateInstance<TBaseSPMethodDelegateInstance<true, const UserClass, ESPMode::ThreadSafe, FuncType, UserPolicy, std::decay_t<VarTypes>...>>(StaticCastSharedRef<const UserClass>(InUserObject->AsShared()), InFunc, Forward<VarTypes>(Vars)...);
+		return Result;
 	}
 
 	/**
-	 * Binds a UFunction-based member function delegate.
+	 * Static: Creates a UFunction-based member function delegate.
 	 *
 	 * UFunction delegates keep a weak reference to your object.
 	 * You can use ExecuteIfBound() to call them.
 	 */
 	template <typename UObjectTemplate, typename... VarTypes>
-	inline void BindUFunction(UObjectTemplate* InUserObject, const FName& InFunctionName, VarTypes&&... Vars)
+	[[nodiscard]] inline static TDelegate<RetValType(ParamTypes...), UserPolicy> CreateUFunction(UObjectTemplate* InUserObject, const FName& InFunctionName, VarTypes&&... Vars)
 	{
-		Super::template CreateDelegateInstance<TBaseUFunctionDelegateInstance<UObjectTemplate, FuncType, UserPolicy, std::decay_t<VarTypes>...>>(InUserObject, InFunctionName, Forward<VarTypes>(Vars)...);
+		TDelegate<RetValType(ParamTypes...), UserPolicy> Result;
+		Result.template CreateDelegateInstance<TBaseUFunctionDelegateInstance<UObjectTemplate, FuncType, UserPolicy, std::decay_t<VarTypes>...>>(InUserObject, InFunctionName, Forward<VarTypes>(Vars)...);
+		return Result;
 	}
 	template <typename UObjectTemplate, typename... VarTypes>
-	inline void BindUFunction(TObjectPtr<UObjectTemplate> InUserObject, const FName& InFunctionName, VarTypes&&... Vars)
+	[[nodiscard]] inline static TDelegate<RetValType(ParamTypes...), UserPolicy> CreateUFunction(TObjectPtr<UObjectTemplate> InUserObject, const FName& InFunctionName, VarTypes&&... Vars)
 	{
-		Super::template CreateDelegateInstance<TBaseUFunctionDelegateInstance<UObjectTemplate, FuncType, UserPolicy, std::decay_t<VarTypes>...>>(ToRawPtr(InUserObject), InFunctionName, Forward<VarTypes>(Vars)...);
+		TDelegate<RetValType(ParamTypes...), UserPolicy> Result;
+		Result.template CreateDelegateInstance<TBaseUFunctionDelegateInstance<UObjectTemplate, FuncType, UserPolicy, std::decay_t<VarTypes>...>>(ToRawPtr(InUserObject), InFunctionName, Forward<VarTypes>(Vars)...);
+		return Result;
 	}
 
 	/**
-	 * Binds a UObject-based member function delegate.
+	 * Static: Creates a UObject-based member function delegate.
 	 *
 	 * UObject delegates keep a weak reference to your object.
 	 * You can use ExecuteIfBound() to call them.
 	 */
 	template <typename UserClass, typename... VarTypes>
-	inline void BindUObject(UserClass* InUserObject, typename TMemFunPtrType<false, UserClass, RetValType (ParamTypes..., std::decay_t<VarTypes>...)>::Type InFunc, VarTypes&&... Vars)
+	[[nodiscard]] inline static TDelegate<RetValType(ParamTypes...), UserPolicy> CreateUObject(UserClass* InUserObject, typename TMemFunPtrType<false, UserClass, RetValType (ParamTypes..., std::decay_t<VarTypes>...)>::Type InFunc, VarTypes&&... Vars)
 	{
 		static_assert(!TIsConst<UserClass>::Value, "Attempting to bind a delegate with a const object pointer and non-const member function.");
 
-		Super::template CreateDelegateInstance<TBaseUObjectMethodDelegateInstance<false, UserClass, FuncType, UserPolicy, std::decay_t<VarTypes>...>>(InUserObject, InFunc, Forward<VarTypes>(Vars)...);
+		TDelegate<RetValType(ParamTypes...), UserPolicy> Result;
+		Result.template CreateDelegateInstance<TBaseUObjectMethodDelegateInstance<false, UserClass, FuncType, UserPolicy, std::decay_t<VarTypes>...>>(InUserObject, InFunc, Forward<VarTypes>(Vars)...);
+		return Result;
 	}
 	template <typename UserClass, typename... VarTypes>
-	inline void BindUObject(const UserClass* InUserObject, typename TMemFunPtrType<true, UserClass, RetValType (ParamTypes..., std::decay_t<VarTypes>...)>::Type InFunc, VarTypes&&... Vars)
+	[[nodiscard]] inline static TDelegate<RetValType(ParamTypes...), UserPolicy> CreateUObject(const UserClass* InUserObject, typename TMemFunPtrType<true, UserClass, RetValType (ParamTypes..., std::decay_t<VarTypes>...)>::Type InFunc, VarTypes&&... Vars)
 	{
-		Super::template CreateDelegateInstance<TBaseUObjectMethodDelegateInstance<true, const UserClass, FuncType, UserPolicy, std::decay_t<VarTypes>...>>(InUserObject, InFunc, Forward<VarTypes>(Vars)...);
+		TDelegate<RetValType(ParamTypes...), UserPolicy> Result;
+		Result.template CreateDelegateInstance<TBaseUObjectMethodDelegateInstance<true, const UserClass, FuncType, UserPolicy, std::decay_t<VarTypes>...>>(InUserObject, InFunc, Forward<VarTypes>(Vars)...);
+		return Result;
 	}
 	template <typename UserClass, typename... VarTypes>
-	inline void BindUObject(TObjectPtr<UserClass> InUserObject, typename TMemFunPtrType<false, UserClass, RetValType (ParamTypes..., std::decay_t<VarTypes>...)>::Type InFunc, VarTypes&&... Vars)
+	[[nodiscard]] inline static TDelegate<RetValType(ParamTypes...), UserPolicy> CreateUObject(TObjectPtr<UserClass> InUserObject, typename TMemFunPtrType<false, UserClass, RetValType (ParamTypes..., std::decay_t<VarTypes>...)>::Type InFunc, VarTypes&&... Vars)
 	{
 		static_assert(!TIsConst<UserClass>::Value, "Attempting to bind a delegate with a const object pointer and non-const member function.");
 
-		Super::template CreateDelegateInstance<TBaseUObjectMethodDelegateInstance<false, UserClass, FuncType, UserPolicy, std::decay_t<VarTypes>...>>(ToRawPtr(InUserObject), InFunc, Forward<VarTypes>(Vars)...);
+		TDelegate<RetValType(ParamTypes...), UserPolicy> Result;
+		Result.template CreateDelegateInstance<TBaseUObjectMethodDelegateInstance<false, UserClass, FuncType, UserPolicy, std::decay_t<VarTypes>...>>(ToRawPtr(InUserObject), InFunc, Forward<VarTypes>(Vars)...);
+		return Result;
 	}
 	template <typename UserClass, typename... VarTypes>
-	inline void BindUObject(TObjectPtr<UserClass> InUserObject, typename TMemFunPtrType<true, UserClass, RetValType (ParamTypes..., std::decay_t<VarTypes>...)>::Type InFunc, VarTypes&&... Vars)
+	[[nodiscard]] inline static TDelegate<RetValType(ParamTypes...), UserPolicy> CreateUObject(TObjectPtr<UserClass> InUserObject, typename TMemFunPtrType<true, UserClass, RetValType (ParamTypes..., std::decay_t<VarTypes>...)>::Type InFunc, VarTypes&&... Vars)
 	{
-		Super::template CreateDelegateInstance<TBaseUObjectMethodDelegateInstance<true, const UserClass, FuncType, UserPolicy, std::decay_t<VarTypes>...>>(ToRawPtr(InUserObject), InFunc, Forward<VarTypes>(Vars)...);
+		TDelegate<RetValType(ParamTypes...), UserPolicy> Result;
+		Result.template CreateDelegateInstance<TBaseUObjectMethodDelegateInstance<true, const UserClass, FuncType, UserPolicy, std::decay_t<VarTypes>...>>(ToRawPtr(InUserObject), InFunc, Forward<VarTypes>(Vars)...);
+		return Result;
 	}
 
 public:
@@ -625,45 +690,60 @@
 /////////////////////////////////////////////////////////////////////////////////////////////////////////////////////////////
 
 template <typename DelegateSignature>
+using TTSDelegateRegistration = TDelegateRegistration<DelegateSignature, FDefaultTSDelegateUserPolicy>;
+
+template <typename DelegateSignature>
 using TTSDelegate = TDelegate<DelegateSignature, FDefaultTSDelegateUserPolicy>;
 
 /////////////////////////////////////////////////////////////////////////////////////////////////////////////////////////////
 
 /**
- * Multicast delegate template base class, used for both normal and event multicast delegates.
+ * Multicast delegate registration class.  A function can expose this base class to users to let them bind
+ * against, but without letting them broadcast.  This gives us the benefit of C# events.
  *
- * This class implements the functionality of multicast delegates. It is templated to the function signature
- * that it is compatible with. Use the various DECLARE_MULTICAST_DELEGATE and DECLARE_EVENT macros to create
- * actual delegate types.
+ * Example:
  *
- * Multicast delegates offer no guarantees for the calling order of bound functions. As bindings get added
- * and removed over time, the calling order may change. Only bindings without return values are supported.
+ * struct FInterestingType
+ * {
+ * public:
+ *     TMulticastDelegateRegistration<void()>& OnInterestingThing() const
+ *     {
+ *         return OnInterestingThingDelegate;
+ *     }
+ *
+ * private:
+ *     mutable TMulticastDelegate<void()> OnInterestingThingDelegate;
+ * };
+ *
+ * void Func(const FInterestingThing& Thing)
+ * {
+ *     // Registration and reregistration are allowed
+ *     FDelegateHandle Handle = Thing.OnInterestingThing().Add([](){ RespondToInterestingThing(); });
+ *     Thing.OnInterestingThing().Remove(Handle);
+ *
+ *     // Broadcast is not
+ *     Thing.OnInterestingThing().Broadcast(); // error: Broadcast is deleted
+ * }
  */
 template <typename DelegateSignature, typename UserPolicy = FDefaultDelegateUserPolicy>
-class TMulticastDelegate
-{
-	static_assert(sizeof(DelegateSignature) == 0, "Expected a function signature for the delegate template parameter");
-};
+class TMulticastDelegateRegistration;
 
-template <typename RetValType, typename... ParamTypes, typename UserPolicy>
-class TMulticastDelegate<RetValType(ParamTypes...), UserPolicy>
-{
-	static_assert(sizeof(RetValType) == 0, "The return type of a multicast delegate must be void");
-};
-
 template <typename... ParamTypes, typename UserPolicy>
-class TMulticastDelegate<void(ParamTypes...), UserPolicy> : public UserPolicy::FMulticastDelegateExtras
+class TMulticastDelegateRegistration<void(ParamTypes...), UserPolicy> : public UserPolicy::FMulticastDelegateExtras
 {
+protected:
 	using Super                         = typename UserPolicy::FMulticastDelegateExtras;
-	using InvocationListType            = typename Super::InvocationListType;
 	using DelegateInstanceInterfaceType = IBaseDelegateInstance<void (ParamTypes...), UserPolicy>;
 
+private:
+	using InvocationListType = typename Super::InvocationListType;
+
 public:
 	/** Type definition for unicast delegate classes whose delegate instances are compatible with this delegate. */
 	using FDelegate = TDelegate<void(ParamTypes...), UserPolicy>;
 
 public:
-	// Make sure TMulticastDelegateBase's public functions are publicly exposed through the TMulticastDelegate API
+	// Make sure TMulticastDelegateBase's public functions are publicly exposed through the TMulticastDelegateRegistration API
 	using Super::Clear;
 	using Super::IsBound;
 	using Super::IsBoundToObject;
@@ -671,29 +751,17 @@
 	using Super::GetAllocatedSize;
 
 private:
-	// Make sure TMulticastDelegateBase's protected functions are not accidentally exposed through the TMulticastDelegate API
+	// Make sure TMulticastDelegateBase's protected functions are not accidentally exposed through the TMulticastDelegateRegistration API
 	using Super::AddDelegateInstance;
 	using Super::RemoveDelegateInstance;
 
-public:
-	TMulticastDelegate() = default;
-
-	TMulticastDelegate(const TMulticastDelegate& Other)
-	{
-		*this = Other;
-	}
-
-	TMulticastDelegate& operator=(const TMulticastDelegate& Other)
-	{
-		if (&Other != this)
-		{
-			Super::template CopyFrom<DelegateInstanceInterfaceType>(Other);
-		}
-		return *this;
-	}
-
-	TMulticastDelegate(TMulticastDelegate&&) = default;
-	TMulticastDelegate& operator=(TMulticastDelegate&&) = default;
+protected:
+	TMulticastDelegateRegistration() = default;
+	TMulticastDelegateRegistration(TMulticastDelegateRegistration&&) = default;
+	TMulticastDelegateRegistration(const TMulticastDelegateRegistration&) = default;
+	TMulticastDelegateRegistration& operator=(TMulticastDelegateRegistration&&) = default;
+	TMulticastDelegateRegistration& operator=(const TMulticastDelegateRegistration&) = default;
+	~TMulticastDelegateRegistration() = default;
 
 public:
 
@@ -946,17 +1014,72 @@
 		return bResult;
 	}
 
+	// Broadcasting via a delegate registration reference is not allowed
+	void Broadcast(ParamTypes... Params) = delete;
+};
+
+/**
+ * Multicast delegate template base class, used for both normal and event multicast delegates.
+ *
+ * This class implements the functionality of multicast delegates. It is templated to the function signature
+ * that it is compatible with. Use the various DECLARE_MULTICAST_DELEGATE and DECLARE_EVENT macros to create
+ * actual delegate types.
+ *
+ * Multicast delegates offer no guarantees for the calling order of bound functions. As bindings get added
+ * and removed over time, the calling order may change. Only bindings without return values are supported.
+ */
+template <typename DelegateSignature, typename UserPolicy = FDefaultDelegateUserPolicy>
+class TMulticastDelegate
+{
+	static_assert(sizeof(DelegateSignature) == 0, "Expected a function signature for the delegate template parameter");
+};
+
+template <typename RetValType, typename... ParamTypes, typename UserPolicy>
+class TMulticastDelegate<RetValType(ParamTypes...), UserPolicy>
+{
+	static_assert(sizeof(RetValType) == 0, "The return type of a multicast delegate must be void");
+};
+
+template <typename... ParamTypes, typename UserPolicy>
+class TMulticastDelegate<void(ParamTypes...), UserPolicy> : public TMulticastDelegateRegistration<void(ParamTypes...), UserPolicy>
+{
+private:
+	using Super                         = TMulticastDelegateRegistration<void(ParamTypes...), UserPolicy>;
+	using DelegateInstanceInterfaceType = typename Super::DelegateInstanceInterfaceType;
+
+public:
+	using RegistrationType = Super;
+
+	TMulticastDelegate() = default;
+
+	TMulticastDelegate(const TMulticastDelegate& Other)
+	{
+		*this = Other;
+	}
+
+	TMulticastDelegate& operator=(const TMulticastDelegate& Other)
+	{
+		if (&Other != this)
+		{
+			Super::template CopyFrom<DelegateInstanceInterfaceType>(Other);
+		}
+		return *this;
+	}
+
 	/**
 	 * Broadcasts this delegate to all bound objects, except to those that may have expired.
 	 *
-	 * The constness of this method is a lie, but it allows for broadcasting from const functions.
+	 * The constness of this method allows for broadcasting from const functions.
 	 */
 	void Broadcast(ParamTypes... Params) const
 	{
-		Super::template Broadcast<DelegateInstanceInterfaceType, ParamTypes...>(Params...);
+		Super::Super::template Broadcast<typename Super::DelegateInstanceInterfaceType, ParamTypes...>(Params...);
 	}
 };
 
 template <typename DelegateSignature>
+using TTSMulticastDelegateRegistration = TMulticastDelegateRegistration<DelegateSignature, FDefaultTSDelegateUserPolicy>;
+
+template <typename DelegateSignature>
 using TTSMulticastDelegate = TMulticastDelegate<DelegateSignature, FDefaultTSDelegateUserPolicy>;
 
diff --git a/Engine/Source/Runtime/Core/Public/Modules/ModuleManager.h b/Engine/Source/Runtime/Core/Public/Modules/ModuleManager.h
--- a/Engine/Source/Runtime/Core/Public/Modules/ModuleManager.h
+++ b/Engine/Source/Runtime/Core/Public/Modules/ModuleManager.h
@@ -510,8 +510,8 @@ public:
 	 *
 	 * @return The multicast delegate.
 	 */
-	using FModulesChangedEvent = TTSMulticastDelegate<void(FName ModuleName, EModuleChangeReason ChangeReason)>;
-	FModulesChangedEvent& OnModulesChanged( )
+	using FModulesChangedEvent UE_DEPRECATED(5.5, "The FModulesChangedEvent typedef has been deprecated - use TTSMulticastDelegateRegistration<void(FName, EModuleChangeReason)> instead.") = TTSMulticastDelegate<void(FName ModuleName, EModuleChangeReason ChangeReason)>;
+	TTSMulticastDelegateRegistration<void(FName ModuleName, EModuleChangeReason ChangeReason)>& OnModulesChanged( )
 	{
 		return ModulesChangedEvent;
 	}
@@ -524,8 +524,8 @@ public:
 	 *
 	 * @return The delegate.
 	 */
-	DECLARE_EVENT_TwoParams(FModuleManager, ProcessLoadedObjectsEvent, FName, bool);
-	ProcessLoadedObjectsEvent& OnProcessLoadedObjectsCallback()
+	using ProcessLoadedObjectsEvent UE_DEPRECATED(5.5, "The ProcessLoadedObjectsEvent typedef has been deprecated - use TMulticastDelegateRegistration<void(FName, bool)> instead.") = TMulticastDelegate<void(FName, bool)>;
+	TMulticastDelegateRegistration<void(FName, bool)>& OnProcessLoadedObjectsCallback()
 	{
 		return ProcessLoadedObjectsCallback;
 	}
@@ -683,10 +683,10 @@ private:
 
 	/** Multicast delegate that will broadcast a notification when modules are loaded, unloaded, or
 		our set of known modules changes */
-	FModulesChangedEvent ModulesChangedEvent;
+	TTSMulticastDelegate<void(FName ModuleName, EModuleChangeReason ChangeReason)> ModulesChangedEvent;
 	
 	/** Multicast delegate called to process any new loaded objects. */
-	ProcessLoadedObjectsEvent ProcessLoadedObjectsCallback;
+	TMulticastDelegate<void(FName, bool)> ProcessLoadedObjectsCallback;
 
 	/** When module manager is linked against an application that supports UObjects, this delegate will be primed
 		at startup to provide information about whether a UObject package is loaded into memory. */
diff --git a/Engine/Source/Runtime/Engine/Classes/Engine/World.h b/Engine/Source/Runtime/Engine/Classes/Engine/World.h
--- a/Engine/Source/Runtime/Engine/Classes/Engine/World.h
+++ b/Engine/Source/Runtime/Engine/Classes/Engine/World.h
@@ -1530,48 +1530,43 @@ private:
 	/** Timestamp (in FPlatformTime::Seconds) when the next call to BuildStreamingData() should be made, if bDirtyStreamingData is true. */
 	double BuildStreamingDataTimer;
 
-	DECLARE_EVENT_OneParam(UWorld, FOnNetTickEvent, float);
-	DECLARE_EVENT(UWorld, FOnTickFlushEvent);
+	using FOnNetTickEvent UE_DEPRECATED(5.5, "The FOnNetTickEvent typedef has been deprecated - use TMulticastDelegateRegistration<void(float)> instead.") = TMulticastDelegate<void(float)>;
+	using FOnTickFlushEvent UE_DEPRECATED(5.5, "The FOnTickFlushEvent typedef has been deprecated - use TMulticastDelegateRegistration<void()> instead.") = TMulticastDelegate<void()>;
+
 	/** Event to gather up all net drivers and call TickDispatch at once */
-	FOnNetTickEvent TickDispatchEvent;
+	UE_DEPRECATED_FORGAME(5.5, "Public access to TickDispatchEvent has been deprecated - use OnTickDispatch() instead.")
+	TMulticastDelegate<void(float)> TickDispatchEvent;
 
 	/** Event to gather up all net drivers and call PostTickDispatch at once */
-	FOnTickFlushEvent PostTickDispatchEvent;
+	UE_DEPRECATED_FORGAME(5.5, "Public access to PostTickDispatchEvent has been deprecated - use OnPostTickDispatch() instead.")
+	TMulticastDelegate<void()> PostTickDispatchEvent;
 
 	/** Event called prior to calling TickFlush */
-	FOnNetTickEvent PreTickFlushEvent;
+	UE_DEPRECATED_FORGAME(5.5, "Public access to PreTickFlushEvent has been deprecated - use OnPreTickFlush() instead.")
+	TMulticastDelegate<void(float)> PreTickFlushEvent;
 
 	/** Event to gather up all net drivers and call TickFlush at once */
-	FOnNetTickEvent TickFlushEvent;
+	UE_DEPRECATED_FORGAME(5.5, "Public access to TickFlushEvent has been deprecated - use OnTickFlush() instead.")
+	TMulticastDelegate<void(float)> TickFlushEvent;
 	
 	/** Event to gather up all net drivers and call PostTickFlush at once */
-	FOnTickFlushEvent PostTickFlushEvent;
+	UE_DEPRECATED_FORGAME(5.5, "Public access to PostTickFlushEvent has been deprecated - use OnPostTickFlush() instead.")
+	TMulticastDelegate<void()> PostTickFlushEvent;
 
 	/** All registered net drivers TickDispatch() */
-	void BroadcastTickDispatch(float DeltaTime)	
-	{
-		TickDispatchEvent.Broadcast(DeltaTime);
-	}
+	void BroadcastTickDispatch(float DeltaTime);
+
 	/** All registered net drivers PostTickDispatch() */
-	void BroadcastPostTickDispatch()
-	{
-		PostTickDispatchEvent.Broadcast();
-	}
+	void BroadcastPostTickDispatch();
+
 	/** PreTickFlush */
-	void BroadcastPreTickFlush(float DeltaTime)
-	{
-		PreTickFlushEvent.Broadcast(DeltaTime);
-	}
+	void BroadcastPreTickFlush(float DeltaTime);
+
 	/** All registered net drivers TickFlush() */
-	void BroadcastTickFlush(float DeltaTime)
-	{
-		TickFlushEvent.Broadcast(DeltaTime);
-	}
+	void BroadcastTickFlush(float DeltaTime);
+
 	/** All registered net drivers PostTickFlush() */
-	void BroadcastPostTickFlush(float DeltaTime)
-	{
-		PostTickFlushEvent.Broadcast();
-	}
+	void BroadcastPostTickFlush(float DeltaTime);
 
 	/** Called when the number of levels changes. */
 	DECLARE_EVENT(UWorld, FOnLevelsChangedEvent);
@@ -3035,11 +3030,11 @@ public:
 public:
 
 	/** Network Tick events */
-	FOnNetTickEvent& OnTickDispatch() { return TickDispatchEvent; }
-	FOnTickFlushEvent& OnPostTickDispatch() { return PostTickDispatchEvent; }	
-	FOnNetTickEvent& OnPreTickFlush() { return PreTickFlushEvent; }
-	FOnNetTickEvent& OnTickFlush() { return TickFlushEvent; }
-	FOnTickFlushEvent& OnPostTickFlush() { return PostTickFlushEvent; }
+	TMulticastDelegateRegistration<void(float)>& OnTickDispatch();
+	TMulticastDelegateRegistration<void()>& OnPostTickDispatch();
+	TMulticastDelegateRegistration<void(float)>& OnPreTickFlush();
+	TMulticastDelegateRegistration<void(float)>& OnTickFlush();
+	TMulticastDelegateRegistration<void()>& OnPostTickFlush();
 
 	/**
 	 * Update the level after a variable amount of time, DeltaSeconds, has passed.
diff --git a/Engine/Source/Runtime/Engine/Private/World.cpp b/Engine/Source/Runtime/Engine/Private/World.cpp
--- a/Engine/Source/Runtime/Engine/Private/World.cpp
+++ b/Engine/Source/Runtime/Engine/Private/World.cpp
@@ -9520,6 +9520,56 @@ FAudioDevice* UWorld::GetAudioDeviceRaw() const
 	}
 }
 
+TMulticastDelegateRegistration<void(float)>& UWorld::OnTickDispatch()
+{
+	return TickDispatchEvent;
+}
+
+TMulticastDelegateRegistration<void()>& UWorld::OnPostTickDispatch()
+{
+	return PostTickDispatchEvent;
+}
+
+TMulticastDelegateRegistration<void(float)>& UWorld::OnPreTickFlush()
+{
+	return PreTickFlushEvent;
+}
+
+TMulticastDelegateRegistration<void(float)>& UWorld::OnTickFlush()
+{
+	return TickFlushEvent;
+}
+
+TMulticastDelegateRegistration<void()>& UWorld::OnPostTickFlush()
+{
+	return PostTickFlushEvent;
+}
+
+void UWorld::BroadcastTickDispatch(float DeltaTime)	
+{
+	TickDispatchEvent.Broadcast(DeltaTime);
+}
+
+void UWorld::BroadcastPostTickDispatch()
+{
+	PostTickDispatchEvent.Broadcast();
+}
+
+void UWorld::BroadcastPreTickFlush(float DeltaTime)
+{
+	PreTickFlushEvent.Broadcast(DeltaTime);
+}
+
+void UWorld::BroadcastTickFlush(float DeltaTime)
+{
+	TickFlushEvent.Broadcast(DeltaTime);
+}
+
+void UWorld::BroadcastPostTickFlush(float DeltaTime)
+{
+	PostTickFlushEvent.Broadcast();
+}
+
 /**
 * Dump visible actors in current world.
 */
