[Scheduler]
 - Add a test for the potential deadlock recently fixed

diff --git a/Engine/Source/Runtime/Core/Tests/Tasks/TasksTest.cpp b/Engine/Source/Runtime/Core/Tests/Tasks/TasksTest.cpp
--- a/Engine/Source/Runtime/Core/Tests/Tasks/TasksTest.cpp
+++ b/Engine/Source/Runtime/Core/Tests/Tasks/TasksTest.cpp
@@ -1665,6 +1665,76 @@ namespace UE { namespace TasksTests
 			Benchmark<5>(*Name, [&Index]() { BenchmarkBlockUnblockWorkers(Index); });
 		}
 	}
+
+	TEST_CASE_NAMED(FTasksLoneStandbyWorker, "System::Core::Async::Tasks::LoneStandbyWorker", "[.][ApplicationContextMask][EngineFilter]")
+	{
+		// We absolutely need oversubscription to kick in to test this.
+		// So only use a single worker to make sure that happens.
+		LowLevelTasks::FScheduler::Get().RestartWorkers(1, 0);
+
+		UE::FManualResetEvent OversubscribeeReadyEvent;
+		UE::FManualResetEvent OversubscriberReadyEvent;
+		UE::FManualResetEvent OversubscribeeDoneEvent;
+		UE::FManualResetEvent OversubscriberDoneEvent;
+		UE::FManualResetEvent LocalQueueEvent;
+
+		UE::Tasks::FTask Oversubscriber =
+			UE::Tasks::Launch(
+				TEXT("Oversubscriber"),
+				[&]()
+				{
+					LowLevelTasks::FOversubscriptionScope _;
+					OversubscriberReadyEvent.Notify();
+					OversubscriberDoneEvent.Wait();
+				}
+		);
+
+		// Wait until the oversubscription scope is active
+		OversubscriberReadyEvent.Wait();
+
+		UE::Tasks::FTask Oversubscribee = UE::Tasks::Launch(
+			TEXT("Oversubscribee"),
+			[&]()
+			{
+				OversubscribeeReadyEvent.Notify();
+				OversubscribeeDoneEvent.Wait();
+			});
+
+		// The first subsequent of a task is sent to the local queue
+		// so setup ourself to be the subsequent of the oversubscribee.
+		UE::Tasks::Launch(
+			TEXT("LocalQueueTask"),
+			[&]()
+			{
+				LocalQueueEvent.Notify();
+			},
+			UE::Tasks::Prerequisites(Oversubscribee)
+		);
+
+		// Wait until the oversubscribee task is launched.
+		OversubscribeeReadyEvent.Wait();
+		// Now close the oversubscription scope while the oversubcribee is still executing.
+		OversubscriberDoneEvent.Notify();
+		// Wait until the oversubscriber has closed it's oversubscription scope.
+		Oversubscriber.Wait();
+		// Now release the oversubscribee so it finishes executing and release its subsequent.
+		// The first subsequent will be sent to the local queue of the standby thread.
+		OversubscribeeDoneEvent.Notify();
+		// Now we're left with only a standby thread executing. If that standby thread
+		// enqueues the first subsequent to its local queue without waking up another worker
+		// and goes immediately to sleep after its execution because the oversubscription
+		// period is over, we could deadlock. The fix for this was to always enqueue
+		// to the global queue for standby workers and to always perform a wakeup
+		// so that normal workers can pick up the work if the standby worker can't.
+		// Enqueuing to the global queue is not stricly required as just performing
+		// a wake up would be sufficient, but the other thread would have to perform
+		// stealing on all threads to find our task, which is less efficient.
+
+		// Verify that we did not timeout (i.e. deadlock).
+		verify(LocalQueueEvent.WaitFor(UE::FMonotonicTimeSpan::FromSeconds(5)));
+
+		LowLevelTasks::FScheduler::Get().RestartWorkers();
+	}
 }}
 
 #endif // WITH_TESTS
