[BinnedMalloc2/3]
 - Improve thread caches memory trimming behavior

diff --git a/Engine/Source/Runtime/Core/Private/Async/Fundamental/WaitingQueue.cpp b/Engine/Source/Runtime/Core/Private/Async/Fundamental/WaitingQueue.cpp
--- a/Engine/Source/Runtime/Core/Private/Async/Fundamental/WaitingQueue.cpp
+++ b/Engine/Source/Runtime/Core/Private/Async/Fundamental/WaitingQueue.cpp
@@ -103,6 +103,27 @@ namespace LowLevelTasks::Private
 			// The StandbyState stores the active thread count in the waiter bits.
 			return (StandbyState & WaiterMask) >> WaiterShift;
 		}
+
+		void EnterWait(FWaitEvent* Node)
+		{
+			// Flush any open scope before going to sleep so that anything that happened
+			// before appears in UnrealInsights right away. If we don't do this,
+			// the thread buffer will be held to this thread until we wake up and fill it
+			// so it might cause events to appear as missing in UnrealInsights, especially
+			// in case we never wake up again (i.e. deadlock / crash).
+			TRACE_CPUPROFILER_EVENT_FLUSH();
+			Private::FOversubscriptionAllowedScope _(false /* Disallow oversubscription for this wait */);
+
+			// Let the memory manager know we're inactive so it can do whatever it wants with our
+			// thread-local memory cache if we have any.
+			FMemory::MarkTLSCachesAsUnusedOnCurrentThread();
+
+			Node->Event->Wait();
+
+			// Let the memory manager know we're active again and need our
+			// thread-local memory cache back if we have any.
+			FMemory::MarkTLSCachesAsUsedOnCurrentThread();
+		}
 	}
 
 void FWaitingQueue::Init(uint32 InThreadCount, uint32 InMaxThreadCount, TFunction<void()> InCreateThread, uint32 InActiveThreadCount)
@@ -334,8 +355,7 @@ void FWaitingQueue::ConditionalStandby(FWaitEvent* Node)
 		if (StandbyState.compare_exchange_weak(LocalState, NewState))
 		{
 			WAITINGQUEUE_EVENT_SCOPE(Standby);
-			Private::FOversubscriptionAllowedScope _(false /* Disallow oversubscription for this wait */);
-			Node->Event->Wait();
+			EnterWait(Node);
 		}
 		else
 		{
@@ -376,14 +396,7 @@ bool FWaitingQueue::CommitStandby(FWaitEvent* Node, FOutOfWork& OutOfWork)
 	}
 
 	OutOfWork.Stop();
-	// Flush any open scope before going to sleep so that anything that happened
-	// before appears in UnrealInsights right away. If we don't do this,
-	// the thread buffer will be held to this thread until we wake up and fill it
-	// so it might cause events to appear as missing in UnrealInsights, especially
-	// in case we never wake up again (i.e. deadlock / crash).
-	TRACE_CPUPROFILER_EVENT_FLUSH();
-	Private::FOversubscriptionAllowedScope _(false /* Disallow oversubscription for this wait */);
-	Node->Event->Wait();
+	EnterWait(Node);
 	return true;
 }
 
@@ -599,14 +612,7 @@ void FWaitingQueue::Park(FWaitEvent* Node, FOutOfWork& OutOfWork, int32 SpinCycl
 		}
 	}
 
-	// Flush any open scope before going to sleep so that anything that happened
-	// before appears in UnrealInsights right away. If we don't do this,
-	// the thread buffer will be held to this thread until we wake up and fill it
-	// so it might cause events to appear as missing in UnrealInsights, especially
-	// in case we never wake up again (i.e. deadlock / crash).
-	TRACE_CPUPROFILER_EVENT_FLUSH();
-	Private::FOversubscriptionAllowedScope _(false /* Disallow oversubscription for this wait */);
-	Node->Event->Wait();
+	EnterWait(Node);
 }
 
 int32 FWaitingQueue::Unpark(FWaitEvent* Node)
diff --git a/Engine/Source/Runtime/Core/Private/HAL/MallocBinned2.cpp b/Engine/Source/Runtime/Core/Private/HAL/MallocBinned2.cpp
--- a/Engine/Source/Runtime/Core/Private/HAL/MallocBinned2.cpp
+++ b/Engine/Source/Runtime/Core/Private/HAL/MallocBinned2.cpp
@@ -1203,72 +1203,35 @@ bool FMallocBinned2::ValidateHeap()
 
 const TCHAR* FMallocBinned2::GetDescriptiveName()
 {
-	return TEXT("binned2");
+	return TEXT("Binned2");
 }
 
-void FMallocBinned2::FlushCurrentThreadCache()
+void FMallocBinned2::FreeBundles(FBundleNode* Bundles, uint32 PoolIndex)
 {
-	TRACE_CPUPROFILER_EVENT_SCOPE(FMallocBinned2::FlushCurrentThreadCache);
-	NOALLOC_SCOPE_CYCLE_COUNTER(STAT_FMallocBinned2_FlushCurrentThreadCache);
-	FPerThreadFreeBlockLists* Lists = FPerThreadFreeBlockLists::Get();
-
-	double WaitForMutexTime = 0.0;
-	double WaitForMutexAndTrimTime = 0.0;
+	Private::FreeBundles(*this, Bundles, PoolIndexToBlockSize(PoolIndex), PoolIndex);
+}
 
-	if (Lists)
-	{
-		const double StartTimeInner = FPlatformTime::Seconds();
-		FScopeLock Lock(&Mutex);
-		WaitForMutexTime = FPlatformTime::Seconds() - StartTimeInner;
-		for (int32 PoolIndex = 0; PoolIndex != BINNED2_SMALL_POOL_COUNT; ++PoolIndex)
-		{
-			FBundleNode* Bundles = Lists->PopBundles(PoolIndex);
-			if (Bundles)
-			{
-				Private::FreeBundles(*this, Bundles, PoolIndexToBlockSize(PoolIndex), PoolIndex);
-			}
-		}
-		WaitForMutexAndTrimTime = FPlatformTime::Seconds() - StartTimeInner;
-	}
+float FMallocBinned2::GetFlushThreadCacheMaxWaitTime()
+{
+	return GMallocBinned2FlushThreadCacheMaxWaitTime;
+}
 
-	// These logs must happen outside the above mutex to avoid deadlocks
-	if (WaitForMutexTime > GMallocBinned2FlushThreadCacheMaxWaitTime)
-	{
-		UE_LOG(LogMemory, Warning, TEXT("FMallocBinned2 took %6.2fms to wait for mutex for trim."), WaitForMutexTime * 1000.0f);
-	}
-	if (WaitForMutexAndTrimTime > GMallocBinned2FlushThreadCacheMaxWaitTime)
-	{
-		UE_LOG(LogMemory, Warning, TEXT("FMallocBinned2 took %6.2fms to wait for mutex AND trim."), WaitForMutexAndTrimTime * 1000.0f);
-	}
+FCriticalSection& FMallocBinned2::GetFreeBlockListsRegistrationMutex()
+{
+	return Private::GetFreeBlockListsRegistrationMutex();
 }
 
-#include "Async/TaskGraphInterfaces.h"
+TArray<FMallocBinned2::FPerThreadFreeBlockLists*>& FMallocBinned2::GetRegisteredFreeBlockLists()
+{
+	return Private::GetRegisteredFreeBlockLists();
+}
 
 void FMallocBinned2::Trim(bool bTrimThreadCaches)
 {
-	NOALLOC_SCOPE_CYCLE_COUNTER(STAT_FMallocBinned2_Trim);
-
-	if (GMallocBinned2PerThreadCaches  &&  bTrimThreadCaches)
-	{
-		//double StartTime = FPlatformTime::Seconds();
-		TFunction<void(ENamedThreads::Type CurrentThread)> Broadcast =
-			[this](ENamedThreads::Type MyThread)
-		{
-			FlushCurrentThreadCache();
-		};
-		// Skip task threads on desktop platforms as it is too slow and they don't have much memory
-		if (PLATFORM_DESKTOP)
-		{
-			FTaskGraphInterface::BroadcastSlow_OnlyUseForSpecialPurposes(false, false, Broadcast);
-		}
-		else
-		{
-			FTaskGraphInterface::BroadcastSlow_OnlyUseForSpecialPurposes(FPlatformProcess::SupportsMultithreading() && FApp::ShouldUseThreadingForPerformance(), false, Broadcast);
-		}
-		//UE_LOG(LogTemp, Display, TEXT("Trim Broadcast = %6.2fms"), 1000.0f * float(FPlatformTime::Seconds() - StartTime));
-	}
+	if (GMallocBinned2PerThreadCaches && bTrimThreadCaches)
 	{
-		//double StartTime = FPlatformTime::Seconds();
+		TMallocBinnedCommon::TrimImpl(*this);
+
 #if !UE_USE_VERYLARGEPAGEALLOCATOR
 		FScopeLock Lock(&Mutex);
 		// this cache is recycled anyway, if you need to trim it based on being OOM, it's already too late.
@@ -1297,10 +1260,27 @@ void FMallocBinned2::SetupTLSCachesOnCurrentThread()
 void FMallocBinned2::ClearAndDisableTLSCachesOnCurrentThread()
 {
 	NOALLOC_SCOPE_CYCLE_COUNTER(STAT_FMallocBinned2_ClearTLSCachesOnCurrentThread);
-	FlushCurrentThreadCache();
+	FlushCurrentThreadCache(*this);
 	FPerThreadFreeBlockLists::ClearTLS();
 }
 
+void FMallocBinned2::MarkTLSCachesAsUsedOnCurrentThread()
+{
+	NOALLOC_SCOPE_CYCLE_COUNTER(STAT_FMallocBinned2_MarkTLSCachesAsUsedOnCurrentThread);
+	FPerThreadFreeBlockLists::LockTLS();
+}
+
+void FMallocBinned2::MarkTLSCachesAsUnusedOnCurrentThread()
+{
+	NOALLOC_SCOPE_CYCLE_COUNTER(STAT_FMallocBinned2_MarkTLSCachesAsUnusedOnCurrentThread);
+
+	// Will only flush if memory trimming epoch has been bumped while the thread was active.
+	const bool bNewEpochOnly = true;
+	FlushCurrentThreadCache(*this, bNewEpochOnly);
+
+	FPerThreadFreeBlockLists::UnlockTLS();
+}
+
 void FMallocBinned2::CanaryTest(const FFreeBlock* Block) const
 {
 #if BINNED2_FORK_SUPPORT
diff --git a/Engine/Source/Runtime/Core/Private/HAL/MallocBinned3.cpp b/Engine/Source/Runtime/Core/Private/HAL/MallocBinned3.cpp
--- a/Engine/Source/Runtime/Core/Private/HAL/MallocBinned3.cpp
+++ b/Engine/Source/Runtime/Core/Private/HAL/MallocBinned3.cpp
@@ -611,6 +611,11 @@ struct FMallocBinned3::Private
 
 FMallocBinned3::Private::FGlobalRecycler FMallocBinned3::Private::GGlobalRecycler;
 
+void FMallocBinned3::FreeBundles(FBundleNode* Bundles, uint32 PoolIndex)
+{
+	Private::FreeBundles(*this, Bundles, PoolIndexToBlockSize(PoolIndex), PoolIndex);
+}
+
 void FMallocBinned3::RegisterThreadFreeBlockLists(FPerThreadFreeBlockLists* FreeBlockLists)
 {
 	Private::RegisterThreadFreeBlockLists(FreeBlockLists);
@@ -1290,60 +1295,27 @@ const TCHAR* FMallocBinned3::GetDescriptiveName()
 	return TEXT("Binned3");
 }
 
-void FMallocBinned3::FlushCurrentThreadCache()
+void FMallocBinned3::Trim(bool bTrimThreadCaches)
 {
-	double StartTimeInner = FPlatformTime::Seconds();
-	TRACE_CPUPROFILER_EVENT_SCOPE(FMallocBinned3::FlushCurrentThreadCache);
-	QUICK_SCOPE_CYCLE_COUNTER(STAT_FMallocBinned3_FlushCurrentThreadCache);
-	FPerThreadFreeBlockLists* Lists = FPerThreadFreeBlockLists::Get();
-
-	float WaitForMutexTime = 0.0f;
-	float WaitForMutexAndTrimTime = 0.0f;
-
-	if (Lists)
+	if (GMallocBinned3PerThreadCaches && bTrimThreadCaches)
 	{
-		FScopeLock Lock(&Mutex);
-		WaitForMutexTime = FPlatformTime::Seconds() - StartTimeInner;
-		for (int32 PoolIndex = 0; PoolIndex != BINNED3_SMALL_POOL_COUNT; ++PoolIndex)
-		{
-			FBundleNode* Bundles = Lists->PopBundles(PoolIndex);
-			if (Bundles)
-			{
-				Private::FreeBundles(*this, Bundles, PoolIndexToBlockSize(PoolIndex), PoolIndex);
-			}
-		}
-		WaitForMutexAndTrimTime = FPlatformTime::Seconds() - StartTimeInner;
-	}
-
-	// These logs must happen outside the above mutex to avoid deadlocks
-	if (WaitForMutexTime > GMallocBinned3FlushThreadCacheMaxWaitTime)
-	{
-		UE_LOG(LogMemory, Warning, TEXT("FMallocBinned3 took %6.2fms to wait for mutex for trim."), WaitForMutexTime * 1000.0f);
-	}
-	if (WaitForMutexAndTrimTime > GMallocBinned3FlushThreadCacheMaxWaitTime)
-	{
-		UE_LOG(LogMemory, Warning, TEXT("FMallocBinned3 took %6.2fms to wait for mutex AND trim."), WaitForMutexAndTrimTime * 1000.0f);
+		TMallocBinnedCommon::TrimImpl(*this);
 	}
 }
 
-#include "Async/TaskGraphInterfaces.h"
+float FMallocBinned3::GetFlushThreadCacheMaxWaitTime()
+{
+	return GMallocBinned3FlushThreadCacheMaxWaitTime;
+}
 
-void FMallocBinned3::Trim(bool bTrimThreadCaches)
+FCriticalSection& FMallocBinned3::GetFreeBlockListsRegistrationMutex()
 {
+	return Private::GetFreeBlockListsRegistrationMutex();
+}
 
-	if (GMallocBinned3PerThreadCaches  &&  bTrimThreadCaches)
-	{
-		QUICK_SCOPE_CYCLE_COUNTER(STAT_FMallocBinned3_Trim);
-		//double StartTime = FPlatformTime::Seconds();
-		TFunction<void(ENamedThreads::Type CurrentThread)> Broadcast =
-			[this](ENamedThreads::Type MyThread)
-		{
-			FlushCurrentThreadCache();
-		};
-		// Skip task threads on desktop platforms as it is too slow and they don't have much memory
-		FTaskGraphInterface::BroadcastSlow_OnlyUseForSpecialPurposes(!PLATFORM_DESKTOP, false, Broadcast);
-		//UE_LOG(LogTemp, Display, TEXT("Trim Broadcast = %6.2fms"), 1000.0f * float(FPlatformTime::Seconds() - StartTime));
-	}
+TArray<FMallocBinned3::FPerThreadFreeBlockLists*>& FMallocBinned3::GetRegisteredFreeBlockLists()
+{
+	return Private::GetRegisteredFreeBlockLists();
 }
 
 void FMallocBinned3::SetupTLSCachesOnCurrentThread()
@@ -1366,10 +1338,33 @@ void FMallocBinned3::ClearAndDisableTLSCachesOnCurrentThread()
 	{
 		return;
 	}
-	FlushCurrentThreadCache();
+	FlushCurrentThreadCache(*this);
 	FPerThreadFreeBlockLists::ClearTLS();
 }
 
+void FMallocBinned3::MarkTLSCachesAsUsedOnCurrentThread()
+{
+	if (!BINNED3_ALLOW_RUNTIME_TWEAKING && !GMallocBinned3PerThreadCaches)
+	{
+		return;
+	}
+
+	FPerThreadFreeBlockLists::LockTLS();
+}
+
+void FMallocBinned3::MarkTLSCachesAsUnusedOnCurrentThread()
+{
+	if (!BINNED3_ALLOW_RUNTIME_TWEAKING && !GMallocBinned3PerThreadCaches)
+	{
+		return;
+	}
+
+	// Will only flush if memory trimming epoch has been bumped while the thread was active.
+	const bool bNewEpochOnly = true;
+	FlushCurrentThreadCache(*this, bNewEpochOnly);
+	FPerThreadFreeBlockLists::UnlockTLS();
+}
+
 void FMallocBinned3::FFreeBlock::CanaryFail() const
 {
 	UE_LOG(LogMemory, Fatal, TEXT("FMallocBinned3 Attempt to realloc an unrecognized block %p   canary == 0x%x != 0x%x"), (void*)this, (int32)Canary, (int32)FMallocBinned3::FFreeBlock::CANARY_VALUE);
diff --git a/Engine/Source/Runtime/Core/Private/HAL/MallocLeakDetectionProxy.h b/Engine/Source/Runtime/Core/Private/HAL/MallocLeakDetectionProxy.h
--- a/Engine/Source/Runtime/Core/Private/HAL/MallocLeakDetectionProxy.h
+++ b/Engine/Source/Runtime/Core/Private/HAL/MallocLeakDetectionProxy.h
@@ -112,6 +112,17 @@ public:
 	{
 		return UsedMalloc->SetupTLSCachesOnCurrentThread();
 	}
+
+	virtual void MarkTLSCachesAsUsedOnCurrentThread() override
+	{
+		UsedMalloc->MarkTLSCachesAsUsedOnCurrentThread();
+	}
+
+	virtual void MarkTLSCachesAsUnusedOnCurrentThread() override
+	{
+		UsedMalloc->MarkTLSCachesAsUnusedOnCurrentThread();
+	}
+
 	virtual void ClearAndDisableTLSCachesOnCurrentThread() override
 	{
 		return UsedMalloc->ClearAndDisableTLSCachesOnCurrentThread();
diff --git a/Engine/Source/Runtime/Core/Private/HAL/MallocVerify.h b/Engine/Source/Runtime/Core/Private/HAL/MallocVerify.h
--- a/Engine/Source/Runtime/Core/Private/HAL/MallocVerify.h
+++ b/Engine/Source/Runtime/Core/Private/HAL/MallocVerify.h
@@ -125,6 +125,15 @@ public:
 	{
 		return UsedMalloc->SetupTLSCachesOnCurrentThread();
 	}
+	virtual void MarkTLSCachesAsUsedOnCurrentThread() override
+	{
+		UsedMalloc->MarkTLSCachesAsUsedOnCurrentThread();
+	}
+
+	virtual void MarkTLSCachesAsUnusedOnCurrentThread() override
+	{
+		UsedMalloc->MarkTLSCachesAsUnusedOnCurrentThread();
+	}
 	virtual void ClearAndDisableTLSCachesOnCurrentThread() override
 	{
 		return UsedMalloc->ClearAndDisableTLSCachesOnCurrentThread();
diff --git a/Engine/Source/Runtime/Core/Private/HAL/UnrealMemory.cpp b/Engine/Source/Runtime/Core/Private/HAL/UnrealMemory.cpp
--- a/Engine/Source/Runtime/Core/Private/HAL/UnrealMemory.cpp
+++ b/Engine/Source/Runtime/Core/Private/HAL/UnrealMemory.cpp
@@ -206,6 +206,14 @@ public:
 	{
 		return UsedMalloc->SetupTLSCachesOnCurrentThread();
 	}
+	virtual void MarkTLSCachesAsUsedOnCurrentThread() override
+	{
+		UsedMalloc->MarkTLSCachesAsUsedOnCurrentThread();
+	}
+	virtual void MarkTLSCachesAsUnusedOnCurrentThread() override
+	{
+		UsedMalloc->MarkTLSCachesAsUnusedOnCurrentThread();
+	}
 	virtual void ClearAndDisableTLSCachesOnCurrentThread() override
 	{
 		return UsedMalloc->ClearAndDisableTLSCachesOnCurrentThread();
@@ -569,6 +577,22 @@ void FMemory::ClearAndDisableTLSCachesOnCurrentThread()
 	}
 }
 
+void FMemory::MarkTLSCachesAsUsedOnCurrentThread()
+{
+	if (GMalloc)
+	{
+		GMalloc->MarkTLSCachesAsUsedOnCurrentThread();
+	}
+}
+
+void FMemory::MarkTLSCachesAsUnusedOnCurrentThread()
+{
+	if (GMalloc)
+	{
+		GMalloc->MarkTLSCachesAsUnusedOnCurrentThread();
+	}
+}
+
 void FMemory::TestMemory()
 {
 #if !UE_BUILD_SHIPPING
diff --git a/Engine/Source/Runtime/Core/Private/ProfilingDebugging/MemoryTrace.cpp b/Engine/Source/Runtime/Core/Private/ProfilingDebugging/MemoryTrace.cpp
--- a/Engine/Source/Runtime/Core/Private/ProfilingDebugging/MemoryTrace.cpp
+++ b/Engine/Source/Runtime/Core/Private/ProfilingDebugging/MemoryTrace.cpp
@@ -191,6 +191,8 @@ private:
 	virtual bool ValidateHeap() override                                    { return InnerMalloc->ValidateHeap(); }
 	virtual bool GetAllocationSize(void* Address, SIZE_T &SizeOut) override { return InnerMalloc->GetAllocationSize(Address, SizeOut); }
 	virtual void SetupTLSCachesOnCurrentThread() override                   { return InnerMalloc->SetupTLSCachesOnCurrentThread(); }
+	virtual void MarkTLSCachesAsUsedOnCurrentThread() override              { InnerMalloc->MarkTLSCachesAsUsedOnCurrentThread(); }
+	virtual void MarkTLSCachesAsUnusedOnCurrentThread() override            { InnerMalloc->MarkTLSCachesAsUnusedOnCurrentThread(); }
 	virtual void OnMallocInitialized() override                             { InnerMalloc->OnMallocInitialized(); }
 	virtual void OnPreFork() override                                       { InnerMalloc->OnPreFork(); }
 	virtual void OnPostFork() override                                      { InnerMalloc->OnPostFork(); }
diff --git a/Engine/Source/Runtime/Core/Public/HAL/MallocBinned2.h b/Engine/Source/Runtime/Core/Public/HAL/MallocBinned2.h
--- a/Engine/Source/Runtime/Core/Public/HAL/MallocBinned2.h
+++ b/Engine/Source/Runtime/Core/Public/HAL/MallocBinned2.h
@@ -494,6 +494,8 @@ public:
 	CORE_API virtual bool ValidateHeap() override;
 	CORE_API virtual void Trim(bool bTrimThreadCaches) override;
 	CORE_API virtual void SetupTLSCachesOnCurrentThread() override;
+	CORE_API virtual void MarkTLSCachesAsUsedOnCurrentThread() override;
+	CORE_API virtual void MarkTLSCachesAsUnusedOnCurrentThread() override;
 	CORE_API virtual void ClearAndDisableTLSCachesOnCurrentThread() override;
 	CORE_API virtual const TCHAR* GetDescriptiveName() override;
 	CORE_API virtual void UpdateStats() override;
@@ -502,13 +504,12 @@ public:
 	CORE_API virtual void OnPostFork() override;
 	// End FMalloc interface.
 
-	CORE_API void FlushCurrentThreadCache();
 	CORE_API void* MallocExternalSmall(SIZE_T Size, uint32 Alignment);
 	CORE_API void* MallocExternalLarge(SIZE_T Size, uint32 Alignment);
 	CORE_API void* ReallocExternal(void* Ptr, SIZE_T NewSize, uint32 Alignment);
 	CORE_API void FreeExternal(void *Ptr);
 	CORE_API bool GetAllocationSizeExternal(void* Ptr, SIZE_T& SizeOut);
-
+	
 	CORE_API void CanaryTest(const FFreeBlock* Block) const;
 	CORE_API void CanaryFail(const FFreeBlock* Block) const;
 
@@ -528,6 +529,9 @@ public:
 
 	static void RegisterThreadFreeBlockLists(FPerThreadFreeBlockLists* FreeBlockLists);
 	static void UnregisterThreadFreeBlockLists(FPerThreadFreeBlockLists* FreeBlockLists);
+	static float GetFlushThreadCacheMaxWaitTime();
+	static FCriticalSection& GetFreeBlockListsRegistrationMutex();
+	static TArray<FPerThreadFreeBlockLists*>& GetRegisteredFreeBlockLists();
 
 	static void* AllocateMetaDataMemory(SIZE_T Size);
 	static void FreeMetaDataMemory(void* Ptr, SIZE_T Size);
@@ -536,6 +540,9 @@ public:
 	{
 		return SmallBlockSizesReversed[BINNED2_SMALL_POOL_COUNT - PoolIndex - 1];
 	}
+
+	void FreeBundles(FBundleNode* Bundles, uint32 PoolIndex);
+	FCriticalSection& GetMutex() { return Mutex; }
 };
 
 #define BINNED2_INLINE (1)
diff --git a/Engine/Source/Runtime/Core/Public/HAL/MallocBinned3.h b/Engine/Source/Runtime/Core/Public/HAL/MallocBinned3.h
--- a/Engine/Source/Runtime/Core/Public/HAL/MallocBinned3.h
+++ b/Engine/Source/Runtime/Core/Public/HAL/MallocBinned3.h
@@ -446,11 +446,12 @@ public:
 	virtual bool ValidateHeap() override;
 	virtual void Trim(bool bTrimThreadCaches) override;
 	virtual void SetupTLSCachesOnCurrentThread() override;
+	virtual void MarkTLSCachesAsUsedOnCurrentThread() override;
+	virtual void MarkTLSCachesAsUnusedOnCurrentThread() override;
 	virtual void ClearAndDisableTLSCachesOnCurrentThread() override;
 	virtual const TCHAR* GetDescriptiveName() override;
 	// End FMalloc interface.
 
-	void FlushCurrentThreadCache();
 	void* MallocExternal(SIZE_T Size, uint32 Alignment);
 	void* ReallocExternal(void* Ptr, SIZE_T NewSize, uint32 Alignment);
 	void FreeExternal(void *Ptr);
@@ -471,6 +472,9 @@ public:
 
 	static void RegisterThreadFreeBlockLists(FPerThreadFreeBlockLists* FreeBlockLists);
 	static void UnregisterThreadFreeBlockLists(FPerThreadFreeBlockLists* FreeBlockLists);
+	static float GetFlushThreadCacheMaxWaitTime();
+	static FCriticalSection& GetFreeBlockListsRegistrationMutex();
+	static TArray<FPerThreadFreeBlockLists*>& GetRegisteredFreeBlockLists();
 
 #if !BINNED3_USE_SEPARATE_VM_PER_POOL
 	static uint8* Binned3BaseVMPtr;
@@ -488,6 +492,9 @@ public:
 	{
 		return uint32(SmallBlockSizesReversedShifted[BINNED3_SMALL_POOL_COUNT - PoolIndex - 1]) << BINNED3_MINIMUM_ALIGNMENT_SHIFT;
 	}
+	
+	void FreeBundles(FBundleNode* Bundles, uint32 PoolIndex);
+	FCriticalSection& GetMutex() { return Mutex; }
 
 	void Commit(uint32 InPoolIndex, void *Ptr, SIZE_T Size);
 	void Decommit(uint32 InPoolIndex, void *Ptr, SIZE_T Size);
diff --git a/Engine/Source/Runtime/Core/Public/HAL/MallocBinnedCommon.h b/Engine/Source/Runtime/Core/Public/HAL/MallocBinnedCommon.h
--- a/Engine/Source/Runtime/Core/Public/HAL/MallocBinnedCommon.h
+++ b/Engine/Source/Runtime/Core/Public/HAL/MallocBinnedCommon.h
@@ -6,6 +6,8 @@
 #include "HAL/MemoryBase.h"
 #include "Math/UnrealMathUtility.h"
 #include "HAL/PlatformTLS.h"
+#include "Misc/App.h"
+#include "Async/TaskGraphInterfaces.h"
 #include "Templates/AlignmentTemplates.h"
 
 #if PLATFORM_HAS_FPlatformVirtualMemoryBlock
@@ -306,6 +308,7 @@ protected:
 	static std::atomic<int64> TLSMemory;
 	static std::atomic<int64> ConsolidatedMemory;
 #endif
+	std::atomic<uint64> MemoryTrimEpoch{ 0 };
 };
 
 template <class AllocType, int MinAlign, int MaxAlign, int MinAlignShift, int NumSmallPools, int MaxSmallPoolSize>
@@ -438,11 +441,30 @@ protected:
 				TLSMemory.fetch_add(TLSSize, std::memory_order_relaxed);
 #endif
 				verify(ThreadSingleton);
+				ThreadSingleton->Lock();
 				FPlatformTLS::SetTlsValue(BinnedTlsSlot, ThreadSingleton);
 				AllocType::RegisterThreadFreeBlockLists(ThreadSingleton);
 			}
 		}
 
+		static void UnlockTLS()
+		{
+			FPerThreadFreeBlockLists* ThreadSingleton = (FPerThreadFreeBlockLists*)FPlatformTLS::GetTlsValue(BinnedTlsSlot);
+			if (ThreadSingleton)
+			{
+				ThreadSingleton->Unlock();
+			}
+		}
+
+		static void LockTLS()
+		{
+			FPerThreadFreeBlockLists* ThreadSingleton = (FPerThreadFreeBlockLists*)FPlatformTLS::GetTlsValue(BinnedTlsSlot);
+			if (ThreadSingleton)
+			{
+				ThreadSingleton->Lock();
+			}
+		}
+
 		static void ClearTLS()
 		{
 			check(FPlatformTLS::IsValidTlsSlot(BinnedTlsSlot));
@@ -454,7 +476,7 @@ protected:
 				TLSMemory.fetch_sub(TLSSize, std::memory_order_relaxed);
 #endif
 				AllocType::UnregisterThreadFreeBlockLists(ThreadSingleton);
-
+				ThreadSingleton->Unlock();
 				ThreadSingleton->~FPerThreadFreeBlockLists();
 
 				AllocType::FreeMetaDataMemory(ThreadSingleton, TLSSize);
@@ -498,11 +520,40 @@ protected:
 			return FreeLists[InPoolIndex].PopBundles(InPoolIndex);
 		}
 
+		void Lock()
+		{
+			Mutex.Lock();
+		}
+
+		bool TryLock()
+		{
+			return Mutex.TryLock();
+		}
+
+		void Unlock()
+		{
+			Mutex.Unlock();
+		}
+
+		// should only be called from inside the Lock.
+		bool UpdateEpoch(uint64 NewEpoch)
+		{
+			if (MemoryTrimEpoch >= NewEpoch)
+			{
+				return false;
+			}
+
+			MemoryTrimEpoch = NewEpoch;
+			return true;
+		}
+
 #if UE_BINNEDCOMMON_ALLOCATOR_STATS
 	public:
 		int64 AllocatedMemory = 0;
 #endif
 	private:
+		UE::FMutex Mutex;
+		uint64 MemoryTrimEpoch = 0;
 		FFreeBlockList FreeLists[NumSmallPools];
 	};
 
@@ -578,4 +629,104 @@ protected:
 
 		return false;
 	}
+
+	static void TrimThreadFreeBlockLists(AllocType& Allocator, FPerThreadFreeBlockLists* FreeBlockLists)
+	{
+		if (FreeBlockLists)
+		{
+			TRACE_CPUPROFILER_EVENT_SCOPE(FMallocBinnedCommon::TrimThreadFreeBlockLists);
+
+			for (int32 PoolIndex = 0; PoolIndex != NumSmallPools; ++PoolIndex)
+			{
+				FBundleNode* Bundles = FreeBlockLists->PopBundles(PoolIndex);
+				if (Bundles)
+				{
+					Allocator.FreeBundles(Bundles, PoolIndex);
+				}
+			}
+		}
+	}
+
+	void FlushCurrentThreadCache(AllocType& Allocator, bool bNewEpochOnly = false)
+	{
+		if (FPerThreadFreeBlockLists* Lists = FPerThreadFreeBlockLists::Get())
+		{
+			if (Lists->UpdateEpoch(MemoryTrimEpoch.load(std::memory_order_relaxed)) || !bNewEpochOnly)
+			{
+				TRACE_CPUPROFILER_EVENT_SCOPE(FMallocBinnedCommon::FlushCurrentThreadCache);
+				QUICK_SCOPE_CYCLE_COUNTER(STAT_FMallocBinnedCommon_FlushCurrentThreadCache);
+
+				double StartTimeInner = FPlatformTime::Seconds();
+
+				double WaitForMutexTime = 0.0f;
+				double WaitForMutexAndTrimTime = 0.0f;
+
+				{
+					FScopeLock Lock(&Allocator.GetMutex());
+					WaitForMutexTime = FPlatformTime::Seconds() - StartTimeInner;
+					TrimThreadFreeBlockLists(Allocator, Lists);
+					WaitForMutexAndTrimTime = FPlatformTime::Seconds() - StartTimeInner;
+				}
+
+				// These logs must happen outside the above mutex to avoid deadlocks
+				if (WaitForMutexTime > AllocType::GetFlushThreadCacheMaxWaitTime())
+				{
+					UE_LOG(LogMemory, Warning, TEXT("FMalloc%s took %6.2fms to wait for mutex for trim."), GetDescriptiveName(), WaitForMutexTime * 1000.0f);
+				}
+				if (WaitForMutexAndTrimTime > AllocType::GetFlushThreadCacheMaxWaitTime())
+				{
+					UE_LOG(LogMemory, Warning, TEXT("FMalloc%s took %6.2fms to wait for mutex AND trim."), GetDescriptiveName(), WaitForMutexAndTrimTime * 1000.0f);
+				}
+			}
+		}
+	}
+
+	void TrimImpl(AllocType& Allocator)
+	{
+		// Update the trim epoch so that threads cleanup their thread-local memory when going to sleep.
+		MemoryTrimEpoch.fetch_add(1, std::memory_order_relaxed);
+
+		QUICK_SCOPE_CYCLE_COUNTER(STAT_FMallocBinnedCommon_Trim);
+
+		// Process thread-local memory caches from as many threads as possible without waking them up.
+		// Skip on desktop as we may have too many threads and this could cause some hitches.
+		if (!PLATFORM_DESKTOP)
+		{
+			FScopeLock Lock(&Allocator.GetMutex());
+			FScopeLock FreeBlockLock(&AllocType::GetFreeBlockListsRegistrationMutex());
+			for (FPerThreadFreeBlockLists* BlockList : AllocType::GetRegisteredFreeBlockLists())
+			{
+				// If we're unable to lock, it's because the thread is currently active so it
+				// will do the flush itself when going back to sleep because we incremented the Epoch.
+				if (BlockList->TryLock())
+				{
+					// Only trim if the epoch has been updated, otherwise the thread already
+					// did the trimming when it went to sleep.
+					if (BlockList->UpdateEpoch(MemoryTrimEpoch.load(std::memory_order_relaxed)))
+					{
+						TrimThreadFreeBlockLists(Allocator, BlockList);
+					}
+					BlockList->Unlock();
+				}
+			}
+		}
+
+		TFunction<void(ENamedThreads::Type CurrentThread)> Broadcast =
+			[this, &Allocator](ENamedThreads::Type MyThread)
+			{
+				// We might already have updated the Epoch so we can skip doing anything costly (i.e. Mutex) in that case.
+				const bool bNewEpochOnly = true;
+				FlushCurrentThreadCache(Allocator, bNewEpochOnly);
+			};
+
+		// Skip task threads on desktop platforms as it is too slow and they don't have much memory
+		if (PLATFORM_DESKTOP)
+		{
+			FTaskGraphInterface::BroadcastSlow_OnlyUseForSpecialPurposes(false, false, Broadcast);
+		}
+		else
+		{
+			FTaskGraphInterface::BroadcastSlow_OnlyUseForSpecialPurposes(FPlatformProcess::SupportsMultithreading() && FApp::ShouldUseThreadingForPerformance(), false, Broadcast);
+		}
+	}
 };
\ No newline at end of file
diff --git a/Engine/Source/Runtime/Core/Public/HAL/MallocCallstackHandler.h b/Engine/Source/Runtime/Core/Public/HAL/MallocCallstackHandler.h
--- a/Engine/Source/Runtime/Core/Public/HAL/MallocCallstackHandler.h
+++ b/Engine/Source/Runtime/Core/Public/HAL/MallocCallstackHandler.h
@@ -71,6 +71,22 @@ public:
 		UsedMalloc->SetupTLSCachesOnCurrentThread();
 	}
 
+	/**
+	* Mark TLS caches for the current thread as used. Thread has woken up to do some processing and needs its TLS caches back.
+	*/
+	virtual void MarkTLSCachesAsUsedOnCurrentThread() override
+	{
+		UsedMalloc->MarkTLSCachesAsUsedOnCurrentThread();
+	}
+
+	/**
+	* Mark TLS caches for current thread as unused. Typically before going to sleep. These are the threads that we can trim without waking them up.
+	*/
+	virtual void MarkTLSCachesAsUnusedOnCurrentThread()
+	{
+		UsedMalloc->MarkTLSCachesAsUnusedOnCurrentThread();
+	}
+
 	/**
 	* Clears the TLS caches on the current thread and disables any future caching.
 	*/
diff --git a/Engine/Source/Runtime/Core/Public/HAL/MallocPoisonProxy.h b/Engine/Source/Runtime/Core/Public/HAL/MallocPoisonProxy.h
--- a/Engine/Source/Runtime/Core/Public/HAL/MallocPoisonProxy.h
+++ b/Engine/Source/Runtime/Core/Public/HAL/MallocPoisonProxy.h
@@ -141,6 +141,16 @@ public:
 		UsedMalloc->SetupTLSCachesOnCurrentThread();
 	}
 
+	virtual void MarkTLSCachesAsUsedOnCurrentThread() override
+	{
+		UsedMalloc->MarkTLSCachesAsUsedOnCurrentThread();
+	}
+
+	virtual void MarkTLSCachesAsUnusedOnCurrentThread() override
+	{
+		UsedMalloc->MarkTLSCachesAsUnusedOnCurrentThread();
+	}
+
 	virtual void ClearAndDisableTLSCachesOnCurrentThread() override
 	{
 		UsedMalloc->ClearAndDisableTLSCachesOnCurrentThread();
diff --git a/Engine/Source/Runtime/Core/Public/HAL/MemoryBase.h b/Engine/Source/Runtime/Core/Public/HAL/MemoryBase.h
--- a/Engine/Source/Runtime/Core/Public/HAL/MemoryBase.h
+++ b/Engine/Source/Runtime/Core/Public/HAL/MemoryBase.h
@@ -146,6 +146,20 @@ public:
 	{
 	}
 
+	/**
+	* Mark TLS caches for the current thread as used. Thread has woken up to do some processing and needs its TLS caches back.
+	*/
+	virtual void MarkTLSCachesAsUsedOnCurrentThread()
+	{
+	}
+
+	/**
+	* Mark TLS caches for current thread as unused. Typically before going to sleep. These are the threads that we can trim without waking them up.
+	*/
+	virtual void MarkTLSCachesAsUnusedOnCurrentThread()
+	{
+	}
+
 	/**
 	* Clears the TLS caches on the current thread and disables any future caching.
 	*/
diff --git a/Engine/Source/Runtime/Core/Public/HAL/UnrealMemory.h b/Engine/Source/Runtime/Core/Public/HAL/UnrealMemory.h
--- a/Engine/Source/Runtime/Core/Public/HAL/UnrealMemory.h
+++ b/Engine/Source/Runtime/Core/Public/HAL/UnrealMemory.h
@@ -244,6 +244,16 @@ struct FMemory
 	*/
 	static CORE_API void ClearAndDisableTLSCachesOnCurrentThread();
 
+	/**
+	* Mark TLS caches for the current thread as used. Thread has woken up to do some processing and needs its TLS caches back.
+	*/
+	static CORE_API void MarkTLSCachesAsUsedOnCurrentThread();
+
+	/**
+	* Mark TLS caches for current thread as unused. Typically before going to sleep. These are the threads that we can trim without waking them up.
+	*/
+	static CORE_API void MarkTLSCachesAsUnusedOnCurrentThread();
+
 	/**
 	 * A helper function that will perform a series of random heap allocations to test
 	 * the internal validity of the heap. Note, this function will "leak" memory, but another call
