[MallocBinned2]
 - Fix fork support regression not caught by preflight

diff --git a/Engine/Source/Runtime/Core/Private/HAL/MallocBinned2.cpp b/Engine/Source/Runtime/Core/Private/HAL/MallocBinned2.cpp
--- a/Engine/Source/Runtime/Core/Private/HAL/MallocBinned2.cpp
+++ b/Engine/Source/Runtime/Core/Private/HAL/MallocBinned2.cpp
@@ -779,7 +779,7 @@ void FMallocBinned2::OnPreFork()
 	// Trim caches so we don't use them in the child process and cause pages to be copied
 	if (GMallocBinned2PerThreadCaches)
 	{
-		FlushCurrentThreadCache();
+		FlushCurrentThreadCache(*this);
 		FMallocBinned2::Private::CheckThreadFreeBlockListsForFork();
 	}
 
@@ -805,7 +805,7 @@ void FMallocBinned2::OnPostFork()
 #if BINNED2_FORK_SUPPORT
 	if (GMallocBinned2PerThreadCaches)
 	{
-		FlushCurrentThreadCache();
+		FlushCurrentThreadCache(*this);
 		FMallocBinned2::Private::CheckThreadFreeBlockListsForFork();
 	}
 
