[TaskGraph]
 - Fix task tracing having a huge cost even when not active

diff --git a/Engine/Source/Runtime/Core/Private/Async/TaskTrace.cpp b/Engine/Source/Runtime/Core/Private/Async/TaskTrace.cpp
--- a/Engine/Source/Runtime/Core/Private/Async/TaskTrace.cpp
+++ b/Engine/Source/Runtime/Core/Private/Async/TaskTrace.cpp
@@ -99,162 +99,195 @@ namespace TaskTrace
 
 	void Created(FId TaskId, uint64 TaskSize)
 	{
-		if (!bGTaskTraceInitialized)
+		if (UE_TRACE_CHANNELEXPR_IS_ENABLED(TaskChannel))
 		{
-			return;
-		}
+			if (!bGTaskTraceInitialized)
+			{
+				return;
+			}
 
-		check(TaskId != InvalidId);
+			check(TaskId != InvalidId);
 
-		UE_TRACE_LOG(TaskTrace, Created, TaskChannel)
-			<< Created.Timestamp(FPlatformTime::Cycles64())
-			<< Created.TaskId(TaskId)
-			<< Created.TaskSize(TaskSize);
+			UE_TRACE_LOG(TaskTrace, Created, TaskChannel)
+				<< Created.Timestamp(FPlatformTime::Cycles64())
+				<< Created.TaskId(TaskId)
+				<< Created.TaskSize(TaskSize);
+		}
 	}
 
 	void Launched(FId TaskId, const TCHAR* DebugName, bool bTracked, ENamedThreads::Type ThreadToExecuteOn, uint64 TaskSize)
 	{
-		if (!bGTaskTraceInitialized)
+		if (UE_TRACE_CHANNELEXPR_IS_ENABLED(TaskChannel))
 		{
-			return;
-		}
+			if (!bGTaskTraceInitialized)
+			{
+				return;
+			}
 
-		check(TaskId != InvalidId);
+			check(TaskId != InvalidId);
 
-		UE_TRACE_LOG(TaskTrace, Launched, TaskChannel)
-			<< Launched.Timestamp(FPlatformTime::Cycles64())
-			<< Launched.TaskId(TaskId)
-			<< Launched.DebugName(DebugName != nullptr ? DebugName : TEXT(""))
-			<< Launched.Tracked(bTracked)
-			<< Launched.ThreadToExecuteOn(ThreadToExecuteOn)
-			<< Launched.TaskSize(TaskSize);
+			UE_TRACE_LOG(TaskTrace, Launched, TaskChannel)
+				<< Launched.Timestamp(FPlatformTime::Cycles64())
+				<< Launched.TaskId(TaskId)
+				<< Launched.DebugName(DebugName != nullptr ? DebugName : TEXT(""))
+				<< Launched.Tracked(bTracked)
+				<< Launched.ThreadToExecuteOn(ThreadToExecuteOn)
+				<< Launched.TaskSize(TaskSize);
+		}
 	}
 
 	void Scheduled(FId TaskId)
 	{
-		if (!bGTaskTraceInitialized)
+		if (UE_TRACE_CHANNELEXPR_IS_ENABLED(TaskChannel))
 		{
-			return;
-		}
+			if (!bGTaskTraceInitialized)
+			{
+				return;
+			}
 
-		check(TaskId != InvalidId);
+			check(TaskId != InvalidId);
 
-		UE_TRACE_LOG(TaskTrace, Scheduled, TaskChannel)
-			<< Scheduled.Timestamp(FPlatformTime::Cycles64())
-			<< Scheduled.TaskId(TaskId);
+			UE_TRACE_LOG(TaskTrace, Scheduled, TaskChannel)
+				<< Scheduled.Timestamp(FPlatformTime::Cycles64())
+				<< Scheduled.TaskId(TaskId);
+		}
 	}
 
 	void SubsequentAdded(FId TaskId, FId SubsequentId)
 	{
-		if (!bGTaskTraceInitialized)
+		if (UE_TRACE_CHANNELEXPR_IS_ENABLED(TaskChannel))
 		{
-			return;
-		}
+			if (!bGTaskTraceInitialized)
+			{
+				return;
+			}
 
-		// "empty" FGraphEvent is used for synchronisation only, to wait for a notification. It doesn't have an associated task and ID.
-		if (TaskId == InvalidId)
-		{
-			TaskId = GenerateTaskId();
-		}
+			// "empty" FGraphEvent is used for synchronisation only, to wait for a notification. It doesn't have an associated task and ID.
+			if (TaskId == InvalidId)
+			{
+				TaskId = GenerateTaskId();
+			}
 
-		check(SubsequentId != InvalidId);
-		UE_TRACE_LOG(TaskTrace, SubsequentAdded, TaskChannel)
-			<< SubsequentAdded.Timestamp(FPlatformTime::Cycles64())
-			<< SubsequentAdded.TaskId(TaskId)
-			<< SubsequentAdded.SubsequentId(SubsequentId);
+			check(SubsequentId != InvalidId);
+			UE_TRACE_LOG(TaskTrace, SubsequentAdded, TaskChannel)
+				<< SubsequentAdded.Timestamp(FPlatformTime::Cycles64())
+				<< SubsequentAdded.TaskId(TaskId)
+				<< SubsequentAdded.SubsequentId(SubsequentId);
+		}
 	}
 
 	void Started(FId TaskId)
 	{
-		if (!bGTaskTraceInitialized)
+		if (UE_TRACE_CHANNELEXPR_IS_ENABLED(TaskChannel))
 		{
-			return;
-		}
+			if (!bGTaskTraceInitialized)
+			{
+				return;
+			}
 
-		check(TaskId != InvalidId);
+			check(TaskId != InvalidId);
 
-		UE_TRACE_LOG(TaskTrace, Started, TaskChannel)
-			<< Started.Timestamp(FPlatformTime::Cycles64())
-			<< Started.TaskId(TaskId);
+			UE_TRACE_LOG(TaskTrace, Started, TaskChannel)
+				<< Started.Timestamp(FPlatformTime::Cycles64())
+				<< Started.TaskId(TaskId);
+		}
 	}
 
 	void Finished(FId TaskId)
 	{
-		if (!bGTaskTraceInitialized)
+		if (UE_TRACE_CHANNELEXPR_IS_ENABLED(TaskChannel))
 		{
-			return;
-		}
+			if (!bGTaskTraceInitialized)
+			{
+				return;
+			}
 
-		check(TaskId != InvalidId);
+			check(TaskId != InvalidId);
 
-		UE_TRACE_LOG(TaskTrace, Finished, TaskChannel)
-			<< Finished.Timestamp(FPlatformTime::Cycles64())
-			<< Finished.TaskId(TaskId);
+			UE_TRACE_LOG(TaskTrace, Finished, TaskChannel)
+				<< Finished.Timestamp(FPlatformTime::Cycles64())
+				<< Finished.TaskId(TaskId);
+		}
 	}
 
 	void Completed(FId TaskId)
 	{
-		if (!bGTaskTraceInitialized)
+		if (UE_TRACE_CHANNELEXPR_IS_ENABLED(TaskChannel))
 		{
-			return;
-		}
+			if (!bGTaskTraceInitialized)
+			{
+				return;
+			}
 
-		// "empty" FGraphEvent is used for synchronisation only, to wait for a notification. It doesn't have an associated task and ID.
-		if (TaskId == InvalidId)
-		{
-			TaskId = GenerateTaskId();
-		}
+			// "empty" FGraphEvent is used for synchronisation only, to wait for a notification. It doesn't have an associated task and ID.
+			if (TaskId == InvalidId)
+			{
+				TaskId = GenerateTaskId();
+			}
 
-		UE_TRACE_LOG(TaskTrace, Completed, TaskChannel)
-			<< Completed.Timestamp(FPlatformTime::Cycles64())
-			<< Completed.TaskId(TaskId);
+			UE_TRACE_LOG(TaskTrace, Completed, TaskChannel)
+				<< Completed.Timestamp(FPlatformTime::Cycles64())
+				<< Completed.TaskId(TaskId);
+		}
 	}
 
 	void Destroyed(FId TaskId)
 	{
-		if (!bGTaskTraceInitialized)
+		if (UE_TRACE_CHANNELEXPR_IS_ENABLED(TaskChannel))
 		{
-			return;
-		}
+			if (!bGTaskTraceInitialized)
+			{
+				return;
+			}
 
-		UE_TRACE_LOG(TaskTrace, Destroyed, TaskChannel)
-			<< Destroyed.Timestamp(FPlatformTime::Cycles64())
-			<< Destroyed.TaskId(TaskId);
+			UE_TRACE_LOG(TaskTrace, Destroyed, TaskChannel)
+				<< Destroyed.Timestamp(FPlatformTime::Cycles64())
+				<< Destroyed.TaskId(TaskId);
+		}
 	}
 
 	FWaitingScope::FWaitingScope(const TArray<FId>& Tasks)
 	{
-		if (!bGTaskTraceInitialized)
+		if (UE_TRACE_CHANNELEXPR_IS_ENABLED(TaskChannel))
 		{
-			return;
-		}
+			if (!bGTaskTraceInitialized)
+			{
+				return;
+			}
 
-		UE_TRACE_LOG(TaskTrace, WaitingStarted, TaskChannel)
-			<< WaitingStarted.Timestamp(FPlatformTime::Cycles64())
-			<< WaitingStarted.Tasks(Tasks.GetData(), Tasks.Num());
+			UE_TRACE_LOG(TaskTrace, WaitingStarted, TaskChannel)
+				<< WaitingStarted.Timestamp(FPlatformTime::Cycles64())
+				<< WaitingStarted.Tasks(Tasks.GetData(), Tasks.Num());
+		}
 	}
 
 	FWaitingScope::FWaitingScope(FId TaskId)
 	{
-		if (!bGTaskTraceInitialized)
+		if (UE_TRACE_CHANNELEXPR_IS_ENABLED(TaskChannel))
 		{
-			return;
-		}
+			if (!bGTaskTraceInitialized)
+			{
+				return;
+			}
 
-		UE_TRACE_LOG(TaskTrace, WaitingStarted, TaskChannel)
-			<< WaitingStarted.Timestamp(FPlatformTime::Cycles64())
-			<< WaitingStarted.Tasks(&TaskId, 1);
+			UE_TRACE_LOG(TaskTrace, WaitingStarted, TaskChannel)
+				<< WaitingStarted.Timestamp(FPlatformTime::Cycles64())
+				<< WaitingStarted.Tasks(&TaskId, 1);
+		}
 	}
 
 	FWaitingScope::~FWaitingScope()
 	{
-		if (!bGTaskTraceInitialized)
+		if (UE_TRACE_CHANNELEXPR_IS_ENABLED(TaskChannel))
 		{
-			return;
-		}
+			if (!bGTaskTraceInitialized)
+			{
+				return;
+			}
 
-		UE_TRACE_LOG(TaskTrace, WaitingFinished, TaskChannel)
-			<< WaitingFinished.Timestamp(FPlatformTime::Cycles64());
+			UE_TRACE_LOG(TaskTrace, WaitingFinished, TaskChannel)
+				<< WaitingFinished.Timestamp(FPlatformTime::Cycles64());
+		}
 	}
 
 	FTaskTimingEventScope::FTaskTimingEventScope(TaskTrace::FId InTaskId)
