[MallocBinned2/3]
 - Unify GMallocBinned2FlushThreadCacheMaxWaitTime and GMallocBinned3FlushThreadCacheMaxWaitTime
 - Move FApp call to cpp file to remove App.h from MallocBinnedCommon.h

diff --git a/Engine/Source/Runtime/Core/Private/HAL/MallocBinned2.cpp b/Engine/Source/Runtime/Core/Private/HAL/MallocBinned2.cpp
--- a/Engine/Source/Runtime/Core/Private/HAL/MallocBinned2.cpp
+++ b/Engine/Source/Runtime/Core/Private/HAL/MallocBinned2.cpp
@@ -86,15 +86,6 @@ static FAutoConsoleVariableRef GGMallocBinned2MoveOSFreesOffTimeCriticalThreadsC
 #endif
 
 
-
-float GMallocBinned2FlushThreadCacheMaxWaitTime = 0.02f;
-static FAutoConsoleVariableRef GMallocBinned2FlushThreadCacheMaxWaitTimeCVar(
-	TEXT("MallocBinned2.FlushThreadCacheMaxWaitTime"),
-	GMallocBinned2FlushThreadCacheMaxWaitTime,
-	TEXT("The threshold of time before warning about FlushCurrentThreadCache taking too long (seconds)."),
-	ECVF_ReadOnly
-);
-
 #if BINNED2_ALLOCATOR_STATS
 TAtomic<int64> AllocatedSmallPoolMemory(0); // memory that's requested to be allocated by the game
 TAtomic<int64> AllocatedOSSmallPoolMemory(0);
@@ -1211,11 +1202,6 @@ void FMallocBinned2::FreeBundles(FBundleNode* Bundles, uint32 PoolIndex)
 	Private::FreeBundles(*this, Bundles, PoolIndexToBlockSize(PoolIndex), PoolIndex);
 }
 
-float FMallocBinned2::GetFlushThreadCacheMaxWaitTime()
-{
-	return GMallocBinned2FlushThreadCacheMaxWaitTime;
-}
-
 FCriticalSection& FMallocBinned2::GetFreeBlockListsRegistrationMutex()
 {
 	return Private::GetFreeBlockListsRegistrationMutex();
diff --git a/Engine/Source/Runtime/Core/Private/HAL/MallocBinned3.cpp b/Engine/Source/Runtime/Core/Private/HAL/MallocBinned3.cpp
--- a/Engine/Source/Runtime/Core/Private/HAL/MallocBinned3.cpp
+++ b/Engine/Source/Runtime/Core/Private/HAL/MallocBinned3.cpp
@@ -54,14 +54,6 @@ static FAutoConsoleVariableRef GMallocBinned3AllocExtraCVar(
 
 #endif
 
-float GMallocBinned3FlushThreadCacheMaxWaitTime = 0.02f;
-static FAutoConsoleVariableRef GMallocBinned3FlushThreadCacheMaxWaitTimeCVar(
-	TEXT("MallocBinned3.FlushThreadCacheMaxWaitTime"),
-	GMallocBinned3FlushThreadCacheMaxWaitTime,
-	TEXT("The threshold of time before warning about FlushCurrentThreadCache taking too long (seconds)."),
-	ECVF_ReadOnly
-);
-
 #if BINNED3_ALLOCATOR_STATS
 int64 Binned3AllocatedSmallPoolMemory = 0; // memory that's requested to be allocated by the game
 int64 Binned3AllocatedOSSmallPoolMemory = 0;
@@ -1303,11 +1295,6 @@ void FMallocBinned3::Trim(bool bTrimThreadCaches)
 	}
 }
 
-float FMallocBinned3::GetFlushThreadCacheMaxWaitTime()
-{
-	return GMallocBinned3FlushThreadCacheMaxWaitTime;
-}
-
 FCriticalSection& FMallocBinned3::GetFreeBlockListsRegistrationMutex()
 {
 	return Private::GetFreeBlockListsRegistrationMutex();
diff --git a/Engine/Source/Runtime/Core/Private/HAL/MallocBinnedCommon.cpp b/Engine/Source/Runtime/Core/Private/HAL/MallocBinnedCommon.cpp
--- a/Engine/Source/Runtime/Core/Private/HAL/MallocBinnedCommon.cpp
+++ b/Engine/Source/Runtime/Core/Private/HAL/MallocBinnedCommon.cpp
@@ -2,6 +2,7 @@
 
 #include "HAL/MallocBinnedCommon.h"
 #include "Algo/Sort.h"
+#include "Misc/App.h"
 #include "Containers/ArrayView.h"
 #include "Misc/AssertionMacros.h"
 #include "Math/NumericLimits.h"
@@ -453,6 +454,19 @@ uint32 FBitTree::CountOnes(uint32 UpTo) const
 
 #endif
 
+bool FMallocBinnedCommonBase::IsAppMultithreaded()
+{
+	return FPlatformProcess::SupportsMultithreading() && FApp::ShouldUseThreadingForPerformance();
+}
+
+float GMallocBinnedFlushThreadCacheMaxWaitTime = 0.2f;
+static FAutoConsoleVariableRef GMallocBinnedFlushThreadCacheMaxWaitTimeCVar(
+	TEXT("MallocBinned.FlushThreadCacheMaxWaitTime"),
+	GMallocBinnedFlushThreadCacheMaxWaitTime,
+	TEXT("The threshold of time before warning about FlushCurrentThreadCache taking too long (seconds)."),
+	ECVF_ReadOnly
+);
+
 #if UE_BINNEDCOMMON_ALLOW_RUNTIME_TWEAKING
 
 int32 GMallocBinnedBundleSize = DEFAULT_GMallocBinnedBundleSize;
diff --git a/Engine/Source/Runtime/Core/Public/HAL/MallocBinned2.h b/Engine/Source/Runtime/Core/Public/HAL/MallocBinned2.h
--- a/Engine/Source/Runtime/Core/Public/HAL/MallocBinned2.h
+++ b/Engine/Source/Runtime/Core/Public/HAL/MallocBinned2.h
@@ -529,7 +529,6 @@ public:
 
 	static void RegisterThreadFreeBlockLists(FPerThreadFreeBlockLists* FreeBlockLists);
 	static void UnregisterThreadFreeBlockLists(FPerThreadFreeBlockLists* FreeBlockLists);
-	static float GetFlushThreadCacheMaxWaitTime();
 	static FCriticalSection& GetFreeBlockListsRegistrationMutex();
 	static TArray<FPerThreadFreeBlockLists*>& GetRegisteredFreeBlockLists();
 
diff --git a/Engine/Source/Runtime/Core/Public/HAL/MallocBinned3.h b/Engine/Source/Runtime/Core/Public/HAL/MallocBinned3.h
--- a/Engine/Source/Runtime/Core/Public/HAL/MallocBinned3.h
+++ b/Engine/Source/Runtime/Core/Public/HAL/MallocBinned3.h
@@ -472,7 +472,6 @@ public:
 
 	static void RegisterThreadFreeBlockLists(FPerThreadFreeBlockLists* FreeBlockLists);
 	static void UnregisterThreadFreeBlockLists(FPerThreadFreeBlockLists* FreeBlockLists);
-	static float GetFlushThreadCacheMaxWaitTime();
 	static FCriticalSection& GetFreeBlockListsRegistrationMutex();
 	static TArray<FPerThreadFreeBlockLists*>& GetRegisteredFreeBlockLists();
 
diff --git a/Engine/Source/Runtime/Core/Public/HAL/MallocBinnedCommon.h b/Engine/Source/Runtime/Core/Public/HAL/MallocBinnedCommon.h
--- a/Engine/Source/Runtime/Core/Public/HAL/MallocBinnedCommon.h
+++ b/Engine/Source/Runtime/Core/Public/HAL/MallocBinnedCommon.h
@@ -158,6 +158,7 @@ extern CORE_API int32 GMallocBinnedBundleCount;
 #	define UE_BINNEDCOMMON_ALLOCATOR_STATS (!UE_BUILD_SHIPPING || WITH_EDITOR)
 #endif
 
+extern CORE_API float GMallocBinnedFlushThreadCacheMaxWaitTime;
 
 class FMallocBinnedCommonBase : public FMalloc
 {
@@ -309,6 +310,9 @@ protected:
 	static std::atomic<int64> ConsolidatedMemory;
 #endif
 	std::atomic<uint64> MemoryTrimEpoch{ 0 };
+
+protected:
+	bool IsAppMultithreaded();
 };
 
 template <class AllocType, int MinAlign, int MaxAlign, int MinAlignShift, int NumSmallPools, int MaxSmallPoolSize>
@@ -669,11 +673,11 @@ protected:
 				}
 
 				// These logs must happen outside the above mutex to avoid deadlocks
-				if (WaitForMutexTime > AllocType::GetFlushThreadCacheMaxWaitTime())
+				if (WaitForMutexTime > GMallocBinnedFlushThreadCacheMaxWaitTime)
 				{
 					UE_LOG(LogMemory, Warning, TEXT("FMalloc%s took %6.2fms to wait for mutex for trim."), GetDescriptiveName(), WaitForMutexTime * 1000.0f);
 				}
-				if (WaitForMutexAndTrimTime > AllocType::GetFlushThreadCacheMaxWaitTime())
+				if (WaitForMutexAndTrimTime > GMallocBinnedFlushThreadCacheMaxWaitTime)
 				{
 					UE_LOG(LogMemory, Warning, TEXT("FMalloc%s took %6.2fms to wait for mutex AND trim."), GetDescriptiveName(), WaitForMutexAndTrimTime * 1000.0f);
 				}
@@ -726,7 +730,7 @@ protected:
 		}
 		else
 		{
-			FTaskGraphInterface::BroadcastSlow_OnlyUseForSpecialPurposes(FPlatformProcess::SupportsMultithreading() && FApp::ShouldUseThreadingForPerformance(), false, Broadcast);
+			FTaskGraphInterface::BroadcastSlow_OnlyUseForSpecialPurposes(IsAppMultithreaded(), false, Broadcast);
 		}
 	}
 };
\ No newline at end of file
