[BinnedCommon]
 - Really remove the App.h this time

diff --git a/Engine/Source/Runtime/Core/Public/HAL/MallocBinnedCommon.h b/Engine/Source/Runtime/Core/Public/HAL/MallocBinnedCommon.h
--- a/Engine/Source/Runtime/Core/Public/HAL/MallocBinnedCommon.h
+++ b/Engine/Source/Runtime/Core/Public/HAL/MallocBinnedCommon.h
@@ -6,7 +6,6 @@
 #include "HAL/MemoryBase.h"
 #include "Math/UnrealMathUtility.h"
 #include "HAL/PlatformTLS.h"
-#include "Misc/App.h"
 #include "Async/TaskGraphInterfaces.h"
 #include "Templates/AlignmentTemplates.h"
 
