[Math] Improvements to TIntVector2, TIntVector3 and TIntVector4 to bring them closer to parity with the functionality available in the corresponding TVector classes. Added a new file to contain IntVector tests. Added unit tests for the functions added.

#tests Ran newly added unit tests successfully.

diff --git a/Engine/Source/Runtime/Core/Public/Math/IntVector.h b/Engine/Source/Runtime/Core/Public/Math/IntVector.h
--- a/Engine/Source/Runtime/Core/Public/Math/IntVector.h
+++ b/Engine/Source/Runtime/Core/Public/Math/IntVector.h
@@ -26,13 +26,13 @@ struct TIntVector3
 	{
 		struct
 		{
-			/** Holds the point's x-coordinate. */
+			/** Holds the vector's x-coordinate. */
 			IntType X;
 
-			/** Holds the point's y-coordinate. */
+			/** Holds the vector's y-coordinate. */
 			IntType Y;
 
-			/**  Holds the point's z-coordinate. */
+			/** Holds the vector's z-coordinate. */
 			IntType Z;
 		};
 
@@ -40,10 +40,10 @@ struct TIntVector3
 		IntType XYZ[3];
 	};
 
-	/** An int point with zeroed values. */
+	/** An int vector with zeroed values. */
 	static const TIntVector3 ZeroValue;
 
-	/** An int point with INDEX_NONE values. */
+	/** An int vector with INDEX_NONE values. */
 	static const TIntVector3 NoneValue;
 
 	/**
@@ -58,7 +58,7 @@ struct TIntVector3
 	 * @param InY The y-coordinate.
 	 * @param InZ The z-coordinate.
 	 */
-	TIntVector3(IntType InX, IntType InY, IntType InZ )
+	TIntVector3(IntType InX, IntType InY, IntType InZ)
 		: X(InX)
 		, Y(InY)
 		, Z(InZ)
@@ -70,7 +70,7 @@ struct TIntVector3
 	 *
 	 * @param InValue replicated to all components
 	 */
-	explicit TIntVector3(IntType InValue )
+	explicit TIntVector3(IntType InValue)
 		: X(InValue)
 		, Y(InValue)
 		, Z(InValue)
@@ -82,15 +82,14 @@ struct TIntVector3
 	 *
 	 * @param InVector float vector converted to int
 	 */
-	explicit TIntVector3( FVector InVector  );
+	explicit TIntVector3(FVector InVector);
 
 	/**
 	 * Constructor
 	 *
 	 * @param EForceInit Force init enum
 	 */
-
-	explicit TIntVector3( EForceInit )
+	explicit TIntVector3(EForceInit)
 		: X(0)
 		, Y(0)
 		, Z(0)
@@ -117,9 +116,9 @@ PRAGMA_DISABLE_DEPRECATION_WARNINGS
 PRAGMA_ENABLE_DEPRECATION_WARNINGS
 
 	/**
-	 * Gets specific component of a point.
+	 * Gets specific component of a vector.
 	 *
-	 * @param ComponentIndex Index of point component.
+	 * @param ComponentIndex Index of vector component.
 	 * @return const reference to component.
 	 */
 	const IntType& operator()(int32 ComponentIndex) const
@@ -130,9 +129,9 @@ PRAGMA_ENABLE_DEPRECATION_WARNINGS
 	}
 
 	/**
-	 * Gets specific component of a point.
+	 * Gets specific component of a vector.
 	 *
-	 * @param ComponentIndex Index of point component.
+	 * @param ComponentIndex Index of vector component.
 	 * @return reference to component.
 	 */
 	IntType& operator()(int32 ComponentIndex)
@@ -143,9 +142,9 @@ PRAGMA_ENABLE_DEPRECATION_WARNINGS
 	}
 
 	/**
-	 * Gets specific component of a point.
+	 * Gets specific component of a vector.
 	 *
-	 * @param ComponentIndex Index of point component.
+	 * @param ComponentIndex Index of vector component.
 	 * @return const reference to component.
 	 */
 	const IntType& operator[](int32 ComponentIndex) const
@@ -156,9 +155,9 @@ PRAGMA_ENABLE_DEPRECATION_WARNINGS
 	}
 
 	/**
-	 * Gets specific component of a point.
+	 * Gets specific component of a vector.
 	 *
-	 * @param ComponentIndex Index of point component.
+	 * @param ComponentIndex Index of vector component.
 	 * @return reference to component.
 	 */
 	IntType& operator[](int32 ComponentIndex)
@@ -169,10 +168,10 @@ PRAGMA_ENABLE_DEPRECATION_WARNINGS
 	}
 
 	/**
-	 * Compares points for equality.
+	 * Compares vectors for equality.
 	 *
-	 * @param Other The other int point being compared.
-	 * @return true if the points are equal, false otherwise..
+	 * @param Other The other int vector being compared.
+	 * @return true if the vectors are equal, false otherwise..
 	 */
 	bool operator==(const TIntVector3& Other) const
 	{
@@ -180,10 +179,10 @@ PRAGMA_ENABLE_DEPRECATION_WARNINGS
 	}
 
 	/**
-	 * Compares points for inequality.
+	 * Compares vectors for inequality.
 	 *
-	 * @param Other The other int point being compared.
-	 * @return true if the points are not equal, false otherwise..
+	 * @param Other The other int vector being compared.
+	 * @return true if the vectors are not equal, false otherwise..
 	 */
 	bool operator!=(const TIntVector3& Other) const
 	{
@@ -193,8 +192,8 @@ PRAGMA_ENABLE_DEPRECATION_WARNINGS
 	/**
 	 * Multiplies this vector with another vector, using component-wise multiplication.
 	 *
-	 * @param Other The point to multiply with.
-	 * @return Reference to this point after multiplication.
+	 * @param Other The vector to multiply with.
+	 * @return Reference to this vector after multiplication.
 	 */
 	TIntVector3& operator*=(const TIntVector3& Other)
 	{
@@ -206,10 +205,10 @@ PRAGMA_ENABLE_DEPRECATION_WARNINGS
 	}
 
 	/**
-	 * Scales this point.
+	 * Scales this vector.
 	 *
-	 * @param Scale What to multiply the point by.
-	 * @return Reference to this point after multiplication.
+	 * @param Scale What to multiply the vector by.
+	 * @return Reference to this vector after multiplication.
 	 */
 	TIntVector3& operator*=(IntType Scale)
 	{
@@ -221,10 +220,10 @@ PRAGMA_ENABLE_DEPRECATION_WARNINGS
 	}
 
 	/**
-	 * Divides this point.
+	 * Divides this vector.
 	 *
-	 * @param Divisor What to divide the point by.
-	 * @return Reference to this point after division.
+	 * @param Divisor What to divide the vector by.
+	 * @return Reference to this vector after division.
 	 */
 	TIntVector3& operator/=(IntType Divisor)
 	{
@@ -236,10 +235,10 @@ PRAGMA_ENABLE_DEPRECATION_WARNINGS
 	}
 
 	/**
-	 * Remainder of division of this point.
+	 * Remainder of division of this vector.
 	 *
-	 * @param Divisor What to divide the point by.
-	 * @return Reference to this point after remainder.
+	 * @param Divisor What to divide the vector by.
+	 * @return Reference to this vector after remainder.
 	 */
 	TIntVector3& operator%=(IntType Divisor)
 	{
@@ -251,10 +250,10 @@ PRAGMA_ENABLE_DEPRECATION_WARNINGS
 	}
 
 	/**
-	 * Adds to this point.
+	 * Adds to this vector.
 	 *
-	 * @param Other The point to add to this point.
-	 * @return Reference to this point after addition.
+	 * @param Other The vector to add to this vector.
+	 * @return Reference to this vector after addition.
 	 */
 	TIntVector3& operator+=(const TIntVector3& Other)
 	{
@@ -266,10 +265,10 @@ PRAGMA_ENABLE_DEPRECATION_WARNINGS
 	}
 
 	/**
-	 * Subtracts from this point.
+	 * Subtracts from this vector.
 	 *
-	 * @param Other The point to subtract from this point.
-	 * @return Reference to this point after subtraction.
+	 * @param Other The vector to subtract from this vector.
+	 * @return Reference to this vector after subtraction.
 	 */
 	TIntVector3& operator-=(const TIntVector3& Other)
 	{
@@ -281,9 +280,9 @@ PRAGMA_ENABLE_DEPRECATION_WARNINGS
 	}
 
 	/**
-	 * Gets the result of component-wise multiplication of this point by another.
+	 * Gets the result of component-wise multiplication of this vector by another.
 	 *
-	 * @param Other The point to multiply with.
+	 * @param Other The vector to multiply with.
 	 * @return The result of multiplication.
 	 */
 	TIntVector3 operator*(const TIntVector3& Other) const
@@ -292,10 +291,10 @@ PRAGMA_ENABLE_DEPRECATION_WARNINGS
 	}
 
 	/**
-	 * Gets the result of scaling on this point.
+	 * Gets the result of scaling on this vector.
 	 *
-	 * @param Scale What to multiply the point by.
-	 * @return A new scaled int point.
+	 * @param Scale What to multiply the vector by.
+	 * @return A new scaled int vector.
 	 */
 	TIntVector3 operator*(IntType Scale) const
 	{
@@ -303,10 +302,10 @@ PRAGMA_ENABLE_DEPRECATION_WARNINGS
 	}
 
 	/**
-	 * Gets the result of division on this point.
+	 * Gets the result of division on this vector.
 	 *
-	 * @param Divisor What to divide the point by.
-	 * @return A new divided int point.
+	 * @param Divisor What to divide the vector by.
+	 * @return A new divided int vector.
 	 */
 	TIntVector3 operator/(IntType Divisor) const
 	{
@@ -314,10 +313,10 @@ PRAGMA_ENABLE_DEPRECATION_WARNINGS
 	}
 
 	/**
-	 * Gets the remainder of division on this point.
+	 * Gets the remainder of division on this vector.
 	 *
-	 * @param Divisor What to divide the point by.
-	 * @return A new remainder int point.
+	 * @param Divisor What to divide the vector by.
+	 * @return A new remainder int vector.
 	 */
 	TIntVector3 operator%(IntType Divisor) const
 	{
@@ -325,10 +324,10 @@ PRAGMA_ENABLE_DEPRECATION_WARNINGS
 	}
 
 	/**
-	 * Gets the result of addition on this point.
+	 * Gets the result of addition on this vector.
 	 *
-	 * @param Other The other point to add to this.
-	 * @return A new combined int point.
+	 * @param Other The other vector to add to this.
+	 * @return A new combined int vector.
 	 */
 	TIntVector3 operator+(const TIntVector3& Other) const
 	{
@@ -336,10 +335,10 @@ PRAGMA_ENABLE_DEPRECATION_WARNINGS
 	}
 
 	/**
-	 * Gets the result of subtraction from this point.
+	 * Gets the result of subtraction from this vector.
 	 *
-	 * @param Other The other point to subtract from this.
-	 * @return A new subtracted int point.
+	 * @param Other The other vector to subtract from this.
+	 * @return A new subtracted int vector.
 	 */
 	TIntVector3 operator-(const TIntVector3& Other) const
 	{
@@ -350,7 +349,7 @@ PRAGMA_ENABLE_DEPRECATION_WARNINGS
 	 * Shifts all components to the right.
 	 *
 	 * @param Shift The number of bits to shift.
-	 * @return A new shifted int point.
+	 * @return A new shifted int vector.
 	 */
 	TIntVector3 operator>>(IntType Shift) const
 	{
@@ -361,7 +360,7 @@ PRAGMA_ENABLE_DEPRECATION_WARNINGS
 	 * Shifts all components to the left.
 	 *
 	 * @param Shift The number of bits to shift.
-	 * @return A new shifted int point.
+	 * @return A new shifted int vector.
 	 */
 	TIntVector3 operator<<(IntType Shift) const
 	{
@@ -372,7 +371,7 @@ PRAGMA_ENABLE_DEPRECATION_WARNINGS
 	 * Component-wise AND.
 	 *
 	 * @param Value Number to AND with the each component.
-	 * @return A new shifted int point.
+	 * @return A new shifted int vector.
 	 */
 	TIntVector3 operator&(IntType Value) const
 	{
@@ -383,7 +382,7 @@ PRAGMA_ENABLE_DEPRECATION_WARNINGS
 	 * Component-wise OR.
 	 *
 	 * @param Value Number to OR with the each component.
-	 * @return A new shifted int point.
+	 * @return A new shifted int vector.
 	 */
 	TIntVector3 operator|(IntType Value) const
 	{
@@ -394,7 +393,7 @@ PRAGMA_ENABLE_DEPRECATION_WARNINGS
 	 * Component-wise XOR.
 	 *
 	 * @param Value Number to XOR with the each component.
-	 * @return A new shifted int point.
+	 * @return A new shifted int vector.
 	 */
 	TIntVector3 operator^(IntType Value) const
 	{
@@ -411,9 +410,9 @@ PRAGMA_ENABLE_DEPRECATION_WARNINGS
 	}
 
 	/**
-	 * Gets the maximum value in the point.
+	 * Gets the maximum value in the vector.
 	 *
-	 * @return The maximum value in the point.
+	 * @return The maximum value in the vector.
 	 */
 	IntType GetMax() const
 	{
@@ -421,9 +420,26 @@ PRAGMA_ENABLE_DEPRECATION_WARNINGS
 	}
 
 	/**
-	 * Gets the minimum value in the point.
+	 * Get the maximum absolute value in the vector.
+	 *
+	 * @return The maximum absolute value in the vector.
+	 */
+	IntType GetAbsMax() const
+	{
+		if constexpr (std::is_signed_v<IntType>)
+		{
+			return FMath::Max(FMath::Max(FMath::Abs(X), FMath::Abs(Y)), FMath::Abs(Z));
+		}
+		else
+		{
+			return GetMax();
+		}
+	}
+
+	/**
+	 * Gets the minimum value in the vector.
 	 *
-	 * @return The minimum value in the point.
+	 * @return The minimum value in the vector.
 	 */
 	IntType GetMin() const
 	{
@@ -431,9 +447,54 @@ PRAGMA_ENABLE_DEPRECATION_WARNINGS
 	}
 
 	/**
-	 * Gets the distance of this point from (0,0,0).
+	 * Get the minimum absolute value in the vector.
+	 *
+	 * @return The minimum absolute value in the vector.
+	 */
+	IntType GetAbsMin() const
+	{
+		if constexpr (std::is_signed_v<IntType>)
+		{
+			return FMath::Min(FMath::Min(FMath::Abs(X), FMath::Abs(Y)), FMath::Abs(Z));
+		}
+		else
+		{
+			return GetMin();
+		}
+	}
+
+	/**
+	 * Get the component-wise max of this vector and the parameter vector.
+	 * 
+	 * @param Other The other vector to compare against.
+	 * @return A vector where each component is the maximum value of the corresponding components of the two vectors.
+	 */
+	TIntVector3 ComponentMax(const TIntVector3& Other) const
+	{
+		return TIntVector3(
+			FMath::Max(X, Other.X),
+			FMath::Max(Y, Other.Y),
+			FMath::Max(Z, Other.Z));
+	}
+
+	/**
+	 * Get the component-wise min of this vector and the parameter vector.
+	 * 
+	 * @param Other The other vector to compare against.
+	 * @return A vector where each component is the minimum value of the corresponding components of the two vectors.
+	 */
+	TIntVector3 ComponentMin(const TIntVector3& Other) const
+	{
+		return TIntVector3(
+			FMath::Min(X, Other.X),
+			FMath::Min(Y, Other.Y),
+			FMath::Min(Z, Other.Z));
+	}
+
+	/**
+	 * Gets the distance of this vector from (0,0,0).
 	 *
-	 * @return The distance of this point from (0,0,0).
+	 * @return The distance of this vector from (0,0,0).
 	 */
 	IntType Size() const
 	{
@@ -443,6 +504,42 @@ PRAGMA_ENABLE_DEPRECATION_WARNINGS
 		return IntType(FMath::Sqrt(double(LocalX64 * LocalX64 + LocalY64 * LocalY64 + LocalZ64 * LocalZ64)));
 	}
 
+	/**
+	 * Appends a textual representation of this vector to the output string builder.
+	 *
+	 * @param Out The string builder to append to.
+	 */
+	template <typename CharType>
+	void AppendString(TStringBuilderBase<CharType>& Out) const
+	{
+		Out << "X=" << X << " Y=" << Y << " Z=" << Z;
+	}
+
+	/**
+	 * Appends a textual representation of the parameter vector to the output string builder.
+	 *
+	 * @param Builder The string builder to append to.
+	 * @param Vector The vector to append.
+	 */
+	template <typename CharType>
+	friend TStringBuilderBase<CharType>& operator<<(TStringBuilderBase<CharType>& Builder, const TIntVector3& Vector)
+	{
+		Vector.AppendString(Builder);
+		return Builder;
+	}
+
+	/**
+	 * Appends a textual representation of this vector to the output string.
+	 *
+	 * @param Out The string to append to.
+	 */
+	void AppendString(FString& Out) const
+	{
+		TStringBuilder<128> Builder;
+		Builder << *this;
+		Out.Append(Builder);
+	}
+
 	/**
 	 * Get a textual representation of this vector.
 	 *
@@ -450,19 +547,21 @@ PRAGMA_ENABLE_DEPRECATION_WARNINGS
 	 */
 	FString ToString() const
 	{
-		return FString::Printf(TEXT("X=%s Y=%s Z=%s"), *LexToString(X), *LexToString(Y), *LexToString(Z));
+		FString Out;
+		AppendString(Out);
+		return Out;
 	}
 
 	/**
-	 * Initialize this FIntPoint based on an FString. The String is expected to contain X=, Y=
-	 * The FIntPoint will be bogus when InitFromString returns false.
+	 * Initialize this vector based on an FString. The String is expected to contain X=, Y=, Z=
+	 * The vector will be bogus when InitFromString returns false.
 	 *
-	 * @param	InSourceString	FString containing the color values.
-	 * @return true if the X,Y values were read successfully; false otherwise.
+	 * @param InSourceString FString containing the color values.
+	 * @return true if the X,Y,Z values were read successfully; false otherwise.
 	 */
 	bool InitFromString(const FString& InSourceString)
 	{
-		X = Y = 0;
+		X = Y = Z = 0;
 
 		// The initialization is only successful if the X, Y and Z values can all be parsed from the string
 		const bool bSuccessful = FParse::Value(*InSourceString, TEXT("X="), X) && FParse::Value(*InSourceString, TEXT("Y="), Y) && FParse::Value(*InSourceString, TEXT("Z="), Z);
@@ -471,26 +570,26 @@ PRAGMA_ENABLE_DEPRECATION_WARNINGS
 	}
 
 	/**
-	 * Divide an int point and round up the result.
+	 * Divide an int vector and round up the result.
 	 *
-	 * @param lhs The int point being divided.
-	 * @param Divisor What to divide the int point by.
-	 * @return A new divided int point.
+	 * @param Lhs The int vector being divided.
+	 * @param Divisor What to divide the int vector by.
+	 * @return A new divided int vector.
 	 */
-	static TIntVector3 DivideAndRoundUp(TIntVector3 lhs, IntType Divisor)
+	static TIntVector3 DivideAndRoundUp(TIntVector3 Lhs, IntType Divisor)
 	{
-		return TIntVector3(FMath::DivideAndRoundUp(lhs.X, Divisor), FMath::DivideAndRoundUp(lhs.Y, Divisor), FMath::DivideAndRoundUp(lhs.Z, Divisor));
+		return TIntVector3(FMath::DivideAndRoundUp(Lhs.X, Divisor), FMath::DivideAndRoundUp(Lhs.Y, Divisor), FMath::DivideAndRoundUp(Lhs.Z, Divisor));
 	}
 
-	static TIntVector3 DivideAndRoundUp(TIntVector3 lhs, TIntVector3 Divisor)
+	static TIntVector3 DivideAndRoundUp(TIntVector3 Lhs, TIntVector3 Divisor)
 	{
-		return TIntVector3(FMath::DivideAndRoundUp(lhs.X, Divisor.X), FMath::DivideAndRoundUp(lhs.Y, Divisor.Y), FMath::DivideAndRoundUp(lhs.Z, Divisor.Z));
+		return TIntVector3(FMath::DivideAndRoundUp(Lhs.X, Divisor.X), FMath::DivideAndRoundUp(Lhs.Y, Divisor.Y), FMath::DivideAndRoundUp(Lhs.Z, Divisor.Z));
 	}
 
 	/**
-	 * Gets the number of components a point has.
+	 * Gets the number of components a vector has.
 	 *
-	 * @return Number of components point has.
+	 * @return Number of components vector has.
 	 */
 	static int32 Num()
 	{
@@ -504,7 +603,7 @@ PRAGMA_ENABLE_DEPRECATION_WARNINGS
 	 * @param Vector The vector to serialize.
 	 * @return Reference to the Archive after serialization.
 	 */
-	friend FArchive& operator<<( FArchive& Ar, TIntVector3& Vector )
+	friend FArchive& operator<<(FArchive& Ar, TIntVector3& Vector)
 	{
 		return Ar << Vector.X << Vector.Y << Vector.Z;
 	}
@@ -517,7 +616,7 @@ PRAGMA_ENABLE_DEPRECATION_WARNINGS
 		Record << SA_VALUE(TEXT("Z"), Vector.Z);
 	}
 
-	bool Serialize( FArchive& Ar )
+	bool Serialize(FArchive& Ar)
 	{
 		Ar << *this;
 		return true;
@@ -526,54 +625,6 @@ PRAGMA_ENABLE_DEPRECATION_WARNINGS
 	bool SerializeFromMismatchedTag(FName StructTag, FArchive& Ar);
 };
 
-template <>
-inline FString TIntVector3<int64>::ToString() const
-{
-	return FString::Printf(TEXT("X=%lld Y=%lld Z=%lld"), X, Y, Z);
-}
-
-template <>
-inline FString TIntVector3<int32>::ToString() const
-{
-	return FString::Printf(TEXT("X=%d Y=%d Z=%d"), X, Y, Z);
-}
-
-template <>
-inline FString TIntVector3<int16>::ToString() const
-{
-	return FString::Printf(TEXT("X=%d Y=%d Z=%d"), X, Y, Z);
-}
-
-template <>
-inline FString TIntVector3<int8>::ToString() const
-{
-	return FString::Printf(TEXT("X=%d Y=%d Z=%d"), X, Y, Z);
-}
-
-template <>
-inline FString TIntVector3<uint64>::ToString() const
-{
-	return FString::Printf(TEXT("X=%llu Y=%llu Z=%llu"), X, Y, Z);
-}
-
-template <>
-inline FString TIntVector3<uint32>::ToString() const
-{
-	return FString::Printf(TEXT("X=%u Y=%u Z=%u"), X, Y, Z);
-}
-
-template <>
-inline FString TIntVector3<uint16>::ToString() const
-{
-	return FString::Printf(TEXT("X=%u Y=%u Z=%u"), X, Y, Z);
-}
-
-template <>
-inline FString TIntVector3<uint8>::ToString() const
-{
-	return FString::Printf(TEXT("X=%u Y=%u Z=%u"), X, Y, Z);
-}
-
 template <typename IntType>
 const TIntVector3<IntType> TIntVector3<IntType>::ZeroValue(0, 0, 0);
 
@@ -592,17 +643,21 @@ struct TIntVector2
 	{
 		struct
 		{
-			IntType X, Y;
+			/** Holds the vector's x-coordinate. */
+			IntType X;
+
+			/** Holds the vector's y-coordinate. */
+			IntType Y;
 		};
 
 		UE_DEPRECATED(all, "For internal use only")
 		IntType XY[2];
 	};
 
-	/** An int point with zeroed values. */
+	/** An int vector with zeroed values. */
 	static const TIntVector2 ZeroValue;
 
-	/** An int point with INDEX_NONE values. */
+	/** An int vector with INDEX_NONE values. */
 	static const TIntVector2 NoneValue;
 
 	TIntVector2() = default;
@@ -670,8 +725,8 @@ PRAGMA_ENABLE_DEPRECATION_WARNINGS
 	/**
 	 * Multiplies this vector with another vector, using component-wise multiplication.
 	 *
-	 * @param Other The point to multiply with.
-	 * @return Reference to this point after multiplication.
+	 * @param Other The vector to multiply with.
+	 * @return Reference to this vector after multiplication.
 	 */
 	TIntVector2& operator*=(const TIntVector2& Other)
 	{
@@ -682,10 +737,10 @@ PRAGMA_ENABLE_DEPRECATION_WARNINGS
 	}
 
 	/**
-	 * Scales this point.
+	 * Scales this vector.
 	 *
-	 * @param Scale What to multiply the point by.
-	 * @return Reference to this point after multiplication.
+	 * @param Scale What to multiply the vector by.
+	 * @return Reference to this vector after multiplication.
 	 */
 	TIntVector2& operator*=(IntType Scale)
 	{
@@ -696,10 +751,10 @@ PRAGMA_ENABLE_DEPRECATION_WARNINGS
 	}
 
 	/**
-	 * Divides this point.
+	 * Divides this vector.
 	 *
-	 * @param Divisor What to divide the point by.
-	 * @return Reference to this point after division.
+	 * @param Divisor What to divide the vector by.
+	 * @return Reference to this vector after division.
 	 */
 	TIntVector2& operator/=(IntType Divisor)
 	{
@@ -710,10 +765,10 @@ PRAGMA_ENABLE_DEPRECATION_WARNINGS
 	}
 
 	/**
-	 * Remainder of division of this point.
+	 * Remainder of division of this vector.
 	 *
-	 * @param Divisor What to divide the point by.
-	 * @return Reference to this point after remainder.
+	 * @param Divisor What to divide the vector by.
+	 * @return Reference to this vector after remainder.
 	 */
 	TIntVector2& operator%=(IntType Divisor)
 	{
@@ -724,10 +779,10 @@ PRAGMA_ENABLE_DEPRECATION_WARNINGS
 	}
 
 	/**
-	 * Adds to this point.
+	 * Adds to this vector.
 	 *
-	 * @param Other The point to add to this point.
-	 * @return Reference to this point after addition.
+	 * @param Other The vector to add to this vector.
+	 * @return Reference to this vector after addition.
 	 */
 	TIntVector2& operator+=(const TIntVector2& Other)
 	{
@@ -738,10 +793,10 @@ PRAGMA_ENABLE_DEPRECATION_WARNINGS
 	}
 
 	/**
-	 * Subtracts from this point.
+	 * Subtracts from this vector.
 	 *
-	 * @param Other The point to subtract from this point.
-	 * @return Reference to this point after subtraction.
+	 * @param Other The vector to subtract from this vector.
+	 * @return Reference to this vector after subtraction.
 	 */
 	TIntVector2& operator-=(const TIntVector2& Other)
 	{
@@ -752,9 +807,9 @@ PRAGMA_ENABLE_DEPRECATION_WARNINGS
 	}
 
 	/**
-	 * Gets the result of component-wise multiplication of this point by another.
+	 * Gets the result of component-wise multiplication of this vector by another.
 	 *
-	 * @param Other The point to multiply with.
+	 * @param Other The vector to multiply with.
 	 * @return The result of multiplication.
 	 */
 	TIntVector2 operator*(const TIntVector2& Other) const
@@ -763,10 +818,10 @@ PRAGMA_ENABLE_DEPRECATION_WARNINGS
 	}
 
 	/**
-	 * Gets the result of scaling on this point.
+	 * Gets the result of scaling on this vector.
 	 *
-	 * @param Scale What to multiply the point by.
-	 * @return A new scaled int point.
+	 * @param Scale What to multiply the vector by.
+	 * @return A new scaled int vector.
 	 */
 	TIntVector2 operator*(IntType Scale) const
 	{
@@ -774,10 +829,10 @@ PRAGMA_ENABLE_DEPRECATION_WARNINGS
 	}
 
 	/**
-	 * Gets the result of division on this point.
+	 * Gets the result of division on this vector.
 	 *
-	 * @param Divisor What to divide the point by.
-	 * @return A new divided int point.
+	 * @param Divisor What to divide the vector by.
+	 * @return A new divided int vector.
 	 */
 	TIntVector2 operator/(IntType Divisor) const
 	{
@@ -785,10 +840,10 @@ PRAGMA_ENABLE_DEPRECATION_WARNINGS
 	}
 
 	/**
-	 * Gets the remainder of division on this point.
+	 * Gets the remainder of division on this vector.
 	 *
-	 * @param Divisor What to divide the point by.
-	 * @return A new remainder int point.
+	 * @param Divisor What to divide the vector by.
+	 * @return A new remainder int vector.
 	 */
 	TIntVector2 operator%(IntType Divisor) const
 	{
@@ -796,10 +851,10 @@ PRAGMA_ENABLE_DEPRECATION_WARNINGS
 	}
 
 	/**
-	 * Gets the result of addition on this point.
+	 * Gets the result of addition on this vector.
 	 *
-	 * @param Other The other point to add to this.
-	 * @return A new combined int point.
+	 * @param Other The other vector to add to this.
+	 * @return A new combined int vector.
 	 */
 	TIntVector2 operator+(const TIntVector2& Other) const
 	{
@@ -807,16 +862,197 @@ PRAGMA_ENABLE_DEPRECATION_WARNINGS
 	}
 
 	/**
-	 * Gets the result of subtraction from this point.
+	 * Gets the result of subtraction from this vector.
 	 *
-	 * @param Other The other point to subtract from this.
-	 * @return A new subtracted int point.
+	 * @param Other The other vector to subtract from this.
+	 * @return A new subtracted int vector.
 	 */
 	TIntVector2 operator-(const TIntVector2& Other) const
 	{
 		return TIntVector2(*this) -= Other;
 	}
 
+	/**
+	 * Is vector equal to zero.
+	 * @return is zero
+	*/
+	bool IsZero() const
+	{
+		return *this == ZeroValue;
+	}
+
+	/**
+	 * Gets the maximum value in the vector.
+	 *
+	 * @return The maximum value in the vector.
+	 */
+	IntType GetMax() const
+	{
+		return FMath::Max(X, Y);
+	}
+
+	/**
+	 * Get the maximum absolute value in the vector.
+	 *
+	 * @return The maximum absolute value in the vector.
+	 */
+	IntType GetAbsMax() const
+	{
+		if constexpr (std::is_signed_v<IntType>)
+		{
+			return FMath::Max(FMath::Abs(X), FMath::Abs(Y));
+		}
+		else
+		{
+			return GetMax();
+		}
+	}
+
+	/**
+	 * Gets the minimum value in the vector.
+	 *
+	 * @return The minimum value in the vector.
+	 */
+	IntType GetMin() const
+	{
+		return FMath::Min(X, Y);
+	}
+
+	/**
+	 * Get the minimum absolute value in the vector.
+	 *
+	 * @return The minimum absolute value in the vector.
+	 */
+	IntType GetAbsMin() const
+	{
+		if constexpr (std::is_signed_v<IntType>)
+		{
+			return FMath::Min(FMath::Abs(X), FMath::Abs(Y));
+		}
+		else
+		{
+			return GetMin();
+		}
+	}
+
+	/**
+	 * Get the component-wise max of this vector and the parameter vector.
+	 * 
+	 * @param Other The other vector to compare against.
+	 * @return A vector where each component is the maximum value of the corresponding components of the two vectors.
+	 */
+	TIntVector2 ComponentMax(const TIntVector2& Other) const
+	{
+		return TIntVector2(
+			FMath::Max(X, Other.X),
+			FMath::Max(Y, Other.Y));
+	}
+
+	/**
+	 * Get the component-wise min of this vector and the parameter vector.
+	 * 
+	 * @param Other The other vector to compare against.
+	 * @return A vector where each component is the minimum value of the corresponding components of the two vectors.
+	 */
+	TIntVector2 ComponentMin(const TIntVector2& Other) const
+	{
+		return TIntVector2(
+			FMath::Min(X, Other.X),
+			FMath::Min(Y, Other.Y));
+	}
+
+	/**
+	 * Appends a textual representation of this vector to the output string builder.
+	 *
+	 * @param Out The string builder to append to.
+	 */
+	template <typename CharType>
+	void AppendString(TStringBuilderBase<CharType>& Out) const
+	{
+		Out << "X=" << X << " Y=" << Y;
+	}
+
+	/**
+	 * Appends a textual representation of the parameter vector to the output string builder.
+	 *
+	 * @param Builder The string builder to append to.
+	 * @param Vector The vector to append.
+	 */
+	template <typename CharType>
+	friend TStringBuilderBase<CharType>& operator<<(TStringBuilderBase<CharType>& Builder, const TIntVector2& Vector)
+	{
+		Vector.AppendString(Builder);
+		return Builder;
+	}
+
+	/**
+	 * Appends a textual representation of this vector to the output string.
+	 *
+	 * @param Out The string to append to.
+	 */
+	void AppendString(FString& Out) const
+	{
+		TStringBuilder<128> Builder;
+		Builder << *this;
+		Out.Append(Builder);
+	}
+
+	/**
+	 * Get a textual representation of this vector.
+	 *
+	 * @return A string describing the vector.
+	 */
+	FString ToString() const
+	{
+		FString Out;
+		AppendString(Out);
+		return Out;
+	}
+
+	/**
+	 * Initialize this FIntVector based on an FString. The String is expected to contain X=, Y=
+	 * The FIntVector will be bogus when InitFromString returns false.
+	 *
+	 * @param InSourceString FString containing the color values.
+	 * @return true if the X,Y values were read successfully; false otherwise.
+	 */
+	bool InitFromString(const FString& InSourceString)
+	{
+		X = Y = 0;
+
+		// The initialization is only successful if the X and Y values can all be parsed from the string
+		const bool bSuccessful = FParse::Value(*InSourceString, TEXT("X="), X) && FParse::Value(*InSourceString, TEXT("Y="), Y);
+
+		return bSuccessful;
+	}
+
+	/**
+	 * Divide an int vector and round up the result.
+	 *
+	 * @param Lhs The int vector being divided.
+	 * @param Divisor What to divide the int vector by.
+	 * @return A new divided int vector.
+	 */
+	static TIntVector2 DivideAndRoundUp(TIntVector2 Lhs, IntType Divisor)
+	{
+		return TIntVector2(FMath::DivideAndRoundUp(Lhs.X, Divisor), FMath::DivideAndRoundUp(Lhs.Y, Divisor));
+	}
+
+	static TIntVector2 DivideAndRoundUp(TIntVector2 Lhs, TIntVector2 Divisor)
+	{
+		return TIntVector2(FMath::DivideAndRoundUp(Lhs.X, Divisor.X), FMath::DivideAndRoundUp(Lhs.Y, Divisor.Y));
+	}
+
+	/**
+	 * Gets the number of components a vector has.
+	 *
+	 * @return Number of components vector has.
+	 */
+	static int32 Num()
+	{
+		return 2;
+	}
+
 	/**
 	 * Serializes the Vector2.
 	 *
@@ -829,6 +1065,13 @@ PRAGMA_ENABLE_DEPRECATION_WARNINGS
 		return Ar << Vector.X << Vector.Y;
 	}
 
+	friend void operator<<(FStructuredArchive::FSlot Slot, TIntVector2& Vector)
+	{
+		FStructuredArchive::FRecord Record = Slot.EnterRecord();
+		Record << SA_VALUE(TEXT("X"), Vector.X);
+		Record << SA_VALUE(TEXT("Y"), Vector.Y);
+	}
+
 	bool Serialize(FArchive& Ar)
 	{
 		Ar << *this;
@@ -856,13 +1099,29 @@ struct TIntVector4
 	{
 		struct
 		{
-			IntType X, Y, Z, W;
+			/** Holds the vector's x-coordinate. */
+			IntType X;
+
+			/** Holds the vector's y-coordinate. */
+			IntType Y;
+
+			/** Holds the vector's z-coordinate. */
+			IntType Z;
+			
+			/** Holds the vector's w-coordinate. */
+			IntType W;
 		};
 
 		UE_DEPRECATED(all, "For internal use only")
 		IntType XYZW[4];
 	};
 
+	/** An int vector with zeroed values. */
+	static const TIntVector4 ZeroValue;
+
+	/** An int vector with INDEX_NONE values. */
+	static const TIntVector4 NoneValue;
+
 	TIntVector4() = default;
 
 	TIntVector4(IntType InX, IntType InY, IntType InZ, IntType InW)
@@ -918,9 +1177,9 @@ PRAGMA_DISABLE_DEPRECATION_WARNINGS
 PRAGMA_ENABLE_DEPRECATION_WARNINGS
 
 	/**
-	 * Gets specific component of a point.
+	 * Gets specific component of a vector.
 	 *
-	 * @param ComponentIndex Index of point component.
+	 * @param ComponentIndex Index of vector component.
 	 * @return const reference to component.
 	 */
 	const IntType& operator()(int32 ComponentIndex) const
@@ -931,9 +1190,9 @@ PRAGMA_ENABLE_DEPRECATION_WARNINGS
 	}
 
 	/**
-	 * Gets specific component of a point.
+	 * Gets specific component of a vector.
 	 *
-	 * @param ComponentIndex Index of point component.
+	 * @param ComponentIndex Index of vector component.
 	 * @return reference to component.
 	 */
 	IntType& operator()(int32 ComponentIndex)
@@ -944,9 +1203,9 @@ PRAGMA_ENABLE_DEPRECATION_WARNINGS
 	}
 
 	/**
-	 * Gets specific component of a point.
+	 * Gets specific component of a vector.
 	 *
-	 * @param ComponentIndex Index of point component.
+	 * @param ComponentIndex Index of vector component.
 	 * @return const reference to component.
 	 */
 	const IntType& operator[](int32 ComponentIndex) const
@@ -957,9 +1216,9 @@ PRAGMA_ENABLE_DEPRECATION_WARNINGS
 	}
 
 	/**
-	 * Gets specific component of a point.
+	 * Gets specific component of a vector.
 	 *
-	 * @param ComponentIndex Index of point component.
+	 * @param ComponentIndex Index of vector component.
 	 * @return reference to component.
 	 */
 	IntType& operator[](int32 ComponentIndex)
@@ -970,10 +1229,10 @@ PRAGMA_ENABLE_DEPRECATION_WARNINGS
 	}
 
 	/**
-	 * Compares points for equality.
+	 * Compares vectors for equality.
 	 *
-	 * @param Other The other int point being compared.
-	 * @return true if the points are equal, false otherwise..
+	 * @param Other The other int vector being compared.
+	 * @return true if the vectors are equal, false otherwise..
 	 */
 	bool operator==(const TIntVector4& Other) const
 	{
@@ -981,10 +1240,10 @@ PRAGMA_ENABLE_DEPRECATION_WARNINGS
 	}
 
 	/**
-	 * Compares points for inequality.
+	 * Compares vectors for inequality.
 	 *
-	 * @param Other The other int point being compared.
-	 * @return true if the points are not equal, false otherwise..
+	 * @param Other The other int vector being compared.
+	 * @return true if the vectors are not equal, false otherwise..
 	 */
 	bool operator!=(const TIntVector4& Other) const
 	{
@@ -994,8 +1253,8 @@ PRAGMA_ENABLE_DEPRECATION_WARNINGS
 	/**
 	 * Multiplies this vector with another vector, using component-wise multiplication.
 	 *
-	 * @param Other The point to multiply with.
-	 * @return Reference to this point after multiplication.
+	 * @param Other The vector to multiply with.
+	 * @return Reference to this vector after multiplication.
 	 */
 	TIntVector4& operator*=(const TIntVector4& Other)
 	{
@@ -1008,10 +1267,10 @@ PRAGMA_ENABLE_DEPRECATION_WARNINGS
 	}
 
 	/**
-	 * Scales this point.
+	 * Scales this vector.
 	 *
-	 * @param Scale What to multiply the point by.
-	 * @return Reference to this point after multiplication.
+	 * @param Scale What to multiply the vector by.
+	 * @return Reference to this vector after multiplication.
 	 */
 	TIntVector4& operator*=(IntType Scale)
 	{
@@ -1024,10 +1283,10 @@ PRAGMA_ENABLE_DEPRECATION_WARNINGS
 	}
 
 	/**
-	 * Divides this point.
+	 * Divides this vector.
 	 *
-	 * @param Divisor What to divide the point by.
-	 * @return Reference to this point after division.
+	 * @param Divisor What to divide the vector by.
+	 * @return Reference to this vector after division.
 	 */
 	TIntVector4& operator/=(IntType Divisor)
 	{
@@ -1040,10 +1299,10 @@ PRAGMA_ENABLE_DEPRECATION_WARNINGS
 	}
 
 	/**
-	 * Remainder of division of this point.
+	 * Remainder of division of this vector.
 	 *
-	 * @param Divisor What to divide the point by.
-	 * @return Reference to this point after remainder.
+	 * @param Divisor What to divide the vector by.
+	 * @return Reference to this vector after remainder.
 	 */
 	TIntVector4& operator%=(IntType Divisor)
 	{
@@ -1056,10 +1315,10 @@ PRAGMA_ENABLE_DEPRECATION_WARNINGS
 	}
 
 	/**
-	 * Adds to this point.
+	 * Adds to this vector.
 	 *
-	 * @param Other The point to add to this point.
-	 * @return Reference to this point after addition.
+	 * @param Other The vector to add to this vector.
+	 * @return Reference to this vector after addition.
 	 */
 	TIntVector4& operator+=(const TIntVector4& Other)
 	{
@@ -1072,10 +1331,10 @@ PRAGMA_ENABLE_DEPRECATION_WARNINGS
 	}
 
 	/**
-	 * Subtracts from this point.
+	 * Subtracts from this vector.
 	 *
-	 * @param Other The point to subtract from this point.
-	 * @return Reference to this point after subtraction.
+	 * @param Other The vector to subtract from this vector.
+	 * @return Reference to this vector after subtraction.
 	 */
 	TIntVector4& operator-=(const TIntVector4& Other)
 	{
@@ -1088,9 +1347,9 @@ PRAGMA_ENABLE_DEPRECATION_WARNINGS
 	}
 
 	/**
-	 * Gets the result of component-wise multiplication of this point by another.
+	 * Gets the result of component-wise multiplication of this vector by another.
 	 *
-	 * @param Other The point to multiply with.
+	 * @param Other The vector to multiply with.
 	 * @return The result of multiplication.
 	 */
 	TIntVector4 operator*(const TIntVector4& Other) const
@@ -1099,10 +1358,10 @@ PRAGMA_ENABLE_DEPRECATION_WARNINGS
 	}
 
 	/**
-	 * Gets the result of scaling on this point.
+	 * Gets the result of scaling on this vector.
 	 *
-	 * @param Scale What to multiply the point by.
-	 * @return A new scaled int point.
+	 * @param Scale What to multiply the vector by.
+	 * @return A new scaled int vector.
 	 */
 	TIntVector4 operator*(IntType Scale) const
 	{
@@ -1110,10 +1369,10 @@ PRAGMA_ENABLE_DEPRECATION_WARNINGS
 	}
 
 	/**
-	 * Gets the result of division on this point.
+	 * Gets the result of division on this vector.
 	 *
-	 * @param Divisor What to divide the point by.
-	 * @return A new divided int point.
+	 * @param Divisor What to divide the vector by.
+	 * @return A new divided int vector.
 	 */
 	TIntVector4 operator/(IntType Divisor) const
 	{
@@ -1121,10 +1380,10 @@ PRAGMA_ENABLE_DEPRECATION_WARNINGS
 	}
 
 	/**
-	 * Gets the remainder of division on this point.
+	 * Gets the remainder of division on this vector.
 	 *
-	 * @param Divisor What to divide the point by.
-	 * @return A new remainder int point.
+	 * @param Divisor What to divide the vector by.
+	 * @return A new remainder int vector.
 	 */
 	TIntVector4 operator%(IntType Divisor) const
 	{
@@ -1132,10 +1391,10 @@ PRAGMA_ENABLE_DEPRECATION_WARNINGS
 	}
 
 	/**
-	 * Gets the result of addition on this point.
+	 * Gets the result of addition on this vector.
 	 *
-	 * @param Other The other point to add to this.
-	 * @return A new combined int point.
+	 * @param Other The other vector to add to this.
+	 * @return A new combined int vector.
 	 */
 	TIntVector4 operator+(const TIntVector4& Other) const
 	{
@@ -1143,10 +1402,10 @@ PRAGMA_ENABLE_DEPRECATION_WARNINGS
 	}
 
 	/**
-	 * Gets the result of subtraction from this point.
+	 * Gets the result of subtraction from this vector.
 	 *
-	 * @param Other The other point to subtract from this.
-	 * @return A new subtracted int point.
+	 * @param Other The other vector to subtract from this.
+	 * @return A new subtracted int vector.
 	 */
 	TIntVector4 operator-(const TIntVector4& Other) const
 	{
@@ -1157,7 +1416,7 @@ PRAGMA_ENABLE_DEPRECATION_WARNINGS
 	 * Shifts all components to the right.
 	 *
 	 * @param Shift The number of bits to shift.
-	 * @return A new shifted int point.
+	 * @return A new shifted int vector.
 	 */
 	TIntVector4 operator>>(IntType Shift) const
 	{
@@ -1168,7 +1427,7 @@ PRAGMA_ENABLE_DEPRECATION_WARNINGS
 	 * Shifts all components to the left.
 	 *
 	 * @param Shift The number of bits to shift.
-	 * @return A new shifted int point.
+	 * @return A new shifted int vector.
 	 */
 	TIntVector4 operator<<(IntType Shift) const
 	{
@@ -1179,7 +1438,7 @@ PRAGMA_ENABLE_DEPRECATION_WARNINGS
 	 * Component-wise AND.
 	 *
 	 * @param Value Number to AND with the each component.
-	 * @return A new shifted int point.
+	 * @return A new shifted int vector.
 	 */
 	TIntVector4 operator&(IntType Value) const
 	{
@@ -1190,7 +1449,7 @@ PRAGMA_ENABLE_DEPRECATION_WARNINGS
 	 * Component-wise OR.
 	 *
 	 * @param Value Number to OR with the each component.
-	 * @return A new shifted int point.
+	 * @return A new shifted int vector.
 	 */
 	TIntVector4 operator|(IntType Value) const
 	{
@@ -1201,7 +1460,7 @@ PRAGMA_ENABLE_DEPRECATION_WARNINGS
 	 * Component-wise XOR.
 	 *
 	 * @param Value Number to XOR with the each component.
-	 * @return A new shifted int point.
+	 * @return A new shifted int vector.
 	 */
 	TIntVector4 operator^(IntType Value) const
 	{
@@ -1209,12 +1468,205 @@ PRAGMA_ENABLE_DEPRECATION_WARNINGS
 	}
 
 	/**
-	* Serializes the Vector4.
-	*
-	* @param Ar The archive to serialize into.
-	* @param Vector The vector to serialize.
-	* @return Reference to the Archive after serialization.
+	 * Is vector equal to zero.
+	 * @return is zero
 	*/
+	bool IsZero() const
+	{
+		return *this == ZeroValue;
+	}
+
+	/**
+	 * Gets the maximum value in the vector.
+	 *
+	 * @return The maximum value in the vector.
+	 */
+	IntType GetMax() const
+	{
+		return FMath::Max(FMath::Max(FMath::Max(X, Y), Z), W);
+	}
+
+	/**
+	 * Get the maximum absolute value in the vector.
+	 *
+	 * @return The maximum absolute value in the vector.
+	 */
+	IntType GetAbsMax() const
+	{
+		if constexpr (std::is_signed_v<IntType>)
+		{
+			return FMath::Max(FMath::Max(FMath::Max(FMath::Abs(X), FMath::Abs(Y)), FMath::Abs(Z)), FMath::Abs(W));
+		}
+		else
+		{
+			return GetMax();
+		}
+	}
+
+	/**
+	 * Gets the minimum value in the vector.
+	 *
+	 * @return The minimum value in the vector.
+	 */
+	IntType GetMin() const
+	{
+		return FMath::Min(FMath::Min(FMath::Min(X, Y), Z), W);
+	}
+
+	/**
+	 * Get the minimum absolute value in the vector.
+	 *
+	 * @return The minimum absolute value in the vector.
+	 */
+	IntType GetAbsMin() const
+	{
+		if constexpr (std::is_signed_v<IntType>)
+		{
+			return FMath::Min(FMath::Min(FMath::Min(FMath::Abs(X), FMath::Abs(Y)), FMath::Abs(Z)), FMath::Abs(W));
+		}
+		else
+		{
+			return GetMin();
+		}
+	}
+
+	/**
+	 * Get the component-wise max of this vector and the parameter vector.
+	 * 
+	 * @param Other The other vector to compare against.
+	 * @return A vector where each component is the maximum value of the corresponding components of the two vectors.
+	 */
+	TIntVector4 ComponentMax(const TIntVector4& Other) const
+	{
+		return TIntVector4(
+			FMath::Max(X, Other.X),
+			FMath::Max(Y, Other.Y),
+			FMath::Max(Z, Other.Z),
+			FMath::Max(W, Other.W));
+	}
+
+	/**
+	 * Get the component-wise min of this vector and the parameter vector.
+	 * 
+	 * @param Other The other vector to compare against.
+	 * @return A vector where each component is the minimum value of the corresponding components of the two vectors.
+	 */
+	TIntVector4 ComponentMin(const TIntVector4& Other) const
+	{
+		return TIntVector4(
+			FMath::Min(X, Other.X),
+			FMath::Min(Y, Other.Y),
+			FMath::Min(Z, Other.Z),
+			FMath::Min(W, Other.W));
+	}
+
+	/**
+	 * Appends a textual representation of this vector to the output string builder.
+	 *
+	 * @param Out The string builder to append to.
+	 */
+	template <typename CharType>
+	void AppendString(TStringBuilderBase<CharType>& Out) const
+	{
+		Out << "X=" << X << " Y=" << Y << " Z=" << Z << " W=" << W;
+	}
+
+	/**
+	 * Appends a textual representation of the parameter vector to the output string builder.
+	 *
+	 * @param Builder The string builder to append to.
+	 * @param Vector The vector to append.
+	 */
+	template <typename CharType>
+	friend TStringBuilderBase<CharType>& operator<<(TStringBuilderBase<CharType>& Builder, const TIntVector4& Vector)
+	{
+		Vector.AppendString(Builder);
+		return Builder;
+	}
+
+	/**
+	 * Appends a textual representation of this vector to the output string.
+	 *
+	 * @param Out The string to append to.
+	 */
+	void AppendString(FString& Out) const
+	{
+		TStringBuilder<128> Builder;
+		Builder << *this;
+		Out.Append(Builder);
+	}
+
+	/**
+	 * Get a textual representation of this vector.
+	 *
+	 * @return A string describing the vector.
+	 */
+	FString ToString() const
+	{
+		FString Out;
+		AppendString(Out);
+		return Out;
+	}
+
+	/**
+	 * Initialize this vector based on an FString. The String is expected to contain X=, Y=, Z=, W=
+	 * The vector will be bogus when InitFromString returns false.
+	 *
+	 * @param InSourceString FString containing the color values.
+	 * @return true if the X,Y,Z,W values were read successfully; false otherwise.
+	 */
+	bool InitFromString(const FString& InSourceString)
+	{
+		X = Y = Z = W = 0;
+
+		// The initialization is only successful if the X, Y, Z and W values can all be parsed from the string
+		const bool bSuccessful = FParse::Value(*InSourceString, TEXT("X="), X) && FParse::Value(*InSourceString, TEXT("Y="), Y) && FParse::Value(*InSourceString, TEXT("Z="), Z) && FParse::Value(*InSourceString, TEXT("W="), W);
+
+		return bSuccessful;
+	}
+
+	/**
+	 * Divide an int vector and round up the result.
+	 *
+	 * @param Lhs The int vector being divided.
+	 * @param Divisor What to divide the int vector by.
+	 * @return A new divided int vector.
+	 */
+	static TIntVector4 DivideAndRoundUp(TIntVector4 Lhs, IntType Divisor)
+	{
+		return TIntVector4(
+			FMath::DivideAndRoundUp(Lhs.X, Divisor),
+			FMath::DivideAndRoundUp(Lhs.Y, Divisor),
+			FMath::DivideAndRoundUp(Lhs.Z, Divisor),
+			FMath::DivideAndRoundUp(Lhs.W, Divisor));
+	}
+
+	static TIntVector4 DivideAndRoundUp(TIntVector4 Lhs, TIntVector4 Divisor)
+	{
+		return TIntVector4(
+			FMath::DivideAndRoundUp(Lhs.X, Divisor.X),
+			FMath::DivideAndRoundUp(Lhs.Y, Divisor.Y),
+			FMath::DivideAndRoundUp(Lhs.Z, Divisor.Z),
+			FMath::DivideAndRoundUp(Lhs.W, Divisor.W));
+	}
+
+	/**
+	 * Gets the number of components a vector has.
+	 *
+	 * @return Number of components vector has.
+	 */
+	static int32 Num()
+	{
+		return 4;
+	}
+
+	/**
+	 * Serializes the Vector4.
+	 *
+	 * @param Ar The archive to serialize into.
+	 * @param Vector The vector to serialize.
+	 * @return Reference to the Archive after serialization.
+	 */
 	friend FArchive& operator<<(FArchive& Ar, TIntVector4& Vector)
 	{
 		return Ar << Vector.X << Vector.Y << Vector.Z << Vector.W;
@@ -1238,6 +1690,12 @@ PRAGMA_ENABLE_DEPRECATION_WARNINGS
 	bool SerializeFromMismatchedTag(FName StructTag, FArchive& Ar);
 };
 
+template <typename IntType>
+const TIntVector4<IntType> TIntVector4<IntType>::ZeroValue(0, 0, 0, 0);
+
+template <typename IntType>
+const TIntVector4<IntType> TIntVector4<IntType>::NoneValue(INDEX_NONE, INDEX_NONE, INDEX_NONE, INDEX_NONE);
+
 /**
  * Creates a hash value from an IntVector2.
  *
diff --git a/Engine/Source/Runtime/Core/Tests/Math/IntVectorTest.cpp b/Engine/Source/Runtime/Core/Tests/Math/IntVectorTest.cpp
new file mode 100644
--- /dev/null
+++ b/Engine/Source/Runtime/Core/Tests/Math/IntVectorTest.cpp
@@ -0,0 +1,387 @@
+// Copyright Epic Games, Inc. All Rights Reserved.
+
+#if WITH_TESTS
+
+#include "CoreTypes.h"
+#include "Math/IntVector.h"
+#include "Tests/TestHarnessAdapter.h"
+
+// UE::Math::TIntVector2
+
+template <typename IntType>
+struct TIntVector2TestData
+{
+};
+
+template <>
+struct TIntVector2TestData<int32>
+{
+	static inline FStringView ZeroString = TEXTVIEW("X=0 Y=0");
+	static inline FStringView MinMaxString = TEXTVIEW("X=-2147483648 Y=2147483647");
+};
+
+template <>
+struct TIntVector2TestData<int64>
+{
+	static inline FStringView ZeroString = TEXTVIEW("X=0 Y=0");
+	static inline FStringView MinMaxString = TEXTVIEW("X=-9223372036854775808 Y=9223372036854775807");
+};
+
+template <>
+struct TIntVector2TestData<uint32>
+{
+	static inline FStringView ZeroString = TEXTVIEW("X=0 Y=0");
+	static inline FStringView MinMaxString = TEXTVIEW("X=0 Y=4294967295");
+};
+
+template <>
+struct TIntVector2TestData<uint64>
+{
+	static inline FStringView ZeroString = TEXTVIEW("X=0 Y=0");
+	static inline FStringView MinMaxString = TEXTVIEW("X=0 Y=18446744073709551615");
+};
+
+template <typename IntType>
+void TestIntVector2()
+{
+	using FIntVector2Type = UE::Math::TIntVector2<IntType>;
+	using FIntVector2TestData = TIntVector2TestData<IntType>;
+
+	constexpr IntType MinValue = TNumericLimits<IntType>::Min();
+	constexpr IntType MaxValue = TNumericLimits<IntType>::Max();
+
+	// GetMax
+	CHECK(FIntVector2Type(1, 2).GetMax() == 2);
+	if constexpr (std::is_signed_v<IntType>)
+	{
+		CHECK(FIntVector2Type( 2, -2).GetMax() == 2);
+		CHECK(FIntVector2Type(-2, -4).GetMax() == -2);
+	}
+
+	// GetAbsMax
+	CHECK(FIntVector2Type(1, 2).GetAbsMax() == 2);
+	if constexpr (std::is_signed_v<IntType>)
+	{
+		CHECK(FIntVector2Type( 2, -2).GetAbsMax() == 2);
+		CHECK(FIntVector2Type(-2, -4).GetAbsMax() == 4);
+	}
+
+	// GetMin
+	CHECK(FIntVector2Type(1, 2).GetMin() == 1);
+	if constexpr (std::is_signed_v<IntType>)
+	{
+		CHECK(FIntVector2Type( 2, -2).GetMin() == -2);
+		CHECK(FIntVector2Type(-2, -4).GetMin() == -4);
+	}
+
+	// GetAbsMin
+	CHECK(FIntVector2Type(1, 2).GetAbsMin() == 1);
+	if constexpr (std::is_signed_v<IntType>)
+	{
+		CHECK(FIntVector2Type( 2, -2).GetAbsMin() == 2);
+		CHECK(FIntVector2Type(-2, -4).GetAbsMin() == 2);
+	}
+
+	// ComponentMax
+	CHECK(FIntVector2Type(1, 2).ComponentMax(FIntVector2Type(2, 1)) == FIntVector2Type(2, 2));
+	if constexpr (std::is_signed_v<IntType>)
+	{
+		CHECK(FIntVector2Type(-1, -2).ComponentMax(FIntVector2Type(-2, -1)) == FIntVector2Type(-1, -1));
+	}
+
+	// ComponentMin
+	CHECK(FIntVector2Type(1, 2).ComponentMin(FIntVector2Type(2, 1)) == FIntVector2Type(1, 1));
+	if constexpr (std::is_signed_v<IntType>)
+	{
+		CHECK(FIntVector2Type(-1, -2).ComponentMin(FIntVector2Type(-2, -1)) == FIntVector2Type(-2, -2));
+	}
+
+	// AppendString(TStringBuilderBase<TCHAR>) via operator<<(TStringBuilderBase<TCHAR>, FIntVector2Type)
+	TStringBuilder<128> Builder;
+	CHECK((Builder << FIntVector2Type(0, 0)) == FIntVector2TestData::ZeroString);
+	Builder.Reset();
+	CHECK((Builder << FIntVector2Type(MinValue, MaxValue)) == FIntVector2TestData::MinMaxString);
+
+	// ToString
+	CHECK(FIntVector2Type(0, 0).ToString() == FIntVector2TestData::ZeroString);
+	CHECK(FIntVector2Type(MinValue, MaxValue).ToString() == FIntVector2TestData::MinMaxString);
+
+	// InitFromString
+	FIntVector2Type InitFromStringVec;
+	CHECK((InitFromStringVec.InitFromString(FString(FIntVector2TestData::ZeroString)) && InitFromStringVec == FIntVector2Type(0, 0)));
+	CHECK((InitFromStringVec.InitFromString(FString(FIntVector2TestData::MinMaxString)) && InitFromStringVec == FIntVector2Type(MinValue, MaxValue)));
+}
+
+TEST_CASE_NAMED(FInt32Vector2Test, "System::Core::Math::FInt32Vector2", "[ApplicationContextMask][SmokeFilter]")
+{
+	TestIntVector2<int32>();
+}
+
+TEST_CASE_NAMED(FInt64Vector2Test, "System::Core::Math::FInt64Vector2", "[ApplicationContextMask][SmokeFilter]")
+{
+	TestIntVector2<int64>();
+}
+
+TEST_CASE_NAMED(FUint32Vector2Test, "System::Core::Math::FUint32Vector2", "[ApplicationContextMask][SmokeFilter]")
+{
+	TestIntVector2<uint32>();
+}
+
+TEST_CASE_NAMED(FUint64Vector2Test, "System::Core::Math::FUint64Vector2", "[ApplicationContextMask][SmokeFilter]")
+{
+	TestIntVector2<uint64>();
+}
+
+// UE::Math::TIntVector3
+
+template <typename IntType>
+struct TIntVector3TestData
+{
+};
+
+template <>
+struct TIntVector3TestData<int32>
+{
+	static inline FStringView ZeroString = TEXTVIEW("X=0 Y=0 Z=0");
+	static inline FStringView MinZeroMaxString = TEXTVIEW("X=-2147483648 Y=0 Z=2147483647");
+};
+
+template <>
+struct TIntVector3TestData<int64>
+{
+	static inline FStringView ZeroString = TEXTVIEW("X=0 Y=0 Z=0");
+	static inline FStringView MinZeroMaxString = TEXTVIEW("X=-9223372036854775808 Y=0 Z=9223372036854775807");
+};
+
+template <>
+struct TIntVector3TestData<uint32>
+{
+	static inline FStringView ZeroString = TEXTVIEW("X=0 Y=0 Z=0");
+	static inline FStringView MinZeroMaxString = TEXTVIEW("X=0 Y=0 Z=4294967295");
+};
+
+template <>
+struct TIntVector3TestData<uint64>
+{
+	static inline FStringView ZeroString = TEXTVIEW("X=0 Y=0 Z=0");
+	static inline FStringView MinZeroMaxString = TEXTVIEW("X=0 Y=0 Z=18446744073709551615");
+};
+
+template <typename IntType>
+void TestIntVector3()
+{
+	using FIntVector3Type = UE::Math::TIntVector3<IntType>;
+	using FIntVector3TestData = TIntVector3TestData<IntType>;
+
+	constexpr IntType MinValue = TNumericLimits<IntType>::Min();
+	constexpr IntType MaxValue = TNumericLimits<IntType>::Max();
+
+	// GetMax
+	CHECK(FIntVector3Type(0, 1, 2).GetMax() == 2);
+	if constexpr (std::is_signed_v<IntType>)
+	{
+		CHECK(FIntVector3Type( 2,  0, -2).GetMax() == 2);
+		CHECK(FIntVector3Type(-2, -4, -6).GetMax() == -2);
+	}
+
+	// GetAbsMax
+	CHECK(FIntVector3Type(0, 1, 2).GetAbsMax() == 2);
+	if constexpr (std::is_signed_v<IntType>)
+	{
+		CHECK(FIntVector3Type( 2,  0, -2).GetAbsMax() == 2);
+		CHECK(FIntVector3Type(-2, -4, -6).GetAbsMax() == 6);
+	}
+
+	// GetMin
+	CHECK(FIntVector3Type(0, 1, 2).GetMin() == 0);
+	if constexpr (std::is_signed_v<IntType>)
+	{
+		CHECK(FIntVector3Type( 2,  0, -2).GetMin() == -2);
+		CHECK(FIntVector3Type(-2, -4, -6).GetMin() == -6);
+	}
+
+	// GetAbsMin
+	CHECK(FIntVector3Type(0, 1, 2).GetAbsMin() == 0);
+	if constexpr (std::is_signed_v<IntType>)
+	{
+		CHECK(FIntVector3Type( 2,  0, -2).GetAbsMin() == 0);
+		CHECK(FIntVector3Type(-2, -4, -6).GetAbsMin() == 2);
+	}
+
+	// ComponentMax
+	CHECK(FIntVector3Type(0, 1, 2).ComponentMax(FIntVector3Type(2, 1, 0)) == FIntVector3Type(2, 1, 2));
+	if constexpr (std::is_signed_v<IntType>)
+	{
+		CHECK(FIntVector3Type(-1, -2, -3).ComponentMax(FIntVector3Type(-3, -2, -1)) == FIntVector3Type(-1, -2, -1));
+	}
+
+	// ComponentMin
+	CHECK(FIntVector3Type(0, 1, 2).ComponentMin(FIntVector3Type(2, 1, 0)) == FIntVector3Type(0, 1, 0));
+	if constexpr (std::is_signed_v<IntType>)
+	{
+		CHECK(FIntVector3Type(-1, -2, -3).ComponentMin(FIntVector3Type(-3, -2, -1)) == FIntVector3Type(-3, -2, -3));
+	}
+
+	// AppendString(TStringBuilderBase<TCHAR>) via operator<<(TStringBuilderBase<TCHAR>, FIntVector3Type)
+	TStringBuilder<128> Builder;
+	CHECK((Builder << FIntVector3Type(0, 0, 0)) == FIntVector3TestData::ZeroString);
+	Builder.Reset();
+	CHECK((Builder << FIntVector3Type(MinValue, 0, MaxValue)) == FIntVector3TestData::MinZeroMaxString);
+
+	// ToString
+	CHECK(FIntVector3Type(0, 0, 0).ToString() == FIntVector3TestData::ZeroString);
+	CHECK(FIntVector3Type(MinValue, 0, MaxValue).ToString() == FIntVector3TestData::MinZeroMaxString);
+
+	// InitFromString
+	FIntVector3Type InitFromStringVec;
+	CHECK((InitFromStringVec.InitFromString(FString(FIntVector3TestData::ZeroString)) && InitFromStringVec == FIntVector3Type(0, 0, 0)));
+	CHECK((InitFromStringVec.InitFromString(FString(FIntVector3TestData::MinZeroMaxString)) && InitFromStringVec == FIntVector3Type(MinValue, 0, MaxValue)));
+}
+
+TEST_CASE_NAMED(FInt32Vector3Test, "System::Core::Math::FInt32Vector3", "[ApplicationContextMask][SmokeFilter]")
+{
+	TestIntVector3<int32>();
+}
+
+TEST_CASE_NAMED(FInt64Vector3Test, "System::Core::Math::FInt64Vector3", "[ApplicationContextMask][SmokeFilter]")
+{
+	TestIntVector3<int64>();
+}
+
+TEST_CASE_NAMED(FUint32Vector3Test, "System::Core::Math::FUint32Vector3", "[ApplicationContextMask][SmokeFilter]")
+{
+	TestIntVector3<uint32>();
+}
+
+TEST_CASE_NAMED(FUint64Vector3Test, "System::Core::Math::FUint64Vector3", "[ApplicationContextMask][SmokeFilter]")
+{
+	TestIntVector3<uint64>();
+}
+
+// UE::Math::TIntVector4
+
+template <typename IntType>
+struct TIntVector4TestData
+{
+};
+
+template <>
+struct TIntVector4TestData<int32>
+{
+	static inline FStringView ZeroString = TEXTVIEW("X=0 Y=0 Z=0 W=0");
+	static inline FStringView MinZeroMaxOneString = TEXTVIEW("X=-2147483648 Y=0 Z=2147483647 W=1");
+};
+
+template <>
+struct TIntVector4TestData<int64>
+{
+	static inline FStringView ZeroString = TEXTVIEW("X=0 Y=0 Z=0 W=0");
+	static inline FStringView MinZeroMaxOneString = TEXTVIEW("X=-9223372036854775808 Y=0 Z=9223372036854775807 W=1");
+};
+
+template <>
+struct TIntVector4TestData<uint32>
+{
+	static inline FStringView ZeroString = TEXTVIEW("X=0 Y=0 Z=0 W=0");
+	static inline FStringView MinZeroMaxOneString = TEXTVIEW("X=0 Y=0 Z=4294967295 W=1");
+};
+
+template <>
+struct TIntVector4TestData<uint64>
+{
+	static inline FStringView ZeroString = TEXTVIEW("X=0 Y=0 Z=0 W=0");
+	static inline FStringView MinZeroMaxOneString = TEXTVIEW("X=0 Y=0 Z=18446744073709551615 W=1");
+};
+
+template <typename IntType>
+void TestIntVector4()
+{
+	using FIntVector4Type = UE::Math::TIntVector4<IntType>;
+	using FIntVector4TestData = TIntVector4TestData<IntType>;
+
+	constexpr IntType MinValue = TNumericLimits<IntType>::Min();
+	constexpr IntType MaxValue = TNumericLimits<IntType>::Max();
+
+	// GetMax
+	CHECK(FIntVector4Type(0, 1, 2, 3).GetMax() == 3);
+	if constexpr (std::is_signed_v<IntType>)
+	{
+		CHECK(FIntVector4Type( 2,  0, -2,  0).GetMax() == 2);
+		CHECK(FIntVector4Type(-2, -4, -6, -8).GetMax() == -2);
+	}
+
+	// GetAbsMax
+	CHECK(FIntVector4Type(0, 1, 2, 3).GetAbsMax() == 3);
+	if constexpr (std::is_signed_v<IntType>)
+	{
+		CHECK(FIntVector4Type( 2,  0, -2,  0).GetAbsMax() == 2);
+		CHECK(FIntVector4Type(-2, -4, -6, -8).GetAbsMax() == 8);
+	}
+
+	// GetMin
+	CHECK(FIntVector4Type(0, 1, 2, 3).GetMin() == 0);
+	if constexpr (std::is_signed_v<IntType>)
+	{
+		CHECK(FIntVector4Type( 2,  0, -2,  0).GetMin() == -2);
+		CHECK(FIntVector4Type(-2, -4, -6, -8).GetMin() == -8);
+	}
+
+	// GetAbsMin
+	CHECK(FIntVector4Type(0, 1, 2, 3).GetAbsMin() == 0);
+	if constexpr (std::is_signed_v<IntType>)
+	{
+		CHECK(FIntVector4Type( 2,  0, -2,  0).GetAbsMin() == 0);
+		CHECK(FIntVector4Type(-2, -4, -6, -8).GetAbsMin() == 2);
+	}
+
+	// ComponentMax
+	CHECK(FIntVector4Type(0, 1, 2, 3).ComponentMax(FIntVector4Type(3, 2, 1, 0)) == FIntVector4Type(3, 2, 2, 3));
+	if constexpr (std::is_signed_v<IntType>)
+	{
+		CHECK(FIntVector4Type(-1, -2, -3, -4).ComponentMax(FIntVector4Type(-4, -3, -2, -1)) == FIntVector4Type(-1, -2, -2, -1));
+	}
+
+	// ComponentMin
+	CHECK(FIntVector4Type(0, 1, 2, 3).ComponentMin(FIntVector4Type(3, 2, 1, 0)) == FIntVector4Type(0, 1, 1, 0));
+	if constexpr (std::is_signed_v<IntType>)
+	{
+		CHECK(FIntVector4Type(-1, -2, -3, -4).ComponentMin(FIntVector4Type(-4, -3, -2, -1)) == FIntVector4Type(-4, -3, -3, -4));
+	}
+
+	// AppendString(TStringBuilderBase<TCHAR>) via operator<<(TStringBuilderBase<TCHAR>, FIntVector4Type)
+	TStringBuilder<128> Builder;
+	CHECK((Builder << FIntVector4Type(0, 0, 0, 0)) == FIntVector4TestData::ZeroString);
+	Builder.Reset();
+	CHECK((Builder << FIntVector4Type(MinValue, 0, MaxValue, 1)) == FIntVector4TestData::MinZeroMaxOneString);
+
+	// ToString
+	CHECK(FIntVector4Type(0, 0, 0, 0).ToString() == FIntVector4TestData::ZeroString);
+	CHECK(FIntVector4Type(MinValue, 0, MaxValue, 1).ToString() == FIntVector4TestData::MinZeroMaxOneString);
+
+	// InitFromString
+	FIntVector4Type InitFromStringVec;
+	CHECK((InitFromStringVec.InitFromString(FString(FIntVector4TestData::ZeroString)) && InitFromStringVec == FIntVector4Type(0, 0, 0, 0)));
+	CHECK((InitFromStringVec.InitFromString(FString(FIntVector4TestData::MinZeroMaxOneString)) && InitFromStringVec == FIntVector4Type(MinValue, 0, MaxValue, 1)));
+}
+
+TEST_CASE_NAMED(FInt32Vector4Test, "System::Core::Math::FInt32Vector4", "[ApplicationContextMask][SmokeFilter]")
+{
+	TestIntVector4<int32>();
+}
+
+TEST_CASE_NAMED(FInt64Vector4Test, "System::Core::Math::FInt64Vector4", "[ApplicationContextMask][SmokeFilter]")
+{
+	TestIntVector4<int64>();
+}
+
+TEST_CASE_NAMED(FUint32Vector4Test, "System::Core::Math::FUint32Vector4", "[ApplicationContextMask][SmokeFilter]")
+{
+	TestIntVector4<uint32>();
+}
+
+TEST_CASE_NAMED(FUint64Vector4Test, "System::Core::Math::FUint64Vector4", "[ApplicationContextMask][SmokeFilter]")
+{
+	TestIntVector4<uint64>();
+}
+
+#endif // WITH_TESTS
