[BinnedCommon]
 - Move TaskGraphInterfaces.h to the cpp too because some monolithic builds are not happy with the include

diff --git a/Engine/Source/Runtime/Core/Private/HAL/MallocBinnedCommon.cpp b/Engine/Source/Runtime/Core/Private/HAL/MallocBinnedCommon.cpp
--- a/Engine/Source/Runtime/Core/Private/HAL/MallocBinnedCommon.cpp
+++ b/Engine/Source/Runtime/Core/Private/HAL/MallocBinnedCommon.cpp
@@ -6,6 +6,7 @@
 #include "Containers/ArrayView.h"
 #include "Misc/AssertionMacros.h"
 #include "Math/NumericLimits.h"
+#include "Async/TaskGraphInterfaces.h"
 #include "Templates/AlignmentTemplates.h"
 #include "Templates/UnrealTemplate.h"
 
@@ -454,9 +455,23 @@ uint32 FBitTree::CountOnes(uint32 UpTo) const
 
 #endif
 
-bool FMallocBinnedCommonBase::IsAppMultithreaded()
+void FMallocBinnedCommonBase::ConditionalBroadcastSlow(TFunction<void()>& Broadcast)
 {
-	return FPlatformProcess::SupportsMultithreading() && FApp::ShouldUseThreadingForPerformance();
+	TFunction<void(ENamedThreads::Type)> ThreadBroadcast =
+		[Broadcast](ENamedThreads::Type)
+		{
+			Broadcast();
+		};
+
+	// Skip task threads on desktop platforms as it is too slow and they don't have much memory
+	if (PLATFORM_DESKTOP)
+	{
+		FTaskGraphInterface::BroadcastSlow_OnlyUseForSpecialPurposes(false, false, ThreadBroadcast);
+	}
+	else
+	{
+		FTaskGraphInterface::BroadcastSlow_OnlyUseForSpecialPurposes(FPlatformProcess::SupportsMultithreading() && FApp::ShouldUseThreadingForPerformance(), false, ThreadBroadcast);
+	}
 }
 
 float GMallocBinnedFlushThreadCacheMaxWaitTime = 0.2f;
diff --git a/Engine/Source/Runtime/Core/Public/HAL/MallocBinnedCommon.h b/Engine/Source/Runtime/Core/Public/HAL/MallocBinnedCommon.h
--- a/Engine/Source/Runtime/Core/Public/HAL/MallocBinnedCommon.h
+++ b/Engine/Source/Runtime/Core/Public/HAL/MallocBinnedCommon.h
@@ -5,8 +5,11 @@
 #include "CoreTypes.h"
 #include "HAL/MemoryBase.h"
 #include "Math/UnrealMathUtility.h"
+#include "Stats/Stats.h"
 #include "HAL/PlatformTLS.h"
-#include "Async/TaskGraphInterfaces.h"
+#include "HAL/PlatformTime.h"
+#include "Async/Mutex.h"
+#include "Misc/ScopeLock.h"
 #include "Templates/AlignmentTemplates.h"
 
 #if PLATFORM_HAS_FPlatformVirtualMemoryBlock
@@ -311,7 +314,7 @@ protected:
 	std::atomic<uint64> MemoryTrimEpoch{ 0 };
 
 protected:
-	bool IsAppMultithreaded();
+	void ConditionalBroadcastSlow(TFunction<void()>& Broadcast);
 };
 
 template <class AllocType, int MinAlign, int MaxAlign, int MinAlignShift, int NumSmallPools, int MaxSmallPoolSize>
@@ -714,22 +717,14 @@ protected:
 			}
 		}
 
-		TFunction<void(ENamedThreads::Type CurrentThread)> Broadcast =
-			[this, &Allocator](ENamedThreads::Type MyThread)
+		TFunction<void()> Broadcast =
+			[this, &Allocator]()
 			{
 				// We might already have updated the Epoch so we can skip doing anything costly (i.e. Mutex) in that case.
 				const bool bNewEpochOnly = true;
 				FlushCurrentThreadCache(Allocator, bNewEpochOnly);
 			};
 
-		// Skip task threads on desktop platforms as it is too slow and they don't have much memory
-		if (PLATFORM_DESKTOP)
-		{
-			FTaskGraphInterface::BroadcastSlow_OnlyUseForSpecialPurposes(false, false, Broadcast);
-		}
-		else
-		{
-			FTaskGraphInterface::BroadcastSlow_OnlyUseForSpecialPurposes(IsAppMultithreaded(), false, Broadcast);
-		}
+		ConditionalBroadcastSlow(Broadcast);
 	}
 };
\ No newline at end of file
