[BinnedCommon]
 - Fix build by removing circular dependencies caused by Stats.h

diff --git a/Engine/Source/Runtime/Core/Private/HAL/MallocBinnedCommon.cpp b/Engine/Source/Runtime/Core/Private/HAL/MallocBinnedCommon.cpp
--- a/Engine/Source/Runtime/Core/Private/HAL/MallocBinnedCommon.cpp
+++ b/Engine/Source/Runtime/Core/Private/HAL/MallocBinnedCommon.cpp
@@ -474,6 +474,18 @@ void FMallocBinnedCommonBase::ConditionalBroadcastSlow(TFunction<void()>& Broadc
 	}
 }
 
+void FMallocBinnedCommonBase::ConditionalLogWarnings(double WaitForMutexTime, double WaitForMutexAndTrimTime)
+{
+	if (WaitForMutexTime > GMallocBinnedFlushThreadCacheMaxWaitTime)
+	{
+		UE_LOG(LogMemory, Warning, TEXT("FMalloc%s took %6.2fms to wait for mutex for trim."), GetDescriptiveName(), WaitForMutexTime * 1000.0f);
+	}
+	if (WaitForMutexAndTrimTime > GMallocBinnedFlushThreadCacheMaxWaitTime)
+	{
+		UE_LOG(LogMemory, Warning, TEXT("FMalloc%s took %6.2fms to wait for mutex AND trim."), GetDescriptiveName(), WaitForMutexAndTrimTime * 1000.0f);
+	}
+}
+
 float GMallocBinnedFlushThreadCacheMaxWaitTime = 0.2f;
 static FAutoConsoleVariableRef GMallocBinnedFlushThreadCacheMaxWaitTimeCVar(
 	TEXT("MallocBinned.FlushThreadCacheMaxWaitTime"),
diff --git a/Engine/Source/Runtime/Core/Public/HAL/MallocBinnedCommon.h b/Engine/Source/Runtime/Core/Public/HAL/MallocBinnedCommon.h
--- a/Engine/Source/Runtime/Core/Public/HAL/MallocBinnedCommon.h
+++ b/Engine/Source/Runtime/Core/Public/HAL/MallocBinnedCommon.h
@@ -5,7 +5,6 @@
 #include "CoreTypes.h"
 #include "HAL/MemoryBase.h"
 #include "Math/UnrealMathUtility.h"
-#include "Stats/Stats.h"
 #include "HAL/PlatformTLS.h"
 #include "HAL/PlatformTime.h"
 #include "Async/Mutex.h"
@@ -315,6 +314,7 @@ protected:
 
 protected:
 	void ConditionalBroadcastSlow(TFunction<void()>& Broadcast);
+	void ConditionalLogWarnings(double WaitForMutexTime, double WaitForMutexAndTrimTime);
 };
 
 template <class AllocType, int MinAlign, int MaxAlign, int MinAlignShift, int NumSmallPools, int MaxSmallPoolSize>
@@ -660,7 +660,6 @@ protected:
 			if (Lists->UpdateEpoch(MemoryTrimEpoch.load(std::memory_order_relaxed)) || !bNewEpochOnly)
 			{
 				TRACE_CPUPROFILER_EVENT_SCOPE(FMallocBinnedCommon::FlushCurrentThreadCache);
-				QUICK_SCOPE_CYCLE_COUNTER(STAT_FMallocBinnedCommon_FlushCurrentThreadCache);
 
 				double StartTimeInner = FPlatformTime::Seconds();
 
@@ -675,14 +674,7 @@ protected:
 				}
 
 				// These logs must happen outside the above mutex to avoid deadlocks
-				if (WaitForMutexTime > GMallocBinnedFlushThreadCacheMaxWaitTime)
-				{
-					UE_LOG(LogMemory, Warning, TEXT("FMalloc%s took %6.2fms to wait for mutex for trim."), GetDescriptiveName(), WaitForMutexTime * 1000.0f);
-				}
-				if (WaitForMutexAndTrimTime > GMallocBinnedFlushThreadCacheMaxWaitTime)
-				{
-					UE_LOG(LogMemory, Warning, TEXT("FMalloc%s took %6.2fms to wait for mutex AND trim."), GetDescriptiveName(), WaitForMutexAndTrimTime * 1000.0f);
-				}
+				ConditionalLogWarnings(WaitForMutexTime, WaitForMutexAndTrimTime);
 			}
 		}
 	}
@@ -692,8 +684,6 @@ protected:
 		// Update the trim epoch so that threads cleanup their thread-local memory when going to sleep.
 		MemoryTrimEpoch.fetch_add(1, std::memory_order_relaxed);
 
-		QUICK_SCOPE_CYCLE_COUNTER(STAT_FMallocBinnedCommon_Trim);
-
 		// Process thread-local memory caches from as many threads as possible without waking them up.
 		// Skip on desktop as we may have too many threads and this could cause some hitches.
 		if (!PLATFORM_DESKTOP)
