[BinnedCommon]
 - Remove tracing until the circular dependency is fixed so this can be restored
 - BinnedCommon is currently included from UnrealMemory on some platform which cause circular dependencies between include files

diff --git a/Engine/Source/Runtime/Core/Public/HAL/MallocBinnedCommon.h b/Engine/Source/Runtime/Core/Public/HAL/MallocBinnedCommon.h
--- a/Engine/Source/Runtime/Core/Public/HAL/MallocBinnedCommon.h
+++ b/Engine/Source/Runtime/Core/Public/HAL/MallocBinnedCommon.h
@@ -640,8 +640,6 @@ protected:
 	{
 		if (FreeBlockLists)
 		{
-			TRACE_CPUPROFILER_EVENT_SCOPE(FMallocBinnedCommon::TrimThreadFreeBlockLists);
-
 			for (int32 PoolIndex = 0; PoolIndex != NumSmallPools; ++PoolIndex)
 			{
 				FBundleNode* Bundles = FreeBlockLists->PopBundles(PoolIndex);
@@ -659,8 +657,6 @@ protected:
 		{
 			if (Lists->UpdateEpoch(MemoryTrimEpoch.load(std::memory_order_relaxed)) || !bNewEpochOnly)
 			{
-				TRACE_CPUPROFILER_EVENT_SCOPE(FMallocBinnedCommon::FlushCurrentThreadCache);
-
 				double StartTimeInner = FPlatformTime::Seconds();
 
 				double WaitForMutexTime = 0.0f;
