[Backout] - CL32468535
Original CL Desc
-----------------------------------------------------------------
[MallocBinned2/3]
 - Unify GMallocBinned2FlushThreadCacheMaxWaitTime and GMallocBinned3FlushThreadCacheMaxWaitTime
 - Move FApp call to cpp file to remove App.h from MallocBinnedCommon.h

diff --git a/Engine/Source/Runtime/Core/Private/HAL/MallocBinnedCommon.cpp b/Engine/Source/Runtime/Core/Private/HAL/MallocBinnedCommon.cpp
--- a/Engine/Source/Runtime/Core/Private/HAL/MallocBinnedCommon.cpp
+++ b/Engine/Source/Runtime/Core/Private/HAL/MallocBinnedCommon.cpp
@@ -2,11 +2,9 @@
 
 #include "HAL/MallocBinnedCommon.h"
 #include "Algo/Sort.h"
-#include "Misc/App.h"
 #include "Containers/ArrayView.h"
 #include "Misc/AssertionMacros.h"
 #include "Math/NumericLimits.h"
-#include "Async/TaskGraphInterfaces.h"
 #include "Templates/AlignmentTemplates.h"
 #include "Templates/UnrealTemplate.h"
 
@@ -455,45 +453,6 @@ uint32 FBitTree::CountOnes(uint32 UpTo) const
 
 #endif
 
-void FMallocBinnedCommonBase::ConditionalBroadcastSlow(TFunction<void()>& Broadcast)
-{
-	TFunction<void(ENamedThreads::Type)> ThreadBroadcast =
-		[Broadcast](ENamedThreads::Type)
-		{
-			Broadcast();
-		};
-
-	// Skip task threads on desktop platforms as it is too slow and they don't have much memory
-	if (PLATFORM_DESKTOP)
-	{
-		FTaskGraphInterface::BroadcastSlow_OnlyUseForSpecialPurposes(false, false, ThreadBroadcast);
-	}
-	else
-	{
-		FTaskGraphInterface::BroadcastSlow_OnlyUseForSpecialPurposes(FPlatformProcess::SupportsMultithreading() && FApp::ShouldUseThreadingForPerformance(), false, ThreadBroadcast);
-	}
-}
-
-void FMallocBinnedCommonBase::ConditionalLogWarnings(double WaitForMutexTime, double WaitForMutexAndTrimTime)
-{
-	if (WaitForMutexTime > GMallocBinnedFlushThreadCacheMaxWaitTime)
-	{
-		UE_LOG(LogMemory, Warning, TEXT("FMalloc%s took %6.2fms to wait for mutex for trim."), GetDescriptiveName(), WaitForMutexTime * 1000.0f);
-	}
-	if (WaitForMutexAndTrimTime > GMallocBinnedFlushThreadCacheMaxWaitTime)
-	{
-		UE_LOG(LogMemory, Warning, TEXT("FMalloc%s took %6.2fms to wait for mutex AND trim."), GetDescriptiveName(), WaitForMutexAndTrimTime * 1000.0f);
-	}
-}
-
-float GMallocBinnedFlushThreadCacheMaxWaitTime = 0.2f;
-static FAutoConsoleVariableRef GMallocBinnedFlushThreadCacheMaxWaitTimeCVar(
-	TEXT("MallocBinned.FlushThreadCacheMaxWaitTime"),
-	GMallocBinnedFlushThreadCacheMaxWaitTime,
-	TEXT("The threshold of time before warning about FlushCurrentThreadCache taking too long (seconds)."),
-	ECVF_ReadOnly
-);
-
 #if UE_BINNEDCOMMON_ALLOW_RUNTIME_TWEAKING
 
 int32 GMallocBinnedBundleSize = DEFAULT_GMallocBinnedBundleSize;
