Fixed FTextHistory_StringTableEntry::GetDisplayString returning the placeholder for a valid entry with no translation data

This was a regression caused by a recent change to no longer generate display strings for untranslated texts

#rnx

diff --git a/Engine/Source/Runtime/Core/Private/Internationalization/TextHistory.cpp b/Engine/Source/Runtime/Core/Private/Internationalization/TextHistory.cpp
--- a/Engine/Source/Runtime/Core/Private/Internationalization/TextHistory.cpp
+++ b/Engine/Source/Runtime/Core/Private/Internationalization/TextHistory.cpp
@@ -2370,7 +2370,6 @@ const FString& FTextHistory_StringTableEntry::GetDisplayString() const
 		{
 			return *DisplayString;
 		}
-		return FStringTableEntry::GetPlaceholderSourceString();
 	}
 	return GetSourceString();
 }
