Fixed non unity compile errors

diff --git a/Engine/Source/Runtime/Core/Private/ProfilingDebugging/TraceAuxiliary.cpp b/Engine/Source/Runtime/Core/Private/ProfilingDebugging/TraceAuxiliary.cpp
--- a/Engine/Source/Runtime/Core/Private/ProfilingDebugging/TraceAuxiliary.cpp
+++ b/Engine/Source/Runtime/Core/Private/ProfilingDebugging/TraceAuxiliary.cpp
@@ -28,7 +28,7 @@
 	#define UE_TRACE_AUTOSTART 1
 #endif
 
-#if UE_TRACE_SERVER_LAUNCH_ENABLED
+#if UE_TRACE_SERVER_LAUNCH_ENABLED || UE_TRACE_SERVER_CONTROLS_ENABLED
 #include "HAL/PlatformProcess.h"
 #include "Misc/CommandLine.h"
 #include "Misc/Parse.h"
