Backed out the TConstSetBitIterator optimization while investigating a crash

#rnx

diff --git a/Engine/Source/Runtime/Core/Public/Containers/BitArray.h b/Engine/Source/Runtime/Core/Public/Containers/BitArray.h
--- a/Engine/Source/Runtime/Core/Public/Containers/BitArray.h
+++ b/Engine/Source/Runtime/Core/Public/Containers/BitArray.h
@@ -1807,56 +1807,68 @@ public:
 
 	explicit TConstSetBitIterator(const TBitArray<Allocator>& InArray UE_LIFETIMEBOUND)
 		: FRelativeBitReference(0)
-		, Array(InArray)
-		, WordIt(InArray.GetData())
-		, BitCount(InArray.Num())
+		, Array                (InArray)
+		, UnvisitedBitMask     (~0U)
+		, CurrentBitIndex      (0)
+		, BaseBitIndex         (0)
 	{
-		if (BitCount)
+		if (Array.Num())
 		{
-			RemainingWord = *WordIt++;
 			FindFirstSetBit();
 		}
 	}
 
 	explicit TConstSetBitIterator(const TBitArray<Allocator>& InArray UE_LIFETIMEBOUND, int32 StartIndex)
 		: FRelativeBitReference(StartIndex)
-		, Array(InArray)
-		, WordIt(InArray.GetData() + WordIndex)
-		, BitCount(InArray.Num())
-		, CurrentBitIndex(StartIndex)
+		, Array                (InArray)
+		, UnvisitedBitMask     ((~0U) << (StartIndex & (NumBitsPerDWORD - 1)))
+		, CurrentBitIndex      (StartIndex)
+		, BaseBitIndex         (StartIndex & ~(NumBitsPerDWORD - 1))
 	{
-		check(CurrentBitIndex >= 0 && CurrentBitIndex <= BitCount);
-		if (CurrentBitIndex != BitCount && BitCount)
+		check(StartIndex >= 0 && StartIndex <= Array.Num());
+		if (StartIndex != Array.Num())
 		{
-			RemainingWord = *WordIt++ & ~(Mask - 1);
 			FindFirstSetBit();
 		}
 	}
 
+	/** Forwards iteration operator. */
 	FORCEINLINE TConstSetBitIterator& operator++()
 	{
+		// Mark the current bit as visited.
+		UnvisitedBitMask &= ~this->Mask;
+
+		// Find the first set bit that hasn't been visited yet.
 		FindFirstSetBit();
+
 		return *this;
 	}
 
 	FORCEINLINE bool operator==(const TConstSetBitIterator& Rhs) const
 	{
-		// We only need to compare the bit index and the array... the rest of the state is unobservable.
+		// We only need to compare the bit index and the array... all the rest of the state is unobservable.
 		return CurrentBitIndex == Rhs.CurrentBitIndex && &Array == &Rhs.Array;
 	}
 
 #if !PLATFORM_COMPILER_HAS_GENERATED_COMPARISON_OPERATORS
 	FORCEINLINE bool operator!=(const TConstSetBitIterator& Rhs) const
-	{
+	{ 
 		return !(*this == Rhs);
 	}
 #endif
 
+	/** conversion to "bool" returning true if the iterator is valid. */
 	FORCEINLINE explicit operator bool() const
+	{ 
+		return CurrentBitIndex < Array.Num(); 
+	}
+	/** inverse of the "bool" operator */
+	FORCEINLINE bool operator !() const 
 	{
-		return CurrentBitIndex < BitCount;
+		return !(bool)*this;
 	}
 
+	/** Index accessor. */
 	FORCEINLINE int32 GetIndex() const
 	{
 		return CurrentBitIndex;
@@ -1866,41 +1878,49 @@ private:
 
 	const TBitArray<Allocator>& Array;
 
-	const uint32* WordIt;
-	const int32 BitCount;
-	int32 CurrentBitIndex = 0;
-	uint32 RemainingWord = 0;
+	uint32 UnvisitedBitMask;
+	int32 CurrentBitIndex;
+	int32 BaseBitIndex;
 
 	/** Find the first set bit starting with the current bit, inclusive. */
 	void FindFirstSetBit()
 	{
-		for (;;)
+		const uint32* ArrayData      = Array.GetData();
+		const int32   ArrayNum       = Array.Num();
+		const int32   LastWordIndex = (ArrayNum - 1) / NumBitsPerDWORD;
+
+		// Advance to the next non-zero uint32.
+		uint32 RemainingBitMask = ArrayData[this->WordIndex] & UnvisitedBitMask;
+		while (!RemainingBitMask)
 		{
-			// Find the index of the first remaining set bit in the word.
-			if (RemainingWord)
+			++this->WordIndex;
+			BaseBitIndex += NumBitsPerDWORD;
+			if (this->WordIndex > LastWordIndex)
 			{
-				// Clear the first remaining set bit in the word.
-				const uint32 NewRemainingWord = RemainingWord & (RemainingWord - 1);
-
-				// Calculate the index of the bit that was cleared.
-				Mask = RemainingWord ^ NewRemainingWord;
-				WordIndex = CurrentBitIndex >> NumBitsPerDWORDLogTwo;
-				CurrentBitIndex = (WordIndex << NumBitsPerDWORDLogTwo) | FMath::FloorLog2NonZero(Mask);
-
-				RemainingWord = NewRemainingWord;
+				// We've advanced past the end of the array.
+				CurrentBitIndex = ArrayNum;
 				return;
 			}
 
-			// Advance the bit index to the next word.
-			CurrentBitIndex = (CurrentBitIndex & ~(NumBitsPerDWORD - 1)) + NumBitsPerDWORD;
-			if (CurrentBitIndex >= BitCount)
-			{
-				CurrentBitIndex = BitCount;
-				return;
-			}
+			RemainingBitMask = ArrayData[this->WordIndex];
+			UnvisitedBitMask = ~0u;
+		}
 
-			// Read the next word into the iterator.
-			RemainingWord = *WordIt++;
+		// This operation has the effect of unsetting the lowest set bit of BitMask
+		const uint32 NewRemainingBitMask = RemainingBitMask & (RemainingBitMask - 1);
+
+		// This operation XORs the above mask with the original mask, which has the effect
+		// of returning only the bits which differ; specifically, the lowest bit
+		this->Mask = NewRemainingBitMask ^ RemainingBitMask;
+
+		// If the Nth bit was the lowest set bit of BitMask, then this gives us N
+		CurrentBitIndex = BaseBitIndex + NumBitsPerDWORD - 1 - FMath::CountLeadingZeros(this->Mask);
+
+		// If we've accidentally iterated off the end of an array but still within the same Word
+		// then set the index to the last index of the array
+		if (CurrentBitIndex > ArrayNum)
+		{
+			CurrentBitIndex = ArrayNum;
 		}
 	}
 };
