[Core]
* Removed macro dependency from Core to CoreUObject in ScopedTimers.h. Added a dedicated macro instead to combine args

diff --git a/Engine/Source/Runtime/Core/Public/ProfilingDebugging/ScopedTimers.h b/Engine/Source/Runtime/Core/Public/ProfilingDebugging/ScopedTimers.h
--- a/Engine/Source/Runtime/Core/Public/ProfilingDebugging/ScopedTimers.h
+++ b/Engine/Source/Runtime/Core/Public/ProfilingDebugging/ScopedTimers.h
@@ -244,5 +244,8 @@ private:
 #define UE_SCOPED_TIMER(Title, Category, Verbosity)
 #else
 #define UE_SCOPED_TIMER(Title, Category, Verbosity) \
-	static double BODY_MACRO_COMBINE(Scoped,TimerTotal,_,__LINE__) = 0.0; FScopedDurationTimeCustomLogger BODY_MACRO_COMBINE(Scoped,Timer,_,__LINE__)(Title, BODY_MACRO_COMBINE(Scoped,TimerTotal,_,__LINE__), [](const TCHAR* Msg) { UE_LOG(Category, Verbosity, TEXT("%s"), Msg) })
+	static double UE_SCOPED_TIMER_COMBINE(ScopedTimerTotal_,__LINE__) = 0.0; FScopedDurationTimeCustomLogger UE_SCOPED_TIMER_COMBINE(ScopedTimer_,__LINE__)(Title, UE_SCOPED_TIMER_COMBINE(ScopedTimerTotal_,__LINE__), [](const TCHAR* Msg) { UE_LOG(Category, Verbosity, TEXT("%s"), Msg) })
+
+#define UE_SCOPED_TIMER_COMBINE_INNER(A,B) A##B
+#define UE_SCOPED_TIMER_COMBINE(A,B) UE_SCOPED_TIMER_COMBINE_INNER(A,B)
 #endif
