Add an AutoSTM flag and make it a feature we can get at via the C++ preprocessor.

diff --git a/Engine/Source/Runtime/Core/Private/AutoRTFM/Context.cpp b/Engine/Source/Runtime/Core/Private/AutoRTFM/Context.cpp
--- a/Engine/Source/Runtime/Core/Private/AutoRTFM/Context.cpp
+++ b/Engine/Source/Runtime/Core/Private/AutoRTFM/Context.cpp
@@ -287,6 +287,12 @@ ETransactionResult FContext::Transact(void (*Function)(void* Arg), void* Arg)
 {
     constexpr bool bVerbose = false;
 
+	if constexpr (0 != UE_AUTOSTM)
+	{
+		UE_LOG(LogAutoRTFM, Warning, TEXT("AutoSTM is not implemented in the runtime yet!"));
+		return ETransactionResult::AbortedByLanguage;
+	}
+
     if (UNLIKELY(EContextStatus::Committing == Status))
     {
     	return ETransactionResult::AbortedByTransactInOnCommit;
diff --git a/Engine/Source/Runtime/Core/Public/AutoRTFM/AutoRTFM.h b/Engine/Source/Runtime/Core/Public/AutoRTFM/AutoRTFM.h
--- a/Engine/Source/Runtime/Core/Public/AutoRTFM/AutoRTFM.h
+++ b/Engine/Source/Runtime/Core/Public/AutoRTFM/AutoRTFM.h
@@ -8,6 +8,18 @@
 #define UE_AUTORTFM 0
 #endif
 
+#if UE_AUTORTFM
+#if defined(__has_feature)
+#if __has_feature(autostm)
+#define UE_AUTOSTM 1
+#endif
+#endif
+#endif
+
+#if !defined(UE_AUTOSTM)
+#define UE_AUTOSTM 0
+#endif
+
 #if UE_AUTORTFM
 #define UE_AUTORTFM_AUTORTFM(F) [[clang::autortfm(F), clang::noinline]]
 #define UE_AUTORTFM_NOAUTORTFM [[clang::noautortfm, clang::noinline]]
