FHashBuilder optimization: now uses CityHash instead of CRC when hashing types not implementing GetTypeHash.

diff --git a/Engine/Source/Runtime/Core/Public/Misc/HashBuilder.h b/Engine/Source/Runtime/Core/Public/Misc/HashBuilder.h
--- a/Engine/Source/Runtime/Core/Public/Misc/HashBuilder.h
+++ b/Engine/Source/Runtime/Core/Public/Misc/HashBuilder.h
@@ -3,13 +3,13 @@
 #pragma once
 
 #include "CoreTypes.h"
-#include "Crc.h"
 #include "Concepts/GetTypeHashable.h"
 #include "Containers/UnrealString.h"
+#include "Hash/CityHash.h"
 
 /**
  * Class for computing a hash of multiple types, going through GetTypeHash when the type implements it, and
- * fallbacks to CRC32 when the type doesn't.
+ * fallbacks to raw data hashing when the type doesn't.
  *
  * Note: this hash builder should be used for transient hashes, as some types implements run-dependent hash
  * computations, such as GetTypeHash(FName).
@@ -21,20 +21,20 @@ public:
 		: Hash(~InHash)
 	{}
 
-	void AppendRaw(const void* Data, int64 Num)
+	FORCEINLINE FHashBuilder& AppendRaw(const void* Data, int64 Num)
 	{
-		Hash = FCrc::MemCrc32(Data, static_cast<int32>(Num), Hash);		// TODO: Update MemCrc32 to take an int64 Length?
+		Hash = HashCombineFast(Hash, GetTypeHash(CityHash64(static_cast<const char*>(Data),  static_cast<int32>(Num))));
+		return *this;
 	}
 
 	template <typename T>
-	typename TEnableIf<TIsPODType<T>::Value, FHashBuilder&>::Type AppendRaw(const T& InData)
+	FORCEINLINE typename TEnableIf<TIsPODType<T>::Value, FHashBuilder&>::Type AppendRaw(const T& InData)
 	{
-		AppendRaw(&InData, sizeof(T));
-		return *this;
+		return AppendRaw(&InData, sizeof(T));
 	}
 
 	template <typename T>
-	FHashBuilder& Append(const T& InData)
+	FORCEINLINE FHashBuilder& Append(const T& InData)
 	{
 		if constexpr (TModels_V<CGetTypeHashable, T>)
 		{
@@ -48,9 +48,9 @@ public:
 	}
 
 	template <typename T>
-	FHashBuilder& Append(const TArray<T>& InArray)
+	FORCEINLINE FHashBuilder& Append(const TArray<T>& InArray)
 	{
-		for (auto& Value: InArray)
+		for (const T& Value: InArray)
 		{
 			Append(Value);
 		}
@@ -58,26 +58,37 @@ public:
 	}
 
 	template <typename T>
-	FHashBuilder& Append(const TSet<T>& InArray)
+	FORCEINLINE FHashBuilder& Append(const TSet<T>& InSet)
 	{
-		for (auto& Value: InArray)
+		for (const T& Value: InSet)
 		{
 			Append(Value);
 		}
 		return *this;
 	}
 
+	template <typename T, typename U>
+	FORCEINLINE FHashBuilder& Append(const TMap<T, U>& InMap)
+	{
+		for (const TPair<T, U>& Value: InMap)
+		{
+			Append(Value.Key);
+			Append(Value.Value);
+		}
+		return *this;
+	}
+
 	template <typename T>
-	FHashBuilder& operator<<(const T& InData)
+	FORCEINLINE FHashBuilder& operator<<(const T& InData)
 	{
 		return Append(InData);
 	}
 
-	uint32 GetHash() const
+	FORCEINLINE uint32 GetHash() const
 	{
 		return ~Hash;
 	}
 
 private:
 	uint32 Hash;
-};
+};
\ No newline at end of file
