[Core]
* Fixed header compile error by adding include

diff --git a/Engine/Source/Runtime/Core/Public/Misc/HashBuilder.h b/Engine/Source/Runtime/Core/Public/Misc/HashBuilder.h
--- a/Engine/Source/Runtime/Core/Public/Misc/HashBuilder.h
+++ b/Engine/Source/Runtime/Core/Public/Misc/HashBuilder.h
@@ -2,8 +2,8 @@
 
 #pragma once
 
-#include "CoreTypes.h"
 #include "Concepts/GetTypeHashable.h"
+#include "Containers/Map.h"
 #include "Containers/UnrealString.h"
 #include "Hash/CityHash.h"
 
