PR #11663 : Avoid division by zero in FastDecimalFormat

diff --git a/Engine/Source/Runtime/Core/Private/Internationalization/FastDecimalFormat.cpp b/Engine/Source/Runtime/Core/Private/Internationalization/FastDecimalFormat.cpp
--- a/Engine/Source/Runtime/Core/Private/Internationalization/FastDecimalFormat.cpp
+++ b/Engine/Source/Runtime/Core/Private/Internationalization/FastDecimalFormat.cpp
@@ -431,7 +431,7 @@ FString CultureInvariantDecimalToString(const double InVal, const TCHAR*& InBuff
 
 	bool bUseGrouping = InFormattingOptions.UseGrouping && InFormattingRules.PrimaryGroupingSize > 0;
 	uint8 NumIntegralDigits = static_cast<uint8>(FMath::Abs((FMath::LogX(10.0, InVal)))) + 1;
-	uint8 NumUntilNextGroup = NumIntegralDigits % InFormattingRules.PrimaryGroupingSize;
+	uint8 NumUntilNextGroup = bUseGrouping ? NumIntegralDigits % InFormattingRules.PrimaryGroupingSize : 0;
 	const TCHAR* InBufferEnd = InBuffer + InBufferLen;
 
 	// Apply front padding
