ShaderKeyGenerator: Suppress static analysis compile warning.
#rnx

diff --git a/Engine/Source/Runtime/Core/Private/Serialization/ShaderKeyGenerator.cpp b/Engine/Source/Runtime/Core/Private/Serialization/ShaderKeyGenerator.cpp
--- a/Engine/Source/Runtime/Core/Private/Serialization/ShaderKeyGenerator.cpp
+++ b/Engine/Source/Runtime/Core/Private/Serialization/ShaderKeyGenerator.cpp
@@ -13,7 +13,7 @@ void FShaderKeyGenerator::Append(const FBlake3Hash& Value)
 	case EOutputType::Text:
 	{
 		CA_SUPPRESS(6260) /* warning C6260: sizeof * sizeof is usually wrong. */
-		constexpr int32 StringSize = sizeof(TCHAR) * sizeof(FBlake3Hash::ByteArray) * 2;
+		constexpr int32 StringSize = sizeof(TCHAR) * sizeof(FBlake3Hash::ByteArray) * 2; // -V531
 		ResultString->Append(WriteToString<StringSize>(Value));
 		return;
 	}
