Fix the uint64_t -> int compiler error.

diff --git a/Engine/Source/Runtime/Core/Private/AutoRTFM/TaskArray.h b/Engine/Source/Runtime/Core/Private/AutoRTFM/TaskArray.h
--- a/Engine/Source/Runtime/Core/Private/AutoRTFM/TaskArray.h
+++ b/Engine/Source/Runtime/Core/Private/AutoRTFM/TaskArray.h
@@ -71,9 +71,9 @@ public:
 
 		auto EraseLastAddedKeyFromArray = [Key](FInternalArray& Array) -> bool
 		{
-			for(uint64 Idx = 0; Idx < Array.Num(); Idx++)
+			for(typename FInternalArray::SizeType Idx = 0; Idx < Array.Num(); Idx++)
 			{
-				uint64 BackwardsIdx = Array.Num() - 1 - Idx;
+				typename FInternalArray::SizeType BackwardsIdx = Array.Num() - 1 - Idx;
 
 				if(Array[BackwardsIdx].Key != Key)
 				{
@@ -189,15 +189,15 @@ public:
 
     size_t Num() const
     {
-        size_t Result = Latest.Num();
+		typename FInternalArray::SizeType Result = Latest.Num();
 
-        for (size_t Index = 0; Index < Stash.Num(); Index++)
+        for (typename FInternalArray::SizeType Index = 0; Index < Stash.Num(); Index++)
         {
             const FInternalArray& StashedVector = Stash[Index];
             Result += StashedVector.Num();
         }
 
-        return Result;
+		return static_cast<size_t>(Result);
     }
 
 private:
