Hook MallocZeroed deeper into FMemory.

This allows FMalloc implementations to provide optimal malloc implementations if they can zero memory cheaply.

diff --git a/Engine/Source/Runtime/Core/Private/HAL/MallocMimalloc.cpp b/Engine/Source/Runtime/Core/Private/HAL/MallocMimalloc.cpp
--- a/Engine/Source/Runtime/Core/Private/HAL/MallocMimalloc.cpp
+++ b/Engine/Source/Runtime/Core/Private/HAL/MallocMimalloc.cpp
@@ -174,6 +174,43 @@ void FMallocMimalloc::Free( void* Ptr )
 	mi_free(Ptr);
 }
 
+void* FMallocMimalloc::MallocZeroed(SIZE_T Size, uint32 Alignment)
+{
+	void* Result = TryMallocZeroed(Size, Alignment);
+
+	if (Result == nullptr && Size)
+	{
+		OutOfMemory(Size, Alignment);
+	}
+
+	return Result;
+}
+
+void* FMallocMimalloc::TryMallocZeroed(SIZE_T Size, uint32 Alignment)
+{
+#if !UE_BUILD_SHIPPING
+	uint64 LocalMaxSingleAlloc = MaxSingleAlloc.Load(EMemoryOrder::Relaxed);
+	if (LocalMaxSingleAlloc != 0 && Size > LocalMaxSingleAlloc)
+	{
+		return nullptr;
+	}
+#endif
+
+	void* NewPtr = nullptr;
+
+	if (Alignment != DEFAULT_ALIGNMENT)
+	{
+		Alignment = FMath::Max(uint32(Size >= 16 ? 16 : 8), Alignment);
+		NewPtr = mi_zalloc_aligned(Size, Alignment);
+	}
+	else
+	{
+		NewPtr = mi_zalloc_aligned(Size, uint32(Size >= 16 ? 16 : 8));
+	}
+
+	return NewPtr;
+}
+
 bool FMallocMimalloc::GetAllocationSize(void *Original, SIZE_T &SizeOut)
 {
 	SizeOut = mi_malloc_size(Original);
diff --git a/Engine/Source/Runtime/Core/Private/HAL/MemoryBase.cpp b/Engine/Source/Runtime/Core/Private/HAL/MemoryBase.cpp
--- a/Engine/Source/Runtime/Core/Private/HAL/MemoryBase.cpp
+++ b/Engine/Source/Runtime/Core/Private/HAL/MemoryBase.cpp
@@ -101,3 +101,27 @@ void* FMalloc::TryRealloc(void* Original, SIZE_T Count, uint32 Alignment)
 {
 	return Realloc(Original, Count, Alignment);
 }
+
+void* FMalloc::MallocZeroed(SIZE_T Count, uint32 Alignment)
+{
+	void* const Memory = Malloc(Count, Alignment);
+
+	if (Memory)
+	{
+		FPlatformMemory::Memzero(Memory, Count);
+	}
+
+	return Memory;
+}
+
+void* FMalloc::TryMallocZeroed(SIZE_T Count, uint32 Alignment)
+{
+	void* const Memory = Malloc(Count, Alignment);
+
+	if (Memory)
+	{
+		FPlatformMemory::Memzero(Memory, Count);
+	}
+
+	return Memory;
+}
diff --git a/Engine/Source/Runtime/Core/Private/HAL/UnrealMemory.cpp b/Engine/Source/Runtime/Core/Private/HAL/UnrealMemory.cpp
--- a/Engine/Source/Runtime/Core/Private/HAL/UnrealMemory.cpp
+++ b/Engine/Source/Runtime/Core/Private/HAL/UnrealMemory.cpp
@@ -531,6 +531,16 @@ SIZE_T FMemory::GetAllocSizeExternal(void* Original)
 	return GMalloc->GetAllocationSize(Original, Size) ? Size : 0;
 }
 
+void* FMemory::MallocZeroedExternal(SIZE_T Count, uint32 Alignment)
+{
+	if (!GMalloc)
+	{
+		GCreateMalloc();
+		CA_ASSUME(GMalloc != NULL);	// Don't want to assert, but suppress static analysis warnings about potentially NULL GMalloc
+	}
+	return GMalloc->MallocZeroed(Count, Alignment);
+}
+
 SIZE_T FMemory::QuantizeSizeExternal(SIZE_T Count, uint32 Alignment)
 { 
 	if (!GMalloc)
diff --git a/Engine/Source/Runtime/Core/Public/HAL/FMemory.inl b/Engine/Source/Runtime/Core/Public/HAL/FMemory.inl
--- a/Engine/Source/Runtime/Core/Public/HAL/FMemory.inl
+++ b/Engine/Source/Runtime/Core/Public/HAL/FMemory.inl
@@ -168,6 +168,45 @@ FMEMORY_INLINE_FUNCTION_DECORATOR SIZE_T FMemory::GetAllocSize(void* Original)
 	return Result;
 }
 
+FMEMORY_INLINE_FUNCTION_DECORATOR void* FMemory::MallocZeroed(SIZE_T Count, uint32 Alignment)
+{
+	void* Ptr = nullptr; // Silence bogus static analysis warnings.
+
+	// AutoRTFM: For non-transactional code, all of these calls optimize away and the
+	// behavior is the same as it always has been.
+	// For transactional code, we call the allocator in the 'open' as an optimization, so that
+	// we don't end up keeping track of the writes to the allocator's internal data structures.
+	// This is because allocators are already transactional - malloc can be rolled back by
+	// calling free.
+	UE_AUTORTFM_OPEN(
+		{
+			if (!FMEMORY_INLINE_GMalloc)
+			{
+				Ptr = MallocZeroedExternal(Count, Alignment);
+			}
+			else
+			{
+				DoGamethreadHook(0);
+				FScopedMallocTimer Timer(0);
+				Ptr = FMEMORY_INLINE_GMalloc->MallocZeroed(Count, Alignment);
+			}
+	// optional tracking of every allocation
+	LLM_IF_ENABLED(FLowLevelMemTracker::Get().OnLowLevelAlloc(ELLMTracker::Default, Ptr, Count, ELLMTag::Untagged, ELLMAllocType::FMalloc));
+		});
+
+	// AutoRTFM: This is a no-op for non-transactional code.
+	// For transactional code, this defers a call to Free if the transaction aborts,
+	// so that rolling back this allocation will end up freeing the memory.
+	AutoRTFM::OnAbort([Ptr]
+		{
+			// Disable the code analysis warning that complains that Free is being passed
+			// a pointer that may be null. Free explicitly handles this case already.
+			Free(Ptr); //-V575
+		});
+
+	return AutoRTFM::DidAllocate(Ptr, Count);
+}
+
 FMEMORY_INLINE_FUNCTION_DECORATOR SIZE_T FMemory::QuantizeSize(SIZE_T Count, uint32 Alignment)
 {
 	SIZE_T Result;
diff --git a/Engine/Source/Runtime/Core/Public/HAL/MallocMimalloc.h b/Engine/Source/Runtime/Core/Public/HAL/MallocMimalloc.h
--- a/Engine/Source/Runtime/Core/Public/HAL/MallocMimalloc.h
+++ b/Engine/Source/Runtime/Core/Public/HAL/MallocMimalloc.h
@@ -31,6 +31,8 @@ public:
 	virtual void* Realloc(void* Ptr, SIZE_T NewSize, uint32 Alignment) override;
 	virtual void* TryRealloc(void* Ptr, SIZE_T NewSize, uint32 Alignment) override;
 	virtual void Free(void* Ptr) override;
+	virtual void* MallocZeroed(SIZE_T Count, uint32 Alignment) override;
+	virtual void* TryMallocZeroed(SIZE_T Count, uint32 Alignment) override;
 	virtual bool GetAllocationSize(void *Original, SIZE_T &SizeOut) override;
 	virtual void Trim(bool bTrimThreadCaches) override;
 
diff --git a/Engine/Source/Runtime/Core/Public/HAL/MemoryBase.h b/Engine/Source/Runtime/Core/Public/HAL/MemoryBase.h
--- a/Engine/Source/Runtime/Core/Public/HAL/MemoryBase.h
+++ b/Engine/Source/Runtime/Core/Public/HAL/MemoryBase.h
@@ -108,6 +108,17 @@ public:
 	 * Free
 	 */
 	virtual void Free( void* Original ) = 0;
+
+	/**
+	 * Malloc zeroed memory
+	 */
+	CORE_API virtual void* MallocZeroed(SIZE_T Count, uint32 Alignment = DEFAULT_ALIGNMENT);
+
+	/**
+	 * TryMalloc - like MallocZeroed(), but may return a nullptr result if the allocation
+	 *             request cannot be satisfied.
+	 */
+	CORE_API virtual void* TryMallocZeroed(SIZE_T Count, uint32 Alignment = DEFAULT_ALIGNMENT);
 		
 	/** 
 	* For some allocators this will return the actual size that should be requested to eliminate
diff --git a/Engine/Source/Runtime/Core/Public/HAL/UnrealMemory.h b/Engine/Source/Runtime/Core/Public/HAL/UnrealMemory.h
--- a/Engine/Source/Runtime/Core/Public/HAL/UnrealMemory.h
+++ b/Engine/Source/Runtime/Core/Public/HAL/UnrealMemory.h
@@ -214,12 +214,7 @@ struct FMemory
 	static CORE_API void Free(void* Original);
 	static CORE_API SIZE_T GetAllocSize(void* Original);
 
-	UE_ALLOCATION_FUNCTION(1, 2) static FORCEINLINE_DEBUGGABLE void* MallocZeroed(SIZE_T Count, uint32 Alignment = DEFAULT_ALIGNMENT)
-	{
-		void* Memory = Malloc(Count, Alignment);
-		Memzero(Memory, Count);
-		return Memory;
-	}
+	UE_ALLOCATION_FUNCTION(1, 2) static CORE_API void* MallocZeroed(SIZE_T Count, uint32 Alignment = DEFAULT_ALIGNMENT);
 
 	/**
 	* For some allocators this will return the actual size that should be requested to eliminate
@@ -296,6 +291,7 @@ private:
 	static CORE_API void* ReallocExternal(void* Original, SIZE_T Count, uint32 Alignment = DEFAULT_ALIGNMENT);
 	static CORE_API void FreeExternal(void* Original);
 	static CORE_API SIZE_T GetAllocSizeExternal(void* Original);
+	static CORE_API void* MallocZeroedExternal(SIZE_T Count, uint32 Alignment = DEFAULT_ALIGNMENT);
 	static CORE_API SIZE_T QuantizeSizeExternal(SIZE_T Count, uint32 Alignment = DEFAULT_ALIGNMENT);
 };
 
