[ClosableMpscQueueTest]
 - Fix uint32 overflow

#rnx

diff --git a/Engine/Source/Runtime/Core/Tests/Containers/ConcurrentQueuesTest.cpp b/Engine/Source/Runtime/Core/Tests/Containers/ConcurrentQueuesTest.cpp
--- a/Engine/Source/Runtime/Core/Tests/Containers/ConcurrentQueuesTest.cpp
+++ b/Engine/Source/Runtime/Core/Tests/Containers/ConcurrentQueuesTest.cpp
@@ -630,7 +630,7 @@ namespace ClosableMpscQueueTests
 			FGraphEventArray Producers;
 			Producers.Reserve(ProducersNum);
 
-			TArray<uint32> NumsProduced;
+			TArray<uint64> NumsProduced;
 			NumsProduced.AddZeroed(ProducersNum);
 
 			for (int32 i = 0; i != ProducersNum; ++i)
@@ -638,7 +638,7 @@ namespace ClosableMpscQueueTests
 				Producers.Add(FFunctionGraphTask::CreateAndDispatchWhenReady(
 					[&Queue, NumProduced = &NumsProduced[i]]
 					{
-						uint32 i = 0;
+						uint64 i = 0;
 						while (Queue.Enqueue((void*)(intptr_t)(++i)))
 						{
 						}
@@ -651,14 +651,14 @@ namespace ClosableMpscQueueTests
 			FPlatformProcess::Yield();
 
 			uint64 Consumed = 0;
-			Queue.Close([&Consumed](void* Value) { Consumed += (uint32)(intptr_t)Value; });
+			Queue.Close([&Consumed](void* Value) { Consumed += (uint64)(intptr_t)Value; });
 
 			FTaskGraphInterface::Get().WaitUntilTasksComplete(Producers);
 
 			uint64 Produced = 0;
 			for (int32 i = 0; i != ProducersNum; ++i)
 			{
-				uint32 N = NumsProduced[i];
+				uint64 N = NumsProduced[i];
 				Produced += N * (N + 1) / 2;
 			}
 
