Improve the performance of the hitset.

I've done this by disabling loop unrolling in a specific loop for a 2% uplift, and changing the hitset cutoff to 16 for a further 1% uplift. Hits > 16 are unlikely except for memset/memcpy/memmove anyway.

diff --git a/Engine/Source/Runtime/Core/Private/AutoRTFM/HitSet.h b/Engine/Source/Runtime/Core/Private/AutoRTFM/HitSet.h
--- a/Engine/Source/Runtime/Core/Private/AutoRTFM/HitSet.h
+++ b/Engine/Source/Runtime/Core/Private/AutoRTFM/HitSet.h
@@ -27,8 +27,8 @@ namespace AutoRTFM
 
 			// The capacity is always a power of two so that the range reduction is optimal.
 			static_assert(0 == (InitialCapacity & (InitialCapacity - 1)));
-
-			Payload = new uintptr_t[InitialCapacity]();
+			
+			Payload = static_cast<uintptr_t*>(FMemory::MallocZeroed(InitialCapacity * sizeof(uintptr_t)));
 			ASSERT(nullptr != Payload);
 
 			SixtyFourMinusLogCapacity = 64 - LogInitialCapacity;
@@ -37,7 +37,7 @@ namespace AutoRTFM
 
 		~FHitSet()
 		{
-			delete[] Payload;
+			FMemory::Free(Payload);
 		}
 
 		// Insert something in the HitSet, returning true if the put succeeded
@@ -107,7 +107,7 @@ namespace AutoRTFM
 
 		void Resize()
 		{
-			const uintptr_t* const OldPayload = Payload;
+			uintptr_t* const OldPayload = Payload;
 			const uint32_t OldCapacity = Capacity();
 			const uint32_t OldSize = Size;
 
@@ -119,7 +119,7 @@ namespace AutoRTFM
 
 				IncreaseCapacity();
 
-				Payload = new uintptr_t[Capacity()]();
+				Payload = static_cast<uintptr_t*>(FMemory::MallocZeroed(Capacity() * sizeof(uintptr_t)));
 				ASSERT(nullptr != Payload);
 
 				// Now we need to rehash and reinsert all the items. We need to
@@ -147,7 +147,7 @@ namespace AutoRTFM
 
 				if (NeedAnotherResize)
 				{
-					delete[] Payload;
+					FMemory::Free(Payload);
 					continue;
 				}
 
@@ -156,7 +156,7 @@ namespace AutoRTFM
 
 			ASSERT(OldSize == Size);
 
-			delete[] OldPayload;
+			FMemory::Free(OldPayload);
 
 			return;
 		}
@@ -227,6 +227,11 @@ namespace AutoRTFM
 			// Then, if that fails, we use a second more complicated hash with a linear probe.
 			const uintptr_t Hash = SecondHash(Raw);
 
+			// Clang goes way over the top with loop unrolling, and makes this code noticably
+			// slower as a result. So we disable it!
+#ifdef __clang__
+#pragma clang loop unroll(disable)
+#endif
 			for (uint32_t D = 0; D < LinearProbeDepth; D++)
 			{
 				const uintptr_t I = FirstHashInRange(Hash + D);
diff --git a/Engine/Source/Runtime/Core/Private/AutoRTFM/TransactionInlines.h b/Engine/Source/Runtime/Core/Private/AutoRTFM/TransactionInlines.h
--- a/Engine/Source/Runtime/Core/Private/AutoRTFM/TransactionInlines.h
+++ b/Engine/Source/Runtime/Core/Private/AutoRTFM/TransactionInlines.h
@@ -34,7 +34,7 @@ AUTORTFM_NO_ASAN UE_AUTORTFM_FORCEINLINE void FTransaction::RecordWriteMaxPageSi
 
 AUTORTFM_NO_ASAN UE_AUTORTFM_FORCEINLINE void FTransaction::RecordWrite(void* LogicalAddress, size_t Size)
 {
-    if (0 == Size)
+    if (UNLIKELY(0 == Size))
     {
         return;
     }
@@ -49,7 +49,12 @@ AUTORTFM_NO_ASAN UE_AUTORTFM_FORCEINLINE void FTransaction::RecordWrite(void* Lo
         return;
     }
 
-    if (Size <= FWriteLogBumpAllocator::MaxSize)
+	// The cutoff here is arbitrarily any number less than UINT16_MAX, but its a
+	// weigh up what a good size is. Because the hitset doesn't detect when you
+	// are trying to write to a subregion of a previous hit (like memset something,
+	// then write to an individual element), we've got to balance the cost of
+	// recording meaningless hits, against the potential to hit again.
+    if (Size <= 16)
     {
         FMemoryLocation Key(LogicalAddress);
         Key.SetTopTag(static_cast<uint16_t>(Size));
