Temporary fix issue where RTFM global data pointer was being shared with a child process

#rnx

diff --git a/Engine/Source/Runtime/Core/Private/AutoRTFM/GlobalData.cpp b/Engine/Source/Runtime/Core/Private/AutoRTFM/GlobalData.cpp
--- a/Engine/Source/Runtime/Core/Private/AutoRTFM/GlobalData.cpp
+++ b/Engine/Source/Runtime/Core/Private/AutoRTFM/GlobalData.cpp
@@ -3,6 +3,13 @@
 #if (defined(__AUTORTFM) && __AUTORTFM)
 #include "GlobalData.h"
 #include "Context.h"
+#include <inttypes.h>
+
+#ifdef _MSC_VER
+#include <process.h>
+#else
+#include <unistd.h>
+#endif
 
 namespace AutoRTFM
 {
@@ -33,7 +40,8 @@ void InitializeGlobalDataIfNecessary()
         // - All AutoRTFM instances coordinate together on things like the function table,
         //   lock table, and TLS key.
 
-        constexpr const char* EnvName = "AutoRTFMGlobalData";
+		char EnvName[64];
+		sprintf(EnvName, "AutoRTFMGlobalData_%" PRIu64, (uint64_t)getpid());
 
         if (char* EnvString = getenv(EnvName))
         {
