Make sure IsAutoRTFMRuntimeEnable returns false when AutoRTFM is not compiled on

diff --git a/Engine/Source/Runtime/Core/Private/AutoRTFM/API.cpp b/Engine/Source/Runtime/Core/Private/AutoRTFM/API.cpp
--- a/Engine/Source/Runtime/Core/Private/AutoRTFM/API.cpp
+++ b/Engine/Source/Runtime/Core/Private/AutoRTFM/API.cpp
@@ -37,7 +37,7 @@ namespace AutoRTFM
 	{
 		bool SetAutoRTFMRuntime(EAutoRTFMEnabledState State)
 		{
-			// #noop if AutoRTFM is not compiled in, as GAutoRTFMRuntimeEnabled is a static constexpr when no AutoRTFM compiled in
+			// #noop if AutoRTFM is not compiled in
 #if UE_AUTORTFM
 			switch (GAutoRTFMRuntimeEnabled)
 			{
@@ -63,6 +63,8 @@ namespace AutoRTFM
 
 		bool IsAutoRTFMRuntimeEnabled()
 		{
+			// #noop if AutoRTFM is not compiled in
+#if UE_AUTORTFM
 			switch (GAutoRTFMRuntimeEnabled)
 			{
 			default:
@@ -71,6 +73,9 @@ namespace AutoRTFM
 			case EAutoRTFMEnabledState::AutoRTFM_ForcedEnabled:
 				return true;
 			}
+#else
+			return false;
+#endif
 		}
 	}
 }
