Add test for IterateComponents.
Add IterateAncestors for reverse path iteration.
Use IterateAncestors when building path tree in asset registry as an example case.

diff --git a/Engine/Source/Runtime/AssetRegistry/Private/PathTree.cpp b/Engine/Source/Runtime/AssetRegistry/Private/PathTree.cpp
--- a/Engine/Source/Runtime/AssetRegistry/Private/PathTree.cpp
+++ b/Engine/Source/Runtime/AssetRegistry/Private/PathTree.cpp
@@ -2,6 +2,7 @@
 
 #include "AssetRegistry/PathTree.h"
 
+#include "Misc/PathViews.h"
 #include "Misc/ReverseIterate.h"
 #include "UObject/NameTypes.h"
 
@@ -43,18 +44,17 @@ bool FPathTree::CachePath(FName InPath, TFunctionRef<void(FName)> OnPathAdded)
 	TArray<FName, TInlineAllocator<16>> NewPaths;
 	NewPaths.Add(InPath);
 	ParentPathToChildPaths.FindOrAdd(InPath, {}); // Add this new path with no children
+	FName LastPath = InPath;
 
 	// Walk backwards through the string until we encounter a path we've already created
-	FName LastPath = InPath;
-	FName ParentPath;
-	while (ParentPath != Root)
-	{
-		// Strip the last path element from PathView to get the parent of LastPath
-		// i.e. /Game/Maps/Something -> /Game/Maps 
-		int32 SlashIndex = UE::String::FindLastChar(PathView, '/');	
-		PathView.LeftInline(SlashIndex);
-		ParentPath = PathView.IsEmpty() ? Root : FName(PathView);
-		
+	FPathViews::IterateAncestors(PathView, [this, &NewPaths, &LastPath, PathView](FStringView InAncestor) {
+		if (InAncestor == PathView)
+		{
+			// Initial input path was already handled
+			return true;
+		}
+
+		FName ParentPath = InAncestor.IsEmpty() ? Root : FName(InAncestor);
 		ChildPathToParentPath.FindOrAdd(LastPath, ParentPath);
 
 		uint32 Hash = GetTypeHash(ParentPath);
@@ -63,16 +63,17 @@ bool FPathTree::CachePath(FName InPath, TFunctionRef<void(FName)> OnPathAdded)
 		{
 			// Parent path already existed in tree, no need to continue looking at parents
 			Children->Add(LastPath);
-			break; 
+			return false;
 		}
 		else
 		{
 			ParentPathToChildPaths.AddByHash(Hash, ParentPath).Add(LastPath);
 			NewPaths.Add(ParentPath);
 			LastPath = ParentPath;
+			return true;
 		}
-	} 
-	
+	});
+
 	// Notify caller of each path created in order from root to leaf
 	for (FName NewPath : ReverseIterate(NewPaths))
 	{
diff --git a/Engine/Source/Runtime/Core/Private/Misc/PathViews.cpp b/Engine/Source/Runtime/Core/Private/Misc/PathViews.cpp
--- a/Engine/Source/Runtime/Core/Private/Misc/PathViews.cpp
+++ b/Engine/Source/Runtime/Core/Private/Misc/PathViews.cpp
@@ -11,6 +11,7 @@
 #include "HAL/PlatformMisc.h"
 #include "HAL/PlatformProcess.h"
 #include "Math/UnrealMathUtility.h"
+#include "Misc/AsciiSet.h"
 #include "Misc/AssertionMacros.h"
 #include "Misc/CString.h"
 #include "Misc/Char.h"
@@ -222,7 +223,56 @@ bool FPathViews::IsPathLeaf(FStringView InPath)
 
 void FPathViews::IterateComponents(FStringView InPath, TFunctionRef<void(FStringView)> ComponentVisitor)
 {
-	UE::String::ParseTokensMultiple(InPath, { TEXT('/'), TEXT('\\') }, ComponentVisitor);
+	UE::String::ParseTokensMultiple(InPath, { '/', '\\' }, ComponentVisitor);
+}
+
+void FPathViews::IterateAncestors(FStringView InPath, TFunctionRef<bool(FStringView)> AncestorVisitor)
+{
+	constexpr const FAsciiSet Separators("\\/");
+	constexpr const FAsciiSet NonSeparators(~Separators);
+
+	// First check for a root specifier such as C:\, /Root, //ShareName
+	// When we reach this root we stop searching for path separators and yield it as the final ancestor.
+	FStringView Root;
+	FStringView Remainder;
+	FPathViews::SplitVolumeSpecifier(InPath, Root, Remainder);
+	if (Root.Len() > 0 && Root[Root.Len() - 1] == ':' && Remainder.Len() > 0 && Separators.Contains(Remainder[0]))
+	{
+		// SplitVolumeSpecifier returns Drive: without separator but it's more convenient to yield Drive:\ for use as a path in other APIs
+		Root = { Root.GetData(), Root.Len() + 1 };
+	}
+	else if (Root.Len() == 0 && InPath.Len() > 0 && Separators.Contains(InPath[0]))
+	{
+		// SplitVolumeSpecifier does not return a linux/package path-style root of / - split it out explicitly
+		Root = { InPath.GetData(), 1 };
+	}
+
+	do
+	{
+		// Strip any trailing path separators
+		InPath = FAsciiSet::TrimSuffixWithout(InPath, NonSeparators);
+
+		// If we've iterated into the root, yield that as the last element
+		// Need to use pointer checks rather than length here because we may have stripped extra separators from the root
+		if (InPath.GetData() + InPath.Len() <= Root.GetData() + Root.Len())
+		{
+			break;
+		}
+
+		if (!AncestorVisitor(InPath))
+		{
+			return;
+		}
+
+		// Find next separator
+		InPath = FAsciiSet::TrimSuffixWithout(InPath, Separators);
+	}
+	while (!InPath.IsEmpty());
+
+	if (!Root.IsEmpty())
+	{
+		AncestorVisitor(Root);
+	}
 }
 
 void FPathViews::Split(const FStringView& InPath, FStringView& OutPath, FStringView& OutName, FStringView& OutExt)
diff --git a/Engine/Source/Runtime/Core/Public/Misc/PathViews.h b/Engine/Source/Runtime/Core/Public/Misc/PathViews.h
--- a/Engine/Source/Runtime/Core/Public/Misc/PathViews.h
+++ b/Engine/Source/Runtime/Core/Public/Misc/PathViews.h
@@ -138,14 +138,43 @@ public:
 
 	/**
 	 * Splits InPath into individual directory components, and calls ComponentVisitor on each.
+	 * Note that this function does not properly handle Windows network paths such as //ShareName/Path and
+	 * drive components such as C: (without trailing slash) may not be considered valid paths by some APIs.
 	 *
 	 * Examples:
 	 * "A/B.C" -> {"A", "B.C"}
 	 * "A/B/C" -> {"A", "B", "C"}
 	 * "../../A/B/C.D" -> {"..", "..", "A", "B", "C.D" }
+	 * "/A/B/C" -> { "", "A", "B", "C" }
+	 * "C:/A/B/C" -> { "C:", "A", "B", "C" }
+	 * "//ShareName/Path" -> { "", "", "ShareName", "Path" }
 	 */
 	static CORE_API void IterateComponents(FStringView InPath, TFunctionRef<void(FStringView)> ComponentVisitor);
 
+	/**
+	 * Calls AncestorVisitor with the given path and its ancestor paths.
+	 * If AncestorVisitor returns false, iteration stops.
+	 * If the string begins with a path root as understood by SplitVolumeSpecifier, that will be the last item visited
+	 * with the modification that a separator following a drive specifier will be included by this function.
+	 *
+	 * Examples:
+	 * "" -> { "" }
+	 * "/" -> { "/" }
+	 * "A/B.C" -> { "A/B.C", "A" }
+	 * "A/B/C" -> { "A/B/C", "A/B", "A" }
+	 * "../../A/B/C.D" -> { "../../A/B/C.D", "../../A/B", "../../A", "../..", ".." }
+	 * "/A/B" -> { "/A/B", "/A", "/" }
+	 * "/A/B/" -> { "/A/B", "/A", "/" }
+	 * "C:/A/B" -> { "C:/A/B", "C:/A", "C:/" }
+	 * "C:A/B" -> { "C:A/B", "C:A", "C:" }
+	 * "\\ShareName/A/B" -> { "\\ShareName/A/B", "\\ShareName/A", "\\ShareName" }
+	 *
+	 * When moving to the next ancestor, extra trailing separators are also removed, e.g.
+	 * "/A//B" -> { "/A//B", "/A" }
+	 * Note that separators are not removed except at the end of the current path, as with visiting "/A//B" above.
+	 */
+	static CORE_API void IterateAncestors(FStringView InPath, TFunctionRef<bool(FStringView)> AncestorVisitor);
+
 	/**
 	 * Splits a path into three parts, any of which may be empty: the path, the clean name, and the extension.
 	 *
diff --git a/Engine/Source/Runtime/Core/Tests/Misc/PathViewsTest.cpp b/Engine/Source/Runtime/Core/Tests/Misc/PathViewsTest.cpp
--- a/Engine/Source/Runtime/Core/Tests/Misc/PathViewsTest.cpp
+++ b/Engine/Source/Runtime/Core/Tests/Misc/PathViewsTest.cpp
@@ -1,6 +1,8 @@
 // Copyright Epic Games, Inc. All Rights Reserved.
 
 #include "PathTests.h"
+
+#include "Algo/Compare.h"
 #include "Containers/StringView.h"
 #include "Misc/PathViews.h"
 #include "Tests/TestHarnessAdapter.h"
@@ -520,6 +522,150 @@ TEST_CASE_NAMED(FPathViewsSplitTest, "System::Core::Misc::PathViews::Split", "[A
 	RunSplitTest(TEXT("C:\\Folder\\First.Last\\File.tar.gz"), TEXT("C:\\Folder\\First.Last"), TEXT("File.tar"), TEXT("gz"));
 }
 
+TEST_CASE_NAMED(FPathViewsIterateTest, "System::Core::Misc::PathViews::IterateComponents", "[ApplicationContextMask][SmokeFilter]")
+{
+	auto RunIterateTest = [](const TCHAR* InPath, auto&&... InExpectedComponents) {
+		TArray<FStringView> Components;
+		TArray<FStringView> ExpectedComponents({ Forward<decltype(InExpectedComponents)>(InExpectedComponents)... });
+		FPathViews::IterateComponents(InPath, [&Components](FStringView Comp) { Components.Add(Comp); });
+		if (!Algo::Compare(Components, ExpectedComponents))
+		{
+			FAIL_CHECK(FString::Printf(TEXT("Incorrect components splitting %s - got %d, expected %d. Got components (%s), expected components (%s)"),
+				InPath,
+				Components.Num(),
+				ExpectedComponents.Num(),
+				*FString::Join(Components, TEXT(",")),
+				*FString::Join(ExpectedComponents, TEXT(","))));
+		}
+	};
+
+	RunIterateTest(TEXT(""), TEXT(""));
+	RunIterateTest(TEXT(".txt"), TEXT(".txt"));
+	RunIterateTest(TEXT(".tar.gz"), TEXT(".tar.gz"));
+	RunIterateTest(TEXT(".tar.gz/"), TEXT(".tar.gz"), TEXT(""));
+	RunIterateTest(TEXT(".tar.gz\\"), TEXT(".tar.gz"), TEXT(""));
+	RunIterateTest(TEXT("."), TEXT("."));
+	RunIterateTest(TEXT(".."), TEXT(".."));
+	RunIterateTest(TEXT("File"), TEXT("File"));
+	RunIterateTest(TEXT("File.txt"), TEXT("File.txt"));
+	RunIterateTest(TEXT("File.tar.gz"), TEXT("File.tar.gz"));
+	RunIterateTest(TEXT("File.tar.gz/"), TEXT("File.tar.gz"), TEXT(""));
+	RunIterateTest(TEXT("File.tar.gz\\"), TEXT("File.tar.gz"), TEXT(""));
+
+	RunIterateTest(TEXT("C:/Folder"), TEXT("C:"), TEXT("Folder"));
+	RunIterateTest(TEXT("C:/Folder/"), TEXT("C:"), TEXT("Folder"), TEXT(""));
+	RunIterateTest(TEXT("C:/Folder/File"), TEXT("C:"), TEXT("Folder"), TEXT("File"));
+	RunIterateTest(TEXT("C:/Folder/File.txt"), TEXT("C:"), TEXT("Folder"), TEXT("File.txt"));
+	RunIterateTest(TEXT("C:/Folder/File.tar.gz"), TEXT("C:"), TEXT("Folder"), TEXT("File.tar.gz"));
+	RunIterateTest(TEXT("C:/Folder/First.Last/File"), TEXT("C:"), TEXT("Folder"), TEXT("First.Last"), TEXT("File"));
+	RunIterateTest(TEXT("C:/Folder/First.Last/File.txt"), TEXT("C:"), TEXT("Folder"), TEXT("First.Last"), TEXT("File.txt"));
+	RunIterateTest(TEXT("C:/Folder/First.Last/File.tar.gz"), TEXT("C:"), TEXT("Folder"), TEXT("First.Last"), TEXT("File.tar.gz"));
+
+	RunIterateTest(TEXT("C:\\Folder"), TEXT("C:"), TEXT("Folder"));
+	RunIterateTest(TEXT("C:\\Folder\\"), TEXT("C:"), TEXT("Folder"), TEXT(""));
+	RunIterateTest(TEXT("C:\\Folder\\File"), TEXT("C:"), TEXT("Folder"), TEXT("File"));
+	RunIterateTest(TEXT("C:\\Folder\\First.Last\\"), TEXT("C:"), TEXT("Folder"), TEXT("First.Last"), TEXT(""));
+	RunIterateTest(TEXT("C:\\Folder\\First.Last\\File"), TEXT("C:"), TEXT("Folder"), TEXT("First.Last"), TEXT("File"));
+	RunIterateTest(TEXT("C:\\Folder\\First.Last\\File.txt"), TEXT("C:"), TEXT("Folder"), TEXT("First.Last"), TEXT("File.txt"));
+	RunIterateTest(TEXT("C:\\Folder\\First.Last\\File.tar.gz"), TEXT("C:"), TEXT("Folder"), TEXT("First.Last"), TEXT("File.tar.gz"));
+
+	RunIterateTest(TEXT("//ShareName/Path"), TEXT(""), TEXT(""), TEXT("ShareName"), TEXT("Path"));
+	RunIterateTest(TEXT("//ShareName/Path/"), TEXT(""), TEXT(""), TEXT("ShareName"), TEXT("Path"), TEXT(""));
+	RunIterateTest(TEXT("//ShareName/Path/Sub"), TEXT(""), TEXT(""), TEXT("ShareName"), TEXT("Path"), TEXT("Sub"));
+	RunIterateTest(TEXT("//ShareName/Path/Sub/File"), TEXT(""), TEXT(""), TEXT("ShareName"), TEXT("Path"), TEXT("Sub"), TEXT("File"));
+	RunIterateTest(TEXT("//ShareName/Path/Sub/File.txt"), TEXT(""), TEXT(""), TEXT("ShareName"), TEXT("Path"), TEXT("Sub"), TEXT("File.txt"));
+	RunIterateTest(TEXT("//ShareName/Path/Sub/File.tar.gz"), TEXT(""), TEXT(""), TEXT("ShareName"), TEXT("Path"), TEXT("Sub"), TEXT("File.tar.gz"));
+
+	RunIterateTest(TEXT("/MountPoint"), TEXT(""), TEXT("MountPoint"));
+	RunIterateTest(TEXT("/MountPoint/"), TEXT(""), TEXT("MountPoint"), TEXT(""));
+	RunIterateTest(TEXT("/MountPoint/Sub"), TEXT(""), TEXT("MountPoint"), TEXT("Sub"));
+	RunIterateTest(TEXT("/MountPoint/Sub/File"), TEXT(""), TEXT("MountPoint"), TEXT("Sub"), TEXT("File"));
+	RunIterateTest(TEXT("/MountPoint/Sub/File.txt"), TEXT(""), TEXT("MountPoint"), TEXT("Sub"), TEXT("File.txt"));
+	RunIterateTest(TEXT("/MountPoint/Sub/File.tar.gz"), TEXT(""), TEXT("MountPoint"), TEXT("Sub"), TEXT("File.tar.gz"));
+}
+
+TEST_CASE_NAMED(FPathViewsAncestorTest, "System::Core::Misc::PathViews::IterateAncestors", "[ApplicationContextMask][SmokeFilter]")
+{
+	auto RunAncestorTest = [](FStringView InPath, auto&&... Args) {
+		TArray<FStringView> Ancestors;
+		TArray<FStringView> ExpectedAncestors{ Args... };
+		FPathViews::IterateAncestors(InPath, [&Ancestors](FStringView Comp) { Ancestors.Add(Comp); return true; });
+		if (!Algo::Compare(Ancestors, ExpectedAncestors))
+		{
+			FStringBuilderBase Builder;
+			Builder << TEXT("Incorrect ancestors splitting ") << InPath << TEXT("- got ") << Ancestors.Num() << TEXT(", expected ") << ExpectedAncestors.Num();
+			Builder << LINE_TERMINATOR TEXT("Got ancestors      (");
+			Builder.Join(Ancestors, TEXT(","));
+			Builder << LINE_TERMINATOR TEXT("Expected ancestors (");
+			Builder.Join(ExpectedAncestors, TEXT(", ")) << TEXT(")");
+			FAIL_CHECK(Builder.ToString());
+		}
+	};
+
+	RunAncestorTest(TEXTVIEW(""), TEXTVIEW(""));
+	RunAncestorTest(TEXTVIEW("/"), TEXTVIEW("/"));
+	RunAncestorTest(TEXTVIEW("\\"), TEXTVIEW("\\"));
+	RunAncestorTest(TEXTVIEW(".txt"), TEXTVIEW(".txt"));
+	RunAncestorTest(TEXTVIEW(".tar.gz"), TEXTVIEW(".tar.gz"));
+	RunAncestorTest(TEXTVIEW(".tar.gz/"), TEXTVIEW(".tar.gz"));
+	RunAncestorTest(TEXTVIEW(".tar.gz\\"), TEXTVIEW(".tar.gz"));
+	RunAncestorTest(TEXTVIEW("."), TEXTVIEW("."));
+	RunAncestorTest(TEXTVIEW(".."), TEXTVIEW(".."));
+	RunAncestorTest(TEXTVIEW("File"), TEXTVIEW("File"));
+	RunAncestorTest(TEXTVIEW("File/"), TEXTVIEW("File"));
+	RunAncestorTest(TEXTVIEW("File.txt"), TEXTVIEW("File.txt"));
+	RunAncestorTest(TEXTVIEW("File.tar.gz"), TEXTVIEW("File.tar.gz"));
+	RunAncestorTest(TEXTVIEW("File.tar.gz/"), TEXTVIEW("File.tar.gz"));
+	RunAncestorTest(TEXTVIEW("File.tar.gz\\"), TEXTVIEW("File.tar.gz"));
+
+	RunAncestorTest(TEXTVIEW("C:/Folder"), TEXTVIEW("C:/Folder"), TEXTVIEW("C:/"));
+	RunAncestorTest(TEXTVIEW("C:/Folder/"), TEXTVIEW("C:/Folder"), TEXTVIEW("C:/"));
+	RunAncestorTest(TEXTVIEW("C:/Folder/File"), TEXTVIEW("C:/Folder/File"), TEXTVIEW("C:/Folder"), TEXTVIEW("C:/"));
+	RunAncestorTest(TEXTVIEW("C:/Folder/File.txt"), TEXTVIEW("C:/Folder/File.txt"), TEXTVIEW("C:/Folder"), TEXTVIEW("C:/"));
+	RunAncestorTest(TEXTVIEW("C:/Folder/File.tar.gz"), TEXTVIEW("C:/Folder/File.tar.gz"), TEXTVIEW("C:/Folder"), TEXTVIEW("C:/"));
+	RunAncestorTest(TEXTVIEW("C:/Folder/First.Last/File"), TEXTVIEW("C:/Folder/First.Last/File"), TEXTVIEW("C:/Folder/First.Last"), TEXTVIEW("C:/Folder"), TEXTVIEW("C:/"));
+	RunAncestorTest(TEXTVIEW("C:/Folder/First.Last/File.txt"), TEXTVIEW("C:/Folder/First.Last/File.txt"), TEXTVIEW("C:/Folder/First.Last"), TEXTVIEW("C:/Folder"), TEXTVIEW("C:/"));
+	RunAncestorTest(TEXTVIEW("C:/Folder/First.Last/File.tar.gz"), TEXTVIEW("C:/Folder/First.Last/File.tar.gz"), TEXTVIEW("C:/Folder/First.Last"), TEXTVIEW("C:/Folder"), TEXTVIEW("C:/"));
+
+	// Windows path specifying drive letter but relative path rather than absolute
+	// Not well handled elsewhere in the engine but we parse it here
+	RunAncestorTest(TEXTVIEW("C:Folder/File"), TEXTVIEW("C:Folder/File"), TEXTVIEW("C:Folder"), TEXTVIEW("C:"));
+
+	RunAncestorTest(TEXTVIEW("C:\\Folder"), TEXTVIEW("C:\\Folder"), TEXTVIEW("C:\\"));
+	RunAncestorTest(TEXTVIEW("C:\\Folder\\"), TEXTVIEW("C:\\Folder"), TEXTVIEW("C:\\"));
+	RunAncestorTest(TEXTVIEW("C:\\Folder\\File"), TEXTVIEW("C:\\Folder\\File"), TEXTVIEW("C:\\Folder"), TEXTVIEW("C:\\"));
+	RunAncestorTest(TEXTVIEW("C:\\Folder\\First.Last\\"), TEXTVIEW("C:\\Folder\\First.Last"), TEXTVIEW("C:\\Folder"), TEXTVIEW("C:\\"));
+	RunAncestorTest(TEXTVIEW("C:\\Folder\\First.Last\\File"), TEXTVIEW("C:\\Folder\\First.Last\\File"), TEXTVIEW("C:\\Folder\\First.Last"), TEXTVIEW("C:\\Folder"), TEXTVIEW("C:\\"));
+	RunAncestorTest(TEXTVIEW("C:\\Folder\\First.Last\\File.txt"), TEXTVIEW("C:\\Folder\\First.Last\\File.txt"), TEXTVIEW("C:\\Folder\\First.Last"), TEXTVIEW("C:\\Folder"), TEXTVIEW("C:\\"));
+	RunAncestorTest(TEXTVIEW("C:\\Folder\\First.Last\\File.tar.gz"), TEXTVIEW("C:\\Folder\\First.Last\\File.tar.gz"), TEXTVIEW("C:\\Folder\\First.Last"), TEXTVIEW("C:\\Folder"), TEXTVIEW("C:\\"));
+
+	RunAncestorTest(TEXTVIEW("/MountPoint"), TEXTVIEW("/MountPoint"), TEXTVIEW("/"));
+	RunAncestorTest(TEXTVIEW("/MountPoint/"), TEXTVIEW("/MountPoint"), TEXTVIEW("/"));
+	RunAncestorTest(TEXTVIEW("/MountPoint/Sub"), TEXTVIEW("/MountPoint/Sub"), TEXTVIEW("/MountPoint"), TEXTVIEW("/"));
+	RunAncestorTest(TEXTVIEW("/MountPoint/Sub/File"), TEXTVIEW("/MountPoint/Sub/File"), TEXTVIEW("/MountPoint/Sub"), TEXTVIEW("/MountPoint"), TEXTVIEW("/"));
+	RunAncestorTest(TEXTVIEW("/MountPoint/Sub/File.txt"), TEXTVIEW("/MountPoint/Sub/File.txt"), TEXTVIEW("/MountPoint/Sub"), TEXTVIEW("/MountPoint"), TEXTVIEW("/"));
+	RunAncestorTest(TEXTVIEW("/MountPoint/Sub/File.tar.gz"), TEXTVIEW("/MountPoint/Sub/File.tar.gz"), TEXTVIEW("/MountPoint/Sub"), TEXTVIEW("/MountPoint"), TEXTVIEW("/"));
+
+	RunAncestorTest(TEXTVIEW("\\\\ShareName/Path"), TEXTVIEW("\\\\ShareName/Path"), TEXTVIEW("\\\\ShareName"));
+	RunAncestorTest(TEXTVIEW("\\\\ShareName/Path/"), TEXTVIEW("\\\\ShareName/Path"), TEXTVIEW("\\\\ShareName"));
+	RunAncestorTest(TEXTVIEW("\\\\ShareName/Path/Sub"), TEXTVIEW("\\\\ShareName/Path/Sub"), TEXTVIEW("\\\\ShareName/Path"), TEXTVIEW("\\\\ShareName"));
+	RunAncestorTest(TEXTVIEW("\\\\ShareName/Path/Sub/File"), TEXTVIEW("\\\\ShareName/Path/Sub/File"), TEXTVIEW("\\\\ShareName/Path/Sub"), TEXTVIEW("\\\\ShareName/Path"), TEXTVIEW("\\\\ShareName"));
+	RunAncestorTest(TEXTVIEW("\\\\ShareName/Path/Sub/File.txt"), TEXTVIEW("\\\\ShareName/Path/Sub/File.txt"), TEXTVIEW("\\\\ShareName/Path/Sub"), TEXTVIEW("\\\\ShareName/Path"), TEXTVIEW("\\\\ShareName"));
+	RunAncestorTest(TEXTVIEW("\\\\ShareName/Path/Sub/File.tar.gz"), TEXTVIEW("\\\\ShareName/Path/Sub/File.tar.gz"), TEXTVIEW("\\\\ShareName/Path/Sub"), TEXTVIEW("\\\\ShareName/Path"), TEXTVIEW("\\\\ShareName"));
+
+	// Duplicated path separators
+	RunAncestorTest(TEXTVIEW("//MountPoint"), TEXTVIEW("//MountPoint"));     // Double leadings slashes treated like a windows share
+	RunAncestorTest(TEXTVIEW("//\\MountPoint"), TEXTVIEW("//\\MountPoint")); // SplitVolumeSpecifier treats this as a root
+	RunAncestorTest(TEXTVIEW("/\\MountPoint"), TEXTVIEW("/\\MountPoint"));   // Looks like a share name to SplitVolumeSpecifier
+	RunAncestorTest(TEXTVIEW("/MountPoint//"), TEXTVIEW("/MountPoint"), TEXT("/"));
+	RunAncestorTest(TEXTVIEW("/MountPoint//Sub"),
+		TEXTVIEW("/MountPoint//Sub"), // We don't remove the internal duplication because we're using the source string
+		TEXTVIEW("/MountPoint"),      // We strip the extras when moving to the next ancestor
+		TEXTVIEW("/"));
+	RunAncestorTest(TEXTVIEW("/MountPoint/Sub//File"), TEXTVIEW("/MountPoint/Sub//File"), TEXTVIEW("/MountPoint/Sub"), TEXTVIEW("/MountPoint"), TEXTVIEW("/"));
+	RunAncestorTest(TEXTVIEW("\\\\\\NotShareName/Path"), TEXTVIEW("\\\\\\NotShareName/Path"), TEXTVIEW("\\\\\\NotShareName")); // SplitVolumeSpecifier allows extra separators after the two-separator prefix
+}
+
 TEST_CASE_NAMED(FPathViewsAppendTest, "System::Core::Misc::PathViews::Append", "[ApplicationContextMask][SmokeFilter]")
 {
 	TStringBuilder<256> Path;
@@ -894,33 +1040,34 @@ TEST_CASE_NAMED(FPathViewsVolumeSpecifierTest, "System::Core::Misc::PathViews::V
 		FStringView Remainder;
 	};
 	FTestCase TestCases[] = {
-		{ TEXTVIEW(""),					false,	TEXTVIEW(""),			TEXTVIEW("") },
-		{ TEXTVIEW("D:"),				true,	TEXTVIEW("D:"),			TEXTVIEW("") },
-		{ TEXTVIEW("D:/"),				false,	TEXTVIEW("D:"),			TEXTVIEW("/") },
-		{ TEXTVIEW("D:\\"),				false,	TEXTVIEW("D:"),			TEXTVIEW("\\") },
-		{ TEXTVIEW("D:root/path"),		true,	TEXTVIEW("D:"),			TEXTVIEW("root/path") },
-		{ TEXTVIEW("D:/root/path"),		false,	TEXTVIEW("D:"),			TEXTVIEW("/root/path") },
-		{ TEXTVIEW("D:\\root\\path"),	false,	TEXTVIEW("D:"),			TEXTVIEW("\\root\\path") },
-		{ TEXTVIEW("//volume"),			false,	TEXTVIEW("//volume"),	TEXTVIEW("") },
-		{ TEXTVIEW("\\\\volume"),		false,	TEXTVIEW("\\\\volume"),	TEXTVIEW("") },
-		{ TEXTVIEW("/\\volume"),		false,	TEXTVIEW("/\\volume"),	TEXTVIEW("") }, // Poorly defined case, somewhat arbitrary
-		{ TEXTVIEW("\\/volume"),		false,	TEXTVIEW("\\/volume"),	TEXTVIEW("") }, // Poorly defined case, somewhat arbitrary
-		{ TEXTVIEW("//volume/"),		false,	TEXTVIEW("//volume"),	TEXTVIEW("/") },
-		{ TEXTVIEW("//volume/root"),	false,	TEXTVIEW("//volume"),	TEXTVIEW("/root") },
-		{ TEXTVIEW("/root/path"),		false,	TEXTVIEW(""),			TEXTVIEW("/root/path") },
-		{ TEXTVIEW("\\root\\path"),		false,	TEXTVIEW(""),			TEXTVIEW("\\root\\path") },
-		{ TEXTVIEW("root/path"),		false,	TEXTVIEW(""),			TEXTVIEW("root/path") },
-		{ TEXTVIEW("/"),				false,	TEXTVIEW(""),			TEXTVIEW("/") },
-		{ TEXTVIEW("\\"),				false,	TEXTVIEW(""),			TEXTVIEW("\\") },
-		{ TEXTVIEW("//"),				false,	TEXTVIEW("//"),			TEXTVIEW("") }, // Poorly defined case, somewhat arbitrary
-		{ TEXTVIEW("\\\\"),				false,	TEXTVIEW("\\\\"),		TEXTVIEW("") }, // Poorly defined case, somewhat arbitrary
-		{ TEXTVIEW("/\\"),				false,	TEXTVIEW("/\\"),		TEXTVIEW("") }, // Poorly defined case, somewhat arbitrary
-		{ TEXTVIEW("\\/"),				false,	TEXTVIEW("\\/"),		TEXTVIEW("") }, // Poorly defined case, somewhat arbitrary
-		{ TEXTVIEW("/:"),				false,	TEXTVIEW(""),			TEXTVIEW("/:") }, // Poorly defined case, somewhat arbitrary
-		{ TEXTVIEW(":"),				true,	TEXTVIEW(":"),			TEXTVIEW("") }, // Poorly defined case, somewhat arbitrary
-		{ TEXTVIEW(":/"),				false,	TEXTVIEW(":"),			TEXTVIEW("/") }, // Poorly defined case, somewhat arbitrary
-		{ TEXTVIEW(":root"),			true,	TEXTVIEW(":"),			TEXTVIEW("root") }, // Poorly defined case, somewhat arbitrary
-		{ TEXTVIEW("////volume/path"),	false,	TEXTVIEW("////volume"),	TEXTVIEW("/path") }, // Poorly defined case, somewhat arbitrary, @see RemoveDuplicateSlashes
+		{TEXTVIEW(""),                 false, TEXTVIEW(""),           TEXTVIEW("")            },
+		{ TEXTVIEW("D:"),              true,  TEXTVIEW("D:"),         TEXTVIEW("")            },
+		{ TEXTVIEW("D:/"),             false, TEXTVIEW("D:"),         TEXTVIEW("/")           },
+		{ TEXTVIEW("D:\\"),            false, TEXTVIEW("D:"),         TEXTVIEW("\\")          },
+		{ TEXTVIEW("D:root/path"),     true,  TEXTVIEW("D:"),         TEXTVIEW("root/path")   },
+		{ TEXTVIEW("D:/root/path"),    false, TEXTVIEW("D:"),         TEXTVIEW("/root/path")  },
+		{ TEXTVIEW("D:\\root\\path"),  false, TEXTVIEW("D:"),         TEXTVIEW("\\root\\path")},
+		{ TEXTVIEW("//volume"),        false, TEXTVIEW("//volume"),   TEXTVIEW("")            },
+		{ TEXTVIEW("\\\\volume"),      false, TEXTVIEW("\\\\volume"), TEXTVIEW("")            },
+		{ TEXTVIEW("/\\volume"),       false, TEXTVIEW("/\\volume"),  TEXTVIEW("")            }, // Poorly defined case, somewhat arbitrary
+		{ TEXTVIEW("\\/volume"),       false, TEXTVIEW("\\/volume"),  TEXTVIEW("")            }, // Poorly defined case, somewhat arbitrary
+		{ TEXTVIEW("//\\volume"),      false, TEXTVIEW("//\\volume"), TEXTVIEW("")            }, // Poorly defined case, somewhat arbitrary
+		{ TEXTVIEW("//volume/"),       false, TEXTVIEW("//volume"),   TEXTVIEW("/")           },
+		{ TEXTVIEW("//volume/root"),   false, TEXTVIEW("//volume"),   TEXTVIEW("/root")       },
+		{ TEXTVIEW("/root/path"),      false, TEXTVIEW(""),           TEXTVIEW("/root/path")  },
+		{ TEXTVIEW("\\root\\path"),    false, TEXTVIEW(""),           TEXTVIEW("\\root\\path")},
+		{ TEXTVIEW("root/path"),       false, TEXTVIEW(""),           TEXTVIEW("root/path")   },
+		{ TEXTVIEW("/"),               false, TEXTVIEW(""),           TEXTVIEW("/")           },
+		{ TEXTVIEW("\\"),              false, TEXTVIEW(""),           TEXTVIEW("\\")          },
+		{ TEXTVIEW("//"),              false, TEXTVIEW("//"),         TEXTVIEW("")            }, // Poorly defined case, somewhat arbitrary
+		{ TEXTVIEW("\\\\"),            false, TEXTVIEW("\\\\"),       TEXTVIEW("")            }, // Poorly defined case, somewhat arbitrary
+		{ TEXTVIEW("/\\"),             false, TEXTVIEW("/\\"),        TEXTVIEW("")            }, // Poorly defined case, somewhat arbitrary
+		{ TEXTVIEW("\\/"),             false, TEXTVIEW("\\/"),        TEXTVIEW("")            }, // Poorly defined case, somewhat arbitrary
+		{ TEXTVIEW("/:"),              false, TEXTVIEW(""),           TEXTVIEW("/:")          }, // Poorly defined case, somewhat arbitrary
+		{ TEXTVIEW(":"),               true,  TEXTVIEW(":"),          TEXTVIEW("")            }, // Poorly defined case, somewhat arbitrary
+		{ TEXTVIEW(":/"),              false, TEXTVIEW(":"),          TEXTVIEW("/")           }, // Poorly defined case, somewhat arbitrary
+		{ TEXTVIEW(":root"),           true,  TEXTVIEW(":"),          TEXTVIEW("root")        }, // Poorly defined case, somewhat arbitrary
+		{ TEXTVIEW("////volume/path"), false, TEXTVIEW("////volume"), TEXTVIEW("/path")       }, // Poorly defined case, somewhat arbitrary, @see RemoveDuplicateSlashes
 	};
 
 	for (const FTestCase& TestCase : TestCases)
