Also exposes the `DeploymentName` tag for custom crash reporting code to be able to make use of.

diff --git a/Engine/Source/Runtime/Core/Private/GenericPlatform/GenericPlatformCrashContext.cpp b/Engine/Source/Runtime/Core/Private/GenericPlatform/GenericPlatformCrashContext.cpp
--- a/Engine/Source/Runtime/Core/Private/GenericPlatform/GenericPlatformCrashContext.cpp
+++ b/Engine/Source/Runtime/Core/Private/GenericPlatform/GenericPlatformCrashContext.cpp
@@ -187,6 +187,7 @@ const TCHAR* const FGenericCrashContext::CrashGUIDRootPrefix = TEXT("UECC-");
 
 const TCHAR* const FGenericCrashContext::CrashContextExtension = TEXT(".runtime-xml");
 const TCHAR* const FGenericCrashContext::RuntimePropertiesTag = TEXT( "RuntimeProperties" );
+const TCHAR* const FGenericCrashContext::DeploymentNameTag = TEXT( "DeploymentName" );
 const TCHAR* const FGenericCrashContext::PlatformPropertiesTag = TEXT( "PlatformProperties" );
 const TCHAR* const FGenericCrashContext::EngineDataTag = TEXT( "EngineData" );
 const TCHAR* const FGenericCrashContext::GameDataTag = TEXT( "GameData" );
@@ -835,7 +836,7 @@ void FGenericCrashContext::SerializeSessionContext(FString& Buffer)
 	AddCrashPropertyInternal(Buffer, TEXT("EngineMode"), NCached::Session.EngineMode);
 	AddCrashPropertyInternal(Buffer, TEXT("EngineModeEx"), NCached::Session.EngineModeEx);
 
-	AddCrashPropertyInternal(Buffer, TEXT("DeploymentName"), NCached::Session.DeploymentName);
+	AddCrashPropertyInternal(Buffer, FGenericCrashContext::DeploymentNameTag, NCached::Session.DeploymentName);
 
 	AddCrashPropertyInternal(Buffer, TEXT("EngineVersion"), NCached::Session.EngineVersion); 
 	AddCrashPropertyInternal(Buffer, TEXT("EngineCompatibleVersion"), NCached::Session.EngineCompatibleVersion); 
diff --git a/Engine/Source/Runtime/Core/Public/GenericPlatform/GenericPlatformCrashContext.h b/Engine/Source/Runtime/Core/Public/GenericPlatform/GenericPlatformCrashContext.h
--- a/Engine/Source/Runtime/Core/Public/GenericPlatform/GenericPlatformCrashContext.h
+++ b/Engine/Source/Runtime/Core/Public/GenericPlatform/GenericPlatformCrashContext.h
@@ -377,6 +377,7 @@ public:
 
 	CORE_API static const TCHAR* const CrashContextExtension;
 	CORE_API static const TCHAR* const RuntimePropertiesTag;
+	CORE_API static const TCHAR* const DeploymentNameTag;
 	CORE_API static const TCHAR* const PlatformPropertiesTag;
 	CORE_API static const TCHAR* const EngineDataTag;
 	CORE_API static const TCHAR* const GameDataTag;
