[Core]
* Added this to lambda capture to fix warning with new intel compiler

diff --git a/Engine/Source/Runtime/Core/Public/Misc/TypeContainer.h b/Engine/Source/Runtime/Core/Public/Misc/TypeContainer.h
--- a/Engine/Source/Runtime/Core/Public/Misc/TypeContainer.h
+++ b/Engine/Source/Runtime/Core/Public/Misc/TypeContainer.h
@@ -307,7 +307,7 @@ public:
 	{
 		TSharedPtr<IInstanceProvider> Provider = MakeShareable(
 			new TFunctionInstanceProvider<R>(
-				[=]() -> TSharedPtr<void, Mode> {
+				[=, this]() -> TSharedPtr<void, Mode> {
 					return CreateFunc();
 				}
 			)
@@ -330,7 +330,7 @@ public:
 	{
 		TSharedPtr<IInstanceProvider> Provider = MakeShareable(
 			new TFunctionInstanceProvider<R>(
-				[=]() -> TSharedPtr<void, Mode> {
+				[=, this]() -> TSharedPtr<void, Mode> {
 					return CreateFunc(GetInstance<P0>(), GetInstance<P>()...);
 				}
 			)
