Allow for disabling AutoRTFM Runtime with a define

diff --git a/Engine/Source/Runtime/Core/Private/AutoRTFM/API.cpp b/Engine/Source/Runtime/Core/Private/AutoRTFM/API.cpp
--- a/Engine/Source/Runtime/Core/Private/AutoRTFM/API.cpp
+++ b/Engine/Source/Runtime/Core/Private/AutoRTFM/API.cpp
@@ -8,7 +8,11 @@
 namespace
 {
 	// Move this to a local only and use functions to access this
+#if UE_AUTORTFM_ENABLED_RUNTIME_BY_DEFAULT
 	int GAutoRTFMRuntimeEnabled = AutoRTFM::EAutoRTFMEnabledState::AutoRTFM_Enabled;
+#else
+	int GAutoRTFMRuntimeEnabled = AutoRTFM::EAutoRTFMEnabledState::AutoRTFM_Disabled;
+#endif // UE_AUTORTFM_ENABLED_RUNTIME_BY_DEFAULT
 
 	void UpdateAutoRTFMRuntimeCrashData()
 	{
diff --git a/Engine/Source/Runtime/Core/Public/AutoRTFM/AutoRTFM.h b/Engine/Source/Runtime/Core/Public/AutoRTFM/AutoRTFM.h
--- a/Engine/Source/Runtime/Core/Public/AutoRTFM/AutoRTFM.h
+++ b/Engine/Source/Runtime/Core/Public/AutoRTFM/AutoRTFM.h
@@ -8,6 +8,10 @@
 #define UE_AUTORTFM 0
 #endif
 
+#if !defined(UE_AUTORTFM_ENABLED_RUNTIME_BY_DEFAULT)
+#define UE_AUTORTFM_ENABLED_RUNTIME_BY_DEFAULT 1
+#endif
+
 #if UE_AUTORTFM
 #if defined(__has_feature)
 #if __has_feature(autostm)
