- Some platforms were trying to query ProjectDir during the first file operation, that ProjectDir was using to determine if it's Staged or not, so I added FPaths::CanGetProjectDir() which will return false when ProjectDir is being computed. This is a very edge case and is unlikely to be needed often or maybe even ever again.

diff --git a/Engine/Source/Runtime/Core/Private/Misc/Paths.cpp b/Engine/Source/Runtime/Core/Private/Misc/Paths.cpp
--- a/Engine/Source/Runtime/Core/Private/Misc/Paths.cpp
+++ b/Engine/Source/Runtime/Core/Private/Misc/Paths.cpp
@@ -144,15 +144,31 @@ namespace UE4Paths_Private
 	}
 }
 
+namespace UE::Paths
+{
+	bool bIsComputingStaged = false;
+}
+
+bool FPaths::CanGetProjectDir()
+{
+	return !UE::Paths::bIsComputingStaged;
+}
+
 bool FPaths::IsStaged()
 {
-	// if it's not a program and the platform requires cooked data, then we can assume staged, but if not, then
-	// check if it went through the staging process
-	static bool bIsStaged =
-#if !IS_PROGRAM
-		FPlatformProperties::RequiresCookedData() ||
-#endif
-		FileExists(Combine(EngineConfigDir(), FString::Printf(TEXT("StagedBuild_%s.ini"), FApp::GetProjectName())));
+	static bool bHasInitialized = false;
+	static bool bIsStaged;
+
+	if (!bHasInitialized)
+	{
+		UE::Paths::bIsComputingStaged = true;
+
+		bIsStaged = FPlatformProperties::RequiresCookedData() ||
+			FileExists(Combine(EngineConfigDir(), FString::Printf(TEXT("StagedBuild_%s.ini"), FApp::GetProjectName())));
+
+		UE::Paths::bIsComputingStaged = false;
+	}
+
 	return bIsStaged;
 }
 
diff --git a/Engine/Source/Runtime/Core/Public/Misc/Paths.h b/Engine/Source/Runtime/Core/Public/Misc/Paths.h
--- a/Engine/Source/Runtime/Core/Public/Misc/Paths.h
+++ b/Engine/Source/Runtime/Core/Public/Misc/Paths.h
@@ -36,6 +36,11 @@ class FPaths
 {
 public:
 
+	/**
+	 * @return true if ProjectDir() is allowed to be called
+	 */
+	static CORE_API bool CanGetProjectDir();
+
 	/**
 	  * @return true if the build has been staged and uses the simplified remapped directory structure made when staging
 	 */
