ColorManagement: Prevent failed config check in Lightmass.

#rnx

diff --git a/Engine/Source/Runtime/Core/Private/ColorManagement/ColorSpace.cpp b/Engine/Source/Runtime/Core/Private/ColorManagement/ColorSpace.cpp
--- a/Engine/Source/Runtime/Core/Private/ColorManagement/ColorSpace.cpp
+++ b/Engine/Source/Runtime/Core/Private/ColorManagement/ColorSpace.cpp
@@ -1,6 +1,9 @@
 // Copyright Epic Games, Inc. All Rights Reserved.
 
 #include "ColorManagement/ColorSpace.h"
+
+#include "Logging/LogCategory.h"
+#include "Logging/LogMacros.h"
 #include "Misc/ConfigCacheIni.h"
 #include "Math/VectorRegister.h"
 
@@ -16,24 +19,37 @@ void PreloadWorkingColorSpace()
 		return;
 	}
 
-	check(GConfig != nullptr && GConfig->IsReadyForUse());
-
-	bool bIsWorkingColorSpaceInConfig = true;
-	TStaticArray<FVector2d, 4> Chromaticities;
-	bIsWorkingColorSpaceInConfig &= GConfig->GetVector2D(TEXT("/Script/Engine.RendererSettings"), TEXT("RedChromaticityCoordinate"),   Chromaticities[0], GEngineIni);
-	bIsWorkingColorSpaceInConfig &= GConfig->GetVector2D(TEXT("/Script/Engine.RendererSettings"), TEXT("GreenChromaticityCoordinate"), Chromaticities[1], GEngineIni);
-	bIsWorkingColorSpaceInConfig &= GConfig->GetVector2D(TEXT("/Script/Engine.RendererSettings"), TEXT("BlueChromaticityCoordinate"),  Chromaticities[2], GEngineIni);
-	bIsWorkingColorSpaceInConfig &= GConfig->GetVector2D(TEXT("/Script/Engine.RendererSettings"), TEXT("WhiteChromaticityCoordinate"), Chromaticities[3], GEngineIni);
-
-	if (bIsWorkingColorSpaceInConfig)
+	// Note: Other non-engine programs may not have access project configs. For those edge cases, we must fall back to default.
+	if (GConfig != nullptr)
 	{
-		FColorSpace::SetWorking(FColorSpace(Chromaticities[0], Chromaticities[1], Chromaticities[2], Chromaticities[3]));
+		check(GConfig->IsReadyForUse());
+
+		int32 LoadedChromaticities = 0;
+		TStaticArray<FVector2d, 4> Chromaticities;
+		LoadedChromaticities += static_cast<int32>(GConfig->GetVector2D(TEXT("/Script/Engine.RendererSettings"), TEXT("RedChromaticityCoordinate"), Chromaticities[0], GEngineIni));
+		LoadedChromaticities += static_cast<int32>(GConfig->GetVector2D(TEXT("/Script/Engine.RendererSettings"), TEXT("GreenChromaticityCoordinate"), Chromaticities[1], GEngineIni));
+		LoadedChromaticities += static_cast<int32>(GConfig->GetVector2D(TEXT("/Script/Engine.RendererSettings"), TEXT("BlueChromaticityCoordinate"), Chromaticities[2], GEngineIni));
+		LoadedChromaticities += static_cast<int32>(GConfig->GetVector2D(TEXT("/Script/Engine.RendererSettings"), TEXT("WhiteChromaticityCoordinate"), Chromaticities[3], GEngineIni));
+
+		if (LoadedChromaticities == 0)
+		{
+			// No-op: The working color space wasn't modified/serialized to the config so we keep the implicit sRGB default.
+		}
+		else if (LoadedChromaticities == 4)
+		{
+			FColorSpace::SetWorking(FColorSpace(Chromaticities[0], Chromaticities[1], Chromaticities[2], Chromaticities[3]));
+		}
+		else
+		{
+			UE_LOG(LogCore, Warning, TEXT("Working color space in config is missing chromaticities, falling back to sRGB."));
+		}
 	}
 	else
 	{
-		// The working color space wasn't modified/serialized to the config so we keep the implicit sRGB default.
-		bIsWorkingColorSpaceReadyForUse = true;
+		UE_LOG(LogCore, Display, TEXT("Working color space in config is not available, falling back to sRGB."));
 	}
+
+	bIsWorkingColorSpaceReadyForUse = true;
 };
 
 const FColorSpace& FColorSpace::GetWorking()
