Core: Tests: Fixed a multithreading issue in some multithreading tests.

The problem was caused by using an ordinary 'int' variable as a counter for the number of async tasks completed. Because of that, the following code:

while (TasksComplete != TaskCount)
{
}

was optimized on some platforms as an infinite loop because TaskComplete was assumed to not be changing (from the main thread's perspective) and the whole check was removed leaving only a single jmp instruction to its own address.

Resolved by replacing the 'int' with std::atomic<int> which makes the compiler aware that this counter may be changed from outside of the main thread.

#rnx

diff --git a/Engine/Source/Runtime/Core/Tests/Async/MutexTest.cpp b/Engine/Source/Runtime/Core/Tests/Async/MutexTest.cpp
--- a/Engine/Source/Runtime/Core/Tests/Async/MutexTest.cpp
+++ b/Engine/Source/Runtime/Core/Tests/Async/MutexTest.cpp
@@ -19,7 +19,7 @@ TEST_CASE("Core::Async::Mutex", "[Core][Async][Slow][LinuxSkip]")
 	SECTION("FMutex IsLocked() and TryLock()")
 	{
 		FThread Threads[TaskCount];
-		uint32 TasksComplete = 0;
+		std::atomic<uint32> TasksComplete = 0;
 		
 		FMutex Mutex;
 		CHECK(Mutex.TryLock());
@@ -58,7 +58,7 @@ TEST_CASE("Core::Async::Mutex", "[Core][Async][Slow][LinuxSkip]")
 	SECTION("FMutex with TUniqueLock which uses SlowLock() and SlowUnlock()")
 	{
 		FThread Threads[TaskCount];
-		uint32 TasksComplete = 0;
+		std::atomic<uint32> TasksComplete = 0;
 		
 		FMutex Mutex;
 		Mutex.Lock();
