Replaced the check for mismatched SP modes between TSharedPtr/Ref and TSharedFromThis with something that should be quicker to compile.

diff --git a/Engine/Source/Runtime/Core/Public/Templates/SharedPointer.h b/Engine/Source/Runtime/Core/Public/Templates/SharedPointer.h
--- a/Engine/Source/Runtime/Core/Public/Templates/SharedPointer.h
+++ b/Engine/Source/Runtime/Core/Public/Templates/SharedPointer.h
@@ -141,16 +141,6 @@ namespace UE::Core::Private
 }
 
 
-// TSharedPtr of one mode to a type which has a TSharedFromThis only of another mode is illegal.
-// A type which does not inherit TSharedFromThis at all is ok.
-// We only check this inside the constructor because we don't necessarily have the full type of T when we declare a TSharedPtr<T>.
-#define UE_TSHAREDPTR_STATIC_ASSERT_VALID_MODE(ObjectType, Mode) \
-	static_assert( \
-		std::is_convertible_v<ObjectType*, TSharedFromThis<ObjectType, Mode>*> || \
-		!std::is_convertible_v<ObjectType*, TSharedFromThis<ObjectType, (Mode == ESPMode::NotThreadSafe) ? ESPMode::ThreadSafe : ESPMode::NotThreadSafe>*>, \
-		"You cannot use a TSharedPtr of one mode with a type which inherits TSharedFromThis of another mode." \
-	);
-
 /**
  * TSharedRef is a non-nullable, non-intrusive reference-counted authoritative object reference.
  *
@@ -182,8 +172,6 @@ public:
 		: Object( InObject )
 		, SharedReferenceCount( SharedPointerInternals::NewDefaultReferenceController< Mode >( InObject ) )
 	{
-		UE_TSHAREDPTR_STATIC_ASSERT_VALID_MODE(ObjectType, Mode)
-
 		Init(InObject);
 	}
 
@@ -202,8 +190,6 @@ public:
 		: Object( InObject )
 		, SharedReferenceCount( SharedPointerInternals::NewCustomReferenceController< Mode >( InObject, Forward< DeleterType >( InDeleter ) ) )
 	{
-		UE_TSHAREDPTR_STATIC_ASSERT_VALID_MODE(ObjectType, Mode)
-
 		Init(InObject);
 	}
 
@@ -235,15 +221,13 @@ public:
 		: Object( InRawPtrProxy.Object )
 		, SharedReferenceCount( SharedPointerInternals::NewDefaultReferenceController< Mode >( InRawPtrProxy.Object ) )
 	{
-		UE_TSHAREDPTR_STATIC_ASSERT_VALID_MODE(ObjectType, Mode)
-
 		// If the following assert goes off, it means a TSharedRef was initialized from a nullptr object pointer.
 		// Shared references must never be nullptr, so either pass a valid object or consider using TSharedPtr instead.
 		check( InRawPtrProxy.Object != nullptr );
 
 		// If the object happens to be derived from TSharedFromThis, the following method
 		// will prime the object with a weak pointer to itself.
-		SharedPointerInternals::EnableSharedFromThis( this, InRawPtrProxy.Object, InRawPtrProxy.Object );
+		SharedPointerInternals::EnableSharedFromThis( this, InRawPtrProxy.Object );
 	}
 
 	/**
@@ -262,15 +246,13 @@ public:
 		: Object( InRawPtrProxy.Object )
 		, SharedReferenceCount( SharedPointerInternals::NewCustomReferenceController< Mode >( InRawPtrProxy.Object, InRawPtrProxy.Deleter ) )
 	{
-		UE_TSHAREDPTR_STATIC_ASSERT_VALID_MODE(ObjectType, Mode)
-
 		// If the following assert goes off, it means a TSharedRef was initialized from a nullptr object pointer.
 		// Shared references must never be nullptr, so either pass a valid object or consider using TSharedPtr instead.
 		check( InRawPtrProxy.Object != nullptr );
 
 		// If the object happens to be derived from TSharedFromThis, the following method
 		// will prime the object with a weak pointer to itself.
-		SharedPointerInternals::EnableSharedFromThis( this, InRawPtrProxy.Object, InRawPtrProxy.Object );
+		SharedPointerInternals::EnableSharedFromThis( this, InRawPtrProxy.Object );
 	}
 
 	/**
@@ -289,15 +271,13 @@ public:
 		: Object( InRawPtrProxy.Object )
 		, SharedReferenceCount( SharedPointerInternals::NewCustomReferenceController< Mode >( InRawPtrProxy.Object, MoveTemp( InRawPtrProxy.Deleter ) ) )
 	{
-		UE_TSHAREDPTR_STATIC_ASSERT_VALID_MODE(ObjectType, Mode)
-
 		// If the following assert goes off, it means a TSharedRef was initialized from a nullptr object pointer.
 		// Shared references must never be nullptr, so either pass a valid object or consider using TSharedPtr instead.
 		check( InRawPtrProxy.Object != nullptr );
 
 		// If the object happens to be derived from TSharedFromThis, the following method
 		// will prime the object with a weak pointer to itself.
-		SharedPointerInternals::EnableSharedFromThis( this, InRawPtrProxy.Object, InRawPtrProxy.Object );
+		SharedPointerInternals::EnableSharedFromThis( this, InRawPtrProxy.Object );
 	}
 
 	/**
@@ -581,7 +561,7 @@ private:
 
 		// If the object happens to be derived from TSharedFromThis, the following method
 		// will prime the object with a weak pointer to itself.
-		SharedPointerInternals::EnableSharedFromThis(this, InObject, InObject);
+		SharedPointerInternals::EnableSharedFromThis(this, InObject);
 	}
 
 	/**
@@ -661,8 +641,6 @@ private:
 		: Object(InObject)
 		, SharedReferenceCount(InSharedReferenceCount)
 	{
-		UE_TSHAREDPTR_STATIC_ASSERT_VALID_MODE(ObjectType, Mode)
-
 		Init(InObject);
 	}
 };
@@ -718,11 +696,9 @@ public:
 		: Object( InObject )
 		, SharedReferenceCount( SharedPointerInternals::NewDefaultReferenceController< Mode >( InObject ) )
 	{
-		UE_TSHAREDPTR_STATIC_ASSERT_VALID_MODE(ObjectType, Mode)
-
 		// If the object happens to be derived from TSharedFromThis, the following method
 		// will prime the object with a weak pointer to itself.
-		SharedPointerInternals::EnableSharedFromThis( this, InObject, InObject );
+		SharedPointerInternals::EnableSharedFromThis( this, InObject );
 	}
 
 	/**
@@ -741,11 +717,9 @@ public:
 		: Object( InObject )
 		, SharedReferenceCount( SharedPointerInternals::NewCustomReferenceController< Mode >( InObject, Forward< DeleterType >( InDeleter ) ) )
 	{
-		UE_TSHAREDPTR_STATIC_ASSERT_VALID_MODE(ObjectType, Mode)
-
 		// If the object happens to be derived from TSharedFromThis, the following method
 		// will prime the object with a weak pointer to itself.
-		SharedPointerInternals::EnableSharedFromThis( this, InObject, InObject );
+		SharedPointerInternals::EnableSharedFromThis( this, InObject );
 	}
 
 	/**
@@ -762,11 +736,9 @@ public:
 		: Object( InRawPtrProxy.Object )
 		, SharedReferenceCount( SharedPointerInternals::NewDefaultReferenceController< Mode >( InRawPtrProxy.Object ) )
 	{
-		UE_TSHAREDPTR_STATIC_ASSERT_VALID_MODE(ObjectType, Mode)
-
 		// If the object happens to be derived from TSharedFromThis, the following method
 		// will prime the object with a weak pointer to itself.
-		SharedPointerInternals::EnableSharedFromThis( this, InRawPtrProxy.Object, InRawPtrProxy.Object );
+		SharedPointerInternals::EnableSharedFromThis( this, InRawPtrProxy.Object );
 	}
 
 	/**
@@ -784,11 +756,9 @@ public:
 		: Object( InRawPtrProxy.Object )
 		, SharedReferenceCount( SharedPointerInternals::NewCustomReferenceController< Mode >( InRawPtrProxy.Object, InRawPtrProxy.Deleter ) )
 	{
-		UE_TSHAREDPTR_STATIC_ASSERT_VALID_MODE(ObjectType, Mode)
-
 		// If the object happens to be derived from TSharedFromThis, the following method
 		// will prime the object with a weak pointer to itself.
-		SharedPointerInternals::EnableSharedFromThis( this, InRawPtrProxy.Object, InRawPtrProxy.Object );
+		SharedPointerInternals::EnableSharedFromThis( this, InRawPtrProxy.Object );
 	}
 
 	/**
@@ -806,11 +776,9 @@ public:
 		: Object( InRawPtrProxy.Object )
 		, SharedReferenceCount( SharedPointerInternals::NewCustomReferenceController< Mode >( InRawPtrProxy.Object, MoveTemp( InRawPtrProxy.Deleter ) ) )
 	{
-		UE_TSHAREDPTR_STATIC_ASSERT_VALID_MODE(ObjectType, Mode)
-
 		// If the object happens to be derived from TSharedFromThis, the following method
 		// will prime the object with a weak pointer to itself.
-		SharedPointerInternals::EnableSharedFromThis( this, InRawPtrProxy.Object, InRawPtrProxy.Object );
+		SharedPointerInternals::EnableSharedFromThis( this, InRawPtrProxy.Object );
 	}
 
 	/**
@@ -1754,11 +1722,11 @@ public:		// @todo: Ideally this would be private, but template sharing problems
 	 * Note that until this function is called, calls to AsShared() will result in an empty pointer.
 	 */
 	template< class SharedPtrType, class OtherType >
-	FORCEINLINE void UpdateWeakReferenceInternal( TSharedPtr< SharedPtrType, Mode > const* InSharedPtr, OtherType* InObject ) const
+	FORCEINLINE void UpdateWeakReferenceInternal( TSharedPtr< SharedPtrType, Mode > const* InSharedPtr ) const
 	{
 		if( !WeakThis.IsValid() )
 		{
-			WeakThis = TSharedPtr< ObjectType, Mode >( *InSharedPtr, InObject );
+			WeakThis = TSharedPtr< ObjectType, Mode >( *InSharedPtr, static_cast<ObjectType*>(const_cast<TSharedFromThis*>(this)) );
 		}
 	}
 
@@ -1768,11 +1736,11 @@ public:		// @todo: Ideally this would be private, but template sharing problems
 	 * Note that until this function is called, calls to AsShared() will result in an empty pointer.
 	 */
 	template< class SharedRefType, class OtherType >
-	FORCEINLINE void UpdateWeakReferenceInternal( TSharedRef< SharedRefType, Mode > const* InSharedRef, OtherType* InObject ) const
+	FORCEINLINE void UpdateWeakReferenceInternal( TSharedRef< SharedRefType, Mode > const* InSharedRef ) const
 	{
 		if( !WeakThis.IsValid() )
 		{
-			WeakThis = TSharedRef< ObjectType, Mode >( *InSharedRef, InObject );
+			WeakThis = TSharedRef< ObjectType, Mode >( *InSharedRef, static_cast<ObjectType*>(const_cast<TSharedFromThis*>(this)) );
 		}
 	}
 
diff --git a/Engine/Source/Runtime/Core/Public/Templates/SharedPointerInternals.h b/Engine/Source/Runtime/Core/Public/Templates/SharedPointerInternals.h
--- a/Engine/Source/Runtime/Core/Public/Templates/SharedPointerInternals.h
+++ b/Engine/Source/Runtime/Core/Public/Templates/SharedPointerInternals.h
@@ -834,51 +834,21 @@ namespace SharedPointerInternals
 		TReferenceControllerBase<Mode>* ReferenceController;
 	};
 
-
-	/** Templated helper function (const) that creates a shared pointer from an object instance */
-	template< class SharedPtrType, class ObjectType, class OtherType, ESPMode Mode >
-	FORCEINLINE void EnableSharedFromThis( TSharedPtr< SharedPtrType, Mode > const* InSharedPtr, ObjectType const* InObject, TSharedFromThis< OtherType, Mode > const* InShareable )
+	/** Templated helper function (const) that creates a shared pointer/reference from an object instance */
+	template <class SharedPtrType, class ObjectType, class OtherType, ESPMode Mode>
+	FORCEINLINE void EnableSharedFromThis(const volatile SharedPtrType* InSharedPtr, const volatile ObjectType* InObject)
 	{
-		if( InShareable != nullptr )
-		{
-			InShareable->UpdateWeakReferenceInternal( InSharedPtr, const_cast< ObjectType* >( InObject ) );
-		}
-	}
-
-
-	/** Templated helper function that creates a shared pointer from an object instance */
-	template< class SharedPtrType, class ObjectType, class OtherType, ESPMode Mode >
-	FORCEINLINE void EnableSharedFromThis( TSharedPtr< SharedPtrType, Mode >* InSharedPtr, ObjectType const* InObject, TSharedFromThis< OtherType, Mode > const* InShareable )
-	{
-		if( InShareable != nullptr )
-		{
-			InShareable->UpdateWeakReferenceInternal( InSharedPtr, const_cast< ObjectType* >( InObject ) );
-		}
-	}
-
+		static_assert(SharedPtrType::Mode == Mode, "You cannot use a TSharedPtr of one mode with a type which inherits TSharedFromThis of another mode");
 
-	/** Templated helper function (const) that creates a shared reference from an object instance */
-	template< class SharedRefType, class ObjectType, class OtherType, ESPMode Mode >
-	FORCEINLINE void EnableSharedFromThis( TSharedRef< SharedRefType, Mode > const* InSharedRef, ObjectType const* InObject, TSharedFromThis< OtherType, Mode > const* InShareable )
-	{
-		if( InShareable != nullptr )
+		if (InObject)
 		{
-			InShareable->UpdateWeakReferenceInternal( InSharedRef, const_cast< ObjectType* >( InObject ) );
+			const_cast<const TSharedFromThis<OtherType, Mode>*>(InObject)->UpdateWeakReferenceInternal(const_cast<const SharedPtrType*>(InSharedPtr), const_cast<ObjectType*>(InObject));
 		}
 	}
 
 
-	/** Templated helper function that creates a shared reference from an object instance */
-	template< class SharedRefType, class ObjectType, class OtherType, ESPMode Mode >
-	FORCEINLINE void EnableSharedFromThis( TSharedRef< SharedRefType, Mode >* InSharedRef, ObjectType const* InObject, TSharedFromThis< OtherType, Mode > const* InShareable )
+	/** Templated helper catch-all function, accomplice to the above helper functions */
+	constexpr void EnableSharedFromThis(const volatile void*, const volatile void*)
 	{
-		if( InShareable != nullptr )
-		{
-			InShareable->UpdateWeakReferenceInternal( InSharedRef, const_cast< ObjectType* >( InObject ) );
-		}
 	}
-
-
-	/** Templated helper catch-all function, accomplice to the above helper functions */
-	constexpr void EnableSharedFromThis( ... ) { }
 }
