Fix invalid Language Code Identifier sent in crash report.
#tests Validated in editor that the value stored in NCached::Session.LanguageLCID is now valid. In editor, triggered a dummy crash and checked the generated CrashContext.runtime-xml file: the language is no longer invalid while it was always 0 before..

diff --git a/Engine/Source/Runtime/Core/Private/GenericPlatform/GenericPlatformCrashContext.cpp b/Engine/Source/Runtime/Core/Private/GenericPlatform/GenericPlatformCrashContext.cpp
--- a/Engine/Source/Runtime/Core/Private/GenericPlatform/GenericPlatformCrashContext.cpp
+++ b/Engine/Source/Runtime/Core/Private/GenericPlatform/GenericPlatformCrashContext.cpp
@@ -368,25 +368,15 @@ void FGenericCrashContext::Initialize()
 	NCached::Set(NCached::Session.CrashGUIDRoot, *CrashGUIDRoot);
 	UE_LOG(LogInit, Log, TEXT("Session CrashGUID >====================================================\n         Session CrashGUID >   %s\n         Session CrashGUID >===================================================="), *CrashGUIDRoot);
 
-	if (GIsRunning)
+	if (FInternationalization::IsAvailable())
 	{
-		if (FInternationalization::IsAvailable())
+		FInternationalization& Internationalization = FInternationalization::Get();
+		NCached::Session.LanguageLCID = Internationalization.GetCurrentCulture()->GetLCID();
+
+		Internationalization.OnCultureChanged().AddLambda([]()
 		{
 			NCached::Session.LanguageLCID = FInternationalization::Get().GetCurrentCulture()->GetLCID();
-		}
-		else
-		{
-			FCulturePtr DefaultCulture = FInternationalization::Get().GetCulture(TEXT("en"));
-			if (DefaultCulture.IsValid())
-			{
-				NCached::Session.LanguageLCID = DefaultCulture->GetLCID();
-			}
-			else
-			{
-				const int DefaultCultureLCID = 1033;
-				NCached::Session.LanguageLCID = DefaultCultureLCID;
-			}
-		}
+		});
 	}
 
 	// Initialize delegate for updating SecondsSinceStart, because FPlatformTime::Seconds() is not POSIX safe.
