Execute the OnHangDelegate when the Present thread hangs too.

diff --git a/Engine/Source/Runtime/Core/Private/HAL/ThreadHeartBeat.cpp b/Engine/Source/Runtime/Core/Private/HAL/ThreadHeartBeat.cpp
--- a/Engine/Source/Runtime/Core/Private/HAL/ThreadHeartBeat.cpp
+++ b/Engine/Source/Runtime/Core/Private/HAL/ThreadHeartBeat.cpp
@@ -164,6 +164,9 @@ void FORCENOINLINE FThreadHeartBeat::OnPresentHang(double HangDuration)
 	// We want to avoid all memory allocations if a hang is detected.
 	// Force a crash in a way that will generate a crash report.
 
+	// Delegate implementation will be called from the hang detector thread and not from the hung thread
+	OnHangDelegate.ExecuteIfBound(FThreadHeartBeat::PresentThreadId);
+
 	// Avoiding calling RaiseException here will keep OnPresentHang on the top of the crash callstack,
 	// making crash bucketing easier when looking at retail crash dumps on supported platforms.
 
