Replaced the check for mismatched SP modes between TSharedPtr/Ref and TSharedFromThis with something that should be quicker to compile.
Removed the multiple overloads of EnableSharedFromThis.

diff --git a/Engine/Source/Runtime/Core/Public/Templates/SharedPointer.h b/Engine/Source/Runtime/Core/Public/Templates/SharedPointer.h
--- a/Engine/Source/Runtime/Core/Public/Templates/SharedPointer.h
+++ b/Engine/Source/Runtime/Core/Public/Templates/SharedPointer.h
@@ -141,16 +141,6 @@ namespace UE::Core::Private
 }
 
 
-// TSharedPtr of one mode to a type which has a TSharedFromThis only of another mode is illegal.
-// A type which does not inherit TSharedFromThis at all is ok.
-// We only check this inside the constructor because we don't necessarily have the full type of T when we declare a TSharedPtr<T>.
-#define UE_TSHAREDPTR_STATIC_ASSERT_VALID_MODE(ObjectType, Mode) \
-	static_assert( \
-		std::is_convertible_v<ObjectType*, TSharedFromThis<ObjectType, Mode>*> || \
-		!std::is_convertible_v<ObjectType*, TSharedFromThis<ObjectType, (Mode == ESPMode::NotThreadSafe) ? ESPMode::ThreadSafe : ESPMode::NotThreadSafe>*>, \
-		"You cannot use a TSharedPtr of one mode with a type which inherits TSharedFromThis of another mode." \
-	);
-
 /**
  * TSharedRef is a non-nullable, non-intrusive reference-counted authoritative object reference.
  *
@@ -182,8 +172,6 @@ public:
 		: Object( InObject )
 		, SharedReferenceCount( SharedPointerInternals::NewDefaultReferenceController< Mode >( InObject ) )
 	{
-		UE_TSHAREDPTR_STATIC_ASSERT_VALID_MODE(ObjectType, Mode)
-
 		Init(InObject);
 	}
 
@@ -202,8 +190,6 @@ public:
 		: Object( InObject )
 		, SharedReferenceCount( SharedPointerInternals::NewCustomReferenceController< Mode >( InObject, Forward< DeleterType >( InDeleter ) ) )
 	{
-		UE_TSHAREDPTR_STATIC_ASSERT_VALID_MODE(ObjectType, Mode)
-
 		Init(InObject);
 	}
 
@@ -235,8 +221,6 @@ public:
 		: Object( InRawPtrProxy.Object )
 		, SharedReferenceCount( SharedPointerInternals::NewDefaultReferenceController< Mode >( InRawPtrProxy.Object ) )
 	{
-		UE_TSHAREDPTR_STATIC_ASSERT_VALID_MODE(ObjectType, Mode)
-
 		// If the following assert goes off, it means a TSharedRef was initialized from a nullptr object pointer.
 		// Shared references must never be nullptr, so either pass a valid object or consider using TSharedPtr instead.
 		check( InRawPtrProxy.Object != nullptr );
@@ -262,8 +246,6 @@ public:
 		: Object( InRawPtrProxy.Object )
 		, SharedReferenceCount( SharedPointerInternals::NewCustomReferenceController< Mode >( InRawPtrProxy.Object, InRawPtrProxy.Deleter ) )
 	{
-		UE_TSHAREDPTR_STATIC_ASSERT_VALID_MODE(ObjectType, Mode)
-
 		// If the following assert goes off, it means a TSharedRef was initialized from a nullptr object pointer.
 		// Shared references must never be nullptr, so either pass a valid object or consider using TSharedPtr instead.
 		check( InRawPtrProxy.Object != nullptr );
@@ -289,8 +271,6 @@ public:
 		: Object( InRawPtrProxy.Object )
 		, SharedReferenceCount( SharedPointerInternals::NewCustomReferenceController< Mode >( InRawPtrProxy.Object, MoveTemp( InRawPtrProxy.Deleter ) ) )
 	{
-		UE_TSHAREDPTR_STATIC_ASSERT_VALID_MODE(ObjectType, Mode)
-
 		// If the following assert goes off, it means a TSharedRef was initialized from a nullptr object pointer.
 		// Shared references must never be nullptr, so either pass a valid object or consider using TSharedPtr instead.
 		check( InRawPtrProxy.Object != nullptr );
@@ -661,8 +641,6 @@ private:
 		: Object(InObject)
 		, SharedReferenceCount(InSharedReferenceCount)
 	{
-		UE_TSHAREDPTR_STATIC_ASSERT_VALID_MODE(ObjectType, Mode)
-
 		Init(InObject);
 	}
 };
@@ -718,8 +696,6 @@ public:
 		: Object( InObject )
 		, SharedReferenceCount( SharedPointerInternals::NewDefaultReferenceController< Mode >( InObject ) )
 	{
-		UE_TSHAREDPTR_STATIC_ASSERT_VALID_MODE(ObjectType, Mode)
-
 		// If the object happens to be derived from TSharedFromThis, the following method
 		// will prime the object with a weak pointer to itself.
 		SharedPointerInternals::EnableSharedFromThis( this, InObject, InObject );
@@ -741,8 +717,6 @@ public:
 		: Object( InObject )
 		, SharedReferenceCount( SharedPointerInternals::NewCustomReferenceController< Mode >( InObject, Forward< DeleterType >( InDeleter ) ) )
 	{
-		UE_TSHAREDPTR_STATIC_ASSERT_VALID_MODE(ObjectType, Mode)
-
 		// If the object happens to be derived from TSharedFromThis, the following method
 		// will prime the object with a weak pointer to itself.
 		SharedPointerInternals::EnableSharedFromThis( this, InObject, InObject );
@@ -762,8 +736,6 @@ public:
 		: Object( InRawPtrProxy.Object )
 		, SharedReferenceCount( SharedPointerInternals::NewDefaultReferenceController< Mode >( InRawPtrProxy.Object ) )
 	{
-		UE_TSHAREDPTR_STATIC_ASSERT_VALID_MODE(ObjectType, Mode)
-
 		// If the object happens to be derived from TSharedFromThis, the following method
 		// will prime the object with a weak pointer to itself.
 		SharedPointerInternals::EnableSharedFromThis( this, InRawPtrProxy.Object, InRawPtrProxy.Object );
@@ -784,8 +756,6 @@ public:
 		: Object( InRawPtrProxy.Object )
 		, SharedReferenceCount( SharedPointerInternals::NewCustomReferenceController< Mode >( InRawPtrProxy.Object, InRawPtrProxy.Deleter ) )
 	{
-		UE_TSHAREDPTR_STATIC_ASSERT_VALID_MODE(ObjectType, Mode)
-
 		// If the object happens to be derived from TSharedFromThis, the following method
 		// will prime the object with a weak pointer to itself.
 		SharedPointerInternals::EnableSharedFromThis( this, InRawPtrProxy.Object, InRawPtrProxy.Object );
@@ -806,8 +776,6 @@ public:
 		: Object( InRawPtrProxy.Object )
 		, SharedReferenceCount( SharedPointerInternals::NewCustomReferenceController< Mode >( InRawPtrProxy.Object, MoveTemp( InRawPtrProxy.Deleter ) ) )
 	{
-		UE_TSHAREDPTR_STATIC_ASSERT_VALID_MODE(ObjectType, Mode)
-
 		// If the object happens to be derived from TSharedFromThis, the following method
 		// will prime the object with a weak pointer to itself.
 		SharedPointerInternals::EnableSharedFromThis( this, InRawPtrProxy.Object, InRawPtrProxy.Object );
diff --git a/Engine/Source/Runtime/Core/Public/Templates/SharedPointerInternals.h b/Engine/Source/Runtime/Core/Public/Templates/SharedPointerInternals.h
--- a/Engine/Source/Runtime/Core/Public/Templates/SharedPointerInternals.h
+++ b/Engine/Source/Runtime/Core/Public/Templates/SharedPointerInternals.h
@@ -837,48 +837,21 @@ namespace SharedPointerInternals
 
 	/** Templated helper function (const) that creates a shared pointer from an object instance */
 	template< class SharedPtrType, class ObjectType, class OtherType, ESPMode Mode >
-	FORCEINLINE void EnableSharedFromThis( TSharedPtr< SharedPtrType, Mode > const* InSharedPtr, ObjectType const* InObject, TSharedFromThis< OtherType, Mode > const* InShareable )
+	FORCEINLINE void EnableSharedFromThis( SharedPtrType* InSharedPtrOrRef, ObjectType const* InObject, TSharedFromThis< OtherType, Mode > const* InShareable )
 	{
-		if( InShareable != nullptr )
+		if constexpr (SharedPtrType::Mode == Mode)
 		{
-			InShareable->UpdateWeakReferenceInternal( InSharedPtr, const_cast< ObjectType* >( InObject ) );
-		}
-	}
-
-
-	/** Templated helper function that creates a shared pointer from an object instance */
-	template< class SharedPtrType, class ObjectType, class OtherType, ESPMode Mode >
-	FORCEINLINE void EnableSharedFromThis( TSharedPtr< SharedPtrType, Mode >* InSharedPtr, ObjectType const* InObject, TSharedFromThis< OtherType, Mode > const* InShareable )
-	{
-		if( InShareable != nullptr )
-		{
-			InShareable->UpdateWeakReferenceInternal( InSharedPtr, const_cast< ObjectType* >( InObject ) );
-		}
-	}
-
-
-	/** Templated helper function (const) that creates a shared reference from an object instance */
-	template< class SharedRefType, class ObjectType, class OtherType, ESPMode Mode >
-	FORCEINLINE void EnableSharedFromThis( TSharedRef< SharedRefType, Mode > const* InSharedRef, ObjectType const* InObject, TSharedFromThis< OtherType, Mode > const* InShareable )
-	{
-		if( InShareable != nullptr )
-		{
-			InShareable->UpdateWeakReferenceInternal( InSharedRef, const_cast< ObjectType* >( InObject ) );
+			if( InShareable != nullptr )
+			{
+				InShareable->UpdateWeakReferenceInternal( InSharedPtrOrRef, const_cast< ObjectType* >( InObject ) );
+			}
 		}
-	}
-
-
-	/** Templated helper function that creates a shared reference from an object instance */
-	template< class SharedRefType, class ObjectType, class OtherType, ESPMode Mode >
-	FORCEINLINE void EnableSharedFromThis( TSharedRef< SharedRefType, Mode >* InSharedRef, ObjectType const* InObject, TSharedFromThis< OtherType, Mode > const* InShareable )
-	{
-		if( InShareable != nullptr )
+		else
 		{
-			InShareable->UpdateWeakReferenceInternal( InSharedRef, const_cast< ObjectType* >( InObject ) );
+			static_assert(sizeof(ObjectType) == 0, "You cannot use a TSharedPtr of one mode with a type which inherits TSharedFromThis of another mode.");
 		}
 	}
 
-
-	/** Templated helper catch-all function, accomplice to the above helper functions */
+	/** Templated helper catch-all function, accomplice to the above helper function */
 	constexpr void EnableSharedFromThis( ... ) { }
 }
