Replaced ImplicitConv usage with std::is_convertible_v in TSharedPtr classes - should compile faster.

diff --git a/Engine/Source/Runtime/Core/Public/Templates/SharedPointer.h b/Engine/Source/Runtime/Core/Public/Templates/SharedPointer.h
--- a/Engine/Source/Runtime/Core/Public/Templates/SharedPointer.h
+++ b/Engine/Source/Runtime/Core/Public/Templates/SharedPointer.h
@@ -166,7 +166,7 @@ public:
 	 */
 	template <
 		typename OtherType
-		UE_REQUIRES(UE_REQUIRES_EXPR(ImplicitConv<ObjectType*>((OtherType*)nullptr)))
+		UE_REQUIRES(std::is_convertible_v<OtherType*, ObjectType*>)
 	>
 	FORCEINLINE explicit TSharedRef( OtherType* InObject )
 		: Object( InObject )
@@ -184,7 +184,7 @@ public:
 	template <
 		typename OtherType,
 		typename DeleterType
-		UE_REQUIRES(UE_REQUIRES_EXPR(ImplicitConv<ObjectType*>((OtherType*)nullptr)))
+		UE_REQUIRES(std::is_convertible_v<OtherType*, ObjectType*>)
 	>
 	FORCEINLINE TSharedRef( OtherType* InObject, DeleterType&& InDeleter )
 		: Object( InObject )
@@ -215,7 +215,7 @@ public:
 	// NOTE: The following is an Unreal extension to standard shared_ptr behavior
 	template <
 		typename OtherType
-		UE_REQUIRES(UE_REQUIRES_EXPR(ImplicitConv<ObjectType*>((OtherType*)nullptr)))
+		UE_REQUIRES(std::is_convertible_v<OtherType*, ObjectType*>)
 	>
 	FORCEINLINE TSharedRef( SharedPointerInternals::TRawPtrProxy< OtherType > const& InRawPtrProxy )
 		: Object( InRawPtrProxy.Object )
@@ -240,7 +240,7 @@ public:
 	template <
 		typename OtherType,
 		typename DeleterType
-		UE_REQUIRES(UE_REQUIRES_EXPR(ImplicitConv<ObjectType*>((OtherType*)nullptr)))
+		UE_REQUIRES(std::is_convertible_v<OtherType*, ObjectType*>)
 	>
 	FORCEINLINE TSharedRef( SharedPointerInternals::TRawPtrProxyWithDeleter< OtherType, DeleterType > const& InRawPtrProxy )
 		: Object( InRawPtrProxy.Object )
@@ -265,7 +265,7 @@ public:
 	template <
 		typename OtherType,
 		typename DeleterType
-		UE_REQUIRES(UE_REQUIRES_EXPR(ImplicitConv<ObjectType*>((OtherType*)nullptr)))
+		UE_REQUIRES(std::is_convertible_v<OtherType*, ObjectType*>)
 	>
 	FORCEINLINE TSharedRef( SharedPointerInternals::TRawPtrProxyWithDeleter< OtherType, DeleterType >&& InRawPtrProxy )
 		: Object( InRawPtrProxy.Object )
@@ -288,7 +288,7 @@ public:
 	 */
 	template <
 		typename OtherType
-		UE_REQUIRES(UE_REQUIRES_EXPR(ImplicitConv<ObjectType*>((OtherType*)nullptr)))
+		UE_REQUIRES(std::is_convertible_v<OtherType*, ObjectType*>)
 	>
 	FORCEINLINE TSharedRef( TSharedRef< OtherType, Mode > const& InSharedRef )
 		: Object( InSharedRef.Object )
@@ -387,7 +387,7 @@ public:
 	// NOTE: The following is an Unreal extension to standard shared_ptr behavior
 	template <
 		typename OtherType
-		UE_REQUIRES(UE_REQUIRES_EXPR(ImplicitConv<ObjectType*>((OtherType*)nullptr)))
+		UE_REQUIRES(std::is_convertible_v<OtherType*, ObjectType*>)
 	>
 	FORCEINLINE TSharedRef& operator=( SharedPointerInternals::TRawPtrProxy< OtherType > const& InRawPtrProxy )
 	{
@@ -410,7 +410,7 @@ public:
 	template <
 		typename OtherType,
 		typename DeleterType
-		UE_REQUIRES(UE_REQUIRES_EXPR(ImplicitConv<ObjectType*>((OtherType*)nullptr)))
+		UE_REQUIRES(std::is_convertible_v<OtherType*, ObjectType*>)
 	>
 	FORCEINLINE TSharedRef& operator=( SharedPointerInternals::TRawPtrProxyWithDeleter< OtherType, DeleterType > const& InRawPtrProxy )
 	{
@@ -433,7 +433,7 @@ public:
 	template <
 		typename OtherType,
 		typename DeleterType
-		UE_REQUIRES(UE_REQUIRES_EXPR(ImplicitConv<ObjectType*>((OtherType*)nullptr)))
+		UE_REQUIRES(std::is_convertible_v<OtherType*, ObjectType*>)
 	>
 	FORCEINLINE TSharedRef& operator=( SharedPointerInternals::TRawPtrProxyWithDeleter< OtherType, DeleterType >&& InRawPtrProxy )
 	{
@@ -572,7 +572,7 @@ private:
 	 */
 	template <
 		typename OtherType
-		UE_REQUIRES(UE_REQUIRES_EXPR(ImplicitConv<ObjectType*>((OtherType*)nullptr)))
+		UE_REQUIRES(std::is_convertible_v<OtherType*, ObjectType*>)
 	>
 	FORCEINLINE explicit TSharedRef( TSharedPtr< OtherType, Mode > const& InSharedPtr )
 		: Object( InSharedPtr.Object )
@@ -585,7 +585,7 @@ private:
 
 	template <
 		typename OtherType
-		UE_REQUIRES(UE_REQUIRES_EXPR(ImplicitConv<ObjectType*>((OtherType*)nullptr)))
+		UE_REQUIRES(std::is_convertible_v<OtherType*, ObjectType*>)
 	>
 	FORCEINLINE explicit TSharedRef( TSharedPtr< OtherType, Mode >&& InSharedPtr )
 		: Object( InSharedPtr.Object )
@@ -690,7 +690,7 @@ public:
 	 */
 	template <
 		typename OtherType
-		UE_REQUIRES(UE_REQUIRES_EXPR(ImplicitConv<ObjectType*>((OtherType*)nullptr)))
+		UE_REQUIRES(std::is_convertible_v<OtherType*, ObjectType*>)
 	>
 	FORCEINLINE explicit TSharedPtr( OtherType* InObject )
 		: Object( InObject )
@@ -711,7 +711,7 @@ public:
 	template <
 		typename OtherType,
 		typename DeleterType
-		UE_REQUIRES(UE_REQUIRES_EXPR(ImplicitConv<ObjectType*>((OtherType*)nullptr)))
+		UE_REQUIRES(std::is_convertible_v<OtherType*, ObjectType*>)
 	>
 	FORCEINLINE TSharedPtr( OtherType* InObject, DeleterType&& InDeleter )
 		: Object( InObject )
@@ -730,7 +730,7 @@ public:
 	// NOTE: The following is an Unreal extension to standard shared_ptr behavior
 	template <
 		typename OtherType
-		UE_REQUIRES(UE_REQUIRES_EXPR(ImplicitConv<ObjectType*>((OtherType*)nullptr)))
+		UE_REQUIRES(std::is_convertible_v<OtherType*, ObjectType*>)
 	>
 	FORCEINLINE TSharedPtr( SharedPointerInternals::TRawPtrProxy< OtherType > const& InRawPtrProxy )
 		: Object( InRawPtrProxy.Object )
@@ -750,7 +750,7 @@ public:
 	template <
 		typename OtherType,
 		typename DeleterType
-		UE_REQUIRES(UE_REQUIRES_EXPR(ImplicitConv<ObjectType*>((OtherType*)nullptr)))
+		UE_REQUIRES(std::is_convertible_v<OtherType*, ObjectType*>)
 	>
 	FORCEINLINE TSharedPtr( SharedPointerInternals::TRawPtrProxyWithDeleter< OtherType, DeleterType > const& InRawPtrProxy )
 		: Object( InRawPtrProxy.Object )
@@ -770,7 +770,7 @@ public:
 	template <
 		typename OtherType,
 		typename DeleterType
-		UE_REQUIRES(UE_REQUIRES_EXPR(ImplicitConv<ObjectType*>((OtherType*)nullptr)))
+		UE_REQUIRES(std::is_convertible_v<OtherType*, ObjectType*>)
 	>
 	FORCEINLINE TSharedPtr( SharedPointerInternals::TRawPtrProxyWithDeleter< OtherType, DeleterType >&& InRawPtrProxy )
 		: Object( InRawPtrProxy.Object )
@@ -789,7 +789,7 @@ public:
 	 */
 	template <
 		typename OtherType
-		UE_REQUIRES(UE_REQUIRES_EXPR(ImplicitConv<ObjectType*>((OtherType*)nullptr)))
+		UE_REQUIRES(std::is_convertible_v<OtherType*, ObjectType*>)
 	>
 	FORCEINLINE TSharedPtr( TSharedPtr< OtherType, Mode > const& InSharedPtr )
 		: Object( InSharedPtr.Object )
@@ -819,7 +819,7 @@ public:
 	// NOTE: The following is an Unreal extension to standard shared_ptr behavior
 	template <
 		typename OtherType
-		UE_REQUIRES(UE_REQUIRES_EXPR(ImplicitConv<ObjectType*>((OtherType*)nullptr)))
+		UE_REQUIRES(std::is_convertible_v<OtherType*, ObjectType*>)
 	>
 	FORCEINLINE TSharedPtr( TSharedRef< OtherType, Mode > const& InSharedRef )
 		: Object( InSharedRef.Object )
@@ -948,7 +948,7 @@ public:
 	// NOTE: The following is an Unreal extension to standard shared_ptr behavior
 	template <
 		typename OtherType
-		UE_REQUIRES(UE_REQUIRES_EXPR(ImplicitConv<ObjectType*>((OtherType*)nullptr)))
+		UE_REQUIRES(std::is_convertible_v<OtherType*, ObjectType*>)
 	>
 	FORCEINLINE TSharedPtr& operator=( SharedPointerInternals::TRawPtrProxy< OtherType > const& InRawPtrProxy )
 	{
@@ -967,7 +967,7 @@ public:
 	template <
 		typename OtherType,
 		typename DeleterType
-		UE_REQUIRES(UE_REQUIRES_EXPR(ImplicitConv<ObjectType*>((OtherType*)nullptr)))
+		UE_REQUIRES(std::is_convertible_v<OtherType*, ObjectType*>)
 	>
 	FORCEINLINE TSharedPtr& operator=( SharedPointerInternals::TRawPtrProxyWithDeleter< OtherType, DeleterType > const& InRawPtrProxy )
 	{
@@ -986,7 +986,7 @@ public:
 	template <
 		typename OtherType,
 		typename DeleterType
-		UE_REQUIRES(UE_REQUIRES_EXPR(ImplicitConv<ObjectType*>((OtherType*)nullptr)))
+		UE_REQUIRES(std::is_convertible_v<OtherType*, ObjectType*>)
 	>
 	FORCEINLINE TSharedPtr& operator=( SharedPointerInternals::TRawPtrProxyWithDeleter< OtherType, DeleterType >&& InRawPtrProxy )
 	{
@@ -1155,7 +1155,7 @@ private:
 	 */
 	template <
 		typename OtherType
-		UE_REQUIRES(UE_REQUIRES_EXPR(ImplicitConv<ObjectType*>((OtherType*)nullptr)))
+		UE_REQUIRES(std::is_convertible_v<OtherType*, ObjectType*>)
 	>
 	FORCEINLINE explicit TSharedPtr( TWeakPtr< OtherType, Mode > const& InWeakPtr )
 		: Object( nullptr )
@@ -1270,7 +1270,7 @@ public:
 	// NOTE: The following is an Unreal extension to standard shared_ptr behavior
 	template <
 		typename OtherType
-		UE_REQUIRES(UE_REQUIRES_EXPR(ImplicitConv<ObjectType*>((OtherType*)nullptr)))
+		UE_REQUIRES(std::is_convertible_v<OtherType*, ObjectType*>)
 	>
 	FORCEINLINE TWeakPtr( TSharedRef< OtherType, Mode > const& InSharedRef )
 		: Object( InSharedRef.Object )
@@ -1285,7 +1285,7 @@ public:
 	 */
 	template <
 		typename OtherType
-		UE_REQUIRES(UE_REQUIRES_EXPR(ImplicitConv<ObjectType*>((OtherType*)nullptr)))
+		UE_REQUIRES(std::is_convertible_v<OtherType*, ObjectType*>)
 	>
 	FORCEINLINE TWeakPtr( TSharedPtr< OtherType, Mode > const& InSharedPtr )
 		: Object( InSharedPtr.Object )
@@ -1331,7 +1331,7 @@ public:
 	 */
 	template <
 		typename OtherType
-		UE_REQUIRES(UE_REQUIRES_EXPR(ImplicitConv<ObjectType*>((OtherType*)nullptr)))
+		UE_REQUIRES(std::is_convertible_v<OtherType*, ObjectType*>)
 	>
 	FORCEINLINE TWeakPtr( TWeakPtr< OtherType, Mode > const& InWeakPtr )
 		: Object( InWeakPtr.Object )
@@ -1341,7 +1341,7 @@ public:
 
 	template <
 		typename OtherType
-		UE_REQUIRES(UE_REQUIRES_EXPR(ImplicitConv<ObjectType*>((OtherType*)nullptr)))
+		UE_REQUIRES(std::is_convertible_v<OtherType*, ObjectType*>)
 	>
 	FORCEINLINE TWeakPtr( TWeakPtr< OtherType, Mode >&& InWeakPtr )
 		: Object( InWeakPtr.Object )
@@ -1404,7 +1404,7 @@ public:
 	 */
 	template <
 		typename OtherType
-		UE_REQUIRES(UE_REQUIRES_EXPR(ImplicitConv<ObjectType*>((OtherType*)nullptr)))
+		UE_REQUIRES(std::is_convertible_v<OtherType*, ObjectType*>)
 	>
 	FORCEINLINE TWeakPtr& operator=( TWeakPtr<OtherType, Mode> const& InWeakPtr )
 	{
@@ -1415,7 +1415,7 @@ public:
 
 	template <
 		typename OtherType
-		UE_REQUIRES(UE_REQUIRES_EXPR(ImplicitConv<ObjectType*>((OtherType*)nullptr)))
+		UE_REQUIRES(std::is_convertible_v<OtherType*, ObjectType*>)
 	>
 	FORCEINLINE TWeakPtr& operator=( TWeakPtr<OtherType, Mode>&& InWeakPtr )
 	{
@@ -1433,7 +1433,7 @@ public:
 	// NOTE: The following is an Unreal extension to standard shared_ptr behavior
 	template <
 		typename OtherType
-		UE_REQUIRES(UE_REQUIRES_EXPR(ImplicitConv<ObjectType*>((OtherType*)nullptr)))
+		UE_REQUIRES(std::is_convertible_v<OtherType*, ObjectType*>)
 	>
 	FORCEINLINE TWeakPtr& operator=( TSharedRef< OtherType, Mode > const& InSharedRef )
 	{
@@ -1449,7 +1449,7 @@ public:
 	 */
 	template <
 		typename OtherType
-		UE_REQUIRES(UE_REQUIRES_EXPR(ImplicitConv<ObjectType*>((OtherType*)nullptr)))
+		UE_REQUIRES(std::is_convertible_v<OtherType*, ObjectType*>)
 	>
 	FORCEINLINE TWeakPtr& operator=( TSharedPtr< OtherType, Mode > const& InSharedPtr )
 	{
