Replaced ImplicitConv usage in TUniquePtr.

diff --git a/Engine/Source/Runtime/Core/Public/Templates/UniquePtr.h b/Engine/Source/Runtime/Core/Public/Templates/UniquePtr.h
--- a/Engine/Source/Runtime/Core/Public/Templates/UniquePtr.h
+++ b/Engine/Source/Runtime/Core/Public/Templates/UniquePtr.h
@@ -35,7 +35,7 @@ struct TDefaultDelete
 
 	template <
 		typename U
-		UE_REQUIRES(UE_REQUIRES_EXPR(ImplicitConv<T*>((U*)nullptr)))
+		UE_REQUIRES(std::is_convertible_v<U*, T*>)
 	>
 	TDefaultDelete(const TDefaultDelete<U>&)
 	{
@@ -43,7 +43,7 @@ struct TDefaultDelete
 
 	template <
 		typename U
-		UE_REQUIRES(UE_REQUIRES_EXPR(ImplicitConv<T*>((U*)nullptr)))
+		UE_REQUIRES(std::is_convertible_v<U*, T*>)
 	>
 	TDefaultDelete& operator=(const TDefaultDelete<U>&)
 	{
@@ -77,7 +77,7 @@ struct TDefaultDelete<T[]>
 
 	template <
 		typename U
-		UE_REQUIRES(UE_REQUIRES_EXPR(ImplicitConv<T(*)[]>((U(*)[])nullptr)))
+		UE_REQUIRES(std::is_convertible_v<U(*)[], T(*)[]>)
 	>
 	TDefaultDelete(const TDefaultDelete<U[]>&)
 	{
@@ -85,7 +85,7 @@ struct TDefaultDelete<T[]>
 
 	template <
 		typename U
-		UE_REQUIRES(UE_REQUIRES_EXPR(ImplicitConv<T(*)[]>((U(*)[])nullptr)))
+		UE_REQUIRES(std::is_convertible_v<U(*)[], T(*)[]>)
 	>
 	TDefaultDelete& operator=(const TDefaultDelete<U[]>&)
 	{
@@ -94,7 +94,7 @@ struct TDefaultDelete<T[]>
 
 	template <
 		typename U
-		UE_REQUIRES(UE_REQUIRES_EXPR(ImplicitConv<T(*)[]>((U(*)[])nullptr)))
+		UE_REQUIRES(std::is_convertible_v<U(*)[], T(*)[]>)
 	>
 	void operator()(U* Ptr) const
 	{
@@ -133,7 +133,7 @@ public:
 	 */
 	template <
 		typename U
-		UE_REQUIRES(UE_REQUIRES_EXPR(ImplicitConv<T*>((U*)nullptr)))
+		UE_REQUIRES(std::is_convertible_v<U*, T*>)
 	>
 	explicit FORCEINLINE TUniquePtr(U* InPtr)
 		: Deleter()
@@ -148,7 +148,7 @@ public:
 	 */
 	template <
 		typename U
-		UE_REQUIRES(UE_REQUIRES_EXPR(ImplicitConv<T*>((U*)nullptr)))
+		UE_REQUIRES(std::is_convertible_v<U*, T*>)
 	>
 	explicit FORCEINLINE TUniquePtr(U* InPtr, Deleter&& InDeleter)
 		: Deleter(MoveTemp(InDeleter))
@@ -163,7 +163,7 @@ public:
 	 */
 	template <
 		typename U
-		UE_REQUIRES(UE_REQUIRES_EXPR(ImplicitConv<T*>((U*)nullptr)))
+		UE_REQUIRES(std::is_convertible_v<U*, T*>)
 	>
 	explicit FORCEINLINE TUniquePtr(U* InPtr, const Deleter& InDeleter)
 		: Deleter(InDeleter)
@@ -198,7 +198,7 @@ public:
 		typename OtherDeleter
 		UE_REQUIRES(
 			!std::is_array_v<OtherT> &&
-			UE_REQUIRES_EXPR(ImplicitConv<T*>((OtherT*)nullptr))
+			std::is_convertible_v<OtherT*, T*>
 		)
 	>
 	FORCEINLINE TUniquePtr(TUniquePtr<OtherT, OtherDeleter>&& Other)
@@ -235,7 +235,7 @@ public:
 		typename OtherDeleter
 		UE_REQUIRES(
 			!std::is_array_v<OtherT> &&
-			UE_REQUIRES_EXPR(ImplicitConv<T*>((OtherT*)nullptr))
+			std::is_convertible_v<OtherT*, T*>
 		)
 	>
 	FORCEINLINE TUniquePtr& operator=(TUniquePtr<OtherT, OtherDeleter>&& Other)
@@ -463,7 +463,7 @@ public:
 	 */
 	template <
 		typename U
-		UE_REQUIRES(UE_REQUIRES_EXPR(ImplicitConv<T(*)[]>((U(*)[])nullptr)))
+		UE_REQUIRES(std::is_convertible_v<U(*)[], T(*)[]>)
 	>
 	explicit FORCEINLINE TUniquePtr(U* InPtr)
 		: Deleter()
@@ -478,7 +478,7 @@ public:
 	 */
 	template <
 		typename U
-		UE_REQUIRES(UE_REQUIRES_EXPR(ImplicitConv<T(*)[]>((U(*)[])nullptr)))
+		UE_REQUIRES(std::is_convertible_v<U(*)[], T(*)[]>)
 	>
 	explicit FORCEINLINE TUniquePtr(U* InPtr, Deleter&& InDeleter)
 		: Deleter(MoveTemp(InDeleter))
@@ -493,7 +493,7 @@ public:
 	 */
 	template <
 		typename U
-		UE_REQUIRES(UE_REQUIRES_EXPR(ImplicitConv<T(*)[]>((U(*)[])nullptr)))
+		UE_REQUIRES(std::is_convertible_v<U(*)[], T(*)[]>)
 	>
 	explicit FORCEINLINE TUniquePtr(U* InPtr, const Deleter& InDeleter)
 		: Deleter(InDeleter)
@@ -526,7 +526,7 @@ public:
 	template <
 		typename OtherT,
 		typename OtherDeleter
-		UE_REQUIRES(UE_REQUIRES_EXPR(ImplicitConv<T(*)[]>((OtherT(*)[])nullptr)))
+		UE_REQUIRES(std::is_convertible_v<OtherT(*)[], T(*)[]>)
 	>
 	FORCEINLINE TUniquePtr(TUniquePtr<OtherT, OtherDeleter>&& Other)
 		: Deleter(MoveTemp(Other.GetDeleter()))
@@ -560,7 +560,7 @@ public:
 	template <
 		typename OtherT,
 		typename OtherDeleter
-		UE_REQUIRES(UE_REQUIRES_EXPR(ImplicitConv<T(*)[]>((OtherT(*)[])nullptr)))
+		UE_REQUIRES(std::is_convertible_v<OtherT(*)[], T(*)[]>)
 	>
 	FORCEINLINE TUniquePtr& operator=(TUniquePtr<OtherT, OtherDeleter>&& Other)
 	{
@@ -655,7 +655,7 @@ public:
 	 */
 	template <
 		typename U
-		UE_REQUIRES(UE_REQUIRES_EXPR(ImplicitConv<T(*)[]>((U(*)[])nullptr)))
+		UE_REQUIRES(std::is_convertible_v<U(*)[], T(*)[]>)
 	>
 	FORCEINLINE void Reset(U* InPtr)
 	{
