Adding exception for programs when checking bIsStaged. Avoid using RequiresCookedData() as an indicative for staged in programs.

Fixes mistaking standalone programs as staged when building with bBuildWithEditorOnlyData=false

diff --git a/Engine/Source/Runtime/Core/Private/Misc/Paths.cpp b/Engine/Source/Runtime/Core/Private/Misc/Paths.cpp
--- a/Engine/Source/Runtime/Core/Private/Misc/Paths.cpp
+++ b/Engine/Source/Runtime/Core/Private/Misc/Paths.cpp
@@ -163,7 +163,10 @@ bool FPaths::IsStaged()
 	{
 		UE::Paths::bIsComputingStaged = true;
 
-		bIsStaged = FPlatformProperties::RequiresCookedData() ||
+		bIsStaged =
+#if !IS_PROGRAM
+			FPlatformProperties::RequiresCookedData() ||
+#endif
 			FileExists(Combine(EngineConfigDir(), FString::Printf(TEXT("StagedBuild_%s.ini"), FApp::GetProjectName())));
 
 		UE::Paths::bIsComputingStaged = false;
