TSan: MallocBinned: Marked a function as TSAN_SAFE to suppress warnings about a case we know should be safe in practice.

BinnedTlsSlot is accessed from multiple threads but we know it shouldn't cause issues because it's written to only once on the first access.

#rnx

diff --git a/Engine/Source/Runtime/Core/Public/HAL/MallocBinnedCommon.h b/Engine/Source/Runtime/Core/Public/HAL/MallocBinnedCommon.h
--- a/Engine/Source/Runtime/Core/Public/HAL/MallocBinnedCommon.h
+++ b/Engine/Source/Runtime/Core/Public/HAL/MallocBinnedCommon.h
@@ -433,7 +433,7 @@ protected:
 
 	struct FPerThreadFreeBlockLists
 	{
-		FORCEINLINE static FPerThreadFreeBlockLists* Get()
+		FORCEINLINE static FPerThreadFreeBlockLists* Get() TSAN_SAFE
 		{
 			return FPlatformTLS::IsValidTlsSlot(BinnedTlsSlot) ? (FPerThreadFreeBlockLists*)FPlatformTLS::GetTlsValue(BinnedTlsSlot) : nullptr;
 		}
