[TaskGraph]
 - Reorder fields so that TryUnlock only touch a single cache line

diff --git a/Engine/Source/Runtime/Core/Public/Tasks/TaskPrivate.h b/Engine/Source/Runtime/Core/Public/Tasks/TaskPrivate.h
--- a/Engine/Source/Runtime/Core/Public/Tasks/TaskPrivate.h
+++ b/Engine/Source/Runtime/Core/Public/Tasks/TaskPrivate.h
@@ -669,51 +669,23 @@
 			CORE_API bool WaitImpl(FTimeout Timeout);
 
 		private:
+			// the number of times that the task should be unlocked before it can be scheduled or completed
+			// initial count is 1 for launching the task (it can't be scheduled before it's launched)
+			// reaches 0 the task is scheduled for execution.
+			// NumLocks's the most significant bit (see `ExecutionFlag`) is set on task execution start, and indicates that now 
+			// NumLocks is about how many times the task must be unlocked to be completed
+			static constexpr uint32 NumInitialLocks = 1;
+			std::atomic<uint32> NumLocks{ NumInitialLocks };
+
+			FPipe* Pipe{ nullptr };
+
 			EExtendedTaskPriority ExtendedPriority; // internal priorities, if any
 
-			LowLevelTasks::FTask LowLevelTask;
+			std::atomic<uint32> ExecutingThreadId = FThread::InvalidThreadId;
 
-			// the task is completed when its subsequents list is closed and no more can be added
-			template <typename AllocatorType = FDefaultAllocator>
-			class FSubsequents
-			{
-			public:
-				bool PushIfNotClosed(FTaskBase* NewItem)
-				{
-					TASKGRAPH_VERBOSE_EVENT_SCOPE(FSubsequents::PushIfNotClosed);
-					if (bIsClosed.load(std::memory_order_relaxed))
-					{
-						return false;
-					}
-					UE::TUniqueLock Lock(Mutex);
-					if (bIsClosed)
-					{
-						return false;
-					}
-					Subsequents.Emplace(NewItem);
-					return true;
-				}
+#if UE_TASK_TRACE_ENABLED
+			std::atomic<TaskTrace::FId> TraceId{ TaskTrace::GenerateTaskId() };
+#endif
 
-				TArray<FTaskBase*, AllocatorType> Close()
-				{
-					TASKGRAPH_VERBOSE_EVENT_SCOPE(FSubsequents::Close);
-					UE::TUniqueLock Lock(Mutex);
-					bIsClosed = true;
-					return MoveTemp(Subsequents);
-				}
-
-				bool IsClosed() const
-				{
-					return bIsClosed;
-				}
-
-			private:
-				TArray<FTaskBase*, AllocatorType> Subsequents;
-				std::atomic<bool>  bIsClosed = false;
-				UE::FMutex Mutex;
-			};
-
-			FSubsequents<TInlineAllocator<1>> Subsequents;
-
 			// stores backlinks to prerequsites, either execution prerequisites or nested tasks (completion prerequisites).
 			// It's populated in three stages:
@@ -751,26 +723,55 @@
 				}
 			private:
 				TArray<FTaskBase*, AllocatorType> Prerequisites;
-				UE::FMutex Mutex { UE::AcquireLock }; // Start locked by default to avoid compare exchange during construction.
+				UE::FMutex Mutex{ UE::AcquireLock }; // Start locked by default to avoid compare exchange during construction.
 			};
 
 			FPrerequisites<TInlineAllocator<1>> Prerequisites;
 
-			FPipe* Pipe{ nullptr };
+			LowLevelTasks::FTask LowLevelTask;
 
-#if UE_TASK_TRACE_ENABLED
-			std::atomic<TaskTrace::FId> TraceId{ TaskTrace::GenerateTaskId() };
-#endif
+			// the task is completed when its subsequents list is closed and no more can be added
+			template <typename AllocatorType = FDefaultAllocator>
+			class FSubsequents
+			{
+			public:
+				bool PushIfNotClosed(FTaskBase* NewItem)
+				{
+					TASKGRAPH_VERBOSE_EVENT_SCOPE(FSubsequents::PushIfNotClosed);
+					if (bIsClosed.load(std::memory_order_relaxed))
+					{
+						return false;
+					}
+					UE::TUniqueLock Lock(Mutex);
+					if (bIsClosed)
+					{
+						return false;
+					}
+					Subsequents.Emplace(NewItem);
+					return true;
+				}
 
-			// the number of times that the task should be unlocked before it can be scheduled or completed
-			// initial count is 1 for launching the task (it can't be scheduled before it's launched)
-			// reaches 0 the task is scheduled for execution.
-			// NumLocks's the most significant bit (see `ExecutionFlag`) is set on task execution start, and indicates that now 
-			// NumLocks is about how many times the task must be unlocked to be completed
-			static constexpr uint32 NumInitialLocks = 1;
-			std::atomic<uint32> NumLocks{ NumInitialLocks };
+				TArray<FTaskBase*, AllocatorType> Close()
+				{
+					TASKGRAPH_VERBOSE_EVENT_SCOPE(FSubsequents::Close);
+					UE::TUniqueLock Lock(Mutex);
+					bIsClosed = true;
+					return MoveTemp(Subsequents);
+				}
 
-			std::atomic<uint32> ExecutingThreadId = FThread::InvalidThreadId;
+				bool IsClosed() const
+				{
+					return bIsClosed;
+				}
 
+			private:
+				TArray<FTaskBase*, AllocatorType> Subsequents;
+				std::atomic<bool>  bIsClosed = false;
+				UE::FMutex Mutex;
+			};
+
+			FSubsequents<TInlineAllocator<1>> Subsequents;
+
+
 protected:
 			void UnlockPrerequisites()
