Removed ImplicitConv usage from MakeArrayView.

diff --git a/Engine/Source/Runtime/Core/Public/Memory/MemoryView.h b/Engine/Source/Runtime/Core/Public/Memory/MemoryView.h
--- a/Engine/Source/Runtime/Core/Public/Memory/MemoryView.h
+++ b/Engine/Source/Runtime/Core/Public/Memory/MemoryView.h
@@ -265,20 +265,15 @@ template <typename DataType>
 	return TMemoryView<const void>(Data, Size);
 }
 
-/** Make a non-owning mutable view starting at Data and ending at DataEnd. */
-template <typename DataEndType,
-	decltype(ImplicitConv<void*>(DeclVal<DataEndType*>()))* = nullptr>
-[[nodiscard]] inline TMemoryView<void> MakeMemoryView(void* Data, DataEndType* DataEnd)
+/** Make a non-owning view starting at Data and ending at DataEnd. */
+template <typename DataType, typename DataEndType>
+[[nodiscard]] inline auto MakeMemoryView(DataType* Data, DataEndType* DataEnd)
 {
-	return TMemoryView<void>(Data, DataEnd);
-}
+	// This function is templated on pointer type to prevent MakeMemoryView(Ptr, 0) being deduced as a null end pointer
 
-/** Make a non-owning const view starting at Data and ending at DataEnd. */
-template <typename DataEndType,
-	decltype(ImplicitConv<const void*>(DeclVal<DataEndType*>()))* = nullptr>
-[[nodiscard]] inline TMemoryView<const void> MakeMemoryView(const void* Data, DataEndType* DataEnd)
-{
-	return TMemoryView<const void>(Data, DataEnd);
+	using VoidType = std::conditional_t<std::is_const_v<DataType> || std::is_const_v<DataEndType>, const void, void>;
+
+	return TMemoryView<VoidType>(Data, DataEnd);
 }
 
 /**
