[TaskGraph]
 - Fix memory leak when waiting repeateadly on a task with a timeout

diff --git a/Engine/Source/Runtime/Core/Private/Tasks/TaskPrivate.cpp b/Engine/Source/Runtime/Core/Private/Tasks/TaskPrivate.cpp
--- a/Engine/Source/Runtime/Core/Private/Tasks/TaskPrivate.cpp
+++ b/Engine/Source/Runtime/Core/Private/Tasks/TaskPrivate.cpp
@@ -43,6 +43,9 @@ namespace UE::Tasks
 #endif
 
 			bWakeUpWorker |= LowLevelTasks::FScheduler::Get().TryLaunch(LowLevelTask, bWakeUpWorker ? LowLevelTasks::EQueuePreference::GlobalQueuePreference : LowLevelTasks::EQueuePreference::LocalQueuePreference, bWakeUpWorker);
+			
+			// In case a thread is waiting on us to perform retraction, now is the time to try retraction again.
+			StateChangeEvent.Notify();
 		}
 
 		thread_local uint32 TaskRetractionRecursion = 0;
@@ -208,37 +211,42 @@ namespace UE::Tasks
 
 		bool FTaskBase::WaitImpl(FTimeout Timeout)
 		{
-			// ignore the result as we still have to make sure the task is completed upon returning from this function call
-			TryRetractAndExecute(Timeout);
-
-			// spin for a while with hope the task is getting completed right now, to avoid getting blocked by a pricy syscall
-			const uint32 MaxSpinCount = 40;
-			for (uint32 SpinCount = 0; SpinCount != MaxSpinCount && !IsCompleted() && !Timeout; ++SpinCount)
+			while (true)
 			{
-				FPlatformProcess::Yield(); // YieldThread() was much slower on some platforms with low core count and contention for CPU
-			}
+				// ignore the result as we still have to make sure the task is completed upon returning from this function call
+				TryRetractAndExecute(Timeout);
 
-			if (IsCompleted() || Timeout)
-			{
-				return IsCompleted();
-			}
+				// spin for a while with hope the task is getting completed right now, to avoid getting blocked by a pricey syscall
+				const uint32 MaxSpinCount = 40;
+				for (uint32 SpinCount = 0; SpinCount != MaxSpinCount && !IsCompleted() && !Timeout; ++SpinCount)
+				{
+					FPlatformProcess::Yield(); // YieldThread() was much slower on some platforms with low core count and contention for CPU
+				}
 
-			// the event must be alive for the task and this function lifetime, we don't know which one will be finished first as waiting can 
-			// time out before the waiting task is completed
-			FSharedEventRef CompletionEvent;
-			auto WaitingTaskBody = [CompletionEvent] { CompletionEvent->Trigger(); };
-			using FWaitingTask = TExecutableTask<decltype(WaitingTaskBody)>;
+				if (IsCompleted() || Timeout)
+				{
+					return IsCompleted();
+				}
 
-			TRefCountPtr<FWaitingTask> WaitingTask{ FWaitingTask::Create(TEXT("Waiting Task"), MoveTemp(WaitingTaskBody), ETaskPriority::Default /* doesn't matter*/, EExtendedTaskPriority::Inline, ETaskFlags::None), /*bAddRef=*/ false };
-			WaitingTask->AddPrerequisites(*this);
+				auto Token = StateChangeEvent.PrepareWait();
 
-			if (WaitingTask->TryLaunch(sizeof(WaitingTask)))
-			{	// was executed inline
-				check(WaitingTask->IsCompleted());
-				return true;
-			}
+				// Important to check the condition a second time after PrepareWait has been called to make sure we don't
+				// miss an important state change event.
+				if (IsCompleted())
+				{
+					return true;
+				}
 
-			return CompletionEvent->Wait(Timeout.GetRemainingRoundedUpMilliseconds());
+				{
+					TRACE_CPUPROFILER_EVENT_SCOPE(FTaskBase::WaitImpl_StateChangeEvent_WaitFor);
+
+					// Always flush events before entering a wait to make sure there's nothing missing in Unreal Insights that could prevent us understanding what's going on.
+					TRACE_CPUPROFILER_EVENT_FLUSH();
+					StateChangeEvent.WaitFor(Token, UE::FMonotonicTimeSpan::FromMilliseconds(Timeout.GetRemainingRoundedUpMilliseconds()));
+				}
+
+				// Once the state of the task has changed (either closed or scheduled), it's time to do another round of retraction to help if possible.
+			}
 		}
 
 		FTaskBase* FTaskBase::TryPushIntoPipe()
diff --git a/Engine/Source/Runtime/Core/Public/Tasks/TaskPrivate.h b/Engine/Source/Runtime/Core/Public/Tasks/TaskPrivate.h
--- a/Engine/Source/Runtime/Core/Public/Tasks/TaskPrivate.h
+++ b/Engine/Source/Runtime/Core/Public/Tasks/TaskPrivate.h
@@ -545,6 +545,9 @@ public:
 				ReleasePrerequisites();
 
 				TaskTrace::Completed(GetTraceId());
+
+				// In case a thread is waiting on us to perform retraction, now is the time to try retraction again.
+				StateChangeEvent.Notify();
 			}
 
 			CORE_API void ClearPipe();
@@ -679,6 +682,8 @@ public:
 
 			FPipe* Pipe{ nullptr };
 
+			FEventCount StateChangeEvent;
+
 			EExtendedTaskPriority ExtendedPriority; // internal priorities, if any
 
 			std::atomic<uint32> ExecutingThreadId = FThread::InvalidThreadId;
