Deprecate bool operator and promote explicit function call IsExpired() instead

diff --git a/Engine/Source/Runtime/Core/Private/Tasks/TaskPrivate.cpp b/Engine/Source/Runtime/Core/Private/Tasks/TaskPrivate.cpp
--- a/Engine/Source/Runtime/Core/Private/Tasks/TaskPrivate.cpp
+++ b/Engine/Source/Runtime/Core/Private/Tasks/TaskPrivate.cpp
@@ -80,7 +80,7 @@ namespace UE::Tasks
 				return false;
 			}
 
-			if (IsCompleted() || Timeout)
+			if (IsCompleted() || Timeout.IsExpired())
 			{
 				return IsCompleted();
 			}
@@ -132,7 +132,7 @@ namespace UE::Tasks
 				return true;
 			}
 
-			if (Timeout)
+			if (Timeout.IsExpired())
 			{
 				return IsCompleted();
 			}
@@ -182,7 +182,7 @@ namespace UE::Tasks
 
 		bool FTaskBase::Wait(FTimeout Timeout)
 		{
-			if (IsCompleted() || Timeout)
+			if (IsCompleted() || Timeout.IsExpired())
 			{
 				return IsCompleted();
 			}
@@ -218,12 +218,12 @@ namespace UE::Tasks
 
 				// spin for a while with hope the task is getting completed right now, to avoid getting blocked by a pricey syscall
 				const uint32 MaxSpinCount = 40;
-				for (uint32 SpinCount = 0; SpinCount != MaxSpinCount && !IsCompleted() && !Timeout; ++SpinCount)
+				for (uint32 SpinCount = 0; SpinCount != MaxSpinCount && !IsCompleted() && !Timeout.IsExpired(); ++SpinCount)
 				{
 					FPlatformProcess::Yield(); // YieldThread() was much slower on some platforms with low core count and contention for CPU
 				}
 
-				if (IsCompleted() || Timeout)
+				if (IsCompleted() || Timeout.IsExpired())
 				{
 					return IsCompleted();
 				}
diff --git a/Engine/Source/Runtime/Core/Public/Misc/Timeout.h b/Engine/Source/Runtime/Core/Public/Misc/Timeout.h
--- a/Engine/Source/Runtime/Core/Public/Misc/Timeout.h
+++ b/Engine/Source/Runtime/Core/Public/Misc/Timeout.h
@@ -13,7 +13,7 @@ namespace UE
 	// usage:
 	// ------------------------------------
 	// FTimeout Timeout(FTimespan::FromMilliseconds(2));
-	// while (!Timeout) { ... }
+	// while (!Timeout.IsExpired()) { ... }
 	// ------------------------------------
 	class FTimeout
 	{
@@ -23,7 +23,13 @@ namespace UE
 		{
 		}
 
+		UE_DEPRECATED(5.5, "Use IsExpired() instead.")
 		explicit operator bool() const
+		{
+			return IsExpired();
+		}
+
+		bool IsExpired() const
 		{
 			return GetRemainingTime() <= FTimespan::Zero();
 		}
diff --git a/Engine/Source/Runtime/Core/Public/Tasks/Task.h b/Engine/Source/Runtime/Core/Public/Tasks/Task.h
--- a/Engine/Source/Runtime/Core/Public/Tasks/Task.h
+++ b/Engine/Source/Runtime/Core/Public/Tasks/Task.h
@@ -367,7 +367,7 @@ namespace UE::Tasks
 
 		for (auto& Task : Tasks)
 		{
-			if (Timeout || !Task.BusyWait(Timeout))
+			if (Timeout.IsExpired() || !Task.BusyWait(Timeout))
 			{
 				return false;
 			}
diff --git a/Engine/Source/Runtime/Core/Public/Tasks/TaskPrivate.h b/Engine/Source/Runtime/Core/Public/Tasks/TaskPrivate.h
--- a/Engine/Source/Runtime/Core/Public/Tasks/TaskPrivate.h
+++ b/Engine/Source/Runtime/Core/Public/Tasks/TaskPrivate.h
@@ -973,7 +973,7 @@ protected:
 					bResult = false;  // do not stop here to let this thread to help in executing tasks as much as possible, as it's waiting for their completion anyway
 				}
 
-				if (Timeout)
+				if (Timeout.IsExpired())
 				{
 					return false;
 				}
