Replaced ImplicitCast in a variety of places.

diff --git a/Engine/Plugins/Enterprise/DatasmithContent/Source/DatasmithContent/Public/DatasmithAdditionalData.h b/Engine/Plugins/Enterprise/DatasmithContent/Source/DatasmithContent/Public/DatasmithAdditionalData.h
--- a/Engine/Plugins/Enterprise/DatasmithContent/Source/DatasmithContent/Public/DatasmithAdditionalData.h
+++ b/Engine/Plugins/Enterprise/DatasmithContent/Source/DatasmithContent/Public/DatasmithAdditionalData.h
@@ -4,10 +4,13 @@
 
 #include "CoreMinimal.h"
 #include "UObject/Object.h"
+#include "Templates/Requires.h"
 #include "Templates/SubclassOf.h"
 #include "AssetRegistry/AssetData.h"
 #include "DatasmithAssetImportData.h"
 
+#include <type_traits>
+
 #include "DatasmithAdditionalData.generated.h"
 
 /**
@@ -35,8 +38,8 @@ namespace Datasmith
 	 * @return DataType*    Created instance
 	 */
 	template<
-		class DataType,
-		typename = decltype(ImplicitConv<UDatasmithAdditionalData*>((DataType*)nullptr))
+		class DataType
+		UE_REQUIRES(std::is_convertible_v<DataType*, UDatasmithAdditionalData*>)
 	>
 	inline DataType* MakeAdditionalData()
 	{
@@ -56,8 +59,8 @@ namespace Datasmith
 	 * @return TArray<DataType*>    The list of additional data stored on the given asset
 	 */
 	template<
-		class DataType,
-		typename = decltype(ImplicitConv<UDatasmithAdditionalData*>((DataType*)nullptr))
+		class DataType
+		UE_REQUIRES(std::is_convertible_v<DataType*, UDatasmithAdditionalData*>)
 	>
 	inline TArray<DataType*> GetMultipleAdditionalData(const FAssetData& SourceAssetData, int MaxCount=0)
 	{
@@ -92,8 +95,8 @@ namespace Datasmith
 	 * @return DataType*        nullptr if no Additional data is stored on the given asset
 	 */
 	template<
-		class DataType,
-		typename = decltype(ImplicitConv<UDatasmithAdditionalData*>((DataType*)nullptr))
+		class DataType
+		UE_REQUIRES(std::is_convertible_v<DataType*, UDatasmithAdditionalData*>)
 	>
 	inline DataType* GetAdditionalData(const FAssetData& SourceAssetData)
 	{
diff --git a/Engine/Plugins/Enterprise/GLTFExporter/Source/GLTFExporter/Public/Converters/GLTFSharedArray.h b/Engine/Plugins/Enterprise/GLTFExporter/Source/GLTFExporter/Public/Converters/GLTFSharedArray.h
--- a/Engine/Plugins/Enterprise/GLTFExporter/Source/GLTFExporter/Public/Converters/GLTFSharedArray.h
+++ b/Engine/Plugins/Enterprise/GLTFExporter/Source/GLTFExporter/Public/Converters/GLTFSharedArray.h
@@ -3,6 +3,9 @@
 #pragma once
 
 #include "CoreMinimal.h"
+#include "Templates/Requires.h"
+
+#include <type_traits>
 
 template <typename ElementType, typename ArrayType = TArray<ElementType>, ESPMode Mode = ESPMode::ThreadSafe>
 class TGLTFSharedArray : public TSharedRef<ArrayType, Mode>
@@ -14,7 +17,10 @@ public:
 	{
 	}
 
-	template <typename OtherType, typename = decltype(ImplicitConv<ArrayType*>(static_cast<OtherType*>(nullptr)))>
+	template <
+		typename OtherType
+		UE_REQUIRES(std::is_convertible_v<OtherType*, ArrayType*>)
+	>
 	TGLTFSharedArray(const TSharedRef<OtherType, Mode>& SharedRef)
 	: TSharedRef<ArrayType, Mode>(SharedRef)
 	{
diff --git a/Engine/Plugins/Experimental/PythonScriptPlugin/Source/PythonScriptPlugin/Private/PyUtil.h b/Engine/Plugins/Experimental/PythonScriptPlugin/Source/PythonScriptPlugin/Private/PyUtil.h
--- a/Engine/Plugins/Experimental/PythonScriptPlugin/Source/PythonScriptPlugin/Private/PyUtil.h
+++ b/Engine/Plugins/Experimental/PythonScriptPlugin/Source/PythonScriptPlugin/Private/PyUtil.h
@@ -7,9 +7,12 @@
 #include "CoreMinimal.h"
 #include "Containers/ArrayView.h"
 #include "Logging/LogMacros.h"
+#include "Templates/Requires.h"
 #include "UObject/Field.h"
 #include "UObject/FieldIterator.h"
 
+#include <type_traits>
+
 DECLARE_LOG_CATEGORY_EXTERN(LogPython, Log, All);
 
 #if WITH_PYTHON
@@ -119,8 +122,8 @@ namespace PyUtil
 		TPropOnScope& operator=(const TPropOnScope&) = delete;
 
 		template <
-			typename TOtherPropType,
-			typename = decltype(ImplicitConv<TPropType*>((TOtherPropType*)nullptr))
+			typename TOtherPropType
+			UE_REQUIRES(std::is_convertible_v<TOtherPropType*, TPropType*>)
 		>
 		TPropOnScope(TPropOnScope<TOtherPropType>&& Other)
 		{
@@ -131,8 +134,8 @@ namespace PyUtil
 		}
 
 		template <
-			typename TOtherPropType,
-			typename = decltype(ImplicitConv<TPropType*>((TOtherPropType*)nullptr))
+			typename TOtherPropType
+			UE_REQUIRES(std::is_convertible_v<TOtherPropType*, TPropType*>)
 		>
 		TPropOnScope& operator=(TPropOnScope<TOtherPropType>&& Other)
 		{
diff --git a/Engine/Plugins/MeshPainting/Source/MeshPaintingToolset/Public/MeshPaintingToolsetTypes.h b/Engine/Plugins/MeshPainting/Source/MeshPaintingToolset/Public/MeshPaintingToolsetTypes.h
--- a/Engine/Plugins/MeshPainting/Source/MeshPaintingToolset/Public/MeshPaintingToolsetTypes.h
+++ b/Engine/Plugins/MeshPainting/Source/MeshPaintingToolset/Public/MeshPaintingToolsetTypes.h
@@ -2,6 +2,9 @@
 
 #pragma once
 
+#include "Templates/Requires.h"
+
+#include <type_traits>
 
 #include "MeshPaintingToolsetTypes.generated.h"
 
@@ -106,7 +109,10 @@ struct FPaintableTexture
 		, UVChannelIndex(0)
 	{}
 
-	template<typename T, decltype(ImplicitConv<UTexture*>(DeclVal<T>()))* = nullptr>
+	template <
+		typename T
+		UE_REQUIRES(std::is_convertible_v<T, UTexture*>)
+	>
 	FPaintableTexture(T InTexture = nullptr, uint32 InUVChannelIndex = 0)
 		: Texture(InTexture)
 		, UVChannelIndex(InUVChannelIndex)
diff --git a/Engine/Plugins/Runtime/Metasound/Source/MetasoundFrontend/Public/MetasoundAccessPtr.h b/Engine/Plugins/Runtime/Metasound/Source/MetasoundFrontend/Public/MetasoundAccessPtr.h
--- a/Engine/Plugins/Runtime/Metasound/Source/MetasoundFrontend/Public/MetasoundAccessPtr.h
+++ b/Engine/Plugins/Runtime/Metasound/Source/MetasoundFrontend/Public/MetasoundAccessPtr.h
@@ -3,9 +3,12 @@
 #pragma once
 
 #include "Templates/Function.h"
+#include "Templates/Requires.h"
 #include "Templates/SharedPointer.h"
 #include "Templates/UnrealTemplate.h"
 
+#include <type_traits>
+
 #if !(UE_BUILD_SHIPPING || UE_BUILD_TEST)
 #define METASOUND_FRONTEND_ACCESSPTR_DEBUG_INFO 1
 #else
@@ -135,8 +138,8 @@ namespace Metasound
 
 			/** Creates an access pointer from another using a static cast. */
 			template <
-				typename OtherType,
-				typename = decltype(ImplicitConv<Type*>((OtherType*)nullptr))
+				typename OtherType
+				UE_REQUIRES(std::is_convertible_v<ObjectType*, Type*>)
 			>
 			TAccessPtr(const TAccessPtr<OtherType>& InOther, EDerivedCopy InTag=EDerivedCopy::Tag)
 			{
diff --git a/Engine/Plugins/Runtime/ReplicationGraph/Source/Public/ReplicationGraphTypes.h b/Engine/Plugins/Runtime/ReplicationGraph/Source/Public/ReplicationGraphTypes.h
--- a/Engine/Plugins/Runtime/ReplicationGraph/Source/Public/ReplicationGraphTypes.h
+++ b/Engine/Plugins/Runtime/ReplicationGraph/Source/Public/ReplicationGraphTypes.h
@@ -7,7 +7,11 @@
 #include "Net/DataBunch.h"
 #include "ProfilingDebugging/CsvProfiler.h"
 #include "Templates/Greater.h"
+#include "Templates/Requires.h"
 #include "UObject/Package.h"
+
+#include <type_traits>
+
 #include "ReplicationGraphTypes.generated.h"
 
 class AActor;
@@ -126,8 +130,8 @@ struct FActorRepListType
 
 	// to support conversion from TObjectPtr<ASubclassOfActor>
 	template <
-		typename T,
-		decltype(ImplicitConv<AActor*>(std::declval<const T&>())) = nullptr
+		typename T
+		UE_REQUIRES(std::is_convertible_v<const T&, AActor*>)
 	>
 	FActorRepListType(const T& InActor)
 		: ActorRaw(InActor)
diff --git a/Engine/Source/Editor/BlueprintGraph/Classes/BlueprintNodeBinder.h b/Engine/Source/Editor/BlueprintGraph/Classes/BlueprintNodeBinder.h
--- a/Engine/Source/Editor/BlueprintGraph/Classes/BlueprintNodeBinder.h
+++ b/Engine/Source/Editor/BlueprintGraph/Classes/BlueprintNodeBinder.h
@@ -4,10 +4,13 @@
 
 #include "CoreMinimal.h"
 #include "Templates/Casts.h"
+#include "Templates/Requires.h"
 #include "UObject/Object.h"
 #include "UObject/WeakObjectPtr.h"
 #include "UObject/WeakFieldPtr.h"
 
+#include <type_traits>
+
 class UEdGraphNode;
 
 class FBindingObject
@@ -19,7 +22,10 @@ public:
 	FBindingObject()
 		: bIsUObject(false)
 	{}
-	template <typename T, decltype(ImplicitConv<UObject*>(DeclVal<T>()))* = nullptr>
+	template <
+		typename T
+		UE_REQUIRES(std::is_convertible_v<T, UObject*>)
+	>
 	FBindingObject(T InObject)
 		: Object(InObject)
 		, bIsUObject(true)
@@ -40,7 +46,10 @@ public:
 			Field = InFieldOrObject.ToField();
 		}
 	}
-	template <typename T, decltype(ImplicitConv<UObject*>(DeclVal<T>()))* = nullptr>
+	template <
+		typename T
+		UE_REQUIRES(std::is_convertible_v<T, UObject*>)
+	>
 	FBindingObject& operator=(T InObject)
 	{
 		Object = ImplicitConv<UObject*>(InObject);
diff --git a/Engine/Source/Editor/MeshPaint/Public/MeshPaintTypes.h b/Engine/Source/Editor/MeshPaint/Public/MeshPaintTypes.h
--- a/Engine/Source/Editor/MeshPaint/Public/MeshPaintTypes.h
+++ b/Engine/Source/Editor/MeshPaint/Public/MeshPaintTypes.h
@@ -5,6 +5,9 @@
 #include "Engine/Texture.h"
 #include "Engine/TextureRenderTarget2D.h"
 #include "Materials/MaterialInterface.h"
+#include "Templates/Requires.h"
+
+#include <type_traits>
 
 #include "MeshPaintTypes.generated.h"
 
@@ -86,7 +89,10 @@ struct FPaintableTexture
 		, UVChannelIndex(0)
 	{}
 
-	template<typename T, decltype(ImplicitConv<UTexture*>(DeclVal<T>()))* = nullptr>
+	template <
+		typename T
+		UE_REQUIRES(std::is_convertible_v<T, UTexture*>)
+	>
 	FPaintableTexture(T InTexture, uint32 InUVChannelIndex = 0)
 		: Texture(InTexture)
 		, UVChannelIndex(InUVChannelIndex)
diff --git a/Engine/Source/Runtime/Core/Public/Internationalization/LocalizedTextSourceTypes.h b/Engine/Source/Runtime/Core/Public/Internationalization/LocalizedTextSourceTypes.h
--- a/Engine/Source/Runtime/Core/Public/Internationalization/LocalizedTextSourceTypes.h
+++ b/Engine/Source/Runtime/Core/Public/Internationalization/LocalizedTextSourceTypes.h
@@ -5,8 +5,11 @@
 #include "CoreTypes.h"
 #include "Misc/EnumClassFlags.h"
 #include "Templates/RefCounting.h"
+#include "Templates/Requires.h"
 #include "Templates/SharedPointer.h"
 
+#include <type_traits>
+
 /**
  * Categories of localized text.
  * @note This enum is mirrored in NoExportTypes.h for UHT.
@@ -166,13 +169,19 @@ public:
 		check(this->IsValid());
 	}
 
-	template <typename OtherType, typename = decltype(ImplicitConv<ObjectType*>((OtherType*)nullptr))>
+	template <
+		typename OtherType
+		UE_REQUIRES(std::is_convertible_v<OtherType*, ObjectType*>)
+	>
 	TDisplayStringRef(const TDisplayStringRef<OtherType>& InOther)
 		: TDisplayStringPtrBase<ObjectType>(InOther.Private_GetDisplayStringPtr())
 	{
 	}
 
-	template <typename OtherType, typename = decltype(ImplicitConv<ObjectType*>((OtherType*)nullptr))>
+	template <
+		typename OtherType
+		UE_REQUIRES(std::is_convertible_v<OtherType*, ObjectType*>)
+	>
 	TDisplayStringRef& operator=(const TDisplayStringRef<OtherType>& InOther)
 	{
 		if (this->DisplayStringPtr != InOther.Private_GetDisplayStringPtr())
@@ -207,19 +216,28 @@ public:
 	{
 	}
 
-	template <typename OtherType, typename = decltype(ImplicitConv<ObjectType*>((OtherType*)nullptr))>
+	template <
+		typename OtherType
+		UE_REQUIRES(std::is_convertible_v<OtherType*, ObjectType*>)
+	>
 	TDisplayStringPtr(const TDisplayStringPtr<OtherType>& InOther)
 		: TDisplayStringPtrBase<ObjectType>(InOther.Private_GetDisplayStringPtr())
 	{
 	}
 
-	template <typename OtherType, typename = decltype(ImplicitConv<ObjectType*>((OtherType*)nullptr))>
+	template <
+		typename OtherType
+		UE_REQUIRES(std::is_convertible_v<OtherType*, ObjectType*>)
+	>
 	TDisplayStringPtr(const TDisplayStringRef<OtherType>& InOther)
 		: TDisplayStringPtrBase<ObjectType>(InOther.Private_GetDisplayStringPtr())
 	{
 	}
 
-	template <typename OtherType, typename = decltype(ImplicitConv<ObjectType*>((OtherType*)nullptr))>
+	template <
+		typename OtherType
+		UE_REQUIRES(std::is_convertible_v<OtherType*, ObjectType*>)
+	>
 	TDisplayStringPtr& operator=(const TDisplayStringPtr<OtherType>& InOther)
 	{
 		if (this->DisplayStringPtr != InOther.Private_GetDisplayStringPtr())
@@ -229,7 +247,10 @@ public:
 		return *this;
 	}
 
-	template <typename OtherType, typename = decltype(ImplicitConv<ObjectType*>((OtherType*)nullptr))>
+	template <
+		typename OtherType
+		UE_REQUIRES(std::is_convertible_v<OtherType*, ObjectType*>)
+	>
 	TDisplayStringPtr& operator=(const TDisplayStringRef<OtherType>& InOther)
 	{
 		if (this->DisplayStringPtr != InOther.Private_GetDisplayStringPtr())
diff --git a/Engine/Source/Runtime/Core/Public/Internationalization/Text.h b/Engine/Source/Runtime/Core/Public/Internationalization/Text.h
--- a/Engine/Source/Runtime/Core/Public/Internationalization/Text.h
+++ b/Engine/Source/Runtime/Core/Public/Internationalization/Text.h
@@ -21,6 +21,9 @@
 #include "Misc/Optional.h"
 #include "Templates/UniquePtr.h"
 #include "Templates/IsConstructible.h"
+#include "Templates/Requires.h"
+
+#include <type_traits>
 
 class FText;
 class FTextHistory;
@@ -766,7 +769,10 @@ public:
 #endif
 
 private:
-	template <typename HistoryType, typename = decltype(ImplicitConv<ITextData*>((HistoryType*)nullptr))>
+	template <
+		typename HistoryType
+		UE_REQUIRES(std::is_convertible_v<HistoryType*, ITextData*>)
+	>
 	explicit FText( TRefCountPtr<HistoryType>&& InTextData )
 		: TextData(MoveTemp(InTextData))
 		, Flags(0)
diff --git a/Engine/Source/Runtime/Core/Public/Templates/PointerVariants.h b/Engine/Source/Runtime/Core/Public/Templates/PointerVariants.h
--- a/Engine/Source/Runtime/Core/Public/Templates/PointerVariants.h
+++ b/Engine/Source/Runtime/Core/Public/Templates/PointerVariants.h
@@ -3,12 +3,15 @@
 #pragma once
 
 #include "Misc/TVariant.h"
+#include "Templates/Requires.h"
 #include "Templates/SharedPointer.h"
 #include "UObject/StrongObjectPtr.h"
 #include "UObject/WeakObjectPtr.h"
 #include "UnrealTypeTraits.h"
 #include "UnrealTemplate.h"
 
+#include <type_traits>
+
 /*
 * `TWeakPtrVariant` and `TStrongPtrVariant` are particularly useful for "interfaces" in areas where
 * UObjects and non-UObjects exist to hold onto a "smart" pointer to the Interface, rather than needing to keep track of the implementing class to check validity and manually cast
@@ -52,10 +55,10 @@ namespace UE::Core::Private
 		{
 		}
 
-		template<typename DerivedType,
-			decltype(ImplicitConv<BaseType*>((DerivedType*)nullptr))* = nullptr
-			UE_REQUIRES(std::is_base_of_v<UObject, DerivedType> || IsDerivedFromSharedFromThis<DerivedType>())
-			>
+		template <
+			typename DerivedType
+			UE_REQUIRES(std::is_convertible_v<DerivedType*, BaseType*> && (std::is_base_of_v<UObject, DerivedType> || IsDerivedFromSharedFromThis<DerivedType>()))
+		>
 		constexpr TPtrVariantBase(DerivedType* InDerived)
 		{   
 			// TODO:	Look for a way to move this "smart pointer selection" logic outside of the constructor; then we can remove "TVariant" as a dependency
diff --git a/Engine/Source/Runtime/Core/Public/Templates/UnrealTemplate.h b/Engine/Source/Runtime/Core/Public/Templates/UnrealTemplate.h
--- a/Engine/Source/Runtime/Core/Public/Templates/UnrealTemplate.h
+++ b/Engine/Source/Runtime/Core/Public/Templates/UnrealTemplate.h
@@ -748,8 +748,8 @@ FORCEINLINE T ImplicitConv(typename TIdentity<T>::Type Obj)
  */
 template <
 	typename T,
-	typename Base,
-	decltype(ImplicitConv<const volatile Base*>((std::remove_reference_t<T>*)nullptr))* = nullptr
+	typename Base
+	UE_REQUIRES(std::is_convertible_v<std::remove_reference_t<T>*, const volatile Base*>)
 >
 UE_INTRINSIC_CAST FORCEINLINE decltype(auto) ForwardAsBase(std::remove_reference_t<T>& Obj)
 {
@@ -758,8 +758,8 @@ UE_INTRINSIC_CAST FORCEINLINE decltype(auto) ForwardAsBase(std::remove_reference
 
 template <
 	typename T,
-	typename Base,
-	decltype(ImplicitConv<const volatile Base*>((std::remove_reference_t<T>*)nullptr))* = nullptr
+	typename Base
+	UE_REQUIRES(std::is_convertible_v<std::remove_reference_t<T>*, const volatile Base*>)
 >
 UE_INTRINSIC_CAST FORCEINLINE decltype(auto) ForwardAsBase(std::remove_reference_t<T>&& Obj)
 {
diff --git a/Engine/Source/Runtime/CoreUObject/Public/UObject/Field.h b/Engine/Source/Runtime/CoreUObject/Public/UObject/Field.h
--- a/Engine/Source/Runtime/CoreUObject/Public/UObject/Field.h
+++ b/Engine/Source/Runtime/CoreUObject/Public/UObject/Field.h
@@ -29,6 +29,7 @@ Field.h: Declares FField property system fundamentals
 #include "Templates/EnableIf.h"
 #include "Templates/IsAbstract.h"
 #include "Templates/IsEnum.h"
+#include "Templates/Requires.h"
 #include "Templates/TypeHash.h"
 #include "Templates/UnrealTemplate.h"
 #include "Templates/UnrealTypeTraits.h"
@@ -293,8 +294,8 @@ public:
 	}
 
 	template <
-		typename T,
-		decltype(ImplicitConv<const UObject*>(std::declval<T>()))* = nullptr
+		typename T
+		UE_REQUIRES(std::is_convertible_v<T, const UObject*>)
 	>
 	FFieldVariant(T&& InObject)
 	{
diff --git a/Engine/Source/Runtime/CoreUObject/Public/UObject/FieldPath.h b/Engine/Source/Runtime/CoreUObject/Public/UObject/FieldPath.h
--- a/Engine/Source/Runtime/CoreUObject/Public/UObject/FieldPath.h
+++ b/Engine/Source/Runtime/CoreUObject/Public/UObject/FieldPath.h
@@ -15,6 +15,7 @@
 #include "Templates/EnableIf.h"
 #include "Templates/LosesQualifiersFromTo.h"
 #include "Templates/PointerIsConvertibleFromTo.h"
+#include "Templates/Requires.h"
 #include "Templates/TypeHash.h"
 #include "Templates/UnrealTemplate.h"
 #include "UObject/Field.h"
@@ -23,6 +24,8 @@
 #include "UObject/WeakObjectPtr.h"
 #include "UObject/WeakObjectPtrTemplates.h"
 
+#include <type_traits>
+
 class FArchive;
 class FLinkerLoad;
 class UField;
@@ -334,8 +337,8 @@ public:
 	* @param Object object to create a weak pointer to
 	**/
 	template <
-		typename OtherPropertyType,
-		typename = decltype(ImplicitConv<PropertyType*>((OtherPropertyType*)nullptr))
+		typename OtherPropertyType
+		UE_REQUIRES(std::is_convertible_v<OtherPropertyType*, PropertyType*>)
 	>
 	FORCEINLINE TFieldPath(OtherPropertyType* InProperty, EDummy1 = Dummy1)
 		: FFieldPath((FField*)CastField<PropertyType>(InProperty))
diff --git a/Engine/Source/Runtime/CoreUObject/Public/UObject/LazyObjectPtr.h b/Engine/Source/Runtime/CoreUObject/Public/UObject/LazyObjectPtr.h
--- a/Engine/Source/Runtime/CoreUObject/Public/UObject/LazyObjectPtr.h
+++ b/Engine/Source/Runtime/CoreUObject/Public/UObject/LazyObjectPtr.h
@@ -14,10 +14,13 @@
 #include "Serialization/Archive.h"
 #include "Serialization/StructuredArchive.h"
 #include "Templates/Casts.h"
+#include "Templates/Requires.h"
 #include "Templates/UnrealTemplate.h"
 #include "UObject/Object.h"
 #include "UObject/PersistentObjectPtr.h"
 
+#include <type_traits>
+
 template <typename T> struct TIsPODType;
 template <typename T> struct TIsWeakPointerType;
 
@@ -193,14 +196,20 @@ public:
 	TLazyObjectPtr<T>& operator=(const TLazyObjectPtr<T>&) = default;
 
 	/** Construct from another lazy pointer with implicit upcasting allowed */
-	template<typename U, typename = decltype(ImplicitConv<T*>((U*)nullptr))>
+	template <
+		typename U
+		UE_REQUIRES(std::is_convertible_v<U*, T*>)
+	>
 	FORCEINLINE TLazyObjectPtr(const TLazyObjectPtr<U>& Other) :
 		FLazyObjectPtr((const FLazyObjectPtr&)Other)
 	{
 	}
 	
 	/** Assign from another lazy pointer with implicit upcasting allowed */
-	template<typename U, typename = decltype(ImplicitConv<T*>((U*)nullptr))>
+	template <
+		typename U
+		UE_REQUIRES(std::is_convertible_v<U*, T*>)
+	>
 	FORCEINLINE TLazyObjectPtr<T>& operator=(const TLazyObjectPtr<U>& Other)
 	{
 		FLazyObjectPtr::operator=((const FLazyObjectPtr&)Other);
diff --git a/Engine/Source/Runtime/CoreUObject/Public/UObject/ObjectKey.h b/Engine/Source/Runtime/CoreUObject/Public/UObject/ObjectKey.h
--- a/Engine/Source/Runtime/CoreUObject/Public/UObject/ObjectKey.h
+++ b/Engine/Source/Runtime/CoreUObject/Public/UObject/ObjectKey.h
@@ -2,6 +2,7 @@
 
 #pragma once
 
+#include "Templates/Requires.h"
 #include "UObject/WeakObjectPtr.h"
 
 #include <type_traits>
@@ -37,8 +38,8 @@ public:
 		}
 	}
 	template <
-		typename U,
-		decltype(ImplicitConv<const UObject*>(std::declval<U>()))* = nullptr
+		typename U
+		UE_REQUIRES(std::is_convertible_v<U, const UObject*>)
 	>
 	FORCEINLINE FObjectKey(U Object)
 		: FObjectKey(ImplicitConv<const UObject*>(Object))
@@ -163,8 +164,8 @@ public:
 
 	/** Construct from an object pointer */
 	template <
-		typename U,
-		decltype(ImplicitConv<const InElementType*>(std::declval<U>()))* = nullptr
+		typename U
+		UE_REQUIRES(std::is_convertible_v<U, const InElementType*>)
 	>
 	FORCEINLINE TObjectKey(U Object)
 		: ObjectKey(ImplicitConv<const InElementType*>(Object))
