Replaced ImplicitConv usage in FWeakObjectPtr and TWeakObjectPtr.

diff --git a/Engine/Source/Runtime/Core/Public/UObject/WeakObjectPtrTemplates.h b/Engine/Source/Runtime/Core/Public/UObject/WeakObjectPtrTemplates.h
--- a/Engine/Source/Runtime/Core/Public/UObject/WeakObjectPtrTemplates.h
+++ b/Engine/Source/Runtime/Core/Public/UObject/WeakObjectPtrTemplates.h
@@ -46,8 +46,8 @@ public:
 	 * @param Object object to create a weak pointer to
 	 */
 	template <
-		typename U,
-		typename = decltype(ImplicitConv<T*>(std::declval<U>()))
+		typename U
+		UE_REQUIRES(std::is_convertible_v<U, T*>)
 	>
 	FORCEINLINE TWeakObjectPtr(U Object) :
 		TWeakObjectPtrBase((const UObject*)Object)
@@ -62,8 +62,8 @@ public:
 	 * @param Other weak pointer to copy from
 	 */
 	template <
-		typename OtherT,
-		typename = decltype(ImplicitConv<T*>((OtherT*)nullptr))
+		typename OtherT
+		UE_REQUIRES(std::is_convertible_v<OtherT*, T*>)
 	>
 	FORCEINLINE TWeakObjectPtr(const TWeakObjectPtr<OtherT, TWeakObjectPtrBase>& Other) :
 		TWeakObjectPtrBase(*(TWeakObjectPtrBase*)&Other) // we do a C-style cast to private base here to avoid clang 3.6.0 compilation problems with friend declarations
@@ -98,8 +98,8 @@ public:
 	 * @param Other weak pointer to copy from
 	 */
 	template <
-		typename OtherT,
-		typename = decltype(ImplicitConv<T*>((OtherT*)nullptr))
+		typename OtherT
+		UE_REQUIRES(std::is_convertible_v<OtherT*, T*>)
 	>
 	FORCEINLINE TWeakObjectPtr& operator=(const TWeakObjectPtr<OtherT, TWeakObjectPtrBase>& Other)
 	{
diff --git a/Engine/Source/Runtime/CoreUObject/Public/UObject/WeakObjectPtr.h b/Engine/Source/Runtime/CoreUObject/Public/UObject/WeakObjectPtr.h
--- a/Engine/Source/Runtime/CoreUObject/Public/UObject/WeakObjectPtr.h
+++ b/Engine/Source/Runtime/CoreUObject/Public/UObject/WeakObjectPtr.h
@@ -76,8 +76,8 @@ public:
 	 * @param Object object to create a weak pointer to
 	 */
 	template <
-		typename U,
-		decltype(ImplicitConv<const UObject*>(std::declval<U>()))* = nullptr
+		typename U
+		UE_REQUIRES(std::is_convertible_v<U, const UObject*>)
 	>
 	FORCEINLINE FWeakObjectPtr(U&& Object)
 	{
