Changed UE_REQUIRES_EXPR to not refer to ImplicitConv, to lessen the chance of users copying it.

diff --git a/Engine/Source/Runtime/Core/Public/Templates/Requires.h b/Engine/Source/Runtime/Core/Public/Templates/Requires.h
--- a/Engine/Source/Runtime/Core/Public/Templates/Requires.h
+++ b/Engine/Source/Runtime/Core/Public/Templates/Requires.h
@@ -19,26 +19,19 @@
 	UE_REQUIRES_EXPR() wraps the effects of a requires expression, used to test the
 	compilability of an expression based on a deduced template parameter.  Usage:
 
-	template <typename Type>
-	struct TSmartPtr
+	template <
+		typename LhsType,
+		typename RhsType
+		UE_REQUIRES(
+			std::is_base_of_v<UObject, LhsType> &&
+			std::is_base_of_v<UObject, RhsType> &&
+			UE_REQUIRES_EXPR(std::declval<const LhsType*>() == std::declval<const RhsType*>())
+		)
+	>
+	bool operator==(const TSmartPtr<LhsType>& Lhs, const TSmartPtr<RhsType>& Rhs)
 	{
-		// Only enable this constructor if the incoming pointer type is a UObject
-		// convertible to the smart pointer type:
-		template <
-			typename OtherType
-			UE_REQUIRES(
-				std::is_base_of_v<UObject, OtherType> &&
-				UE_REQUIRES_EXPR(ImplicitConv<Type*>((OtherType*)nullptr)))
-			)
-		>
-		explicit TSmartPtr(OtherType* OtherPtr)
-			: Ptr(Cast<Type>(OtherPtr))
-		{
-		}
-
-	private:
-		BaseType* Ptr;
-	};
+		return Lhs.Get() == Rhs.Get();
+	}
 
 	Unlike a C++20 requires expression, UE_REQUIRES_EXPR() can only be used in
 	the body of a UE_REQUIRES() macro - standalone concept checks must still be
