Fix a race condition with TLS initialization on CsvProfiler startup where GCsvProfilerIsCapturing=true  is visible on the renderthread before the initialization of TlsSlot. We could add a memory barrier in FCsvProfiler::Get() to ensure the variables are visible in the correct order from all threads, that but that would add overhead to every call, so instead we'll initialize the TLSSlot earlier (in CsvProfiler::Begin, rather than in the deferred command in FCsvProfiler::BeginFrame), ensuring a frame of separation.

diff --git a/Engine/Source/Runtime/Core/Private/ProfilingDebugging/CsvProfiler.cpp b/Engine/Source/Runtime/Core/Private/ProfilingDebugging/CsvProfiler.cpp
--- a/Engine/Source/Runtime/Core/Private/ProfilingDebugging/CsvProfiler.cpp
+++ b/Engine/Source/Runtime/Core/Private/ProfilingDebugging/CsvProfiler.cpp
@@ -1989,6 +1989,11 @@ public:
 		}
 	}
 
+	static bool IsTlsSlotInitialized()
+	{
+		return FPlatformTLS::IsValidTlsSlot(TlsSlot);
+	}
+
 	static FORCENOINLINE FCsvProfilerThreadData* CreateTLSData()
 	{
 		QUICK_SCOPE_CYCLE_COUNTER(CSVProfiler_ThreadData_CreateTLSData);
@@ -3067,8 +3072,8 @@ void FCsvProfiler::BeginFrame()
 
 					GCsvStatCounts = !!CVarCsvStatCounts.GetValueOnGameThread();
 
-					// Initialize tls before setting the capturing flag to true.
-					FCsvProfilerThreadData::InitTls();
+					// Check TLS is initialized before starting the capture. This should have happened in BeginCapture
+					check(FCsvProfilerThreadData::IsTlsSlotInitialized());
 					TRACE_CSV_PROFILER_BEGIN_CAPTURE(*Filename, RenderThreadId, RHIThreadId, GDefaultWaitStatName, GCsvStatCounts);
 					GCsvProfilerIsCapturing = true;
 				}
@@ -3339,6 +3344,8 @@ void FCsvProfiler::BeginCapture(int InNumFramesToCapture,
 		return;
 	}
 
+	// Lazy init TLS before starting the capture
+	FCsvProfilerThreadData::InitTls();
 
 	CommandQueue.Enqueue(FCsvCaptureCommand(ECsvCommandType::Start, GCsvProfilerFrameNumber, InNumFramesToCapture, InDestinationFolder, InFilename, InFlags));
 }
