Implemented denormals flush to zero floaring unit control flag
Plus fixed ARM FPU control flags API

diff --git a/Engine/Source/Runtime/Core/Public/Math/UnrealMathFPU.h b/Engine/Source/Runtime/Core/Public/Math/UnrealMathFPU.h
--- a/Engine/Source/Runtime/Core/Public/Math/UnrealMathFPU.h
+++ b/Engine/Source/Runtime/Core/Public/Math/UnrealMathFPU.h
@@ -1941,6 +1941,11 @@ FORCEINLINE uint32 VectorAnyGreaterThan(const VectorRegister4Double& Vec1, const
  */
 #define VECTOR_ROUND_TOWARD_ZERO		0
 
+ /**
+  * Denormal operands and results will be flushed to zero
+  */
+#define VECTOR_DENORMALS_FLUSH_TO_ZERO	0
+
 // Returns true if the vector contains a component that is either NAN or +/-infinite.
 inline bool VectorContainsNaNOrInfinite(const VectorRegister4Float& Vec)
 {
diff --git a/Engine/Source/Runtime/Core/Public/Math/UnrealMathNeon.h b/Engine/Source/Runtime/Core/Public/Math/UnrealMathNeon.h
--- a/Engine/Source/Runtime/Core/Public/Math/UnrealMathNeon.h
+++ b/Engine/Source/Runtime/Core/Public/Math/UnrealMathNeon.h
@@ -2264,19 +2264,42 @@ FORCEINLINE int32 VectorAnyGreaterThan(VectorRegister4Double Vec1, VectorRegiste
  *
  * @return			The uint32 control register
  */
-#define VectorGetControlRegister()		0
+FORCEINLINE uint32_t VectorGetControlRegister()
+{
+#if PLATFORM_WINDOWS
+	return (uint32_t)_ReadStatusReg(ARM64_FPCR);
+#else
+	uint64_t Value;
+	// The system register read/write instructions use 64-bit registers,
+	__asm__ volatile("mrs %0, fpcr" : "=r"(Value));
+	return (uint32_t)Value;
+#endif
+}
 
 /**
  * Sets the control register.
  *
  * @param ControlStatus		The uint32 control status value to set
  */
-#define	VectorSetControlRegister(ControlStatus)
+FORCEINLINE void VectorSetControlRegister(uint32_t ControlStatus)
+{
+#if PLATFORM_WINDOWS
+	_WriteStatusReg(ARM64_FPCR, ControlStatus);
+#else
+	uint64_t State64 = ControlStatus; // instruction needs a 64b reg, but all control bits fit in the lower 32b
+	__asm__ volatile("msr fpcr, %0" : : "r"(State64));
+#endif
+}
 
 /**
  * Control status bit to round all floating point math results towards zero.
  */
-#define VECTOR_ROUND_TOWARD_ZERO		0
+#define VECTOR_ROUND_TOWARD_ZERO		(3 << 22)
+
+ /**
+  * Denormal operands and results will be flushed to zero
+  */
+#define VECTOR_DENORMALS_FLUSH_TO_ZERO	(1 << 24)
 
 
 /**
diff --git a/Engine/Source/Runtime/Core/Public/Math/UnrealMathSSE.h b/Engine/Source/Runtime/Core/Public/Math/UnrealMathSSE.h
--- a/Engine/Source/Runtime/Core/Public/Math/UnrealMathSSE.h
+++ b/Engine/Source/Runtime/Core/Public/Math/UnrealMathSSE.h
@@ -2699,6 +2699,11 @@ FORCEINLINE int VectorAnyGreaterThan(const VectorRegister4Double& Vec1, const Ve
  */
 #define VECTOR_ROUND_TOWARD_ZERO		_MM_ROUND_TOWARD_ZERO
 
+ /**
+  * Denormal operands and results will be flushed to zero
+  */
+#define VECTOR_DENORMALS_FLUSH_TO_ZERO	_MM_FLUSH_ZERO_ON
+
 /**
 * Multiplies two quaternions; the order matters.
 *
