Add support for creating thread pools after forking.

diff --git a/Engine/Source/Runtime/Core/Private/HAL/ThreadingBase.cpp b/Engine/Source/Runtime/Core/Private/HAL/ThreadingBase.cpp
--- a/Engine/Source/Runtime/Core/Private/HAL/ThreadingBase.cpp
+++ b/Engine/Source/Runtime/Core/Private/HAL/ThreadingBase.cpp
@@ -1015,6 +1015,17 @@ protected:
 	 */
 	virtual uint32 Run() override;
 
+	/** Common code use by Create functions */
+	FString SetupForCreate(class FQueuedThreadPoolBase* InPool, const TCHAR* ThreadName)
+	{
+		static int32 PoolThreadIndex = 0;
+		const FString PoolThreadName = ThreadName ? FString(ThreadName) : FString::Printf(TEXT("PoolThread %d"), PoolThreadIndex++);
+
+		OwningThreadPool = InPool;
+		DoWorkEvent = FPlatformProcess::GetSynchEventFromPool();
+		return PoolThreadName;
+	}
+
 public:
 
 	/** Default constructor **/
@@ -1031,15 +1042,32 @@ public:
 	 */
 	virtual bool Create(class FQueuedThreadPoolBase* InPool,uint32 InStackSize = 0, EThreadPriority ThreadPriority=TPri_Normal, const TCHAR* ThreadName = nullptr)
 	{
-		static int32 PoolThreadIndex = 0;
-		const FString PoolThreadName = ThreadName ? FString(ThreadName) : FString::Printf( TEXT( "PoolThread %d" ), PoolThreadIndex++ );
+		const FString PoolThreadName = SetupForCreate(InPool, ThreadName);
 
-		OwningThreadPool = InPool;
-		DoWorkEvent = FPlatformProcess::GetSynchEventFromPool();
 		Thread = FRunnableThread::Create(this, *PoolThreadName, InStackSize, ThreadPriority, FPlatformAffinity::GetPoolThreadMask());
 		check(Thread);
 		return true;
 	}
+
+	/**
+	 * Attempts to creates a forkable thread (see FForkProcessHelper) with the specified stack size and creates the various
+	 * events to be able to communicate with it.
+	 * 
+	 * Currently only works after the process has been forked because GetSingleThreadInterface() isn't implemented
+	 *
+	 * @param InPool The thread pool interface used to place this thread back into the pool of available threads when its work is done
+	 * @param InStackSize The size of the stack to create. 0 means use the current thread's stack size
+	 * @param ThreadPriority priority of new thread
+	 * @return True if the thread and all of its initialization was successful, false otherwise
+	 */
+	virtual bool CreateForkable(class FQueuedThreadPoolBase* InPool, uint32 InStackSize = 0, EThreadPriority ThreadPriority = TPri_Normal, const TCHAR* ThreadName = nullptr)
+	{
+		const FString PoolThreadName = SetupForCreate(InPool, ThreadName);
+
+		// Currently no single threaded support, so this is expected to fail before forking
+		Thread = FForkProcessHelper::CreateForkableThread(this, *PoolThreadName, InStackSize, ThreadPriority, FPlatformAffinity::GetPoolThreadMask());
+		return !!Thread;
+	}
 	
 	/**
 	 * Tells the thread to exit. If the caller needs to know when the thread
@@ -1123,7 +1151,7 @@ public:
 		Destroy();
 	}
 
-	virtual bool Create(uint32 InNumQueuedThreads, uint32 StackSize, EThreadPriority ThreadPriority, const TCHAR* Name) override
+	bool CreateInternal(bool bForkable, uint32 InNumQueuedThreads, uint32 StackSize, EThreadPriority ThreadPriority, const TCHAR* Name)
 	{
 		UE::Trace::ThreadGroupBegin(Name);
 
@@ -1150,7 +1178,10 @@ public:
 			FQueuedThread* pThread = new FQueuedThread();
 			// Now create the thread and add it if ok
 			const FString ThreadName = FString::Printf(TEXT("%s #%d"), Name, Count);
-			if (pThread->Create(this, StackSize, ThreadPriority, *ThreadName) == true)
+			const bool bCreateSuccess = bForkable ?
+				pThread->CreateForkable(this, StackSize, ThreadPriority, *ThreadName) :
+				pThread->Create(this, StackSize, ThreadPriority, *ThreadName);
+			if (bCreateSuccess)
 			{
 				QueuedThreads.Add(pThread);
 				AllThreads.Add(pThread);
@@ -1172,6 +1203,16 @@ public:
 		return bWasSuccessful;
 	}
 
+	virtual bool Create(uint32 InNumQueuedThreads, uint32 StackSize, EThreadPriority ThreadPriority, const TCHAR* Name) override
+	{
+		return CreateInternal(false, InNumQueuedThreads, StackSize, ThreadPriority, Name);
+	}
+
+	virtual bool CreateForkable(uint32 InNumQueuedThreads, uint32 StackSize, EThreadPriority ThreadPriority, const TCHAR* Name) override
+	{
+		return CreateInternal(true, InNumQueuedThreads, StackSize, ThreadPriority, Name);
+	}
+
 	virtual void Destroy() override final
 	{
 		if (SynchQueue)
diff --git a/Engine/Source/Runtime/Core/Public/Misc/QueuedThreadPool.h b/Engine/Source/Runtime/Core/Public/Misc/QueuedThreadPool.h
--- a/Engine/Source/Runtime/Core/Public/Misc/QueuedThreadPool.h
+++ b/Engine/Source/Runtime/Core/Public/Misc/QueuedThreadPool.h
@@ -114,6 +114,22 @@ public:
 	 */
 	virtual bool Create(uint32 InNumQueuedThreads, uint32 StackSize = (32 * 1024), EThreadPriority ThreadPriority = TPri_Normal, const TCHAR* Name = TEXT("UnknownThreadPool")) = 0;
 
+	/**
+	 * Creates the thread pool with the specified number of forkable threads (see FForkProcessHelper)
+	 * 
+	 * Currently only works after the process has been forked because GetSingleThreadInterface() isn't implemented for pooled threads
+	 *
+	 * @param InNumQueuedThreads Specifies the number of threads to use in the pool
+	 * @param StackSize The size of stack the threads in the pool need (32K default)
+	 * @param ThreadPriority priority of new pool thread
+	 * @param Name optional name for the pool to be used for instrumentation
+	 * @return Whether the pool creation was successful or not
+	 */
+	virtual bool CreateForkable(uint32 InNumQueuedThreads, uint32 StackSize = (32 * 1024), EThreadPriority ThreadPriority = TPri_Normal, const TCHAR* Name = TEXT("UnknownThreadPool"))
+	{
+		return false;
+	}
+
 	/** Tells the pool to clean up all background threads */
 	virtual void Destroy() = 0;
 
