Add ADD_WARNING and ADD_ERROR macros to match TestHarness.h

diff --git a/Engine/Source/Runtime/Core/Public/Misc/LowLevelTestAdapter.h b/Engine/Source/Runtime/Core/Public/Misc/LowLevelTestAdapter.h
--- a/Engine/Source/Runtime/Core/Public/Misc/LowLevelTestAdapter.h
+++ b/Engine/Source/Runtime/Core/Public/Misc/LowLevelTestAdapter.h
@@ -148,8 +148,11 @@ FString CaptureExpressionsAndValues(const FString& InExpressions, ArgTypes&&...
 #define CAPTURE(...) FAutomationTestFramework::Get().GetCurrentTest()->AddInfo(CaptureExpressionsAndValues(#__VA_ARGS__, __VA_ARGS__));
 #define INFO(Message) FAutomationTestFramework::Get().GetCurrentTest()->AddInfo(Message);
 #define WARN(Message) FAutomationTestFramework::Get().GetCurrentTest()->AddWarning(Message); 
+#define ADD_WARNING(Message) FAutomationTestFramework::Get().GetCurrentTest()->AddWarning(Message); 
+#define ADD_ERROR(Message) FAutomationTestFramework::Get().GetCurrentTest()->AddError(Message); 
 #define FAIL_ON_MESSAGE(Message) FAutomationTestFramework::Get().GetCurrentTest()->AddExpectedError(Message);
 
+
 #define SKIP(Message)
 
 #endif // !WITH_LOW_LEVEL_TESTS
\ No newline at end of file
