TDirectoryTree updates:
Add copy constructor.
Add ContainsChildPaths method.
Add tests to FoundationTests runner.

diff --git a/Engine/Source/Runtime/Core/Internal/Containers/DirectoryTree.h b/Engine/Source/Runtime/Core/Internal/Containers/DirectoryTree.h
--- a/Engine/Source/Runtime/Core/Internal/Containers/DirectoryTree.h
+++ b/Engine/Source/Runtime/Core/Internal/Containers/DirectoryTree.h
@@ -20,7 +20,7 @@ CORE_API void FixupPathSeparator(FStringBuilderBase& InOutPath, int32 StartIndex
 CORE_API int32 FindInsertionIndex(int32 NumChildNodes, const TUniquePtr<FString[]>& RelPaths,
 	FStringView FirstPathComponent, bool& bOutExists);
 
-}
+} // namespace UE::DirectoryTree
 
 enum class EDirectoryTreeGetFlags
 {
@@ -84,6 +84,12 @@ class TDirectoryTree
 public:
 	TDirectoryTree();
 
+	TDirectoryTree(const TDirectoryTree& Other) = default;
+	TDirectoryTree& operator=(const TDirectoryTree& Other) = default;
+
+	TDirectoryTree(TDirectoryTree&& Other) = default;
+	TDirectoryTree& operator=(TDirectoryTree&& Other) = default;
+
 	/**
 	 * Add a path to the tree if it does not already exist. Construct default Value for it if it did not already exist.
 	 * Return a reference to the added or existing Value. Optionally report whether the path already existed.
@@ -143,6 +149,11 @@ public:
 	 */
 	bool TryFindClosestPath(FStringView Path, FString& OutPath, ValueType** OutValue);
 
+	/**
+	 * Return whether any children of the given path have been added to the tree.
+	 */
+	bool ContainsChildPaths(FStringView Path) const;
+
 	/**
 	 * Report the children (optionally recursive or not, optionally implied or not) in the tree of a given Path
 	 * (optionally skipped if implied). @see EDirectoryTreeGetFlags.
@@ -165,8 +176,8 @@ private:
 	 * Root
 	 *		/				(FullPath: /)
 	 *			A			(FullPath: /A
-	 *				X		(FullPath: /A/Y
-	 *				Y/M		(FullPath: /A/Z/N)
+	 *				X		(FullPath: /A/X
+	 *				Y/M		(FullPath: /A/Y/M)
 	 *			B/Z			(FullPath: /B/Z)
 	 *				N		(FullPath: /B/Z/N)
 	 *				O		(FullPath: /B/Z/O)
@@ -178,8 +189,8 @@ private:
 		FTreeNode() = default;
 		FTreeNode(FTreeNode&& Other);
 		FTreeNode& operator=(FTreeNode&& Other);
-		FTreeNode(const FTreeNode& Other) = delete;
-		FTreeNode& operator=(const FTreeNode& Other) = delete;
+		FTreeNode(const FTreeNode& Other);
+		FTreeNode& operator=(const FTreeNode& Other);
 
 		/** Remove Value and ChildNodes, return state to default-constructed state. */
 		void Reset();
@@ -196,6 +207,14 @@ private:
 		/** Return pointer to the Value stored in RelPath, if RelPath exists in the tree. */
 		ValueType* Find(FStringView InRelPath);
 
+		/**
+		 * Recursively search this node's subtree for the given relative path, then return whether that path has any
+		 * children. Returns false if the given relative path doesn't exist in the tree or if it is an explicit path
+		 * with no children. Returns true if the given path is implicit (there is a node representing a sub-path of it)
+		 * or if it is an explicit path with direct or indirect children.
+		 */
+		bool ContainsChildPaths(FStringView InRelPath) const;
+
 		bool TryGetChildren(FStringBuilderBase& ReportedPathPrefix, TCHAR InPathSeparator, FStringView InRelPath,
 			TArray<FString>& OutRelativeChildNames, EDirectoryTreeGetFlags Flags) const;
 
@@ -221,6 +240,8 @@ private:
 		bool HasValue() const;
 		/** Get a reference to the node's Value. Invalid to call if !HasValue. */
 		ValueType& GetValue();
+		/** Get a reference to the node's Value. Invalid to call if !HasValue. */
+		const ValueType& GetValue() const;
 		/** Set HasValue=true, and move InValue into the node's value, after destructing any existing old value. */
 		void SetValue(ValueType&& InValue);
 		/** Set HasValue=true, and default-construct the node's value, after destructing any existing old value. */
@@ -585,7 +606,16 @@ inline void TDirectoryTree<ValueType>::InitializePathSeparator(TCHAR InPathSepar
 	bPathSeparatorInitialized = true;
 }
 
-template <typename ValueType>
+template<typename ValueType>
+inline bool TDirectoryTree<ValueType>::ContainsChildPaths(FStringView Path) const
+{
+	TStringBuilder<16> NormalizeBuffer;
+	NormalizePathForReading(Path, NormalizeBuffer);
+
+	return Root.ContainsChildPaths(Path);
+}
+
+template<typename ValueType>
 inline bool TDirectoryTree<ValueType>::TryGetChildren(FStringView Path, TArray<FString>& OutRelativeChildNames,
 	EDirectoryTreeGetFlags Flags) const
 {
@@ -600,6 +630,37 @@ inline bool TDirectoryTree<ValueType>::TryGetChildren(FStringView Path, TArray<F
 	return Root.TryGetChildren(ReportedPathPrefix, PathSeparator, Path, OutRelativeChildNames, Flags);
 }
 
+template<typename ValueType>
+TDirectoryTree<ValueType>::FTreeNode::FTreeNode(const FTreeNode& Other)
+{
+	*this = Other;
+}
+
+template<typename ValueType>
+typename TDirectoryTree<ValueType>::FTreeNode& TDirectoryTree<ValueType>::FTreeNode::operator=(const FTreeNode& Other)
+{
+	Reset();
+
+	if (Other.HasValue())
+	{
+		SetValue(CopyTemp(Other.GetValue()));
+	}
+
+	Realloc(Other.CapacityChildNodes);
+	int32 NumChildren = Other.GetNumChildNodes();
+	for (int32 i = 0; i < NumChildren; ++i)
+	{
+		RelPaths[i] = Other.RelPaths[i];
+	}
+	for (int32 i = 0; i < NumChildren; ++i)
+	{
+		ChildNodes[i] = Other.ChildNodes[i];
+	}
+
+	SetNumChildNodes(Other.GetNumChildNodes());
+	return *this;
+}
+
 template <typename ValueType>
 TDirectoryTree<ValueType>::FTreeNode::FTreeNode(FTreeNode&& Other)
 {
@@ -665,6 +726,12 @@ inline ValueType& TDirectoryTree<ValueType>::FTreeNode::GetValue()
 	return *reinterpret_cast<ValueType*>(&Value);
 }
 
+template<typename ValueType>
+inline const ValueType& TDirectoryTree<ValueType>::FTreeNode::GetValue() const
+{
+	return *reinterpret_cast<const ValueType*>(&Value);
+}
+
 template <typename ValueType>
 inline void TDirectoryTree<ValueType>::FTreeNode::SetValue(ValueType&& InValue)
 {
@@ -913,6 +980,85 @@ inline ValueType* TDirectoryTree<ValueType>::FTreeNode::Find(FStringView InRelPa
 	}
 }
 
+template<typename ValueType>
+inline bool TDirectoryTree<ValueType>::FTreeNode::ContainsChildPaths(FStringView InRelPath) const
+{
+	if (InRelPath.IsEmpty())
+	{
+		// This is the node we were searching for as an explicit entry
+		return GetNumChildNodes() > 0;
+	}
+
+	// We are still looking for the requested InRelPath and are not reporting results yet. Look for an existing
+	// stored child that has the same FirstComponent of its path as InRelPath does.
+	FStringView FirstComponent;
+	FStringView RemainingPath;
+	FPathViews::SplitFirstComponent(InRelPath, FirstComponent, RemainingPath);
+	bool bExists;
+	int32 InsertionIndex = FindInsertionIndex(FirstComponent, bExists);
+	if (!bExists)
+	{
+		// No child has the same FirstComponent as InRelPath, so InRelPath does not exist in the tree, not even as
+		// an implied path.
+		return false;
+	}
+
+	FTreeNode& ChildNode = ChildNodes[InsertionIndex];
+	FString& ChildRelPath = RelPaths[InsertionIndex];
+
+	FStringView ExistingFirstComponent;
+	FStringView ExistingRemainingPath;
+	FPathViews::SplitFirstComponent(ChildRelPath, ExistingFirstComponent, ExistingRemainingPath);
+	check(FPathViews::Equals(
+		FirstComponent,
+		ExistingFirstComponent)); // Otherwise FindInsertionIndex would have returned bExists=false
+
+	// Same logic as TryGetChildren - progressively match the InRelPath against ChildRelPath until one is empty
+	for (int32 RunawayLoop = 0; RunawayLoop <= InRelPath.Len(); ++RunawayLoop)
+	{
+		if (ExistingRemainingPath.IsEmpty())
+		{
+			// We've reached the end of the existing path, so InRelPath is either equal to or a child of the
+			// ChildNode. Delegate to that node to keep searching or return results
+			return ChildNode.ContainsChildPaths(RemainingPath);
+		}
+		else if (RemainingPath.IsEmpty())
+		{
+			// We've reached the end of the input path, but not the end of the existing path, so InRelPath is a
+			// parent of the existing path, and is an implied path rather than an added path.
+			return true;
+		}
+		else
+		{
+			// Both existing and remaining have more directory components
+			FStringView NextFirstComponent;
+			FStringView NextRemainingPath;
+			FPathViews::SplitFirstComponent(RemainingPath, NextFirstComponent, NextRemainingPath);
+			check(!NextFirstComponent.IsEmpty());
+			FStringView NextExistingFirstComponent;
+			FStringView NextExistingRemainingPath;
+			FPathViews::SplitFirstComponent(ExistingRemainingPath, NextExistingFirstComponent, NextExistingRemainingPath);
+			check(!NextExistingFirstComponent.IsEmpty());
+			if (NextFirstComponent == NextExistingFirstComponent)
+			{
+				// Next component is also a match, go to the next loop iteration to handle the new remainingpaths
+				RemainingPath = NextRemainingPath;
+				ExistingRemainingPath = NextExistingRemainingPath;
+
+				continue;
+			}
+			else
+			{
+				// The existing child diverges from the path components of InRelPath, so InRelPath does not exist
+				// in the tree, not even as an implied path.
+				return false;
+			}
+		}
+	}
+	checkf(false, TEXT("Infinite loop trying to split path %.*s into components."), InRelPath.Len(), InRelPath.GetData());
+	return false;
+}
+
 template <typename ValueType>
 inline bool TDirectoryTree<ValueType>::FTreeNode::TryGetChildren(FStringBuilderBase& ReportPathPrefix,
 	TCHAR InPathSeparator, FStringView InRelPath, TArray<FString>& OutRelativeChildNames,
diff --git a/Engine/Source/Runtime/Core/Private/Containers/DirectoryTree.cpp b/Engine/Source/Runtime/Core/Private/Containers/DirectoryTree.cpp
--- a/Engine/Source/Runtime/Core/Private/Containers/DirectoryTree.cpp
+++ b/Engine/Source/Runtime/Core/Private/Containers/DirectoryTree.cpp
@@ -37,923 +37,3 @@ int32 FindInsertionIndex(int32 NumChildNodes, const TUniquePtr<FString[]>& RelPa
 }
 
 } // namespace UE::DirectoryTree
-
-#if WITH_DEV_AUTOMATION_TESTS
-
-IMPLEMENT_SIMPLE_AUTOMATION_TEST(FDirectoryTreeTest, "System.Core.DirectoryTree", EAutomationTestFlags::EditorContext | EAutomationTestFlags::EngineFilter)
-bool FDirectoryTreeTest::RunTest(const FString& Parameters)
-{
-	constexpr int32 NumPathTypes = 5;
-	constexpr int32 NumPaths = 9;
-	FStringView PathsByTypeAndIndex[NumPathTypes][NumPaths] =
-	{
-		{
-			TEXTVIEW("/Game/Dir2"),
-			TEXTVIEW("/Game/Path1"),
-			TEXTVIEW("/Game/Dir2/Path2"),
-			TEXTVIEW("/Plugin1/Path1"),
-			TEXTVIEW("/Plugin1/Dir2/Path2"),
-			TEXTVIEW("/Engine/Path1"),
-			TEXTVIEW("/Plugin2/Path1"),
-			// Make sure we handle suffixes of an existing string with a leading value that sorts after /
-			TEXTVIEW("/Game/Foo/Leaf"),
-			TEXTVIEW("/Game/Foo-Bar/Leaf"),
-		},
-		{
-			TEXTVIEW("d:\\root\\Project\\Content\\Dir2"),
-			TEXTVIEW("d:\\root\\Project\\Content\\Path1.uasset"),
-			TEXTVIEW("d:\\root\\Project\\Content\\Dir2\\Path2.uasset"),
-			TEXTVIEW("d:\\root\\Project\\Plugins\\Plugin1\\Content\\Path1.uasset"),
-			TEXTVIEW("d:\\root\\Project\\Plugins\\Plugin1\\Content\\Dir2\\Path2.uasset"),
-			TEXTVIEW("d:\\root\\Engine\\Content\\Path1.uasset"),
-			TEXTVIEW("e:\\root\\Project\\Plugins\\Plugin2\\Content\\Path1.uasset"),
-			TEXTVIEW("d:\\root\\Project\\Content\\Foo\\Leaf"),
-			TEXTVIEW("d:\\root\\Project\\Content\\Foo-Bar\\Leaf"),
-		},
-		{
-			TEXTVIEW("d:/root/Project/Content/Dir2"),
-			TEXTVIEW("d:/root/Project/Content/Path1.uasset"),
-			TEXTVIEW("d:/root/Project/Content/Dir2/Path2.uasset"),
-			TEXTVIEW("d:/root/Project/Plugins/Plugin1/Content/Path1.uasset"),
-			TEXTVIEW("d:/root/Project/Plugins/Plugin1/Content/Dir2/Path2.uasset"),
-			TEXTVIEW("d:/root/Engine/Content/Path1.uasset"),
-			TEXTVIEW("e:/root/Project/Plugins/Plugin2/Content/Path1.uasset"),
-			TEXTVIEW("d:/root/Project/Content/Foo/Leaf"),
-			TEXTVIEW("d:/root/Project/Content/Foo-Bar/Leaf"),
-		},
-		{
-			TEXTVIEW("..\\..\\..\\Project\\Content\\Dir2"),
-			TEXTVIEW("..\\..\\..\\Project\\Content\\Path1.uasset"),
-			TEXTVIEW("..\\..\\..\\Project\\Content\\Dir2\\Path2.uasset"),
-			TEXTVIEW("..\\..\\..\\Project\\Plugins\\Plugin1\\Content\\Path1.uasset"),
-			TEXTVIEW("..\\..\\..\\Project\\Plugins\\Plugin1\\Content\\Dir2\\Path2.uasset"),
-			TEXTVIEW("..\\..\\..\\Engine\\Content\\Path1.uasset"),
-			TEXTVIEW("e:\\root\\Project\\Plugins\\Plugin2\\Content\\Path1.uasset"),
-			TEXTVIEW("..\\..\\..\\Project\\Content\\Foo\\Leaf"),
-			TEXTVIEW("..\\..\\..\\Project\\Content\\Foo-Bar\\Leaf"),
-		},
-		{
-			TEXTVIEW("../../../Project/Content/Dir2"),
-			TEXTVIEW("../../../Project/Content/Path1.uasset"),
-			TEXTVIEW("../../../Project/Content/Dir2/Path2.uasset"),
-			TEXTVIEW("../../../Project/Plugins/Plugin1/Content/Path1.uasset"),
-			TEXTVIEW("../../../Project/Plugins/Plugin1/Content/Dir2/Path2.uasset"),
-			TEXTVIEW("../../../Engine/Content/Path1.uasset"),
-			TEXTVIEW("e:/root/Project/Plugins/Plugin2/Content/Path1.uasset"),
-			TEXTVIEW("../../../Project/Content/Foo/Leaf"),
-			TEXTVIEW("../../../Project/Content/Foo-Bar/Leaf"),
-		},
-	};
-	// PathSub0SubPath[i] provides a sub path of PathsByTypeAndIndex[i][0]
-	FStringView Path0SubPath[NumPathTypes] =
-	{
-		TEXTVIEW("/Game/Dir2/Sub"),
-		TEXTVIEW("d:\\root\\Project\\Content\\Dir2\\Sub"),
-		TEXTVIEW("d:/root/Project/Content/Dir2/Sub"),
-		TEXTVIEW("..\\..\\..\\Project\\Content\\Dir2"),
-		TEXTVIEW("../../../Project/Content/Dir2/Sub"),
-	};
-	int32 ValueByIndex[NumPaths] = { 1,2,3,4,5,6,7,8,9 };
-	FStringView NonPathsByTypeAndIndex0[] =
-		{
-			TEXTVIEW(""),
-			TEXTVIEW("/"),
-			TEXTVIEW("/Game"),
-			TEXTVIEW("/Game/"),
-			TEXTVIEW("/Plugin1"),
-			TEXTVIEW("/Plugin1/"),
-			TEXTVIEW("/Plugin1/Dir2"),
-			TEXTVIEW("/Plugin1/Dir2/"),
-			TEXTVIEW("/Engine"),
-			TEXTVIEW("/Engine/"),
-		};
-	FStringView NonPathsByTypeAndIndex1[] =
-		{
-			TEXTVIEW(""),
-			TEXTVIEW("d:\\"),
-			TEXTVIEW("d:\\root1"),
-			TEXTVIEW("d:\\root1\\"),
-			TEXTVIEW("d:\\root1\\Project"),
-			TEXTVIEW("d:\\root1\\Project\\"),
-			TEXTVIEW("d:\\root1\\Project\\Content"),
-			TEXTVIEW("d:\\root1\\Project\\Plugins\\"),
-			TEXTVIEW("d:\\root1\\Project\\Plugins\\Content"),
-			TEXTVIEW("d:\\root1\\Project\\Plugins\\Content\\"),
-			TEXTVIEW("d:\\root1\\Project\\Plugins\\Content\\Plugin1"),
-			TEXTVIEW("d:\\root1\\Project\\Plugins\\Content\\Plugin1\\"),
-			TEXTVIEW("d:\\root1\\Project\\Plugins\\Content\\Plugin1\\Dir2"),
-			TEXTVIEW("d:\\root1\\Project\\Plugins\\Content\\Plugin1\\Dir2\\"),
-			TEXTVIEW("d:\\root1\\Engine"),
-			TEXTVIEW("d:\\root1\\Engine\\"),
-			TEXTVIEW("d:\\root1\\Engine\\Content"),
-			TEXTVIEW("d:\\root1\\Engine\\Content\\"),
-		};
-	FStringView NonPathsByTypeAndIndex2[] =
-		{
-			TEXTVIEW(""),
-			TEXTVIEW("d:/"),
-			TEXTVIEW("d:/root1"),
-			TEXTVIEW("d:/root1/Project"),
-			TEXTVIEW("d:/root1/Project/Content"),
-			TEXTVIEW("d:/root1/Project/Plugins/Content"),
-			TEXTVIEW("d:/root1/Project/Plugins/Content/Plugin1"),
-			TEXTVIEW("d:/root1/Project/Plugins/Content/Plugin1/Dir2"),
-			TEXTVIEW("d:/root1/Engine"),
-			TEXTVIEW("d:/root1/Engine/Content"),
-		};
-	FStringView NonPathsByTypeAndIndex3[] =
-		{
-			TEXTVIEW(""),
-			TEXTVIEW("\\"),
-			TEXTVIEW(".."),
-			TEXTVIEW("..\\"),
-			TEXTVIEW("..\\.."),
-			TEXTVIEW("..\\..\\"),
-			TEXTVIEW("..\\..\\.."),
-			TEXTVIEW("..\\..\\..\\"),
-			TEXTVIEW("..\\..\\..\\Project"),
-			TEXTVIEW("..\\..\\..\\Project\\Content"),
-			TEXTVIEW("..\\..\\..\\Project\\Plugins\\Content"),
-			TEXTVIEW("..\\..\\..\\Project\\Plugins\\Content\\Plugin1"),
-			TEXTVIEW("..\\..\\..\\Project\\Plugins\\Content\\Plugin1\\Dir2"),
-			TEXTVIEW("..\\..\\..\\Engine"),
-			TEXTVIEW("..\\..\\..\\Engine\\Content"),
-		};
-	FStringView NonPathsByTypeAndIndex4[] =
-		{
-			TEXTVIEW(""),
-			TEXTVIEW("/"),
-			TEXTVIEW(".."),
-			TEXTVIEW("../.."),
-			TEXTVIEW("../../.."),
-			TEXTVIEW("../../../Project"),
-			TEXTVIEW("../../../Project/Content"),
-			TEXTVIEW("../../../Project/Plugins/Content"),
-			TEXTVIEW("../../../Project/Plugins/Content/Plugin1"),
-			TEXTVIEW("../../../Project/Plugins/Content/Plugin1/Dir2"),
-			TEXTVIEW("../../../Engine"),
-			TEXTVIEW("../../../Engine/Content"),
-		};
-	TConstArrayView<FStringView> NonPathsByTypeAndIndex[NumPathTypes] =
-	{
-		NonPathsByTypeAndIndex0, NonPathsByTypeAndIndex1, NonPathsByTypeAndIndex2, NonPathsByTypeAndIndex3,
-		NonPathsByTypeAndIndex4
-	};
-
-	constexpr int32 NumPermutations = 2;
-	int32 Permutations[NumPermutations][NumPaths] =
-	{
-		{ 0, 1, 2, 3, 4, 5, 6, 7, 8 },
-		{ 8, 7, 6, 5, 4, 3, 2, 1, 0 },
-	};
-
-	auto GetDirTreeTestName = [](FStringView InTestName, int32 PathType, int32 Permutation, int32 EditPermutationIndex,
-		int32 OtherPermutationIndex)
-	{
-		return FString::Printf(TEXT("%.*s(%d, %d, %d, %d)"),
-			InTestName.Len(), InTestName.GetData(),
-			PathType, Permutation, EditPermutationIndex, OtherPermutationIndex);
-	};
-
-	TArray<const TCHAR*> ScratchA;
-	TArray<const TCHAR*> ScratchB;
-	auto UnorderedEquals = [&ScratchA, &ScratchB](const TArray<FString>& A, TConstArrayView<const TCHAR*> B)
-		{
-			if (A.Num() != B.Num())
-			{
-				return false;
-			}
-			ScratchA.Reset(A.Num());
-			ScratchB.Reset(B.Num());
-			for (const FString& AStr : A)
-			{
-				ScratchA.Add(*AStr);
-			}
-			for (const TCHAR* BStr : B)
-			{
-				ScratchB.Add(BStr);
-			}
-			Algo::Sort(ScratchA, [](const TCHAR* StrA, const TCHAR* StrB)
-				{
-					return FCString::Stricmp(StrA, StrB) < 0;
-				});
-			Algo::Sort(ScratchB, [](const TCHAR* StrA, const TCHAR* StrB)
-				{
-					return FCString::Stricmp(StrA, StrB) < 0;
-				});
-			for (int32 Index = 0; Index < A.Num(); ++Index)
-			{
-				if (FCString::Stricmp(ScratchA[Index], ScratchB[Index]) != 0)
-				{
-					return false;
-				}
-			}
-			return true;
-		};
-
-	for (int32 PathType = 0; PathType < NumPathTypes; ++PathType)
-	{
-		int32 NumNonPaths = NonPathsByTypeAndIndex[PathType].Num();
-		for (int32 Permutation = 0; Permutation < NumPermutations; ++Permutation)
-		{
-			TDirectoryTree<int32> Tree;
-
-			// Add all the Paths in the given order and make Contains assertions after each addition
-			for (int32 AddPathPermutationIndex = 0; AddPathPermutationIndex < NumPaths; ++AddPathPermutationIndex)
-			{
-				int32 AddPathIndex = Permutations[Permutation][AddPathPermutationIndex];
-				FStringView AddPath = PathsByTypeAndIndex[PathType][AddPathIndex];
-
-				// Add the path
-				Tree.FindOrAdd(AddPath) = ValueByIndex[AddPathIndex];
-
-				if (Tree.Num() != AddPathPermutationIndex+1)
-				{
-					AddError(GetDirTreeTestName(TEXTVIEW("TreeNum has expected value"), PathType, Permutation,
-						AddPathPermutationIndex, 0));
-				}
-
-				// Assert all paths up to and including this one are included
-				for (int32 OtherPermutationIndex = 0; OtherPermutationIndex <= AddPathPermutationIndex;
-					++OtherPermutationIndex)
-				{
-					int32 OtherIndex = Permutations[Permutation][OtherPermutationIndex];
-					FStringView OtherPath = PathsByTypeAndIndex[PathType][OtherIndex];
-					int32 OtherValue = ValueByIndex[OtherIndex];
-					int32* ExistingValue = Tree.Find(OtherPath);
-					if (!ExistingValue)
-					{
-						AddError(GetDirTreeTestName(TEXTVIEW("OtherInList"), PathType, Permutation,
-							AddPathPermutationIndex, OtherPermutationIndex));
-					}
-					else
-					{
-						if (*ExistingValue != OtherValue)
-						{
-							AddError(GetDirTreeTestName(TEXTVIEW("OtherInListMatchesValue"), PathType, Permutation,
-								AddPathPermutationIndex, OtherPermutationIndex));
-						}
-						if (!Tree.ContainsPathOrParent(OtherPath))
-						{
-							AddError(GetDirTreeTestName(TEXTVIEW("ContainsPathOrParentOtherInList"), PathType, Permutation,
-								AddPathPermutationIndex, OtherPermutationIndex));
-						}
-						else
-						{
-							FString ClosestPath;
-							if (!Tree.TryFindClosestPath(OtherPath, ClosestPath))
-							{
-								AddError(GetDirTreeTestName(TEXTVIEW("TryFindClosestPathOtherInListSucceeds"), PathType, Permutation,
-									AddPathPermutationIndex, OtherPermutationIndex));
-							}
-							else if (ClosestPath != OtherPath)
-							{
-								AddError(GetDirTreeTestName(TEXTVIEW("TryFindClosestPathOtherInListMatches"), PathType, Permutation,
-									AddPathPermutationIndex, OtherPermutationIndex));
-							}
-						}
-					}
-				}
-
-				// Assert all paths not yet added are not included
-				for (int32 OtherPermutationIndex = AddPathPermutationIndex + 1; OtherPermutationIndex < NumPaths;
-					++OtherPermutationIndex)
-				{
-					int32 OtherIndex = Permutations[Permutation][OtherPermutationIndex];
-					FStringView OtherPath = PathsByTypeAndIndex[PathType][OtherIndex];
-					if (Tree.Contains(OtherPath))
-					{
-						AddError(GetDirTreeTestName(TEXTVIEW("OtherNotInList"), PathType, Permutation,
-							AddPathPermutationIndex, OtherPermutationIndex));
-					}
-					// Call ContainsPathOrParent to test whether it crashes, but not yet implemented that we verify 
-					// what its return value should be.
-					(void)Tree.ContainsPathOrParent(OtherPath);
-				}
-
-				// Assert all non paths are not included
-				for (int32 NonPathIndex = 0; NonPathIndex < NumNonPaths; ++NonPathIndex)
-				{
-					FStringView NonPath = NonPathsByTypeAndIndex[PathType][NonPathIndex];
-					if (Tree.Contains(NonPath))
-					{
-						AddError(GetDirTreeTestName(TEXTVIEW("NonPathNotInList"), PathType, Permutation,
-							AddPathPermutationIndex, NonPathIndex));
-					}
-					// Call ContainsPathOrParent to test whether it crashes, but not yet implemented that we verify 
-					// what its return value should be.
-					(void)Tree.ContainsPathOrParent(NonPath);
-				}
-			}
-
-			// Verify that the SubPath is present
-			FString ExistingSubParentPath;
-			int32* ExistingSubParentValue;
-			if (!Tree.TryFindClosestPath(Path0SubPath[PathType], ExistingSubParentPath, &ExistingSubParentValue))
-			{
-				AddError(GetDirTreeTestName(TEXTVIEW("SubPathInTree"), PathType, Permutation, 0, 0));
-			}
-			else if (ExistingSubParentPath != PathsByTypeAndIndex[PathType][0] || *ExistingSubParentValue != ValueByIndex[0])
-			{
-				AddError(GetDirTreeTestName(TEXTVIEW("SubPathInTreeMatches"), PathType, Permutation, 0, 0));
-			}
-
-			// Remove all the Paths (in specified order) and make Contains assertions after each removal
-			// Currently we only test removal in FIFO order; bugs that are specific to a removal-order
-			// should be dependent only on the final added state and should not be dependent on the earlier add-order.
-			for (int32 RemovePathPermutationIndex = 0; RemovePathPermutationIndex < NumPaths;
-				++RemovePathPermutationIndex)
-			{
-				int32 RemovePathIndex = Permutations[Permutation][RemovePathPermutationIndex];
-				FStringView RemovePath = PathsByTypeAndIndex[PathType][RemovePathIndex];
-
-				// Remove the path
-				bool bExisted;
-				Tree.Remove(RemovePath, &bExisted);
-				if (!bExisted)
-				{
-					AddError(GetDirTreeTestName(TEXTVIEW("RemoveFoundSomethingToRemove"), PathType, Permutation,
-						RemovePathPermutationIndex, 0));
-				}
-				if (Tree.Num() != NumPaths - (RemovePathPermutationIndex+1))
-				{
-					AddError(GetDirTreeTestName(TEXTVIEW("TreeNum has expected value"), PathType, Permutation,
-						RemovePathPermutationIndex, 0));
-				}
-
-
-				// Assert all paths not yet removed are still included
-				for (int32 OtherPermutationIndex = RemovePathPermutationIndex+1; OtherPermutationIndex < NumPaths;
-					++OtherPermutationIndex)
-				{
-					int32 OtherIndex = Permutations[Permutation][OtherPermutationIndex];
-					FStringView OtherPath = PathsByTypeAndIndex[PathType][OtherIndex];
-					int32 OtherValue = ValueByIndex[OtherIndex];
-					int32* ExistingValue = Tree.Find(OtherPath);
-					if (!ExistingValue)
-					{
-						AddError(GetDirTreeTestName(TEXTVIEW("OtherInListAfterRemoval"), PathType, Permutation,
-							RemovePathPermutationIndex, OtherPermutationIndex));
-					}
-					else
-					{
-						if (*ExistingValue != OtherValue)
-						{
-							AddError(GetDirTreeTestName(TEXTVIEW("OtherInListAfterRemovalMatches"), PathType, Permutation,
-								RemovePathPermutationIndex, OtherPermutationIndex));
-						}
-						if (!Tree.ContainsPathOrParent(OtherPath))
-						{
-							AddError(GetDirTreeTestName(TEXTVIEW("OtherContainsPathOrParentAfterRemoval"), PathType, Permutation,
-								RemovePathPermutationIndex, OtherPermutationIndex));
-						}
-					}
-				}
-
-				// Assert all paths up to and including this one have been removed
-				for (int32 OtherPermutationIndex = 0; OtherPermutationIndex <= RemovePathPermutationIndex;
-					++OtherPermutationIndex)
-				{
-					int32 OtherIndex = Permutations[Permutation][OtherPermutationIndex];
-					FStringView OtherPath = PathsByTypeAndIndex[PathType][OtherIndex];
-					if (Tree.Contains(OtherPath))
-					{
-						AddError(GetDirTreeTestName(TEXT("OtherNotInListAfterRemoval"), PathType, Permutation,
-							RemovePathPermutationIndex, OtherPermutationIndex));
-					}
-					// Call ContainsPathOrParent to test whether it crashes, but not yet implemented that we verify 
-					// what its return value should be.
-					(void)Tree.ContainsPathOrParent(OtherPath);
-				}
-
-				// Assert all non paths are not still included
-				for (int32 NonPathIndex = 0; NonPathIndex < NumNonPaths; ++NonPathIndex)
-				{
-					FStringView NonPath = NonPathsByTypeAndIndex[PathType][NonPathIndex];
-					if (Tree.Contains(NonPath))
-					{
-						AddError(GetDirTreeTestName(TEXTVIEW("NonPathNotInListAfterRemoval"), PathType, Permutation,
-							RemovePathPermutationIndex, NonPathIndex));
-					}
-					// Call ContainsPathOrParent to test whether it crashes, but not yet implemented that we verify 
-					// what its return value should be.
-					(void)Tree.ContainsPathOrParent(NonPath);
-				}
-			}
-			if (!Tree.IsEmpty())
-			{
-				AddError(GetDirTreeTestName(TEXTVIEW("TreeEmptyAfterRemoval"), PathType, Permutation, 0, 0));
-			}
-		}
-	}
-
-	// Testing some pathtype-independent scenarios
-	{
-		TDirectoryTree<int32> Tree;
-		Tree.FindOrAdd(TEXTVIEW("/Root/Path1")) = 1;
-		Tree.FindOrAdd(TEXTVIEW("/Root/Path2")) = 2;
-		int32* FoundRoot = Tree.FindClosestValue(TEXTVIEW("/Root"));
-		int32* FoundPath1 = Tree.FindClosestValue(TEXTVIEW("/Root/Path1"));
-		int32* FoundPath1Sub = Tree.FindClosestValue(TEXTVIEW("/Root/Path1/Sub"));
-		int32* FoundPath2 = Tree.FindClosestValue(TEXTVIEW("/Root/Path2"));
-		int32* FoundPath2Sub = Tree.FindClosestValue(TEXTVIEW("/Root/Path2/Sub"));
-		TestTrue(TEXT("TwoPaths Root does not exist"), FoundRoot == nullptr);
-		TestTrue(TEXT("TwoPaths Path1 Value matches"), FoundPath1 && *FoundPath1 == 1);
-		TestTrue(TEXT("TwoPaths Path1Sub Value matches"), FoundPath1Sub &&* FoundPath1Sub == 1);
-		TestTrue(TEXT("TwoPaths Path2 Value matches"), FoundPath2 && *FoundPath2 == 2);
-		TestTrue(TEXT("TwoPaths Path2Sub Value matches"), FoundPath2Sub && *FoundPath2Sub == 2);
-	}
-	{
-		TDirectoryTree<int32> Tree;
-		Tree.FindOrAdd(TEXTVIEW("/Root/Path1/A/B/C")) = 1;
-		Tree.FindOrAdd(TEXTVIEW("/Root/Path2/A/B/C")) = 2;
-		int32* FoundRoot = Tree.FindClosestValue(TEXTVIEW("/Root"));
-		int32* FoundPath1 = Tree.FindClosestValue(TEXTVIEW("/Root/Path1/A/B/C"));
-		int32* FoundPath1Sub = Tree.FindClosestValue(TEXTVIEW("/Root/Path1/A/B/C/Sub"));
-		int32* FoundPath1Parent = Tree.FindClosestValue(TEXTVIEW("/Root/Path1/A"));
-		int32* FoundPath2 = Tree.FindClosestValue(TEXTVIEW("/Root/Path2/A/B/C"));
-		int32* FoundPath2Sub = Tree.FindClosestValue(TEXTVIEW("/Root/Path2/A/B/C/Sub"));
-		int32* FoundPath2Parent = Tree.FindClosestValue(TEXTVIEW("/Root/Path2/A"));
-		TestTrue(TEXT("TwoPathsLong Root does not exist"), FoundRoot == nullptr);
-		TestTrue(TEXT("TwoPathsLong Path1 Value matches"), FoundPath1 && *FoundPath1 == 1);
-		TestTrue(TEXT("TwoPathsLong Path1Sub Value matches"), FoundPath1Sub && *FoundPath1Sub == 1);
-		TestTrue(TEXT("TwoPathsLong Path1 Parent does not exist"), FoundPath1Parent == nullptr);
-		TestTrue(TEXT("TwoPathsLong Path2 Value matches"), FoundPath2 && *FoundPath2 == 2);
-		TestTrue(TEXT("TwoPathsLong Path2Sub Value matches"), FoundPath2Sub && *FoundPath2Sub == 2);
-		TestTrue(TEXT("TwoPathsLong Path2 Parent does not exist"), FoundPath2Parent == nullptr);
-	}
-
-	struct FMoveConstructOnly
-	{
-		FMoveConstructOnly() { Value = 437; }
-		FMoveConstructOnly(FMoveConstructOnly&& Other) { Value = Other.Value; Other.Value = -1; }
-		FMoveConstructOnly(const FMoveConstructOnly& Other) = delete;
-		FMoveConstructOnly& operator=(FMoveConstructOnly&& Other) = delete;
-		FMoveConstructOnly& operator=(const FMoveConstructOnly& Other) = delete;
-		
-		int32 Value;
-	};
-
-	{
-		TDirectoryTree<FMoveConstructOnly> Tree;
-		Tree.FindOrAdd(TEXTVIEW("/Root/PathM")).Value = 1;
-		Tree.FindOrAdd(TEXTVIEW("/Root/PathP")).Value = 2;
-		Tree.FindOrAdd(TEXTVIEW("/Root/PathA")).Value = 3;
-		Tree.FindOrAdd(TEXTVIEW("/Root/PathZ"));
-
-		FMoveConstructOnly* Value = Tree.Find(TEXTVIEW("/Root/PathA"));
-		TestTrue(TEXT("MoveConstructOnlyValueA correct"), Value && Value->Value == 3);
-		Value = Tree.Find(TEXTVIEW("/Root/PathM"));
-		TestTrue(TEXT("MoveConstructOnlyValueM correct"), Value && Value->Value == 1);
-		Value = Tree.Find(TEXTVIEW("/Root/PathP"));
-		TestTrue(TEXT("MoveConstructOnlyValueP correct"), Value && Value->Value == 2);
-		Value = Tree.Find(TEXTVIEW("/Root/PathZ"));
-		TestTrue(TEXT("MoveConstructOnlyValueZ correct"), Value && Value->Value == 437);
-	}
-
-	// Handling special case of drive specifiers without a path
-	{
-		TStringBuilder<16> FoundPath;
-		TArray<FString> ChildNames;
-		int* FoundValue = nullptr;
-		auto Reset = [&FoundPath, &ChildNames, &FoundValue]()
-			{
-				FoundPath.Reset();
-				ChildNames.Reset();
-				FoundValue = nullptr;
-			};
-
-		{
-			TDirectoryTree<int32> Tree;
-			Tree.FindOrAdd(TEXTVIEW("D:")) = 1;
-
-			Reset();
-			TestTrue(TEXT("DriveSpecifier: Before PathSep: Without PathSep: Tree.Contains"), Tree.Contains(TEXTVIEW("D:")));
-			TestTrue(TEXT("DriveSpecifier: Before PathSep: Without PathSep: Tree.Find"), Tree.Find(TEXTVIEW("D:")) != nullptr);
-			TestTrue(TEXT("DriveSpecifier: Before PathSep: Without PathSep: Tree.ContainsPathOrParent"), Tree.ContainsPathOrParent(TEXTVIEW("D:")));
-			TestTrue(TEXT("DriveSpecifier: Before PathSep: Without PathSep: Tree.FindClosestValue"), Tree.FindClosestValue(TEXTVIEW("D:")) != nullptr);
-			TestTrue(TEXT("DriveSpecifier: Before PathSep: Without PathSep: Tree.TryFindClosestPath"),
-				Tree.TryFindClosestPath(TEXTVIEW("D:"), FoundPath, &FoundValue) && FoundPath.Len() > 0 && FoundValue != nullptr);
-			TestTrue(TEXT("DriveSpecifier: Before PathSep: Without PathSep: Tree.TryGetChildren"), Tree.TryGetChildren(TEXTVIEW("D:"), ChildNames));
-
-			Tree.FindOrAdd(TEXTVIEW("D:/root")) = 1;
-
-			Reset();
-			TestTrue(TEXT("DriveSpecifier: After PathSep('/'): Without PathSep: Tree.Contains"), Tree.Contains(TEXTVIEW("D:")));
-			TestTrue(TEXT("DriveSpecifier: After PathSep('/'): Without PathSep: Tree.Find"), Tree.Find(TEXTVIEW("D:")) != nullptr);
-			TestTrue(TEXT("DriveSpecifier: After PathSep('/'): Without PathSep: Tree.ContainsPathOrParent"), Tree.ContainsPathOrParent(TEXTVIEW("D:")));
-			TestTrue(TEXT("DriveSpecifier: After PathSep('/'): Without PathSep: Tree.FindClosestValue"), Tree.FindClosestValue(TEXTVIEW("D:")) != nullptr);
-			TestTrue(TEXT("DriveSpecifier: After PathSep('/'): Without PathSep: Tree.TryFindClosestPath"),
-				Tree.TryFindClosestPath(TEXTVIEW("D:"), FoundPath, &FoundValue) && FoundPath.Len() > 0 && FoundValue != nullptr);
-			TestTrue(TEXT("DriveSpecifier: After PathSep('/'): Without PathSep: Tree.TryGetChildren"), Tree.TryGetChildren(TEXTVIEW("D:"), ChildNames));
-
-			Reset();
-			TestTrue(TEXT("DriveSpecifier: After PathSep('/'): With PathSep: Tree.Contains"), Tree.Contains(TEXTVIEW("D:/")));
-			TestTrue(TEXT("DriveSpecifier: After PathSep('/'): With PathSep: Tree.Find"), Tree.Find(TEXTVIEW("D:/")) != nullptr);
-			TestTrue(TEXT("DriveSpecifier: After PathSep('/'): With PathSep: Tree.ContainsPathOrParent"), Tree.ContainsPathOrParent(TEXTVIEW("D:/")));
-			TestTrue(TEXT("DriveSpecifier: After PathSep('/'): With PathSep: Tree.FindClosestValue"), Tree.FindClosestValue(TEXTVIEW("D:/")) != nullptr);
-			TestTrue(TEXT("DriveSpecifier: After PathSep('/'): With PathSep: Tree.TryFindClosestPath"),
-				Tree.TryFindClosestPath(TEXTVIEW("D:/"), FoundPath, &FoundValue) && FoundPath.Len() > 0 && FoundValue != nullptr);
-			TestTrue(TEXT("DriveSpecifier: After PathSep('/'): With PathSep: Tree.TryGetChildren"), Tree.TryGetChildren(TEXTVIEW("D:/"), ChildNames));
-		}
-		{
-			TDirectoryTree<int32> Tree;
-			Tree.FindOrAdd(TEXTVIEW("D:")) = 1;
-			Tree.FindOrAdd(TEXTVIEW("D:\\root")) = 1;
-
-			Reset();
-			TestTrue(TEXT("DriveSpecifier: After PathSep('\\'): Without PathSep: Tree.Contains"), Tree.Contains(TEXTVIEW("D:")));
-			TestTrue(TEXT("DriveSpecifier: After PathSep('\\'): Without PathSep: Tree.Find"), Tree.Find(TEXTVIEW("D:")) != nullptr);
-			TestTrue(TEXT("DriveSpecifier: After PathSep('\\'): Without PathSep: Tree.ContainsPathOrParent"), Tree.ContainsPathOrParent(TEXTVIEW("D:")));
-			TestTrue(TEXT("DriveSpecifier: After PathSep('\\'): Without PathSep: Tree.FindClosestValue"), Tree.FindClosestValue(TEXTVIEW("D:")) != nullptr);
-			TestTrue(TEXT("DriveSpecifier: After PathSep('\\'): Without PathSep: Tree.TryFindClosestPath"),
-				Tree.TryFindClosestPath(TEXTVIEW("D:"), FoundPath, &FoundValue) && FoundPath.Len() > 0 && FoundValue != nullptr);
-			TestTrue(TEXT("DriveSpecifier: After PathSep('\\'): Without PathSep: Tree.TryGetChildren"), Tree.TryGetChildren(TEXTVIEW("D:"), ChildNames));
-
-			Reset();
-			TestTrue(TEXT("DriveSpecifier: After PathSep('\\'): With PathSep: Tree.Contains"), Tree.Contains(TEXTVIEW("D:\\")));
-			TestTrue(TEXT("DriveSpecifier: After PathSep('\\'): With PathSep: Tree.Find"), Tree.Find(TEXTVIEW("D:\\")) != nullptr);
-			TestTrue(TEXT("DriveSpecifier: After PathSep('\\'): With PathSep: Tree.ContainsPathOrParent"), Tree.ContainsPathOrParent(TEXTVIEW("D:\\")));
-			TestTrue(TEXT("DriveSpecifier: After PathSep('\\'): With PathSep: Tree.FindClosestValue"), Tree.FindClosestValue(TEXTVIEW("D:\\")) != nullptr);
-			TestTrue(TEXT("DriveSpecifier: After PathSep('\\'): With PathSep: Tree.TryFindClosestPath"),
-				Tree.TryFindClosestPath(TEXTVIEW("D:\\"), FoundPath, &FoundValue) && FoundPath.Len() > 0 && FoundValue != nullptr);
-			TestTrue(TEXT("DriveSpecifier: After PathSep('\\'): With PathSep: Tree.TryGetChildren"), Tree.TryGetChildren(TEXTVIEW("D:\\"), ChildNames));
-		}
-		{
-			TDirectoryTree<int32> Tree;
-			Tree.FindOrAdd(TEXTVIEW("D:root")) = 1;
-
-			Reset();
-			TestTrue(TEXT("DriveSpecifierLong: Before PathSep: Without PathSep: Tree.Contains"), Tree.Contains(TEXTVIEW("D:root")));
-			TestTrue(TEXT("DriveSpecifierLong: Before PathSep: Without PathSep: Tree.Find"), Tree.Find(TEXTVIEW("D:root")) != nullptr);
-			TestTrue(TEXT("DriveSpecifierLong: Before PathSep: Without PathSep: Tree.ContainsPathOrParent"), Tree.ContainsPathOrParent(TEXTVIEW("D:root")));
-			TestTrue(TEXT("DriveSpecifierLong: Before PathSep: Without PathSep: Tree.FindClosestValue"), Tree.FindClosestValue(TEXTVIEW("D:root")) != nullptr);
-			TestTrue(TEXT("DriveSpecifierLong: Before PathSep: Without PathSep: Tree.TryFindClosestPath"),
-				Tree.TryFindClosestPath(TEXTVIEW("D:root"), FoundPath, &FoundValue) && FoundPath.Len() > 0 && FoundValue != nullptr);
-			TestTrue(TEXT("DriveSpecifierLong: Before PathSep: Without PathSep: Tree.TryGetChildren"), Tree.TryGetChildren(TEXTVIEW("D:root"), ChildNames));
-
-			Tree.FindOrAdd(TEXTVIEW("D:\\root\\path")) = 1;
-
-			Reset();
-			TestTrue(TEXT("DriveSpecifierLong: After PathSep('\\'): Without PathSep: Tree.Contains"), Tree.Contains(TEXTVIEW("D:root")));
-			TestTrue(TEXT("DriveSpecifierLong: After PathSep('\\'): Without PathSep: Tree.Find"), Tree.Find(TEXTVIEW("D:root")) != nullptr);
-			TestTrue(TEXT("DriveSpecifierLong: After PathSep('\\'): Without PathSep: Tree.ContainsPathOrParent"), Tree.ContainsPathOrParent(TEXTVIEW("D:root")));
-			TestTrue(TEXT("DriveSpecifierLong: After PathSep('\\'): Without PathSep: Tree.FindClosestValue"), Tree.FindClosestValue(TEXTVIEW("D:root")) != nullptr);
-			TestTrue(TEXT("DriveSpecifierLong: After PathSep('\\'): Without PathSep: Tree.TryFindClosestPath"),
-				Tree.TryFindClosestPath(TEXTVIEW("D:root"), FoundPath, &FoundValue) && FoundPath.Len() > 0 && FoundValue != nullptr);
-			TestTrue(TEXT("DriveSpecifierLong: After PathSep('\\'): Without PathSep: Tree.TryGetChildren"), Tree.TryGetChildren(TEXTVIEW("D:root"), ChildNames));
-
-			Reset();
-			TestTrue(TEXT("DriveSpecifierLong: After PathSep('\\'): With PathSep: Tree.Contains"), Tree.Contains(TEXTVIEW("D:\\root")));
-			TestTrue(TEXT("DriveSpecifierLong: After PathSep('\\'): With PathSep: Tree.Find"), Tree.Find(TEXTVIEW("D:\\root")) != nullptr);
-			TestTrue(TEXT("DriveSpecifierLong: After PathSep('\\'): With PathSep: Tree.ContainsPathOrParent"), Tree.ContainsPathOrParent(TEXTVIEW("D:\\root")));
-			TestTrue(TEXT("DriveSpecifierLong: After PathSep('\\'): With PathSep: Tree.FindClosestValue"), Tree.FindClosestValue(TEXTVIEW("D:\\root")) != nullptr);
-			TestTrue(TEXT("DriveSpecifierLong: After PathSep('\\'): With PathSep: Tree.TryFindClosestPath"),
-				Tree.TryFindClosestPath(TEXTVIEW("D:root"), FoundPath, &FoundValue) && FoundPath.Len() > 0 && FoundValue != nullptr);
-			TestTrue(TEXT("DriveSpecifierLong: After PathSep('\\'): With PathSep: Tree.TryGetChildren"), Tree.TryGetChildren(TEXTVIEW("D:\\root"), ChildNames));
-		}
-	}
-
-	// Testing accessors
-	{
-		// GetChildren
-		TDirectoryTree<FMoveConstructOnly> Tree;
-		bool bExists;
-		TArray<FString> Children;
-
-		Children.Reset();
-		bExists = Tree.TryGetChildren(TEXTVIEW(""), Children,
-			EDirectoryTreeGetFlags::None);
-		TestTrue(TEXT("GetChildrenEmpty, Root, !ImpliedParent"),
-			bExists == false && Children.IsEmpty());
-
-		Children.Reset();
-		bExists = Tree.TryGetChildren(TEXTVIEW(""), Children,
-			EDirectoryTreeGetFlags::ImpliedParent |
-			EDirectoryTreeGetFlags::ImpliedChildren |
-			EDirectoryTreeGetFlags::Recursive);
-		TestTrue(TEXT("GetChildrenEmpty, Root, ImpliedParent"),
-			bExists == true && Children.IsEmpty());
-
-		Children.Reset();
-		bExists = Tree.TryGetChildren(TEXTVIEW("/SomePath"), Children,
-			EDirectoryTreeGetFlags::ImpliedParent |
-			EDirectoryTreeGetFlags::ImpliedChildren |
-			EDirectoryTreeGetFlags::Recursive);
-		TestTrue(TEXT("GetChildrenEmpty, Non-root"),
-			bExists == false && Children.IsEmpty());
-
-		Tree.FindOrAdd(TEXTVIEW("")).Value = 1;
-		Children.Reset();
-		bExists = Tree.TryGetChildren(TEXTVIEW(""), Children,
-			EDirectoryTreeGetFlags::None);
-		TestTrue(TEXT("GetChildrenRoot, !ImpliedParent, !ImpliedChildren"),
-			bExists == true && Children.IsEmpty());
-
-		Tree.Empty();
-		Tree.FindOrAdd(TEXTVIEW("/")).Value = 1;
-
-		bExists = Tree.TryGetChildren(TEXTVIEW(""), Children,
-			EDirectoryTreeGetFlags::ImpliedChildren |
-			EDirectoryTreeGetFlags::Recursive);
-		TestTrue(TEXT("GetChildrenRoot, !ImpliedParent, ImpliedChildren"),
-			bExists == false && Children.IsEmpty());
-
-		Children.Reset();
-		bExists = Tree.TryGetChildren(TEXTVIEW(""), Children,
-			EDirectoryTreeGetFlags::ImpliedParent |
-			EDirectoryTreeGetFlags::ImpliedChildren |
-			EDirectoryTreeGetFlags::Recursive);
-		TestTrue(TEXT("GetChildrenRoot, ImpliedParent, ImpliedChildren"),
-			bExists == true && UnorderedEquals(Children, { TEXT("/") }));
-
-		bExists = Tree.TryGetChildren(TEXTVIEW(""), Children,
-			EDirectoryTreeGetFlags::ImpliedParent |
-			EDirectoryTreeGetFlags::ImpliedChildren |
-			EDirectoryTreeGetFlags::Recursive);
-		TestTrue(TEXT("GetChildren appends to the outdir rather than resetting"),
-			bExists == true && Children.Num() == 2);
-
-		Children.Reset();
-		bExists = Tree.TryGetChildren(TEXTVIEW(""), Children,
-			EDirectoryTreeGetFlags::ImpliedParent |
-			EDirectoryTreeGetFlags::Recursive);
-		TestTrue(TEXT("GetChildrenRoot, ImpliedParent, !ImpliedChildren"),
-			bExists == true && UnorderedEquals(Children, { TEXT("/") }));
-
-		Tree.Empty();
-		Tree.FindOrAdd(TEXTVIEW("/Root/Child")).Value = 1;
-
-		Children.Reset();
-		bExists = Tree.TryGetChildren(TEXTVIEW(""), Children,
-			EDirectoryTreeGetFlags::ImpliedParent |
-			EDirectoryTreeGetFlags::ImpliedChildren);
-		TestTrue(TEXT("GetChildrenRootImpliedChild, ImpliedParent, ImpliedChildren, !Recursive"),
-			bExists == true && UnorderedEquals(Children, { TEXT("/") }));
-		Children.Reset();
-		bExists = Tree.TryGetChildren(TEXTVIEW(""), Children,
-			EDirectoryTreeGetFlags::ImpliedParent);
-		TestTrue(TEXT("GetChildrenRootImpliedChild, ImpliedParent, !ImpliedChildren, !Recursive"),
-			bExists == true && UnorderedEquals(Children, { TEXT("/Root/Child") }));
-		Children.Reset();
-		bExists = Tree.TryGetChildren(TEXTVIEW(""), Children,
-			EDirectoryTreeGetFlags::ImpliedParent |
-			EDirectoryTreeGetFlags::ImpliedChildren |
-			EDirectoryTreeGetFlags::Recursive);
-		TestTrue(TEXT("GetChildrenRootImpliedChild, ImpliedParent, ImpliedChildren, Recursive"),
-			bExists == true && UnorderedEquals(Children, { TEXT("/"), TEXT("/Root"), TEXT("/Root/Child") }));
-		Children.Reset();
-		bExists = Tree.TryGetChildren(TEXTVIEW(""), Children,
-			EDirectoryTreeGetFlags::ImpliedParent |
-			EDirectoryTreeGetFlags::Recursive);
-		TestTrue(TEXT("GetChildrenRootImpliedChild, ImpliedParent, !ImpliedChildren, Recursive"),
-			bExists == true && UnorderedEquals(Children, { TEXT("/Root/Child") }));
-
-		Tree.FindOrAdd(TEXTVIEW("/Root/Child2")).Value = 1;
-
-		Tree.Empty();
-		Tree.FindOrAdd(TEXTVIEW("/Stem/A_OtherChild")).Value = 1;
-		Tree.FindOrAdd(TEXTVIEW("/Stem/B_ImpliedRoot/AddedChild")).Value = 1;
-		Tree.FindOrAdd(TEXTVIEW("/Stem/B_ImpliedRoot/AddedChild/Child")).Value = 1;
-		Tree.FindOrAdd(TEXTVIEW("/Stem/B_ImpliedRoot/ImpliedChild/AddedChild")).Value = 1;
-		Tree.FindOrAdd(TEXTVIEW("/Stem/B_ImpliedRoot/ImpliedChild/AddedChild/AddedChild")).Value = 1;
-		Tree.FindOrAdd(TEXTVIEW("/Stem/B_ImpliedRoot/ImpliedChild/AddedChild/ImpliedChild/AddedChild")).Value = 1;
-		Tree.FindOrAdd(TEXTVIEW("/Stem/C_MiddleRoot/MiddlePath/ImpliedChild/AddedChild")).Value = 1;
-		Tree.FindOrAdd(TEXTVIEW("/Stem/C_MiddleRoot/MiddlePath/ImpliedChild/AddedChild/Child")).Value = 1;
-		Tree.FindOrAdd(TEXTVIEW("/Stem/D_MiddleRoot/MiddlePath/AddedChild")).Value = 1;
-		Tree.FindOrAdd(TEXTVIEW("/Stem/D_MiddleRoot/MiddlePath/AddedChild/Child")).Value = 1;
-		Tree.FindOrAdd(TEXTVIEW("/Stem/E_AddedRoot")).Value = 1;
-		Tree.FindOrAdd(TEXTVIEW("/Stem/E_AddedRoot/AddedChild")).Value = 1;
-		Tree.FindOrAdd(TEXTVIEW("/Stem/E_AddedRoot/ImpliedChild/AddedChild")).Value = 1;
-		Tree.FindOrAdd(TEXTVIEW("/Stem/E_AddedRoot/ImpliedChild/AddedChild/ImpliedChild/AddedChild")).Value = 1;
-		Tree.FindOrAdd(TEXTVIEW("/Stem/F_AddedRoot")).Value = 1;
-		Tree.FindOrAdd(TEXTVIEW("/Stem/F_AddedRoot/ImpliedChild/AddedChild")).Value = 1;
-
-
-		// Case: Requested path is an implied path that is a stored child in the tree.
-		Children.Reset();
-		bExists = Tree.TryGetChildren(TEXTVIEW("/Stem/B_ImpliedRoot"), Children,
-			EDirectoryTreeGetFlags::None);
-		TestTrue(TEXT("GetChildrenComplexA B_ImpliedRoot, !ImpliedParent, !ImpliedChildren, !Recursive"),
-			bExists == false && Children.IsEmpty());
-		Children.Reset();
-		bExists = Tree.TryGetChildren(TEXTVIEW("/Stem/B_ImpliedRoot"), Children,
-			EDirectoryTreeGetFlags::Recursive);
-		TestTrue(TEXT("GetChildrenComplexA B_ImpliedRoot, !ImpliedParent, !ImpliedChildren, Recursive"),
-			bExists == false && Children.IsEmpty());
-		Children.Reset();
-		bExists = Tree.TryGetChildren(TEXTVIEW("/Stem/B_ImpliedRoot"), Children,
-			EDirectoryTreeGetFlags::ImpliedChildren);
-		TestTrue(TEXT("GetChildrenComplexA B_ImpliedRoot, !ImpliedParent, ImpliedChildren, !Recursive"),
-			bExists == false && Children.IsEmpty());
-		Children.Reset();
-		bExists = Tree.TryGetChildren(TEXTVIEW("/Stem/B_ImpliedRoot"), Children,
-			EDirectoryTreeGetFlags::ImpliedChildren |
-			EDirectoryTreeGetFlags::Recursive);
-		TestTrue(TEXT("GetChildrenComplexA B_ImpliedRoot, !ImpliedParent, ImpliedChildren, Recursive"),
-			bExists == false && Children.IsEmpty());
-		Children.Reset();
-		bExists = Tree.TryGetChildren(TEXTVIEW("/Stem/B_ImpliedRoot"), Children,
-			EDirectoryTreeGetFlags::ImpliedParent);
-		TestTrue(TEXT("GetChildrenComplexA B_ImpliedRoot, ImpliedParent, !ImpliedChildren, !Recursive"),
-			bExists == true && UnorderedEquals(Children, { TEXT("AddedChild"), TEXT("ImpliedChild/AddedChild") }));
-		Children.Reset();
-		bExists = Tree.TryGetChildren(TEXTVIEW("/Stem/B_ImpliedRoot"), Children,
-			EDirectoryTreeGetFlags::ImpliedParent |
-			EDirectoryTreeGetFlags::Recursive);
-		TestTrue(TEXT("GetChildrenComplexA B_ImpliedRoot, ImpliedParent, !ImpliedChildren, Recursive"),
-			bExists == true && UnorderedEquals(Children, { TEXT("AddedChild"), TEXT("AddedChild/Child"),
-				TEXT("ImpliedChild/AddedChild"), TEXT("ImpliedChild/AddedChild/AddedChild"),
-				TEXT("ImpliedChild/AddedChild/ImpliedChild/AddedChild") }));
-		Children.Reset();
-		bExists = Tree.TryGetChildren(TEXTVIEW("/Stem/B_ImpliedRoot"), Children,
-			EDirectoryTreeGetFlags::ImpliedParent |
-			EDirectoryTreeGetFlags::ImpliedChildren);
-		TestTrue(TEXT("GetChildrenComplexA B_ImpliedRoot, ImpliedParent, ImpliedChildren, !Recursive"),
-			bExists == true && UnorderedEquals(Children, { TEXT("AddedChild"), TEXT("ImpliedChild") }));
-		Children.Reset();
-		bExists = Tree.TryGetChildren(TEXTVIEW("/Stem/B_ImpliedRoot"), Children,
-			EDirectoryTreeGetFlags::ImpliedParent |
-			EDirectoryTreeGetFlags::ImpliedChildren |
-			EDirectoryTreeGetFlags::Recursive);
-		TestTrue(TEXT("GetChildrenComplexA B_ImpliedRoot, ImpliedParent, ImpliedChildren, Recursive"),
-			bExists == true && UnorderedEquals(Children, { TEXT("AddedChild"), TEXT("AddedChild/Child"),
-				TEXT("ImpliedChild"), TEXT("ImpliedChild/AddedChild"), TEXT("ImpliedChild/AddedChild/AddedChild"),
-				TEXT("ImpliedChild/AddedChild/ImpliedChild"), TEXT("ImpliedChild/AddedChild/ImpliedChild/AddedChild") }));
-
-		// Case: Requested path is an implied path that is not a stored child in the tree - it is an in-between dir in a
-		// relpath - and it has an implied child
-		Children.Reset();
-		bExists = Tree.TryGetChildren(TEXTVIEW("/Stem/C_MiddleRoot/MiddlePath"), Children,
-			EDirectoryTreeGetFlags::None);
-		TestTrue(TEXT("GetChildrenComplexA C_MiddleRoot, !ImpliedParent, !ImpliedChildren, !Recursive"),
-			bExists == false && Children.IsEmpty());
-		Children.Reset();
-		bExists = Tree.TryGetChildren(TEXTVIEW("/Stem/C_MiddleRoot/MiddlePath"), Children,
-			EDirectoryTreeGetFlags::Recursive);
-		TestTrue(TEXT("GetChildrenComplexA C_MiddleRoot, !ImpliedParent, !ImpliedChildren, Recursive"),
-			bExists == false && Children.IsEmpty());
-		Children.Reset();
-		bExists = Tree.TryGetChildren(TEXTVIEW("/Stem/C_MiddleRoot/MiddlePath"), Children,
-			EDirectoryTreeGetFlags::ImpliedChildren);
-		TestTrue(TEXT("GetChildrenComplexA C_MiddleRoot, !ImpliedParent, ImpliedChildren, !Recursive"),
-			bExists == false && Children.IsEmpty());
-		Children.Reset();
-		bExists = Tree.TryGetChildren(TEXTVIEW("/Stem/C_MiddleRoot/MiddlePath"), Children,
-			EDirectoryTreeGetFlags::ImpliedChildren |
-			EDirectoryTreeGetFlags::Recursive);
-		TestTrue(TEXT("GetChildrenComplexA C_MiddleRoot, !ImpliedParent, ImpliedChildren, Recursive"),
-			bExists == false && Children.IsEmpty());
-		Children.Reset();
-		bExists = Tree.TryGetChildren(TEXTVIEW("/Stem/C_MiddleRoot/MiddlePath"), Children,
-			EDirectoryTreeGetFlags::ImpliedParent);
-		TestTrue(TEXT("GetChildrenComplexA C_MiddleRoot, ImpliedParent, !ImpliedChildren, !Recursive"),
-			bExists == true && UnorderedEquals(Children, { TEXT("ImpliedChild/AddedChild") }));
-		Children.Reset();
-		bExists = Tree.TryGetChildren(TEXTVIEW("/Stem/C_MiddleRoot/MiddlePath"), Children,
-			EDirectoryTreeGetFlags::ImpliedParent |
-			EDirectoryTreeGetFlags::Recursive);
-		TestTrue(TEXT("GetChildrenComplexA C_MiddleRoot, ImpliedParent, !ImpliedChildren, Recursive"),
-			bExists == true && UnorderedEquals(Children, { TEXT("ImpliedChild/AddedChild"),
-				TEXT("ImpliedChild/AddedChild/Child") }));
-		Children.Reset();
-		bExists = Tree.TryGetChildren(TEXTVIEW("/Stem/C_MiddleRoot/MiddlePath"), Children,
-			EDirectoryTreeGetFlags::ImpliedParent |
-			EDirectoryTreeGetFlags::ImpliedChildren);
-		TestTrue(TEXT("GetChildrenComplexA C_MiddleRoot, ImpliedParent, ImpliedChildren, !Recursive"),
-			bExists == true && UnorderedEquals(Children, { TEXT("ImpliedChild") }));
-		Children.Reset();
-		bExists = Tree.TryGetChildren(TEXTVIEW("/Stem/C_MiddleRoot/MiddlePath"), Children,
-			EDirectoryTreeGetFlags::ImpliedParent |
-			EDirectoryTreeGetFlags::ImpliedChildren |
-			EDirectoryTreeGetFlags::Recursive);
-		TestTrue(TEXT("GetChildrenComplexA C_MiddleRoot, ImpliedParent, ImpliedChildren, Recursive"),
-			bExists == true && UnorderedEquals(Children, { TEXT("ImpliedChild"), TEXT("ImpliedChild/AddedChild"),
-				TEXT("ImpliedChild/AddedChild/Child") }));
-
-		// Case: Requested path is a non-existent sibling path of an implied path that is not a stored
-		// path.
-		Children.Reset();
-		bExists = Tree.TryGetChildren(TEXTVIEW("/Stem/C_MiddleRoot/MiddlePathExceptItDoesNotExist"), Children,
-			EDirectoryTreeGetFlags::ImpliedParent |
-			EDirectoryTreeGetFlags::ImpliedChildren |
-			EDirectoryTreeGetFlags::Recursive);
-		TestTrue(TEXT("GetChildrenComplexA MiddlePathExceptItDoesNotExist, ImpliedParent, ImpliedChildren, Recursive"),
-			bExists == false && Children.IsEmpty());
-
-		// Case: Requested path is an implied path that is not a stored child in the tree - it is an in-between dir in a
-		// relpath - and it has an added child
-		Children.Reset();
-		bExists = Tree.TryGetChildren(TEXTVIEW("/Stem/D_MiddleRoot/MiddlePath"), Children,
-			EDirectoryTreeGetFlags::None);
-		TestTrue(TEXT("GetChildrenComplexA D_MiddleRoot, !ImpliedParent, !ImpliedChildren, !Recursive"),
-			bExists == false && Children.IsEmpty());
-		Children.Reset();
-		bExists = Tree.TryGetChildren(TEXTVIEW("/Stem/D_MiddleRoot/MiddlePath"), Children,
-			EDirectoryTreeGetFlags::Recursive);
-		TestTrue(TEXT("GetChildrenComplexA D_MiddleRoot, !ImpliedParent, !ImpliedChildren, Recursive"),
-			bExists == false && Children.IsEmpty());
-		Children.Reset();
-		bExists = Tree.TryGetChildren(TEXTVIEW("/Stem/D_MiddleRoot/MiddlePath"), Children,
-			EDirectoryTreeGetFlags::ImpliedChildren);
-		TestTrue(TEXT("GetChildrenComplexA D_MiddleRoot, !ImpliedParent, ImpliedChildren, !Recursive"),
-			bExists == false && Children.IsEmpty());
-		Children.Reset();
-		bExists = Tree.TryGetChildren(TEXTVIEW("/Stem/D_MiddleRoot/MiddlePath"), Children,
-			EDirectoryTreeGetFlags::ImpliedChildren |
-			EDirectoryTreeGetFlags::Recursive);
-		TestTrue(TEXT("GetChildrenComplexA D_MiddleRoot, !ImpliedParent, ImpliedChildren, Recursive"),
-			bExists == false && Children.IsEmpty());
-		Children.Reset();
-		bExists = Tree.TryGetChildren(TEXTVIEW("/Stem/D_MiddleRoot/MiddlePath"), Children,
-			EDirectoryTreeGetFlags::ImpliedParent);
-		TestTrue(TEXT("GetChildrenComplexA D_MiddleRoot, ImpliedParent, !ImpliedChildren, !Recursive"),
-			bExists == true && UnorderedEquals(Children, { TEXT("AddedChild") }));
-		Children.Reset();
-		bExists = Tree.TryGetChildren(TEXTVIEW("/Stem/D_MiddleRoot/MiddlePath"), Children,
-			EDirectoryTreeGetFlags::ImpliedParent |
-			EDirectoryTreeGetFlags::Recursive);
-		TestTrue(TEXT("GetChildrenComplexA D_MiddleRoot, ImpliedParent, !ImpliedChildren, Recursive"),
-			bExists == true && UnorderedEquals(Children, { TEXT("AddedChild"), TEXT("AddedChild/Child") }));
-		Children.Reset();
-		bExists = Tree.TryGetChildren(TEXTVIEW("/Stem/D_MiddleRoot/MiddlePath"), Children,
-			EDirectoryTreeGetFlags::ImpliedParent |
-			EDirectoryTreeGetFlags::ImpliedChildren);
-		TestTrue(TEXT("GetChildrenComplexA D_MiddleRoot, ImpliedParent, ImpliedChildren, !Recursive"),
-			bExists == true && UnorderedEquals(Children, { TEXT("AddedChild") }));
-		Children.Reset();
-		bExists = Tree.TryGetChildren(TEXTVIEW("/Stem/D_MiddleRoot/MiddlePath"), Children,
-			EDirectoryTreeGetFlags::ImpliedParent |
-			EDirectoryTreeGetFlags::ImpliedChildren |
-			EDirectoryTreeGetFlags::Recursive);
-		TestTrue(TEXT("GetChildrenComplexA D_MiddleRoot, ImpliedParent, ImpliedChildren, Recursive"),
-			bExists == true && UnorderedEquals(Children, { TEXT("AddedChild"), TEXT("AddedChild/Child") }));
-
-		// Case: Requested path is an added path and it has an added child and an implied child
-		Children.Reset();
-		bExists = Tree.TryGetChildren(TEXTVIEW("/Stem/E_AddedRoot"), Children,
-			EDirectoryTreeGetFlags::None);
-		TestTrue(TEXT("GetChildrenComplexA E_AddedRoot, !ImpliedParent, !ImpliedChildren, !Recursive"),
-			bExists == true && UnorderedEquals(Children, { TEXT("AddedChild"), TEXT("ImpliedChild/AddedChild")}));
-		Children.Reset();
-		bExists = Tree.TryGetChildren(TEXTVIEW("/Stem/E_AddedRoot"), Children,
-			EDirectoryTreeGetFlags::Recursive);
-		TestTrue(TEXT("GetChildrenComplexA E_AddedRoot, !ImpliedParent, !ImpliedChildren, Recursive"),
-			bExists == true && UnorderedEquals(Children, { TEXT("AddedChild"), TEXT("ImpliedChild/AddedChild"),
-				TEXT("ImpliedChild/AddedChild/ImpliedChild/AddedChild") }));
-		Children.Reset();
-		bExists = Tree.TryGetChildren(TEXTVIEW("/Stem/E_AddedRoot"), Children,
-			EDirectoryTreeGetFlags::ImpliedChildren);
-		TestTrue(TEXT("GetChildrenComplexA E_AddedRoot, !ImpliedParent, ImpliedChildren, !Recursive"),
-			bExists == true && UnorderedEquals(Children, { TEXT("AddedChild"), TEXT("ImpliedChild")}));
-		Children.Reset();
-		bExists = Tree.TryGetChildren(TEXTVIEW("/Stem/E_AddedRoot"), Children,
-			EDirectoryTreeGetFlags::ImpliedChildren |
-			EDirectoryTreeGetFlags::Recursive);
-		TestTrue(TEXT("GetChildrenComplexA E_MiddleRoot, !ImpliedParent, ImpliedChildren, Recursive"),
-			bExists == true && UnorderedEquals(Children, { TEXT("AddedChild"), TEXT("ImpliedChild"),
-				TEXT("ImpliedChild/AddedChild"), TEXT("ImpliedChild/AddedChild/ImpliedChild"),
-				TEXT("ImpliedChild/AddedChild/ImpliedChild/AddedChild") }));
-
-		// Case: Requested path is an added path and it has only an implied child
-		Children.Reset();
-		bExists = Tree.TryGetChildren(TEXTVIEW("/Stem/F_AddedRoot"), Children,
-			EDirectoryTreeGetFlags::None);
-		TestTrue(TEXT("GetChildrenComplexA F_AddedRoot, !ImpliedParent, !ImpliedChildren, !Recursive"),
-			bExists == true && UnorderedEquals(Children, { TEXT("ImpliedChild/AddedChild") }));
-		Children.Reset();
-		bExists = Tree.TryGetChildren(TEXTVIEW("/Stem/F_AddedRoot"), Children,
-			EDirectoryTreeGetFlags::Recursive);
-		TestTrue(TEXT("GetChildrenComplexA F_AddedRoot, !ImpliedParent, !ImpliedChildren, Recursive"),
-			bExists == true && UnorderedEquals(Children, { TEXT("ImpliedChild/AddedChild") }));
-		Children.Reset();
-		bExists = Tree.TryGetChildren(TEXTVIEW("/Stem/F_AddedRoot"), Children,
-			EDirectoryTreeGetFlags::ImpliedChildren);
-		TestTrue(TEXT("GetChildrenComplexA F_AddedRoot, !ImpliedParent, ImpliedChildren, !Recursive"),
-			bExists == true && UnorderedEquals(Children, { TEXT("ImpliedChild") }));
-		Children.Reset();
-		bExists = Tree.TryGetChildren(TEXTVIEW("/Stem/F_AddedRoot"), Children,
-			EDirectoryTreeGetFlags::ImpliedChildren |
-			EDirectoryTreeGetFlags::Recursive);
-		TestTrue(TEXT("GetChildrenComplexA F_AddedRoot, !ImpliedParent, ImpliedChildren, Recursive"),
-			bExists == true && UnorderedEquals(Children, { TEXT("ImpliedChild"),
-				TEXT("ImpliedChild/AddedChild") }));
-
-		// Case: Requesting !ImpliedChildren and !Recursive on a path with an implied child, should report
-		// the added path children of the Implied child
-		Tree.Empty();
-		Tree.FindOrAdd(TEXTVIEW("/Root/Implied1/Added1")).Value = 1;
-		Tree.FindOrAdd(TEXTVIEW("/Root/Implied1/Added2")).Value = 1;
-		Tree.FindOrAdd(TEXTVIEW("/Root/Implied2/Added")).Value = 1;
-
-		Children.Reset();
-		bExists = Tree.TryGetChildren(TEXTVIEW("/Root"), Children,
-			EDirectoryTreeGetFlags::ImpliedParent);
-		TestTrue(TEXT("!ImpliedChildren, !Recursive, and direct child is implied."),
-			bExists = true && UnorderedEquals(Children, { TEXT("Implied1/Added1"), TEXT("Implied1/Added2"), TEXT("Implied2/Added") }));
-	}
-
-	return true;
-}
-
-#endif
diff --git a/Engine/Source/Runtime/Core/Tests/Containers/DirectoryTreeTest.cpp b/Engine/Source/Runtime/Core/Tests/Containers/DirectoryTreeTest.cpp
new file mode 100644
--- /dev/null
+++ b/Engine/Source/Runtime/Core/Tests/Containers/DirectoryTreeTest.cpp
@@ -0,0 +1,944 @@
+// Copyright Epic Games, Inc. All Rights Reserved.
+
+// Additional tests are present in Core/Private/DirectoryTree.cpp
+
+#if WITH_TESTS
+
+#include "Containers/DirectoryTree.h"
+
+#include "Tests/TestHarnessAdapter.h"
+
+TEST_CASE_NAMED(FDirectoryTreeContainsChildPathTests, "System::Core::Containers::DirectoryTree::ContainsChildPaths", "[Core][Containers][DirectoryTree]")
+{
+	TDirectoryTree<int32> Tree;
+
+	Tree.FindOrAdd(TEXTVIEW("/Game/Dir1"));
+	Tree.FindOrAdd(TEXTVIEW("/Game/Dir2"));
+	Tree.FindOrAdd(TEXTVIEW("/Game/Dir2/Grandchild"));
+	Tree.FindOrAdd(TEXTVIEW("/Plugin1/Path1"));
+	Tree.FindOrAdd(TEXTVIEW("/Plugin1/Path2"));
+
+	CHECK(Tree.ContainsChildPaths(TEXTVIEW("/")));
+	CHECK(Tree.ContainsChildPaths(TEXTVIEW("/Game")));
+	CHECK(Tree.ContainsChildPaths(TEXTVIEW("/Game/Dir2")));
+	CHECK(Tree.ContainsChildPaths(TEXTVIEW("/Plugin1")));
+
+	CHECK_FALSE(Tree.ContainsChildPaths(TEXTVIEW("/Game/Dir1")));
+	CHECK_FALSE(Tree.ContainsChildPaths(TEXTVIEW("/Game/Dir2/Grandchild")));
+	CHECK_FALSE(Tree.ContainsChildPaths(TEXTVIEW("/Plugin1/Path1")));
+	CHECK_FALSE(Tree.ContainsChildPaths(TEXTVIEW("/Plugin1/Path2")));
+}
+
+TEST_CASE_NAMED(FDirectoryTreeAPITests, "System::Core::Containers::DirectoryTree::API", "[Core][Containers][DirectoryTree]")
+{
+	constexpr int32 NumPathTypes = 5;
+	constexpr int32 NumPaths = 9;
+	FStringView PathsByTypeAndIndex[NumPathTypes][NumPaths] =
+	{
+		{
+			TEXTVIEW("/Game/Dir2"),
+			TEXTVIEW("/Game/Path1"),
+			TEXTVIEW("/Game/Dir2/Path2"),
+			TEXTVIEW("/Plugin1/Path1"),
+			TEXTVIEW("/Plugin1/Dir2/Path2"),
+			TEXTVIEW("/Engine/Path1"),
+			TEXTVIEW("/Plugin2/Path1"),
+			// Make sure we handle suffixes of an existing string with a leading value that sorts after /
+			TEXTVIEW("/Game/Foo/Leaf"),
+			TEXTVIEW("/Game/Foo-Bar/Leaf"),
+		},
+		{
+			TEXTVIEW("d:\\root\\Project\\Content\\Dir2"),
+			TEXTVIEW("d:\\root\\Project\\Content\\Path1.uasset"),
+			TEXTVIEW("d:\\root\\Project\\Content\\Dir2\\Path2.uasset"),
+			TEXTVIEW("d:\\root\\Project\\Plugins\\Plugin1\\Content\\Path1.uasset"),
+			TEXTVIEW("d:\\root\\Project\\Plugins\\Plugin1\\Content\\Dir2\\Path2.uasset"),
+			TEXTVIEW("d:\\root\\Engine\\Content\\Path1.uasset"),
+			TEXTVIEW("e:\\root\\Project\\Plugins\\Plugin2\\Content\\Path1.uasset"),
+			TEXTVIEW("d:\\root\\Project\\Content\\Foo\\Leaf"),
+			TEXTVIEW("d:\\root\\Project\\Content\\Foo-Bar\\Leaf"),
+		},
+		{
+			TEXTVIEW("d:/root/Project/Content/Dir2"),
+			TEXTVIEW("d:/root/Project/Content/Path1.uasset"),
+			TEXTVIEW("d:/root/Project/Content/Dir2/Path2.uasset"),
+			TEXTVIEW("d:/root/Project/Plugins/Plugin1/Content/Path1.uasset"),
+			TEXTVIEW("d:/root/Project/Plugins/Plugin1/Content/Dir2/Path2.uasset"),
+			TEXTVIEW("d:/root/Engine/Content/Path1.uasset"),
+			TEXTVIEW("e:/root/Project/Plugins/Plugin2/Content/Path1.uasset"),
+			TEXTVIEW("d:/root/Project/Content/Foo/Leaf"),
+			TEXTVIEW("d:/root/Project/Content/Foo-Bar/Leaf"),
+		},
+		{
+			TEXTVIEW("..\\..\\..\\Project\\Content\\Dir2"),
+			TEXTVIEW("..\\..\\..\\Project\\Content\\Path1.uasset"),
+			TEXTVIEW("..\\..\\..\\Project\\Content\\Dir2\\Path2.uasset"),
+			TEXTVIEW("..\\..\\..\\Project\\Plugins\\Plugin1\\Content\\Path1.uasset"),
+			TEXTVIEW("..\\..\\..\\Project\\Plugins\\Plugin1\\Content\\Dir2\\Path2.uasset"),
+			TEXTVIEW("..\\..\\..\\Engine\\Content\\Path1.uasset"),
+			TEXTVIEW("e:\\root\\Project\\Plugins\\Plugin2\\Content\\Path1.uasset"),
+			TEXTVIEW("..\\..\\..\\Project\\Content\\Foo\\Leaf"),
+			TEXTVIEW("..\\..\\..\\Project\\Content\\Foo-Bar\\Leaf"),
+		},
+		{
+			TEXTVIEW("../../../Project/Content/Dir2"),
+			TEXTVIEW("../../../Project/Content/Path1.uasset"),
+			TEXTVIEW("../../../Project/Content/Dir2/Path2.uasset"),
+			TEXTVIEW("../../../Project/Plugins/Plugin1/Content/Path1.uasset"),
+			TEXTVIEW("../../../Project/Plugins/Plugin1/Content/Dir2/Path2.uasset"),
+			TEXTVIEW("../../../Engine/Content/Path1.uasset"),
+			TEXTVIEW("e:/root/Project/Plugins/Plugin2/Content/Path1.uasset"),
+			TEXTVIEW("../../../Project/Content/Foo/Leaf"),
+			TEXTVIEW("../../../Project/Content/Foo-Bar/Leaf"),
+		},
+	};
+	// PathSub0SubPath[i] provides a sub path of PathsByTypeAndIndex[i][0]
+	FStringView Path0SubPath[NumPathTypes] =
+	{
+		TEXTVIEW("/Game/Dir2/Sub"),
+		TEXTVIEW("d:\\root\\Project\\Content\\Dir2\\Sub"),
+		TEXTVIEW("d:/root/Project/Content/Dir2/Sub"),
+		TEXTVIEW("..\\..\\..\\Project\\Content\\Dir2"),
+		TEXTVIEW("../../../Project/Content/Dir2/Sub"),
+	};
+	int32 ValueByIndex[NumPaths] = { 1,2,3,4,5,6,7,8,9 };
+	FStringView NonPathsByTypeAndIndex0[] =
+		{
+			TEXTVIEW(""),
+			TEXTVIEW("/"),
+			TEXTVIEW("/Game"),
+			TEXTVIEW("/Game/"),
+			TEXTVIEW("/Plugin1"),
+			TEXTVIEW("/Plugin1/"),
+			TEXTVIEW("/Plugin1/Dir2"),
+			TEXTVIEW("/Plugin1/Dir2/"),
+			TEXTVIEW("/Engine"),
+			TEXTVIEW("/Engine/"),
+		};
+	FStringView NonPathsByTypeAndIndex1[] =
+		{
+			TEXTVIEW(""),
+			TEXTVIEW("d:\\"),
+			TEXTVIEW("d:\\root1"),
+			TEXTVIEW("d:\\root1\\"),
+			TEXTVIEW("d:\\root1\\Project"),
+			TEXTVIEW("d:\\root1\\Project\\"),
+			TEXTVIEW("d:\\root1\\Project\\Content"),
+			TEXTVIEW("d:\\root1\\Project\\Plugins\\"),
+			TEXTVIEW("d:\\root1\\Project\\Plugins\\Content"),
+			TEXTVIEW("d:\\root1\\Project\\Plugins\\Content\\"),
+			TEXTVIEW("d:\\root1\\Project\\Plugins\\Content\\Plugin1"),
+			TEXTVIEW("d:\\root1\\Project\\Plugins\\Content\\Plugin1\\"),
+			TEXTVIEW("d:\\root1\\Project\\Plugins\\Content\\Plugin1\\Dir2"),
+			TEXTVIEW("d:\\root1\\Project\\Plugins\\Content\\Plugin1\\Dir2\\"),
+			TEXTVIEW("d:\\root1\\Engine"),
+			TEXTVIEW("d:\\root1\\Engine\\"),
+			TEXTVIEW("d:\\root1\\Engine\\Content"),
+			TEXTVIEW("d:\\root1\\Engine\\Content\\"),
+		};
+	FStringView NonPathsByTypeAndIndex2[] =
+		{
+			TEXTVIEW(""),
+			TEXTVIEW("d:/"),
+			TEXTVIEW("d:/root1"),
+			TEXTVIEW("d:/root1/Project"),
+			TEXTVIEW("d:/root1/Project/Content"),
+			TEXTVIEW("d:/root1/Project/Plugins/Content"),
+			TEXTVIEW("d:/root1/Project/Plugins/Content/Plugin1"),
+			TEXTVIEW("d:/root1/Project/Plugins/Content/Plugin1/Dir2"),
+			TEXTVIEW("d:/root1/Engine"),
+			TEXTVIEW("d:/root1/Engine/Content"),
+		};
+	FStringView NonPathsByTypeAndIndex3[] =
+		{
+			TEXTVIEW(""),
+			TEXTVIEW("\\"),
+			TEXTVIEW(".."),
+			TEXTVIEW("..\\"),
+			TEXTVIEW("..\\.."),
+			TEXTVIEW("..\\..\\"),
+			TEXTVIEW("..\\..\\.."),
+			TEXTVIEW("..\\..\\..\\"),
+			TEXTVIEW("..\\..\\..\\Project"),
+			TEXTVIEW("..\\..\\..\\Project\\Content"),
+			TEXTVIEW("..\\..\\..\\Project\\Plugins\\Content"),
+			TEXTVIEW("..\\..\\..\\Project\\Plugins\\Content\\Plugin1"),
+			TEXTVIEW("..\\..\\..\\Project\\Plugins\\Content\\Plugin1\\Dir2"),
+			TEXTVIEW("..\\..\\..\\Engine"),
+			TEXTVIEW("..\\..\\..\\Engine\\Content"),
+		};
+	FStringView NonPathsByTypeAndIndex4[] =
+		{
+			TEXTVIEW(""),
+			TEXTVIEW("/"),
+			TEXTVIEW(".."),
+			TEXTVIEW("../.."),
+			TEXTVIEW("../../.."),
+			TEXTVIEW("../../../Project"),
+			TEXTVIEW("../../../Project/Content"),
+			TEXTVIEW("../../../Project/Plugins/Content"),
+			TEXTVIEW("../../../Project/Plugins/Content/Plugin1"),
+			TEXTVIEW("../../../Project/Plugins/Content/Plugin1/Dir2"),
+			TEXTVIEW("../../../Engine"),
+			TEXTVIEW("../../../Engine/Content"),
+		};
+	TConstArrayView<FStringView> NonPathsByTypeAndIndex[NumPathTypes] =
+	{
+		NonPathsByTypeAndIndex0, NonPathsByTypeAndIndex1, NonPathsByTypeAndIndex2, NonPathsByTypeAndIndex3,
+		NonPathsByTypeAndIndex4
+	};
+
+	constexpr int32 NumPermutations = 2;
+	int32 Permutations[NumPermutations][NumPaths] =
+	{
+		{ 0, 1, 2, 3, 4, 5, 6, 7, 8 },
+		{ 8, 7, 6, 5, 4, 3, 2, 1, 0 },
+	};
+
+	auto GetDirTreeTestName = [](FStringView InTestName, int32 PathType, int32 Permutation, int32 EditPermutationIndex,
+		int32 OtherPermutationIndex)
+	{
+		return FString::Printf(TEXT("%.*s(%d, %d, %d, %d)"),
+			InTestName.Len(), InTestName.GetData(),
+			PathType, Permutation, EditPermutationIndex, OtherPermutationIndex);
+	};
+
+	TArray<const TCHAR*> ScratchA;
+	TArray<const TCHAR*> ScratchB;
+	auto UnorderedEquals = [&ScratchA, &ScratchB](const TArray<FString>& A, TConstArrayView<const TCHAR*> B)
+		{
+			if (A.Num() != B.Num())
+			{
+				return false;
+			}
+			ScratchA.Reset(A.Num());
+			ScratchB.Reset(B.Num());
+			for (const FString& AStr : A)
+			{
+				ScratchA.Add(*AStr);
+			}
+			for (const TCHAR* BStr : B)
+			{
+				ScratchB.Add(BStr);
+			}
+			Algo::Sort(ScratchA, [](const TCHAR* StrA, const TCHAR* StrB)
+				{
+					return FCString::Stricmp(StrA, StrB) < 0;
+				});
+			Algo::Sort(ScratchB, [](const TCHAR* StrA, const TCHAR* StrB)
+				{
+					return FCString::Stricmp(StrA, StrB) < 0;
+				});
+			for (int32 Index = 0; Index < A.Num(); ++Index)
+			{
+				if (FCString::Stricmp(ScratchA[Index], ScratchB[Index]) != 0)
+				{
+					return false;
+				}
+			}
+			return true;
+		};
+
+	for (int32 PathType = 0; PathType < NumPathTypes; ++PathType)
+	{
+		int32 NumNonPaths = NonPathsByTypeAndIndex[PathType].Num();
+		for (int32 Permutation = 0; Permutation < NumPermutations; ++Permutation)
+		{
+			TDirectoryTree<int32> Tree;
+
+			// Add all the Paths in the given order and make Contains assertions after each addition
+			for (int32 AddPathPermutationIndex = 0; AddPathPermutationIndex < NumPaths; ++AddPathPermutationIndex)
+			{
+				int32 AddPathIndex = Permutations[Permutation][AddPathPermutationIndex];
+				FStringView AddPath = PathsByTypeAndIndex[PathType][AddPathIndex];
+
+				// Add the path
+				Tree.FindOrAdd(AddPath) = ValueByIndex[AddPathIndex];
+
+				if (Tree.Num() != AddPathPermutationIndex+1)
+				{
+					ADD_ERROR(GetDirTreeTestName(TEXTVIEW("TreeNum has expected value"), PathType, Permutation,
+						AddPathPermutationIndex, 0));
+				}
+
+				// Assert all paths up to and including this one are included
+				for (int32 OtherPermutationIndex = 0; OtherPermutationIndex <= AddPathPermutationIndex;
+					++OtherPermutationIndex)
+				{
+					int32 OtherIndex = Permutations[Permutation][OtherPermutationIndex];
+					FStringView OtherPath = PathsByTypeAndIndex[PathType][OtherIndex];
+					int32 OtherValue = ValueByIndex[OtherIndex];
+					int32* ExistingValue = Tree.Find(OtherPath);
+					if (!ExistingValue)
+					{
+						ADD_ERROR(GetDirTreeTestName(TEXTVIEW("OtherInList"), PathType, Permutation,
+							AddPathPermutationIndex, OtherPermutationIndex));
+					}
+					else
+					{
+						if (*ExistingValue != OtherValue)
+						{
+							ADD_ERROR(GetDirTreeTestName(TEXTVIEW("OtherInListMatchesValue"), PathType, Permutation,
+								AddPathPermutationIndex, OtherPermutationIndex));
+						}
+						if (!Tree.ContainsPathOrParent(OtherPath))
+						{
+							ADD_ERROR(GetDirTreeTestName(TEXTVIEW("ContainsPathOrParentOtherInList"), PathType, Permutation,
+								AddPathPermutationIndex, OtherPermutationIndex));
+						}
+						else
+						{
+							FString ClosestPath;
+							if (!Tree.TryFindClosestPath(OtherPath, ClosestPath))
+							{
+								ADD_ERROR(GetDirTreeTestName(TEXTVIEW("TryFindClosestPathOtherInListSucceeds"), PathType, Permutation,
+									AddPathPermutationIndex, OtherPermutationIndex));
+							}
+							else if (ClosestPath != OtherPath)
+							{
+								ADD_ERROR(GetDirTreeTestName(TEXTVIEW("TryFindClosestPathOtherInListMatches"), PathType, Permutation,
+									AddPathPermutationIndex, OtherPermutationIndex));
+							}
+						}
+					}
+				}
+
+				// Assert all paths not yet added are not included
+				for (int32 OtherPermutationIndex = AddPathPermutationIndex + 1; OtherPermutationIndex < NumPaths;
+					++OtherPermutationIndex)
+				{
+					int32 OtherIndex = Permutations[Permutation][OtherPermutationIndex];
+					FStringView OtherPath = PathsByTypeAndIndex[PathType][OtherIndex];
+					if (Tree.Contains(OtherPath))
+					{
+						ADD_ERROR(GetDirTreeTestName(TEXTVIEW("OtherNotInList"), PathType, Permutation,
+							AddPathPermutationIndex, OtherPermutationIndex));
+					}
+					// Call ContainsPathOrParent to test whether it crashes, but not yet implemented that we verify 
+					// what its return value should be.
+					(void)Tree.ContainsPathOrParent(OtherPath);
+				}
+
+				// Assert all non paths are not included
+				for (int32 NonPathIndex = 0; NonPathIndex < NumNonPaths; ++NonPathIndex)
+				{
+					FStringView NonPath = NonPathsByTypeAndIndex[PathType][NonPathIndex];
+					if (Tree.Contains(NonPath))
+					{
+						ADD_ERROR(GetDirTreeTestName(TEXTVIEW("NonPathNotInList"), PathType, Permutation,
+							AddPathPermutationIndex, NonPathIndex));
+					}
+					// Call ContainsPathOrParent to test whether it crashes, but not yet implemented that we verify 
+					// what its return value should be.
+					(void)Tree.ContainsPathOrParent(NonPath);
+				}
+			}
+
+			// Verify that the SubPath is present
+			FString ExistingSubParentPath;
+			int32* ExistingSubParentValue;
+			if (!Tree.TryFindClosestPath(Path0SubPath[PathType], ExistingSubParentPath, &ExistingSubParentValue))
+			{
+				ADD_ERROR(GetDirTreeTestName(TEXTVIEW("SubPathInTree"), PathType, Permutation, 0, 0));
+			}
+			else if (ExistingSubParentPath != PathsByTypeAndIndex[PathType][0] || *ExistingSubParentValue != ValueByIndex[0])
+			{
+				ADD_ERROR(GetDirTreeTestName(TEXTVIEW("SubPathInTreeMatches"), PathType, Permutation, 0, 0));
+			}
+
+			// Remove all the Paths (in specified order) and make Contains assertions after each removal
+			// Currently we only test removal in FIFO order; bugs that are specific to a removal-order
+			// should be dependent only on the final added state and should not be dependent on the earlier add-order.
+			for (int32 RemovePathPermutationIndex = 0; RemovePathPermutationIndex < NumPaths;
+				++RemovePathPermutationIndex)
+			{
+				int32 RemovePathIndex = Permutations[Permutation][RemovePathPermutationIndex];
+				FStringView RemovePath = PathsByTypeAndIndex[PathType][RemovePathIndex];
+
+				// Remove the path
+				bool bExisted;
+				Tree.Remove(RemovePath, &bExisted);
+				if (!bExisted)
+				{
+					ADD_ERROR(GetDirTreeTestName(TEXTVIEW("RemoveFoundSomethingToRemove"), PathType, Permutation,
+						RemovePathPermutationIndex, 0));
+				}
+				if (Tree.Num() != NumPaths - (RemovePathPermutationIndex+1))
+				{
+					ADD_ERROR(GetDirTreeTestName(TEXTVIEW("TreeNum has expected value"), PathType, Permutation,
+						RemovePathPermutationIndex, 0));
+				}
+
+
+				// Assert all paths not yet removed are still included
+				for (int32 OtherPermutationIndex = RemovePathPermutationIndex+1; OtherPermutationIndex < NumPaths;
+					++OtherPermutationIndex)
+				{
+					int32 OtherIndex = Permutations[Permutation][OtherPermutationIndex];
+					FStringView OtherPath = PathsByTypeAndIndex[PathType][OtherIndex];
+					int32 OtherValue = ValueByIndex[OtherIndex];
+					int32* ExistingValue = Tree.Find(OtherPath);
+					if (!ExistingValue)
+					{
+						ADD_ERROR(GetDirTreeTestName(TEXTVIEW("OtherInListAfterRemoval"), PathType, Permutation,
+							RemovePathPermutationIndex, OtherPermutationIndex));
+					}
+					else
+					{
+						if (*ExistingValue != OtherValue)
+						{
+							ADD_ERROR(GetDirTreeTestName(TEXTVIEW("OtherInListAfterRemovalMatches"), PathType, Permutation,
+								RemovePathPermutationIndex, OtherPermutationIndex));
+						}
+						if (!Tree.ContainsPathOrParent(OtherPath))
+						{
+							ADD_ERROR(GetDirTreeTestName(TEXTVIEW("OtherContainsPathOrParentAfterRemoval"), PathType, Permutation,
+								RemovePathPermutationIndex, OtherPermutationIndex));
+						}
+					}
+				}
+
+				// Assert all paths up to and including this one have been removed
+				for (int32 OtherPermutationIndex = 0; OtherPermutationIndex <= RemovePathPermutationIndex;
+					++OtherPermutationIndex)
+				{
+					int32 OtherIndex = Permutations[Permutation][OtherPermutationIndex];
+					FStringView OtherPath = PathsByTypeAndIndex[PathType][OtherIndex];
+					if (Tree.Contains(OtherPath))
+					{
+						ADD_ERROR(GetDirTreeTestName(TEXT("OtherNotInListAfterRemoval"), PathType, Permutation,
+							RemovePathPermutationIndex, OtherPermutationIndex));
+					}
+					// Call ContainsPathOrParent to test whether it crashes, but not yet implemented that we verify 
+					// what its return value should be.
+					(void)Tree.ContainsPathOrParent(OtherPath);
+				}
+
+				// Assert all non paths are not still included
+				for (int32 NonPathIndex = 0; NonPathIndex < NumNonPaths; ++NonPathIndex)
+				{
+					FStringView NonPath = NonPathsByTypeAndIndex[PathType][NonPathIndex];
+					if (Tree.Contains(NonPath))
+					{
+						ADD_ERROR(GetDirTreeTestName(TEXTVIEW("NonPathNotInListAfterRemoval"), PathType, Permutation,
+							RemovePathPermutationIndex, NonPathIndex));
+					}
+					// Call ContainsPathOrParent to test whether it crashes, but not yet implemented that we verify 
+					// what its return value should be.
+					(void)Tree.ContainsPathOrParent(NonPath);
+				}
+			}
+			if (!Tree.IsEmpty())
+			{
+				ADD_ERROR(GetDirTreeTestName(TEXTVIEW("TreeEmptyAfterRemoval"), PathType, Permutation, 0, 0));
+			}
+		}
+	}
+
+	// Testing some pathtype-independent scenarios
+	{
+		TDirectoryTree<int32> Tree;
+		Tree.FindOrAdd(TEXTVIEW("/Root/Path1")) = 1;
+		Tree.FindOrAdd(TEXTVIEW("/Root/Path2")) = 2;
+		int32* FoundRoot = Tree.FindClosestValue(TEXTVIEW("/Root"));
+		int32* FoundPath1 = Tree.FindClosestValue(TEXTVIEW("/Root/Path1"));
+		int32* FoundPath1Sub = Tree.FindClosestValue(TEXTVIEW("/Root/Path1/Sub"));
+		int32* FoundPath2 = Tree.FindClosestValue(TEXTVIEW("/Root/Path2"));
+		int32* FoundPath2Sub = Tree.FindClosestValue(TEXTVIEW("/Root/Path2/Sub"));
+		CHECK_MESSAGE(TEXT("TwoPaths Root does not exist"), FoundRoot == nullptr);
+		CHECK_MESSAGE(TEXT("TwoPaths Path1 Value matches"), FoundPath1 && *FoundPath1 == 1);
+		CHECK_MESSAGE(TEXT("TwoPaths Path1Sub Value matches"), FoundPath1Sub &&* FoundPath1Sub == 1);
+		CHECK_MESSAGE(TEXT("TwoPaths Path2 Value matches"), FoundPath2 && *FoundPath2 == 2);
+		CHECK_MESSAGE(TEXT("TwoPaths Path2Sub Value matches"), FoundPath2Sub && *FoundPath2Sub == 2);
+	}
+	{
+		TDirectoryTree<int32> Tree;
+		Tree.FindOrAdd(TEXTVIEW("/Root/Path1/A/B/C")) = 1;
+		Tree.FindOrAdd(TEXTVIEW("/Root/Path2/A/B/C")) = 2;
+		int32* FoundRoot = Tree.FindClosestValue(TEXTVIEW("/Root"));
+		int32* FoundPath1 = Tree.FindClosestValue(TEXTVIEW("/Root/Path1/A/B/C"));
+		int32* FoundPath1Sub = Tree.FindClosestValue(TEXTVIEW("/Root/Path1/A/B/C/Sub"));
+		int32* FoundPath1Parent = Tree.FindClosestValue(TEXTVIEW("/Root/Path1/A"));
+		int32* FoundPath2 = Tree.FindClosestValue(TEXTVIEW("/Root/Path2/A/B/C"));
+		int32* FoundPath2Sub = Tree.FindClosestValue(TEXTVIEW("/Root/Path2/A/B/C/Sub"));
+		int32* FoundPath2Parent = Tree.FindClosestValue(TEXTVIEW("/Root/Path2/A"));
+		CHECK_MESSAGE(TEXT("TwoPathsLong Root does not exist"), FoundRoot == nullptr);
+		CHECK_MESSAGE(TEXT("TwoPathsLong Path1 Value matches"), FoundPath1 && *FoundPath1 == 1);
+		CHECK_MESSAGE(TEXT("TwoPathsLong Path1Sub Value matches"), FoundPath1Sub && *FoundPath1Sub == 1);
+		CHECK_MESSAGE(TEXT("TwoPathsLong Path1 Parent does not exist"), FoundPath1Parent == nullptr);
+		CHECK_MESSAGE(TEXT("TwoPathsLong Path2 Value matches"), FoundPath2 && *FoundPath2 == 2);
+		CHECK_MESSAGE(TEXT("TwoPathsLong Path2Sub Value matches"), FoundPath2Sub && *FoundPath2Sub == 2);
+		CHECK_MESSAGE(TEXT("TwoPathsLong Path2 Parent does not exist"), FoundPath2Parent == nullptr);
+	}
+
+	struct FMoveConstructOnly
+	{
+		FMoveConstructOnly() { Value = 437; }
+		FMoveConstructOnly(FMoveConstructOnly&& Other) { Value = Other.Value; Other.Value = -1; }
+		FMoveConstructOnly(const FMoveConstructOnly& Other) = delete;
+		FMoveConstructOnly& operator=(FMoveConstructOnly&& Other) = delete;
+		FMoveConstructOnly& operator=(const FMoveConstructOnly& Other) = delete;
+		
+		int32 Value;
+	};
+
+	{
+		TDirectoryTree<FMoveConstructOnly> Tree;
+		Tree.FindOrAdd(TEXTVIEW("/Root/PathM")).Value = 1;
+		Tree.FindOrAdd(TEXTVIEW("/Root/PathP")).Value = 2;
+		Tree.FindOrAdd(TEXTVIEW("/Root/PathA")).Value = 3;
+		Tree.FindOrAdd(TEXTVIEW("/Root/PathZ"));
+
+		FMoveConstructOnly* Value = Tree.Find(TEXTVIEW("/Root/PathA"));
+		CHECK_MESSAGE(TEXT("MoveConstructOnlyValueA correct"), Value && Value->Value == 3);
+		Value = Tree.Find(TEXTVIEW("/Root/PathM"));
+		CHECK_MESSAGE(TEXT("MoveConstructOnlyValueM correct"), Value && Value->Value == 1);
+		Value = Tree.Find(TEXTVIEW("/Root/PathP"));
+		CHECK_MESSAGE(TEXT("MoveConstructOnlyValueP correct"), Value && Value->Value == 2);
+		Value = Tree.Find(TEXTVIEW("/Root/PathZ"));
+		CHECK_MESSAGE(TEXT("MoveConstructOnlyValueZ correct"), Value && Value->Value == 437);
+	}
+
+	// Handling special case of drive specifiers without a path
+	{
+		TStringBuilder<16> FoundPath;
+		TArray<FString> ChildNames;
+		int* FoundValue = nullptr;
+		auto Reset = [&FoundPath, &ChildNames, &FoundValue]()
+			{
+				FoundPath.Reset();
+				ChildNames.Reset();
+				FoundValue = nullptr;
+			};
+
+		{
+			TDirectoryTree<int32> Tree;
+			Tree.FindOrAdd(TEXTVIEW("D:")) = 1;
+
+			Reset();
+			CHECK_MESSAGE(TEXT("DriveSpecifier: Before PathSep: Without PathSep: Tree.Contains"), Tree.Contains(TEXTVIEW("D:")));
+			CHECK_MESSAGE(TEXT("DriveSpecifier: Before PathSep: Without PathSep: Tree.Find"), Tree.Find(TEXTVIEW("D:")) != nullptr);
+			CHECK_MESSAGE(TEXT("DriveSpecifier: Before PathSep: Without PathSep: Tree.ContainsPathOrParent"), Tree.ContainsPathOrParent(TEXTVIEW("D:")));
+			CHECK_MESSAGE(TEXT("DriveSpecifier: Before PathSep: Without PathSep: Tree.FindClosestValue"), Tree.FindClosestValue(TEXTVIEW("D:")) != nullptr);
+			CHECK_MESSAGE(TEXT("DriveSpecifier: Before PathSep: Without PathSep: Tree.TryFindClosestPath"),
+				Tree.TryFindClosestPath(TEXTVIEW("D:"), FoundPath, &FoundValue) && FoundPath.Len() > 0 && FoundValue != nullptr);
+			CHECK_MESSAGE(TEXT("DriveSpecifier: Before PathSep: Without PathSep: Tree.TryGetChildren"), Tree.TryGetChildren(TEXTVIEW("D:"), ChildNames));
+
+			Tree.FindOrAdd(TEXTVIEW("D:/root")) = 1;
+
+			Reset();
+			CHECK_MESSAGE(TEXT("DriveSpecifier: After PathSep('/'): Without PathSep: Tree.Contains"), Tree.Contains(TEXTVIEW("D:")));
+			CHECK_MESSAGE(TEXT("DriveSpecifier: After PathSep('/'): Without PathSep: Tree.Find"), Tree.Find(TEXTVIEW("D:")) != nullptr);
+			CHECK_MESSAGE(TEXT("DriveSpecifier: After PathSep('/'): Without PathSep: Tree.ContainsPathOrParent"), Tree.ContainsPathOrParent(TEXTVIEW("D:")));
+			CHECK_MESSAGE(TEXT("DriveSpecifier: After PathSep('/'): Without PathSep: Tree.FindClosestValue"), Tree.FindClosestValue(TEXTVIEW("D:")) != nullptr);
+			CHECK_MESSAGE(TEXT("DriveSpecifier: After PathSep('/'): Without PathSep: Tree.TryFindClosestPath"),
+				Tree.TryFindClosestPath(TEXTVIEW("D:"), FoundPath, &FoundValue) && FoundPath.Len() > 0 && FoundValue != nullptr);
+			CHECK_MESSAGE(TEXT("DriveSpecifier: After PathSep('/'): Without PathSep: Tree.TryGetChildren"), Tree.TryGetChildren(TEXTVIEW("D:"), ChildNames));
+
+			Reset();
+			CHECK_MESSAGE(TEXT("DriveSpecifier: After PathSep('/'): With PathSep: Tree.Contains"), Tree.Contains(TEXTVIEW("D:/")));
+			CHECK_MESSAGE(TEXT("DriveSpecifier: After PathSep('/'): With PathSep: Tree.Find"), Tree.Find(TEXTVIEW("D:/")) != nullptr);
+			CHECK_MESSAGE(TEXT("DriveSpecifier: After PathSep('/'): With PathSep: Tree.ContainsPathOrParent"), Tree.ContainsPathOrParent(TEXTVIEW("D:/")));
+			CHECK_MESSAGE(TEXT("DriveSpecifier: After PathSep('/'): With PathSep: Tree.FindClosestValue"), Tree.FindClosestValue(TEXTVIEW("D:/")) != nullptr);
+			CHECK_MESSAGE(TEXT("DriveSpecifier: After PathSep('/'): With PathSep: Tree.TryFindClosestPath"),
+				Tree.TryFindClosestPath(TEXTVIEW("D:/"), FoundPath, &FoundValue) && FoundPath.Len() > 0 && FoundValue != nullptr);
+			CHECK_MESSAGE(TEXT("DriveSpecifier: After PathSep('/'): With PathSep: Tree.TryGetChildren"), Tree.TryGetChildren(TEXTVIEW("D:/"), ChildNames));
+		}
+		{
+			TDirectoryTree<int32> Tree;
+			Tree.FindOrAdd(TEXTVIEW("D:")) = 1;
+			Tree.FindOrAdd(TEXTVIEW("D:\\root")) = 1;
+
+			Reset();
+			CHECK_MESSAGE(TEXT("DriveSpecifier: After PathSep('\\'): Without PathSep: Tree.Contains"), Tree.Contains(TEXTVIEW("D:")));
+			CHECK_MESSAGE(TEXT("DriveSpecifier: After PathSep('\\'): Without PathSep: Tree.Find"), Tree.Find(TEXTVIEW("D:")) != nullptr);
+			CHECK_MESSAGE(TEXT("DriveSpecifier: After PathSep('\\'): Without PathSep: Tree.ContainsPathOrParent"), Tree.ContainsPathOrParent(TEXTVIEW("D:")));
+			CHECK_MESSAGE(TEXT("DriveSpecifier: After PathSep('\\'): Without PathSep: Tree.FindClosestValue"), Tree.FindClosestValue(TEXTVIEW("D:")) != nullptr);
+			CHECK_MESSAGE(TEXT("DriveSpecifier: After PathSep('\\'): Without PathSep: Tree.TryFindClosestPath"),
+				Tree.TryFindClosestPath(TEXTVIEW("D:"), FoundPath, &FoundValue) && FoundPath.Len() > 0 && FoundValue != nullptr);
+			CHECK_MESSAGE(TEXT("DriveSpecifier: After PathSep('\\'): Without PathSep: Tree.TryGetChildren"), Tree.TryGetChildren(TEXTVIEW("D:"), ChildNames));
+
+			Reset();
+			CHECK_MESSAGE(TEXT("DriveSpecifier: After PathSep('\\'): With PathSep: Tree.Contains"), Tree.Contains(TEXTVIEW("D:\\")));
+			CHECK_MESSAGE(TEXT("DriveSpecifier: After PathSep('\\'): With PathSep: Tree.Find"), Tree.Find(TEXTVIEW("D:\\")) != nullptr);
+			CHECK_MESSAGE(TEXT("DriveSpecifier: After PathSep('\\'): With PathSep: Tree.ContainsPathOrParent"), Tree.ContainsPathOrParent(TEXTVIEW("D:\\")));
+			CHECK_MESSAGE(TEXT("DriveSpecifier: After PathSep('\\'): With PathSep: Tree.FindClosestValue"), Tree.FindClosestValue(TEXTVIEW("D:\\")) != nullptr);
+			CHECK_MESSAGE(TEXT("DriveSpecifier: After PathSep('\\'): With PathSep: Tree.TryFindClosestPath"),
+				Tree.TryFindClosestPath(TEXTVIEW("D:\\"), FoundPath, &FoundValue) && FoundPath.Len() > 0 && FoundValue != nullptr);
+			CHECK_MESSAGE(TEXT("DriveSpecifier: After PathSep('\\'): With PathSep: Tree.TryGetChildren"), Tree.TryGetChildren(TEXTVIEW("D:\\"), ChildNames));
+		}
+		{
+			TDirectoryTree<int32> Tree;
+			Tree.FindOrAdd(TEXTVIEW("D:root")) = 1;
+
+			Reset();
+			CHECK_MESSAGE(TEXT("DriveSpecifierLong: Before PathSep: Without PathSep: Tree.Contains"), Tree.Contains(TEXTVIEW("D:root")));
+			CHECK_MESSAGE(TEXT("DriveSpecifierLong: Before PathSep: Without PathSep: Tree.Find"), Tree.Find(TEXTVIEW("D:root")) != nullptr);
+			CHECK_MESSAGE(TEXT("DriveSpecifierLong: Before PathSep: Without PathSep: Tree.ContainsPathOrParent"), Tree.ContainsPathOrParent(TEXTVIEW("D:root")));
+			CHECK_MESSAGE(TEXT("DriveSpecifierLong: Before PathSep: Without PathSep: Tree.FindClosestValue"), Tree.FindClosestValue(TEXTVIEW("D:root")) != nullptr);
+			CHECK_MESSAGE(TEXT("DriveSpecifierLong: Before PathSep: Without PathSep: Tree.TryFindClosestPath"),
+				Tree.TryFindClosestPath(TEXTVIEW("D:root"), FoundPath, &FoundValue) && FoundPath.Len() > 0 && FoundValue != nullptr);
+			CHECK_MESSAGE(TEXT("DriveSpecifierLong: Before PathSep: Without PathSep: Tree.TryGetChildren"), Tree.TryGetChildren(TEXTVIEW("D:root"), ChildNames));
+
+			Tree.FindOrAdd(TEXTVIEW("D:\\root\\path")) = 1;
+
+			Reset();
+			CHECK_MESSAGE(TEXT("DriveSpecifierLong: After PathSep('\\'): Without PathSep: Tree.Contains"), Tree.Contains(TEXTVIEW("D:root")));
+			CHECK_MESSAGE(TEXT("DriveSpecifierLong: After PathSep('\\'): Without PathSep: Tree.Find"), Tree.Find(TEXTVIEW("D:root")) != nullptr);
+			CHECK_MESSAGE(TEXT("DriveSpecifierLong: After PathSep('\\'): Without PathSep: Tree.ContainsPathOrParent"), Tree.ContainsPathOrParent(TEXTVIEW("D:root")));
+			CHECK_MESSAGE(TEXT("DriveSpecifierLong: After PathSep('\\'): Without PathSep: Tree.FindClosestValue"), Tree.FindClosestValue(TEXTVIEW("D:root")) != nullptr);
+			CHECK_MESSAGE(TEXT("DriveSpecifierLong: After PathSep('\\'): Without PathSep: Tree.TryFindClosestPath"),
+				Tree.TryFindClosestPath(TEXTVIEW("D:root"), FoundPath, &FoundValue) && FoundPath.Len() > 0 && FoundValue != nullptr);
+			CHECK_MESSAGE(TEXT("DriveSpecifierLong: After PathSep('\\'): Without PathSep: Tree.TryGetChildren"), Tree.TryGetChildren(TEXTVIEW("D:root"), ChildNames));
+
+			Reset();
+			CHECK_MESSAGE(TEXT("DriveSpecifierLong: After PathSep('\\'): With PathSep: Tree.Contains"), Tree.Contains(TEXTVIEW("D:\\root")));
+			CHECK_MESSAGE(TEXT("DriveSpecifierLong: After PathSep('\\'): With PathSep: Tree.Find"), Tree.Find(TEXTVIEW("D:\\root")) != nullptr);
+			CHECK_MESSAGE(TEXT("DriveSpecifierLong: After PathSep('\\'): With PathSep: Tree.ContainsPathOrParent"), Tree.ContainsPathOrParent(TEXTVIEW("D:\\root")));
+			CHECK_MESSAGE(TEXT("DriveSpecifierLong: After PathSep('\\'): With PathSep: Tree.FindClosestValue"), Tree.FindClosestValue(TEXTVIEW("D:\\root")) != nullptr);
+			CHECK_MESSAGE(TEXT("DriveSpecifierLong: After PathSep('\\'): With PathSep: Tree.TryFindClosestPath"),
+				Tree.TryFindClosestPath(TEXTVIEW("D:root"), FoundPath, &FoundValue) && FoundPath.Len() > 0 && FoundValue != nullptr);
+			CHECK_MESSAGE(TEXT("DriveSpecifierLong: After PathSep('\\'): With PathSep: Tree.TryGetChildren"), Tree.TryGetChildren(TEXTVIEW("D:\\root"), ChildNames));
+		}
+	}
+
+	// Testing accessors
+	{
+		// GetChildren
+		TDirectoryTree<FMoveConstructOnly> Tree;
+		bool bExists;
+		TArray<FString> Children;
+
+		Children.Reset();
+		bExists = Tree.TryGetChildren(TEXTVIEW(""), Children,
+			EDirectoryTreeGetFlags::None);
+		CHECK_MESSAGE(TEXT("GetChildrenEmpty, Root, !ImpliedParent"),
+			bExists == false && Children.IsEmpty());
+
+		Children.Reset();
+		bExists = Tree.TryGetChildren(TEXTVIEW(""), Children,
+			EDirectoryTreeGetFlags::ImpliedParent |
+			EDirectoryTreeGetFlags::ImpliedChildren |
+			EDirectoryTreeGetFlags::Recursive);
+		CHECK_MESSAGE(TEXT("GetChildrenEmpty, Root, ImpliedParent"),
+			bExists == true && Children.IsEmpty());
+
+		Children.Reset();
+		bExists = Tree.TryGetChildren(TEXTVIEW("/SomePath"), Children,
+			EDirectoryTreeGetFlags::ImpliedParent |
+			EDirectoryTreeGetFlags::ImpliedChildren |
+			EDirectoryTreeGetFlags::Recursive);
+		CHECK_MESSAGE(TEXT("GetChildrenEmpty, Non-root"),
+			bExists == false && Children.IsEmpty());
+
+		Tree.FindOrAdd(TEXTVIEW("")).Value = 1;
+		Children.Reset();
+		bExists = Tree.TryGetChildren(TEXTVIEW(""), Children,
+			EDirectoryTreeGetFlags::None);
+		CHECK_MESSAGE(TEXT("GetChildrenRoot, !ImpliedParent, !ImpliedChildren"),
+			bExists == true && Children.IsEmpty());
+
+		Tree.Empty();
+		Tree.FindOrAdd(TEXTVIEW("/")).Value = 1;
+
+		bExists = Tree.TryGetChildren(TEXTVIEW(""), Children,
+			EDirectoryTreeGetFlags::ImpliedChildren |
+			EDirectoryTreeGetFlags::Recursive);
+		CHECK_MESSAGE(TEXT("GetChildrenRoot, !ImpliedParent, ImpliedChildren"),
+			bExists == false && Children.IsEmpty());
+
+		Children.Reset();
+		bExists = Tree.TryGetChildren(TEXTVIEW(""), Children,
+			EDirectoryTreeGetFlags::ImpliedParent |
+			EDirectoryTreeGetFlags::ImpliedChildren |
+			EDirectoryTreeGetFlags::Recursive);
+		CHECK_MESSAGE(TEXT("GetChildrenRoot, ImpliedParent, ImpliedChildren"),
+			bExists == true && UnorderedEquals(Children, { TEXT("/") }));
+
+		bExists = Tree.TryGetChildren(TEXTVIEW(""), Children,
+			EDirectoryTreeGetFlags::ImpliedParent |
+			EDirectoryTreeGetFlags::ImpliedChildren |
+			EDirectoryTreeGetFlags::Recursive);
+		CHECK_MESSAGE(TEXT("GetChildren appends to the outdir rather than resetting"),
+			bExists == true && Children.Num() == 2);
+
+		Children.Reset();
+		bExists = Tree.TryGetChildren(TEXTVIEW(""), Children,
+			EDirectoryTreeGetFlags::ImpliedParent |
+			EDirectoryTreeGetFlags::Recursive);
+		CHECK_MESSAGE(TEXT("GetChildrenRoot, ImpliedParent, !ImpliedChildren"),
+			bExists == true && UnorderedEquals(Children, { TEXT("/") }));
+
+		Tree.Empty();
+		Tree.FindOrAdd(TEXTVIEW("/Root/Child")).Value = 1;
+
+		Children.Reset();
+		bExists = Tree.TryGetChildren(TEXTVIEW(""), Children,
+			EDirectoryTreeGetFlags::ImpliedParent |
+			EDirectoryTreeGetFlags::ImpliedChildren);
+		CHECK_MESSAGE(TEXT("GetChildrenRootImpliedChild, ImpliedParent, ImpliedChildren, !Recursive"),
+			bExists == true && UnorderedEquals(Children, { TEXT("/") }));
+		Children.Reset();
+		bExists = Tree.TryGetChildren(TEXTVIEW(""), Children,
+			EDirectoryTreeGetFlags::ImpliedParent);
+		CHECK_MESSAGE(TEXT("GetChildrenRootImpliedChild, ImpliedParent, !ImpliedChildren, !Recursive"),
+			bExists == true && UnorderedEquals(Children, { TEXT("/Root/Child") }));
+		Children.Reset();
+		bExists = Tree.TryGetChildren(TEXTVIEW(""), Children,
+			EDirectoryTreeGetFlags::ImpliedParent |
+			EDirectoryTreeGetFlags::ImpliedChildren |
+			EDirectoryTreeGetFlags::Recursive);
+		CHECK_MESSAGE(TEXT("GetChildrenRootImpliedChild, ImpliedParent, ImpliedChildren, Recursive"),
+			bExists == true && UnorderedEquals(Children, { TEXT("/"), TEXT("/Root"), TEXT("/Root/Child") }));
+		Children.Reset();
+		bExists = Tree.TryGetChildren(TEXTVIEW(""), Children,
+			EDirectoryTreeGetFlags::ImpliedParent |
+			EDirectoryTreeGetFlags::Recursive);
+		CHECK_MESSAGE(TEXT("GetChildrenRootImpliedChild, ImpliedParent, !ImpliedChildren, Recursive"),
+			bExists == true && UnorderedEquals(Children, { TEXT("/Root/Child") }));
+
+		Tree.FindOrAdd(TEXTVIEW("/Root/Child2")).Value = 1;
+
+		Tree.Empty();
+		Tree.FindOrAdd(TEXTVIEW("/Stem/A_OtherChild")).Value = 1;
+		Tree.FindOrAdd(TEXTVIEW("/Stem/B_ImpliedRoot/AddedChild")).Value = 1;
+		Tree.FindOrAdd(TEXTVIEW("/Stem/B_ImpliedRoot/AddedChild/Child")).Value = 1;
+		Tree.FindOrAdd(TEXTVIEW("/Stem/B_ImpliedRoot/ImpliedChild/AddedChild")).Value = 1;
+		Tree.FindOrAdd(TEXTVIEW("/Stem/B_ImpliedRoot/ImpliedChild/AddedChild/AddedChild")).Value = 1;
+		Tree.FindOrAdd(TEXTVIEW("/Stem/B_ImpliedRoot/ImpliedChild/AddedChild/ImpliedChild/AddedChild")).Value = 1;
+		Tree.FindOrAdd(TEXTVIEW("/Stem/C_MiddleRoot/MiddlePath/ImpliedChild/AddedChild")).Value = 1;
+		Tree.FindOrAdd(TEXTVIEW("/Stem/C_MiddleRoot/MiddlePath/ImpliedChild/AddedChild/Child")).Value = 1;
+		Tree.FindOrAdd(TEXTVIEW("/Stem/D_MiddleRoot/MiddlePath/AddedChild")).Value = 1;
+		Tree.FindOrAdd(TEXTVIEW("/Stem/D_MiddleRoot/MiddlePath/AddedChild/Child")).Value = 1;
+		Tree.FindOrAdd(TEXTVIEW("/Stem/E_AddedRoot")).Value = 1;
+		Tree.FindOrAdd(TEXTVIEW("/Stem/E_AddedRoot/AddedChild")).Value = 1;
+		Tree.FindOrAdd(TEXTVIEW("/Stem/E_AddedRoot/ImpliedChild/AddedChild")).Value = 1;
+		Tree.FindOrAdd(TEXTVIEW("/Stem/E_AddedRoot/ImpliedChild/AddedChild/ImpliedChild/AddedChild")).Value = 1;
+		Tree.FindOrAdd(TEXTVIEW("/Stem/F_AddedRoot")).Value = 1;
+		Tree.FindOrAdd(TEXTVIEW("/Stem/F_AddedRoot/ImpliedChild/AddedChild")).Value = 1;
+
+
+		// Case: Requested path is an implied path that is a stored child in the tree.
+		Children.Reset();
+		bExists = Tree.TryGetChildren(TEXTVIEW("/Stem/B_ImpliedRoot"), Children,
+			EDirectoryTreeGetFlags::None);
+		CHECK_MESSAGE(TEXT("GetChildrenComplexA B_ImpliedRoot, !ImpliedParent, !ImpliedChildren, !Recursive"),
+			bExists == false && Children.IsEmpty());
+		Children.Reset();
+		bExists = Tree.TryGetChildren(TEXTVIEW("/Stem/B_ImpliedRoot"), Children,
+			EDirectoryTreeGetFlags::Recursive);
+		CHECK_MESSAGE(TEXT("GetChildrenComplexA B_ImpliedRoot, !ImpliedParent, !ImpliedChildren, Recursive"),
+			bExists == false && Children.IsEmpty());
+		Children.Reset();
+		bExists = Tree.TryGetChildren(TEXTVIEW("/Stem/B_ImpliedRoot"), Children,
+			EDirectoryTreeGetFlags::ImpliedChildren);
+		CHECK_MESSAGE(TEXT("GetChildrenComplexA B_ImpliedRoot, !ImpliedParent, ImpliedChildren, !Recursive"),
+			bExists == false && Children.IsEmpty());
+		Children.Reset();
+		bExists = Tree.TryGetChildren(TEXTVIEW("/Stem/B_ImpliedRoot"), Children,
+			EDirectoryTreeGetFlags::ImpliedChildren |
+			EDirectoryTreeGetFlags::Recursive);
+		CHECK_MESSAGE(TEXT("GetChildrenComplexA B_ImpliedRoot, !ImpliedParent, ImpliedChildren, Recursive"),
+			bExists == false && Children.IsEmpty());
+		Children.Reset();
+		bExists = Tree.TryGetChildren(TEXTVIEW("/Stem/B_ImpliedRoot"), Children,
+			EDirectoryTreeGetFlags::ImpliedParent);
+		CHECK_MESSAGE(TEXT("GetChildrenComplexA B_ImpliedRoot, ImpliedParent, !ImpliedChildren, !Recursive"),
+			bExists == true && UnorderedEquals(Children, { TEXT("AddedChild"), TEXT("ImpliedChild/AddedChild") }));
+		Children.Reset();
+		bExists = Tree.TryGetChildren(TEXTVIEW("/Stem/B_ImpliedRoot"), Children,
+			EDirectoryTreeGetFlags::ImpliedParent |
+			EDirectoryTreeGetFlags::Recursive);
+		CHECK_MESSAGE(TEXT("GetChildrenComplexA B_ImpliedRoot, ImpliedParent, !ImpliedChildren, Recursive"),
+			bExists == true && UnorderedEquals(Children, { TEXT("AddedChild"), TEXT("AddedChild/Child"),
+				TEXT("ImpliedChild/AddedChild"), TEXT("ImpliedChild/AddedChild/AddedChild"),
+				TEXT("ImpliedChild/AddedChild/ImpliedChild/AddedChild") }));
+		Children.Reset();
+		bExists = Tree.TryGetChildren(TEXTVIEW("/Stem/B_ImpliedRoot"), Children,
+			EDirectoryTreeGetFlags::ImpliedParent |
+			EDirectoryTreeGetFlags::ImpliedChildren);
+		CHECK_MESSAGE(TEXT("GetChildrenComplexA B_ImpliedRoot, ImpliedParent, ImpliedChildren, !Recursive"),
+			bExists == true && UnorderedEquals(Children, { TEXT("AddedChild"), TEXT("ImpliedChild") }));
+		Children.Reset();
+		bExists = Tree.TryGetChildren(TEXTVIEW("/Stem/B_ImpliedRoot"), Children,
+			EDirectoryTreeGetFlags::ImpliedParent |
+			EDirectoryTreeGetFlags::ImpliedChildren |
+			EDirectoryTreeGetFlags::Recursive);
+		CHECK_MESSAGE(TEXT("GetChildrenComplexA B_ImpliedRoot, ImpliedParent, ImpliedChildren, Recursive"),
+			bExists == true && UnorderedEquals(Children, { TEXT("AddedChild"), TEXT("AddedChild/Child"),
+				TEXT("ImpliedChild"), TEXT("ImpliedChild/AddedChild"), TEXT("ImpliedChild/AddedChild/AddedChild"),
+				TEXT("ImpliedChild/AddedChild/ImpliedChild"), TEXT("ImpliedChild/AddedChild/ImpliedChild/AddedChild") }));
+
+		// Case: Requested path is an implied path that is not a stored child in the tree - it is an in-between dir in a
+		// relpath - and it has an implied child
+		Children.Reset();
+		bExists = Tree.TryGetChildren(TEXTVIEW("/Stem/C_MiddleRoot/MiddlePath"), Children,
+			EDirectoryTreeGetFlags::None);
+		CHECK_MESSAGE(TEXT("GetChildrenComplexA C_MiddleRoot, !ImpliedParent, !ImpliedChildren, !Recursive"),
+			bExists == false && Children.IsEmpty());
+		Children.Reset();
+		bExists = Tree.TryGetChildren(TEXTVIEW("/Stem/C_MiddleRoot/MiddlePath"), Children,
+			EDirectoryTreeGetFlags::Recursive);
+		CHECK_MESSAGE(TEXT("GetChildrenComplexA C_MiddleRoot, !ImpliedParent, !ImpliedChildren, Recursive"),
+			bExists == false && Children.IsEmpty());
+		Children.Reset();
+		bExists = Tree.TryGetChildren(TEXTVIEW("/Stem/C_MiddleRoot/MiddlePath"), Children,
+			EDirectoryTreeGetFlags::ImpliedChildren);
+		CHECK_MESSAGE(TEXT("GetChildrenComplexA C_MiddleRoot, !ImpliedParent, ImpliedChildren, !Recursive"),
+			bExists == false && Children.IsEmpty());
+		Children.Reset();
+		bExists = Tree.TryGetChildren(TEXTVIEW("/Stem/C_MiddleRoot/MiddlePath"), Children,
+			EDirectoryTreeGetFlags::ImpliedChildren |
+			EDirectoryTreeGetFlags::Recursive);
+		CHECK_MESSAGE(TEXT("GetChildrenComplexA C_MiddleRoot, !ImpliedParent, ImpliedChildren, Recursive"),
+			bExists == false && Children.IsEmpty());
+		Children.Reset();
+		bExists = Tree.TryGetChildren(TEXTVIEW("/Stem/C_MiddleRoot/MiddlePath"), Children,
+			EDirectoryTreeGetFlags::ImpliedParent);
+		CHECK_MESSAGE(TEXT("GetChildrenComplexA C_MiddleRoot, ImpliedParent, !ImpliedChildren, !Recursive"),
+			bExists == true && UnorderedEquals(Children, { TEXT("ImpliedChild/AddedChild") }));
+		Children.Reset();
+		bExists = Tree.TryGetChildren(TEXTVIEW("/Stem/C_MiddleRoot/MiddlePath"), Children,
+			EDirectoryTreeGetFlags::ImpliedParent |
+			EDirectoryTreeGetFlags::Recursive);
+		CHECK_MESSAGE(TEXT("GetChildrenComplexA C_MiddleRoot, ImpliedParent, !ImpliedChildren, Recursive"),
+			bExists == true && UnorderedEquals(Children, { TEXT("ImpliedChild/AddedChild"),
+				TEXT("ImpliedChild/AddedChild/Child") }));
+		Children.Reset();
+		bExists = Tree.TryGetChildren(TEXTVIEW("/Stem/C_MiddleRoot/MiddlePath"), Children,
+			EDirectoryTreeGetFlags::ImpliedParent |
+			EDirectoryTreeGetFlags::ImpliedChildren);
+		CHECK_MESSAGE(TEXT("GetChildrenComplexA C_MiddleRoot, ImpliedParent, ImpliedChildren, !Recursive"),
+			bExists == true && UnorderedEquals(Children, { TEXT("ImpliedChild") }));
+		Children.Reset();
+		bExists = Tree.TryGetChildren(TEXTVIEW("/Stem/C_MiddleRoot/MiddlePath"), Children,
+			EDirectoryTreeGetFlags::ImpliedParent |
+			EDirectoryTreeGetFlags::ImpliedChildren |
+			EDirectoryTreeGetFlags::Recursive);
+		CHECK_MESSAGE(TEXT("GetChildrenComplexA C_MiddleRoot, ImpliedParent, ImpliedChildren, Recursive"),
+			bExists == true && UnorderedEquals(Children, { TEXT("ImpliedChild"), TEXT("ImpliedChild/AddedChild"),
+				TEXT("ImpliedChild/AddedChild/Child") }));
+
+		// Case: Requested path is a non-existent sibling path of an implied path that is not a stored
+		// path.
+		Children.Reset();
+		bExists = Tree.TryGetChildren(TEXTVIEW("/Stem/C_MiddleRoot/MiddlePathExceptItDoesNotExist"), Children,
+			EDirectoryTreeGetFlags::ImpliedParent |
+			EDirectoryTreeGetFlags::ImpliedChildren |
+			EDirectoryTreeGetFlags::Recursive);
+		CHECK_MESSAGE(TEXT("GetChildrenComplexA MiddlePathExceptItDoesNotExist, ImpliedParent, ImpliedChildren, Recursive"),
+			bExists == false && Children.IsEmpty());
+
+		// Case: Requested path is an implied path that is not a stored child in the tree - it is an in-between dir in a
+		// relpath - and it has an added child
+		Children.Reset();
+		bExists = Tree.TryGetChildren(TEXTVIEW("/Stem/D_MiddleRoot/MiddlePath"), Children,
+			EDirectoryTreeGetFlags::None);
+		CHECK_MESSAGE(TEXT("GetChildrenComplexA D_MiddleRoot, !ImpliedParent, !ImpliedChildren, !Recursive"),
+			bExists == false && Children.IsEmpty());
+		Children.Reset();
+		bExists = Tree.TryGetChildren(TEXTVIEW("/Stem/D_MiddleRoot/MiddlePath"), Children,
+			EDirectoryTreeGetFlags::Recursive);
+		CHECK_MESSAGE(TEXT("GetChildrenComplexA D_MiddleRoot, !ImpliedParent, !ImpliedChildren, Recursive"),
+			bExists == false && Children.IsEmpty());
+		Children.Reset();
+		bExists = Tree.TryGetChildren(TEXTVIEW("/Stem/D_MiddleRoot/MiddlePath"), Children,
+			EDirectoryTreeGetFlags::ImpliedChildren);
+		CHECK_MESSAGE(TEXT("GetChildrenComplexA D_MiddleRoot, !ImpliedParent, ImpliedChildren, !Recursive"),
+			bExists == false && Children.IsEmpty());
+		Children.Reset();
+		bExists = Tree.TryGetChildren(TEXTVIEW("/Stem/D_MiddleRoot/MiddlePath"), Children,
+			EDirectoryTreeGetFlags::ImpliedChildren |
+			EDirectoryTreeGetFlags::Recursive);
+		CHECK_MESSAGE(TEXT("GetChildrenComplexA D_MiddleRoot, !ImpliedParent, ImpliedChildren, Recursive"),
+			bExists == false && Children.IsEmpty());
+		Children.Reset();
+		bExists = Tree.TryGetChildren(TEXTVIEW("/Stem/D_MiddleRoot/MiddlePath"), Children,
+			EDirectoryTreeGetFlags::ImpliedParent);
+		CHECK_MESSAGE(TEXT("GetChildrenComplexA D_MiddleRoot, ImpliedParent, !ImpliedChildren, !Recursive"),
+			bExists == true && UnorderedEquals(Children, { TEXT("AddedChild") }));
+		Children.Reset();
+		bExists = Tree.TryGetChildren(TEXTVIEW("/Stem/D_MiddleRoot/MiddlePath"), Children,
+			EDirectoryTreeGetFlags::ImpliedParent |
+			EDirectoryTreeGetFlags::Recursive);
+		CHECK_MESSAGE(TEXT("GetChildrenComplexA D_MiddleRoot, ImpliedParent, !ImpliedChildren, Recursive"),
+			bExists == true && UnorderedEquals(Children, { TEXT("AddedChild"), TEXT("AddedChild/Child") }));
+		Children.Reset();
+		bExists = Tree.TryGetChildren(TEXTVIEW("/Stem/D_MiddleRoot/MiddlePath"), Children,
+			EDirectoryTreeGetFlags::ImpliedParent |
+			EDirectoryTreeGetFlags::ImpliedChildren);
+		CHECK_MESSAGE(TEXT("GetChildrenComplexA D_MiddleRoot, ImpliedParent, ImpliedChildren, !Recursive"),
+			bExists == true && UnorderedEquals(Children, { TEXT("AddedChild") }));
+		Children.Reset();
+		bExists = Tree.TryGetChildren(TEXTVIEW("/Stem/D_MiddleRoot/MiddlePath"), Children,
+			EDirectoryTreeGetFlags::ImpliedParent |
+			EDirectoryTreeGetFlags::ImpliedChildren |
+			EDirectoryTreeGetFlags::Recursive);
+		CHECK_MESSAGE(TEXT("GetChildrenComplexA D_MiddleRoot, ImpliedParent, ImpliedChildren, Recursive"),
+			bExists == true && UnorderedEquals(Children, { TEXT("AddedChild"), TEXT("AddedChild/Child") }));
+
+		// Case: Requested path is an added path and it has an added child and an implied child
+		Children.Reset();
+		bExists = Tree.TryGetChildren(TEXTVIEW("/Stem/E_AddedRoot"), Children,
+			EDirectoryTreeGetFlags::None);
+		CHECK_MESSAGE(TEXT("GetChildrenComplexA E_AddedRoot, !ImpliedParent, !ImpliedChildren, !Recursive"),
+			bExists == true && UnorderedEquals(Children, { TEXT("AddedChild"), TEXT("ImpliedChild/AddedChild")}));
+		Children.Reset();
+		bExists = Tree.TryGetChildren(TEXTVIEW("/Stem/E_AddedRoot"), Children,
+			EDirectoryTreeGetFlags::Recursive);
+		CHECK_MESSAGE(TEXT("GetChildrenComplexA E_AddedRoot, !ImpliedParent, !ImpliedChildren, Recursive"),
+			bExists == true && UnorderedEquals(Children, { TEXT("AddedChild"), TEXT("ImpliedChild/AddedChild"),
+				TEXT("ImpliedChild/AddedChild/ImpliedChild/AddedChild") }));
+		Children.Reset();
+		bExists = Tree.TryGetChildren(TEXTVIEW("/Stem/E_AddedRoot"), Children,
+			EDirectoryTreeGetFlags::ImpliedChildren);
+		CHECK_MESSAGE(TEXT("GetChildrenComplexA E_AddedRoot, !ImpliedParent, ImpliedChildren, !Recursive"),
+			bExists == true && UnorderedEquals(Children, { TEXT("AddedChild"), TEXT("ImpliedChild")}));
+		Children.Reset();
+		bExists = Tree.TryGetChildren(TEXTVIEW("/Stem/E_AddedRoot"), Children,
+			EDirectoryTreeGetFlags::ImpliedChildren |
+			EDirectoryTreeGetFlags::Recursive);
+		CHECK_MESSAGE(TEXT("GetChildrenComplexA E_MiddleRoot, !ImpliedParent, ImpliedChildren, Recursive"),
+			bExists == true && UnorderedEquals(Children, { TEXT("AddedChild"), TEXT("ImpliedChild"),
+				TEXT("ImpliedChild/AddedChild"), TEXT("ImpliedChild/AddedChild/ImpliedChild"),
+				TEXT("ImpliedChild/AddedChild/ImpliedChild/AddedChild") }));
+
+		// Case: Requested path is an added path and it has only an implied child
+		Children.Reset();
+		bExists = Tree.TryGetChildren(TEXTVIEW("/Stem/F_AddedRoot"), Children,
+			EDirectoryTreeGetFlags::None);
+		CHECK_MESSAGE(TEXT("GetChildrenComplexA F_AddedRoot, !ImpliedParent, !ImpliedChildren, !Recursive"),
+			bExists == true && UnorderedEquals(Children, { TEXT("ImpliedChild/AddedChild") }));
+		Children.Reset();
+		bExists = Tree.TryGetChildren(TEXTVIEW("/Stem/F_AddedRoot"), Children,
+			EDirectoryTreeGetFlags::Recursive);
+		CHECK_MESSAGE(TEXT("GetChildrenComplexA F_AddedRoot, !ImpliedParent, !ImpliedChildren, Recursive"),
+			bExists == true && UnorderedEquals(Children, { TEXT("ImpliedChild/AddedChild") }));
+		Children.Reset();
+		bExists = Tree.TryGetChildren(TEXTVIEW("/Stem/F_AddedRoot"), Children,
+			EDirectoryTreeGetFlags::ImpliedChildren);
+		CHECK_MESSAGE(TEXT("GetChildrenComplexA F_AddedRoot, !ImpliedParent, ImpliedChildren, !Recursive"),
+			bExists == true && UnorderedEquals(Children, { TEXT("ImpliedChild") }));
+		Children.Reset();
+		bExists = Tree.TryGetChildren(TEXTVIEW("/Stem/F_AddedRoot"), Children,
+			EDirectoryTreeGetFlags::ImpliedChildren |
+			EDirectoryTreeGetFlags::Recursive);
+		CHECK_MESSAGE(TEXT("GetChildrenComplexA F_AddedRoot, !ImpliedParent, ImpliedChildren, Recursive"),
+			bExists == true && UnorderedEquals(Children, { TEXT("ImpliedChild"),
+				TEXT("ImpliedChild/AddedChild") }));
+
+		// Case: Requesting !ImpliedChildren and !Recursive on a path with an implied child, should report
+		// the added path children of the Implied child
+		Tree.Empty();
+		Tree.FindOrAdd(TEXTVIEW("/Root/Implied1/Added1")).Value = 1;
+		Tree.FindOrAdd(TEXTVIEW("/Root/Implied1/Added2")).Value = 1;
+		Tree.FindOrAdd(TEXTVIEW("/Root/Implied2/Added")).Value = 1;
+
+		Children.Reset();
+		bExists = Tree.TryGetChildren(TEXTVIEW("/Root"), Children,
+			EDirectoryTreeGetFlags::ImpliedParent);
+		CHECK_MESSAGE(TEXT("!ImpliedChildren, !Recursive, and direct child is implied."),
+			bExists == true && UnorderedEquals(Children, { TEXT("Implied1/Added1"), TEXT("Implied1/Added2"), TEXT("Implied2/Added") }));
+	}
+}
+#endif // WITH_TESTS
\ No newline at end of file
