Fix includes

diff --git a/Engine/Source/Runtime/Core/Tests/HAL/FileManagerGenericTest.cpp b/Engine/Source/Runtime/Core/Tests/HAL/FileManagerGenericTest.cpp
--- a/Engine/Source/Runtime/Core/Tests/HAL/FileManagerGenericTest.cpp
+++ b/Engine/Source/Runtime/Core/Tests/HAL/FileManagerGenericTest.cpp
@@ -4,6 +4,7 @@
 #include "Containers/UnrealString.h"
 #include "HAL/FileManager.h"
 #include "HAL/PlatformFileManager.h"
+#include "HAL/PlatformProcess.h"
 #include "Misc/Paths.h"
 #include "Misc/StringBuilder.h"
 #include "Tests/TestHarnessAdapter.h"
