Fix a wait stat push/pop mismatch. This would result in a crash if the CSV profiler started while another thread is inside a wait
#tests tested with a cooked build on multiple projects, preflighted

diff --git a/Engine/Source/Runtime/Core/Private/ProfilingDebugging/CsvProfiler.cpp b/Engine/Source/Runtime/Core/Private/ProfilingDebugging/CsvProfiler.cpp
--- a/Engine/Source/Runtime/Core/Private/ProfilingDebugging/CsvProfiler.cpp
+++ b/Engine/Source/Runtime/Core/Private/ProfilingDebugging/CsvProfiler.cpp
@@ -2215,10 +2215,12 @@ public:
 	}
 	CSV_PROFILER_INLINE const char* PopWaitStatName()
 	{
-		check(WaitStatNameStack.Num() > 0);
-		return WaitStatNameStack.Pop();
+		if (WaitStatNameStack.Num() > 0)
+		{
+			return WaitStatNameStack.Pop();
+		}
+		return nullptr;
 	}
-
 	// Raw stat data (written from the thread)
 	TSingleProducerSingleConsumerList<FCsvTimingMarker, 256> TimingMarkers;
 	TSingleProducerSingleConsumerList<FCsvCustomStat, 256> CustomStats;
