Small write log optimization to improve perf with AutoRTFM by 5.1%.

diff --git a/Engine/Source/Runtime/Core/Private/AutoRTFM/CompilerInterface.cpp b/Engine/Source/Runtime/Core/Private/AutoRTFM/CompilerInterface.cpp
--- a/Engine/Source/Runtime/Core/Private/AutoRTFM/CompilerInterface.cpp
+++ b/Engine/Source/Runtime/Core/Private/AutoRTFM/CompilerInterface.cpp
@@ -41,6 +41,45 @@ extern "C" UE_AUTORTFM_API AUTORTFM_NO_ASAN void autortfm_record_write(void* Ptr
 	Context->RecordWrite(Ptr, Size);
 }
 
+extern "C" UE_AUTORTFM_API AUTORTFM_NO_ASAN void autortfm_record_write_1(void* Ptr)
+{
+	// check for writes to null here so we end up crashing in the user
+	// code rather than in the autortfm runtime.
+	if (UNLIKELY(nullptr == Ptr))
+	{
+		return;
+	}
+
+	FContext* Context = FContext::Get();
+	Context->RecordWrite<1>(Ptr);
+}
+
+extern "C" UE_AUTORTFM_API AUTORTFM_NO_ASAN void autortfm_record_write_2(void* Ptr)
+{
+	// check for writes to null here so we end up crashing in the user
+	// code rather than in the autortfm runtime.
+	if (UNLIKELY(nullptr == Ptr))
+	{
+		return;
+	}
+
+	FContext* Context = FContext::Get();
+	Context->RecordWrite<2>(Ptr);
+}
+
+extern "C" UE_AUTORTFM_API AUTORTFM_NO_ASAN void autortfm_record_write_4(void* Ptr)
+{
+	// check for writes to null here so we end up crashing in the user
+	// code rather than in the autortfm runtime.
+	if (UNLIKELY(nullptr == Ptr))
+	{
+		return;
+	}
+
+	FContext* Context = FContext::Get();
+	Context->RecordWrite<4>(Ptr);
+}
+
 extern "C" UE_AUTORTFM_API AUTORTFM_NO_ASAN void autortfm_record_write_8(void* Ptr)
 {
 	// check for writes to null here so we end up crashing in the user
diff --git a/Engine/Source/Runtime/Core/Private/AutoRTFM/Transaction.cpp b/Engine/Source/Runtime/Core/Private/AutoRTFM/Transaction.cpp
--- a/Engine/Source/Runtime/Core/Private/AutoRTFM/Transaction.cpp
+++ b/Engine/Source/Runtime/Core/Private/AutoRTFM/Transaction.cpp
@@ -82,17 +82,17 @@ void FTransaction::Undo()
 	for(auto Iter = WriteLog.rbegin(); Iter != WriteLog.rend(); ++Iter)
     {
 		FWriteLogEntry& Entry = *Iter;
+		void* const Original = Entry.GetOriginal();
 
         // Skip writes to our current transaction nest if we're scoped. We're about to
 		// leave so the changes don't matter. 
-        if (IsScopedTransaction() && Context->IsInnerTransactionStack(Entry.OriginalAndSize.Get()))
+        if (IsScopedTransaction() && Context->IsInnerTransactionStack(Original))
         {
             continue;
         }
 
-        void* const Original = Entry.OriginalAndSize.Get();
-        const size_t Size = Entry.OriginalAndSize.GetTopTag();
-        void* const Copy = Entry.Copy;
+		const size_t Size = Entry.GetSize();
+        void* const Copy = Entry.GetCopy();
 
 		if (UE_LOG_ACTIVE(LogAutoRTFM, Verbose))
 		{
@@ -161,13 +161,17 @@ void FTransaction::CommitNested()
     for (FWriteLogEntry& Write : WriteLog)
     {
         // Skip writes that are into our current transactions stack.
-        if (IsScopedTransaction() && Context->IsInnerTransactionStack(Write.OriginalAndSize.Get()))
+        if (IsScopedTransaction() && Context->IsInnerTransactionStack(Write.GetOriginal()))
         {
             continue;
         }
 
         Parent->WriteLog.Push(Write);
-        Parent->HitSet.Insert(Write.OriginalAndSize);
+
+		FHitSet::Key HitSetEntry(Write.GetOriginal());
+		HitSetEntry.SetTopTag(static_cast<uint16_t>(Write.GetSize()));
+
+        Parent->HitSet.Insert(HitSetEntry);
     }
 
     Parent->WriteLogBumpAllocator.Merge(MoveTemp(WriteLogBumpAllocator));
diff --git a/Engine/Source/Runtime/Core/Private/AutoRTFM/TransactionInlines.h b/Engine/Source/Runtime/Core/Private/AutoRTFM/TransactionInlines.h
--- a/Engine/Source/Runtime/Core/Private/AutoRTFM/TransactionInlines.h
+++ b/Engine/Source/Runtime/Core/Private/AutoRTFM/TransactionInlines.h
@@ -88,10 +88,9 @@ AUTORTFM_NO_ASAN UE_AUTORTFM_FORCEINLINE void FTransaction::RecordWrite(void* Lo
     // Remainder at the end of the memcpy.
     RecordWriteMaxPageSized(Address + I, Size - I);
 }
-
 template<unsigned SIZE> AUTORTFM_NO_ASAN UE_AUTORTFM_FORCEINLINE void FTransaction::RecordWrite(void* LogicalAddress)
 {
-    static_assert(SIZE <= FWriteLogBumpAllocator::MaxSize);
+    static_assert(SIZE <= 8);
 
     // If we are recording a stack address that is relative to our current
     // transactions stack location, we do not need to record the data in the
@@ -122,7 +121,7 @@ template<unsigned SIZE> AUTORTFM_NO_ASAN UE_AUTORTFM_FORCEINLINE void FTransacti
 
 	Stats.Collect<EStatsKind::NewMemoryTrackerMiss>();
 
-    RecordWriteMaxPageSized(LogicalAddress, SIZE);
+	WriteLog.Push(FWriteLogEntry::CreateSmall<SIZE>(LogicalAddress));
 }
 
 UE_AUTORTFM_FORCEINLINE void FTransaction::DidAllocate(void* LogicalAddress, const size_t Size)
diff --git a/Engine/Source/Runtime/Core/Private/AutoRTFM/WriteLog.h b/Engine/Source/Runtime/Core/Private/AutoRTFM/WriteLog.h
--- a/Engine/Source/Runtime/Core/Private/AutoRTFM/WriteLog.h
+++ b/Engine/Source/Runtime/Core/Private/AutoRTFM/WriteLog.h
@@ -10,20 +10,72 @@ namespace AutoRTFM
 {
 	using FMemoryLocation = FHitSet::Key;
 
-    struct FWriteLogEntry final
+    class FWriteLogEntry final
     {
-        FMemoryLocation OriginalAndSize;
-        void* Copy;
+		static constexpr uintptr_t IsSmallBit = 0x8000;
+
+        FMemoryLocation OriginalAndIsSmallAndSize;
+        uintptr_t Copy;
+	public:
 
         UE_AUTORTFM_FORCEINLINE FWriteLogEntry() = default;
         UE_AUTORTFM_FORCEINLINE FWriteLogEntry(const FWriteLogEntry&) = default;
         UE_AUTORTFM_FORCEINLINE FWriteLogEntry& operator=(const FWriteLogEntry&) = default;
 
         UE_AUTORTFM_FORCEINLINE explicit FWriteLogEntry(void* Original, size_t Size, void* Copy) :
-            OriginalAndSize(Original), Copy(Copy)
+			OriginalAndIsSmallAndSize(Original), Copy(reinterpret_cast<uintptr_t>(Copy))
         {
-            OriginalAndSize.SetTopTag(static_cast<uint16_t>(Size));
+			OriginalAndIsSmallAndSize.SetTopTag(static_cast<uint16_t>(Size));
         }
+
+		template<unsigned int SIZE> UE_AUTORTFM_FORCEINLINE static FWriteLogEntry CreateSmall(void* Original)
+		{
+			static_assert(SIZE <= 8);
+
+			FWriteLogEntry Entry;
+
+			Entry.OriginalAndIsSmallAndSize = Original;
+
+			// We store the 8-byte or less write log entry in the write log itself since it is small.
+			uintptr_t Temp = 0;
+			memcpy(&Temp, Original, SIZE);
+			Entry.Copy = Temp;
+
+			size_t Size = IsSmallBit | SIZE;
+
+			ASSERT(Size <= UINT16_MAX);
+
+			Entry.OriginalAndIsSmallAndSize.SetTopTag(static_cast<uint16_t>(Size));
+
+			return Entry;
+		}
+
+		UE_AUTORTFM_FORCEINLINE void* GetOriginal()
+		{
+			return OriginalAndIsSmallAndSize.Get();
+		}
+
+		UE_AUTORTFM_FORCEINLINE void* GetCopy()
+		{
+			if (IsSmall())
+			{
+				return reinterpret_cast<void*>(&Copy);
+			}
+			else
+			{
+				return reinterpret_cast<void*>(Copy);
+			}
+		}
+
+		UE_AUTORTFM_FORCEINLINE size_t GetSize() const
+		{
+			return OriginalAndIsSmallAndSize.GetTopTag() & ~IsSmallBit;
+		}
+
+		UE_AUTORTFM_FORCEINLINE bool IsSmall() const
+		{
+			return OriginalAndIsSmallAndSize.GetTopTag() & IsSmallBit;
+		}
     };
 
     class FWriteLog final
diff --git a/Engine/Source/Runtime/Core/Public/AutoRTFM/AutoRTFMConstants.h b/Engine/Source/Runtime/Core/Public/AutoRTFM/AutoRTFMConstants.h
--- a/Engine/Source/Runtime/Core/Public/AutoRTFM/AutoRTFMConstants.h
+++ b/Engine/Source/Runtime/Core/Public/AutoRTFM/AutoRTFMConstants.h
@@ -13,7 +13,7 @@ namespace Constants
 {
 
 	inline constexpr uint32_t Major = 0;
-	inline constexpr uint32_t Minor = 1;
+	inline constexpr uint32_t Minor = 2;
 	inline constexpr uint32_t Patch = 0;
 
 } // namespace Constants
