[TaskGraph/Pipe]
  - Fix use-after-free that could happen when launching tasks with prerequisites as they were not known to the pipe until the prerequisites were fulfilled
  - Fix HasWork() and FPipe destruction to work correctly with tasks that have prerequisites
  - Fix task leak when calling WaitUntilEmpty with a timeout repeatedly.

diff --git a/Engine/Source/Runtime/Core/Private/Tasks/Pipe.cpp b/Engine/Source/Runtime/Core/Private/Tasks/Pipe.cpp
--- a/Engine/Source/Runtime/Core/Private/Tasks/Pipe.cpp
+++ b/Engine/Source/Runtime/Core/Private/Tasks/Pipe.cpp
@@ -30,7 +30,7 @@ namespace UE::Tasks
 		return LastTask_Local; // transfer the reference to the caller that must release it
 	}
 
-	bool FPipe::TryClearTask(Private::FTaskBase& Task)
+	void FPipe::ClearTask(Private::FTaskBase& Task)
 	{
 		Private::FTaskBase* Task_Local = &Task;
 		// try clearing the task if it's still pipe's "last task". if succeeded, release the ref accounted for pipe's last task. otherwise whoever replaced it
@@ -38,54 +38,50 @@ namespace UE::Tasks
 		if (LastTask.compare_exchange_strong(Task_Local, nullptr, std::memory_order_acquire, std::memory_order_relaxed))
 		{
 			Task.Release(); // it was still pipe's last task. now that we cleared it, release the reference
-			return true;
 		}
 
-		return false;
+		if (TaskCount.fetch_sub(1, std::memory_order_relaxed) == 1)
+		{
+			EmptyEvent.Notify();
+		}
 	}
 
-	bool FPipe::WaitUntilEmpty(FTimespan Timeout/* = FTimespan::MaxValue()*/)
+	bool FPipe::WaitUntilEmpty(FTimespan InTimeout/* = FTimespan::MaxValue()*/)
 	{
-		struct FPlaceholderTask final : Private::FTaskBase
+		if (TaskCount.load(std::memory_order_relaxed) == 0)
+		{
+			return true;
+		}
+
+		TRACE_CPUPROFILER_EVENT_SCOPE(FPipe::WaitUntilEmpty);
+
+		UE::FTimeout Timeout(InTimeout);
+		while (true)
 		{
-			FPlaceholderTask()
-				// InitRefCount: one for the initial local ref, and one for the internal reference
-				: Private::FTaskBase(/*InitRefCount =*/ 2)
+			if (TaskCount.load(std::memory_order_relaxed) == 0)
 			{
-				Init(
-					TEXT("Pipe::WaitUntilEmpty() placeholder"),
-					ETaskPriority::Normal, // doesn't matter
-					EExtendedTaskPriority::TaskEvent, // no need for scheduling or execution for this dummy task
-					ETaskFlags::None
-				);
+				return true;
 			}
 
-			virtual void ExecuteTask() override final
+			if (Timeout.IsExpired())
 			{
-				checkNoEntry(); // the method won't be called because the task was initialized with `EExtendedTaskPriority::TaskEvent`
+				return false;
 			}
-		};
 
-		bool bRes = true;
+			UE::FEventCountToken Token = EmptyEvent.PrepareWait();
 
-		TRefCountPtr<FPlaceholderTask> PlaceholderTask{ new FPlaceholderTask, /*bAddRef =*/ false }; // the initial ref was already accounted for
-		Private::FTaskBase* LastTask_Local = PushIntoPipe(*PlaceholderTask);
-		if (LastTask_Local)
-		{
-			LastTask_Local->Release(); // release pipe's ref to the last task
-			// when the placeholder replaced "the last task" in the pipe, it became its subsequent, so the last task will "complete" it
-			bRes = PlaceholderTask->Wait(FTimeout{ Timeout });
-		}
-		else
-		{
-			// there was no "last task" in the pipe, so we have to launch the placeholder to complete it to not assert on its destruction
-			PlaceholderTask->TryLaunch(/*TaskSize =*/ 0);
-			checkSlow(PlaceholderTask->IsCompleted());
+			if (TaskCount.load(std::memory_order_relaxed) == 0)
+			{
+				return true;
+			}
+
+			if (!EmptyEvent.WaitFor(Token, UE::FMonotonicTimeSpan::FromMilliseconds(Timeout.GetRemainingRoundedUpMilliseconds())))
+			{
+				break;
+			}
 		}
 
-		verifyf(TryClearTask(*PlaceholderTask), TEXT("More tasks were launched (concurrently) in the pipe after `WaitUntilEmpty()` was called"));
-		
-		return bRes;
+		return false;
 	}
 
 	// Maintains pipe callstack. Due to busy waiting tasks from multiple pipes can be executed nested.
@@ -131,4 +127,4 @@ namespace UE::Tasks
 	{
 		return FPipeCallStack::IsOnTop(*this);
 	}
-}
\ No newline at end of file
+}
diff --git a/Engine/Source/Runtime/Core/Private/Tasks/TaskPrivate.cpp b/Engine/Source/Runtime/Core/Private/Tasks/TaskPrivate.cpp
--- a/Engine/Source/Runtime/Core/Private/Tasks/TaskPrivate.cpp
+++ b/Engine/Source/Runtime/Core/Private/Tasks/TaskPrivate.cpp
@@ -266,7 +266,7 @@ namespace UE::Tasks
 
 		void FTaskBase::ClearPipe()
 		{
-			GetPipe()->TryClearTask(*this);
+			GetPipe()->ClearTask(*this);
 		}
 
 		static thread_local FTaskBase* CurrentTask = nullptr;
diff --git a/Engine/Source/Runtime/Core/Public/Tasks/Pipe.h b/Engine/Source/Runtime/Core/Public/Tasks/Pipe.h
--- a/Engine/Source/Runtime/Core/Public/Tasks/Pipe.h
+++ b/Engine/Source/Runtime/Core/Public/Tasks/Pipe.h
@@ -7,6 +7,7 @@
 #include "CoreTypes.h"
 #include "Misc/AssertionMacros.h"
 #include "Tasks/Task.h"
+#include "Async/EventCount.h"
 #include "Tasks/TaskPrivate.h"
 #include "Templates/Invoke.h"
 #include "Templates/UnrealTemplate.h"
@@ -20,7 +21,8 @@ namespace UE::Tasks
 	// non-concurrent tasks execution. FPipe is a replacement for named threads because it's lightweight and flexible -
 	// there can be a large dynamic number of pipes each controlling its own shared resource. Can be used as a replacement for
 	// dedicated threads.
-	// Execution order is FIFO, i.e. it's the same as launching order. This means that if launching order is unspecified
+	// Execution order is FIFO for tasks that don't have prerequisites, i.e. it's the same as launching order. 
+	// Adding prerequisites to a pipe task can alter when the task is queued to the pipe, hence can change the execution order.
 	// A pipe must be alive until its last task is completed.
 	// See `FTasksPipeTest` for tests and examples.
 	class FPipe
@@ -42,7 +44,7 @@ namespace UE::Tasks
 		// returns `true` if the pipe has any not completed tasks
 		bool HasWork() const
 		{
-			return LastTask.load(std::memory_order_relaxed) != nullptr;
+			return TaskCount.load(std::memory_order_relaxed) != 0;
 		}
 
 		// waits until the pipe is empty (its last task is executed)
@@ -70,8 +72,7 @@ namespace UE::Tasks
 			using FExecutableTask = Private::TExecutableTask<std::decay_t<TaskBodyType>>;
 
 			FExecutableTask* Task = FExecutableTask::Create(InDebugName, Forward<TaskBodyType>(TaskBody), Priority, ExtendedPriority, Flags);
-			Task->SetPipe(*this);
-			Task->TryLaunch(sizeof(*Task));
+			AddToPipe(Task);
 			return TTask<FResult>{ Task };
 		}
 
@@ -97,9 +98,8 @@ namespace UE::Tasks
 			using FExecutableTask = Private::TExecutableTask<std::decay_t<TaskBodyType>>;
 
 			FExecutableTask* Task = FExecutableTask::Create(InDebugName, Forward<TaskBodyType>(TaskBody), Priority, ExtendedPriority, Flags);
-			Task->SetPipe(*this);
 			Task->AddPrerequisites(Forward<PrerequisitesCollectionType>(Prerequisites));
-			Task->TryLaunch(sizeof(*Task));
+			AddToPipe(Task);
 			return TTask<FResult>{ Task };
 		}
 
@@ -116,18 +116,26 @@ namespace UE::Tasks
 		[[nodiscard]] Private::FTaskBase* PushIntoPipe(Private::FTaskBase& Task);
 
 		// pipe holds a "weak" reference to a task. the task must be cleared from the pipe when its execution finished before its completion, 
-		// otherwise the next piped task can try to add itself as a subsequent to already destroyed task.
-		// returns true if the task was still the last one and so was successfully cleared
-		bool TryClearTask(Private::FTaskBase& Task);
+		// otherwise the next piped task can try to add itself as a subsequent to an already destroyed task.
+		void ClearTask(Private::FTaskBase& Task);
 
 		// notifications about pipe's task execution
 		void ExecutionStarted();
 		void ExecutionFinished();
 
+		template <typename TaskType>
+		void AddToPipe(TaskType* Task)
+		{
+			TaskCount.fetch_add(1, std::memory_order_acq_rel);
+			Task->SetPipe(*this);
+			Task->TryLaunch(sizeof(*Task));
+		}
+
 	private:
 		// pipe builds a chain (a linked list) of tasks and so needs to store only the last one. the last task is null if the pipe is not blocked
 		std::atomic<Private::FTaskBase*> LastTask{ nullptr };
-
+		std::atomic<uint64> TaskCount { 0 };
+		UE::FEventCount EmptyEvent;
 	public:
 		FORCENOINLINE const TCHAR* GetDebugName() const
 		{
diff --git a/Engine/Source/Runtime/Core/Tests/Tasks/TasksTest.cpp b/Engine/Source/Runtime/Core/Tests/Tasks/TasksTest.cpp
--- a/Engine/Source/Runtime/Core/Tests/Tasks/TasksTest.cpp
+++ b/Engine/Source/Runtime/Core/Tests/Tasks/TasksTest.cpp
@@ -323,8 +323,8 @@ namespace UE { namespace TasksTests
 					Pipe.Launch(UE_SOURCE_LOCATION,
 						[this]
 						{
-							SuspendSignal.Trigger();
 							AddNested(ResumeSignal);
+							SuspendSignal.Trigger();
 						}
 					);
 
@@ -344,7 +344,6 @@ namespace UE { namespace TasksTests
 			FTask Task;
 			{
 				FPipeSuspensionScope Suspension(Pipe);
-				FPlatformProcess::Sleep(0.1f); // let pipe suspension finish its business, which was a bug as pipe was cleared (and unblocked) before AddNested kicked in
 				Task = Pipe.Launch(UE_SOURCE_LOCATION, [] {});
 				verify(!Task.Wait(FTimespan::FromMilliseconds(100)));
 			}
@@ -683,7 +682,37 @@ namespace UE { namespace TasksTests
 			Pipe.WaitUntilEmpty();
 		}
 
-		//for (int i = 0; i != 100000; ++i)
+		{	// waiting until a not empty pipe is empty with prereq
+			FPipe Pipe{ UE_SOURCE_LOCATION };
+			
+			FTaskEvent Prereq{ UE_SOURCE_LOCATION };
+
+			check(!Pipe.HasWork());
+
+			FTask Task1{ Pipe.Launch(UE_SOURCE_LOCATION, [] {}, Prereq) };
+
+			//Make sure the pipe knows about the task even if it has prereq
+			check(Pipe.HasWork());
+
+			check(!Task1.IsCompleted());
+
+			check(!Pipe.WaitUntilEmpty(FTimespan::FromMilliseconds(100)));
+
+			FTask Task2{ Pipe.Launch(UE_SOURCE_LOCATION, [] {}) };
+			check(Task2.Wait(FTimespan::FromMilliseconds(100)));
+
+			check(!Pipe.WaitUntilEmpty(FTimespan::FromMilliseconds(100)));
+
+			check(!Task1.IsCompleted());
+			check(Task2.IsCompleted());
+
+			Prereq.Trigger();
+
+			check(Pipe.WaitUntilEmpty(FTimespan::FromMilliseconds(100)));
+			check(Task1.IsCompleted());
+			check(Task2.IsCompleted());
+		}
+
 		{	// waiting until a not empty pipe is empty
 			FPipe Pipe{ UE_SOURCE_LOCATION };
 			Pipe.Launch(UE_SOURCE_LOCATION, [] {});
