Added hardware support runtime detection to AES.

diff --git a/Engine/Source/Runtime/Core/Private/IOS/IOSPlatformMisc.cpp b/Engine/Source/Runtime/Core/Private/IOS/IOSPlatformMisc.cpp
--- a/Engine/Source/Runtime/Core/Private/IOS/IOSPlatformMisc.cpp
+++ b/Engine/Source/Runtime/Core/Private/IOS/IOSPlatformMisc.cpp
@@ -1906,46 +1906,62 @@ void FIOSPlatformMisc::MetalAssert()
     *(int32 *)7 = 123;
 }
 
-bool FIOSPlatformMisc::CPUHasHwCrcSupport()
+struct FCPUFeatures
 {
-	// HW CRC instructions support is available on Apple A10+
-	static int HwCrcSupported = -1;
-	if (HwCrcSupported == -1)
+	// CRC instructions support is available on Apple A10 and beyond.
+	bool bHasCrc : 1 = false;
+	// AES instructions support is available on Apple A7  and beyond.
+	// A8 is minspec, so assuming it is always available.
+	bool bHasAes : 1 = true;
+
+	FCPUFeatures()
 	{
-		HwCrcSupported = 0;
-		
 		const FString DeviceIDString = GetIOSDeviceIDString();
 		if (DeviceIDString.StartsWith(TEXT("iPod")))
 		{
 			const int Major = FCString::Atoi(&DeviceIDString[4]);
-			//iPod Touch 6 and lower don't support hw CRC32
-			HwCrcSupported = Major > 7;
+			// iPod Touch 6 and lower don't support it.
+			bHasCrc = Major > 7;
 		}
 		else if (DeviceIDString.StartsWith(TEXT("iPad")))
 		{
-			// get major revision number
 			const int Major = FCString::Atoi(&DeviceIDString[4]);
-
-			//iPad 5, iPad Pro and lower
-			HwCrcSupported = Major > 6;
+			// iPad 5, iPad Pro and lower don't support it.
+			bHasCrc = Major > 6;
 		}
 		else if (DeviceIDString.StartsWith(TEXT("iPhone")))
 		{
 			const int Major = FCString::Atoi(&DeviceIDString[6]);
-			
-			// iPhone 6S, iPhone SE and below
-			HwCrcSupported = Major > 9;
+			// iPhone 6S, iPhone SE and below don't support it.
+			bHasCrc = Major > 9;
 		}
 		else if (DeviceIDString.StartsWith(TEXT("AppleTV")))
 		{
 			const int Major = FCString::Atoi(&DeviceIDString[7]);
-			
-			// Apple TV
-			HwCrcSupported = Major > 5;
+			// Apple TV 4 and lower don't support it.
+			bHasCrc = Major > 5;
+		}
+		else
+		{
+			bHasCrc = false;
 		}
 	}
+};
 
-	return HwCrcSupported == 1;
+static FCPUFeatures DetectCPUFeatures()
+{
+	static FCPUFeatures Features;
+	return Features;
+}
+
+bool FIOSPlatformMisc::CPUHasHwCrcSupport()
+{
+	return DetectCPUFeatures().bHasCrc;
+}
+
+bool FIOSPlatformMisc::CPUHasHwAesSupport()
+{
+	return DetectCPUFeatures().bHasAes;
 }
 
 static FCriticalSection EnsureLock;
diff --git a/Engine/Source/Runtime/Core/Private/Misc/AES.cpp b/Engine/Source/Runtime/Core/Private/Misc/AES.cpp
--- a/Engine/Source/Runtime/Core/Private/Misc/AES.cpp
+++ b/Engine/Source/Runtime/Core/Private/Misc/AES.cpp
@@ -3,50 +3,193 @@
 #include "Misc/AES.h"
 #include "Misc/AssertionMacros.h"
 #include "Misc/CString.h"
-#include "Misc/CoreDelegates.h"
 #include "HAL/PlatformMemory.h"
 
-#if defined(_M_AMD64) || defined(__x86_64__)
-
-#include <wmmintrin.h>
-#if PLATFORM_COMPILER_CLANG && !PLATFORM_WINDOWS
-#include <cpuid.h>
+// if x86 AESNI is enabled with compiler, use AESNI without fallback
+#if PLATFORM_CPU_X86_FAMILY
+#	include <wmmintrin.h>
+#	if defined(__AES__) || PLATFORM_ALWAYS_HAS_AESNI
+#		define DETECT_HW_AES_SUPPORT_IN_RUNTIME 0
+#	else
+#		define DETECT_HW_AES_SUPPORT_IN_RUNTIME 1
+#	endif
+// on ARM enable AES instructions unconditionally only when they are enabled with compiler
+#elif PLATFORM_CPU_ARM_FAMILY
+#	include <arm_neon.h>
+#	ifdef __ARM_FEATURE_AES
+#		define DETECT_HW_AES_SUPPORT_IN_RUNTIME 0
+#	else
+#		define DETECT_HW_AES_SUPPORT_IN_RUNTIME 1
+#	endif
+#else
+#	error Unknown CPU
 #endif
 
-#define UE_PLATFORM_AES_X86 1
+#ifndef DETECT_HW_AES_SUPPORT_IN_RUNTIME
+#	define DETECT_HW_AES_SUPPORT_IN_RUNTIME 0
+#endif
 
-// if x86 AESNI is enabled with compiler, use AESNI without fallback
-#if defined(__AES__) || PLATFORM_ALWAYS_HAS_AESNI
-#define UE_PLATFORM_AES_FALLBACK 0
+#if DETECT_HW_AES_SUPPORT_IN_RUNTIME
+#	if PLATFORM_CPU_X86_FAMILY
+#		if PLATFORM_COMPILER_CLANG
+#			include <cpuid.h>
+#		else
+#			include <intrin.h> // __cpuid
+#		endif
+#	elif PLATFORM_ANDROID_ARM64
+#		include <cpu-features.h>
+#	elif PLATFORM_LINUXARM64
+#		include <sys/auxv.h>
+#		include <asm/hwcap.h>
+#		ifndef HWCAP_AES
+#			define HWCAP_AES (1 << 3)
+#		endif
+#	elif PLATFORM_IOS
+#		include "IOS/IOSPlatformMisc.h"
+#	elif PLATFORM_WINDOWS // This one handles Windows on ARM64.
+#		include "Windows/AllowWindowsPlatformTypes.h"
+#		include <Windows.h>
+#		include <processthreadsapi.h>
+#		include "Windows/HideWindowsPlatformTypes.h"
+#	endif
+#	if !PLATFORM_WINDOWS
+#		include <signal.h>
+#		include <setjmp.h>
+#	endif
 #endif
 
-#elif defined(__aarch64__)
+#define AES256_ROUND_COUNT 14
 
-// on ARMv8 enable AES instructions unconditionally only when they are enabled with compiler
-#ifdef __ARM_FEATURE_CRYPTO
-#include <arm_neon.h>
-#define UE_PLATFORM_AES_ARMV8 1
-#define UE_PLATFORM_AES_FALLBACK 0
-#endif
+typedef void AesFunc(const uint8* Key, uint8* Contents, uint64 NumBytes);
 
-#endif
+static AesFunc AesEncryptX86HW;
+static AesFunc AesDecryptX86HW;
+static AesFunc AesEncryptArmHW;
+static AesFunc AesDecryptArmHW;
+
+static inline uint32 RotateRight(uint32 X, int N)
+{
+	return (X >> N) | (X << (32 - N));
+}
 
+#if !DETECT_HW_AES_SUPPORT_IN_RUNTIME
 
-#ifndef UE_PLATFORM_AES_FALLBACK
-#define UE_PLATFORM_AES_FALLBACK 1
+#if PLATFORM_CPU_X86_FAMILY
+constexpr static AesFunc* AesEncrypt = AesEncryptX86HW;
+constexpr static AesFunc* AesDecrypt = AesDecryptX86HW;
+#elif PLATFORM_CPU_ARM_FAMILY
+constexpr static AesFunc* AesEncrypt = AesEncryptArmHW;
+constexpr static AesFunc* AesDecrypt = AesDecryptArmHW;
+#else
+#error Unknown CPU
 #endif
-#ifndef UE_PLATFORM_AES_X86
-#define UE_PLATFORM_AES_X86 0
+
+#else
+
+static AesFunc AesEncryptSW;
+static AesFunc AesDecryptSW;
+static AesFunc AesEncryptRuntimeDispatch;
+static AesFunc AesDecryptRuntimeDispatch;
+
+// Those are set to the dispatch functions, so if FAesRuntimeHWSupportDetection
+// is gone (or static initialization is not desired), the code still works,
+// doing the implementation selection on demand.
+static AesFunc* AesEncrypt = &AesEncryptRuntimeDispatch;
+static AesFunc* AesDecrypt = &AesDecryptRuntimeDispatch;
+
+#if !PLATFORM_WINDOWS
+static sigjmp_buf AesHwJump;
+static void AesHwSignalHandler(int Sig)
+{
+	siglongjmp(AesHwJump, Sig);
+}
 #endif
-#ifndef UE_PLATFORM_AES_ARMV8
-#define UE_PLATFORM_AES_ARMV8 0
+
+static bool HasHardwareSupport()
+{
+#if PLATFORM_CPU_X86_FAMILY
+	int info[4];
+#	if PLATFORM_COMPILER_CLANG
+	__cpuid(1, info[0], info[1], info[2], info[3]);
+#	else
+	__cpuid(info, 1);
+#	endif
+	return (info[2] & 0x02000000) != 0;
+#elif PLATFORM_ANDROID_ARM64
+	return android_getCpuFeatures() & ANDROID_CPU_ARM64_FEATURE_AES;
+#elif PLATFORM_LINUXARM64
+	return getauxval(AT_HWCAP) & HWCAP_AES;
+#elif PLATFORM_IOS
+	return FPlatformMisc::CPUHasHwAesSupport();
+#elif PLATFORM_WINDOWS
+	return ::IsProcessorFeaturePresent(PF_ARM_V8_CRYPTO_INSTRUCTIONS_AVAILABLE) != 0;
+#else
+	sigset_t SignalSet;
+	sigemptyset(&SignalSet);
+	sigaddset(&SignalSet, SIGILL);
+
+	struct sigaction SigAction;
+	struct sigaction OldSigAction;
+	sigset_t OldSignalSet;
+	memset(&SigAction, 0, sizeof(SigAction));
+	SigAction.sa_handler = AesHwSignalHandler;
+	SigAction.sa_mask = SignalSet;
+
+	sigprocmask(SIG_SETMASK, &SigAction.sa_mask, &OldSignalSet);
+	sigaction(SIGILL, &SigAction, &OldSigAction);
+
+	bool bPassed = false;
+
+	if (sigsetjmp(AesHwJump, 1) == 0)
+	{
+		uint8x16_t Dummy = vaeseq_u8(vdupq_n_u8(0), vdupq_n_u8(0));
+		bPassed = true;
+	}
+
+	sigaction(SIGILL, &OldSigAction, nullptr);
+	sigprocmask(SIG_SETMASK, &OldSignalSet, nullptr);
+
+	return bPassed;
 #endif
+}
 
+static void PickImplementation()
+{
+	AesEncrypt = AesEncryptSW;
+	AesDecrypt = AesDecryptSW;
+	if (HasHardwareSupport())
+	{
+#if PLATFORM_CPU_X86_FAMILY		
+		AesEncrypt = AesEncryptX86HW;
+		AesDecrypt = AesDecryptX86HW;
+#elif PLATFORM_CPU_ARM_FAMILY
+		AesEncrypt = AesEncryptArmHW;
+		AesDecrypt = AesDecryptArmHW;
+#endif		
+	}
+}
 
-#define AES256_ROUND_COUNT 14
+static void AesEncryptRuntimeDispatch(const uint8* Key, uint8* Contents, uint64 NumBytes)
+{
+	PickImplementation();
+	check(*AesEncrypt != AesEncryptRuntimeDispatch);
+	return AesEncrypt(Key, Contents, NumBytes);
+}
 
+static void AesDecryptRuntimeDispatch(const uint8* Key, uint8* Contents, uint64 NumBytes)
+{
+	PickImplementation();
+	check(*AesDecrypt != AesDecryptRuntimeDispatch);
+	return AesDecrypt(Key, Contents, NumBytes);
+}
 
-#if UE_PLATFORM_AES_FALLBACK
+static struct FAesRuntimeHWSupportDetection
+{
+	FAesRuntimeHWSupportDetection()
+	{
+		PickImplementation();
+	}
+} AesRuntimeHWSupportDetection;
 
 struct FAesExpandedKey
 {
@@ -389,11 +532,6 @@ static inline uint32 AesDecryptMix(uint32 x)
 		 ^ (AesITbox[3][AesSbox[(uint8)(x >> 24)]]);
 }
 
-static inline uint32 RotateRight(uint32 X, int N)
-{
-	return (X >> N) | (X << (32 - N));
-}
-
 static inline void AesEncryptExpand(FAesExpandedKey* EncryptKey, const uint8* Key)
 {
 	uint32* EKey = EncryptKey->Key;
@@ -407,7 +545,7 @@ static inline void AesEncryptExpand(FAesExpandedKey* EncryptKey, const uint8* Ke
 	EKey[6] = FPlatformMemory::ReadUnaligned<uint32>(Key + 24);
 	EKey[7] = FPlatformMemory::ReadUnaligned<uint32>(Key + 28);
 
-	for (int Index = 0; Index < AES256_ROUND_COUNT/2; Index++, EKey += 8)
+	for (int Index = 0; Index < AES256_ROUND_COUNT / 2; Index++, EKey += 8)
 	{
 		// rcon[] = { 0x01, 0x02, 0x04, 0x08, 0x10, 0x20, 0x40, 0x80 }
 		uint32 Rcon = 1U << Index;
@@ -417,7 +555,7 @@ static inline void AesEncryptExpand(FAesExpandedKey* EncryptKey, const uint8* Ke
 		EKey[10] = EKey[2] ^ EKey[ 9];
 		EKey[11] = EKey[3] ^ EKey[10];
 
-		if (Index < AES256_ROUND_COUNT/2 - 1)
+		if (Index < AES256_ROUND_COUNT / 2 - 1)
 		{
 			EKey[12] = EKey[4] ^ AesEncryptMix(EKey[11]);
 			EKey[13] = EKey[5] ^ EKey[12];
@@ -458,7 +596,7 @@ static inline void AesDecryptExpand(FAesExpandedKey* DecryptKey, const uint8* Ke
 	DKey[3] = EKey[3];
 }
 
-static inline void AesEncrypt(const uint8* Key, uint8* Contents, uint64 NumBytes)
+static inline void AesEncryptSW(const uint8* Key, uint8* Contents, uint64 NumBytes)
 {
 #define AES_ENC(a, b, c, d) (      \
 	AesTbox[0][(uint8)(a >>  0)] ^ \
@@ -521,7 +659,7 @@ static inline void AesEncrypt(const uint8* Key, uint8* Contents, uint64 NumBytes
 #undef AES_ENC_LAST
 }
 
-static inline void AesDecrypt(const uint8* Key, uint8* Contents, uint64 NumBytes)
+static inline void AesDecryptSW(const uint8* Key, uint8* Contents, uint64 NumBytes)
 {
 #define AES_DEC(a, b, c, d) (       \
 	AesITbox[0][(uint8)(a >>  0)] ^ \
@@ -584,32 +722,10 @@ static inline void AesDecrypt(const uint8* Key, uint8* Contents, uint64 NumBytes
 #undef AES_DEC_LAST
 }
 
-#endif
 
-#if UE_PLATFORM_AES_X86
+#endif // DETECT_HW_AES_SUPPORT_IN_RUNTIME
 
-static inline bool DetectAesInstructions()
-{
-	int info[4];
-#if PLATFORM_COMPILER_CLANG && !PLATFORM_WINDOWS
-	__cpuid(1, info[0], info[1], info[2], info[3]);
-#else
-	__cpuid(info, 1);
-#endif
-	return (info[2] & 0x02000000) != 0;
-}
-
-static bool CanUseAesInstructions()
-{
-#if UE_PLATFORM_AES_FALLBACK
-	// run cpuid only once
-	static bool bCanUse = DetectAesInstructions();
-	return bCanUse;
-#else
-	// no AES fallback code means AESNI will be used unconditionally
-	return true;
-#endif
-}
+#if PLATFORM_CPU_X86_FAMILY
 
 // https://www.intel.com/content/dam/doc/white-paper/advanced-encryption-standard-new-instructions-set-paper.pdf
 
@@ -647,7 +763,7 @@ static bool CanUseAesInstructions()
 #pragma clang attribute push (__attribute__((target("aes"))), apply_to=function)
 #endif
 
-static inline void AesEncryptX86(const void* Key, void* Contents, uint64 NumBytes)
+static inline void AesEncryptX86HW(const uint8* Key, uint8* Contents, uint64 NumBytes)
 {
 	const __m128i* Key128 = reinterpret_cast<const __m128i*>(Key);
 	__m128i x = _mm_loadu_si128(Key128 + 0);
@@ -671,7 +787,7 @@ static inline void AesEncryptX86(const void* Key, void* Contents, uint64 NumByte
 	EKey[14] = x;
 
 	// unrolling loop to do multiple AES blocks is faster on newer CPUs
-	// for example, ~10% improvement on Zen2 threadripper (4989 -> 5501 MB/s)
+	// for example, ~10% improvement on Zen2 Threadripper (4989 -> 5501 MB/s)
 	__m128i* Data = reinterpret_cast<__m128i*>(Contents);
 	while (NumBytes >= 32)
 	{
@@ -739,7 +855,7 @@ static inline void AesEncryptX86(const void* Key, void* Contents, uint64 NumByte
 	}
 }
 
-static inline void AesDecryptX86(const void* Key, void* Contents, uint64 NumBytes)
+static inline void AesDecryptX86HW(const uint8* Key, uint8* Contents, uint64 NumBytes)
 {
 	const __m128i* Key128 = reinterpret_cast<const __m128i*>(Key);
 	__m128i x = _mm_loadu_si128(Key128 + 0);
@@ -833,77 +949,84 @@ static inline void AesDecryptX86(const void* Key, void* Contents, uint64 NumByte
 #pragma clang attribute pop
 #endif
 
-#endif
-
+#endif // PLATFORM_CPU_X86_FAMILY
 
-#if UE_PLATFORM_AES_ARMV8
+#if PLATFORM_CPU_ARM_FAMILY
 
-struct FAesExpandedKeyARMV8
+struct FAesExpandedKeyArm
 {
 	uint8x16_t Key[AES256_ROUND_COUNT + 1];
 };
 
-static inline uint32 AesEncryptMixARMV8(uint32 x)
+#if PLATFORM_COMPILER_CLANG
+#pragma clang attribute push (__attribute__((target("aes"))), apply_to=function)
+#endif
+
+// Unfortunately, ARM AES intrinsics are preprocessor-gated and not target-gated in Android NDK 25.
+// [REDACTED_PROFILE_URL]/llvm-project/issues/56480
+#if defined(__NDK_MAJOR__) && (__NDK_MAJOR__ <= 25) && !defined(__ARM_FEATURE_AES)
+#	define vaeseq_u8(p0, p1) __builtin_neon_vaeseq_v(p0, p1, 48)
+#	define vaesdq_u8(p0, p1) __builtin_neon_vaesdq_v(p0, p1, 48)
+#	define vaesimcq_u8(p0)   __builtin_neon_vaesimcq_v(p0, 48)
+#	define vaesmcq_u8(p0)    __builtin_neon_vaesmcq_v(p0, 48)
+#endif
+
+static inline uint32 AesEncryptMixArm(uint32 x)
 {
 	uint8x16_t Block = vreinterpretq_u8_u32(vmovq_n_u32(x));
 	Block = vaeseq_u8(Block, vdupq_n_u8(0));
 	return vgetq_lane_u32(vreinterpretq_u32_u8(Block), 0);
 }
 
-static inline uint32 RotateRight(uint32 X, int N)
-{
-	return (X >> N) | (X << (32 - N));
-}
-
-static inline void AesEncryptExpandARMV8(FAesExpandedKeyARMV8* EncryptKey, const uint8* Key)
+static inline void AesEncryptExpandArm(FAesExpandedKeyArm* EncryptKey, const uint8* Key)
 {
 	uint32 K[4 * (AES256_ROUND_COUNT + 1)];
 
 	uint32* EKey = K;
 	FMemory::Memcpy(EKey, Key, FAES::FAESKey::KeySize);
 
-	for (int Index = 0; Index < AES256_ROUND_COUNT/2; Index++, EKey += 8)
+	for (int Index = 0; Index < AES256_ROUND_COUNT / 2; Index++, EKey += 8)
 	{
 		// rcon[] = { 0x01, 0x02, 0x04, 0x08, 0x10, 0x20, 0x40, 0x80 }
 		uint32 Rcon = 1U << Index;
 
-		EKey[ 8] = EKey[0] ^ RotateRight(AesEncryptMixARMV8(EKey[7]), 8) ^ Rcon;
+		EKey[ 8] = EKey[0] ^ RotateRight(AesEncryptMixArm(EKey[7]), 8) ^ Rcon;
 		EKey[ 9] = EKey[1] ^ EKey[ 8];
 		EKey[10] = EKey[2] ^ EKey[ 9];
 		EKey[11] = EKey[3] ^ EKey[10];
 
-		if (Index < AES256_ROUND_COUNT/2 - 1)
+		if (Index < AES256_ROUND_COUNT / 2 - 1)
 		{
-			EKey[12] = EKey[4] ^ AesEncryptMixARMV8(EKey[11]);
+			EKey[12] = EKey[4] ^ AesEncryptMixArm(EKey[11]);
 			EKey[13] = EKey[5] ^ EKey[12];
 			EKey[14] = EKey[6] ^ EKey[13];
 			EKey[15] = EKey[7] ^ EKey[14];
 		}
 	}
 
-	for (int Round=0; Round<=AES256_ROUND_COUNT; Round++)
+	for (int Round = 0; Round <= AES256_ROUND_COUNT; Round++)
 	{
-		EncryptKey->Key[Round] = vreinterpretq_u8_u32(vld1q_u32(K + Round*4));
+		EncryptKey->Key[Round] = vreinterpretq_u8_u32(vld1q_u32(K + Round * 4));
 	}
 }
 
-static inline void AesDecryptExpandARMV8(FAesExpandedKeyARMV8* DecryptKey, const uint8* Key)
+static inline void AesDecryptExpandArm(FAesExpandedKeyArm* DecryptKey, const uint8* Key)
 {
-	FAesExpandedKeyARMV8 EncryptKey;
-	AesEncryptExpandARMV8(&EncryptKey, Key);
+	FAesExpandedKeyArm EncryptKey;
+	AesEncryptExpandArm(&EncryptKey, Key);
 
 	DecryptKey->Key[0] = EncryptKey.Key[AES256_ROUND_COUNT];
-	for (int Round=1; Round<AES256_ROUND_COUNT; Round++)
+	for (int Round = 1; Round < AES256_ROUND_COUNT; Round++)
 	{
-		DecryptKey->Key[Round] = vaesimcq_u8(EncryptKey.Key[AES256_ROUND_COUNT-Round]);
+		DecryptKey->Key[Round] = vaesimcq_u8(EncryptKey.Key[AES256_ROUND_COUNT - Round]);
 	}
 	DecryptKey->Key[AES256_ROUND_COUNT] = EncryptKey.Key[0];
 }
 
-static inline void AesEncryptARMV8(const uint8* Key, uint8* Contents, uint64 NumBytes)
+static inline void AesEncryptArmHW(const uint8* Key, uint8* Contents, uint64 NumBytes)
 {
-	FAesExpandedKeyARMV8 EncryptKey;
-	AesEncryptExpandARMV8(&EncryptKey, Key);
+	FAesExpandedKeyArm EncryptKey;
+	AesEncryptExpandArm(&EncryptKey, Key);
 
 	const uint8x16_t* EKey = EncryptKey.Key;
 
@@ -931,10 +1054,10 @@ static inline void AesEncryptARMV8(const uint8* Key, uint8* Contents, uint64 Num
 	}
 }
 
-static inline void AesDecryptARMV8(const uint8* Key, uint8* Contents, uint64 NumBytes)
+static inline void AesDecryptArmHW(const uint8* Key, uint8* Contents, uint64 NumBytes)
 {
-	FAesExpandedKeyARMV8 DecryptKey;
-	AesDecryptExpandARMV8(&DecryptKey, Key);
+	FAesExpandedKeyArm DecryptKey;
+	AesDecryptExpandArm(&DecryptKey, Key);
 
 	const uint8x16_t* DKey = DecryptKey.Key;
 
@@ -962,18 +1085,21 @@ static inline void AesDecryptARMV8(const uint8* Key, uint8* Contents, uint64 Num
 	}
 }
 
+#if PLATFORM_COMPILER_CLANG
+#pragma clang attribute pop
 #endif
 
+#endif // PLATFORM_CPU_ARM_FAMILY
 
-void FAES::EncryptData(uint8 *Contents, uint64 NumBytes, const FAESKey& Key)
+void FAES::EncryptData(uint8* Contents, uint64 NumBytes, const FAESKey& Key)
 {
 	checkf(Key.IsValid(), TEXT("No valid encryption key specified"));
-	EncryptData(Contents, NumBytes, (const uint8*)Key.Key, sizeof(Key.Key));
+	EncryptData(Contents, NumBytes, Key.Key, sizeof(Key.Key));
 }
 
-void FAES::EncryptData(uint8 *Contents, uint64 NumBytes, const ANSICHAR* Key)
+void FAES::EncryptData(uint8* Contents, uint64 NumBytes, const ANSICHAR* Key)
 {
-	checkf(Key!=nullptr, TEXT("No encryption key specified"));
+	checkf(Key != nullptr, TEXT("No encryption key specified"));
 	EncryptData(Contents, NumBytes, (const uint8*)Key, TCString<ANSICHAR>::Strlen(Key));
 }
 
@@ -981,23 +1107,7 @@ void FAES::EncryptData(uint8* Contents, uint64 NumBytes, const uint8* KeyBytes,
 {
 	checkf(NumBytes % AESBlockSize == 0, TEXT("NumBytes needs to be a multiple of 16 bytes"));
 	checkf(NumKeyBytes == FAESKey::KeySize, TEXT("AES-256 key needs to be 32 bytes long"));
-
-#if UE_PLATFORM_AES_ARMV8
-	AesEncryptARMV8(KeyBytes, Contents, NumBytes);
-	return;
-#endif
-
-#if UE_PLATFORM_AES_X86
-	if (CanUseAesInstructions())
-	{
-		AesEncryptX86(KeyBytes, Contents, NumBytes);
-		return;
-	}
-#endif
-
-#if UE_PLATFORM_AES_FALLBACK
 	AesEncrypt(KeyBytes, Contents, NumBytes);
-#endif
 }
 
 void FAES::DecryptData(uint8* Contents, uint64 NumBytes, const FAESKey& Key)
@@ -1006,9 +1116,9 @@ void FAES::DecryptData(uint8* Contents, uint64 NumBytes, const FAESKey& Key)
 	DecryptData(Contents, NumBytes, Key.Key, sizeof(Key.Key));
 }
 
-void FAES::DecryptData(uint8 *Contents, uint64 NumBytes, const ANSICHAR* Key)
+void FAES::DecryptData(uint8* Contents, uint64 NumBytes, const ANSICHAR* Key)
 {
-	checkf(Key!=nullptr, TEXT("No valid decryption key specified"));
+	checkf(Key != nullptr, TEXT("No valid decryption key specified"));
 	DecryptData(Contents, NumBytes, (const uint8*)Key, TCString<ANSICHAR>::Strlen(Key));
 }
 
@@ -1016,21 +1126,5 @@ void FAES::DecryptData(uint8* Contents, uint64 NumBytes, const uint8* KeyBytes,
 {
 	checkf(NumBytes % AESBlockSize == 0, TEXT("NumBytes needs to be a multiple of 16 bytes"));
 	checkf(NumKeyBytes == FAESKey::KeySize, TEXT("AES-256 key needs to be 32 bytes long"));
-
-#if UE_PLATFORM_AES_ARMV8
-	AesDecryptARMV8(KeyBytes, Contents, NumBytes);
-	return;
-#endif
-
-#if UE_PLATFORM_AES_X86
-	if (CanUseAesInstructions())
-	{
-		AesDecryptX86(KeyBytes, Contents, NumBytes);
-		return;
-	}
-#endif
-
-#if UE_PLATFORM_AES_FALLBACK
 	AesDecrypt(KeyBytes, Contents, NumBytes);
-#endif
-}
+}
\ No newline at end of file
diff --git a/Engine/Source/Runtime/Core/Public/IOS/IOSPlatformMisc.h b/Engine/Source/Runtime/Core/Public/IOS/IOSPlatformMisc.h
--- a/Engine/Source/Runtime/Core/Public/IOS/IOSPlatformMisc.h
+++ b/Engine/Source/Runtime/Core/Public/IOS/IOSPlatformMisc.h
@@ -214,6 +214,7 @@ struct CORE_API FIOSPlatformMisc : public FApplePlatformMisc
     static FORCENOINLINE CA_NO_RETURN void MetalAssert();
 
 	static bool CPUHasHwCrcSupport();
+	static bool CPUHasHwAesSupport();
 };
 
 typedef FIOSPlatformMisc FPlatformMisc;
