fix compilation error

diff --git a/Engine/Source/Runtime/Core/Public/Async/TaskGraphFwd.h b/Engine/Source/Runtime/Core/Public/Async/TaskGraphFwd.h
--- a/Engine/Source/Runtime/Core/Public/Async/TaskGraphFwd.h
+++ b/Engine/Source/Runtime/Core/Public/Async/TaskGraphFwd.h
@@ -6,12 +6,9 @@
 
 #pragma once
 
+#include "TaskGraphNewFrontEndDef.h"
 #include "Templates/RefCounting.h"
 
-#if !defined(TASKGRAPH_NEW_FRONTEND)
-#define TASKGRAPH_NEW_FRONTEND 0
-#endif
-
 class FBaseGraphTask;
 
 #if TASKGRAPH_NEW_FRONTEND
diff --git a/Engine/Source/Runtime/Core/Public/Async/TaskGraphNewFrontEndDef.h b/Engine/Source/Runtime/Core/Public/Async/TaskGraphNewFrontEndDef.h
new file mode 100644
--- /dev/null
+++ b/Engine/Source/Runtime/Core/Public/Async/TaskGraphNewFrontEndDef.h
@@ -0,0 +1,7 @@
+// Copyright Epic Games, Inc. All Rights Reserved.
+#pragma once
+
+#if !defined(TASKGRAPH_NEW_FRONTEND)
+#define TASKGRAPH_NEW_FRONTEND 0
+#endif
+
