Fixed PVS Studio warnings.

diff --git a/Engine/Source/Runtime/Core/Private/Misc/AES.cpp b/Engine/Source/Runtime/Core/Private/Misc/AES.cpp
--- a/Engine/Source/Runtime/Core/Private/Misc/AES.cpp
+++ b/Engine/Source/Runtime/Core/Private/Misc/AES.cpp
@@ -172,14 +172,14 @@ static void PickImplementation()
 static void AesEncryptRuntimeDispatch(const uint8* Key, uint8* Contents, uint64 NumBytes)
 {
 	PickImplementation();
-	check(*AesEncrypt != AesEncryptRuntimeDispatch);
+	check(AesEncrypt != &AesEncryptRuntimeDispatch);
 	return AesEncrypt(Key, Contents, NumBytes);
 }
 
 static void AesDecryptRuntimeDispatch(const uint8* Key, uint8* Contents, uint64 NumBytes)
 {
 	PickImplementation();
-	check(*AesDecrypt != AesDecryptRuntimeDispatch);
+	check(AesDecrypt != &AesDecryptRuntimeDispatch);
 	return AesDecrypt(Key, Contents, NumBytes);
 }
 
