Fixed truncation error causing invalid memory access when the array index times BytesBetweenElements is larger than what a signed integer can represent.

diff --git a/Engine/Source/Runtime/Core/Public/Containers/StridedView.h b/Engine/Source/Runtime/Core/Public/Containers/StridedView.h
--- a/Engine/Source/Runtime/Core/Public/Containers/StridedView.h
+++ b/Engine/Source/Runtime/Core/Public/Containers/StridedView.h
@@ -155,7 +155,7 @@ private:
 		using ByteType = typename std::conditional_t<std::is_const_v<ElementType>, const uint8, uint8>;
 
 		ByteType* AsBytes = reinterpret_cast<ByteType*>(FirstElementPtr);
-		ElementType* AsElement = reinterpret_cast<ElementType*>(AsBytes + Index * BytesBetweenElements);
+		ElementType* AsElement = reinterpret_cast<ElementType*>(AsBytes + uint64(Index) * uint64(BytesBetweenElements));
 
 		return AsElement;
 	}
