[TaskGraph]
 - Pass the thread index including the queue index to fix some checks being performed by RHI

diff --git a/Engine/Source/Runtime/Core/Public/Async/TaskGraphInterfaces.h b/Engine/Source/Runtime/Core/Public/Async/TaskGraphInterfaces.h
--- a/Engine/Source/Runtime/Core/Public/Async/TaskGraphInterfaces.h
+++ b/Engine/Source/Runtime/Core/Public/Async/TaskGraphInterfaces.h
@@ -640,9 +640,7 @@ private:
 	{
 		FGraphEventRef GraphEventRef{ this };
 		TTask* TaskObject = TaskStorage.GetTypedPtr();
-		ENamedThreads::Type ThreadIndex = ENamedThreads::GetThreadIndex(TaskObject->GetDesiredThread());
-
-		TaskObject->DoTask(ThreadIndex, GraphEventRef);
+		TaskObject->DoTask(TaskObject->GetDesiredThread(), GraphEventRef);
 		DestructItem(TaskObject);
 	}
 
