[TaskGraph]
  - Change include name so it's more future proof.

diff --git a/Engine/Source/Runtime/Core/Public/Async/TaskGraphDefinitions.h b/Engine/Source/Runtime/Core/Public/Async/TaskGraphDefinitions.h
new file mode 100644
--- /dev/null
+++ b/Engine/Source/Runtime/Core/Public/Async/TaskGraphDefinitions.h
@@ -0,0 +1,7 @@
+// Copyright Epic Games, Inc. All Rights Reserved.
+#pragma once
+
+#if !defined(TASKGRAPH_NEW_FRONTEND)
+#define TASKGRAPH_NEW_FRONTEND 0
+#endif
+
diff --git a/Engine/Source/Runtime/Core/Public/Async/TaskGraphFwd.h b/Engine/Source/Runtime/Core/Public/Async/TaskGraphFwd.h
--- a/Engine/Source/Runtime/Core/Public/Async/TaskGraphFwd.h
+++ b/Engine/Source/Runtime/Core/Public/Async/TaskGraphFwd.h
@@ -6,7 +6,7 @@
 
 #pragma once
 
-#include "TaskGraphNewFrontEndDef.h"
+#include "TaskGraphDefinitions.h"
 #include "Templates/RefCounting.h"
 
 class FBaseGraphTask;
diff --git a/Engine/Source/Runtime/Core/Public/Async/TaskGraphNewFrontEndDef.h b/Engine/Source/Runtime/Core/Public/Async/TaskGraphNewFrontEndDef.h
deleted file mode 100644
