Changed CRC hardware support detection on iOS.

It was previously done as device identifier string parsing. We use the system information query API now.

diff --git a/Engine/Source/Runtime/Core/Private/IOS/IOSPlatformMisc.cpp b/Engine/Source/Runtime/Core/Private/IOS/IOSPlatformMisc.cpp
--- a/Engine/Source/Runtime/Core/Private/IOS/IOSPlatformMisc.cpp
+++ b/Engine/Source/Runtime/Core/Private/IOS/IOSPlatformMisc.cpp
@@ -50,6 +50,8 @@
 #import <UserNotifications/UserNotifications.h>
 #endif
 
+#include <sys/sysctl.h> // sysctlbyname
+
 #include "Apple/PostAppleSystemHeaders.h"
 
 #if !defined ENABLE_ADVERTISING_IDENTIFIER
@@ -1916,35 +1918,14 @@ struct FCPUFeatures
 
 	FCPUFeatures()
 	{
-		const FString DeviceIDString = GetIOSDeviceIDString();
-		if (DeviceIDString.StartsWith(TEXT("iPod")))
-		{
-			const int Major = FCString::Atoi(&DeviceIDString[4]);
-			// iPod Touch 6 and lower don't support it.
-			bHasCrc = Major > 7;
-		}
-		else if (DeviceIDString.StartsWith(TEXT("iPad")))
-		{
-			const int Major = FCString::Atoi(&DeviceIDString[4]);
-			// iPad 5, iPad Pro and lower don't support it.
-			bHasCrc = Major > 6;
-		}
-		else if (DeviceIDString.StartsWith(TEXT("iPhone")))
-		{
-			const int Major = FCString::Atoi(&DeviceIDString[6]);
-			// iPhone 6S, iPhone SE and below don't support it.
-			bHasCrc = Major > 9;
-		}
-		else if (DeviceIDString.StartsWith(TEXT("AppleTV")))
-		{
-			const int Major = FCString::Atoi(&DeviceIDString[7]);
-			// Apple TV 4 and lower don't support it.
-			bHasCrc = Major > 5;
-		}
-		else
+		int32 Value = 0;
+		size_t Size = sizeof(Value);
+		// https://developer.apple.com/documentation/kernel/1387446-sysctlbyname/determining_instruction_set_characteristics
+		if (sysctlbyname("hw.optional.armv8_crc32", &Value, &Size, nullptr, 0) == 0)
 		{
-			bHasCrc = false;
+			bHasCrc = Value != 0;
 		}
+		// AES support could be checked by hw.optional.arm.FEAT_AES.
 	}
 };
 
