[EventCount]
 - Prevent PrepareWait from being reordered with memory operations around it
 - This fixes a potential deadlock on some platform

diff --git a/Engine/Source/Runtime/Core/Public/Async/EventCount.h b/Engine/Source/Runtime/Core/Public/Async/EventCount.h
--- a/Engine/Source/Runtime/Core/Public/Async/EventCount.h
+++ b/Engine/Source/Runtime/Core/Public/Async/EventCount.h
@@ -78,7 +78,7 @@ public:
 	inline TEventCountToken<CounterType> PrepareWait()
 	{
 		TEventCountToken<CounterType> Token;
-		Token.Value = Count.fetch_or(1, std::memory_order_relaxed) & ~CounterType(1);
+		Token.Value = Count.fetch_or(1, std::memory_order_acq_rel) & ~CounterType(1);
 		return Token;
 	}
 
