Add support for CommandLine versioning within a run so parse results can be cached safely
Cache IsHDRAllowed cmd line parse results as it is called per window per frame in slate

#tests Console RR, PIE

diff --git a/Engine/Source/Runtime/Core/Private/Misc/CommandLine.cpp b/Engine/Source/Runtime/Core/Private/Misc/CommandLine.cpp
--- a/Engine/Source/Runtime/Core/Private/Misc/CommandLine.cpp
+++ b/Engine/Source/Runtime/Core/Private/Misc/CommandLine.cpp
@@ -13,6 +13,7 @@
 -----------------------------------------------------------------------------*/
 
 bool FCommandLine::bIsInitialized = false;
+uint32 FCommandLine::CmdLineVersion = 0;
 TCHAR FCommandLine::CmdLine[FCommandLine::MaxCommandLineSize] = {};
 TCHAR FCommandLine::OriginalCmdLine[FCommandLine::MaxCommandLineSize] = {};
 TCHAR FCommandLine::LoggingCmdLine[FCommandLine::MaxCommandLineSize] = {};
@@ -29,6 +30,11 @@ bool FCommandLine::IsInitialized()
 	return bIsInitialized;
 }
 
+uint32 FCommandLine::GetCommandLineVersion()
+{
+	return CmdLineVersion;
+}
+
 void FCommandLine::Reset()
 {
 	CmdLine[0] = TEXT('\0');
@@ -36,6 +42,7 @@ void FCommandLine::Reset()
 	LoggingCmdLine[0] = TEXT('\0');
 	LoggingOriginalCmdLine[0] = TEXT('\0');
 	bIsInitialized = false;
+	CmdLineVersion++;
 }
 
 const TCHAR* FCommandLine::Get()
@@ -76,6 +83,7 @@ bool FCommandLine::Set(const TCHAR* NewCommandLine)
 	ApplyCommandLineAllowList();
 
 	bIsInitialized = true;
+	CmdLineVersion++;
 
 	// Check for the '-' that normal ones get converted to in Outlook. It's important to do it AFTER the command line is initialized
 	if (StringHasBadDashes(NewCommandLine))
@@ -95,6 +103,7 @@ bool FCommandLine::Set(const TCHAR* NewCommandLine)
 void FCommandLine::Append(const TCHAR* AppendString)
 {
 	FCString::Strncat( CmdLine, AppendString, UE_ARRAY_COUNT(CmdLine) );
+	CmdLineVersion++;
 	// If configured as part of the build, strip out any unapproved args
 	ApplyCommandLineAllowList();
 }
diff --git a/Engine/Source/Runtime/Core/Public/Misc/CommandLine.h b/Engine/Source/Runtime/Core/Public/Misc/CommandLine.h
--- a/Engine/Source/Runtime/Core/Public/Misc/CommandLine.h
+++ b/Engine/Source/Runtime/Core/Public/Misc/CommandLine.h
@@ -55,6 +55,11 @@ struct FCommandLine
 	 */
 	CORE_API static bool IsInitialized();
 
+	/**
+	 * Gets a number representing this version of the command line, incremented each command line change
+	 */
+	CORE_API static uint32 GetCommandLineVersion();
+
 	/**
 	 * Sets CmdLine to the string given
 	 */
@@ -151,5 +156,7 @@ private:
 	static TCHAR LoggingOriginalCmdLine[MaxCommandLineSize];
 	/** subprocess command line */
 	static FString& GetSubprocessCommandLine_Internal();
+	/** What version of the command line this is, incremented on any modification to the command line */
+	static uint32 CmdLineVersion;
 };
 
diff --git a/Engine/Source/Runtime/RenderCore/Private/RenderCore.cpp b/Engine/Source/Runtime/RenderCore/Private/RenderCore.cpp
--- a/Engine/Source/Runtime/RenderCore/Private/RenderCore.cpp
+++ b/Engine/Source/Runtime/RenderCore/Private/RenderCore.cpp
@@ -431,11 +431,22 @@ RENDERCORE_API bool IsHDREnabled()
 RENDERCORE_API bool IsHDRAllowed()
 {
 	// HDR can be forced on or off on the commandline. Otherwise we check the cvar r.AllowHDR
-	if (FParse::Param(FCommandLine::Get(), TEXT("hdr")))
+	static uint32 LastCommandLineVersionParsed = FCommandLine::GetCommandLineVersion();
+	static bool bYesHDR = FParse::Param(FCommandLine::Get(), TEXT("hdr"));
+	static bool bNoHDR = FParse::Param(FCommandLine::Get(), TEXT("nohdr"));
+
+	if (LastCommandLineVersionParsed != FCommandLine::GetCommandLineVersion())
+	{
+		LastCommandLineVersionParsed = FCommandLine::GetCommandLineVersion();
+		bYesHDR = FParse::Param(FCommandLine::Get(), TEXT("hdr"));
+		bNoHDR = FParse::Param(FCommandLine::Get(), TEXT("nohdr"));
+	}
+
+	if (bYesHDR)
 	{
 		return true;
 	}
-	else if (FParse::Param(FCommandLine::Get(), TEXT("nohdr")))
+	else if (bNoHDR)
 	{
 		return false;
 	}
