Add FStringView comparison to TextFilterUtils. Requires string to already be uppercased, as with FString version.

diff --git a/Engine/Source/Runtime/Core/Private/Misc/TextFilterUtils.cpp b/Engine/Source/Runtime/Core/Private/Misc/TextFilterUtils.cpp
--- a/Engine/Source/Runtime/Core/Private/Misc/TextFilterUtils.cpp
+++ b/Engine/Source/Runtime/Core/Private/Misc/TextFilterUtils.cpp
@@ -296,6 +296,11 @@ bool FTextFilterString::CompareFString(const FString& InOtherUpper, const ETextF
 	return TextFilterInternal::CompareStrings<TCHAR>(*InOtherUpper, InOtherUpper.Len(), *InternalString, InternalString.Len(), InTextComparisonMode);
 }
 
+bool FTextFilterString::CompareFStringView(FStringView InOtherUpper, const ETextFilterTextComparisonMode InTextComparisonMode) const
+{
+	return TextFilterInternal::CompareStrings<TCHAR>(InOtherUpper.GetData(), InOtherUpper.Len(), *InternalString, InternalString.Len(), InTextComparisonMode);
+}
+
 bool FTextFilterString::CompareName(const FName& InOther, const ETextFilterTextComparisonMode InTextComparisonMode) const
 {
 	TextFilterUtils::FNameBufferWithNumber OtherNameBuffer(InOther);
diff --git a/Engine/Source/Runtime/Core/Public/Misc/TextFilterUtils.h b/Engine/Source/Runtime/Core/Public/Misc/TextFilterUtils.h
--- a/Engine/Source/Runtime/Core/Public/Misc/TextFilterUtils.h
+++ b/Engine/Source/Runtime/Core/Public/Misc/TextFilterUtils.h
@@ -62,6 +62,9 @@ public:
 	/** Compare this string against the other FString, using the text comparison mode provided */
 	CORE_API bool CompareFString(const FString& InOther, const ETextFilterTextComparisonMode InTextComparisonMode) const;
 
+	/** Compare this string against a string view which has previously been converted to uppercase */
+	CORE_API bool CompareFStringView(FStringView InOther, const ETextFilterTextComparisonMode InTextComparisonMode) const;
+
 	/** Compare this string against the other FName, using the text comparison mode provided */
 	CORE_API bool CompareName(const FName& InOther, const ETextFilterTextComparisonMode InTextComparisonMode) const;
 
