Add CA_ASSUME to UTEST_NOT_NULL_EXPR (like it was added to UTEST_NOT_NULL in CL 28919862). This addresses static analysis issues arising from the null check not being visible, and then dereferencing the pointer afterwards.

diff --git a/Engine/Source/Runtime/Core/Public/Misc/AutomationTest.h b/Engine/Source/Runtime/Core/Public/Misc/AutomationTest.h
--- a/Engine/Source/Runtime/Core/Public/Misc/AutomationTest.h
+++ b/Engine/Source/Runtime/Core/Public/Misc/AutomationTest.h
@@ -4422,7 +4422,8 @@ public: \
 	if (!TestNotNull(TEXT(#Pointer), Pointer))\
 	{\
 		return false;\
-	}
+	}\
+	CA_ASSUME(Pointer);
 
 //////////////////////////////////////////////////
 // Basic Latent Commands
