Neon and FPU implementations of VectorDoubleToInt, VectorShuffleByte4, VectorCastDoubleToInt and VectorCastIntToDouble.

diff --git a/Engine/Source/Runtime/Core/Public/Math/UnrealMathFPU.h b/Engine/Source/Runtime/Core/Public/Math/UnrealMathFPU.h
--- a/Engine/Source/Runtime/Core/Public/Math/UnrealMathFPU.h
+++ b/Engine/Source/Runtime/Core/Public/Math/UnrealMathFPU.h
@@ -302,6 +302,14 @@ FORCEINLINE constexpr VectorRegister2Double MakeVectorRegister2DoubleConstant(do
 	return VectorRegister2Double { X, Y };
 }
 
+FORCEINLINE VectorRegister4Int MakeVectorRegisterInt64(int64 X, int64 Y)
+{
+	VectorRegister4Int Vec;
+	((int64&)Vec.V[0]) = X;
+	((int64&)Vec.V[2]) = Y;
+	return Vec;
+}
+
 /*=============================================================================
  *	Constants:
  *============================================================================*/
@@ -2497,6 +2505,33 @@ FORCEINLINE VectorRegister4Int VectorFloatToInt(const VectorRegister4Double& A)
 		(int32)A.V[3]);
 }
 
+FORCEINLINE VectorRegister4Int VectorDoubleToInt(const VectorRegister4Double& A)
+{
+	return VectorFloatToInt(A);
+}
+
+FORCEINLINE VectorRegister4Int VectorShuffleByte4(const VectorRegister4Int& Vec, const VectorRegister4Int& Mask)
+{
+	const int8* VecPtr = (const int8*)&Vec;
+	const uint8* MaskPtr = (const uint8*)&Mask;
+	
+	VectorRegister4Int OutVec;
+	int8* OutPtr = (int8*)&OutVec;
+	
+	constexpr uint32 NumBytes = sizeof(VectorRegister4Int);
+
+	// We use 0x10 here to denote zero fill which is compatible with both SSE, which uses >= 0x80,
+	// and Neon, which uses >= number of source bytes
+	constexpr uint32 ZeroFill = NumBytes;
+
+	for (uint32 Index = 0; Index < NumBytes; ++Index)
+	{
+		uint32 M = MaskPtr[Index];
+		OutPtr[Index] = (M < ZeroFill) ? VecPtr[M] : 0;
+	}
+	
+	return OutVec;
+}
 
 //Loads and stores
 
@@ -2580,11 +2615,43 @@ FORCEINLINE VectorRegister4Int VectorIntLoad1(const void* Ptr)
 		IntSplat);
 }
 
+/**
+* Loads 2 int64s from unaligned memory.
+*
+* @param Ptr	Unaligned memory pointer to the 2 int64s
+* @return		VectorRegister4Int(Ptr[0], Ptr[1])
+*/
+
+FORCEINLINE VectorRegister4Int VectorInt64Load(const void* Ptr)
+{
+	int64* IntPtr = (int64*)Ptr;
+	return MakeVectorRegisterInt64(
+		IntPtr[0],
+		IntPtr[1]);
+}
+
+/**
+* Loads 2 doubles from unaligned memory.
+*
+* @param Ptr	Unaligned memory pointer to the 2 doubles
+* @return		VectorRegister2Double(Ptr[0], Ptr[1])
+*/
+
+FORCEINLINE VectorRegister2Double VectorDoubleLoad(const void* Ptr)
+{
+	double* DoublePtr = (double*)Ptr;
+	return MakeVectorRegister2Double(
+		DoublePtr[0],
+		DoublePtr[1]);
+}
+
 #define VectorSetZero()								MakeVectorRegisterFloat(0.f, 0.f, 0.f, 0.f)
 #define VectorSet1(F)								VectorSetFloat1(F)
 #define VectorIntSet1(F)							MakeVectorRegisterInt(F, F, F, F)
 #define VectorCastIntToFloat(Vec)                   VectorLoad((float*)(Vec.V))
 #define VectorCastFloatToInt(Vec)                   VectorIntLoad(Vec.V)
+#define VectorCastDoubleToInt(Vec)                  VectorInt64Load(Vec.V)
+#define VectorCastIntToDouble(Vec)                  VectorDoubleLoad(Vec.V)
 #define VectorShuffleImmediate(Vec, I0, I1, I2, I3) VectorShuffle(Vec, Vec, I3, I2, I1, I0) // ShuffleImmediate shuffle is reversed from our logical shuffle.
 #define VectorShiftLeftImm(Vec, ImmAmt)             static_assert(false, "Unimplemented") // TODO: implement
 #define VectorShiftRightImmArithmetic(Vec, ImmAmt)  static_assert(false, "Unimplemented") // TODO: implement
diff --git a/Engine/Source/Runtime/Core/Public/Math/UnrealMathNeon.h b/Engine/Source/Runtime/Core/Public/Math/UnrealMathNeon.h
--- a/Engine/Source/Runtime/Core/Public/Math/UnrealMathNeon.h
+++ b/Engine/Source/Runtime/Core/Public/Math/UnrealMathNeon.h
@@ -2861,6 +2861,19 @@ FORCEINLINE VectorRegister4Int VectorFloatToInt(const VectorRegister4Double& A)
 	return VectorFloatToInt(MakeVectorRegisterFloatFromDouble(A));
 }
 
+FORCEINLINE VectorRegister4Int VectorDoubleToInt(const VectorRegister4Double& Vec)
+{
+	VectorRegister2Int64 A = vcvtq_s64_f64(Vec.XY);
+	VectorRegister2Int64 B = vcvtq_s64_f64(Vec.ZW);
+	
+	return vcombine_s32(vqmovn_s64(A), vqmovn_s64(B));
+}
+
+FORCEINLINE VectorRegister4Int VectorShuffleByte4(const VectorRegister4Int& Vec, const VectorRegister4Int& Mask)
+{
+	return vqtbl1q_u8(Vec, Mask);
+}
+
 //Loads and stores
 
 /**
@@ -2911,6 +2924,8 @@ FORCEINLINE VectorRegister4Int VectorFloatToInt(const VectorRegister4Double& A)
 #define VectorSet1(F)                               vdupq_n_f32(F)
 #define VectorCastIntToFloat(Vec)                   ((VectorRegister4f)vreinterpretq_f32_s32(Vec))
 #define VectorCastFloatToInt(Vec)					((VectorRegister4i)vreinterpretq_s32_f32(Vec))
+#define VectorCastDoubleToInt(Vec)                  ((VectorRegister2Int64)vreinterpretq_s64_f64(Vec))
+#define VectorCastIntToDouble(Vec)                  ((VectorRegister2Double)vreinterpretq_f64_s64(Vec))
 #define VectorShiftLeftImm(Vec, ImmAmt)             vshlq_n_s32(Vec, ImmAmt)
 #define VectorShiftRightImmArithmetic(Vec, ImmAmt)  vshrq_n_s32(Vec, ImmAmt)
 #define VectorShiftRightImmLogical(Vec, ImmAmt)     vshrq_n_u32(Vec, ImmAmt)
diff --git a/Engine/Source/Runtime/Core/Public/Math/UnrealMathSSE.h b/Engine/Source/Runtime/Core/Public/Math/UnrealMathSSE.h
--- a/Engine/Source/Runtime/Core/Public/Math/UnrealMathSSE.h
+++ b/Engine/Source/Runtime/Core/Public/Math/UnrealMathSSE.h
@@ -3476,7 +3476,7 @@ FORCEINLINE VectorRegister4Int VectorDoubleToInt(const VectorRegister4Double& Ve
 {
 	VectorRegister4Int A = _mm_cvttpd_epi32(Vec.GetXY());
 	VectorRegister4Int B = _mm_cvttpd_epi32(Vec.GetZW());
-	return _mm_xor_si128(A, _mm_shuffle_epi32(B, SHUFFLEMASK(2, 3, 0, 1)));
+	return _mm_unpacklo_epi64(A, B);
 }
 
 FORCEINLINE VectorRegister4Int VectorShuffleByte4(const VectorRegister4Int& Vec, const VectorRegister4Int& Mask)
