Enable all Verse code to run under AutoRTFM, but have it runtime disabled for now (via our existing CVar).

This is effectively a partial undo of CL 28855705, but with an extra check at each site.

#okforversepublic

diff --git a/Engine/Source/Runtime/Core/Private/AutoRTFM/API.cpp b/Engine/Source/Runtime/Core/Private/AutoRTFM/API.cpp
--- a/Engine/Source/Runtime/Core/Private/AutoRTFM/API.cpp
+++ b/Engine/Source/Runtime/Core/Private/AutoRTFM/API.cpp
@@ -75,6 +75,23 @@ namespace AutoRTFM
 				return false;
 			case EAutoRTFMEnabledState::AutoRTFM_Enabled:
 			case EAutoRTFMEnabledState::AutoRTFM_ForcedEnabled:
+			case EAutoRTFMEnabledState::AutoRTFM_EnabledForAllVerse:
+				return true;
+			}
+#else
+			return false;
+#endif
+		}
+
+		bool IsAutoRTFMRuntimeEnabledForAllVerse()
+		{
+			// #noop if AutoRTFM is not compiled in
+#if UE_AUTORTFM
+			switch (GAutoRTFMRuntimeEnabled)
+			{
+			default:
+				return false;
+			case EAutoRTFMEnabledState::AutoRTFM_EnabledForAllVerse:
 				return true;
 			}
 #else
diff --git a/Engine/Source/Runtime/Core/Public/AutoRTFM/AutoRTFM.h b/Engine/Source/Runtime/Core/Public/AutoRTFM/AutoRTFM.h
--- a/Engine/Source/Runtime/Core/Public/AutoRTFM/AutoRTFM.h
+++ b/Engine/Source/Runtime/Core/Public/AutoRTFM/AutoRTFM.h
@@ -584,6 +584,9 @@ namespace ForTheRuntime
 
 		// Force enable AutoRTFM - once set the AutoRTFM runtime cannot be re-enabled.
 		AutoRTFM_ForcedEnabled,
+
+		// Enable AutoRTFM for *all* Verse code (not just failure contexts).
+		AutoRTFM_EnabledForAllVerse,
 	};
 
 	// Set whether the AutoRTFM runtime is enabled or disabled.
@@ -592,6 +595,9 @@ namespace ForTheRuntime
 	// Query whether the AutoRTFM runtime is enabled.
 	UE_AUTORTFM_API bool IsAutoRTFMRuntimeEnabled();
 
+	// Query whether the AutoRTFM runtime is enabled for *all* Verse code (not just failure contexts).
+	UE_AUTORTFM_API bool IsAutoRTFMRuntimeEnabledForAllVerse();
+
 	// Manually create a new transaction from open code and push it as a transaction nest.
 	// Can only be called within an already active parent transaction (EG. this cannot start
 	// a transaction nest itself).
