Fixed comment.

#rnx

diff --git a/Engine/Source/Runtime/Core/Public/Delegates/DelegateSignatureImpl.inl b/Engine/Source/Runtime/Core/Public/Delegates/DelegateSignatureImpl.inl
--- a/Engine/Source/Runtime/Core/Public/Delegates/DelegateSignatureImpl.inl
+++ b/Engine/Source/Runtime/Core/Public/Delegates/DelegateSignatureImpl.inl
@@ -41,7 +41,7 @@ template<typename UserPolicy> class TMulticastDelegateBase;
  *
  * Example:
  *
- * struct FInterestingType
+ * struct FInterestingThing
  * {
  * public:
  *     TDelegateRegistration<void()>& OnInterestingThing() const
@@ -703,7 +703,7 @@ using TTSDelegate = TDelegate<DelegateSignature, FDefaultTSDelegateUserPolicy>;
  *
  * Example:
  *
- * struct FInterestingType
+ * struct FInterestingThing
  * {
  * public:
  *     TMulticastDelegateRegistration<void()>& OnInterestingThing() const
