Fixing a change in behavior introduced in CL 31140890 in which FParse::Token would consume trailing delimitters rather than only leading ones.

diff --git a/Engine/Source/Runtime/Core/Private/Misc/Parse.cpp b/Engine/Source/Runtime/Core/Private/Misc/Parse.cpp
--- a/Engine/Source/Runtime/Core/Private/Misc/Parse.cpp
+++ b/Engine/Source/Runtime/Core/Private/Misc/Parse.cpp
@@ -799,8 +799,9 @@ bool FParse::Token( const TCHAR*& Str, TCHAR* Result, int32 MaxLen, bool UseEsca
 {
 	int32 Len=0;
 
-	// Skip preceeding spaces and tabs.
-	while( FChar::IsWhitespace(*Str) )
+	// Skip preceeding delimiters (either spaces and tabs or custom delimiters)
+	while((SingleCharacterDelimiter == TEXT('\0') && (FChar::IsWhitespace(*Str)))
+		|| ((SingleCharacterDelimiter != TEXT('\0') && *Str == SingleCharacterDelimiter)))
 	{
 		Str++;
 	}
@@ -849,9 +850,6 @@ bool FParse::Token( const TCHAR*& Str, TCHAR* Result, int32 MaxLen, bool UseEsca
 				if ((SingleCharacterDelimiter != TEXT('\0') && Character == SingleCharacterDelimiter)
 					|| (SingleCharacterDelimiter == TEXT('\0') && (FChar::IsWhitespace(Character))))
 				{
-					// Consume the delimiter. If it's whitespace this isn't critical since we'll consume it at the start
-					// of the next call to Token() but if it's not whitespace we won't, so we better do it now.
-					Str++;
 					break;
 				}
 			}
diff --git a/Engine/Source/Runtime/Core/Tests/Misc/ParseTests.cpp b/Engine/Source/Runtime/Core/Tests/Misc/ParseTests.cpp
--- a/Engine/Source/Runtime/Core/Tests/Misc/ParseTests.cpp
+++ b/Engine/Source/Runtime/Core/Tests/Misc/ParseTests.cpp
@@ -576,5 +576,54 @@ TEST_CASE("Parse::GrammaredCLIParse::Callback", "[Smoke]")
 	}
 }
 
+TEST_CASE("Parse::Token", "[Parse][Token][Smoke]")
+{
+	const int32 BufferLen = 256;
+	TCHAR Buffer[BufferLen];
+
+	SECTION("Space Delimited")
+	{
+		const TCHAR* Line = TEXT("a=a1 b=b2 c=c3");
+		FParse::Token(Line, Buffer, BufferLen, false);
+		CHECK(FCString::Strcmp(Buffer, TEXT("a=a1"))==0);
+		CHECK_MESSAGE(TEXT("FParse::Token unexpectedly consumed trailing whitespace"), Line[0] == TEXT(' '));
+		FParse::Token(Line, Buffer, BufferLen, false);
+		CHECK(FCString::Strcmp(Buffer, TEXT("b=b2")) == 0);
+		CHECK_MESSAGE(TEXT("FParse::Token unexpectedly consumed trailing whitespace"), Line[0] == TEXT(' '));
+	}
+
+	SECTION("Custom Delimiter")
+	{
+		const TCHAR* Line = TEXT("-ini:EditorPerProjectUserSettings:[/Script/Project.Setting]:GameFeaturePluginActiveProfile=\"My Default\"");
+		FParse::Token(Line, Buffer, BufferLen, false, TEXT(':'));
+		CHECK(FCString::Strcmp(Buffer, TEXT("-ini")) == 0);
+		FParse::Token(Line, Buffer, BufferLen, false, TEXT(':'));
+		CHECK(FCString::Strcmp(Buffer, TEXT("EditorPerProjectUserSettings")) == 0);
+		FParse::Token(Line, Buffer, BufferLen, false, TEXT(':'));
+		CHECK(FCString::Strcmp(Buffer, TEXT("[/Script/Project.Setting]")) == 0);
+		FParse::Token(Line, Buffer, BufferLen, false, TEXT(':'));
+		CHECK(FCString::Strcmp(Buffer, TEXT("GameFeaturePluginActiveProfile=\"My Default\"")) == 0);
+	}
+
+	SECTION("Leading/Trailing Delimiters")
+	{
+		{
+			const TCHAR* Line = TEXT(":::Foo::Bar::");
+			FParse::Token(Line, Buffer, BufferLen, false, TEXT(':'));
+			CHECK(FCString::Strcmp(Buffer, TEXT("Foo")) == 0);
+			FParse::Token(Line, Buffer, BufferLen, false, TEXT(':'));
+			CHECK(FCString::Strcmp(Buffer, TEXT("Bar")) == 0);
+			CHECK(FParse::Token(Line, Buffer, BufferLen, false, TEXT(':')) == false);
+		}
+		{
+			const TCHAR* Line = TEXT("   Foo  Bar  ");
+			FParse::Token(Line, Buffer, BufferLen, false);
+			CHECK(FCString::Strcmp(Buffer, TEXT("Foo")) == 0);
+			FParse::Token(Line, Buffer, BufferLen, false);
+			CHECK(FCString::Strcmp(Buffer, TEXT("Bar")) == 0);
+			CHECK(FParse::Token(Line, Buffer, BufferLen, false) == false);
+		}
+	}
+}
 #endif
 	
