Add == and != operator to Interval

diff --git a/Engine/Source/Runtime/Core/Public/Math/Interval.h b/Engine/Source/Runtime/Core/Public/Math/Interval.h
--- a/Engine/Source/Runtime/Core/Public/Math/Interval.h
+++ b/Engine/Source/Runtime/Core/Public/Math/Interval.h
@@ -90,6 +90,28 @@ public:
 		}
 	}
 
+	/**
+	 * Compare two intervals for equality.
+	 *
+	 * @param Other The other interval being compared.
+	 * @return true if the intervals are equal, false otherwise.
+	 */
+	bool operator==(const TInterval& Other) const
+	{
+		return Min == Other.Min && Max == Other.Max;
+	}
+
+	/**
+	 * Compare two intervals for inequality.
+	 *
+	 * @param Other The other interval being compared.
+	 * @return true if the intervals are not equal, false otherwise.
+	 */
+	bool operator!=(const TInterval& Other) const
+	{
+		return (Min != Other.Min) || (Max != Other.Max);
+	}
+
 public:
 	
 	/**
