Make `check` macros fail correctly under AutoRTFM.

diff --git a/Engine/Source/Runtime/Core/Private/Misc/AssertionMacros.cpp b/Engine/Source/Runtime/Core/Private/Misc/AssertionMacros.cpp
--- a/Engine/Source/Runtime/Core/Private/Misc/AssertionMacros.cpp
+++ b/Engine/Source/Runtime/Core/Private/Misc/AssertionMacros.cpp
@@ -26,6 +26,7 @@
 #include "HAL/ExceptionHandling.h"
 #include "HAL/ThreadHeartBeat.h"
 #include "HAL/IConsoleManager.h"
+#include "AutoRTFM/AutoRTFM.h"
 
 namespace 
 {
@@ -618,6 +619,7 @@ bool FORCENOINLINE FDebug::CheckVerifyFailedImpl(
 #endif
 }
 
+UE_AUTORTFM_ALWAYS_OPEN
 bool FORCENOINLINE FDebug::CheckVerifyFailedImpl2(
 	const ANSICHAR* Expr,
 	const ANSICHAR* File,
