Fix stall detector trying to report crashes as stalls then crashing.
Now shuts down stall detector when crash detected.

diff --git a/Engine/Source/Runtime/Core/Private/ProfilingDebugging/StallDetector.cpp b/Engine/Source/Runtime/Core/Private/ProfilingDebugging/StallDetector.cpp
--- a/Engine/Source/Runtime/Core/Private/ProfilingDebugging/StallDetector.cpp
+++ b/Engine/Source/Runtime/Core/Private/ProfilingDebugging/StallDetector.cpp
@@ -146,6 +146,13 @@ uint32 UE::FStallDetectorRunnable::Run()
 
 			TArray<FDetectedStall> DetectedStalls;
 			FScopeLock ScopeLock(&StallScopesSection);
+
+			// Stop if there has been a crash, do not report crash as stall, Detector->Stats may also be deleted memory
+			if (!GIsCriticalError)
+			{
+				break;
+			}
+
 			for (FStallDetector* Detector : StallScopes)
 			{
 				if (Detector->bTriggered)
