Wrap FindOrAdd with OPEN and avoid un-do'ing this if we abort as we will likekly re-cache

diff --git a/Engine/Source/Runtime/Core/Private/Internationalization/TextKey.cpp b/Engine/Source/Runtime/Core/Private/Internationalization/TextKey.cpp
--- a/Engine/Source/Runtime/Core/Private/Internationalization/TextKey.cpp
+++ b/Engine/Source/Runtime/Core/Private/Internationalization/TextKey.cpp
@@ -2,6 +2,7 @@
 
 #include "Internationalization/TextKey.h"
 
+#include "AutoRTFM/AutoRTFM.h"
 #include "Containers/ContainerAllocationPolicies.h"
 #include "Containers/Map.h"
 #include "CoreGlobals.h"
@@ -24,13 +25,23 @@ public:
 	void FindOrAdd(FStringView InStr, FTextKey& OutTextKey)
 	{
 		check(!InStr.IsEmpty());
-		FindOrAddImpl(FKeyData(InStr), OutTextKey);
+		// Open around adding this in a cache, if we abort just leak the value in the cache
+		// as the cache takes ownership
+		UE_AUTORTFM_OPEN(
+		{
+			FindOrAddImpl(FKeyData(InStr), OutTextKey);
+		});
 	}
 
 	void FindOrAdd(FStringView InStr, const uint32 InStrHash, FTextKey& OutTextKey)
 	{
 		check(!InStr.IsEmpty());
-		FindOrAddImpl(FKeyData(InStr, InStrHash), OutTextKey);
+		// Open around adding this in a cache, if we abort just leak the value in the cache
+		// as the cache takes ownership
+		UE_AUTORTFM_OPEN(
+		{
+			FindOrAddImpl(FKeyData(InStr, InStrHash), OutTextKey);
+		});
 	}
 
 	void Shrink()
