Properly complete unresolved I/O requests

diff --git a/Engine/Source/Runtime/Core/Private/IO/IoDispatcher.cpp b/Engine/Source/Runtime/Core/Private/IO/IoDispatcher.cpp
--- a/Engine/Source/Runtime/Core/Private/IO/IoDispatcher.cpp
+++ b/Engine/Source/Runtime/Core/Private/IO/IoDispatcher.cpp
@@ -808,13 +808,15 @@ private:
 						Request.Backend = &Backend.Value.Get();
 					}
 					Backend.Value->ResolveIoRequests(MoveTemp(Batch), Unresolved);
-					if (Unresolved.IsEmpty())
+					Batch = MoveTemp(Unresolved);
+					if (Batch.IsEmpty())
 					{
 						break;
 					}
-					Batch = MoveTemp(Unresolved);
 				}
 
+				Unresolved = MoveTemp(Batch);
+
 				int32 UnresolvedCount = 0;
 				while (FIoRequestImpl* Request = Unresolved.PopHead())
 				{
