enabled TaskGraph unit- and stress-tests that have been disabled by mistake for a while

diff --git a/Engine/Source/Runtime/Core/Tests/Async/TaskGraphTest.cpp b/Engine/Source/Runtime/Core/Tests/Async/TaskGraphTest.cpp
--- a/Engine/Source/Runtime/Core/Tests/Async/TaskGraphTest.cpp
+++ b/Engine/Source/Runtime/Core/Tests/Async/TaskGraphTest.cpp
@@ -714,7 +714,7 @@ extern int32 GNumForegroundWorkers;
 
 namespace TaskGraphTests
 {
-	TEST_CASE_NAMED(FTaskGraphGraphEventTest, "System::Core::Async::TaskGraph::GraphEventTest", "[.][ApplicationContextMask][EngineFilter][Disabled]")
+	TEST_CASE_NAMED(FTaskGraphGraphEventTest, "System::Core::Async::TaskGraph::GraphEventTest", "[.][ApplicationContextMask][EngineFilter]")
 	{
 		{	// task completes before it's waited for
 			FGraphEventRef Event = FFunctionGraphTask::CreateAndDispatchWhenReady(
@@ -803,7 +803,7 @@ namespace TaskGraphTests
 		}
 	}
 
-	TEST_CASE_NAMED(FTaskGraphRecursionTest, "System::Core::Async::TaskGraph::RecursionTest", "[.][ApplicationContextMask][EngineFilter][Disabled]")
+	TEST_CASE_NAMED(FTaskGraphRecursionTest, "System::Core::Async::TaskGraph::RecursionTest", "[.][ApplicationContextMask][EngineFilter]")
 	{
 		{	// recursive call on game thread
 			FGraphEventRef Event = FFunctionGraphTask::CreateAndDispatchWhenReady(
@@ -829,7 +829,7 @@ namespace TaskGraphTests
 		//}
 	}
 
-	TEST_CASE_NAMED(FTaskGraphBasicTest, "System::Core::Async::TaskGraph::BasicTest", "[.][ApplicationContextMask][EngineFilter][Disabled]")
+	TEST_CASE_NAMED(FTaskGraphBasicTest, "System::Core::Async::TaskGraph::BasicTest", "[.][ApplicationContextMask][EngineFilter]")
 	{
 		// thread and task priorities
 
@@ -1279,7 +1279,7 @@ PRAGMA_ENABLE_DEPRECATION_WARNINGS
 		}
 	}
 
-	TEST_CASE_NAMED(FTaskGraphPerfTest, "System::Core::Async::TaskGraph::PerfTest", "[.][ApplicationContextMask][EngineFilter][Disabled]")
+	TEST_CASE_NAMED(FTaskGraphPerfTest, "System::Core::Async::TaskGraph::PerfTest", "[.][ApplicationContextMask][EngineFilter]")
 	{
 		TRACE_CPUPROFILER_EVENT_SCOPE(TaskGraphTests_PerfTest);
 
@@ -1329,7 +1329,7 @@ PRAGMA_ENABLE_DEPRECATION_WARNINGS
 		}
 	}
 
-	TEST_CASE_NAMED(FTaskGraphOversubscriptionTest, "System::Core::Async::TaskGraph::Oversubscription", "[.][ApplicationContextMask][EngineFilter][Disabled]")
+	TEST_CASE_NAMED(FTaskGraphOversubscriptionTest, "System::Core::Async::TaskGraph::Oversubscription", "[.][ApplicationContextMask][EngineFilter]")
 	{
 		UE_BENCHMARK(5, OversubscriptionStressTest<10>);
 	}
@@ -1361,12 +1361,12 @@ PRAGMA_ENABLE_DEPRECATION_WARNINGS
 		verify(Event->Wait(FTimespan::FromSeconds(30.f)));
 	}
 
-	TEST_CASE_NAMED(FTaskGraphSquaredOversubscriptionTest, "System::Core::Async::TaskGraph::SquaredOversubscription", "[.][ApplicationContextMask][EngineFilter][Disabled]")
+	TEST_CASE_NAMED(FTaskGraphSquaredOversubscriptionTest, "System::Core::Async::TaskGraph::SquaredOversubscription", "[.][ApplicationContextMask][EngineFilter]")
 	{
 		UE_BENCHMARK(5, SquaredOversubscriptionStressTest<10>);
 	}
 
-	TEST_CASE_NAMED(FTaskGraphTaskDestructionTest, "System::Core::Async::TaskGraph::TaskDestruction", "[.][ApplicationContextMask][EngineFilter][Disabled]")
+	TEST_CASE_NAMED(FTaskGraphTaskDestructionTest, "System::Core::Async::TaskGraph::TaskDestruction", "[.][ApplicationContextMask][EngineFilter]")
 	{
 		struct FDestructionTest
 		{
