[TaskGraphTests]
 - Remove invalid tests that were recently activated by a missing integration from another branch

diff --git a/Engine/Source/Runtime/Core/Tests/Async/TaskGraphTest.cpp b/Engine/Source/Runtime/Core/Tests/Async/TaskGraphTest.cpp
--- a/Engine/Source/Runtime/Core/Tests/Async/TaskGraphTest.cpp
+++ b/Engine/Source/Runtime/Core/Tests/Async/TaskGraphTest.cpp
@@ -803,32 +803,6 @@ namespace TaskGraphTests
 		}
 	}
 
-	TEST_CASE_NAMED(FTaskGraphRecursionTest, "System::Core::Async::TaskGraph::RecursionTest", "[.][ApplicationContextMask][EngineFilter]")
-	{
-		{	// recursive call on game thread
-			FGraphEventRef Event = FFunctionGraphTask::CreateAndDispatchWhenReady(
-				[]
-				{
-					FGraphEventRef Inner = FFunctionGraphTask::CreateAndDispatchWhenReady(
-						[]
-						{
-							check(IsInGameThread());
-						},
-						TStatId{}, nullptr, ENamedThreads::GameThread
-					);
-					Inner->Wait(ENamedThreads::GameThread);
-				},
-				TStatId{}, nullptr, ENamedThreads::GameThread
-			);
-			Event->Wait(ENamedThreads::GameThread);
-		}
-
-		//{	// didn't work in the old version
-		//	FGraphEventRef Event = FFunctionGraphTask::CreateAndDispatchWhenReady([] {}, TStatId{}, nullptr, ENamedThreads::GameThread_Local);
-		//	Event->Wait(ENamedThreads::GameThread);
-		//}
-	}
-
 	TEST_CASE_NAMED(FTaskGraphBasicTest, "System::Core::Async::TaskGraph::BasicTest", "[.][ApplicationContextMask][EngineFilter]")
 	{
 		// thread and task priorities
@@ -1334,38 +1308,6 @@ PRAGMA_ENABLE_DEPRECATION_WARNINGS
 		UE_BENCHMARK(5, OversubscriptionStressTest<10>);
 	}
 
-	template<uint32 Nujm>
-	void SquaredOversubscriptionStressTest()
-	{
-		// same as before but using two ParallelFor nested one into another to simulate oversubscription in square
-
-		FSharedEventRef Event;
-		FFunctionGraphTask::CreateAndDispatchWhenReady(
-			[Event]
-			{
-				ParallelFor(200,
-					[](int32)
-					{
-						ParallelFor(200,
-							[](int32)
-							{
-								FPlatformProcess::Sleep(0.01f); // simulate some work and let all workers to pick up ParallelFor tasks
-								FFunctionGraphTask::CreateAndDispatchWhenReady([] {})->Wait();
-							}
-						);
-					}
-				);
-				Event->Trigger();
-			}
-		);
-		verify(Event->Wait(FTimespan::FromSeconds(30.f)));
-	}
-
-	TEST_CASE_NAMED(FTaskGraphSquaredOversubscriptionTest, "System::Core::Async::TaskGraph::SquaredOversubscription", "[.][ApplicationContextMask][EngineFilter]")
-	{
-		UE_BENCHMARK(5, SquaredOversubscriptionStressTest<10>);
-	}
-
 	TEST_CASE_NAMED(FTaskGraphTaskDestructionTest, "System::Core::Async::TaskGraph::TaskDestruction", "[.][ApplicationContextMask][EngineFilter]")
 	{
 		struct FDestructionTest
