Updated MaterialTranslationDDCVersion to fix DDC issues caused by CL 32997527, where the DDC lookup for the translated shader source succeeds even though the code gen has changed.

#rnx

diff --git a/Engine/Source/Runtime/Core/Private/UObject/FortniteMainBranchObjectVersion.cpp b/Engine/Source/Runtime/Core/Private/UObject/FortniteMainBranchObjectVersion.cpp
--- a/Engine/Source/Runtime/Core/Private/UObject/FortniteMainBranchObjectVersion.cpp
+++ b/Engine/Source/Runtime/Core/Private/UObject/FortniteMainBranchObjectVersion.cpp
@@ -14,6 +14,6 @@ TMap<FGuid, FGuid> FFortniteMainBranchObjectVersion::GetSystemGuids()
 	SystemGuids.Add(DevGuids.Niagara_LatestScriptCompileVersion, FGuid("E267F600682A4E94885E8AEC5CE1DC91"));
 	SystemGuids.Add(DevGuids.SkeletalMeshDerivedDataVersion, FGuid("B9AEAA3EE7AC4FCFB3860E8F3E06DD75"));
 	SystemGuids.Add(DevGuids.STATICMESH_DERIVEDDATA_VER, FGuid("543843B103794E0E9BA4BE80FB602F79"));
-	SystemGuids.Add(DevGuids.MaterialTranslationDDCVersion, FGuid("05477E9B490F4FA9BE05B20BA3BF938C"));
+	SystemGuids.Add(DevGuids.MaterialTranslationDDCVersion, FGuid("84B1DF444194438F8379C1FCE1C5EF9B"));
 	return SystemGuids;
 }
