Add pthread_getspecific to always be opened

diff --git a/Engine/Source/Runtime/Core/Private/AutoRTFM/StdLibSupport.cpp b/Engine/Source/Runtime/Core/Private/AutoRTFM/StdLibSupport.cpp
--- a/Engine/Source/Runtime/Core/Private/AutoRTFM/StdLibSupport.cpp
+++ b/Engine/Source/Runtime/Core/Private/AutoRTFM/StdLibSupport.cpp
@@ -326,6 +326,8 @@ UE_AUTORTFM_REGISTER_SELF_FUNCTION(clock_gettime);
 UE_AUTORTFM_REGISTER_SELF_FUNCTION(gettimeofday);
 UE_AUTORTFM_REGISTER_SELF_FUNCTION(gmtime_r);
 UE_AUTORTFM_REGISTER_SELF_FUNCTION(bcmp);
+
+UE_AUTORTFM_REGISTER_SELF_FUNCTION(pthread_getspecific);
 #endif // PLATFORM_LINUX
 
 wchar_t* RTFM_wcsncpy(wchar_t* Dst, const wchar_t* Src, size_t Count)
