- Missed this file in previous checkin

diff --git a/Engine/Source/Runtime/Core/Public/HAL/Platform.h b/Engine/Source/Runtime/Core/Public/HAL/Platform.h
--- a/Engine/Source/Runtime/Core/Public/HAL/Platform.h
+++ b/Engine/Source/Runtime/Core/Public/HAL/Platform.h
@@ -166,6 +166,15 @@
 #ifndef PLATFORM_ENABLE_VECTORINTRINSICS
 	#define PLATFORM_ENABLE_VECTORINTRINSICS	0
 #endif
+
+#ifndef PLATFORM_SUPPORTS_VECTOR_ROUNDING_MODE
+	#define PLATFORM_SUPPORTS_VECTOR_ROUNDING_MODE 1
+#endif
+
+#ifndef PLATFORM_USE_SSE2_FOR_THREAD_YIELD
+	#define PLATFORM_USE_SSE2_FOR_THREAD_YIELD	PLATFORM_CPU_X86_FAMILY
+#endif
+
 // If PLATFORM_MAYBE_HAS_### is 1, then ### intrinsics are compilable.
 // This does not guarantee that the intrinsics are runnable on all instances of the platform however; a runtime check such as cpuid may be required to confirm availability.
 // If PLATFORM_ALWAYS_HAS_### is 1, then ## intrinsics will compile and run on all instances of the platform.  PLATFORM_ALWAYS_HAS_### == 1 implies PLATFORM_MAYBE_HAS_### == 1.
