Android: Resolve core.VLASize analysis warning by removing unnecessary array

#rnx

diff --git a/Engine/Source/Programs/UnrealBuildTool/Configuration/ModuleRules.cs b/Engine/Source/Programs/UnrealBuildTool/Configuration/ModuleRules.cs
--- a/Engine/Source/Programs/UnrealBuildTool/Configuration/ModuleRules.cs
+++ b/Engine/Source/Programs/UnrealBuildTool/Configuration/ModuleRules.cs
@@ -1040,9 +1040,6 @@ namespace UnrealBuildTool
 			"unix.cstring.NullArg",
 			"unix.Malloc",
 
-			// Android
-			"core.VLASize",
-
 			// Needs evalulation
 			"deadcode.DeadStores",
 		};
diff --git a/Engine/Source/Runtime/Core/Private/Android/AndroidPlatformMisc.cpp b/Engine/Source/Runtime/Core/Private/Android/AndroidPlatformMisc.cpp
--- a/Engine/Source/Runtime/Core/Private/Android/AndroidPlatformMisc.cpp
+++ b/Engine/Source/Runtime/Core/Private/Android/AndroidPlatformMisc.cpp
@@ -870,7 +870,7 @@ int32 FAndroidMisc::NumberOfCoresIncludingHyperthreads()
 static FAndroidMisc::FCPUState CurrentCPUState;
 
 FAndroidMisc::FCPUState& FAndroidMisc::GetCPUState(){
-	uint64_t UserTime, NiceTime, SystemTime, SoftIRQTime, IRQTime, IdleTime, IOWaitTime;
+	uint64_t UserTime, NiceTime, SystemTime, SoftIRQTime, IRQTime, IdleTime, WallTime, IOWaitTime;
 	int32		Index = 0;
 	ANSICHAR	Buffer[500];
 
@@ -919,12 +919,9 @@ FAndroidMisc::FCPUState& FAndroidMisc::GetCPUState(){
 		}
 		fclose(FileHandle);
 
-		uint64_t WallTime;
-		double CPULoad[CurrentCPUState.CoreCount];
 		CurrentCPUState.AverageUtilization = 0.0;
 		for (size_t n = 0; n < CurrentCPUState.CoreCount; n++) {
 			if (CurrentCPUState.CurrentUsage[n].TotalTime <= CurrentCPUState.PreviousUsage[n].TotalTime) {
-				CPULoad[n] = 0;
 				continue;
 			}
 
@@ -932,12 +929,10 @@ FAndroidMisc::FCPUState& FAndroidMisc::GetCPUState(){
 			IdleTime = CurrentCPUState.CurrentUsage[n].IdleTime - CurrentCPUState.PreviousUsage[n].IdleTime;
 
 			if (!WallTime || WallTime <= IdleTime) {
-				CPULoad[n] = 0;
 				continue;
 			}
-			CPULoad[n] = ((double)WallTime - (double)IdleTime) * 100.0 / (double)WallTime;
-			CurrentCPUState.Utilization[n] = CPULoad[n];
-			CurrentCPUState.AverageUtilization += CPULoad[n];
+			CurrentCPUState.Utilization[n] = ((double)WallTime - (double)IdleTime) * 100.0 / (double)WallTime;
+			CurrentCPUState.AverageUtilization += CurrentCPUState.Utilization[n];
 		}
 		CurrentCPUState.AverageUtilization /= (double)CurrentCPUState.CoreCount;
 	}else{
