Added new windows platform overides to the generic platform for patch check purposes.

diff --git a/Engine/Source/Runtime/Core/Private/GenericPlatform/GenericPlatformMisc.cpp b/Engine/Source/Runtime/Core/Private/GenericPlatform/GenericPlatformMisc.cpp
--- a/Engine/Source/Runtime/Core/Private/GenericPlatform/GenericPlatformMisc.cpp
+++ b/Engine/Source/Runtime/Core/Private/GenericPlatform/GenericPlatformMisc.cpp
@@ -1140,6 +1140,17 @@ TArray<FString> FGenericPlatformMisc::GetAdditionalRootDirectories()
 	return TLazySingleton<FStaticData>::Get().AdditionalRootDirectories;
 }
 
+
+const TCHAR* FGenericPlatformMisc::VersionCheckPlatformName()
+{
+	FString OverrideVersionCheckPlatformName;
+	if (!GConfig->GetString(TEXT("PatchConfig"), TEXT("PlatformOverrideName"), OverrideVersionCheckPlatformName, GEngineIni))
+	{
+		return ANSI_TO_TCHAR(FPlatformProperties::IniPlatformName());
+	}
+	return *OverrideVersionCheckPlatformName;
+}
+
 void FGenericPlatformMisc::AddAdditionalRootDirectory(const FString& RootDir)
 {
 	FRWScopeLock Lock(TLazySingleton<FStaticData>::Get().AdditionalRootDirectoriesLock, SLT_Write);
diff --git a/Engine/Source/Runtime/Core/Public/GenericPlatform/GenericPlatformMisc.h b/Engine/Source/Runtime/Core/Public/GenericPlatform/GenericPlatformMisc.h
--- a/Engine/Source/Runtime/Core/Public/GenericPlatform/GenericPlatformMisc.h
+++ b/Engine/Source/Runtime/Core/Public/GenericPlatform/GenericPlatformMisc.h
@@ -1892,6 +1892,8 @@ public:
 	 */
 	static CORE_API bool IsPGIActive();
 
+	static CORE_API const TCHAR* VersionCheckPlatformName();
+
 	/**
 	 * On platforms that support run-time PGI on/off, stops the current collection and writes out the file (or completes its write).
 	 * The filename must have been given to the StartNewPGICollection() at the start of collection (it is needed in advance since the
