Fix for compilation issue in my previous cl.

diff --git a/Engine/Source/Runtime/Core/Private/GenericPlatform/GenericPlatformMisc.cpp b/Engine/Source/Runtime/Core/Private/GenericPlatform/GenericPlatformMisc.cpp
--- a/Engine/Source/Runtime/Core/Private/GenericPlatform/GenericPlatformMisc.cpp
+++ b/Engine/Source/Runtime/Core/Private/GenericPlatform/GenericPlatformMisc.cpp
@@ -1146,7 +1146,7 @@ const TCHAR* FGenericPlatformMisc::VersionCheckPlatformName()
 	FString OverrideVersionCheckPlatformName;
 	if (!GConfig->GetString(TEXT("PatchConfig"), TEXT("PlatformOverrideName"), OverrideVersionCheckPlatformName, GEngineIni))
 	{
-		return ANSI_TO_TCHAR(FPlatformProperties::IniPlatformName());
+		OverrideVersionCheckPlatformName = ANSI_TO_TCHAR(FPlatformProperties::IniPlatformName());
 	}
 	return *OverrideVersionCheckPlatformName;
 }
