[TaskGraph]
 - Add automatic support for named threads during task waits without timeout
   - Required for a test in FTasksBasicTest only active with the new frontend
 - Always use retraction before processing the named thread queue
 - Fix named thread logic that didn't work on render thread

diff --git a/Engine/Source/Runtime/Core/Private/Async/TaskGraph.cpp b/Engine/Source/Runtime/Core/Private/Async/TaskGraph.cpp
--- a/Engine/Source/Runtime/Core/Private/Async/TaskGraph.cpp
+++ b/Engine/Source/Runtime/Core/Private/Async/TaskGraph.cpp
@@ -2076,6 +2076,25 @@ private:
 		TaskTrace::FWaitingScope WaitingScope(GetTraceIds(Tasks));
 		TRACE_CPUPROFILER_EVENT_SCOPE(WaitUntilTasksComplete);
 
+#if TASKGRAPH_NEW_FRONTEND
+		// try retraction first and only invoke named thread waiting as a last resort
+		bool bAllTasksCompleted = true;
+		for (const FGraphEventRef& Task : Tasks)
+		{
+			if (Task.IsValid())
+			{
+				Task->TryRetractAndExecute(UE::FTimeout::Never());
+
+				bAllTasksCompleted &= Task->IsCompleted();
+			}
+		}
+
+		if (bAllTasksCompleted)
+		{
+			return;
+		}
+#endif
+
 		ENamedThreads::Type CurrentThread = CurrentThreadIfKnown;
 		if (ENamedThreads::GetThreadIndex(CurrentThreadIfKnown) == ENamedThreads::AnyThread)
 		{
@@ -2118,27 +2137,6 @@ private:
 			return;
 		}
 
-#if TASKGRAPH_NEW_FRONTEND
-		// use the waiting logic of the new frontend, except for named thread tasks
-		bool HasNamedThreadTasks = false;
-		for (const FGraphEventRef& Task : Tasks)
-		{
-			if (!Task->IsNamedThreadTask())
-			{
-				Task->Wait();
-			}
-			else
-			{
-				HasNamedThreadTasks = true;
-			}
-		}
-
-		if (!HasNamedThreadTasks)
-		{
-			return;
-		}
-#endif
-
 		if (!FTaskGraphInterface::IsMultithread())
 		{
 			bool bAnyPending = false;
diff --git a/Engine/Source/Runtime/Core/Private/Tasks/TaskPrivate.cpp b/Engine/Source/Runtime/Core/Private/Tasks/TaskPrivate.cpp
--- a/Engine/Source/Runtime/Core/Private/Tasks/TaskPrivate.cpp
+++ b/Engine/Source/Runtime/Core/Private/Tasks/TaskPrivate.cpp
@@ -193,6 +193,11 @@ namespace UE::Tasks
 			return WaitImpl(Timeout);
 		}
 
+		void FTaskBase::Wait()
+		{
+			WaitWithNamedThreadsSupport();
+		}
+
 		void FTaskBase::WaitWithNamedThreadsSupport()
 		{
 			if (IsCompleted())
@@ -203,6 +208,13 @@ namespace UE::Tasks
 			TRACE_CPUPROFILER_EVENT_SCOPE(FTaskBase::WaitWithNamedThreadsSupport);
 			TaskTrace::FWaitingScope WaitingScope(GetTraceId());
 
+			TryRetractAndExecute(FTimeout::Never());
+
+			if (IsCompleted())
+			{
+				return;
+			}
+
 			if (!TryWaitOnNamedThread(*this))
 			{
 				WaitImpl(FTimeout::Never());
@@ -289,7 +301,7 @@ namespace UE::Tasks
 			// handle waiting only on a named thread and if not called from inside a task
 			FTaskGraphInterface& TaskGraph = FTaskGraphInterface::Get();
 			ENamedThreads::Type CurrentThread = TaskGraph.GetCurrentThreadIfKnown();
-			if (CurrentThread < ENamedThreads::ActualRenderingThread /* is a named thread? */ && !TaskGraph.IsThreadProcessingTasks(CurrentThread))
+			if (CurrentThread <= ENamedThreads::ActualRenderingThread /* is a named thread? */ && !TaskGraph.IsThreadProcessingTasks(CurrentThread))
 			{
 				// execute other tasks of this named thread while waiting
 				ETaskPriority Dummy;
diff --git a/Engine/Source/Runtime/Core/Public/Tasks/Task.h b/Engine/Source/Runtime/Core/Public/Tasks/Task.h
--- a/Engine/Source/Runtime/Core/Public/Tasks/Task.h
+++ b/Engine/Source/Runtime/Core/Public/Tasks/Task.h
@@ -79,11 +79,26 @@ namespace UE::Tasks
 			// waits for task's completion with timeout. Tries to retract the task and execute it in-place, if failed - blocks until the task 
 			// is completed by another thread. If timeout is zero, tries to retract the task and returns immedially after that.
 			// @return true if the task is completed
-			bool Wait(FTimespan Timeout = FTimespan::MaxValue()) const
+			bool Wait(FTimespan Timeout) const
 			{
 				return !IsValid() || Pimpl->Wait(FTimeout{ Timeout });
 			}
 
+			// waits for task's completion without timeout. Tries to retract the task and execute it in-place, if failed - blocks until the task 
+			// is completed by another thread.
+			// @return true if the task is completed
+			bool Wait() const
+			{
+				if (IsValid())
+				{
+					Pimpl->Wait();
+				}
+				
+				// we keep the return type as boolean even if always returning true
+				// as we could merge some versions of the Wait API back at some point.
+				return true;
+			}
+
 			// waits for task's completion for at least the specified amount of time, while executing other tasks.
 			// the call can return much later than the given timeout
 			// @return true if the task is completed
diff --git a/Engine/Source/Runtime/Core/Public/Tasks/TaskPrivate.h b/Engine/Source/Runtime/Core/Public/Tasks/TaskPrivate.h
--- a/Engine/Source/Runtime/Core/Public/Tasks/TaskPrivate.h
+++ b/Engine/Source/Runtime/Core/Public/Tasks/TaskPrivate.h
@@ -430,6 +430,10 @@ public:
 			// @return true if the task is completed
 			CORE_API bool Wait(FTimeout Timeout);
 
+			// waits for task's completion. Tries to retract the task and execute it in-place, if failed - blocks until the task 
+			// is completed by another thread. 
+			CORE_API void Wait();
+
 			// mimics the old tasks (TaskGraph) behaviour on named threads: waiting for a task on a named thread pulls other tasks from this
 			// named thread queue and executes them
 			CORE_API void WaitWithNamedThreadsSupport();
