Create container chunk ID from container ID

diff --git a/Engine/Source/Runtime/Core/Private/IO/IoChunkId.cpp b/Engine/Source/Runtime/Core/Private/IO/IoChunkId.cpp
--- a/Engine/Source/Runtime/Core/Private/IO/IoChunkId.cpp
+++ b/Engine/Source/Runtime/Core/Private/IO/IoChunkId.cpp
@@ -2,6 +2,7 @@
 
 #include "IO/IoChunkId.h"
 #include "IO/PackageId.h"
+#include "IO/IoContainerId.h"
 #include "Hash/Blake3.h"
 #include "Memory/MemoryView.h"
 #include "Serialization/Archive.h"
@@ -120,3 +121,8 @@ FIoChunkId CreateExternalFileChunkId(const FStringView Filename)
 
 	return ChunkId;
 }
+
+FIoChunkId CreateContainerHeaderChunkId(const FIoContainerId& ContainerId)
+{
+	return CreateIoChunkId(ContainerId.Value(), 0, EIoChunkType::ContainerHeader);
+}
diff --git a/Engine/Source/Runtime/Core/Public/IO/IoChunkId.h b/Engine/Source/Runtime/Core/Public/IO/IoChunkId.h
--- a/Engine/Source/Runtime/Core/Public/IO/IoChunkId.h
+++ b/Engine/Source/Runtime/Core/Public/IO/IoChunkId.h
@@ -12,6 +12,7 @@ class FArchive;
 class FCbFieldView;
 class FCbWriter;
 class FPackageId;
+class FIoContainerId;
 
 /**
  * Addressable chunk types.
@@ -154,3 +155,6 @@ CORE_API FIoChunkId CreatePackageDataChunkId(const FPackageId& PackageId);
 
 /** Returns a file data I/O chunk ID for the specified filename. */
 CORE_API FIoChunkId CreateExternalFileChunkId(const FStringView Filename);
+
+/** Returns a container header I/O chunk ID. */
+CORE_API FIoChunkId CreateContainerHeaderChunkId(const FIoContainerId& ContainerId);
