[Insights] Added compile time validation for the format string passed to TRACE_BOOKMARK.

diff --git a/Engine/Source/Runtime/Core/Public/ProfilingDebugging/MiscTrace.h b/Engine/Source/Runtime/Core/Public/ProfilingDebugging/MiscTrace.h
--- a/Engine/Source/Runtime/Core/Public/ProfilingDebugging/MiscTrace.h
+++ b/Engine/Source/Runtime/Core/Public/ProfilingDebugging/MiscTrace.h
@@ -64,7 +64,7 @@ struct FMiscTrace
 
 	CORE_API static void OutputBeginRegion(const TCHAR* RegionName);
 	CORE_API static void OutputEndRegion(const TCHAR* RegionName);
-	
+
 	CORE_API static void OutputBeginFrame(ETraceFrameType FrameType);
 	CORE_API static void OutputEndFrame(ETraceFrameType FrameType);
 
@@ -74,13 +74,15 @@ struct FMiscTrace
 
 private:
 	CORE_API static void OutputBookmarkInternal(const void* BookmarkPoint, uint16 EncodedFormatArgsSize, uint8* EncodedFormatArgs);
-	
 };
 
 #define TRACE_BOOKMARK(Format, ...) \
 	static bool PREPROCESSOR_JOIN(__BookmarkPoint, __LINE__); \
 	if (!PREPROCESSOR_JOIN(__BookmarkPoint, __LINE__)) \
 	{ \
+		static_assert(std::is_const_v<std::remove_reference_t<decltype(Format)>>, "Formatting string must be a const TCHAR array."); \
+		static_assert(TIsArrayOrRefOfTypeByPredicate<decltype(Format), TIsCharEncodingCompatibleWithTCHAR>::Value, "Formatting string must be a TCHAR array."); \
+		UE_VALIDATE_FORMAT_STRING(Format, ##__VA_ARGS__); \
 		FMiscTrace::OutputBookmarkSpec(&PREPROCESSOR_JOIN(__BookmarkPoint, __LINE__), __FILE__, __LINE__, Format); \
 		PREPROCESSOR_JOIN(__BookmarkPoint, __LINE__) = true; \
 	} \
