[TaskGraph]
 - Fix use-after-free
 - Add comments to prevent the same mistake from happening again

diff --git a/Engine/Source/Runtime/Core/Private/Tasks/TaskPrivate.cpp b/Engine/Source/Runtime/Core/Private/Tasks/TaskPrivate.cpp
--- a/Engine/Source/Runtime/Core/Private/Tasks/TaskPrivate.cpp
+++ b/Engine/Source/Runtime/Core/Private/Tasks/TaskPrivate.cpp
@@ -41,11 +41,15 @@ namespace UE::Tasks
 				return;
 			}
 #endif
-
-			bWakeUpWorker |= LowLevelTasks::FScheduler::Get().TryLaunch(LowLevelTask, bWakeUpWorker ? LowLevelTasks::EQueuePreference::GlobalQueuePreference : LowLevelTasks::EQueuePreference::LocalQueuePreference, bWakeUpWorker);
 			
 			// In case a thread is waiting on us to perform retraction, now is the time to try retraction again.
+			// This needs to be before the launch as performing the execution can destroy the task.
 			StateChangeEvent.Notify();
+
+			// This needs to be the last line touching any of the task's properties.
+			bWakeUpWorker |= LowLevelTasks::FScheduler::Get().TryLaunch(LowLevelTask, bWakeUpWorker ? LowLevelTasks::EQueuePreference::GlobalQueuePreference : LowLevelTasks::EQueuePreference::LocalQueuePreference, bWakeUpWorker);
+
+			// Use-after-free territory, do not touch any of the task's properties here.
 		}
 
 		thread_local uint32 TaskRetractionRecursion = 0;
diff --git a/Engine/Source/Runtime/Core/Public/Tasks/TaskPrivate.h b/Engine/Source/Runtime/Core/Public/Tasks/TaskPrivate.h
--- a/Engine/Source/Runtime/Core/Public/Tasks/TaskPrivate.h
+++ b/Engine/Source/Runtime/Core/Public/Tasks/TaskPrivate.h
@@ -611,6 +611,8 @@ public:
 						TryExecuteTask(); // result doesn't matter, this can fail if task retraction jumped in and got execution
 						// permission between this thread unlocked the task and tried to execute it
 						ReleaseInternalReference();
+
+						// Use-after-free territory, do not touch any of the task's properties here.
 					}
 					else if (ExtendedPriority == EExtendedTaskPriority::TaskEvent)
 					{
@@ -622,11 +624,15 @@ public:
 							ReleasePrerequisites();
 							Close();
 							ReleaseInternalReference();
+
+							// Use-after-free territory, do not touch any of the task's properties here.
 						}
 					}
 					else
 					{
 						Schedule(bWakeUpWorker);
+
+						// Use-after-free territory, do not touch any of the task's properties here.
 					}
 
 					return true;
@@ -642,6 +648,9 @@ public:
 				// this thread unlocked the task, no other thread can reach this point concurrently, we can touch the task again
 				Close();
 				Release(); // the internal reference that kept the task alive for nested tasks
+
+				// Use-after-free territory, do not touch any of the task's properties here.
+
 				return true;
 			}
 
