[TRACE] Allow to disable callstack channel and not pay for backtrace overhead
#tests local

diff --git a/Engine/Source/Runtime/Core/Private/Android/AndroidCallstackTrace.cpp b/Engine/Source/Runtime/Core/Private/Android/AndroidCallstackTrace.cpp
--- a/Engine/Source/Runtime/Core/Private/Android/AndroidCallstackTrace.cpp
+++ b/Engine/Source/Runtime/Core/Private/Android/AndroidCallstackTrace.cpp
@@ -145,6 +145,11 @@ void CallstackTrace_InitializeInternal()
 ////////////////////////////////////////////////////////////////////////////////
 uint32 CallstackTrace_GetCurrentId()
 {
+	if (!UE_TRACE_CHANNELEXPR_IS_ENABLED(CallstackChannel))
+	{
+		return 0;
+	}
+
 	void* ReturnAddress = PLATFORM_RETURN_ADDRESS_FOR_CALLSTACKTRACING();
 	if (FBacktracer* Instance = FBacktracer::Get())
 	{
diff --git a/Engine/Source/Runtime/Core/Private/ProfilingDebugging/Microsoft/WindowsCallstackTrace.cpp b/Engine/Source/Runtime/Core/Private/ProfilingDebugging/Microsoft/WindowsCallstackTrace.cpp
--- a/Engine/Source/Runtime/Core/Private/ProfilingDebugging/Microsoft/WindowsCallstackTrace.cpp
+++ b/Engine/Source/Runtime/Core/Private/ProfilingDebugging/Microsoft/WindowsCallstackTrace.cpp
@@ -866,6 +866,11 @@ void CallstackTrace_InitializeInternal()
 ////////////////////////////////////////////////////////////////////////////////
 uint32 CallstackTrace_GetCurrentId()
 {
+	if (!UE_TRACE_CHANNELEXPR_IS_ENABLED(CallstackChannel))
+	{
+		return 0;
+	}
+
 	void* AddressOfReturnAddress = PLATFORM_RETURN_ADDRESS_FOR_CALLSTACKTRACING();
 	if (FBacktracer* Instance = FBacktracer::Get())
 	{
diff --git a/Engine/Source/Runtime/Core/Private/ProfilingDebugging/Unix/UnixCallstackTrace.cpp b/Engine/Source/Runtime/Core/Private/ProfilingDebugging/Unix/UnixCallstackTrace.cpp
--- a/Engine/Source/Runtime/Core/Private/ProfilingDebugging/Unix/UnixCallstackTrace.cpp
+++ b/Engine/Source/Runtime/Core/Private/ProfilingDebugging/Unix/UnixCallstackTrace.cpp
@@ -123,6 +123,11 @@ void CallstackTrace_InitializeInternal()
 ////////////////////////////////////////////////////////////////////////////////
 uint32 CallstackTrace_GetCurrentId()
 {
+	if (!UE_TRACE_CHANNELEXPR_IS_ENABLED(CallstackChannel))
+	{
+		return 0;
+	}
+
 	void* ReturnAddress = PLATFORM_RETURN_ADDRESS_FOR_CALLSTACKTRACING();
 	if (FBacktracer* Instance = FBacktracer::Get())
 	{
