autortfm: implement link-time initialization of instrumented function cached lookup

this is used to work around an infinite loop that occurs in the following situation:
- foo() and bar() exist in separate compilation units
- foo() simply calls bar(), and at link time the foo and bar symbols may collapse to the same location
- the instrumented foo() and instrumented bar() do NOT collapse because their bodies are not identical - instrumented foo() contains a call to lookup instrumented bar(), whereas instrumented bar() contains the actual function body
- at runtime, we end up calling instrumented foo(), which passes the address of bar() to lookup instrumented bar(), but since foo and bar are the same at runtime, the lookup can return 'instrumented foo()' and we end up in an infinite callstack (if this becomes a tail call then it becomes an infinite loop)

to solve this, we modify the caching mechanism that we already use to accelerate the runtime lookup, where we store the results of lookup into a global variable so we don't have to walk the lookup table on every call

previously, this global variable is simply initialized to null. with this change, we employ the linker to initialize the global to the instrumented function, with appropriate hacks to not cause undefined link errors if the instrumented function doesn't exist
- on non-windows platforms, we simply use a weak symbol to the instrumented function. if the instrumented function doesn't end up existing, the linker will happily initialize the global to null
- on windows platforms we declare the instrumented function as a strong symbol but use the /alternatename linker metadata to redirect the symbol to another symbol (autortfm_llvm_missing_function) if the symbol ends up not existing. at runtime we compare the pointer to autortfm_llvm_missing_function and treat that as "null" / uninitialized

diff --git a/Engine/Source/Runtime/Core/Private/AutoRTFM/CompilerInterface.cpp b/Engine/Source/Runtime/Core/Private/AutoRTFM/CompilerInterface.cpp
--- a/Engine/Source/Runtime/Core/Private/AutoRTFM/CompilerInterface.cpp
+++ b/Engine/Source/Runtime/Core/Private/AutoRTFM/CompilerInterface.cpp
@@ -132,6 +132,14 @@ extern "C" UE_AUTORTFM_API void autortfm_llvm_fail(const char* Message)
     Context->AbortByLanguageAndThrow();
 }
 
+extern "C" UE_AUTORTFM_API void autortfm_llvm_missing_function()
+{
+	UE_LOG(LogAutoRTFM, Fatal, TEXT("Transaction failing because of missing function."));
+
+	FContext* Context = FContext::Get();
+    Context->AbortByLanguageAndThrow();
+}
+
 extern "C" UE_AUTORTFM_API void autortfm_llvm_alignment_error(FContext* Context, void* Ptr, size_t Alignment, const char* Message)
 {
     AbortDueToBadAlignment(Context, Ptr, Alignment, Message);
