UnrealPak: Fix non-deterministic .utocs by adding and zero initializing TocEntryMeta Pad bytes (follow up change for cl 33032385).

#rnx

diff --git a/Engine/Source/Runtime/Core/Internal/IO/IoStore.h b/Engine/Source/Runtime/Core/Internal/IO/IoStore.h
--- a/Engine/Source/Runtime/Core/Internal/IO/IoStore.h
+++ b/Engine/Source/Runtime/Core/Internal/IO/IoStore.h
@@ -92,7 +92,8 @@ struct FIoStoreTocEntryMeta
 {
 	// Source data hash (i.e. not the on disk data)
 	FIoHash ChunkHash;
-	FIoStoreTocEntryMetaFlags Flags;
+	FIoStoreTocEntryMetaFlags Flags = FIoStoreTocEntryMetaFlags::None;
+	uint8 Pad[3] = { 0 };
 };
 
 /**
