Multiply rotations even if the angle is very small to ensure no motion is lost

diff --git a/Engine/Source/Runtime/Core/Public/Math/TransformNonVectorized.h b/Engine/Source/Runtime/Core/Public/Math/TransformNonVectorized.h
--- a/Engine/Source/Runtime/Core/Public/Math/TransformNonVectorized.h
+++ b/Engine/Source/Runtime/Core/Public/Math/TransformNonVectorized.h
@@ -905,12 +905,7 @@ public:
 	*/
 	FORCEINLINE void Accumulate(const TTransform<T>& SourceAtom)
 	{
-		// Add ref pose relative animation to base animation, only if rotation is significant.
-		if (FMath::Square(SourceAtom.Rotation.W) < 1.f - UE_DELTA * UE_DELTA)
-		{
-			Rotation = SourceAtom.Rotation * Rotation;
-		}
-
+		Rotation = SourceAtom.Rotation * Rotation;
 		Translation += SourceAtom.Translation;
 		Scale3D *= SourceAtom.Scale3D;
 
@@ -935,12 +930,7 @@ public:
 	{
 		TTransform<T> SourceAtom(Atom * BlendWeight);
 
-		// Add ref pose relative animation to base animation, only if rotation is significant.
-		if (FMath::Square(SourceAtom.Rotation.W) < 1.f - UE_DELTA * UE_DELTA)
-		{
-			Rotation = SourceAtom.Rotation * Rotation;
-		}
-
+		Rotation = SourceAtom.Rotation * Rotation;
 		Translation += SourceAtom.Translation;
 		Scale3D *= SourceAtom.Scale3D;
 
@@ -1007,12 +997,7 @@ public:
 
 		TTransform<T> SourceAtom(Atom * BlendWeight);
 
-		// Add ref pose relative animation to base animation, only if rotation is significant.
-		if (FMath::Square(SourceAtom.Rotation.W) < 1.f - UE_DELTA * UE_DELTA)
-		{
-			Rotation = SourceAtom.Rotation * Rotation;
-		}
-
+		Rotation = SourceAtom.Rotation * Rotation;
 		Translation += SourceAtom.Translation;
 		Scale3D *= (DefaultScale + SourceAtom.Scale3D);
 
@@ -1080,12 +1065,7 @@ public:
 			DeltaAtom.Blend(AdditiveIdentity, DeltaAtom, BlendWeight);
 		}
 
-		// Add ref pose relative animation to base animation, only if rotation is significant.
-		if (FMath::Square(DeltaAtom.Rotation.W) < 1.f - UE_DELTA * UE_DELTA)
-		{
-			FinalAtom.Rotation = DeltaAtom.Rotation * FinalAtom.Rotation;
-		}
-
+		FinalAtom.Rotation = DeltaAtom.Rotation * FinalAtom.Rotation;
 		FinalAtom.Translation += DeltaAtom.Translation;
 		FinalAtom.Scale3D *= (DefaultScale + DeltaAtom.Scale3D);
 
diff --git a/Engine/Source/Runtime/Core/Public/Math/TransformVectorized.h b/Engine/Source/Runtime/Core/Public/Math/TransformVectorized.h
--- a/Engine/Source/Runtime/Core/Public/Math/TransformVectorized.h
+++ b/Engine/Source/Runtime/Core/Public/Math/TransformVectorized.h
@@ -979,15 +979,8 @@ public:
 	*/
 	FORCEINLINE void Accumulate(const TTransform<T>& SourceAtom)
 	{
-		const TransformVectorRegister BlendedRotation = SourceAtom.Rotation;
-		const TransformVectorRegister RotationW = VectorReplicate(BlendedRotation, 3);
-
-		// if( Square(SourceAtom.Rotation.W) < 1.f - DELTA * DELTA )
-		if (VectorAnyGreaterThan(GlobalVectorConstants::RotationSignificantThreshold, VectorMultiply(RotationW, RotationW)))
-		{
-			// Rotation = SourceAtom.Rotation * Rotation;
-			Rotation = VectorQuaternionMultiply2(BlendedRotation, Rotation);
-		}
+		// Rotation = SourceAtom.Rotation * Rotation;
+		Rotation = VectorQuaternionMultiply2(SourceAtom.Rotation, Rotation);
 
 		// Translation += SourceAtom.Translation;
 		// Scale *= SourceAtom.Scale;
@@ -1019,15 +1012,8 @@ public:
 		const TransformVectorRegister BlendedRotation = VectorMultiply(Atom.Rotation, BlendWeight);
 		const TransformVectorRegister BlendedScale = VectorMultiply(Atom.Scale3D, BlendWeight);
 
-		const TransformVectorRegister RotationW = VectorReplicate(BlendedRotation, 3);
-
-		// Add ref pose relative animation to base animation, only if rotation is significant.
-		// if( Square(SourceAtom.Rotation.W) < 1.f - DELTA * DELTA )
-		if (VectorAnyGreaterThan(GlobalVectorConstants::RotationSignificantThreshold, VectorMultiply(RotationW, RotationW)))
-		{
-			// Rotation = SourceAtom.Rotation * Rotation;
-			Rotation = VectorQuaternionMultiply2(BlendedRotation, Rotation);
-		}
+		// Rotation = SourceAtom.Rotation * Rotation;
+		Rotation = VectorQuaternionMultiply2(BlendedRotation, Rotation);
 
 		// Translation += SourceAtom.Translation;
 		// Scale *= SourceAtom.Scale;
@@ -1085,15 +1071,8 @@ public:
 		// SourceAtom = Atom * BlendWeight;
 		const TransformVectorRegister BlendedRotation = VectorMultiply(Atom.Rotation, BlendWeight);
 
-		const TransformVectorRegister RotationW = VectorReplicate(BlendedRotation, 3);
-
-		// Add ref pose relative animation to base animation, only if rotation is significant.
-		// if( Square(SourceAtom.Rotation.W) < 1.f - DELTA * DELTA )
-		if (VectorAnyGreaterThan(GlobalVectorConstants::RotationSignificantThreshold, VectorMultiply(RotationW, RotationW)))
-		{
-			// Rotation = SourceAtom.Rotation * Rotation;
-			Rotation = VectorQuaternionMultiply2(BlendedRotation, Rotation);
-		}
+		// Rotation = SourceAtom.Rotation * Rotation;
+		Rotation = VectorQuaternionMultiply2(BlendedRotation, Rotation);
 
 		// Translation += SourceAtom.Translation;
 		// Scale *= SourceAtom.Scale;
diff --git a/Engine/Source/Runtime/Engine/Private/Animation/AnimationRuntime.ispc b/Engine/Source/Runtime/Engine/Private/Animation/AnimationRuntime.ispc
--- a/Engine/Source/Runtime/Engine/Private/Animation/AnimationRuntime.ispc
+++ b/Engine/Source/Runtime/Engine/Private/Animation/AnimationRuntime.ispc
@@ -120,14 +120,8 @@ export void AccumulateWithAdditiveScale(uniform FTransform BasePose[],
 		const uniform FVector4 TranslationResult = Base.Translation + (Atom.Translation * BlendWeight);
 		const uniform FVector4 Scale3DResult = Base.Scale3D * ((Atom.Scale3D * BlendWeight) + DefaultScale);
 
-		const uniform float SquareRotationW = BlendedRotation.V[3] * BlendedRotation.V[3];
-
-		// Add ref pose relative animation to base animation, only if rotation is significant.
-		if(SquareRotationW < 1.f - DELTA * DELTA)
-		{
-			// Rotation = SourceAtom.Rotation * Rotation;
-			BasePose[BoneIndex].Rotation = VectorQuaternionMultiply2(BlendedRotation, Base.Rotation);
-		}
+		// Rotation = SourceAtom.Rotation * Rotation;
+		BasePose[BoneIndex].Rotation = VectorQuaternionMultiply2(BlendedRotation, Base.Rotation);
 
 		BasePose[BoneIndex].Translation = TranslationResult;
 		BasePose[BoneIndex].Scale3D = Scale3DResult;
