[Templates]
 - Deprecate dereference operator on TGuardValue/TOptionalGuardValue and favor explicit function instead

diff --git a/Engine/Source/Runtime/Core/Private/Windows/WindowsPlatformCrashContext.cpp b/Engine/Source/Runtime/Core/Private/Windows/WindowsPlatformCrashContext.cpp
--- a/Engine/Source/Runtime/Core/Private/Windows/WindowsPlatformCrashContext.cpp
+++ b/Engine/Source/Runtime/Core/Private/Windows/WindowsPlatformCrashContext.cpp
@@ -1740,7 +1740,7 @@ static void ReportEvent(ECrashContextType InType, const TCHAR* ErrorMessage, uin
 
 	// Ignore any ensure that could be fired by the code reporting an ensure.
 	TGuardValue<bool> ReentranceGuard(bReentranceGuard, true);
-	if (*ReentranceGuard) // Read the old value.
+	if (ReentranceGuard.GetOriginalValue()) // Read the old value.
 	{
 		return; // Already handling an ensure.
 	}
diff --git a/Engine/Source/Runtime/Core/Public/Templates/UnrealTemplate.h b/Engine/Source/Runtime/Core/Public/Templates/UnrealTemplate.h
--- a/Engine/Source/Runtime/Core/Public/Templates/UnrealTemplate.h
+++ b/Engine/Source/Runtime/Core/Public/Templates/UnrealTemplate.h
@@ -345,13 +345,13 @@ template <typename RefType, typename AssignedType = RefType>
 struct TGuardValue : private FNoncopyable
 {
 	[[nodiscard]] TGuardValue(RefType& ReferenceValue, const AssignedType& NewValue)
-	: RefValue(ReferenceValue), OldValue(ReferenceValue)
+	: RefValue(ReferenceValue), OriginalValue(ReferenceValue)
 	{
 		RefValue = NewValue;
 	}
 	~TGuardValue()
 	{
-		RefValue = OldValue;
+		RefValue = OriginalValue;
 	}
 
 	/**
@@ -360,14 +360,25 @@ struct TGuardValue : private FNoncopyable
 	 *
 	 * @return	a const reference to the original data value
 	 */
+	UE_DEPRECATED(5.5, "Use GetOriginalValue() instead.")
 	FORCEINLINE const AssignedType& operator*() const
 	{
-		return OldValue;
+		return OriginalValue;
+	}
+
+	/**
+	 * Provides read-only access to the original value of the data being tracked by this struct
+	 *
+	 * @return	a const reference to the original data value
+	 */
+	FORCEINLINE const AssignedType& GetOriginalValue() const
+	{
+		return OriginalValue;
 	}
 
 private:
 	RefType& RefValue;
-	AssignedType OldValue;
+	AssignedType OriginalValue;
 };
 
 
@@ -382,7 +393,7 @@ template <typename RefType, typename AssignedType = RefType>
 struct TOptionalGuardValue : private FNoncopyable
 {
 	[[nodiscard]] TOptionalGuardValue(RefType& ReferenceValue, const AssignedType& NewValue)
-		: RefValue(ReferenceValue), OldValue(ReferenceValue)
+		: RefValue(ReferenceValue), OriginalValue(ReferenceValue)
 	{
 		if (RefValue != NewValue)
 		{
@@ -391,9 +402,9 @@ struct TOptionalGuardValue : private FNoncopyable
 	}
 	~TOptionalGuardValue()
 	{
-		if (RefValue != OldValue)
+		if (RefValue != OriginalValue)
 		{
-			RefValue = OldValue;
+			RefValue = OriginalValue;
 		}
 	}
 
@@ -403,14 +414,25 @@ struct TOptionalGuardValue : private FNoncopyable
 	 *
 	 * @return	a const reference to the original data value
 	 */
+	UE_DEPRECATED(5.5, "Use GetOriginalValue() instead.")
 	FORCEINLINE const AssignedType& operator*() const
 	{
-		return OldValue;
+		return OriginalValue;
+	}
+
+	/**
+	 * Provides read-only access to the original value of the data being tracked by this struct
+	 *
+	 * @return	a const reference to the original data value
+	 */
+	FORCEINLINE const AssignedType& GetOriginalValue() const
+	{
+		return OriginalValue;
 	}
 
 private:
 	RefType& RefValue;
-	AssignedType OldValue;
+	AssignedType OriginalValue;
 };
 
 template <typename FuncType>
diff --git a/Engine/Source/Runtime/CoreUObject/Private/UObject/LinkerLoad.cpp b/Engine/Source/Runtime/CoreUObject/Private/UObject/LinkerLoad.cpp
--- a/Engine/Source/Runtime/CoreUObject/Private/UObject/LinkerLoad.cpp
+++ b/Engine/Source/Runtime/CoreUObject/Private/UObject/LinkerLoad.cpp
@@ -4824,7 +4824,7 @@ void FLinkerLoad::Preload( UObject* Object )
 #if USE_CIRCULAR_DEPENDENCY_LOAD_DEFERRING
 				{
 					SCOPE_CYCLE_COUNTER(STAT_LinkerLoadDeferred);
-					if ((LoadFlags & LOAD_DeferDependencyLoads) != (*LoadFlagsGuard & LOAD_DeferDependencyLoads))
+					if ((LoadFlags & LOAD_DeferDependencyLoads) != (LoadFlagsGuard.GetOriginalValue() & LOAD_DeferDependencyLoads))
 					{
 						if (bIsBlueprintStruct)
 						{
diff --git a/Engine/Source/Runtime/CoreUObject/Private/UObject/UObjectGlobals.cpp b/Engine/Source/Runtime/CoreUObject/Private/UObject/UObjectGlobals.cpp
--- a/Engine/Source/Runtime/CoreUObject/Private/UObject/UObjectGlobals.cpp
+++ b/Engine/Source/Runtime/CoreUObject/Private/UObject/UObjectGlobals.cpp
@@ -1977,7 +1977,7 @@ UPackage* LoadPackageInternal(UPackage* InOuter, const FPackagePath& PackagePath
 		EndLoadAndCopyLocalizationGatherFlag();
 
 #if WITH_EDITOR
-		GIsEditorLoadingPackage = *IsEditorLoadingPackage;
+		GIsEditorLoadingPackage = IsEditorLoadingPackage.GetOriginalValue();
 #endif
 
 		// if we are calculating the script SHA for a package, do the comparison now
