[Templates]
 - Add new class similar to TGuardValue but which works using accessors instead

diff --git a/Engine/Source/Runtime/Core/Public/Templates/GuardValueAccessors.h b/Engine/Source/Runtime/Core/Public/Templates/GuardValueAccessors.h
new file mode 100644
--- /dev/null
+++ b/Engine/Source/Runtime/Core/Public/Templates/GuardValueAccessors.h
@@ -0,0 +1,44 @@
+// Copyright Epic Games, Inc. All Rights Reserved.
+
+#pragma once
+
+#include "CoreTypes.h"
+#include "Templates/Function.h"
+
+/**
+ * Guard around saving/restoring a value.
+ * Commonly used to make sure a value is restored
+ * even if the code early outs in the future.
+ * Usage:
+ *  	TGuardValueAccessors GuardSomeBool(UE::SomeGetterFunction, UE::SomeSetterFunction, false); // Saves the state, sets the value, and restores it in dtor.
+ */
+template <typename AssignedType>
+struct TGuardValueAccessors
+{
+	UE_NONCOPYABLE(TGuardValueAccessors)
+
+	[[nodiscard]] explicit TGuardValueAccessors(TFunctionRef<AssignedType()> Getter, TFunction<void(const AssignedType&)>&& InSetter, const AssignedType& NewValue)
+		: Setter(MoveTemp(InSetter))
+		, OriginalValue(Getter())
+	{
+		Setter(NewValue);
+	}
+
+	~TGuardValueAccessors()
+	{
+		Setter(OriginalValue);
+	}
+
+	/**
+	 * Provides read-only access to the original value of the data being tracked by this struct
+	 *
+	 * @return	a const reference to the original data value
+	 */
+	FORCEINLINE const AssignedType& GetOriginalValue() const
+	{
+		return OriginalValue;
+	}
+private:
+	TFunction<void(const AssignedType&)> Setter;
+	AssignedType OriginalValue;
+};
