Added AsSharedSubobject(&Subobject) and AsWeakSubobject(&Subobject) to allow easy access to the aliasing constructor from TSharedFromThis.
Fixed a bunch of comments.

diff --git a/Engine/Source/Runtime/Core/Public/Templates/SharedPointer.h b/Engine/Source/Runtime/Core/Public/Templates/SharedPointer.h
--- a/Engine/Source/Runtime/Core/Public/Templates/SharedPointer.h
+++ b/Engine/Source/Runtime/Core/Public/Templates/SharedPointer.h
@@ -1587,11 +1587,11 @@ class TSharedFromThis
 public:
 
 	/**
-	 * Provides access to a shared reference to this object.  Note that is only valid to call
+	 * Provides a shared reference to this object.  Note that is only valid to call
 	 * this after a shared reference (or shared pointer) to the object has already been created.
 	 * Also note that it is illegal to call this in the object's destructor.
 	 *
-	 * @return	Returns this object as a shared pointer
+	 * @return	Returns this object as a shared reference
 	 */
 	[[nodiscard]] TSharedRef< ObjectType, Mode > AsShared()
 	{
@@ -1600,8 +1600,8 @@ public:
 		//
 		// If the following assert goes off, it means one of the following:
 		//
-		//     - You tried to request a shared pointer before the object was ever assigned to one. (e.g. constructor)
-		//     - You tried to request a shared pointer while the object is being destroyed (destructor chain)
+		//     - You tried to request a shared reference before the object was ever assigned to one. (e.g. constructor)
+		//     - You tried to request a shared reference while the object is being destroyed (destructor chain)
 		//
 		// To fix this, make sure you create at least one shared reference to your object instance before requested,
 		// and also avoid calling this function from your object's destructor.
@@ -1614,11 +1614,11 @@ public:
 	}
 
 	/**
-	 * Provides access to a shared reference to this object (const.)  Note that is only valid to call
+	 * Provides a shared reference to this object (const.)  Note that is only valid to call
 	 * this after a shared reference (or shared pointer) to the object has already been created.
 	 * Also note that it is illegal to call this in the object's destructor.
 	 *
-	 * @return	Returns this object as a shared pointer (const)
+	 * @return	Returns this object as a shared reference (const)
 	 */
 	[[nodiscard]] TSharedRef< ObjectType const, Mode > AsShared() const
 	{
@@ -1627,8 +1627,8 @@ public:
 		//
 		// If the following assert goes off, it means one of the following:
 		//
-		//     - You tried to request a shared pointer before the object was ever assigned to one. (e.g. constructor)
-		//     - You tried to request a shared pointer while the object is being destroyed (destructor chain)
+		//     - You tried to request a shared reference before the object was ever assigned to one. (e.g. constructor)
+		//     - You tried to request a shared reference while the object is being destroyed (destructor chain)
 		//
 		// To fix this, make sure you create at least one shared reference to your object instance before requested,
 		// and also avoid calling this function from your object's destructor.
@@ -1640,12 +1640,26 @@ public:
 		return MoveTemp( SharedThis ).ToSharedRef();
 	}
 
+	/**
+	 * Provides a shared reference to a subobject of this object, i.e. points to an object which shares
+	 * the lifetime and refcount of this object, but isn't derived from TSharedFromThis.  Note that is only valid to call
+	 * this after a shared reference (or shared pointer) to the object has already been created.
+	 * Also note that it is illegal to call this in the object's destructor.
+	 *
+	 * @return	Returns the subobject as a shared reference.
+	 */
+	template< typename SubobjectType >
+	[[nodiscard]] TSharedRef< SubobjectType, Mode > AsSharedSubobject( SubobjectType* SubobjectPtr ) const
+	{
+		return TSharedRef< SubobjectType, Mode >( AsShared(), SubobjectPtr );
+	}
+
 	/**
 	 * Provides a weak reference to this object.  Note that is only valid to call
 	 * this after a shared reference (or shared pointer) to the object has already been created.
 	 * Also note that it is illegal to call this in the object's destructor.
 	 *
-	 * @return	Returns this object as a shared pointer
+	 * @return	Returns this object as a weak pointer
 	 */
 	[[nodiscard]] TWeakPtr< ObjectType, Mode > AsWeak()
 	{
@@ -1684,6 +1698,20 @@ public:
 		return Result;
 	}
 
+	/**
+	 * Provides a weak pointer to a subobject of this object, i.e. points to an object which shares
+	 * the lifetime and refcount of this object, but isn't derived from TSharedFromThis.  Note that is only valid to call
+	 * this after a shared reference (or shared pointer) to the object has already been created.
+	 * Also note that it is illegal to call this in the object's destructor.
+	 *
+	 * @return	Returns the subobject as a weak pointer.
+	 */
+	template <typename SubobjectType>
+	[[nodiscard]] TWeakPtr< SubobjectType, Mode > AsWeakSubobject( SubobjectType* SubobjectPtr ) const
+	{
+		return AsSharedSubobject( SubobjectPtr ).ToWeakPtr();
+	}
+
 protected:
 
 	/**
