Open around GetLocalTime, we will have to revisit this later as we need to get the same time if we abort and try again

diff --git a/Engine/Source/Runtime/Core/Private/AutoRTFM/StdLibSupport.cpp b/Engine/Source/Runtime/Core/Private/AutoRTFM/StdLibSupport.cpp
--- a/Engine/Source/Runtime/Core/Private/AutoRTFM/StdLibSupport.cpp
+++ b/Engine/Source/Runtime/Core/Private/AutoRTFM/StdLibSupport.cpp
@@ -306,6 +306,7 @@ UE_AUTORTFM_REGISTER_SELF_FUNCTION(QueryPerformanceFrequency);
 UE_AUTORTFM_REGISTER_SELF_FUNCTION(GetCurrentThreadId);
 
 UE_AUTORTFM_REGISTER_SELF_FUNCTION(TlsGetValue);
+UE_AUTORTFM_REGISTER_SELF_FUNCTION(GetLocalTime);
 
 BOOL RTFM_TlsSetValue(DWORD dwTlsIndex, LPVOID lpTlsValue)
 {
