Added 144 and 176 bytes MB3 bins to reduce overall memory footprint

diff --git a/Engine/Source/Runtime/Core/Private/HAL/MallocBinnedCommon.cpp b/Engine/Source/Runtime/Core/Private/HAL/MallocBinnedCommon.cpp
--- a/Engine/Source/Runtime/Core/Private/HAL/MallocBinnedCommon.cpp
+++ b/Engine/Source/Runtime/Core/Private/HAL/MallocBinnedCommon.cpp
@@ -18,8 +18,8 @@ PRAGMA_DISABLE_UNSAFE_TYPECAST_WARNINGS
 // They must be 16-byte aligned as well.
 static uint32 BinnedCommonSmallBlockSizes4k[] = 
 {
-	16, 32, 48, 64, 80, 96, 112, 128, 160, // +16
-	192, 224, 256, 288, 320, // +32
+	16, 32, 48, 64, 80, 96, 112, 128, 144, 160, 176, 192, // +16
+	224, 256, 288, 320, // +32
 	368,  // /11 ish
 	400,  // /10 ish
 	448,  // /9 ish
diff --git a/Engine/Source/Runtime/Core/Public/HAL/MallocBinnedCommon.h b/Engine/Source/Runtime/Core/Public/HAL/MallocBinnedCommon.h
--- a/Engine/Source/Runtime/Core/Public/HAL/MallocBinnedCommon.h
+++ b/Engine/Source/Runtime/Core/Public/HAL/MallocBinnedCommon.h
@@ -18,7 +18,7 @@
 #include "Templates/Function.h"
 
 #define BINNEDCOMMON_MAX_LISTED_SMALL_POOL_SIZE	28672
-#define BINNEDCOMMON_NUM_LISTED_SMALL_POOLS	49
+#define BINNEDCOMMON_NUM_LISTED_SMALL_POOLS	51
 
 #if !defined(BINNEDCOMMON_USE_SEPARATE_VM_PER_POOL)
 	#if PLATFORM_WINDOWS
