Remove log line in ConfigContext initialization to avoid a deadlock (as well as potentially not being able to suppress this log through config).

diff --git a/Engine/Source/Runtime/Core/Private/Misc/ConfigContext.cpp b/Engine/Source/Runtime/Core/Private/Misc/ConfigContext.cpp
--- a/Engine/Source/Runtime/Core/Private/Misc/ConfigContext.cpp
+++ b/Engine/Source/Runtime/Core/Private/Misc/ConfigContext.cpp
@@ -476,8 +476,6 @@ bool FConfigContext::PerformLoad()
 	if (IsInGameThread()) GPerformLoadTime -= FPlatformTime::Seconds();
 
 #if DISABLE_GENERATED_INI_WHEN_COOKED
-	UE_LOG(LogConfig, Display, TEXT("IniName = %s"), *BaseIniName);
-
 	if (BaseIniName == TEXT("GameUserSettings"))
 	{
 		bAllowGeneratedIniWhenCooked = true;
