Open around FindOrAdd to FTextCache, we will leak this in the cache, assuming it will be used again

diff --git a/Engine/Source/Runtime/Core/Private/Internationalization/TextCache.cpp b/Engine/Source/Runtime/Core/Private/Internationalization/TextCache.cpp
--- a/Engine/Source/Runtime/Core/Private/Internationalization/TextCache.cpp
+++ b/Engine/Source/Runtime/Core/Private/Internationalization/TextCache.cpp
@@ -30,16 +30,28 @@ FText FTextCache::FindOrCache(const TCHAR* InTextLiteral, const FTextId& InTextI
 
 	// First try and find a cached instance
 	{
-		FScopeLock Lock(&CachedTextCS);
-	
-		const FText* FoundText = CachedText.Find(InTextId);
-		if (FoundText)
+		FText* ReturnFoundText = nullptr;
+
+		UE_AUTORTFM_OPEN(
 		{
-			const FString* FoundTextLiteral = FTextInspector::GetSourceString(*FoundText);
-			if (FoundTextLiteral && FCString::Strcmp(**FoundTextLiteral, InTextLiteral) == 0)
+
+			FScopeLock Lock(&CachedTextCS);
+
+			FText* FoundText = CachedText.Find(InTextId);
+			if (FoundText)
 			{
-				return *FoundText;
+				const FString* FoundTextLiteral = FTextInspector::GetSourceString(*FoundText);
+				if (FoundTextLiteral && FCString::Strcmp(**FoundTextLiteral, InTextLiteral) == 0)
+				{
+
+					ReturnFoundText = FoundText;
+				}
 			}
+		});
+
+		if (ReturnFoundText)
+		{
+			return *ReturnFoundText;
 		}
 	}
 
@@ -47,11 +59,12 @@ FText FTextCache::FindOrCache(const TCHAR* InTextLiteral, const FTextId& InTextI
 	FText NewText = FText(InTextLiteral, InTextId.GetNamespace(), InTextId.GetKey(), ETextFlag::Immutable);
 
 	// ... and add it to the cache
+	UE_AUTORTFM_OPEN(
 	{
 		FScopeLock Lock(&CachedTextCS);
 
 		CachedText.Emplace(InTextId, NewText);
-	}
+	});
 
 	return NewText;
 }
