Add an option for StorageServerPlatformFile to control if needs to access local files or no
Used during ini files hierarchy reading in order to speed up boot time on specific platforms

diff --git a/Engine/Source/Runtime/Core/Private/Android/AndroidPlatformFile.cpp b/Engine/Source/Runtime/Core/Private/Android/AndroidPlatformFile.cpp
--- a/Engine/Source/Runtime/Core/Private/Android/AndroidPlatformFile.cpp
+++ b/Engine/Source/Runtime/Core/Private/Android/AndroidPlatformFile.cpp
@@ -1408,6 +1408,11 @@ public:
 	{
 		return FileSize(Filename, false);
 	}
+	
+	int64 FileSize(const TCHAR* Filename, EPlatformFileFlags PlatformFlags) override
+	{
+		return IPlatformFile::FileSize(Filename, PlatformFlags);
+	}
 
 	int64 FileSize(const TCHAR* Filename, bool AllowLocal)
 	{
diff --git a/Engine/Source/Runtime/Core/Private/HAL/FileManagerGeneric.cpp b/Engine/Source/Runtime/Core/Private/HAL/FileManagerGeneric.cpp
--- a/Engine/Source/Runtime/Core/Private/HAL/FileManagerGeneric.cpp
+++ b/Engine/Source/Runtime/Core/Private/HAL/FileManagerGeneric.cpp
@@ -133,6 +133,11 @@ int64 FFileManagerGeneric::FileSize( const TCHAR* Filename )
 	return GetLowLevel().FileSize( Filename );
 }
 
+int64 FFileManagerGeneric::FileSize(const TCHAR* Filename, EPlatformFileFlags PlatformFlags)
+{
+	return GetLowLevel().FileSize(Filename, PlatformFlags);
+}
+
 uint32 FFileManagerGeneric::Copy( const TCHAR* Dest, const TCHAR* Src, bool Replace, bool EvenIfReadOnly, bool Attributes, FCopyProgress* Progress, EFileRead ReadFlags, EFileWrite WriteFlags)
 {
 	uint32	Result = COPY_OK;
diff --git a/Engine/Source/Runtime/Core/Private/Misc/ConfigCacheIni.cpp b/Engine/Source/Runtime/Core/Private/Misc/ConfigCacheIni.cpp
--- a/Engine/Source/Runtime/Core/Private/Misc/ConfigCacheIni.cpp
+++ b/Engine/Source/Runtime/Core/Private/Misc/ConfigCacheIni.cpp
@@ -470,7 +470,7 @@ void FixupArrayOfStructKeysForSection(SectionType* Section, const FString& Secti
 /**
  * Check if an ini file exists, allowing a delegate to determine if it will handle loading it
  */
-/*static*/ bool DoesConfigFileExistWrapper(const TCHAR* IniFile, const TSet<FString>* IniCacheSet)
+/*static*/ bool DoesConfigFileExistWrapper(const TCHAR* IniFile, const TSet<FString>* IniCacheSet, EPlatformFileFlags PlatformFlag)
 {
 	// will any delegates return contents via TSPreLoadConfigFileDelegate()?
 	int32 ResponderCount = 0;
@@ -510,7 +510,7 @@ void FixupArrayOfStructKeysForSection(SectionType* Section, const FString& Secti
 	}
 
 	// otherwise just look for the normal file to exist
-	const bool bFileExistsCached = IFileManager::Get().FileSize(IniFile) >= 0;
+	const bool bFileExistsCached = IFileManager::Get().FileSize(IniFile, PlatformFlag) >= 0;
 	return bFileExistsCached;
 }
 
diff --git a/Engine/Source/Runtime/Core/Private/Misc/ConfigContext.cpp b/Engine/Source/Runtime/Core/Private/Misc/ConfigContext.cpp
--- a/Engine/Source/Runtime/Core/Private/Misc/ConfigContext.cpp
+++ b/Engine/Source/Runtime/Core/Private/Misc/ConfigContext.cpp
@@ -1000,7 +1000,7 @@ bool FConfigContext::GenerateDestIniFile()
 		{
 			static bool bDumpIniLoadInfo = FParse::Param(FCommandLine::Get(), TEXT("dumpiniloads"));
 			UE_CLOG(bDumpIniLoadInfo, LogConfig, Display, TEXT("Looking for saved user ini file: %s"), *DestIniFilename);
-			if (DoesConfigFileExistWrapper(*DestIniFilename))
+			if (DoesConfigFileExistWrapper(*DestIniFilename, nullptr, EPlatformFileFlags::Default))
 			{
 				UE_CLOG(bDumpIniLoadInfo, LogConfig, Display, TEXT("   Found!"));
 				Branch->SavedLayer.FillFileFromDisk(*DestIniFilename, false);
diff --git a/Engine/Source/Runtime/Core/Public/GenericPlatform/GenericPlatformFile.h b/Engine/Source/Runtime/Core/Public/GenericPlatform/GenericPlatformFile.h
--- a/Engine/Source/Runtime/Core/Public/GenericPlatform/GenericPlatformFile.h
+++ b/Engine/Source/Runtime/Core/Public/GenericPlatform/GenericPlatformFile.h
@@ -94,6 +94,15 @@ enum class ESymlinkResult : int8
 
 ENUM_CLASS_FLAGS(EDirectoryVisitorFlags);
 
+enum class EPlatformFileFlags : uint8
+{
+	Default = 0x00,
+	StorageServerOnly = 0x01,
+};
+
+ENUM_CLASS_FLAGS(EPlatformFileFlags);
+
+
 /** Results that can be returned from IPlatformFile FileJournal API. */
 enum class EFileJournalResult
 {
@@ -371,6 +380,12 @@ public:
 	virtual bool		FileExists(const TCHAR* Filename) = 0;
 	/** Return the size of the file, or -1 if it doesn't exist. **/
 	virtual int64		FileSize(const TCHAR* Filename) = 0;
+	/** Fallback to vanilla FileSize if this function is not overidden by a Platform implementation**/
+	virtual int64		FileSize(const TCHAR* Filename, EPlatformFileFlags PlatformFlags)
+	{
+		return FileSize(Filename);
+	}
+
 	/** Delete a file and return true if the file exists. Will not delete read only files. **/
 	virtual bool		DeleteFile(const TCHAR* Filename) = 0;
 	/** Delete an array of files and return true if ALL deletes are succeeded. **/
diff --git a/Engine/Source/Runtime/Core/Public/HAL/FileManager.h b/Engine/Source/Runtime/Core/Public/HAL/FileManager.h
--- a/Engine/Source/Runtime/Core/Public/HAL/FileManager.h
+++ b/Engine/Source/Runtime/Core/Public/HAL/FileManager.h
@@ -53,7 +53,6 @@ enum EFileOpenFlags
 	IO_APPEND		= 0x40,					// When writing, keep the existing data, set the filepointer to the end of the existing data
 };
 
-
 class IFileManager
 {
 protected:
@@ -269,7 +268,16 @@ public:
 	 *	@param Filename		Platform-independent Unreal filename.
 	 *	@return				File size in bytes or INDEX_NONE if the file didn't exist.
 	 **/
-	virtual int64 FileSize( const TCHAR* Filename )=0;
+	virtual int64 FileSize(const TCHAR* Filename ) = 0;
+
+	/**
+	 *	Returns the size of a file. (Thread-safe)
+	 *
+	 *	@param Filename		Platform-independent Unreal filename.
+	 *  @param PlatformFlag	Flags to control special behavior of some platforms files implementations
+	 *	@return				File size in bytes or INDEX_NONE if the file didn't exist.
+	 **/
+	virtual int64 FileSize(const TCHAR* Filename, EPlatformFileFlags PlatformFlag) = 0;
 
 	/**
 	 * Sends a message to the file server, and will block until it's complete. Will return 
diff --git a/Engine/Source/Runtime/Core/Public/HAL/FileManagerGeneric.h b/Engine/Source/Runtime/Core/Public/HAL/FileManagerGeneric.h
--- a/Engine/Source/Runtime/Core/Public/HAL/FileManagerGeneric.h
+++ b/Engine/Source/Runtime/Core/Public/HAL/FileManagerGeneric.h
@@ -194,6 +194,15 @@ public:
 	 **/
 	CORE_API int64 FileSize( const TCHAR* Filename ) override;
 
+	/**
+	 *	Returns the size of a file. (Thread-safe)
+	 *
+	 *	@param Filename		Platform-independent Unreal filename.
+	 *  @param PlatformFlag	Flags to control special behavior of some platforms files, if not implemented fallback to vanilla FileSize
+	 *	@return				File size in bytes or INDEX_NONE if the file didn't exist.
+	 **/
+	CORE_API int64 FileSize(const TCHAR* Filename, EPlatformFileFlags PlatformFlags) override;
+
 	/**
 	 * Sends a message to the file server, and will block until it's complete. Will return 
 	 * immediately if the file manager doesn't support talking to a server.
diff --git a/Engine/Source/Runtime/Core/Public/Misc/ConfigContext.h b/Engine/Source/Runtime/Core/Public/Misc/ConfigContext.h
--- a/Engine/Source/Runtime/Core/Public/Misc/ConfigContext.h
+++ b/Engine/Source/Runtime/Core/Public/Misc/ConfigContext.h
@@ -12,7 +12,7 @@
 #include "Misc/ConfigTypes.h"
 // @todo mvoe types into ConfigTypes.h so we don't need to include the monster here
 #include "Misc/ConfigCacheIni.h"
-
+#include "GenericPlatform/GenericPlatformFile.h"
 
 #ifndef CUSTOM_CONFIG
 #define CUSTOM_CONFIG ""
@@ -233,4 +233,4 @@ private:
 	void LogVariables(const TCHAR* InBaseIniName, const FString& Platform);
 };
 
-bool DoesConfigFileExistWrapper(const TCHAR* IniFile, const TSet<FString>* IniCacheSet = nullptr);
+bool DoesConfigFileExistWrapper(const TCHAR* IniFile, const TSet<FString>* IniCacheSet = nullptr, EPlatformFileFlags PlatformFlag = EPlatformFileFlags::StorageServerOnly);
diff --git a/Engine/Source/Runtime/StorageServerClient/Private/StorageServerPlatformFile.cpp b/Engine/Source/Runtime/StorageServerClient/Private/StorageServerPlatformFile.cpp
--- a/Engine/Source/Runtime/StorageServerClient/Private/StorageServerPlatformFile.cpp
+++ b/Engine/Source/Runtime/StorageServerClient/Private/StorageServerPlatformFile.cpp
@@ -517,6 +517,28 @@ int64 FStorageServerPlatformFile::FileSize(const TCHAR* Filename)
 	return LowerLevel->FileSize(Filename);
 }
 
+int64 FStorageServerPlatformFile::FileSize(const TCHAR* Filename, EPlatformFileFlags PlatformFlags)
+{
+	TStringBuilder<1024> StorageServerFilename;
+	if (MakeStorageServerPath(Filename, StorageServerFilename))
+	{
+		if (const FIoChunkId* FileChunkId = ServerToc.GetFileChunkId(*StorageServerFilename))
+		{
+			const FFileStatData FileStatData = SendGetStatDataMessage(*FileChunkId);
+			check(FileStatData.bIsValid);
+			return FileStatData.FileSize;
+		}
+	}
+
+	if (PlatformFlags == EPlatformFileFlags::StorageServerOnly)
+	{
+		return -1;
+	}
+
+	return LowerLevel->FileSize(Filename);
+}
+
+
 bool FStorageServerPlatformFile::IsReadOnly(const TCHAR* Filename)
 {
 	TStringBuilder<1024> StorageServerFilename;
diff --git a/Engine/Source/Runtime/StorageServerClient/Private/StorageServerPlatformFile.h b/Engine/Source/Runtime/StorageServerClient/Private/StorageServerPlatformFile.h
--- a/Engine/Source/Runtime/StorageServerClient/Private/StorageServerPlatformFile.h
+++ b/Engine/Source/Runtime/StorageServerClient/Private/StorageServerPlatformFile.h
@@ -80,6 +80,7 @@ public:
 
 	virtual bool FileExists(const TCHAR* Filename) override;
 	virtual int64 FileSize(const TCHAR* Filename) override;
+	virtual int64 FileSize(const TCHAR* Filename, EPlatformFileFlags PlatformFlags) override;
 	virtual bool IsReadOnly(const TCHAR* Filename) override;
 	virtual FDateTime GetTimeStamp(const TCHAR* Filename) override;
 	virtual FDateTime GetAccessTimeStamp(const TCHAR* Filename) override;
