CompressedBuffer: Added an assert to prevent saving a null buffer to an archive

Added tests of serialization to and from an archive.

diff --git a/Engine/Source/Runtime/Core/Private/Compression/CompressedBuffer.cpp b/Engine/Source/Runtime/Core/Private/Compression/CompressedBuffer.cpp
--- a/Engine/Source/Runtime/Core/Private/Compression/CompressedBuffer.cpp
+++ b/Engine/Source/Runtime/Core/Private/Compression/CompressedBuffer.cpp
@@ -1080,7 +1080,8 @@ FCompressedBuffer FCompressedBuffer::Load(FArchive& Ar)
 	Header.ByteSwap();
 
 	FCompressedBuffer Local;
-	if (Header.Magic == Header.ExpectedMagic && Header.TotalCompressedSize >= sizeof(FHeader))
+	constexpr uint64 MaxCompressedSize = uint64(1) << 48;
+	if (Header.Magic == Header.ExpectedMagic && Header.TotalCompressedSize >= sizeof(FHeader) && Header.TotalCompressedSize <= MaxCompressedSize)
 	{
 		FUniqueBuffer MutableBuffer = FUniqueBuffer::Alloc(Header.TotalCompressedSize);
 		Header.ByteSwap();
@@ -1101,6 +1102,7 @@ FCompressedBuffer FCompressedBuffer::Load(FArchive& Ar)
 void FCompressedBuffer::Save(FArchive& Ar) const
 {
 	check(Ar.IsSaving());
+	checkf(!CompressedData.IsNull(), TEXT("Serializing a null compressed buffer is not supported."));
 	for (const FSharedBuffer& Segment : CompressedData.GetSegments())
 	{
 		Ar.Serialize(const_cast<void*>(Segment.GetData()), int64(Segment.GetSize()));
diff --git a/Engine/Source/Runtime/Core/Tests/Compression/CompressedBufferTest.cpp b/Engine/Source/Runtime/Core/Tests/Compression/CompressedBufferTest.cpp
--- a/Engine/Source/Runtime/Core/Tests/Compression/CompressedBufferTest.cpp
+++ b/Engine/Source/Runtime/Core/Tests/Compression/CompressedBufferTest.cpp
@@ -7,7 +7,8 @@
 #include "Algo/Compare.h"
 #include "Compression/OodleDataCompression.h"
 #include "IO/IoHash.h"
-#include "Serialization/BufferReader.h"
+#include "Serialization/MemoryReader.h"
+#include "Serialization/MemoryWriter.h"
 #include "TestHarness.h"
 
 #include <catch2/generators/catch_generators.hpp>
@@ -17,9 +18,28 @@ TEST_CASE("Core::Compression::CompressedBuffer::Compress", "[Core][Compression][
 	const uint8 ZeroBuffer[1024]{};
 	const FIoHash ZeroBufferHash = FIoHash::HashBuffer(MakeMemoryView(ZeroBuffer));
 
+	const auto SerializeBuffer = [](const FCompressedBuffer& Buffer) -> FCompressedBuffer
+	{
+		FCompressedBuffer SerializedBuffer = Buffer;
+		TArray<uint8> Bytes;
+		{
+			FMemoryWriter Ar(Bytes, /*bIsPersistent*/ true);
+			Ar << SerializedBuffer;
+			CHECK_FALSE(Ar.IsError());
+		}
+		SerializedBuffer.Reset();
+		{
+			FMemoryReader Ar(Bytes, /*bIsPersistent*/ true);
+			Ar << SerializedBuffer;
+			CHECK_FALSE(Ar.IsError());
+		}
+		return SerializedBuffer;
+	};
+
 	SECTION("Null")
 	{
 		const FCompressedBuffer Buffer;
+
 		CHECK_FALSE(Buffer);
 		CHECK(Buffer.IsNull());
 		CHECK(Buffer.IsOwned());
@@ -30,15 +50,33 @@ TEST_CASE("Core::Compression::CompressedBuffer::Compress", "[Core][Compression][
 		CHECK(Buffer.DecompressToComposite().IsNull());
 	}
 
+	SECTION("Empty")
+	{
+		const FSharedBuffer EmptyBuffer = FUniqueBuffer::Alloc(0).MoveToShared();
+		const FIoHash ExpectedRawHash = FIoHash::HashBuffer(EmptyBuffer);
+		const FCompressedBuffer OriginalBuffer = FCompressedBuffer::Compress(EmptyBuffer);
+		const FCompressedBuffer SerializedBuffer = SerializeBuffer(OriginalBuffer);
+		const FCompressedBuffer Buffer = GENERATE_REF(OriginalBuffer, SerializedBuffer);
+
+		CHECK(Buffer);
+		CHECK_FALSE(Buffer.IsNull());
+		CHECK(Buffer.IsOwned());
+		CHECK(Buffer.GetRawSize() == 0);
+		CHECK(Buffer.GetRawHash() == ExpectedRawHash);
+		CHECK_FALSE(Buffer.Decompress().IsNull());
+		CHECK_FALSE(Buffer.DecompressToComposite().IsNull());
+	}
+
 	SECTION("Method None")
 	{
-		const FCompressedBuffer BufferNew = FCompressedBuffer::Compress(FSharedBuffer::MakeView(MakeMemoryView(ZeroBuffer)),
+		const FCompressedBuffer OriginalBuffer = FCompressedBuffer::Compress(FSharedBuffer::MakeView(MakeMemoryView(ZeroBuffer)),
 			ECompressedBufferCompressor::NotSet, ECompressedBufferCompressionLevel::None);
-		const FCompressedBuffer BufferCopy = FCompressedBuffer::FromCompressed(BufferNew.GetCompressed());
-		const FCompressedBuffer Buffer = GENERATE_REF(BufferNew, BufferCopy);
+		const FCompressedBuffer CopiedBuffer = FCompressedBuffer::FromCompressed(OriginalBuffer.GetCompressed());
+		const FCompressedBuffer SerializedBuffer = SerializeBuffer(OriginalBuffer);
+		const FCompressedBuffer Buffer = GENERATE_REF(OriginalBuffer, CopiedBuffer, SerializedBuffer);
 
 		CHECK(Buffer);
-		CHECK(!Buffer.IsNull());
+		CHECK_FALSE(Buffer.IsNull());
 		CHECK(Buffer.IsOwned());
 		CHECK(Buffer.GetCompressedSize() == sizeof(ZeroBuffer) + 64);
 		CHECK(Buffer.GetRawSize() == sizeof(ZeroBuffer));
@@ -57,13 +95,14 @@ TEST_CASE("Core::Compression::CompressedBuffer::Compress", "[Core][Compression][
 
 	SECTION("Method Oodle")
 	{
-		const FCompressedBuffer BufferNew = FCompressedBuffer::Compress(FSharedBuffer::MakeView(MakeMemoryView(ZeroBuffer)),
+		const FCompressedBuffer OriginalBuffer = FCompressedBuffer::Compress(FSharedBuffer::MakeView(MakeMemoryView(ZeroBuffer)),
 			ECompressedBufferCompressor::Mermaid, ECompressedBufferCompressionLevel::VeryFast);
-		const FCompressedBuffer BufferCopy = FCompressedBuffer::FromCompressed(BufferNew.GetCompressed());
-		const FCompressedBuffer Buffer = GENERATE_REF(BufferNew, BufferCopy);
+		const FCompressedBuffer CopiedBuffer = FCompressedBuffer::FromCompressed(OriginalBuffer.GetCompressed());
+		const FCompressedBuffer SerializedBuffer = SerializeBuffer(OriginalBuffer);
+		const FCompressedBuffer Buffer = GENERATE_REF(OriginalBuffer, CopiedBuffer, SerializedBuffer);
 
 		CHECK(Buffer);
-		CHECK(!Buffer.IsNull());
+		CHECK_FALSE(Buffer.IsNull());
 		CHECK(Buffer.IsOwned());
 		CHECK(Buffer.GetCompressedSize() < sizeof(ZeroBuffer));
 		CHECK(Buffer.GetRawSize() == sizeof(ZeroBuffer));
@@ -161,7 +200,7 @@ TEST_CASE("Core::Compression::CompressedBuffer::Decompress", "[Core][Compression
 		constexpr uint64 OffsetCount = 150;
 		{
 			FSharedBuffer Buffer = Compressed.GetCompressed().ToShared();
-			FBufferReader Ar(const_cast<void*>(Buffer.GetData()), int64(Buffer.GetSize()), /*bFreeOnClose*/ false, /*bIsPersistent*/ true);
+			FMemoryReaderView Ar(Buffer, /*bIsPersistent*/ true);
 			FCompressedBufferReaderSourceScope Source(Reader, Ar);
 			const FSharedBuffer Uncompressed = Reader.Decompress(OffsetCount * sizeof(uint64));
 			CHECK(Algo::Compare(CastToArrayView(Uncompressed), MakeArrayView(ExpectedValues).RightChop(OffsetCount)));
@@ -198,7 +237,7 @@ TEST_CASE("Core::Compression::CompressedBuffer::Decompress", "[Core][Compression
 		constexpr uint64 Count = 50;
 		{
 			FSharedBuffer Buffer = Compressed.GetCompressed().ToShared();
-			FBufferReader Ar(const_cast<void*>(Buffer.GetData()), int64(Buffer.GetSize()), /*bFreeOnClose*/ false, /*bIsPersistent*/ true);
+			FMemoryReaderView Ar(Buffer, /*bIsPersistent*/ true);
 			FCompressedBufferReaderSourceScope Source(Reader, Ar);
 			const FSharedBuffer Uncompressed = Reader.Decompress(OffsetCount * sizeof(uint64), Count * sizeof(uint64));
 			CHECK(Algo::Compare(CastToArrayView(Uncompressed), MakeArrayView(ExpectedValues).Mid(OffsetCount, Count)));
