Made MB3 bins configurable via project defines

diff --git a/Engine/Source/Runtime/Core/Private/HAL/MallocBinnedCommon.cpp b/Engine/Source/Runtime/Core/Private/HAL/MallocBinnedCommon.cpp
--- a/Engine/Source/Runtime/Core/Private/HAL/MallocBinnedCommon.cpp
+++ b/Engine/Source/Runtime/Core/Private/HAL/MallocBinnedCommon.cpp
@@ -115,38 +115,31 @@ FSizeTableEntry::FSizeTableEntry(uint32 InBlockSize, uint64 PlatformPageSize, ui
 	BlocksPerBlockOfBlocks = (PlatformPageSize * PagesPlatformForBlockOfBlocks) / BlockSize;
 }
 
+static inline void FillTable(FSizeTableEntry* SizeTable, int32& Index, const uint32* BlockList, uint32 BlockListSize, uint32 NumPages, uint64 PlatformPageSize, uint32 BasePageSize, uint32 MinimumAlignment)
+{
+	for (uint32 Sub = 0; Sub < BlockListSize; Sub++)
+	{
+		// if we override BINNEDCOMMON_MAX_LISTED_SMALL_POOL_SIZE externally, we need to filter out predefined bins of a larger size
+		if (BlockList[Sub] <= BINNEDCOMMON_MAX_LISTED_SMALL_POOL_SIZE)
+		{
+			SizeTable[Index++] = FSizeTableEntry(BlockList[Sub], PlatformPageSize, NumPages, BasePageSize, MinimumAlignment);
+		}
+	}
+}
+
 uint8 FSizeTableEntry::FillSizeTable(uint64 PlatformPageSize, FSizeTableEntry* SizeTable, uint32 BasePageSize, uint32 MinimumAlignment, uint32 MaxSize, uint32 SizeIncrement)
 {
 	int32 Index = 0;
+	FillTable(SizeTable, Index, BinnedCommonSmallBlockSizes4k,  UE_ARRAY_COUNT(BinnedCommonSmallBlockSizes4k),  1, PlatformPageSize, BasePageSize, MinimumAlignment);
+	FillTable(SizeTable, Index, BinnedCommonSmallBlockSizes8k,  UE_ARRAY_COUNT(BinnedCommonSmallBlockSizes8k),  2, PlatformPageSize, BasePageSize, MinimumAlignment);
+	FillTable(SizeTable, Index, BinnedCommonSmallBlockSizes12k, UE_ARRAY_COUNT(BinnedCommonSmallBlockSizes12k), 3, PlatformPageSize, BasePageSize, MinimumAlignment);
+	FillTable(SizeTable, Index, BinnedCommonSmallBlockSizes16k, UE_ARRAY_COUNT(BinnedCommonSmallBlockSizes16k), 4, PlatformPageSize, BasePageSize, MinimumAlignment);
+	FillTable(SizeTable, Index, BinnedCommonSmallBlockSizes20k, UE_ARRAY_COUNT(BinnedCommonSmallBlockSizes20k), 5, PlatformPageSize, BasePageSize, MinimumAlignment);
+	FillTable(SizeTable, Index, BinnedCommonSmallBlockSizes24k, UE_ARRAY_COUNT(BinnedCommonSmallBlockSizes24k), 6, PlatformPageSize, BasePageSize, MinimumAlignment);
+	FillTable(SizeTable, Index, BinnedCommonSmallBlockSizes28k, UE_ARRAY_COUNT(BinnedCommonSmallBlockSizes28k), 7, PlatformPageSize, BasePageSize, MinimumAlignment);
+
+	check(Index == BINNEDCOMMON_NUM_LISTED_SMALL_POOLS);
 
-	for (int32 Sub = 0; Sub < UE_ARRAY_COUNT(BinnedCommonSmallBlockSizes4k); Sub++)
-	{
-		SizeTable[Index++] = FSizeTableEntry(BinnedCommonSmallBlockSizes4k[Sub], PlatformPageSize, 1, BasePageSize, MinimumAlignment);
-	}
-	for (int32 Sub = 0; Sub < UE_ARRAY_COUNT(BinnedCommonSmallBlockSizes8k); Sub++)
-	{
-		SizeTable[Index++] = FSizeTableEntry(BinnedCommonSmallBlockSizes8k[Sub], PlatformPageSize, 2, BasePageSize, MinimumAlignment);
-	}
-	for (int32 Sub = 0; Sub < UE_ARRAY_COUNT(BinnedCommonSmallBlockSizes12k); Sub++)
-	{
-		SizeTable[Index++] = FSizeTableEntry(BinnedCommonSmallBlockSizes12k[Sub], PlatformPageSize, 3, BasePageSize, MinimumAlignment);
-	}
-	for (int32 Sub = 0; Sub < UE_ARRAY_COUNT(BinnedCommonSmallBlockSizes16k); Sub++)
-	{
-		SizeTable[Index++] = FSizeTableEntry(BinnedCommonSmallBlockSizes16k[Sub], PlatformPageSize, 4, BasePageSize, MinimumAlignment);
-	}
-	for (int32 Sub = 0; Sub < UE_ARRAY_COUNT(BinnedCommonSmallBlockSizes20k); Sub++)
-	{
-		SizeTable[Index++] = FSizeTableEntry(BinnedCommonSmallBlockSizes20k[Sub], PlatformPageSize, 5, BasePageSize, MinimumAlignment);
-	}
-	for (int32 Sub = 0; Sub < UE_ARRAY_COUNT(BinnedCommonSmallBlockSizes24k); Sub++)
-	{
-		SizeTable[Index++] = FSizeTableEntry(BinnedCommonSmallBlockSizes24k[Sub], PlatformPageSize, 6, BasePageSize, MinimumAlignment);
-	}
-	for (int32 Sub = 0; Sub < UE_ARRAY_COUNT(BinnedCommonSmallBlockSizes28k); Sub++)
-	{
-		SizeTable[Index++] = FSizeTableEntry(BinnedCommonSmallBlockSizes28k[Sub], PlatformPageSize, 7, BasePageSize, MinimumAlignment);
-	}
 	Algo::Sort(MakeArrayView(SizeTable, Index));
 	check(SizeTable[Index - 1].BlockSize == BINNEDCOMMON_MAX_LISTED_SMALL_POOL_SIZE);
 	check(IsAligned(BINNEDCOMMON_MAX_LISTED_SMALL_POOL_SIZE, BasePageSize));
diff --git a/Engine/Source/Runtime/Core/Public/HAL/MallocBinnedCommon.h b/Engine/Source/Runtime/Core/Public/HAL/MallocBinnedCommon.h
--- a/Engine/Source/Runtime/Core/Public/HAL/MallocBinnedCommon.h
+++ b/Engine/Source/Runtime/Core/Public/HAL/MallocBinnedCommon.h
@@ -17,8 +17,13 @@
 #include "HAL/PlatformMemory.h"
 #include "Templates/Function.h"
 
-#define BINNEDCOMMON_MAX_LISTED_SMALL_POOL_SIZE	28672
-#define BINNEDCOMMON_NUM_LISTED_SMALL_POOLS	51
+#if !defined(BINNEDCOMMON_MAX_LISTED_SMALL_POOL_SIZE)
+	#define BINNEDCOMMON_MAX_LISTED_SMALL_POOL_SIZE	28672
+#endif
+
+#if !defined(BINNEDCOMMON_NUM_LISTED_SMALL_POOLS)
+	#define BINNEDCOMMON_NUM_LISTED_SMALL_POOLS	51
+#endif
 
 #if !defined(BINNEDCOMMON_USE_SEPARATE_VM_PER_POOL)
 	#if PLATFORM_WINDOWS
