Added a commentary that explains how an application can configure MallocBinnedCommon and MallocBinned3 to reduce runtime memory footprint at the expense of reduced performance

diff --git a/Engine/Source/Runtime/Core/Public/HAL/MallocBinnedCommon.h b/Engine/Source/Runtime/Core/Public/HAL/MallocBinnedCommon.h
--- a/Engine/Source/Runtime/Core/Public/HAL/MallocBinnedCommon.h
+++ b/Engine/Source/Runtime/Core/Public/HAL/MallocBinnedCommon.h
@@ -17,6 +17,17 @@
 #include "HAL/PlatformMemory.h"
 #include "Templates/Function.h"
 
+// A project can define it's own BINNEDCOMMON_MAX_LISTED_SMALL_POOL_SIZE and BINNEDCOMMON_NUM_LISTED_SMALL_POOLS to reduce runtime memory usage
+// MallocBinnedCommon.cpp has a list of predefined bins that go up to 28672
+// By default allocators (i.e. MB3) that use these bins will rely on this number as a baseline for a small bins count
+// These allocators can increase the amount of small bins they want to manage by going over the default MBC bins list
+// In MB3 case that means defining BINNED3_MAX_SMALL_POOL_SIZE to something like 65536
+// Every bin over the BINNEDCOMMON_MAX_LISTED_SMALL_POOL_SIZE would come with a 4kb increment
+// These small bins would be kept in user mode, increasing application's memory footprint and reducing the time it takes to allocate memory from the said bins
+// If application needs to aggressively reduce it's memory footprint, potentially trading some perf due to an increased amount of kernel calls to allocate memory
+// it can redefine BINNEDCOMMON_MAX_LISTED_SMALL_POOL_SIZE and BINNED3_MAX_SMALL_POOL_SIZE to smaller numbers, the good value is 16384 for both
+// This, however, would require the app to redefine BINNEDCOMMON_NUM_LISTED_SMALL_POOLS too to match the number of bins that fall under the new define's threshold
+// In case of 16384, we'll skip 3 larger bins and so BINNEDCOMMON_NUM_LISTED_SMALL_POOLS should be set to 48 at the time of writing
 #if !defined(BINNEDCOMMON_MAX_LISTED_SMALL_POOL_SIZE)
 	#define BINNEDCOMMON_MAX_LISTED_SMALL_POOL_SIZE	28672
 #endif
