Change iOS malloc selection logging to use low-level output available before engine initialization
Add IOSRuntimeSettings array RequiredEntitlements and check that the app has these entitlements at startup.

diff --git a/Engine/Source/Runtime/Core/Private/Apple/ApplePlatformMemory.cpp b/Engine/Source/Runtime/Core/Private/Apple/ApplePlatformMemory.cpp
--- a/Engine/Source/Runtime/Core/Private/Apple/ApplePlatformMemory.cpp
+++ b/Engine/Source/Runtime/Core/Private/Apple/ApplePlatformMemory.cpp
@@ -295,13 +295,13 @@ void FApplePlatformMemory::SetAllocatorToUse()
     // force Ansi allocator in particular cases
     if(getenv("UE4_FORCE_MALLOC_ANSI") != nullptr)
     {
-        UE_LOG(LogTemp, Display, TEXT("Using Ansi allocator."));
+		FPlatformMisc::LowLevelOutputDebugString(TEXT("UE4_FORCE_MALLOC_ANSI is set, using Ansi allocator.\n"));
         AllocatorToUse = EMemoryAllocatorToUse::Ansi;
         return;
     }
     if (FORCE_ANSI_ALLOCATOR)
     {
-        UE_LOG(LogTemp, Display, TEXT("Using Ansi allocator."));
+		FPlatformMisc::LowLevelOutputDebugString(TEXT("FORCE_ANSI_ALLOCATOR defined, using Ansi allocator.\n"));
         AllocatorToUse = EMemoryAllocatorToUse::Ansi;
         return;
     }
@@ -310,22 +310,22 @@ void FApplePlatformMemory::SetAllocatorToUse()
  #if PLATFORM_IOS || PLATFORM_TVOS
         if(!FIOSPlatformMisc::IsEntitlementEnabled("com.apple.developer.kernel.extended-virtual-addressing"))
         {
-            UE_LOG(LogTemp, Warning, TEXT("MallocBinned2 requested but Virtual Address Space entitlement not found. Check your entitlements. Falling back to Ansi."));
+			FPlatformMisc::LowLevelOutputDebugString(TEXT("MallocBinned2 requested but Virtual Address Space entitlement not found. Check your entitlements. Falling back to Ansi.\n"));
             AllocatorToUse = EMemoryAllocatorToUse::Ansi;
             return;
         }
         else
         {
-            UE_LOG(LogTemp, Warning, TEXT("Virtual Address Space entitlement found. Using MallocBinned2 allocator"));
+			FPlatformMisc::LowLevelOutputDebugString(TEXT("Virtual Address Space entitlement found. Using MallocBinned2 allocator.\n"));
         }
 #endif
-        UE_LOG(LogTemp, Display, TEXT("Using MallocBinned2 allocator."));
+		FPlatformMisc::LowLevelOutputDebugString(TEXT("Using MallocBinned2 allocator.\n"));
         AllocatorToUse = EMemoryAllocatorToUse::Binned2;
         return;
     }
     else
     {
-        UE_LOG(LogTemp, Display, TEXT("Defaulting to Ansi allocator."));
+		FPlatformMisc::LowLevelOutputDebugString(TEXT("Defaulting to Ansi allocator.\n"));
         AllocatorToUse = EMemoryAllocatorToUse::Ansi;
         return;
     }
diff --git a/Engine/Source/Runtime/Core/Private/IOS/IOSPlatformMisc.cpp b/Engine/Source/Runtime/Core/Private/IOS/IOSPlatformMisc.cpp
--- a/Engine/Source/Runtime/Core/Private/IOS/IOSPlatformMisc.cpp
+++ b/Engine/Source/Runtime/Core/Private/IOS/IOSPlatformMisc.cpp
@@ -32,7 +32,7 @@
 #include "Misc/SecureHash.h"
 #include "Modules/ModuleManager.h"
 #include "Templates/Function.h"
-
+#include "Containers/StringConv.h"
 #include "Apple/PreAppleSystemHeaders.h"
 
 #if !PLATFORM_TVOS && !PLATFORM_VISIONOS
@@ -91,6 +91,17 @@ void FIOSPlatformMisc::PlatformInit()
 	int32 Result = setrlimit(RLIMIT_NOFILE, &Limit);
 	check(Result == 0);
 
+	// Check for required entitlements
+	TArray<FString> RequiredEntitlements;
+	GConfig->GetArray(TEXT("/Script/IOSRuntimeSettings.IOSRuntimeSettings"), TEXT("RequiredEntitlements"), RequiredEntitlements, GEngineIni);
+	for (const FString& Entitlement : RequiredEntitlements)
+	{
+		if (!FIOSPlatformMisc::IsEntitlementEnabled(TCHAR_TO_ANSI(*Entitlement)))
+		{
+			UE_LOG(LogInit, Fatal, TEXT("App does not have required entitlement %s."), *Entitlement);
+		}
+	}
+
 	// Identity.
 	UE_LOG(LogInit, Log, TEXT("Computer: %s"), FPlatformProcess::ComputerName() );
 	UE_LOG(LogInit, Log, TEXT("User: %s"), FPlatformProcess::UserName() );
@@ -1305,7 +1316,7 @@ extern NSString *EntitlementsData(void)
     // verify that it's a 64bit app
     if (executableHeader->magic != MH_MAGIC_64)
     {
-        UE_LOG(LogIOS, Error, TEXT("Executable is NOT 64bit. Entitlement retrieval not supported."));
+		FPlatformMisc::LowLevelOutputDebugString(TEXT("Executable is NOT 64bit. Entitlement retrieval not supported.\n"));
         return nil;
     }
     uintptr_t cursor = (uintptr_t)executableHeader + sizeof(struct mach_header_64);
@@ -1318,7 +1329,7 @@ extern NSString *EntitlementsData(void)
         switch (segmentCommand->cmd)
         {
             case LC_CODE_SIGNATURE:
-                UE_LOG(LogIOS, Log, TEXT("LC_CODE_SIGNATURE found"));
+				FPlatformMisc::LowLevelOutputDebugString(TEXT("LC_CODE_SIGNATURE found\n"));
                 break;
             default:
                 continue;
@@ -1329,7 +1340,7 @@ extern NSString *EntitlementsData(void)
         FILE* file = fopen(ImageName, "rb");
         if (file == NULL)
         {
-            UE_LOG(LogIOS, Error, TEXT("Could not open binary file"));
+			FPlatformMisc::LowLevelOutputDebugString(TEXT("Could not open binary file\n"));
             return nil;
         }
         CS_MultiBlob multiBlob;
@@ -1424,17 +1435,17 @@ extern bool IsEntitlementPresentInEmbeddedProvision(const char *EntitlementsToFi
             
             if(trueptr == NULL && falseptr == NULL)
             {
-                UE_LOG(LogIOS, Error, TEXT("Unexpected Behaviour. The entitlement key is found but its value is not set."));
+				FPlatformMisc::LowLevelOutputDebugString(TEXT("Unexpected Behaviour. The entitlement key is found but its value is not set.\n"));
                 return false;
             }
             if(trueptr && falseptr == NULL)  // only true
             {
-                UE_LOG(LogIOS, Log, TEXT("Entitlements found in embedded mobile provision file."));
+				FPlatformMisc::LowLevelOutputDebugString(TEXT("Entitlements found in embedded mobile provision file.\n"));
                 return true;
             }
             if(trueptr == NULL && falseptr) // only false
             {
-                UE_LOG(LogIOS, Log, TEXT("Entitlements found but set to false."));
+				FPlatformMisc::LowLevelOutputDebugString(TEXT("Entitlements found but set to false.\n"));
                 return false;
             }
             return (trueptr < falseptr); // return true if true comes before false
@@ -1464,7 +1475,7 @@ bool FIOSPlatformMisc::IsEntitlementEnabled(const char * EntitlementToCheck)
 
     if ([CleanedEntitlementData rangeOfString: (@"%s", EntitlementsToFind)].location == NSNotFound)
     {
-        UE_LOG(LogIOS, Log, TEXT("Entitlements not found in binary Mach-O header. Looking at the embedded mobile provision file."));
+		FPlatformMisc::LowLevelOutputDebugString(TEXT("Entitlements not found in binary Mach-O header. Looking at the embedded mobile provision file.\n"));
         return IsEntitlementPresentInEmbeddedProvision(EntitlementToCheck);
     }
     else
diff --git a/Engine/Source/Runtime/IOS/IOSRuntimeSettings/Classes/IOSRuntimeSettings.h b/Engine/Source/Runtime/IOS/IOSRuntimeSettings/Classes/IOSRuntimeSettings.h
--- a/Engine/Source/Runtime/IOS/IOSRuntimeSettings/Classes/IOSRuntimeSettings.h
+++ b/Engine/Source/Runtime/IOS/IOSRuntimeSettings/Classes/IOSRuntimeSettings.h
@@ -258,6 +258,10 @@ public:
     UPROPERTY(GlobalConfig, EditAnywhere, Category = Build)
     FString AdditionalPlistData;
     
+	// List of any entitlements we require. The app will not initialize if the specified entitlements are not available.
+	UPROPERTY(GlobalConfig, EditAnywhere, Category = Build)
+	TArray<FString> RequiredEntitlements;
+
     // Whether or not to add support for iPad devices
     UPROPERTY(GlobalConfig, EditAnywhere, Category = Build, meta = (DisplayName = "Supports iPad"))
     bool bSupportsIPad;
