Ensure that 32-bit platforms are not supported.

diff --git a/Engine/Source/Runtime/Core/Public/HAL/Platform.h b/Engine/Source/Runtime/Core/Public/HAL/Platform.h
--- a/Engine/Source/Runtime/Core/Public/HAL/Platform.h
+++ b/Engine/Source/Runtime/Core/Public/HAL/Platform.h
@@ -637,6 +637,9 @@
 
 // These is computed, not predefined
 #define PLATFORM_32BITS					(!PLATFORM_64BITS)
+#if !PLATFORM_64BITS
+	#error "UE only supports 64-bit platforms"
+#endif
 
 // not supported by the platform system yet or maybe ever
 #define PLATFORM_VTABLE_AT_END_OF_CLASS 0
