Made FUtf8String::Printf's and Appendf's formatting strings be ANSI, because there's no printf which takes a char8_t formatting string.  UTF-8-formatted string arguments are still possible via %s.

diff --git a/Engine/Source/Runtime/Core/Private/Containers/String.cpp.inl b/Engine/Source/Runtime/Core/Private/Containers/String.cpp.inl
--- a/Engine/Source/Runtime/Core/Private/Containers/String.cpp.inl
+++ b/Engine/Source/Runtime/Core/Private/Containers/String.cpp.inl
@@ -1125,7 +1125,7 @@ UE_STRING_CLASS UE_STRING_CLASS::FromBlob(const uint8* SrcBuffer,const uint32 Sr
 	// Convert and append each byte in the buffer
 	for (uint32 Count = 0; Count < SrcSize; Count++)
 	{
-		Result += UE_STRING_CLASS::Printf(CHARTEXT(ElementType, "%03d"),(uint8)SrcBuffer[Count]);
+		Result += UE_STRING_CLASS::Printf(CHARTEXT(FmtCharType, "%03d"),(uint8)SrcBuffer[Count]);
 	}
 	return Result;
 }
@@ -1160,7 +1160,7 @@ UE_STRING_CLASS UE_STRING_CLASS::FromHexBlob( const uint8* SrcBuffer, const uint
 	// Convert and append each byte in the buffer
 	for (uint32 Count = 0; Count < SrcSize; Count++)
 	{
-		Result += UE_STRING_CLASS::Printf( CHARTEXT(ElementType,  "%02X" ), (uint8)SrcBuffer[Count] );
+		Result += UE_STRING_CLASS::Printf(CHARTEXT(FmtCharType, "%02X"), (uint8)SrcBuffer[Count]);
 	}
 	return Result;
 }
@@ -1194,7 +1194,7 @@ UE_STRING_CLASS UE_STRING_CLASS::SanitizeFloat( double InFloat, const int32 InMi
 	UE::Core::Private::StripNegativeZero(InFloat);
 
 	// First create the string
-	UE_STRING_CLASS TempString = UE_STRING_CLASS::Printf(CHARTEXT(ElementType, "%f"), InFloat);
+	UE_STRING_CLASS TempString = UE_STRING_CLASS::Printf(CHARTEXT(FmtCharType, "%f"), InFloat);
 	if (!TempString.IsNumeric())
 	{
 		// String did not format as a valid decimal number so avoid messing with it
diff --git a/Engine/Source/Runtime/Core/Public/Containers/AnsiString.h b/Engine/Source/Runtime/Core/Public/Containers/AnsiString.h
--- a/Engine/Source/Runtime/Core/Public/Containers/AnsiString.h
+++ b/Engine/Source/Runtime/Core/Public/Containers/AnsiString.h
@@ -8,9 +8,11 @@
 #define UE_STRING_CLASS                        FAnsiString
 #define UE_STRING_CHARTYPE                     ANSICHAR
 #define UE_STRING_CHARTYPE_IS_TCHAR            0
+#define UE_STRING_PRINTF_FMT_CHARTYPE          ANSICHAR
 #define UE_STRING_DEPRECATED(Version, Message) UE_DEPRECATED(Version, Message)
 	#include "Containers/UnrealString.h.inl"
 #undef UE_STRING_DEPRECATED
+#undef UE_STRING_PRINTF_FMT_CHARTYPE
 #undef UE_STRING_CHARTYPE_IS_TCHAR
 #undef UE_STRING_CHARTYPE
 #undef UE_STRING_CLASS
diff --git a/Engine/Source/Runtime/Core/Public/Containers/UnrealString.h b/Engine/Source/Runtime/Core/Public/Containers/UnrealString.h
--- a/Engine/Source/Runtime/Core/Public/Containers/UnrealString.h
+++ b/Engine/Source/Runtime/Core/Public/Containers/UnrealString.h
@@ -11,9 +11,11 @@
 #define UE_STRING_CLASS                        FString
 #define UE_STRING_CHARTYPE                     TCHAR
 #define UE_STRING_CHARTYPE_IS_TCHAR            1
+#define UE_STRING_PRINTF_FMT_CHARTYPE          TCHAR
 #define UE_STRING_DEPRECATED(Version, Message)
 	#include "Containers/UnrealString.h.inl"
 #undef UE_STRING_DEPRECATED
+#undef UE_STRING_PRINTF_FMT_CHARTYPE
 #undef UE_STRING_CHARTYPE_IS_TCHAR
 #undef UE_STRING_CHARTYPE
 #undef UE_STRING_CLASS
diff --git a/Engine/Source/Runtime/Core/Public/Containers/UnrealString.h.inl b/Engine/Source/Runtime/Core/Public/Containers/UnrealString.h.inl
--- a/Engine/Source/Runtime/Core/Public/Containers/UnrealString.h.inl
+++ b/Engine/Source/Runtime/Core/Public/Containers/UnrealString.h.inl
@@ -23,6 +23,9 @@
 #ifndef UE_STRING_CHARTYPE_IS_TCHAR
 	#error "UnrealString.h.inl should only be included after defining UE_STRING_CHARTYPE_IS_TCHAR"
 #endif
+#ifndef UE_STRING_PRINTF_FMT_CHARTYPE
+	#error "UnrealString.h.inl should only be included after defining UE_STRING_PRINTF_FMT_CHARTYPE"
+#endif
 #ifndef UE_STRING_DEPRECATED
 	#error "UnrealString.h.inl should only be included after defining UE_STRING_DEPRECATED"
 #endif
@@ -53,6 +56,7 @@ class UE_STRING_CLASS
 public:
 	using AllocatorType = TSizedDefaultAllocator<32>;
 	using ElementType   = UE_STRING_CHARTYPE;
+	using FmtCharType   = UE_STRING_PRINTF_FMT_CHARTYPE;
 
 private:
 	/** Array holding the character data */
@@ -61,7 +65,7 @@ private:
 
 	/** Like the TIsCharEncodingCompatibleWithTCHAR trait, but for the element type of the string */
 	template <typename SrcEncoding>
-	using TIsCharEncodingCompatibleWithElementType = TIsCharEncodingCompatibleWith<SrcEncoding, ElementType>;
+	using TIsCharEncodingCompatibleWithPrintfFmt = TIsCharEncodingCompatibleWith<SrcEncoding, UE_STRING_PRINTF_FMT_CHARTYPE>;
 
 public:
 	UE_STRING_CLASS() = default;
@@ -1539,7 +1543,7 @@ public:
 	template <typename FmtType, typename... Types>
 	[[nodiscard]] static UE_STRING_CLASS Printf(const FmtType& Fmt, Types... Args)
 	{
-		static_assert(TIsArrayOrRefOfTypeByPredicate<FmtType, TIsCharEncodingCompatibleWithElementType>::Value, "Formatting string must be a character array.");
+		static_assert(TIsArrayOrRefOfTypeByPredicate<FmtType, TIsCharEncodingCompatibleWithPrintfFmt>::Value, "Formatting string must be a literal " PREPROCESSOR_TO_STRING(UE_STRING_PRINTF_FMT_CHARTYPE) " array.");
 		static_assert((TIsValidVariadicFunctionArg<Types>::Value && ...), "Invalid argument(s) passed to Printf");
 
 		return PrintfImpl((const ElementType*)Fmt, Args...);
@@ -1552,7 +1556,7 @@ public:
 	template <typename FmtType, typename... Types>
 	UE_STRING_CLASS& Appendf(const FmtType& Fmt, Types... Args)
 	{
-		static_assert(TIsArrayOrRefOfTypeByPredicate<FmtType, TIsCharEncodingCompatibleWithElementType>::Value, "Formatting string must be a character array.");
+		static_assert(TIsArrayOrRefOfTypeByPredicate<FmtType, TIsCharEncodingCompatibleWithPrintfFmt>::Value, "Formatting string must be a literal " PREPROCESSOR_TO_STRING(UE_STRING_PRINTF_FMT_CHARTYPE) " array.");
 		static_assert((TIsValidVariadicFunctionArg<Types>::Value && ...), "Invalid argument(s) passed to TString::Appendf");
 
 		AppendfImpl(*this, (const ElementType*)Fmt, Args...);
diff --git a/Engine/Source/Runtime/Core/Public/Containers/Utf8String.h b/Engine/Source/Runtime/Core/Public/Containers/Utf8String.h
--- a/Engine/Source/Runtime/Core/Public/Containers/Utf8String.h
+++ b/Engine/Source/Runtime/Core/Public/Containers/Utf8String.h
@@ -5,12 +5,16 @@
 // Include UnrealString.h.inl's includes before defining the macros, in case the macros 'poison' other headers or there are re-entrant includes.
 #include "Containers/UnrealStringIncludes.h.inl"
 
+// Note that FUtf8String's Printf formatting type is not UTF8CHAR - even though the formatting string must be ANSI,
+// the string still supports UTF8CHAR strings as arguments, e.g. FUtf8String::Printf("Name: %s", Utf8Name)
 #define UE_STRING_CLASS                        FUtf8String
 #define UE_STRING_CHARTYPE                     UTF8CHAR
 #define UE_STRING_CHARTYPE_IS_TCHAR            0
+#define UE_STRING_PRINTF_FMT_CHARTYPE          ANSICHAR
 #define UE_STRING_DEPRECATED(Version, Message) UE_DEPRECATED(Version, Message)
 	#include "Containers/UnrealString.h.inl"
 #undef UE_STRING_DEPRECATED
+#undef UE_STRING_PRINTF_FMT_CHARTYPE
 #undef UE_STRING_CHARTYPE_IS_TCHAR
 #undef UE_STRING_CHARTYPE
 #undef UE_STRING_CLASS
diff --git a/Engine/Source/Runtime/CoreUObject/Private/UObject/StrProperty.cpp.inl b/Engine/Source/Runtime/CoreUObject/Private/UObject/StrProperty.cpp.inl
--- a/Engine/Source/Runtime/CoreUObject/Private/UObject/StrProperty.cpp.inl
+++ b/Engine/Source/Runtime/CoreUObject/Private/UObject/StrProperty.cpp.inl
@@ -97,7 +97,7 @@ void UE_STRPROPERTY_CLASS::ExportText_Internal( FString& ValueStr, const void* P
 	}
 	else if ( StringValue.Len() > 0 )
 	{
-		ValueStr += UE_STRPROPERTY_STRINGTYPE::Printf( CHARTEXT(UE_STRPROPERTY_STRINGTYPE::ElementType, "\"%s\""), *StringValue.ReplaceCharWithEscapedChar() );
+		ValueStr += UE_STRPROPERTY_STRINGTYPE::Printf( CHARTEXT(UE_STRPROPERTY_STRINGTYPE::FmtCharType, "\"%s\""), *StringValue.ReplaceCharWithEscapedChar() );
 	}
 	else
 	{
