Always OPEN around these ensure macros

diff --git a/Engine/Source/Runtime/Core/Private/Misc/AssertionMacros.cpp b/Engine/Source/Runtime/Core/Private/Misc/AssertionMacros.cpp
--- a/Engine/Source/Runtime/Core/Private/Misc/AssertionMacros.cpp
+++ b/Engine/Source/Runtime/Core/Private/Misc/AssertionMacros.cpp
@@ -9,6 +9,7 @@
 #include "Misc/Crc.h"
 #include "Async/UniqueLock.h"
 #include "Async/WordMutex.h"
+#include "AutoRTFM/AutoRTFM.h"
 #include "Containers/UnrealString.h"
 #include "Containers/StringConv.h"
 #include "GenericPlatform/GenericPlatformStackWalk.h"
@@ -776,6 +777,7 @@ bool UE_DEBUG_SECTION VARARGS CheckVerifyImpl(std::atomic<bool>& bExecuted, bool
 #endif
 }
 
+UE_AUTORTFM_ALWAYS_OPEN
 bool UE_DEBUG_SECTION UE::Assert::Private::ExecCheckImplInternal(std::atomic<bool>& bExecuted, bool bAlways, const ANSICHAR* File, int32 Line, const ANSICHAR* Expr)
 {
 	if (((bAlways && GEnsureAlwaysEnabled) || !bExecuted.load(std::memory_order_relaxed)) && FPlatformMisc::IsEnsureAllowed())
@@ -792,6 +794,7 @@ bool UE_DEBUG_SECTION UE::Assert::Private::ExecCheckImplInternal(std::atomic<boo
 	return false;
 }
 
+UE_AUTORTFM_ALWAYS_OPEN
 bool UE_DEBUG_SECTION VARARGS UE::Assert::Private::EnsureFailed(std::atomic<bool>& bExecuted, const FStaticEnsureRecord* Ensure, ...)
 {
 	if (bExecuted.exchange(true, std::memory_order_release) && !(Ensure->bAlways && GEnsureAlwaysEnabled))
