Switch header to IntrusiveUnsetOptionalState.h rather than Optional.h

diff --git a/Engine/Source/Runtime/Core/Public/UObject/NameTypes.h b/Engine/Source/Runtime/Core/Public/UObject/NameTypes.h
--- a/Engine/Source/Runtime/Core/Public/UObject/NameTypes.h
+++ b/Engine/Source/Runtime/Core/Public/UObject/NameTypes.h
@@ -14,7 +14,7 @@
 #include "UObject/UnrealNames.h"
 #include "Templates/Atomic.h"
 #include "Serialization/MemoryLayout.h"
-#include "Misc/Optional.h"
+#include "Misc/IntrusiveUnsetOptionalState.h"
 #include "Misc/StringBuilder.h"
 #include "Trace/Trace.h"
 
