Fixed a shadowed variable warning when an ensure macro is used in the same scope as a variable named bExecuted.

diff --git a/Engine/Source/Runtime/Core/Public/Misc/AssertionMacros.h b/Engine/Source/Runtime/Core/Public/Misc/AssertionMacros.h
--- a/Engine/Source/Runtime/Core/Public/Misc/AssertionMacros.h
+++ b/Engine/Source/Runtime/Core/Public/Misc/AssertionMacros.h
@@ -404,16 +404,16 @@ RetType FORCENOINLINE UE_DEBUG_SECTION DispatchCheckVerify(InnerType&& Inner, Ar
 
 	#define UE_ENSURE_IMPL(Always, InExpression) \
 		(LIKELY(!!(InExpression)) \
-			|| (::UE::Assert::Private::ExecCheckImplInternal([]() UE_DEBUG_SECTION -> std::atomic<bool>& { static std::atomic<bool> bExecuted = false; return bExecuted; } (), Always, __FILE__, __LINE__, #InExpression) \
+			|| (::UE::Assert::Private::ExecCheckImplInternal([]() UE_DEBUG_SECTION -> std::atomic<bool>& { static std::atomic<bool> ENSURE_bExecuted = false; return ENSURE_bExecuted; } (), Always, __FILE__, __LINE__, #InExpression) \
 			&& [] () { PLATFORM_BREAK(); return false; } ()))
 
 	#define UE_ENSURE_IMPL2(Capture, Always, InExpression, InFormat, ...) \
 		(LIKELY(!!(InExpression)) || ([Capture] () UE_DEBUG_SECTION \
 		{ \
 			UE_VALIDATE_FORMAT_STRING(InFormat, ##__VA_ARGS__); \
-			static std::atomic<bool> bExecuted = false; \
+			static std::atomic<bool> ENSURE_bExecuted = false; \
 			static constexpr ::UE::Assert::Private::FStaticEnsureRecord ENSURE_Static(InFormat, #InExpression, __builtin_FILE(), __builtin_LINE(), Always); \
-			if ((Always || !bExecuted.load(std::memory_order_relaxed)) && FPlatformMisc::IsEnsureAllowed() && ::UE::Assert::Private::EnsureFailed(bExecuted, &ENSURE_Static, ##__VA_ARGS__)) \
+			if ((Always || !ENSURE_bExecuted.load(std::memory_order_relaxed)) && FPlatformMisc::IsEnsureAllowed() && ::UE::Assert::Private::EnsureFailed(ENSURE_bExecuted, &ENSURE_Static, ##__VA_ARGS__)) \
 			{ \
 				PLATFORM_BREAK(); \
 			} \
