Add methods to FHashTable for getting hash and index size for users that want to grow the table after the fact.
Add a move constructor.

diff --git a/Engine/Source/Runtime/Core/Public/Containers/HashTable.h b/Engine/Source/Runtime/Core/Public/Containers/HashTable.h
--- a/Engine/Source/Runtime/Core/Public/Containers/HashTable.h
+++ b/Engine/Source/Runtime/Core/Public/Containers/HashTable.h
@@ -195,13 +195,20 @@ class FHashTable
 {
 public:
 					FHashTable( uint32 InHashSize = 1024, uint32 InIndexSize = 0 );
-					FHashTable( const FHashTable& Other );
+					FHashTable(const FHashTable& Other);
+					FHashTable(FHashTable&& Other);
 					~FHashTable();
 
 	void			Clear();
 	void			Clear( uint32 InHashSize, uint32 InIndexSize = 0 );
 	void			Free();
+	/** 
+	 * Increases or decreases the size of the index but not the hash lookup.
+	 * If the previous size was empty, allocates the hash at its desired size.
+	 */
 	CORE_API void	Resize( uint32 NewIndexSize );
+	inline uint32	GetIndexSize() { return IndexSize; }
+	inline uint32	GetHashSize() { return HashSize; }
 
 	// Functions used to search
 	uint32			First( uint32 Key ) const;
@@ -209,6 +216,7 @@ public:
 	bool			IsValid( uint32 Index ) const;
 	
 	void			Add( uint32 Key, uint32 Index );
+	// Safe to call concurrently with other threads calling Add_Concurrent with different values for Index.
 	void			Add_Concurrent( uint32 Key, uint32 Index );
 	void			Remove( uint32 Key, uint32 Index );
 
@@ -216,6 +224,7 @@ public:
 	CORE_API float	AverageSearch() const;
 
 	FHashTable&		operator=(const FHashTable& Other);
+	FHashTable&		operator=(FHashTable&& Other);
 
 protected:
 	// Avoids allocating hash until first add
@@ -267,6 +276,20 @@ FORCEINLINE FHashTable::FHashTable( const FHashTable& Other )
 	}
 }
 
+FORCEINLINE FHashTable::FHashTable(FHashTable&& Other )
+	: HashSize( Other.HashSize )
+	, HashMask( Other.HashMask )
+	, IndexSize( Other.IndexSize )
+	, Hash(Other.Hash)
+	, NextIndex(Other.NextIndex)
+{
+	Other.HashSize = 0;
+	Other.HashMask = 0;
+	Other.IndexSize = 0;
+	Other.Hash = EmptyHash;
+	Other.NextIndex = nullptr;
+}
+
 FORCEINLINE FHashTable& FHashTable::operator=(const FHashTable& Other)
 {
 	Free();
@@ -285,6 +308,24 @@ FORCEINLINE FHashTable& FHashTable::operator=(const FHashTable& Other)
 	return *this;
 }
 
+FORCEINLINE FHashTable& FHashTable::operator=(FHashTable&& Other)
+{
+	Free();
+
+	HashSize = Other.HashSize;
+	HashMask = Other.HashMask;
+	IndexSize = Other.IndexSize;
+	Hash = Other.Hash;
+	NextIndex = Other.NextIndex;
+
+	Other.HashSize = 0;
+	Other.HashMask = 0;
+	Other.IndexSize = 0;
+	Other.Hash = EmptyHash;
+	Other.NextIndex = nullptr;
+	return *this;
+}
+
 FORCEINLINE FHashTable::~FHashTable()
 {
 	Free();
diff --git a/Engine/Source/Runtime/Core/Tests/Containers/HashTableTest.cpp b/Engine/Source/Runtime/Core/Tests/Containers/HashTableTest.cpp
new file mode 100644
--- /dev/null
+++ b/Engine/Source/Runtime/Core/Tests/Containers/HashTableTest.cpp
@@ -0,0 +1,75 @@
+// Copyright Epic Games, Inc. All Rights Reserved.
+
+#if WITH_TESTS
+
+#include "Containers/HashTable.h"
+#include "Tests/TestHarnessAdapter.h"
+
+TEST_CASE_NAMED(FHashTableTestMove, "System::Core::Containers::HashTable::Move", "[ApplicationContextMask][SmokeFilter]")
+{
+    // Hash values by address and store in hash table
+    int32 Values[] = { 1, 2, 3, 4, 5 };
+    FHashTable A(4, 16);
+    for (int32 i=0; i < UE_ARRAY_COUNT(Values); ++i)
+    {
+        A.Add(GetTypeHash(&Values[i]), i);
+    }
+
+    FHashTable B = MoveTemp(A);
+    CHECK(A.GetIndexSize() == 0);
+    CHECK(A.GetHashSize() == 0);
+    CHECK(B.GetIndexSize() == 16);
+    CHECK(B.GetHashSize() == 4);
+
+    auto Find = [&Values](const FHashTable& Table, int32 TargetIndex)
+    {
+        for (int32 Index = Table.First(GetTypeHash(&Values[TargetIndex])); Table.IsValid(Index); Index = Table.Next(Index))
+        {
+            if (Index == TargetIndex)
+            {
+                return true;
+            }
+        }
+        return false;
+    };
+    for (int32 i=0; i < UE_ARRAY_COUNT(Values); ++i)
+    {
+        CHECK_FALSE(Find(A, i));
+        CHECK(Find(B, i));
+    }
+}
+
+TEST_CASE_NAMED(FHashTableTestGrow, "System::Core::Containers::HashTable::Grow", "[ApplicationContextMask][SmokeFilter]")
+{
+    // Hash values by address and store in hash table
+    int32 Values[] = { 1, 2, 3, 4, 5 };
+    FHashTable Table(4, 16);
+    for (int32 i=0; i < UE_ARRAY_COUNT(Values); ++i)
+    {
+        Table.Add(GetTypeHash(&Values[i]), i);
+    }
+
+    auto Find = [&Values](const FHashTable& Table, int32 TargetIndex)
+    {
+        for (int32 Index = Table.First(GetTypeHash(&Values[TargetIndex])); Table.IsValid(Index); Index = Table.Next(Index))
+        {
+            if (Index == TargetIndex)
+            {
+                return true;
+            }
+        }
+        return false;
+    };
+    for (int32 i=0; i < UE_ARRAY_COUNT(Values); ++i)
+    {
+        CHECK(Find(Table, i));
+    }
+
+    Table.Resize(128);
+    for (int32 i=0; i < UE_ARRAY_COUNT(Values); ++i)
+    {
+        CHECK(Find(Table, i));
+    }
+}
+
+#endif // WITH_TESTS
