Fix FramePro, CpuProfiler for AutoRTFM

diff --git a/Engine/Source/Runtime/Core/Private/FramePro/FramePro.cpp b/Engine/Source/Runtime/Core/Private/FramePro/FramePro.cpp
--- a/Engine/Source/Runtime/Core/Private/FramePro/FramePro.cpp
+++ b/Engine/Source/Runtime/Core/Private/FramePro/FramePro.cpp
@@ -20,6 +20,8 @@ OR IN CONNECTION WITH THE USE OR PERFORMANCE OF THIS SOFTWARE.
 	#include "FramePro.h"
 #endif
 
+#include "AutoRTFM/AutoRTFM.h"
+
 //------------------------------------------------------------------------
 #if FRAMEPRO_ENABLED
 
@@ -3906,6 +3908,7 @@ void FramePro::AddTimeSpan(const wchar_t* p_name_and_source_info, int64 start_ti
 }
 
 //------------------------------------------------------------------------
+UE_AUTORTFM_ALWAYS_OPEN
 void FramePro::AddTimeSpan(StringId name, const char* p_source_info, int64 start_time, int64 end_time)
 {
 	FRAMEPRO_ASSERT(start_time <= end_time);
@@ -7197,6 +7200,7 @@ namespace FramePro
 	}
 
 	//------------------------------------------------------------------------
+	UE_AUTORTFM_ALWAYS_OPEN
 	StringId FrameProTLS::RegisterString(const char* p_str)
 	{
 		FRAMEPRO_ASSERT(IsOnTLSThread());
@@ -7220,6 +7224,7 @@ namespace FramePro
 	}
 
 	//------------------------------------------------------------------------
+	UE_AUTORTFM_ALWAYS_OPEN
 	StringId FrameProTLS::RegisterString(const wchar_t* p_str)
 	{
 		FRAMEPRO_ASSERT(IsOnTLSThread());
diff --git a/Engine/Source/Runtime/Core/Private/FramePro/FrameProProfiler.cpp b/Engine/Source/Runtime/Core/Private/FramePro/FrameProProfiler.cpp
--- a/Engine/Source/Runtime/Core/Private/FramePro/FrameProProfiler.cpp
+++ b/Engine/Source/Runtime/Core/Private/FramePro/FrameProProfiler.cpp
@@ -14,6 +14,7 @@
 #include "Misc/EngineVersion.h"
 #include "Misc/Paths.h"
 #include "Runtime/Launch/Resources/Version.h"
+#include "AutoRTFM/AutoRTFM.h"
 
 DEFINE_LOG_CATEGORY_STATIC(LogFramePro, Log, All);
 
@@ -130,11 +131,28 @@ public:
 
 	virtual void* Alloc(size_t size) override
 	{
-		return BaseMalloc.Malloc(size, DEFAULT_ALIGNMENT);
+		void* Ptr = nullptr;
+		UE_AUTORTFM_OPEN(
+		{
+			Ptr = BaseMalloc.Malloc(size, DEFAULT_ALIGNMENT);
+		});
+
+		// no-op for non-transactional code
+		UE_AUTORTFM_ONABORT(
+		{
+			// Disable the code analysis warning that complains that Free is being passed
+			// a pointer that may be null. Free explicitly handles this case already.
+			this->Free(Ptr); //-V575
+		});
+
+		return Ptr;
 	}
 	virtual void Free(void* p) override
 	{
-		return BaseMalloc.Free(p);
+		UE_AUTORTFM_ONCOMMIT(
+		{
+			BaseMalloc.Free(p);
+		});
 	}
 
 private:
@@ -147,17 +165,61 @@ class FrameProMalloc
 public:
 	static void* Malloc(SIZE_T Count, uint32 Alignment = DEFAULT_ALIGNMENT)
 	{
-		return FrameProAllocator::Get().GetBaseMalloc()->Malloc(Count, Alignment);
+		void* Ptr = nullptr;
+		UE_AUTORTFM_OPEN(
+		{
+			Ptr = FrameProAllocator::Get().GetBaseMalloc()->Malloc(Count, Alignment);
+		});
+
+		// no-op for non-transactional code
+		UE_AUTORTFM_ONABORT(
+		{
+			// Disable the code analysis warning that complains that Free is being passed
+			// a pointer that may be null. Free explicitly handles this case already.
+			FrameProMalloc::Free(Ptr); //-V575
+		});
+
+		return Ptr;
 	}
 
 	static void* Realloc(void* Original, SIZE_T Count, uint32 Alignment = DEFAULT_ALIGNMENT)
 	{
-		return FrameProAllocator::Get().GetBaseMalloc()->Realloc(Original, Count, Alignment);
+		FMalloc* BaseMalloc = FrameProAllocator::Get().GetBaseMalloc();
+
+		// This logic is from FMemory.inl
+		if (AutoRTFM::IsClosed())
+		{
+			void* Ptr = FrameProMalloc::Malloc(Count, Alignment);
+			if (!Ptr)
+			{
+				return nullptr;
+			}
+
+			if (Original)
+			{
+				SIZE_T OriginalCount = 0;
+				UE_AUTORTFM_OPEN(
+				{
+					BaseMalloc->GetAllocationSize(Original, OriginalCount);
+				});
+
+				SIZE_T CopyCount = FGenericPlatformMath::Min(Count, OriginalCount);
+				FMemory::Memcpy(Ptr, Original, CopyCount);
+				FrameProMalloc::Free(Original);
+			}
+
+			return Ptr;
+		}
+
+		return BaseMalloc->Realloc(Original, Count, Alignment);
 	}
 
 	static void Free(void* Original)
 	{
-		return FrameProAllocator::Get().GetBaseMalloc()->Free(Original);
+		UE_AUTORTFM_ONCOMMIT(
+		{
+			FrameProAllocator::Get().GetBaseMalloc()->Free(Original);
+		});
 	}
 };
 
diff --git a/Engine/Source/Runtime/Core/Private/ProfilingDebugging/CpuProfilerTrace.cpp b/Engine/Source/Runtime/Core/Private/ProfilingDebugging/CpuProfilerTrace.cpp
--- a/Engine/Source/Runtime/Core/Private/ProfilingDebugging/CpuProfilerTrace.cpp
+++ b/Engine/Source/Runtime/Core/Private/ProfilingDebugging/CpuProfilerTrace.cpp
@@ -11,6 +11,7 @@
 #include "Misc/MemStack.h"
 #include "Misc/Crc.h"
 #include "UObject/NameTypes.h"
+#include "AutoRTFM/AutoRTFM.h"
 
 #if CPUPROFILERTRACE_ENABLED
 
@@ -132,6 +133,7 @@ FCpuProfilerTraceInternal::FThreadBuffer* FCpuProfilerTraceInternal::CreateThrea
 	return ThreadBuffer;
 }
 
+UE_AUTORTFM_ALWAYS_OPEN
 void FCpuProfilerTraceInternal::FlushThreadBuffer(FThreadBuffer* InThreadBuffer)
 {
 	UE_TRACE_LOG(CpuProfiler, EventBatchV2, true)
@@ -313,6 +315,7 @@ uint32 FCpuProfilerTraceInternal::GetNextSpecId()
 	return (NextSpecId++) + 1;
 }
 
+UE_AUTORTFM_ALWAYS_OPEN
 uint32 FCpuProfilerTrace::OutputEventType(const TCHAR* Name, const ANSICHAR* File, uint32 Line)
 {
 	uint32 SpecId = FCpuProfilerTraceInternal::GetNextSpecId();
@@ -333,6 +336,7 @@ uint32 FCpuProfilerTrace::OutputEventType(const TCHAR* Name, const ANSICHAR* Fil
 	return SpecId;
 }
 
+UE_AUTORTFM_ALWAYS_OPEN
 uint32 FCpuProfilerTrace::OutputEventType(const ANSICHAR* Name, const ANSICHAR* File, uint32 Line)
 {
 	uint32 SpecId = FCpuProfilerTraceInternal::GetNextSpecId();
diff --git a/Engine/Source/Runtime/TraceLog/TraceLog.Build.cs b/Engine/Source/Runtime/TraceLog/TraceLog.Build.cs
--- a/Engine/Source/Runtime/TraceLog/TraceLog.Build.cs
+++ b/Engine/Source/Runtime/TraceLog/TraceLog.Build.cs
@@ -12,5 +12,7 @@ public class TraceLog : ModuleRules
 
 		bRequiresImplementModule = false;
 		PublicIncludePathModuleNames.Add("Core");
+
+		bAllowAutoRTFMInstrumentation = true;
 	}
 }
