Moved live coding defines to only exist in LaunchModule's Definitions.h to make builds more deterministic.

diff --git a/Engine/Source/Programs/UnrealBuildTool/Configuration/UEBuildTarget.cs b/Engine/Source/Programs/UnrealBuildTool/Configuration/UEBuildTarget.cs
--- a/Engine/Source/Programs/UnrealBuildTool/Configuration/UEBuildTarget.cs
+++ b/Engine/Source/Programs/UnrealBuildTool/Configuration/UEBuildTarget.cs
@@ -4781,6 +4781,16 @@ namespace UnrealBuildTool
 				OutputDirectory = GetOutputDirectoryForExecutable(Unreal.EngineDirectory, Rules.File);
 			}
 
+			// We only want these defines in the launch module to not create non-deterministic Definitions.h on other modules
+			if (Rules.bWithLiveCoding && Rules.LinkType == TargetLinkType.Monolithic)
+			{
+				LaunchModule.Rules.PrivateDefinitions.Add(String.Format("UE_LIVE_CODING_ENGINE_DIR=\"{0}\"", Unreal.EngineDirectory.FullName.Replace("\\", "\\\\")));
+				if (Rules.ProjectFile != null)
+				{
+					LaunchModule.Rules.PrivateDefinitions.Add(String.Format("UE_LIVE_CODING_PROJECT=\"{0}\"", Rules.ProjectFile.FullName.Replace("\\", "\\\\")));
+				}
+			}
+
 			bool bCompileAsDLL = Rules.bShouldCompileAsDLL && bCompileMonolithic;
 			List<FileReference> OutputPaths = MakeBinaryPaths(OutputDirectory, bCompileMonolithic ? TargetName : AppName, Platform, Configuration, bCompileAsDLL ? UEBuildBinaryType.DynamicLinkLibrary : UEBuildBinaryType.Executable, Rules.Architectures, Rules.UndecoratedConfiguration, bCompileMonolithic && ProjectFile != null, Rules.ExeBinariesSubFolder, ProjectFile, Rules);
 
@@ -5214,14 +5224,6 @@ namespace UnrealBuildTool
 			if (Rules.bWithLiveCoding)
 			{
 				GlobalCompileEnvironment.Definitions.Add("WITH_LIVE_CODING=1");
-				if (Rules.LinkType == TargetLinkType.Monolithic)
-				{
-					GlobalCompileEnvironment.Definitions.Add(String.Format("UE_LIVE_CODING_ENGINE_DIR=\"{0}\"", Unreal.EngineDirectory.FullName.Replace("\\", "\\\\")));
-					if (ProjectFile != null)
-					{
-						GlobalCompileEnvironment.Definitions.Add(String.Format("UE_LIVE_CODING_PROJECT=\"{0}\"", ProjectFile.FullName.Replace("\\", "\\\\")));
-					}
-				}
 			}
 			else
 			{
diff --git a/Engine/Source/Runtime/Core/Public/Modules/ModuleManager.h b/Engine/Source/Runtime/Core/Public/Modules/ModuleManager.h
--- a/Engine/Source/Runtime/Core/Public/Modules/ModuleManager.h
+++ b/Engine/Source/Runtime/Core/Public/Modules/ModuleManager.h
@@ -854,21 +854,6 @@ class FDefaultGameModuleImpl
 	#define IMPLEMENT_FOREIGN_ENGINE_DIR() 
 #endif
 
-/**
- * Macros for setting the source directories for live coding builds. This allows locally packaging a target and patching code into it.
- */
-#ifdef UE_LIVE_CODING_ENGINE_DIR
-	#define IMPLEMENT_LIVE_CODING_ENGINE_DIR() const TCHAR* GLiveCodingEngineDir = TEXT(UE_LIVE_CODING_ENGINE_DIR);
-	#ifdef UE_LIVE_CODING_PROJECT
-		#define IMPLEMENT_LIVE_CODING_PROJECT() const TCHAR* GLiveCodingProject = TEXT(UE_LIVE_CODING_PROJECT);
-	#else
-		#define IMPLEMENT_LIVE_CODING_PROJECT() const TCHAR* GLiveCodingProject = nullptr;
-	#endif
-#else
-	#define IMPLEMENT_LIVE_CODING_ENGINE_DIR()
-	#define IMPLEMENT_LIVE_CODING_PROJECT()
-#endif
-
 /**
  * Macro for passing a list argument to a macro
  */
@@ -939,8 +924,6 @@ class FDefaultGameModuleImpl
 			/* For monolithic builds, we must statically define the game's name string (See Core.h) */ \
 			TCHAR GInternalProjectName[64] = TEXT( GameName ); \
 			IMPLEMENT_FOREIGN_ENGINE_DIR() \
-			IMPLEMENT_LIVE_CODING_ENGINE_DIR() \
-			IMPLEMENT_LIVE_CODING_PROJECT() \
 			IMPLEMENT_SIGNING_KEY_REGISTRATION() \
 			IMPLEMENT_ENCRYPTION_KEY_REGISTRATION() \
 			IMPLEMENT_GAME_MODULE(FDefaultGameModuleImpl, ModuleName) \
@@ -958,8 +941,6 @@ class FDefaultGameModuleImpl
 					FCString::Strncpy(GInternalProjectName, TEXT( GameName ), UE_ARRAY_COUNT(GInternalProjectName)); \
 				} \
 			} AutoSet##ModuleName; \
-			IMPLEMENT_LIVE_CODING_ENGINE_DIR() \
-			IMPLEMENT_LIVE_CODING_PROJECT() \
 			PER_MODULE_BOILERPLATE \
 			PER_MODULE_BOILERPLATE_ANYLINK(FDefaultGameModuleImpl, ModuleName) \
 			FEngineLoop GEngineLoop;
@@ -979,8 +960,6 @@ class FDefaultGameModuleImpl
 			/* Implement the GIsGameAgnosticExe variable (See Core.h). */ \
 			bool GIsGameAgnosticExe = false; \
 			IMPLEMENT_FOREIGN_ENGINE_DIR() \
-			IMPLEMENT_LIVE_CODING_ENGINE_DIR() \
-			IMPLEMENT_LIVE_CODING_PROJECT() \
 			IMPLEMENT_SIGNING_KEY_REGISTRATION() \
 			IMPLEMENT_ENCRYPTION_KEY_REGISTRATION() \
 			IMPLEMENT_TARGET_NAME_REGISTRATION() \
@@ -994,8 +973,6 @@ class FDefaultGameModuleImpl
 			TCHAR GInternalProjectName[64] = TEXT( PREPROCESSOR_TO_STRING(UE_PROJECT_NAME) ); \
 			PER_MODULE_BOILERPLATE \
 			IMPLEMENT_FOREIGN_ENGINE_DIR() \
-			IMPLEMENT_LIVE_CODING_ENGINE_DIR() \
-			IMPLEMENT_LIVE_CODING_PROJECT() \
 			IMPLEMENT_SIGNING_KEY_REGISTRATION() \
 			IMPLEMENT_ENCRYPTION_KEY_REGISTRATION() \
 			IMPLEMENT_TARGET_NAME_REGISTRATION() \
diff --git a/Engine/Source/Runtime/Launch/Private/Launch.cpp b/Engine/Source/Runtime/Launch/Private/Launch.cpp
--- a/Engine/Source/Runtime/Launch/Private/Launch.cpp
+++ b/Engine/Source/Runtime/Launch/Private/Launch.cpp
@@ -21,6 +21,14 @@
 	#include "Windows/WindowsHWrapper.h"
 #endif
 
+#ifdef UE_LIVE_CODING_ENGINE_DIR
+const TCHAR* GLiveCodingEngineDir = TEXT(UE_LIVE_CODING_ENGINE_DIR);
+#ifdef UE_LIVE_CODING_PROJECT
+	const TCHAR* GLiveCodingProject = TEXT(UE_LIVE_CODING_PROJECT);
+#else
+	const TCHAR* GLiveCodingProject = nullptr;
+#endif
+#endif
 
 IMPLEMENT_MODULE(FDefaultModuleImpl, Launch);
 
diff --git a/Engine/Source/Runtime/UnrealGame/Private/UnrealGame.cpp b/Engine/Source/Runtime/UnrealGame/Private/UnrealGame.cpp
--- a/Engine/Source/Runtime/UnrealGame/Private/UnrealGame.cpp
+++ b/Engine/Source/Runtime/UnrealGame/Private/UnrealGame.cpp
@@ -13,6 +13,4 @@ PER_MODULE_BOILERPLATE
 bool GIsGameAgnosticExe = true;
 TCHAR GInternalProjectName[64] = TEXT("");
 IMPLEMENT_FOREIGN_ENGINE_DIR()
-IMPLEMENT_LIVE_CODING_ENGINE_DIR()
-IMPLEMENT_LIVE_CODING_PROJECT()
 #endif
