[Coroutines]
 - Make sure nothing is compiled-in while coroutines are inactive
 - Optimize IsInGameThread() when coroutines are inactive

diff --git a/Engine/Source/Runtime/Core/Private/Experimental/Coroutine/Coroutine.cpp b/Engine/Source/Runtime/Core/Private/Experimental/Coroutine/Coroutine.cpp
--- a/Engine/Source/Runtime/Core/Private/Experimental/Coroutine/Coroutine.cpp
+++ b/Engine/Source/Runtime/Core/Private/Experimental/Coroutine/Coroutine.cpp
@@ -4,6 +4,8 @@
 
 #include "Experimental/Coroutine/CoroLocalVariable.h"
 
+#if WITH_CPP_COROUTINES
+
 namespace CoroTask_Detail
 {
 
@@ -52,4 +54,6 @@ bool FCoroLocalState::IsCoroLaunchedTask()
 	return CombinedCoroLocalState.CoroLocalState != nullptr;
 }
 
-}
\ No newline at end of file
+}
+
+#endif // WITH_CPP_COROUTINES
\ No newline at end of file
diff --git a/Engine/Source/Runtime/Core/Private/HAL/ThreadingBase.cpp b/Engine/Source/Runtime/Core/Private/HAL/ThreadingBase.cpp
--- a/Engine/Source/Runtime/Core/Private/HAL/ThreadingBase.cpp
+++ b/Engine/Source/Runtime/Core/Private/HAL/ThreadingBase.cpp
@@ -182,7 +182,9 @@ PRAGMA_DISABLE_DEPRECATION_WARNINGS
 		bool newValue = FTaskTagScope::IsCurrentTag(ETaskTag::EGameThread) || FTaskTagScope::IsRunningDuringStaticInit();
 #if !UE_BUILD_SHIPPING && !UE_BUILD_TEST
 		if (!LowLevelTasks::FSchedulerTls::IsBusyWaiting() &&
+#if WITH_CPP_COROUTINES
 			!CoroTask_Detail::FCoroLocalState::IsCoroLaunchedTask() &&
+#endif
 			!UE::Tasks::Private::IsThreadRetractingTask())
 		{
 			const uint32 CurrentThreadId = FPlatformTLS::GetCurrentThreadId();
@@ -210,7 +212,9 @@ PRAGMA_DISABLE_DEPRECATION_WARNINGS
 	bool newValue = FTaskTagScope::IsCurrentTag(ETaskTag::ESlateThread);
 #if !UE_BUILD_SHIPPING && !UE_BUILD_TEST
 	if (!LowLevelTasks::FSchedulerTls::IsBusyWaiting() &&
+#if WITH_CPP_COROUTINES
 		!CoroTask_Detail::FCoroLocalState::IsCoroLaunchedTask() &&
+#endif
 		!UE::Tasks::Private::IsThreadRetractingTask())
 	{
 		bool oldValue = GSlateLoadingThreadId != 0 && FPlatformTLS::GetCurrentThreadId() == GSlateLoadingThreadId;
@@ -246,7 +250,11 @@ CORE_API bool IsInActualRenderingThread()
 PRAGMA_DISABLE_DEPRECATION_WARNINGS
 	bool newValue = FTaskTagScope::IsCurrentTag(ETaskTag::ERenderingThread);
 #if !UE_BUILD_SHIPPING && !UE_BUILD_TEST
-	if (!LowLevelTasks::FSchedulerTls::IsBusyWaiting() && !CoroTask_Detail::FCoroLocalState::IsCoroLaunchedTask())
+	if (!LowLevelTasks::FSchedulerTls::IsBusyWaiting() 
+#if WITH_CPP_COROUTINES
+		&& !CoroTask_Detail::FCoroLocalState::IsCoroLaunchedTask()
+#endif
+		)
 	{
 		bool oldValue = FPlatformTLS::GetCurrentThreadId() == GRenderThreadId;
 		ensureMsgf(oldValue == newValue, TEXT("oldValue(%i) newValue(%i) If this check fails make sure that there is a FTaskTagScope(ETaskTag::ERenderingThread) as deep as possible on the current callstack, you can see the current value in ActiveNamedThreads(%x)"), oldValue, newValue, FTaskTagScope::GetCurrentTag());
@@ -266,7 +274,9 @@ PRAGMA_DISABLE_DEPRECATION_WARNINGS
 
 #if !UE_BUILD_SHIPPING && !UE_BUILD_TEST
 	if (!LowLevelTasks::FSchedulerTls::IsBusyWaiting() && 
+#if WITH_CPP_COROUTINES
 		!CoroTask_Detail::FCoroLocalState::IsCoroLaunchedTask() && 
+#endif
 		!UE::Tasks::Private::IsThreadRetractingTask())
 	{
 		const uint32 CurrentThreadId = FPlatformTLS::GetCurrentThreadId();
@@ -297,7 +307,9 @@ PRAGMA_DISABLE_DEPRECATION_WARNINGS
 
 #if !UE_BUILD_SHIPPING && !UE_BUILD_TEST
 	if (!LowLevelTasks::FSchedulerTls::IsBusyWaiting() &&
+#if WITH_CPP_COROUTINES
 		!CoroTask_Detail::FCoroLocalState::IsCoroLaunchedTask() &&
+#endif
 		!UE::Tasks::Private::IsThreadRetractingTask())
 	{
 		const uint32 CurrentThreadId = FPlatformTLS::GetCurrentThreadId();
diff --git a/Engine/Source/Runtime/Core/Public/Experimental/Coroutine/CoroLocalVariable.h b/Engine/Source/Runtime/Core/Public/Experimental/Coroutine/CoroLocalVariable.h
--- a/Engine/Source/Runtime/Core/Public/Experimental/Coroutine/CoroLocalVariable.h
+++ b/Engine/Source/Runtime/Core/Public/Experimental/Coroutine/CoroLocalVariable.h
@@ -6,6 +6,8 @@
 #include "Experimental/Containers/RobinHoodHashTable.h"
 #include "CoroutineHandle.h"
 
+#if WITH_CPP_COROUTINES
+
 namespace CoroTask_Detail
 {
 	class FPromise;
@@ -106,4 +108,6 @@ public:
 		GetValue() = Other;
 		return *this;
 	}
-};
\ No newline at end of file
+};
+
+#endif // WITH_CPP_COROUTINES
\ No newline at end of file
diff --git a/Engine/Source/Runtime/Core/Public/Experimental/Coroutine/CoroParallelFor.h b/Engine/Source/Runtime/Core/Public/Experimental/Coroutine/CoroParallelFor.h
--- a/Engine/Source/Runtime/Core/Public/Experimental/Coroutine/CoroParallelFor.h
+++ b/Engine/Source/Runtime/Core/Public/Experimental/Coroutine/CoroParallelFor.h
@@ -4,6 +4,8 @@
 #include "Async/ParallelFor.h"
 #include "Coroutine.h"
 
+#if WITH_CPP_COROUTINES
+
 template<typename BodyType, typename... ARGS>
 inline void CoroParallelFor(const TCHAR* DebugName, int32 Num, BodyType Body, EParallelForFlags Flags, ARGS&... Args)
 {
@@ -118,4 +120,6 @@ inline void CoroParallelFor(const TCHAR* DebugName, int32 Num, BodyType Body, EP
 		TaskIndex--;
 	}
 	checkSlow(BatchItem.load(std::memory_order_relaxed) * BatchSize >= Num);
-}
\ No newline at end of file
+}
+
+#endif // WITH_CPP_COROUTINES
\ No newline at end of file
diff --git a/Engine/Source/Runtime/Core/Public/Experimental/Coroutine/CoroTimeout.h b/Engine/Source/Runtime/Core/Public/Experimental/Coroutine/CoroTimeout.h
--- a/Engine/Source/Runtime/Core/Public/Experimental/Coroutine/CoroTimeout.h
+++ b/Engine/Source/Runtime/Core/Public/Experimental/Coroutine/CoroTimeout.h
@@ -1,9 +1,12 @@
 // Copyright Epic Games, Inc. All Rights Reserved.
 #pragma once
+
 #include "Coroutine.h"
 #include "Misc/EnumClassFlags.h"
 #include "Misc/Timespan.h"
 
+#if WITH_CPP_COROUTINES
+
 enum class ECoroTimeoutFlags : uint8
 {
 	Suspend_Named	= (1U << 1U),
@@ -12,7 +15,6 @@ enum class ECoroTimeoutFlags : uint8
 };
 ENUM_CLASS_FLAGS(ECoroTimeoutFlags)
 
-#if WITH_CPP_COROUTINES
 /*
  * FCoroTimeoutAwaitable allows to suspend a Coroutine after it's timelimit has passed
  */
@@ -59,14 +61,4 @@ public:
 	}
 };
 
-#else
-
-class FCoroTimeoutAwaitable
-{
-public:
-	inline FCoroTimeoutAwaitable(FTimespan InTimedelta = FTimespan::FromMilliseconds(33), ECoroTimeoutFlags InFlags = ECoroTimeoutFlags::Suspend_Worker)
-	{
-	}
-};
-
-#endif
\ No newline at end of file
+#endif // WITH_CPP_COROUTINES
\ No newline at end of file
diff --git a/Engine/Source/Runtime/Core/Public/Experimental/Coroutine/CoroutineHandle.h b/Engine/Source/Runtime/Core/Public/Experimental/Coroutine/CoroutineHandle.h
--- a/Engine/Source/Runtime/Core/Public/Experimental/Coroutine/CoroutineHandle.h
+++ b/Engine/Source/Runtime/Core/Public/Experimental/Coroutine/CoroutineHandle.h
@@ -9,7 +9,7 @@
 #pragma once
 
 #ifndef WITH_CPP_COROUTINES
-#define WITH_CPP_COROUTINES 1
+#define WITH_CPP_COROUTINES 0
 #endif
 
 #if WITH_CPP_COROUTINES
