Exposed Trace.RegionBegin and Trace.RegionEnd commands.
These allow developers to tag regions of insights profiles manually with custom names, which can be useful when performing certain tests and then analyzing those results.

#tests Tested in a local test build

diff --git a/Engine/Source/Runtime/Core/Private/ProfilingDebugging/TraceAuxiliary.cpp b/Engine/Source/Runtime/Core/Private/ProfilingDebugging/TraceAuxiliary.cpp
--- a/Engine/Source/Runtime/Core/Private/ProfilingDebugging/TraceAuxiliary.cpp
+++ b/Engine/Source/Runtime/Core/Private/ProfilingDebugging/TraceAuxiliary.cpp
@@ -1105,6 +1105,24 @@ static void TraceBookmark(const TArray<FString>& Args)
 	TRACE_BOOKMARK(TEXT("%s"), Args.Num() ? *Args[0] : TEXT(""));
 }
 
+////////////////////////////////////////////////////////////////////////////////
+static void TraceRegionBegin(const TArray<FString>& Args)
+{
+	if (Args.Num() > 0)
+	{
+		TRACE_BEGIN_REGION(*FString::Join(Args, TEXT(" ")));
+	}
+}
+
+////////////////////////////////////////////////////////////////////////////////
+static void TraceRegionEnd(const TArray<FString>& Args)
+{
+	if (Args.Num() > 0)
+	{
+		TRACE_END_REGION(*FString::Join(Args, TEXT(" ")));
+	}
+}
+
 ////////////////////////////////////////////////////////////////////////////////
 static FAutoConsoleCommand TraceAuxiliarySendCmd(
 	TEXT("Trace.Send"),
@@ -1205,6 +1223,22 @@ static FAutoConsoleCommand TraceBookmarkCmd(
 	FConsoleCommandWithArgsDelegate::CreateStatic(TraceBookmark)
 );
 
+////////////////////////////////////////////////////////////////////////////////
+static FAutoConsoleCommand TraceRegionBeginCmd(
+	TEXT("Trace.RegionBegin"),
+	TEXT("[Name] - Emits a TRACE_BEGIN_REGION() event with the given string name."
+	),
+	FConsoleCommandWithArgsDelegate::CreateStatic(TraceRegionBegin)
+);
+
+////////////////////////////////////////////////////////////////////////////////
+static FAutoConsoleCommand TraceRegionEndCmd(
+	TEXT("Trace.RegionEnd"),
+	TEXT("[Name] - Emits a TRACE_END_REGION() event with the given string name."
+	),
+	FConsoleCommandWithArgsDelegate::CreateStatic(TraceRegionEnd)
+);
+
 #endif // UE_TRACE_ENABLED
 
 
