Fixed usage of FString::Printf without format arguments in Math header files.

diff --git a/Engine/Source/Runtime/Core/Public/Math/Rotator.h b/Engine/Source/Runtime/Core/Public/Math/Rotator.h
--- a/Engine/Source/Runtime/Core/Public/Math/Rotator.h
+++ b/Engine/Source/Runtime/Core/Public/Math/Rotator.h
@@ -808,7 +808,7 @@ FORCEINLINE FString TRotator<T>::ToCompactString() const
 {
 	if( IsNearlyZero() )
 	{
-		return FString::Printf(TEXT("R(0)"));
+		return FString(TEXT("R(0)"));
 	}
 
 	FString ReturnString(TEXT("R("));
diff --git a/Engine/Source/Runtime/Core/Public/Math/Vector.h b/Engine/Source/Runtime/Core/Public/Math/Vector.h
--- a/Engine/Source/Runtime/Core/Public/Math/Vector.h
+++ b/Engine/Source/Runtime/Core/Public/Math/Vector.h
@@ -2340,7 +2340,7 @@ FORCEINLINE FString TVector<T>::ToCompactString() const
 {
     if(IsNearlyZero())
     {
-        return FString::Printf(TEXT("V(0)"));
+        return FString(TEXT("V(0)"));
     }
 
     FString ReturnString(TEXT("V("));
