Trace message bus controls

Adds a control layer for Trace (TraceService) and a controller (TraceController) integrated into SessionServices. This allows remote control of trace over messge bus.

diff --git a/Engine/Source/Runtime/Core/Private/ProfilingDebugging/TraceAuxiliary.cpp b/Engine/Source/Runtime/Core/Private/ProfilingDebugging/TraceAuxiliary.cpp
--- a/Engine/Source/Runtime/Core/Private/ProfilingDebugging/TraceAuxiliary.cpp
+++ b/Engine/Source/Runtime/Core/Private/ProfilingDebugging/TraceAuxiliary.cpp
@@ -69,6 +69,7 @@
 const TCHAR* GDefaultChannels = TEXT("cpu,gpu,frame,log,bookmark,screenshot,region");
 const TCHAR* GMemoryChannels = TEXT("memtag,memalloc,callstack,module");
 const TCHAR* GTraceConfigSection = TEXT("Trace.Config");
+static UE::Trace::FInitializeDesc GInitializeDesc;
 
 ////////////////////////////////////////////////////////////////////////////////
 CSV_DEFINE_CATEGORY(Trace, true);
@@ -108,11 +109,12 @@ public:
 	void ResetCommandlineChannels();
 	bool HasCommandlineChannels() const { return !CommandlineChannels.IsEmpty(); }
 	void EnableChannels(const TCHAR* ChannelList, const FTraceAuxiliary::FLogCategoryAlias& LogCategory);
+	void EnableChannels(TConstArrayView<uint32> ChannelIds);
 	void DisableChannels(const TCHAR* ChannelList, const FTraceAuxiliary::FLogCategoryAlias& LogCategory);
+	void DisableChannels(TConstArrayView<uint32> ChannelIds);
 	bool Connect(FTraceAuxiliary::EConnectionType Type, const TCHAR* Parameter, const FTraceAuxiliary::FLogCategoryAlias& LogCategory, uint16 SendFlags);
 	bool Stop();
 	void FreezeReadOnlyChannels();
-	static bool IsReadOnlyChannel(const TCHAR* Channel);
 	void ResumeChannels();
 	void PauseChannels();
 	bool IsPaused();
@@ -144,12 +146,15 @@ private:
 	void AddCommandlineChannel(const TCHAR* Name, const FTraceAuxiliary::FLogCategoryAlias& LogCategory);
 	template <class T> void ForEachChannel(const TCHAR* ChannelList, bool bResolvePresets, const FTraceAuxiliary::FLogCategoryAlias& LogCategory, T Callable);
 	static uint32 HashChannelName(const TCHAR* Name);
-	bool EnableChannel(const TCHAR* Channel, const FTraceAuxiliary::FLogCategoryAlias& LogCategory);
-	void DisableChannel(const TCHAR* Channel, const FTraceAuxiliary::FLogCategoryAlias& LogCategory);
 	bool SendToHost(const TCHAR* Host, const FTraceAuxiliary::FLogCategoryAlias& LogCategory, uint16 SendFlags);
 	bool WriteToFile(const TCHAR* Path, const FTraceAuxiliary::FLogCategoryAlias& LogCategory, uint16 SendFlags);
 	bool FinalizeFilePath(const TCHAR* InPath, FString& OutPath, const FTraceAuxiliary::FLogCategoryAlias& LogCategory);
 
+	template<typename ChannelType>
+	bool EnableChannel(ChannelType Channel, const FTraceAuxiliary::FLogCategoryAlias& LogCategory);
+	template<typename ChannelType>
+	bool DisableChannel(ChannelType Channel, const FTraceAuxiliary::FLogCategoryAlias& LogCategory);
+
 	typedef TMap<uint32, FChannelEntry, TInlineSetAllocator<128>> ChannelSet;
 	ChannelSet CommandlineChannels;
 	bool bWorkerThreadStarted = false;
@@ -216,23 +221,42 @@ void FTraceAuxiliaryImpl::EnableChannels(const TCHAR* ChannelList, const FTraceA
 {
 	if (ChannelList)
 	{
-		ForEachChannel(ChannelList, true, LogCategory, &FTraceAuxiliaryImpl::EnableChannel);
+		ForEachChannel(ChannelList, true, LogCategory, &FTraceAuxiliaryImpl::EnableChannel<const TCHAR*>);
 	}
 }
 
+////////////////////////////////////////////////////////////////////////////////
+void FTraceAuxiliaryImpl::EnableChannels(TConstArrayView<uint32> ChannelIds)
+{
+	for (const auto ChannelId : ChannelIds)
+	{
+		EnableChannel(ChannelId, LogTrace);
+	}
+}
+
+
 ////////////////////////////////////////////////////////////////////////////////
 void FTraceAuxiliaryImpl::DisableChannels(const TCHAR* ChannelList, const FTraceAuxiliary::FLogCategoryAlias& LogCategory)
 {
 	if (ChannelList)
 	{
-		ForEachChannel(ChannelList, true, LogCategory, &FTraceAuxiliaryImpl::DisableChannel);
+		ForEachChannel(ChannelList, true, LogCategory, &FTraceAuxiliaryImpl::DisableChannel<const TCHAR*>);
 	}
 	else
 	{
 		// Disable all channels.
 		TStringBuilder<128> EnabledChannels;
 		GetActiveChannelsString(EnabledChannels);
-		ForEachChannel(EnabledChannels.ToString(), true, LogCategory, &FTraceAuxiliaryImpl::DisableChannel);
+		ForEachChannel(EnabledChannels.ToString(), true, LogCategory, &FTraceAuxiliaryImpl::DisableChannel<const TCHAR*>);
+	}
+}
+
+////////////////////////////////////////////////////////////////////////////////
+void FTraceAuxiliaryImpl::DisableChannels(TConstArrayView<uint32> ChannelIds)
+{
+	for (const auto ChannelId : ChannelIds)
+	{
+		DisableChannel(ChannelId, LogTrace);
 	}
 }
 
@@ -271,6 +295,87 @@ void FTraceAuxiliaryImpl::ForEachChannel(const TCHAR* ChannelList, bool bResolve
 	});
 }
 
+////////////////////////////////////////////////////////////////////////////////
+template <typename IdentifierType>
+bool FTraceAuxiliaryImpl::EnableChannel(IdentifierType ChannelIdentifier, const FTraceAuxiliary::FLogCategoryAlias& LogCategory)
+{
+	// Channel names have been provided by the user and may not exist yet. As
+	// we want to maintain bActive accurately (channels toggles are reference
+	// counted), we will first check Trace knows of the channel.
+	UE::Trace::FChannel* Channel = UE::Trace::FindChannel(ChannelIdentifier);
+	if (!Channel)
+	{
+		return false;
+	}
+
+	// Build an wide representation of the name for logging and platform
+	// events purposes.
+	// todo: This can be moved into the log scope once platform events is being triggered by channel callbacks
+	TStringBuilder<64> ChannelName;
+	{
+		const ANSICHAR* ChannelNameA = nullptr;
+		const uint32 ChannelNameLen = Channel->GetName(&ChannelNameA);
+		ChannelName << FAnsiStringView(ChannelNameA, ChannelNameLen);
+	}
+	
+	// It is not possible to change read only channels once trace is initialized.
+	if (bReadOnlyChannelsFrozen && Channel->IsReadOnly())
+	{
+		UE_LOG_REF(
+			LogCategory,
+			Error,
+			TEXT("Channel '%s' is read only. It is not allowed to manually enable this channel."),
+			ChannelName.ToString()
+		);
+		return Channel->IsEnabled();
+	}
+
+	const bool bIsEnabled = Channel->Toggle(true);
+	FPlatformEventsTrace::OnTraceChannelUpdated(ChannelName.ToString(), bIsEnabled);
+
+	return bIsEnabled;
+}
+
+////////////////////////////////////////////////////////////////////////////////
+template <typename ChannelType>
+bool FTraceAuxiliaryImpl::DisableChannel(ChannelType ChannelIdentifier, const FTraceAuxiliary::FLogCategoryAlias& LogCategory)
+{
+	// Channel names have been provided by the user and may not exist yet. As
+	// we want to maintain bActive accurately we will first check Trace knows of the channel.
+	UE::Trace::FChannel* Channel = UE::Trace::FindChannel(ChannelIdentifier);
+	if (!Channel)
+	{
+		return false;
+	}
+
+	// Build an wide representation of the name for logging and platform
+	// events purposes.
+	// todo: This can be moved into the log scope once platform events is being triggered by channel callbacks
+	TStringBuilder<64> ChannelName;
+	{
+		const ANSICHAR* ChannelNameA = nullptr;
+		const uint32 ChannelNameLen = Channel->GetName(&ChannelNameA);
+		ChannelName << FAnsiStringView(ChannelNameA, ChannelNameLen);
+	}
+	
+	// It is not possible to change read only channels once trace is initialized.
+	if (bReadOnlyChannelsFrozen && Channel->IsReadOnly())
+	{
+		UE_LOG_REF(
+			LogCategory,
+			Error,
+			TEXT("Channel '%s' is read only. It is not allowed to manually disable this channel."),
+			ChannelName.ToString()
+		);
+		return Channel->IsEnabled();
+	}
+
+	const bool bIsEnabled = Channel->Toggle(false);
+	FPlatformEventsTrace::OnTraceChannelUpdated(ChannelName.ToString(), bIsEnabled);
+
+	return bIsEnabled;
+}
+
 ////////////////////////////////////////////////////////////////////////////////
 uint32 FTraceAuxiliaryImpl::HashChannelName(const TCHAR* Name)
 {
@@ -376,92 +481,11 @@ void FTraceAuxiliaryImpl::FreezeReadOnlyChannels()
 	bReadOnlyChannelsFrozen = true;
 }
 
-////////////////////////////////////////////////////////////////////////////////
-bool FTraceAuxiliaryImpl::IsReadOnlyChannel(const TCHAR* InChannel)
-{
-	const auto ChannelA = StringCast<ANSICHAR, 80>(InChannel);
-	struct FUserData
-	{
-		FAnsiStringView Channel;
-		bool bIsReadOnlyChannel = false;
-	} UserData = { FAnsiStringView(ChannelA.Get(), ChannelA.Length()), false};
-
-	UE::Trace::EnumerateChannels([](const UE::Trace::FChannelInfo& Info, void* User)
-	{
-		FAnsiStringView NameView = FAnsiStringView(Info.Name).LeftChop(7); // Remove "Channel" suffix
-		FUserData* UserData = (FUserData*) User;
-		if (NameView.Equals(UserData->Channel, ESearchCase::IgnoreCase))
-		{
-			UserData->bIsReadOnlyChannel = Info.bIsReadOnly;
-			return false;
-		}
-		return true;
-	}, &UserData);
-
-	return UserData.bIsReadOnlyChannel;
-}
-
-////////////////////////////////////////////////////////////////////////////////
-bool FTraceAuxiliaryImpl::EnableChannel(const TCHAR* Channel, const FTraceAuxiliary::FLogCategoryAlias& LogCategory)
-{
-	// Channel names have been provided by the user and may not exist yet. As
-	// we want to maintain bActive accurately (channels toggles are reference
-	// counted), we will first check Trace knows of the channel.
-	if (!UE::Trace::IsChannel(Channel))
-	{
-		return false;
-	}
-
-	// It is not possible to change read only channels once trace is initialized.
-	if (bReadOnlyChannelsFrozen && IsReadOnlyChannel(Channel))
-	{
-		UE_LOG_REF(
-			LogCategory,
-			Error,
-			TEXT("Channel '%s' is read only. It is not allowed to manually enable this channel."),
-			Channel
-		);
-		return false;
-	}
-
-	bool bIsEnabled = UE::Trace::ToggleChannel(Channel, true);
-	FPlatformEventsTrace::OnTraceChannelUpdated(Channel, bIsEnabled);
-
-	return bIsEnabled;
-}
-
-////////////////////////////////////////////////////////////////////////////////
-void FTraceAuxiliaryImpl::DisableChannel(const TCHAR* Channel, const FTraceAuxiliary::FLogCategoryAlias& LogCategory)
-{
-	// Channel names have been provided by the user and may not exist yet. As
-	// we want to maintain bActive accurately (channels toggles are reference
-	// counted), we will first check Trace knows of the channel.
-	if (!UE::Trace::IsChannel(Channel))
-	{
-		return;
-	}
-
-	// It is not possible to change read only channels once trace is initialized.
-	if (bReadOnlyChannelsFrozen && IsReadOnlyChannel(Channel))
-	{
-		UE_LOG_REF(
-			LogCategory,
-			Error,
-			TEXT("Channel '%s' is read only. It is not allowed to manually disable this channel."),
-			Channel
-		);
-		return;
-	}
-
-	bool bIsEnabled = UE::Trace::ToggleChannel(Channel, false);
-	FPlatformEventsTrace::OnTraceChannelUpdated(Channel, bIsEnabled);
-}
-
 ////////////////////////////////////////////////////////////////////////////////
 void FTraceAuxiliaryImpl::ResumeChannels()
 {
 	// Enable channels from the "paused" preset.
-	ForEachChannel(*PausedPreset, false, LogTrace, &FTraceAuxiliaryImpl::EnableChannel);
+	ForEachChannel(*PausedPreset, false, LogTrace, &FTraceAuxiliaryImpl::EnableChannel<const TCHAR*>);
 
 	PausedPreset.Empty();
 }
@@ -477,7 +501,7 @@ void FTraceAuxiliaryImpl::PauseChannels()
 	PausedPreset = EnabledChannels.ToString();
 
 	// Disable all "paused" channels.
-	ForEachChannel(*PausedPreset, true, LogTrace, &FTraceAuxiliaryImpl::DisableChannel);
+	ForEachChannel(*PausedPreset, true, LogTrace, &FTraceAuxiliaryImpl::DisableChannel<const TCHAR*>);
 }
 
 ////////////////////////////////////////////////////////////////////////////////
@@ -1565,8 +1589,8 @@ void FTraceAuxiliary::Initialize(const TCHAR* CommandLine)
 	bool bShouldStartWorkerThread = false;
 	StartFromCommandlineArguments(CommandLine, bShouldStartWorkerThread);
 
-	// Initialize Trace
-	UE::Trace::FInitializeDesc Desc;
+	// Initialize Trace. The settings are stored in a static for posterity.
+	UE::Trace::FInitializeDesc& Desc = GInitializeDesc;
 #if WITH_EDITOR
 	Desc.TailSizeBytes = 32 << 20;
 #endif
@@ -1579,7 +1603,7 @@ void FTraceAuxiliary::Initialize(const TCHAR* CommandLine)
 	FGuid SessionGuid;
 	if (!FParse::Value(CommandLine, TEXT("-tracesessionguid="), SessionGuid))
 	{
-		SessionGuid = FApp::GetInstanceId();
+		SessionGuid = FApp::GetSessionId();
 	}
 	FMemory::Memcpy((FGuid&)Desc.SessionGuid, SessionGuid);
 
@@ -1687,13 +1711,27 @@ void FTraceAuxiliary::Shutdown()
 }
 
 ////////////////////////////////////////////////////////////////////////////////
-void FTraceAuxiliary::EnableChannels()
+void FTraceAuxiliary::EnableCommandlineChannels()
 {
 #if UE_TRACE_ENABLED
 	GTraceAuxiliary.EnableCommandlineChannels();
 #endif
 }
 
+void FTraceAuxiliary::EnableChannels(TConstArrayView<uint32> ChannelIds)
+{
+#if UE_TRACE_ENABLED
+	GTraceAuxiliary.EnableChannels(ChannelIds);
+#endif
+}
+
+void FTraceAuxiliary::DisableChannels(TConstArrayView<uint32> ChannelIds)
+{
+#if UE_TRACE_ENABLED
+	GTraceAuxiliary.DisableChannels(ChannelIds);
+#endif
+}
+
 void FTraceAuxiliary::DisableChannels(const TCHAR* Channels)
 {
 #if UE_TRACE_ENABLED
@@ -1759,6 +1797,15 @@ void FTraceAuxiliary::Panic()
 	UE::Trace::Panic();
 }
 
+UE::Trace::FInitializeDesc const* FTraceAuxiliary::GetInitializeDesc()
+{
+#if UE_TRACE_ENABLED
+	return &GInitializeDesc;
+#else
+	return nullptr;
+#endif
+}
+
 ////////////////////////////////////////////////////////////////////////////////
 void FTraceAuxiliary::TryAutoConnect()
 {
diff --git a/Engine/Source/Runtime/Core/Public/ProfilingDebugging/TraceAuxiliary.h b/Engine/Source/Runtime/Core/Public/ProfilingDebugging/TraceAuxiliary.h
--- a/Engine/Source/Runtime/Core/Public/ProfilingDebugging/TraceAuxiliary.h
+++ b/Engine/Source/Runtime/Core/Public/ProfilingDebugging/TraceAuxiliary.h
@@ -147,7 +147,21 @@ public:
 	 *  Enable previously selected channels. This method can be called multiple times
 	 *  as channels can be announced on module loading.
 	 */
-	static CORE_API void EnableChannels();
+	static CORE_API void EnableCommandlineChannels();
+
+	/**
+	 * Enable channels to emit events belonging to this category.
+	 * @note Note that presets cannot be used when specifying channel ids.
+	 * @param ChannelIds List of channels (or a preset) to enable. 
+	 */
+	static CORE_API void EnableChannels(TConstArrayView<uint32> ChannelIds);
+
+	/**
+	 * Disable channels to mute events belonging to this category.
+	 * @note Note that presets cannot be used when specifying channel ids.
+	 * @param ChannelIds List of channels (or a preset) to disable. 
+	 */
+	static CORE_API void DisableChannels(TConstArrayView<uint32> ChannelIds);
 
 	/**
 	 * Disable channels to stop recording traces with them.
@@ -195,7 +209,12 @@ public:
 	 * Used when process is panicking. Stops all tracing immediately to avoid further allocations. Process is not
 	 * expected to continue after this call.
 	 */
-	 static CORE_API void Panic();
+	static CORE_API void Panic();
+
+	/**
+	 * Get the settings used to initialize TraceLog
+	 */
+	static CORE_API struct UE::Trace::FInitializeDesc const* GetInitializeDesc();
 
 	/**
 	 * Delegate that triggers when a connection is established. Gives subscribers a chance to trace events that appear
diff --git a/Engine/Source/Runtime/Engine/Private/TraceService.cpp b/Engine/Source/Runtime/Engine/Private/TraceService.cpp
new file mode 100644
--- /dev/null
+++ b/Engine/Source/Runtime/Engine/Private/TraceService.cpp
@@ -0,0 +1,262 @@
+// Copyright Epic Games, Inc. All Rights Reserved.
+#include "TraceService.h"
+
+#include "MessageEndpoint.h"
+#include "MessageEndpointBuilder.h"
+#include "Misc/App.h"
+#include "ProfilingDebugging/TraceAuxiliary.h"
+#include "ProfilingDebugging/TraceScreenshot.h"
+#include "Trace/Trace.h"
+#include "TraceControlMessages.h"
+
+class FTraceServiceImpl
+{
+public:
+	FTraceServiceImpl();
+	virtual ~FTraceServiceImpl() {};
+	
+private:
+	void OnStatusPing(const FTraceControlStatusPing& Message, const TSharedRef<IMessageContext>& Context);
+	void OnChannelsPing(const FTraceControlChannelsPing& Message, const TSharedRef<IMessageContext>& Context);
+	void OnSettingsPing(const FTraceControlSettingsPing& Message, const TSharedRef<IMessageContext, ESPMode::ThreadSafe>& Context);
+	void OnConnectRequest(const FTraceControlDiscoveryPing& Message, const TSharedRef<IMessageContext>& Context);
+	void OnStop(const FTraceControlStop& Message, const TSharedRef<IMessageContext>& Context);
+	void OnSend(const FTraceControlSend& Message, const TSharedRef<IMessageContext>& Context);
+	void OnChannelSet(const FTraceControlChannelsSet& Message, const TSharedRef<IMessageContext>& Context);
+	void OnFile(const FTraceControlFile& Message, const TSharedRef<IMessageContext>& Context);
+	void OnSnapshotSend(const FTraceControlSnapshotSend& Message, const TSharedRef<IMessageContext>& Context);
+	void OnSnapshotFile(const FTraceControlSnapshotFile& Message, const TSharedRef<IMessageContext>& Context);
+	void OnPause(const FTraceControlPause& Message, const TSharedRef<IMessageContext>& Context);
+	void OnResume(const FTraceControlResume& Message, const TSharedRef<IMessageContext>& Context);
+	void OnBookmark(const FTraceControlBookmark& Message, const TSharedRef<IMessageContext>& Context);
+	void OnScreenshot(const FTraceControlScreenshot& Message, const TSharedRef<IMessageContext>& Context);
+	
+	static void FillTraceStatusMessage(FTraceControlStatus* Message);
+
+	TSharedPtr<FMessageEndpoint, ESPMode::ThreadSafe> MessageEndpoint;
+	FGuid SessionId;
+	FGuid InstanceId;
+};
+
+
+FTraceServiceImpl::FTraceServiceImpl()
+{
+	SessionId = FApp::GetSessionId();
+	InstanceId = FApp::GetInstanceId();
+	
+	MessageEndpoint = FMessageEndpoint::Builder("FTraceService")
+		.Handling<FTraceControlDiscoveryPing>(this, &FTraceServiceImpl::OnConnectRequest)
+		.Handling<FTraceControlChannelsSet>(this, &FTraceServiceImpl::OnChannelSet)
+		.Handling<FTraceControlStop>(this, &FTraceServiceImpl::OnStop)
+		.Handling<FTraceControlSend>(this, &FTraceServiceImpl::OnSend)
+		.Handling<FTraceControlFile>(this, &FTraceServiceImpl::OnFile)
+		.Handling<FTraceControlSnapshotSend>(this, &FTraceServiceImpl::OnSnapshotSend)
+		.Handling<FTraceControlSnapshotFile>(this, &FTraceServiceImpl::OnSnapshotFile)
+		.Handling<FTraceControlPause>(this, &FTraceServiceImpl::OnPause)
+		.Handling<FTraceControlResume>(this, &FTraceServiceImpl::OnResume)
+		.Handling<FTraceControlBookmark>(this, &FTraceServiceImpl::OnBookmark)
+		.Handling<FTraceControlScreenshot>(this, &FTraceServiceImpl::OnScreenshot)
+		.Handling<FTraceControlStatusPing>(this, &FTraceServiceImpl::OnStatusPing)
+		.Handling<FTraceControlSettingsPing>(this, &FTraceServiceImpl::OnSettingsPing)
+		.Handling<FTraceControlChannelsPing>(this, &FTraceServiceImpl::OnChannelsPing);
+
+	if (!MessageEndpoint.IsValid())
+	{
+		 return;
+	}
+
+	MessageEndpoint->Subscribe<FTraceControlStatusPing>();
+	MessageEndpoint->Subscribe<FTraceControlSettingsPing>();
+	MessageEndpoint->Subscribe<FTraceControlDiscoveryPing>();
+	MessageEndpoint->Subscribe<FTraceControlChannelsPing>();
+	MessageEndpoint->Subscribe<FTraceControlStop>();
+	MessageEndpoint->Subscribe<FTraceControlSend>();
+	MessageEndpoint->Subscribe<FTraceControlChannelsSet>();
+	MessageEndpoint->Subscribe<FTraceControlFile>();
+	MessageEndpoint->Subscribe<FTraceControlSnapshotSend>();
+	MessageEndpoint->Subscribe<FTraceControlSnapshotSend>();
+	MessageEndpoint->Subscribe<FTraceControlPause>();
+	MessageEndpoint->Subscribe<FTraceControlResume>();
+	MessageEndpoint->Subscribe<FTraceControlBookmark>();
+	MessageEndpoint->Subscribe<FTraceControlScreenshot>();
+	
+}
+
+void FTraceServiceImpl::FillTraceStatusMessage(FTraceControlStatus* Message)
+{
+	// Get the current endpoint and ids
+	Message->Endpoint = FTraceAuxiliary::GetTraceDestinationString();
+	FTraceAuxiliary::IsConnected(Message->SessionGuid, Message->TraceGuid);
+	
+	// For stats we can query TraceLog directly.
+	UE::Trace::FStatistics Stats;
+	UE::Trace::GetStatistics(Stats);
+	Message->BytesSent = Stats.BytesSent;
+	Message->BytesTraced = Stats.BytesTraced;
+	Message->MemoryUsed = Stats.MemoryUsed;
+	Message->CacheAllocated = Stats.CacheAllocated;
+	Message->CacheUsed = Stats.CacheUsed;
+	Message->CacheWaste = Stats.CacheWaste;
+
+}
+
+void FTraceServiceImpl::OnChannelSet(const FTraceControlChannelsSet& Message, const TSharedRef<IMessageContext>& Context)
+{
+	FTraceAuxiliary::EnableChannels(Message.ChannelIdsToEnable);
+	FTraceAuxiliary::DisableChannels(Message.ChannelIdsToDisable);
+}
+
+void FTraceServiceImpl::OnStop(const FTraceControlStop& Message, const TSharedRef<IMessageContext>& Context)
+{
+	FTraceAuxiliary::Stop();
+}
+
+void FTraceServiceImpl::OnSend(const FTraceControlSend& Message, const TSharedRef<IMessageContext>& Context)
+{
+	FTraceAuxiliary::FOptions Options;
+	Options.bExcludeTail = Message.bExcludeTail;
+	
+	FTraceAuxiliary::Start(
+		FTraceAuxiliary::EConnectionType::Network,
+		*Message.Host,
+		*Message.Channels,
+		&Options
+	);
+}
+
+void FTraceServiceImpl::OnFile(const FTraceControlFile& Message, const TSharedRef<IMessageContext>& Context)
+{
+	FTraceAuxiliary::FOptions Options;
+	Options.bTruncateFile = Message.bTruncateFile;
+	Options.bExcludeTail = Message.bExcludeTail;
+	
+	FTraceAuxiliary::Start(
+		FTraceAuxiliary::EConnectionType::File,
+		*Message.File,
+		*Message.Channels,
+		&Options
+	);
+}
+
+void FTraceServiceImpl::OnSnapshotSend(const FTraceControlSnapshotSend& Message, const TSharedRef<IMessageContext>& Context)
+{
+	FTraceAuxiliary::SendSnapshot(*Message.Host);
+}
+
+void FTraceServiceImpl::OnSnapshotFile(const FTraceControlSnapshotFile& Message, const TSharedRef<IMessageContext>& Context)
+{
+	FTraceAuxiliary::SendSnapshot(*Message.File);
+}
+
+void FTraceServiceImpl::OnPause(const FTraceControlPause& Message, const TSharedRef<IMessageContext>& Context)
+{
+	FTraceAuxiliary::Pause();
+}
+
+void FTraceServiceImpl::OnResume(const FTraceControlResume& Message, const TSharedRef<IMessageContext>& Context)
+{
+	FTraceAuxiliary::Resume();
+}
+
+void FTraceServiceImpl::OnBookmark(const FTraceControlBookmark& Message, const TSharedRef<IMessageContext>& Context)
+{
+	TRACE_BOOKMARK(TEXT("%s"), *Message.Label);
+}
+
+void FTraceServiceImpl::OnScreenshot(const FTraceControlScreenshot& Message, const TSharedRef<IMessageContext>& Context)
+{
+	FTraceScreenshot::RequestScreenshot(Message.Name, Message.bShowUI);
+}
+
+
+void FTraceServiceImpl::OnStatusPing(const FTraceControlStatusPing& Message, const TSharedRef<IMessageContext>& Context)
+{
+	FTraceControlStatus* Response = FMessageEndpoint::MakeMessage<FTraceControlStatus>();
+	FillTraceStatusMessage(Response);
+	
+	MessageEndpoint->Send(Response, Context->GetSender());
+}
+
+void FTraceServiceImpl::OnChannelsPing(const FTraceControlChannelsPing& Message, const TSharedRef<IMessageContext>& Context)
+{
+	
+	struct FEnumerateUserData
+	{
+		TArray<FString> Channels;
+		TArray<FString> Descriptions;
+		TArray<uint32> Ids;
+		TArray<uint32> ReadOnlyIds;
+		TArray<uint32> EnabledIds;
+	} UserData;
+	
+	UE::Trace::EnumerateChannels([](const UE::Trace::FChannelInfo& ChannelInfo, void* User)
+	{
+		FEnumerateUserData* UserData = static_cast<FEnumerateUserData*>(User);
+		FAnsiStringView NameView = FAnsiStringView(ChannelInfo.Name).LeftChop(7); // Remove "Channel" suffix
+		const uint32 ChannelId = ChannelInfo.Id;
+		UserData->Channels.Emplace(NameView);
+		UserData->Ids.Emplace(ChannelId);
+		UserData->Descriptions.Emplace(ChannelInfo.Desc);
+		if (ChannelInfo.bIsReadOnly)
+		{
+			UserData->ReadOnlyIds.Add(ChannelId);
+		}
+		if (ChannelInfo.bIsEnabled)
+		{
+			UserData->EnabledIds.Add(ChannelId);
+		}
+		return true;
+	}, &UserData);
+
+	// Only send channel description message if the number of channels has changed.
+	if (Message.KnownChannelCount < uint32(UserData.Channels.Num()))
+	{
+		FTraceControlChannelsDesc* DescResponse = FMessageEndpoint::MakeMessage<FTraceControlChannelsDesc>();
+		DescResponse->Channels = MoveTemp(UserData.Channels);
+		DescResponse->Ids = MoveTemp(UserData.Ids);
+		DescResponse->Descriptions = MoveTemp(UserData.Descriptions);
+		DescResponse->ReadOnlyIds = MoveTemp(UserData.ReadOnlyIds);
+		MessageEndpoint->Send(DescResponse, Context->GetSender());
+	}
+
+	// Always send status response
+	FTraceControlChannelsStatus* StatusResponse = FMessageEndpoint::MakeMessage<FTraceControlChannelsStatus>();
+	StatusResponse->EnabledIds = MoveTemp(UserData.EnabledIds);
+	MessageEndpoint->Send(StatusResponse, Context->GetSender());
+}
+
+void FTraceServiceImpl::OnSettingsPing(const FTraceControlSettingsPing& Message, const TSharedRef<IMessageContext>& Context)
+{
+	FTraceControlSettings* Response = FMessageEndpoint::MakeMessage<FTraceControlSettings>();
+	UE::Trace::FInitializeDesc const* InitDesc = FTraceAuxiliary::GetInitializeDesc();
+
+	Response->bStatNamedEvents = GCycleStatsShouldEmitNamedEvents > 0;
+	Response->bUseImportantCache = InitDesc->bUseImportantCache;
+	Response->bUseWorkerThread = InitDesc->bUseWorkerThread;
+	Response->TailSizeBytes = InitDesc->TailSizeBytes;
+	
+	MessageEndpoint->Send(Response, Context->GetSender());
+}
+
+void FTraceServiceImpl::OnConnectRequest(const FTraceControlDiscoveryPing& Message, const TSharedRef<IMessageContext>& Context)
+{
+	if (Message.InstanceId == FApp::GetInstanceId() && Message.SessionId == FApp::GetSessionId())
+	{
+		const auto Response = FMessageEndpoint::MakeMessage<FTraceControlDiscovery>();
+		Response->SessionId = FApp::GetSessionId();
+		Response->InstanceId = FApp::GetInstanceId();
+		FillTraceStatusMessage(Response);
+
+		MessageEndpoint->Send(
+			Response,
+			Context->GetSender()
+		);
+	}
+}
+
+
+FTraceService::FTraceService()
+{
+	Impl = MakePimpl<FTraceServiceImpl>();
+}
+
diff --git a/Engine/Source/Runtime/Engine/Public/TraceService.h b/Engine/Source/Runtime/Engine/Public/TraceService.h
new file mode 100644
--- /dev/null
+++ b/Engine/Source/Runtime/Engine/Public/TraceService.h
@@ -0,0 +1,14 @@
+// Copyright Epic Games, Inc. All Rights Reserved.
+#pragma once
+
+#include <Templates/PimplPtr.h>
+
+class FTraceService
+{
+public:
+	ENGINE_API FTraceService();
+	virtual ~FTraceService() {}
+	
+private:
+	TPimplPtr<class FTraceServiceImpl> Impl;
+};
diff --git a/Engine/Source/Runtime/EngineMessages/Public/TraceControlMessages.h b/Engine/Source/Runtime/EngineMessages/Public/TraceControlMessages.h
new file mode 100644
--- /dev/null
+++ b/Engine/Source/Runtime/EngineMessages/Public/TraceControlMessages.h
@@ -0,0 +1,305 @@
+﻿// Copyright Epic Games, Inc. All Rights Reserved.
+
+#pragma once
+
+#include "CoreMinimal.h"
+#include "UObject/ObjectMacros.h"
+#include "Misc/Guid.h"
+#include "Containers/BitArray.h"
+#include "TraceControlMessages.generated.h"
+
+/* Trace messages
+ *****************************************************************************/
+
+/**
+ * Request a status update
+ */
+USTRUCT()
+struct FTraceControlStatusPing
+{
+	GENERATED_USTRUCT_BODY()
+};
+
+/**
+ * Response for FTraceControlStatusPing. Meant to be fairly high
+ * frequency.
+ */
+USTRUCT()
+struct FTraceControlStatus
+{
+	GENERATED_USTRUCT_BODY()
+
+	UPROPERTY(EditAnywhere, Category="Message")
+	FString Endpoint;
+	
+	UPROPERTY(EditAnywhere, Category="Message")
+	FGuid SessionGuid;
+	
+	UPROPERTY(EditAnywhere, Category="Message")
+	FGuid TraceGuid;
+
+	UPROPERTY(EditAnywhere, Category="Message")
+	uint64 BytesSent;
+
+	UPROPERTY(EditAnywhere, Category="Message")
+	uint64 BytesTraced;
+
+	UPROPERTY(EditAnywhere, Category="Message")
+	uint64 MemoryUsed;
+
+	UPROPERTY(EditAnywhere, Category="Message")
+	uint32 CacheAllocated; 
+
+	UPROPERTY(EditAnywhere, Category="Message")
+	uint32 CacheUsed;
+
+	UPROPERTY(EditAnywhere, Category="Message")
+	uint32 CacheWaste;
+};
+
+/**
+ * Request update of settings
+ */
+USTRUCT()
+struct FTraceControlSettingsPing
+{
+	GENERATED_USTRUCT_BODY()
+};
+
+/**
+ * Response for FTraceControlSettingsPing containing
+ * TraceLog and system settings.
+ */
+USTRUCT()
+struct FTraceControlSettings
+{
+	GENERATED_USTRUCT_BODY()
+
+	UPROPERTY(EditAnywhere, Category="Message")
+	bool bUseWorkerThread;
+	
+	UPROPERTY(EditAnywhere, Category="Message")
+	bool bUseImportantCache;
+
+	UPROPERTY(EditAnywhere, Category="Message")
+	uint32 TailSizeBytes;
+	
+	UPROPERTY(EditAnywhere, Category="Message")
+	bool bStatNamedEvents;
+};
+
+/**
+ * Request update of channels and status
+ */
+USTRUCT()
+struct FTraceControlChannelsPing
+{
+	GENERATED_USTRUCT_BODY()
+	
+	UPROPERTY(EditAnywhere, Category="Message")
+	uint32 KnownChannelCount;
+};
+
+/**
+ * Response for FTraceControlChannelsPing with information
+ * about available channels.
+ */
+USTRUCT()
+struct FTraceControlChannelsDesc
+{
+	GENERATED_USTRUCT_BODY()
+	
+	UPROPERTY(EditAnywhere, Category="Message")
+	TArray<FString> Channels;
+
+	UPROPERTY(EditAnywhere, Category="Message")
+	TArray<uint32> Ids;
+	
+	UPROPERTY(EditAnywhere, Category="Message")
+	TArray<FString> Descriptions;
+	
+	UPROPERTY(EditAnywhere, Category="Message")
+	TArray<uint32> ReadOnlyIds;
+};
+
+/**
+ * Response for FTraceControlChannelsPing with information about
+ * enabled channels.
+ */
+USTRUCT()
+struct FTraceControlChannelsStatus
+{
+	GENERATED_USTRUCT_BODY()
+	
+	UPROPERTY(EditAnywhere, Category="Message")
+	TArray<uint32> EnabledIds;
+};
+
+
+/**
+ * Command to set the state of channels
+ */
+USTRUCT()
+struct FTraceControlChannelsSet
+{
+	GENERATED_USTRUCT_BODY()
+	
+	UPROPERTY(EditAnywhere, Category="Message")
+	TArray<uint32> ChannelIdsToEnable;
+	
+	UPROPERTY(EditAnywhere, Category="Message")
+	TArray<uint32> ChannelIdsToDisable;
+};
+
+/**
+ * Message sent when first connecting to a session. Expected to return
+ * FTraceControlPong message.
+ */
+USTRUCT()
+struct FTraceControlDiscoveryPing
+{
+	GENERATED_USTRUCT_BODY()
+	
+	UPROPERTY(EditAnywhere, Category="Message")
+	FGuid SessionId;
+	
+	UPROPERTY(EditAnywhere, Category="Message")
+	FGuid InstanceId;
+};
+
+/**
+ * Response for FTraceControlPing. A status message with additional
+ * information about session instance and properties that doesn't change
+ * over the course of the session lifetime.
+ */
+USTRUCT()
+struct FTraceControlDiscovery : public FTraceControlStatus
+{
+	GENERATED_USTRUCT_BODY()
+	
+	UPROPERTY(EditAnywhere, Category="Message")
+	FGuid SessionId;
+	
+	UPROPERTY(EditAnywhere, Category="Message")
+	FGuid InstanceId;
+};
+
+/**
+ * Stop current trace.
+ */
+USTRUCT()
+struct FTraceControlStop
+{
+	GENERATED_USTRUCT_BODY()
+};
+
+/**
+ * Shared properties for starting a trace. Never use this message directly.
+ */
+USTRUCT()
+struct FTraceControlStartCommon
+{
+	GENERATED_USTRUCT_BODY()
+	
+	UPROPERTY(EditAnywhere, Category="Message")
+	FString Channels;
+
+	UPROPERTY(EditAnywhere, Category="Message")
+	bool bExcludeTail;
+};
+
+/**
+ * Start tracing to a host
+ */
+USTRUCT()
+struct FTraceControlSend : public FTraceControlStartCommon
+{
+	GENERATED_USTRUCT_BODY()
+
+	UPROPERTY(EditAnywhere, Category="Message")
+	FString Host;
+};
+
+/**
+ * Start tracing to a file (on the receiver)
+ */
+USTRUCT()
+struct FTraceControlFile : public FTraceControlStartCommon
+{
+	GENERATED_USTRUCT_BODY()
+	
+	UPROPERTY(EditAnywhere, Category="Message")
+	FString File;
+	
+	UPROPERTY(EditAnywhere, Category="Message")
+	bool bTruncateFile;
+};
+
+/**
+ * Pause all tracing (i.e. mute all channels)
+ */
+USTRUCT()
+struct FTraceControlPause
+{
+	GENERATED_USTRUCT_BODY()
+};
+
+/**
+ * Resume all tracing (from pause)
+ */
+USTRUCT()
+struct FTraceControlResume
+{
+	GENERATED_USTRUCT_BODY()
+};
+
+/**
+ * Start tracing to a host
+ */
+USTRUCT()
+struct FTraceControlSnapshotSend
+{
+	GENERATED_USTRUCT_BODY()
+	
+	UPROPERTY(EditAnywhere, Category="Message")
+	FString Host;
+};
+
+/**
+ * Start tracing to a host
+ */
+USTRUCT()
+struct FTraceControlSnapshotFile
+{
+	GENERATED_USTRUCT_BODY()
+	
+	UPROPERTY(EditAnywhere, Category="Message")
+	FString File;
+};
+
+/**
+ * Start tracing to a host
+ */
+USTRUCT()
+struct FTraceControlBookmark
+{
+	GENERATED_USTRUCT_BODY()
+
+	UPROPERTY(EditAnywhere, Category="Message")
+	FString Label;
+};
+
+/**
+ * Start tracing to a host
+ */
+USTRUCT()
+struct FTraceControlScreenshot
+{
+	GENERATED_USTRUCT_BODY()
+	
+	UPROPERTY(EditAnywhere, Category="Message")
+	FString Name;
+	
+	UPROPERTY(EditAnywhere, Category="Message")
+	bool bShowUI;
+};
diff --git a/Engine/Source/Runtime/Launch/Private/LaunchEngineLoop.cpp b/Engine/Source/Runtime/Launch/Private/LaunchEngineLoop.cpp
--- a/Engine/Source/Runtime/Launch/Private/LaunchEngineLoop.cpp
+++ b/Engine/Source/Runtime/Launch/Private/LaunchEngineLoop.cpp
@@ -176,6 +176,7 @@
 	#include "DeviceProfiles/DeviceProfileManager.h"
 	#include "Commandlets/Commandlet.h"
 	#include "EngineService.h"
+	#include "TraceService.h"
 	#include "ContentStreaming.h"
 	#include "HighResScreenshot.h"
 	#include "Misc/HotReloadInterface.h"
@@ -4812,6 +4813,7 @@ int32 FEngineLoop::Init()
 		}
 
 		EngineService = new FEngineService();
+		TraceService = new FTraceService();
 	}
 
 	{
@@ -4846,7 +4848,7 @@ int32 FEngineLoop::Init()
 		GetMoviePlayer()->WaitForMovieToFinish();
     }
 
-	FTraceAuxiliary::EnableChannels();
+	FTraceAuxiliary::EnableCommandlineChannels();
 
 #if !UE_SERVER
 	// initialize media framework
@@ -4968,6 +4970,9 @@ void FEngineLoop::Exit()
 	delete EngineService;
 	EngineService = nullptr;
 
+	delete TraceService;
+	TraceService = nullptr;
+
 	if (SessionService.IsValid())
 	{
 		SessionService->Stop();
diff --git a/Engine/Source/Runtime/Launch/Public/LaunchEngineLoop.h b/Engine/Source/Runtime/Launch/Public/LaunchEngineLoop.h
--- a/Engine/Source/Runtime/Launch/Public/LaunchEngineLoop.h
+++ b/Engine/Source/Runtime/Launch/Public/LaunchEngineLoop.h
@@ -9,6 +9,7 @@
 #endif
 
 class FEngineService;
+class FTraceService;
 class FPendingCleanupObjects;
 class ISessionService;
 class FSlateRenderer;
@@ -188,6 +189,9 @@ private:
 	/** Holds the engine service. */
 	FEngineService* EngineService;
 
+	/** Trace control service */
+	FTraceService* TraceService;
+
 	/** Holds the application session service. */
 	TSharedPtr<ISessionService> SessionService;
 
diff --git a/Engine/Source/Runtime/SessionServices/Private/SessionServicesModule.cpp b/Engine/Source/Runtime/SessionServices/Private/SessionServicesModule.cpp
--- a/Engine/Source/Runtime/SessionServices/Private/SessionServicesModule.cpp
+++ b/Engine/Source/Runtime/SessionServices/Private/SessionServicesModule.cpp
@@ -1,16 +1,16 @@
 // Copyright Epic Games, Inc. All Rights Reserved.
 
 #include "CoreMinimal.h"
-#include "Misc/CoreMisc.h"
-#include "Misc/App.h"
-#include "Modules/ModuleManager.h"
 #include "IMessagingModule.h"
 #include "ISessionManager.h"
-#include "SessionManager.h"
 #include "ISessionService.h"
-#include "SessionService.h"
 #include "ISessionServicesModule.h"
-
+#include "Misc/App.h"
+#include "Misc/CoreMisc.h"
+#include "Modules/ModuleManager.h"
+#include "SessionManager.h"
+#include "SessionService.h"
+#include "TraceController.h"
 
 /**
  * Implements the SessionServices module.
@@ -112,6 +112,16 @@ public:
 		return SessionService;
 	}
 
+	virtual TSharedPtr<ITraceController> GetTraceController() override
+	{
+		if (!TraceController.IsValid() && MessageBusPtr.IsValid())
+		{
+			TraceController = MakeShareable(new FTraceController(MessageBusPtr.Pin().ToSharedRef()));
+		}
+
+		return TraceController;
+	}
+
 public:
 
 	//~ IModuleInterface interface
@@ -123,6 +133,7 @@ public:
 
 	virtual void ShutdownModule() override
 	{
+		TraceController.Reset();
 		SessionManager.Reset();
 		SessionService.Reset();
 	}
@@ -137,6 +148,9 @@ private:
 
 	/** Holds the session service singleton. */
 	TSharedPtr<ISessionService> SessionService;
+
+	/** Holds the control for Tracing */
+	TSharedPtr<FTraceController> TraceController;
 };
 
 
diff --git a/Engine/Source/Runtime/SessionServices/Private/TraceController.cpp b/Engine/Source/Runtime/SessionServices/Private/TraceController.cpp
new file mode 100644
--- /dev/null
+++ b/Engine/Source/Runtime/SessionServices/Private/TraceController.cpp
@@ -0,0 +1,356 @@
+﻿// Copyright Epic Games, Inc. All Rights Reserved.
+
+#include "TraceController.h"
+
+#include "IMessageBus.h"
+#include "ISessionServicesModule.h"
+#include "MessageEndpoint.h"
+#include "MessageEndpointBuilder.h"
+#include "TraceControlMessages.h"
+
+FTraceController::FTraceController(const TSharedRef<IMessageBus>& InMessageBus)
+	: MessageBus(InMessageBus)
+{
+	MessageEndpoint = FMessageEndpoint::Builder("FTraceController", InMessageBus)
+		.Handling<FTraceControlDiscovery>(this, &FTraceController::OnDiscoveryResponse)
+		.Handling<FTraceControlStatus>(this, &FTraceController::OnStatus)
+		.Handling<FTraceControlSettings>(this, &FTraceController::OnSettings)
+		.Handling<FTraceControlChannelsDesc>(this, &FTraceController::OnChannelsDesc)
+		.Handling<FTraceControlChannelsStatus>(this, &FTraceController::OnChannelsStatus)
+		.NotificationHandling(FOnBusNotification::CreateRaw(this, &FTraceController::OnNotification));
+	
+	ISessionServicesModule& SessionServicesModule = FModuleManager::LoadModuleChecked<ISessionServicesModule>("SessionServices");
+	SessionManager = SessionServicesModule.GetSessionManager();
+	
+	if (SessionManager.IsValid())
+	{
+		//Note that we should not add ourselves as shared pointer to session manager,
+		//that will create a circular dependency
+		SessionManager->OnInstanceSelectionChanged().AddRaw(this, &FTraceController::OnInstanceSelectionChanged);
+	}
+}
+
+FTraceController::~FTraceController()
+{
+	if (SessionManager.IsValid())
+	{
+		SessionManager->OnInstanceSelectionChanged().RemoveAll(this);
+	}
+}
+
+template<typename StringType>
+TArray<uint32> ResolveChannelNames(TConstArrayView<StringType> ChannelNames, const TMap<uint32, FTraceStatus::FChannel> Channels)
+{
+	TArray<uint32> Ids;
+	if (ChannelNames.Num() > 0)
+	{
+		for (const auto& [Id, Channel] : Channels)
+		{
+			if (ChannelNames.Contains(Channel.Name))
+			{
+				Ids.Add(Channel.Id);
+			}
+		}
+	}
+	return MoveTemp(Ids);
+}
+
+void FTraceController::SetChannels(TConstArrayView<FStringView> ChannelsToEnable, TConstArrayView<FStringView> ChannelsToDisable)
+{
+	FReadScopeLock _(InstancesLock);
+
+	// Since each session might have different notion about channel ids
+	// we need to send a specific message to each instance using their unique
+	// channel mappings.
+	for (const auto& Instance : SessionManager->GetSelectedInstances())
+	{
+		if (const auto Address = InstanceToAddress.Find(Instance->GetInstanceId()))
+		{
+			const auto& InstanceStatus = Instances[*Address];
+			const auto Message = FMessageEndpoint::MakeMessage<FTraceControlChannelsSet>();
+			Message->ChannelIdsToEnable = ResolveChannelNames(ChannelsToEnable, InstanceStatus.Channels);
+			Message->ChannelIdsToDisable = ResolveChannelNames(ChannelsToDisable, InstanceStatus.Channels);
+			MessageEndpoint->Send(Message, *Address);
+		}
+	}
+}
+
+void FTraceController::SetChannels(TConstArrayView<FString> ChannelsToEnable, TConstArrayView<FString> ChannelsToDisable)
+{
+	FReadScopeLock _(InstancesLock);
+	
+	for (const auto& Instance : SessionManager->GetSelectedInstances())
+	{
+		if (const auto Address = InstanceToAddress.Find(Instance->GetInstanceId()))
+		{
+			const auto& InstanceStatus = Instances[*Address];
+			const auto Message = FMessageEndpoint::MakeMessage<FTraceControlChannelsSet>();
+			Message->ChannelIdsToEnable = ResolveChannelNames(ChannelsToEnable, InstanceStatus.Channels);
+			Message->ChannelIdsToDisable = ResolveChannelNames(ChannelsToDisable, InstanceStatus.Channels);
+			MessageEndpoint->Send(Message, *Address);
+		}
+	}
+}
+
+void FTraceController::Send(FStringView Host, FStringView Channels, bool bExcludeTail)
+{
+	const auto Message = FMessageEndpoint::MakeMessage<FTraceControlSend>();
+	Message->Host = Host;
+	Message->Channels = Channels;
+	Message->bExcludeTail = bExcludeTail;
+
+	SendToSelectedSessions(Message);
+}
+
+void FTraceController::File(FStringView File, FStringView Channels, bool bExcludeTail, bool bTruncateFile)
+{
+	const auto Message = FMessageEndpoint::MakeMessage<FTraceControlFile>();
+	Message->File = File;
+	Message->Channels = Channels;
+	Message->bExcludeTail = bExcludeTail;
+	Message->bTruncateFile = bTruncateFile;
+
+	SendToSelectedSessions(Message);
+}
+
+void FTraceController::Stop()
+{
+	SendToSelectedSessions(FMessageEndpoint::MakeMessage<FTraceControlStop>());
+}
+
+void FTraceController::SnapshotSend(FStringView Host)
+{
+	const auto Message = FMessageEndpoint::MakeMessage<FTraceControlSnapshotSend>();
+	Message->Host = Host;
+
+	SendToSelectedSessions(Message);
+}
+
+void FTraceController::SnapshotFile(FStringView File)
+{
+	const auto Message = FMessageEndpoint::MakeMessage<FTraceControlSnapshotFile>();
+	Message->File = File;
+
+	SendToSelectedSessions(Message);
+}
+
+void FTraceController::Pause()
+{
+	SendToSelectedSessions(FMessageEndpoint::MakeMessage<FTraceControlPause>());
+}
+
+void FTraceController::Resume()
+{
+	SendToSelectedSessions(FMessageEndpoint::MakeMessage<FTraceControlResume>());
+}
+
+void FTraceController::Bookmark(FStringView Label)
+{
+	const auto Message = FMessageEndpoint::MakeMessage<FTraceControlBookmark>();
+	Message->Label = Label;
+	
+	SendToSelectedSessions(Message);
+}
+
+void FTraceController::Screenshot(FStringView Name, bool bShowUI)
+{
+	const auto Message = FMessageEndpoint::MakeMessage<FTraceControlScreenshot>();
+	Message->Name = Name;
+	Message->bShowUI = bShowUI;
+	
+	SendToSelectedSessions(Message);
+}
+
+void FTraceController::SendStatusUpdateRequest()
+{
+	MessageEndpoint->Publish(FMessageEndpoint::MakeMessage<FTraceControlStatusPing>());
+}
+
+void FTraceController::SendChannelUpdateRequest()
+{
+	FReadScopeLock _(InstancesLock);
+	
+	for (const auto& Instance : Instances)
+	{
+		const auto Message = FMessageEndpoint::MakeMessage<FTraceControlChannelsPing>();
+		Message->KnownChannelCount = uint32(Instance.Value.Channels.Num());
+		MessageEndpoint->Send(Message, Instance.Key);
+	}
+}
+
+void FTraceController::SendSettingsUpdateRequest()
+{
+	MessageEndpoint->Publish(FMessageEndpoint::MakeMessage<FTraceControlSettingsPing>());
+}
+
+void FTraceController::OnNotification(const FMessageBusNotification& Event)
+{
+	FWriteScopeLock _(InstancesLock);
+	
+	if (Event.NotificationType == EMessageBusNotification::Unregistered)
+	{
+		// Many endpoints may be removed from one instance, look for the
+		// one we have registered.
+		if(Instances.Remove(Event.RegistrationAddress) > 0)
+		{
+			InstanceToAddress = InstanceToAddress.FilterByPredicate([&](auto It) { return It.Value == Event.RegistrationAddress; } );
+		}
+	}
+}
+
+void FTraceController::OnDiscoveryResponse(const FTraceControlDiscovery& Message, const TSharedRef<IMessageContext>& Context)
+{
+	FWriteScopeLock _(InstancesLock);
+	
+	FTraceStatus& Status = Instances.FindOrAdd(Context->GetSender());
+	InstanceToAddress.FindOrAdd(Message.InstanceId, Context->GetSender());
+	
+	UpdateStatus(Message, Status);
+	// This is the application session id, which is not necessarily the
+	// same as trace session (may be overridden by commandline)
+	Status.SessionId = Message.SessionId;
+	Status.InstanceId = Message.InstanceId;
+
+	if (SelectedInstanceId == Message.InstanceId)
+	{
+		StatusReceivedEvent.Broadcast(Status, FTraceStatus::EUpdateType::All);
+	}
+}
+
+void FTraceController::OnStatus(const FTraceControlStatus& Message, const TSharedRef<IMessageContext>& Context)
+{
+	FWriteScopeLock _(InstancesLock);
+	
+	if (FTraceStatus* Status = Instances.Find(Context->GetSender()))
+	{
+		UpdateStatus(Message, *Status);
+		
+		if (SelectedInstanceId == Status->InstanceId)
+		{
+			StatusReceivedEvent.Broadcast(*Status, FTraceStatus::EUpdateType::Status);
+		}
+	}
+	
+}
+
+void FTraceController::OnChannelsDesc(const FTraceControlChannelsDesc& Message, const TSharedRef<IMessageContext>& Context)
+{
+	FWriteScopeLock _(InstancesLock);
+	
+	const auto Status = Instances.Find(Context->GetSender());
+	
+	if (!Status)
+	{
+		return;
+	}
+	
+	check(Message.Channels.Num() == Message.Ids.Num() && Message.Channels.Num() == Message.Descriptions.Num());
+	const int32 Count = Message.Channels.Num();
+	for(int32 Index = 0; Index < Count; ++Index)
+	{
+		const uint32 Id = Message.Ids[Index];
+		if (Status->Channels.Contains(Id))
+		{
+			continue;
+		}
+
+		FTraceStatus::FChannel& NewChannel = Status->Channels.Add(Id);
+		NewChannel.Name = Message.Channels[Index];
+		NewChannel.Description = Message.Descriptions[Index];
+		NewChannel.Id = Id;
+		NewChannel.bEnabled = false;
+		NewChannel.bReadOnly = Message.ReadOnlyIds.Contains(Id);
+		
+		if (SelectedInstanceId == Status->InstanceId)
+		{
+			StatusReceivedEvent.Broadcast(*Status, FTraceStatus::EUpdateType::ChannelsDesc);
+		}
+	}
+}
+
+void FTraceController::OnChannelsStatus(const FTraceControlChannelsStatus& Message, const TSharedRef<IMessageContext>& Context)
+{
+	FWriteScopeLock _(InstancesLock);
+	
+	if (const auto Status = Instances.Find(Context->GetSender()))
+	{
+		for (auto& [Id, Channel] : Status->Channels)
+		{
+			Channel.bEnabled = Message.EnabledIds.Contains(Id);
+		}
+		
+		if (SelectedInstanceId == Status->InstanceId)
+		{
+			StatusReceivedEvent.Broadcast(*Status, FTraceStatus::EUpdateType::ChannelsStatus);
+		}
+	}
+}
+
+void FTraceController::OnSettings(const FTraceControlSettings& Message, const TSharedRef<IMessageContext>& Context)
+{
+	FWriteScopeLock _(InstancesLock);
+	
+	if (const auto Status = Instances.Find(Context->GetSender()))
+	{
+		FTraceStatus::FSettings& Settings = Status->Settings;
+		Settings.bStatNamedEvents = Message.bStatNamedEvents;
+		Settings.bUseImportantCache = Message.bUseImportantCache;
+		Settings.bUseWorkerThread = Message.bUseWorkerThread;
+		Settings.TailSizeBytes = Message.TailSizeBytes;
+		
+		if (SelectedInstanceId == Status->InstanceId)
+		{
+			StatusReceivedEvent.Broadcast(*Status, FTraceStatus::EUpdateType::Settings);
+		}
+	}
+}
+
+void FTraceController::OnInstanceSelectionChanged(const TSharedPtr<ISessionInstanceInfo>& Instance, bool bSelected)
+{
+	FReadScopeLock _(InstancesLock);
+	
+	SelectedInstanceId = Instance->GetInstanceId();
+	if (const auto Address = InstanceToAddress.Find(SelectedInstanceId))
+	{
+		if (const auto Status = Instances.Find(*Address))
+		{
+			StatusReceivedEvent.Broadcast(*Status, FTraceStatus::EUpdateType::All);
+		}
+	}
+	else
+	{
+		// We haven't discovered this instance yet, send a discovery message specifically
+		// to that instance.
+		const auto Message = FMessageEndpoint::MakeMessage<FTraceControlDiscoveryPing>();
+		Message->SessionId = Instance->GetOwnerSession()->GetSessionId();
+		Message->InstanceId = Instance->GetInstanceId();
+		MessageEndpoint->Publish<FTraceControlDiscoveryPing>(Message);
+	}
+}
+
+TArray<FMessageAddress> FTraceController::GetSelectedSessionAddresses()
+{
+	TArray<FMessageAddress> Selected;
+	auto SelectedInstances = SessionManager->GetSelectedInstances();
+	for (const auto& Instance : SelectedInstances)
+	{
+		if (const auto Address = InstanceToAddress.Find(Instance->GetInstanceId()))
+		{
+			Selected.Add(*Address);
+		}
+	}
+	return MoveTemp(Selected);
+}
+
+void FTraceController::UpdateStatus(const FTraceControlStatus& Message, FTraceStatus& Status)
+{
+	Status.bIsTracing = !Message.Endpoint.IsEmpty();
+	Status.Endpoint = Message.Endpoint;
+	Status.SessionGuid = Message.SessionGuid;
+	Status.TraceGuid = Message.TraceGuid;
+	Status.Stats.BytesSent = Message.BytesSent;
+	Status.Stats.BytesTraced = Message.BytesTraced;
+	Status.Stats.CacheAllocated = Message.CacheAllocated;
+	Status.Stats.CacheUsed = Message.CacheUsed;
+	Status.Stats.CacheWaste = Message.CacheWaste;
+}
diff --git a/Engine/Source/Runtime/SessionServices/Private/TraceController.h b/Engine/Source/Runtime/SessionServices/Private/TraceController.h
new file mode 100644
--- /dev/null
+++ b/Engine/Source/Runtime/SessionServices/Private/TraceController.h
@@ -0,0 +1,115 @@
+﻿// Copyright Epic Games, Inc. All Rights Reserved.
+
+#pragma once
+
+#include "IMessageContext.h"
+#include "ISessionManager.h"
+#include "ITraceController.h"
+#include "MessageEndpoint.h"
+#include "TraceControlMessages.h"
+
+struct FTraceControlSettings;
+struct FTraceControlStatus;
+struct FTraceControlDiscovery;
+class ISessionManager;
+class IMessageBus;
+class FMessageEndpoint;
+
+/**
+ * Interface to control other sessions tracing.
+ */
+class FTraceController : public ITraceController 
+{
+public:
+	FTraceController(const TSharedRef<IMessageBus>& InMessageBus);
+	virtual ~FTraceController() override;
+
+private:
+	/* ITraceController interface */
+	virtual void SetChannels(TConstArrayView<FStringView> ChannelsToEnable, TConstArrayView<FStringView> ChannelsToDisable) override;
+	virtual void SetChannels(TConstArrayView<FString> ChannelsToEnable, TConstArrayView<FString> ChannelsToDisable) override;
+	virtual void Send(FStringView Host, FStringView Channels, bool bExcludeTail) override;
+	virtual void File(FStringView File, FStringView Channels, bool bExcludeTail, bool bTruncateFile) override;
+	virtual void Stop() override;
+	virtual void SnapshotSend(FStringView Host) override;
+	virtual void SnapshotFile(FStringView File) override;
+	virtual void Pause() override;
+	virtual void Resume() override;
+	virtual void Bookmark(FStringView Label) override;
+	virtual void Screenshot(FStringView Name, bool bShowUI) override;
+
+	virtual void SendStatusUpdateRequest() override;
+	virtual void SendChannelUpdateRequest() override;
+	virtual void SendSettingsUpdateRequest() override;
+	
+	DECLARE_DERIVED_EVENT(FTraceController, ITraceController::FStatusRecievedEvent, FStatusRecievedEvent);
+	virtual FStatusRecievedEvent& OnStatusReceived() override
+	{
+		return StatusReceivedEvent;
+	}
+
+	/* Message handlers */
+	void OnNotification(const FMessageBusNotification& MessageBusNotification);
+	void OnDiscoveryResponse(const FTraceControlDiscovery& Message, const TSharedRef<IMessageContext>& Context);
+	void OnStatus(const FTraceControlStatus& Message, const TSharedRef<IMessageContext>& Context);
+	void OnChannelsDesc(const FTraceControlChannelsDesc& Message, const TSharedRef<IMessageContext>& Context);
+	void OnChannelsStatus(const FTraceControlChannelsStatus& Message, const TSharedRef<IMessageContext>& Context);
+	void OnSettings(const FTraceControlSettings& Message, const TSharedRef<IMessageContext>& Context);
+	static void UpdateStatus(const FTraceControlStatus& Message, FTraceStatus& Status);
+
+	/* Events from SessionManager handlers */
+	void OnInstanceSelectionChanged(const TSharedPtr<class ISessionInstanceInfo>&, bool);
+
+	/** Utility for sending a message to all selected sessions */
+	template <class MessageType>
+	void SendToSelectedSessions(MessageType* Message);
+
+	TArray<FMessageAddress> GetSelectedSessionAddresses();
+
+public:
+
+private:
+	TWeakPtr<IMessageBus> MessageBus;
+
+	/** Our own endpoint for messages */
+	TSharedPtr<FMessageEndpoint> MessageEndpoint;
+
+	/** Session manager used for selecting sessions */
+	TSharedPtr<ISessionManager> SessionManager;
+
+	/** Address of the runtime endpoint for trace controls */
+	FMessageAddress TraceControlAddress;
+
+	FStatusRecievedEvent StatusReceivedEvent;
+
+	/** Lock to protect access to Instances list */
+	FRWLock InstancesLock;
+
+	/** Map of instances and their respective last reported status */
+	TMap<FMessageAddress, FTraceStatus> Instances;
+
+	/** Secondary lookup from instance -> address */
+	TMap<FGuid, FMessageAddress> InstanceToAddress;
+
+	/** Currently selected instance */
+	FGuid SelectedInstanceId;
+};
+
+template <typename MessageType>
+void FTraceController::SendToSelectedSessions(MessageType* Message)
+{
+	TArray<FMessageAddress> Recipients;
+	auto SelectedInstances = SessionManager->GetSelectedInstances();
+	for (const auto& Instance : SelectedInstances)
+	{
+		if (const auto Address = InstanceToAddress.Find(Instance->GetInstanceId()))
+		{
+			Recipients.Add(*Address);
+		}
+	}
+
+	if (!Recipients.IsEmpty())
+	{
+		MessageEndpoint->Send(Message, Recipients);
+	}
+}
diff --git a/Engine/Source/Runtime/SessionServices/Public/ISessionServicesModule.h b/Engine/Source/Runtime/SessionServices/Public/ISessionServicesModule.h
--- a/Engine/Source/Runtime/SessionServices/Public/ISessionServicesModule.h
+++ b/Engine/Source/Runtime/SessionServices/Public/ISessionServicesModule.h
@@ -7,6 +7,8 @@
 #include "ISessionManager.h"
 #include "ISessionService.h"
 
+class ITraceController;
+
 /**
  * Interface for session core modules.
  */
@@ -31,6 +33,12 @@ public:
 	 */
 	virtual TSharedPtr<ISessionService> GetSessionService() = 0;
 
+
+	/**
+	 * Get the trace controller
+	 */
+	virtual TSharedPtr<ITraceController> GetTraceController() = 0; 
+
 public:
 
 	/** Virtual destructor. */
diff --git a/Engine/Source/Runtime/SessionServices/Public/ITraceController.h b/Engine/Source/Runtime/SessionServices/Public/ITraceController.h
new file mode 100644
--- /dev/null
+++ b/Engine/Source/Runtime/SessionServices/Public/ITraceController.h
@@ -0,0 +1,187 @@
+﻿// Copyright Epic Games, Inc. All Rights Reserved.
+
+#pragma once
+
+
+struct FTraceStatus
+{
+	/**
+	 * Update types
+	 */
+	enum class EUpdateType : uint8
+	{
+		Status			= 1 << 0,
+		Settings		= 1 << 1,
+		ChannelsDesc	= 1 << 2,
+		ChannelsStatus	= 1 << 3,
+		All				= Status|Settings|ChannelsDesc|ChannelsStatus
+	};
+	
+	struct FSettings
+	{
+		/** If worker thread is used or TraceLog is pumped on end frame. */
+		bool bUseWorkerThread;
+		/** If important cache is enabled */
+		bool bUseImportantCache;
+		/** Size of tail buffer */
+		uint32 TailSizeBytes;
+		/** If stats are emitted as named events */
+		bool bStatNamedEvents;
+	};
+
+	struct FChannel
+	{
+		FString Name;
+		FString Description;
+		uint32 Id;
+		bool bEnabled;
+		bool bReadOnly;
+	};
+
+	struct FStats
+	{
+		/** Number of bytes sent to server or file. */
+		uint64 BytesSent;
+		/** Number of (uncompressed) bytes traced from process. */
+		uint64 BytesTraced;
+		/** Total memory used by TraceLog */
+		uint64 MemoryUsed;
+		/** Allocated memory for important events */
+		uint32 CacheAllocated;
+		/** Memory used for important events */
+		uint32 CacheUsed;
+		/** Wasted space for important events cache */
+		uint32 CacheWaste;
+	};
+	
+	/** Session id of the process we're controlling */
+	FGuid SessionId;
+	/** Instance id of the process we're controlling */
+	FGuid InstanceId;
+	/** If tracing is active */
+	bool bIsTracing = false;
+	/** Endpoint of active trace */
+	FString Endpoint;
+	/** Session identifier for the trace. */
+	FGuid SessionGuid;
+	/** Unique identifier for the trace */
+	FGuid TraceGuid;
+	/** Settings */
+	FSettings Settings;
+	/** State of channels */
+	TMap<uint32, FChannel> Channels;
+	/** Stats */
+	FStats Stats;
+};
+
+ENUM_CLASS_FLAGS(FTraceStatus::EUpdateType);
+
+/**
+ * Interface to control other sessions tracing.
+ */
+class ITraceController
+{
+public:
+	
+	virtual ~ITraceController() = default;
+
+	/**
+	 * Enables or disables channels by name
+	 * @param ChannelsToEnable  List of channels to enable
+	 * @param ChannelsToDisable List of channels to disable
+	 */
+	virtual void SetChannels(TConstArrayView<FStringView> ChannelsToEnable, TConstArrayView<FStringView> ChannelsToDisable) = 0;
+
+	/**
+	 * Enables or disables channels by name
+	 * @param ChannelsToEnable  List of channels to enable
+	 * @param ChannelsToDisable List of channels to disable
+	 */
+	virtual void SetChannels(TConstArrayView<FString> ChannelsToEnable, TConstArrayView<FString> ChannelsToDisable) = 0;
+	
+	/**
+	 * Start a trace on selected instances to the provided host, using a set of channels.
+	 * @param Host Host to send the trace to
+	 * @param Channels Comma separated list of channels to enable
+	 * @param bExcludeTail If the tail (circular buffer of recent events) should be included
+	 */
+	virtual void Send(FStringView Host, FStringView Channels, bool bExcludeTail = false) = 0;
+
+	/**
+	 * Start a trace on selected instances to a file on the instance, using a set of channels.
+	 * @param File Path on the instance. ".utrace" will be appended
+	 * @param Channels Comma separated list of channels to enable
+	 * @param bExcludeTail If the tail (circular buffer of recent events) should be included
+	 * @param bTruncateFile If the file should be truncated (if already exists)
+	 */
+	virtual void File(FStringView File, FStringView Channels, bool bExcludeTail = false, bool bTruncateFile = false) = 0;
+
+	/**
+	 * On selected instances, make a snapshot of the tail (circular buffer of recent events)
+	 * and send to the provided host.
+	 * @param Host Host to send the trace to
+	 */
+	virtual void SnapshotSend(FStringView Host) = 0;
+	
+	/**
+	 * On selected instances, make a snapshot of the tail (circular buffer of recent events)
+	 * and save to a file.
+	 * @param File Path on the instance. ".utrace" will be appended
+	 */
+	virtual void SnapshotFile(FStringView File) = 0;
+
+	/**
+	 * Pause tracing by muting all (non-readonly) channels.
+	 */
+	virtual void Pause() = 0;
+
+	/**
+	 * Resume tracing (from paused) by enabling the previously enabled channels.
+	 */
+	virtual void Resume() = 0;
+
+	/**
+	 * Stop active trace.
+	 */
+	virtual void Stop() = 0;
+
+	/**
+	 * Insert bookmark into the trace.
+	 * @param Label Label of bookmark
+	 */
+	virtual void Bookmark(FStringView Label) = 0;
+
+	/**
+	 * Insert screenshot into the trace.
+	 * @param Name Name of the screenshot
+	 * @param bShowUI If the UI should be visible in the image
+	 */
+	virtual void Screenshot(FStringView Name, bool bShowUI) = 0;
+
+	/**
+	 * Request update of the status from all sessions and instances.
+	 */
+	virtual void SendStatusUpdateRequest() = 0;
+
+	/**
+	 * Request update of channel status from all sessions and instances.
+	 */
+	virtual void SendChannelUpdateRequest() = 0;
+
+	/**
+	 * Request update of setting state from all sessions and instances.
+	 */
+	virtual void SendSettingsUpdateRequest() = 0;
+
+	/**
+	 * Generic event for updates of status
+	 */
+	DECLARE_EVENT_TwoParams(ITraceController, FStatusRecievedEvent, const FTraceStatus&, FTraceStatus::EUpdateType);
+
+	/**
+	 * Event triggered whenever the FTraceStatus for the selected instance (in session manager)
+	 * is updated. A reference to the status and what has changed is provided.
+	 */
+	virtual FStatusRecievedEvent& OnStatusReceived() = 0;
+	
+};
