Open into GetTextRev call which we do in another place

diff --git a/Engine/Source/Runtime/Core/Private/Internationalization/TextHistory.cpp b/Engine/Source/Runtime/Core/Private/Internationalization/TextHistory.cpp
--- a/Engine/Source/Runtime/Core/Private/Internationalization/TextHistory.cpp
+++ b/Engine/Source/Runtime/Core/Private/Internationalization/TextHistory.cpp
@@ -711,7 +711,10 @@ void FTextHistory::MarkDisplayStringUpToDate()
 {
 	if (CanUpdateDisplayString())
 	{
-		FTextLocalizationManager::Get().GetTextRevisions(GetTextId(), GlobalRevision, LocalRevision);
+		UE_AUTORTFM_OPEN(
+		{
+			FTextLocalizationManager::Get().GetTextRevisions(GetTextId(), GlobalRevision, LocalRevision);
+		});
 	}
 	else
 	{
