Added bool FPlatform::CanOverallocateVirtualMemory() function to check if we can allocate more virtual memory than physical + swap

diff --git a/Engine/Source/Runtime/Core/Private/Apple/ApplePlatformMemory.cpp b/Engine/Source/Runtime/Core/Private/Apple/ApplePlatformMemory.cpp
--- a/Engine/Source/Runtime/Core/Private/Apple/ApplePlatformMemory.cpp
+++ b/Engine/Source/Runtime/Core/Private/Apple/ApplePlatformMemory.cpp
@@ -307,8 +307,7 @@ void FApplePlatformMemory::SetAllocatorToUse()
     }
     if (USE_MALLOC_BINNED2)
     {
- #if PLATFORM_IOS || PLATFORM_TVOS
-        if(!FIOSPlatformMisc::IsEntitlementEnabled("com.apple.developer.kernel.extended-virtual-addressing"))
+        if(!CanOverallocateVirtualMemory())
         {
 			FPlatformMisc::LowLevelOutputDebugString(TEXT("MallocBinned2 requested but Virtual Address Space entitlement not found. Check your entitlements. Falling back to Ansi.\n"));
             AllocatorToUse = EMemoryAllocatorToUse::Ansi;
@@ -318,7 +317,6 @@ void FApplePlatformMemory::SetAllocatorToUse()
         {
 			FPlatformMisc::LowLevelOutputDebugString(TEXT("Virtual Address Space entitlement found. Using MallocBinned2 allocator.\n"));
         }
-#endif
 		FPlatformMisc::LowLevelOutputDebugString(TEXT("Using MallocBinned2 allocator.\n"));
         AllocatorToUse = EMemoryAllocatorToUse::Binned2;
         return;
@@ -852,4 +850,13 @@ bool FApplePlatformMemory::GetLLMAllocFunctions(void*_Nonnull(*_Nonnull&OutAlloc
 #endif
 }
 
+static bool FApplePlatformMemory::CanOverallocateVirtualMemory()
+{
+#if PLATFORM_IOS || PLATFORM_TVOS
+	static bool bHasExtendedVirtualAddressingEntitlement = FIOSPlatformMisc::IsEntitlementEnabled("com.apple.developer.kernel.extended-virtual-addressing");
+	return bHasExtendedVirtualAddressingEntitlement;
+#endif
+	return true;	// 64 bit Mac process can allocate ~18 exabytes of addressable space
+}
+
 NS_ASSUME_NONNULL_END
diff --git a/Engine/Source/Runtime/Core/Private/Unix/UnixPlatformMemory.cpp b/Engine/Source/Runtime/Core/Private/Unix/UnixPlatformMemory.cpp
--- a/Engine/Source/Runtime/Core/Private/Unix/UnixPlatformMemory.cpp
+++ b/Engine/Source/Runtime/Core/Private/Unix/UnixPlatformMemory.cpp
@@ -829,7 +829,38 @@ static float ReadOvercommitRatio()
 	}
 
 	return OutVal;
+}
+
+static int ReadOvercommitMode()
+{
+	int PolicyType = 2;	// Do not overcommit by default
+	int Fd = open("/proc/sys/vm/overcommit_memory", O_RDONLY);
 
+	if (Fd < 0)
+	{
+		static bool bLogOnceFileNotFound = false;
+		if (!bLogOnceFileNotFound)
+		{
+			fprintf(stderr, "Warning: ReadOvercommitMode failed to open /proc/sys/vm/overcommit_memory (Err %d).\n", errno);
+			fflush(stderr);
+			bLogOnceFileNotFound = true;
+		}
+	}
+	else
+	{
+		// The overcommit_memory policy is defined by a number from 0 to 2
+		// https://www.kernel.org/doc/Documentation/vm/overcommit-accounting
+		char Buffer[512] = { 0 };
+		ssize_t ReadBytes = read(Fd, Buffer, sizeof(Buffer) - 1);
+		if (ReadBytes > 0)
+		{
+			PolicyType = FCStringAnsi::Atoi(Buffer);
+		}
+
+		close(Fd);
+	}
+
+	return PolicyType;
 }
 
 // struct CORE_API FGenericPlatformMemoryStats : public FPlatformMemoryConstants
@@ -1426,3 +1457,9 @@ bool FUnixPlatformMemory::GetLLMAllocFunctions(void*(*&OutAllocFunction)(size_t)
 	return false;
 #endif
 }
+
+bool FUnixPlatformMemory::CanOverallocateVirtualMemory()
+{
+	static int OvercommitPolicy = ReadOvercommitMode();
+	return OvercommitPolicy != 2;
+}
diff --git a/Engine/Source/Runtime/Core/Public/Apple/ApplePlatformMemory.h b/Engine/Source/Runtime/Core/Public/Apple/ApplePlatformMemory.h
--- a/Engine/Source/Runtime/Core/Public/Apple/ApplePlatformMemory.h
+++ b/Engine/Source/Runtime/Core/Public/Apple/ApplePlatformMemory.h
@@ -151,6 +151,8 @@ struct CORE_API FApplePlatformMemory : public FGenericPlatformMemory
 	
 	/** Setup the current default CFAllocator to use our malloc functions. */
 	static void ConfigureDefaultCFAllocator(void);
+
+	static bool CanOverallocateVirtualMemory();
 	
 	static vm_address_t NanoRegionStart;
 	static vm_address_t NanoRegionEnd;
diff --git a/Engine/Source/Runtime/Core/Public/GenericPlatform/GenericPlatformMemory.h b/Engine/Source/Runtime/Core/Public/GenericPlatform/GenericPlatformMemory.h
--- a/Engine/Source/Runtime/Core/Public/GenericPlatform/GenericPlatformMemory.h
+++ b/Engine/Source/Runtime/Core/Public/GenericPlatform/GenericPlatformMemory.h
@@ -812,6 +812,15 @@ public:
 	 */
 	static CORE_API FString PrettyMemory( uint64 Memory );
 
+	/**
+	* Return true if the platform can allocate a lot more virtual memory than physical memory
+	* It's true for most platforms, but iOS needs a special entitlement and Linux a kernel config for this
+	*/
+	FORCEINLINE static bool CanOverallocateVirtualMemory()
+	{
+		return true;
+	}
+
 protected:
 	friend struct FGenericStatsUpdater;
 
diff --git a/Engine/Source/Runtime/Core/Public/Unix/UnixPlatformMemory.h b/Engine/Source/Runtime/Core/Public/Unix/UnixPlatformMemory.h
--- a/Engine/Source/Runtime/Core/Public/Unix/UnixPlatformMemory.h
+++ b/Engine/Source/Runtime/Core/Public/Unix/UnixPlatformMemory.h
@@ -141,6 +141,8 @@ struct FUnixPlatformMemory : public FGenericPlatformMemory
 	//~ End FGenericPlatformMemory Interface
 
 	static CORE_API bool HasForkPageProtectorEnabled();
+
+	static CORE_API bool CanOverallocateVirtualMemory();
 };
 
 typedef FUnixPlatformMemory FPlatformMemory;
