Fixed compilation on Apple platforms

diff --git a/Engine/Source/Runtime/Core/Private/Apple/ApplePlatformMemory.cpp b/Engine/Source/Runtime/Core/Private/Apple/ApplePlatformMemory.cpp
--- a/Engine/Source/Runtime/Core/Private/Apple/ApplePlatformMemory.cpp
+++ b/Engine/Source/Runtime/Core/Private/Apple/ApplePlatformMemory.cpp
@@ -850,7 +850,7 @@ bool FApplePlatformMemory::GetLLMAllocFunctions(void*_Nonnull(*_Nonnull&OutAlloc
 #endif
 }
 
-static bool FApplePlatformMemory::CanOverallocateVirtualMemory()
+bool FApplePlatformMemory::CanOverallocateVirtualMemory()
 {
 #if PLATFORM_IOS || PLATFORM_TVOS
 	static bool bHasExtendedVirtualAddressingEntitlement = FIOSPlatformMisc::IsEntitlementEnabled("com.apple.developer.kernel.extended-virtual-addressing");
