Fixed case where platform ShadingModel mask was not applied in all cases

diff --git a/Engine/Source/Runtime/Core/Private/UObject/FortniteReleaseBranchCustomObjectVersion.cpp b/Engine/Source/Runtime/Core/Private/UObject/FortniteReleaseBranchCustomObjectVersion.cpp
--- a/Engine/Source/Runtime/Core/Private/UObject/FortniteReleaseBranchCustomObjectVersion.cpp
+++ b/Engine/Source/Runtime/Core/Private/UObject/FortniteReleaseBranchCustomObjectVersion.cpp
@@ -8,5 +8,6 @@ TMap<FGuid, FGuid> FFortniteReleaseBranchCustomObjectVersion::GetSystemGuids()
 
 	SystemGuids.Add(DevGuids.Niagara_LatestScriptCompileVersion, FGuid("5F589BAA37394A51BF718ECB10D73C19"));
 	SystemGuids.Add(DevGuids.SkeletalMeshDerivedDataVersion, FGuid("5705956EC7134274A5A1BD3FC74DB3A9"));
+	SystemGuids.Add(DevGuids.MaterialTranslationDDCVersion, FGuid("555395300FC34E35A49DB477FC8ADD85"));
 	return SystemGuids;
 }
