Added a comment to TUnion to suggest TVariant instead.

#rnx

diff --git a/Engine/Source/Runtime/Core/Public/Containers/Union.h b/Engine/Source/Runtime/Core/Public/Containers/Union.h
--- a/Engine/Source/Runtime/Core/Public/Containers/Union.h
+++ b/Engine/Source/Runtime/Core/Public/Containers/Union.h
@@ -56,6 +56,8 @@ public:
  * Represents a type which is the union of several other types; i.e. it can have a value whose type is of any the union's subtypes.
  * This differs from C union types by being type-safe, and supporting non-trivial data types as subtypes.
  * Since a value for the union must be of a single subtype, the union stores potential values of different subtypes in overlapped memory, and keeps track of which one is currently valid.
+ * 
+ * TUnion should be considered deprecated and TVariant should be used instead for new code.
  */
 template<typename TypeA,typename TypeB = FNull,typename TypeC = FNull,typename TypeD = FNull,typename TypeE = FNull,typename TypeF = FNull>
 class TUnion
