Always OPEN around GetTextRevision

diff --git a/Engine/Source/Runtime/Core/Private/Internationalization/TextHistory.cpp b/Engine/Source/Runtime/Core/Private/Internationalization/TextHistory.cpp
--- a/Engine/Source/Runtime/Core/Private/Internationalization/TextHistory.cpp
+++ b/Engine/Source/Runtime/Core/Private/Internationalization/TextHistory.cpp
@@ -687,9 +687,7 @@ void FTextHistory::UpdateDisplayStringIfOutOfDate()
 	{
 		uint16 CurrentGlobalRevision = 0;
 		uint16 CurrentLocalRevision = 0;
-		UE_AUTORTFM_OPEN({
-			FTextLocalizationManager::Get().GetTextRevisions(GetTextId(), CurrentGlobalRevision, CurrentLocalRevision);
-		});
+		FTextLocalizationManager::Get().GetTextRevisions(GetTextId(), CurrentGlobalRevision, CurrentLocalRevision);
 
 		if (GlobalRevision != CurrentGlobalRevision || LocalRevision != CurrentLocalRevision)
 		{
@@ -711,10 +709,7 @@ void FTextHistory::MarkDisplayStringUpToDate()
 {
 	if (CanUpdateDisplayString())
 	{
-		UE_AUTORTFM_OPEN(
-		{
-			FTextLocalizationManager::Get().GetTextRevisions(GetTextId(), GlobalRevision, LocalRevision);
-		});
+		FTextLocalizationManager::Get().GetTextRevisions(GetTextId(), GlobalRevision, LocalRevision);
 	}
 	else
 	{
diff --git a/Engine/Source/Runtime/Core/Private/Internationalization/TextLocalizationManager.cpp b/Engine/Source/Runtime/Core/Private/Internationalization/TextLocalizationManager.cpp
--- a/Engine/Source/Runtime/Core/Private/Internationalization/TextLocalizationManager.cpp
+++ b/Engine/Source/Runtime/Core/Private/Internationalization/TextLocalizationManager.cpp
@@ -1106,6 +1106,7 @@ bool FTextLocalizationManager::GetLocResID(const FTextKey& Namespace, const FTex
 }
 #endif
 
+UE_AUTORTFM_ALWAYS_OPEN
 uint16 FTextLocalizationManager::GetTextRevision() const
 {
 	FReadScopeLock ScopeLock(TextRevisionRW);
@@ -1125,6 +1126,7 @@ uint16 FTextLocalizationManager::GetLocalRevisionForTextId(const FTextId& InText
 	return 0;
 }
 
+UE_AUTORTFM_ALWAYS_OPEN
 void FTextLocalizationManager::GetTextRevisions(const FTextId& InTextId, uint16& OutGlobalTextRevision, uint16& OutLocalTextRevision) const
 {
 	FReadScopeLock ScopeLock(TextRevisionRW);
