Add DeploymentName in the platform CrashContext on some platforms

diff --git a/Engine/Source/Runtime/Core/Private/GenericPlatform/GenericPlatformCrashContext.cpp b/Engine/Source/Runtime/Core/Private/GenericPlatform/GenericPlatformCrashContext.cpp
--- a/Engine/Source/Runtime/Core/Private/GenericPlatform/GenericPlatformCrashContext.cpp
+++ b/Engine/Source/Runtime/Core/Private/GenericPlatform/GenericPlatformCrashContext.cpp
@@ -1025,6 +1025,11 @@ void FGenericCrashContext::SetDeploymentName(const FString& EpicApp)
 	NCached::Set(NCached::Session.DeploymentName, *EpicApp);
 }
 
+const TCHAR* FGenericCrashContext::GetDeploymentName()
+{
+	return NCached::Session.DeploymentName;
+}
+
 void FGenericCrashContext::SetCrashTrigger(ECrashTrigger Type)
 {
 	NCached::Session.CrashTrigger = (int32)Type;
diff --git a/Engine/Source/Runtime/Core/Public/GenericPlatform/GenericPlatformCrashContext.h b/Engine/Source/Runtime/Core/Public/GenericPlatform/GenericPlatformCrashContext.h
--- a/Engine/Source/Runtime/Core/Public/GenericPlatform/GenericPlatformCrashContext.h
+++ b/Engine/Source/Runtime/Core/Public/GenericPlatform/GenericPlatformCrashContext.h
@@ -728,6 +728,11 @@ public:
 	 */
 	CORE_API static void SetDeploymentName(const FString& EpicApp);
 
+	/**
+	 * Get the current deployment name (ie. EpicApp)
+	 */
+	CORE_API static const TCHAR* GetDeploymentName();
+
 	/**
 	 * Sets the type of crash triggered. Used to distinguish crashes caused for debugging purposes.
 	 */
