Always Open around some Text Formatting

diff --git a/Engine/Source/Runtime/Core/Private/Internationalization/TextFormatter.cpp b/Engine/Source/Runtime/Core/Private/Internationalization/TextFormatter.cpp
--- a/Engine/Source/Runtime/Core/Private/Internationalization/TextFormatter.cpp
+++ b/Engine/Source/Runtime/Core/Private/Internationalization/TextFormatter.cpp
@@ -379,6 +379,7 @@ public:
 	/**
 	 * Test to see whether this instance contains valid compiled data.
 	 */
+	UE_AUTORTFM_ALWAYS_OPEN
 	FORCEINLINE bool IsValid() const
 	{
 		UE::TScopeLock Lock(CompiledDataMutex);
@@ -394,6 +395,7 @@ public:
 	 * Validate the format pattern is valid based on the rules of the given culture (or null to use the current language).
 	 * @return true if the pattern is valid, or false if not (false may also fill in OutValidationErrors).
 	 */
+	UE_AUTORTFM_ALWAYS_OPEN
 	FORCEINLINE bool ValidatePattern(const FCulturePtr& InCulture, TArray<FString>& OutValidationErrors)
 	{
 		UE::TScopeLock Lock(CompiledDataMutex);
@@ -403,6 +405,7 @@ public:
 	/**
 	 * Produce a formatted string using the given argument look-up.
 	 */
+	UE_AUTORTFM_ALWAYS_OPEN
 	FORCEINLINE FString Format(const FPrivateTextFormatArguments& InFormatArgs)
 	{
 		UE::TScopeLock Lock(CompiledDataMutex);
@@ -412,6 +415,7 @@ public:
 	/**
 	 * Append the names of any arguments to the given array.
 	 */
+	UE_AUTORTFM_ALWAYS_OPEN
 	FORCEINLINE void GetFormatArgumentNames(TArray<FString>& OutArgumentNames)
 	{
 		UE::TScopeLock Lock(CompiledDataMutex);
@@ -439,6 +443,7 @@ public:
 	/**
 	 * Get the type of expression currently compiled.
 	 */
+	UE_AUTORTFM_ALWAYS_OPEN
 	FORCEINLINE FTextFormat::EExpressionType GetExpressionType() const
 	{
 		UE::TScopeLock Lock(CompiledDataMutex);
