Added FMath::WrapExclusive which gives more expected results for integral types.
Deprecated FMath::Wrap for integers and moved existing code over to using FMath::WrapExclusive instead.
Rewrote tests for FMath::Wrap to be more obvious and implemented some for FMath::WrapExclusive too.

diff --git a/Engine/Plugins/FX/Niagara/Source/Niagara/Private/NiagaraDebugHud.cpp b/Engine/Plugins/FX/Niagara/Source/Niagara/Private/NiagaraDebugHud.cpp
--- a/Engine/Plugins/FX/Niagara/Source/Niagara/Private/NiagaraDebugHud.cpp
+++ b/Engine/Plugins/FX/Niagara/Source/Niagara/Private/NiagaraDebugHud.cpp
@@ -3294,7 +3294,7 @@ void FNiagaraDebugHud::DrawMessages(class FNiagaraWorldManager* WorldManager, cl
 void FNiagaraDebugHudStatHistory::AddFrame_GT(double Time)
 {
 	GTFrames.SetNumZeroed(NiagaraDebugLocal::Settings.PerfHistoryFrames);
-	CurrFrame = FMath::Wrap(CurrFrame + 1, 0, GTFrames.Num() - 1);
+	CurrFrame = FMath::WrapExclusive(CurrFrame + 1, 0, GTFrames.Num());
 	GTFrames[CurrFrame] = Time;
 }
 
@@ -3302,7 +3302,7 @@ void FNiagaraDebugHudStatHistory::AddFrame_RT(double Time)
 {
 	FScopeLock Lock(&NiagaraDebugLocal::RTFramesGuard);
 	RTFrames.SetNumZeroed(NiagaraDebugLocal::Settings.PerfHistoryFrames);
-	CurrFrameRT = FMath::Wrap(CurrFrameRT + 1, 0, RTFrames.Num() - 1);
+	CurrFrameRT = FMath::WrapExclusive(CurrFrameRT + 1, 0, RTFrames.Num());
 	RTFrames[CurrFrameRT] = Time;
 }
 
@@ -3310,7 +3310,7 @@ void FNiagaraDebugHudStatHistory::AddFrame_GPU(double Time)
 {
 	FScopeLock Lock(&NiagaraDebugLocal::RTFramesGuard);
 	GPUFrames.SetNumZeroed(NiagaraDebugLocal::Settings.PerfHistoryFrames);
-	CurrFrameGPU = FMath::Wrap(CurrFrameGPU + 1, 0, GPUFrames.Num() - 1);
+	CurrFrameGPU = FMath::WrapExclusive(CurrFrameGPU + 1, 0, GPUFrames.Num());
 	GPUFrames[CurrFrameGPU] = Time;
 }
 
@@ -3325,7 +3325,7 @@ void FNiagaraDebugHudStatHistory::GetHistoryFrames_GT(TArray<double>& OutHistory
 	do
 	{
 		OutHistoryGT.Add(GTFrames[WriteFrame]);
-		WriteFrame = FMath::Wrap(WriteFrame + 1, 0, GTFrames.Num() - 1);
+		WriteFrame = FMath::WrapExclusive(WriteFrame + 1, 0, GTFrames.Num());
 	} while (WriteFrame != CurrFrame);
 }
 
@@ -3343,7 +3343,7 @@ void FNiagaraDebugHudStatHistory::GetHistoryFrames_RT(TArray<double>& OutHistory
 	do 
 	{
 		OutHistoryRT.Add(RTFrames[WriteFrame]);
-		WriteFrame = FMath::Wrap(WriteFrame + 1, 0, RTFrames.Num() - 1);
+		WriteFrame = FMath::WrapExclusive(WriteFrame + 1, 0, RTFrames.Num());
 	} while (WriteFrame != CurrFrameRT);
 }
 
@@ -3361,7 +3361,7 @@ void FNiagaraDebugHudStatHistory::GetHistoryFrames_GPU(TArray<double>& OutHistor
 	do
 	{
 		OutHistoryGPU.Add(GPUFrames[WriteFrame]);
-		WriteFrame = FMath::Wrap(WriteFrame + 1, 0, GPUFrames.Num() - 1);
+		WriteFrame = FMath::WrapExclusive(WriteFrame + 1, 0, GPUFrames.Num());
 	} while (WriteFrame != CurrFrameGPU);
 }
 
diff --git a/Engine/Plugins/Media/ImgMedia/Source/ImgMedia/Private/Loader/ImgMediaLoader.cpp b/Engine/Plugins/Media/ImgMedia/Source/ImgMedia/Private/Loader/ImgMediaLoader.cpp
--- a/Engine/Plugins/Media/ImgMedia/Source/ImgMedia/Private/Loader/ImgMediaLoader.cpp
+++ b/Engine/Plugins/Media/ImgMedia/Source/ImgMedia/Private/Loader/ImgMediaLoader.cpp
@@ -71,7 +71,7 @@ namespace ImgMediaLoader
 		// We start one frame away, since it's generally to late to enqueue additional update work for the current cached frame.
 		for (int32 Offset = 1; Offset <= OffsetCount; ++Offset)
 		{
-			int32 OffsetFrame = FMath::Wrap(OriginFrame + PlayRateSign * Offset, 0, TotalNumFrames - 1);
+			int32 OffsetFrame = FMath::WrapExclusive(OriginFrame + PlayRateSign * Offset, 0, TotalNumFrames);
 			if (OffsetFrame == CurrentFrame)
 			{
 				return true;
diff --git a/Engine/Plugins/Runtime/Metasound/Source/MetasoundStandardNodes/Private/MetasoundOscillators.h b/Engine/Plugins/Runtime/Metasound/Source/MetasoundStandardNodes/Private/MetasoundOscillators.h
--- a/Engine/Plugins/Runtime/Metasound/Source/MetasoundStandardNodes/Private/MetasoundOscillators.h
+++ b/Engine/Plugins/Runtime/Metasound/Source/MetasoundStandardNodes/Private/MetasoundOscillators.h
@@ -375,9 +375,13 @@ namespace Metasound
 			float operator()(float InPhase, float, const FGeneratorArgs&) const
 			{
 				const TArray<float>& WaveTable = GetWaveTable();
-				int32 LastIndex = WaveTable.Num() - 1;
-				int32 TableIndex = (int32)(InPhase * (float)(LastIndex));
-				TableIndex = FMath::Wrap(TableIndex, 0, LastIndex);
+				int32 NumEntries = WaveTable.Num();
+				float FractionalPhase = FMath::Wrap(InPhase, 0.0f, 1.0f);
+				int32 TableIndex = (int32)(FractionalPhase * NumEntries + 0.5f);
+				if (TableIndex >= NumEntries)
+				{
+					TableIndex = 0;
+				}
 				return WaveTable[TableIndex];
 			}
 		};
diff --git a/Engine/Source/Developer/TraceInsights/Private/Insights/Widgets/SStartPageWindow.cpp b/Engine/Source/Developer/TraceInsights/Private/Insights/Widgets/SStartPageWindow.cpp
--- a/Engine/Source/Developer/TraceInsights/Private/Insights/Widgets/SStartPageWindow.cpp
+++ b/Engine/Source/Developer/TraceInsights/Private/Insights/Widgets/SStartPageWindow.cpp
@@ -2588,7 +2588,7 @@ void STraceStoreWindow::RefreshTraceList()
 					DirColor[ColorIdx],
 					ETraceDirOperations::Delete|ETraceDirOperations::Explore
 				));
-				ColorIdx = FMath::Wrap(++ColorIdx, int32(-1), int32(UE_ARRAY_COUNT(DirColor))-1);
+				ColorIdx = FMath::WrapExclusive(++ColorIdx, (int32)0, int32(UE_ARRAY_COUNT(DirColor)));
 			}
 			if (WatchDirsListView)
 			{
diff --git a/Engine/Source/Runtime/Core/Public/Math/UnrealMathUtility.h b/Engine/Source/Runtime/Core/Public/Math/UnrealMathUtility.h
--- a/Engine/Source/Runtime/Core/Public/Math/UnrealMathUtility.h
+++ b/Engine/Source/Runtime/Core/Public/Math/UnrealMathUtility.h
@@ -603,10 +603,9 @@ public:
 	/** Clamps X to be between Min and Max, inclusive. Overload to support mixed int64/int32 types. */
 	[[nodiscard]] static constexpr FORCEINLINE int64 Clamp(const int64 X, const int32 Min, const int32 Max) { return Clamp<int64>(X, Min, Max); }
 
-	/** Wraps X to be between Min and Max, inclusive. */
-	/** When X can wrap to both Min and Max, it will wrap to Min if it lies below the range and wrap to Max if it is above the range. */
+private:
 	template< class T >
-	[[nodiscard]] static constexpr FORCEINLINE T Wrap(const T X, const T Min, const T Max)
+	[[nodiscard]] static constexpr FORCEINLINE T WrapImpl(const T X, const T Min, const T Max)
 	{
 		// Use unsigned type for integers to allow for large ranges which don't overflow on subtraction
 		// We don't do that for floating point types because there are no unsigned versions of those.  We
@@ -632,16 +631,74 @@ public:
 		}
 		else if (EndVal > Max)
 		{
-			T Mod = FMath::Modulo((SizeType)((SizeType)EndVal - (SizeType)Max), Size);
+			SizeType Mod = FMath::Modulo((SizeType)((SizeType)EndVal - (SizeType)Max), Size);
 			EndVal = (Mod != (T)0) ? (T)((SizeType)Min + Mod) : Max;
 		}
 		return EndVal;
 	}
 
-	template <typename T>
-	[[nodiscard]] static constexpr FORCEINLINE T Modulo(T Value, T Base)
+public:
+	/** Wraps X to be between Min and Max, inclusive. */
+	/** When X can wrap to both Min and Max, it will wrap to Min if it lies below the range and wrap to Max if it is above the range. */
+	template <
+		typename T
+		UE_REQUIRES(std::is_floating_point_v<T>)
+	>
+	[[nodiscard]] static constexpr FORCEINLINE T Wrap(const T X, const T Min, const T Max)
+	{
+		return WrapImpl(X, Min, Max);
+	}
+	template <
+		typename T
+		UE_REQUIRES(!std::is_floating_point_v<T>)
+	>
+	UE_DEPRECATED(5.5, "Use of FMath::Wrap with non-floating point arguments is deprecated - consider FMath::WrapExclusive instead.")
+	[[nodiscard]] static constexpr FORCEINLINE T Wrap(const T X, const T Min, const T Max)
+	{
+		return WrapImpl(X, Min, Max);
+	}
+
+	/** Wraps X to be between Min and Max, exclusive. Will never return Max. */
+	template <
+		typename T
+		UE_REQUIRES(std::is_integral_v<T>)
+	>
+	[[nodiscard]] static constexpr FORCEINLINE T WrapExclusive(const T X, const T Min, const T Max)
+	{
+		// Use unsigned type allow for large ranges which don't overflow on subtraction.
+		using SizeType = std::make_unsigned_t<T>;
+
+		// Our asserts are not constexpr-friendly yet
+		// checkSlow(Min <= Max);
+
+		SizeType Size = (SizeType)((SizeType)Max - (SizeType)Min);
+		if (Size == 0)
+		{
+			// Guard against zero-sized ranges causing division by zero.
+			return Max;
+		}
+
+		SizeType Mod;
+		if (X < Min)
+		{
+			Mod = (SizeType)FMath::Modulo((SizeType)((SizeType)Min - (SizeType)X), Size);
+			if (Mod > 0)
+			{
+				Mod = (SizeType)(Size - Mod);
+			}
+		}
+		else
+		{
+			Mod = (SizeType)FMath::Modulo((SizeType)((SizeType)X - (SizeType)Min), Size);
+		}
+
+		return (T)(Min + Mod);
+	}
+
+	template <typename ValueType, typename BaseType>
+	[[nodiscard]] static constexpr FORCEINLINE auto Modulo(ValueType Value, BaseType Base)
 	{
-		if constexpr (std::is_floating_point_v<T>)
+		if constexpr (std::is_floating_point_v<ValueType>)
 		{
 			return FMath::Fmod(Value, Base);
 		}
diff --git a/Engine/Source/Runtime/Core/Tests/Math/UnrealMathTest.cpp b/Engine/Source/Runtime/Core/Tests/Math/UnrealMathTest.cpp
--- a/Engine/Source/Runtime/Core/Tests/Math/UnrealMathTest.cpp
+++ b/Engine/Source/Runtime/Core/Tests/Math/UnrealMathTest.cpp
@@ -4557,181 +4557,120 @@ TEST_CASE_NAMED(FBitCastTest, "System::Core::Math::Bitcast", "[ApplicationContex
 
 TEST_CASE_NAMED(FMathWrapTest, "System::Core::Math::Wrap", "[ApplicationContextMask][SmokeFilter]")
 {
-	// Tests wrapping of FMath::Wrap(Val, Min, Max) with a set of values in a set of ranges.
-	//
-	// The values to test are of the form ValFrom + X*ValStep, and do not exceed ValTo.
-	// The minimum values in the wrapping range are of the form: MinFrom + Y*MinStep, and do not exceed MinTo.
-	// The sizes of each range (where Max == Min + Size) are of the form: SizeFrom + Z*SizeStep, and do not exceed SizeTo.
-	auto WrapTest = []<typename T>(T ValFrom, T ValTo, T ValStep, T MinFrom, T MinTo, T MinStep, T SizeFrom, T SizeTo, T SizeStep)
-	{
-		for (T Val = ValFrom; Val < ValTo; Val += ValStep)
-		{
-			for (T Min = MinFrom; Min < MinTo; Min += MinStep)
-			{
-				for (T Size = SizeFrom; Size < SizeTo; Size += SizeStep)
-				{
-					if (Size == (T)0)
-					{
-						T Wrap = FMath::Wrap(Val, Min, Min);
-
-						CHECK_MESSAGE(TEXT("Wrapped value should be in the empty range"), Wrap == Min);
-					}
-					else
-					{
-						T Max = Min + Size;
-						T Wrap = FMath::Wrap(Val, Min, Max);
-
-						CHECK_MESSAGE(TEXT("Wrapped value should be in the non-empty range"), Wrap >= Min && Wrap <= Max);
-						T Mod = FMath::Modulo((Wrap - Val), Size);
-						if constexpr (std::is_integral_v<T>)
-						{
-							CHECK_MESSAGE(FString::Printf(TEXT("Wrapped value should be at a distance which is an exact multiple of the range size: (Val: %d, Min: %d, Max: %d, Wrap: %d, Mod: %d)"), Val, Min, Max, Wrap, Mod), Mod == 0);
-						}
-						else
-						{
-							T Tolerance;
-							if constexpr (std::is_same_v<T, float>)
-							{
-								Tolerance = UE_KINDA_SMALL_NUMBER;
-							}
-							else
-							{
-								Tolerance = UE_DOUBLE_KINDA_SMALL_NUMBER;
-							}
-
-							// We need to check that we're in the range of zero *or* +/- size because of rounding
-							bool bIsExactMultipleOfSize = FMath::Square(Mod) < Tolerance || FMath::Square(Mod - Size) < Tolerance || FMath::Square(Mod + Size) < Tolerance;
-							CHECK_MESSAGE(FString::Printf(TEXT("Wrapped value should be at a distance which is an exact multiple of the range size: (Val: %f, Min: %f, Max: %f, Wrap: %f, Mod: %f)"), Val, Min, Max, Wrap, Mod), bIsExactMultipleOfSize);
-						}
-
-						if (Val < Min)
-						{
-							CHECK_FALSE_MESSAGE(TEXT("Wrapping a value from below a non-empty range should never give the max"), Wrap == Max);
-						}
-						else if (Val > Max)
-						{
-							CHECK_FALSE_MESSAGE(TEXT("Wrapping a value from above a non-empty range should never give the min"), Wrap == Min);
-						}
-					}
-				}
-			}
-		}
-	};
-
-	// Integral
-	WrapTest(
-		/* ValFrom  */ -25,
-		/* ValTo    */  25,
-		/* ValStep  */   1,
-		/* MinFrom  */  -5,
-		/* MinTo    */   5,
-		/* MinStep  */   1,
-		/* SizeFrom */   0,
-		/* SizeTo   */   5,
-		/* SizeStep */   1
-	);
-
 	// Floats (with integral values)
-	WrapTest(
-		/* ValFrom  */ -25.0f,
-		/* ValTo    */  25.0f,
-		/* ValStep  */   1.0f,
-		/* MinFrom  */  -5.0f,
-		/* MinTo    */   5.0f,
-		/* MinStep  */   1.0f,
-		/* SizeFrom */   0.0f,
-		/* SizeTo   */   5.0f,
-		/* SizeStep */   1.0f
-	);
-
-	// Floats (with fractional values)
-	WrapTest(
-		/* ValFrom  */  -7.34f,
-		/* ValTo    */  12.19f,
-		/* ValStep  */   0.7361f,
-		/* MinFrom  */  -8.43f,
-		/* MinTo    */  11.84f,
-		/* MinStep  */   0.69f,
-		/* SizeFrom */   0.0f,
-		/* SizeTo   */   7.23f,
-		/* SizeStep */   0.59f
-	);
+	CHECK(FMath::Wrap(-17.0f, 10.0f, 20.0f) == 13.0f);
+	CHECK(FMath::Wrap(-1.0f,  10.0f, 20.0f) == 19.0f);
+	CHECK(FMath::Wrap(0.0f,   10.0f, 20.0f) == 10.0f);
+	CHECK(FMath::Wrap(5.0f,   10.0f, 20.0f) == 15.0f);
+	CHECK(FMath::Wrap(9.0f,   10.0f, 20.0f) == 19.0f);
+	CHECK(FMath::Wrap(10.0f,  10.0f, 20.0f) == 10.0f);
+	CHECK(FMath::Wrap(15.0f,  10.0f, 20.0f) == 15.0f);
+	CHECK(FMath::Wrap(19.0f,  10.0f, 20.0f) == 19.0f);
+	CHECK(FMath::Wrap(20.0f,  10.0f, 20.0f) == 20.0f);
+	CHECK(FMath::Wrap(21.0f,  10.0f, 20.0f) == 11.0f);
+	CHECK(FMath::Wrap(25.0f,  10.0f, 20.0f) == 15.0f);
+	CHECK(FMath::Wrap(29.0f,  10.0f, 20.0f) == 19.0f);
+	CHECK(FMath::Wrap(30.0f,  10.0f, 20.0f) == 20.0f);
+	CHECK(FMath::Wrap(54.0f,  10.0f, 20.0f) == 14.0f);
+	CHECK(FMath::Wrap(60.0f,  10.0f, 20.0f) == 20.0f);
 
 	// Doubles (with integral values)
-	WrapTest(
-		/* ValFrom  */ -25.0,
-		/* ValTo    */  25.0,
-		/* ValStep  */   1.0,
-		/* MinFrom  */  -5.0,
-		/* MinTo    */   5.0,
-		/* MinStep  */   1.0,
-		/* SizeFrom */   0.0,
-		/* SizeTo   */   5.0,
-		/* SizeStep */   1.0
-	);
-
-	// Doubles (with fractional values)
-	WrapTest(
-		/* ValFrom  */  -7.34,
-		/* ValTo    */  12.19,
-		/* ValStep  */   0.7361,
-		/* MinFrom  */  -8.43,
-		/* MinTo    */  11.84,
-		/* MinStep  */   0.69,
-		/* SizeFrom */   0.0,
-		/* SizeTo   */   7.23,
-		/* SizeStep */   0.59
-	);
-
-	// Test large values far away from the range
-	WrapTest(
-		/* ValFrom  */  1 << 30,
-		/* ValTo    */  (1 << 30) + 1,
-		/* ValStep  */  1,
-		/* MinFrom  */  0,
-		/* MinTo    */  1,
-		/* MinStep  */  1,
-		/* SizeFrom */  2,
-		/* SizeTo   */  3,
-		/* SizeStep */  1
-	);
-	WrapTest(
-		/* ValFrom  */  (float)(1 << 25),
-		/* ValTo    */  (float)(1 << 25) + 4.0f,
-		/* ValStep  */  4.0f,
-		/* MinFrom  */  0.0f,
-		/* MinTo    */  1.0f,
-		/* MinStep  */  1.0f,
-		/* SizeFrom */  1.0f,
-		/* SizeTo   */  2.0f,
-		/* SizeStep */  1.0f
-	);
-	WrapTest(
-		/* ValFrom  */  (double)(1ull << 54),
-		/* ValTo    */  (double)(1ull << 54) + 4.0,
-		/* ValStep  */  4.0,
-		/* MinFrom  */  0.0,
-		/* MinTo    */  1.0,
-		/* MinStep  */  1.0,
-		/* SizeFrom */  1.0,
-		/* SizeTo   */  2.0,
-		/* SizeStep */  1.0
-	);
+	CHECK(FMath::Wrap(-17.0, 10.0, 20.0) == 13.0);
+	CHECK(FMath::Wrap( -1.0, 10.0, 20.0) == 19.0);
+	CHECK(FMath::Wrap(  0.0, 10.0, 20.0) == 10.0);
+	CHECK(FMath::Wrap(  5.0, 10.0, 20.0) == 15.0);
+	CHECK(FMath::Wrap(  9.0, 10.0, 20.0) == 19.0);
+	CHECK(FMath::Wrap( 10.0, 10.0, 20.0) == 10.0);
+	CHECK(FMath::Wrap( 15.0, 10.0, 20.0) == 15.0);
+	CHECK(FMath::Wrap( 19.0, 10.0, 20.0) == 19.0);
+	CHECK(FMath::Wrap( 20.0, 10.0, 20.0) == 20.0);
+	CHECK(FMath::Wrap( 21.0, 10.0, 20.0) == 11.0);
+	CHECK(FMath::Wrap( 25.0, 10.0, 20.0) == 15.0);
+	CHECK(FMath::Wrap( 29.0, 10.0, 20.0) == 19.0);
+	CHECK(FMath::Wrap( 30.0, 10.0, 20.0) == 20.0);
+	CHECK(FMath::Wrap( 54.0, 10.0, 20.0) == 14.0);
+	CHECK(FMath::Wrap( 60.0, 10.0, 20.0) == 20.0);
 
-	// Test constexpr
-	static_assert(FMath::Wrap(-3, 0, 5) == 2);
-	//static_assert(FMath::Wrap(-3.0f, 0.0f, 5.0f) == 2); // needs constexpr fmod support in C++23
-	//static_assert(FMath::Wrap(-3.0, 0.0, 5.0) == 2); // needs constexpr fmod support in C++23
+	// Floats (with fractional values)
+	CHECK(FMath::IsNearlyEqual(FMath::Wrap(-17.62f, 2.37f, 8.44f), 6.66f, UE_KINDA_SMALL_NUMBER));
+	CHECK(FMath::IsNearlyEqual(FMath::Wrap( -1.89f, 2.37f, 8.44f), 4.18f, UE_KINDA_SMALL_NUMBER));
+	CHECK(FMath::IsNearlyEqual(FMath::Wrap(  7.12f, 2.37f, 8.44f), 7.12f, UE_KINDA_SMALL_NUMBER));
+	CHECK(FMath::IsNearlyEqual(FMath::Wrap( 11.31f, 2.37f, 8.44f), 5.24f, UE_KINDA_SMALL_NUMBER));
+	CHECK(FMath::IsNearlyEqual(FMath::Wrap( 23.65f, 2.37f, 8.44f), 5.44f, UE_KINDA_SMALL_NUMBER));
 
-	// These will fail to compile due to signed overload if we don't use unsigned diffs in the implementation
-	static_assert(FMath::Wrap(MAX_int32, MIN_int32 + 2, MAX_int32)); // range size will overflow
-	static_assert(FMath::Wrap(MIN_int32, MAX_int32 - 2, MAX_int32)); // distance from val to min will overflow
-	static_assert(FMath::Wrap(MAX_int32, MIN_int32, MIN_int32 + 2)); // distance from max to val will overflow
+	// Doubles (with fractional values)
+	CHECK(FMath::IsNearlyEqual(FMath::Wrap(-17.62, 2.37, 8.44), 6.66, UE_DOUBLE_KINDA_SMALL_NUMBER));
+	CHECK(FMath::IsNearlyEqual(FMath::Wrap( -1.89, 2.37, 8.44), 4.18, UE_DOUBLE_KINDA_SMALL_NUMBER));
+	CHECK(FMath::IsNearlyEqual(FMath::Wrap(  7.12, 2.37, 8.44), 7.12, UE_DOUBLE_KINDA_SMALL_NUMBER));
+	CHECK(FMath::IsNearlyEqual(FMath::Wrap( 11.31, 2.37, 8.44), 5.24, UE_DOUBLE_KINDA_SMALL_NUMBER));
+	CHECK(FMath::IsNearlyEqual(FMath::Wrap( 23.65, 2.37, 8.44), 5.44, UE_DOUBLE_KINDA_SMALL_NUMBER));
+
+	// Floats (large values far away from the range)
+	CHECK(FMath::Wrap(-(float)(1 << 25), 1.0f, 2.0f) == 1.0f);
+	CHECK(FMath::Wrap( (float)(1 << 25), 1.0f, 2.0f) == 2.0f);
+
+	// Doubles (large values far away from the range)
+	CHECK(FMath::Wrap(-(double)(1 << 25), 1.0, 2.0) == 1.0);
+	CHECK(FMath::Wrap( (double)(1 << 25), 1.0, 2.0) == 2.0);
+}
 
-	// Try with bytes too, where the value is more than 128 away from the range
-	static_assert(FMath::Wrap((int8)123, (int8)-20, (int8)-10) == (int8)-17);
-	static_assert(FMath::Wrap((int8)-123, (int8)10, (int8)20) == (int8)17);
+TEST_CASE_NAMED(FMathWrapExclusiveTest, "System::Core::Math::WrapExclusive", "[ApplicationContextMask][SmokeFilter]")
+{
+	// Ints (compile time)
+	static_assert(FMath::WrapExclusive(-17, 10, 20) == 13);
+	static_assert(FMath::WrapExclusive(-1,  10, 20) == 19);
+	static_assert(FMath::WrapExclusive(0,   10, 20) == 10);
+	static_assert(FMath::WrapExclusive(5,   10, 20) == 15);
+	static_assert(FMath::WrapExclusive(9,   10, 20) == 19);
+	static_assert(FMath::WrapExclusive(10,  10, 20) == 10);
+	static_assert(FMath::WrapExclusive(15,  10, 20) == 15);
+	static_assert(FMath::WrapExclusive(19,  10, 20) == 19);
+	static_assert(FMath::WrapExclusive(20,  10, 20) == 10);
+	static_assert(FMath::WrapExclusive(21,  10, 20) == 11);
+	static_assert(FMath::WrapExclusive(25,  10, 20) == 15);
+	static_assert(FMath::WrapExclusive(29,  10, 20) == 19);
+	static_assert(FMath::WrapExclusive(30,  10, 20) == 10);
+	static_assert(FMath::WrapExclusive(54,  10, 20) == 14);
+	static_assert(FMath::WrapExclusive(60,  10, 20) == 10);
+
+	// Ints (runtime)
+	CHECK(FMath::WrapExclusive(-17, 10, 20) == 13);
+	CHECK(FMath::WrapExclusive(-1,  10, 20) == 19);
+	CHECK(FMath::WrapExclusive(0,   10, 20) == 10);
+	CHECK(FMath::WrapExclusive(5,   10, 20) == 15);
+	CHECK(FMath::WrapExclusive(9,   10, 20) == 19);
+	CHECK(FMath::WrapExclusive(10,  10, 20) == 10);
+	CHECK(FMath::WrapExclusive(15,  10, 20) == 15);
+	CHECK(FMath::WrapExclusive(19,  10, 20) == 19);
+	CHECK(FMath::WrapExclusive(20,  10, 20) == 10);
+	CHECK(FMath::WrapExclusive(21,  10, 20) == 11);
+	CHECK(FMath::WrapExclusive(25,  10, 20) == 15);
+	CHECK(FMath::WrapExclusive(29,  10, 20) == 19);
+	CHECK(FMath::WrapExclusive(30,  10, 20) == 10);
+	CHECK(FMath::WrapExclusive(54,  10, 20) == 14);
+	CHECK(FMath::WrapExclusive(60,  10, 20) == 10);
+
+	// Ints (compile time - distance to range is greater than std::numeric_limits<T>::max())
+	static_assert(FMath::WrapExclusive(MIN_int32, MAX_int32 - 8, MAX_int32) == MAX_int32 - 7);
+	static_assert(FMath::WrapExclusive(MAX_int32, MIN_int32, MIN_int32 + 8) == MIN_int32 + 7);
+	static_assert(FMath::WrapExclusive((int8)123, (int8)-20, (int8)-10) == (int8)-17);
+	static_assert(FMath::WrapExclusive((int8)-123, (int8)10, (int8)20) == (int8)17);
+
+	// Ints (runtime - distance to range is greater than std::numeric_limits<T>::max())
+	CHECK(FMath::WrapExclusive(MIN_int32, MAX_int32 - 8, MAX_int32) == MAX_int32 - 7);
+	CHECK(FMath::WrapExclusive(MAX_int32, MIN_int32, MIN_int32 + 8) == MIN_int32 + 7);
+	CHECK(FMath::WrapExclusive((int8)123, (int8)-20, (int8)-10) == (int8)-17);
+	CHECK(FMath::WrapExclusive((int8)-123, (int8)10, (int8)20) == (int8)17);
+
+	// Ints (compile time - range-size is greater than std::numeric_limits<T>::max())
+	static_assert(FMath::WrapExclusive(MAX_int32, MIN_int32 + 2, MAX_int32));
+	static_assert(FMath::WrapExclusive((int8)-123, (int8)-120, (int8)120) == (int8)117);
+
+	// Ints (runtime - range-size is greater than std::numeric_limits<T>::max())
+	CHECK(FMath::WrapExclusive(MAX_int32, MIN_int32 + 2, MAX_int32));
+	CHECK(FMath::WrapExclusive((int8)-123, (int8)-120, (int8)120) == (int8)117);
 }
+
 class FInitVectorTestClass {
 public:
 
diff --git a/Engine/Source/Runtime/Engine/Classes/Kismet/KismetMathLibrary.inl b/Engine/Source/Runtime/Engine/Classes/Kismet/KismetMathLibrary.inl
--- a/Engine/Source/Runtime/Engine/Classes/Kismet/KismetMathLibrary.inl
+++ b/Engine/Source/Runtime/Engine/Classes/Kismet/KismetMathLibrary.inl
@@ -273,7 +273,9 @@ int32 UKismetMathLibrary::Clamp(int32 V, int32 A, int32 B)
 KISMET_MATH_FORCEINLINE
 int32 UKismetMathLibrary::Wrap(int32 Value, int32 Min, int32 Max)
 {
+PRAGMA_DISABLE_DEPRECATION_WARNINGS
 	return FMath::Wrap(Value, Min, Max);
+PRAGMA_ENABLE_DEPRECATION_WARNINGS
 }
 
 KISMET_MATH_FORCEINLINE
diff --git a/Engine/Source/Runtime/SignalProcessing/Private/InterpolatedMultiTapDelay.cpp b/Engine/Source/Runtime/SignalProcessing/Private/InterpolatedMultiTapDelay.cpp
--- a/Engine/Source/Runtime/SignalProcessing/Private/InterpolatedMultiTapDelay.cpp
+++ b/Engine/Source/Runtime/SignalProcessing/Private/InterpolatedMultiTapDelay.cpp
@@ -48,7 +48,7 @@ namespace Audio
 		{
 			FMemory::Memcpy(&DelayLine[WriteIndex], InBuffer.GetData(), InNumSamples * sizeof(float));
 			WriteIndex += InNumSamples;
-			WriteIndex = FMath::Wrap<uint32>(WriteIndex, 0, DelayLine.Num());
+			WriteIndex = FMath::WrapExclusive((uint32)WriteIndex, (uint32)0, (uint32)DelayLine.Num());
 		}
 	}
 
@@ -77,7 +77,7 @@ namespace Audio
 		const float SampleStride = FMath::Clamp((float)(EndSample - StartSample) / (float)OutputNumSamples, 0.25f, 4.f);
 		const uint32 FixedSampleRate = (uint32)(SampleStride * 65536.f);
 
-		StartSample = FMath::Wrap(StartSample, 0, DelayBufferNumSamples - 1);
+		StartSample = FMath::WrapExclusive(StartSample, 0, DelayBufferNumSamples);
 		
 		Resampler.CurrentFrameFraction = StartSampleFraction;
 		const int32 FramesNeeded = (int32)Resampler.SourceFramesNeeded(OutputNumSamples, FixedSampleRate);
