Replaced persistent auxilary memory with a new Persistent Linear Allocator

diff --git a/Engine/Config/BaseEngine.ini b/Engine/Config/BaseEngine.ini
--- a/Engine/Config/BaseEngine.ini
+++ b/Engine/Config/BaseEngine.ini
@@ -1527,7 +1527,6 @@ s.FlushStreamingOnExit=True
 
 [/Script/Engine.GarbageCollectionSettings]
 gc.MaxObjectsNotConsideredByGC=1
-gc.SizeOfPermanentObjectPool=0
 gc.FlushStreamingOnGC=0
 gc.NumRetriesBeforeForcingGC=10
 gc.AllowParallelGC=True
@@ -2200,6 +2199,7 @@ bRecompileOnLoad=true
 
 [MemoryPools]
 FLightPrimitiveInteractionInitialBlockSize=512
+PersistentAllocatorReserveSizeMB=512
 
 [ProcessLimits]
 VirtualMemoryLimitInKB=0
diff --git a/Engine/Source/Runtime/Core/Core.Build.cs b/Engine/Source/Runtime/Core/Core.Build.cs
--- a/Engine/Source/Runtime/Core/Core.Build.cs
+++ b/Engine/Source/Runtime/Core/Core.Build.cs
@@ -3,6 +3,8 @@
 using UnrealBuildTool;
 using System;
 using System.IO;
+using EpicGames.Core;
+using Microsoft.Extensions.Logging;
 
 public class Core : ModuleRules
 {
@@ -346,6 +348,8 @@ public class Core : ModuleRules
 		IWYUSupport = IWYUSupport.KeepAsIs;
 
 		bAllowAutoRTFMInstrumentation = true;
+
+		PrivateDefinitions.Add("UE_PERSISTENT_ALLOCATOR_RESERVE_SIZE=" + GetPersistentAllocatorReserveSize().ToString() + "ULL");
 	}
 
 	protected virtual bool SupportsBinaryConfig(ReadOnlyTargetRules Target)
@@ -364,4 +368,17 @@ public class Core : ModuleRules
 			 Target.Platform == UnrealTargetPlatform.LinuxArm64 ||
 			 Target.Platform == UnrealTargetPlatform.Win64);
 	}
+
+	private ulong GetPersistentAllocatorReserveSize()
+	{
+		ConfigHierarchy Ini = ConfigCache.ReadHierarchy(ConfigHierarchyType.Engine, DirectoryReference.FromFile(Target.ProjectFile), Target.Platform);
+		int SizeOfPermanentObjectPool = 0;
+		if (Ini.GetInt32("/Script/Engine.GarbageCollectionSettings", "gc.SizeOfPermanentObjectPool", out SizeOfPermanentObjectPool))
+		{
+			Target.Logger.LogWarning("/Script/Engine.GarbageCollectionSettings, gc.SizeOfPermanentObjectPool ini setting was deprecated in a favor of MemoryPools, PersistentAllocatorReserveSizeMB");
+		}
+		int PersistentAllocatorReserveSizeMB = 0;
+		Ini.GetInt32("MemoryPools", "PersistentAllocatorReserveSizeMB", out PersistentAllocatorReserveSizeMB);
+		return (ulong)PersistentAllocatorReserveSizeMB * 1024 * 1024;
+	}
 }
diff --git a/Engine/Source/Runtime/Core/Private/HAL/UnrealMemory.cpp b/Engine/Source/Runtime/Core/Private/HAL/UnrealMemory.cpp
--- a/Engine/Source/Runtime/Core/Private/HAL/UnrealMemory.cpp
+++ b/Engine/Source/Runtime/Core/Private/HAL/UnrealMemory.cpp
@@ -23,6 +23,7 @@
 #include "HAL/MallocDoubleFreeFinder.h"
 #include "HAL/MallocFrameProfiler.h"
 #include "HAL/MallocStomp2.h"
+#include "Memory/LinearVirtualMemoryAllocator.h"
 
 #if MALLOC_GT_HOOKS
 
@@ -671,69 +672,24 @@ void FUseSystemMallocForNew::operator delete[](void* Ptr)
 	FMemory::SystemFree(Ptr);
 }
 
-static bool GPersistentAuxiliaryEnabled = true;
-static uint8 * GPersistentAuxiliary = nullptr;
-static uint8 * GPersistentAuxiliaryEnd = nullptr;
-static TAtomic<SIZE_T> GPersistentAuxiliaryCurrentOffset;
-static SIZE_T GPersistentAuxiliarySize = 0;
-
-
-void FMemory::RegisterPersistentAuxiliary(void* InMemory, SIZE_T InSize)
-{
-	check(GPersistentAuxiliary == nullptr);
-	GPersistentAuxiliaryCurrentOffset = 0;
-	GPersistentAuxiliarySize = InSize;
-	GPersistentAuxiliary = (uint8 *)InMemory;
-	GPersistentAuxiliaryEnd = GPersistentAuxiliary + InSize;
-}
 void* FMemory::MallocPersistentAuxiliary(SIZE_T InSize, uint32 InAlignment)
 {
-	if (GPersistentAuxiliary != nullptr && GPersistentAuxiliaryEnabled)
-	{
-		const uint32 Alignment = FMath::Max<uint32>(InAlignment, 16u);
-		const SIZE_T AlignedSize = Align(InSize, Alignment);
-		// 1st check if there is room, this is atomic but could still fail when actually incrementing the offset.
-		if (GPersistentAuxiliaryCurrentOffset + AlignedSize <= GPersistentAuxiliarySize)
-		{
-			SIZE_T OldOffset = GPersistentAuxiliaryCurrentOffset.AddExchange(AlignedSize);
-			if (OldOffset + AlignedSize <= GPersistentAuxiliarySize)
-			{
-				// we were able to increment the offset and it's still within the bounds of the aux memory.
-				return &GPersistentAuxiliary[OldOffset];
-			}
-			// we've gone over the end of the aux memory, this could waste some space, if it's a problem protect with a critical section.
-		}
-	}
-	return FMemory::Malloc(InSize, InAlignment);
+	return GetPersistentLinearAllocator().Allocate(InSize, InAlignment);
 }
+
 void FMemory::FreePersistentAuxiliary(void* InPtr)
 {
-	if (GPersistentAuxiliary != nullptr)
-	{
-		uint8* Ptr = (uint8*)InPtr;
-		if (Ptr >= GPersistentAuxiliary && Ptr < GPersistentAuxiliaryEnd)
-		{
-			// it is part of the GPersistentAuxiliary
-			return;
-		}
-	}
-	return FMemory::Free(InPtr);
+	GetPersistentLinearAllocator().TryDeallocate(InPtr, 0);
 }
+
 bool FMemory::IsPersistentAuxiliaryActive()
 {
-	return GPersistentAuxiliary != nullptr && GPersistentAuxiliaryEnabled;
-}
-void FMemory::DisablePersistentAuxiliary()
-{
-	GPersistentAuxiliaryEnabled = false;
-}
-void FMemory::EnablePersistentAuxiliary()
-{
-	GPersistentAuxiliaryEnabled = true;
+	return GetPersistentLinearAllocator().IsInitialized();
 }
+
 SIZE_T FMemory::GetUsedPersistentAuxiliary()
 {
-	return GPersistentAuxiliaryCurrentOffset;
+	return GetPersistentLinearAllocator().GetAllocatedMemorySize();
 }
 
 
diff --git a/Engine/Source/Runtime/Core/Private/Memory/LinearVirtualMemoryAllocator.cpp b/Engine/Source/Runtime/Core/Private/Memory/LinearVirtualMemoryAllocator.cpp
new file mode 100644
--- /dev/null
+++ b/Engine/Source/Runtime/Core/Private/Memory/LinearVirtualMemoryAllocator.cpp
@@ -0,0 +1,132 @@
+// Copyright Epic Games, Inc. All Rights Reserved.
+
+#include "Memory/LinearVirtualMemoryAllocator.h"
+
+
+CORE_API FLinearVirtualMemoryAllocatorExtends GPersistentLinearAllocatorExtends;
+
+static struct FInitializer
+{
+	FInitializer()
+	{
+		GPersistentLinearAllocatorExtends.Address = (uint64)GetPersistentLinearAllocator().GetBasePointer();
+		GPersistentLinearAllocatorExtends.Size = (uint64)GetPersistentLinearAllocator().GetReservedMemorySize();
+	}
+} Initializer;
+
+
+#if PLATFORM_HAS_FPlatformVirtualMemoryBlock
+
+#include "CoreGlobals.h"
+#include "Misc/ScopeLock.h"
+#include "HAL/LowLevelMemTracker.h"
+
+FLinearVirtualMemoryAllocator::FLinearVirtualMemoryAllocator(SIZE_T ReserveMemorySize)
+	: Reserved(ReserveMemorySize)
+{
+	if (FPlatformMemory::CanOverallocateVirtualMemory())
+	{
+		VirtualMemory = VirtualMemory.AllocateVirtual(ReserveMemorySize);
+		if (!VirtualMemory.GetVirtualPointer())
+		{
+			UE_LOG(LogMemory, Warning, TEXT("LinearVirtualMemoryAllocator failed to reserve %u MB and will default to FMemory::Malloc instead"), ReserveMemorySize / 1024 / 1024);
+			Reserved = 0;
+		}
+	}
+	else
+	{
+#if PLATFORM_IOS || PLATFORM_TVOS
+		UE_LOG(LogMemory, Warning, TEXT("LinearVirtualMemoryAllocator requires com.apple.developer.kernel.extended-virtual-addressing entitlement to work"));
+#else
+		UE_LOG(LogMemory, Warning, TEXT("This platform does not allow to allocate more virtual memory than there is physical memory. LinearVirtualMemoryAllocator will default to FMemory::Malloc instead"));
+#endif
+		Reserved = 0;
+	}
+
+	GPersistentLinearAllocatorExtends.Address = (uint64)VirtualMemory.GetVirtualPointer();
+	GPersistentLinearAllocatorExtends.Size = (uint64)Reserved;
+}
+
+void* FLinearVirtualMemoryAllocator::Allocate(SIZE_T Size, uint32 Alignment)
+{
+	Alignment = FMath::Max(Alignment, 8u);
+	{
+		void* Mem = nullptr;
+		{
+			FScopeLock AutoLock(&Lock);
+			if (CanFit(Size, Alignment))
+			{
+				CurrentOffset = Align(CurrentOffset, Alignment);
+				const SIZE_T NewOffset = CurrentOffset + Size;
+				if (NewOffset > Committed)
+				{
+					LLM_PLATFORM_SCOPE(ELLMTag::FMalloc);
+					const size_t CommitGranularity = UE_USE_VERYLARGEPAGEALLOCATOR ? 2ul * 1024 * 1024 : 65536ul;
+					const SIZE_T ToCommit = Align(NewOffset - Committed, FMath::Max(VirtualMemory.GetCommitAlignment(), CommitGranularity));
+					VirtualMemory.Commit(Committed, ToCommit);
+					Committed += ToCommit;
+					LLM_IF_ENABLED(FLowLevelMemTracker::Get().OnLowLevelAlloc(ELLMTracker::Platform, (uint8*)VirtualMemory.GetVirtualPointer() + CurrentOffset, ToCommit));
+				}
+				Mem = (uint8*)VirtualMemory.GetVirtualPointer() + CurrentOffset;
+				CurrentOffset += Size;
+			}
+		}
+		if (Mem)
+		{
+			LLM(FLowLevelMemTracker::Get().OnLowLevelAlloc(ELLMTracker::Default, Mem, Size, ELLMTag::Untagged, ELLMAllocType::FMalloc));
+			return Mem;
+		}
+	}
+
+	ExceedsReservation.fetch_add(Size, std::memory_order_relaxed);
+	return FMemory::Malloc(Size, Alignment);
+}
+
+void FLinearVirtualMemoryAllocator::PreAllocate(SIZE_T Size, uint32 Alignment)
+{
+	Alignment = FMath::Max(Alignment, 8u);
+
+	FScopeLock AutoLock(&Lock);
+	if (CanFit(Size, Alignment))
+	{
+		CurrentOffset = Align(CurrentOffset, Alignment);
+		const SIZE_T NewOffset = CurrentOffset + Size;
+		if (NewOffset > Committed)
+		{
+			const SIZE_T ToCommit = Align(NewOffset - Committed, VirtualMemory.GetCommitAlignment());
+			VirtualMemory.Commit(Committed, ToCommit);
+			Committed += ToCommit;
+		}
+	}
+}
+
+bool FLinearVirtualMemoryAllocator::TryDeallocate(void* Ptr, SIZE_T Size)
+{
+	if (ContainsPointer(Ptr))
+	{
+		FScopeLock AutoLock(&Lock);
+		if ((uint8*)Ptr + Size == (uint8*)VirtualMemory.GetVirtualPointer() + CurrentOffset)
+		{
+			LLM(FLowLevelMemTracker::Get().OnLowLevelFree(ELLMTracker::Default, Ptr, ELLMAllocType::FMalloc));
+			CurrentOffset -= Size;
+			return true;
+		}
+
+		return false;
+	}
+
+	FMemory::Free(Ptr);
+	return true;
+}
+
+bool FLinearVirtualMemoryAllocator::CanFit(SIZE_T Size, uint32 Alignment) const
+{
+	return (Reserved - Align(CurrentOffset, Alignment)) >= Size;
+}
+
+FLinearVirtualMemoryAllocator& GetPersistentLinearAllocator()
+{
+	static FLinearVirtualMemoryAllocator GPersistentLinearAllocator(UE_PERSISTENT_ALLOCATOR_RESERVE_SIZE);
+	return GPersistentLinearAllocator;
+}
+#endif //~PLATFORM_HAS_FPlatformVirtualMemoryBlock
\ No newline at end of file
diff --git a/Engine/Source/Runtime/Core/Private/UObject/UnrealNames.cpp b/Engine/Source/Runtime/Core/Private/UObject/UnrealNames.cpp
--- a/Engine/Source/Runtime/Core/Private/UObject/UnrealNames.cpp
+++ b/Engine/Source/Runtime/Core/Private/UObject/UnrealNames.cpp
@@ -10,6 +10,7 @@
 #include "Math/NumericLimits.h"
 #include "Math/UnrealMathUtility.h"
 #include "HAL/UnrealMemory.h"
+#include "Memory/LinearVirtualMemoryAllocator.h"
 #include "Templates/UnrealTemplate.h"
 #include "Misc/CString.h"
 #include "Misc/Crc.h"
@@ -453,9 +454,9 @@ public:
 
 	~FNameEntryAllocator()
 	{
-		for (uint32 Index = 0; Index <= CurrentBlock; ++Index)
+		for (int32 Index = CurrentBlock; Index >= 0; --Index)
 		{
-			FMemory::Free(Blocks[Index]);
+			GetPersistentLinearAllocator().TryDeallocate(Blocks[Index], BlockSizeBytes);
 		}
 	}
 
@@ -656,7 +657,7 @@ private:
 	static uint8* AllocBlock()
 	{
 		LLM_SCOPE(ELLMTag::FName);
-		return (uint8*)FMemory::MallocPersistentAuxiliary(BlockSizeBytes, alignof(FNameEntry));
+		return (uint8*)GetPersistentLinearAllocator().Allocate(BlockSizeBytes, alignof(FNameEntry));
 	}
 	
 	void AllocateNewBlock()
diff --git a/Engine/Source/Runtime/Core/Public/Containers/LockFreeFixedSizeAllocator.h b/Engine/Source/Runtime/Core/Public/Containers/LockFreeFixedSizeAllocator.h
--- a/Engine/Source/Runtime/Core/Public/Containers/LockFreeFixedSizeAllocator.h
+++ b/Engine/Source/Runtime/Core/Public/Containers/LockFreeFixedSizeAllocator.h
@@ -17,7 +17,7 @@
 #define USE_NAIVE_TLockFreeFixedSizeAllocator_TLSCacheBase (0) // this is useful for find who really leaked
 #endif
 
-template<int32 SIZE, typename TBundleRecycler, typename TTrackingCounter = FNoopCounter, bool AllowDisablingOfTrim = false>
+template<int32 SIZE, typename TBundleRecycler, typename TTrackingCounter = FNoopCounter>
 class TLockFreeFixedSizeAllocator_TLSCacheBase : public FNoncopyable
 {
 	enum
@@ -66,7 +66,7 @@ public:
 				TLS.PartialBundle = GlobalFreeListBundles.Pop();
 				if (!TLS.PartialBundle)
 				{
-					TLS.PartialBundle = (void**)FMemory::MallocPersistentAuxiliary(SIZE_PER_BUNDLE);
+					TLS.PartialBundle = (void**)FMemory::Malloc(SIZE_PER_BUNDLE);
 					void **Next = TLS.PartialBundle;
 					for (int32 Index = 0; Index < NUM_PER_BUNDLE - 1; Index++)
 					{
@@ -190,7 +190,7 @@ private:
  * never returns free space until program shutdown.
  * alignment isn't handled, assumes FMemory::Malloc will work
  */
-template<int32 SIZE, int TPaddingForCacheContention, typename TTrackingCounter = FNoopCounter, bool AllowDisablingOfTrim = false>
+template<int32 SIZE, int TPaddingForCacheContention, typename TTrackingCounter = FNoopCounter>
 class TLockFreeFixedSizeAllocator
 {
 public:
@@ -218,7 +218,7 @@ public:
 		}
 		else
 		{
-			Memory = FMemory::MallocPersistentAuxiliary(SIZE);
+			Memory = FMemory::Malloc(SIZE);
 		}
 		return Memory;
 	}
@@ -241,16 +241,9 @@ public:
 	*/
 	void Trim()
 	{
-		if (AllowDisablingOfTrim)
-		{
-			if (FMemory::IsPersistentAuxiliaryActive())
-			{
-				return;
-			}
-		}
 		while (void* Mem = FreeList.Pop())
 		{
-			FMemory::FreePersistentAuxiliary(Mem);
+			FMemory::Free(Mem);
 			NumFree.Decrement();
 		}
 	}
@@ -294,8 +287,8 @@ private:
  * never returns free space, even at shutdown
  * alignment isn't handled, assumes FMemory::Malloc will work
  */
-template<int32 SIZE, int TPaddingForCacheContention, typename TTrackingCounter = FNoopCounter, bool AllowDisablingOfTrim = false>
-class TLockFreeFixedSizeAllocator_TLSCache : public TLockFreeFixedSizeAllocator_TLSCacheBase<SIZE, TLockFreePointerListUnordered<void*, TPaddingForCacheContention>, TTrackingCounter, AllowDisablingOfTrim>
+template<int32 SIZE, int TPaddingForCacheContention, typename TTrackingCounter = FNoopCounter>
+class TLockFreeFixedSizeAllocator_TLSCache : public TLockFreeFixedSizeAllocator_TLSCacheBase<SIZE, TLockFreePointerListUnordered<void*, TPaddingForCacheContention>, TTrackingCounter>
 {
 };
 
@@ -304,8 +297,8 @@ class TLockFreeFixedSizeAllocator_TLSCache : public TLockFreeFixedSizeAllocator_
  *
  * Never returns free space until program shutdown.
  */
-template<class T, int TPaddingForCacheContention, bool AllowDisablingOfTrim = false>
-class TLockFreeClassAllocator : private TLockFreeFixedSizeAllocator<sizeof(T), TPaddingForCacheContention, FNoopCounter, AllowDisablingOfTrim>
+template<class T, int TPaddingForCacheContention>
+class TLockFreeClassAllocator : private TLockFreeFixedSizeAllocator<sizeof(T), TPaddingForCacheContention, FNoopCounter>
 {
 public:
 	/**
@@ -348,8 +341,8 @@ public:
  *
  * Never returns free space until program shutdown.
  */
-template<class T, int TPaddingForCacheContention, bool AllowDisablingOfTrim = false>
-class TLockFreeClassAllocator_TLSCache : private TLockFreeFixedSizeAllocator_TLSCache<sizeof(T), TPaddingForCacheContention, FNoopCounter , AllowDisablingOfTrim>
+template<class T, int TPaddingForCacheContention>
+class TLockFreeClassAllocator_TLSCache : private TLockFreeFixedSizeAllocator_TLSCache<sizeof(T), TPaddingForCacheContention, FNoopCounter>
 {
 public:
 	/**
diff --git a/Engine/Source/Runtime/Core/Public/HAL/UnrealMemory.h b/Engine/Source/Runtime/Core/Public/HAL/UnrealMemory.h
--- a/Engine/Source/Runtime/Core/Public/HAL/UnrealMemory.h
+++ b/Engine/Source/Runtime/Core/Public/HAL/UnrealMemory.h
@@ -277,12 +277,25 @@ struct FMemory
 	* Functions to handle special memory given to the title from the platform
 	* This memory is allocated like a stack, it's never really freed
 	*/
-	static CORE_API void RegisterPersistentAuxiliary(void* InMemory, SIZE_T InSize);
+	UE_DEPRECATED(5.4, "Persistent Auxiliary allocator is obsolete and is replaced by a GetPersistentLinearAllocator()")
+	static inline void RegisterPersistentAuxiliary(void* /*InMemory*/, SIZE_T /*InSize*/) {}
+
+	UE_DEPRECATED(5.4, "Persistent Auxiliary allocator is obsolete and is replaced by a GetPersistentLinearAllocator()")
 	static CORE_API void* MallocPersistentAuxiliary(SIZE_T InSize, uint32 InAlignment = 0);
+
+	UE_DEPRECATED(5.4, "Persistent Auxiliary allocator is obsolete and is replaced by a GetPersistentLinearAllocator()")
 	static CORE_API void FreePersistentAuxiliary(void* InPtr);
+
+	UE_DEPRECATED(5.4, "Persistent Auxiliary allocator is obsolete and is replaced by a GetPersistentLinearAllocator()")
 	static CORE_API bool IsPersistentAuxiliaryActive();
-	static CORE_API void DisablePersistentAuxiliary();
-	static CORE_API void EnablePersistentAuxiliary();
+
+	UE_DEPRECATED(5.4, "Persistent Auxiliary allocator is obsolete and is replaced by a GetPersistentLinearAllocator()")
+	static inline void DisablePersistentAuxiliary() {}
+
+	UE_DEPRECATED(5.4, "Persistent Auxiliary allocator is obsolete and is replaced by a GetPersistentLinearAllocator()")
+	static inline void EnablePersistentAuxiliary() {}
+
+	UE_DEPRECATED(5.4, "Persistent Auxiliary allocator is obsolete and is replaced by a GetPersistentLinearAllocator()")
 	static CORE_API SIZE_T GetUsedPersistentAuxiliary();
 private:
 	static CORE_API void GCreateMalloc();
diff --git a/Engine/Source/Runtime/Core/Public/Memory/LinearVirtualMemoryAllocator.h b/Engine/Source/Runtime/Core/Public/Memory/LinearVirtualMemoryAllocator.h
new file mode 100644
--- /dev/null
+++ b/Engine/Source/Runtime/Core/Public/Memory/LinearVirtualMemoryAllocator.h
@@ -0,0 +1,99 @@
+// Copyright Epic Games, Inc. All Rights Reserved.
+
+#pragma once
+
+#include "HAL/Platform.h"
+
+#if PLATFORM_HAS_FPlatformVirtualMemoryBlock
+#include <atomic>
+#include "HAL/CriticalSection.h"
+
+struct FLinearVirtualMemoryAllocator
+{
+	CORE_API FLinearVirtualMemoryAllocator(SIZE_T ReserveMemorySize);
+
+	FORCEINLINE ~FLinearVirtualMemoryAllocator()
+	{
+		VirtualMemory.FreeVirtual();
+	}
+
+	CORE_API void* Allocate(SIZE_T Size, uint32 Alignment = 8);
+	CORE_API void  PreAllocate(SIZE_T Size, uint32 Alignment = 8);
+
+	// This will succeed only when deallocating the last allocation
+	CORE_API bool  TryDeallocate(void* Ptr, SIZE_T Size);
+
+	FORCEINLINE bool ContainsPointer(const void* Ptr) const
+	{
+		return (uintptr_t)Ptr - (uintptr_t)VirtualMemory.GetVirtualPointer() < Reserved;
+	}
+
+	FORCEINLINE SIZE_T GetExceedingSize() const
+	{
+		return ExceedsReservation.load(std::memory_order_relaxed);
+	}
+
+	FORCEINLINE bool IsInitialized() const
+	{
+		return Reserved != 0;
+	}
+
+	FORCEINLINE SIZE_T GetAllocatedMemorySize() const
+	{
+		return Committed;
+	}
+
+	FORCEINLINE SIZE_T GetReservedMemorySize() const
+	{
+		return Reserved;
+	}
+
+	FORCEINLINE const void* GetBasePointer() const
+	{
+		return VirtualMemory.GetVirtualPointer();
+	}
+
+private:
+	FCriticalSection Lock;
+	FPlatformMemory::FPlatformVirtualMemoryBlock VirtualMemory;
+	SIZE_T Reserved;
+	SIZE_T Committed = 0;
+	SIZE_T CurrentOffset = 0;
+	std::atomic<SIZE_T> ExceedsReservation = 0;
+
+	bool CanFit(SIZE_T Size, uint32 Alignment) const;
+};
+
+CORE_API FLinearVirtualMemoryAllocator& GetPersistentLinearAllocator();
+
+#else
+// stub implementation with most functions being nop
+struct FLinearVirtualMemoryAllocator
+{
+	FORCEINLINE FLinearVirtualMemoryAllocator() {}
+	FORCEINLINE FLinearVirtualMemoryAllocator(SIZE_T) {}
+
+	FORCEINLINE void*	Allocate(SIZE_T Size, uint32 Alignment = 8) { return FMemory::Malloc(Size, Alignment); }
+	FORCEINLINE void	PreAllocate(SIZE_T, uint32) {}
+	FORCEINLINE bool	TryDeallocate(void* Ptr, SIZE_T Size)	{ FMemory::Free(Ptr); return true; }
+
+	FORCEINLINE bool	ContainsPointer(const void*) const	{ return false; }
+	FORCEINLINE SIZE_T	GetExceedingSize() const			{ return 0; }
+	FORCEINLINE bool	IsInitialized() const				{ return false; }
+	FORCEINLINE SIZE_T	GetAllocatedMemorySize() const		{ return 0; }
+	FORCEINLINE SIZE_T	GetReservedMemorySize() const		{ return 0; }
+	FORCEINLINE const void* GetBasePointer() const			{ return nullptr; }
+};
+
+FORCEINLINE FLinearVirtualMemoryAllocator GetPersistentLinearAllocator() { return FLinearVirtualMemoryAllocator(); }
+
+#endif	//~PLATFORM_HAS_FPlatformVirtualMemoryBlock
+
+struct FLinearVirtualMemoryAllocatorExtends
+{
+	uint64 Address;
+	uint64 Size;
+};
+
+// Special case for the FPermanentObjectPoolExtents to reduce the amount of pointer dereferencing
+extern CORE_API FLinearVirtualMemoryAllocatorExtends GPersistentLinearAllocatorExtends;
\ No newline at end of file
diff --git a/Engine/Source/Runtime/CoreUObject/Private/UObject/UObjectAllocator.cpp b/Engine/Source/Runtime/CoreUObject/Private/UObject/UObjectAllocator.cpp
--- a/Engine/Source/Runtime/CoreUObject/Private/UObject/UObjectAllocator.cpp
+++ b/Engine/Source/Runtime/CoreUObject/Private/UObject/UObjectAllocator.cpp
@@ -12,32 +12,15 @@ DEFINE_LOG_CATEGORY_STATIC(LogUObjectAllocator, Log, All);
 /** Global UObjectBase allocator							*/
 COREUOBJECT_API FUObjectAllocator GUObjectAllocator;
 
-/**
- * Allocates and initializes the permanent object pool
- *
- * @param InPermanentObjectPoolSize size of permanent object pool
- */
-void FUObjectAllocator::AllocatePermanentObjectPool(int32 InPermanentObjectPoolSize)
-{
-	PermanentObjectPoolSize	= InPermanentObjectPoolSize;
-	PermanentObjectPool		= (uint8*) FMemory::MallocPersistentAuxiliary( PermanentObjectPoolSize );
-	PermanentObjectPoolTail	= PermanentObjectPool;
-	PermanentObjectPoolExceededTail = PermanentObjectPoolTail;
-}
-
-
 /**
  * Prints a debugf message to allow tuning
  */
 void FUObjectAllocator::BootMessage()
 {
-	if (PermanentObjectPoolSize && PermanentObjectPoolExceededTail - PermanentObjectPool > PermanentObjectPoolSize)
-	{
-		UE_LOG(LogUObjectAllocator, Warning, TEXT("%i Exceeds size of permanent object pool %i, please tune SizeOfPermanentObjectPool."), PermanentObjectPoolExceededTail - PermanentObjectPool, PermanentObjectPoolSize );
-	}
-	else
+	const SIZE_T ExceedingSize = GetPersistentLinearAllocator().GetExceedingSize();
+	if (ExceedingSize > 0)
 	{
-		UE_LOG(LogUObjectAllocator, Log, TEXT("%i out of %i bytes used by permanent object pool."), PermanentObjectPoolExceededTail - PermanentObjectPool, PermanentObjectPoolSize );
+		UE_LOG(LogUObjectAllocator, Warning, TEXT("Persistent memory pool exceeded by %u KB, please tune PersistentAllocatorReserveSizeMB setting in [MemoryPools] ini group."), ExceedingSize / 1024);
 	}
 }
 
@@ -51,48 +34,17 @@ void FUObjectAllocator::BootMessage()
  */
 UObjectBase* FUObjectAllocator::AllocateUObject(int32 Size, int32 Alignment, bool bAllowPermanent)
 {
-	// Force alignment to minimal of 16 bytes
-	Alignment = FMath::Max(16, Alignment);
-	int32 AlignedSize = Align( Size, Alignment );
-	UObjectBase* Result = nullptr;
-
-	bAllowPermanent &= PermanentObjectPool != nullptr;
-	const bool bPlaceInPerm = bAllowPermanent && (Align(PermanentObjectPoolTail,Alignment) + Size) <= (PermanentObjectPool + PermanentObjectPoolSize);
-	if (bAllowPermanent && !bPlaceInPerm)
-	{
-		// advance anyway so we can determine how much space we should set aside in the ini
-		uint8* AlignedPtr = Align( PermanentObjectPoolExceededTail, Alignment );
-		PermanentObjectPoolExceededTail = AlignedPtr + Size;
-	}
-	// Use object memory pool for objects disregarded by GC (initially loaded ones). This allows identifying their
-	// GC status by simply looking at their address.
-	if (bPlaceInPerm)
+	// we want to perform this allocation uninstrumented so the GC can clean this up if the transaction is aborted
+	UE_AUTORTFM_OPEN(
 	{
-		// Align current tail pointer and use it for object. 
-		uint8* AlignedPtr = Align( PermanentObjectPoolTail, Alignment );
-		// Update tail pointer.
-		PermanentObjectPoolTail = AlignedPtr + Size;
-		Result = (UObjectBase*)AlignedPtr;
-		if (PermanentObjectPoolExceededTail < PermanentObjectPoolTail)
+		if (bAllowPermanent)
 		{
-			PermanentObjectPoolExceededTail = PermanentObjectPoolTail;
+			// this allocation might go over the reserved memory amount and default to FMemory::Malloc, so we are moving it into the ARTFM scope
+			return (UObjectBase*)GetPersistentLinearAllocator().Allocate(Size, Alignment);
 		}
-	}
-	else
-	{
-		// we want to perform this allocation uninstrumented so the GC can clean this up if the transaction is aborted
-		UE_AUTORTFM_OPEN(
-		{
-			// Allocate new memory of the appropriate size and alignment.
-			Result = (UObjectBase*)FMemory::Malloc( Size, Alignment );
-		});
-	}
-
-#if !UE_BUILD_TEST && !UE_BUILD_SHIPPING
-	checkf(IsAligned(Result, Alignment), TEXT("Allocated memory address does not match requirement of %d byte alignment for size %d"), Alignment, Size);
-#endif
 
-	return Result;
+		return (UObjectBase*)FMemory::Malloc(Size, Alignment);
+	});
 }
 
 /**
@@ -104,7 +56,7 @@ void FUObjectAllocator::FreeUObject(UObjectBase *Object) const
 {
 	check(Object);
 	// Only free memory if it was allocated directly from allocator and not from permanent object pool.
-	if (FPermanentObjectPoolExtents(*this).Contains(Object) == false)
+	if (FPermanentObjectPoolExtents().Contains(Object) == false)
 	{
 		FMemory::Free(Object);
 	}
diff --git a/Engine/Source/Runtime/CoreUObject/Private/UObject/UObjectBase.cpp b/Engine/Source/Runtime/CoreUObject/Private/UObject/UObjectBase.cpp
--- a/Engine/Source/Runtime/CoreUObject/Private/UObject/UObjectBase.cpp
+++ b/Engine/Source/Runtime/CoreUObject/Private/UObject/UObjectBase.cpp
@@ -984,7 +984,6 @@ void UObjectBaseInit()
 
 	// Zero initialize and later on get value from .ini so it is overridable per game/ platform...
 	int32 MaxObjectsNotConsideredByGC = 0;
-	int32 SizeOfPermanentObjectPool = 0;
 	int32 MaxUObjects = 2 * 1024 * 1024; // Default to ~2M UObjects
 	bool bPreAllocateUObjectArray = false;	
 
@@ -1002,9 +1001,6 @@ void UObjectBaseInit()
 		else
 		{
 			GConfig->GetInt(TEXT("/Script/Engine.GarbageCollectionSettings"), TEXT("gc.MaxObjectsNotConsideredByGC"), MaxObjectsNotConsideredByGC, GEngineIni);
-
-			// Not used on PC as in-place creation inside bigger pool interacts with the exit purge and deleting UObject directly.
-			GConfig->GetInt(TEXT("/Script/Engine.GarbageCollectionSettings"), TEXT("gc.SizeOfPermanentObjectPool"), SizeOfPermanentObjectPool, GEngineIni);
 		}
 
 		// Maximum number of UObjects in cooked game
@@ -1025,23 +1021,10 @@ void UObjectBaseInit()
 #endif
 	}
 
-	if (MaxObjectsNotConsideredByGC <= 0 && SizeOfPermanentObjectPool > 0)
-	{
-		// If permanent object pool is enabled but disregard for GC is disabled, GC will mark permanent object pool objects
-		// as unreachable and may destroy them so disable permanent object pool too.
-		// An alternative would be to make GC not mark permanent object pool objects as unreachable but then they would have to
-		// be considered as root set objects because they could be referencing objects from outside of permanent object pool.
-		// This would be inconsistent and confusing and also counter productive (the more root set objects the more expensive MarkAsUnreachable phase is).
-		SizeOfPermanentObjectPool = 0;
-		UE_LOG(LogInit, Warning, TEXT("Disabling permanent object pool because disregard for GC is disabled (gc.MaxObjectsNotConsideredByGC=%d)."), MaxObjectsNotConsideredByGC);
-	}
-
 	// Log what we're doing to track down what really happens as log in LaunchEngineLoop doesn't report those settings in pristine form.
-	UE_LOG(LogInit, Log, TEXT("%s for max %d objects, including %i objects not considered by GC, pre-allocating %i bytes for permanent pool."), 
-		bPreAllocateUObjectArray ? TEXT("Pre-allocating") : TEXT("Presizing"),
-		MaxUObjects, MaxObjectsNotConsideredByGC, SizeOfPermanentObjectPool);
+	UE_LOG(LogInit, Log, TEXT("%s for max %d objects, including %i objects not considered by GC."), 
+		bPreAllocateUObjectArray ? TEXT("Pre-allocating") : TEXT("Presizing"), MaxUObjects, MaxObjectsNotConsideredByGC);
 
-	GUObjectAllocator.AllocatePermanentObjectPool(SizeOfPermanentObjectPool);
 	GUObjectArray.AllocateObjectPool(MaxUObjects, MaxObjectsNotConsideredByGC, bPreAllocateUObjectArray);
 #if UE_WITH_OBJECT_HANDLE_LATE_RESOLVE
 	UE::CoreUObject::Private::InitObjectHandles(GUObjectArray.GetObjectArrayCapacity());
diff --git a/Engine/Source/Runtime/CoreUObject/Public/UObject/UObjectAllocator.h b/Engine/Source/Runtime/CoreUObject/Public/UObject/UObjectAllocator.h
--- a/Engine/Source/Runtime/CoreUObject/Public/UObject/UObjectAllocator.h
+++ b/Engine/Source/Runtime/CoreUObject/Public/UObject/UObjectAllocator.h
@@ -8,42 +8,26 @@
 
 #include "CoreMinimal.h"
 #include "HAL/PlatformMath.h"
+#include "Memory/LinearVirtualMemoryAllocator.h"
 
 class UObjectBase;
 
 class FUObjectAllocator
 {
 public:
-
-	/**
-	 * Constructor, initializes to no permanent object pool
-	 */
-	FUObjectAllocator() :
-	  PermanentObjectPoolSize(0),
-	  PermanentObjectPool(NULL),
-	  PermanentObjectPoolTail(NULL),
-		PermanentObjectPoolExceededTail(NULL)
-	{
-	}
-
 	/**
 	 * Allocates and initializes the permanent object pool
 	 *
 	 * @param InPermanentObjectPoolSize size of permanent object pool
 	 */
-	COREUOBJECT_API void AllocatePermanentObjectPool(int32 InPermanentObjectPoolSize);
+	UE_DEPRECATED(5.4, "Permanent Object Pool is handled by the global instance of FLinearVirtualMemoryAllocator now")
+	COREUOBJECT_API void AllocatePermanentObjectPool(int32 InPermanentObjectPoolSize) {}
 
 	/**
 	 * Prints a debugf message to allow tuning
 	 */
 	COREUOBJECT_API void BootMessage();
 
-	UE_DEPRECATED(5.1, "Use the more efficient FPermanentObjectPoolExtents instead")
-	FORCEINLINE bool ResidesInPermanentPool(const UObjectBase *Object) const
-	{
-		return ((const uint8*)Object >= PermanentObjectPool) && ((const uint8*)Object < PermanentObjectPoolTail);
-	}
-
 	/**
 	 * Allocates a UObjectBase from the free store or the permanent object pool
 	 *
@@ -60,18 +44,6 @@ public:
 	 * @param Object object to free
 	 */
 	COREUOBJECT_API void FreeUObject(UObjectBase *Object) const;
-
-private:
-	friend class FPermanentObjectPoolExtents;
-
-	/** Size in bytes of pool for objects disregarded for GC.								*/
-	int32							PermanentObjectPoolSize;
-	/** Begin of pool for objects disregarded for GC.										*/
-	uint8*						PermanentObjectPool;
-	/** Current position in pool for objects disregarded for GC.							*/
-	uint8*						PermanentObjectPoolTail;
-	/** Tail that exceeded the size of the permanent object pool, >= PermanentObjectPoolTail.		*/
-	uint8*						PermanentObjectPoolExceededTail;
 };
 
 /** Global UObjectBase allocator							*/
@@ -81,11 +53,11 @@ extern COREUOBJECT_API FUObjectAllocator GUObjectAllocator;
 class FPermanentObjectPoolExtents
 {
 public:
-	FORCEINLINE FPermanentObjectPoolExtents(const FUObjectAllocator& ObjectAllocator = GUObjectAllocator)
-		: Address(reinterpret_cast<uint64>(ObjectAllocator.PermanentObjectPool))
-		, Size(static_cast<uint64>(ObjectAllocator.PermanentObjectPoolSize))
+	FORCEINLINE FPermanentObjectPoolExtents(const FLinearVirtualMemoryAllocatorExtends& InAllocatorExtends = GPersistentLinearAllocatorExtends)
+		: Address(InAllocatorExtends.Address)
+		, Size(InAllocatorExtends.Size)
 	{}
-	
+
 	FORCEINLINE bool Contains(const UObjectBase* Object) const
 	{
 		return reinterpret_cast<uint64>(Object) - Address < Size;
