Added a missing include in the LinearVirtualMemoryAllocator.h to fix compilation failure in Non-Unity builds

diff --git a/Engine/Source/Runtime/Core/Public/Memory/LinearVirtualMemoryAllocator.h b/Engine/Source/Runtime/Core/Public/Memory/LinearVirtualMemoryAllocator.h
--- a/Engine/Source/Runtime/Core/Public/Memory/LinearVirtualMemoryAllocator.h
+++ b/Engine/Source/Runtime/Core/Public/Memory/LinearVirtualMemoryAllocator.h
@@ -6,6 +6,7 @@
 
 #if PLATFORM_HAS_FPlatformVirtualMemoryBlock
 #include <atomic>
+#include "HAL/PlatformMemory.h"
 #include "HAL/CriticalSection.h"
 
 struct FLinearVirtualMemoryAllocator
