Disable gc.SizeOfPermanentObjectPool deprecation warning for now

diff --git a/Engine/Source/Runtime/Core/Core.Build.cs b/Engine/Source/Runtime/Core/Core.Build.cs
--- a/Engine/Source/Runtime/Core/Core.Build.cs
+++ b/Engine/Source/Runtime/Core/Core.Build.cs
@@ -375,7 +375,7 @@ public class Core : ModuleRules
 		int SizeOfPermanentObjectPool = 0;
 		if (Ini.GetInt32("/Script/Engine.GarbageCollectionSettings", "gc.SizeOfPermanentObjectPool", out SizeOfPermanentObjectPool))
 		{
-			Target.Logger.LogWarning("/Script/Engine.GarbageCollectionSettings, gc.SizeOfPermanentObjectPool ini setting was deprecated in a favor of MemoryPools, PersistentAllocatorReserveSizeMB");
+			//Target.Logger.LogWarning("/Script/Engine.GarbageCollectionSettings, gc.SizeOfPermanentObjectPool ini setting was deprecated in a favor of MemoryPools, PersistentAllocatorReserveSizeMB");
 		}
 		int PersistentAllocatorReserveSizeMB = 0;
 		Ini.GetInt32("MemoryPools", "PersistentAllocatorReserveSizeMB", out PersistentAllocatorReserveSizeMB);
