Fixed TOptional::Emplace constructing over an already-constructed inner object when it's in the intrustive unset state.
Removed operator=(FIntrusiveUnsetOptionalState) overloads and documentation as they are no longer relevant now that TOptional::Reset always destructs and reconstructs for an intrusive state.

diff --git a/Engine/Source/Runtime/Core/Public/Internationalization/Text.h b/Engine/Source/Runtime/Core/Public/Internationalization/Text.h
--- a/Engine/Source/Runtime/Core/Public/Internationalization/Text.h
+++ b/Engine/Source/Runtime/Core/Public/Internationalization/Text.h
@@ -779,11 +779,6 @@ public:
 	, Flags(0)
 	{
 	}
-	void operator=(FIntrusiveUnsetOptionalState)
-	{
-		TextData = nullptr;
-		Flags = 0;	
-	}
 	bool operator==(FIntrusiveUnsetOptionalState) const
 	{
 		return TextData == nullptr;
diff --git a/Engine/Source/Runtime/Core/Public/Misc/IntrusiveUnsetOptionalState.h b/Engine/Source/Runtime/Core/Public/Misc/IntrusiveUnsetOptionalState.h
--- a/Engine/Source/Runtime/Core/Public/Misc/IntrusiveUnsetOptionalState.h
+++ b/Engine/Source/Runtime/Core/Public/Misc/IntrusiveUnsetOptionalState.h
@@ -7,10 +7,9 @@
  * TOptional's 'unset' state, saving space.  A class in such a state will only ever be
  * compared against FIntrusiveUnsetOptionalState or destructed.
  *
- * A class should implement a constructor taking FIntrusiveUnsetOptionalState, an
- * assignment operator from FIntrusiveUnsetOptionalState, and an equality comparison
- * operator against FIntrusiveUnsetOptionalState, which will put a class instance into
- * the 'unset' state (in the case of the constructor and assignment) and allow testing
+ * A class should implement a constructor taking FIntrusiveUnsetOptionalState, and an 
+ * equality comparison operator against FIntrusiveUnsetOptionalState, which will put a
+  class instance into the 'unset' state (in the case of the constructor) and allow testing
  * of its unset state.
  *
  * A public constexpr static data member of type bool called bHasIntrusiveUnsetOptionalState
@@ -45,13 +44,6 @@
  *         Index = -1;
  *     }
  *
- *     // Similarly with this assignment operator.  It need not return a reference to itself
- *     // as most assignment operators do.
- *     void operator=(FIntrusiveUnsetOptionalState)
- *     {
- *         Index = -1;
- *     }
- *
  *     // This comparison function will only ever be called by TOptional to check if the
  *     // object is in the 'unset' state.  It does not need to be commutative like most
  *     // comparison operators, nor is an operator!= necessary.
diff --git a/Engine/Source/Runtime/Core/Public/Misc/Optional.h b/Engine/Source/Runtime/Core/Public/Misc/Optional.h
--- a/Engine/Source/Runtime/Core/Public/Misc/Optional.h
+++ b/Engine/Source/Runtime/Core/Public/Misc/Optional.h
@@ -228,10 +228,17 @@ public:
 	OptionalType& Emplace(ArgsType&&... Args)
 	{
 		// Destroy the member in-place before replacing it - a bit nasty, but it'll work since we don't support exceptions
-		if (IsSet())
+		if constexpr (bUsingIntrusiveUnsetState)
 		{
 			DestroyValue();
 		}
+		else
+		{
+			if (IsSet())
+			{
+				DestroyValue();
+			}
+		}
 
 		// If this fails to compile when trying to call Emplace with a non-public constructor,
 		// do not make TOptional a friend.
diff --git a/Engine/Source/Runtime/Core/Public/Templates/Function.h b/Engine/Source/Runtime/Core/Public/Templates/Function.h
--- a/Engine/Source/Runtime/Core/Public/Templates/Function.h
+++ b/Engine/Source/Runtime/Core/Public/Templates/Function.h
@@ -722,10 +722,6 @@ public:
 	explicit TFunctionRef(FIntrusiveUnsetOptionalState)
 	{
 	}
-	void operator=(FIntrusiveUnsetOptionalState)
-	{
-		Super::Reset();
-	}
 	bool operator==(FIntrusiveUnsetOptionalState) const
 	{
 		return !Super::IsSet();
diff --git a/Engine/Source/Runtime/Core/Public/Templates/SharedPointer.h b/Engine/Source/Runtime/Core/Public/Templates/SharedPointer.h
--- a/Engine/Source/Runtime/Core/Public/Templates/SharedPointer.h
+++ b/Engine/Source/Runtime/Core/Public/Templates/SharedPointer.h
@@ -538,11 +538,6 @@ public:
 		: Object(nullptr)
 	{
 	}
-	void operator=(FIntrusiveUnsetOptionalState)
-	{
-		Object = nullptr;
-		SharedReferenceCount = {};
-	}
 	bool operator==(FIntrusiveUnsetOptionalState) const
 	{
 		return !IsValid();
@@ -1130,11 +1125,6 @@ public:
 		: Object((ObjectType*)-1)
 	{
 	}
-	void operator=(FIntrusiveUnsetOptionalState)
-	{
-		Object = (ObjectType*)-1;
-		SharedReferenceCount = {};
-	}
 	bool operator==(FIntrusiveUnsetOptionalState) const
 	{
 		return Object == (ObjectType*)-1;
@@ -1526,11 +1516,6 @@ public:
 		: Object((ObjectType*)-1)
 	{
 	}
-	void operator=(FIntrusiveUnsetOptionalState)
-	{
-		Object = (ObjectType*)-1;
-		WeakReferenceCount = SharedPointerInternals::FWeakReferencer< Mode >{};
-	}
 	bool operator==(FIntrusiveUnsetOptionalState) const
 	{
 		return Object == (ObjectType*)-1;
diff --git a/Engine/Source/Runtime/Core/Public/UObject/NameTypes.h b/Engine/Source/Runtime/Core/Public/UObject/NameTypes.h
--- a/Engine/Source/Runtime/Core/Public/UObject/NameTypes.h
+++ b/Engine/Source/Runtime/Core/Public/UObject/NameTypes.h
@@ -58,7 +58,6 @@ struct FNameEntryId
 	FNameEntryId() : Value(0) {}
 	FNameEntryId(ENoInit) {}
 	explicit FNameEntryId(FIntrusiveUnsetOptionalState) : Value(~0u) {}
-	void operator=(FIntrusiveUnsetOptionalState) { Value = ~0u; }
 
 	bool IsNone() const
 	{
@@ -944,18 +943,14 @@ public:
 
 	/** Special constructor used by TOptional<FName> */
 	explicit FName(FIntrusiveUnsetOptionalState I)
-	: FName()
-	{
-		ComparisonIndex = I;
+		: ComparisonIndex(I)
+#if !UE_FNAME_OUTLINE_NUMBER
+		, Number(NAME_NO_NUMBER_INTERNAL)
+#endif
 #if WITH_CASE_PRESERVING_NAME
-		DisplayIndex = I;
+		, DisplayIndex(I)
 #endif
-	}
-
-	/** Special assignment operator used by TOptional<FName> */
-	inline void operator=(FIntrusiveUnsetOptionalState I)
 	{
-		*this = FName(I);
 	}
 
 	FORCEINLINE explicit FName(FMinimalName InName);
