Include Missing Header

diff --git a/Engine/Source/Runtime/Core/Public/Misc/ScopeRWLock.h b/Engine/Source/Runtime/Core/Public/Misc/ScopeRWLock.h
--- a/Engine/Source/Runtime/Core/Public/Misc/ScopeRWLock.h
+++ b/Engine/Source/Runtime/Core/Public/Misc/ScopeRWLock.h
@@ -6,7 +6,7 @@
 #include "Misc/AssertionMacros.h"
 #include "HAL/CriticalSection.h"
 #include "AutoRTFM/AutoRTFM.h"
-
+#include "Templates/Function.h"
 
 /** Keeps a FRWLock read-locked while this scope lives */
 class FReadScopeLock
