Fix Csv profiler using old metadata for new capture if end capture is called when no capture is in progress.

#tests verified the fix doing captures in PIE.

diff --git a/Engine/Source/Runtime/Core/Private/ProfilingDebugging/CsvProfiler.cpp b/Engine/Source/Runtime/Core/Private/ProfilingDebugging/CsvProfiler.cpp
--- a/Engine/Source/Runtime/Core/Private/ProfilingDebugging/CsvProfiler.cpp
+++ b/Engine/Source/Runtime/Core/Private/ProfilingDebugging/CsvProfiler.cpp
@@ -3357,6 +3357,12 @@ TSharedFuture<FString> FCsvProfiler::EndCapture(FGraphEventRef EventToSignal)
 	LLM_SCOPE(ELLMTag::CsvProfiler);
 	check(IsInGameThread());
 
+	if (!IsCapturing())
+	{
+		UE_LOG(LogCsvProfiler, Warning, TEXT("EndCapture() called, but no capture was in progress. Ignoring!"));
+		return {};
+	}
+
 	// If there's already a stop command in flight for this capture, warn and continue
 	FCsvCaptureCommand CurrentCommand;
 	if (CommandQueue.Peek(CurrentCommand) && CurrentCommand.CommandType == ECsvCommandType::Stop)
