Added intrusive unset states to various containers.

diff --git a/Engine/Source/Runtime/Core/Public/Containers/Array.h b/Engine/Source/Runtime/Core/Public/Containers/Array.h
--- a/Engine/Source/Runtime/Core/Public/Containers/Array.h
+++ b/Engine/Source/Runtime/Core/Public/Containers/Array.h
@@ -4,6 +4,7 @@
 
 #include "CoreTypes.h"
 #include "Misc/AssertionMacros.h"
+#include "Misc/IntrusiveUnsetOptionalState.h"
 #include "Misc/ReverseIterate.h"
 #include "HAL/UnrealMemory.h"
 #include "Templates/UnrealTypeTraits.h"
@@ -683,6 +684,26 @@ public:
 		// they are left unchanged and use-after-destruct will see them the same as before destruct
 	}
 
+	///////////////////////////////////////////////
+	// Start - intrusive TOptional<TArray> state //
+	///////////////////////////////////////////////
+	constexpr static bool bHasIntrusiveUnsetOptionalState = true;
+	using IntrusiveUnsetOptionalStateType = TArray;
+
+	explicit TArray(FIntrusiveUnsetOptionalState Tag)
+		: ArrayNum(0)
+		, ArrayMax(-1)
+	{
+		// Use ArrayMax == -1 as our intrusive state so that the destructor still works without change, as it doesn't use ArrayMax.
+	}
+	bool operator==(FIntrusiveUnsetOptionalState Tag) const
+	{
+		return ArrayMax == -1;
+	}
+	/////////////////////////////////////////////
+	// End - intrusive TOptional<TArray> state //
+	/////////////////////////////////////////////
+
 	/**
 	 * Helper function for returning a typed pointer to the first array entry.
 	 *
diff --git a/Engine/Source/Runtime/Core/Public/Containers/ArrayView.h b/Engine/Source/Runtime/Core/Public/Containers/ArrayView.h
--- a/Engine/Source/Runtime/Core/Public/Containers/ArrayView.h
+++ b/Engine/Source/Runtime/Core/Public/Containers/ArrayView.h
@@ -5,6 +5,7 @@
 #include "CoreTypes.h"
 #include "Misc/AssertionMacros.h"
 #include "Misc/ReverseIterate.h"
+#include "Misc/IntrusiveUnsetOptionalState.h"
 #include "Templates/Invoke.h"
 #include "Templates/UnrealTypeTraits.h"
 #include "Traits/ElementType.h"
@@ -263,6 +264,25 @@ public:
 		static_assert(std::is_const_v<ElementType>, "Only views of const elements can bind to initializer lists");
 	}
 
+	///////////////////////////////////////////////////
+	// Start - intrusive TOptional<TArrayView> state //
+	///////////////////////////////////////////////////
+	constexpr static bool bHasIntrusiveUnsetOptionalState = true;
+	using IntrusiveUnsetOptionalStateType = TArrayView;
+
+	explicit TArrayView(FIntrusiveUnsetOptionalState)
+		: DataPtr(nullptr)
+		, ArrayNum(-1)
+	{
+	}
+	bool operator==(FIntrusiveUnsetOptionalState) const
+	{
+		return ArrayNum == -1;
+	}
+	/////////////////////////////////////////////////
+	// End - intrusive TOptional<TArrayView> state //
+	/////////////////////////////////////////////////
+
 public:
 
 	/**
diff --git a/Engine/Source/Runtime/Core/Public/Containers/ChunkedArray.h b/Engine/Source/Runtime/Core/Public/Containers/ChunkedArray.h
--- a/Engine/Source/Runtime/Core/Public/Containers/ChunkedArray.h
+++ b/Engine/Source/Runtime/Core/Public/Containers/ChunkedArray.h
@@ -4,6 +4,7 @@
 
 #include "CoreTypes.h"
 #include "Misc/AssertionMacros.h"
+#include "Misc/IntrusiveUnsetOptionalState.h"
 #include "Templates/UnrealTypeTraits.h"
 #include "Containers/IndirectArray.h"
 
@@ -100,6 +101,24 @@ public:
 	TChunkedArray(const TChunkedArray&) = default;
 	TChunkedArray& operator=(const TChunkedArray&) = default;
 
+	//////////////////////////////////////////////////////
+	// Start - intrusive TOptional<TChunkedArray> state //
+	//////////////////////////////////////////////////////
+	constexpr static bool bHasIntrusiveUnsetOptionalState = true;
+	using IntrusiveUnsetOptionalStateType = TChunkedArray;
+
+	explicit TChunkedArray(FIntrusiveUnsetOptionalState)
+		: NumElements(-1)
+	{
+	}
+	bool operator==(FIntrusiveUnsetOptionalState) const
+	{
+		return NumElements == -1;
+	}
+	////////////////////////////////////////////////////
+	// End - intrusive TOptional<TChunkedArray> state //
+	////////////////////////////////////////////////////
+
 	// Accessors.
 	ElementType& operator()(int32 ElementIndex)
 	{
diff --git a/Engine/Source/Runtime/Core/Public/Containers/Map.h b/Engine/Source/Runtime/Core/Public/Containers/Map.h
--- a/Engine/Source/Runtime/Core/Public/Containers/Map.h
+++ b/Engine/Source/Runtime/Core/Public/Containers/Map.h
@@ -161,6 +161,24 @@ protected:
 		: Pairs(Other.Pairs)
 	{ }
 
+	/////////////////////////////////////////////////
+	// Start - intrusive TOptional<TMapBase> state //
+	/////////////////////////////////////////////////
+	constexpr static bool bHasIntrusiveUnsetOptionalState = true;
+	using IntrusiveUnsetOptionalStateType = TMapBase;
+
+	explicit TMapBase(FIntrusiveUnsetOptionalState Tag)
+		: Pairs(Tag)
+	{
+	}
+	bool operator==(FIntrusiveUnsetOptionalState Tag) const
+	{
+		return Pairs == Tag;
+	}
+	///////////////////////////////////////////////
+	// End - intrusive TOptional<TMapBase> state //
+	///////////////////////////////////////////////
+
 	/** Assignment operator for moving elements from a TMap with a different SetAllocator */
 	template<typename OtherSetAllocator>
 	TMapBase& operator=(TMapBase<KeyType, ValueType, OtherSetAllocator, KeyFuncs>&& Other)
@@ -1064,6 +1082,24 @@ protected:
 		return *this;
 	}
 
+	/////////////////////////////////////////////////////////
+	// Start - intrusive TOptional<TSortableMapBase> state //
+	/////////////////////////////////////////////////////////
+	constexpr static bool bHasIntrusiveUnsetOptionalState = true;
+	using IntrusiveUnsetOptionalStateType = TSortableMapBase;
+
+	explicit TSortableMapBase(FIntrusiveUnsetOptionalState Tag)
+		: Super(Tag)
+	{
+	}
+	bool operator==(FIntrusiveUnsetOptionalState Tag) const
+	{
+		return Super::operator==(Tag);
+	}
+	///////////////////////////////////////////////////////
+	// End - intrusive TOptional<TSortableMapBase> state //
+	///////////////////////////////////////////////////////
+
 public:
 	/**
 	 * Sorts the pairs array using each pair's Key as the sort criteria, then rebuilds the map's hash.
@@ -1206,6 +1242,24 @@ public:
 		}
 	}
 
+	/////////////////////////////////////////////
+	// Start - intrusive TOptional<TMap> state //
+	/////////////////////////////////////////////
+	constexpr static bool bHasIntrusiveUnsetOptionalState = true;
+	using IntrusiveUnsetOptionalStateType = TMap;
+
+	explicit TMap(FIntrusiveUnsetOptionalState Tag)
+		: Super(Tag)
+	{
+	}
+	bool operator==(FIntrusiveUnsetOptionalState Tag) const
+	{
+		return Super::operator==(Tag);
+	}
+	///////////////////////////////////////////
+	// End - intrusive TOptional<TMap> state //
+	///////////////////////////////////////////
+
 	/** Assignment operator for moving elements from a TMap with a different SetAllocator */
 	template<typename OtherSetAllocator>
 	TMap& operator=(TMap<KeyType, ValueType, OtherSetAllocator, KeyFuncs>&& Other)
@@ -1391,6 +1445,24 @@ public:
 		}
 	}
 
+	//////////////////////////////////////////////////
+	// Start - intrusive TOptional<TMultiMap> state //
+	//////////////////////////////////////////////////
+	constexpr static bool bHasIntrusiveUnsetOptionalState = true;
+	using IntrusiveUnsetOptionalStateType = TMultiMap;
+
+	explicit TMultiMap(FIntrusiveUnsetOptionalState Tag)
+		: Super(Tag)
+	{
+	}
+	bool operator==(FIntrusiveUnsetOptionalState Tag) const
+	{
+		return Super::operator==(Tag);
+	}
+	////////////////////////////////////////////////
+	// End - intrusive TOptional<TMultiMap> state //
+	////////////////////////////////////////////////
+
 	/** Assignment operator for moving elements from a TMap with a different SetAllocator */
 	template<typename OtherSetAllocator>
 	TMultiMap& operator=(TMultiMap<KeyType, ValueType, OtherSetAllocator, KeyFuncs>&& Other)
diff --git a/Engine/Source/Runtime/Core/Public/Containers/Set.h b/Engine/Source/Runtime/Core/Public/Containers/Set.h
--- a/Engine/Source/Runtime/Core/Public/Containers/Set.h
+++ b/Engine/Source/Runtime/Core/Public/Containers/Set.h
@@ -307,24 +307,21 @@ private:
 public:
 	/** Initialization constructor. */
 	FORCEINLINE TSet()
-	:	HashSize(0)
-	{}
+	{
+	}
 
 	/** Copy constructor. */
 	FORCEINLINE TSet(const TSet& Copy)
-	:	HashSize(0)
 	{
 		*this = Copy;
 	}
 
 	FORCEINLINE explicit TSet(TArrayView<const ElementType> InArrayView)
-		: HashSize(0)
 	{
 		Append(InArrayView);
 	}
 
 	FORCEINLINE explicit TSet(TArray<ElementType>&& InArray)
-		: HashSize(0)
 	{
 		Append(MoveTemp(InArray));
 	}
@@ -335,6 +332,24 @@ public:
 		HashSize = 0;
 	}
 
+	/////////////////////////////////////////////
+	// Start - intrusive TOptional<TSet> state //
+	/////////////////////////////////////////////
+	constexpr static bool bHasIntrusiveUnsetOptionalState = true;
+	using IntrusiveUnsetOptionalStateType = TSet;
+
+	explicit TSet(FIntrusiveUnsetOptionalState Tag)
+		: Elements(Tag)
+	{
+	}
+	bool operator==(FIntrusiveUnsetOptionalState Tag) const
+	{
+		return Elements == Tag;
+	}
+	///////////////////////////////////////////
+	// End - intrusive TOptional<TSet> state //
+	///////////////////////////////////////////
+
 	/** Assignment operator. */
 	TSet& operator=(const TSet& Copy)
 	{
@@ -1391,7 +1406,7 @@ private:
 	ElementArrayType Elements;
 
 	mutable HashType Hash;
-	mutable int32	 HashSize;
+	mutable int32	 HashSize = 0;
 
 public:
 	void WriteMemoryImage(FMemoryImageWriter& Writer) const
diff --git a/Engine/Source/Runtime/Core/Public/Containers/SortedMap.h b/Engine/Source/Runtime/Core/Public/Containers/SortedMap.h
--- a/Engine/Source/Runtime/Core/Public/Containers/SortedMap.h
+++ b/Engine/Source/Runtime/Core/Public/Containers/SortedMap.h
@@ -59,6 +59,24 @@ public:
 		}
 	}
 
+	///////////////////////////////////////////////////
+	// Start - intrusive TOptional<TSortedMap> state //
+	///////////////////////////////////////////////////
+	constexpr static bool bHasIntrusiveUnsetOptionalState = true;
+	using IntrusiveUnsetOptionalStateType = TSortedMap;
+
+	explicit TSortedMap(FIntrusiveUnsetOptionalState Tag)
+		: Pairs(Tag)
+	{
+	}
+	bool operator==(FIntrusiveUnsetOptionalState Tag) const
+	{
+		return Pairs == Tag;
+	}
+	/////////////////////////////////////////////////
+	// End - intrusive TOptional<TSortedMap> state //
+	/////////////////////////////////////////////////
+
 	/** Assignment operator for moving elements from a TSortedMap with a different ArrayAllocator. */
 	template<typename OtherArrayAllocator>
 	TSortedMap& operator=(TSortedMap<KeyType, ValueType, OtherArrayAllocator, SortPredicate>&& Other)
diff --git a/Engine/Source/Runtime/Core/Public/Containers/SparseArray.h b/Engine/Source/Runtime/Core/Public/Containers/SparseArray.h
--- a/Engine/Source/Runtime/Core/Public/Containers/SparseArray.h
+++ b/Engine/Source/Runtime/Core/Public/Containers/SparseArray.h
@@ -16,6 +16,7 @@
 #include "Serialization/StructuredArchive.h"
 #include "Serialization/MemoryImageWriter.h"
 #include "Containers/UnrealString.h"
+#include "Misc/IntrusiveUnsetOptionalState.h"
 
 
 #if UE_BUILD_SHIPPING || UE_BUILD_TEST
@@ -84,6 +85,24 @@ public:
 		Empty();
 	}
 
+	/////////////////////////////////////////////////////
+	// Start - intrusive TOptional<TSparseArray> state //
+	/////////////////////////////////////////////////////
+	constexpr static bool bHasIntrusiveUnsetOptionalState = true;
+	using IntrusiveUnsetOptionalStateType = TSparseArray;
+
+	explicit TSparseArray(FIntrusiveUnsetOptionalState)
+		: NumFreeIndices(-1)
+	{
+	}
+	bool operator==(FIntrusiveUnsetOptionalState) const
+	{
+		return NumFreeIndices == -1;
+	}
+	///////////////////////////////////////////////////
+	// End - intrusive TOptional<TSparseArray> state //
+	///////////////////////////////////////////////////
+
 	/** Marks an index as allocated, and returns information about the allocation. */
 	FSparseArrayAllocationInfo AllocateIndex(int32 Index)
 	{
diff --git a/Engine/Source/Runtime/Core/Public/Containers/UnrealString.h.inl b/Engine/Source/Runtime/Core/Public/Containers/UnrealString.h.inl
--- a/Engine/Source/Runtime/Core/Public/Containers/UnrealString.h.inl
+++ b/Engine/Source/Runtime/Core/Public/Containers/UnrealString.h.inl
@@ -196,7 +196,25 @@ public:
         }
     }
 #endif
-    
+
+	////////////////////////////////////////////////////////
+	// Start - intrusive TOptional<UE_STRING_CLASS> state //
+	////////////////////////////////////////////////////////
+	constexpr static bool bHasIntrusiveUnsetOptionalState = true;
+	using IntrusiveUnsetOptionalStateType = UE_STRING_CLASS;
+
+	explicit UE_STRING_CLASS(FIntrusiveUnsetOptionalState Tag)
+		: Data(Tag)
+	{
+	}
+	bool operator==(FIntrusiveUnsetOptionalState Tag) const
+	{
+		return Data == Tag;
+	}
+	//////////////////////////////////////////////////////
+	// End - intrusive TOptional<UE_STRING_CLASS> state //
+	//////////////////////////////////////////////////////
+
 #if defined(__OBJC__) && UE_STRING_CHARTYPE_IS_TCHAR
 	/** Convert Objective-C NSString* to string class */
 	FORCEINLINE UE_STRING_CLASS(const NSString* In) : UE_STRING_CLASS((__bridge CFStringRef)In)
