Fix position of [[nodiscard]] attribute.

diff --git a/Engine/Source/Runtime/Core/Public/Misc/NamePermissionList.h b/Engine/Source/Runtime/Core/Public/Misc/NamePermissionList.h
--- a/Engine/Source/Runtime/Core/Public/Misc/NamePermissionList.h
+++ b/Engine/Source/Runtime/Core/Public/Misc/NamePermissionList.h
@@ -357,7 +357,7 @@ public:
 	 * Result will contain AllowList paths that pass both filters.
 	 * @return new combined filter.
 	 */
-	CORE_API [[nodiscard]] FPathPermissionList CombinePathFilters(const FPathPermissionList& OtherFilter) const;
+	[[nodiscard]] CORE_API FPathPermissionList CombinePathFilters(const FPathPermissionList& OtherFilter) const;
 
 	/**
 	* Unregisters specified owners then adds specified filters in one operation (to avoid multiple filters changed events).
