Allow projects to opt out reporting loaded plugins

diff --git a/Engine/Source/Runtime/Core/Private/GenericPlatform/GenericPlatformCrashContext.cpp b/Engine/Source/Runtime/Core/Private/GenericPlatform/GenericPlatformCrashContext.cpp
--- a/Engine/Source/Runtime/Core/Private/GenericPlatform/GenericPlatformCrashContext.cpp
+++ b/Engine/Source/Runtime/Core/Private/GenericPlatform/GenericPlatformCrashContext.cpp
@@ -40,6 +40,11 @@
 #define CRASH_REPORTER_WITH_ANALYTICS 0
 #endif
 
+// Allow projects to opt out of reporting loaded plugins.
+#ifndef UE_CRASH_REPORTER_WITH_LOADED_PLUGINS
+#define UE_CRASH_REPORTER_WITH_LOADED_PLUGINS 1
+#endif
+
 DEFINE_LOG_CATEGORY_STATIC(LogCrashContext, Display, All);
 
 extern CORE_API bool GIsGPUCrashed;
@@ -545,6 +550,7 @@ void FGenericCrashContext::CopySharedCrashContext(FSharedCrashContext& Dst)
 	// -1 to allow space for null terminator
 	#define CR_DYNAMIC_BUFFER_REMAIN uint32((CR_MAX_DYNAMIC_BUFFER_CHARS) - (DynamicDataPtr-DynamicDataStart) - 1)
 
+#if UE_CRASH_REPORTER_WITH_LOADED_PLUGINS
 	Dst.EnabledPluginsOffset = (uint32)(DynamicDataPtr - DynamicDataStart);
 	Dst.EnabledPluginsNum = NCached::EnabledPluginsList.Num();
 	for (const FString& Plugin : NCached::EnabledPluginsList)
@@ -553,6 +559,10 @@ void FGenericCrashContext::CopySharedCrashContext(FSharedCrashContext& Dst)
 		FCString::Strncat(DynamicDataPtr, CR_PAIR_DELIM, CR_DYNAMIC_BUFFER_REMAIN);
 	}
 	DynamicDataPtr += FCString::Strlen(DynamicDataPtr) + 1;
+#else
+	Dst.EnabledPluginsOffset = (uint32)(DynamicDataPtr - DynamicDataStart);
+	Dst.EnabledPluginsNum = 0;
+#endif
 
 	Dst.EngineDataOffset = (uint32)(DynamicDataPtr - DynamicDataStart);
 	Dst.EngineDataNum = NCached::EngineData.Num();
