Add missing input string length checks in FGuid::ParseExact (that FGuid::Parse is doing prior, but otherwise can cause invalid index ensures when called directly)

diff --git a/Engine/Source/Runtime/Core/Private/Misc/Guid.cpp b/Engine/Source/Runtime/Core/Private/Misc/Guid.cpp
--- a/Engine/Source/Runtime/Core/Private/Misc/Guid.cpp
+++ b/Engine/Source/Runtime/Core/Private/Misc/Guid.cpp
@@ -333,9 +333,14 @@ bool FGuid::ParseExact(const FString& GuidString, EGuidFormats Format, FGuid& Ou
 	// The below normalizing into the Digits format doesn't work with Short and Base36 Guids
 	if (Format == EGuidFormats::Short)
 	{
+		if (GuidString.Len() != 22)
+		{
+			return false;
+		}
+
 		uint32 Data[4] = {};
 
-		// We don't need to do replacements to get an accurate size (defering allocations till later)
+		// We don't need to do replacements to get an accurate size (deferring allocations till later)
 		if (FBase64::GetDecodedDataSize(GuidString) != sizeof(Data))
 		{
 			// This isn't a Short GUID if it's not 128 bits / 16 bytes (the size of Data)
@@ -363,6 +368,11 @@ bool FGuid::ParseExact(const FString& GuidString, EGuidFormats Format, FGuid& Ou
 
 	if (Format == EGuidFormats::Base36Encoded)
 	{
+		if (GuidString.Len() != 25)
+		{
+			return false;
+		}
+
 		FUInt128 Value;
 		const TCHAR* c = *GuidString;
 
@@ -400,10 +410,20 @@ bool FGuid::ParseExact(const FString& GuidString, EGuidFormats Format, FGuid& Ou
 
 	if (Format == EGuidFormats::Digits || Format == EGuidFormats::DigitsLower)
 	{
+		if (GuidString.Len() != 32)
+		{
+			return false;
+		}
+
 		NormalizedGuidString = GuidString;
 	}
 	else if (Format == EGuidFormats::DigitsWithHyphens || Format == EGuidFormats::DigitsWithHyphensLower)
 	{
+		if (GuidString.Len() != 36)
+		{
+			return false;
+		}
+
 		if ((GuidString[8] != TCHAR('-')) ||
 			(GuidString[13] != TCHAR('-')) ||
 			(GuidString[18] != TCHAR('-')) ||
@@ -420,6 +440,11 @@ bool FGuid::ParseExact(const FString& GuidString, EGuidFormats Format, FGuid& Ou
 	}
 	else if (Format == EGuidFormats::DigitsWithHyphensInBraces)
 	{
+		if (GuidString.Len() != 38)
+		{
+			return false;
+		}
+
 		if ((GuidString[0] != TCHAR('{')) ||
 			(GuidString[9] != TCHAR('-')) ||
 			(GuidString[14] != TCHAR('-')) ||
@@ -438,6 +463,11 @@ bool FGuid::ParseExact(const FString& GuidString, EGuidFormats Format, FGuid& Ou
 	}
 	else if (Format == EGuidFormats::DigitsWithHyphensInParentheses)
 	{
+		if (GuidString.Len() != 38)
+		{
+			return false;
+		}
+
 		if ((GuidString[0] != TCHAR('(')) ||
 			(GuidString[9] != TCHAR('-')) ||
 			(GuidString[14] != TCHAR('-')) ||
@@ -456,6 +486,11 @@ bool FGuid::ParseExact(const FString& GuidString, EGuidFormats Format, FGuid& Ou
 	}
 	else if (Format == EGuidFormats::HexValuesInBraces)
 	{
+		if (GuidString.Len() != 68)
+		{
+			return false;
+		}
+
 		if ((GuidString[0] != TCHAR('{')) ||
 			(GuidString[1] != TCHAR('0')) ||
 			(GuidString[2] != TCHAR('x')) ||
@@ -510,6 +545,11 @@ bool FGuid::ParseExact(const FString& GuidString, EGuidFormats Format, FGuid& Ou
 	}
 	else if (Format == EGuidFormats::UniqueObjectGuid)
 	{
+		if (GuidString.Len() != 35)
+		{
+			return false;
+		}
+
 		if ((GuidString[8] != TCHAR('-')) ||
 			(GuidString[17] != TCHAR('-')) ||
 			(GuidString[26] != TCHAR('-')))
