Fix comment spelling
#rnx

diff --git a/Engine/Source/Runtime/Core/Public/Misc/ConfigCacheIni.h b/Engine/Source/Runtime/Core/Public/Misc/ConfigCacheIni.h
--- a/Engine/Source/Runtime/Core/Public/Misc/ConfigCacheIni.h
+++ b/Engine/Source/Runtime/Core/Public/Misc/ConfigCacheIni.h
@@ -62,7 +62,7 @@ CORE_API DECLARE_LOG_CATEGORY_EXTERN(LogConfig, Log, All);
 // make some additional code changes:
 //
 // Modifying:
-//    * Replace your direct modification with calls to SetSeting, SetBool, etc for non-array values
+//    * Replace your direct modification with calls to SetString, SetBool, etc for non-array values
 //    * Replace your direct modifications of array type values with AddToSection, AddUniqueToSection, RemoveKeyFromSection, RemoveFromSection
 //    * Fully construct a local new FConfigSection and then add that fully into the FConfigFile with Add
 //
