Always open around an hitch detecor as there is nothing we want to unwind here

diff --git a/Engine/Source/Runtime/Core/Private/Stats/StatsMisc.cpp b/Engine/Source/Runtime/Core/Private/Stats/StatsMisc.cpp
--- a/Engine/Source/Runtime/Core/Private/Stats/StatsMisc.cpp
+++ b/Engine/Source/Runtime/Core/Private/Stats/StatsMisc.cpp
@@ -4,12 +4,14 @@
 #include "Stats/Stats.h"
 #include "Logging/LogMacros.h"
 #include "CoreGlobals.h"
+#include "AutoRTFM/AutoRTFM.h"
 
 
 #if !STATS && !UE_BUILD_DEBUG && defined(USE_LIGHTWEIGHT_STATS_FOR_HITCH_DETECTION) && USE_LIGHTWEIGHT_STATS_FOR_HITCH_DETECTION && USE_HITCH_DETECTION
 #include "HAL/ThreadHeartBeat.h"
 #include "HAL/ThreadManager.h"
 
+UE_AUTORTFM_ALWAYS_OPEN
 void FLightweightStatScope::ReportHitch()
 {
 	if (StatString)
@@ -86,4 +88,4 @@ FString FConditionalScopeLogTime::GetDisplayUnitsString() const
 	}
 
 	return TEXT("ms");
-}
\ No newline at end of file
+}
