Silence bogus clang static analysis warnings

diff --git a/Engine/Source/Runtime/Core/Private/HAL/MallocBinned3.cpp b/Engine/Source/Runtime/Core/Private/HAL/MallocBinned3.cpp
--- a/Engine/Source/Runtime/Core/Private/HAL/MallocBinned3.cpp
+++ b/Engine/Source/Runtime/Core/Private/HAL/MallocBinned3.cpp
@@ -310,17 +310,20 @@ struct FMallocBinned3::Private
 	*/
 	static FPoolInfoSmall* GetOrCreatePoolInfoSmall(FMallocBinned3& Allocator, uint32 InPoolIndex, uint32 BlockOfBlocksIndex)
 	{
-		FPoolInfoSmall*& InfoBlock = Allocator.SmallPoolTables[InPoolIndex].PoolInfos[BlockOfBlocksIndex / Allocator.SmallPoolInfosPerPlatformPage];
+		const uint32 InfosPerPage = Allocator.SmallPoolInfosPerPlatformPage;
+		const uint32 InfoOuterIndex = BlockOfBlocksIndex / InfosPerPage;
+		const uint32 InfoInnerIndex = BlockOfBlocksIndex % InfosPerPage;
+		FPoolInfoSmall*& InfoBlock = Allocator.SmallPoolTables[InPoolIndex].PoolInfos[InfoOuterIndex];
 		if (!InfoBlock)
 		{
 			InfoBlock = (FPoolInfoSmall*)Allocator.AllocateMetaDataMemory(Allocator.OsAllocationGranularity);
 #if BINNED3_ALLOCATOR_STATS
 			Binned3PoolInfoMemory += Allocator.OsAllocationGranularity;
 #endif
-			DefaultConstructItems<FPoolInfoSmall>((void*)InfoBlock, Allocator.SmallPoolInfosPerPlatformPage);
+			DefaultConstructItems<FPoolInfoSmall>((void*)InfoBlock, InfosPerPage);
 		}
 
-		FPoolInfoSmall* Result = &InfoBlock[BlockOfBlocksIndex % Allocator.SmallPoolInfosPerPlatformPage];
+		FPoolInfoSmall* Result = &InfoBlock[InfoInnerIndex];
 
 		bool bGuaranteedToBeNew = false;
 		if (BlockOfBlocksIndex >= Allocator.SmallPoolTables[InPoolIndex].NumEverUsedBlockOfBlocks)
diff --git a/Engine/Source/Runtime/Core/Private/HAL/MallocBinnedGPU.cpp b/Engine/Source/Runtime/Core/Private/HAL/MallocBinnedGPU.cpp
--- a/Engine/Source/Runtime/Core/Private/HAL/MallocBinnedGPU.cpp
+++ b/Engine/Source/Runtime/Core/Private/HAL/MallocBinnedGPU.cpp
@@ -245,16 +245,19 @@ struct FMallocBinnedGPU::Private
 	*/
 	static FPoolInfoSmall* GetOrCreatePoolInfoSmall(FMallocBinnedGPU& Allocator, uint32 InPoolIndex, uint32 BlockOfBlocksIndex)
 	{
-		FPoolInfoSmall*& InfoBlock = Allocator.SmallPoolTables[InPoolIndex].PoolInfos[BlockOfBlocksIndex / Allocator.SmallPoolInfosPerPlatformPage];
+		const uint32 InfosPerPage = Allocator.SmallPoolInfosPerPlatformPage;
+		const uint32 InfoOuterIndex = BlockOfBlocksIndex / InfosPerPage;
+		const uint32 InfoInnerIndex = BlockOfBlocksIndex % InfosPerPage;
+		FPoolInfoSmall*& InfoBlock = Allocator.SmallPoolTables[InPoolIndex].PoolInfos[InfoOuterIndex];
 		if (!InfoBlock)
 		{
 			InfoBlock = (FPoolInfoSmall*)FMemory::Malloc(Allocator.ArenaParams.BasePageSize);
 			Allocator.MallocedPointers.Add(InfoBlock);
 			MBG_STAT(Allocator.BinnedGPUPoolInfoMemory += Allocator.ArenaParams.AllocationGranularity;)
-			DefaultConstructItems<FPoolInfoSmall>((void*)InfoBlock, Allocator.SmallPoolInfosPerPlatformPage);
+			DefaultConstructItems<FPoolInfoSmall>((void*)InfoBlock, InfosPerPage);
 		}
 
-		FPoolInfoSmall* Result = &InfoBlock[BlockOfBlocksIndex % Allocator.SmallPoolInfosPerPlatformPage];
+		FPoolInfoSmall* Result = &InfoBlock[InfoInnerIndex];
 
 		bool bGuaranteedToBeNew = false;
 		if (BlockOfBlocksIndex >= Allocator.SmallPoolTables[InPoolIndex].NumEverUsedBlockOfBlocks)
