Made TVariant trivially destructible if all the types are trivially destructible.

diff --git a/Engine/Source/Runtime/Core/Public/Misc/TVariant.h b/Engine/Source/Runtime/Core/Public/Misc/TVariant.h
--- a/Engine/Source/Runtime/Core/Public/Misc/TVariant.h
+++ b/Engine/Source/Runtime/Core/Public/Misc/TVariant.h
@@ -29,8 +29,19 @@ struct FEmptyVariantState
  * Attempting to use the value of a Get() when the underlying type is different leads to undefined behavior.
  */
 template <typename T, typename... Ts>
-class TVariant final : private UE::Core::Private::TVariantStorage<T, Ts...>
+class TVariant final
+#if UE_TVARIANT_TRIVIAL_DESTRUCTOR_USING_CONCEPTS
+	: private UE::Core::Private::TVariantStorage<T, Ts...>
+#else
+	: private std::conditional_t<!std::is_trivially_destructible_v<T> || (!std::is_trivially_destructible_v<Ts> || ...), UE::Core::Private::TDestructibleVariantStorage<T, Ts...> , UE::Core::Private::TVariantStorage<T, Ts...>>
+#endif
 {
+#if UE_TVARIANT_TRIVIAL_DESTRUCTOR_USING_CONCEPTS
+	using Super = UE::Core::Private::TVariantStorage<T, Ts...>;
+#else
+	using Super = std::conditional_t<!std::is_trivially_destructible_v<T> || (!std::is_trivially_destructible_v<Ts> || ...), UE::Core::Private::TDestructibleVariantStorage<T, Ts...> , UE::Core::Private::TVariantStorage<T, Ts...>>;
+#endif
+
 	static_assert(!UE::Core::Private::TTypePackContainsDuplicates<T, Ts...>::Value, "All the types used in TVariant should be unique");
 	static_assert(!UE::Core::Private::TContainsReferenceType<T, Ts...>::Value, "TVariant cannot hold reference types");
 
@@ -59,15 +70,15 @@ public:
 
 	/** Copy construct the variant from another variant of the same type */
 	TVariant(const TVariant& Other)
-		: TypeIndex(Other.TypeIndex)
 	{
+		TypeIndex = Other.TypeIndex;
 		UE::Core::Private::TCopyConstructorLookup<T, Ts...>::Construct(TypeIndex, &UE::Core::Private::CastToStorage(*this).Storage, &UE::Core::Private::CastToStorage(Other).Storage);
 	}
 
 	/** Move construct the variant from another variant of the same type */
 	TVariant(TVariant&& Other)
-		: TypeIndex(Other.TypeIndex)
 	{
+		TypeIndex = Other.TypeIndex;
 		UE::Core::Private::TMoveConstructorLookup<T, Ts...>::Construct(TypeIndex, &UE::Core::Private::CastToStorage(*this).Storage, &UE::Core::Private::CastToStorage(Other).Storage);
 	}
 
@@ -93,11 +104,20 @@ public:
 		return *this;
 	}
 
+#if UE_TVARIANT_TRIVIAL_DESTRUCTOR_USING_CONCEPTS
 	/** Destruct the underlying type (if appropriate) */
 	~TVariant()
+		requires(!std::is_trivially_destructible_v<T> || (!std::is_trivially_destructible_v<Ts> || ...))
 	{
 		UE::Core::Private::TDestructorLookup<T, Ts...>::Destruct(TypeIndex, &UE::Core::Private::CastToStorage(*this).Storage);
 	}
+	~TVariant()
+		requires(std::is_trivially_destructible_v<T> && (std::is_trivially_destructible_v<Ts> && ...))
+	= default;
+#else
+	// Defer to the storage as to how to destruct the elements
+	~TVariant() = default;
+#endif
 
 	/** Determine if the variant holds the specific type */
 	template <typename U>
@@ -189,8 +209,12 @@ public:
 	}
 
 private:
+#if UE_TVARIANT_TRIVIAL_DESTRUCTOR_USING_CONCEPTS
 	/** Index into the template parameter pack for the type held. */
 	uint8 TypeIndex;
+#else
+	using Super::TypeIndex;
+#endif
 };
 
 /** Apply a visitor function to the list of variants */
diff --git a/Engine/Source/Runtime/Core/Public/Misc/TVariantMeta.h b/Engine/Source/Runtime/Core/Public/Misc/TVariantMeta.h
--- a/Engine/Source/Runtime/Core/Public/Misc/TVariantMeta.h
+++ b/Engine/Source/Runtime/Core/Public/Misc/TVariantMeta.h
@@ -11,6 +11,19 @@
 
 #include "Misc/AssertionMacros.h"
 
+// Disable the concept checks on MSVC until this internal compile error gets fixed:
+//
+// https://developercommunity.visualstudio.com/t/Internal-compile-error-when-returning-a/10650542
+#ifndef _MSC_VER
+	#if defined(__cpp_concepts) && __cpp_concepts >= 202002L
+		#define UE_TVARIANT_TRIVIAL_DESTRUCTOR_USING_CONCEPTS 1
+	#endif
+#endif
+
+#if !defined(UE_TVARIANT_TRIVIAL_DESTRUCTOR_USING_CONCEPTS)
+	#define UE_TVARIANT_TRIVIAL_DESTRUCTOR_USING_CONCEPTS 0
+#endif
+
 template <typename T, typename... Ts>
 class TVariant;
 
@@ -156,6 +169,11 @@ namespace Private
 		}
 
 		TAlignedBytes<SizeofValue, AlignofValue> Storage;
+
+#if !UE_TVARIANT_TRIVIAL_DESTRUCTOR_USING_CONCEPTS
+		/** Index into the template parameter pack for the type held. */
+		uint8 TypeIndex;
+#endif
 	};
 
 	/** Helper to lookup indices of each type in a template parameter pack */
@@ -377,6 +395,19 @@ namespace Private
 		static constexpr InvokeFn Invokers[] = { WrapperType::template FuncPtr<EncodedIndices>... };
 		return Invokers[EncodedIndex](Forward<Func>(Callable), Forward<Variants>(Args)...);
 	}
+
+#if !UE_TVARIANT_TRIVIAL_DESTRUCTOR_USING_CONCEPTS
+	template <typename T, typename... Ts>
+	struct TDestructibleVariantStorage : TVariantStorage<T, Ts...>
+	{
+	protected:
+		~TDestructibleVariantStorage()
+		{
+			UE::Core::Private::TDestructorLookup<T, Ts...>::Destruct(this->TypeIndex, this);
+		}
+	};
+#endif
+
 } // namespace Private
 } // namespace Core
 } // namespace UE
diff --git a/Engine/Source/Runtime/Core/Tests/Misc/TVariantTest.cpp b/Engine/Source/Runtime/Core/Tests/Misc/TVariantTest.cpp
--- a/Engine/Source/Runtime/Core/Tests/Misc/TVariantTest.cpp
+++ b/Engine/Source/Runtime/Core/Tests/Misc/TVariantTest.cpp
@@ -31,6 +31,14 @@ namespace UE::Core::Private
 	static_assert(TVariantSize_V<const TVariant<int, float>&>  == 0, "Expected TVariantSize_V<const TVariant&> to be 0");
 	static_assert(TVariantSize_V<      TVariant<int, float>&&> == 0, "Expected TVariantSize_V<TVariant&&> to be 0");
 	static_assert(TVariantSize_V<      bool>                   == 0, "Expected TVariantSize_V<non-TVariant> to be 0");
+
+	// Ensure that TVariant has a trivial destructor if all the types are
+	static_assert( std::is_trivially_destructible_v<TVariant<int>>);
+	static_assert( std::is_trivially_destructible_v<TVariant<int, FVector>>);
+	static_assert( std::is_trivially_destructible_v<TVariant<int, FVector, FBox>>);
+	static_assert(!std::is_trivially_destructible_v<TVariant<FString>>);
+	static_assert(!std::is_trivially_destructible_v<TVariant<int, FString>>);
+	static_assert(!std::is_trivially_destructible_v<TVariant<int, FVector, FString>>);
 } // namespace UE::Core::Private
 
 struct FNonDefaultConstructible
