PR #11788: fix: clang coroutine support when building with C++17

diff --git a/Engine/Source/Programs/UnrealBuildTool/Platform/Windows/VCToolChain.cs b/Engine/Source/Programs/UnrealBuildTool/Platform/Windows/VCToolChain.cs
--- a/Engine/Source/Programs/UnrealBuildTool/Platform/Windows/VCToolChain.cs
+++ b/Engine/Source/Programs/UnrealBuildTool/Platform/Windows/VCToolChain.cs
@@ -1136,7 +1136,10 @@ namespace UnrealBuildTool
 				}
 				else if (Target.WindowsPlatform.Compiler.IsClang() && !Target.WindowsPlatform.Compiler.IsIntel())
 				{
-					Arguments.Add("-fcoroutines-ts");
+					if (CompileEnvironment.CppStandard < CppStandardVersion.Cpp20)
+					{
+						Arguments.Add("-fcoroutines-ts");
+					}
 				}
 			}
 
diff --git a/Engine/Source/Programs/UnrealBuildTool/ToolChain/ClangToolChain.cs b/Engine/Source/Programs/UnrealBuildTool/ToolChain/ClangToolChain.cs
--- a/Engine/Source/Programs/UnrealBuildTool/ToolChain/ClangToolChain.cs
+++ b/Engine/Source/Programs/UnrealBuildTool/ToolChain/ClangToolChain.cs
@@ -372,7 +372,11 @@ namespace UnrealBuildTool
 
 			if (CompileEnvironment.bEnableCoroutines)
 			{
-				Arguments.Add("-fcoroutines-ts");
+				if (CompileEnvironment.CppStandard < CppStandardVersion.Cpp20)
+				{
+					Arguments.Add("-fcoroutines-ts");
+				}
+
 				if (!CompileEnvironment.bEnableExceptions)
 				{
 					Arguments.Add("-Wno-coroutine-missing-unhandled-exception");
diff --git a/Engine/Source/Runtime/Core/Public/Experimental/Coroutine/Coroutine.h b/Engine/Source/Runtime/Core/Public/Experimental/Coroutine/Coroutine.h
--- a/Engine/Source/Runtime/Core/Public/Experimental/Coroutine/Coroutine.h
+++ b/Engine/Source/Runtime/Core/Public/Experimental/Coroutine/Coroutine.h
@@ -291,7 +291,7 @@ namespace CoroTask_Detail
 			if(!CoroutineHandle.done())
 			{
 				TraceId = TaskTrace::GenerateTaskId();
-				TaskTrace::Launched(TraceId, GetDebugName(), true, ENamedThreads::Type::AnyThread);
+				TaskTrace::Launched(TraceId, GetDebugName(), true, ENamedThreads::Type::AnyThread, 0);
 				TaskTrace::SubsequentAdded(TraceId, OldTraceId);
 				if(Prerequisite != nullptr)
 				{
@@ -311,7 +311,7 @@ namespace CoroTask_Detail
 		{
 			coroCheck(Prerequisite == nullptr);
 			IncrementRefCount();
-			COROTASKTRACE(TaskTrace::Launched(TraceId, DebugName, true, ENamedThreads::Type::AnyThread));
+			COROTASKTRACE(TaskTrace::Launched(TraceId, DebugName, true, ENamedThreads::Type::AnyThread, 0));
 			Task.Init(DebugName, Priority, FTaskExecutor(this), Flags);
 			COROTASKTRACE(TaskTrace::Scheduled(TraceId));
 			coroVerify(LowLevelTasks::TryLaunch(Task, QueuePreference));
@@ -322,7 +322,7 @@ namespace CoroTask_Detail
 			ensureMsgf(bAllowNested || !FCoroLocalState::IsCoroLaunchedTask(), TEXT("Unconnected Expedition within a CoroTask, connect the Callstacks for better Performance."));
 
 			coroCheck(IsExpeditable());
-			if (!Task.TryCancel(LowLevelTasks::ECancellationFlags::PrelaunchCancellation)) 
+			if (!Task.TryCancel(LowLevelTasks::ECancellationFlags::PrelaunchCancellation))
 			{
 				return CoroutineHandle.done();
 			}
@@ -386,7 +386,7 @@ namespace CoroTask_Detail
 
 		COROFORCEINLINE bool InvokeOnceAndLaunch(const TCHAR* DebugName = nullptr, LowLevelTasks::ETaskPriority Priority = LowLevelTasks::ETaskPriority::Inherit, LowLevelTasks::ETaskFlags Flags = LowLevelTasks::ETaskFlags::DefaultFlags)
 		{
-			COROTASKTRACE(TaskTrace::Launched(TraceId, DebugName, true, ENamedThreads::Type::AnyThread));
+			COROTASKTRACE(TaskTrace::Launched(TraceId, DebugName, true, ENamedThreads::Type::AnyThread, 0));
 			COROTASKTRACE(TaskTrace::Scheduled(TraceId));
 			Execute();
 			if(!CoroutineHandle.done())
