Fixing compiler error when using TNonNullPtr constructor with convertible object types.

#rnx

diff --git a/Engine/Source/Runtime/Core/Public/Templates/NonNullPointer.h b/Engine/Source/Runtime/Core/Public/Templates/NonNullPointer.h
--- a/Engine/Source/Runtime/Core/Public/Templates/NonNullPointer.h
+++ b/Engine/Source/Runtime/Core/Public/Templates/NonNullPointer.h
@@ -55,7 +55,7 @@ public:
 		UE_REQUIRES(std::is_convertible_v<OtherObjectType*, ObjectType*>)
 	>
 	FORCEINLINE TNonNullPtr(const TNonNullPtr<OtherObjectType>& Other)
-		: Object(Other.Object)
+		: Object(Other.Get())
 	{
 	}
 
@@ -88,7 +88,7 @@ public:
 	>
 	FORCEINLINE TNonNullPtr& operator=(const TNonNullPtr<OtherObjectType>& Other)
 	{
-		Object = Other.Object;
+		Object = Other.Get();
 		return *this;
 	}
 
