Make Shrink at the end of filling from buffer do what it originally intended to do
#rnx

diff --git a/Engine/Source/Runtime/Core/Private/Misc/ConfigCacheIni.cpp b/Engine/Source/Runtime/Core/Private/Misc/ConfigCacheIni.cpp
--- a/Engine/Source/Runtime/Core/Private/Misc/ConfigCacheIni.cpp
+++ b/Engine/Source/Runtime/Core/Private/Misc/ConfigCacheIni.cpp
@@ -1055,6 +1055,21 @@ bool FConfigFile::Combine(const FString& Filename)
 	return FillFileFromDisk(Filename, true);
 }
 
+void FConfigFile::Shrink()
+{
+	FConfigFileMap::Shrink();
+	for (FConfigFileMap::TIterator It(*this); It; ++It)
+	{
+		It.Value().Shrink();
+	}
+
+	PerObjectConfigArrayOfStructKeys.Shrink();
+	for (auto& Pair : PerObjectConfigArrayOfStructKeys)
+	{
+		Pair.Value.Shrink();
+	}
+}
+
 // Assumes GetTypeHash(AltKeyType) matches GetTypeHash(KeyType)
 template<class KeyType, class ValueType, class AltKeyType>
 ValueType& FindOrAddHeterogeneous(TMap<KeyType, ValueType>& Map, const AltKeyType& Key)
@@ -1532,11 +1547,6 @@ void FillFileFromBuffer(FileType* File, FStringView Buffer, bool bHandleSymbolCo
 
 	// Avoid memory wasted in array slack.
 	File->Shrink();
-	for (auto Pair : AsConst(*File))
-	{
-		Pair.Value.Shrink();
-	}
-
 }
 
 template<typename FileType>
@@ -2763,6 +2773,20 @@ FConfigCommandStreamSection* FConfigCommandStream::FindOrAddSectionInternal(cons
 	return &FindOrAdd(SectionName);
 }
 
+void FConfigCommandStream::Shrink()
+{
+	TMap<FString, FConfigCommandStreamSection>::Shrink();
+	for (auto& Pair : *this)
+	{
+		Pair.Value.Shrink();
+	}
+
+	PerObjectConfigArrayOfStructKeys.Shrink();
+	for (auto& Pair : PerObjectConfigArrayOfStructKeys)
+	{
+		Pair.Value.Shrink();
+	}
+}
 
 
 
diff --git a/Engine/Source/Runtime/Core/Public/Misc/ConfigCacheIni.h b/Engine/Source/Runtime/Core/Public/Misc/ConfigCacheIni.h
--- a/Engine/Source/Runtime/Core/Public/Misc/ConfigCacheIni.h
+++ b/Engine/Source/Runtime/Core/Public/Misc/ConfigCacheIni.h
@@ -510,6 +510,8 @@ public:
 	FConfigCommandStreamSection* FindOrAddSectionInternal(const FString& SectionName);
 	bool FillFileFromDisk(const FString& Filename, bool bHandleSymbolCommands);
 
+	void Shrink();
+
 	// This holds per-object config class names, with their ArrayOfStructKeys. Since the POC sections are all unique,
 	// we can't track it just in that section. This is expected to be empty/small
 	TMap<FString, TMap<FName, FString> > PerObjectConfigArrayOfStructKeys;
@@ -519,7 +521,7 @@ public:
 	
 	uint8 Dirty;
 
-	// used to determine if existing settings should be removed from a FCOnfigFile before applying this stream - used for compatibility with
+	// used to determine if existing settings should be removed from a FConfigFile before applying this stream - used for compatibility with
 	// how Saved config files are stored (replace the static layers values fully if at least one key exists)
 	uint8 bIsSavedConfigFile = false;
 	
@@ -691,6 +693,8 @@ private:
 	 */
 	void Cleanup();
 
+	void Shrink();
+
 	void FillFileFromBuffer(FStringView Buffer, bool bHandleSymbolCommands, const FString& FileHint);
 	bool FillFileFromDisk(const FString& Filename, bool bHandleSymbolCommands);
 	void ProcessCommand(FConfigSection* Section, FStringView SectionName, FConfigValue::EValueType Command, FName Key, FString&& Value);
