Submit Josh's config memusage functionality
#rnx

diff --git a/Engine/Source/Runtime/Core/Private/Misc/ConfigCacheIni.cpp b/Engine/Source/Runtime/Core/Private/Misc/ConfigCacheIni.cpp
--- a/Engine/Source/Runtime/Core/Private/Misc/ConfigCacheIni.cpp
+++ b/Engine/Source/Runtime/Core/Private/Misc/ConfigCacheIni.cpp
@@ -3189,7 +3189,39 @@ FConfigCacheIni::~FConfigCacheIni()
 	Flush( 1 );
 }
 
+FConfigBranch* FConfigCacheIni::FindBranchWithNoReload(FName BaseIniName, const FString& Filename)
+{
+	// look for a known file, if there's no ini extension
+	FConfigBranch* Branch = KnownFiles.GetBranch(BaseIniName);
+
+	if (Branch == nullptr)
+	{
+		Branch = KnownFiles.GetBranch(*Filename);
+	}
+	if (Branch == nullptr)
+	{
+		Branch = OtherFiles.FindRef(Filename);
+		if (Branch == nullptr)
+		{
+			for (TPair<FString, FConfigBranch*>& CurrentFilePair : OtherFiles)
+			{
+				if (CurrentFilePair.Value->IniName == BaseIniName)
+				{
+					Branch = CurrentFilePair.Value;
+					break;
+				}
+			}
+		}
+	}
+
+	// if Filename is a .ini and it doesn't match what the KnownFile has (if it has one yet), then we can't use it
+	if (Branch && Branch->IniPath.Len() > 0 && Filename.Len() > 0 && Filename.EndsWith(".ini") && Branch->IniPath != Filename)
+	{
+		Branch = nullptr;
+	}
 
+	return Branch;
+}
 	
 FConfigBranch* FConfigCacheIni::FindBranch(FName BaseIniName, const FString& Filename)
 {
@@ -6010,6 +6042,54 @@ class FIniExec : public FSelfRegisteringExec
 			Ar.Logf(TEXT("INITIME : PrepareForLoad: %fms, PreformLoad: %fms"), GPrepareForLoadTime * 1000.0, GPerformLoadTime * 1000.0);
 		}
 		
+		if (FParse::Command(&Cmd, TEXT("MemUsage")))
+		{
+			uint64 Total = 0;
+			int NumSkipped = 0;
+			uint64 SkippedTotal = 0;
+			int SingleSection = 0;
+			uint64 SingleSectionTotal = 0;
+			int NoSection = 0;
+			uint64 NoSectionTotal = 0;
+			for (const FString& Filename : GConfig->GetFilenames())
+			{
+				FArchiveCountConfigMem MemAr;
+
+				FConfigBranch* Branch = GConfig->FindBranchWithNoReload(*Filename, Filename);
+				MemAr << *Branch;
+
+				uint64 Mem = MemAr.GetMax();
+				Total += Mem;
+
+				if (Branch->InMemoryFile.Num() == 1)
+				{
+					SingleSection++;
+					SingleSectionTotal += Mem;
+				}
+				if (Branch->InMemoryFile.Num() == 0)
+				{
+					NoSection++;
+					NoSectionTotal += Mem;
+				}
+
+				// don't bother printing the neglibly sized ones as they are just noise, so cut off anything < 10kb
+				if (Mem < 10 * 1024)
+				{
+					NumSkipped++;
+					SkippedTotal += Mem;
+				}
+				else
+				{
+					Ar.Logf(TEXT("[%0.2fmb] - %s"), (double)MemAr.GetMax() / 1024.0 / 1024.0, *Filename);
+				}
+			}
+
+			Ar.Logf(TEXT("[%0.2fmb] - %d All Configs"), (double)Total / 1024.0 / 1024.0, GConfig->GetFilenames().Num());
+			Ar.Logf(TEXT("[%0.2fmb] - %d Tiny Configs (not displayed)"), (double)SkippedTotal / 1024.0 / 1024.0, NumSkipped);
+			Ar.Logf(TEXT("[%0.2fmb] - %d Single Section Configs"), (double)SingleSectionTotal / 1024.0 / 1024.0, SingleSection);
+			Ar.Logf(TEXT("[%0.2fmb] - %d ZeroSection Configs"), (double)NoSectionTotal / 1024.0 / 1024.0, NoSection);
+		}
+
 		return true;
 	}
 	
diff --git a/Engine/Source/Runtime/Core/Public/Misc/ConfigCacheIni.h b/Engine/Source/Runtime/Core/Public/Misc/ConfigCacheIni.h
--- a/Engine/Source/Runtime/Core/Public/Misc/ConfigCacheIni.h
+++ b/Engine/Source/Runtime/Core/Public/Misc/ConfigCacheIni.h
@@ -1126,6 +1126,9 @@ public:
     /** Finds an FConfigBranch, using base name or a filename (can pass in NAME_None or empty string as needed) */
     CORE_API FConfigBranch* FindBranch(FName BaseIniName, const FString& Filename);
 
+	/** Same as FConfigBRanch, but this will not reload a formerly SafeUnload'd branch */
+	CORE_API FConfigBranch* FindBranchWithNoReload(FName BaseIniName, const FString& Filename);
+
     /** Create a new branch for FIlename, and return it */
 	CORE_API FConfigBranch& AddNewBranch(const FString& Filename);
 	
