Always open around finding modules in ModuleManager

diff --git a/Engine/Source/Runtime/Core/Private/Modules/ModuleManager.cpp b/Engine/Source/Runtime/Core/Private/Modules/ModuleManager.cpp
--- a/Engine/Source/Runtime/Core/Private/Modules/ModuleManager.cpp
+++ b/Engine/Source/Runtime/Core/Private/Modules/ModuleManager.cpp
@@ -20,6 +20,7 @@
 #include "Stats/Stats.h"
 #include "Trace/Trace.h"
 #include "Trace/Trace.inl"
+#include "AutoRTFM/AutoRTFM.h"
 
 DEFINE_LOG_CATEGORY_STATIC(LogModuleManager, Log, All);
 
@@ -47,6 +48,7 @@ void FModuleManager::WarnIfItWasntSafeToLoadHere(const FName InModuleName)
 	}
 }
 
+UE_AUTORTFM_ALWAYS_OPEN
 FModuleManager::ModuleInfoPtr FModuleManager::FindModule(FName InModuleName)
 {
 	FModuleManager::ModuleInfoPtr Result = nullptr;
