Cherrypick of 33033713. In addition to the below description, this CL also bumps MaterialTranslationDDCVersion in FFortniteReleaseBranchCustomObjectVersion to ensure that the material translation cache is invalidated and picks up on the changed code gen.

Fix for invalid ddx/ddy calls in generated material code when using the DDX and DDY nodes in a material that can be executed in a compute shader context other than Nanite shading (such as HeterogeneousVolumes or Volumetric Clouds).
Fixed by using existing DDX() and DDY() wrapper functions which are already set up to return 0 in all shader types except pixel shaders and nanite compute shading compute shaders.

#rnx

diff --git a/Engine/Source/Runtime/Core/Private/UObject/FortniteReleaseBranchCustomObjectVersion.cpp b/Engine/Source/Runtime/Core/Private/UObject/FortniteReleaseBranchCustomObjectVersion.cpp
--- a/Engine/Source/Runtime/Core/Private/UObject/FortniteReleaseBranchCustomObjectVersion.cpp
+++ b/Engine/Source/Runtime/Core/Private/UObject/FortniteReleaseBranchCustomObjectVersion.cpp
@@ -8,6 +8,6 @@ TMap<FGuid, FGuid> FFortniteReleaseBranchCustomObjectVersion::GetSystemGuids()
 
 	SystemGuids.Add(DevGuids.Niagara_LatestScriptCompileVersion, FGuid("5F589BAA37394A51BF718ECB10D73C19"));
 	SystemGuids.Add(DevGuids.SkeletalMeshDerivedDataVersion, FGuid("5705956EC7134274A5A1BD3FC74DB3A9"));
-	SystemGuids.Add(DevGuids.MaterialTranslationDDCVersion, FGuid("555395300FC34E35A49DB477FC8ADD85"));
+	SystemGuids.Add(DevGuids.MaterialTranslationDDCVersion, FGuid("CC1B3A336E144A0B9E78014E07775988"));
 	return SystemGuids;
 }
