Add NaN check to Matrix2x2

diff --git a/Engine/Source/Runtime/Core/Public/Math/TransformCalculus2D.h b/Engine/Source/Runtime/Core/Public/Math/TransformCalculus2D.h
--- a/Engine/Source/Runtime/Core/Public/Math/TransformCalculus2D.h
+++ b/Engine/Source/Runtime/Core/Public/Math/TransformCalculus2D.h
@@ -413,6 +413,7 @@ public:
 	{
 		M[0][0] = m00; M[0][1] = m01;
 		M[1][0] = m10; M[1][1] = m11;
+		DiagnosticCheckNaN();
 	}
 
 
@@ -421,6 +422,7 @@ public:
 	{
 		M[0][0] = UniformScale; M[0][1] = 0;
 		M[1][0] = 0; M[1][1] = UniformScale;
+		DiagnosticCheckNaN();
 	}
 
 	/** Ctor. initialize from a scale. */
@@ -430,6 +432,7 @@ public:
 		T ScaleY = (T)Scale.GetVector().Y;
 		M[0][0] = ScaleX; M[0][1] = 0;
 		M[1][0] = 0; M[1][1] = ScaleY;
+		DiagnosticCheckNaN();
 	}
 
 	/** Factory function. initialize from a 2D shear. */
@@ -438,7 +441,8 @@ public:
 		T XX = (T)Shear.GetVector().X;
 		T YY = (T)Shear.GetVector().Y;
 		M[0][0] = 1; M[0][1] =YY;
-		M[1][0] =XX; M[1][1] = 1;
+		M[1][0] = XX; M[1][1] = 1;
+		DiagnosticCheckNaN();
 	}
 
 	/** Ctor. initialize from a rotation. */
@@ -448,6 +452,7 @@ public:
 		T SinAngle = (T)Rotation.GetVector().Y;
 		M[0][0] = CosAngle; M[0][1] = SinAngle;
 		M[1][0] = -SinAngle; M[1][1] = CosAngle;
+		DiagnosticCheckNaN();
 	}
 
 	/**
@@ -481,9 +486,11 @@ public:
 		GetMatrix(A, B, C, D);
 		T E, F, G, H;
 		RHS.GetMatrix(E, F, G, H);
-		return TMatrix2x2(
+		TMatrix2x2 Result = TMatrix2x2(
 			A*E + B*G, A*F + B*H,
 			C*E + D*G, C*F + D*H);
+		Result.DiagnosticCheckNaN();
+		return Result;
 	}
 	/**
 	 * Invert the transform.
@@ -493,9 +500,11 @@ public:
 		T A, B, C, D;
 		GetMatrix(A, B, C, D);
 		T InvDet = InverseDeterminant();
-		return TMatrix2x2(
+		TMatrix2x2 Result = TMatrix2x2(
 			 D*InvDet, -B*InvDet,
 			-C*InvDet,  A*InvDet);
+		Result.DiagnosticCheckNaN();
+		return Result;
 	}
 
 	/** Equality. */
@@ -583,6 +592,27 @@ public:
 			FMath::IsNearlyEqual(M[1][1], 1.0f, ErrorTolerance);
 	}
 
+	bool ContainsNaN() const
+	{
+		return (!FMath::IsFinite(M[0][0]) ||
+				!FMath::IsFinite(M[0][1]) ||
+				!FMath::IsFinite(M[1][0]) ||
+				!FMath::IsFinite(M[1][1]));
+	}
+
+#if ENABLE_NAN_DIAGNOSTIC
+	FORCEINLINE void DiagnosticCheckNaN() const
+	{
+		if (ContainsNaN())
+		{
+			logOrEnsureNanError(TEXT("Matrix2x2 contains NaN"));
+			*const_cast<TMatrix2x2<T>*>(static_cast<const TMatrix2x2<T>*>(this)) = TMatrix2x2<T>(ForceInitToZero);
+		}
+	}
+#else
+	FORCEINLINE void DiagnosticCheckNaN() const {}
+#endif
+
 private:
 	T M[2][2];
 };
@@ -805,6 +835,12 @@ public:
 		);
 	}
 
+	/** Utility to check if there are any non-finite values (NaN or Inf) in this transform. */
+	bool ContainsNaN() const
+	{
+		return M.ContainsNaN() || Trans.ContainsNaN();
+	}
+
 private:
 	Matrix2Type M;
 	Vector2Type Trans;
