Deprecate the 'LazyInit' mode for virtualization systems.

#rnx

- The original use case for this feature was to avoid connecting to external services that VA might rely on (source control, horde etc) either until later on in the start up of the process or until they are actually needed. However since then we have added the option of deferring the connection to these services as a specific feature which allows for the main system to start up early on (so we can query things like is VA enabled on the project etc) without risking the expensive connections.
- Due to changes with telemetry in the last year, VA systems marked as 'LazyInit' will be initialized quite early on in the process start up making the feature pretty pointless anyway.
- If anyone is still using this feature they probably would've complained about it being broken by now and it doesn't really offer any compelling advantage over the newer features, so we might as well deprecate it rather than fix it.

diff --git a/Engine/Source/Runtime/Core/Private/Virtualization/VirtualizationSystem.cpp b/Engine/Source/Runtime/Core/Private/Virtualization/VirtualizationSystem.cpp
--- a/Engine/Source/Runtime/Core/Private/Virtualization/VirtualizationSystem.cpp
+++ b/Engine/Source/Runtime/Core/Private/Virtualization/VirtualizationSystem.cpp
@@ -19,6 +19,8 @@
 // the virtualization system to be a lazy initialization.
 #ifndef UE_VIRTUALIZATION_SYSTEM_LAZY_INIT
 	#define UE_VIRTUALIZATION_SYSTEM_LAZY_INIT 0
+#else
+	#pragma message("WARNING: Use of UE_VIRTUALIZATION_SYSTEM_LAZY_INIT is deprecated as of UE 5.5, consider replacing with 'ini:Engine:[Core.VirtualizationModule]:LazyInitConnections=true'")
 #endif //UE_VIRTUALIZATION_SYSTEM_LAZY_INIT
 
 // When enabled we will log if FNullVirtualizationSystem tries to push or pull payloads
@@ -297,6 +299,7 @@ void Initialize(const FInitParams& InitParams, EInitializationFlags Flags)
 	{
 		if (ShouldLazyInitializeSystem(InitParams.ConfigFile))
 		{
+			UE_LOG(LogVirtualization, Warning, TEXT("The 'LazyInit' feature is deprecated as of UE 5.5, consider replacing with 'ini:Engine:[Core.VirtualizationModule]:LazyInitConnections=true'"));
 			return;
 		}
 	}
