FConfigBranch Phase 2:
- Making use of unload configs in PluginManager, GameFeaturePlugins, and HotFixes
- To help with memory increase from the branch memory for unloadable configs, I added FConfigCacheIni::SafeUnload(BranchName), which can unload the memory usage by a branch, but then will on-demand reload the hierarchy if accessed
- Used SafeUnload to GGameplayTagsIni because that is read once and cached
- Various other minor optimizations, like not creating hundreds of empty configs in GConfig,
- Added "config memusage" (similar to configmem as it turns out)
- Removed some old deprecated code
- Added ability for editor to remember and write out comments in .ini files
- Changed commandline ini overrides to a dynamic layer (this is a work in progress - old method is still in there - it's tied to the ini.UseNewDynamicLayers cvar - old GFP code used the old way, etc, which i still support)
- Experimental support for runtime config change tracking, so saving out .ini files will have more context about how the user changed a config section (mostly important for arrays). This is not fully functional yet [WILL PROBABLY REMOVE]
- Experimental new config layer expansion system that uses FPaths::ConvertPath to find Restricted/Platform locations, instead of hardcoding

diff --git a/Engine/Plugins/Online/OnlineFramework/Source/Hotfix/Private/OnlineHotfixManager.cpp b/Engine/Plugins/Online/OnlineFramework/Source/Hotfix/Private/OnlineHotfixManager.cpp
--- a/Engine/Plugins/Online/OnlineFramework/Source/Hotfix/Private/OnlineHotfixManager.cpp
+++ b/Engine/Plugins/Online/OnlineFramework/Source/Hotfix/Private/OnlineHotfixManager.cpp
@@ -926,6 +926,24 @@ FString UOnlineHotfixManager::BuildConfigCacheKey(const FString& IniName)
 	return GConfig->GetConfigFilename(*IniNameNoExtension);
 }
 
+FConfigBranch* UOnlineHotfixManager::GetBranch(const FString& IniName)
+{
+	const FString StrippedIniName(GetStrippedConfigFileName(IniName));
+	const FString StrippedIniNameNoExtension = FPaths::GetBaseFilename(StrippedIniName);
+
+	// find the branch by basename or full filename
+	FConfigBranch* Branch = GConfig->FindBranch(*StrippedIniName, IniName);
+	if (Branch == nullptr)
+	{
+		// does this really happen, seems pretty unexpected
+		UE_LOG(LogHotfixManager, Warning, TEXT("No config braanch found, creating new branch for hotfix %s"), *IniName);
+		
+		Branch = &GConfig->AddNewBranch(IniName);
+	}
+	
+	return Branch;
+}
+
 FConfigFile* UOnlineHotfixManager::GetConfigFile(const FString& IniName)
 {
 	const FString StrippedIniName(GetStrippedConfigFileName(IniName));
@@ -971,6 +989,20 @@ bool UOnlineHotfixManager::HotfixIniFile(const FString& FileName, const FString&
 {
 	// Flush async loading before modifying GConfig.
 	FlushAsyncLoading();
+	
+	static bool bUseNewDynamicLayers = IConsoleManager::Get().FindConsoleVariable(TEXT("ini.UseNewDynamicLayers"))->GetInt() != 0;
+	if (bUseNewDynamicLayers)
+	{
+		FConfigBranch* Branch = GetBranch(FileName);
+		FConfigModificationTracker ChangeTracker;
+		
+		FName Tag = *BuildConfigCacheKey(FileName);
+		Branch->AddDynamicLayerStringToHierarchy(FileName, IniData, Tag, DynamicLayerPriority::Hotfix, &ChangeTracker);
+		
+		UObjectBaseUtility::ReloadObjectsFromModifiedConfigSections(ChangeTracker.ModifiedSections, FileName);
+
+		return true;
+	}
 
 	FConfigFile* ConfigFile = GetConfigFile(FileName);
 	// Store the original file so we can undo this later
@@ -1198,6 +1230,39 @@ UOnlineHotfixManager::FConfigFileBackup& UOnlineHotfixManager::BackupIniFile(con
 
 void UOnlineHotfixManager::RestoreBackupIniFiles()
 {
+	static bool bUseNewDynamicLayers = IConsoleManager::Get().FindConsoleVariable(TEXT("ini.UseNewDynamicLayers"))->GetInt() != 0;
+	if (bUseNewDynamicLayers)
+	{
+		// @todo branch - would be nice to have a way to know nothing was backed up yet, with Branch mode, so we can skip the FlushAsyncLoading call when there's nothing to do
+		// Flush async loading before modifying GConfig.
+		FlushAsyncLoading();
+
+		FConfigModificationTracker ChangeTracker;
+		for (const FCloudFileHeader& FileHeader : ChangedHotfixFileList)
+		{
+			if (FileHeader.FileName.EndsWith(TEXT(".INI")))
+			{
+				FName Tag = *BuildConfigCacheKey(FileHeader.FileName);
+				FConfigCacheIni::RemoveTagFromAllBranches(Tag, &ChangeTracker);
+			}
+		}
+		for (const FCloudFileHeader& FileHeader : RemovedHotfixFileList)
+		{
+			if (FileHeader.FileName.EndsWith(TEXT(".INI")))
+			{
+				FName Tag = *BuildConfigCacheKey(FileHeader.FileName);
+				FConfigCacheIni::RemoveTagFromAllBranches(Tag, &ChangeTracker);
+			}
+		}
+		
+		// @todo: the last param is for passing to a delegate on "SectionsChanged" but we are doing all the files at once -
+		// we could move the delegate out of this function, or call it once per file - but if a class was modified and removed
+		// we will do extra work - and also, the old way wasn't even calling the delegate!
+		UObjectBaseUtility::ReloadObjectsFromModifiedConfigSections(ChangeTracker.ModifiedSections, FString());
+		
+		return;
+	}
+
 	if (IniBackups.Num() == 0)
 	{
 		return;
diff --git a/Engine/Plugins/Online/OnlineFramework/Source/Hotfix/Public/OnlineHotfixManager.h b/Engine/Plugins/Online/OnlineFramework/Source/Hotfix/Public/OnlineHotfixManager.h
--- a/Engine/Plugins/Online/OnlineFramework/Source/Hotfix/Public/OnlineHotfixManager.h
+++ b/Engine/Plugins/Online/OnlineFramework/Source/Hotfix/Public/OnlineHotfixManager.h
@@ -227,6 +227,7 @@ protected:
 
 	/** @return the config file entry for the ini file name in question */
 	FConfigFile* GetConfigFile(const FString& IniName);
+	FConfigBranch* GetBranch(const FString& IniName);
 
 	/** @return the config cache key used to associate ini file entries within the config cache */
 	FString BuildConfigCacheKey(const FString& IniName);
diff --git a/Engine/Plugins/Runtime/GameFeatures/Source/GameFeatures/Private/GameFeatureData.cpp b/Engine/Plugins/Runtime/GameFeatures/Source/GameFeatures/Private/GameFeatureData.cpp
--- a/Engine/Plugins/Runtime/GameFeatures/Source/GameFeatures/Private/GameFeatureData.cpp
+++ b/Engine/Plugins/Runtime/GameFeatures/Source/GameFeatures/Private/GameFeatureData.cpp
@@ -84,21 +84,63 @@ static TAutoConsoleVariable<bool> CVarAllowRuntimeDeviceProfiles(
 void UGameFeatureData::InitializeBasePluginIniFile(const FString& PluginInstalledFilename) const
 {
 	const FString PluginName = FPaths::GetBaseFilename(PluginInstalledFilename);
-	const FString PluginConfigDir = FPaths::GetPath(PluginInstalledFilename) / TEXT("Config/");
-	const FString EngineConfigDir = FPaths::EngineConfigDir();
 
-	const bool bIsBaseIniName = false;
-	const bool bForceReloadFromDisk = false;
-	const bool bWriteDestIni = false;
+	// DEPRECATED NAMING PATH - must keep because these files are read in as a single file, not in a hierarchical way, so 
+	// they don't have the + syntax for arrays
+	{
+		const FString PluginConfigDir = FPaths::GetPath(PluginInstalledFilename) / TEXT("Config/");
+		const FString EngineConfigDir = FPaths::EngineConfigDir();
+
+		const bool bIsBaseIniName = false;
+		const bool bForceReloadFromDisk = false;
+		const bool bWriteDestIni = false;
+
+		// This will be the generated path including platform
+		FString PluginConfigFilename = GConfig->GetConfigFilename(*PluginName);
+
+		// Try the deprecated path first that doesn't include the Default prefix
+		FConfigFile& PluginConfig = GConfig->Add(PluginConfigFilename, FConfigFile());
+		if (FConfigCacheIni::LoadExternalIniFile(PluginConfig, *PluginName, *EngineConfigDir, *PluginConfigDir, bIsBaseIniName, nullptr, bForceReloadFromDisk, bWriteDestIni))
+		{
+			// This is the deprecated loading path that doesn't handle cases like + in arrays
+			UE_LOG(LogGameFeatures, Log, TEXT("[GameFeatureData %s]: Loaded deprecated config %s, rename to start with Default for normal parsing"), *GetPathNameSafe(this), *PluginConfigFilename);
+
+			FCoreRedirects::ReadRedirectsFromIni(PluginConfigFilename);
+			ReloadConfigs(PluginConfig);
 
-	// This will be the generated path including platform
-	FString PluginConfigFilename = GConfig->GetConfigFilename(*PluginName);
+			return;
+		}
+	}
 
-	// Try the deprecated path first that doesn't include the Default prefix
-	FConfigFile& PluginConfig = GConfig->Add(PluginConfigFilename, FConfigFile());
-	if (!FConfigCacheIni::LoadExternalIniFile(PluginConfig, *PluginName, *EngineConfigDir, *PluginConfigDir, bIsBaseIniName, nullptr, bForceReloadFromDisk, bWriteDestIni))
+	static bool bUseNewDynamicLayers = IConsoleManager::Get().FindConsoleVariable(TEXT("ini.UseNewDynamicLayers"))->GetInt() != 0;
+	if (bUseNewDynamicLayers)
 	{
-		// Now try the same rules as PluginManager using Default and the config hierarchy
+		TSharedPtr<IPlugin> Plugin = IPluginManager::Get().FindPlugin(PluginName);
+
+		// attempt to load config branch named for the plugin
+
+		UE_LOG(LogGameFeatures, Verbose, TEXT("Loading GameFeature base plugin hierarchy for %s"), *PluginName);
+
+		// register this plugin, so the ConfigContext.Load, and future loads, know about it
+		bool bIncludePluginNameInBranchName = true;
+		FConfigCacheIni::RegisterPlugin(*Plugin->GetName(), Plugin->GetBaseDir(), Plugin->GetExtensionBaseDirs(), DynamicLayerPriority::GameFeature, bIncludePluginNameInBranchName);
+
+		// load the plugin inis and track the modified sections
+		FConfigContext Context = FConfigContext::ReadIntoGConfig();
+		FString PluginConfigFilename;
+		Context.ConfigFileTag = *Plugin->GetName();
+		Context.Load(*PluginName, PluginConfigFilename);
+
+		// @todo move this into ReloadObjectsFromModifiedConfigSections?
+		FCoreRedirects::ReadRedirectsFromIni(PluginConfigFilename);
+
+		return;
+	}
+	
+	{
+		// Now use old method to load the plugin hierarchy without the registering plugin stuff above
+		FString PluginConfigFilename = GConfig->GetConfigFilename(*PluginName);
+		FConfigFile& PluginConfig = GConfig->Add(PluginConfigFilename, FConfigFile());
 		FConfigContext Context = FConfigContext::ReadIntoPluginFile(PluginConfig, *FPaths::GetPath(PluginInstalledFilename),
 				IPluginManager::Get().FindPluginFromPath(PluginName)->GetExtensionBaseDirs());
 
@@ -113,20 +155,49 @@ void UGameFeatureData::InitializeBasePluginIniFile(const FString& PluginInstalle
 			ReloadConfigs(PluginConfig);
 		}
 	}
-	else
-	{
-		// This is the deprecated loading path that doesn't handle cases like + in arrays
-		UE_LOG(LogGameFeatures, Log, TEXT("[GameFeatureData %s]: Loaded deprecated config %s, rename to start with Default for normal parsing"), *GetPathNameSafe(this), *PluginConfigFilename);
-
-		FCoreRedirects::ReadRedirectsFromIni(PluginConfigFilename);
-		ReloadConfigs(PluginConfig);
-	}
 }
 
 void UGameFeatureData::InitializeHierarchicalPluginIniFiles(const FString& PluginInstalledFilename) const
 {
 	UDeviceProfileManager& DeviceProfileManager = UDeviceProfileManager::Get();
 
+	static bool bUseNewDynamicLayers = IConsoleManager::Get().FindConsoleVariable(TEXT("ini.UseNewDynamicLayers"))->GetInt() != 0;
+	if (bUseNewDynamicLayers)
+	{
+		const FString PluginName = FPaths::GetBaseFilename(PluginInstalledFilename);
+		TSharedPtr<IPlugin> PluginSystemPlugin = IPluginManager::Get().FindPlugin(PluginName);
+		//	checkf(FPaths::GetPath(PluginInstalledStandardFilename) == PluginSystemPlugin->GetBaseDir(), TEXT("Expected plugin system to have matching BaseDir to GFD plugin"));
+
+		// read the plugin into any configs it is overriding
+		FConfigModificationTracker ChangeTracker;
+		ChangeTracker.SectionsToTrackContents.Add(TEXT("ConsoleVariables"));
+
+		UE_LOG(LogGameFeatures, Verbose, TEXT("Loading GameFeature config modification for %s"), *PluginName);
+
+		FConfigCacheIni::AddPluginToAllBranches(*PluginName, &ChangeTracker);
+
+		const FConfigSection* CVars = ChangeTracker.TrackedSections.Find(TEXT("ConsoleVariables"));
+		if (CVars != nullptr)
+		{
+			for (auto Pair : *CVars)
+			{
+				IConsoleVariable* CVar = IConsoleManager::Get().FindConsoleVariable(*Pair.Key.ToString());
+				CVar->Set(*Pair.Value.GetValue(), ECVF_SetByPluginLowPriority, *PluginName);
+			}
+		}
+
+		// reload objects that had their configs changed
+		UObjectBaseUtility::ReloadObjectsFromModifiedConfigSections(ChangeTracker.ModifiedSections, PluginName);
+
+		// update active DP if it was modified (including parents)
+		if (UDeviceProfileManager::Get().DoActiveProfilesReference(ChangeTracker.ModifiedSections))
+		{
+			DeviceProfileManager.ReapplyDeviceProfile();
+		}
+
+		return;
+	}
+
 	FString PlatformName = FPlatformProperties::IniPlatformName();
 
 #if ALLOW_OTHER_PLATFORM_CONFIG
@@ -479,9 +550,7 @@ void UGameFeatureData::InitializeHierarchicalPluginIniFiles(const FString& Plugi
 				const FString PluginIniPath = FString::Printf(TEXT("%s%s.ini"), *ConfigDirectory, *PluginIniName);
 				ExistingConfig->CombineFromBuffer(ConfigAsString, PluginIniPath);
 
-#if ALLOW_INI_OVERRIDE_FROM_COMMANDLINE
 				FConfigFile::OverrideFromCommandline(ExistingConfig, Ini.Name);
-#endif // ALLOW_INI_OVERRIDE_FROM_COMMANDLINE
 
 				if (Ini.bCreateDeviceProfiles)
 				{
diff --git a/Engine/Plugins/Runtime/GameFeatures/Source/GameFeatures/Private/GameFeaturePluginStateMachine.cpp b/Engine/Plugins/Runtime/GameFeatures/Source/GameFeatures/Private/GameFeaturePluginStateMachine.cpp
--- a/Engine/Plugins/Runtime/GameFeatures/Source/GameFeatures/Private/GameFeaturePluginStateMachine.cpp
+++ b/Engine/Plugins/Runtime/GameFeatures/Source/GameFeatures/Private/GameFeaturePluginStateMachine.cpp
@@ -37,6 +37,7 @@
 #include "Misc/PathViews.h"
 #include "Containers/Queue.h"
 #include "ShaderCodeLibrary.h"
+#include "DeviceProfiles/DeviceProfileManager.h"
 #include "Trace/Trace.h"
 
 #include UE_INLINE_GENERATED_CPP_BY_NAME(GameFeaturePluginStateMachine)
@@ -2970,6 +2971,25 @@ struct FGameFeaturePluginState_Deactivating : public FGameFeaturePluginState
 		NumExpectedPausers = 0;
 		bInProcessOfDeactivating = false;
 		bHasUnloaded = false;
+
+		static bool bUseNewDynamicLayers = IConsoleManager::Get().FindConsoleVariable(TEXT("ini.UseNewDynamicLayers"))->GetInt() != 0;
+		if (bUseNewDynamicLayers)
+		{
+			// we added the plugin in GameFeatureData.cpp, but we don't need to add any complexity into there, since we just unload by PluginName
+			FConfigModificationTracker ChangeTracker;
+			FConfigCacheIni::RemoveTagFromAllBranches(*StateProperties.PluginName, &ChangeTracker);
+
+			IConsoleManager::Get().UnsetAllConsoleVariablesWithTag(*StateProperties.PluginName);
+
+			// reload any modified uobjects
+			UObjectBaseUtility::ReloadObjectsFromModifiedConfigSections(ChangeTracker.ModifiedSections, StateProperties.PluginName);
+			
+			if (UDeviceProfileManager::Get().DoActiveProfilesReference(ChangeTracker.ModifiedSections))
+			{
+				UDeviceProfileManager& DeviceProfileManager = UDeviceProfileManager::Get();
+				DeviceProfileManager.ReapplyDeviceProfile();
+			}
+		}
 	}
 
 	void OnPauserCompleted(FStringView InPauserTag)
diff --git a/Engine/Source/Runtime/Core/Private/HAL/ConsoleManager.cpp b/Engine/Source/Runtime/Core/Private/HAL/ConsoleManager.cpp
--- a/Engine/Source/Runtime/Core/Private/HAL/ConsoleManager.cpp
+++ b/Engine/Source/Runtime/Core/Private/HAL/ConsoleManager.cpp
@@ -24,7 +24,6 @@ ConsoleManager.cpp: console command handling
 #include "HAL/PlatformProcess.h"
 #include "ProfilingDebugging/CsvProfiler.h"
 #include "HAL/FileManager.h"
-#include "Serialization/ArchiveCountMem.h"
 #include "Logging/MessageLog.h"
 
 #include <clocale>
@@ -373,6 +372,11 @@ public:
 		return this;
 	}
 	
+	virtual void LogHistory(FOutputDevice& Ar)
+	{
+	
+	}
+	
 	/** Legacy funciton to add old single delegates to the new multicast delegate. */
 	virtual void SetOnChangedCallback(const FConsoleVariableDelegate& Callback) 
 	{
@@ -2329,9 +2333,7 @@ void FConsoleManager::LoadHistoryIfNeeded()
 	HistoryEntriesMap.Reset();
 
 	FConfigFile Ini;
-
-	const FString ConfigPath = FPaths::GeneratedConfigDir() + TEXT("ConsoleHistory.ini");
-	ProcessIniContents(*ConfigPath, *ConfigPath, &Ini, false, false);
+	FConfigContext::ReadSingleIntoLocalFile(Ini).Load(*(FPaths::GeneratedConfigDir() + TEXT("ConsoleHistory.ini")));
 
 	const FString SectionName = TEXT("ConsoleHistory");
 	const FName KeyName = TEXT("History");
diff --git a/Engine/Source/Runtime/Core/Private/Misc/ConfigCacheIni.cpp b/Engine/Source/Runtime/Core/Private/Misc/ConfigCacheIni.cpp
--- a/Engine/Source/Runtime/Core/Private/Misc/ConfigCacheIni.cpp
+++ b/Engine/Source/Runtime/Core/Private/Misc/ConfigCacheIni.cpp
@@ -55,6 +55,26 @@ namespace
 DEFINE_LOG_CATEGORY(LogConfig);
 #define LOCTEXT_NAMESPACE "ConfigCache"
 
+
+static TAutoConsoleVariable<int32> CVarUseNewDynamicLayers(
+	TEXT("ini.UseNewDynamicLayers"),
+	0,
+	TEXT("If true, use the new dynamic layers that load/unload, with GameFeatures and Hotfixes"),
+	ECVF_Default);
+
+static int GUseNewSaveTracking = 0;
+static FAutoConsoleVariableRef CVarUseNewSaveTracking(
+	TEXT("ini.UseNewSaveTracking"),
+	GUseNewSaveTracking,
+	TEXT("If true, use the new method for tracking modifications to GConfig when saving"));
+
+#if WITH_EDITOR
+// editor wants full replay - and we can't put this in a cvar since we need way too early for cvar processing!
+static int GDefaultReplayMethod = 2;
+#else
+static int GDefaultReplayMethod = 1;
+#endif
+
 /*-----------------------------------------------------------------------------
 FConfigValue
 -----------------------------------------------------------------------------*/
@@ -789,6 +809,13 @@ static bool BuildOutputString(FString& String, const FileType& FileToWrite)
 		
 		for (const TPair<FName, FConfigValue>& Value : Section.Value)
 		{
+#if CONFIG_CAN_SAVE_COMMENTS
+			if (Value.Value.Comment.Len() > 0)
+			{
+				String.Append(Value.Value.Comment);
+				String.Append(LINE_TERMINATOR);
+			}
+#endif
 			if (Value.Value.ValueType != FConfigValue::EValueType::Set)
 			{
 				const TCHAR* Cmd = CommandLookup.FindKey(Value.Value.ValueType);
@@ -819,16 +846,32 @@ static bool AreWritesAllowedGlobally()
 
 static bool SaveBranch(FConfigBranch& Branch)
 {
-	FString Output;
-	Branch.InMemoryFile.WriteToString(Output, Branch.IniPath);
-
 	if (!Branch.InMemoryFile.Dirty || Branch.InMemoryFile.NoSave || !AreWritesAllowedGlobally())
 	{
 		return true;
 	}
 
-	if (Output.Len() > 0)
+	FString Output;
+	bool bBuiltString;
+	if (GUseNewSaveTracking == 1)
+	{
+		bBuiltString = BuildOutputString(Output, Branch.SavedLayer);
+	}
+	else if (GUseNewSaveTracking == 2)
+	{
+		bBuiltString = BuildDiffOutputString(Output, Branch.InMemoryFile, Branch.FinalCombinedLayers);
+	}
+	else
 	{
+		Branch.InMemoryFile.WriteToString(Output, Branch.IniPath);
+		bBuiltString = true;
+	}
+		
+	if (bBuiltString && Output.Len() > 0)
+	{
+		Output = FString(TEXT(";METADATA=(Diff=true, UseCommands=true)" LINE_TERMINATOR_ANSI)) + Output;
+		
+	
 		return SaveConfigFileWrapper(*Branch.IniPath, Output);
 	}
 
@@ -897,6 +940,7 @@ FConfigFile& FConfigFile::operator=(const FConfigFile& Other)
 	Tag = Other.Tag;
 	Branch = Other.Branch;
 
+	// @todo branch - remove this right?
 #if ALLOW_INI_OVERRIDE_FROM_COMMANDLINE
 	CommandlineOptions = Other.CommandlineOptions;
 #endif // ALLOW_INI_OVERRIDE_FROM_COMMANDLINE
@@ -941,6 +985,7 @@ FConfigFile& FConfigFile::operator=(FConfigFile&& Other)
 	Tag = MoveTemp(Other.Tag);
 	Branch = MoveTemp(Other.Branch);
 
+	// @todo branch - remove this right?
 #if ALLOW_INI_OVERRIDE_FROM_COMMANDLINE
 	CommandlineOptions = MoveTemp(Other.CommandlineOptions);
 #endif // ALLOW_INI_OVERRIDE_FROM_COMMANDLINE
@@ -1364,6 +1409,7 @@ void FillFileFromBuffer(FileType* File, FStringView Buffer, bool bHandleSymbolCo
 	TStringBuilder<128> TheLine;
 	FString ProcessedValue;
 	bool Done = false;
+	bool bHasHandledMetadata = false;
 	while( !Done && Ptr - Buffer.GetData() < Buffer.Len() )
 	{
 		// Advance past new line characters
@@ -1387,6 +1433,26 @@ void FillFileFromBuffer(FileType* File, FStringView Buffer, bool bHandleSymbolCo
 			Start[FCString::Strlen(Start)-1] = TEXT('\0');
 		}
 
+		// look for the ;METADATA line as the first line in the file, which can override passed in settings
+		if (!bHasHandledMetadata)
+		{
+			bHasHandledMetadata = true;
+			
+			if (FCString::Strnicmp(Start, TEXT(";METADATA="), 10) == 0)
+			{
+				FString MetadataStruct(Start + 10);
+				FString MetadataValue;
+				ExtractPropertyValue(MetadataStruct, TEXT("UseCommands="), MetadataValue);
+				if (MetadataValue.Len() > 0)
+				{
+					bHandleSymbolCommands = FCString::ToBool(*MetadataValue);
+				}
+				
+				// move on to the next line
+				continue;
+			}
+			
+		}
 		// If the first character in the line is [ and last char is ], this line indicates a section name
 		if( *Start=='[' && Start[FCString::Strlen(Start)-1]==']' )
 		{
@@ -1769,8 +1835,20 @@ void FConfigFile::OverrideFromCommandline(FConfigCommandStream* File, const FStr
 						Value = Value.Mid(1, Value.Len() - 2);
 					}
 
+					FConfigValue::EValueType ValueType = FConfigValue::EValueType::Set;
+					if (PropertyKey.StartsWith(TEXT("-")))
+					{
+						PropertyKey.RemoveFromStart(TEXT("-"));
+						ValueType = FConfigValue::EValueType::Remove;
+					}
+					else if (PropertyKey.StartsWith(TEXT("+")))
+					{
+						PropertyKey.RemoveFromStart(TEXT("+"));
+						ValueType = FConfigValue::EValueType::ArrayAdd;
+					}
+
 					FConfigCommandStreamSection* SectionToModify = File->FindOrAddSectionInternal(Section);
-					SectionToModify->Emplace(*PropertyKey, FConfigValue(MoveTemp(Value), FConfigValue::EValueType::Set));
+					SectionToModify->Emplace(*PropertyKey, FConfigValue(MoveTemp(Value), ValueType));
 				}
 			}
 		}
@@ -1781,6 +1859,14 @@ void FConfigFile::OverrideFromCommandline(FConfigCommandStream* File, const FStr
 void FConfigFile::OverrideFromCommandline(FConfigFile* File, const FString& Filename)
 {
 #if ALLOW_INI_OVERRIDE_FROM_COMMANDLINE
+
+	// this would already be handled with new dynamic layer stuff
+	static bool bUseNewDynamicLayers = IConsoleManager::Get().FindConsoleVariable(TEXT("ini.UseNewDynamicLayers"))->GetInt() != 0;
+	if (bUseNewDynamicLayers)
+	{
+		return;
+	}
+
 	// look for this filename on the commandline in the format:
 	//		-ini:IniName:[Section1]:Key=Value
 	// for example:
@@ -1890,14 +1976,6 @@ void FConfigFile::OverrideFromCommandline(FConfigFile* File, const FString& File
 }
 
 
-void FConfigFile::AddDynamicLayerToHierarchy(const FString& Filename)
-{
-	checkf(Branch != nullptr, TEXT("Unable to AddDynamicLayerToHierarchy on an FConfigFile that doesn't have an internal Branch. ConfigFile = %s"), *Name.ToString());
-	
-	Branch->AddDynamicLayerToHierarchy(Filename);
-}
-
-
 namespace UE::ConfigCacheIni::Private
 {
 
@@ -2536,6 +2614,23 @@ void FConfigFile::SetArray(const TCHAR* SectionName, const TCHAR* Key, const TAr
 		Section->Add(Key, FConfigValue(Section, FName(Key), Value[i]));
 		Dirty = true;
 	}
+	
+	if (ChangeTracker)
+	{
+		// remove anything to do with this array in the tracker
+		FConfigCommandStreamSection* Sec = ChangeTracker->FindOrAddSectionInternal(SectionName);
+		
+		// if there were any entries to remove, then add a clear operation
+		if (Sec->Remove(Key) > 0)
+		{
+			Sec->Add(Key, FConfigValue(TEXT("__Clear__"), FConfigValue::EValueType::Clear));
+		}
+		// then add all entries
+		for (int32 i = 0; i < Value.Num(); i++)
+		{
+			Sec->Add(Key, FConfigValue(Value[i], FConfigValue::EValueType::ArrayAdd));
+		}
+	}
 }
 
 bool FConfigFile::AddToSection(const TCHAR* SectionName, FName Key, const FString& Value)
@@ -2543,6 +2638,13 @@ bool FConfigFile::AddToSection(const TCHAR* SectionName, FName Key, const FStrin
 	FConfigSection* Section = FindOrAddSectionInternal(SectionName);
 	Section->Add(Key, FConfigValue(Section, Key, Value));
 	Dirty = true;
+	
+	if (ChangeTracker != nullptr)
+	{
+		FConfigCommandStreamSection* Sec = ChangeTracker->FindOrAddSectionInternal(SectionName);
+		Sec->Add(Key, FConfigValue(Value, FConfigValue::EValueType::ArrayAdd));
+	}
+	
 	return true;
 }
 
@@ -2557,6 +2659,13 @@ bool FConfigFile::AddUniqueToSection(const TCHAR* SectionName, FName Key, const
 	// just call Add since we already checked above if it exists (AddUnique can't return whether or not it existed)
 	Section->Add(Key, FConfigValue(Section, Key, Value));
 	Dirty = true;
+	
+	if (ChangeTracker != nullptr)
+	{
+		FConfigCommandStreamSection* Sec = ChangeTracker->FindOrAddSectionInternal(SectionName);
+		Sec->Add(Key, FConfigValue(Value, FConfigValue::EValueType::ArrayAddUnique));
+	}
+
 	return true;
 }
 
@@ -2571,6 +2680,15 @@ bool FConfigFile::RemoveKeyFromSection(const TCHAR* SectionName, FName Key)
 
 	Section->Remove(Key);
 	Dirty = true;
+	
+	if (ChangeTracker != nullptr)
+	{
+		FConfigCommandStreamSection* Sec = ChangeTracker->FindOrAddSectionInternal(SectionName);
+		// remove any tracked changes for this key as they are all blown away now
+		Sec->Remove(Key);
+		Sec->Add(Key, FConfigValue(TEXT("__Clear__"), FConfigValue::EValueType::Clear));
+	}
+
 	return true;
 }
 
@@ -2600,6 +2718,13 @@ bool FConfigFile::ResetKeyInSection(const TCHAR* SectionName, FName Key)
 
 	Section->Remove(Key);
 
+	if (ChangeTracker != nullptr)
+	{
+		// remove this key from being tracked, which is the difference between this and RemvoeKeyFromSection
+		FConfigCommandStreamSection* Sec = ChangeTracker->FindOrAddSectionInternal(SectionName);
+		Sec->Remove(Key);
+	}
+	
 	Dirty = true;
 	return true;
 }
@@ -2795,21 +2920,14 @@ void FConfigCommandStream::Shrink()
 -----------------------------------------------------------------------------*/
 
 FConfigBranch::FConfigBranch()
+	: bIsSafeUnloaded(false)
 {
 	static int DefaultReplayMethod = -1;
 	if (DefaultReplayMethod == -1)
 	{
 		if (!FParse::Value(FCommandLine::Get(), TEXT("ConfigReplayMethod="), DefaultReplayMethod))
 		{
-#if WITH_EDITOR
-			// editor wants full replay - and we can't put this in a cvar since we need way too early for cvar processing!
-			DefaultReplayMethod = 2;
-#else
-			// @todo We are defaulting to 0 (no replay) here because nothing is using the "unload" functionality yet until Phase 2 of the 
-			// changes, and this is using more memory. An mem saving solution will be needed before we can set the default to 1
-			// like it will need to be
-			DefaultReplayMethod = 0;
-#endif
+			DefaultReplayMethod = GDefaultReplayMethod;
 		}
 	}
 	
@@ -2828,7 +2946,8 @@ FConfigBranch::FConfigBranch()
 }
 
 FConfigBranch::FConfigBranch(const FConfigFile& ExistingFile)
-	: bIsHierarchical(false)
+	: bIsSafeUnloaded(false)
+	, bIsHierarchical(false)
 	, InMemoryFile(ExistingFile)
 {
 	ReplayMethod = EBranchReplayMethod::NoReplay;
@@ -2843,6 +2962,11 @@ void FConfigBranch::InitFiles()
 	FinalCombinedLayers.Branch = this;
 	CommandLineOverrides.Branch = this;
 	InMemoryFile.Branch = this;
+	
+	if (GUseNewSaveTracking)
+	{
+		InMemoryFile.ChangeTracker = &SavedLayer;
+	}
 }
 
 bool FConfigBranch::AddDynamicLayerToHierarchy(const FString& Filename, FConfigModificationTracker* ModificationTracker)
@@ -3106,14 +3230,15 @@ bool FConfigBranch::RemoveDynamicLayersFromHierarchy(const TArray<FString>& File
 
 }
 
-void FConfigBranch::RemovePluginFromHierarchy(FName PluginName, FConfigModificationTracker* ModificationTracker)
+void FConfigBranch::RemoveTagFromHierarchy(FName Tag, FConfigModificationTracker* ModificationTracker)
 {
 	// gather tagged layers
 	TArray<FString> LayersToRemove;
 	for (DynamicLayerList::TIterator Node(DynamicLayers.GetHead()); Node; ++Node)
 	{
-		if (Node->Tag == PluginName)
+		if (Node->Tag == Tag)
 		{
+			UE_LOG(LogConfig, Verbose, TEXT("Removing dynamic layer %s from branch %s with tag %s"), *Node->Filename, *IniName.ToString(), *Tag.ToString());
 			// @todo make a version that takes CommandStreams, not Filenames, for speed
 			LayersToRemove.Add(Node->Filename);
 		}
@@ -3126,6 +3251,22 @@ void FConfigBranch::RemovePluginFromHierarchy(FName PluginName, FConfigModificat
 	}
 }
 
+bool FConfigBranch::SafeUnload()
+{
+	bIsSafeUnloaded = true;
+
+	InMemoryFile.Cleanup();
+	CombinedStaticLayers.Cleanup();
+	SavedLayer.Empty();
+	CommandLineOverrides.Empty();
+	StaticLayers.Empty();
+	DynamicLayers.Empty();
+	FinalCombinedLayers.Empty();
+
+	return true;
+}
+
+
 void FConfigBranch::Flush()
 {
 	SaveBranch(*this);
@@ -3222,38 +3363,22 @@ FConfigBranch* FConfigCacheIni::FindBranchWithNoReload(FName BaseIniName, const
 
 	return Branch;
 }
-	
+
 FConfigBranch* FConfigCacheIni::FindBranch(FName BaseIniName, const FString& Filename)
 {
-	// look for a known file, if there's no ini extension
-	FConfigBranch* Branch = KnownFiles.GetBranch(BaseIniName);
+	FConfigBranch* Branch = FindBranchWithNoReload(BaseIniName, Filename);
 
-	if (Branch == nullptr)
-	{
-		Branch = KnownFiles.GetBranch(*Filename);
-	}
-	if (Branch == nullptr)
+	if (Branch && Branch->bIsSafeUnloaded)
 	{
-		Branch = OtherFiles.FindRef(Filename);
-		if (Branch == nullptr)
-		{
-			for (TPair<FString,FConfigBranch*>& CurrentFilePair : OtherFiles)
-			{
-				if (CurrentFilePair.Value->IniName == BaseIniName)
-				{
-					Branch = CurrentFilePair.Value;
-					break;
-				}
-			}
-		}
-	}
+		double StartTime = FPlatformTime::Seconds();
+		FConfigContext Context = FConfigContext::ReadIntoConfigSystem(this, PlatformName.ToString());
+		Context.Branch = Branch;
+		Context.DestIniFilename = Branch->IniPath;
+		Context.Load(*BaseIniName.ToString());
 
-	// if Filename is a .ini and it doesn't match what the KnownFile has (if it has one yet), then we can't use it
-	if (Branch && Branch->IniPath.Len() > 0 && Filename.Len() > 0 && Filename.EndsWith(".ini") && Branch->IniPath != Filename)
-	{
-		Branch = nullptr;
+		UE_LOG(LogConfig, Warning, TEXT("Branch '%s' had been unloaded. Reloading on-demand took %.2fms"), *Branch->IniName.ToString(), (FPlatformTime::Seconds() - StartTime) * 1000.0f);
 	}
-	
+
 	return Branch;
 }
 
@@ -3310,15 +3435,17 @@ FConfigFile* FConfigCacheIni::Find(const FString& Filename)
 		
 		if (!Result)
 		{
-			Result = &Add(UnrealFileName, FConfigFile());
-			UE_LOG(LogConfig, Verbose, TEXT("GConfig::Find is looking for file:  %s"), *UnrealFileName);
 			if (DoesConfigFileExistWrapper(*UnrealFileName))
 			{
-				Result->Read(UnrealFileName);
-				// Files added through Find are treated the same as ReadSingleIntoConfigSystem contexts,
-				// and do not use a hierarchy so they do not use a generatedini and should never be saved.
-				Result->NoSave = true;
-				UE_LOG(LogConfig, Verbose, TEXT("GConfig::Find has loaded file:  %s"), *UnrealFileName);
+				Result = &Add(UnrealFileName, FConfigFile());
+				UE_LOG(LogConfig, Verbose, TEXT("GConfig::Find is looking for file:  %s"), *UnrealFileName);
+				{
+					Result->Read(UnrealFileName);
+					// Files added through Find are treated the same as ReadSingleIntoConfigSystem contexts,
+					// and do not use a hierarchy so they do not use a generatedini and should never be saved.
+					Result->NoSave = true;
+					UE_LOG(LogConfig, Verbose, TEXT("GConfig::Find has loaded file:  %s"), *UnrealFileName);
+				}
 			}
 		}
 		else
@@ -3866,6 +3993,19 @@ bool FConfigCacheIni::RemoveKey( const TCHAR* Section, const TCHAR* Key, const F
 	return false;
 }
 
+bool FConfigCacheIni::SafeUnloadBranch(const TCHAR* BranchName)
+{
+	FConfigBranch* Branch = FindBranchWithNoReload(BranchName, BranchName);
+	if (Branch)
+	{
+		return Branch->SafeUnload();
+	}
+
+	return false;
+}
+
+
+
 bool FConfigCacheIni::EmptySection( const TCHAR* Section, const FString& Filename )
 {
 	FConfigFile* File = Find(Filename);
@@ -4592,7 +4732,7 @@ void FConfigCacheIni::ShowMemoryUsage( FOutputDevice& Ar )
 	for (TPair<FString, FConfigBranch*>& Pair : OtherFiles)
 	{
 		FString Filename = Pair.Key;
-		FConfigFile& ConfigFile = Pair.Value->InMemoryFile;
+		FConfigBranch& ConfigBranch = *Pair.Value;
 
 		FArchiveCountConfigMem MemAr;
 
@@ -4600,7 +4740,7 @@ void FConfigCacheIni::ShowMemoryUsage( FOutputDevice& Ar )
 		MemAr << Filename;
 
 		// count the bytes used for storing the array of SectionName->Section pairs
-		MemAr << ConfigFile;
+		MemAr << ConfigBranch;
 		
 		ConfigCacheMemoryData.AddConfigFile(Filename, MemAr);
 	}
@@ -4925,6 +5065,28 @@ FArchive& operator<<(FArchive& Ar, FConfigBranch& ConfigBranch)
 {
 	Ar << ConfigBranch.bIsHierarchical;
 	Ar << ConfigBranch.InMemoryFile;
+
+	// needed to count full memory usage
+	if (!Ar.IsPersistent())
+	{
+		Ar << ConfigBranch.IniName;
+		Ar << ConfigBranch.IniPath;
+		Ar << ConfigBranch.Platform;
+		Ar << ConfigBranch.SourceEngineConfigDir;
+		Ar << ConfigBranch.SourceProjectConfigDir;
+		Ar << ConfigBranch.Hierarchy;
+		Ar << ConfigBranch.StaticLayers;
+		Ar << ConfigBranch.SavedLayer;
+		Ar << ConfigBranch.CombinedStaticLayers;
+		Ar << ConfigBranch.FinalCombinedLayers;
+		Ar << ConfigBranch.CommandLineOverrides;
+		Ar << ConfigBranch.RuntimeChanges;
+		for (FConfigBranch::DynamicLayerList::TIterator Node(ConfigBranch.DynamicLayers.GetHead()); Node; ++Node)
+		{
+			FConfigCommandStream& S = *Node.GetNode()->GetValue();
+			Ar << S;
+		}
+	}
 	return Ar;
 }
 
@@ -5709,15 +5871,15 @@ void FConfigCacheIni::AddPluginToAllBranches(FName PluginName, FConfigModificati
 #endif
 }
 
-void FConfigCacheIni::RemovePluginFromAllBranches(FName PluginName, FConfigModificationTracker* ModificationTracker)
+void FConfigCacheIni::RemoveTagFromAllBranches(FName Tag, FConfigModificationTracker* ModificationTracker)
 {
-	GConfig->RemovePluginFromBranches(PluginName, ModificationTracker);
+	GConfig->RemoveTagFromBranches(Tag, ModificationTracker);
 	
 #if ALLOW_OTHER_PLATFORM_CONFIG
 	// need to walk over the other platforms without calling ForPlatform because that could end up loading pending plugins
 	for (auto Pair : ConfigForPlatform)
 	{
-		Pair.Value->RemovePluginFromBranches(PluginName, ModificationTracker);
+		Pair.Value->RemoveTagFromBranches(Tag, ModificationTracker);
 	}
 #endif
 
@@ -5778,20 +5940,22 @@ void FConfigCacheIni::AddPluginToBranches(FName PluginName, FConfigModificationT
 			continue;
 		}
 		
+		UE_LOG(LogConfig, Verbose, TEXT("Modifying branch %s with plugin ini %s"), *BranchName.ToString(), *ConfigFilename);
+
 		Context.ConfigFileTag = PluginName;
 		Context.Load(*BranchName.ToString());
 	}
 }
 
-void FConfigCacheIni::RemovePluginFromBranches(FName PluginName, FConfigModificationTracker* ModificationTracker)
+void FConfigCacheIni::RemoveTagFromBranches(FName Tag, FConfigModificationTracker* ModificationTracker)
 {
 	for (uint8 KnownIndex = 0; KnownIndex < (uint8)EKnownIniFile::NumKnownFiles; KnownIndex++)
 	{
-		KnownFiles.Branches[KnownIndex].RemovePluginFromHierarchy(PluginName, ModificationTracker);
+		KnownFiles.Branches[KnownIndex].RemoveTagFromHierarchy(Tag, ModificationTracker);
 	}
 	for (auto& Pair : OtherFiles)
 	{
-		Pair.Value->RemovePluginFromHierarchy(PluginName, ModificationTracker);
+		Pair.Value->RemoveTagFromHierarchy(Tag, ModificationTracker);
 	}
 }
 
@@ -5915,9 +6079,6 @@ void FConfigCacheIni::RegisterPlugin(FName PluginName, const FString& PluginDir,
 
 
 
-
-
-
 double GPrepareForLoadTime = 0;
 double GPerformLoadTime = 0;
 
@@ -6011,30 +6172,45 @@ class FIniExec : public FSelfRegisteringExec
 				GConfig->Flush(false);
 			}
 		}
-		
+
+		if (FParse::Command(&Cmd, TEXT("Unload")))
+		{
+			TCHAR BranchName[256];
+			if (FParse::Token(Cmd, BranchName, UE_ARRAY_COUNT(BranchName), true))
+			{
+				GConfig->SafeUnloadBranch(BranchName);
+			}
+		}
+
 		if (FParse::Command(&Cmd, TEXT("AddHotFix")))
 		{
 			TCHAR FileName[256];
 			if (FParse::Token(Cmd, FileName, UE_ARRAY_COUNT(FileName), true))
 			{
 				FConfigModificationTracker ChangeTracker;
-				FConfigBranch* Branch = GConfig->FindBranch(*FPaths::GetBaseFilename(FileName), FileName);
+				FString FilenameBase = FPaths::GetBaseFilename(FileName);
+				FConfigBranch* Branch = GConfig->FindBranch(*FilenameBase, FilenameBase);
 				
 				if (Branch != nullptr)
 				{
 					Branch->AddDynamicLayersToHierarchy({ FileName }, "HotfixTest", DynamicLayerPriority::Hotfix, &ChangeTracker);
-					Ar.Logf(TEXT("HF Modified sections:"));
+					Ar.Logf(TEXT("HF Modified sections (reloading as needed):"));
 					for (const FString& Section : ChangeTracker.ModifiedSections)
 					{
 						Ar.Logf(TEXT("  %s"), *Section);
 					}
+
+					FCoreDelegates::ReloadObjectsAfterDynamicConfigChange.Broadcast(ChangeTracker.ModifiedSections, FileName);
 				}
 			}
 		}
 		
 		if (FParse::Command(&Cmd, TEXT("RemoveHotFixes")))
 		{
-			FConfigCacheIni::RemovePluginFromAllBranches("HotfixTest");
+			FConfigModificationTracker ChangeTracker;
+			FConfigCacheIni::RemoveTagFromAllBranches("HotfixTest", &ChangeTracker);
+
+			FCoreDelegates::ReloadObjectsAfterDynamicConfigChange.Broadcast(ChangeTracker.ModifiedSections, "");
 		}
 		
 		if (FParse::Command(&Cmd, TEXT("Timing")))
diff --git a/Engine/Source/Runtime/Core/Private/Misc/ConfigContext.cpp b/Engine/Source/Runtime/Core/Private/Misc/ConfigContext.cpp
--- a/Engine/Source/Runtime/Core/Private/Misc/ConfigContext.cpp
+++ b/Engine/Source/Runtime/Core/Private/Misc/ConfigContext.cpp
@@ -188,6 +188,12 @@ bool FConfigContext::Load(const TCHAR* InBaseIniName, FString& OutFinalFilename)
 {
 	TRACE_CPUPROFILER_EVENT_SCOPE(FConfigContext::Load);
 
+	if (Branch != nullptr && Branch->bIsSafeUnloaded)
+	{
+		Branch->bIsSafeUnloaded = false;
+		return LoadIniFileHierarchy();
+	}
+
 	// set up a branch if needed
 	if (ExistingFile != nullptr)
 	{
@@ -273,6 +279,16 @@ bool FConfigContext::Load(const TCHAR* InBaseIniName, FString& OutFinalFilename)
 			// we need to copy the temporary branch's final result back into the output
 			*ExistingFile = TemporaryBranch->InMemoryFile;
 		}
+		// delete the branch if it is useless
+		if (!bSuccess && ConfigSystem != nullptr && ExistingFile == nullptr && TemporaryBranch == nullptr)
+		{
+			static bool bRemoveEmptyPlugins = FParse::Param(FCommandLine::Get(), TEXT("RemoveEmptyConfigs"));
+			if (bRemoveEmptyPlugins)
+			{
+				ConfigSystem->Remove(DestIniFilename);
+				Branch = nullptr;
+			}
+		}
 	}
 	return bSuccess;
 }
@@ -351,6 +367,7 @@ bool FConfigContext::PerformSingleFileLoad()
 		}
 		else
 		{
+			// @todo should we pass in a Name to AddNewBranch? could pass is BaseIniName
 			Branch = &ConfigSystem->AddNewBranch(DestIniFilename);
 			Branch->bIsHierarchical = false;
 		}
@@ -385,11 +402,11 @@ bool FConfigContext::PrepareForLoad(bool& bPerformLoad)
 
 	// assume we will load, unless some code below determines not to
 	bPerformLoad = true;
-	
+
 	// first, make sure the DestIniFilename is set, if needed
 	if (bWriteDestIni || bAllowGeneratedIniWhenCooked || FPlatformProperties::RequiresCookedData())
 	{
-		// delay filling out GeneratedConfigDir because some early configs can be read in that set -savedir, and
+		// delay filling out GeneratedConfigDir because some early configs can be read in that set -savedir, and 
 		// FPaths::GeneratedConfigDir() will permanently cache the value
 		if (GeneratedConfigDir.IsEmpty())
 		{
@@ -904,6 +921,8 @@ bool FConfigContext::LoadIniFileHierarchy()
 	
 	// LogVariables(*BaseIniName, Platform);
 
+	bool bReadAnyFile = false;
+
 	TRACE_CPUPROFILER_EVENT_SCOPE(LoadIniFileHierarchy);
 	// Traverse ini list back to front, merging along the way.
 	for (const TPair<int32, FString>& HierarchyIt : Branch->Hierarchy)
@@ -918,6 +937,11 @@ bool FConfigContext::LoadIniFileHierarchy()
 		{
 			continue;
 		}
+		
+		if (HierarchyIt.Key != 0)
+		{
+			bReadAnyFile = true;
+		}
 
 		UE_CLOG(bDumpIniLoadInfo, LogConfig, Display, TEXT("   Found!"));
 
@@ -960,9 +984,12 @@ bool FConfigContext::LoadIniFileHierarchy()
 	if (Branch->ReplayMethod != EBranchReplayMethod::NoReplay)
 	{
 		Branch->InMemoryFile = Branch->CombinedStaticLayers;
+		
+		// need to reset this since it just got blown away
+		Branch->InMemoryFile.ChangeTracker = &Branch->SavedLayer;
 	}
 	Branch->FinalCombinedLayers = Branch->InMemoryFile;
-	return true;
+	return bReadAnyFile;
 }
 
 /**
@@ -982,6 +1009,7 @@ bool FConfigContext::GenerateDestIniFile()
 	Branch->CommandLineOverrides.Empty();
 	Branch->StaticLayers.Empty();
 	Branch->DynamicLayers.Empty();
+	Branch->FinalCombinedLayers.Empty();
 
 	// read the static files into the branch
 	bool bResult = LoadIniFileHierarchy();
@@ -1102,8 +1130,6 @@ bool FConfigContext::GenerateDestIniFile()
 
 
 
-
-
 /*-----------------------------------------------------------------------------
 	FConfigFileHierarchy
 -----------------------------------------------------------------------------*/
diff --git a/Engine/Source/Runtime/Core/Private/Misc/CoreDelegates.cpp b/Engine/Source/Runtime/Core/Private/Misc/CoreDelegates.cpp
--- a/Engine/Source/Runtime/Core/Private/Misc/CoreDelegates.cpp
+++ b/Engine/Source/Runtime/Core/Private/Misc/CoreDelegates.cpp
@@ -230,6 +230,8 @@ TTSMulticastDelegate<void(const FString&, const TSet<FString>&)>& FCoreDelegates
 	return Singleton;
 }
 
+TMulticastDelegate<void(const TSet<FString>& ModifiedSections, const FString& IniFilename)> FCoreDelegates::ReloadObjectsAfterDynamicConfigChange;
+
 TMulticastDelegate<void(const TCHAR*, const TCHAR*, uint32, bool)> FCoreDelegates::OnApplyCVarFromIni;
 TMulticastDelegate<void(uint32, uint32)> FCoreDelegates::OnSystemResolutionChanged;
 
diff --git a/Engine/Source/Runtime/Core/Public/Misc/ConfigCacheIni.h b/Engine/Source/Runtime/Core/Public/Misc/ConfigCacheIni.h
--- a/Engine/Source/Runtime/Core/Public/Misc/ConfigCacheIni.h
+++ b/Engine/Source/Runtime/Core/Public/Misc/ConfigCacheIni.h
@@ -49,6 +49,7 @@ CORE_API DECLARE_LOG_CATEGORY_EXTERN(LogConfig, Log, All);
 
 // Server builds should be tweakable even in Shipping
 #define ALLOW_INI_OVERRIDE_FROM_COMMANDLINE			(UE_SERVER || !(UE_BUILD_SHIPPING))
+#define CONFIG_CAN_SAVE_COMMENTS (WITH_EDITOR)
 
 
 ///////////////////////////////////////////////////////////////////////////////
@@ -371,6 +372,10 @@ public:
 	// Add, subtract, stc
 	EValueType ValueType;
 
+#if CONFIG_CAN_SAVE_COMMENTS
+	FString Comment;
+#endif
+
 	/** Gets the expanded value (GetValue) without marking it as having been accessed for e.g. writing out to a ConfigFile to disk */
 	const FString& GetValueForWriting() const
 	{
@@ -567,6 +572,9 @@ private:
 	// we can't track it just in that section. This is expected to be empty/small
 	TMap<FString, TMap<FName, FString> > PerObjectConfigArrayOfStructKeys;
 
+	// if this is set, then we track changes made to sections for saving or replaying later (currently unused)
+	FConfigCommandStream* ChangeTracker = nullptr;
+	
 #if UE_WITH_CONFIG_TRACKING
 	mutable TRefCountPtr<UE::ConfigAccessTracking::FFile> FileAccess;
 #endif
@@ -953,6 +961,7 @@ public:
 	FString SourceEngineConfigDir;
 	FString SourceProjectConfigDir;
 	
+	bool bIsSafeUnloaded;
 	bool bIsHierarchical;
 	EBranchReplayMethod ReplayMethod = EBranchReplayMethod::NoReplay;
 	
@@ -996,8 +1005,12 @@ public:
 	CORE_API bool RemoveDynamicLayerFromHierarchy(const FString& Filename, FConfigModificationTracker* ModificationTracker =nullptr);
 	CORE_API bool RemoveDynamicLayersFromHierarchy(const TArray<FString>& Filenames, FConfigModificationTracker* ModificationTracker =nullptr);
 	
-	
-	
+	/** 
+	 * Frees up the static layer memory, which can be useful if a branch is loaded from, cached, and never used again. However,
+	 * if something does try to access it, it will reload in-place (can hitch your game, so be aware)
+	 */
+	CORE_API bool SafeUnload();
+
 	CORE_API void Flush();
 	
 	CORE_API void Dump(FOutputDevice& Ar);
@@ -1005,7 +1018,7 @@ public:
 private:
 	void InitFiles();
 	
-	void RemovePluginFromHierarchy(FName PluginName, FConfigModificationTracker* ModificationTracker);
+	void RemoveTagFromHierarchy(FName Tag, FConfigModificationTracker* ModificationTracker);
 	
 	friend class FConfigCacheIni;
 };
@@ -1021,7 +1034,7 @@ public:
 	/** DO NOT USE. This constructor is for internal usage only for hot-reload purposes. */
 	CORE_API FConfigCacheIni();
 
-	CORE_API ~FConfigCacheIni();
+	CORE_API virtual ~FConfigCacheIni();
 
 	/**
 	* Disables any file IO by the config cache system
@@ -1521,6 +1534,11 @@ public:
 	 */
 	CORE_API bool ResetKeyInSection(const TCHAR* Section, FName Key, const FString& Filename);
 
+	/**
+	 * Clears out the memory of a branch, but can reload it on demand if needed
+	 */
+	CORE_API bool SafeUnloadBranch(const TCHAR* Filename);
+
 
 	// Static helper functions
 
@@ -1728,7 +1746,7 @@ public:
 	static CORE_API void RegisterPlugin(FName PluginName, const FString& PluginDir, const TArray<FString>& ChildPluginDirs, DynamicLayerPriority Priority, bool bIncludePluginNameInBranchName);
 
 	static CORE_API void AddPluginToAllBranches(FName PluginName, FConfigModificationTracker* ModificationTracker=nullptr);
-	static CORE_API void RemovePluginFromAllBranches(FName PluginName, FConfigModificationTracker* ModificationTracker=nullptr);
+	static CORE_API void RemoveTagFromAllBranches(FName Tag, FConfigModificationTracker* ModificationTracker=nullptr);
 
 private:
 #if WITH_EDITOR
@@ -1742,7 +1760,7 @@ private:
 	void DumpFile(FOutputDevice& Ar, const FString& Filename, const FConfigFile& File);
 
 	void AddPluginToBranches(FName PluginName, FConfigModificationTracker* ModificationTracker);
-	void RemovePluginFromBranches(FName PluginName, FConfigModificationTracker* ModificationTracker);
+	void RemoveTagFromBranches(FName Tag, FConfigModificationTracker* ModificationTracker);
 
 	/** true if file operations should not be performed */
 	bool bAreFileOperationsDisabled;
diff --git a/Engine/Source/Runtime/Core/Public/Misc/ConfigContext.h b/Engine/Source/Runtime/Core/Public/Misc/ConfigContext.h
--- a/Engine/Source/Runtime/Core/Public/Misc/ConfigContext.h
+++ b/Engine/Source/Runtime/Core/Public/Misc/ConfigContext.h
@@ -109,6 +109,37 @@ public:
 		return Context;
 	}
 
+	/**
+	  * Inserts plugin ini files into an existing Branch
+	 */
+	static FConfigContext ReadPluginIntoConfigSystem(FConfigCacheIni* ConfigSystem, FName PluginName, const FString& PluginRootDir, const TArray<FString>& ChildPluginsBaseDirs)
+	{
+		FConfigContext Context(ConfigSystem, true, FString());
+		Context.bIsForPlugin = true;
+		Context.PluginRootDir = PluginRootDir;
+		Context.ChildPluginBaseDirs = ChildPluginsBaseDirs;
+		Context.ConfigFileTag = PluginName;
+
+		return Context;
+	}
+
+	/**
+	  * Inserts plugin ini files into an existing Branch
+	 */
+	static FConfigContext ReadPluginToModifyConfigSystem(FConfigCacheIni* ConfigSystem, DynamicLayerPriority Priority, FName PluginName, const FString& PluginRootDir, const TArray<FString>& ChildPluginsBaseDirs, FConfigModificationTracker* ChangeTracker=nullptr)
+	{
+		FConfigContext Context(ConfigSystem, true, FString());
+		Context.bIsForPlugin = true;
+		Context.bIsForPluginModification = true;
+		Context.PluginRootDir = PluginRootDir;
+		Context.ChildPluginBaseDirs = ChildPluginsBaseDirs;
+		Context.ConfigFileTag = PluginName;
+		Context.ChangeTracker = ChangeTracker;
+		Context.PluginModificationPriority = Priority;
+
+		return Context;
+	}
+
 	/**
 	 * Create a context to read a hierarchy, but once it reaches the given filename (StartDeletingFilename), it will not read in anymore 
 	 * files at that point
diff --git a/Engine/Source/Runtime/Core/Public/Misc/CoreDelegates.h b/Engine/Source/Runtime/Core/Public/Misc/CoreDelegates.h
--- a/Engine/Source/Runtime/Core/Public/Misc/CoreDelegates.h
+++ b/Engine/Source/Runtime/Core/Public/Misc/CoreDelegates.h
@@ -291,6 +291,9 @@ public:
 	static CORE_API TTSMulticastDelegate<void(const TCHAR* IniFilename, const TCHAR* SectionName)>& TSOnConfigSectionNameRead();
 	static CORE_API TTSMulticastDelegate<void(const FString& IniFilename, const TSet<FString>& SectionNames)>& TSOnConfigSectionsChanged();
 
+	// this calls the UObjectBaseUtility from code where object system is not linked (it also calls TSOnConfigSectionsChanged()!)
+	static CORE_API TMulticastDelegate<void(const TSet<FString>& ModifiedSections, const FString& IniFilename)> ReloadObjectsAfterDynamicConfigChange;
+
 	static CORE_API TMulticastDelegate<void(const TCHAR* SectionName, const TCHAR* IniFilename, uint32 SetBy, bool bAllowCheating)> OnApplyCVarFromIni;
 
 	static CORE_API TMulticastDelegate<void(uint32 ResX, uint32 ResY)> OnSystemResolutionChanged;
diff --git a/Engine/Source/Runtime/CoreUObject/Private/UObject/Obj.cpp b/Engine/Source/Runtime/CoreUObject/Private/UObject/Obj.cpp
--- a/Engine/Source/Runtime/CoreUObject/Private/UObject/Obj.cpp
+++ b/Engine/Source/Runtime/CoreUObject/Private/UObject/Obj.cpp
@@ -2923,11 +2923,11 @@ void UObject::LoadConfig( UClass* ConfigClass/*=NULL*/, const TCHAR* InFilename/
 
 	if ( PropertyToLoad == NULL )
 	{
-		UE_LOG(LogConfig, Verbose, TEXT("(%s) '%s' loading configuration from %s"), *ConfigClass->GetName(), *GetName(), *Filename);
+		UE_LOG(LogConfig, VeryVerbose, TEXT("(%s) '%s' loading configuration from %s"), *ConfigClass->GetName(), *GetName(), *Filename);
 	}
 	else
 	{
-		UE_LOG(LogConfig, Verbose, TEXT("(%s) '%s' loading configuration for property %s from %s"), *ConfigClass->GetName(), *GetName(), *PropertyToLoad->GetName(), *Filename);
+		UE_LOG(LogConfig, VeryVerbose, TEXT("(%s) '%s' loading configuration for property %s from %s"), *ConfigClass->GetName(), *GetName(), *PropertyToLoad->GetName(), *Filename);
 	}
 
 	auto GetConfigValue = [&OverrideConfigFile, &bUseConfigOverride](const TCHAR* ClassSection, const TCHAR* Key, const TCHAR* ConfigName, FString& OutValue)
@@ -3019,7 +3019,7 @@ void UObject::LoadConfig( UClass* ConfigClass/*=NULL*/, const TCHAR* InFilename/
 		// Track if we loaded this config value using special handling (e.g. array or set)
 		bool bProcessedProperty = false;
 
-		UE_LOG(LogConfig, Verbose, TEXT("   Loading value for %s from [%s]"), *Key, *ClassSection);
+		UE_LOG(LogConfig, VeryVerbose, TEXT("   Loading value for %s from [%s]"), *Key, *ClassSection);
 
 		FArrayProperty* Array = CastField<FArrayProperty>(Property);
 		FSetProperty* SetProperty = CastField<FSetProperty>(Property);
@@ -3240,8 +3240,6 @@ void UObject::SaveConfig( uint64 Flags, const TCHAR* InFilename, FConfigCacheIni
 		return;
 	}
 
-	uint32 PropagationFlags = UE::LCPF_None;
-
 	const FString Filename
 	// if a filename was specified, always load from that file
 	=	InFilename
@@ -3293,13 +3291,7 @@ void UObject::SaveConfig( uint64 Flags, const TCHAR* InFilename, FConfigCacheIni
 			if (Property->PropertyFlags & CPF_GlobalConfig)
 			{
 				// call LoadConfig() on child classes if any of the properties were global config
-				PropagationFlags |= UE::LCPF_PropagateToChildDefaultObjects;
 				BaseClass = Property->GetOwnerClass();
-				if ( BaseClass != GetClass() )
-				{
-					// call LoadConfig() on parent classes only if the global config property was declared in a parent class
-					PropagationFlags |= UE::LCPF_ReadParentSections;
-				}
 			}
 
 			FString Key				= Property->GetName();
@@ -3345,7 +3337,7 @@ void UObject::SaveConfig( uint64 Flags, const TCHAR* InFilename, FConfigCacheIni
 				if (Sec)
 				{
 					// Delete the old value for the property in the ConfigCache before (conditionally) adding in the new value
-					Config->RemoveKeyFromSection(*Section, *CompleteKey, PropFileName);
+					Config->ResetKeyInSection(*Section, *CompleteKey, PropFileName);
 				}
 
 				if (!bPropDeprecated && (!bShouldCheckIfIdenticalBeforeAdding || !Property->Identical_InContainer(this, SuperClassDefaultObject)))
@@ -3408,7 +3400,7 @@ void UObject::SaveConfig( uint64 Flags, const TCHAR* InFilename, FConfigCacheIni
 					else
 					{
 						// If we are not writing it to config above, we should make sure that this property isn't stagnant in the cache.
-						Config->RemoveKey( *Section, *Key, PropFileName );
+						Config->ResetKeyInSection( *Section, *Key, PropFileName );
 					}
 				}
 			}
@@ -3426,7 +3418,7 @@ void UObject::SaveConfig( uint64 Flags, const TCHAR* InFilename, FConfigCacheIni
 	// only write out the config file if this is GConfig
 	if (Config == GConfig)
 	{
-		Config->Flush( 0, Filename );
+		Config->Flush(0, *Filename);
 	}
 }
 
@@ -5270,20 +5262,34 @@ bool StaticExec( UWorld* InWorld, const TCHAR* Cmd, FOutputDevice& Ar )
 		// Determine the object/class name
 		if (FParse::Token(Str,ClassName,UE_ARRAY_COUNT(ClassName),1))
 		{
+			UObject* ObjectToReload = nullptr;
+			
 			// Try to find a corresponding class
 			UClass* ClassToReload = FindFirstObject<UClass>(ClassName, EFindFirstObjectOptions::None, ELogVerbosity::Warning, TEXT("StaticExec RELOADCONFIG"));
 			if (ClassToReload)
 			{
-				ClassToReload->ReloadConfig();
+				ObjectToReload = ClassToReload->GetDefaultObject();
 			}
 			else
 			{
 				// If the class is missing, search for an object with that name
-				UObject* ObjectToReload = FindFirstObject<UObject>(ClassName, EFindFirstObjectOptions::None, ELogVerbosity::Warning, TEXT("StaticExec RELOADCONFIG"));
-				if (ObjectToReload)
+				ObjectToReload = FindFirstObject<UObject>(ClassName, EFindFirstObjectOptions::None, ELogVerbosity::Warning, TEXT("StaticExec RELOADCONFIG"));
+			}
+			if (ObjectToReload)
+			{
+				if (ObjectToReload->GetClass()->HasAnyClassFlags(EClassFlags::CLASS_Config))
 				{
+					// reload the object's config
+					FConfigContext Context = FConfigContext::ForceReloadIntoGConfig();
+					Context.Load(*ObjectToReload->GetClass()->GetConfigName());
+					
+					// now updates all the class properties now that the config was reloaded from disk
 					ObjectToReload->ReloadConfig();
 				}
+				else
+				{
+					Ar.Logf(TEXT("Class %s is not a config-enabled class."), *ObjectToReload->GetClass()->GetName());
+				}
 			}
 		}
 		return true;
diff --git a/Engine/Source/Runtime/CoreUObject/Private/UObject/UObjectBaseUtility.cpp b/Engine/Source/Runtime/CoreUObject/Private/UObject/UObjectBaseUtility.cpp
--- a/Engine/Source/Runtime/CoreUObject/Private/UObject/UObjectBaseUtility.cpp
+++ b/Engine/Source/Runtime/CoreUObject/Private/UObject/UObjectBaseUtility.cpp
@@ -18,6 +18,7 @@
 #include "Misc/ConfigCacheIni.h"
 #include "Containers/VersePath.h"
 #include "Misc/CoreDelegates.h"
+#include "Misc/DelayedAutoRegister.h"
 
 /***********************/
 /******** Names ********/
@@ -536,6 +537,8 @@ void UObjectBaseUtility::ReloadObjectsFromModifiedConfigSections(const TSet<FStr
 	// Reload configs so objects get the changes
 	for (const FString& SectionName : ModifiedSections)
 	{
+		UE_LOG(LogConfig, Verbose, TEXT("Found modified section: %s"), *SectionName);
+
 		// @todo: This entire overarching process is very similar in its goals as that of UOnlineHotfixManager::HotfixIniFile.
 		// Could consider a combined refactor of the hotfix manager, the base config cache system, etc. to expose an easier way to support this pattern
 
@@ -597,7 +600,7 @@ void UObjectBaseUtility::ReloadObjectsFromModifiedConfigSections(const TSet<FStr
 
 	auto ReloadObjectImpl = [&NumObjectsReloaded](UObject* ReloadObject)
 	{
-		UE_LOG(LogCore, Verbose, TEXT("Reloading %s"), *ReloadObject->GetPathName());
+		UE_LOG(LogConfig, Verbose, TEXT("Reloading %s"), *ReloadObject->GetPathName());
 		// Intentionally using LoadConfig instead of ReloadConfig, since we do not want to call modify/preeditchange/posteditchange on the objects changed when GIsEditor
 		ReloadObject->LoadConfig(nullptr, nullptr, UE::LCPF_ReloadingConfigData | UE::LCPF_ReadParentSections, nullptr);
 //		ReloadObject->ReloadConfig();
@@ -624,6 +627,10 @@ void UObjectBaseUtility::ReloadObjectsFromModifiedConfigSections(const TSet<FStr
 		*IniFilename, FPlatformTime::Seconds() - StartTime, NumObjectsReloaded);
 }
 
+static FDelayedAutoRegisterHelper GSetupReload(EDelayedRegisterRunPhase::ObjectSystemReady, []
+	{
+		FCoreDelegates::ReloadObjectsAfterDynamicConfigChange.AddStatic(&UObjectBaseUtility::ReloadObjectsFromModifiedConfigSections);
+	});
 
 
 UClass* GetParentNativeClass(UClass* Class)
diff --git a/Engine/Source/Runtime/GameplayTags/Private/GameplayTagsManager.cpp b/Engine/Source/Runtime/GameplayTags/Private/GameplayTagsManager.cpp
--- a/Engine/Source/Runtime/GameplayTags/Private/GameplayTagsManager.cpp
+++ b/Engine/Source/Runtime/GameplayTags/Private/GameplayTagsManager.cpp
@@ -386,20 +386,12 @@ void UGameplayTagsManager::AddTagsFromAdditionalLooseIniFiles(const TArray<FStri
 		{
 			FoundSource->SourceTagList->ConfigFileName = IniFilePath;
 
-			// Check deprecated locations
-			TArray<FString> Tags;
-			if (GConfig->GetArray(TEXT("UserTags"), TEXT("GameplayTags"), Tags, IniFilePath))
-			{
-				for (const FString& Tag : Tags)
-				{
-					FoundSource->SourceTagList->GameplayTagList.AddUnique(FGameplayTagTableRow(FName(*Tag)));
-				}
-			}
-			else
-			{
-				// Load from new ini
-				FoundSource->SourceTagList->LoadConfig(UGameplayTagsList::StaticClass(), *IniFilePath);
-			}
+			FoundSource->SourceTagList->LoadConfig(UGameplayTagsList::StaticClass(), *IniFilePath);
+
+			// we don't actually need this in GConfig because they aren't read from again, and they take a lot of memory,
+			// and aren't tagged with the plugin name, so can't be unloaded along with the plugin anyway, but
+			// since LoadConfig can't take an existing FConfigFile* to load from, we put it into GConfig, then remove it
+			GConfig->Remove(IniFilePath);
 
 #if WITH_EDITOR
 			if (GIsEditor || IsRunningCommandlet()) // Sort tags for UI Purposes but don't sort in -game scenario since this would break compat with noneditor cooked builds
@@ -483,24 +475,6 @@ void UGameplayTagsManager::ConstructGameplayTagTree()
 		{
 			SCOPE_LOG_GAMEPLAYTAGS(TEXT("UGameplayTagsManager::ConstructGameplayTagTree: ImportINI tags"));
 
-			// Copy from deprecated list in DefaultEngine.ini
-			TArray<FString> EngineConfigTags;
-			GConfig->GetArray(TEXT("/Script/GameplayTags.GameplayTagsSettings"), TEXT("+GameplayTags"), EngineConfigTags, GEngineIni);
-			
-			for (const FString& EngineConfigTag : EngineConfigTags)
-			{
-				MutableDefault->GameplayTagList.AddUnique(FGameplayTagTableRow(FName(*EngineConfigTag)));
-			}
-
-			// Copy from deprecated list in DefaultGamplayTags.ini
-			EngineConfigTags.Empty();
-			GConfig->GetArray(TEXT("/Script/GameplayTags.GameplayTagsSettings"), TEXT("+GameplayTags"), EngineConfigTags, MutableDefault->GetDefaultConfigFilename());
-
-			for (const FString& EngineConfigTag : EngineConfigTags)
-			{
-				MutableDefault->GameplayTagList.AddUnique(FGameplayTagTableRow(FName(*EngineConfigTag)));
-			}
-
 #if WITH_EDITOR
 			MutableDefault->SortTags();
 #endif
@@ -527,6 +501,11 @@ void UGameplayTagsManager::ConstructGameplayTagTree()
 			}
 		}
 
+		if (!GIsEditor)
+		{
+			GConfig->SafeUnloadBranch(*GGameplayTagsIni);
+		}
+
 #if WITH_EDITOR
 		// Add any transient editor-only tags
 		for (FName TransientTag : TransientEditorTags)
diff --git a/Engine/Source/Runtime/Projects/Private/PluginManager.cpp b/Engine/Source/Runtime/Projects/Private/PluginManager.cpp
--- a/Engine/Source/Runtime/Projects/Private/PluginManager.cpp
+++ b/Engine/Source/Runtime/Projects/Private/PluginManager.cpp
@@ -791,7 +791,7 @@ void FPluginManager::ReadAllPlugins(FDiscoveredPluginMap& Plugins, const TSet<FS
 
 		SlowTask_ReadAll.EnterProgressFrame(0.5f);
 		// Find "built-in" plugins.  That is, plugins situated right within the Engine directory.
-		TArray<FString> EnginePluginDirs = FPaths::GetExtensionDirs(FPaths::EngineDir(), TEXT("Plugins"), !GIsEditor);
+		TArray<FString> EnginePluginDirs = FPaths::GetExtensionDirs(FPaths::EngineDir(), TEXT("Plugins"));
 		if (OutPluginSources)
 		{
 			OptionalOutPluginRoots.Append(EnginePluginDirs);
@@ -811,7 +811,7 @@ void FPluginManager::ReadAllPlugins(FDiscoveredPluginMap& Plugins, const TSet<FS
 		// assume that the game plugin version is preferred.
 		if (Project != nullptr)
 		{
-			TArray<FString> ProjectPluginDirs = FPaths::GetExtensionDirs(FPaths::GetPath(FPaths::GetProjectFilePath()), TEXT("Plugins"), !GIsEditor);
+			TArray<FString> ProjectPluginDirs = FPaths::GetExtensionDirs(FPaths::GetPath(FPaths::GetProjectFilePath()), TEXT("Plugins"));
 			if (OutPluginSources)
 			{
 				OptionalOutPluginRoots.Append(ProjectPluginDirs);
@@ -1345,14 +1345,12 @@ bool FPluginManager::IntegratePluginsIntoConfig(FConfigCacheIni& ConfigSystem, c
 			{
 				FString BaseConfigFile = *FPaths::GetBaseFilename(ConfigFile);
 
-				// Use GetConfigFilename to find the proper config file to combine into, since it manages command line overrides and path sanitization
-				FString PluginConfigFilename = ConfigSystem.GetConfigFilename(*BaseConfigFile);
-				FConfigFile* FoundConfig = ConfigSystem.FindConfigFile(PluginConfigFilename);
-				if (FoundConfig != nullptr)
+				FConfigBranch* FoundBranch = ConfigSystem.FindBranch(*BaseConfigFile, TEXT(""));
+				if (FoundBranch != nullptr)
 				{
-					UE_LOG(LogPluginManager, Log, TEXT("Found config from plugin[%s] %s"), *Plugin.GetName(), *PluginConfigFilename);
+					UE_LOG(LogPluginManager, Log, TEXT("Found config from plugin[%s] %s"), *Plugin.GetName(), *BaseConfigFile);
 
-					FoundConfig->AddDynamicLayerToHierarchy(FPaths::Combine(PluginConfigDir, ConfigFile));
+					FoundBranch->AddDynamicLayerToHierarchy(FPaths::Combine(PluginConfigDir, ConfigFile));
 				}
 			}
 
