Fixed TUniqueFunction compile error when constructed one with an lvalue.
Fixed some comments.

diff --git a/Engine/Source/Runtime/Core/Public/Templates/Function.h b/Engine/Source/Runtime/Core/Public/Templates/Function.h
--- a/Engine/Source/Runtime/Core/Public/Templates/Function.h
+++ b/Engine/Source/Runtime/Core/Public/Templates/Function.h
@@ -169,6 +169,15 @@ namespace UE::Core::Private::Function
 	template <typename T, bool bOnHeap>
 	struct TFunction_UniqueOwnedObject final : public TFunction_OwnedObject<T, bOnHeap>
 	{
+		/**
+		 * Constructor which creates its T by copying.  T is allowed to be copied into the
+		 * storage, but it will never be copied again after that.
+		 */
+		explicit TFunction_UniqueOwnedObject(const T& InObj)
+			: TFunction_OwnedObject<T, bOnHeap>(InObj)
+		{
+		}
+
 		/**
 		 * Constructor which creates its T by moving.
 		 */
@@ -916,14 +925,14 @@ public:
 	{
 		// This constructor is disabled for TUniqueFunction types so it isn't incorrectly selected as copy/move constructors.
 
-		// This is probably a mistake if you expect TFunction to take a copy of what
+		// This is probably a mistake if you expect TUniqueFunction to take a copy of what
 		// TFunctionRef is bound to, because that's not possible.
 		//
-		// If you really intended to bind a TFunction to a TFunctionRef, you can just
+		// If you really intended to bind a TUniqueFunction to a TFunctionRef, you can just
 		// wrap it in a lambda (and thus it's clear you're just binding to a call to another
 		// reference):
 		//
-		// TFunction<int32(float)> MyFunction = [MyFunctionRef](float F) { return MyFunctionRef(F); };
+		// TUniqueFunction<int32(float)> MyFunction = [MyFunctionRef](float F) { return MyFunctionRef(F); };
 		static_assert(!TIsTFunctionRef<std::decay_t<FunctorType>>::Value, "Cannot construct a TUniqueFunction from a TFunctionRef");
 	}
 
