Always open localtime_r for now

Later will need to consider have a consistent time when rolling back

diff --git a/Engine/Source/Runtime/Core/Private/AutoRTFM/StdLibSupport.cpp b/Engine/Source/Runtime/Core/Private/AutoRTFM/StdLibSupport.cpp
--- a/Engine/Source/Runtime/Core/Private/AutoRTFM/StdLibSupport.cpp
+++ b/Engine/Source/Runtime/Core/Private/AutoRTFM/StdLibSupport.cpp
@@ -326,6 +326,7 @@ UE_AUTORTFM_REGISTER_OPEN_FUNCTION(TlsSetValue);
 UE_AUTORTFM_REGISTER_SELF_FUNCTION(clock_gettime);
 UE_AUTORTFM_REGISTER_SELF_FUNCTION(gettimeofday);
 UE_AUTORTFM_REGISTER_SELF_FUNCTION(gmtime_r);
+UE_AUTORTFM_REGISTER_SELF_FUNCTION(localtime_r);
 UE_AUTORTFM_REGISTER_SELF_FUNCTION(bcmp);
 
 UE_AUTORTFM_REGISTER_SELF_FUNCTION(pthread_getspecific);
