Fix regression in reference replacement archives behavior wrt delegates - return to old invariant that loading archives are compacted after operator<<, !IsSaving() archives operate on data in place, not on proxies. This is a partial backout of 33463669
#rnx

diff --git a/Engine/Source/Runtime/Core/Public/UObject/ScriptDelegates.h b/Engine/Source/Runtime/Core/Public/UObject/ScriptDelegates.h
--- a/Engine/Source/Runtime/Core/Public/UObject/ScriptDelegates.h
+++ b/Engine/Source/Runtime/Core/Public/UObject/ScriptDelegates.h
@@ -837,31 +837,31 @@ public:
 	/** Multi-cast delegate serialization */
 	friend FArchive& operator<<( FArchive& Ar, TMulticastScriptDelegate& D )
 	{
-		// Warning: This order of if/else is necessary because there are archivers (eg FArchiveGatherExternalActorRefs)
-		// that do not report as either IsLoading or IsSaving but should be treated as IsSaving in this case
-		if (Ar.IsLoading())
+		// Special case to avoid taking a lock on empty script delegate.
+		// This is required for avoiding asserts on EmptyDelegate serialization.
+		if (Ar.IsSaving() && D.InvocationList.Num() == 0)
 		{
-			FWriteAccessScope WriteScope = D.GetWriteAccessScope();
+			FReadAccessScope ReadScope = D.GetReadAccessScope();
+
 			Ar << D.InvocationList;
-			// After loading the delegate, clean up the list to make sure there are no bad object references
-			D.CompactInvocationList();
 		}
 		else
 		{
-			FReadAccessScope ReadScope = D.GetReadAccessScope();
-			// When saving the delegate, clean up the list to make sure there are no bad object references
-			// Don't do this in place because we don't want to require a write lock
-			typedef TArray<UnicastDelegateType, TInlineAllocator<4>> FInlineInvocationList;
-			FInlineInvocationList CompactedList;
+			FWriteAccessScope WriteScope = D.GetWriteAccessScope();
 
-			for (const UnicastDelegateType& Delegate : D.InvocationList)
+			if( Ar.IsSaving() )
 			{
-				if (!Delegate.IsCompactable())
-				{
-					CompactedList.Add(Delegate);
-				}
+				// When saving the delegate, clean up the list to make sure there are no bad object references
+				D.CompactInvocationList();
+			}
+
+			Ar << D.InvocationList;
+
+			if( Ar.IsLoading() )
+			{
+				// After loading the delegate, clean up the list to make sure there are no bad object references
+				D.CompactInvocationList();
 			}
-			Ar << CompactedList;
 		}
 
 		return Ar;
