Python parse mode

diff --git a/Engine/Source/Runtime/Core/Private/Misc/ConfigCacheIni.cpp b/Engine/Source/Runtime/Core/Private/Misc/ConfigCacheIni.cpp
--- a/Engine/Source/Runtime/Core/Private/Misc/ConfigCacheIni.cpp
+++ b/Engine/Source/Runtime/Core/Private/Misc/ConfigCacheIni.cpp
@@ -890,6 +890,7 @@ FConfigFile::FConfigFile()
     : Dirty( false )
     , NoSave( false )
     , bHasPlatformName( false )
+    , bPythonConfigParserMode( false )
     , bCanSaveAllSections( true )
     , Name( NAME_None )
 {
@@ -931,6 +932,7 @@ FConfigFile& FConfigFile::operator=(const FConfigFile& Other)
 	Dirty = Other.Dirty;
 	NoSave = Other.NoSave;
 	bHasPlatformName = Other.bHasPlatformName;
+	bPythonConfigParserMode = Other.bPythonConfigParserMode;
 	bCanSaveAllSections = Other.bCanSaveAllSections;
 
 	// LoadType is not copied; each FConfigFile has to set it itself
@@ -1405,6 +1407,7 @@ void FillFileFromBuffer(FileType* File, FStringView Buffer, bool bHandleSymbolCo
 	SectionType* CurrentSection = nullptr;
 	
 	FString CurrentSectionName;
+	FName CurrentKeyName;
 	TMap<FString, FString>* CurrentKeyRemap = nullptr;
 	TStringBuilder<128> TheLine;
 	FString ProcessedValue;
@@ -1420,7 +1423,7 @@ void FillFileFromBuffer(FileType* File, FStringView Buffer, bool bHandleSymbolCo
 
 		// read the next line
 		int32 LinesConsumed = 0;
-		FParse::LineExtended(&Ptr, /* reset */ TheLine, LinesConsumed, false);
+		FParse::LineExtended(&Ptr, /* reset */ TheLine, LinesConsumed, !!File->bPythonConfigParserMode);
 		if (Ptr == nullptr || *Ptr == 0)
 		{
 			Done = true;
@@ -1462,6 +1465,7 @@ void FillFileFromBuffer(FileType* File, FStringView Buffer, bool bHandleSymbolCo
 
 			// If we don't have an existing section by this name, add one
 			CurrentSectionName = Start;
+			CurrentKeyName = NAME_None;
 			
 			// lookup to see if there is an entry in the SectionName remap
 			const FString* FoundRemap;
@@ -1492,89 +1496,102 @@ void FillFileFromBuffer(FileType* File, FStringView Buffer, bool bHandleSymbolCo
 			// ignore [comment] lines that start with ;
 			if(*Start != (TCHAR)';')
 			{
-				Value = FCString::Strstr(Start,TEXT("="));
+				// If we're in python mode and the line starts with whitespace
+				// then we should consider it a part of the prior key
+				if (File->bPythonConfigParserMode && !CurrentKeyName.IsNone() && FChar::IsWhitespace(*Start))
+				{
+					Value = Start;
+				}
+				else
+				{
+					Value = FCString::Strstr(Start,TEXT("="));
+				}
 			}
 
 			// Ignore any lines that don't contain a key-value pair
 			if( Value )
 			{
-				// Terminate the property name, advancing past the =
-				*Value++ = TEXT('\0');
-
-				// strip leading whitespace from the property name
-				while ( *Start && FChar::IsWhitespace(*Start) )
-				{
-					Start++;
-				}
-
-				// ~ is a packaging and should be skipped at runtime
-				if (Start[0] == '~')
-				{
-					Start++;
-				}
+				SectionType* OriginalCurrentSection = CurrentSection;
 
 				// determine how this line will be merged
-				FConfigValue::EValueType Command;
+				// when we don't want commands, the default action is to add new entries (this is for standalone ini files that have arrays,
+				// without any + cmds) - there's no difference between a single value and an array of 1 (in terms of the Config system)
+				FConfigValue::EValueType Command = FConfigValue::EValueType::ArrayAdd;
 
-				if (!bHandleSymbolCommands)
+				// Value will be Start in the python configparser extending case in which case
+				// we want to continue using the CurrentKeyName
+				if (Value != Start)
 				{
-					// when we don't want commands, the default action is to add new entries (this is for standalone ini files that have arrays,
-					// without any + cmds) - there's no difference between a single value and an array of 1 (in terms of the Config system)
-					Command = FConfigValue::EValueType::ArrayAdd;
-				}
-				else
-				{
-					TCHAR Cmd = Start[0];
-					if (Cmd == '+' || Cmd == '-' || Cmd == '.' || Cmd == '!' || Cmd == '@' || Cmd == '*')
+					// Terminate the property name, advancing past the =
+					*Value++ = TEXT('\0');
+
+					// strip leading whitespace from the property name
+					while (*Start && FChar::IsWhitespace(*Start))
 					{
 						Start++;
 					}
-					else
+
+					// ~ is a packaging and should be skipped at runtime
+					if (Start[0] == '~')
 					{
-						Cmd = TEXT('\0');
+						Start++;
 					}
-					
-					// turn into a command
-					FConfigValue::EValueType* Lookup = CommandLookup.Find(Cmd);
-					if (Lookup == nullptr)
+
+					if (bHandleSymbolCommands)
 					{
-						UE_LOG(LogConfig, Log, TEXT("Found unknown ini command %c in an ini"), Cmd);
-						continue;
+						TCHAR Cmd = Start[0];
+						if (Cmd == '+' || Cmd == '-' || Cmd == '.' || Cmd == '!' || Cmd == '@' || Cmd == '*')
+						{
+							Start++;
+						}
+						else
+						{
+							Cmd = TEXT('\0');
+						}
+					
+						// turn into a command
+						FConfigValue::EValueType* Lookup = CommandLookup.Find(Cmd);
+						if (Lookup == nullptr)
+						{
+							UE_LOG(LogConfig, Log, TEXT("Found unknown ini command %c in an ini"), Cmd);
+							continue;
+						}
+						Command = *Lookup;
 					}
-					Command = *Lookup;
-				}
 
-				// Strip trailing spaces from the property name.
-				while( *Start && FChar::IsWhitespace(Start[FCString::Strlen(Start)-1]) )
-				{
-					Start[FCString::Strlen(Start)-1] = TEXT('\0');
-				}
-
-				const TCHAR* KeyName = Start;
-				SectionType* OriginalCurrentSection = CurrentSection;
-				// look up for key remap
-				if (CurrentKeyRemap != nullptr)
-				{
-					const FString* FoundRemap;
-					if ((FoundRemap = CurrentKeyRemap->Find(KeyName)) != nullptr)
+					// Strip trailing spaces from the property name.
+					while (*Start && FChar::IsWhitespace(Start[FCString::Strlen(Start)-1]))
 					{
-						WarnAboutKeyRemap(KeyName, *FoundRemap, CurrentSectionName, FileHint);
-
-						// the Remap will not ever reallocate, so we can just point right into the FString
-						KeyName = **FoundRemap;
+						Start[FCString::Strlen(Start)-1] = TEXT('\0');
+					}
 
-						// look for a section:name remap
-						int32 ColonLoc;
-						if (FoundRemap->FindChar(':', ColonLoc))
+					const TCHAR* KeyName = Start;
+					// look up for key remap
+					if (CurrentKeyRemap != nullptr)
+					{
+						const FString* FoundRemap;
+						if ((FoundRemap = CurrentKeyRemap->Find(KeyName)) != nullptr)
 						{
-							// find or create a section for name before the :
-							CurrentSection = File->FindOrAddSectionInternal(*FoundRemap->Mid(0, ColonLoc));
-							// the name can still point right into the FString, but right after the :
-							KeyName = **FoundRemap + ColonLoc + 1;
+							WarnAboutKeyRemap(KeyName, *FoundRemap, CurrentSectionName, FileHint);
+
+							// the Remap will not ever reallocate, so we can just point right into the FString
+							KeyName = **FoundRemap;
+
+							// look for a section:name remap
+							int32 ColonLoc;
+							if (FoundRemap->FindChar(':', ColonLoc))
+							{
+								// find or create a section for name before the :
+								CurrentSection = File->FindOrAddSectionInternal(*FoundRemap->Mid(0, ColonLoc));
+								// the name can still point right into the FString, but right after the :
+								KeyName = **FoundRemap + ColonLoc + 1;
+							}
 						}
 					}
+
+					CurrentKeyName = FName(KeyName);
 				}
-				
+
 				// Strip leading whitespace from the property value
 				while ( *Value && FChar::IsWhitespace(*Value) )
 				{
@@ -1599,8 +1616,7 @@ void FillFileFromBuffer(FileType* File, FStringView Buffer, bool bHandleSymbolCo
 					ProcessedValue = Value;
 				}
 
-				const FName Key(KeyName);
-				File->ProcessCommand(CurrentSection, FStringView(CurrentSectionName), Command, Key, MoveTemp(ProcessedValue));
+				File->ProcessCommand(CurrentSection, FStringView(CurrentSectionName), Command, CurrentKeyName, MoveTemp(ProcessedValue));
 				
 				// restore the current section, in case it was overridden
 				CurrentSection = OriginalCurrentSection;
diff --git a/Engine/Source/Runtime/Core/Public/Misc/ConfigCacheIni.h b/Engine/Source/Runtime/Core/Public/Misc/ConfigCacheIni.h
--- a/Engine/Source/Runtime/Core/Public/Misc/ConfigCacheIni.h
+++ b/Engine/Source/Runtime/Core/Public/Misc/ConfigCacheIni.h
@@ -525,6 +525,7 @@ public:
 	FName Tag;
 	
 	uint8 Dirty;
+	uint8 bPythonConfigParserMode = false;
 
 	// used to determine if existing settings should be removed from a FConfigFile before applying this stream - used for compatibility with
 	// how Saved config files are stored (replace the static layers values fully if at least one key exists)
@@ -543,6 +544,7 @@ public:
 	bool Dirty : 1; // = false;
 	bool NoSave : 1; // = false;
 	bool bHasPlatformName : 1; // = false;
+	bool bPythonConfigParserMode : 1; // = false;
 
 	// by default, we allow saving - this is going to be applied to config files that are not loaded from disk
 	// (when loading, this will get set to false, and then the ini sections will be checked)
