Removed a lot of inheritance layers from TFunction implementation and replaced it with if constexpr blocks.

diff --git a/Engine/Source/Runtime/Core/Public/Templates/Function.h b/Engine/Source/Runtime/Core/Public/Templates/Function.h
--- a/Engine/Source/Runtime/Core/Public/Templates/Function.h
+++ b/Engine/Source/Runtime/Core/Public/Templates/Function.h
@@ -40,7 +40,7 @@
  */
 namespace UE::Core::Private::Function
 {
-	template <typename T, bool bOnHeap>
+	template <typename T, bool bUnique, bool bOnHeap>
 	struct TFunction_OwnedObject;
 
 	template <bool bUnique>
@@ -72,61 +72,12 @@ namespace UE::Core::Private::Function
 		virtual ~IFunction_OwnedObject() = default;
 	};
 
-	/**
-	 * Common interface to a callable object owned by TFunction.
-	 */
-	template <typename T>
-	struct IFunction_OwnedObject_OnHeap : public IFunction_OwnedObject
+	template <typename T, bool bUnique, bool bOnHeap>
+	struct TFunction_OwnedObject : public IFunction_OwnedObject
 	{
-		/**
-		 * Destructor.
-		 */
-		virtual void Destroy() override
-		{
-			void* This = this;
-			this->~IFunction_OwnedObject_OnHeap();
-			FMemory::Free(This);
-		}
-
-		~IFunction_OwnedObject_OnHeap() override
-		{
-			// It is not necessary to define this destructor but MSVC will
-			// erroneously issue warning C5046 without it.
-		}
-	};
-
-	/**
-	 * Common interface to a callable object owned by TFunction.
-	 */
-	template <typename T>
-	struct IFunction_OwnedObject_Inline : public IFunction_OwnedObject
-	{
-		/**
-		 * Destructor.
-		 */
-		virtual void Destroy() override
-		{
-			this->~IFunction_OwnedObject_Inline();
-		}
-
-		~IFunction_OwnedObject_Inline() override
-		{
-			// It is not necessary to define this destructor but MSVC will
-			// erroneously issue warning C5046 without it.
-		}
-	};
-
-	template <typename T, bool bOnHeap>
-	struct TFunction_OwnedObject : public
-#if TFUNCTION_USES_INLINE_STORAGE
-		std::conditional_t<bOnHeap, IFunction_OwnedObject_OnHeap<T>, IFunction_OwnedObject_Inline<T>>
-#else
-		IFunction_OwnedObject_OnHeap<T>
-#endif
-	{
-		template <typename... ArgTypes>
-		explicit TFunction_OwnedObject(ArgTypes&&... Args)
-			: Obj(Forward<ArgTypes>(Args)...)
+		template <typename ArgType>
+		explicit TFunction_OwnedObject(ArgType&& Arg)
+			: Obj(Forward<ArgType>(Arg))
 		{
 		}
 
@@ -135,63 +86,29 @@ namespace UE::Core::Private::Function
 			return &Obj;
 		}
 
-		T Obj;
-	};
-
-	/**
-	 * Implementation of IFunction_OwnedObject for a given copyable T.
-	 */
-	template <typename T, bool bOnHeap>
-	struct TFunction_CopyableOwnedObject final : public TFunction_OwnedObject<T, bOnHeap>
-	{
-		/**
-		 * Constructor which creates its T by copying.
-		 */
-		explicit TFunction_CopyableOwnedObject(const T& InObj)
-			: TFunction_OwnedObject<T, bOnHeap>(InObj)
-		{
-		}
-
-		/**
-		 * Constructor which creates its T by moving.
-		 */
-		explicit TFunction_CopyableOwnedObject(T&& InObj)
-			: TFunction_OwnedObject<T, bOnHeap>(MoveTemp(InObj))
-		{
-		}
-
 		void* CloneToEmptyStorage(void* UntypedStorage) const override;
-	};
 
-	/**
-	 * Implementation of IFunction_OwnedObject for a given non-copyable T.
-	 */
-	template <typename T, bool bOnHeap>
-	struct TFunction_UniqueOwnedObject final : public TFunction_OwnedObject<T, bOnHeap>
-	{
-		/**
-		 * Constructor which creates its T by copying.  T is allowed to be copied into the
-		 * storage, but it will never be copied again after that.
-		 */
-		explicit TFunction_UniqueOwnedObject(const T& InObj)
-			: TFunction_OwnedObject<T, bOnHeap>(InObj)
+		virtual void Destroy() override
 		{
+			if constexpr (bOnHeap)
+			{
+			    void* This = this;
+			    this->~TFunction_OwnedObject();
+			    FMemory::Free(This);
+			}
+            else
+            {
+				this->~TFunction_OwnedObject();
+            }
 		}
 
-		/**
-		 * Constructor which creates its T by moving.
-		 */
-		explicit TFunction_UniqueOwnedObject(T&& InObj)
-			: TFunction_OwnedObject<T, bOnHeap>(MoveTemp(InObj))
+		~TFunction_OwnedObject() override
 		{
+			// It is not necessary to define this destructor but MSVC will
+			// erroneously issue warning C5046 without it.
 		}
 
-		void* CloneToEmptyStorage(void* Storage) const override
-		{
-			// Should never get here - copy functions are deleted for TUniqueFunction
-			check(false);
-			return nullptr;
-		}
+		T Obj;
 	};
 
 	template <typename T>
@@ -210,24 +127,6 @@ namespace UE::Core::Private::Function
 		}
 	}
 
-	template <typename FunctorType, bool bUnique, bool bOnHeap>
-	struct TStorageOwnerType;
-
-	template <typename FunctorType, bool bOnHeap>
-	struct TStorageOwnerType<FunctorType, true, bOnHeap>
-	{
-		using Type = TFunction_UniqueOwnedObject<std::decay_t<FunctorType>, bOnHeap>;
-	};
-
-	template <typename FunctorType, bool bOnHeap>
-	struct TStorageOwnerType<FunctorType, false, bOnHeap>
-	{
-		using Type = TFunction_CopyableOwnedObject<std::decay_t<FunctorType>, bOnHeap>;
-	};
-
-	template <typename FunctorType, bool bUnique, bool bOnHeap>
-	using TStorageOwnerTypeT = typename TStorageOwnerType<FunctorType, bUnique, bOnHeap>::Type;
-
 	struct FFunctionStorage
 	{
 		constexpr static bool bCanBeNull = true;
@@ -316,22 +215,24 @@ namespace UE::Core::Private::Function
 		template <typename FunctorType>
 		std::decay_t<FunctorType>* Bind(FunctorType&& InFunc)
 		{
+			using DecayedFunctorType = std::decay_t<FunctorType>;
+
 			if (!IsBound(InFunc))
 			{
 				return nullptr;
 			}
 
 #if TFUNCTION_USES_INLINE_STORAGE
-			constexpr bool bUseInline = sizeof(TStorageOwnerTypeT<FunctorType, bUnique, false>) <= TFUNCTION_INLINE_SIZE;
+			constexpr bool bOnHeap = sizeof(TFunction_OwnedObject<DecayedFunctorType, bUnique, false>) > TFUNCTION_INLINE_SIZE;
 #else
-			constexpr bool bUseInline = false;
+			constexpr bool bOnHeap = true;
 #endif
 
-			using OwnedType = TStorageOwnerTypeT<FunctorType, bUnique, !bUseInline>;
+			using OwnedType = TFunction_OwnedObject<DecayedFunctorType, bUnique, bOnHeap>;
 
 			void* NewAlloc;
 #if TFUNCTION_USES_INLINE_STORAGE
-			if constexpr (bUseInline)
+			if constexpr (!bOnHeap)
 			{
 				NewAlloc = &InlineAllocation;
 			}
@@ -348,28 +249,37 @@ namespace UE::Core::Private::Function
 		}
 	};
 
-	template <typename T, bool bOnHeap>
-	void* TFunction_CopyableOwnedObject<T, bOnHeap>::CloneToEmptyStorage(void* UntypedStorage) const
+	template <typename T, bool bUnique, bool bOnHeap>
+	void* TFunction_OwnedObject<T, bUnique, bOnHeap>::CloneToEmptyStorage(void* UntypedStorage) const
 	{
-		TFunctionStorage<false>& Storage = *(TFunctionStorage<false>*)UntypedStorage;
-
-		void* NewAlloc;
-		#if TFUNCTION_USES_INLINE_STORAGE
-		if /* constexpr */ (!bOnHeap)
+		if constexpr (bUnique)
 		{
-			NewAlloc = &Storage.InlineAllocation;
+			// Should never get here - copy functions are deleted for TUniqueFunction
+			check(false);
+			return nullptr;
 		}
 		else
-		#endif
 		{
-			NewAlloc = FMemory::Malloc(sizeof(TFunction_CopyableOwnedObject), alignof(TFunction_CopyableOwnedObject));
-			Storage.HeapAllocation = NewAlloc;
-			CA_ASSUME(NewAlloc);
-		}
+			TFunctionStorage<false>& Storage = *(TFunctionStorage<false>*)UntypedStorage;
+
+			void* NewAlloc;
+#if TFUNCTION_USES_INLINE_STORAGE
+			if constexpr (!bOnHeap)
+			{
+				NewAlloc = &Storage.InlineAllocation;
+			}
+			else
+#endif
+			{
+				NewAlloc = FMemory::Malloc(sizeof(TFunction_OwnedObject), alignof(TFunction_OwnedObject));
+				Storage.HeapAllocation = NewAlloc;
+				CA_ASSUME(NewAlloc);
+			}
 
-		auto* NewOwned = ::new (NewAlloc) TFunction_CopyableOwnedObject(this->Obj);
+			auto* NewOwned = ::new (NewAlloc) TFunction_OwnedObject(this->Obj);
 
-		return &NewOwned->Obj;
+			return &NewOwned->Obj;
+		}
 	}
 
 	#if UE_ENABLE_TFUNCTIONREF_VISUALIZATION
@@ -395,24 +305,19 @@ namespace UE::Core::Private::Function
 	/**
 	 * A class which is used to instantiate the code needed to call a bound function.
 	 */
-	template <typename Functor, typename FuncType>
-	struct TFunctionRefCaller;
-
 	template <typename Functor, typename Ret, typename... ParamTypes>
-	struct TFunctionRefCaller<Functor, Ret (ParamTypes...)>
+	struct TFunctionRefCaller
 	{
 		static Ret Call(void* Obj, ParamTypes&... Params)
 		{
-			return Invoke(*(Functor*)Obj, Forward<ParamTypes>(Params)...);
-		}
-	};
-
-	template <typename Functor, typename... ParamTypes>
-	struct TFunctionRefCaller<Functor, void (ParamTypes...)>
-	{
-		static void Call(void* Obj, ParamTypes&... Params)
-		{
-			Invoke(*(Functor*)Obj, Forward<ParamTypes>(Params)...);
+			if constexpr (std::is_void_v<Ret>)
+			{
+				Invoke(*(Functor*)Obj, Forward<ParamTypes>(Params)...);
+			}
+			else
+			{
+				return Invoke(*(Functor*)Obj, Forward<ParamTypes>(Params)...);
+			}
 		}
 	};
 
@@ -533,7 +438,7 @@ namespace UE::Core::Private::Function
 
 			using DecayedFunctorType = typename TRemovePointer<decltype(Binding)>::Type;
 
-			Callable = &TFunctionRefCaller<DecayedFunctorType, Ret (ParamTypes...)>::Call;
+			Callable = &TFunctionRefCaller<DecayedFunctorType, Ret, ParamTypes...>::Call;
 
 			#if UE_ENABLE_TFUNCTIONREF_VISUALIZATION
 				// We placement new over the top of the same object each time.  This is illegal,
