[TaskGraph]
 - Activate new frontend by default allowing the old taskgraph to benefit from retraction

diff --git a/Engine/Source/Runtime/Core/Public/Async/TaskGraphDefinitions.h b/Engine/Source/Runtime/Core/Public/Async/TaskGraphDefinitions.h
--- a/Engine/Source/Runtime/Core/Public/Async/TaskGraphDefinitions.h
+++ b/Engine/Source/Runtime/Core/Public/Async/TaskGraphDefinitions.h
@@ -2,6 +2,6 @@
 #pragma once
 
 #if !defined(TASKGRAPH_NEW_FRONTEND)
-#define TASKGRAPH_NEW_FRONTEND 0
+#define TASKGRAPH_NEW_FRONTEND 1
 #endif
 
