Making TMulticastScriptDelegate archive stream operator a bit more conservative about what types of access rights are required based on the type of archive being processed. An archive marked IsSaving will now only require read access. This code has been restructed from a previous attempt to properly account for reference rewriter archives that are !IsSaving and !IsLoading but need write access. Made a corresponding change in FArchiveGatherExternalActorRefs to ensure that it reports as IsSaving and thus only needs read access. This enables use during concurrent save.

diff --git a/Engine/Source/Runtime/Core/Public/UObject/ScriptDelegates.h b/Engine/Source/Runtime/Core/Public/UObject/ScriptDelegates.h
--- a/Engine/Source/Runtime/Core/Public/UObject/ScriptDelegates.h
+++ b/Engine/Source/Runtime/Core/Public/UObject/ScriptDelegates.h
@@ -835,35 +835,35 @@ public:
 	}
 
 	/** Multi-cast delegate serialization */
-	friend FArchive& operator<<( FArchive& Ar, TMulticastScriptDelegate& D )
+	friend FArchive& operator<<(FArchive& Ar, TMulticastScriptDelegate& D)
 	{
-		// Special case to avoid taking a lock on empty script delegate.
-		// This is required for avoiding asserts on EmptyDelegate serialization.
-		if (Ar.IsSaving() && D.InvocationList.Num() == 0)
-		{
-			FReadAccessScope ReadScope = D.GetReadAccessScope();
-
-			Ar << D.InvocationList;
-		}
-		else
+		// Note that !IsSaving is not the same as IsLoading. See, e.g., FArchiveReplaceObjectRef
+		if (!Ar.IsSaving())
 		{
 			FWriteAccessScope WriteScope = D.GetWriteAccessScope();
-
-			if( Ar.IsSaving() )
+			Ar << D.InvocationList;
+			// After loading the delegate, clean up the list to make sure there are no bad object references
+			if (Ar.IsLoading())
 			{
-				// When saving the delegate, clean up the list to make sure there are no bad object references
 				D.CompactInvocationList();
 			}
-
-			Ar << D.InvocationList;
-
-			if( Ar.IsLoading() )
+		}
+		else
+		{
+			FReadAccessScope ReadScope = D.GetReadAccessScope();
+			// When saving the delegate, clean up the list to make sure there are no bad object references
+			// Don't do this in place because we don't want to require a write lock
+			typedef TArray<UnicastDelegateType, TInlineAllocator<4>> FInlineInvocationList;
+			FInlineInvocationList CompactedList;
+			for (const UnicastDelegateType& Delegate : D.InvocationList)
 			{
-				// After loading the delegate, clean up the list to make sure there are no bad object references
-				D.CompactInvocationList();
+				if (!Delegate.IsCompactable())
+				{
+					CompactedList.Add(Delegate);
+				}
 			}
+			Ar << CompactedList;
 		}
-
 		return Ar;
 	}
 
diff --git a/Engine/Source/Runtime/Engine/Private/ActorReferencesUtils.cpp b/Engine/Source/Runtime/Engine/Private/ActorReferencesUtils.cpp
--- a/Engine/Source/Runtime/Engine/Private/ActorReferencesUtils.cpp
+++ b/Engine/Source/Runtime/Engine/Private/ActorReferencesUtils.cpp
@@ -38,6 +38,9 @@ public:
 		// Don't trigger serialization of compilable assets
 		SetShouldSkipCompilingAssets(true);
 
+		// We are a collector and don't modify the objects we visit
+		SetIsSaving(true);
+
 		ArIgnoreOuterRef = true;
 		ArIsObjectReferenceCollector = true;
 		ArShouldSkipBulkData = true;
