Added uint32 Slow_NextAllocBits(uint32 NumBits, uint64 StartIndex) to FBitTree, which finds the first available area of size NumBits at or after StartIndex, and returns the index of the first bit in the area. NumBits must be a power of two or a multiple of 64, and the function only searches regions that are aligned to the requested number of bits.

This function is slow and only suitable for infrequent allocations.

diff --git a/Engine/Source/Runtime/Core/Private/HAL/MallocBinnedCommon.cpp b/Engine/Source/Runtime/Core/Private/HAL/MallocBinnedCommon.cpp
--- a/Engine/Source/Runtime/Core/Private/HAL/MallocBinnedCommon.cpp
+++ b/Engine/Source/Runtime/Core/Private/HAL/MallocBinnedCommon.cpp
@@ -445,6 +445,81 @@ uint32 FBitTree::CountOnes(uint32 UpTo) const
 	return Result;
 }
 
+/**
+ * Finds a contiguous span of unallocated bits.
+ * NumBits must be a power of two or a multiple of 64.
+ * Only checks regions aligned to min(NumBits, 64).
+ * 
+ * Warning, slow!
+ * Requires a linear search along the bottom row! O(Capacity / min(NumBits,64)) iterations.
+ * 
+ * Returns the index of the first unallocated bit in the span.
+ */
+uint32 FBitTree::Slow_NextAllocBits(uint32 NumBits, uint64 StartIndex)
+{
+	check(FMath::IsPowerOfTwo(NumBits) || NumBits % 64 == 0);
+
+	uint32 Offset = OffsetOfLastRow + StartIndex / 64;
+	uint32 MaxOffset = OffsetOfLastRow + DesiredCapacity / 64;
+
+	// If the number of bits is >= 64, we can search int-by-int instead of bit-by-bit on the lowest row
+	if (NumBits >= 64)
+	{
+		uint32 NumInts = NumBits / 64; // Number of uint64s required to hold our allocation
+		uint32 FreeInts = 0; // Number of free uint64s we've encountered in a row
+		
+		while (Offset < MaxOffset)
+		{
+			// Increment FreeInts if the uint64 is empty, otherwise reset to 0
+			FreeInts = Bits[Offset] ? 0 : (FreeInts + 1);
+
+			if (FreeInts == NumInts)
+			{
+				// Offset will be pointing at the last uint64 in the free span - correct to point at the first
+				return ((Offset - OffsetOfLastRow) - (NumInts - 1)) * 64;
+			}
+			
+			Offset++;
+		}
+	}
+
+	// If the number of bits is < 64, we need to loop over ints and then power of two bits within those ints
+	else
+	{
+		uint32 SlotsPerInt = 64 / NumBits; // How many possible positions for our allocation there are per uint64
+
+		while (Offset < MaxOffset)
+		{
+			// Allocation of bits within a given uint64 goes RIGHT-TO-LEFT
+			// Create right-aligned mask based on number of bits
+			uint64 Mask = (1ull << NumBits) - 1;
+
+			// Shift that mask left to check each aligned section of the int
+			// Return offset + bit position if we find a free section
+			for (uint32 i = 0; i < SlotsPerInt; i++)
+			{
+				if (!(Mask & Bits[Offset]))
+				{
+					uint64 Result = (Offset - OffsetOfLastRow) * 64 + (i * NumBits);
+
+					// Ensure we don't return results before our requested start index
+					if (Result >= StartIndex) {
+						return Result;
+					}
+				}
+
+				// Shift mask left
+				Mask <<= NumBits;
+			}
+
+			Offset++;
+		}
+	}
+
+	check(false); // If we get here, no aligned span of this size is available
+	return MAX_uint32;
+}
+
 #endif
 
 float GMallocBinnedFlushThreadCacheMaxWaitTime = 0.2f;
diff --git a/Engine/Source/Runtime/Core/Public/HAL/MallocBinnedCommon.h b/Engine/Source/Runtime/Core/Public/HAL/MallocBinnedCommon.h
--- a/Engine/Source/Runtime/Core/Public/HAL/MallocBinnedCommon.h
+++ b/Engine/Source/Runtime/Core/Public/HAL/MallocBinnedCommon.h
@@ -89,6 +89,8 @@ public:
 	uint32 NextAllocBit(uint32 StartIndex) const;
 	void FreeBit(uint32 Index);
 	uint32 CountOnes(uint32 UpTo) const;
+
+	uint32 Slow_NextAllocBits(uint32 NumBits, uint64 StartIndex); // Warning, slow! NumBits must be a power of two or a multiple of 64.
 };
 
 struct FSizeTableEntry
