Making the GEngineStartupModuleLoadingComplete flag into an atomic.

diff --git a/Engine/Source/Runtime/Core/Private/Misc/CoreGlobals.cpp b/Engine/Source/Runtime/Core/Private/Misc/CoreGlobals.cpp
--- a/Engine/Source/Runtime/Core/Private/Misc/CoreGlobals.cpp
+++ b/Engine/Source/Runtime/Core/Private/Misc/CoreGlobals.cpp
@@ -293,17 +293,17 @@ static void appNoop()
 {
 }
 
-bool GEngineStartupModuleLoadingComplete = false;
+std::atomic<bool> GEngineStartupModuleLoadingComplete = false;
 CORE_API bool IsEngineStartupModuleLoadingComplete()
 {
-	return GEngineStartupModuleLoadingComplete;
+	return GEngineStartupModuleLoadingComplete.load(std::memory_order_acquire);
 }
 
 CORE_API void SetEngineStartupModuleLoadingComplete()
 {
-	if (ensure(!GEngineStartupModuleLoadingComplete))
+	if (ensure(!GEngineStartupModuleLoadingComplete.load(std::memory_order_relaxed)))
 	{
-		GEngineStartupModuleLoadingComplete = true;
+		GEngineStartupModuleLoadingComplete.store(true, std::memory_order_release);
 		SCOPED_BOOT_TIMING("OnAllModuleLoadingPhasesComplete.Broadcast");
 		FCoreDelegates::OnAllModuleLoadingPhasesComplete.Broadcast();
 	}
