[CpuProfiler]
 - Fix TSAN warning and provide a bit of context of why that race is perfectly fine.

#rnx

diff --git a/Engine/Source/Runtime/Core/Public/ProfilingDebugging/CpuProfilerTrace.h b/Engine/Source/Runtime/Core/Public/ProfilingDebugging/CpuProfilerTrace.h
--- a/Engine/Source/Runtime/Core/Public/ProfilingDebugging/CpuProfilerTrace.h
+++ b/Engine/Source/Runtime/Core/Public/ProfilingDebugging/CpuProfilerTrace.h
@@ -174,11 +174,14 @@ struct FCpuProfilerTrace
 		{
 			if (bEnabled)
 			{
+				// We only do relaxed here to avoid barrier cost as the worst case that can happen is multiple threads could each create an event type.
+				// At some point the last thread in the race will set the output event and no more thread will try to create new ones from then on.
+				// We don't care which event type wins as long as all threads eventually converge and stop creating new ones.
 				if (FPlatformAtomics::AtomicRead_Relaxed((volatile int32*)&InOutSpecId) == 0)
 				{
 					FPlatformAtomics::AtomicStore_Relaxed((volatile int32*)&InOutSpecId, FCpuProfilerTrace::OutputEventType(InEventString, File, Line));
 				}
-				OutputBeginEvent(InOutSpecId);
+				OutputBeginEvent(FPlatformAtomics::AtomicRead_Relaxed((volatile int32*)&InOutSpecId));
 			}
 		}
 
@@ -186,11 +189,14 @@ struct FCpuProfilerTrace
 		{
 			if (bEnabled)
 			{
+				// We only do relaxed here to avoid barrier cost as the worst case that can happen is multiple threads could each create an event type.
+				// At some point the last thread in the race will set the output event and no more thread will try to create new ones from then on.
+				// We don't care which event type wins as long as all threads eventually converge and stop creating new ones.
 				if (FPlatformAtomics::AtomicRead_Relaxed((volatile int32*)&InOutSpecId) == 0)
 				{
 					FPlatformAtomics::AtomicStore_Relaxed((volatile int32*)&InOutSpecId, FCpuProfilerTrace::OutputEventType(InEventString, File, Line));
 				}
-				OutputBeginEvent(InOutSpecId);
+				OutputBeginEvent(FPlatformAtomics::AtomicRead_Relaxed((volatile int32*)&InOutSpecId));
 			}
 		}
 
