Extend support for using tracehost argument

Add a way to specify `-tracehost` argument without specifying a hostname, which will default to 'localhost'. Analog to how the file arguments work.

diff --git a/Engine/Source/Runtime/Core/Private/ProfilingDebugging/TraceAuxiliary.cpp b/Engine/Source/Runtime/Core/Private/ProfilingDebugging/TraceAuxiliary.cpp
--- a/Engine/Source/Runtime/Core/Private/ProfilingDebugging/TraceAuxiliary.cpp
+++ b/Engine/Source/Runtime/Core/Private/ProfilingDebugging/TraceAuxiliary.cpp
@@ -1322,6 +1322,11 @@ static bool StartFromCommandlineArguments(const TCHAR* CommandLine, bool& bOutSt
 		Type = FTraceAuxiliary::EConnectionType::Network;
 		Target = *Parameter;
 	}
+	else if (FParse::Value(CommandLine, TEXT("-tracehost"), Parameter))
+	{
+		Type = FTraceAuxiliary::EConnectionType::Network;
+		Target = TEXT("localhost");
+	}
 	else if (FParse::Value(CommandLine, TEXT("-tracefile="), Parameter))
 	{
 		Type = FTraceAuxiliary::EConnectionType::File;
