ConfigIni: Fix issue with ini configs being overwritten.

diff --git a/Engine/Source/Runtime/Core/Private/Misc/ConfigCacheIni.cpp b/Engine/Source/Runtime/Core/Private/Misc/ConfigCacheIni.cpp
--- a/Engine/Source/Runtime/Core/Private/Misc/ConfigCacheIni.cpp
+++ b/Engine/Source/Runtime/Core/Private/Misc/ConfigCacheIni.cpp
@@ -1437,25 +1437,26 @@ void FillFileFromBuffer(FileType* File, FStringView Buffer, bool bHandleSymbolCo
 		}
 
 		// look for the ;METADATA line as the first line in the file, which can override passed in settings
-		if (!bHasHandledMetadata)
-		{
-			bHasHandledMetadata = true;
+		//@todo (UE-214768) Comment this back in
+		// if (!bHasHandledMetadata)
+		// {
+		// 	bHasHandledMetadata = true;
 			
-			if (FCString::Strnicmp(Start, TEXT(";METADATA="), 10) == 0)
-			{
-				FString MetadataStruct(Start + 10);
-				FString MetadataValue;
-				ExtractPropertyValue(MetadataStruct, TEXT("UseCommands="), MetadataValue);
-				if (MetadataValue.Len() > 0)
-				{
-					bHandleSymbolCommands = FCString::ToBool(*MetadataValue);
-				}
+		// 	if (FCString::Strnicmp(Start, TEXT(";METADATA="), 10) == 0)
+		// 	{
+		// 		FString MetadataStruct(Start + 10);
+		// 		FString MetadataValue;
+		// 		ExtractPropertyValue(MetadataStruct, TEXT("UseCommands="), MetadataValue);
+		// 		if (MetadataValue.Len() > 0)
+		// 		{
+		// 			bHandleSymbolCommands = FCString::ToBool(*MetadataValue);
+		// 		}
 				
-				// move on to the next line
-				continue;
-			}
+		// 		// move on to the next line
+		// 		continue;
+		// 	}
 			
-		}
+		// }
 		// If the first character in the line is [ and last char is ], this line indicates a section name
 		if( *Start=='[' && Start[FCString::Strlen(Start)-1]==']' )
 		{
