Support serialize container ID to CbObject

diff --git a/Engine/Source/Runtime/Core/Private/IO/IoContainerId.cpp b/Engine/Source/Runtime/Core/Private/IO/IoContainerId.cpp
--- a/Engine/Source/Runtime/Core/Private/IO/IoContainerId.cpp
+++ b/Engine/Source/Runtime/Core/Private/IO/IoContainerId.cpp
@@ -6,6 +6,8 @@
 #include "Misc/AssertionMacros.h"
 #include "Misc/Char.h"
 #include "Serialization/Archive.h"
+#include "Serialization/CompactBinaryWriter.h"
+#include "Serialization/CompactBinarySerialization.h"
 #include "Serialization/StructuredArchiveSlots.h"
 
 FIoContainerId FIoContainerId::FromName(const FName& Name)
@@ -28,12 +30,29 @@ FArchive& operator<<(FArchive& Ar, FIoContainerId& ContainerId)
 	return Ar;
 }
 
-void operator<<(FStructuredArchiveSlot Slot, FIoContainerId& Value)
+void operator<<(FStructuredArchiveSlot Slot, FIoContainerId& ContainerId)
 {
-	Slot << Value.Id;
+	Slot << ContainerId.Id;
+}
+
+FCbWriter& operator<<(FCbWriter& Writer, const FIoContainerId& ContainerId)
+{
+	Writer << ContainerId.Id;
+	return Writer;
 }
 
 FString LexToString(const FIoContainerId& ContainerId)
 {
 	return FString::Printf(TEXT("%llx"), ContainerId.Id);
 }
+
+bool LoadFromCompactBinary(FCbFieldView Field, FIoContainerId& OutContainerId)
+{
+	if (Field.IsInteger())
+	{
+		OutContainerId.Id = Field.AsUInt64();
+		return true;
+	}
+
+	return false;
+}
diff --git a/Engine/Source/Runtime/Core/Public/IO/IoContainerId.h b/Engine/Source/Runtime/Core/Public/IO/IoContainerId.h
--- a/Engine/Source/Runtime/Core/Public/IO/IoContainerId.h
+++ b/Engine/Source/Runtime/Core/Public/IO/IoContainerId.h
@@ -8,6 +8,8 @@
 
 class FArchive;
 class FStructuredArchiveSlot;
+class FCbFieldView;
+class FCbWriter;
 
 /**
  * Container ID.
@@ -54,10 +56,14 @@ public:
 
 	CORE_API friend FArchive& operator<<(FArchive& Ar, FIoContainerId& ContainerId);
 
-	CORE_API friend void operator<<(FStructuredArchiveSlot Slot, FIoContainerId& Value);
+	CORE_API friend void operator<<(FStructuredArchiveSlot Slot, FIoContainerId& ContainerId);
+
+	CORE_API friend FCbWriter& operator<<(FCbWriter& Writer, const FIoContainerId& ContainerId);
 
 	CORE_API friend FString LexToString(const FIoContainerId& ContainerId);
 
+	CORE_API friend bool LoadFromCompactBinary(FCbFieldView Field, FIoContainerId& OutContainerId);
+
 private:
 	inline explicit FIoContainerId(const uint64 InId)
 		: Id(InId) { }
@@ -66,4 +72,3 @@ private:
 
 	uint64 Id = InvalidId;
 };
-
