Avoid comparison name reference mismatch between log and report

#rnx

diff --git a/Engine/Source/Developer/ScreenShotComparisonTools/Private/ScreenShotManager.cpp b/Engine/Source/Developer/ScreenShotComparisonTools/Private/ScreenShotManager.cpp
--- a/Engine/Source/Developer/ScreenShotComparisonTools/Private/ScreenShotManager.cpp
+++ b/Engine/Source/Developer/ScreenShotComparisonTools/Private/ScreenShotManager.cpp
@@ -82,14 +82,17 @@ FString FScreenShotManager::GetPathComponentForPlatformAndRHI(const FAutomationS
  */
 FString FScreenShotManager::GetPathComponentForTestImages(const FAutomationScreenshotMetadata& MetaData, bool bIncludeVariantName) const
 {
+	FString FilePath;
 	if (bIncludeVariantName)
 	{
-		return FPaths::Combine(MetaData.Context, *MetaData.ScreenShotName, *MetaData.VariantName);
+		FilePath = FPaths::Combine(MetaData.Context, *MetaData.ScreenShotName, *MetaData.VariantName);
 	}
 	else
 	{
-		return FPaths::Combine(MetaData.Context, *MetaData.ScreenShotName);
+		FilePath = FPaths::Combine(MetaData.Context, *MetaData.ScreenShotName);
 	}
+	FPaths::NormalizeDirectoryName(FilePath);
+	return FilePath;
 }
 
 FString FScreenShotManager::GetApprovedFolderForImageWithOptions(const FAutomationScreenshotMetadata& MetaData, EApprovedFolderOptions InOptions) const
@@ -359,7 +362,7 @@ FImageComparisonResult FScreenShotManager::CompareScreenshot(const FString& InUn
 	{
 		// We can't find a ground truth, so it's a new comparison.
 		ComparisonResult.IncomingFilePath = InUnapprovedIncomingFilePath;
-		ComparisonResult.CreationTime = FDateTime::Now();
+		ComparisonResult.CreationTime = FDateTime::UtcNow();
 
 		UE_LOG(LogScreenShotManager, Log, TEXT("No ideal-image found. Assuming %s is a new test image"), *InUnapprovedIncomingFilePath);
 	}
diff --git a/Engine/Source/Runtime/Core/Private/Misc/AutomationTest.cpp b/Engine/Source/Runtime/Core/Private/Misc/AutomationTest.cpp
--- a/Engine/Source/Runtime/Core/Private/Misc/AutomationTest.cpp
+++ b/Engine/Source/Runtime/Core/Private/Misc/AutomationTest.cpp
@@ -1294,14 +1294,11 @@ void FAutomationTestExecutionInfo::AddError(const FString& ErrorMessage)
 FAutomationEvent FAutomationScreenshotCompareResults::ToAutomationEvent() const
 {
 	FAutomationEvent Event(EAutomationEventType::Info, TEXT(""));
-	FString OutputScreenshotName = ScreenshotPath;
-	FPaths::NormalizeDirectoryName(OutputScreenshotName);
-	OutputScreenshotName.ReplaceInline(TEXT("/"), TEXT("."));
 
 	if (bWasNew)
 	{
 		Event.Type = EAutomationEventType::Warning;
-		Event.Message = FString::Printf(TEXT("New Screenshot '%s' was discovered!  Please add a ground truth version of it."), *OutputScreenshotName);
+		Event.Message = FString::Printf(TEXT("New Screenshot '%s' was discovered!  Please add a ground truth version of it."), *ScreenshotPath);
 	}
 	else
 	{
@@ -1309,7 +1306,7 @@ FAutomationEvent FAutomationScreenshotCompareResults::ToAutomationEvent() const
 		{
 			Event.Type = EAutomationEventType::Info;
 			Event.Message = FString::Printf(TEXT("Screenshot '%s' was similar!  Global Difference = %f, Max Local Difference = %f"),
-				*OutputScreenshotName, GlobalDifference, MaxLocalDifference);
+				*ScreenshotPath, GlobalDifference, MaxLocalDifference);
 		}
 		else
 		{
@@ -1318,11 +1315,11 @@ FAutomationEvent FAutomationScreenshotCompareResults::ToAutomationEvent() const
 			if (ErrorMessage.IsEmpty())
 			{
 				Event.Message = FString::Printf(TEXT("Screenshot '%s' test failed, Screenshots were different!  Global Difference = %f, Max Local Difference = %f"),
-					*OutputScreenshotName, GlobalDifference, MaxLocalDifference);
+					*ScreenshotPath, GlobalDifference, MaxLocalDifference);
 			}
 			else
 			{
-				Event.Message = FString::Printf(TEXT("Screenshot '%s' test failed; Error = %s"), *OutputScreenshotName, *ErrorMessage);
+				Event.Message = FString::Printf(TEXT("Screenshot '%s' test failed; Error = %s"), *ScreenshotPath, *ErrorMessage);
 			}
 		}
 	}
