Fixed IsDerivedFromSharedFromThis() when there are multiple TSharedFromThis bases.
Refactored the TSharedFromThis code to enforce a compile error when trying to create a TSharedPtr with multiple TSharedFromThis bases.
Fixed fallout.

diff --git a/Engine/Plugins/Developer/PerforceSourceControl/Source/PerforceSourceControl/Private/PerforceSourceControlLabel.h b/Engine/Plugins/Developer/PerforceSourceControl/Source/PerforceSourceControl/Private/PerforceSourceControlLabel.h
--- a/Engine/Plugins/Developer/PerforceSourceControl/Source/PerforceSourceControl/Private/PerforceSourceControlLabel.h
+++ b/Engine/Plugins/Developer/PerforceSourceControl/Source/PerforceSourceControl/Private/PerforceSourceControlLabel.h
@@ -11,7 +11,7 @@ class FPerforceSourceControlProvider;
 /** 
  * Abstraction of a Perforce label.
  */
-class FPerforceSourceControlLabel : public ISourceControlLabel, public TSharedFromThis<FPerforceSourceControlLabel>
+class FPerforceSourceControlLabel : public ISourceControlLabel
 {
 public:
 
diff --git a/Engine/Plugins/Developer/PerforceSourceControl/Source/PerforceSourceControl/Private/PerforceSourceControlRevision.h b/Engine/Plugins/Developer/PerforceSourceControl/Source/PerforceSourceControl/Private/PerforceSourceControlRevision.h
--- a/Engine/Plugins/Developer/PerforceSourceControl/Source/PerforceSourceControl/Private/PerforceSourceControlRevision.h
+++ b/Engine/Plugins/Developer/PerforceSourceControl/Source/PerforceSourceControl/Private/PerforceSourceControlRevision.h
@@ -9,7 +9,7 @@
 class FPerforceConnection;
 class FPerforceSourceControlProvider;
 
-class FPerforceSourceControlRevision final : public ISourceControlRevision, public TSharedFromThis<FPerforceSourceControlRevision, ESPMode::ThreadSafe>
+class FPerforceSourceControlRevision final : public ISourceControlRevision
 {
 public:
 
diff --git a/Engine/Plugins/Developer/PerforceSourceControl/Source/PerforceSourceControl/Private/PerforceSourceControlState.h b/Engine/Plugins/Developer/PerforceSourceControl/Source/PerforceSourceControl/Private/PerforceSourceControlState.h
--- a/Engine/Plugins/Developer/PerforceSourceControl/Source/PerforceSourceControl/Private/PerforceSourceControlState.h
+++ b/Engine/Plugins/Developer/PerforceSourceControl/Source/PerforceSourceControl/Private/PerforceSourceControlState.h
@@ -46,7 +46,7 @@ namespace EPerforceState
 	};
 }
 
-class FPerforceSourceControlState : public ISourceControlState, public TSharedFromThis<FPerforceSourceControlState, ESPMode::ThreadSafe>
+class FPerforceSourceControlState : public ISourceControlState
 {
 public:
 	FPerforceSourceControlState( const FString& InLocalFilename, EPerforceState::Type InState = EPerforceState::DontCare )
diff --git a/Engine/Plugins/Developer/TraceDataFiltering/Source/TraceDataFiltering/Private/STraceObjectRowWidget.cpp b/Engine/Plugins/Developer/TraceDataFiltering/Source/TraceDataFiltering/Private/STraceObjectRowWidget.cpp
--- a/Engine/Plugins/Developer/TraceDataFiltering/Source/TraceDataFiltering/Private/STraceObjectRowWidget.cpp
+++ b/Engine/Plugins/Developer/TraceDataFiltering/Source/TraceDataFiltering/Private/STraceObjectRowWidget.cpp
@@ -5,7 +5,7 @@
 #include "ITraceObject.h"
 #include "EventFilterStyle.h"
 
-class FSetFilteringDragDropOp : public FDragDropOperation, public TSharedFromThis<FSetFilteringDragDropOp>
+class FSetFilteringDragDropOp : public FDragDropOperation
 {
 public:
 	DRAG_DROP_OPERATOR_TYPE(FSetFilteringDragDropOp, FDragDropOperation)
diff --git a/Engine/Source/Developer/WidgetRegistration/Public/ToolElementRegistry.h b/Engine/Source/Developer/WidgetRegistration/Public/ToolElementRegistry.h
--- a/Engine/Source/Developer/WidgetRegistration/Public/ToolElementRegistry.h
+++ b/Engine/Source/Developer/WidgetRegistration/Public/ToolElementRegistry.h
@@ -93,7 +93,7 @@ FString KeyString;
 };
 
 /** Represents one Tool Element, and is responsible for displaying the UI for that element  */
-class WIDGETREGISTRATION_API FToolElement : public FToolElementRegistrationKey, public TSharedFromThis<FToolElement>
+class WIDGETREGISTRATION_API FToolElement : public FToolElementRegistrationKey
 {
 
 public:
diff --git a/Engine/Source/Editor/AudioEditor/Private/SubmixDetailsCustomization.cpp b/Engine/Source/Editor/AudioEditor/Private/SubmixDetailsCustomization.cpp
--- a/Engine/Source/Editor/AudioEditor/Private/SubmixDetailsCustomization.cpp
+++ b/Engine/Source/Editor/AudioEditor/Private/SubmixDetailsCustomization.cpp
@@ -31,6 +31,12 @@ class SWidget;
 
 #define LOCTEXT_NAMESPACE "SubmixDetailsInspector"
 
+void FNameSelectorGenerator::SetWeakThis(TWeakPtr<FNameSelectorGenerator>&& InWeakThis)
+{
+	checkSlow(InWeakThis.IsValid() && !WeakThis.IsValid());
+	WeakThis = MoveTemp(InWeakThis);
+}
+
 FText FNameSelectorGenerator::GetComboBoxContent() const
 {
 	FName DesiredName = CachedCallbacks.GetCurrentlySelectedName();
@@ -65,21 +71,22 @@ TSharedRef<SWidget> FNameSelectorGenerator::MakeNameSelectorWidget(TArray<FName>
 			break;
 		}
 	}
-	
+
+	TSharedRef<FNameSelectorGenerator> SharedThis = WeakThis.Pin().ToSharedRef();
 
 	return SNew(SComboBox<TSharedPtr<FName>>)
-		.OnGenerateWidget(this, &FNameSelectorGenerator::HandleResponseComboBoxGenerateWidget)
+		.OnGenerateWidget(SharedThis, &FNameSelectorGenerator::HandleResponseComboBoxGenerateWidget)
 		.OptionsSource(&CachedNameArray)
-		.OnSelectionChanged(this, &FNameSelectorGenerator::OnSelectionChanged)
+		.OnSelectionChanged(SharedThis, &FNameSelectorGenerator::OnSelectionChanged)
 		.InitiallySelectedItem(InitialSelectedName)
 		.IsEnabled(FSlateApplication::Get().GetNormalExecutionAttribute())
 		.ContentPadding(FMargin(2.0f, 2.0f))
 		.Content()
 		[
 			SNew(STextBlock)
-			.Text(this, &FNameSelectorGenerator::GetComboBoxContent)
+			.Text(SharedThis, &FNameSelectorGenerator::GetComboBoxContent)
 			.Font(IDetailLayoutBuilder::GetDetailFont())
-			.ToolTipText(this, &FNameSelectorGenerator::GetComboBoxToolTip)
+			.ToolTipText(SharedThis, &FNameSelectorGenerator::GetComboBoxToolTip)
 		];
 }
 
@@ -97,6 +104,13 @@ TSharedRef<SWidget> FNameSelectorGenerator::HandleResponseComboBoxGenerateWidget
 		.Font(IDetailLayoutBuilder::GetDetailFont());
 }
 
+TSharedRef<FNameSelectorGenerator> FNameSelectorGenerator::MakeInstance()
+{
+	TSharedRef<FNameSelectorGenerator> Result = MakeShared<FNameSelectorGenerator>();
+	Result->SetWeakThis(Result.ToWeakPtr());
+	return Result;
+}
+
 TSharedRef<IDetailCustomization> FSoundfieldSubmixDetailsCustomization::MakeInstance()
 {
 	return MakeShareable(new FSoundfieldSubmixDetailsCustomization);
@@ -106,7 +120,7 @@ void FSoundfieldSubmixDetailsCustomization::CustomizeDetails(IDetailLayoutBuilde
 {
 	if (!SoundfieldFormatNameSelectorGenerator)
 	{
-		SoundfieldFormatNameSelectorGenerator = MakeShareable(new FNameSelectorGenerator());
+		SoundfieldFormatNameSelectorGenerator = FNameSelectorGenerator::MakeInstance();
 	}
 
 	IDetailCategoryBuilder& SoundfieldCategory = DetailLayout.EditCategory(TEXT("Soundfield"));
@@ -164,7 +178,9 @@ void FSoundfieldSubmixDetailsCustomization::CustomizeDetails(IDetailLayoutBuilde
 
 TSharedRef<IDetailCustomization> FEndpointSubmixDetailsCustomization::MakeInstance()
 {
-	return MakeShareable(new FEndpointSubmixDetailsCustomization);
+	TSharedRef<FEndpointSubmixDetailsCustomization> Result = MakeShared<FEndpointSubmixDetailsCustomization>();
+	Result->SetWeakThis(StaticCastWeakPtr<FNameSelectorGenerator>(Result.ToWeakPtr()));
+	return Result;
 }
 
 void FEndpointSubmixDetailsCustomization::CustomizeDetails(IDetailLayoutBuilder& DetailLayout)
@@ -222,7 +238,9 @@ void FEndpointSubmixDetailsCustomization::CustomizeDetails(IDetailLayoutBuilder&
 
 TSharedRef<IDetailCustomization> FSoundfieldEndpointSubmixDetailsCustomization::MakeInstance()
 {
-	return MakeShareable(new FSoundfieldEndpointSubmixDetailsCustomization);
+	TSharedRef<FSoundfieldEndpointSubmixDetailsCustomization> Result = MakeShared<FSoundfieldEndpointSubmixDetailsCustomization>();
+	Result->SetWeakThis(StaticCastWeakPtr<FNameSelectorGenerator>(Result.ToWeakPtr()));
+	return Result;
 }
 
 void FSoundfieldEndpointSubmixDetailsCustomization::CustomizeDetails(IDetailLayoutBuilder& DetailLayout)
diff --git a/Engine/Source/Editor/AudioEditor/Public/SubmixDetailsCustomization.h b/Engine/Source/Editor/AudioEditor/Public/SubmixDetailsCustomization.h
--- a/Engine/Source/Editor/AudioEditor/Public/SubmixDetailsCustomization.h
+++ b/Engine/Source/Editor/AudioEditor/Public/SubmixDetailsCustomization.h
@@ -21,7 +21,7 @@ class IDetailLayoutBuilder;
 class SWidget;
 
 // Utility class to build combo boxes out of arrays of names.
-class FNameSelectorGenerator : public TSharedFromThis<FNameSelectorGenerator>
+class FNameSelectorGenerator
 {
 
 public:
@@ -35,9 +35,13 @@ public:
 	// Use this to generate a combo box widget.
 	TSharedRef<SWidget> MakeNameSelectorWidget(TArray<FName>& InNameArray, FNameSelectorCallbacks&& InCallbacks);
 
-	
+	// Makes a new instance of the name selector generator class
+	static TSharedRef<FNameSelectorGenerator> MakeInstance();
 
 protected:
+	// This needs to be called after construction and after it has been bound to a TSharedPtr
+	void SetWeakThis(TWeakPtr<FNameSelectorGenerator>&& InWeakThis);
+
 	void OnSelectionChanged(TSharedPtr<FName> NameItem, ESelectInfo::Type SelectInfo);
 	TSharedRef<SWidget> HandleResponseComboBoxGenerateWidget(TSharedPtr<FName> StringItem);
 	FText GetComboBoxToolTip() const;
@@ -45,6 +49,9 @@ protected:
 
 	TArray<TSharedPtr<FName>> CachedNameArray;
 	FNameSelectorCallbacks CachedCallbacks;
+
+private:
+	TWeakPtr<FNameSelectorGenerator> WeakThis;
 };
 
 class AUDIOEDITOR_API FSoundfieldSubmixDetailsCustomization : public IDetailCustomization
@@ -61,7 +68,7 @@ private:
 	TSharedPtr<FNameSelectorGenerator> SoundfieldFormatNameSelectorGenerator;
 };
 
-class AUDIOEDITOR_API FEndpointSubmixDetailsCustomization : public IDetailCustomization, FNameSelectorGenerator
+class AUDIOEDITOR_API FEndpointSubmixDetailsCustomization : public IDetailCustomization, public FNameSelectorGenerator
 {
 public:
 	// Makes a new instance of this detail layout class
@@ -75,7 +82,7 @@ private:
 	TSharedPtr<FNameSelectorGenerator> EndpointTypeNameSelectorGenerator;
 };
 
-class AUDIOEDITOR_API FSoundfieldEndpointSubmixDetailsCustomization : public IDetailCustomization, FNameSelectorGenerator
+class AUDIOEDITOR_API FSoundfieldEndpointSubmixDetailsCustomization : public IDetailCustomization, public FNameSelectorGenerator
 {
 public:
 	// Makes a new instance of this detail layout class
diff --git a/Engine/Source/Editor/SequencerCore/Private/MVVM/Views/OutlinerColumns/SColumnToggleWidget.cpp b/Engine/Source/Editor/SequencerCore/Private/MVVM/Views/OutlinerColumns/SColumnToggleWidget.cpp
--- a/Engine/Source/Editor/SequencerCore/Private/MVVM/Views/OutlinerColumns/SColumnToggleWidget.cpp
+++ b/Engine/Source/Editor/SequencerCore/Private/MVVM/Views/OutlinerColumns/SColumnToggleWidget.cpp
@@ -13,7 +13,7 @@
 namespace UE::Sequencer
 {
 
-class FColumnToggleDragDropOp : public FDragDropOperation, public TSharedFromThis<FColumnToggleDragDropOp>
+class FColumnToggleDragDropOp : public FDragDropOperation
 {
 public:
 
diff --git a/Engine/Source/Editor/UnrealEd/Classes/Factories/FbxSceneImportFactory.h b/Engine/Source/Editor/UnrealEd/Classes/Factories/FbxSceneImportFactory.h
--- a/Engine/Source/Editor/UnrealEd/Classes/Factories/FbxSceneImportFactory.h
+++ b/Engine/Source/Editor/UnrealEd/Classes/Factories/FbxSceneImportFactory.h
@@ -105,7 +105,7 @@ private:
 	UObject *ContentObject;
 };
 
-class FFbxMeshInfo : public FFbxAttributeInfo, public TSharedFromThis<FFbxMeshInfo>
+class FFbxMeshInfo : public FFbxAttributeInfo
 {
 public:
 	int32 FaceNum;
@@ -141,7 +141,7 @@ public:
 	virtual UClass *GetType();
 };
 
-class FFbxTextureInfo : public FFbxAttributeInfo, public TSharedFromThis<FFbxTextureInfo>
+class FFbxTextureInfo : public FFbxAttributeInfo
 {
 public:
 	FString TexturePath;
@@ -153,7 +153,7 @@ public:
 	virtual UClass *GetType();
 };
 
-class FFbxMaterialInfo : public FFbxAttributeInfo, public TSharedFromThis<FFbxMaterialInfo>
+class FFbxMaterialInfo : public FFbxAttributeInfo
 {
 public:
 	//This string is use to help match the material when doing a reimport
diff --git a/Engine/Source/Runtime/Core/Public/Templates/SharedPointer.h b/Engine/Source/Runtime/Core/Public/Templates/SharedPointer.h
--- a/Engine/Source/Runtime/Core/Public/Templates/SharedPointer.h
+++ b/Engine/Source/Runtime/Core/Public/Templates/SharedPointer.h
@@ -227,7 +227,7 @@ public:
 
 		// If the object happens to be derived from TSharedFromThis, the following method
 		// will prime the object with a weak pointer to itself.
-		SharedPointerInternals::EnableSharedFromThis( this, InRawPtrProxy.Object, InRawPtrProxy.Object );
+		SharedPointerInternals::EnableSharedFromThis( this, InRawPtrProxy.Object );
 	}
 
 	/**
@@ -252,7 +252,7 @@ public:
 
 		// If the object happens to be derived from TSharedFromThis, the following method
 		// will prime the object with a weak pointer to itself.
-		SharedPointerInternals::EnableSharedFromThis( this, InRawPtrProxy.Object, InRawPtrProxy.Object );
+		SharedPointerInternals::EnableSharedFromThis( this, InRawPtrProxy.Object );
 	}
 
 	/**
@@ -277,7 +277,7 @@ public:
 
 		// If the object happens to be derived from TSharedFromThis, the following method
 		// will prime the object with a weak pointer to itself.
-		SharedPointerInternals::EnableSharedFromThis( this, InRawPtrProxy.Object, InRawPtrProxy.Object );
+		SharedPointerInternals::EnableSharedFromThis( this, InRawPtrProxy.Object );
 	}
 
 	/**
@@ -556,7 +556,7 @@ private:
 
 		// If the object happens to be derived from TSharedFromThis, the following method
 		// will prime the object with a weak pointer to itself.
-		SharedPointerInternals::EnableSharedFromThis(this, InObject, InObject);
+		SharedPointerInternals::EnableSharedFromThis(this, InObject);
 	}
 
 	/**
@@ -693,7 +693,7 @@ public:
 	{
 		// If the object happens to be derived from TSharedFromThis, the following method
 		// will prime the object with a weak pointer to itself.
-		SharedPointerInternals::EnableSharedFromThis( this, InObject, InObject );
+		SharedPointerInternals::EnableSharedFromThis( this, InObject );
 	}
 
 	/**
@@ -714,7 +714,7 @@ public:
 	{
 		// If the object happens to be derived from TSharedFromThis, the following method
 		// will prime the object with a weak pointer to itself.
-		SharedPointerInternals::EnableSharedFromThis( this, InObject, InObject );
+		SharedPointerInternals::EnableSharedFromThis( this, InObject );
 	}
 
 	/**
@@ -733,7 +733,7 @@ public:
 	{
 		// If the object happens to be derived from TSharedFromThis, the following method
 		// will prime the object with a weak pointer to itself.
-		SharedPointerInternals::EnableSharedFromThis( this, InRawPtrProxy.Object, InRawPtrProxy.Object );
+		SharedPointerInternals::EnableSharedFromThis( this, InRawPtrProxy.Object );
 	}
 
 	/**
@@ -753,7 +753,7 @@ public:
 	{
 		// If the object happens to be derived from TSharedFromThis, the following method
 		// will prime the object with a weak pointer to itself.
-		SharedPointerInternals::EnableSharedFromThis( this, InRawPtrProxy.Object, InRawPtrProxy.Object );
+		SharedPointerInternals::EnableSharedFromThis( this, InRawPtrProxy.Object );
 	}
 
 	/**
@@ -773,7 +773,7 @@ public:
 	{
 		// If the object happens to be derived from TSharedFromThis, the following method
 		// will prime the object with a weak pointer to itself.
-		SharedPointerInternals::EnableSharedFromThis( this, InRawPtrProxy.Object, InRawPtrProxy.Object );
+		SharedPointerInternals::EnableSharedFromThis( this, InRawPtrProxy.Object );
 	}
 
 	/**
@@ -1567,7 +1567,7 @@ struct TCallTraits<TWeakPtr<ObjectType, Mode>> : public TCallTraitsBase<TWeakPtr
  * instance that's already been allocated.  Use the optional Mode template argument for thread-safety.
  */
 template< class ObjectType, ESPMode Mode >
-class TSharedFromThis
+class TSharedFromThis : private UE::Core::Private::FSharedFromThisBase
 {
 public:
 
@@ -1734,9 +1734,11 @@ public:		// @todo: Ideally this would be private, but template sharing problems
 	 * the supplied object pointer along with the authoritative shared reference to the object.
 	 * Note that until this function is called, calls to AsShared() will result in an empty pointer.
 	 */
-	template< class SharedPtrType, class OtherType >
-	FORCEINLINE void UpdateWeakReferenceInternal( TSharedPtr< SharedPtrType, Mode > const* InSharedPtr, OtherType* InObject ) const
+	template< class SharedPtrType, ESPMode SharedPtrMode, class OtherType >
+	FORCEINLINE void UpdateWeakReferenceInternal( TSharedPtr< SharedPtrType, SharedPtrMode > const* InSharedPtr, OtherType* InObject ) const
 	{
+		static_assert(SharedPtrMode == Mode, "You cannot use a TSharedPtr of one mode with a type which inherits TSharedFromThis of another mode.");
+
 		if( !WeakThis.IsValid() )
 		{
 			WeakThis = TSharedPtr< ObjectType, Mode >( *InSharedPtr, InObject );
@@ -1748,9 +1750,11 @@ public:		// @todo: Ideally this would be private, but template sharing problems
 	 * the supplied object pointer along with the authoritative shared reference to the object.
 	 * Note that until this function is called, calls to AsShared() will result in an empty pointer.
 	 */
-	template< class SharedRefType, class OtherType >
-	FORCEINLINE void UpdateWeakReferenceInternal( TSharedRef< SharedRefType, Mode > const* InSharedRef, OtherType* InObject ) const
+	template< class SharedRefType, ESPMode SharedPtrMode, class OtherType >
+	FORCEINLINE void UpdateWeakReferenceInternal( TSharedRef< SharedRefType, SharedPtrMode > const* InSharedRef, OtherType* InObject ) const
 	{
+		static_assert(SharedPtrMode == Mode, "You cannot use a TSharedPtr of one mode with a type which inherits TSharedFromThis of another mode.");
+
 		if( !WeakThis.IsValid() )
 		{
 			WeakThis = TSharedRef< ObjectType, Mode >( *InSharedRef, InObject );
@@ -1790,7 +1794,9 @@ protected:
 	}
 
 	/** Hidden destructor */
-	~TSharedFromThis() { }
+	~TSharedFromThis()
+	{
+	}
 
 private:
 
@@ -1800,27 +1806,6 @@ private:
 };
 
 
-namespace UE::Core::Private
-{
-	template <typename T>
-	constexpr bool IsDerivedFromSharedFromThisImpl(const TSharedFromThis<T>*)
-	{
-		return true;
-	}
-
-	constexpr bool IsDerivedFromSharedFromThisImpl(...)
-	{
-		return false;
-	}
-}
-
-template <typename T>
-constexpr bool IsDerivedFromSharedFromThis()
-{
-	return UE::Core::Private::IsDerivedFromSharedFromThisImpl((const T*)nullptr);
-}
-
-
 /**
  * Global equality operator for TSharedRef
  *
@@ -2202,6 +2187,11 @@ template< class CastToType, class CastFromType, ESPMode Mode >
 template< class ObjectType >
 [[nodiscard]] FORCEINLINE SharedPointerInternals::TRawPtrProxy< ObjectType > MakeShareable( ObjectType* InObject )
 {
+	///////////////////////////////////////////////////////////////////////////////////////////////
+	// If you get an 'ambiguous call' compile error in this function, it means you have multiple //
+	// TSharedFromThis bases in your inheritance hierarchy.  This is not supported.              //
+	///////////////////////////////////////////////////////////////////////////////////////////////
+
 	if constexpr (IsDerivedFromSharedFromThis<ObjectType>())
 	{
 		// If this goes off, you should probably be using Ptr->AsShared() or Ptr->AsWeak() instead.
@@ -2221,6 +2211,11 @@ template< class ObjectType >
 template< class ObjectType, class DeleterType >
 [[nodiscard]] FORCEINLINE SharedPointerInternals::TRawPtrProxyWithDeleter< ObjectType, DeleterType > MakeShareable( ObjectType* InObject, DeleterType&& InDeleter )
 {
+	///////////////////////////////////////////////////////////////////////////////////////////////
+	// If you get an 'ambiguous call' compile error in this function, it means you have multiple //
+	// TSharedFromThis bases in your inheritance hierarchy.  This is not supported.              //
+	///////////////////////////////////////////////////////////////////////////////////////////////
+
 	if constexpr (IsDerivedFromSharedFromThis<ObjectType>())
 	{
 		// If this goes off, you should probably be using Ptr->AsShared() or Ptr->AsWeak() instead.
diff --git a/Engine/Source/Runtime/Core/Public/Templates/SharedPointerInternals.h b/Engine/Source/Runtime/Core/Public/Templates/SharedPointerInternals.h
--- a/Engine/Source/Runtime/Core/Public/Templates/SharedPointerInternals.h
+++ b/Engine/Source/Runtime/Core/Public/Templates/SharedPointerInternals.h
@@ -833,25 +833,41 @@ namespace SharedPointerInternals
 		/** Pointer to the reference controller for the object a TWeakPtr is referencing */
 		TReferenceControllerBase<Mode>* ReferenceController;
 	};
+}
+
+
+namespace UE::Core::Private
+{
+	// This base class only exists to implement IsDerivedFromSharedFromThis
+	struct FSharedFromThisBase
+	{
+	};
+}
+
+template <typename T>
+constexpr bool IsDerivedFromSharedFromThis()
+{
+	return std::is_base_of_v<UE::Core::Private::FSharedFromThisBase, T>;
+}
 
 
+namespace SharedPointerInternals
+{
 	/** Templated helper function (const) that creates a shared pointer from an object instance */
-	template< class SharedPtrType, class ObjectType, class OtherType, ESPMode Mode >
-	FORCEINLINE void EnableSharedFromThis( SharedPtrType* InSharedPtrOrRef, ObjectType const* InObject, TSharedFromThis< OtherType, Mode > const* InShareable )
+	template< class SharedPtrType, class ObjectType >
+	FORCEINLINE void EnableSharedFromThis( SharedPtrType* InSharedPtrOrRef, ObjectType const* InObject )
 	{
-		if constexpr (SharedPtrType::Mode == Mode)
+		///////////////////////////////////////////////////////////////////////////////////////////////
+		// If you get an 'ambiguous call' compile error in this function, it means you have multiple //
+		// TSharedFromThis bases in your inheritance hierarchy.  This is not supported.              //
+		///////////////////////////////////////////////////////////////////////////////////////////////
+
+		if constexpr (IsDerivedFromSharedFromThis<ObjectType>())
 		{
-			if( InShareable != nullptr )
+			if( InObject != nullptr )
 			{
-				InShareable->UpdateWeakReferenceInternal( InSharedPtrOrRef, const_cast< ObjectType* >( InObject ) );
+				InObject->UpdateWeakReferenceInternal( InSharedPtrOrRef, const_cast< ObjectType* >( InObject ) );
 			}
 		}
-		else
-		{
-			static_assert(sizeof(ObjectType) == 0, "You cannot use a TSharedPtr of one mode with a type which inherits TSharedFromThis of another mode.");
-		}
 	}
-
-	/** Templated helper catch-all function, accomplice to the above helper function */
-	constexpr void EnableSharedFromThis( ... ) { }
 }
