Use the correct determination for when to use __builtin_unreachable()
Wrap some of the headers up with checks that __AUTORTFM is enabled to avoid compile issues on platforms that are not yet supported.
#rnx

diff --git a/Engine/Source/Runtime/Core/Private/AutoRTFM/CallNest.h b/Engine/Source/Runtime/Core/Private/AutoRTFM/CallNest.h
--- a/Engine/Source/Runtime/Core/Private/AutoRTFM/CallNest.h
+++ b/Engine/Source/Runtime/Core/Private/AutoRTFM/CallNest.h
@@ -2,6 +2,7 @@
 
 #pragma once
 
+#if (defined(__AUTORTFM) && __AUTORTFM)
 #include "LongJump.h"
 
 namespace AutoRTFM
@@ -27,3 +28,4 @@ public:
 
 
 } // namespace AutoRTFM
+#endif
\ No newline at end of file
diff --git a/Engine/Source/Runtime/Core/Private/AutoRTFM/CallNestInlines.h b/Engine/Source/Runtime/Core/Private/AutoRTFM/CallNestInlines.h
--- a/Engine/Source/Runtime/Core/Private/AutoRTFM/CallNestInlines.h
+++ b/Engine/Source/Runtime/Core/Private/AutoRTFM/CallNestInlines.h
@@ -2,6 +2,7 @@
 
 #pragma once
 
+#if (defined(__AUTORTFM) && __AUTORTFM)
 #include "Context.h"
 #include "CallNest.h"
 
@@ -25,3 +26,4 @@ void FCallNest::Try(const TTryFunctor& TryFunctor)
 }
 
 } // namespace AutoRTFM
+#endif
\ No newline at end of file
diff --git a/Engine/Source/Runtime/Core/Private/AutoRTFM/LongJump.h b/Engine/Source/Runtime/Core/Private/AutoRTFM/LongJump.h
--- a/Engine/Source/Runtime/Core/Private/AutoRTFM/LongJump.h
+++ b/Engine/Source/Runtime/Core/Private/AutoRTFM/LongJump.h
@@ -2,6 +2,7 @@
 
 #pragma once
 
+#if (defined(__AUTORTFM) && __AUTORTFM)
 #include "AutoRTFM/AutoRTFM.h"
 #include "Utils.h"
 #include <setjmp.h>
@@ -55,7 +56,7 @@ void FLongJump::TryCatch(const TTryFunctor& TryFunctor, const TCatchFunctor& Cat
 inline void FLongJump::Throw()
 {
     ASSERT(bIsSet);
-#if PLATFORM_WINDOWS
+#if defined(_MSC_VER) && !defined(__clang__)
 	longjmp(JmpBuf, 1);
 #else
 	_longjmp(JmpBuf, 1);
@@ -63,3 +64,4 @@ inline void FLongJump::Throw()
 }
 
 } // namespace AutoRTFM
+#endif
\ No newline at end of file
diff --git a/Engine/Source/Runtime/Core/Private/AutoRTFM/Utils.cpp b/Engine/Source/Runtime/Core/Private/AutoRTFM/Utils.cpp
--- a/Engine/Source/Runtime/Core/Private/AutoRTFM/Utils.cpp
+++ b/Engine/Source/Runtime/Core/Private/AutoRTFM/Utils.cpp
@@ -25,7 +25,7 @@ void PrettyAbort(const char* const File, const unsigned Line, const char* const
 {
     UE_LOG(LogAutoRTFM, Fatal, TEXT("%s:%d:%s: assertion %s failed."), ANSI_TO_TCHAR(File), Line, ANSI_TO_TCHAR(Function), ANSI_TO_TCHAR(Expression));
 
-#if PLATFORM_WINDOWS
+#if defined(_MSC_VER) && !defined(__clang__)
 	__assume(false);
 #else
 	__builtin_unreachable();
diff --git a/Engine/Source/Runtime/Core/Private/AutoRTFM/Utils.h b/Engine/Source/Runtime/Core/Private/AutoRTFM/Utils.h
--- a/Engine/Source/Runtime/Core/Private/AutoRTFM/Utils.h
+++ b/Engine/Source/Runtime/Core/Private/AutoRTFM/Utils.h
@@ -22,7 +22,7 @@ namespace AutoRTFM
 {
 	UE_LOG(LogAutoRTFM, Fatal, TEXT("Unreachable encountered!"));
 
-#if PLATFORM_WINDOWS
+#if defined(_MSC_VER) && !defined(__clang__)
     __assume(false);
 #else
     __builtin_unreachable();
