Remove unitentionally submitted change
#rnx

diff --git a/Engine/Source/Runtime/Core/Private/AutoRTFM/LongJump.h b/Engine/Source/Runtime/Core/Private/AutoRTFM/LongJump.h
--- a/Engine/Source/Runtime/Core/Private/AutoRTFM/LongJump.h
+++ b/Engine/Source/Runtime/Core/Private/AutoRTFM/LongJump.h
@@ -56,7 +56,7 @@ void FLongJump::TryCatch(const TTryFunctor& TryFunctor, const TCatchFunctor& Cat
 inline void FLongJump::Throw()
 {
     ASSERT(bIsSet);
-#if defined(_MSC_VER) && !defined(__clang__)
+#if PLATFORM_WINDOWS
 	longjmp(JmpBuf, 1);
 #else
 	_longjmp(JmpBuf, 1);
