[TaskGraph/Pipe]
 - Add missing memory barrier when pipe tasks finishes
 - Fix multiple TSAN warnings about unprotected data when using pipes.

diff --git a/Engine/Source/Runtime/Core/Private/Tasks/Pipe.cpp b/Engine/Source/Runtime/Core/Private/Tasks/Pipe.cpp
--- a/Engine/Source/Runtime/Core/Private/Tasks/Pipe.cpp
+++ b/Engine/Source/Runtime/Core/Private/Tasks/Pipe.cpp
@@ -35,7 +35,10 @@ namespace UE::Tasks
 		Private::FTaskBase* Task_Local = &Task;
 		// try clearing the task if it's still pipe's "last task". if succeeded, release the ref accounted for pipe's last task. otherwise whoever replaced it
 		// as the last task will do this
-		if (LastTask.compare_exchange_strong(Task_Local, nullptr, std::memory_order_acquire, std::memory_order_relaxed))
+
+		// important to have a barrier even in case of failure so that whenever a pipe task finished, we have a barrier protecting any produced data
+		// so that it can be passed across threads on the same pipe without synchronization.
+		if (LastTask.compare_exchange_strong(Task_Local, nullptr, std::memory_order_acq_rel, std::memory_order_acquire))
 		{
 			Task.Release(); // it was still pipe's last task. now that we cleared it, release the reference
 		}
