UE_REQUIRES_FRIEND and UE_REQUIRES_DEFINITION to allow forward declarations and friend declarations with constraints.

diff --git a/Engine/Source/Runtime/Core/Public/Templates/Requires.h b/Engine/Source/Runtime/Core/Public/Templates/Requires.h
--- a/Engine/Source/Runtime/Core/Public/Templates/Requires.h
+++ b/Engine/Source/Runtime/Core/Public/Templates/Requires.h
@@ -14,7 +14,47 @@
 		typename U  // note - no trailing comma before the constraint
 		UE_REQUIRES(std::is_integral_v<T> && sizeof(U) <= 4)
 	>
-	void IntegralUpTo32Bit(T Lhs, U Rhs) {}
+	void IntegralUpTo32Bit(T Lhs, U Rhs)
+	{
+	}
+
+	When a template is both declared and defined at once, UE_REQUIRES must be used.
+
+	If a forward declaration is needed, UE_REQUIRES must be used to declare the function
+	outside the scope, then UE_REQUIRES_DEFINITION should be used for the constraint on the
+	function declaration.
+
+	If a friend declaration is needed, a forward declaration using UE_REQUIRES must first be
+	specified at the required scope, then UE_REQUIRES_FRIEND should be used for the constraint
+	in the friendship declaration.  Usage:
+
+	template <
+		typename T
+		UE_REQUIRES(trait_v<T>)
+	>
+	void Function(T Val);
+
+	class FThing
+	{
+	private:
+		template <
+			typename T
+			UE_REQUIRES_FRIEND(trait_v<T>)
+		>
+		friend void FuncB(T);
+
+		template <typename T>
+		static void Use(T);
+	};
+
+	template <
+		typename T
+		UE_REQUIRES_DEFINITION(trait_v<T>)
+	>
+	void Function(T Val)
+	{
+		Use(Val);
+	}
 
 	UE_REQUIRES_EXPR() wraps the effects of a requires expression, used to test the
 	compilability of an expression based on a deduced template parameter.  Usage:
@@ -39,6 +79,8 @@
  -----------------------------------------------------------------------------*/
 #if __cplusplus < 202000
 	#define UE_REQUIRES(...) , std::enable_if_t<(__VA_ARGS__), int> = 0
+	#define UE_REQUIRES_FRIEND(...) , std::enable_if_t<(__VA_ARGS__), int>
+	#define UE_REQUIRES_DEFINITION(...) , std::enable_if_t<(__VA_ARGS__), int>
 	#define UE_REQUIRES_EXPR(...) (!std::is_same_v<decltype(__VA_ARGS__), long long************>) // This is highly unlikely to be the type of any expression
 #else
 	namespace UE::Core::Private
@@ -49,7 +91,9 @@
 	}
 
 	#define UE_REQUIRES(...) > requires (!!(__VA_ARGS__)) && UE::Core::Private::BoolIdentityConcept<true
-    #define UE_REQUIRES_EXPR(...) requires { (__VA_ARGS__); }
+	#define UE_REQUIRES_FRIEND(...) UE_REQUIRES(__VA_ARGS__)
+	#define UE_REQUIRES_DEFINITION(...) UE_REQUIRES(__VA_ARGS__)
+	#define UE_REQUIRES_EXPR(...) requires { (__VA_ARGS__); }
 #endif
 
 // This should be regarded as deprecated - please use UE_REQUIRES instead.
