Added Algo::Mismatch and some tests.

diff --git a/Engine/Source/Runtime/Core/Private/Containers/Algo/AlgosTest.cpp b/Engine/Source/Runtime/Core/Private/Containers/Algo/AlgosTest.cpp
--- a/Engine/Source/Runtime/Core/Private/Containers/Algo/AlgosTest.cpp
+++ b/Engine/Source/Runtime/Core/Private/Containers/Algo/AlgosTest.cpp
@@ -13,6 +13,7 @@
 #include "Algo/IntroSort.h"
 #include "Algo/IsHeap.h"
 #include "Algo/IsSorted.h"
+#include "Algo/Mismatch.h"
 #include "Algo/Sort.h"
 #include "Algo/Transform.h"
 #include "Algo/IndexOf.h"
@@ -680,6 +681,131 @@ public:
 		}
 	}
 
+	void TestMismatch()
+	{
+		{
+			TArray<int32> Empty;
+			TArray<int32> DataA = { 1, 2, 3, 4, 5, 6 };
+			TArray<int32> DataB = { 1, 2, 3, 7, 8, 9 };
+			TArray<int32> DataC = { 1, 2, 3, 4, 5, 6, 7, 8, 9 };
+
+			// Test empty ranges
+			check(Algo::Mismatch(Empty, Empty) == 0);
+			check(Algo::Mismatch(Empty, DataA) == 0);
+			check(Algo::Mismatch(DataA, Empty) == 0);
+
+			// Test common initial sequences
+			check(Algo::Mismatch(DataA, DataB) == 3);
+			check(Algo::Mismatch(DataB, DataA) == 3);
+
+			// Test equal sequences
+			check(Algo::Mismatch(DataA, DataA) == 6);
+			check(Algo::Mismatch(DataB, DataB) == 6);
+			check(Algo::Mismatch(DataC, DataC) == 9);
+
+			// Test subsequences
+			check(Algo::Mismatch(DataA, DataC) == 6);
+			check(Algo::Mismatch(DataC, DataA) == 6);
+		}
+
+		{
+			auto CompareCaseInsensitive = [](TCHAR Lhs, TCHAR Rhs)
+			{
+				return FChar::ToUpper(Lhs) == FChar::ToUpper(Rhs);
+			};
+
+			FString Empty;
+			FString DataA = TEXT("HeLlO wOrLd");
+			FString DataB = TEXT("HELLO GOODBYE");
+			FString DataC = TEXT("hello");
+
+			// Test empty ranges with custom equality
+			check(Algo::Mismatch(Empty, Empty, CompareCaseInsensitive) == 0);
+			check(Algo::Mismatch(Empty, DataA, CompareCaseInsensitive) == 0);
+			check(Algo::Mismatch(DataA, Empty, CompareCaseInsensitive) == 0);
+
+			// Test common initial sequences
+			check(Algo::Mismatch(DataA, DataB, CompareCaseInsensitive) == 6);
+			check(Algo::Mismatch(DataB, DataA, CompareCaseInsensitive) == 6);
+
+			// Test equal sequences
+			check(Algo::Mismatch(DataA, DataA, CompareCaseInsensitive) == 11);
+			check(Algo::Mismatch(DataB, DataB, CompareCaseInsensitive) == 13);
+			check(Algo::Mismatch(DataC, DataC, CompareCaseInsensitive) == 5);
+
+			// Test subsequences
+			check(Algo::Mismatch(DataA, DataC, CompareCaseInsensitive) == 5);
+			check(Algo::Mismatch(DataC, DataA, CompareCaseInsensitive) == 5);
+			check(Algo::Mismatch(DataB, DataC, CompareCaseInsensitive) == 5);
+			check(Algo::Mismatch(DataC, DataB, CompareCaseInsensitive) == 5);
+		}
+
+		{
+			auto Square = [](int32 Val)
+			{
+				return Val * Val;
+			};
+
+			TArray<int32> Empty;
+			TArray<int32> DataA = { 1, 2, 3, 4, 5, 6 };
+			TArray<int32> DataB = { -1, -2, -3, -7, -8, -9 };
+			TArray<int32> DataC = { 1, -2, 3, -4, 5, -6, 7, -8, 9 };
+
+			// Test empty ranges with projection
+			check(Algo::MismatchBy(Empty, Empty, Square) == 0);
+			check(Algo::MismatchBy(Empty, DataA, Square) == 0);
+			check(Algo::MismatchBy(DataA, Empty, Square) == 0);
+
+			// Test common initial sequences with projection
+			check(Algo::MismatchBy(DataA, DataB, Square) == 3);
+			check(Algo::MismatchBy(DataB, DataA, Square) == 3);
+
+			// Test equal sequences with projection
+			check(Algo::MismatchBy(DataA, DataA, Square) == 6);
+			check(Algo::MismatchBy(DataB, DataB, Square) == 6);
+			check(Algo::MismatchBy(DataC, DataC, Square) == 9);
+
+			// Test subsequences with projection
+			check(Algo::MismatchBy(DataA, DataC, Square) == 6);
+			check(Algo::MismatchBy(DataC, DataA, Square) == 6);
+		}
+
+		{
+			struct FStringWrapper
+			{
+				FString Str;
+			};
+
+			auto CompareCaseInsensitive = [](const FString& Lhs, const FString& Rhs)
+			{
+				return Lhs.Equals(Rhs, ESearchCase::IgnoreCase);
+			};
+
+			TArray<FStringWrapper> Empty;
+			TArray<FStringWrapper> DataA = { { TEXT("Class") }, { TEXT("Struct") }, { TEXT("Enum") }, { TEXT("Float") }, { TEXT("Int") }, { TEXT("Char") } };
+			TArray<FStringWrapper> DataB = { { TEXT("class") }, { TEXT("struct") }, { TEXT("enum") }, { TEXT("public") }, { TEXT("protected") }, { TEXT("private") } };
+			TArray<FStringWrapper> DataC = { { TEXT("CLASS") }, { TEXT("STRUCT") }, { TEXT("ENUM") }, { TEXT("FLOAT") }, { TEXT("INT") }, { TEXT("CHAR") }, { TEXT("PUBLIC") }, { TEXT("PROTECTED") }, { TEXT("PRIVATE") } };
+
+			// Test empty ranges with projection and custom equality
+			check(Algo::MismatchBy(Empty, Empty, &FStringWrapper::Str, CompareCaseInsensitive) == 0);
+			check(Algo::MismatchBy(Empty, DataA, &FStringWrapper::Str, CompareCaseInsensitive) == 0);
+			check(Algo::MismatchBy(DataA, Empty, &FStringWrapper::Str, CompareCaseInsensitive) == 0);
+
+			// Test common initial sequences with projection and custom equality
+			check(Algo::MismatchBy(DataA, DataB, &FStringWrapper::Str, CompareCaseInsensitive) == 3);
+			check(Algo::MismatchBy(DataB, DataA, &FStringWrapper::Str, CompareCaseInsensitive) == 3);
+
+			// Test equal sequences with projection and custom equality
+			check(Algo::MismatchBy(DataA, DataA, &FStringWrapper::Str, CompareCaseInsensitive) == 6);
+			check(Algo::MismatchBy(DataB, DataB, &FStringWrapper::Str, CompareCaseInsensitive) == 6);
+			check(Algo::MismatchBy(DataC, DataC, &FStringWrapper::Str, CompareCaseInsensitive) == 9);
+
+			// Test subsequences with projection and custom equality
+			check(Algo::MismatchBy(DataA, DataC, &FStringWrapper::Str, CompareCaseInsensitive) == 6);
+			check(Algo::MismatchBy(DataC, DataA, &FStringWrapper::Str, CompareCaseInsensitive) == 6);
+		}
+	}
+
 private:
 	TArray<int> TestData;
 	TArray<int> TestData2;
@@ -703,6 +829,7 @@ bool FAlgosTest::RunTest(const FString& Parameters)
 	TestEditDistance();
 	TestEditDistanceArray();
 	TestIncludes();
+	TestMismatch();
 	Cleanup();
 
 	return true;
diff --git a/Engine/Source/Runtime/Core/Public/Algo/Mismatch.h b/Engine/Source/Runtime/Core/Public/Algo/Mismatch.h
new file mode 100644
--- /dev/null
+++ b/Engine/Source/Runtime/Core/Public/Algo/Mismatch.h
@@ -0,0 +1,108 @@
+// Copyright Epic Games, Inc. All Rights Reserved.
+
+#pragma once
+
+#include "CoreTypes.h"
+#include "Templates/EqualTo.h"
+#include "Templates/IdentityFunctor.h"
+#include "Templates/Invoke.h"
+#include "Templates/UnrealTemplate.h"
+
+#include <type_traits>
+
+namespace Algo::Private
+{
+	template <typename InputAType, typename InputBType, typename ProjectionType, typename PredicateType>
+	constexpr auto Mismatch(InputAType&& InputA, InputBType&& InputB, ProjectionType Projection, PredicateType Predicate)
+	{
+		using IndexType = std::common_type_t<decltype(GetNum(InputA)), decltype(GetNum(InputB))>;
+
+		IndexType Index = 0;
+
+		const IndexType SizeA = (IndexType)GetNum(InputA);
+		const IndexType SizeB = (IndexType)GetNum(InputB);
+
+		auto* A = GetData(InputA);
+		auto* B = GetData(InputB);
+
+		for (;;)
+		{
+			if (Index == SizeA || Index == SizeB)
+			{
+				break;
+			}
+
+			if (!Invoke(Predicate, Invoke(Projection, *A++), Invoke(Projection, *B++)))
+			{
+				break;
+			}
+
+			++Index;
+		}
+
+		return Index;
+	}
+}
+
+namespace Algo
+{
+	/**
+	 * Returns the index at which two contiguous containers differ, using operator== to compare pairs of elements.
+	 *
+	 * @param  InputA     Container of elements that are used as the first argument to operator==.
+	 * @param  InputB     Container of elements that are used as the second argument to operator==.
+	 *
+	 * @return The index at which the containers differ.
+	 */
+	template <typename InputAType, typename InputBType>
+	constexpr auto Mismatch(InputAType&& InputA, InputBType&& InputB) -> std::common_type_t<decltype(GetNum(InputA)), decltype(GetNum(InputB))>
+	{
+		return Private::Mismatch(Forward<InputAType>(InputA), Forward<InputBType>(InputB), FIdentityFunctor(), TEqualTo<>());
+	}
+
+	/**
+	 * Returns the index at which two contiguous containers differ, using operator== to compare pairs of elements.
+	 *
+	 * @param  InputA     Container of elements that are used as the first argument to operator==.
+	 * @param  InputB     Container of elements that are used as the second argument to operator==.
+	 * @param  Predicate  Condition which returns true for elements which are deemed equal.
+	 *
+	 * @return The index at which the containers differ.
+	 */
+	template <typename InputAType, typename InputBType, typename PredicateType>
+	constexpr auto Mismatch(InputAType&& InputA, InputBType&& InputB, PredicateType Predicate) -> std::common_type_t<decltype(GetNum(InputA)), decltype(GetNum(InputB))>
+	{
+		return Private::Mismatch(Forward<InputAType>(InputA), Forward<InputBType>(InputB), FIdentityFunctor(), MoveTemp(Predicate));
+	}
+
+	/**
+	 * Returns the index at which two contiguous containers differ, using operator== to compare pairs of projected elements.
+	 *
+	 * @param  InputA     Container of elements that are used as the first argument to operator==.
+	 * @param  InputB     Container of elements that are used as the second argument to operator==.
+	 * @param  Projection Projection to apply to the elements before comparing them.
+	 *
+	 * @return The index at which the containers differ.
+	 */
+	template <typename InputAType, typename InputBType, typename ProjectionType>
+	constexpr auto MismatchBy(InputAType&& InputA, InputBType&& InputB, ProjectionType Projection) -> std::common_type_t<decltype(GetNum(InputA)), decltype(GetNum(InputB))>
+	{
+		return Private::Mismatch(Forward<InputAType>(InputA), Forward<InputBType>(InputB), MoveTemp(Projection), TEqualTo<>());
+	}
+
+	/**
+	 * Returns the index at which two contiguous containers differ, using a predicate to compare pairs of projected elements.
+	 *
+	 * @param  InputA     Container of elements that are used as the first argument to the predicate.
+	 * @param  InputB     Container of elements that are used as the second argument to the predicate.
+	 * @param  Projection Projection to apply to the elements before comparing them.
+	 * @param  Predicate  Condition which returns true for elements which are deemed equal.
+	 *
+	 * @return The index at which the containers differ.
+	 */
+	template <typename InputAType, typename InputBType, typename ProjectionType, typename PredicateType>
+	constexpr auto MismatchBy(InputAType&& InputA, InputBType&& InputB, ProjectionType Projection, PredicateType Predicate) -> std::common_type_t<decltype(GetNum(InputA)), decltype(GetNum(InputB))>
+	{
+		return Private::Mismatch(Forward<InputAType>(InputA), Forward<InputBType>(InputB), MoveTemp(Projection), MoveTemp(Predicate));
+	}
+}
