Adding FUniqueBuffer::AllocZeroed
#rnx

diff --git a/Engine/Source/Runtime/Core/Private/Memory/SharedBuffer.cpp b/Engine/Source/Runtime/Core/Private/Memory/SharedBuffer.cpp
--- a/Engine/Source/Runtime/Core/Private/Memory/SharedBuffer.cpp
+++ b/Engine/Source/Runtime/Core/Private/Memory/SharedBuffer.cpp
@@ -105,6 +105,13 @@ FUniqueBuffer FUniqueBuffer::Alloc(uint64 InSize)
 	return FUniqueBuffer(new UE::SharedBuffer::Private::FBufferOwnerHeap(InSize));
 }
 
+FUniqueBuffer FUniqueBuffer::AllocZeroed(uint64 InSize)
+{
+	FUniqueBuffer Buffer = Alloc(InSize);
+	FMemory::Memzero(Buffer.GetData(), Buffer.GetSize());
+	return Buffer;
+}
+
 FUniqueBuffer FUniqueBuffer::Clone(FMemoryView View)
 {
 	FUniqueBuffer Buffer = Alloc(View.GetSize());
diff --git a/Engine/Source/Runtime/Core/Public/Memory/SharedBuffer.h b/Engine/Source/Runtime/Core/Public/Memory/SharedBuffer.h
--- a/Engine/Source/Runtime/Core/Public/Memory/SharedBuffer.h
+++ b/Engine/Source/Runtime/Core/Public/Memory/SharedBuffer.h
@@ -219,6 +219,8 @@ class FUniqueBuffer
 public:
 	/** Make an uninitialized owned buffer of the specified size. */
 	[[nodiscard]] CORE_API static FUniqueBuffer Alloc(uint64 Size);
+	/** Make a zeroed owned buffer of the specified size. */
+	[[nodiscard]] CORE_API static FUniqueBuffer AllocZeroed(uint64 Size);
 
 	/** Make an owned clone of the input. */
 	[[nodiscard]] CORE_API static FUniqueBuffer Clone(FMemoryView View);
