Added dll exports to fix linker errors

diff --git a/Engine/Source/Runtime/Core/Public/HAL/IPlatformFileManagedStorageWrapper.h b/Engine/Source/Runtime/Core/Public/HAL/IPlatformFileManagedStorageWrapper.h
--- a/Engine/Source/Runtime/Core/Public/HAL/IPlatformFileManagedStorageWrapper.h
+++ b/Engine/Source/Runtime/Core/Public/HAL/IPlatformFileManagedStorageWrapper.h
@@ -36,7 +36,7 @@
 
 #include <atomic>
 
-DECLARE_LOG_CATEGORY_EXTERN(LogPlatformFileManagedStorage, Log, All);
+CORE_API DECLARE_LOG_CATEGORY_EXTERN(LogPlatformFileManagedStorage, Log, All);
 
 namespace ManagedStorageInternal
 {
@@ -126,10 +126,10 @@ public:
 	FManagedStorageScopeFileLock(FManagedStorageScopeFileLock&& InOther) = delete;
 	FManagedStorageScopeFileLock& operator=(FManagedStorageScopeFileLock&& InOther) = delete;
 	
-	void Unlock();
+	CORE_API void Unlock();
 
 private:
-	void Lock();
+	CORE_API void Lock();
 
 private:
 	FPersistentManagedFile ManagedFile;
