Fail gracefully if we fail to allocate TLS data for the CSV profiler.
Fix a bug where IsValidTlsData doesn't work on some platforms
Speculative fix for OOM issues

diff --git a/Engine/Source/Runtime/Core/Private/ProfilingDebugging/CsvProfiler.cpp b/Engine/Source/Runtime/Core/Private/ProfilingDebugging/CsvProfiler.cpp
--- a/Engine/Source/Runtime/Core/Private/ProfilingDebugging/CsvProfiler.cpp
+++ b/Engine/Source/Runtime/Core/Private/ProfilingDebugging/CsvProfiler.cpp
@@ -3349,7 +3349,15 @@ void FCsvProfiler::BeginCapture(int InNumFramesToCapture,
 	// Lazy init TLS before starting the capture
 	FCsvProfilerThreadData::InitTls();
 
-	CommandQueue.Enqueue(FCsvCaptureCommand(ECsvCommandType::Start, GCsvProfilerFrameNumber, InNumFramesToCapture, InDestinationFolder, InFilename, InFlags));
+	// Check if we actually have valid TLS data before starting the profile
+	if (!FCsvProfilerThreadData::IsTlsSlotInitialized())
+	{
+		UE_LOG(LogCsvProfiler, Error, TEXT("Failed to allocate TLS! Not starting the CSV capture"));
+	}
+	else
+	{
+		CommandQueue.Enqueue(FCsvCaptureCommand(ECsvCommandType::Start, GCsvProfilerFrameNumber, InNumFramesToCapture, InDestinationFolder, InFilename, InFlags));
+	}
 }
 
 TSharedFuture<FString> FCsvProfiler::EndCapture(FGraphEventRef EventToSignal)
