Preset for "light" memory tracing

In some cases it could be useful to trace detailed allocations without callstack, and solely rely on tags for analysis. Enable by starting the process with `-trace=memory_light`.

diff --git a/Engine/Build/Android/Java/src/com/epicgames/unreal/GameActivity.java.template b/Engine/Build/Android/Java/src/com/epicgames/unreal/GameActivity.java.template
--- a/Engine/Build/Android/Java/src/com/epicgames/unreal/GameActivity.java.template
+++ b/Engine/Build/Android/Java/src/com/epicgames/unreal/GameActivity.java.template
@@ -7135,8 +7135,17 @@ public class GameActivity extends $${gameActivitySuperClass}$$ implements Surfac
 				bundle.getString("com.epicgames.unreal.GameActivity.BuildConfiguration").equals("Shipping");
 			if (!bShipping)
 			{ 
-				boolean bEnableMemoryTracing = new File(GameApplication.getAppContext().getFilesDir() + "/UEEnableMemoryTracing.txt").exists();
-				Os.setenv("UEEnableMemoryTracing", bEnableMemoryTracing ? "1" : "0", true);
+				boolean bEnableFullMemoryTracing = new File(GameApplication.getAppContext().getFilesDir() + "/UEEnableMemoryTracing.txt").exists();
+				boolean bEnableLightlMemoryTracing = new File(GameApplication.getAppContext().getFilesDir() + "/UEEnableMemoryTracingLight.txt").exists();
+				
+				if (bEnableFullMemoryTracing)
+				{
+					Os.setenv("UEEnableMemoryTracing", "full", true);
+				}
+				else if (bEnableLightlMemoryTracing)
+				{
+					Os.setenv("UEEnableMemoryTracing", "light", true);
+				}
 			}
 		}
 		catch (ErrnoException | NameNotFoundException e)
diff --git a/Engine/Source/Runtime/Core/Private/Android/AndroidMemoryTrace.cpp b/Engine/Source/Runtime/Core/Private/Android/AndroidMemoryTrace.cpp
--- a/Engine/Source/Runtime/Core/Private/Android/AndroidMemoryTrace.cpp
+++ b/Engine/Source/Runtime/Core/Private/Android/AndroidMemoryTrace.cpp
@@ -8,7 +8,7 @@
 #include <string.h>
 
 ////////////////////////////////////////////////////////////////////////////////
-FMalloc* MemoryTrace_CreateInternal(FMalloc*);
+FMalloc* MemoryTrace_CreateInternal(FMalloc*, EMemoryTraceInit);
 
 ////////////////////////////////////////////////////////////////////////////////
 FMalloc* MemoryTrace_Create(FMalloc* InMalloc)
@@ -16,8 +16,18 @@ FMalloc* MemoryTrace_Create(FMalloc* InMalloc)
 	const char* UEEnableMemoryTracing = getenv("UEEnableMemoryTracing");
 	__android_log_print(ANDROID_LOG_DEBUG, "UE", "getenv(\"UEEnableMemoryTracing\") == \"%s\"", UEEnableMemoryTracing ? UEEnableMemoryTracing : "nullptr");
 
-	const bool bEnableMemoryTracing = UEEnableMemoryTracing != nullptr && !strncmp(UEEnableMemoryTracing, "1", 1);
-	return bEnableMemoryTracing ? MemoryTrace_CreateInternal(InMalloc) :  InMalloc;
+	// If environment variable UEEnableMemoryTracing is present memory tracing will be fully active
+	// unless equal to "light" in which case the light memory preset is used.
+	EMemoryTraceInit Mode = EMemoryTraceInit::Disabled;
+	if (UEEnableMemoryTracing != nullptr && !strncmp(UEEnableMemoryTracing, "light", 5))
+	{
+		Mode = EMemoryTraceInit::Light;
+	}
+	else if (UEEnableMemoryTracing != nullptr)
+	{
+		Mode = EMemoryTraceInit::Full;
+	}
+	return MemoryTrace_CreateInternal(InMalloc, Mode);
 }
 
 #endif // UE_MEMORY_TRACE_ENABLED
diff --git a/Engine/Source/Runtime/Core/Private/ProfilingDebugging/MemoryTrace.cpp b/Engine/Source/Runtime/Core/Private/ProfilingDebugging/MemoryTrace.cpp
--- a/Engine/Source/Runtime/Core/Private/ProfilingDebugging/MemoryTrace.cpp
+++ b/Engine/Source/Runtime/Core/Private/ProfilingDebugging/MemoryTrace.cpp
@@ -290,8 +290,9 @@ static FUndestructed<FTraceMalloc> GTraceMalloc;
 
 ////////////////////////////////////////////////////////////////////////////////
 template <typename ArgCharType>
-static bool MemoryTrace_ShouldEnable(int32 ArgC, const ArgCharType* const* ArgV)
+static EMemoryTraceInit MemoryTrace_ShouldEnable(int32 ArgC, const ArgCharType* const* ArgV)
 {
+	EMemoryTraceInit Mode = EMemoryTraceInit::Disabled;
 	for (int32 ArgIndex = 1; ArgIndex < ArgC; ArgIndex++)
 	{
 		const ArgCharType* Arg = ArgV[ArgIndex];
@@ -319,9 +320,28 @@ static bool MemoryTrace_ShouldEnable(int32 ArgC, const ArgCharType* const* ArgV)
 				{
 					ChannelsOrPresets.FindChar(',', CommaPos);
 					TStringView<ArgCharType> Channel = ChannelsOrPresets.SubStr(0, CommaPos == INDEX_NONE ? ChannelsOrPresets.Len() : CommaPos);
-					if (Channel.Equals("memalloc", ESearchCase::IgnoreCase) || Channel.Equals("memory", ESearchCase::IgnoreCase))
+					if (Channel.Equals("memalloc", ESearchCase::IgnoreCase))
 					{
-						return true;
+						Mode |= EMemoryTraceInit::AllocEvents;
+					}
+					else if (Channel.Equals("callstack", ESearchCase::IgnoreCase))
+					{
+						Mode |= EMemoryTraceInit::Callstacks;
+					}
+					else if (Channel.Equals("memtag", ESearchCase::IgnoreCase))
+					{
+						Mode |= EMemoryTraceInit::Tags;
+					}
+					else if (Channel.Left(6).Equals("memory", ESearchCase::IgnoreCase))
+					{
+						if (Channel.Equals("memory", ESearchCase::IgnoreCase))
+						{
+							return EMemoryTraceInit::Full;
+						}
+						else if (Channel.Equals("memory_light", ESearchCase::IgnoreCase))
+						{
+							return EMemoryTraceInit::Light;
+						}
 					}
 					ChannelsOrPresets.RightChopInline(Channel.Len() + 1);
 				}
@@ -332,12 +352,19 @@ static bool MemoryTrace_ShouldEnable(int32 ArgC, const ArgCharType* const* ArgV)
 		}
 	}
 
-	return false;
+	return Mode;
 }
 
 ////////////////////////////////////////////////////////////////////////////////
-FMalloc* MemoryTrace_CreateInternal(FMalloc* InMalloc)
+FMalloc* MemoryTrace_CreateInternal(FMalloc* InMalloc, EMemoryTraceInit Mode)
 {
+	// If allocation events are not desired we don't need to do anything, even
+	// if user has enabled only callstacks it will be enabled later. 
+	if (!EnumHasAnyFlags(Mode, EMemoryTraceInit::AllocEvents))
+	{
+		return InMalloc;
+	}
+	
 	// Some OSes (i.e. Windows) will terminate all threads except the main
 	// one as part of static deinit. However we may receive more memory
 	// trace events that would get lost as Trace's worker thread has been
@@ -351,9 +378,14 @@ FMalloc* MemoryTrace_CreateInternal(FMalloc* InMalloc)
 
 	// Both tag and callstack tracing need to use the wrapped trace malloc
 	// so we can break out tracing memory overhead (and not cause recursive behaviour).
-	MemoryTrace_InitTags(&GTraceMalloc);
-	CallstackTrace_Create(&GTraceMalloc);
-
+	if (EnumHasAnyFlags(Mode, EMemoryTraceInit::Tags))
+	{
+		MemoryTrace_InitTags(&GTraceMalloc);
+	}
+	if (EnumHasAnyFlags(Mode, EMemoryTraceInit::Callstacks))
+	{
+		CallstackTrace_Create(&GTraceMalloc);
+	}
 
 	static FUndestructed<FMallocWrapper> SMallocWrapper;
 	SMallocWrapper.Construct(InMalloc);
@@ -364,23 +396,15 @@ FMalloc* MemoryTrace_CreateInternal(FMalloc* InMalloc)
 ////////////////////////////////////////////////////////////////////////////////
 FMalloc* MemoryTrace_CreateInternal(FMalloc* InMalloc, int32 ArgC, const WIDECHAR* const* ArgV)
 {
-	if (!MemoryTrace_ShouldEnable(ArgC, ArgV))
-	{
-		return nullptr;
-	}
-
-	return MemoryTrace_CreateInternal(InMalloc);
+	const EMemoryTraceInit Mode = MemoryTrace_ShouldEnable(ArgC, ArgV);
+	return MemoryTrace_CreateInternal(InMalloc, Mode);
 }
 
 ////////////////////////////////////////////////////////////////////////////////
 FMalloc* MemoryTrace_CreateInternal(FMalloc* InMalloc, int32 ArgC, const ANSICHAR* const* ArgV)
 {
-	if (!MemoryTrace_ShouldEnable(ArgC, ArgV))
-	{
-		return nullptr;
-	}
-
-	return MemoryTrace_CreateInternal(InMalloc);
+	const EMemoryTraceInit Mode = MemoryTrace_ShouldEnable(ArgC, ArgV);
+	return MemoryTrace_CreateInternal(InMalloc, Mode);
 }
 
 ////////////////////////////////////////////////////////////////////////////////
diff --git a/Engine/Source/Runtime/Core/Private/ProfilingDebugging/TraceAuxiliary.cpp b/Engine/Source/Runtime/Core/Private/ProfilingDebugging/TraceAuxiliary.cpp
--- a/Engine/Source/Runtime/Core/Private/ProfilingDebugging/TraceAuxiliary.cpp
+++ b/Engine/Source/Runtime/Core/Private/ProfilingDebugging/TraceAuxiliary.cpp
@@ -68,6 +68,7 @@
 ////////////////////////////////////////////////////////////////////////////////
 const TCHAR* GDefaultChannels = TEXT("cpu,gpu,frame,log,bookmark,screenshot,region");
 const TCHAR* GMemoryChannels = TEXT("memtag,memalloc,callstack,module");
+const TCHAR* GMemoryLightChannels = TEXT("memtag,memalloc");
 const TCHAR* GTraceConfigSection = TEXT("Trace.Config");
 static UE::Trace::FInitializeDesc GInitializeDesc;
 
@@ -279,10 +280,14 @@ void FTraceAuxiliaryImpl::ForEachChannel(const TCHAR* ChannelList, bool bResolve
 			{
 				ForEachChannel(GDefaultChannels, false, LogCategory, Callable);
 			}
-			else if (FCString::Stricmp(Name, TEXT("memory"))== 0)
+			else if (FCString::Stricmp(Name, TEXT("memory")) == 0)
 			{
 				ForEachChannel(GMemoryChannels, false, LogCategory, Callable);
 			}
+			else if (FCString::Stricmp(Name, TEXT("memory_light")) == 0)
+			{
+				ForEachChannel(GMemoryLightChannels, false, LogCategory, Callable);
+			}
 			// Check against data driven presets (if available)
 			else if (GConfig && GConfig->GetString(TEXT("Trace.ChannelPresets"), Name, Value, GEngineIni))
 			{
diff --git a/Engine/Source/Runtime/Core/Public/ProfilingDebugging/MemoryTrace.h b/Engine/Source/Runtime/Core/Public/ProfilingDebugging/MemoryTrace.h
--- a/Engine/Source/Runtime/Core/Public/ProfilingDebugging/MemoryTrace.h
+++ b/Engine/Source/Runtime/Core/Public/ProfilingDebugging/MemoryTrace.h
@@ -68,6 +68,22 @@ enum class EMemoryTraceSwapOperation : uint8
 	FreeInSwap = 2, // Freed while being paged out in swap
 };
 
+////////////////////////////////////////////////////////////////////////////////
+
+// Internal options for early initialization of memory tracing systems. Exposed
+// here due to visibility in platform implementations.
+enum class EMemoryTraceInit : uint8
+{
+	Disabled		= 0,
+	AllocEvents		= 1 << 0,
+	Callstacks		= 1 << 1,
+	Tags			= 1 << 2,
+	Full			= AllocEvents|Callstacks|Tags,
+	Light			= AllocEvents|Tags,
+};
+
+ENUM_CLASS_FLAGS(EMemoryTraceInit);
+
 ////////////////////////////////////////////////////////////////////////////////
 #if UE_MEMORY_TRACE_ENABLED
 
