Fixed the UE_REQUIRES Clang 16 workaround on Apple.

diff --git a/Engine/Source/Runtime/Core/Public/Templates/Requires.h b/Engine/Source/Runtime/Core/Public/Templates/Requires.h
--- a/Engine/Source/Runtime/Core/Public/Templates/Requires.h
+++ b/Engine/Source/Runtime/Core/Public/Templates/Requires.h
@@ -77,8 +77,8 @@
 	the body of a UE_REQUIRES() macro - standalone concept checks must still be
 	done via something like a TModels-type concept.
  -----------------------------------------------------------------------------*/
-#if __cplusplus < 202000 || (defined(__clang__) && __clang_major__ == 16)
-	// Clang 16 treats a UE_REQUIRES_FRIEND declaration as a separate function, causing ambiguous overload resolution,
+#if __cplusplus < 202000 || (defined(__clang__) && ((!PLATFORM_APPLE && __clang_major__ == 16) || (PLATFORM_APPLE && __clang_major__ == 15)))
+	// Clang 16 (or Apple Clang 15) treats a UE_REQUIRES_FRIEND declaration as a separate function, causing ambiguous overload resolution,
 	// so fall back to the non-concept implementation.
 	//
 	// [REDACTED_PROFILE_URL]/llvm-project/issues/60749
