Support for subframe on FTimecode such that we can reflect subframe information in widgets, e.g. STimecode widget uses FTimecode::ToString() for conversion to TEXT()

diff --git a/Engine/Source/Runtime/Core/Public/Misc/QualifiedFrameTime.h b/Engine/Source/Runtime/Core/Public/Misc/QualifiedFrameTime.h
--- a/Engine/Source/Runtime/Core/Public/Misc/QualifiedFrameTime.h
+++ b/Engine/Source/Runtime/Core/Public/Misc/QualifiedFrameTime.h
@@ -56,13 +56,12 @@ public:
 	/**
 	 * Create an FTimecode from this qualified frame time.
 	 *
-	 * Any subframe value in the frame time will be ignored.
 	 * Whether or not the returned timecode is a drop-frame timecode will be determined by the qualified frame time's frame rate
 	 * and the CVar specifying whether to generate drop-frame timecodes by default for supported frame rates.
 	 */
 	FTimecode ToTimecode() const
 	{
-		return FTimecode::FromFrameNumber(Time.FloorToFrame(), Rate);
+		return FTimecode::FromFrameTime(Time, Rate);
 	}
 
 	/**
@@ -75,7 +74,7 @@ public:
 	 */
 	FTimecode ToTimecode(bool bDropFrame) const
 	{
-		return FTimecode::FromFrameNumber(Time.FloorToFrame(), Rate, bDropFrame);
+		return FTimecode::FromFrameTime(Time, Rate, bDropFrame);
 	}
 
 public:
diff --git a/Engine/Source/Runtime/Core/Public/Misc/Timecode.h b/Engine/Source/Runtime/Core/Public/Misc/Timecode.h
--- a/Engine/Source/Runtime/Core/Public/Misc/Timecode.h
+++ b/Engine/Source/Runtime/Core/Public/Misc/Timecode.h
@@ -8,6 +8,7 @@
 #include "Math/UnrealMathUtility.h"
 #include "Misc/FrameNumber.h"
 #include "Misc/FrameRate.h"
+#include "Misc/FrameTime.h"
 #include "Misc/Timespan.h"
 
 /**
@@ -24,6 +25,7 @@ struct FTimecode
 		, Minutes(0)
 		, Seconds(0)
 		, Frames(0)
+		, Subframe(0)
 		, bDropFrameFormat(false)
 	{}
 
@@ -41,6 +43,21 @@ struct FTimecode
 		, bDropFrameFormat(InbDropFrame)
 	{}
 
+	/**
+	 * User construction from a number of hours minutes seconds frames, and subframes.
+	 * @param InbDropFrame - If true, this Timecode represents a "Drop Frame Timecode" format which
+							skips the first frames of every minute (except those ending in multiples of 10)
+							to account for drift when using a fractional NTSC framerate.
+ 	 */
+	explicit FTimecode(int32 InHours, int32 InMinutes, int32 InSeconds, int32 InFrames, float InSubframe, bool InbDropFrame)
+		: Hours(InHours)
+		, Minutes(InMinutes)
+		, Seconds(InSeconds)
+		, Frames(InFrames)
+		, Subframe(InSubframe)
+		, bDropFrameFormat(InbDropFrame)
+	{}
+
 	/**
 	 * User construction from a time in seconds
 	 * @param InbDropFrame	- If true, this Timecode represents a "Drop Frame Timecode" format which
@@ -227,6 +244,22 @@ public:
 		return FromFrameNumber(InFrameNumber, InFrameRate, UseDropFormatTimecode(InFrameRate));
 	}
 
+	/**
+	 * Create a FTimecode from a specific frame time at the given frame rate.
+	 *
+	 * @param InFrameTime - The frame time to convert into a timecode. This should already be converted to InFrameTimes's resolution.
+	 * @param InFrameRate - The framerate that this timecode is based in. This should be the playback framerate as it is used to determine
+	 *					    when the Frame value wraps over.
+	 * @param bDropFrame  - Optional parameter to indicate if we are in drop frame format.
+	 */
+	static FTimecode FromFrameTime(const FFrameTime& InFrameTime, const FFrameRate& InFrameRate, TOptional<bool> bDropFrame = {})
+	{
+		bool bShouldUseDropFrame = bDropFrame ? *bDropFrame : UseDropFormatTimecode(InFrameRate);
+		FTimecode Timecode = FromFrameNumber(InFrameTime.FloorToFrame(), InFrameRate, bShouldUseDropFrame);
+		Timecode.Subframe = InFrameTime.GetSubFrame();
+		return Timecode;
+	}
+
 	/**
 	 * Converts this Timecode back into a timespan at the given framerate, taking into account if this is a drop-frame format timecode.
 	 */
@@ -319,14 +352,23 @@ public:
 			SignText = PositiveSign;
 		}
 
+		TStringBuilder<32> Builder;
+
 		if (bDropFrameFormat)
 		{
-			return FString::Printf(TEXT("%s%02d:%02d:%02d;%02d"), SignText, FMath::Abs(Hours), FMath::Abs(Minutes), FMath::Abs(Seconds), FMath::Abs(Frames));
+			Builder.Appendf(TEXT("%s%02d:%02d:%02d;%02d"), SignText, FMath::Abs(Hours), FMath::Abs(Minutes), FMath::Abs(Seconds), FMath::Abs(Frames));
 		}
 		else
 		{
-			return FString::Printf(TEXT("%s%02d:%02d:%02d:%02d"), SignText, FMath::Abs(Hours), FMath::Abs(Minutes), FMath::Abs(Seconds), FMath::Abs(Frames));
+			Builder.Appendf(TEXT("%s%02d:%02d:%02d:%02d"), SignText, FMath::Abs(Hours), FMath::Abs(Minutes), FMath::Abs(Seconds), FMath::Abs(Frames));
+		}
+
+		if (Subframe > 0)
+		{
+			int32 ClampedSubframe = static_cast<int32>(FMath::Clamp(100*Subframe,0,99));
+			Builder.Appendf(TEXT(".%02d"), ClampedSubframe);
 		}
+		return Builder.ToString();
 	}
 
 public:
@@ -343,13 +385,15 @@ public:
 	/** How many frames does this timecode represent */
 	int32 Frames;
 
+	/** The subframe value provided when converting to FFrameTime */
+	float Subframe = 0;
+
 	/** If true, this Timecode represents a Drop Frame timecode used to account for fractional frame rates in NTSC play rates. */
 	bool bDropFrameFormat;
 
-
 	friend inline bool operator==(const FTimecode& A, const FTimecode& B)
 	{
-		return A.Hours == B.Hours && A.Minutes == B.Minutes && A.Seconds == B.Seconds && A.Frames == B.Frames;
+		return A.Hours == B.Hours && A.Minutes == B.Minutes && A.Seconds == B.Seconds && A.Frames == B.Frames && A.Subframe == B.Subframe;
 	}
 
 	friend inline bool operator!=(const FTimecode& A, const FTimecode& B)
