Fix non-unity
#rnx

diff --git a/Engine/Source/Runtime/Core/Public/Misc/CoreMiscDefines.h b/Engine/Source/Runtime/Core/Public/Misc/CoreMiscDefines.h
--- a/Engine/Source/Runtime/Core/Public/Misc/CoreMiscDefines.h
+++ b/Engine/Source/Runtime/Core/Public/Misc/CoreMiscDefines.h
@@ -4,7 +4,8 @@
 
 // HEADER_UNIT_SKIP - Included by rc files and break if I include proper files
 
-#include "../HAL/PreprocessorHelpers.h"
+#include "HAL/Platform.h"
+#include "HAL/PreprocessorHelpers.h"
 
 // When passed to pragma message will result in clickable warning in VS
 #define WARNING_LOCATION(Line) __FILE__ "(" PREPROCESSOR_TO_STRING(Line) ")"
