Change headers to work for rc and standalone header compile

diff --git a/Engine/Source/Runtime/Core/Public/Misc/CoreMiscDefines.h b/Engine/Source/Runtime/Core/Public/Misc/CoreMiscDefines.h
--- a/Engine/Source/Runtime/Core/Public/Misc/CoreMiscDefines.h
+++ b/Engine/Source/Runtime/Core/Public/Misc/CoreMiscDefines.h
@@ -4,8 +4,12 @@
 
 // HEADER_UNIT_SKIP - Included by rc files and break if I include proper files
 
+#ifndef RC_INVOKED
 #include "HAL/Platform.h"
 #include "HAL/PreprocessorHelpers.h"
+#else 
+#include "../HAL/PreprocessorHelpers.h"
+#endif
 
 // When passed to pragma message will result in clickable warning in VS
 #define WARNING_LOCATION(Line) __FILE__ "(" PREPROCESSOR_TO_STRING(Line) ")"
@@ -520,4 +524,4 @@ struct FPlatformInputDeviceState
 
 	/** The connection state of this input device */
 	EInputDeviceConnectionState ConnectionState = EInputDeviceConnectionState::Invalid;
-};
\ No newline at end of file
+};
