Fix additional rc files and return CoreMiscDefines.h to not need special RC behavior
#rnx

diff --git a/Engine/Source/Programs/SwitchboardListener/SblCore/Private/Windows/Resources/Resource.rc2 b/Engine/Source/Programs/SwitchboardListener/SblCore/Private/Windows/Resources/Resource.rc2
--- a/Engine/Source/Programs/SwitchboardListener/SblCore/Private/Windows/Resources/Resource.rc2
+++ b/Engine/Source/Programs/SwitchboardListener/SblCore/Private/Windows/Resources/Resource.rc2
@@ -1,7 +1,7 @@
 // Copyright Epic Games, Inc. All Rights Reserved.
 
 #include "../../../../../../Runtime/Launch/Resources/Windows/resource.h"
-#include "../../../../../../Runtime/Core/Public/Misc/CoreMiscDefines.h"
+#include "../../../../../../Runtime/Core/Public/HAL/PreprocessorHelpers.h"
 #include "../../../../../../Runtime/Launch/Resources/Version.h"
 
 #include "../../../Internal/SwitchboardListenerVersion.h"
diff --git a/Engine/Source/Runtime/Core/Public/Misc/CoreMiscDefines.h b/Engine/Source/Runtime/Core/Public/Misc/CoreMiscDefines.h
--- a/Engine/Source/Runtime/Core/Public/Misc/CoreMiscDefines.h
+++ b/Engine/Source/Runtime/Core/Public/Misc/CoreMiscDefines.h
@@ -2,14 +2,8 @@
 
 #pragma once
 
-// HEADER_UNIT_SKIP - Included by rc files and break if I include proper files
-
-#ifndef RC_INVOKED
 #include "HAL/Platform.h"
 #include "HAL/PreprocessorHelpers.h"
-#else 
-#include "../HAL/PreprocessorHelpers.h"
-#endif
 
 // When passed to pragma message will result in clickable warning in VS
 #define WARNING_LOCATION(Line) __FILE__ "(" PREPROCESSOR_TO_STRING(Line) ")"
@@ -524,4 +518,4 @@ struct FPlatformInputDeviceState
 
 	/** The connection state of this input device */
 	EInputDeviceConnectionState ConnectionState = EInputDeviceConnectionState::Invalid;
-};
+};
\ No newline at end of file
