Deprecated TEMPLATE_REQUIRES.

diff --git a/Engine/Source/Runtime/Core/Public/Templates/Requires.h b/Engine/Source/Runtime/Core/Public/Templates/Requires.h
--- a/Engine/Source/Runtime/Core/Public/Templates/Requires.h
+++ b/Engine/Source/Runtime/Core/Public/Templates/Requires.h
@@ -101,5 +101,4 @@
 	#define UE_REQUIRES_EXPR(...) requires { (__VA_ARGS__); }
 #endif
 
-// This should be regarded as deprecated - please use UE_REQUIRES instead.
-#define TEMPLATE_REQUIRES(...) typename TEnableIf<__VA_ARGS__, int>::type = 0
+#define TEMPLATE_REQUIRES(...) UE_DEPRECATED_MACRO(5.5, "TEMPLATE_REQUIRES has been deprecated - please use UE_REQUIRES instead.") typename TEnableIf<__VA_ARGS__, int>::type = 0
