Add a command line switch to set process-wide affinity on Windows.

Now -processaffinity=<number of logical cores> can be used to restrict the process to use a specific number of logical cores, and -processaffinityphysical=<number of physical cores> can be used to restrict the process to use a specific number of physical (non-hyperthreaded) cores.

When using this option, explicitly setting thread affinities is disabled for simplicity.

diff --git a/Engine/Source/Runtime/Core/Private/Windows/WindowsPlatformProcess.cpp b/Engine/Source/Runtime/Core/Private/Windows/WindowsPlatformProcess.cpp
--- a/Engine/Source/Runtime/Core/Private/Windows/WindowsPlatformProcess.cpp
+++ b/Engine/Source/Runtime/Core/Private/Windows/WindowsPlatformProcess.cpp
@@ -95,6 +95,77 @@ namespace WindowsPlatformProcess
 	}
 }
 
+static bool bCustomProcessAffinity = false;
+bool FWindowsPlatformProcess::SetProcessAffinity(uint32 NumCoresForAffinity, bool bPhysicalCoresOnly)
+{
+	if (NumCoresForAffinity == 0)
+	{
+		return false;
+	}
+
+	DWORD_PTR AffinityMask = 0;
+	if (!bPhysicalCoresOnly)
+	{
+		if (NumCoresForAffinity > 64)
+		{
+			UE_LOG(LogWindows, Warning, TEXT("Requested process affinity to %d logical cores but the maximum affinity is 64 cores. Will use 64."), NumCoresForAffinity);
+			NumCoresForAffinity = 64;
+		}
+
+		if (NumCoresForAffinity == 64)
+		{
+			AffinityMask = 0xFFFFFFFFFFFFFFFF;
+		}
+		else
+		{
+			AffinityMask = (((DWORD_PTR)1) << NumCoresForAffinity) - 1;
+		}
+	}
+	else
+	{
+		if (NumCoresForAffinity > 32)
+		{
+			UE_LOG(LogWindows, Warning, TEXT("Requested process affinity to %d physical cores but the maximum affinity is 32 cores. Will use 32."), NumCoresForAffinity);
+			NumCoresForAffinity = 32;
+		}
+
+		// Windows numbers physical and logical (hyperthreaded) cores by interleaving them.
+		// So 0 is physical, 1 is logical, 2 is physical, 3 is logical, etc.
+		constexpr DWORD_PTR PhysicalMask = 0x5555555555555555;
+
+		if (NumCoresForAffinity == 32)
+		{
+			AffinityMask = PhysicalMask;
+		}
+		else
+		{
+			AffinityMask = (((DWORD_PTR)1) << (NumCoresForAffinity * 2)) - 1;
+			AffinityMask &= PhysicalMask;
+		}
+	}
+
+
+	if (!SetProcessAffinityMask(GetCurrentProcess(), AffinityMask))
+	{
+		DWORD LastError = GetLastError();
+		TCHAR ErrorMsg[1024];
+		FPlatformMisc::GetSystemErrorMessage(ErrorMsg, 1024, LastError);
+		UE_LOG(LogWindows, Error, TEXT("Failed to set process affinity, process may run on all available logical cores. Error: %d [%s]"), LastError, ErrorMsg);
+	}
+	else
+	{
+		UE_LOG(LogWindows, Log, TEXT("Successfully set process affinity, process will only use %d cores"), NumCoresForAffinity);
+		bCustomProcessAffinity = true;
+	}
+
+	return bCustomProcessAffinity;
+}
+
+bool FWindowsPlatformProcess::IsProcessAffinitySet()
+{
+	return bCustomProcessAffinity;
+}
+
 void FWindowsPlatformProcess::AddDllDirectory(const TCHAR* Directory)
 {
 	FString NormalizedDirectory = FPaths::ConvertRelativePathToFull(Directory);
@@ -592,6 +663,16 @@ uint32 FWindowsPlatformProcess::GetCurrentCoreNumber()
 
 void FWindowsPlatformProcess::SetThreadAffinityMask( uint64 AffinityMask )
 {
+	// While it's technically possible to use both process and thread affinities,
+	// it requires restricting the set of cores eligible for affinity to respect 
+	// the process affinity mask.
+	// For simplicity, as long as the process-wide affinity is a debugging option,
+	// disallow thread affinities when using process affinity.
+	if (FWindowsPlatformProcess::IsProcessAffinitySet())
+	{
+		return;
+	}
+
 	if( AffinityMask != FPlatformAffinity::GetNoAffinityMask() )
 	{
 		::SetThreadAffinityMask( ::GetCurrentThread(), (DWORD_PTR)AffinityMask );
diff --git a/Engine/Source/Runtime/Core/Private/Windows/WindowsRunnableThread.cpp b/Engine/Source/Runtime/Core/Private/Windows/WindowsRunnableThread.cpp
--- a/Engine/Source/Runtime/Core/Private/Windows/WindowsRunnableThread.cpp
+++ b/Engine/Source/Runtime/Core/Private/Windows/WindowsRunnableThread.cpp
@@ -113,6 +113,16 @@ uint32 FRunnableThreadWin::GuardedRun()
 
 bool FRunnableThreadWin::SetThreadAffinity(const FThreadAffinity& Affinity)
 {
+	// While it's technically possible to use both process and thread affinities,
+	// it requires restricting the set of cores eligible for affinity to respect 
+	// the process affinity mask.
+	// For simplicity, as long as the process-wide affinity is a debugging option,
+	// disallow thread affinities when using process affinity.
+	if (FWindowsPlatformProcess::IsProcessAffinitySet())
+	{
+		return false;
+	}
+
 	const FProcessorGroupDesc& ProcessorGroups = FPlatformMisc::GetProcessorGroupDesc();
 	int32 CpuGroupCount = ProcessorGroups.NumProcessorGroups;
 	check(Affinity.ProcessorGroup < CpuGroupCount);
diff --git a/Engine/Source/Runtime/Core/Public/Windows/WindowsPlatformProcess.h b/Engine/Source/Runtime/Core/Public/Windows/WindowsPlatformProcess.h
--- a/Engine/Source/Runtime/Core/Public/Windows/WindowsPlatformProcess.h
+++ b/Engine/Source/Runtime/Core/Public/Windows/WindowsPlatformProcess.h
@@ -123,6 +123,9 @@ struct FWindowsPlatformProcess
 
 public:
 
+	static CORE_API bool SetProcessAffinity(uint32 NumCoresForAffinity, bool bPhysicalCoresOnly);
+	static CORE_API bool IsProcessAffinitySet();
+
 	// FGenericPlatformProcess interface
 
 	static CORE_API void* GetDllHandle( const TCHAR* Filename );
diff --git a/Engine/Source/Runtime/Launch/Private/Windows/LaunchWindows.cpp b/Engine/Source/Runtime/Launch/Private/Windows/LaunchWindows.cpp
--- a/Engine/Source/Runtime/Launch/Private/Windows/LaunchWindows.cpp
+++ b/Engine/Source/Runtime/Launch/Private/Windows/LaunchWindows.cpp
@@ -217,6 +217,20 @@ LAUNCH_API int32 LaunchWindowsStartup( HINSTANCE hInInstance, HINSTANCE hPrevIns
 		GAlwaysReportCrash = true;
 	}
 
+	uint32 LogicalCoreAffinity = 0;
+	uint32 PhysicalCoreAffinity = 0;
+	FParse::Value(CmdLine, TEXT("-processaffinity="), LogicalCoreAffinity);
+	FParse::Value(CmdLine, TEXT("-processaffinityphysical="), PhysicalCoreAffinity);
+
+	if (LogicalCoreAffinity > 0)
+	{
+		FWindowsPlatformProcess::SetProcessAffinity(LogicalCoreAffinity, false);
+	}
+	else if (PhysicalCoreAffinity > 0)
+	{
+		FWindowsPlatformProcess::SetProcessAffinity(PhysicalCoreAffinity, true);
+	}
+
 	bool bNoExceptionHandler = FParse::Param(CmdLine,TEXT("noexceptionhandler"));
 	(void)bNoExceptionHandler;
 
