[IoStore]
 - Properly chain decompression prerequisites so everything can be retracted

diff --git a/Engine/Source/Runtime/Core/Private/IO/IoStore.cpp b/Engine/Source/Runtime/Core/Private/IO/IoStore.cpp
--- a/Engine/Source/Runtime/Core/Private/IO/IoStore.cpp
+++ b/Engine/Source/Runtime/Core/Private/IO/IoStore.cpp
@@ -566,15 +566,13 @@ public:
 		{			
 			UE::Tasks::FTaskEvent DecompressionDoneEvent(TEXT("FIoStoreReader::DecompressionDone"));
 
-			std::atomic_int32_t DecompressionJobsRemaining = LastBlockIndex - FirstBlockIndex + 1;
-
 			uint64 CompressedSourceOffset = 0;
 			uint64 UncompressedDestinationOffset = 0;
 			uint64 OffsetInBlock = ResolvedOffset % CompressionBlockSize;
 			uint64 RemainingSize = ResolvedSize;
 			for (int32 BlockIndex = FirstBlockIndex; BlockIndex <= LastBlockIndex; ++BlockIndex)
 			{
-				UE::Tasks::FTask DecompressBlockTask = UE::Tasks::Launch(TEXT("FIoStoreReader::Decompress"), [this, State, BlockIndex, CompressedSourceOffset, UncompressedDestinationOffset, OffsetInBlock, RemainingSize, &DecompressionDoneEvent, &DecompressionJobsRemaining]()
+				UE::Tasks::FTask DecompressBlockTask = UE::Tasks::Launch(TEXT("FIoStoreReader::Decompress"), [this, State, BlockIndex, CompressedSourceOffset, UncompressedDestinationOffset, OffsetInBlock, RemainingSize]()
 				{
 					if (State->bReadSucceeded)
 					{
@@ -619,12 +617,10 @@ public:
 						}
 					} // end if read succeeded
 
-					if (DecompressionJobsRemaining.fetch_add(-1) == 1)
-					{
-						DecompressionDoneEvent.Trigger();
-					}
 				}); // end decompression lambda
 
+				DecompressionDoneEvent.AddPrerequisites(DecompressBlockTask);
+
 				const FIoStoreTocCompressedBlockEntry& CompressionBlock = TocResource.CompressionBlocks[BlockIndex];
 				const uint32 RawSize = Align(CompressionBlock.GetCompressedSize(), FAES::AESBlockSize);
 				CompressedSourceOffset += RawSize;
@@ -633,8 +629,10 @@ public:
 				OffsetInBlock = 0;
 			} // end for each block
 
-			// Wait for everything
-			DecompressionDoneEvent.BusyWait();
+			// Unlock the event so we're now only waiting on the prerequisites
+			DecompressionDoneEvent.Trigger();
+			// Wait for everything and potentially help with the decompression tasks by retraction.
+			DecompressionDoneEvent.Wait();
 
 			TIoStatusOr<FIoBuffer> Result;
 			if (State->bReadSucceeded == false)
