[TextLocalizationManager]
 - Move the event broadcast to a safe place on the game-thread

#rnx

diff --git a/Engine/Source/Runtime/Core/Private/Internationalization/TextLocalizationManager.cpp b/Engine/Source/Runtime/Core/Private/Internationalization/TextLocalizationManager.cpp
--- a/Engine/Source/Runtime/Core/Private/Internationalization/TextLocalizationManager.cpp
+++ b/Engine/Source/Runtime/Core/Private/Internationalization/TextLocalizationManager.cpp
@@ -15,6 +15,7 @@
 #include "Misc/FileHelper.h"
 #include "Misc/Parse.h"
 #include "Misc/ScopeLock.h"
+#include "Containers/Ticker.h"
 #include "Misc/ScopeRWLock.h"
 #include "Misc/CommandLine.h"
 #include "Misc/LazySingleton.h"
@@ -1894,7 +1895,7 @@ void FTextLocalizationManager::DirtyTextRevision()
 	}
 	else
 	{
-		AsyncTask(ENamedThreads::GameThread, []()
+		ExecuteOnGameThread(TEXT("OnTextRevisionChangedEventBroadcastGT"), []()
 		{
 			FTextLocalizationManager::Get().OnTextRevisionChangedEvent.Broadcast();
 		});
