Fix non-unity issues

diff --git a/Engine/Source/Runtime/Core/Public/Microsoft/COMPointer.h b/Engine/Source/Runtime/Core/Public/Microsoft/COMPointer.h
--- a/Engine/Source/Runtime/Core/Public/Microsoft/COMPointer.h
+++ b/Engine/Source/Runtime/Core/Public/Microsoft/COMPointer.h
@@ -2,7 +2,10 @@
 
 #pragma once
 
+#include "Microsoft/AllowMicrosoftPlatformTypes.h"
 #include <Unknwn.h>
+#include "Microsoft/HideMicrosoftPlatformTypes.h"
+
 #include "Misc/AssertionMacros.h"
 
 #ifdef TEXT_PASTE
