Fix nesting of allow windows types and resolve non-unity different way
#rnx

diff --git a/Engine/Source/Runtime/Core/Public/Microsoft/COMPointer.h b/Engine/Source/Runtime/Core/Public/Microsoft/COMPointer.h
--- a/Engine/Source/Runtime/Core/Public/Microsoft/COMPointer.h
+++ b/Engine/Source/Runtime/Core/Public/Microsoft/COMPointer.h
@@ -2,11 +2,8 @@
 
 #pragma once
 
-#include "Microsoft/AllowMicrosoftPlatformTypes.h"
-#include <Unknwn.h>
-#include "Microsoft/HideMicrosoftPlatformTypes.h"
-
 #include "Misc/AssertionMacros.h"
+#include <Unknwn.h>
 
 #ifdef TEXT_PASTE
 	#undef TEXT
