Disabling the Windows Compiler warning C6393.

The Warning:
warning C6393: A lookup table of size 365 is not sufficient to handle leap years.

This is a warning that goes off if any static array is 365 members long.  The intent of the warning is to stop common misuse of static arrays to look up days of the year.

In practice someone hitting this could just disable the warning locally to approve the array and remove the warning.  In our case though the code generation step can easily generate arrays 365 members long and there is no (good) way to fix it.

So turning it off globally makes the most sense for us.

diff --git a/Engine/Source/Runtime/Core/Public/Windows/WindowsPlatformCodeAnalysis.h b/Engine/Source/Runtime/Core/Public/Windows/WindowsPlatformCodeAnalysis.h
--- a/Engine/Source/Runtime/Core/Public/Windows/WindowsPlatformCodeAnalysis.h
+++ b/Engine/Source/Runtime/Core/Public/Windows/WindowsPlatformCodeAnalysis.h
@@ -102,6 +102,10 @@
 	// Likewise, expressions involving constants can also be valuable to check code paths compile.
 	#pragma warning(disable : 6240) // warning C6240 : (<expression> && <non-zero constant>) always evaluates to the result of <expression>. Did you intend to use the bitwise-and operator?
 
+	// This is a warning that a static array is exactly 365 members long.  We get false positives because our code gen can make arrays exactly this size.  
+	// EX: 365 properties in a UObject will auto gen a table exactly this long.
+	#pragma warning(disable : 6393) // warning C6393: A lookup table of size 365 is not sufficient to handle leap years.
+
 	//PVS-Studio settings:
 	//-V::505,542,581,601,623,668,677,690,704,719,720,730,735,751,758,780,781,1002,1008,1055,1062,1100
 	//-V:TRYCOMPRESSION:519,547
