Replaced some macro names with references to UE4.

diff --git a/Engine/Plugins/FX/Niagara/Source/Niagara/Private/NiagaraModule.cpp b/Engine/Plugins/FX/Niagara/Source/Niagara/Private/NiagaraModule.cpp
--- a/Engine/Plugins/FX/Niagara/Source/Niagara/Private/NiagaraModule.cpp
+++ b/Engine/Plugins/FX/Niagara/Source/Niagara/Private/NiagaraModule.cpp
@@ -174,7 +174,7 @@ void INiagaraModule::OnDataChannelsEnabledChanged(IConsoleVariable* Variable)
 // these two globals are intended to help ensure that a global variable is accessible to natvis (as defined in Niagara.natvis)
 // while debugging.  GCoreTypeRegistrySingletonPtr is the pointer to the actual data stored within the TLazySingleton<FNiagaraTypeRegistry>
 // while GTypeRegistrySingletonPtr can be declared in each module to ensure that it can be accessed while debugging any Niagara
-// module.  See UE4_VISUALIZERS_HELPERS.  Note that currently it seems like we can effectively debug NiagaraEditor/NiagaraShader without
+// module.  See UE_VISUALIZERS_HELPERS.  Note that currently it seems like we can effectively debug NiagaraEditor/NiagaraShader without
 // further declarations, which is nice, so will be leaving it as it is.
 namespace NiagaraDebugVisHelper
 {
diff --git a/Engine/Plugins/Runtime/USDCore/Source/UnrealUSDWrapper/Public/USDMemory.h b/Engine/Plugins/Runtime/USDCore/Source/UnrealUSDWrapper/Public/USDMemory.h
--- a/Engine/Plugins/Runtime/USDCore/Source/UnrealUSDWrapper/Public/USDMemory.h
+++ b/Engine/Plugins/Runtime/USDCore/Source/UnrealUSDWrapper/Public/USDMemory.h
@@ -294,7 +294,7 @@ TSharedRef<ObjectType> MakeSharedUnreal(ArgTypes&&... Args)
 		} \
 		/* Forced reference to this function is added by the linker to check that each module uses IMPLEMENT_MODULE */ \
 		extern "C" void IMPLEMENT_MODULE_##ModuleName() { } \
-		UE4_VISUALIZERS_HELPERS \
+		UE_VISUALIZERS_HELPERS \
 		REPLACEMENT_OPERATOR_NEW_AND_DELETE_USD
 #else
 	#define IMPLEMENT_MODULE_USD( ModuleImplClass, ModuleName ) IMPLEMENT_MODULE( ModuleImplClass, ModuleName )
diff --git a/Engine/Source/Runtime/Core/Public/HAL/PerModuleInline.inl b/Engine/Source/Runtime/Core/Public/HAL/PerModuleInline.inl
--- a/Engine/Source/Runtime/Core/Public/HAL/PerModuleInline.inl
+++ b/Engine/Source/Runtime/Core/Public/HAL/PerModuleInline.inl
@@ -5,7 +5,7 @@
 #include "CoreGlobals.h"
 #include "Modules/Boilerplate/ModuleBoilerplate.h"
 
-UE4_VISUALIZERS_HELPERS
+UE_VISUALIZERS_HELPERS
 REPLACEMENT_OPERATOR_NEW_AND_DELETE
 UE_DEFINE_FMEMORY_WRAPPERS
 
diff --git a/Engine/Source/Runtime/Core/Public/Modules/Boilerplate/ModuleBoilerplate.h b/Engine/Source/Runtime/Core/Public/Modules/Boilerplate/ModuleBoilerplate.h
--- a/Engine/Source/Runtime/Core/Public/Modules/Boilerplate/ModuleBoilerplate.h
+++ b/Engine/Source/Runtime/Core/Public/Modules/Boilerplate/ModuleBoilerplate.h
@@ -76,12 +76,20 @@ class FChunkedFixedUObjectArray;
 class FChunkedFixedUObjectArray;
 
 #ifdef DISABLE_UE4_VISUALIZER_HELPERS
-	#define UE4_VISUALIZERS_HELPERS
+	UE_DEPRECATED_HEADER(5.5, "DISABLE_UE4_VISUALIZER_HELPERS has been disabled, please define UE_ENABLE_VISUALIZER_HELPERS=0 instead.")
+	#define UE_ENABLE_VISUALIZER_HELPERS 0
+#endif
+#ifndef UE_ENABLE_VISUALIZER_HELPERS
+	#define UE_ENABLE_VISUALIZER_HELPERS 1
+#endif
+
+#if !UE_ENABLE_VISUALIZER_HELPERS
+	#define UE_VISUALIZERS_HELPERS
 #elif PLATFORM_UNIX
 	// GDB/LLDB pretty printers don't use these - no need to export additional symbols. This also solves ODR violation reported by ASan on Linux
-	#define UE4_VISUALIZERS_HELPERS
+	#define UE_VISUALIZERS_HELPERS
 #else
-	#define UE4_VISUALIZERS_HELPERS \
+	#define UE_VISUALIZERS_HELPERS \
 		uint8** GNameBlocksDebug = FNameDebugVisualizer::GetBlocks(); \
 		const UPTRINT*& GObjectIndexToPackedObjectRefDebug = GCoreObjectIndexToPackedObjectRefDebug; \
 		FChunkedFixedUObjectArray*& GObjectArrayForDebugVisualizers = GCoreObjectArrayForDebugVisualizers; \
@@ -92,6 +100,6 @@ class FChunkedFixedUObjectArray;
 // in DLL builds, these are done per-module, otherwise we just need one in the application
 // visual studio cannot find cross dll data for visualizers, so these provide access
 #define PER_MODULE_BOILERPLATE \
-	UE4_VISUALIZERS_HELPERS \
+	UE_VISUALIZERS_HELPERS \
 	REPLACEMENT_OPERATOR_NEW_AND_DELETE \
 	UE_DEFINE_FMEMORY_WRAPPERS
