Added UE_SELECT_ANY that can be used to have multiple definitions of a global variable (is used sometimes to simplify debugging and natvis)

diff --git a/Engine/Plugins/FX/Niagara/Source/Niagara/Private/NiagaraModule.cpp b/Engine/Plugins/FX/Niagara/Source/Niagara/Private/NiagaraModule.cpp
--- a/Engine/Plugins/FX/Niagara/Source/Niagara/Private/NiagaraModule.cpp
+++ b/Engine/Plugins/FX/Niagara/Source/Niagara/Private/NiagaraModule.cpp
@@ -179,7 +179,7 @@ void INiagaraModule::OnDataChannelsEnabledChanged(IConsoleVariable* Variable)
 namespace NiagaraDebugVisHelper
 {
 	const FNiagaraTypeRegistry* GCoreTypeRegistrySingletonPtr = nullptr;
-	const FNiagaraTypeRegistry*& GTypeRegistrySingletonPtr = GCoreTypeRegistrySingletonPtr;
+	UE_SELECT_ANY const FNiagaraTypeRegistry*& GTypeRegistrySingletonPtr = GCoreTypeRegistrySingletonPtr;
 }
 
 FNiagaraVariable INiagaraModule::Engine_WorldDeltaTime;
diff --git a/Engine/Plugins/FX/Niagara/Source/NiagaraEditor/Private/NiagaraEditorModule.cpp b/Engine/Plugins/FX/Niagara/Source/NiagaraEditor/Private/NiagaraEditorModule.cpp
--- a/Engine/Plugins/FX/Niagara/Source/NiagaraEditor/Private/NiagaraEditorModule.cpp
+++ b/Engine/Plugins/FX/Niagara/Source/NiagaraEditor/Private/NiagaraEditorModule.cpp
@@ -224,7 +224,7 @@ static FNiagaraShaderQueueTickable NiagaraShaderQueueProcessor;
 #if !IS_MONOLITHIC
 namespace NiagaraDebugVisHelper
 {
-	const FNiagaraTypeRegistry*& GTypeRegistrySingletonPtr = GCoreTypeRegistrySingletonPtr;
+	UE_SELECT_ANY const FNiagaraTypeRegistry*& GTypeRegistrySingletonPtr = GCoreTypeRegistrySingletonPtr;
 }
 #endif
 
diff --git a/Engine/Source/Runtime/Core/Public/HAL/Platform.h b/Engine/Source/Runtime/Core/Public/HAL/Platform.h
--- a/Engine/Source/Runtime/Core/Public/HAL/Platform.h
+++ b/Engine/Source/Runtime/Core/Public/HAL/Platform.h
@@ -714,6 +714,14 @@
 	#error "Compiler is expected to support [[noreturn]]"
 #endif
 
+/* Use in front of initialized global variable to have link once-semantics
+   When multiple definitions of the variable are encountered by the linker, the first is selected and the remainder are discarded */
+#ifdef _MSC_VER
+    #define UE_SELECT_ANY __declspec(selectany)
+#else
+    #define UE_SELECT_ANY __attribute__((selectany))
+#endif
+
 /* Macro wrapper for the consteval keyword which isn't yet present on all compilers - constexpr
    can be used as a workaround but is less strict and so may let some non-consteval code pass */
 #if defined(__cpp_consteval)
diff --git a/Engine/Source/Runtime/MovieScene/Public/EntitySystem/MovieSceneEntityManager.h b/Engine/Source/Runtime/MovieScene/Public/EntitySystem/MovieSceneEntityManager.h
--- a/Engine/Source/Runtime/MovieScene/Public/EntitySystem/MovieSceneEntityManager.h
+++ b/Engine/Source/Runtime/MovieScene/Public/EntitySystem/MovieSceneEntityManager.h
@@ -1089,12 +1089,6 @@ private:
 	FEntityManager* EntityManager;
 };
 
-#if PLATFORM_COMPILER_CLANG
-#define UE_SELECT_ANY __attribute__((weak))
-#else
-#define UE_SELECT_ANY __declspec(selectany)
-#endif
-
 extern MOVIESCENE_API FEntityManager* GEntityManagerForDebuggingVisualizers;
 
 
