Move the AutoRTFM pass before ASan to speed up ASan builds by 300x.

#rnx

diff --git a/Engine/Source/Runtime/Core/Private/AutoRTFM/Transaction.h b/Engine/Source/Runtime/Core/Private/AutoRTFM/Transaction.h
--- a/Engine/Source/Runtime/Core/Private/AutoRTFM/Transaction.h
+++ b/Engine/Source/Runtime/Core/Private/AutoRTFM/Transaction.h
@@ -70,9 +70,6 @@ public:
     void AbortWithoutThrowing();
     bool AttemptToCommit();
 
-    // Should we skip recording for the given address (used by ASan).
-    bool SkipRecord(void* LogicalAddress);
-
     // Record that a write is about to occur at the given LogicalAddress of Size bytes.
     void RecordWrite(void* LogicalAddress, size_t Size);
     void RecordWriteMaxPageSized(void* LogicalAddress, size_t Size);
diff --git a/Engine/Source/Runtime/Core/Private/AutoRTFM/TransactionInlines.h b/Engine/Source/Runtime/Core/Private/AutoRTFM/TransactionInlines.h
--- a/Engine/Source/Runtime/Core/Private/AutoRTFM/TransactionInlines.h
+++ b/Engine/Source/Runtime/Core/Private/AutoRTFM/TransactionInlines.h
@@ -7,29 +7,9 @@
 
 #include "HAL/Platform.h"
 
-#if PLATFORM_HAS_ASAN_INCLUDE && USING_ADDRESS_SANITISER
-#include <sanitizer/asan_interface.h>
-#endif
-
 namespace AutoRTFM
 {
 
-AUTORTFM_NO_ASAN UE_AUTORTFM_FORCEINLINE bool FTransaction::SkipRecord(void* LogicalAddress)
-{
-#if PLATFORM_HAS_ASAN_INCLUDE && USING_ADDRESS_SANITISER
-    // TODO(SOL-5123): Can we detect shadow memory locations at compile-time instead?
-    const char* const Location = __asan_locate_address(LogicalAddress, nullptr, 0, nullptr, nullptr);
-
-    if (strstr(Location, "shadow"))
-    {
-        return true;
-    }
-#endif
-
-    return false;
-}
-
-
 AUTORTFM_NO_ASAN UE_AUTORTFM_FORCEINLINE void FTransaction::RecordWriteMaxPageSized(void* LogicalAddress, size_t Size)
 {
     void* CopyAddress = WriteLogBumpAllocator.Allocate(Size);
@@ -45,11 +25,6 @@ AUTORTFM_NO_ASAN UE_AUTORTFM_FORCEINLINE void FTransaction::RecordWrite(void* Lo
         return;
     }
 
-    if (UNLIKELY(SkipRecord(LogicalAddress)))
-    {
-        return;
-    }
-
     // If we are recording a stack address that is relative to our current
     // transactions stack location, we do not need to record the data in the
     // write log because if that transaction aborted, that memory will cease to
@@ -104,11 +79,6 @@ template<unsigned SIZE> AUTORTFM_NO_ASAN UE_AUTORTFM_FORCEINLINE void FTransacti
 {
     static_assert(SIZE <= 8);
 
-    if (UNLIKELY(SkipRecord(LogicalAddress)))
-    {
-        return;
-    }
-
     // If we are recording a stack address that is relative to our current
     // transactions stack location, we do not need to record the data in the
     // write log because if that transaction aborted, that memory will cease to
