Changed so "InitializeModule()" is a fallback solution for initializing modules in loaded dll. Now we instead use an intrusive linked list to find the module. This way we can have multiple modules in the same dll and also allow the search to be case insensitive

diff --git a/Engine/Plugins/Runtime/USDCore/Source/UnrealUSDWrapper/Public/USDMemory.h b/Engine/Plugins/Runtime/USDCore/Source/UnrealUSDWrapper/Public/USDMemory.h
--- a/Engine/Plugins/Runtime/USDCore/Source/UnrealUSDWrapper/Public/USDMemory.h
+++ b/Engine/Plugins/Runtime/USDCore/Source/UnrealUSDWrapper/Public/USDMemory.h
@@ -288,12 +288,14 @@ TSharedRef<ObjectType> MakeSharedUnreal(ArgTypes&&... Args)
 		/**/ \
 		/* @return	Returns an instance of this module */ \
 		/**/ \
-		extern "C" DLLEXPORT IModuleInterface* InitializeModule() \
+		static IModuleInterface* Initialize##ModuleName##Module() \
 		{ \
 			return new ModuleImplClass(); \
 		} \
+		static FModuleInitializerEntry ModuleName##InitializerEntry(TEXT(#ModuleName), Initialize##ModuleName##Module); \
 		/* Forced reference to this function is added by the linker to check that each module uses IMPLEMENT_MODULE */ \
 		extern "C" void IMPLEMENT_MODULE_##ModuleName() { } \
+		extern "C" DLLEXPORT void ThisIsAnUnrealEngineModule() {} \
 		UE_VISUALIZERS_HELPERS \
 		REPLACEMENT_OPERATOR_NEW_AND_DELETE_USD
 #else
diff --git a/Engine/Source/Runtime/Core/Core.Build.cs b/Engine/Source/Runtime/Core/Core.Build.cs
--- a/Engine/Source/Runtime/Core/Core.Build.cs
+++ b/Engine/Source/Runtime/Core/Core.Build.cs
@@ -350,6 +350,8 @@ public class Core : ModuleRules
 		bAllowAutoRTFMInstrumentation = true;
 
 		PrivateDefinitions.Add("UE_PERSISTENT_ALLOCATOR_RESERVE_SIZE=" + GetPersistentAllocatorReserveSize().ToString() + "ULL");
+
+		PrivateDefinitions.Add("UE_MERGED_MODULES=" + (Target.bMergeModules ? "1" : "0"));
 	}
 
 	protected virtual bool SupportsBinaryConfig(ReadOnlyTargetRules Target)
diff --git a/Engine/Source/Runtime/Core/Private/Modules/ModuleManager.cpp b/Engine/Source/Runtime/Core/Private/Modules/ModuleManager.cpp
--- a/Engine/Source/Runtime/Core/Private/Modules/ModuleManager.cpp
+++ b/Engine/Source/Runtime/Core/Private/Modules/ModuleManager.cpp
@@ -37,6 +37,52 @@ DEFINE_LOG_CATEGORY_STATIC(LogModuleManager, Log, All);
 	}
 #endif
 
+FModuleInitializerEntry* GFirstModuleInitializerEntry;
+
+FModuleInitializerEntry::FModuleInitializerEntry(const TCHAR* InName, FInitializeModuleFunctionPtr InFunction)
+:	Name(InName)
+,	Function(InFunction)
+{
+	Prev = nullptr;
+	Next = GFirstModuleInitializerEntry;
+
+	if (GFirstModuleInitializerEntry)
+	{
+		GFirstModuleInitializerEntry->Prev = this;
+	}
+
+	GFirstModuleInitializerEntry = this;
+}
+
+FModuleInitializerEntry::~FModuleInitializerEntry()
+{
+	if (Next)
+	{
+		Next->Prev = Prev;
+	}
+
+	if (Prev)
+	{
+		Prev->Next = Next;
+	}
+	else
+	{
+		GFirstModuleInitializerEntry = Next;
+	}
+}
+
+FInitializeModuleFunctionPtr FModuleInitializerEntry::FindModule(const TCHAR* Name)
+{
+	for (FModuleInitializerEntry* Entry = GFirstModuleInitializerEntry; Entry; Entry = Entry->Next)
+	{
+		if (FCString::Stricmp(Name, Entry->Name) == 0)
+		{
+			return Entry->Function;
+		}
+	}
+	return nullptr;
+}
+
 
 int32 FModuleManager::FModuleInfo::CurrentLoadOrder = 1;
 
@@ -359,6 +405,15 @@ void FModuleManager::RefreshModuleFilenameFromManifestImpl(const FName InModuleN
 	FString ModuleFilename = MoveTemp(TMap<FName, FString>::TIterator(ModulePathMap).Value());
 
 	const int32 MatchPos = ModuleFilename.Find(ModuleNameString, ESearchCase::IgnoreCase, ESearchDir::FromEnd);
+
+	// If modules are merged it is likely that binary name will not match module. TODO: Revisit this to see if we can get this to work with hot reloading etc
+	#if UE_MERGED_MODULES
+	if (MatchPos == INDEX_NONE)
+	{
+		return;
+	}
+	#endif
+
 	if (!ensureMsgf(MatchPos != INDEX_NONE, TEXT("Could not find module name '%s' in module filename '%s'"), *InModuleName.ToString(), *ModuleFilename))
 	{
 		return;
@@ -650,9 +705,16 @@ IModuleInterface* FModuleManager::LoadModuleWithFailureReason(const FName InModu
 					ProcessLoadedObjectsCallback.Broadcast(InModuleName, bCanProcessNewlyLoadedObjects);
 				}
 
-				// Find our "InitializeModule" global function, which must exist for all module DLLs
-				FInitializeModuleFunctionPtr InitializeModuleFunctionPtr =
-					(FInitializeModuleFunctionPtr)FPlatformProcess::GetDllExport(ModuleInfo->Handle, TEXT("InitializeModule"));
+
+				// Find our "Initialize<Name>Module" global function, which must exist for all module DLLs
+				FInitializeModuleFunctionPtr InitializeModuleFunctionPtr = FModuleInitializerEntry::FindModule(*InModuleName.ToString());
+
+				if (!InitializeModuleFunctionPtr)
+				{
+					// If not found this might be some special case module so look for "InitializeModule" global function
+					InitializeModuleFunctionPtr = (FInitializeModuleFunctionPtr)FPlatformProcess::GetDllExport(ModuleInfo->Handle, TEXT("InitializeModule"));
+				}
+
 				if (InitializeModuleFunctionPtr != nullptr)
 				{
 					if ( ModuleInfo->Module.IsValid() )
diff --git a/Engine/Source/Runtime/Core/Public/Modules/ModuleManager.h b/Engine/Source/Runtime/Core/Public/Modules/ModuleManager.h
--- a/Engine/Source/Runtime/Core/Public/Modules/ModuleManager.h
+++ b/Engine/Source/Runtime/Core/Public/Modules/ModuleManager.h
@@ -744,18 +744,37 @@ public:
 	}
 };
 
-
 /**
  * Function pointer type for InitializeModule().
  *
- * All modules must have an InitializeModule() function. Usually this is declared automatically using
- * the IMPLEMENT_MODULE macro below. The function must be declared using as 'extern "C"' so that the
+ * All modules must have a FModuleInitializerEntry instance or an InitializeModule() function. Usually this is declared automatically using
+ * the IMPLEMENT_MODULE macro below. If using the function it must be declared using as 'extern "C"' so that the
  * name remains undecorated. The object returned will be "owned" by the caller, and will be deleted
  * by the caller before the module is unloaded.
  */
 typedef IModuleInterface* ( *FInitializeModuleFunctionPtr )( void );
 
 
+/**
+ * Intrusive linked list containing name and initializer function pointer of loaded modules.
+ * Use this instead of "InitializeModule()" when possible
+ */
+class FModuleInitializerEntry
+{
+public:
+	CORE_API FModuleInitializerEntry(const TCHAR* InName, FInitializeModuleFunctionPtr InFunction);
+	CORE_API ~FModuleInitializerEntry();
+
+	static FInitializeModuleFunctionPtr FindModule(const TCHAR* Name);
+
+private:
+	FModuleInitializerEntry* Prev;
+	FModuleInitializerEntry* Next;
+	const TCHAR* Name;
+	FInitializeModuleFunctionPtr Function;
+};
+
+
 /**
  * A default minimal implementation of a module that does nothing at startup and shutdown
  */
@@ -815,10 +834,11 @@ class FDefaultGameModuleImpl
 		/**/ \
 		/* @return	Returns an instance of this module */ \
 		/**/ \
-		extern "C" DLLEXPORT IModuleInterface* InitializeModule() \
+		static IModuleInterface* Initialize##ModuleName##Module() \
 		{ \
 			return new ModuleImplClass(); \
 		} \
+		static FModuleInitializerEntry ModuleName##InitializerEntry(TEXT(#ModuleName), Initialize##ModuleName##Module); \
 		/* Forced reference to this function is added by the linker to check that each module uses IMPLEMENT_MODULE */ \
 		extern "C" void IMPLEMENT_MODULE_##ModuleName() { } \
 		PER_MODULE_BOILERPLATE_ANYLINK(ModuleImplClass, ModuleName)
