[Core]
* Added so code identifying module as UE module on unix tests the old path too.. in case there are modules out there that are not using our macros (highly unlikely because they need our operator new/delete etc)

diff --git a/Engine/Source/Runtime/Core/Private/Unix/UnixPlatformProcess.cpp b/Engine/Source/Runtime/Core/Private/Unix/UnixPlatformProcess.cpp
--- a/Engine/Source/Runtime/Core/Private/Unix/UnixPlatformProcess.cpp
+++ b/Engine/Source/Runtime/Core/Private/Unix/UnixPlatformProcess.cpp
@@ -86,6 +86,11 @@ void* FUnixPlatformProcess::GetDllHandle( const TCHAR* Filename )
 		{
 			// is it ia ue4 module ? if not, move it to GLOBAL
 			void *IsUE4Module = dlsym(Handle, "ThisIsAnUnrealEngineModule");
+			if (!IsUE4Module)
+			{
+				IsUE4Module = dlsym(Handle, "InitializeModule");
+			}
+
 			if (!IsUE4Module)
 			{
 				UpgradeToGlobal = true;
