[Core]
* Added so FModuleInitializerEntry tracks two names per module. One is the one provided in IMPLEMENT_MODULE and one is UE_MODULE_NAME. This to get around problem where names are mismatching

diff --git a/Engine/Plugins/Runtime/USDCore/Source/UnrealUSDWrapper/Public/USDMemory.h b/Engine/Plugins/Runtime/USDCore/Source/UnrealUSDWrapper/Public/USDMemory.h
--- a/Engine/Plugins/Runtime/USDCore/Source/UnrealUSDWrapper/Public/USDMemory.h
+++ b/Engine/Plugins/Runtime/USDCore/Source/UnrealUSDWrapper/Public/USDMemory.h
@@ -292,7 +292,7 @@ TSharedRef<ObjectType> MakeSharedUnreal(ArgTypes&&... Args)
 		{ \
 			return new ModuleImplClass(); \
 		} \
-		static FModuleInitializerEntry ModuleName##InitializerEntry(TEXT(#ModuleName), Initialize##ModuleName##Module); \
+		static FModuleInitializerEntry ModuleName##InitializerEntry(TEXT(#ModuleName), Initialize##ModuleName##Module, TEXT(UE_MODULE_NAME)); \
 		/* Forced reference to this function is added by the linker to check that each module uses IMPLEMENT_MODULE */ \
 		extern "C" void IMPLEMENT_MODULE_##ModuleName() { } \
 		extern "C" DLLEXPORT void ThisIsAnUnrealEngineModule() {} \
diff --git a/Engine/Source/Runtime/Core/Private/Modules/ModuleManager.cpp b/Engine/Source/Runtime/Core/Private/Modules/ModuleManager.cpp
--- a/Engine/Source/Runtime/Core/Private/Modules/ModuleManager.cpp
+++ b/Engine/Source/Runtime/Core/Private/Modules/ModuleManager.cpp
@@ -39,10 +39,16 @@ DEFINE_LOG_CATEGORY_STATIC(LogModuleManager, Log, All);
 
 FModuleInitializerEntry* GFirstModuleInitializerEntry;
 
-FModuleInitializerEntry::FModuleInitializerEntry(const TCHAR* InName, FInitializeModuleFunctionPtr InFunction)
+FModuleInitializerEntry::FModuleInitializerEntry(const TCHAR* InName, FInitializeModuleFunctionPtr InFunction, const TCHAR* InName2)
 :	Name(InName)
+,	Name2(nullptr)
 ,	Function(InFunction)
 {
+	if (FCString::Stricmp(InName, InName2) != 0)
+	{
+		Name2 = InName2;
+	}
+
 	Prev = nullptr;
 	Next = GFirstModuleInitializerEntry;
 
@@ -79,6 +85,10 @@ FInitializeModuleFunctionPtr FModuleInitializerEntry::FindModule(const TCHAR* Na
 		{
 			return Entry->Function;
 		}
+		if (Entry->Name2 && FCString::Stricmp(Name, Entry->Name2) == 0)
+		{
+			return Entry->Function;
+		}
 	}
 	return nullptr;
 }
diff --git a/Engine/Source/Runtime/Core/Public/Modules/ModuleManager.h b/Engine/Source/Runtime/Core/Public/Modules/ModuleManager.h
--- a/Engine/Source/Runtime/Core/Public/Modules/ModuleManager.h
+++ b/Engine/Source/Runtime/Core/Public/Modules/ModuleManager.h
@@ -762,7 +762,7 @@ typedef IModuleInterface* ( *FInitializeModuleFunctionPtr )( void );
 class FModuleInitializerEntry
 {
 public:
-	CORE_API FModuleInitializerEntry(const TCHAR* InName, FInitializeModuleFunctionPtr InFunction);
+	CORE_API FModuleInitializerEntry(const TCHAR* InName, FInitializeModuleFunctionPtr InFunction, const TCHAR* InName2 = nullptr);
 	CORE_API ~FModuleInitializerEntry();
 
 	static FInitializeModuleFunctionPtr FindModule(const TCHAR* Name);
@@ -771,6 +771,7 @@ private:
 	FModuleInitializerEntry* Prev;
 	FModuleInitializerEntry* Next;
 	const TCHAR* Name;
+	const TCHAR* Name2;
 	FInitializeModuleFunctionPtr Function;
 };
 
@@ -838,7 +839,7 @@ class FDefaultGameModuleImpl
 		{ \
 			return new ModuleImplClass(); \
 		} \
-		static FModuleInitializerEntry ModuleName##InitializerEntry(TEXT(#ModuleName), Initialize##ModuleName##Module); \
+		static FModuleInitializerEntry ModuleName##InitializerEntry(TEXT(#ModuleName), Initialize##ModuleName##Module, TEXT(UE_MODULE_NAME)); \
 		/* Forced reference to this function is added by the linker to check that each module uses IMPLEMENT_MODULE */ \
 		extern "C" void IMPLEMENT_MODULE_##ModuleName() { } \
 		PER_MODULE_BOILERPLATE_ANYLINK(ModuleImplClass, ModuleName)
