Removed deprecated FPlatformAtomics::AtomicRead64 functions.

diff --git a/Engine/Source/Runtime/Core/Public/Clang/ClangPlatformAtomics.h b/Engine/Source/Runtime/Core/Public/Clang/ClangPlatformAtomics.h
--- a/Engine/Source/Runtime/Core/Public/Clang/ClangPlatformAtomics.h
+++ b/Engine/Source/Runtime/Core/Public/Clang/ClangPlatformAtomics.h
@@ -265,12 +265,6 @@ struct FClangPlatformAtomics : public FGenericPlatformAtomics
 		__atomic_store_n((volatile int64*)Src, Val, __ATOMIC_RELAXED);
 	}
 
-	UE_DEPRECATED(4.19, "AtomicRead64 has been deprecated, please use AtomicRead's overload instead")
-	static FORCEINLINE int64 AtomicRead64(volatile const int64* Src)
-	{
-		return InterlockedCompareExchange((volatile int64*)Src, 0, 0);
-	}
-
 	static FORCEINLINE void* InterlockedCompareExchangePointer(void*volatile* Dest, void* Exchange, void* Comparand)
 	{
 		__atomic_compare_exchange_n(Dest, &Comparand, Exchange, false, __ATOMIC_SEQ_CST, __ATOMIC_SEQ_CST);
diff --git a/Engine/Source/Runtime/Core/Public/GenericPlatform/GenericPlatformAtomics.h b/Engine/Source/Runtime/Core/Public/GenericPlatform/GenericPlatformAtomics.h
--- a/Engine/Source/Runtime/Core/Public/GenericPlatform/GenericPlatformAtomics.h
+++ b/Engine/Source/Runtime/Core/Public/GenericPlatform/GenericPlatformAtomics.h
@@ -263,15 +263,6 @@ struct FGenericPlatformAtomics
 		return InterlockedCompareExchange((volatile int64*)Src, 0, 0);
 	}
 
-	/**
-	 * Atomic read of 64 bit value with an implicit memory barrier.
-	 */
-	UE_DEPRECATED(4.19, "AtomicRead64 has been deprecated, please use AtomicRead's overload instead")
-	static FORCEINLINE int64 AtomicRead64(volatile const int64* Src)
-	{
-		return AtomicRead(Src);
-	}
-
 	/**
 	 * Atomically compares the pointer to comparand and replaces with the exchange
 	 * pointer if they are equal and returns the original value
diff --git a/Engine/Source/Runtime/Core/Public/Windows/WindowsPlatformAtomics.h b/Engine/Source/Runtime/Core/Public/Windows/WindowsPlatformAtomics.h
--- a/Engine/Source/Runtime/Core/Public/Windows/WindowsPlatformAtomics.h
+++ b/Engine/Source/Runtime/Core/Public/Windows/WindowsPlatformAtomics.h
@@ -369,12 +369,6 @@ struct FWindowsPlatformAtomics
 		#endif
 	}
 
-	UE_DEPRECATED(4.19, "AtomicRead64 has been deprecated, please use AtomicRead's overload instead")
-	static FORCEINLINE int64 AtomicRead64(volatile const int64* Src)
-	{
-		return AtomicRead(Src);
-	}
-
 	/**
 	 *	The function compares the Destination value with the Comparand value:
 	 *		- If the Destination value is equal to the Comparand value, the Exchange value is stored in the address specified by Destination, 
