Probable fix for ClangPlatformAtomics.h being compiled under MSVC.

diff --git a/Engine/Source/Runtime/Core/Public/Clang/ClangPlatformAtomics.h b/Engine/Source/Runtime/Core/Public/Clang/ClangPlatformAtomics.h
--- a/Engine/Source/Runtime/Core/Public/Clang/ClangPlatformAtomics.h
+++ b/Engine/Source/Runtime/Core/Public/Clang/ClangPlatformAtomics.h
@@ -6,6 +6,8 @@
 
 #pragma once
 
+#ifdef __clang__
+
 #include "GenericPlatform/GenericPlatformAtomics.h"
 #include "CoreTypes.h"
 
@@ -298,3 +300,5 @@ struct FClangPlatformAtomics : public FGenericPlatformAtomics
 		return !!PLATFORM_HAS_128BIT_ATOMICS;
 	}
 };
+
+#endif
