Fixed FThreadSafeCounter64::GetValue() thread safety.

diff --git a/Engine/Source/Runtime/Core/Public/HAL/ThreadSafeCounter64.h b/Engine/Source/Runtime/Core/Public/HAL/ThreadSafeCounter64.h
--- a/Engine/Source/Runtime/Core/Public/HAL/ThreadSafeCounter64.h
+++ b/Engine/Source/Runtime/Core/Public/HAL/ThreadSafeCounter64.h
@@ -134,7 +134,7 @@ public:
 	*/
 	int64 GetValue() const
 	{
-		return Counter;
+		return FPlatformAtomics::AtomicRead(&Counter);
 	}
 
 private:
