Adding in missing FPlatformMisc specifiers when for device orientation APIs. They were calling the GenericPlatformMisc version instead of the platform specialized ones.

diff --git a/Engine/Source/Runtime/Core/Private/GenericPlatform/GenericPlatformMisc.cpp b/Engine/Source/Runtime/Core/Private/GenericPlatform/GenericPlatformMisc.cpp
--- a/Engine/Source/Runtime/Core/Private/GenericPlatform/GenericPlatformMisc.cpp
+++ b/Engine/Source/Runtime/Core/Private/GenericPlatform/GenericPlatformMisc.cpp
@@ -1788,7 +1788,7 @@ EDeviceScreenOrientation FGenericPlatformMisc::GetDeviceOrientation()
 
 void FGenericPlatformMisc::SetDeviceOrientation(EDeviceScreenOrientation NewDeviceOrientation)
 {
-	SetAllowedDeviceOrientation(NewDeviceOrientation);
+	FPlatformMisc::SetAllowedDeviceOrientation(NewDeviceOrientation);
 }
 
 EDeviceScreenOrientation FGenericPlatformMisc::GetAllowedDeviceOrientation()
@@ -1803,7 +1803,7 @@ void FGenericPlatformMisc::SetAllowedDeviceOrientation(EDeviceScreenOrientation
 
 void FGenericPlatformMisc::ForceUpdateDeviceOrientation(EDeviceScreenOrientation NewDeviceOrientation)
 {
-    SetAllowedDeviceOrientation(NewDeviceOrientation);
+    FPlatformMisc::SetAllowedDeviceOrientation(NewDeviceOrientation);
 }
 
 int32 FGenericPlatformMisc::GetDeviceVolume()
