Add MallocBinned3 to iOS

diff --git a/Engine/Source/Runtime/Core/Core.Build.cs b/Engine/Source/Runtime/Core/Core.Build.cs
--- a/Engine/Source/Runtime/Core/Core.Build.cs
+++ b/Engine/Source/Runtime/Core/Core.Build.cs
@@ -143,6 +143,12 @@ public class Core : ModuleRules
 
 			// export Core symbols for embedded Dlls
 			ModuleSymbolVisibility = ModuleRules.SymbolVisibility.VisibileForDll;
+			
+			if (Target.Platform == UnrealTargetPlatform.IOS)
+			{
+				PublicDefinitions.Add("USE_MALLOC_BINNED3=1");
+				PublicDefinitions.Add("USE_512MB_MAX_MEMORY_PER_BLOCK_SIZE=1");
+			}
 		}
 		else if (Target.IsInPlatformGroup(UnrealPlatformGroup.Android))
 		{
diff --git a/Engine/Source/Runtime/Core/Private/Apple/ApplePlatformMemory.cpp b/Engine/Source/Runtime/Core/Private/Apple/ApplePlatformMemory.cpp
--- a/Engine/Source/Runtime/Core/Private/Apple/ApplePlatformMemory.cpp
+++ b/Engine/Source/Runtime/Core/Private/Apple/ApplePlatformMemory.cpp
@@ -17,6 +17,7 @@
 #include "HAL/MallocAnsi.h"
 #include "HAL/MallocBinned.h"
 #include "HAL/MallocBinned2.h"
+#include "HAL/MallocBinned3.h"
 #include "CoreGlobals.h"
 
 #include <stdlib.h>
@@ -287,8 +288,16 @@ void FApplePlatformMemory::Init()
 	
 }
 
-// Use MallocBinned2 as default, can be overriden below.
-#define USE_MALLOC_BINNED2 1
+// if USE_MALLOC_BINNED3 isn't configured enable it by default
+#ifndef USE_MALLOC_BINNED3
+	#define USE_MALLOC_BINNED3 1
+#endif // USE_MALLOC_BINNED3
+
+
+// if USE_MALLOC_BINNED2 isn't configured enable it if USE_MALLOC_BINNED3 isn't true
+#ifndef USE_MALLOC_BINNED2
+	#define USE_MALLOC_BINNED2 !USE_MALLOC_BINNED3
+#endif // USE_MALLOC_BINNED2
 
 void FApplePlatformMemory::SetAllocatorToUse()
 {
@@ -305,18 +314,29 @@ void FApplePlatformMemory::SetAllocatorToUse()
         AllocatorToUse = EMemoryAllocatorToUse::Ansi;
         return;
     }
-    if (USE_MALLOC_BINNED2)
+
+	if (USE_MALLOC_BINNED3)
+	{
+		if (!CanOverallocateVirtualMemory())
+		{
+			FPlatformMisc::LowLevelOutputDebugString(TEXT("MallocBinned3 requested but com.apple.developer.kernel.extended-virtual-addressing entitlement not found. Check your entitlements. Falling back to Ansi.\n"));
+			AllocatorToUse = EMemoryAllocatorToUse::Ansi;
+			return;
+		}
+
+		FPlatformMisc::LowLevelOutputDebugString(TEXT("Using MallocBinned3 allocator.\n"));
+		AllocatorToUse = EMemoryAllocatorToUse::Binned3;
+		return;
+	}
+    else if (USE_MALLOC_BINNED2)
     {
         if(!CanOverallocateVirtualMemory())
         {
-			FPlatformMisc::LowLevelOutputDebugString(TEXT("MallocBinned2 requested but Virtual Address Space entitlement not found. Check your entitlements. Falling back to Ansi.\n"));
+			FPlatformMisc::LowLevelOutputDebugString(TEXT("MallocBinned2 requested but com.apple.developer.kernel.extended-virtual-addressing entitlement not found. Check your entitlements. Falling back to Ansi.\n"));
             AllocatorToUse = EMemoryAllocatorToUse::Ansi;
             return;
         }
-        else
-        {
-			FPlatformMisc::LowLevelOutputDebugString(TEXT("Virtual Address Space entitlement found. Using MallocBinned2 allocator.\n"));
-        }
+
 		FPlatformMisc::LowLevelOutputDebugString(TEXT("Using MallocBinned2 allocator.\n"));
         AllocatorToUse = EMemoryAllocatorToUse::Binned2;
         return;
@@ -352,6 +372,12 @@ FMalloc* FApplePlatformMemory::BaseAllocator()
             break;
         }
 
+		case EMemoryAllocatorToUse::Binned3:
+		{
+			Instance = new FMallocBinned3();
+			break;
+		}
+
         case EMemoryAllocatorToUse::Binned2:
         {
             Instance = new FMallocBinned2();
