[Core]
* Changed so ModuleManager ctor will not load bootstrap file when running with merged modules

diff --git a/Engine/Source/Runtime/Core/Private/Modules/ModuleManager.cpp b/Engine/Source/Runtime/Core/Private/Modules/ModuleManager.cpp
--- a/Engine/Source/Runtime/Core/Private/Modules/ModuleManager.cpp
+++ b/Engine/Source/Runtime/Core/Private/Modules/ModuleManager.cpp
@@ -149,7 +149,7 @@ FModuleManager::FModuleManager(FPrivateToken)
 {
 	check(IsInGameThread());
 
-#if !IS_MONOLITHIC
+#if !IS_MONOLITHIC && !UE_MERGED_MODULES
 	// Modules bootstrapping is useful to avoid costly directory enumeration by reloading
 	// a serialized state of the module manager. Can only be used when run in the exact
 	// same context multiple times (i.e. starting multiple shader compile workers)
