Add FindByHashChecked

diff --git a/Engine/Source/Runtime/Core/Public/Containers/Map.h b/Engine/Source/Runtime/Core/Public/Containers/Map.h
--- a/Engine/Source/Runtime/Core/Public/Containers/Map.h
+++ b/Engine/Source/Runtime/Core/Public/Containers/Map.h
@@ -569,6 +569,19 @@ public:
 		return const_cast<TMapBase*>(this)->FindByHash(KeyHash, Key);
 	}
 
+	template<typename ComparableKey>
+	FORCEINLINE ValueType& FindByHashChecked(uint32 KeyHash, const ComparableKey& Key)
+	{
+		auto* Pair = Pairs.FindByHash(KeyHash, Key);
+		check(Pair != nullptr);
+		return Pair->Value;
+	}
+	template<typename ComparableKey>
+	FORCEINLINE const ValueType& FindByHashChecked(uint32 KeyHash, const ComparableKey& Key) const
+	{
+		return const_cast<TMapBase*>(this)->FindByHashChecked(KeyHash, Key);
+	}
+
 private:
 	FORCEINLINE static uint32 HashKey(const KeyType& Key)
 	{
