Globally disable C5105 instead of just for direct header including
Include Platform.h before doing a single file header include (to set up the compiler correctly)
#rnx

diff --git a/Engine/Source/Programs/UnrealBuildTool/Configuration/UEBuildModuleCPP.cs b/Engine/Source/Programs/UnrealBuildTool/Configuration/UEBuildModuleCPP.cs
--- a/Engine/Source/Programs/UnrealBuildTool/Configuration/UEBuildModuleCPP.cs
+++ b/Engine/Source/Programs/UnrealBuildTool/Configuration/UEBuildModuleCPP.cs
@@ -2432,6 +2432,7 @@ namespace UnrealBuildTool
 			GeneratedHeaderCppContents.Add("// This file is automatically generated at compile-time to include a user created header file.");
 			GeneratedHeaderCppContents.Add("#define UE_DIRECT_HEADER_COMPILE 1");
 			GeneratedHeaderCppContents.Add($"#define __COMPILING_{HeaderName}");
+			GeneratedHeaderCppContents.Add($"#include \"HAL/Platform.h\"");
 			GeneratedHeaderCppContents.Add($"#include \"{IncludeFileString.Replace('\\', '/')}\"");
 			return GeneratedHeaderCppContents;
 		}
diff --git a/Engine/Source/Runtime/Core/Public/Misc/DirectHeaderCompiling.h b/Engine/Source/Runtime/Core/Public/Misc/DirectHeaderCompiling.h
--- a/Engine/Source/Runtime/Core/Public/Misc/DirectHeaderCompiling.h
+++ b/Engine/Source/Runtime/Core/Public/Misc/DirectHeaderCompiling.h
@@ -2,10 +2,6 @@
 
 #pragma once
 
-// C5105: macro expansion producing 'defined' has undefined behavior
-// however, in all compilers we use, the behavior is as desired
-#pragma warning(disable : 5105) 
-
 // Used to allow control flow when direct header compiling to avoid issues with headers
 // that typically expect to be included after another header
 #ifndef UE_DIRECT_HEADER_COMPILING
diff --git a/Engine/Source/Runtime/Core/Public/Windows/WindowsPlatformCompilerSetup.h b/Engine/Source/Runtime/Core/Public/Windows/WindowsPlatformCompilerSetup.h
--- a/Engine/Source/Runtime/Core/Public/Windows/WindowsPlatformCompilerSetup.h
+++ b/Engine/Source/Runtime/Core/Public/Windows/WindowsPlatformCompilerSetup.h
@@ -233,6 +233,7 @@ static_assert(_MSC_VER >= 1934, "Visual Studio 2022 v17.4 or greater is required
 #pragma warning(disable: 4946) // reinterpret_cast used between related classes: 'class1' and 'class2'								https://docs.microsoft.com/en-us/cpp/error-messages/compiler-warnings/compiler-warning-level-1-c4946
 #pragma warning(default: 5038) // data member 'A::y' will be initialized after data member 'A::x'									https://docs.microsoft.com/en-us/cpp/error-messages/compiler-warnings/c5038
 #pragma warning(disable: 4984) // 'if constexpr' is a C++17 language extension
+#pragma warning(disable : 5105) // macro expansion producing 'defined' has undefined behavior
 
 // #pragma warning(default: 4996) // Deprecation: controlled via command line
 
