Correct how we now can return the original allocator from MemoryTrace_CreateInternal.

This fixes incorrectly trying to hook virtual memory allocators on Windows when memory tracing is disabled.

diff --git a/Engine/Source/Runtime/Core/Private/ProfilingDebugging/Microsoft/WindowsMemoryTrace.cpp b/Engine/Source/Runtime/Core/Private/ProfilingDebugging/Microsoft/WindowsMemoryTrace.cpp
--- a/Engine/Source/Runtime/Core/Private/ProfilingDebugging/Microsoft/WindowsMemoryTrace.cpp
+++ b/Engine/Source/Runtime/Core/Private/ProfilingDebugging/Microsoft/WindowsMemoryTrace.cpp
@@ -401,16 +401,14 @@ FMalloc* MemoryTrace_Create(FMalloc* InMalloc)
 	FMalloc* OutMalloc = MemoryTrace_CreateInternal(InMalloc, ArgC, ArgV);
 	::LocalFree(HLOCAL(ArgV));
 
-	if (OutMalloc != nullptr)
+	if (OutMalloc != InMalloc)
 	{
 #if defined(PLATFORM_SUPPORTS_TRACE_WIN32_VIRTUAL_MEMORY_HOOKS)
 		FVirtualWinApiHooks::Initialize(false);
 #endif
-
-		return OutMalloc;
 	}
 
-	return InMalloc;
+	return OutMalloc;
 }
 
 #include "Windows/HideWindowsPlatformTypes.h"
diff --git a/Engine/Source/Runtime/Core/Private/ProfilingDebugging/Unix/UnixMemoryTrace.cpp b/Engine/Source/Runtime/Core/Private/ProfilingDebugging/Unix/UnixMemoryTrace.cpp
--- a/Engine/Source/Runtime/Core/Private/ProfilingDebugging/Unix/UnixMemoryTrace.cpp
+++ b/Engine/Source/Runtime/Core/Private/ProfilingDebugging/Unix/UnixMemoryTrace.cpp
@@ -36,7 +36,7 @@ FMalloc* MemoryTrace_Create(FMalloc* InMalloc)
 			free((void*)ArgV[--ArgC]);
 		}
 
-		return (OutMalloc != nullptr) ? OutMalloc : InMalloc;
+		return OutMalloc;
 	}
 	return InMalloc;
 }
