ALWAYS_OPEN can be removed when being forceinlined, so moving OPEN inside the functions

diff --git a/Engine/Source/Runtime/Core/Private/Internationalization/TextFormatter.cpp b/Engine/Source/Runtime/Core/Private/Internationalization/TextFormatter.cpp
--- a/Engine/Source/Runtime/Core/Private/Internationalization/TextFormatter.cpp
+++ b/Engine/Source/Runtime/Core/Private/Internationalization/TextFormatter.cpp
@@ -379,11 +379,16 @@ public:
 	/**
 	 * Test to see whether this instance contains valid compiled data.
 	 */
-	UE_AUTORTFM_ALWAYS_OPEN
 	FORCEINLINE bool IsValid() const
 	{
-		UE::TScopeLock Lock(CompiledDataMutex);
-		return IsValid_NoLock();
+		bool bIsValid = false;
+		UE_AUTORTFM_OPEN(
+		{
+			UE::TScopeLock Lock(CompiledDataMutex);
+			bIsValid = IsValid_NoLock();
+		});
+
+		return bIsValid;
 	}
 
 	/**
@@ -395,31 +400,43 @@ public:
 	 * Validate the format pattern is valid based on the rules of the given culture (or null to use the current language).
 	 * @return true if the pattern is valid, or false if not (false may also fill in OutValidationErrors).
 	 */
-	UE_AUTORTFM_ALWAYS_OPEN
 	FORCEINLINE bool ValidatePattern(const FCulturePtr& InCulture, TArray<FString>& OutValidationErrors)
 	{
-		UE::TScopeLock Lock(CompiledDataMutex);
-		return ValidatePattern_NoLock(InCulture, OutValidationErrors);
+		bool bIsValid = false;
+		UE_AUTORTFM_OPEN(
+		{
+			UE::TScopeLock Lock(CompiledDataMutex);
+			bIsValid = ValidatePattern_NoLock(InCulture, OutValidationErrors);
+		});
+
+		return bIsValid;
 	}
 
 	/**
 	 * Produce a formatted string using the given argument look-up.
 	 */
-	UE_AUTORTFM_ALWAYS_OPEN
 	FORCEINLINE FString Format(const FPrivateTextFormatArguments& InFormatArgs)
 	{
-		UE::TScopeLock Lock(CompiledDataMutex);
-		return Format_NoLock(InFormatArgs);
+		FString RetString;
+		UE_AUTORTFM_OPEN(
+		{
+			UE::TScopeLock Lock(CompiledDataMutex);
+			RetString = Format_NoLock(InFormatArgs);
+		});
+
+		return RetString;
 	}
 
 	/**
 	 * Append the names of any arguments to the given array.
 	 */
-	UE_AUTORTFM_ALWAYS_OPEN
 	FORCEINLINE void GetFormatArgumentNames(TArray<FString>& OutArgumentNames)
 	{
-		UE::TScopeLock Lock(CompiledDataMutex);
-		return GetFormatArgumentNames_NoLock(OutArgumentNames);
+		UE_AUTORTFM_OPEN(
+		{
+			UE::TScopeLock Lock(CompiledDataMutex);
+			GetFormatArgumentNames_NoLock(OutArgumentNames);
+		});
 	}
 
 	/**
@@ -443,11 +460,16 @@ public:
 	/**
 	 * Get the type of expression currently compiled.
 	 */
-	UE_AUTORTFM_ALWAYS_OPEN
 	FORCEINLINE FTextFormat::EExpressionType GetExpressionType() const
 	{
-		UE::TScopeLock Lock(CompiledDataMutex);
-		return CompiledExpressionType;
+		FTextFormat::EExpressionType RetType = FTextFormat::EExpressionType::Invalid;
+		UE_AUTORTFM_OPEN(
+		{
+			UE::TScopeLock Lock(CompiledDataMutex);
+			RetType = CompiledExpressionType;
+		});
+
+		return RetType;
 	}
 
 	/**
