Fix Preview Platform FriendlyName overrides

diff --git a/Engine/Config/IOS/DataDrivenPlatformInfo.ini b/Engine/Config/IOS/DataDrivenPlatformInfo.ini
--- a/Engine/Config/IOS/DataDrivenPlatformInfo.ini
+++ b/Engine/Config/IOS/DataDrivenPlatformInfo.ini
@@ -32,8 +32,8 @@ ActiveIconName=LevelEditor.PreviewMode.iOS.Enabled
 InactiveIconName=LevelEditor.PreviewMode.iOS.Disabled
 MenuTooltip=NSLOCTEXT("PreviewPlatform", "PreviewMenuTooltip_iOS", "Mobile preview using iOS material quality settings.")
 IconText=NSLOCTEXT("PreviewPlatform", "PreviewIconText_iOS", "iOS")
-DeviceProfileName=IOS,iPhone6S,iPadPro6_129
-FriendlyName=Generic IOS,iPhone 6S,iPad Pro 6 12.9in
+DeviceProfileName=IOS|iPhone6S|iPadPro6_129
+FriendlyName=Generic IOS|iPhone 6S|iPad Pro 6 12.9in
 
 [PreviewPlatform IOSMetalSM5]
 PlatformName=IOS
diff --git a/Engine/Source/Runtime/Core/Private/Misc/DataDrivenPlatformInfoRegistry.cpp b/Engine/Source/Runtime/Core/Private/Misc/DataDrivenPlatformInfoRegistry.cpp
--- a/Engine/Source/Runtime/Core/Private/Misc/DataDrivenPlatformInfoRegistry.cpp
+++ b/Engine/Source/Runtime/Core/Private/Misc/DataDrivenPlatformInfoRegistry.cpp
@@ -299,11 +299,12 @@ static void ParsePreviewPlatforms(const FConfigFile& IniFile)
 
 
 				FString AllDeviceProfiles = GetSectionString(Section.Value, FName("DeviceProfileName"));
-				FString AllFriendlyName = GetSectionString(Section.Value, FName("FriendlyName"));
+				FText AllFriendlyName;
+				FTextStringHelper::ReadFromBuffer(*GetSectionString(Section.Value, FName("FriendlyName")), AllFriendlyName);
 				TArray<FString> DeviceProfileNames, FriendlyNames;
-				AllDeviceProfiles.ParseIntoArray(DeviceProfileNames, TEXT(","));
-				AllFriendlyName.ParseIntoArray(FriendlyNames, TEXT(","));
-				
+				AllDeviceProfiles.ParseIntoArray(DeviceProfileNames, TEXT("|"));
+				AllFriendlyName.ToString().ParseIntoArray(FriendlyNames, TEXT("|"));
+
 				if (DeviceProfileNames.Num() == 0)
 				{
 					DeviceProfileNames.Add(TEXT(""));
