autortfm: optimization/cleanup pass

- Several function definitions have been moved to the AutoRTFM.h header file so they can be more heavily optimized, especially in uninstrumented code. For example, autortfm_abort_if_closed is always a no-op in open code, so the definition is now inlined. That, coupled with our existing optimization of constant folding calls to autortfm_is_closed() means that this function completely goes away in uninstrumented code.

- Some functions, like OnAbort, had a public API that took a TFunction<> as a parameter. These functions are no-ops in uninstrumented code, but we were paying the cost of constructing a TFunction<> and doing the function call. Now the definition is inlined and the construction of the TFunction<> is wrapped behind a call to autortfm_is_closed(), so it completely folds away in uninstrumented code.

- Functions like Transact, TransactThenOpen, Commit, ... were modified to have an inlined call to autortfm_lookup_function that does the function pointer translation at the callsite rather than inside the runtime. This plays along with recent changes to the compiler that evaluates calls to autortfm_lookup_function at compile time and eliminates the runtime lookup.

- autortfm_abort_if_transactional / abort_if_closed were moved to the header file so is_transactional/is_closed can be evaluated at compile time and dead stripped as appropriate

diff --git a/Engine/Source/Programs/AutoRTFMTests/Private/API.cpp b/Engine/Source/Programs/AutoRTFMTests/Private/API.cpp
--- a/Engine/Source/Programs/AutoRTFMTests/Private/API.cpp
+++ b/Engine/Source/Programs/AutoRTFMTests/Private/API.cpp
@@ -7,34 +7,6 @@
 #include <memory>
 #include <thread>
 
-TEST_CASE("API.autortfm_result")
-{
-    int Answer = 6 * 9;
-
-    REQUIRE(autortfm_committed == autortfm_transact([](void* const Arg)
-    {
-        *static_cast<int* const>(Arg) = 42;
-    }, &Answer));
-
-    REQUIRE(42 == Answer);
-
-    REQUIRE(autortfm_aborted_by_request == autortfm_transact([](void* const Arg)
-    {
-        *static_cast<int* const>(Arg) = 13;
-        AutoRTFM::AbortTransaction();
-    }, &Answer));
-
-    REQUIRE(42 == Answer);
-
-    REQUIRE(autortfm_aborted_by_request == autortfm_transact([](void* const Arg)
-    {
-        *static_cast<int* const>(Arg) = 13;
-        AutoRTFM::AbortIfTransactional();
-    }, &Answer));
-
-    REQUIRE(42 == Answer);
-}
-
 TEST_CASE("API.autortfm_is_transactional")
 {
     REQUIRE(false == autortfm_is_transactional());
@@ -85,30 +57,6 @@ TEST_CASE("API.autortfm_is_closed")
     REQUIRE(true == InClosedNestInOpenNest);
 }
 
-TEST_CASE("API.autortfm_transact")
-{
-    int Answer = 6 * 9;
-
-    REQUIRE(autortfm_committed == autortfm_transact([](void* const Arg)
-    {
-        *static_cast<int* const>(Arg) = 42;
-    }, &Answer));
-
-    REQUIRE(42 == Answer);
-}
-
-TEST_CASE("API.autortfm_commit")
-{
-    int Answer = 6 * 9;
-
-    autortfm_commit([](void* const Arg)
-    {
-        *static_cast<int* const>(Arg) = 42;
-    }, &Answer);
-
-    REQUIRE(42 == Answer);
-}
-
 TEST_CASE("API.autortfm_abort_transaction")
 {
     bool BeforeNest = false;
@@ -225,39 +173,6 @@ TEST_CASE("API.autortfm_open")
     REQUIRE(84 == Answer);
 }
 
-TEST_CASE("API.autortfm_close")
-{
-    bool InClosedNest = false;
-    bool InOpenNest = false;
-    bool InClosedNestInOpenNest = false;
-
-    REQUIRE(AutoRTFM::ETransactionResult::AbortedByRequest == AutoRTFM::Transact([&]
-    {
-        // A closed call inside a transaction does not abort.
-        REQUIRE(autortfm_status_ontrack == autortfm_close([](void* const Arg)
-        {
-            *static_cast<bool* const>(Arg) = true;
-        }, &InClosedNest));
-
-        AutoRTFM::Open([&]
-        {
-            // A closed call inside an open does not abort either.
-			REQUIRE(autortfm_status_ontrack == autortfm_close([](void* const Arg)
-            {
-                *static_cast<bool* const>(Arg) = true;
-            }, &InClosedNestInOpenNest));
-
-            InOpenNest = true;
-        });
-
-        AutoRTFM::AbortTransaction();
-    }));
-
-    REQUIRE(false == InClosedNest);
-    REQUIRE(true == InOpenNest);
-    REQUIRE(false == InClosedNestInOpenNest);
-}
-
 TEST_CASE("API.autortfm_register_open_function")
 {
     autortfm_register_open_function(
diff --git a/Engine/Source/Programs/AutoRTFMTests/Private/Catch2Includes.h b/Engine/Source/Programs/AutoRTFMTests/Private/Catch2Includes.h
--- a/Engine/Source/Programs/AutoRTFMTests/Private/Catch2Includes.h
+++ b/Engine/Source/Programs/AutoRTFMTests/Private/Catch2Includes.h
@@ -5,4 +5,7 @@
 // we need to include thie MainCPP include *before* we include catch_amalgamated or we get issues with windows symbols
 #include "CoreMinimal.h"
 #include "LaunchEngineLoop.h"
+
+THIRD_PARTY_INCLUDES_START
 #include <catch_amalgamated.hpp>
+THIRD_PARTY_INCLUDES_END
diff --git a/Engine/Source/Runtime/Core/Private/AutoRTFM/API.cpp b/Engine/Source/Runtime/Core/Private/AutoRTFM/API.cpp
--- a/Engine/Source/Runtime/Core/Private/AutoRTFM/API.cpp
+++ b/Engine/Source/Runtime/Core/Private/AutoRTFM/API.cpp
@@ -120,15 +120,12 @@ namespace AutoRTFM
 namespace AutoRTFM
 {
 
-UE_AUTORTFM_FORCEINLINE autortfm_result TransactThenOpenImpl(void (*Work)(void* Arg), void* Arg)
+UE_AUTORTFM_FORCEINLINE autortfm_result TransactThenOpenImpl(void (*UninstrumentedWork)(void*), void* Arg)
 {
 	return static_cast<autortfm_result>(
 		AutoRTFM::Transact([&]
 		{
-			AutoRTFM::Open([&]
-			{
-				Work(Arg);
-			});
+			autortfm_open(UninstrumentedWork, Arg);
 		}));
 }
 
@@ -148,25 +145,25 @@ extern "C" UE_AUTORTFM_AUTORTFM("RTFM_autortfm_is_closed") bool autortfm_is_clos
 }
 
 // First Part - the API exposed outside transactions.
-extern "C" UE_AUTORTFM_AUTORTFM("RTFM_autortfm_transact") autortfm_result autortfm_transact(void (*Work)(void* Arg), void* Arg)
+extern "C" UE_AUTORTFM_AUTORTFM("RTFM_autortfm_transact") autortfm_result autortfm_transact(void (*UninstrumentedWork)(void*), void (*InstrumentedWork)(void*), void* Arg)
 {
 	if (ForTheRuntime::IsAutoRTFMRuntimeEnabled())
 	{
-	    return static_cast<autortfm_result>(FContext::Get()->Transact(Work, Arg));
+	    return static_cast<autortfm_result>(FContext::Get()->Transact(InstrumentedWork, Arg));
 	}
 
-	Work(Arg);
+	(*UninstrumentedWork)(Arg);
 	return autortfm_committed;
 }
 
-extern "C" UE_AUTORTFM_AUTORTFM("RTFM_autortfm_transact_then_open") autortfm_result autortfm_transact_then_open(void (*Work)(void* Arg), void* Arg)
+extern "C" UE_AUTORTFM_AUTORTFM("RTFM_autortfm_transact_then_open") autortfm_result autortfm_transact_then_open(void (*UninstrumentedWork)(void*), void (*InstrumentedWork)(void*), void* Arg)
 {
-    return TransactThenOpenImpl(Work, Arg);
+	return TransactThenOpenImpl(UninstrumentedWork, Arg);
 }
 
-extern "C" UE_AUTORTFM_AUTORTFM("RTFM_autortfm_commit") void autortfm_commit(void (*Work)(void* Arg), void* Arg)
+extern "C" UE_AUTORTFM_AUTORTFM("RTFM_autortfm_commit") void autortfm_commit(void (*UninstrumentedWork)(void*), void (*InstrumentedWork)(void*), void* Arg)
 {
-    autortfm_result Result = autortfm_transact(Work, Arg);
+    autortfm_result Result = autortfm_transact(UninstrumentedWork, InstrumentedWork, Arg);
 	UE_CLOG(Result != autortfm_committed, LogAutoRTFM, Fatal, TEXT("Unexpected transaction result: %u."), Result);
 }
 
@@ -224,21 +221,12 @@ extern "C" UE_AUTORTFM_NOAUTORTFM bool autortfm_current_nest_throw()
 	return true;
 }
 
-extern "C" UE_AUTORTFM_AUTORTFM("RTFM_autortfm_abort_if_transactional") void autortfm_abort_if_transactional()
-{
-	UE_CLOG(FContext::IsTransactional(), LogAutoRTFM, Fatal, TEXT("The function `autortfm_abort_if_transactional` was called from an open inside a transaction."));
-}
-
-extern "C" UE_AUTORTFM_AUTORTFM("RTFM_autortfm_abort_if_closed") void autortfm_abort_if_closed()
-{
-}
-
-extern "C" UE_AUTORTFM_AUTORTFM("RTFM_autortfm_open") void autortfm_open(void (*Work)(void* Arg), void* Arg)
+extern "C" UE_AUTORTFM_AUTORTFM("RTFM_autortfm_open") void autortfm_open(void (*Work)(void*), void* Arg)
 {
 	Work(Arg);
 }
 
-extern "C" UE_AUTORTFM_AUTORTFM("RTFM_autortfm_close") autortfm_status autortfm_close(void (*Work)(void* Arg), void* Arg)
+extern "C" UE_AUTORTFM_AUTORTFM("RTFM_autortfm_close") autortfm_status autortfm_close(void (*UninstrumentedWork)(void*), void (*InstrumentedWork)(void*), void* Arg)
 {
 	autortfm_status Result = autortfm_status_ontrack;
 
@@ -247,15 +235,20 @@ extern "C" UE_AUTORTFM_AUTORTFM("RTFM_autortfm_close") autortfm_status autortfm_
 		UE_CLOG(!FContext::IsTransactional(), LogAutoRTFM, Fatal, TEXT("Close called from an outside a transaction."));
 
 		FContext* const Context = FContext::Get();
-		void (*WorkClone)(void* Arg) = FunctionMapLookup(Work, "autortfm_close");
-		if (WorkClone)
+
+		if (InstrumentedWork)
+		{
+			Result = static_cast<autortfm_status>(Context->CallClosedNest(InstrumentedWork, Arg));
+		}
+		else
 		{
-			Result = static_cast<autortfm_status>(Context->CallClosedNest(WorkClone, Arg));
+			UE_LOG(LogAutoRTFM, Fatal, TEXT("Could not find function %p '%s' where '%s'."), UninstrumentedWork, *GetFunctionDescription(UninstrumentedWork), ANSI_TO_TCHAR("autortfm_close"));
+	        Context->AbortByLanguageAndThrow();
 		}
 	}
 	else
 	{
-		Work(Arg);
+		UninstrumentedWork(Arg);
 	}
 
 	return Result;
@@ -275,31 +268,31 @@ extern "C" UE_AUTORTFM_NOAUTORTFM void autortfm_register_open_function(void* Ori
     FunctionMapAdd(OriginalFunction, NewFunction);
 }
 
-UE_AUTORTFM_AUTORTFM("RTFM_OnCommit") void OnCommit(TFunction<void()> && Work)
+UE_AUTORTFM_AUTORTFM("RTFM_OnCommitInternal") void ForTheRuntime::OnCommitInternal(TFunction<void()> && Work)
 {
     Work();
 }
 
-UE_AUTORTFM_AUTORTFM("RTFM_OnAbort") void OnAbort(TFunction<void()> && Work)
+UE_AUTORTFM_AUTORTFM("RTFM_OnAbortInternal") void ForTheRuntime::OnAbortInternal(TFunction<void()> && Work)
 {
 }
 
-UE_AUTORTFM_AUTORTFM("RTFM_PushOnAbortHandler") void PushOnAbortHandler(const void* Key, TFunction<void()> && Work)
+UE_AUTORTFM_AUTORTFM("RTFM_PushOnAbortHandlerInternal") void ForTheRuntime::PushOnAbortHandlerInternal(const void* Key, TFunction<void()> && Work)
 {
 }
 
-UE_AUTORTFM_AUTORTFM("RTFM_PopOnAbortHandler") void PopOnAbortHandler(const void* Key)
+UE_AUTORTFM_AUTORTFM("RTFM_PopOnAbortHandlerInternal") void ForTheRuntime::PopOnAbortHandlerInternal(const void* Key)
 {
 }
 
 void OpenCommit(TFunction<void()>&& Work)
 {
-	OnCommit(MoveTemp(Work));
+	ForTheRuntime::OnCommitInternal(MoveTemp(Work));
 }
 
 void OpenAbort(TFunction<void()>&& Work)
 {
-	OnAbort(MoveTemp(Work));
+	ForTheRuntime::OnAbortInternal(MoveTemp(Work));
 }
 
 extern "C" UE_AUTORTFM_AUTORTFM("RTFM_autortfm_on_commit") void autortfm_on_commit(void (*Work)(void*), void* Arg)
@@ -405,20 +398,20 @@ extern "C" UE_AUTORTFM_NOAUTORTFM bool RTFM_autortfm_is_closed()
     return true;
 }
 
-extern "C" UE_AUTORTFM_NOAUTORTFM autortfm_result RTFM_autortfm_transact(void (*Work)(void* Arg), void* Arg)
+extern "C" UE_AUTORTFM_NOAUTORTFM autortfm_result RTFM_autortfm_transact(void (*UninstrumentedWork)(void*), void (*InstrumentedWork)(void*), void* Arg)
 {
 	FContext* Context = FContext::Get();
-    return static_cast<autortfm_result>(Context->Transact(Work, Arg));
+    return static_cast<autortfm_result>(Context->Transact(InstrumentedWork, Arg));
 }
 
-extern "C" UE_AUTORTFM_NOAUTORTFM autortfm_result RTFM_autortfm_transact_then_open(void (*Work)(void* Arg), void* Arg)
+extern "C" UE_AUTORTFM_NOAUTORTFM autortfm_result RTFM_autortfm_transact_then_open(void (*UninstrumentedWork)(void*), void (*InstrumentedWork)(void*), void* Arg)
 {
-    return TransactThenOpenImpl(Work, Arg);
+    return TransactThenOpenImpl(UninstrumentedWork, Arg);
 }
 
-extern "C" UE_AUTORTFM_NOAUTORTFM void RTFM_autortfm_commit(void (*Work)(void* Arg), void* Arg)
+extern "C" UE_AUTORTFM_NOAUTORTFM void RTFM_autortfm_commit(void (*UninstrumentedWork)(void*), void (*InstrumentedWork)(void*), void* Arg)
 {
-    autortfm_result Result = autortfm_transact(Work, Arg);
+    autortfm_result Result = autortfm_transact(UninstrumentedWork, InstrumentedWork, Arg);
 	UE_CLOG(Result != autortfm_committed, LogAutoRTFM, Fatal, TEXT("Unexpected transaction result: %u."), Result);
 }
 
@@ -456,21 +449,7 @@ extern "C" UE_AUTORTFM_NOAUTORTFM void RTFM_autortfm_clear_transaction_status()
 	AutoRTFM::Unreachable();
 }
 
-extern "C" UE_AUTORTFM_NOAUTORTFM void RTFM_autortfm_abort_if_transactional()
-{
-    UE_LOG(LogAutoRTFM, Verbose, TEXT("The function `autortfm_abort_if_transactional` was called from inside a transaction."));
-	FContext* Context = FContext::Get();
-    Context->AbortByRequestAndThrow();
-}
-
-extern "C" UE_AUTORTFM_NOAUTORTFM void RTFM_autortfm_abort_if_closed()
-{
-	UE_LOG(LogAutoRTFM, Verbose, TEXT("The function `autortfm_abort_if_closed` was called from closed inside a transaction."));
-	FContext* Context = FContext::Get();
-    Context->AbortByRequestAndThrow();
-}
-
-extern "C" UE_AUTORTFM_NOAUTORTFM void RTFM_autortfm_open(void (*Work)(void* Arg), void* Arg)
+extern "C" UE_AUTORTFM_NOAUTORTFM void RTFM_autortfm_open(void (*Work)(void*), void* Arg)
 {
 	Work(Arg);
 
@@ -481,15 +460,21 @@ extern "C" UE_AUTORTFM_NOAUTORTFM void RTFM_autortfm_open(void (*Work)(void* Arg
 	}
 }
 
-extern "C" UE_AUTORTFM_NOAUTORTFM autortfm_status RTFM_autortfm_close(void (*Work)(void* Arg), void* Arg)
+extern "C" UE_AUTORTFM_NOAUTORTFM autortfm_status RTFM_autortfm_close(void (*UninstrumentedWork)(void*), void (*InstrumentedWork)(void*), void* Arg)
 {
-    void (*WorkClone)(void* Arg) = FunctionMapLookup(Work, "RTFM_autortfm_close");
-    if (WorkClone)
+	FContext* const Context = FContext::Get();
+
+	if (InstrumentedWork)
     {
-        WorkClone(Arg);
+        InstrumentedWork(Arg);
     }
+	else
+	{
+		UE_LOG(LogAutoRTFM, Fatal, TEXT("Could not find function %p '%s' where '%s'."), UninstrumentedWork, *GetFunctionDescription(UninstrumentedWork), ANSI_TO_TCHAR("autortfm_close"));
+	    Context->AbortByLanguageAndThrow();
+	}
 
-	return static_cast<autortfm_status>(FContext::Get()->GetStatus());
+	return static_cast<autortfm_status>(Context->GetStatus());
 }
 
 extern "C" UE_AUTORTFM_NOAUTORTFM void RTFM_autortfm_record_open_write(void*, size_t)
@@ -497,28 +482,28 @@ extern "C" UE_AUTORTFM_NOAUTORTFM void RTFM_autortfm_record_open_write(void*, si
 	UE_LOG(LogAutoRTFM, Fatal, TEXT("The function `autortfm_record_open_write` was called from closed code."));
 }
 
-extern "C" UE_AUTORTFM_NOAUTORTFM void RTFM_OnCommit(TFunction<void()> && Work)
+extern "C" UE_AUTORTFM_NOAUTORTFM void RTFM_OnCommitInternal(TFunction<void()> && Work)
 {
 	FContext* Context = FContext::Get();
     ASSERT(Context->GetStatus() == EContextStatus::OnTrack);
     Context->GetCurrentTransaction()->DeferUntilCommit(MoveTemp(Work));
 }
 
-extern "C" UE_AUTORTFM_NOAUTORTFM void RTFM_OnAbort(TFunction<void()>&& Work)
+extern "C" UE_AUTORTFM_NOAUTORTFM void RTFM_OnAbortInternal(TFunction<void()>&& Work)
 {
 	FContext* Context = FContext::Get();
     ASSERT(Context->GetStatus() == EContextStatus::OnTrack);
     Context->GetCurrentTransaction()->DeferUntilAbort(MoveTemp(Work));
 }
 
-extern "C" UE_AUTORTFM_NOAUTORTFM void RTFM_PushOnAbortHandler(const void* Key, TFunction<void()>&& Work)
+extern "C" UE_AUTORTFM_NOAUTORTFM void RTFM_PushOnAbortHandlerInternal(const void* Key, TFunction<void()>&& Work)
 {
 	FContext* Context = FContext::Get();
     ASSERT(Context->GetStatus() == EContextStatus::OnTrack);
     Context->GetCurrentTransaction()->PushDeferUntilAbortHandler(Key, MoveTemp(Work));
 }
 
-extern "C" UE_AUTORTFM_NOAUTORTFM void RTFM_PopOnAbortHandler(const void* Key)
+extern "C" UE_AUTORTFM_NOAUTORTFM void RTFM_PopOnAbortHandlerInternal(const void* Key)
 {
 	FContext* Context = FContext::Get();
     ASSERT(Context->GetStatus() == EContextStatus::OnTrack);
@@ -527,22 +512,22 @@ extern "C" UE_AUTORTFM_NOAUTORTFM void RTFM_PopOnAbortHandler(const void* Key)
 
 extern "C" UE_AUTORTFM_NOAUTORTFM void RTFM_autortfm_on_commit(void (*Work)(void*), void* Arg)
 {
-    RTFM_OnCommit([Work, Arg] { Work(Arg); });
+    RTFM_OnCommitInternal([Work, Arg] { Work(Arg); });
 }
 
 extern "C" UE_AUTORTFM_NOAUTORTFM void RTFM_autortfm_on_abort(void (*Work)(void*), void* Arg)
 {
-    RTFM_OnAbort([Work, Arg] { Work(Arg); });
+    RTFM_OnAbortInternal([Work, Arg] { Work(Arg); });
 }
 
 extern "C" UE_AUTORTFM_NOAUTORTFM void RTFM_autortfm_push_on_abort_handler(const void* Key, void (*Work)(void*), void* Arg)
 {
-    RTFM_PushOnAbortHandler(Key, [Work, Arg] { Work(Arg); });
+    RTFM_PushOnAbortHandlerInternal(Key, [Work, Arg] { Work(Arg); });
 }
 
 extern "C" UE_AUTORTFM_NOAUTORTFM void RTFM_autortfm_pop_on_abort_handler(const void* Key)
 {
-    RTFM_PopOnAbortHandler(Key);
+    RTFM_PopOnAbortHandlerInternal(Key);
 }
 
 extern "C" UE_AUTORTFM_NOAUTORTFM void* RTFM_autortfm_did_allocate(void* Ptr, size_t Size)
diff --git a/Engine/Source/Runtime/Core/Private/AutoRTFM/Context.cpp b/Engine/Source/Runtime/Core/Private/AutoRTFM/Context.cpp
--- a/Engine/Source/Runtime/Core/Private/AutoRTFM/Context.cpp
+++ b/Engine/Source/Runtime/Core/Private/AutoRTFM/Context.cpp
@@ -283,7 +283,7 @@ ETransactionResult FContext::ResolveNestedTransaction(FTransaction* NewTransacti
 	}
 }
 
-ETransactionResult FContext::Transact(void (*Function)(void* Arg), void* Arg)
+ETransactionResult FContext::Transact(void (*InstrumentedFunction)(void*), void* Arg)
 {
     constexpr bool bVerbose = false;
 
@@ -305,10 +305,9 @@ ETransactionResult FContext::Transact(void (*Function)(void* Arg), void* Arg)
     
     ASSERT(Status == EContextStatus::Idle || Status == EContextStatus::OnTrack);
 
-    void (*ClonedFunction)(void* Arg) = FunctionMapTryLookup(Function);
-    if (!ClonedFunction)
+    if (!InstrumentedFunction)
     {
-		UE_LOG(LogAutoRTFM, Warning, TEXT("Could not find function %p (%s) in AutoRTFM::FContext::Transact."), Function, *GetFunctionDescription(Function));
+		UE_LOG(LogAutoRTFM, Warning, TEXT("Could not find function in AutoRTFM::FContext::Transact."));
         return ETransactionResult::AbortedByLanguage;
     }
     
@@ -338,7 +337,7 @@ ETransactionResult FContext::Transact(void (*Function)(void* Arg), void* Arg)
         {
             Status = EContextStatus::OnTrack;
             ASSERT(CurrentTransaction->IsFresh());
-			CurrentNest->Try([&] () { ClonedFunction(Arg); });
+			CurrentNest->Try([&] () { InstrumentedFunction(Arg); });
 			ASSERT(CurrentTransaction == NewTransaction); // The transaction lambda should have unwound any nested transactions.
             ASSERT(Status != EContextStatus::Idle);
 
@@ -399,7 +398,7 @@ ETransactionResult FContext::Transact(void (*Function)(void* Arg), void* Arg)
 		PushTransaction(NewTransaction);
 		PushCallNest(NewNest);
 
-		CurrentNest->Try([&]() { ClonedFunction(Arg); });
+		CurrentNest->Try([&]() { InstrumentedFunction(Arg); });
 		ASSERT(CurrentTransaction == NewTransaction);
 
 		Result = ResolveNestedTransaction(NewTransaction);
diff --git a/Engine/Source/Runtime/Core/Private/AutoRTFM/Context.h b/Engine/Source/Runtime/Core/Private/AutoRTFM/Context.h
--- a/Engine/Source/Runtime/Core/Private/AutoRTFM/Context.h
+++ b/Engine/Source/Runtime/Core/Private/AutoRTFM/Context.h
@@ -20,7 +20,7 @@ public:
     static bool IsTransactional();
     
     // This is public API
-    ETransactionResult Transact(void (*Function)(void* Arg), void* Arg);
+    ETransactionResult Transact(void (*InstrumentedFunction)(void*), void* Arg);
     
 	EContextStatus CallClosedNest(void (*ClosedFunction)(void* Arg), void* Arg);
 
diff --git a/Engine/Source/Runtime/Core/Public/AutoRTFM/AutoRTFM.h b/Engine/Source/Runtime/Core/Public/AutoRTFM/AutoRTFM.h
--- a/Engine/Source/Runtime/Core/Public/AutoRTFM/AutoRTFM.h
+++ b/Engine/Source/Runtime/Core/Public/AutoRTFM/AutoRTFM.h
@@ -131,31 +131,34 @@ UE_AUTORTFM_FORCEINLINE bool autortfm_is_closed(void)
 #endif
 
 #if UE_AUTORTFM
-UE_AUTORTFM_API autortfm_result autortfm_transact(void (*work)(void* arg), void* arg);
+UE_AUTORTFM_API autortfm_result autortfm_transact(void (*UninstrumentedWork)(void*), void (*InstrumentedWork)(void*), void* Arg);
 #else
-UE_AUTORTFM_FORCEINLINE autortfm_result autortfm_transact(void (*work)(void* arg), void* arg)
+UE_AUTORTFM_FORCEINLINE autortfm_result autortfm_transact(void (*UninstrumentedWork)(void*), void (*InstrumentedWork)(void*), void* Arg)
 {
-	work(arg);
+	UE_AUTORTFM_UNUSED(InstrumentedWork);
+	UninstrumentedWork(Arg);
 	return autortfm_committed;
 }
 #endif
 
 #if UE_AUTORTFM
-UE_AUTORTFM_API autortfm_result autortfm_transact_then_open(void (*work)(void* arg), void* arg);
+UE_AUTORTFM_API autortfm_result autortfm_transact_then_open(void (*UninstrumentedWork)(void*), void (*InstrumentedWork)(void*), void* Arg);
 #else
-UE_AUTORTFM_FORCEINLINE autortfm_result autortfm_transact_then_open(void (*work)(void* arg), void* arg)
+UE_AUTORTFM_FORCEINLINE autortfm_result autortfm_transact_then_open(void (*UninstrumentedWork)(void*), void (*InstrumentedWork)(void*), void* Arg)
 {
-	work(arg);
+	UE_AUTORTFM_UNUSED(InstrumentedWork);
+	UninstrumentedWork(Arg);
     return autortfm_committed;
 }
 #endif
 
 #if UE_AUTORTFM
-UE_AUTORTFM_API void autortfm_commit(void (*work)(void* arg), void* arg);
+UE_AUTORTFM_API void autortfm_commit(void (*UninstrumentedWork)(void*), void (*InstrumentedWork)(void*), void* Arg);
 #else
-UE_AUTORTFM_FORCEINLINE void autortfm_commit(void (*work)(void* arg), void* arg)
+UE_AUTORTFM_FORCEINLINE void autortfm_commit(void (*UninstrumentedWork)(void*), void (*InstrumentedWork)(void*), void* Arg)
 {
-	work(arg);
+	UE_AUTORTFM_UNUSED(InstrumentedWork);
+	UninstrumentedWork(Arg);
 }
 #endif
 
@@ -190,13 +193,25 @@ UE_AUTORTFM_FORCEINLINE void autortfm_clear_transaction_status() {}
 #endif
 
 #if UE_AUTORTFM
-UE_AUTORTFM_API void autortfm_abort_if_transactional(void);
+UE_AUTORTFM_FORCEINLINE void autortfm_abort_if_transactional(void)
+{
+	if (autortfm_is_transactional())
+	{
+		autortfm_abort_transaction();
+	}
+}
 #else
 UE_AUTORTFM_FORCEINLINE void autortfm_abort_if_transactional(void) { }
 #endif
 
 #if UE_AUTORTFM
-void autortfm_abort_if_closed(void);
+UE_AUTORTFM_FORCEINLINE void autortfm_abort_if_closed(void)
+{
+	if (autortfm_is_closed())
+	{
+		autortfm_abort_transaction();
+	}
+}
 #else
 UE_AUTORTFM_FORCEINLINE void autortfm_abort_if_closed(void) { }
 #endif
@@ -208,13 +223,14 @@ UE_AUTORTFM_FORCEINLINE void autortfm_open(void (*work)(void* arg), void* arg) {
 #endif
 
 #if UE_AUTORTFM
-[[nodiscard]] UE_AUTORTFM_API autortfm_status autortfm_close(void (*work)(void* arg), void* arg);
+[[nodiscard]] UE_AUTORTFM_API autortfm_status autortfm_close(void (*UninstrumentedWork)(void*), void (*InstrumentedWork)(void*), void* Arg);
 #else
 PRAGMA_DISABLE_UNREACHABLE_CODE_WARNINGS
-[[nodiscard]] UE_AUTORTFM_FORCEINLINE autortfm_status autortfm_close(void (*work)(void* arg), void* arg)
+[[nodiscard]] UE_AUTORTFM_FORCEINLINE autortfm_status autortfm_close(void (*UninstrumentedWork)(void*), void (*InstrumentedWork)(void*), void* Arg)
 {
-	UE_AUTORTFM_UNUSED(work);
-	UE_AUTORTFM_UNUSED(arg);
+	UE_AUTORTFM_UNUSED(UninstrumentedWork);
+	UE_AUTORTFM_UNUSED(InstrumentedWork);
+	UE_AUTORTFM_UNUSED(Arg);
     abort();
 	return autortfm_status_aborted_by_language;
 }
@@ -394,8 +410,44 @@ enum class EContextStatus
 	AbortedByCascade = autortfm_status_aborted_by_cascade
 };
 
+#if UE_AUTORTFM
+namespace ForTheRuntime
+{
+	UE_AUTORTFM_API void OnCommitInternal(TFunction<void()>&& Work);
+	UE_AUTORTFM_API void OnAbortInternal(TFunction<void()>&& Work);
+	UE_AUTORTFM_API void PushOnAbortHandlerInternal(const void* Key, TFunction<void()>&& Work);
+	UE_AUTORTFM_API void PopOnAbortHandlerInternal(const void* Key);
+} // namespace ForTheRuntime
+#endif
+
+template<typename TFunctor>
+void AutoRTFMFunctorInvoker(void* Arg) { (*static_cast<const TFunctor*>(Arg))(); }
+
+#if UE_AUTORTFM
+extern "C" UE_AUTORTFM_API void* autortfm_lookup_function(void* OriginalFunction, const char* Where);
+
+template<typename TFunctor>
+auto AutoRTFMLookupInstrumentedFunctorInvoker(const TFunctor& Functor) -> void(*)(void*)
+{
+	// keep this as a single expression to help ensure that even Debug builds optimize this.
+	// if we put intermediate results in local variables then the compiler emits loads
+	// and stores to the stack which confuse our custom pass that tries to strip away
+	// the actual call to autortfm_lookup_function
+	void (*Result)(void*) = reinterpret_cast<void(*)(void*)>(autortfm_lookup_function(reinterpret_cast<void*>(&AutoRTFMFunctorInvoker<TFunctor>), "AutoRTFMLookupInstrumentedFunctorInvoker"));
+	return Result;
+}
+#else
+template<typename TFunctor>
+auto AutoRTFMLookupInstrumentedFunctorInvoker(const TFunctor& Functor) -> void(*)(void*)
+{
+	return nullptr;
+}
+#endif
+
 // Tells if we are currently running in a transaction. This will return true in an
-// open nest (see Open).
+// open nest (see Open). This function is handled specially in the compiler, it
+// will be constant folded as 'true' in closed code, or
+// preserved as a function call in open code
 UE_AUTORTFM_FORCEINLINE bool IsTransactional() { return autortfm_is_transactional(); }
 
 // Tells if we are currently running in the closed nest of a transaction. By
@@ -421,9 +473,14 @@ UE_AUTORTFM_FORCEINLINE bool IsClosed() { return autortfm_is_closed(); }
 template<typename TFunctor>
 UE_AUTORTFM_FORCEINLINE ETransactionResult Transact(const TFunctor& Functor)
 {
-    return static_cast<ETransactionResult>(autortfm_transact(
-        [] (void* Arg) { (*static_cast<const TFunctor*>(Arg))(); },
-        const_cast<void*>(static_cast<const void*>(&Functor))));
+	ETransactionResult Result =
+		static_cast<ETransactionResult>(
+			autortfm_transact(
+				&AutoRTFMFunctorInvoker<TFunctor>,
+				AutoRTFMLookupInstrumentedFunctorInvoker<TFunctor>(Functor),
+				const_cast<void*>(static_cast<const void*>(&Functor))));
+
+	return Result;
 }
 
 // This is just like calling Transact([&] { Open([&] { Functor(); }); });  
@@ -433,9 +490,14 @@ UE_AUTORTFM_FORCEINLINE ETransactionResult Transact(const TFunctor& Functor)
 template<typename TFunctor>
 UE_AUTORTFM_FORCEINLINE ETransactionResult TransactThenOpen(const TFunctor& Functor)
 {
-	return static_cast<ETransactionResult>(autortfm_transact_then_open(
-		[] (void* Arg) { (*static_cast<const TFunctor*>(Arg))(); },
-		const_cast<void*>(static_cast<const void*>(&Functor))));
+	ETransactionResult Result =
+		static_cast<ETransactionResult>(
+			autortfm_transact_then_open(
+				&AutoRTFMFunctorInvoker<TFunctor>,
+				AutoRTFMLookupInstrumentedFunctorInvoker<TFunctor>(Functor),
+				const_cast<void*>(static_cast<const void*>(&Functor))));
+
+	return Result;
 }
 
 // Run the callback in a transaction like Transact, but abort program
@@ -445,8 +507,9 @@ template<typename TFunctor>
 UE_AUTORTFM_FORCEINLINE void Commit(const TFunctor& Functor)
 {
     autortfm_commit(
-        [] (void* Arg) { (*static_cast<const TFunctor*>(Arg))(); },
-        const_cast<void*>(static_cast<const void*>(&Functor)));
+		&AutoRTFMFunctorInvoker<TFunctor>,
+		AutoRTFMLookupInstrumentedFunctorInvoker<TFunctor>(Functor),
+		const_cast<void*>(static_cast<const void*>(&Functor)));
 }
 
 // End a transaction and discard all effects.
@@ -490,9 +553,11 @@ template<typename TFunctor> UE_AUTORTFM_FORCEINLINE void Open(const TFunctor& Fu
 // Will crash if called outside of a transaction nest.
 template<typename TFunctor> [[nodiscard]] UE_AUTORTFM_FORCEINLINE EContextStatus Close(const TFunctor& Functor)
 {
-    return static_cast<EContextStatus>(autortfm_close(
-        [] (void* Arg) { (*static_cast<const TFunctor*>(Arg))(); },
-        const_cast<void*>(static_cast<const void*>(&Functor))));
+    return static_cast<EContextStatus>(
+		autortfm_close(
+			&AutoRTFMFunctorInvoker<TFunctor>,
+			AutoRTFMLookupInstrumentedFunctorInvoker<TFunctor>(Functor),
+			const_cast<void*>(static_cast<const void*>(&Functor))));
 }
 
 #if UE_AUTORTFM
@@ -500,7 +565,17 @@ template<typename TFunctor> [[nodiscard]] UE_AUTORTFM_FORCEINLINE EContextStatus
 // this just adds the work to the work deferred until the outer nest's commit.
 // If this is called outside a transaction or from an open nest then the work
 // happens immediately.
-UE_AUTORTFM_API void OnCommit(TFunction<void()>&& Work);
+template<typename TFunctor> UE_AUTORTFM_FORCEINLINE void OnCommit(const TFunctor& Work)
+{
+	if (autortfm_is_closed())
+	{
+		ForTheRuntime::OnCommitInternal(Work);
+	}
+	else
+	{
+		Work();
+	}
+}
 #else
 // Have some work happen when this transaction commits. For nested transactions,
 // this just adds the work to the work deferred until the outer nest's commit.
@@ -512,17 +587,35 @@ template<typename TFunctor> UE_AUTORTFM_FORCEINLINE void OnCommit(const TFunctor
 #if UE_AUTORTFM
 // Have some work happen when this transaction aborts. If this is called
 // outside a transaction or from an open nest then the work is ignored.
-UE_AUTORTFM_API void OnAbort(TFunction<void()>&& Work);
+template<typename TFunctor> UE_AUTORTFM_FORCEINLINE void OnAbort(const TFunctor& Work)
+{
+	if (autortfm_is_closed())
+	{
+		ForTheRuntime::OnAbortInternal(Work);
+	}
+}
 
 // Register a handler for transaction abort. Takes a key parameter so that
 // the handler can be unregistered (see PopOnAbortHandler). This is useful
 // for scoped mutations that need an abort handler present unless execution
 // reaches the end of the relevant scope.
-UE_AUTORTFM_API void PushOnAbortHandler(const void* Key, TFunction<void()>&& Work);
+template<typename TFunctor> UE_AUTORTFM_FORCEINLINE void PushOnAbortHandler(const void* Key, const TFunctor& Work)
+{
+	if (autortfm_is_closed())
+	{
+		ForTheRuntime::PushOnAbortHandlerInternal(Key, Work);
+	}
+}
 
 // Unregister all handlers for transaction abort that were previously pushed
 // via PushOnAbortHandler with the given key
-UE_AUTORTFM_API void PopOnAbortHandler(const void* Key);
+UE_AUTORTFM_FORCEINLINE void PopOnAbortHandler(const void* Key)
+{
+	if (autortfm_is_closed())
+	{
+		ForTheRuntime::PopOnAbortHandlerInternal(Key);
+	}
+}
 #else
 // Have some work happen when this transaction aborts. If this is called
 // outside a transaction or from an open nest then the work is ignored.
