Removed template specialization TOptional<TNonNullPtr<OptionalType>> and replaced with the IntrusiveUnsetOptionalState mechanism instead
- so it behaves identically to the generic TOptional<T> and can be used in templates over TOptional<T> that expect consistent behavior regardless of T (e.g. the * operator will return T, the -> operator operates on T, not something inside of T)
- added convenience function T* GetRawPointerOrNull(const TOptional<TNonNullPtr<T>>&) to turn an `TOptional<TNonNullPtr<T>>` back into a nullable T*
- same for TOptional<TNonNullSubclassOf<T>>
- fixed up various use locations to adhere to the new behavior

#rnx

diff --git a/Engine/Source/Developer/UndoHistory/Private/Widgets/ClassLookupReflectionDataProvider.cpp b/Engine/Source/Developer/UndoHistory/Private/Widgets/ClassLookupReflectionDataProvider.cpp
--- a/Engine/Source/Developer/UndoHistory/Private/Widgets/ClassLookupReflectionDataProvider.cpp
+++ b/Engine/Source/Developer/UndoHistory/Private/Widgets/ClassLookupReflectionDataProvider.cpp
@@ -16,7 +16,7 @@ namespace UE::UndoHistory
 	TOptional<FString> FClassLookupReflectionDataProvider::GetClassDisplayName(const FSoftClassPath& ClassPath) const
 	{
 		const TOptional<TNonNullPtr<UClass>> Class = LookUpClass(ClassPath);
-		return Class ? Class->GetName() : TOptional<FString>{};
+		return Class ? (*Class)->GetName() : TOptional<FString>{};
 	}
 
 	TOptional<FPropertyReflectionData> FClassLookupReflectionDataProvider::GetPropertyReflectionData(const FSoftClassPath& ClassPath, FName PropertyName) const
@@ -24,10 +24,10 @@ namespace UE::UndoHistory
 		if (const TOptional<TNonNullPtr<FProperty>> Property = LookUpProperty(ClassPath, PropertyName))
 		{
 			FPropertyReflectionData Result{}; // Value init so if new members are added they have defined values
-			Result.PropertyType = GetPropertyType(*Property);
-			Property->GetCPPMacroType(Result.CppMacroType);
-			Result.TypeName = Property->GetClass()->GetName();
-			Result.PropertyFlags = Property->PropertyFlags;
+			Result.PropertyType = GetPropertyType(**Property);
+			(*Property)->GetCPPMacroType(Result.CppMacroType);
+			Result.TypeName = (*Property)->GetClass()->GetName();
+			Result.PropertyFlags = (*Property)->PropertyFlags;
 			return Result;
 		}
 		return {};
@@ -43,7 +43,7 @@ namespace UE::UndoHistory
 	{
 		if (const TOptional<TNonNullPtr<UClass>> Class = LookUpClass(ClassPath))
 		{
-			FProperty* Property = Class->FindPropertyByName(PropertyName);
+			FProperty* Property = (*Class)->FindPropertyByName(PropertyName);
 			return LIKELY(Property) ? TOptional<TNonNullPtr<FProperty>>{ Property } : TOptional<TNonNullPtr<FProperty>>{};
 		}
 		return {};
diff --git a/Engine/Source/Runtime/Core/Public/Templates/NonNullPointer.h b/Engine/Source/Runtime/Core/Public/Templates/NonNullPointer.h
--- a/Engine/Source/Runtime/Core/Public/Templates/NonNullPointer.h
+++ b/Engine/Source/Runtime/Core/Public/Templates/NonNullPointer.h
@@ -4,6 +4,7 @@
 
 #include "CoreTypes.h"
 #include "Misc/AssertionMacros.h"
+#include "Misc/IntrusiveUnsetOptionalState.h"
 #include "Misc/OptionalFwd.h"
 #include "Templates/Requires.h"
 #include "Templates/UnrealTypeTraits.h"
@@ -138,90 +139,36 @@ public:
 		return Object != nullptr;
 	}
 
-private:
-
-	/** The object we're holding a reference to. */
-	ObjectType* Object;
-
-};
-
-
-/**
- * Specialization of TOptional for TNonNullPtr value types
- */
-template<typename OptionalType>
-struct TOptional<TNonNullPtr<OptionalType>>
-{
-public:
-	/** Construct an OptionaType with a valid value. */
-	TOptional(const TNonNullPtr<OptionalType>& InPointer)
-		: Pointer(InPointer)
-	{
-	}
-
-	/** Construct an OptionalType with no value; i.e. unset */
-	TOptional()
-		: Pointer(nullptr)
-	{
-	}
-
-	/** Construct an OptionalType with an invalid value. */
-	TOptional(FNullOpt)
-		: TOptional()
-	{
-	}
-
-	TOptional& operator=(OptionalType* InPointer)
-	{
-		Pointer = InPointer;
-		return *this;
-	}
-
-	void Reset()
-	{
-		Pointer = nullptr;
-	}
-
-	OptionalType* Emplace(OptionalType* InPointer)
-	{
-		Pointer = InPointer;
-		return InPointer;
-	}
-
-	bool operator==(const TOptional& rhs) const
-	{
-		return Pointer == rhs.Pointer;
-	}
-
-	bool operator!=(const TOptional& rhs) const
+	////////////////////////////////////////////////////
+	// Start - intrusive TOptional<TNonNullPtr> state //
+	////////////////////////////////////////////////////
+	constexpr static bool bHasIntrusiveUnsetOptionalState = true;
+	using IntrusiveUnsetOptionalStateType = TNonNullPtr;
+	FORCEINLINE explicit TNonNullPtr(FIntrusiveUnsetOptionalState)
+		: Object(nullptr)
 	{
-		return Pointer != rhs.Pointer;
 	}
-
-	friend FArchive& operator<<(FArchive& Ar, TOptional& Optional)
+	FORCEINLINE bool operator==(FIntrusiveUnsetOptionalState) const
 	{
-		Ar << Optional.Pointer;
-		return Ar;
+		return Object == nullptr;
 	}
-
-	/** @return true when the value is meaningful; false if calling GetValue() is undefined. */
-	bool IsSet() const { return Pointer != nullptr; }
-	FORCEINLINE explicit operator bool() const { return Pointer != nullptr; }
-
-	/** @return The optional value; undefined when IsSet() returns false. */
-	OptionalType* GetValue() const { checkf(IsSet(), TEXT("It is an error to call GetValue() on an unset TOptional. Please either check IsSet() or use Get(DefaultValue) instead.")); return Pointer; }
-
-	OptionalType* operator->() const { return Pointer; }
-
-	OptionalType& operator*() const { return *Pointer; }
-
-	/** @return The optional value when set; DefaultValue otherwise. */
-	OptionalType* Get(OptionalType* DefaultPointer) const { return IsSet() ? Pointer : DefaultPointer; }
+	//////////////////////////////////////////////////
+	// End - intrusive TOptional<TNonNullPtr> state //
+	//////////////////////////////////////////////////
 
 private:
-	OptionalType* Pointer;
+
+	/** The object we're holding a reference to. */
+	ObjectType* Object;
 };
 
+/** Convenience function to turn an `TOptional<TNonNullPtr<T>>` back into a nullable T* */
+template<typename ObjectType>
+inline ObjectType* GetRawPointerOrNull(const TOptional<TNonNullPtr<ObjectType>>& Optional)
+{
+	return Optional.IsSet() ? Optional->Get() : nullptr;
+}
+
 #if UE_ENABLE_INCLUDE_ORDER_DEPRECATED_IN_5_4
 #include "Templates/EnableIf.h"
 #include "Templates/PointerIsConvertibleFromTo.h"
diff --git a/Engine/Source/Runtime/CoreUObject/Public/Templates/NonNullSubclassOf.h b/Engine/Source/Runtime/CoreUObject/Public/Templates/NonNullSubclassOf.h
--- a/Engine/Source/Runtime/CoreUObject/Public/Templates/NonNullSubclassOf.h
+++ b/Engine/Source/Runtime/CoreUObject/Public/Templates/NonNullSubclassOf.h
@@ -6,6 +6,7 @@
 #include "UObject/Class.h"
 #include "UObject/Field.h"
 #include "SubclassOf.h"
+#include "Misc/IntrusiveUnsetOptionalState.h"
 
 // So we can construct uninitialized TNonNullSubclassOf
 enum class EDefaultConstructNonNullSubclassOf { UnsafeDoNotUse };
@@ -64,102 +65,21 @@ public:
 		Super::operator=(From);
 		return *this;
 	}
-};
-
-
-/**
- * Specialization of TOptional for TNonNullSubclassOf value types
- */
-template<typename T>
-struct TOptional<TNonNullSubclassOf<T>>
-{
-	using ClassType = std::remove_pointer_t<decltype(std::declval<TNonNullSubclassOf<T>>().Get())>;
-
-public:
-	/** Construct an OptionaType with a valid value. */
-	TOptional(const TNonNullSubclassOf<T>& InPointer)
-		: Pointer(InPointer)
-	{
-	}
-
-	/** Construct an TClass with no value; i.e. unset */
-	TOptional()
-		: Pointer(nullptr)
-	{
-	}
 
-	/** Construct an TClass with an invalid value. */
-	TOptional(FNullOpt)
-		: TOptional()
-	{
-	}
-
-	TOptional& operator=(const TOptional& Other)
-	{
-		Pointer = Other.Pointer;
-		return *this;
-	}
-
-	TOptional& operator=(T* InPointer)
-	{
-		Pointer = InPointer;
-		return *this;
-	}
-
-	void Reset()
-	{
-		Pointer = nullptr;
-	}
-
-	T* Emplace(T* InPointer)
-	{
-		Pointer = InPointer;
-		return InPointer;
-	}
-
-	friend bool operator==(const TOptional& lhs, const TOptional& rhs)
-	{
-		return lhs.Pointer == rhs.Pointer;
-	}
-
-	friend bool operator!=(const TOptional& lhs, const TOptional& rhs)
-	{
-		return !(lhs == rhs);
-	}
-
-	friend FArchive& operator<<(FArchive& Ar, TOptional& Optional)
-	{
-		Ar << Optional.Pointer;
-		return Ar;
-	}
-
-	/** @return true when the value is meaningful; false if calling GetValue() is undefined. */
-	bool IsSet() const { return Pointer != nullptr; }
-	FORCEINLINE explicit operator bool() const { return Pointer != nullptr; }
-
-	/** @return The optional value; undefined when IsSet() returns false. */
-	ClassType* GetValue() const
-	{
-		checkf(IsSet(), TEXT("It is an error to call GetValue() on an unset TOptional. Please either check IsSet() or use Get(DefaultValue) instead."));
-		return Pointer;
-	}
-
-	ClassType* operator->() const
-	{
-		return Pointer;
-	}
-
-	ClassType& operator*() const
+	///////////////////////////////////////////////////////////
+	// Start - intrusive TOptional<TNonNullSubclassOf> state //
+	///////////////////////////////////////////////////////////
+	constexpr static bool bHasIntrusiveUnsetOptionalState = true;
+	using IntrusiveUnsetOptionalStateType = TNonNullSubclassOf;
+	FORCEINLINE explicit TNonNullSubclassOf(FIntrusiveUnsetOptionalState)
+		: Super(nullptr)
 	{
-		return *Pointer;
 	}
-
-	/** @return The optional value when set; DefaultValue otherwise. */
-	ClassType* Get(ClassType* DefaultPointer) const
+	FORCEINLINE bool operator==(FIntrusiveUnsetOptionalState) const
 	{
-		return IsSet() ? Pointer : DefaultPointer;
+		return Super::Get() == nullptr;
 	}
-
-private:
-	ClassType* Pointer;
+	/////////////////////////////////////////////////////////
+	// End - intrusive TOptional<TNonNullSubclassOf> state //
+	/////////////////////////////////////////////////////////
 };
