Enabled UE_BINNEDCOMMON_ALLOW_RUNTIME_TWEAKING for IOS

diff --git a/Engine/Source/Runtime/Core/Public/HAL/MallocBinnedCommon.h b/Engine/Source/Runtime/Core/Public/HAL/MallocBinnedCommon.h
--- a/Engine/Source/Runtime/Core/Public/HAL/MallocBinnedCommon.h
+++ b/Engine/Source/Runtime/Core/Public/HAL/MallocBinnedCommon.h
@@ -164,7 +164,10 @@ struct FArenaParams
 
 #define DEFAULT_GMallocBinnedBundleCount 64
 
-#define UE_BINNEDCOMMON_ALLOW_RUNTIME_TWEAKING 0
+#ifndef UE_BINNEDCOMMON_ALLOW_RUNTIME_TWEAKING
+	#define UE_BINNEDCOMMON_ALLOW_RUNTIME_TWEAKING 0
+#endif
+
 #if UE_BINNEDCOMMON_ALLOW_RUNTIME_TWEAKING
 extern CORE_API int32 GMallocBinnedBundleSize;
 extern CORE_API int32 GMallocBinnedBundleCount;
