ParseTokens: Fix bug with ParseTokens[Multiple] when using the EParseTokensOptions::IgnoreCase option with a single upper case delimiter.
When using EParseTokensOptions::IgnoreCase always compare against LowerDelimiter instead of Delimiter.
Improve test coverage for single lower and upper case delimiter characters.

#rnx
#tests FoundationTests

diff --git a/Engine/Source/Runtime/Core/Private/String/ParseTokens.cpp b/Engine/Source/Runtime/Core/Private/String/ParseTokens.cpp
--- a/Engine/Source/Runtime/Core/Private/String/ParseTokens.cpp
+++ b/Engine/Source/Runtime/Core/Private/String/ParseTokens.cpp
@@ -56,7 +56,7 @@ inline static void ParseTokens1Delim1Char(
 			{
 				break;
 			}
-			if (TChar<CharType>::ToLower(*ViewIt) != Delimiter)
+			if (TChar<CharType>::ToLower(*ViewIt) != LowerDelimiter)
 			{
 				++ViewIt;
 				continue;
diff --git a/Engine/Source/Runtime/Core/Tests/String/ParseTokensTest.cpp b/Engine/Source/Runtime/Core/Tests/String/ParseTokensTest.cpp
--- a/Engine/Source/Runtime/Core/Tests/String/ParseTokensTest.cpp
+++ b/Engine/Source/Runtime/Core/Tests/String/ParseTokensTest.cpp
@@ -127,8 +127,12 @@ TEST_CASE("Core::String::ParseTokens", "[Core][String][Smoke]")
 			{TEXTVIEW("A\u2022B,C"),      {TEXT(','), TEXT('\u2022')},      {TEXT("A"), TEXT("B"), TEXT("C")},                      KeepEmpty},
 			{TEXTVIEW("A\u2022B\u0085C"), {TEXT('\u0085'), TEXT('\u2022')}, {TEXT("A"), TEXT("B"), TEXT("C")},                      KeepEmpty},
 
-			{TEXTVIEW("ABC"), {TEXT('b')}, {TEXT("ABC")}, SkipEmpty},
+			{TEXTVIEW("ABC"), {TEXT('b')}, {TEXT("ABC")},          SkipEmpty},
+			{TEXTVIEW("ABC"), {TEXT('B')}, {TEXT("A"), TEXT("C")}, SkipEmpty},
 			{TEXTVIEW("ABC"), {TEXT('b')}, {TEXT("A"), TEXT("C")}, SkipEmpty | IgnoreCase},
+			{TEXTVIEW("abc"), {TEXT('B')}, {TEXT("abc")},          SkipEmpty},
+			{TEXTVIEW("abc"), {TEXT('b')}, {TEXT("a"), TEXT("c")}, SkipEmpty},
+			{TEXTVIEW("abc"), {TEXT('B')}, {TEXT("a"), TEXT("c")}, SkipEmpty | IgnoreCase},
 
 			{TEXTVIEW("AbCdaBcD"), {TEXT('B'), TEXT('d')}, {TEXT("AbC"), TEXT("A"), TEXT("cD")},         SkipEmpty},
 			{TEXTVIEW("AbCdaBcD"), {TEXT('B'), TEXT('d')}, {TEXT("A"), TEXT("C"), TEXT("a"), TEXT("c")}, SkipEmpty | IgnoreCase},
@@ -167,4 +171,4 @@ TEST_CASE("Core::String::ParseTokens", "[Core][String][Smoke]")
 
 } // UE::String
 
-#endif
\ No newline at end of file
+#endif
