[TaskGraph]
 - Fix broken AnyTaskCompleted that didn't actually wait on any of the tasks

diff --git a/Engine/Source/Runtime/Core/Private/Async/TaskGraph.cpp b/Engine/Source/Runtime/Core/Private/Async/TaskGraph.cpp
--- a/Engine/Source/Runtime/Core/Private/Async/TaskGraph.cpp
+++ b/Engine/Source/Runtime/Core/Private/Async/TaskGraph.cpp
@@ -2162,7 +2162,9 @@ FGraphEventRef AnyTaskCompleted(const FGraphEventArray& GraphEvents)
 				{	// the last competed task
 					delete SharedData;
 				}
-			}
+			},
+			TStatId{},
+			GraphEvent
 		);
 	}
 
diff --git a/Engine/Source/Runtime/Core/Tests/Async/TaskGraphTest.cpp b/Engine/Source/Runtime/Core/Tests/Async/TaskGraphTest.cpp
--- a/Engine/Source/Runtime/Core/Tests/Async/TaskGraphTest.cpp
+++ b/Engine/Source/Runtime/Core/Tests/Async/TaskGraphTest.cpp
@@ -1425,12 +1425,13 @@ PRAGMA_ENABLE_DEPRECATION_WARNINGS
 				FFunctionGraphTask::CreateAndDispatchWhenReady([] {}, TStatId{}, Blocker)
 			};
 
+			FGraphEventRef CompletedEvent = AnyTaskCompleted(Tasks);
 			FPlatformProcess::Sleep(0.1f);
-			verify(!AnyTaskCompleted(Tasks)->IsComplete());
+			verify(!CompletedEvent->IsComplete());
 
 			Blocker->DispatchSubsequents();
 
-			AnyTaskCompleted(Tasks)->Wait();
+			CompletedEvent->Wait();
 		}
 
 		{	// doesn't wait for all tasks
