Add new AUTORTFM macros which can be followed by braced blocks to make them work better with the debugger
e.g. UE_AUTORTFM_OPEN2{ ... code ... };

diff --git a/Engine/Source/Runtime/Core/Private/Delegates/DelegateHandle.cpp b/Engine/Source/Runtime/Core/Private/Delegates/DelegateHandle.cpp
--- a/Engine/Source/Runtime/Core/Private/Delegates/DelegateHandle.cpp
+++ b/Engine/Source/Runtime/Core/Private/Delegates/DelegateHandle.cpp
@@ -15,7 +15,7 @@ uint64 FDelegateHandle::GenerateNewID()
 	// Just increment a counter to generate an ID.
 	uint64 Result = 0; // Initialize just to silence static analysis.
 	
-	UE_AUTORTFM_OPEN(
+	UE_AUTORTFM_OPEN2
 	{
 		Result = ++UE::Delegates::Private::GNextID;
 
@@ -25,7 +25,7 @@ uint64 FDelegateHandle::GenerateNewID()
 			// Increment it again - it might not be zero, so don't just assign it to 1.
 			Result = ++UE::Delegates::Private::GNextID;
 		}
-	});
+	};
 
 	return Result;
 }
diff --git a/Engine/Source/Runtime/Core/Private/FramePro/FrameProProfiler.cpp b/Engine/Source/Runtime/Core/Private/FramePro/FrameProProfiler.cpp
--- a/Engine/Source/Runtime/Core/Private/FramePro/FrameProProfiler.cpp
+++ b/Engine/Source/Runtime/Core/Private/FramePro/FrameProProfiler.cpp
@@ -132,27 +132,27 @@ public:
 	virtual void* Alloc(size_t size) override
 	{
 		void* Ptr = nullptr;
-		UE_AUTORTFM_OPEN(
+		UE_AUTORTFM_OPEN2
 		{
 			Ptr = BaseMalloc.Malloc(size, DEFAULT_ALIGNMENT);
-		});
+		};
 
 		// no-op for non-transactional code
-		UE_AUTORTFM_ONABORT(
+		UE_AUTORTFM_ONABORT2(=, this)
 		{
 			// Disable the code analysis warning that complains that Free is being passed
 			// a pointer that may be null. Free explicitly handles this case already.
 			this->Free(Ptr); //-V575
-		});
+		};
 
 		return Ptr;
 	}
 	virtual void Free(void* p) override
 	{
-		UE_AUTORTFM_ONCOMMIT(
+		UE_AUTORTFM_ONCOMMIT2(=, this)
 		{
 			BaseMalloc.Free(p);
-		});
+		};
 	}
 
 private:
@@ -166,18 +166,18 @@ public:
 	static void* Malloc(SIZE_T Count, uint32 Alignment = DEFAULT_ALIGNMENT)
 	{
 		void* Ptr = nullptr;
-		UE_AUTORTFM_OPEN(
+		UE_AUTORTFM_OPEN2
 		{
 			Ptr = FrameProAllocator::Get().GetBaseMalloc()->Malloc(Count, Alignment);
-		});
+		};
 
 		// no-op for non-transactional code
-		UE_AUTORTFM_ONABORT(
+		UE_AUTORTFM_ONABORT2(Ptr)
 		{
 			// Disable the code analysis warning that complains that Free is being passed
 			// a pointer that may be null. Free explicitly handles this case already.
 			FrameProMalloc::Free(Ptr); //-V575
-		});
+		};
 
 		return Ptr;
 	}
@@ -198,10 +198,10 @@ public:
 			if (Original)
 			{
 				SIZE_T OriginalCount = 0;
-				UE_AUTORTFM_OPEN(
+				UE_AUTORTFM_OPEN2
 				{
 					BaseMalloc->GetAllocationSize(Original, OriginalCount);
-				});
+				};
 
 				SIZE_T CopyCount = FGenericPlatformMath::Min(Count, OriginalCount);
 				FMemory::Memcpy(Ptr, Original, CopyCount);
@@ -216,10 +216,10 @@ public:
 
 	static void Free(void* Original)
 	{
-		UE_AUTORTFM_ONCOMMIT(
+		UE_AUTORTFM_ONCOMMIT2(Original)
 		{
 			FrameProAllocator::Get().GetBaseMalloc()->Free(Original);
-		});
+		};
 	}
 };
 
diff --git a/Engine/Source/Runtime/Core/Private/HAL/ThreadingBase.cpp b/Engine/Source/Runtime/Core/Private/HAL/ThreadingBase.cpp
--- a/Engine/Source/Runtime/Core/Private/HAL/ThreadingBase.cpp
+++ b/Engine/Source/Runtime/Core/Private/HAL/ThreadingBase.cpp
@@ -1430,7 +1430,7 @@ FQueuedThread::Run()
 FTlsAutoCleanup* FThreadSingletonInitializer::Get( TFunctionRef<FTlsAutoCleanup*()> CreateInstance, uint32& InOutTlsSlot )
 {
 	uint32 TlsSlot;
-	UE_AUTORTFM_OPEN(
+	UE_AUTORTFM_OPEN2
 	{
 		TlsSlot = (uint32)FPlatformAtomics::AtomicRead_Relaxed((int32*)&InOutTlsSlot);
 		if (TlsSlot == FPlatformTLS::InvalidTlsSlot)
@@ -1448,10 +1448,10 @@ FTlsAutoCleanup* FThreadSingletonInitializer::Get( TFunctionRef<FTlsAutoCleanup*
 				TlsSlot = ThisTlsSlot;
 			}
 		}
-	});
+	};
 
 	FTlsAutoCleanup* ThreadSingleton = nullptr;
-	UE_AUTORTFM_OPEN(
+	UE_AUTORTFM_OPEN2
 	{
 		ThreadSingleton = (FTlsAutoCleanup*)FPlatformTLS::GetTlsValue( TlsSlot );
 		if( !ThreadSingleton )
@@ -1461,7 +1461,7 @@ FTlsAutoCleanup* FThreadSingletonInitializer::Get( TFunctionRef<FTlsAutoCleanup*
 			ThreadSingleton->Register();
 			FPlatformTLS::SetTlsValue( TlsSlot, ThreadSingleton );
 		}
-	});
+	};
 	return ThreadSingleton;
 }
 
diff --git a/Engine/Source/Runtime/Core/Private/Internationalization/TextCache.cpp b/Engine/Source/Runtime/Core/Private/Internationalization/TextCache.cpp
--- a/Engine/Source/Runtime/Core/Private/Internationalization/TextCache.cpp
+++ b/Engine/Source/Runtime/Core/Private/Internationalization/TextCache.cpp
@@ -32,7 +32,7 @@ FText FTextCache::FindOrCache(const TCHAR* InTextLiteral, const FTextId& InTextI
 	{
 		FText* ReturnFoundText = nullptr;
 
-		UE_AUTORTFM_OPEN(
+		UE_AUTORTFM_OPEN2
 		{
 
 			FScopeLock Lock(&CachedTextCS);
@@ -47,7 +47,7 @@ FText FTextCache::FindOrCache(const TCHAR* InTextLiteral, const FTextId& InTextI
 					ReturnFoundText = FoundText;
 				}
 			}
-		});
+		};
 
 		if (ReturnFoundText)
 		{
@@ -59,12 +59,12 @@ FText FTextCache::FindOrCache(const TCHAR* InTextLiteral, const FTextId& InTextI
 	FText NewText = FText(InTextLiteral, InTextId.GetNamespace(), InTextId.GetKey(), ETextFlag::Immutable);
 
 	// ... and add it to the cache
-	UE_AUTORTFM_OPEN(
+	UE_AUTORTFM_OPEN2
 	{
 		FScopeLock Lock(&CachedTextCS);
 
 		CachedText.Emplace(InTextId, NewText);
-	});
+	};
 
 	return NewText;
 }
diff --git a/Engine/Source/Runtime/Core/Private/Internationalization/TextHistory.cpp b/Engine/Source/Runtime/Core/Private/Internationalization/TextHistory.cpp
--- a/Engine/Source/Runtime/Core/Private/Internationalization/TextHistory.cpp
+++ b/Engine/Source/Runtime/Core/Private/Internationalization/TextHistory.cpp
@@ -875,9 +875,10 @@ void FTextHistory_Base::UpdateDisplayString()
 
 	// Create a temp to hold the old value in case we abort, in which case we assign out of the OPEN so the old value will be preserved
 	FTextConstDisplayStringPtr NewLocalizedString;
-	UE_AUTORTFM_OPEN({
+	UE_AUTORTFM_OPEN2
+	{
 		NewLocalizedString = FTextLocalizationManager::Get().GetDisplayString(TextId.GetNamespace(), TextId.GetKey(), SourceString.IsEmpty() ? nullptr : &SourceString);
-	});
+	};
 
 	LocalizedString = NewLocalizedString;
 }
diff --git a/Engine/Source/Runtime/Core/Private/Internationalization/TextKey.cpp b/Engine/Source/Runtime/Core/Private/Internationalization/TextKey.cpp
--- a/Engine/Source/Runtime/Core/Private/Internationalization/TextKey.cpp
+++ b/Engine/Source/Runtime/Core/Private/Internationalization/TextKey.cpp
@@ -27,10 +27,10 @@ public:
 		check(!InStr.IsEmpty());
 		// Open around adding this in a cache, if we abort just leak the value in the cache
 		// as the cache takes ownership
-		UE_AUTORTFM_OPEN(
+		UE_AUTORTFM_OPEN2
 		{
 			FindOrAddImpl(FKeyData(InStr), OutTextKey);
-		});
+		};
 	}
 
 	void FindOrAdd(FStringView InStr, const uint32 InStrHash, FTextKey& OutTextKey)
@@ -38,10 +38,10 @@ public:
 		check(!InStr.IsEmpty());
 		// Open around adding this in a cache, if we abort just leak the value in the cache
 		// as the cache takes ownership
-		UE_AUTORTFM_OPEN(
+		UE_AUTORTFM_OPEN2
 		{
 			FindOrAddImpl(FKeyData(InStr, InStrHash), OutTextKey);
-		});
+		};
 	}
 
 	void Shrink()
diff --git a/Engine/Source/Runtime/Core/Private/Serialization/CustomVersion.cpp b/Engine/Source/Runtime/Core/Private/Serialization/CustomVersion.cpp
--- a/Engine/Source/Runtime/Core/Private/Serialization/CustomVersion.cpp
+++ b/Engine/Source/Runtime/Core/Private/Serialization/CustomVersion.cpp
@@ -417,9 +417,9 @@ void FCustomVersionContainer::SetVersionUsingRegistry(FGuid CustomKey, ESetCusto
 		}
 
 		TOptional<FCustomVersion> RegisteredVersion;
-		UE_AUTORTFM_OPEN({
+		UE_AUTORTFM_OPEN2{
 			RegisteredVersion = FCurrentCustomVersions::Get(CustomKey);
-		});
+		};
 		checkf(RegisteredVersion, TEXT("Attempted to set a version that is not registered"));
 
 		Found->Version      = RegisteredVersion->Version;
@@ -428,9 +428,9 @@ void FCustomVersionContainer::SetVersionUsingRegistry(FGuid CustomKey, ESetCusto
 	else
 	{
 		TOptional<FCustomVersion> RegisteredVersion;
-		UE_AUTORTFM_OPEN({
+		UE_AUTORTFM_OPEN2{
 			RegisteredVersion = FCurrentCustomVersions::Get(CustomKey);
-		});
+		};
 
 		checkf(RegisteredVersion, TEXT("Attempted to set a version that is not registered"));
 
diff --git a/Engine/Source/Runtime/Core/Private/UObject/UnrealNames.cpp b/Engine/Source/Runtime/Core/Private/UObject/UnrealNames.cpp
--- a/Engine/Source/Runtime/Core/Private/UObject/UnrealNames.cpp
+++ b/Engine/Source/Runtime/Core/Private/UObject/UnrealNames.cpp
@@ -1164,13 +1164,13 @@ public:
 	FNameEntryId Find(const FNameValue<Sensitivity>& Value) const
 	{
 		FNameEntryId Result;
-		UE_AUTORTFM_OPEN(
+		UE_AUTORTFM_OPEN2
 		{
 			FRWScopeLock _(Lock, FRWScopeLockType::SLT_ReadOnly);
 
 			FNameSlot& Slot = Probe(Value);
 			Result = Slot.GetId();
-		});
+		};
 		return Result;
 	}
 
@@ -1695,7 +1695,7 @@ FNameEntryId FNamePool::StoreWithNumber(FNameEntryIds StringParts, int32 NumberP
 {
 	FNameEntryId Result;
 
-	AutoRTFM::Open([&]
+	UE_AUTORTFM_OPEN2
 	{
 #if WITH_CASE_PRESERVING_NAME
 		// Look for an exact match with the right casing first
@@ -1720,7 +1720,7 @@ FNameEntryId FNamePool::StoreWithNumber(FNameEntryIds StringParts, int32 NumberP
 #else
 		Result = ComparisonId;
 #endif
-	});
+	};
 
 	return Result;
 }
@@ -2017,12 +2017,12 @@ static FNamePool& GetNamePool()
 
 	FNamePool* Singleton = nullptr;
 	
-	UE_AUTORTFM_OPEN(
+	UE_AUTORTFM_OPEN2
 	{
 		Singleton = new (NamePoolData) FNamePool;
 		bNamePoolInitialized = true;
 		LLM(FLowLevelMemTracker::Get().FinishInitialise());
-	});
+	};
 	
 	return *Singleton;
 }
@@ -2887,9 +2887,10 @@ struct FNameHelper
 		if (FindType == FNAME_Add)
 		{
 			FNameEntryId DisplayId;
-			UE_AUTORTFM_OPEN({
+			UE_AUTORTFM_OPEN2
+			{
 				DisplayId = Pool.StoreWithNumber(BaseIds, InternalNumber);
-			});
+			};
 			return FinalConstruct(FNameEntryIds{ ResolveComparisonId(DisplayId), DisplayId });
 		}
 		else
@@ -2995,9 +2996,10 @@ private:
 	static FName MakeInternal(FNameStringView View, EFindName FindType, int32 InternalNumber)
 	{
 		FNameEntryIds Ids;
-		UE_AUTORTFM_OPEN({
+		UE_AUTORTFM_OPEN2
+		{
 			Ids = FindOrStoreString(View, FindType);
-		});
+		};
 #if UE_FNAME_OUTLINE_NUMBER
 		if (FindType == FNAME_Find && !Ids.DisplayId)
 		{
@@ -3055,7 +3057,7 @@ private:
 	{
 		FNameEntryIds Result{};
 
-		UE_AUTORTFM_OPEN(
+		UE_AUTORTFM_OPEN2
 		{
 			if (View.Len >= NAME_SIZE)
 			{
@@ -3089,7 +3091,7 @@ private:
 					Result = FNameEntryIds{ ResolveComparisonId(DisplayId), DisplayId };
 				}
 			}
-		});
+		};
 
 		return Result;
 	}
diff --git a/Engine/Source/Runtime/Core/Private/Unix/UnixPlatformMisc.cpp b/Engine/Source/Runtime/Core/Private/Unix/UnixPlatformMisc.cpp
--- a/Engine/Source/Runtime/Core/Private/Unix/UnixPlatformMisc.cpp
+++ b/Engine/Source/Runtime/Core/Private/Unix/UnixPlatformMisc.cpp
@@ -1397,7 +1397,7 @@ bool FUnixPlatformMisc::IsDebuggerPresent()
 
 	int StatusFile = -1;
 	bool bDebugging = false;
-	UE_AUTORTFM_OPEN(
+	UE_AUTORTFM_OPEN2
 	{
 		StatusFile = open("/proc/self/status", O_RDONLY);
 
@@ -1425,7 +1425,7 @@ bool FUnixPlatformMisc::IsDebuggerPresent()
 
 			close(StatusFile);
 		}
-	});
+	};
 	return bDebugging;
 }
 #endif // !UE_BUILD_SHIPPING
diff --git a/Engine/Source/Runtime/Core/Public/AutoRTFM/AutoRTFM.h b/Engine/Source/Runtime/Core/Public/AutoRTFM/AutoRTFM.h
--- a/Engine/Source/Runtime/Core/Public/AutoRTFM/AutoRTFM.h
+++ b/Engine/Source/Runtime/Core/Public/AutoRTFM/AutoRTFM.h
@@ -670,8 +670,8 @@ UE_AUTORTFM_FORCEINLINE void DidFree(void* Ptr)
 // A collection of power-user functions that are reserved for use by the AutoRTFM runtime only.
 namespace ForTheRuntime
 {
-	[[deprecated("This macro is deprecated. Use UE_AUTORTFM_ONABORT instead!")]] UE_AUTORTFM_FORCEINLINE void DeprecatedUseOnAbortMacro() {}
-	[[deprecated("This macro is deprecated. Use UE_AUTORTFM_ONCOMMIT instead!")]] UE_AUTORTFM_FORCEINLINE void DeprecatedUseOnCommitMacro() {}
+	[[deprecated("This macro is deprecated. Use UE_AUTORTFM_ONABORT2 instead!")]] UE_AUTORTFM_FORCEINLINE void DeprecatedUseOnAbortMacro() {}
+	[[deprecated("This macro is deprecated. Use UE_AUTORTFM_ONCOMMIT2 instead!")]] UE_AUTORTFM_FORCEINLINE void DeprecatedUseOnCommitMacro() {}
 
 	// An enum to represent the various ways we want to enable/disable the AutoRTFM runtime.
 	enum EAutoRTFMEnabledState
@@ -877,6 +877,42 @@ UE_AUTORTFM_FORCEINLINE void CheckConsistencyAssumingNoRaces() { return ForTheRu
 // Macro-based variants so we completely compile away when not in use, even in debug builds
 #if UE_AUTORTFM
 
+namespace AutoRTFM::Private
+{
+	struct FOpenHelper
+	{
+		template<typename FunctorType>
+		void operator+(FunctorType F)
+		{
+			AutoRTFM::Open(MoveTemp(F));
+		}
+	};
+	struct FOnAbortHelper
+	{
+		template<typename FunctorType>
+		void operator+(FunctorType F)
+		{
+			AutoRTFM::OnAbort(MoveTemp(F));
+		}
+	};
+	struct FOnCommitHelper
+	{
+		template<typename FunctorType>
+		void operator+(FunctorType F)
+		{
+			AutoRTFM::OnCommit(MoveTemp(F));
+		}
+	};
+	struct FTransactHelper
+	{
+		template<typename FunctorType>
+		void operator+(FunctorType F)
+		{
+			AutoRTFM::Transact(MoveTemp(F));
+		}
+	};
+} // namespace AutoRTFM::Private
+
 #if defined(__clang__) && __has_warning("-Wdeprecated-this-capture")
 #define UE_AUTORTFM_BEGIN_DISABLE_WARNINGS _Pragma("clang diagnostic push") _Pragma("clang diagnostic ignored \"-Wdeprecated-this-capture\"")
 #define UE_AUTORTFM_END_DISABLE_WARNINGS _Pragma("clang diagnostic pop")
@@ -885,32 +921,56 @@ UE_AUTORTFM_FORCEINLINE void CheckConsistencyAssumingNoRaces() { return ForTheRu
 #define UE_AUTORTFM_END_DISABLE_WARNINGS
 #endif
 
+// Older macros where the code is passed as a macro argument. These should be phased out as they make debugging more difficult
 #define UE_AUTORTFM_OPEN_IMPL(...) AutoRTFM::Open([&]() { __VA_ARGS__ })
 #define UE_AUTORTFM_ONABORT_IMPL(...) UE_AUTORTFM_BEGIN_DISABLE_WARNINGS AutoRTFM::OnAbort([=]() { __VA_ARGS__ }) UE_AUTORTFM_END_DISABLE_WARNINGS
 #define UE_AUTORTFM_ONCOMMIT_IMPL(...) UE_AUTORTFM_BEGIN_DISABLE_WARNINGS AutoRTFM::OnCommit([=]() { __VA_ARGS__ }) UE_AUTORTFM_END_DISABLE_WARNINGS
 #define UE_AUTORTFM_TRANSACT_IMPL(...) AutoRTFM::Transact([&]() { __VA_ARGS__ })
+
+#define UE_AUTORTFM_OPEN_IMPL2 AutoRTFM::Private::FOpenHelper{} + [&]()
+#define UE_AUTORTFM_ONABORT_IMPL2(...) AutoRTFM::Private::FOnAbortHelper{} + [__VA_ARGS__]()
+#define UE_AUTORTFM_ONCOMMIT_IMPL2(...) AutoRTFM::Private::FOnCommitHelper{} + [__VA_ARGS__]()
+#define UE_AUTORTFM_TRANSACT_IMPL2 AutoRTFM::Private::FTransactHelper{} + [&]() { __VA_ARGS__ })
 #else
+
+// Older macros where the code is passed as a macro argument. These should be phased out as they make debugging more difficult
 #define UE_AUTORTFM_OPEN_IMPL(...) do { __VA_ARGS__ } while (false)
 #define UE_AUTORTFM_ONABORT_IMPL(...) do { /* do nothing */ } while (false)
 #define UE_AUTORTFM_ONCOMMIT_IMPL(...) do { __VA_ARGS__ } while (false)
 #define UE_AUTORTFM_TRANSACT_IMPL(...) do { __VA_ARGS__ } while (false)
+
+// Do nothing, these should be followed by blocks that should be either executed or not executed
+#define UE_AUTORTFM_OPEN_IMPL2 
+#define UE_AUTORTFM_ONABORT_IMPL2(...) while (false)
+#define UE_AUTORTFM_ONCOMMIT_IMPL2(...)
+#define UE_AUTORTFM_TRANSACT_IMPL2
 #endif
 
 // Runs a block of code in the open, non-transactionally. Anything performed in the open will not be undone if a transaction fails.
 #define UE_AUTORTFM_OPEN(...) UE_AUTORTFM_OPEN_IMPL(__VA_ARGS__)
+// This new version is used like UE_AUTORTFM_OPEN2 { ... code ... }; 
+#define UE_AUTORTFM_OPEN2 UE_AUTORTFM_OPEN_IMPL2
 
 // Runs a block of code if a transaction aborts.
 // In non-transactional code paths the block of code will not be executed at all.
 // This captures any used variables from the parent function by-value.
 #define UE_AUTORTFM_ONABORT(...) UE_AUTORTFM_ONABORT_IMPL(__VA_ARGS__)
+// In the new version of the macro, the macro arguments are the capture specification for the lambda
+// Used like UE_AUTORTFM_ABORT2(=) { ... code ... };
+#define UE_AUTORTFM_ONABORT2(...) UE_AUTORTFM_ONABORT_IMPL2(__VA_ARGS__)
 
 // Runs a block of code if a transaction commits successfully.
 // In non-transactional code paths the block of code will be executed immediately.
 // This captures any used variables from the parent function by-value.
 #define UE_AUTORTFM_ONCOMMIT(...) UE_AUTORTFM_ONCOMMIT_IMPL(__VA_ARGS__)
+// In the new version of the macro, the macro arguments are the capture specification for the lambda
+// Used like UE_AUTORTFM_COMMIT2(=) { ... code ... };
+#define UE_AUTORTFM_ONCOMMIT2(...) UE_AUTORTFM_ONCOMMIT_IMPL2(__VA_ARGS__)
 
 // Runs a block of code in the closed, transactionally, within a new transaction.
 #define UE_AUTORTFM_TRANSACT(...) UE_AUTORTFM_TRANSACT_IMPL(__VA_ARGS__)
+// New version is used like a block: UE_AUTORTFM_TRANSACT2 { ... code ... };
+#define UE_AUTORTFM_TRANSACT2 UE_AUTORTFM_TRANSACT_IMPL2
 
 // Deprecated. Use UE_AUTORTFM_ONABORT instead.
 #define UE_AUTORTFM_OPENABORT(...) UE_AUTORTFM_ONABORT(AutoRTFM::ForTheRuntime::DeprecatedUseOnAbortMacro(); __VA_ARGS__)
diff --git a/Engine/Source/Runtime/Core/Public/Delegates/Delegate.h b/Engine/Source/Runtime/Core/Public/Delegates/Delegate.h
--- a/Engine/Source/Runtime/Core/Public/Delegates/Delegate.h
+++ b/Engine/Source/Runtime/Core/Public/Delegates/Delegate.h
@@ -372,11 +372,11 @@ class DynamicMulticastDelegateClassName : public TBaseDynamicMulticastDelegate<F
 			static FName Get()
 			{
 				FName* Result = nullptr;
-				UE_AUTORTFM_OPEN(
+				UE_AUTORTFM_OPEN2
 				{
 					static FName StaticResult = Create();
 					Result = &StaticResult;
-				});
+				};
 				return *Result;
 			}
 
diff --git a/Engine/Source/Runtime/Core/Public/HAL/FMemory.inl b/Engine/Source/Runtime/Core/Public/HAL/FMemory.inl
--- a/Engine/Source/Runtime/Core/Public/HAL/FMemory.inl
+++ b/Engine/Source/Runtime/Core/Public/HAL/FMemory.inl
@@ -24,7 +24,7 @@ FMEMORY_INLINE_FUNCTION_DECORATOR void* FMemory::Malloc(SIZE_T Count, uint32 Ali
 	// we don't end up keeping track of the writes to the allocator's internal data structures.
 	// This is because allocators are already transactional - malloc can be rolled back by
 	// calling free.
-	UE_AUTORTFM_OPEN(
+	UE_AUTORTFM_OPEN2
 	{
 		if (!FMEMORY_INLINE_GMalloc)
 		{
@@ -38,7 +38,7 @@ FMEMORY_INLINE_FUNCTION_DECORATOR void* FMemory::Malloc(SIZE_T Count, uint32 Ali
 		}
 		// optional tracking of every allocation
 		LLM_IF_ENABLED(FLowLevelMemTracker::Get().OnLowLevelAlloc(ELLMTracker::Default, Ptr, Count, ELLMTag::Untagged, ELLMAllocType::FMalloc));
-	});
+	};
 
 	// AutoRTFM: This is a no-op for non-transactional code.
 	// For transactional code, this defers a call to Free if the transaction aborts,
@@ -131,7 +131,7 @@ FMEMORY_INLINE_FUNCTION_DECORATOR void FMemory::Free(void* Original)
 	// AutoRTFM: For transactional code, in order to support the transaction 
 	// aborting and needing to 'roll back' the Free, we defer the actual
 	// free until commit time.
-	UE_AUTORTFM_ONCOMMIT(
+	UE_AUTORTFM_ONCOMMIT2(=)
 	{
 		// optional tracking of every allocation
 		LLM_IF_ENABLED(FLowLevelMemTracker::Get().OnLowLevelFree(ELLMTracker::Default, Original, ELLMAllocType::FMalloc));
@@ -146,13 +146,13 @@ FMEMORY_INLINE_FUNCTION_DECORATOR void FMemory::Free(void* Original)
 		FMEMORY_INLINE_GMalloc->Free(Original);
 
 		AutoRTFM::DidFree(Original);
-	});
+	};
 }
 
 FMEMORY_INLINE_FUNCTION_DECORATOR SIZE_T FMemory::GetAllocSize(void* Original)
 {
 	SIZE_T Result;
-	UE_AUTORTFM_OPEN(
+	UE_AUTORTFM_OPEN2
 	{
 		if (!FMEMORY_INLINE_GMalloc)
 		{
@@ -163,7 +163,7 @@ FMEMORY_INLINE_FUNCTION_DECORATOR SIZE_T FMemory::GetAllocSize(void* Original)
 			SIZE_T Size = 0;
 			Result = FMEMORY_INLINE_GMalloc->GetAllocationSize(Original, Size) ? Size : 0;
 		}
-	});
+	};
 
 	return Result;
 }
@@ -178,7 +178,7 @@ FMEMORY_INLINE_FUNCTION_DECORATOR void* FMemory::MallocZeroed(SIZE_T Count, uint
 	// we don't end up keeping track of the writes to the allocator's internal data structures.
 	// This is because allocators are already transactional - malloc can be rolled back by
 	// calling free.
-	UE_AUTORTFM_OPEN(
+	UE_AUTORTFM_OPEN2
 		{
 			if (!FMEMORY_INLINE_GMalloc)
 			{
@@ -192,7 +192,7 @@ FMEMORY_INLINE_FUNCTION_DECORATOR void* FMemory::MallocZeroed(SIZE_T Count, uint
 			}
 	// optional tracking of every allocation
 	LLM_IF_ENABLED(FLowLevelMemTracker::Get().OnLowLevelAlloc(ELLMTracker::Default, Ptr, Count, ELLMTag::Untagged, ELLMAllocType::FMalloc));
-		});
+		};
 
 	// AutoRTFM: This is a no-op for non-transactional code.
 	// For transactional code, this defers a call to Free if the transaction aborts,
@@ -210,7 +210,7 @@ FMEMORY_INLINE_FUNCTION_DECORATOR void* FMemory::MallocZeroed(SIZE_T Count, uint
 FMEMORY_INLINE_FUNCTION_DECORATOR SIZE_T FMemory::QuantizeSize(SIZE_T Count, uint32 Alignment)
 {
 	SIZE_T Result;
-	UE_AUTORTFM_OPEN(
+	UE_AUTORTFM_OPEN2
 	{
 		if (!FMEMORY_INLINE_GMalloc)
 		{
@@ -220,7 +220,7 @@ FMEMORY_INLINE_FUNCTION_DECORATOR SIZE_T FMemory::QuantizeSize(SIZE_T Count, uin
 		{
 			Result = FMEMORY_INLINE_GMalloc->QuantizeSize(Count, Alignment);
 		}
-	});
+	};
 
 	return Result;
 }
diff --git a/Engine/Source/Runtime/Core/Public/Misc/MTAccessDetector.h b/Engine/Source/Runtime/Core/Public/Misc/MTAccessDetector.h
--- a/Engine/Source/Runtime/Core/Public/Misc/MTAccessDetector.h
+++ b/Engine/Source/Runtime/Core/Public/Misc/MTAccessDetector.h
@@ -365,7 +365,7 @@ private:
 	{
 		const SIZE_T StackTraceSize = 65536;
 		ANSICHAR StackTrace[StackTraceSize] = { 0 };
-		UE_AUTORTFM_OPEN({ FPlatformStackWalk::StackWalkAndDump(StackTrace, StackTraceSize, 0); });
+		UE_AUTORTFM_OPEN2{ FPlatformStackWalk::StackWalkAndDump(StackTrace, StackTraceSize, 0); };
 		return ANSI_TO_TCHAR(StackTrace);
 	}
 
@@ -384,7 +384,7 @@ private:
 
 		const SIZE_T StackTraceSize = 65536;
 		ANSICHAR StackTrace[StackTraceSize] = { 0 };
-		UE_AUTORTFM_OPEN({ FPlatformStackWalk::ThreadStackWalkAndDump(StackTrace, StackTraceSize, 0, ThreadId); });
+		UE_AUTORTFM_OPEN2{ FPlatformStackWalk::ThreadStackWalkAndDump(StackTrace, StackTraceSize, 0, ThreadId); };
 		return ANSI_TO_TCHAR(StackTrace);
 	}
 
diff --git a/Engine/Source/Runtime/Core/Public/Misc/ScopeRWLock.h b/Engine/Source/Runtime/Core/Public/Misc/ScopeRWLock.h
--- a/Engine/Source/Runtime/Core/Public/Misc/ScopeRWLock.h
+++ b/Engine/Source/Runtime/Core/Public/Misc/ScopeRWLock.h
@@ -15,13 +15,13 @@ public:
 	UE_NODISCARD_CTOR explicit FReadScopeLock(FRWLock& InLock)
 		: Lock(InLock)
 	{
-		UE_AUTORTFM_OPEN({ Lock.ReadLock(); });
+		UE_AUTORTFM_OPEN2{ Lock.ReadLock(); };
 		AutoRTFM::PushOnAbortHandler(this, [this](){ this->Lock.ReadUnlock(); });
 	}
 
 	~FReadScopeLock()
 	{
-		UE_AUTORTFM_OPEN({ Lock.ReadUnlock(); });
+		UE_AUTORTFM_OPEN2{ Lock.ReadUnlock(); };
 		AutoRTFM::PopOnAbortHandler(this);
 	}
 
diff --git a/Engine/Source/Runtime/Core/Public/ProfilingDebugging/CsvProfiler.h b/Engine/Source/Runtime/Core/Public/ProfilingDebugging/CsvProfiler.h
--- a/Engine/Source/Runtime/Core/Public/ProfilingDebugging/CsvProfiler.h
+++ b/Engine/Source/Runtime/Core/Public/ProfilingDebugging/CsvProfiler.h
@@ -539,14 +539,14 @@ public:
 		: StatName(InStatName)
 		, CategoryIndex(InCategoryIndex)
 	{
-		UE_AUTORTFM_OPEN({ FCsvProfiler::BeginStat(StatName, CategoryIndex, InNamedEventName); });
+		UE_AUTORTFM_OPEN2{ FCsvProfiler::BeginStat(StatName, CategoryIndex, InNamedEventName); };
 		AutoRTFM::PushOnAbortHandler(this, [SN = StatName, CI = CategoryIndex](){ FCsvProfiler::EndStat(SN, CI); });
 	}
 
 	~FScopedCsvStat()
 	{
 		AutoRTFM::PopOnAbortHandler(this);
-		UE_AUTORTFM_OPEN({ FCsvProfiler::EndStat(StatName, CategoryIndex); });
+		UE_AUTORTFM_OPEN2{ FCsvProfiler::EndStat(StatName, CategoryIndex); };
 	}
 	const char * StatName;
 	uint32 CategoryIndex;
@@ -558,14 +558,14 @@ public:
 	FScopedCsvStatExclusive(const char * InStatName, const char* InNamedEventName = nullptr)
 		: StatName(InStatName)
 	{
-		UE_AUTORTFM_OPEN({ FCsvProfiler::BeginExclusiveStat(StatName, InNamedEventName); });
+		UE_AUTORTFM_OPEN2{ FCsvProfiler::BeginExclusiveStat(StatName, InNamedEventName); };
 		AutoRTFM::PushOnAbortHandler(this, [SN = StatName](){ FCsvProfiler::EndExclusiveStat(SN); });
 	}
 
 	~FScopedCsvStatExclusive()
 	{
 		AutoRTFM::PopOnAbortHandler(this);
-		UE_AUTORTFM_OPEN({ FCsvProfiler::EndExclusiveStat(StatName); });
+		UE_AUTORTFM_OPEN2{ FCsvProfiler::EndExclusiveStat(StatName); };
 	}
 	const char * StatName;
 };
diff --git a/Engine/Source/Runtime/Core/Public/Templates/RefCounting.h b/Engine/Source/Runtime/Core/Public/Templates/RefCounting.h
--- a/Engine/Source/Runtime/Core/Public/Templates/RefCounting.h
+++ b/Engine/Source/Runtime/Core/Public/Templates/RefCounting.h
@@ -254,11 +254,11 @@ public:
 
 			uint32 Count = 0;
 
-			UE_AUTORTFM_OPEN(
+			UE_AUTORTFM_OPEN2
 			{
 				// This reference count may be accessed by multiple threads
 				Count = RefCount.load(std::memory_order_relaxed);
-			});
+			};
 
 			return Count;
 		}
diff --git a/Engine/Source/Runtime/Core/Public/Templates/SharedPointerInternals.h b/Engine/Source/Runtime/Core/Public/Templates/SharedPointerInternals.h
--- a/Engine/Source/Runtime/Core/Public/Templates/SharedPointerInternals.h
+++ b/Engine/Source/Runtime/Core/Public/Templates/SharedPointerInternals.h
@@ -81,12 +81,12 @@ namespace SharedPointerInternals
 
 				int32 Count = 0;
 
-				UE_AUTORTFM_OPEN(
+				UE_AUTORTFM_OPEN2
 					{
 						// This reference count may be accessed by multiple threads
 						Count = SharedReferenceCount.load(std::memory_order_relaxed);
 					}
-				);
+				;
 
 				return Count;
 			}
@@ -111,24 +111,24 @@ namespace SharedPointerInternals
 				// in response to the increment, so there's nothing to order with.
 
 #if defined(_MSC_VER) && (defined(_M_X64) || defined(_M_IX86))
-				UE_AUTORTFM_OPEN(
+				UE_AUTORTFM_OPEN2
 				{
 					// We do a regular SC increment here because it maps to an _InterlockedIncrement (lock inc).
 					// The codegen for a relaxed fetch_add is actually much worse under MSVC (lock xadd).
 					++SharedReferenceCount;
-				});
+				};
 #else
-				UE_AUTORTFM_OPEN(
+				UE_AUTORTFM_OPEN2
 				{
 					SharedReferenceCount.fetch_add(1, std::memory_order_relaxed);
-				});
+				};
 #endif
 
 				// If the transaction would abort, we need to undo adding the shared reference.
-				UE_AUTORTFM_ONABORT(
+				UE_AUTORTFM_ONABORT2(this)
 				{
 					ReleaseSharedReference();
-				});
+				};
 			}
 			else
 			{
@@ -147,7 +147,7 @@ namespace SharedPointerInternals
 			{
 				bool bSucceeded = false;
 
-				UE_AUTORTFM_OPEN(
+				UE_AUTORTFM_OPEN2
 				{
 					// See AddSharedReference for the same reasons that std::memory_order_relaxed is used in this function.
 
@@ -182,15 +182,15 @@ namespace SharedPointerInternals
 							break;
 						}
 					}
-				});
+				};
 
 				// If we succeeded in taking a shared reference count, we need to undo that on an abort.
 				if (bSucceeded)
 				{
-					UE_AUTORTFM_ONABORT(
+					UE_AUTORTFM_ONABORT2(this)
 					{
 						ReleaseSharedReference();
-					});
+					};
 				}
 
 				return bSucceeded;
@@ -213,7 +213,7 @@ namespace SharedPointerInternals
 		{
 			if constexpr (Mode == ESPMode::ThreadSafe)
 			{
-				AutoRTFM::OnCommit([this]
+				UE_AUTORTFM_ONCOMMIT2(this)
 				{
 					// std::memory_order_acq_rel is used here so that, if we do end up executing the destructor, it's not possible
 					// for side effects from executing the destructor end up being visible before we've determined that the shared
@@ -230,7 +230,7 @@ namespace SharedPointerInternals
 						// reference count reaches zero, this object will be deleted.
 						ReleaseWeakReference();
 					}
-				});
+				};
 			}
 			else
 			{
@@ -256,24 +256,24 @@ namespace SharedPointerInternals
 				// See AddSharedReference for the same reasons that std::memory_order_relaxed is used in this function.
 
 #if defined(_MSC_VER) && (defined(_M_X64) || defined(_M_IX86))
-				UE_AUTORTFM_OPEN(
+				UE_AUTORTFM_OPEN2
 					{
 						// We do a regular SC increment here because it maps to an _InterlockedIncrement (lock inc).
 						// The codegen for a relaxed fetch_add is actually much worse under MSVC (lock xadd).
 						++WeakReferenceCount;
-					});
+					};
 #else
-				UE_AUTORTFM_OPEN(
+				UE_AUTORTFM_OPEN2
 					{
 						WeakReferenceCount.fetch_add(1, std::memory_order_relaxed);
-					});
+					};
 #endif
 
 				// If the transaction would abort, we need to undo adding the reference.
-				UE_AUTORTFM_ONABORT(
+				UE_AUTORTFM_ONABORT2(this)
 					{
 						ReleaseWeakReference();
-					});
+					};
 			}
 			else
 			{
@@ -286,7 +286,7 @@ namespace SharedPointerInternals
 		{
 			if constexpr (Mode == ESPMode::ThreadSafe)
 			{
-				AutoRTFM::OnCommit([this]
+				UE_AUTORTFM_ONCOMMIT2(this)
 					{
 						// See ReleaseSharedReference for the same reasons that std::memory_order_acq_rel is used in this function.
 
@@ -302,7 +302,7 @@ namespace SharedPointerInternals
 							delete this;
 #endif
 						}
-					});
+					};
 			}
 			else
 			{
diff --git a/Engine/Source/Runtime/CoreUObject/Private/Misc/GCObjectReferencer.cpp b/Engine/Source/Runtime/CoreUObject/Private/Misc/GCObjectReferencer.cpp
--- a/Engine/Source/Runtime/CoreUObject/Private/Misc/GCObjectReferencer.cpp
+++ b/Engine/Source/Runtime/CoreUObject/Private/Misc/GCObjectReferencer.cpp
@@ -272,16 +272,16 @@ void FGCObject::RegisterGCObject()
 		StaticInit();
 
 		// Add this instance to the referencer's list
-		UE_AUTORTFM_OPEN(
+		UE_AUTORTFM_OPEN2
 		{
 			GGCObjectReferencer->AddObject(this);
-		});
+		};
 
 		// But if we abort, we don't want to leave a dangling reference!
-		UE_AUTORTFM_ONABORT(
+		UE_AUTORTFM_ONABORT2(this)
 		{
 			GGCObjectReferencer->RemoveObject(this);
-		});
+		};
 
 		bReferenceAdded = true;
 	}
diff --git a/Engine/Source/Runtime/CoreUObject/Private/Serialization/AsyncLoadingThread.h b/Engine/Source/Runtime/CoreUObject/Private/Serialization/AsyncLoadingThread.h
--- a/Engine/Source/Runtime/CoreUObject/Private/Serialization/AsyncLoadingThread.h
+++ b/Engine/Source/Runtime/CoreUObject/Private/Serialization/AsyncLoadingThread.h
@@ -475,7 +475,7 @@ public:
 	static bool GetIsInAsyncLoadingTick()
 	{
 		bool Result = false;
-		UE_AUTORTFM_OPEN({ Result = !!AsyncLoadingTickCounter.GetValue(); });
+		UE_AUTORTFM_OPEN2{ Result = !!AsyncLoadingTickCounter.GetValue(); };
 		return Result;
 	}
 
diff --git a/Engine/Source/Runtime/CoreUObject/Private/UObject/Class.cpp b/Engine/Source/Runtime/CoreUObject/Private/UObject/Class.cpp
--- a/Engine/Source/Runtime/CoreUObject/Private/UObject/Class.cpp
+++ b/Engine/Source/Runtime/CoreUObject/Private/UObject/Class.cpp
@@ -6501,7 +6501,7 @@ UFunction* UClass::FindFunctionByName(FName InName, EIncludeSuperFlag::Type Incl
 
 	UFunction* Result = nullptr;
 
-	UE_AUTORTFM_OPEN(
+	UE_AUTORTFM_OPEN2
 	{
 		UClass* SuperClass = GetSuperClass();
 		if (IncludeSuper == EIncludeSuperFlag::ExcludeSuper || ( Interfaces.Num() == 0 && SuperClass == nullptr ) )
@@ -6578,7 +6578,7 @@ UFunction* UClass::FindFunctionByName(FName InName, EIncludeSuperFlag::Type Incl
 				}
 			}
 		}
-	});
+	};
 
 	return Result;
 }
diff --git a/Engine/Source/Runtime/CoreUObject/Private/UObject/GarbageCollection.cpp b/Engine/Source/Runtime/CoreUObject/Private/UObject/GarbageCollection.cpp
--- a/Engine/Source/Runtime/CoreUObject/Private/UObject/GarbageCollection.cpp
+++ b/Engine/Source/Runtime/CoreUObject/Private/UObject/GarbageCollection.cpp
@@ -5259,7 +5259,6 @@ bool GatherUnreachableObjects(UE::GC::EGatherOptions Options, double TimeLimit /
 		}
 	}, (GGatherUnreachableObjectsState.NumWorkerThreads() == 1) ? EParallelForFlags::ForceSingleThread : EParallelForFlags::None);
 
-	// Need to grab the stats now as Finish() will reset the global state but we still want to include Finish() execution time when logging
 	const int32 NumGathered = GGatherUnreachableObjectsState.NumGathered();
 	const int32 NumScanned = GGatherUnreachableObjectsState.NumScanned();
 	const int32 NumThreads = GGatherUnreachableObjectsState.NumWorkerThreads();
@@ -6391,10 +6390,10 @@ void UObjectBase::MarkAsReachable() const
 {
 	// It is safe to perform mark as reachable in the open - the worst case is that we'll mark an object reachable that
 	// should/would be destroyed, and so in the next GC iteration it will be destroyed instead of in this iteration.
-	UE_AUTORTFM_OPEN(
+	UE_AUTORTFM_OPEN2
 	{
 		::MarkAsReachable<false>(this);
-	});
+	};
 }
 
 void UObjectBase::AddRef() const
diff --git a/Engine/Source/Runtime/CoreUObject/Private/UObject/Obj.cpp b/Engine/Source/Runtime/CoreUObject/Private/UObject/Obj.cpp
--- a/Engine/Source/Runtime/CoreUObject/Private/UObject/Obj.cpp
+++ b/Engine/Source/Runtime/CoreUObject/Private/UObject/Obj.cpp
@@ -892,9 +892,9 @@ FClassExclusionData GDedicatedClientExclusionList;
 bool UObject::NeedsLoadForServer() const
 {
 	bool Ret = false;
-	UE_AUTORTFM_OPEN({
+	UE_AUTORTFM_OPEN2{
 		Ret = !GDedicatedServerExclusionList.IsExcluded(GetClass());
-	});
+	};
 	return Ret;
 }
 
@@ -907,9 +907,9 @@ void UObject::UpdateClassesExcludedFromDedicatedServer(const TArray<FString>& In
 bool UObject::NeedsLoadForClient() const
 {
 	bool Ret = false;
-	UE_AUTORTFM_OPEN({
+	UE_AUTORTFM_OPEN2{
 		Ret = !GDedicatedClientExclusionList.IsExcluded(GetClass());
-	});
+	};
 	return Ret;
 }
 
diff --git a/Engine/Source/Runtime/CoreUObject/Private/UObject/ObjectHandleTracking.cpp b/Engine/Source/Runtime/CoreUObject/Private/UObject/ObjectHandleTracking.cpp
--- a/Engine/Source/Runtime/CoreUObject/Private/UObject/ObjectHandleTracking.cpp
+++ b/Engine/Source/Runtime/CoreUObject/Private/UObject/ObjectHandleTracking.cpp
@@ -94,19 +94,19 @@ namespace UE::CoreUObject
 
 			void OnHandleRead(TArrayView<const UObject* const> Objects)
 			{
-				UE_AUTORTFM_OPEN(
+				UE_AUTORTFM_OPEN2
 				{
 					FLightweightReadScopeLock _(HandleLock);
 					for (auto&& Pair : ReadHandleCallbacks)
 					{
 						Pair.Value(Objects);
 					}
-				});
+				};
 			}
 
 			void OnHandleRead(const UObject* Object)
 			{
-				UE_AUTORTFM_OPEN(
+				UE_AUTORTFM_OPEN2
 				{
 					FLightweightReadScopeLock _(HandleLock);
 					TArrayView<const UObject* const> Objects(&Object, 1);
@@ -114,62 +114,62 @@ namespace UE::CoreUObject
 					{
 						Pair.Value(Objects);
 					}
-				});
+				};
 			}
 
 			void OnClassReferenceResolved(const FObjectRef& ObjectRef, UPackage* Package, UClass* Class)
 			{
-				UE_AUTORTFM_OPEN(
+				UE_AUTORTFM_OPEN2
 				{
 					FLightweightReadScopeLock _(HandleLock);
 					for (auto&& Pair : ClassResolvedCallbacks)
 					{
 						Pair.Value(ObjectRef, Package, Class);
 					}
-				});
+				};
 			}
 
 			void OnReferenceResolved(const FObjectRef& ObjectRef, UPackage* Package, UObject* Object)
 			{
-				UE_AUTORTFM_OPEN(
+				UE_AUTORTFM_OPEN2
 				{
 					FLightweightReadScopeLock _(HandleLock);
 					for (auto&& Pair : HandleResolvedCallbacks)
 					{
 						Pair.Value(ObjectRef, Package, Object);
 					}
-				});
+				};
 			}
 
 			void OnReferenceLoaded(const FObjectRef& ObjectRef, UPackage* Package, UObject* Object)
 			{
-				UE_AUTORTFM_OPEN(
+				UE_AUTORTFM_OPEN2
 				{
 					FLightweightReadScopeLock _(HandleLock);
 					for (auto&& Pair : HandleLoadedCallbacks)
 					{
 						Pair.Value(ObjectRef, Package, Object);
 					}
-				});
+				};
 			}
 
 			FObjectHandleTrackingCallbackId AddObjectHandleReadCallback(FObjectHandleReadFunc Func)
 			{
 				FObjectHandleTrackingCallbackId Result;
-				UE_AUTORTFM_OPEN(
+				UE_AUTORTFM_OPEN2
 				{
 					HandleReadCallbackQuantity.fetch_add(1, std::memory_order_release);
 					FLightweightWriteScopeLock _(HandleLock);
 					NextHandleId++;
 					ReadHandleCallbacks.Add({ NextHandleId, Func });
 					Result = FObjectHandleTrackingCallbackId{ NextHandleId };
-				});
+				};
 				return Result;
 			}
 
 			void RemoveObjectHandleReadCallback(FObjectHandleTrackingCallbackId Handle)
 			{
-				UE_AUTORTFM_OPEN(
+				UE_AUTORTFM_OPEN2
 				{
 					HandleReadCallbackQuantity.fetch_sub(1, std::memory_order_release);
 					FLightweightWriteScopeLock _(HandleLock);
@@ -181,25 +181,25 @@ namespace UE::CoreUObject
 							ReadHandleCallbacks.RemoveAt(i);
 						}
 					}
-				});
+				};
 			}
 
 			FObjectHandleTrackingCallbackId AddObjectHandleClassResolvedCallback(FObjectHandleClassResolvedFunc Func)
 			{
 				FObjectHandleTrackingCallbackId Result;
-				UE_AUTORTFM_OPEN(
+				UE_AUTORTFM_OPEN2
 				{
 					FLightweightWriteScopeLock _(HandleLock);
 					NextHandleId++;
 					ClassResolvedCallbacks.Add({ NextHandleId, Func });
 					Result = FObjectHandleTrackingCallbackId{ NextHandleId };
-				});
+				};
 				return Result;
 			}
 
 			void RemoveObjectHandleClassResolvedCallback(FObjectHandleTrackingCallbackId Handle)
 			{
-				UE_AUTORTFM_OPEN(
+				UE_AUTORTFM_OPEN2
 				{
 					FLightweightWriteScopeLock _(HandleLock);
 					for (int32 i = ClassResolvedCallbacks.Num() - 1; i >= 0; --i)
@@ -210,25 +210,25 @@ namespace UE::CoreUObject
 							ClassResolvedCallbacks.RemoveAt(i);
 						}
 					}
-				});
+				};
 			}
 
 			FObjectHandleTrackingCallbackId AddObjectHandleReferenceResolvedCallback(FObjectHandleReferenceResolvedFunc Func)
 			{
 				FObjectHandleTrackingCallbackId Result;
-				UE_AUTORTFM_OPEN(
+				UE_AUTORTFM_OPEN2
 				{
 					FLightweightWriteScopeLock _(HandleLock);
 					NextHandleId++;
 					HandleResolvedCallbacks.Add({ NextHandleId, Func });
 					Result = FObjectHandleTrackingCallbackId{ NextHandleId };
-				});
+				};
 				return Result;
 			}
 
 			void RemoveObjectHandleReferenceResolvedCallback(FObjectHandleTrackingCallbackId Handle)
 			{
-				UE_AUTORTFM_OPEN(
+				UE_AUTORTFM_OPEN2
 				{
 					FLightweightWriteScopeLock _(HandleLock);
 					for (int32 i = HandleResolvedCallbacks.Num() - 1; i >= 0; --i)
@@ -239,25 +239,25 @@ namespace UE::CoreUObject
 							HandleResolvedCallbacks.RemoveAt(i);
 						}
 					}
-				});
+				};
 			}
 
 			FObjectHandleTrackingCallbackId AddObjectHandleReferenceLoadedCallback(FObjectHandleReferenceLoadedFunc Func)
 			{
 				FObjectHandleTrackingCallbackId Result;
-				UE_AUTORTFM_OPEN(
+				UE_AUTORTFM_OPEN2
 				{
 					FLightweightWriteScopeLock _(HandleLock);
 					NextHandleId++;
 					HandleLoadedCallbacks.Add({ NextHandleId, Func });
 					Result = FObjectHandleTrackingCallbackId{ NextHandleId };
-				});
+				};
 				return Result;
 			}
 
 			void RemoveObjectHandleReferenceLoadedCallback(FObjectHandleTrackingCallbackId Handle)
 			{
-				UE_AUTORTFM_OPEN(
+				UE_AUTORTFM_OPEN2
 				{
 					FLightweightWriteScopeLock _(HandleLock);
 					for (int32 i = HandleLoadedCallbacks.Num() - 1; i >= 0; --i)
@@ -268,7 +268,7 @@ namespace UE::CoreUObject
 							HandleLoadedCallbacks.RemoveAt(i);
 						}
 					}
-				});
+				};
 			}
 
 		private:
diff --git a/Engine/Source/Runtime/CoreUObject/Private/UObject/ScriptCore.cpp b/Engine/Source/Runtime/CoreUObject/Private/UObject/ScriptCore.cpp
--- a/Engine/Source/Runtime/CoreUObject/Private/UObject/ScriptCore.cpp
+++ b/Engine/Source/Runtime/CoreUObject/Private/UObject/ScriptCore.cpp
@@ -189,7 +189,7 @@ FBlueprintContext* FBlueprintContextGetThreadSingletonImpl()
 FBlueprintContext* FBlueprintContext::GetThreadSingleton()
 {
 	FBlueprintContext* Result;
-	UE_AUTORTFM_OPEN({ Result = FBlueprintContextGetThreadSingletonImpl(); });
+	UE_AUTORTFM_OPEN2{ Result = FBlueprintContextGetThreadSingletonImpl(); };
 	return Result;
 }
 
@@ -3093,7 +3093,7 @@ void UObject::ProcessContextOpcode( FFrame& Stack, RESULT_DECL, bool bCanFailSil
 
 		if (!bCanFailSilently)
 		{
-			UE_AUTORTFM_OPEN(
+			UE_AUTORTFM_OPEN2
 			{
 				if (NewContext && !IsValid(NewContext))
 				{
@@ -3130,7 +3130,7 @@ void UObject::ProcessContextOpcode( FFrame& Stack, RESULT_DECL, bool bCanFailSil
 					);
 					FBlueprintCoreDelegates::ThrowScriptException(this, Stack, ExceptionInfo);
 				}
-			});
+			};
 		}
 
 		const CodeSkipSizeType wSkip = Stack.ReadCodeSkipCount(); // Code offset for NULL expressions. Code += sizeof(CodeSkipSizeType)
@@ -4071,7 +4071,7 @@ DEFINE_FUNCTION(UObject::execAutoRtfmTransact)
 
 					case EAutoRtfmStopTransactMode::AbortingExitAndAbortParent:
 						// abort this transaction and also abort the parent
-						UE_AUTORTFM_OPEN({ bAbortParentOnCommit = true; });
+						UE_AUTORTFM_OPEN2{ bAbortParentOnCommit = true; };
 						AutoRTFM::AbortTransaction();
 						break;
 					}
diff --git a/Engine/Source/Runtime/CoreUObject/Private/UObject/UObjectAllocator.cpp b/Engine/Source/Runtime/CoreUObject/Private/UObject/UObjectAllocator.cpp
--- a/Engine/Source/Runtime/CoreUObject/Private/UObject/UObjectAllocator.cpp
+++ b/Engine/Source/Runtime/CoreUObject/Private/UObject/UObjectAllocator.cpp
@@ -43,7 +43,7 @@ UObjectBase* FUObjectAllocator::AllocateUObject(int32 Size, int32 Alignment, boo
 {
 	void* Result = nullptr;
 	// we want to perform this allocation uninstrumented so the GC can clean this up if the transaction is aborted
-	UE_AUTORTFM_OPEN(
+	UE_AUTORTFM_OPEN2
 	{
 		if (bAllowPermanent && !GPersistentAllocatorIsDisabled)
 		{
@@ -54,7 +54,7 @@ UObjectBase* FUObjectAllocator::AllocateUObject(int32 Size, int32 Alignment, boo
 		{
 			Result = FMemory::Malloc(Size, Alignment);
 		}
-	});
+	};
 
 	return (UObjectBase*)Result;
 }
diff --git a/Engine/Source/Runtime/CoreUObject/Private/UObject/UObjectArchetype.cpp b/Engine/Source/Runtime/CoreUObject/Private/UObject/UObjectArchetype.cpp
--- a/Engine/Source/Runtime/CoreUObject/Private/UObject/UObjectArchetype.cpp
+++ b/Engine/Source/Runtime/CoreUObject/Private/UObject/UObjectArchetype.cpp
@@ -223,16 +223,16 @@ UObject* GetArchetypeImpl(const UObject* InObject, const FObjectArchetypeHelper:
 		if (Archetype && !(InObject->GetOuter() && InObject->GetOuter()->HasAnyFlags(RF_NeedLoad)))
 		{
 			ArchetypeIndex = GUObjectArray.ObjectToIndex(Archetype);
-			UE_AUTORTFM_OPEN(
+			UE_AUTORTFM_OPEN2
 			{
 				ArchetypeAnnotation.AddAnnotation(InObject, FArchetypeInfo{ ArchetypeIndex, GUObjectArray.AllocateSerialNumber(ArchetypeIndex) });
-			});
+			};
 
 			// If we abort we need to remove this annotation
-			UE_AUTORTFM_ONABORT(
+			UE_AUTORTFM_ONABORT2(=)
 			{
 				ArchetypeAnnotation.RemoveAnnotation(InObject);
-			});
+			};
 		}
 	}
 	else
diff --git a/Engine/Source/Runtime/CoreUObject/Private/UObject/UObjectArray.cpp b/Engine/Source/Runtime/CoreUObject/Private/UObject/UObjectArray.cpp
--- a/Engine/Source/Runtime/CoreUObject/Private/UObject/UObjectArray.cpp
+++ b/Engine/Source/Runtime/CoreUObject/Private/UObject/UObjectArray.cpp
@@ -74,13 +74,13 @@ void FUObjectItem::CreateStatID() const
 	FMemory::Memcpy(StoragePtr, ConversionData.Get(), NumStorageChars * sizeof(PROFILER_CHAR));
 
 	// delay the delete of the StatIDStringStorage swap until after the transaction
-	UE_AUTORTFM_ONCOMMIT(
+	UE_AUTORTFM_ONCOMMIT2(=)
 	{
 		if (FPlatformAtomics::InterlockedCompareExchangePointer((void**)&StatIDStringStorage, StoragePtr, nullptr) != nullptr)
 		{
 			delete[] StoragePtr;
 		}
-	});
+	};
 
 	StatID = TStatId(StatIDStringStorage);
 #endif
@@ -460,7 +460,7 @@ int32 FUObjectArray::AllocateSerialNumber(int32 Index)
 	if (!SerialNumber)
 	{
 		// Open around PrimarySerialNumber as if we fail/abort a transaction we dont need to undo this, simply allow it to grow for the next use
-		UE_AUTORTFM_OPEN({
+		UE_AUTORTFM_OPEN2{
 			SerialNumber = PrimarySerialNumber.Increment();
 			UE_CLOG(SerialNumber <= START_SERIAL_NUMBER, LogUObjectArray, Fatal, TEXT("UObject serial numbers overflowed (trying to allocate serial number %d)."), SerialNumber);
 			int32 ValueWas = FPlatformAtomics::InterlockedCompareExchange((int32*)SerialNumberPtr, SerialNumber, 0);
@@ -469,7 +469,7 @@ int32 FUObjectArray::AllocateSerialNumber(int32 Index)
 				// someone else go it first, use their value
 				SerialNumber = ValueWas;
 			}
-		});
+		};
 	}
 	checkSlow(SerialNumber > START_SERIAL_NUMBER);
 	return SerialNumber;
diff --git a/Engine/Source/Runtime/CoreUObject/Private/UObject/UObjectGlobals.cpp b/Engine/Source/Runtime/CoreUObject/Private/UObject/UObjectGlobals.cpp
--- a/Engine/Source/Runtime/CoreUObject/Private/UObject/UObjectGlobals.cpp
+++ b/Engine/Source/Runtime/CoreUObject/Private/UObject/UObjectGlobals.cpp
@@ -2724,7 +2724,7 @@ namespace NameReuse
 		{
 			FName Result;
 
-			UE_AUTORTFM_OPEN(
+			UE_AUTORTFM_OPEN2
 			{
 				Lock.ReadLock();
 				FNameRangeEntry* Entry = Find(BaseId);
@@ -2760,7 +2760,7 @@ namespace NameReuse
 						Result = Entry->AllocateName(Parent, BaseId, BaseName);
 					}
 				}
-			});
+			};
 
 			return Result;
 		}
@@ -2801,7 +2801,7 @@ namespace NameReuse
 
 		FName ReturnName;
 
-		UE_AUTORTFM_OPEN(
+		UE_AUTORTFM_OPEN2
 		{
 			FNameEntryId BaseId = BaseName.GetDisplayIndex();
 			ReturnName = GRecentNameCache.Find(Parent, BaseId, BaseName);
@@ -2816,7 +2816,7 @@ namespace NameReuse
 				// Store this name for reuse 
 				GRecentNameCache.Store(Parent, BaseId, ReturnName);
 			}
-		});
+		};
 
 		return ReturnName;
 	}
@@ -2864,7 +2864,7 @@ FName MakeUniqueObjectName(UObject* Parent, const UClass* Class, FName InBaseNam
 			else
 			{
 				int32 NameNumber = 0;
-				UE_AUTORTFM_OPEN(
+				UE_AUTORTFM_OPEN2
 				{
 					if (Parent && (Parent != ANY_PACKAGE_DEPRECATED) && !(Options & EUniqueObjectNameOptions::GloballyUnique))
 					{
@@ -2897,7 +2897,7 @@ FName MakeUniqueObjectName(UObject* Parent, const UClass* Class, FName InBaseNam
 					{
 						NameNumber = ++Class->ClassUnique;
 					}
-				});
+				};
 				TestName = FName(BaseName, NameNumber);
 			}
 
@@ -3615,11 +3615,11 @@ UObject* StaticAllocateObject
 	{
 		// perform the UObjectBase construction in the open - we expect the GC to invoke
 		// the destructor in the case of transaction abort
-		UE_AUTORTFM_OPEN(
+		UE_AUTORTFM_OPEN2
 		{
 			FMemory::Memzero((void *)Obj, TotalSize);
 			new ((void *)Obj) UObjectBase(const_cast<UClass*>(InClass), InFlags|RF_NeedInitialization, InternalSetFlags, InOuter, InName, OldIndex, OldSerialNumber);
-		});
+		};
 	}
 	else
 	{
@@ -4096,7 +4096,7 @@ void FObjectInitializer::PostConstructInit()
 	// need a fence to guarantee that the cleared flag is only visible to other threads
 	// after the other initialization-related writes
 	std::atomic_thread_fence(std::memory_order_release);
-	UE_AUTORTFM_OPEN({ Obj->ClearFlags(RF_NeedInitialization); });
+	UE_AUTORTFM_OPEN2{ Obj->ClearFlags(RF_NeedInitialization); };
 
 	// clear the object pointer so we can guard against running this function again
 	Obj = nullptr;
diff --git a/Engine/Source/Runtime/CoreUObject/Private/UObject/UObjectHash.cpp b/Engine/Source/Runtime/CoreUObject/Private/UObject/UObjectHash.cpp
--- a/Engine/Source/Runtime/CoreUObject/Private/UObject/UObjectHash.cpp
+++ b/Engine/Source/Runtime/CoreUObject/Private/UObject/UObjectHash.cpp
@@ -524,7 +524,7 @@ public:
 		if (!(IsGarbageCollectingAndLockingUObjectHashTables() && IsInGameThread()))
 		{
 			Tables = &InTables;
-			UE_AUTORTFM_OPEN({ InTables.Lock(); });
+			UE_AUTORTFM_OPEN2{ InTables.Lock(); };
 			AutoRTFM::PushOnAbortHandler(this, [this](){ if (this->Tables){ this->Tables->Unlock(); }});
 		}
 		else
@@ -540,7 +540,7 @@ public:
 #if THREADSAFE_UOBJECTS
 		if (Tables)
 		{
-			UE_AUTORTFM_OPEN({ Tables->Unlock(); });
+			UE_AUTORTFM_OPEN2{ Tables->Unlock(); };
 			AutoRTFM::PopOnAbortHandler(this);
 		}
 #endif
@@ -851,7 +851,7 @@ UObject* StaticFindObjectFastInternal(const UClass* ObjectClass, const UObject*
 	UObject* Result = nullptr;
 
 	// Transactionally, a static find operation has to occur in the open as it touches shared state.
-	UE_AUTORTFM_OPEN(
+	UE_AUTORTFM_OPEN2
 	{
 		INC_DWORD_STAT(STAT_FindObjectFast);
 
@@ -860,7 +860,7 @@ UObject* StaticFindObjectFastInternal(const UClass* ObjectClass, const UObject*
 		// If they specified an outer use that during the hashing
 		FUObjectHashTables& ThreadHash = FUObjectHashTables::Get();
 		Result = StaticFindObjectFastInternalThreadSafe(ThreadHash, ObjectClass, ObjectPackage, ObjectName, bExactClass, bAnyPackage, ExcludeFlags, ExclusiveInternalFlags);
-	});
+	};
 
 	return Result;
 }
@@ -870,7 +870,7 @@ UObject* StaticFindObjectFastInternal(const UClass* ObjectClass, const UObject*
 	UObject* Result = nullptr;
 
 	// Transactionally, a static find operation has to occur in the open as it touches shared state.
-	UE_AUTORTFM_OPEN(
+	UE_AUTORTFM_OPEN2
 	{
 		INC_DWORD_STAT(STAT_FindObjectFast);
 
@@ -879,7 +879,7 @@ UObject* StaticFindObjectFastInternal(const UClass* ObjectClass, const UObject*
 		// If they specified an outer use that during the hashing
 		FUObjectHashTables& ThreadHash = FUObjectHashTables::Get();
 		Result = StaticFindObjectFastInternalThreadSafe(ThreadHash, ObjectClass, ObjectPackage, ObjectName, bExactClass, /*bAnyPackage =*/ false, ExcludeFlags, ExclusiveInternalFlags);
-	});
+	};
 
 	return Result;
 }
diff --git a/Engine/Source/Runtime/CoreUObject/Public/UObject/DynamicallyTypedValue.h b/Engine/Source/Runtime/CoreUObject/Public/UObject/DynamicallyTypedValue.h
--- a/Engine/Source/Runtime/CoreUObject/Public/UObject/DynamicallyTypedValue.h
+++ b/Engine/Source/Runtime/CoreUObject/Public/UObject/DynamicallyTypedValue.h
@@ -62,9 +62,9 @@ namespace UE
 	{
 		static COREUOBJECT_API FDynamicallyTypedValueType& NullType();
 
-		FDynamicallyTypedValue() { UE_AUTORTFM_OPEN({ InitializeToNull(); }); }
-		FDynamicallyTypedValue(const FDynamicallyTypedValue& Copyee) { UE_AUTORTFM_OPEN({ InitializeFromCopy(Copyee); }); }
-		FDynamicallyTypedValue(FDynamicallyTypedValue&& Movee) { UE_AUTORTFM_OPEN({ InitializeFromMove(MoveTemp(Movee)); }); }
+		FDynamicallyTypedValue() { UE_AUTORTFM_OPEN2{ InitializeToNull(); }; }
+		FDynamicallyTypedValue(const FDynamicallyTypedValue& Copyee) { UE_AUTORTFM_OPEN2{ InitializeFromCopy(Copyee); }; }
+		FDynamicallyTypedValue(FDynamicallyTypedValue&& Movee) { UE_AUTORTFM_OPEN2{ InitializeFromMove(MoveTemp(Movee)); }; }
 
 		~FDynamicallyTypedValue() { Deinit(); }
 
diff --git a/Engine/Source/Runtime/CoreUObject/Public/UObject/UObjectAnnotation.h b/Engine/Source/Runtime/CoreUObject/Public/UObject/UObjectAnnotation.h
--- a/Engine/Source/Runtime/CoreUObject/Public/UObject/UObjectAnnotation.h
+++ b/Engine/Source/Runtime/CoreUObject/Public/UObject/UObjectAnnotation.h
@@ -90,7 +90,7 @@ private:
 			{
 				// We hit this from many locations, for now lets open around adding this, and if we happen to be in a transcation
 				// and abort after we have added let us use the UObject to remove it from the annotation map
-				UE_AUTORTFM_OPEN(
+				UE_AUTORTFM_OPEN2
 				{
 					FScopeLock AnnotationMapLock(&AnnotationMapCritical);
 					bWasEmpty = (AnnotationMap.Num() == 0);
@@ -107,12 +107,12 @@ private:
 					{
 						AnnotationMap.Add(Object, LocalAnnotation);
 					}
-				});
+				};
 
-				UE_AUTORTFM_ONABORT(
+				UE_AUTORTFM_ONABORT2(this, Object)
 				{
 					RemoveAnnotation(Object);
-				});
+				};
 			}
 
 			if (bWasEmpty)
@@ -122,15 +122,15 @@ private:
 				if (bAutoRemove)
 #endif
 				{
-					UE_AUTORTFM_OPEN(
+					UE_AUTORTFM_OPEN2
 					{
 						GUObjectArray.AddUObjectDeleteListener(this);
-					});
+					};
 
-					UE_AUTORTFM_ONABORT(
+					UE_AUTORTFM_ONABORT2(this)
 					{
 						GUObjectArray.RemoveUObjectDeleteListener(this);
-					});
+					};
 				}
 			}
 		}
@@ -193,7 +193,7 @@ public:
 	 */
 	void RemoveAnnotation(const UObjectBase *Object)
 	{
-		AutoRTFM::OnCommit([this, Object]
+		UE_AUTORTFM_ONCOMMIT2(this, Object)
 		{
 			check(Object);
 			bool bHadElements = false;
@@ -217,7 +217,7 @@ public:
 					GUObjectArray.RemoveUObjectDeleteListener(this);
 				}
 			}
-		});
+		};
 	}
 	/**
 	 * Removes all annotation from the annotation list. 
@@ -256,7 +256,7 @@ public:
 	{
 		check(Object);
 		TAnnotation Result;
-		UE_AUTORTFM_OPEN(
+		UE_AUTORTFM_OPEN2
 		{
 			FScopeLock AnnotationMapLock(&AnnotationMapCritical);
 			if (Object != AnnotationCacheKey)
@@ -273,7 +273,7 @@ public:
 				}
 			}
 			Result = AnnotationCacheValue;
-		});
+		};
 		return Result;
 	}
 
@@ -875,7 +875,7 @@ public:
 
 		TAnnotation Result = TAnnotation();
 
-		UE_AUTORTFM_OPEN(
+		UE_AUTORTFM_OPEN2
 		{
 			FRWScopeLock AnnotationArrayLock(AnnotationArrayCritical, SLT_ReadOnly);
 
@@ -890,7 +890,7 @@ public:
 					Result = Chunk.Items[WithinChunkIndex];
 				}
 			}
-		});
+		};
 
 		return Result;
 	}
diff --git a/Engine/Source/Runtime/CoreUObject/Public/UObject/UObjectArray.h b/Engine/Source/Runtime/CoreUObject/Public/UObject/UObjectArray.h
--- a/Engine/Source/Runtime/CoreUObject/Public/UObject/UObjectArray.h
+++ b/Engine/Source/Runtime/CoreUObject/Public/UObject/UObjectArray.h
@@ -154,7 +154,7 @@ public:
 	{
 		checkf((int32(FlagToClear) & ~int32(EInternalObjectFlags_AllFlags)) == 0, TEXT("%d is not a valid internal flag value"), int32(FlagToClear));
 		bool Result = false;
-		UE_AUTORTFM_OPEN(
+		UE_AUTORTFM_OPEN2
 		{
 			FlagToClear &= ~EInternalObjectFlags_ReachabilityFlags; // reachability flags can only be cleared by GC through *_ForGC functions
 			FlagToClear &= ~EInternalObjectFlags::RefCounted; // refcounted flag is internal and must only be cleared internally by AddRef/ReleaseRef.
@@ -166,7 +166,7 @@ public:
 			{
 				Result = AtomicallyClearFlag_ForGC(FlagToClear);
 			}
-		});
+		};
 
 		return Result;
 	}
@@ -191,7 +191,7 @@ public:
 	{
 		checkf((int32(FlagToSet) & ~int32(EInternalObjectFlags_AllFlags)) == 0, TEXT("%d is not a valid internal flag value"), int32(FlagToSet));
 		bool Result = false;
-		UE_AUTORTFM_OPEN(
+		UE_AUTORTFM_OPEN2
 		{
 			FlagToSet &= ~EInternalObjectFlags_ReachabilityFlags; // reachability flags can only be cleared by GC through *_ForGC functions
 			FlagToSet &= ~EInternalObjectFlags::RefCounted; // refcounted flag is internal and must only be set by AddRef/ReleaseRef.
@@ -203,7 +203,7 @@ public:
 			{
 				Result = AtomicallySetFlag_ForGC(FlagToSet);
 			}
-		});
+		};
 		return Result;
 	}
 
@@ -292,19 +292,19 @@ public:
 
 	void AddRef()
 	{
-		UE_AUTORTFM_OPEN(
+		UE_AUTORTFM_OPEN2
 		{
 			FPlatformAtomics::InterlockedIncrement(&RefCount);
 			if ((GetFlags() & EInternalObjectFlags::RefCounted) != EInternalObjectFlags::RefCounted)
 			{
 				SetRootFlags(EInternalObjectFlags::RefCounted);
 			}
-		});
+		};
 	}
 
 	void ReleaseRef()
 	{
-		UE_AUTORTFM_OPEN(
+		UE_AUTORTFM_OPEN2
 		{
 			// This alone is not thread-safe as we may race with AddRef and in that case we don't want ClearRootFlags to apply.
 			// We fix this by validating that the refcount is still 0 while inside the root locks in ClearRootFlags.
@@ -314,7 +314,7 @@ public:
 			{
 				ClearRootFlags(EInternalObjectFlags::RefCounted);
 			}
-		});
+		};
 	}
 
 #if STATS || ENABLE_STATNAMEDEVENTS_UOBJECT
diff --git a/Engine/Source/Runtime/Engine/Private/Actor.cpp b/Engine/Source/Runtime/Engine/Private/Actor.cpp
--- a/Engine/Source/Runtime/Engine/Private/Actor.cpp
+++ b/Engine/Source/Runtime/Engine/Private/Actor.cpp
@@ -226,7 +226,7 @@ void AActor::InitializeDefaults()
 #if (CSV_PROFILER && !UE_BUILD_SHIPPING)
 	// Increment actor class count
 	// Update our ActorClassName count after the transaction finishes
-	UE_AUTORTFM_ONCOMMIT(
+	UE_AUTORTFM_ONCOMMIT2(=)
 	{
 		if (!HasAnyFlags(RF_ArchetypeObject | RF_ClassDefaultObject))
 		{
@@ -238,7 +238,7 @@ void AActor::InitializeDefaults()
 			CurrentCount++;
 			CSVActorTotalCount++;
 		}
-	});
+	};
 #endif // (CSV_PROFILER && !UE_BUILD_SHIPPING)
 
 #if WITH_EDITORONLY_DATA
@@ -782,7 +782,7 @@ void AActor::BeginDestroy()
 #if (CSV_PROFILER && !UE_BUILD_SHIPPING)
 	// Decrement actor class count
 	// Update our ActorClassName count after the transaction finishes
-	UE_AUTORTFM_ONCOMMIT(
+	UE_AUTORTFM_ONCOMMIT2(=)
 	{
 		if (!HasAnyFlags(RF_ArchetypeObject | RF_ClassDefaultObject))
 		{
@@ -797,7 +797,7 @@ void AActor::BeginDestroy()
 			}
 			CSVActorTotalCount--;
 		}
-	});
+	};
 #endif // (CSV_PROFILER && !UE_BUILD_SHIPPING)
 
 #if WITH_EDITOR
diff --git a/Engine/Source/Runtime/Engine/Private/Components/ActorComponent.cpp b/Engine/Source/Runtime/Engine/Private/Components/ActorComponent.cpp
--- a/Engine/Source/Runtime/Engine/Private/Components/ActorComponent.cpp
+++ b/Engine/Source/Runtime/Engine/Private/Components/ActorComponent.cpp
@@ -1819,7 +1819,7 @@ void UActorComponent::ExecuteRegisterEvents(FRegisterComponentContext* Context)
 void UActorComponent::ExecuteUnregisterEvents()
 {
 	// Delay the destroying the physics, as well as the rendering state until after we commit
-	UE_AUTORTFM_ONCOMMIT(
+	UE_AUTORTFM_ONCOMMIT2(=)
 	{
 		DestroyPhysicsState();
 
@@ -1837,7 +1837,7 @@ void UActorComponent::ExecuteUnregisterEvents()
 			OnUnregister();
 			checkf(!bRegistered, TEXT("Failed to route OnUnregister (%s)"), *GetFullName());
 		}
-	});
+	};
 }
 
 void UActorComponent::ReregisterComponent()
diff --git a/Engine/Source/Runtime/Engine/Private/DataReplication.cpp b/Engine/Source/Runtime/Engine/Private/DataReplication.cpp
--- a/Engine/Source/Runtime/Engine/Private/DataReplication.cpp
+++ b/Engine/Source/Runtime/Engine/Private/DataReplication.cpp
@@ -399,16 +399,16 @@ void FObjectReplicator::InitRecentProperties(uint8* Source)
 	// If acting as a server and are IsInternalAck, that means we're recording.
 	// In that case, we don't need to create any receiving state, as no one will be sending data to us.
 	ECreateRepStateFlags Flags = (Connection->IsInternalAck() && bIsServer) ? ECreateRepStateFlags::SkipCreateReceivingState : ECreateRepStateFlags::None;
-	UE_AUTORTFM_OPEN(
+	UE_AUTORTFM_OPEN2
 	{
 		RepState = LocalRepLayout.CreateRepState(Source, RepChangedPropertyTracker, Flags);
-	});
+	};
 
 	// RepState not valid at the start of this function so just go back to being a nullptr, and let the memory be cleaned up 
-	UE_AUTORTFM_ONABORT(
+	UE_AUTORTFM_ONABORT2(this)
 	{
 		RepState = nullptr;
-	});
+	};
 
 	if (!bCreateSendingState)
 	{
diff --git a/Engine/Source/Runtime/Engine/Private/NetDriver.cpp b/Engine/Source/Runtime/Engine/Private/NetDriver.cpp
--- a/Engine/Source/Runtime/Engine/Private/NetDriver.cpp
+++ b/Engine/Source/Runtime/Engine/Private/NetDriver.cpp
@@ -7278,9 +7278,9 @@ TSharedPtr<FRepLayout> UNetDriver::GetFunctionRepLayout(UFunction * Function)
 
 		// Use a temp shared ptr to hold onto a previous value in case we abort. The adding to RepLayoutMap copies this TSharedPtr, and will end up taking a ref count and owning it
 		TSharedPtr<FRepLayout> NewLayoutPtr;
-		UE_AUTORTFM_OPEN({
+		UE_AUTORTFM_OPEN2{
 			NewLayoutPtr = FRepLayout::CreateFromFunction(Function, ServerConnection, Flags);
-		});
+		};
 
 		RepLayoutPtr = &RepLayoutMap.Add(Function, NewLayoutPtr);
 	}
diff --git a/Engine/Source/Runtime/Engine/Private/PrimitiveComponentPhysics.cpp b/Engine/Source/Runtime/Engine/Private/PrimitiveComponentPhysics.cpp
--- a/Engine/Source/Runtime/Engine/Private/PrimitiveComponentPhysics.cpp
+++ b/Engine/Source/Runtime/Engine/Private/PrimitiveComponentPhysics.cpp
@@ -1214,15 +1214,15 @@ void UPrimitiveComponent::SetCollisionEnabled(ECollisionEnabled::Type NewType)
 
 	if (CurrentType != NewType)
 	{
-		UE_AUTORTFM_OPEN({
+		UE_AUTORTFM_OPEN2{
 			BodyInstance.SetCollisionEnabled(NewType);
-		});
+		};
 
 		// If we fail set the CollisionEnabled back to the CurrentType
-		UE_AUTORTFM_ONABORT(
+		UE_AUTORTFM_ONABORT2(=, this)
 		{
 			BodyInstance.SetCollisionEnabled(CurrentType);
-		});
+		};
 
 		EnsurePhysicsStateCreated();
 		OnComponentCollisionSettingsChanged();
diff --git a/Engine/Source/Runtime/Engine/Private/SimpleConstructionScript.cpp b/Engine/Source/Runtime/Engine/Private/SimpleConstructionScript.cpp
--- a/Engine/Source/Runtime/Engine/Private/SimpleConstructionScript.cpp
+++ b/Engine/Source/Runtime/Engine/Private/SimpleConstructionScript.cpp
@@ -694,10 +694,10 @@ void USimpleConstructionScript::ExecuteScriptOnActor(AActor* Actor, const TInlin
 		USceneComponent* SceneComp = NewObject<USceneComponent>(Actor);
 
 		// The object is new, so its safe for us to atomically set the flag in the open.
-		UE_AUTORTFM_OPEN(
+		UE_AUTORTFM_OPEN2
 			{
 				SceneComp->SetFlags(RF_Transactional);
-			});
+			};
 
 		SceneComp->CreationMethod = EComponentCreationMethod::SimpleConstructionScript;
 		if (RootRelativeRotationCache)
diff --git a/Engine/Source/Runtime/Experimental/Chaos/Public/Chaos/ImplicitObject.h b/Engine/Source/Runtime/Experimental/Chaos/Public/Chaos/ImplicitObject.h
--- a/Engine/Source/Runtime/Experimental/Chaos/Public/Chaos/ImplicitObject.h
+++ b/Engine/Source/Runtime/Experimental/Chaos/Public/Chaos/ImplicitObject.h
@@ -117,17 +117,17 @@ public:
 	FChaosRefCountedObject& operator=(const FChaosRefCountedObject& Rhs) = delete;
 	uint32 AddRef() const
 	{
-		UE_AUTORTFM_ONCOMMIT(
+		UE_AUTORTFM_ONCOMMIT2(this)
 		{
 			NumRefs.Increment();
-		});
+		};
 
 		// Note: TRefCountPtr doesn't use the return value
 		return 0;
 	}
 	uint32 Release() const
 	{
-		UE_AUTORTFM_ONCOMMIT(
+		UE_AUTORTFM_ONCOMMIT2(this)
 		{
 			uint32 Refs = uint32(NumRefs.Decrement());
 			if (Refs == 0)
@@ -137,7 +137,7 @@ public:
 					delete this;
 				}
 			}
-		});
+		};
 
 		// Note: TRefCountPtr doesn't use the return value
 		return 0;
@@ -145,10 +145,10 @@ public:
 	uint32 GetRefCount() const
 	{
 		uint32 Ret = 0;
-		UE_AUTORTFM_OPEN(
+		UE_AUTORTFM_OPEN2
 		{
 			Ret = uint32(NumRefs.GetValue());
-		});
+		};
 
 		return Ret;
 	}
diff --git a/Engine/Source/Runtime/GameplayTags/Private/GameplayTagsManager.cpp b/Engine/Source/Runtime/GameplayTags/Private/GameplayTagsManager.cpp
--- a/Engine/Source/Runtime/GameplayTags/Private/GameplayTagsManager.cpp
+++ b/Engine/Source/Runtime/GameplayTags/Private/GameplayTagsManager.cpp
@@ -2251,7 +2251,7 @@ bool UGameplayTagsManager::ExtractParentTags(const FGameplayTag& GameplayTag, TA
 	FName RawTag = GameplayTag.GetTagName();
 
 	// Need to run in the open as it takes a lock. 
-	UE_AUTORTFM_OPEN(
+	UE_AUTORTFM_OPEN2
 		{
 			FScopeLock Lock(&GameplayTagMapCritical);
 
@@ -2277,7 +2277,7 @@ bool UGameplayTagsManager::ExtractParentTags(const FGameplayTag& GameplayTag, TA
 				// If we don't clear invalid tags, we need to extract the parents now in case they get registered later
 				GameplayTag.ParseParentTags(UniqueParentTags);
 			}
-		});
+		};
 
 	return UniqueParentTags.Num() != OldSize;
 }
