Fix an AutoRTFM issue where if we abort we might memory stomp the text cache.

diff --git a/Engine/Source/Runtime/Core/Private/Internationalization/TextCache.cpp b/Engine/Source/Runtime/Core/Private/Internationalization/TextCache.cpp
--- a/Engine/Source/Runtime/Core/Private/Internationalization/TextCache.cpp
+++ b/Engine/Source/Runtime/Core/Private/Internationalization/TextCache.cpp
@@ -56,11 +56,13 @@ FText FTextCache::FindOrCache(const TCHAR* InTextLiteral, const FTextId& InTextI
 	}
 
 	// Not currently cached, make a new instance...
-	FText NewText = FText(InTextLiteral, InTextId.GetNamespace(), InTextId.GetKey(), ETextFlag::Immutable);
+	FText NewText;
 
 	// ... and add it to the cache
 	UE_AUTORTFM_OPEN2
 	{
+		NewText = FText(InTextLiteral, InTextId.GetNamespace(), InTextId.GetKey(), ETextFlag::Immutable);
+
 		FScopeLock Lock(&CachedTextCS);
 
 		CachedText.Emplace(InTextId, NewText);
