Fix commas in unit test expressions in LowLevelTestAdapter for FAutomationTestFramework

[RN] Fix commas in unit test expressions in LowLevelTestAdapter for FAutomationTestFramework

diff --git a/Engine/Source/Runtime/Core/Public/Misc/LowLevelTestAdapter.h b/Engine/Source/Runtime/Core/Public/Misc/LowLevelTestAdapter.h
--- a/Engine/Source/Runtime/Core/Public/Misc/LowLevelTestAdapter.h
+++ b/Engine/Source/Runtime/Core/Public/Misc/LowLevelTestAdapter.h
@@ -122,19 +122,19 @@ FString CaptureExpressionsAndValues(const FString& InExpressions, ArgTypes&&...
 #define DISABLED_SECTION(...)						auto MAKE_UNIQUE_IDENT(disabled_section_) = []()
 
 //-V:CHECK:571,501,547
-#define CHECK(Expr) if (!(Expr)) { FAutomationTestFramework::Get().GetCurrentTest()->AddError(TEXT("Condition failed")); }
+#define CHECK(...) if (!(__VA_ARGS__)) { FAutomationTestFramework::Get().GetCurrentTest()->AddError(TEXT("Condition failed")); }
 //-V:CHECK_FALSE:571,501,547
-#define CHECK_FALSE(Expr) if (Expr) { FAutomationTestFramework::Get().GetCurrentTest()->AddError(TEXT("Condition expected to return false but returned true")); }
-#define CHECKED_IF(Expr) if (Expr)
-#define CHECKED_ELSE(Expr) if (!(Expr))
+#define CHECK_FALSE(...) if (!!(__VA_ARGS__)) { FAutomationTestFramework::Get().GetCurrentTest()->AddError(TEXT("Condition expected to return false but returned true")); }
+#define CHECKED_IF(...) if (!!(__VA_ARGS__))
+#define CHECKED_ELSE(...) if (!(__VA_ARGS__))
 //-V:CHECK_MESSAGE:571,501,547
-#define CHECK_MESSAGE(Message, Expr) if (!(Expr)) { FAutomationTestFramework::Get().GetCurrentTest()->AddError(Message); }
+#define CHECK_MESSAGE(Message, ...) if (!(__VA_ARGS__)) { FAutomationTestFramework::Get().GetCurrentTest()->AddError(Message); }
 //-V:CHECK_FALSE_MESSAGE:571,501,547
-#define CHECK_FALSE_MESSAGE(Message, Expr) if (Expr) { FAutomationTestFramework::Get().GetCurrentTest()->AddError(Message); }
+#define CHECK_FALSE_MESSAGE(Message, ...) if (!!(__VA_ARGS__)) { FAutomationTestFramework::Get().GetCurrentTest()->AddError(Message); }
 //-V:REQUIRE:571,501,547
-#define REQUIRE(Expr) if (!(Expr)) { FAutomationTestFramework::Get().GetCurrentTest()->AddError(TEXT("Required condition failed, interrupting test")); return; }
+#define REQUIRE(...) if (!(__VA_ARGS__)) { FAutomationTestFramework::Get().GetCurrentTest()->AddError(TEXT("Required condition failed, interrupting test")); return; }
 //-V:REQUIRE_MESSAGE:571,501,547
-#define REQUIRE_MESSAGE(Message, Expr) if (!(Expr)) { FAutomationTestFramework::Get().GetCurrentTest()->AddError(Message); return; }
+#define REQUIRE_MESSAGE(Message, ...) if (!(__VA_ARGS__)) { FAutomationTestFramework::Get().GetCurrentTest()->AddError(Message); return; }
 #define STATIC_REQUIRE(...) static_assert(__VA_ARGS__, #__VA_ARGS__);
 #define STATIC_CHECK(...) static_assert(__VA_ARGS__, #__VA_ARGS__);
 #define STATIC_CHECK_FALSE(...) static_assert(!(__VA_ARGS__), "!(" #__VA_ARGS__ ")");
