Add TArray/TArrayView/FStringView::NumBytes() and FString::NumBytesWith[out]Null()

[RN]

diff --git a/Engine/Source/Runtime/Core/Public/Containers/Array.h b/Engine/Source/Runtime/Core/Public/Containers/Array.h
--- a/Engine/Source/Runtime/Core/Public/Containers/Array.h
+++ b/Engine/Source/Runtime/Core/Public/Containers/Array.h
@@ -810,6 +810,12 @@ public:
 		return ArrayNum;
 	}
 
+	/** @returns Number of bytes used, excluding slack */
+	FORCEINLINE SIZE_T NumBytes() const
+	{
+		return static_cast<SIZE_T>(ArrayNum) * sizeof(ElementType);
+	}
+
 	/**
 	 * Returns maximum number of elements in array.
 	 *
diff --git a/Engine/Source/Runtime/Core/Public/Containers/ArrayView.h b/Engine/Source/Runtime/Core/Public/Containers/ArrayView.h
--- a/Engine/Source/Runtime/Core/Public/Containers/ArrayView.h
+++ b/Engine/Source/Runtime/Core/Public/Containers/ArrayView.h
@@ -360,6 +360,12 @@ public:
 		return ArrayNum;
 	}
 
+	/** @returns Number of bytes used, excluding slack */
+	FORCEINLINE SIZE_T NumBytes() const
+	{
+		return static_cast<SIZE_T>(ArrayNum) * sizeof(ElementType);
+	}
+
 	/**
 	 * Array bracket operator. Returns reference to element at given index.
 	 *
diff --git a/Engine/Source/Runtime/Core/Public/Containers/StringView.h b/Engine/Source/Runtime/Core/Public/Containers/StringView.h
--- a/Engine/Source/Runtime/Core/Public/Containers/StringView.h
+++ b/Engine/Source/Runtime/Core/Public/Containers/StringView.h
@@ -159,6 +159,9 @@ public:
 
 	// Capacity
 
+	/** @returns Number of bytes used for the characters in the string view */
+	[[nodiscard]] constexpr inline SIZE_T NumBytes() const { return static_cast<SIZE_T>(Size) * sizeof(CharType); }
+
 	/** Returns the length of the string view. */
 	[[nodiscard]] constexpr inline int32 Len() const { return Size; }
 
diff --git a/Engine/Source/Runtime/Core/Public/Containers/UnrealString.h.inl b/Engine/Source/Runtime/Core/Public/Containers/UnrealString.h.inl
--- a/Engine/Source/Runtime/Core/Public/Containers/UnrealString.h.inl
+++ b/Engine/Source/Runtime/Core/Public/Containers/UnrealString.h.inl
@@ -1099,6 +1099,18 @@ public:
 		return Data.Num() ? Data.Num() - 1 : 0;
 	}
 
+	/** @returns Number of bytes used for characters, excluding the null-terminator and slack */
+	[[nodiscard]] FORCEINLINE SIZE_T NumBytesWithoutNull() const
+	{
+		return static_cast<SIZE_T>(Len()) * sizeof(ElementType);
+	}
+
+	/** @returns Number of bytes used for characters, including the null-terminator but excluding slack */
+	[[nodiscard]] FORCEINLINE SIZE_T NumBytesWithNull() const
+	{
+		return static_cast<SIZE_T>(Len() + 1) * sizeof(ElementType);
+	}
+
 	/** Returns the left most given number of characters */
 	[[nodiscard]] FORCEINLINE UE_STRING_CLASS Left( int32 Count ) const &
 	{
