Fix bogus static analysis warning

#rnx

diff --git a/Engine/Source/Runtime/Core/Public/Containers/UnrealString.h.inl b/Engine/Source/Runtime/Core/Public/Containers/UnrealString.h.inl
--- a/Engine/Source/Runtime/Core/Public/Containers/UnrealString.h.inl
+++ b/Engine/Source/Runtime/Core/Public/Containers/UnrealString.h.inl
@@ -1108,7 +1108,7 @@ public:
 	/** @returns Number of bytes used for characters, including the null-terminator but excluding slack */
 	[[nodiscard]] FORCEINLINE SIZE_T NumBytesWithNull() const
 	{
-		return static_cast<SIZE_T>(Len() + 1) * sizeof(ElementType);
+		return (static_cast<SIZE_T>(Len()) + 1) * sizeof(ElementType);
 	}
 
 	/** Returns the left most given number of characters */
