Added FAutoConsoleCommandDeprecated which can be used to change the name of a console command but leave a shadow console command behind with the old name so old workflows are not broken, and can control if the user gets warnings, asserts, etc. when used.

diff --git a/Engine/Source/Runtime/Core/Private/HAL/ConsoleManager.cpp b/Engine/Source/Runtime/Core/Private/HAL/ConsoleManager.cpp
--- a/Engine/Source/Runtime/Core/Private/HAL/ConsoleManager.cpp
+++ b/Engine/Source/Runtime/Core/Private/HAL/ConsoleManager.cpp
@@ -4455,10 +4455,11 @@
 	ECVF_ReadOnly | ECVF_RenderThreadSafe);
 
 
-class FConsoleVariableShadow : public IConsoleVariable
+class FConsoleObjectShadowData
 {
-public:
-	FConsoleVariableShadow(const TCHAR* CVarToShadow, const TCHAR* InDeprecatedVersion, EShadowCVarBehavior InLookupBehavior, EShadowCVarBehavior InUsageBehavior)
+protected:
+	
+	FConsoleObjectShadowData(const TCHAR* CVarToShadow, const TCHAR* InDeprecatedVersion, EShadowCVarBehavior InLookupBehavior, EShadowCVarBehavior InUsageBehavior)
 		: ShadowName(CVarToShadow)
 		, DeprecatedVersion(InDeprecatedVersion)
 		, LookupBehavior(InLookupBehavior)
@@ -4466,6 +4467,158 @@
 		, bHasLooked(false)
 		, bHasMessagedForUsage(false)
 		, bHasMessagedEditorForUsage(false)
+	{
+	}
+
+	virtual ~FConsoleObjectShadowData()
+	{
+	}
+
+	virtual IConsoleObject* GetRealObject() const = 0;
+	virtual void SetRealObject(IConsoleObject* Object) const = 0;
+	virtual FString GetThisName() const = 0;
+
+	void LogOrEditorMessage(const FText& Msg, bool bIsError) const
+	{
+		if (GIsEditor && !bHasMessagedEditorForUsage)
+		{
+			bHasMessagedEditorForUsage = true;
+
+			FMessageLog EditorErrors("EditorErrors");
+			TSharedRef<FTokenizedMessage> Message = EditorErrors.Message(bIsError ? EMessageSeverity::Error : EMessageSeverity::Warning);
+			Message->AddToken(FTextToken::Create(Msg));
+			EditorErrors.Notify();
+		}
+
+		// always spit to log
+		if (bIsError)
+		{
+			UE_LOG(LogConsoleManager, Error, TEXT("%s"), *Msg.ToString());
+		}
+		else
+		{
+			UE_LOG(LogConsoleManager, Warning, TEXT("%s"), *Msg.ToString());
+		}
+	}
+
+	bool Bind() const
+	{
+		if (GetRealObject() == nullptr)
+		{
+			// if we looked but it wasn't found, then just return false
+			if (bHasLooked)
+			{
+				return false;
+			}
+			bHasLooked = true;
+
+			IConsoleObject* Object = IConsoleManager::Get().FindConsoleObject(*ShadowName, false);
+			if (Object == nullptr)
+			{
+				if (LookupBehavior != EShadowCVarBehavior::NoMessaging)
+				{
+					FFormatNamedArguments Arguments;
+					Arguments.Add(TEXT("ThisName"), FText::FromString(GetThisName()));
+					Arguments.Add(TEXT("ShadowName"), FText::FromString(ShadowName));
+
+					FText Message;
+					if (LookupBehavior == EShadowCVarBehavior::Assert)
+					{
+						Message = FText::Format(LOCTEXT("FailedShadowCVarLookup_Assert", "Attempted to delay-load real CVar '{ThisName}' for shadowed CVar '{ShadowName}' failed."), Arguments);
+					}
+					else
+					{
+						Message = FText::Format(LOCTEXT("FailedShadowCVarLookup", "Attempted to delay-load real CVar '{ThisName}' for shadowed CVar '{ShadowName}' failed. Uses of '{ThisName}' will do nothing."), Arguments);
+					}
+
+					switch (LookupBehavior)
+					{
+					case EShadowCVarBehavior::Warn:
+						LogOrEditorMessage(Message, false);
+						break;
+					case EShadowCVarBehavior::Error:
+						LogOrEditorMessage(Message, true);
+						break;
+					case EShadowCVarBehavior::Ensure:
+						ensureMsgf(false, TEXT("%s"), *Message.ToString());
+						break;
+					case EShadowCVarBehavior::Assert:
+						UE_LOG(LogConsoleManager, Fatal, TEXT("%s"), *Message.ToString());
+						break;
+					}
+				}
+				return false;
+			}
+
+			SetRealObject(Object);
+		}
+
+		return true;
+	}
+
+	bool BindForUsage() const
+	{
+		if (!Bind())
+		{
+			return false;
+		}
+
+		if (!bHasMessagedForUsage)
+		{
+			if (UsageBehavior != EShadowCVarBehavior::NoMessaging)
+			{
+				FFormatNamedArguments Arguments;
+				Arguments.Add(TEXT("ThisName"), FText::FromString(GetThisName()));
+				Arguments.Add(TEXT("ShadowName"), FText::FromString(ShadowName));
+				Arguments.Add(TEXT("DeprecatedVersion"), FText::FromString(DeprecatedVersion));
+
+				FText Message;
+				if (DeprecatedVersion != TEXT(""))
+				{
+					Message = FText::Format(LOCTEXT("ShadowCVarUsage_Deprecated", "Using a deprecated (as of UE {DeprecatedVersion}) CVar: '{ThisName}'. It will be removed in the future. Change all uses to '{ShadowName}' instead."), Arguments);
+				}
+				else
+				{
+					Message = FText::Format(LOCTEXT("ShadowCVarUsage_Deprecated", "Using a shadowed CVar '{ThisName}'. It is recommended to change all uses to '{ShadowName}' instead."), Arguments);
+				}
+				switch (UsageBehavior)
+				{
+				case EShadowCVarBehavior::Warn:
+					LogOrEditorMessage(Message, false);
+					break;
+				case EShadowCVarBehavior::Error:
+					LogOrEditorMessage(Message, true);
+					break;
+				case EShadowCVarBehavior::Ensure:
+					// don't keep re-ensuring for this cvar
+					bHasMessagedForUsage = true;
+					ensureAlwaysMsgf(false, TEXT("%s"), *Message.ToString());
+					break;
+				case EShadowCVarBehavior::Assert:
+					UE_LOG(LogConsoleManager, Fatal, TEXT("%s"), *Message.ToString());
+					break;
+				}
+			}
+		}
+
+		return true;
+	}
+
+	FString ShadowName;
+	FString DeprecatedVersion;
+	EShadowCVarBehavior LookupBehavior;
+	EShadowCVarBehavior UsageBehavior;
+	mutable uint8 bHasLooked : 1;
+	mutable uint8 bHasMessagedForUsage : 1;
+	mutable uint8 bHasMessagedEditorForUsage : 1;
+};
+
+
+class FConsoleVariableShadow : public FConsoleObjectShadowData, public IConsoleVariable
+{
+public:
+	FConsoleVariableShadow(const TCHAR* CVarToShadow, const TCHAR* InDeprecatedVersion, EShadowCVarBehavior InLookupBehavior, EShadowCVarBehavior InUsageBehavior)
+		: FConsoleObjectShadowData(CVarToShadow, InDeprecatedVersion, InLookupBehavior, InUsageBehavior)
 		, RealVariable(nullptr)
 	{
 
@@ -4719,140 +4872,106 @@
 
 #endif
 
+protected:
+	virtual IConsoleObject* GetRealObject() const override
+	{
+		return RealVariable;
+	}
+	virtual void SetRealObject(IConsoleObject* Object) const override
+	{
+		RealVariable = (IConsoleVariable*)Object;
+	}
+	virtual FString GetThisName() const override
+	{
+		return IConsoleManager::Get().FindConsoleObjectName(this);
+	}
+
 private:
 
-	void LogOrEditorMessage(const FText& Msg, bool bIsError) const
-	{
-		if (GIsEditor && !bHasMessagedEditorForUsage)
-		{
-			bHasMessagedEditorForUsage = true;
-
-			FMessageLog EditorErrors("EditorErrors");
-			TSharedRef<FTokenizedMessage> Message = EditorErrors.Message(bIsError ? EMessageSeverity::Error : EMessageSeverity::Warning);
-			Message->AddToken(FTextToken::Create(Msg));
-			EditorErrors.Notify();
-		}
-
-		// always spit to log
-		if (bIsError)
-		{
-			UE_LOG(LogConsoleManager, Error, TEXT("%s"), *Msg.ToString());
-		}
-		else
-		{
-			UE_LOG(LogConsoleManager, Warning, TEXT("%s"), *Msg.ToString());
-		}
-	}
-
-	bool Bind() const 
-	{
-		if (RealVariable == nullptr)
-		{
-			// if we looked but it wasn't found, then just return false
-			if (bHasLooked)
-			{
-				return false;
-			}
-			bHasLooked = true;
-
-			RealVariable = IConsoleManager::Get().FindConsoleVariable(*ShadowName, false);
-			if (RealVariable == nullptr)
-			{
-				if (LookupBehavior != EShadowCVarBehavior::NoMessaging)
-				{
-					FFormatNamedArguments Arguments;
-					Arguments.Add(TEXT("ThisName"), FText::FromString(IConsoleManager::Get().FindConsoleObjectName(this)));
-					Arguments.Add(TEXT("ShadowName"), FText::FromString(ShadowName));
-					
-					FText Message;
-					if (LookupBehavior == EShadowCVarBehavior::Assert)
-					{
-						Message = FText::Format(LOCTEXT("FailedShadowCVarLookup_Assert", "Attempted to delay-load real CVar '{ThisName}' for shadowed CVar '{ShadowName}' failed."), Arguments);
-					}
-					else
-					{
-						Message = FText::Format(LOCTEXT("FailedShadowCVarLookup", "Attempted to delay-load real CVar '{ThisName}' for shadowed CVar '{ShadowName}' failed. Uses of '{ThisName}' will do nothing."), Arguments);
-					}
-
-					switch (LookupBehavior)
-					{
-						case EShadowCVarBehavior::Warn:
-							LogOrEditorMessage(Message, false);
-							break;
-						case EShadowCVarBehavior::Error:
-							LogOrEditorMessage(Message, true);
-							break;
-						case EShadowCVarBehavior::Ensure:
-							ensureMsgf(false, TEXT("%s"), *Message.ToString());
-							break;
-						case EShadowCVarBehavior::Assert:
-							UE_LOG(LogConsoleManager, Fatal, TEXT("%s"), *Message.ToString());
-							break;
-					}
-				}
-				return false;
-			}
-		}
-
-		return true;
-	}
-
-	bool BindForUsage() const
-	{
-		if (!Bind())
-		{
-			return false;
-		}
-
-		if (!bHasMessagedForUsage)
-		{
-			if (UsageBehavior != EShadowCVarBehavior::NoMessaging)
-			{
-				FFormatNamedArguments Arguments;
-				Arguments.Add(TEXT("ThisName"), FText::FromString(IConsoleManager::Get().FindConsoleObjectName(this)));
-				Arguments.Add(TEXT("ShadowName"), FText::FromString(ShadowName));
-				Arguments.Add(TEXT("DeprecatedVersion"), FText::FromString(DeprecatedVersion));
-
-				FText Message;
-				if (DeprecatedVersion != TEXT(""))
-				{
-					Message = FText::Format(LOCTEXT("ShadowCVarUsage_Deprecated", "Using a deprecated (as of UE {DeprecatedVersion}) CVar: '{ThisName}'. It will be removed in the future. Change all uses to '{ShadowName}' instead."), Arguments);
-				}
-				else
-				{
-					Message = FText::Format(LOCTEXT("ShadowCVarUsage", "Using a shadowed CVar '{ThisName}'. It is recommended to change all uses to '{ShadowName}' instead."), Arguments);
-				}
-				switch (UsageBehavior)
-				{
-				case EShadowCVarBehavior::Warn:
-					LogOrEditorMessage(Message, false);
-					break;
-				case EShadowCVarBehavior::Error:
-					LogOrEditorMessage(Message, true);
-					break;
-				case EShadowCVarBehavior::Ensure:
-					// don't keep re-ensuring for this cvar
-					bHasMessagedForUsage = true;
-					ensureAlwaysMsgf(false, TEXT("%s"), *Message.ToString());
-					break;
-				case EShadowCVarBehavior::Assert:
-					UE_LOG(LogConsoleManager, Fatal, TEXT("%s"), *Message.ToString());
-					break;
-				}
-			}
-		}
-
-		return true;
-	}
-
-	FString ShadowName;
-	FString DeprecatedVersion;
-	EShadowCVarBehavior LookupBehavior;
-	EShadowCVarBehavior UsageBehavior;
-	mutable uint8 bHasLooked:1;
-	mutable uint8 bHasMessagedForUsage:1;
-	mutable uint8 bHasMessagedEditorForUsage:1;
 	mutable IConsoleVariable* RealVariable;
 };
 
+struct FConsoleCommandShadow : public FConsoleObjectShadowData, public IConsoleCommand
+{
 
+	FConsoleCommandShadow(const TCHAR* CVarToShadow, const TCHAR* InDeprecatedVersion, EShadowCVarBehavior InLookupBehavior, EShadowCVarBehavior InUsageBehavior)
+		: FConsoleObjectShadowData(CVarToShadow, InDeprecatedVersion, InLookupBehavior, InUsageBehavior)
+		, RealCommand(nullptr)
+	{
+
+	}
+
+	virtual struct IConsoleCommand* AsCommand() override
+	{
+		return this;
+	}
+
+	virtual const TCHAR* GetHelp() const override
+	{
+		if (Bind())
+		{
+			return RealCommand->GetHelp();
+		}
+		return TEXT("");
+	}
+
+	virtual void SetHelp(const TCHAR* InValue) override
+	{
+		if (Bind())
+		{
+			return RealCommand->SetHelp(InValue);
+		}
+	}
+	virtual EConsoleVariableFlags GetFlags() const override
+	{
+		if (Bind())
+		{
+			return RealCommand->GetFlags();
+		}
+		return (EConsoleVariableFlags)0;
+
+	}
+	virtual void SetFlags(const EConsoleVariableFlags Value) override
+	{
+		if (BindForUsage())
+		{
+			RealCommand->SetFlags(Value);
+		}
+	}
+
+	virtual bool Execute(const TArray< FString >& Args, UWorld* InWorld, class FOutputDevice& OutputDevice) override
+	{
+		if (BindForUsage())
+		{
+			return RealCommand->Execute(Args, InWorld, OutputDevice);
+		}
+		return false;
+	}
+
+	virtual void Release() override
+	{
+		// the real var should do it's own Release when it's time
+		RealCommand = nullptr;
+	}
+
+protected:
+	virtual IConsoleObject* GetRealObject() const override
+	{
+		return RealCommand;
+	}
+	virtual void SetRealObject(IConsoleObject* Object) const override
+	{
+		RealCommand = (IConsoleCommand*)Object;
+	}
+	virtual FString GetThisName() const override
+	{
+		return IConsoleManager::Get().FindConsoleObjectName(this);
+	}
+
+private:
+
+	mutable IConsoleCommand* RealCommand;
+
+};
+
+
@@ -4871,4 +4990,13 @@
 #endif
 }
 
+FAutoConsoleCommandDeprecated::FAutoConsoleCommandDeprecated(const TCHAR* Name, const TCHAR* CVarToShadow, const TCHAR* DeprecatedAtVersion, EShadowCVarBehavior UsageBehavior, EShadowCVarBehavior LookupFailureBehavior)
+{
+#if !NO_CVARS
+	GetManager().AddShadowConsoleObject(Name, new FConsoleCommandShadow(CVarToShadow, DeprecatedAtVersion, LookupFailureBehavior, UsageBehavior));
+#endif
+}
+
+
+
 #undef LOCTEXT_NAMESPACE
diff --git a/Engine/Source/Runtime/Core/Public/HAL/IConsoleManager.h b/Engine/Source/Runtime/Core/Public/HAL/IConsoleManager.h
--- a/Engine/Source/Runtime/Core/Public/HAL/IConsoleManager.h
+++ b/Engine/Source/Runtime/Core/Public/HAL/IConsoleManager.h
@@ -2120,5 +2120,21 @@ public:
 	CORE_API FAutoConsoleVariableDeprecated(const TCHAR* Name, const TCHAR* CVarToShadow, const TCHAR* DeprecatedAtVersion, EShadowCVarBehavior UsageBehavior=EShadowCVarBehavior::Ensure, EShadowCVarBehavior LookupFailureBehavior = EShadowCVarBehavior::Ensure);
 };
 
+/**
+ * A class that can be used to deprecate a ConsoleCommand - use just like FAutoConsoleVariableDeprecated
+ */
+class FAutoConsoleCommandDeprecated
+{
+public:
+	/**
+	 * @param Name The name of this cvar
+	 * @param CVarToShadow The name of the existing cvar that this one will shadow
+	 * @param DeprecatedAtVersion THe first UE version when this was deprecated
+	 * @param UsageBehavior How to handle uses of the shadowed cvar by the deprecated name - defaults to ensure
+	 * @param LookupFailureBehavior How to handle failure of looking up the CVarToShadow name - defaults to ensure
+	 */
+	CORE_API FAutoConsoleCommandDeprecated(const TCHAR* Name, const TCHAR* CVarToShadow, const TCHAR* DeprecatedAtVersion, EShadowCVarBehavior UsageBehavior = EShadowCVarBehavior::Ensure, EShadowCVarBehavior LookupFailureBehavior = EShadowCVarBehavior::Ensure);
+};
+
 CORE_API DECLARE_LOG_CATEGORY_EXTERN(LogConsoleResponse, Log, All);
 
