CompactBinary: Added [[nodiscard]]

diff --git a/Engine/Source/Runtime/Core/Private/Tests/Serialization/CompactBinaryTest.cpp b/Engine/Source/Runtime/Core/Private/Tests/Serialization/CompactBinaryTest.cpp
--- a/Engine/Source/Runtime/Core/Private/Tests/Serialization/CompactBinaryTest.cpp
+++ b/Engine/Source/Runtime/Core/Private/Tests/Serialization/CompactBinaryTest.cpp
@@ -374,7 +374,7 @@ bool FCbFieldObjectTest::RunTest(const FString& Parameters)
 		for (int Count = 16; Count > 0; --Count)
 		{
 			++It;
-			It->AsInt32();
+			(void)It->AsInt32();
 		}
 		TestFalse(TEXT("FCbObjectView()::CreateViewIterator() At End"), bool(It));
 		TestTrue(TEXT("FCbObjectView()::CreateViewIterator() At End"), !It);
@@ -586,7 +586,7 @@ bool FCbFieldArrayTest::RunTest(const FString& Parameters)
 		for (int Count = 16; Count > 0; --Count)
 		{
 			++It;
-			It->AsInt32();
+			(void)It->AsInt32();
 		}
 		TestFalse(TEXT("FCbArrayView()::CreateViewIterator() At End"), bool(It));
 		TestTrue(TEXT("FCbArrayView()::CreateViewIterator() At End"), !It);
@@ -744,9 +744,9 @@ bool FCbFieldBinaryTest::RunTest(const FString& Parameters)
 		TestFieldNoClone<ECbFieldType::Binary>(TEXT("Binary, Value, View"), FieldView, MakeMemoryView(Value + 1, 3));
 
 		FCbField Field = FCbField::Clone(FieldView);
-		Field.AsBinary();
+		(void)Field.AsBinary();
 		TestFalse(TEXT("Binary, Value, AsBinary -> GetOuterBuffer()"), Field.GetOuterBuffer().IsNull());
-		MoveTemp(Field).AsBinary();
+		(void)MoveTemp(Field).AsBinary();
 		TestTrue(TEXT("Binary, Value, AsBinary -> GetOuterBuffer()"), Field.GetOuterBuffer().IsNull());
 	}
 
diff --git a/Engine/Source/Runtime/Core/Public/Serialization/CompactBinary.h b/Engine/Source/Runtime/Core/Public/Serialization/CompactBinary.h
--- a/Engine/Source/Runtime/Core/Public/Serialization/CompactBinary.h
+++ b/Engine/Source/Runtime/Core/Public/Serialization/CompactBinary.h
@@ -496,7 +496,7 @@ public:
 	/** Returns a view of the raw byte array for the ObjectId. */
 	constexpr inline FMemoryView GetView() const { return MakeMemoryView(Bytes); }
 
-	CORE_API static FCbObjectId NewObjectId();
+	[[nodiscard]] CORE_API static FCbObjectId NewObjectId();
 
 	inline bool operator==(const FCbObjectId& B) const
 	{
@@ -624,7 +624,7 @@ public:
 	inline explicit FCbFieldView(const FCbValue& Value);
 
 	/** Returns a copy of the field with the name removed. */
-	constexpr inline FCbFieldView RemoveName() const
+	[[nodiscard]] constexpr inline FCbFieldView RemoveName() const
 	{
 		FCbFieldView Field;
 		Field.TypeWithFlags = TypeWithFlags & ~(ECbFieldType::HasFieldType | ECbFieldType::HasFieldName);
@@ -633,94 +633,94 @@ public:
 	}
 
 	/** Returns the name of the field if it has a name, otherwise an empty view. */
-	constexpr inline FUtf8StringView GetName() const
+	[[nodiscard]] constexpr inline FUtf8StringView GetName() const
 	{
 		return FUtf8StringView(static_cast<const UTF8CHAR*>(Value) - NameLen, NameLen);
 	}
 
 	/** Returns the value for unchecked access. Prefer the typed accessors below. */
-	inline FCbValue GetValue() const;
+	[[nodiscard]] inline FCbValue GetValue() const;
 
 	/** Access the field as an object. Defaults to an empty object on error. */
-	CORE_API FCbObjectView AsObjectView();
+	[[nodiscard]] CORE_API FCbObjectView AsObjectView();
 
 	/** Access the field as an array. Defaults to an empty array on error. */
-	CORE_API FCbArrayView AsArrayView();
+	[[nodiscard]] CORE_API FCbArrayView AsArrayView();
 
 	/** Access the field as binary. Returns the provided default on error. */
-	CORE_API FMemoryView AsBinaryView(FMemoryView Default = FMemoryView());
+	[[nodiscard]] CORE_API FMemoryView AsBinaryView(FMemoryView Default = FMemoryView());
 
 	/** Access the field as a string. Returns the provided default on error. */
-	CORE_API FUtf8StringView AsString(FUtf8StringView Default = FUtf8StringView());
+	[[nodiscard]] CORE_API FUtf8StringView AsString(FUtf8StringView Default = FUtf8StringView());
 
 	/** Access the field as an int8. Returns the provided default on error. */
-	inline int8 AsInt8(int8 Default = 0)       { return AsInteger<int8>(Default); }
+	[[nodiscard]] inline int8 AsInt8(int8 Default = 0) { return AsInteger<int8>(Default); }
 	/** Access the field as an int16. Returns the provided default on error. */
-	inline int16 AsInt16(int16 Default = 0)    { return AsInteger<int16>(Default); }
+	[[nodiscard]] inline int16 AsInt16(int16 Default = 0) { return AsInteger<int16>(Default); }
 	/** Access the field as an int32. Returns the provided default on error. */
-	inline int32 AsInt32(int32 Default = 0)    { return AsInteger<int32>(Default); }
+	[[nodiscard]] inline int32 AsInt32(int32 Default = 0) { return AsInteger<int32>(Default); }
 	/** Access the field as an int64. Returns the provided default on error. */
-	inline int64 AsInt64(int64 Default = 0)    { return AsInteger<int64>(Default); }
+	[[nodiscard]] inline int64 AsInt64(int64 Default = 0) { return AsInteger<int64>(Default); }
 	/** Access the field as a uint8. Returns the provided default on error. */
-	inline uint8 AsUInt8(uint8 Default = 0)    { return AsInteger<uint8>(Default); }
+	[[nodiscard]] inline uint8 AsUInt8(uint8 Default = 0) { return AsInteger<uint8>(Default); }
 	/** Access the field as a uint16. Returns the provided default on error. */
-	inline uint16 AsUInt16(uint16 Default = 0) { return AsInteger<uint16>(Default); }
+	[[nodiscard]] inline uint16 AsUInt16(uint16 Default = 0) { return AsInteger<uint16>(Default); }
 	/** Access the field as a uint32. Returns the provided default on error. */
-	inline uint32 AsUInt32(uint32 Default = 0) { return AsInteger<uint32>(Default); }
+	[[nodiscard]] inline uint32 AsUInt32(uint32 Default = 0) { return AsInteger<uint32>(Default); }
 	/** Access the field as a uint64. Returns the provided default on error. */
-	inline uint64 AsUInt64(uint64 Default = 0) { return AsInteger<uint64>(Default); }
+	[[nodiscard]] inline uint64 AsUInt64(uint64 Default = 0) { return AsInteger<uint64>(Default); }
 
 	/** Access the field as a float. Returns the provided default on error. */
-	CORE_API float AsFloat(float Default = 0.0f);
+	[[nodiscard]] CORE_API float AsFloat(float Default = 0.0f);
 	/** Access the field as a double. Returns the provided default on error. */
-	CORE_API double AsDouble(double Default = 0.0);
+	[[nodiscard]] CORE_API double AsDouble(double Default = 0.0);
 
 	/** Access the field as a bool. Returns the provided default on error. */
-	CORE_API bool AsBool(bool bDefault = false);
+	[[nodiscard]] CORE_API bool AsBool(bool bDefault = false);
 
 	/** Access the field as a hash referencing an object attachment. Returns the provided default on error. */
-	CORE_API FIoHash AsObjectAttachment(const FIoHash& Default = FIoHash());
+	[[nodiscard]] CORE_API FIoHash AsObjectAttachment(const FIoHash& Default = FIoHash());
 	/** Access the field as a hash referencing a binary attachment. Returns the provided default on error. */
-	CORE_API FIoHash AsBinaryAttachment(const FIoHash& Default = FIoHash());
+	[[nodiscard]] CORE_API FIoHash AsBinaryAttachment(const FIoHash& Default = FIoHash());
 	/** Access the field as a hash referencing an attachment. Returns the provided default on error. */
-	CORE_API FIoHash AsAttachment(const FIoHash& Default = FIoHash());
+	[[nodiscard]] CORE_API FIoHash AsAttachment(const FIoHash& Default = FIoHash());
 
 	/** Access the field as a hash. Returns the provided default on error. */
-	CORE_API FIoHash AsHash(const FIoHash& Default = FIoHash());
+	[[nodiscard]] CORE_API FIoHash AsHash(const FIoHash& Default = FIoHash());
 
 	/** Access the field as a UUID. Returns a nil UUID on error. */
-	CORE_API FGuid AsUuid();
+	[[nodiscard]] CORE_API FGuid AsUuid();
 	/** Access the field as a UUID. Returns the provided default on error. */
-	CORE_API FGuid AsUuid(const FGuid& Default);
+	[[nodiscard]] CORE_API FGuid AsUuid(const FGuid& Default);
 
 	/** Access the field as a date/time tick count. Returns the provided default on error. */
-	CORE_API int64 AsDateTimeTicks(int64 Default = 0);
+	[[nodiscard]] CORE_API int64 AsDateTimeTicks(int64 Default = 0);
 
 	/** Access the field as a date/time. Returns a date/time at the epoch on error. */
-	CORE_API FDateTime AsDateTime();
+	[[nodiscard]] CORE_API FDateTime AsDateTime();
 	/** Access the field as a date/time. Returns the provided default on error. */
-	CORE_API FDateTime AsDateTime(FDateTime Default);
+	[[nodiscard]] CORE_API FDateTime AsDateTime(FDateTime Default);
 
 	/** Access the field as a timespan tick count. Returns the provided default on error. */
-	CORE_API int64 AsTimeSpanTicks(int64 Default = 0);
+	[[nodiscard]] CORE_API int64 AsTimeSpanTicks(int64 Default = 0);
 
 	/** Access the field as a timespan. Returns an empty timespan on error. */
-	CORE_API FTimespan AsTimeSpan();
+	[[nodiscard]] CORE_API FTimespan AsTimeSpan();
 	/** Access the field as a timespan. Returns the provided default on error. */
-	CORE_API FTimespan AsTimeSpan(FTimespan Default);
+	[[nodiscard]] CORE_API FTimespan AsTimeSpan(FTimespan Default);
 
 	/** Access the field as an object identifier. Returns the provided default on error. */
-	CORE_API FCbObjectId AsObjectId(const FCbObjectId& Default = FCbObjectId());
+	[[nodiscard]] CORE_API FCbObjectId AsObjectId(const FCbObjectId& Default = FCbObjectId());
 
 	/** Access the field as a custom sub-type with an integer identifier. Returns the provided default on error. */
-	CORE_API FCbCustomById AsCustomById(FCbCustomById Default = FCbCustomById());
+	[[nodiscard]] CORE_API FCbCustomById AsCustomById(FCbCustomById Default = FCbCustomById());
 	/** Access the field as a custom sub-type with a string identifier. Returns the provided default on error. */
-	CORE_API FCbCustomByName AsCustomByName(FCbCustomByName Default = FCbCustomByName());
+	[[nodiscard]] CORE_API FCbCustomByName AsCustomByName(FCbCustomByName Default = FCbCustomByName());
 
 	/** Access the field as a custom sub-type with an integer identifier. Returns the provided default on error. */
-	CORE_API FMemoryView AsCustom(uint64 Id, FMemoryView Default = FMemoryView());
+	[[nodiscard]] CORE_API FMemoryView AsCustom(uint64 Id, FMemoryView Default = FMemoryView());
 	/** Access the field as a custom sub-type with a string identifier. Returns the provided default on error. */
-	CORE_API FMemoryView AsCustom(FUtf8StringView Name, FMemoryView Default = FMemoryView());
+	[[nodiscard]] CORE_API FMemoryView AsCustom(FUtf8StringView Name, FMemoryView Default = FMemoryView());
 
 	/** True if the field has a name. */
 	constexpr inline bool HasName() const           { return FCbFieldType::HasFieldName(TypeWithFlags); }
@@ -769,13 +769,13 @@ public:
 	constexpr inline bool HasError() const          { return Error != ECbFieldError::None; }
 
 	/** The type of error that occurred on the last field access, or None. */
-	constexpr inline ECbFieldError GetError() const { return Error; }
+	[[nodiscard]] constexpr inline ECbFieldError GetError() const { return Error; }
 
 	/** Returns the size of the field in bytes, including the type and name. */
-	CORE_API uint64 GetSize() const;
+	[[nodiscard]] CORE_API uint64 GetSize() const;
 
 	/** Calculate the hash of the field, including the type and name. */
-	CORE_API FIoHash GetHash() const;
+	[[nodiscard]] CORE_API FIoHash GetHash() const;
 	/** Append the hash of the field, including the type and name. */
 	CORE_API void AppendHash(FIoHashBuilder& Builder) const;
 
@@ -816,10 +816,10 @@ public:
 	}
 
 	/** Find a field of an object by case-sensitive name comparison, otherwise a field with no value. */
-	CORE_API FCbFieldView operator[](FUtf8StringView Name) const;
+	[[nodiscard]] CORE_API FCbFieldView operator[](FUtf8StringView Name) const;
 
 	/** Create an iterator for the fields of an array or object, otherwise an empty iterator. */
-	CORE_API FCbFieldViewIterator CreateViewIterator() const;
+	[[nodiscard]] CORE_API FCbFieldViewIterator CreateViewIterator() const;
 
 	/** DO NOT USE DIRECTLY. These functions enable range-based for loop support. */
 	inline FCbFieldViewIterator begin() const;
@@ -892,7 +892,7 @@ class FCbFieldViewIterator : public TCbFieldIterator<FCbFieldView>
 {
 public:
 	/** Construct a field range that contains exactly one field. */
-	static inline FCbFieldViewIterator MakeSingle(const FCbFieldView& Field)
+	[[nodiscard]] static inline FCbFieldViewIterator MakeSingle(const FCbFieldView& Field)
 	{
 		return FCbFieldViewIterator(Field);
 	}
@@ -903,7 +903,7 @@ public:
 	 * @param View A buffer containing zero or more valid fields.
 	 * @param Type HasFieldType means that View contains the type. Otherwise, use the given type.
 	 */
-	static inline FCbFieldViewIterator MakeRange(FMemoryView View, ECbFieldType Type = ECbFieldType::HasFieldType)
+	[[nodiscard]] static inline FCbFieldViewIterator MakeRange(FMemoryView View, ECbFieldType Type = ECbFieldType::HasFieldType)
 	{
 		return !View.IsEmpty() ? FCbFieldViewIterator(FCbFieldView(View.GetData(), Type), View.GetDataEnd()) : FCbFieldViewIterator();
 	}
@@ -951,10 +951,10 @@ public:
 	CORE_API uint64 Num() const;
 
 	/** Access the array as an array field. */
-	inline FCbFieldView AsFieldView() const { return RemoveName(); }
+	[[nodiscard]] inline FCbFieldView AsFieldView() const { return RemoveName(); }
 
 	/** Construct an array from an array field. No type check is performed! */
-	static inline FCbArrayView FromFieldNoCheck(const FCbFieldView& Field) { return FCbArrayView(Field); }
+	[[nodiscard]] static inline FCbArrayView FromFieldNoCheck(const FCbFieldView& Field) { return FCbArrayView(Field); }
 
 	/** Whether the array has any fields. */
 	inline explicit operator bool() const { return Num() > 0; }
@@ -1039,19 +1039,19 @@ public:
 	 * @param Name The name of the field.
 	 * @return The matching field if found, otherwise a field with no value.
 	 */
-	CORE_API FCbFieldView FindView(FUtf8StringView Name) const;
+	[[nodiscard]] CORE_API FCbFieldView FindView(FUtf8StringView Name) const;
 
 	/** Find a field by case-insensitive name comparison. */
-	CORE_API FCbFieldView FindViewIgnoreCase(FUtf8StringView Name) const;
+	[[nodiscard]] CORE_API FCbFieldView FindViewIgnoreCase(FUtf8StringView Name) const;
 
 	/** Find a field by case-sensitive name comparison. */
-	inline FCbFieldView operator[](FUtf8StringView Name) const { return FindView(Name); }
+	[[nodiscard]] inline FCbFieldView operator[](FUtf8StringView Name) const { return FindView(Name); }
 
 	/** Access the object as an object field. */
-	inline FCbFieldView AsFieldView() const { return RemoveName(); }
+	[[nodiscard]] inline FCbFieldView AsFieldView() const { return RemoveName(); }
 
 	/** Construct an object from an object field. No type check is performed! */
-	static inline FCbObjectView FromFieldNoCheck(const FCbFieldView& Field) { return FCbObjectView(Field); }
+	[[nodiscard]] static inline FCbObjectView FromFieldNoCheck(const FCbFieldView& Field) { return FCbObjectView(Field); }
 
 	/** Whether the object has any fields. */
 	CORE_API explicit operator bool() const;
@@ -1168,7 +1168,7 @@ public:
 	}
 
 	/** Returns the value as a view. */
-	inline const ViewType& AsView() const { return *this; }
+	[[nodiscard]] inline const ViewType& AsView() const { return *this; }
 
 	/**
 	 * Returns the outer buffer (if any) that contains this value.
@@ -1176,14 +1176,14 @@ public:
 	 * The outer buffer might contain other data before and/or after this value. Use GetBuffer to
 	 * request a buffer that exactly contains this value, or TryGetView for a contiguous view.
 	 */
-	inline const FSharedBuffer& GetOuterBuffer() const & { return Buffer; }
-	inline FSharedBuffer GetOuterBuffer() && { return MoveTemp(Buffer); }
+	[[nodiscard]] inline const FSharedBuffer& GetOuterBuffer() const & { return Buffer; }
+	[[nodiscard]] inline FSharedBuffer GetOuterBuffer() && { return MoveTemp(Buffer); }
 
 	/** Find a field of an object by case-sensitive name comparison, otherwise a field with no value. */
-	inline FCbField operator[](FUtf8StringView Name) const;
+	[[nodiscard]] inline FCbField operator[](FUtf8StringView Name) const;
 
 	/** Create an iterator for the fields of an array or object, otherwise an empty iterator. */
-	inline FCbFieldIterator CreateIterator() const;
+	[[nodiscard]] inline FCbFieldIterator CreateIterator() const;
 
 	/** DO NOT USE DIRECTLY. These functions enable range-based for loop support. */
 	inline FCbFieldIterator begin() const;
@@ -1209,13 +1209,13 @@ class TCbBufferFactory
 {
 public:
 	/** Construct a value from an owned clone of its memory. */
-	static inline Type Clone(const void* const Data)
+	[[nodiscard]] static inline Type Clone(const void* const Data)
 	{
 		return Clone(ViewType(Data));
 	}
 
 	/** Construct a value from an owned clone of its memory. */
-	static inline Type Clone(const ViewType& Value)
+	[[nodiscard]] static inline Type Clone(const ViewType& Value)
 	{
 		Type Owned = MakeView(Value);
 		Owned.MakeOwned();
@@ -1223,13 +1223,13 @@ public:
 	}
 
 	/** Construct a value from a read-only view of its memory and its optional outer buffer. */
-	static inline Type MakeView(const void* const Data, FSharedBuffer OuterBuffer = FSharedBuffer())
+	[[nodiscard]] static inline Type MakeView(const void* const Data, FSharedBuffer OuterBuffer = FSharedBuffer())
 	{
 		return MakeView(ViewType(Data), MoveTemp(OuterBuffer));
 	}
 
 	/** Construct a value from a read-only view of its memory and its optional outer buffer. */
-	static inline Type MakeView(const ViewType& Value, FSharedBuffer OuterBuffer = FSharedBuffer())
+	[[nodiscard]] static inline Type MakeView(const ViewType& Value, FSharedBuffer OuterBuffer = FSharedBuffer())
 	{
 		return Type(Value, MoveTemp(OuterBuffer));
 	}
@@ -1254,19 +1254,19 @@ public:
 	using TCbBuffer::TCbBuffer;
 
 	/** Access the field as an object. Defaults to an empty object on error. */
-	inline FCbObject AsObject() &;
-	inline FCbObject AsObject() &&;
+	[[nodiscard]] inline FCbObject AsObject() &;
+	[[nodiscard]] inline FCbObject AsObject() &&;
 
 	/** Access the field as an array. Defaults to an empty array on error. */
-	inline FCbArray AsArray() &;
-	inline FCbArray AsArray() &&;
+	[[nodiscard]] inline FCbArray AsArray()&;
+	[[nodiscard]] inline FCbArray AsArray() &&;
 
 	/** Access the field as binary. Returns the provided default on error. */
-	inline FSharedBuffer AsBinary(const FSharedBuffer& Default = FSharedBuffer()) &;
-	inline FSharedBuffer AsBinary(const FSharedBuffer& Default = FSharedBuffer()) &&;
+	[[nodiscard]] inline FSharedBuffer AsBinary(const FSharedBuffer& Default = FSharedBuffer())&;
+	[[nodiscard]] inline FSharedBuffer AsBinary(const FSharedBuffer& Default = FSharedBuffer()) &&;
 
 	/** Returns a buffer that contains the field as it would be serialized by CopyTo. */
-	CORE_API FCompositeBuffer GetBuffer() const;
+	[[nodiscard]] CORE_API FCompositeBuffer GetBuffer() const;
 };
 
 template <typename ViewType>
@@ -1288,16 +1288,16 @@ class FCbFieldIterator : public TCbFieldIterator<FCbField>
 {
 public:
 	/** Construct a field range from an owned clone of a range. */
-	CORE_API static FCbFieldIterator CloneRange(const FCbFieldViewIterator& It);
+	[[nodiscard]] CORE_API static FCbFieldIterator CloneRange(const FCbFieldViewIterator& It);
 
 	/** Construct a field range from an owned clone of a range. */
-	static inline FCbFieldIterator CloneRange(const FCbFieldIterator& It)
+	[[nodiscard]] static inline FCbFieldIterator CloneRange(const FCbFieldIterator& It)
 	{
 		return CloneRange(FCbFieldViewIterator(It));
 	}
 
 	/** Construct a field range that contains exactly one field. */
-	static inline FCbFieldIterator MakeSingle(FCbField Field)
+	[[nodiscard]] static inline FCbFieldIterator MakeSingle(FCbField Field)
 	{
 		return FCbFieldIterator(MoveTemp(Field));
 	}
@@ -1308,7 +1308,7 @@ public:
 	 * @param Buffer A buffer containing zero or more valid fields.
 	 * @param Type HasFieldType means that Buffer contains the type. Otherwise, use the given type.
 	 */
-	static inline FCbFieldIterator MakeRange(FSharedBuffer Buffer, ECbFieldType Type = ECbFieldType::HasFieldType)
+	[[nodiscard]] static inline FCbFieldIterator MakeRange(FSharedBuffer Buffer, ECbFieldType Type = ECbFieldType::HasFieldType)
 	{
 		if (Buffer.GetSize())
 		{
@@ -1319,7 +1319,7 @@ public:
 	}
 
 	/** Construct a field range from an iterator and its optional outer buffer. */
-	static inline FCbFieldIterator MakeRangeView(const FCbFieldViewIterator& It, FSharedBuffer OuterBuffer = FSharedBuffer())
+	[[nodiscard]] static inline FCbFieldIterator MakeRangeView(const FCbFieldViewIterator& It, FSharedBuffer OuterBuffer = FSharedBuffer())
 	{
 		return FCbFieldIterator(FCbField(It, MoveTemp(OuterBuffer)), GetFieldsEnd(It));
 	}
@@ -1370,11 +1370,11 @@ public:
 	using TCbBuffer::TCbBuffer;
 
 	/** Access the array as an array field. */
-	inline FCbField AsField() const & { return FCbField(FCbArrayView::AsFieldView(), *this); }
-	inline FCbField AsField() && { return FCbField(FCbArrayView::AsFieldView(), MoveTemp(*this)); }
+	[[nodiscard]] inline FCbField AsField() const & { return FCbField(FCbArrayView::AsFieldView(), *this); }
+	[[nodiscard]] inline FCbField AsField() && { return FCbField(FCbArrayView::AsFieldView(), MoveTemp(*this)); }
 
 	/** Returns a buffer that contains the array as it would be serialized by CopyTo. */
-	CORE_API FCompositeBuffer GetBuffer() const;
+	[[nodiscard]] CORE_API FCompositeBuffer GetBuffer() const;
 };
 
 ///////////////////////////////////////////////////////////////////////////////////////////////////////////////////////
@@ -1391,7 +1391,7 @@ public:
 	using TCbBuffer::TCbBuffer;
 
 	/** Find a field by case-sensitive name comparison. */
-	inline FCbField Find(FUtf8StringView Name) const
+	[[nodiscard]] inline FCbField Find(FUtf8StringView Name) const
 	{
 		if (::FCbFieldView Field = FindView(Name))
 		{
@@ -1401,7 +1401,7 @@ public:
 	}
 
 	/** Find a field by case-insensitive name comparison. */
-	inline FCbField FindIgnoreCase(FUtf8StringView Name) const
+	[[nodiscard]] inline FCbField FindIgnoreCase(FUtf8StringView Name) const
 	{
 		if (::FCbFieldView Field = FindViewIgnoreCase(Name))
 		{
@@ -1411,14 +1411,14 @@ public:
 	}
 
 	/** Find a field by case-sensitive name comparison. */
-	inline FCbField operator[](FUtf8StringView Name) const { return Find(Name); }
+	[[nodiscard]] inline FCbField operator[](FUtf8StringView Name) const { return Find(Name); }
 
 	/** Access the object as an object field. */
-	inline FCbField AsField() const & { return FCbField(FCbObjectView::AsFieldView(), *this); }
-	inline FCbField AsField() && { return FCbField(FCbObjectView::AsFieldView(), MoveTemp(*this)); }
+	[[nodiscard]] inline FCbField AsField() const & { return FCbField(FCbObjectView::AsFieldView(), *this); }
+	[[nodiscard]] inline FCbField AsField() && { return FCbField(FCbObjectView::AsFieldView(), MoveTemp(*this)); }
 
 	/** Returns a buffer that contains the object as it would be serialized by CopyTo. */
-	CORE_API FCompositeBuffer GetBuffer() const;
+	[[nodiscard]] CORE_API FCompositeBuffer GetBuffer() const;
 };
 
 ///////////////////////////////////////////////////////////////////////////////////////////////////////////////////////
