Fix no-unity builds with autortfm missing MoveTemp. Also add a macro to use std::move when building in standalone mode.

diff --git a/Engine/Source/Runtime/Core/Public/AutoRTFM/AutoRTFM.h b/Engine/Source/Runtime/Core/Public/AutoRTFM/AutoRTFM.h
--- a/Engine/Source/Runtime/Core/Public/AutoRTFM/AutoRTFM.h
+++ b/Engine/Source/Runtime/Core/Public/AutoRTFM/AutoRTFM.h
@@ -50,15 +50,19 @@
 #include <stdlib.h>
 #include <memory.h>
 #include <type_traits>
+#include <utility>
 #define UE_AUTORTFM_API
 #define UE_AUTORTFM_FORCEINLINE inline
 #define UE_AUTORTFM_MEMCPY ::memcpy
+#define UE_AUTORTFM_MOVE std::move
 #else
 #include <HAL/Platform.h>
 #include <HAL/PlatformMemory.h>
+#include <Templates/UnrealTemplate.h>
 #define UE_AUTORTFM_API CORE_API
 #define UE_AUTORTFM_FORCEINLINE FORCEINLINE
 #define UE_AUTORTFM_MEMCPY FPlatformMemory::Memcpy
+#define UE_AUTORTFM_MOVE MoveTemp
 #endif
 
 #if UE_AUTORTFM
@@ -884,7 +888,7 @@ namespace AutoRTFM::Private
 		template<typename FunctorType>
 		void operator+(FunctorType F)
 		{
-			AutoRTFM::Open(MoveTemp(F));
+			AutoRTFM::Open(UE_AUTORTFM_MOVE(F));
 		}
 	};
 	struct FOnAbortHelper
@@ -892,7 +896,7 @@ namespace AutoRTFM::Private
 		template<typename FunctorType>
 		void operator+(FunctorType F)
 		{
-			AutoRTFM::OnAbort(MoveTemp(F));
+			AutoRTFM::OnAbort(UE_AUTORTFM_MOVE(F));
 		}
 	};
 	struct FOnCommitHelper
@@ -900,7 +904,7 @@ namespace AutoRTFM::Private
 		template<typename FunctorType>
 		void operator+(FunctorType F)
 		{
-			AutoRTFM::OnCommit(MoveTemp(F));
+			AutoRTFM::OnCommit(UE_AUTORTFM_MOVE(F));
 		}
 	};
 	struct FTransactHelper
@@ -908,7 +912,7 @@ namespace AutoRTFM::Private
 		template<typename FunctorType>
 		void operator+(FunctorType F)
 		{
-			AutoRTFM::Transact(MoveTemp(F));
+			AutoRTFM::Transact(UE_AUTORTFM_MOVE(F));
 		}
 	};
 } // namespace AutoRTFM::Private
