Fallback to regular allocator if memory can't be committed in the reserved virtual adress space (On platform that have very large page allocator enabled, 2MB physical pages are used when reserve / commit size is aligned to 2MB.
In certain cases, the platform can run out of 2MB physical pages. Falling back to regular allocator will ensure it will use a smaller physical page size).

#rnx

diff --git a/Engine/Source/Runtime/Core/Private/Memory/LinearVirtualMemoryAllocator.cpp b/Engine/Source/Runtime/Core/Private/Memory/LinearVirtualMemoryAllocator.cpp
--- a/Engine/Source/Runtime/Core/Private/Memory/LinearVirtualMemoryAllocator.cpp
+++ b/Engine/Source/Runtime/Core/Private/Memory/LinearVirtualMemoryAllocator.cpp
@@ -63,7 +63,16 @@ void* FLinearVirtualMemoryAllocator::Allocate(SIZE_T Size, uint32 Alignment)
 					LLM_PLATFORM_SCOPE(ELLMTag::FMalloc);
 					const size_t CommitGranularity = UE_USE_VERYLARGEPAGEALLOCATOR ? 2ul * 1024 * 1024 : 65536ul;
 					const SIZE_T ToCommit = Align(NewOffset - Committed, FMath::Max(VirtualMemory.GetCommitAlignment(), CommitGranularity));
+				#if UE_USE_VERYLARGEPAGEALLOCATOR
+					if (!VirtualMemory.Commit(Committed, ToCommit, false))
+					{
+						// do not try to use linear allocator anymore and fallback to FMemory::Malloc from now on
+						Reserved = Committed;
+						return FMemory::Malloc(Size, Alignment);
+					}
+				#else
 					VirtualMemory.Commit(Committed, ToCommit);
+				#endif
 					Committed += ToCommit;
 					LLM_IF_ENABLED(FLowLevelMemTracker::Get().OnLowLevelAlloc(ELLMTracker::Platform, (uint8*)VirtualMemory.GetVirtualPointer() + CurrentOffset, ToCommit));
 				}
