Fix a second AutoRTFM issue found with ASan and the text cache - use-after-free with FText::Empty() being double freed.

The bug would mean that an AutoRTFM transaction that aborts could mess up the internal reference count of the static empty string (got via FText::Empty()), eventually resulting in it being freed while a reference for it was still around.

diff --git a/Engine/Source/Programs/AutoRTFMTests/Private/UECoreTests.cpp b/Engine/Source/Programs/AutoRTFMTests/Private/UECoreTests.cpp
--- a/Engine/Source/Programs/AutoRTFMTests/Private/UECoreTests.cpp
+++ b/Engine/Source/Programs/AutoRTFMTests/Private/UECoreTests.cpp
@@ -5,6 +5,7 @@
 #include "Delegates/IDelegateInstance.h"
 #include "HAL/ThreadSingleton.h"
 #include "Internationalization/TextHistory.h"
+#include "Internationalization/TextCache.h"
 #include "Serialization/CustomVersion.h"
 #include "UObject/NameTypes.h"
 
@@ -574,3 +575,19 @@ TEST_CASE("UECore.TIntrusiveReferenceController")
 		}
 	}
 }
+
+TEST_CASE("UECore.FTextCache")
+{
+	FText Text;
+	REQUIRE(Text.IsEmpty());
+
+	AutoRTFM::ETransactionResult Result = AutoRTFM::Transact([&]
+		{
+			FString S("Sheesh");
+			Text = FTextCache::Get().FindOrCache(ANSI_TO_TCHAR("Oh my"), FTextInspector::GetTextId(FText::FromString(S)));
+			AutoRTFM::AbortTransaction();
+		});
+
+	REQUIRE(AutoRTFM::ETransactionResult::AbortedByRequest == Result);
+	REQUIRE(Text.IsEmpty());
+}
diff --git a/Engine/Source/Runtime/Core/Private/Internationalization/TextCache.cpp b/Engine/Source/Runtime/Core/Private/Internationalization/TextCache.cpp
--- a/Engine/Source/Runtime/Core/Private/Internationalization/TextCache.cpp
+++ b/Engine/Source/Runtime/Core/Private/Internationalization/TextCache.cpp
@@ -24,7 +24,8 @@ FText FTextCache::FindOrCache(const TCHAR* InTextLiteral, const TCHAR* InNamespa
 	return FindOrCache(InTextLiteral, FTextId(InNamespace, InKey));
 }
 
-FText FTextCache::FindOrCache(const TCHAR* InTextLiteral, const FTextId& InTextId)
+// #jira SOL-6589: We have to mark this noinline for now until our compiler gets smarter.
+UE_AUTORTFM_ALWAYS_OPEN FORCENOINLINE FText FTextCache::FindOrCache(const TCHAR* InTextLiteral, const FTextId& InTextId)
 {
 	LLM_SCOPE(ELLMTag::Localization);
 
@@ -32,22 +33,18 @@ FText FTextCache::FindOrCache(const TCHAR* InTextLiteral, const FTextId& InTextI
 	{
 		FText* ReturnFoundText = nullptr;
 
-		UE_AUTORTFM_OPEN2
-		{
-
-			FScopeLock Lock(&CachedTextCS);
+		FScopeLock Lock(&CachedTextCS);
 
-			FText* FoundText = CachedText.Find(InTextId);
-			if (FoundText)
+		FText* FoundText = CachedText.Find(InTextId);
+		if (FoundText)
+		{
+			const FString* FoundTextLiteral = FTextInspector::GetSourceString(*FoundText);
+			if (FoundTextLiteral && FCString::Strcmp(**FoundTextLiteral, InTextLiteral) == 0)
 			{
-				const FString* FoundTextLiteral = FTextInspector::GetSourceString(*FoundText);
-				if (FoundTextLiteral && FCString::Strcmp(**FoundTextLiteral, InTextLiteral) == 0)
-				{
 
-					ReturnFoundText = FoundText;
-				}
+				ReturnFoundText = FoundText;
 			}
-		};
+		}
 
 		if (ReturnFoundText)
 		{
@@ -56,17 +53,12 @@ FText FTextCache::FindOrCache(const TCHAR* InTextLiteral, const FTextId& InTextI
 	}
 
 	// Not currently cached, make a new instance...
-	FText NewText;
+	FText NewText = FText(InTextLiteral, InTextId.GetNamespace(), InTextId.GetKey(), ETextFlag::Immutable);
 
 	// ... and add it to the cache
-	UE_AUTORTFM_OPEN2
-	{
-		NewText = FText(InTextLiteral, InTextId.GetNamespace(), InTextId.GetKey(), ETextFlag::Immutable);
-
-		FScopeLock Lock(&CachedTextCS);
+	FScopeLock Lock(&CachedTextCS);
 
-		CachedText.Emplace(InTextId, NewText);
-	};
+	CachedText.Emplace(InTextId, NewText);
 
 	return NewText;
 }
