Change the slack growth factor from 1.25% to 1.2% IOS

diff --git a/Engine/Source/Runtime/Core/Public/Containers/ContainerAllocationPolicies.h b/Engine/Source/Runtime/Core/Public/Containers/ContainerAllocationPolicies.h
--- a/Engine/Source/Runtime/Core/Public/Containers/ContainerAllocationPolicies.h
+++ b/Engine/Source/Runtime/Core/Public/Containers/ContainerAllocationPolicies.h
@@ -108,6 +108,33 @@ struct FArraySlackTrackingHeader
 #define CONTAINER_INITIAL_ALLOC_ZERO_SLACK 1 // ON
 #endif
 
+#if defined(UE_CONTAINER_SLACK_GROWTH_FACTOR_NUMERATOR) && !defined(UE_CONTAINER_SLACK_GROWTH_FACTOR_DENOMINATOR)
+	#error If UE_CONTAINER_SLACK_GROWTH_FACTOR_NUMERATOR is defined you must also define UE_CONTAINER_SLACK_GROWTH_FACTOR_DENOMINATOR
+#endif
+
+#if defined(UE_CONTAINER_SLACK_GROWTH_FACTOR_DENOMINATOR) && !defined(UE_CONTAINER_SLACK_GROWTH_FACTOR_NUMERATOR)
+	#error If UE_CONTAINER_SLACK_GROWTH_FACTOR_DENOMINATOR is defined you must also define UE_CONTAINER_SLACK_GROWTH_FACTOR_NUMERATOR
+#endif
+
+#ifndef UE_CONTAINER_SLACK_GROWTH_FACTOR_NUMERATOR
+	#if AGGRESSIVE_MEMORY_SAVING
+		#define UE_CONTAINER_SLACK_GROWTH_FACTOR_NUMERATOR 1
+	#else
+		#define UE_CONTAINER_SLACK_GROWTH_FACTOR_NUMERATOR 3
+	#endif
+#endif
+
+#ifndef UE_CONTAINER_SLACK_GROWTH_FACTOR_DENOMINATOR
+	#if AGGRESSIVE_MEMORY_SAVING
+		#define UE_CONTAINER_SLACK_GROWTH_FACTOR_DENOMINATOR 4
+	#else
+		#define UE_CONTAINER_SLACK_GROWTH_FACTOR_DENOMINATOR 8
+	#endif
+#endif
+static_assert(UE_CONTAINER_SLACK_GROWTH_FACTOR_DENOMINATOR > 0, "UE_CONTAINER_SLACK_GROWTH_FACTOR_DENOMINATOR must be greater than 0");
+static_assert(UE_CONTAINER_SLACK_GROWTH_FACTOR_DENOMINATOR > UE_CONTAINER_SLACK_GROWTH_FACTOR_NUMERATOR, "UE_CONTAINER_SLACK_GROWTH_FACTOR_DENOMINATOR must be greater than UE_CONTAINER_SLACK_GROWTH_FACTOR_NUMERATOR");
+
+
 class FDefaultBitArrayAllocator;
 
 template<int IndexSize> class TSizedDefaultAllocator;
@@ -151,6 +178,7 @@ FORCEINLINE SizeType DefaultCalculateSlackGrow(SizeType NumElements, SizeType Nu
 #endif
 #if AGGRESSIVE_MEMORY_SAVING
 	const SIZE_T FirstGrow = 1;
+	const SIZE_T ConstantGrow = 0;
 #else
 	const SIZE_T FirstGrow = 4;
 	const SIZE_T ConstantGrow = 16;
@@ -165,11 +193,7 @@ FORCEINLINE SizeType DefaultCalculateSlackGrow(SizeType NumElements, SizeType Nu
 	if (NumAllocatedElements)
 	{
 		// Allocate slack for the array proportional to its size.
-#if AGGRESSIVE_MEMORY_SAVING
-		Grow = SIZE_T(NumElements) + SIZE_T(NumElements) / 4;
-#else
-		Grow = SIZE_T(NumElements) + 3 * SIZE_T(NumElements) / 8 + ConstantGrow;
-#endif
+		Grow = SIZE_T(NumElements) + UE_CONTAINER_SLACK_GROWTH_FACTOR_NUMERATOR * SIZE_T(NumElements) / UE_CONTAINER_SLACK_GROWTH_FACTOR_DENOMINATOR + ConstantGrow;
 	}
 	else if (SIZE_T(NumElements) > Grow)
 	{
@@ -179,11 +203,7 @@ FORCEINLINE SizeType DefaultCalculateSlackGrow(SizeType NumElements, SizeType Nu
 	if (NumAllocatedElements || SIZE_T(NumElements) > Grow)
 	{
 		// Allocate slack for the array proportional to its size.
-#if AGGRESSIVE_MEMORY_SAVING
-		Grow = SIZE_T(NumElements) + SIZE_T(NumElements) / 4;
-#else
-		Grow = SIZE_T(NumElements) + 3 * SIZE_T(NumElements) / 8 + ConstantGrow;
-#endif
+		Grow = SIZE_T(NumElements) + UE_CONTAINER_SLACK_GROWTH_FACTOR_NUMERATOR * SIZE_T(NumElements) / UE_CONTAINER_SLACK_GROWTH_FACTOR_DENOMINATOR + ConstantGrow;
 	}
 #endif
 	
