Added comparison operators to TNonNullPtr to allow it to function properly inside TOptional<TNonNullPtr<T>>
- TOptional, when using bUsingIntrusiveUnsetState, will rely on its inner type for equality comparison, regardless if it's set or not
- the problem was that when you tried to compare two unset TOptional<TNonNullPtr<T>>, the compiler resorted to comparing the raw pointers which are guarded by a null check so it lead to a crash
- adding explicit comparison operators to TNonNullPtr fixes this problem

#rnx

diff --git a/Engine/Source/Runtime/Core/Public/Templates/NonNullPointer.h b/Engine/Source/Runtime/Core/Public/Templates/NonNullPointer.h
--- a/Engine/Source/Runtime/Core/Public/Templates/NonNullPointer.h
+++ b/Engine/Source/Runtime/Core/Public/Templates/NonNullPointer.h
@@ -93,6 +93,58 @@ public:
 		return *this;
 	}
 
+	/**
+	 * Comparison, will also handle default constructed state
+	 */
+	FORCEINLINE bool operator==(const TNonNullPtr& Other) const
+	{
+		return Object == Other.Object;
+	}
+#if !PLATFORM_COMPILER_HAS_GENERATED_COMPARISON_OPERATORS
+	FORCEINLINE bool operator!=(const TNonNullPtr& Other) const
+	{
+		return Object != Other.Object;
+	}
+#endif
+
+	/**
+	 * Comparison with a raw pointer
+	 */
+	template <
+		typename OtherObjectType
+		UE_REQUIRES(UE_REQUIRES_EXPR(std::declval<ObjectType*>() == std::declval<OtherObjectType*>()))
+	>
+	FORCEINLINE bool operator==(OtherObjectType* Other) const
+	{
+		return Object == Other;
+	}
+	template <
+		typename OtherObjectType
+		UE_REQUIRES(UE_REQUIRES_EXPR(std::declval<OtherObjectType*>() == std::declval<ObjectType*>()))
+	>
+	FORCEINLINE friend bool operator==(OtherObjectType* Lhs, const TNonNullPtr& Rhs)
+	{
+		return Lhs == Rhs.Object;
+	}
+#if !PLATFORM_COMPILER_HAS_GENERATED_COMPARISON_OPERATORS
+	template <
+		typename OtherObjectType
+		UE_REQUIRES(UE_REQUIRES_EXPR(std::declval<ObjectType*>() == std::declval<OtherObjectType*>()))
+	>
+	FORCEINLINE bool operator!=(OtherObjectType* Other) const
+	{
+		return Object != Other;
+	}
+	template <
+		typename OtherObjectType
+		UE_REQUIRES(UE_REQUIRES_EXPR(std::declval<OtherObjectType*>() == std::declval<ObjectType*>()))
+	>
+	FORCEINLINE friend bool operator!=(OtherObjectType* Lhs, const TNonNullPtr& Rhs)
+	{
+		return Lhs != Rhs.Object;
+	}
+#endif
+
 	/**
 	 * Returns the internal pointer
 	 */
