Remove all possible PLATFORM_SWITCH macro usage from Plugins.

#rnx

#changelist validated

diff --git a/Engine/Source/Runtime/Core/Public/GenericPlatform/GenericPlatformMisc.h b/Engine/Source/Runtime/Core/Public/GenericPlatform/GenericPlatformMisc.h
--- a/Engine/Source/Runtime/Core/Public/GenericPlatform/GenericPlatformMisc.h
+++ b/Engine/Source/Runtime/Core/Public/GenericPlatform/GenericPlatformMisc.h
@@ -1578,6 +1578,8 @@ public:
 		return ENetworkConnectionType::Unknown;
 	}
 
+	static void EnsureNetworkIfNecessary()	{}
+
 	/**
 	 * Returns whether the platform has variable hardware (configurable/upgradeable system).
 	 */
