Make delegates in stat system thread safe

Users have reported intermittent crashes in stat system due to registering delegates while stat system is broadcasting using the same delegate. Use a thread safe policy to make this safe.

diff --git a/Engine/Source/Runtime/Core/Public/Stats/StatsData.h b/Engine/Source/Runtime/Core/Public/Stats/StatsData.h
--- a/Engine/Source/Runtime/Core/Public/Stats/StatsData.h
+++ b/Engine/Source/Runtime/Core/Public/Stats/StatsData.h
@@ -500,10 +500,10 @@ class FStatsThreadState
 	friend struct FStatsReadFile;
 
 	/** Delegate that FStatsThreadState calls on the stats thread whenever we have a new frame. */
-	DECLARE_MULTICAST_DELEGATE_OneParam(FOnNewFrameHistory, int64);
+	DECLARE_TS_MULTICAST_DELEGATE_OneParam(FOnNewFrameHistory, int64);
 
 	/** Delegate that FStatsThreadState calls on the stats thread whenever we have a new raw stats packet. */
-	DECLARE_MULTICAST_DELEGATE_OneParam(FOnNewRawStatPacket, const FStatPacket*);
+	DECLARE_TS_MULTICAST_DELEGATE_OneParam(FOnNewRawStatPacket, const FStatPacket*);
 
 	/** Internal method to scan the messages to update the current frame. */
 	CORE_API void ScanForAdvance(const FStatMessagesArray& Data);
