Add test to make sure the module name in the IMPLEMENT_MODULE  macros matches the name declared in the uproject, uplugin or the name of the folder for engine modules.

A mismatch for Runtime module results in a error in FModuleManager::LoadModuleWithFailureReason because it can't find the ModuleInitializerPtr

See ModuleManager.h for the actual change. The rest are fixes of wrong names.

diff --git a/Engine/Plugins/Experimental/AVCodecs/WMFCodecs/Source/WMFCodecs/Private/WMFCodecsModule.cpp b/Engine/Plugins/Experimental/AVCodecs/WMFCodecs/Source/WMFCodecs/Private/WMFCodecsModule.cpp
--- a/Engine/Plugins/Experimental/AVCodecs/WMFCodecs/Source/WMFCodecs/Private/WMFCodecsModule.cpp
+++ b/Engine/Plugins/Experimental/AVCodecs/WMFCodecs/Source/WMFCodecs/Private/WMFCodecsModule.cpp
@@ -23,4 +23,4 @@ public:
 	}
 };
 
-IMPLEMENT_MODULE(FWMFCodecModule, WMFCodec);
+IMPLEMENT_MODULE(FWMFCodecModule, WMFCodecs);
diff --git a/Engine/Plugins/Experimental/Avalanche/Source/AvalancheRemoteControlEditor/Private/AvalancheRemoteControlEditorModule.cpp b/Engine/Plugins/Experimental/Avalanche/Source/AvalancheRemoteControlEditor/Private/AvalancheRemoteControlEditorModule.cpp
--- a/Engine/Plugins/Experimental/Avalanche/Source/AvalancheRemoteControlEditor/Private/AvalancheRemoteControlEditorModule.cpp
+++ b/Engine/Plugins/Experimental/Avalanche/Source/AvalancheRemoteControlEditor/Private/AvalancheRemoteControlEditorModule.cpp
@@ -72,4 +72,4 @@ void FAvalancheRemoteControlEditorModule::UnregisterCustomizations()
 
 #undef LOCTEXT_NAMESPACE
 
-IMPLEMENT_MODULE(FAvalancheRemoteControlEditorModule, AvalancheRemoteControlComponentsEditor)
+IMPLEMENT_MODULE(FAvalancheRemoteControlEditorModule, AvalancheRemoteControlEditor)
diff --git a/Engine/Plugins/Experimental/Avalanche/Source/AvalancheSVGEditor/Private/AvaSVGEditorModule.cpp b/Engine/Plugins/Experimental/Avalanche/Source/AvalancheSVGEditor/Private/AvaSVGEditorModule.cpp
--- a/Engine/Plugins/Experimental/Avalanche/Source/AvalancheSVGEditor/Private/AvaSVGEditorModule.cpp
+++ b/Engine/Plugins/Experimental/Avalanche/Source/AvalancheSVGEditor/Private/AvaSVGEditorModule.cpp
@@ -165,4 +165,4 @@ void FAvaSVGEditorModule::OnSVGShapesUpdated(AActor* InActor) const
 
 #undef LOCTEXT_NAMESPACE
 
-IMPLEMENT_MODULE(FAvaSVGEditorModule, AvalancheSVGEditorModule)
+IMPLEMENT_MODULE(FAvaSVGEditorModule, AvalancheSVGEditor)
diff --git a/Engine/Plugins/Experimental/Avalanche/Source/AvalancheTextEditor/Private/AvaTextEditorModule.cpp b/Engine/Plugins/Experimental/Avalanche/Source/AvalancheTextEditor/Private/AvaTextEditorModule.cpp
--- a/Engine/Plugins/Experimental/Avalanche/Source/AvalancheTextEditor/Private/AvaTextEditorModule.cpp
+++ b/Engine/Plugins/Experimental/Avalanche/Source/AvalancheTextEditor/Private/AvaTextEditorModule.cpp
@@ -152,4 +152,4 @@ void FAvaTextEditorModule::RegisterDynamicMaterialPropertyGenerator()
 	);
 }
 
-IMPLEMENT_MODULE(FAvaTextEditorModule, AvalancheTextEditorModule)
+IMPLEMENT_MODULE(FAvaTextEditorModule, AvalancheTextEditor)
diff --git a/Engine/Plugins/Experimental/UserToolBoxBasicCommand/Source/UserToolBoxBasicCommand/Private/UserToolBoxBasicCommand.cpp b/Engine/Plugins/Experimental/UserToolBoxBasicCommand/Source/UserToolBoxBasicCommand/Private/UserToolBoxBasicCommand.cpp
--- a/Engine/Plugins/Experimental/UserToolBoxBasicCommand/Source/UserToolBoxBasicCommand/Private/UserToolBoxBasicCommand.cpp
+++ b/Engine/Plugins/Experimental/UserToolBoxBasicCommand/Source/UserToolBoxBasicCommand/Private/UserToolBoxBasicCommand.cpp
@@ -31,4 +31,4 @@ void FUserToolBox_BasicCommandModule::ShutdownModule()
 
 #undef LOCTEXT_NAMESPACE
 
-IMPLEMENT_MODULE(FUserToolBox_BasicCommandModule, UserToolBox_BasicCommand)
\ No newline at end of file
+IMPLEMENT_MODULE(FUserToolBox_BasicCommandModule, UserToolBoxBasicCommand)
\ No newline at end of file
diff --git a/Engine/Plugins/Media/NDIMedia/Source/NDIMedia/Private/NDIMediaModule.cpp b/Engine/Plugins/Media/NDIMedia/Source/NDIMedia/Private/NDIMediaModule.cpp
--- a/Engine/Plugins/Media/NDIMedia/Source/NDIMedia/Private/NDIMediaModule.cpp
+++ b/Engine/Plugins/Media/NDIMedia/Source/NDIMedia/Private/NDIMediaModule.cpp
@@ -117,4 +117,4 @@ bool FNDIMediaModule::LoadModuleDependencies()
 
 #undef LOCTEXT_NAMESPACE
 	
-IMPLEMENT_MODULE(FNDIMediaModule, NDIMediaOut)
+IMPLEMENT_MODULE(FNDIMediaModule, NDIMedia)
diff --git a/Engine/Plugins/Media/NDIMedia/Source/NDIMediaEditor/Private/NDIMediaEditorModule.cpp b/Engine/Plugins/Media/NDIMedia/Source/NDIMediaEditor/Private/NDIMediaEditorModule.cpp
--- a/Engine/Plugins/Media/NDIMedia/Source/NDIMediaEditor/Private/NDIMediaEditorModule.cpp
+++ b/Engine/Plugins/Media/NDIMedia/Source/NDIMediaEditor/Private/NDIMediaEditorModule.cpp
@@ -60,4 +60,4 @@ private:
 
 #undef LOCTEXT_NAMESPACE
 
-IMPLEMENT_MODULE(FNDIMediaEditorModule, NDIMediaOutEditor)
+IMPLEMENT_MODULE(FNDIMediaEditorModule, NDIMediaEditor)
diff --git a/Engine/Plugins/Runtime/Metasound/Source/MetasoundEngine/Private/MetasoundEngineModule.cpp b/Engine/Plugins/Runtime/Metasound/Source/MetasoundEngine/Private/MetasoundEngineModule.cpp
--- a/Engine/Plugins/Runtime/Metasound/Source/MetasoundEngine/Private/MetasoundEngineModule.cpp
+++ b/Engine/Plugins/Runtime/Metasound/Source/MetasoundEngine/Private/MetasoundEngineModule.cpp
@@ -174,4 +174,4 @@ namespace Metasound::Engine
 	};
 } // namespace Metasound::Engine
 
-IMPLEMENT_MODULE(Metasound::Engine::FModule, MetaSoundEngine);
+IMPLEMENT_MODULE(Metasound::Engine::FModule, MetasoundEngine);
diff --git a/Engine/Source/Developer/LowLevelTestsRunner/Private/TestStubs/PlatformFileModuleStub.cpp b/Engine/Source/Developer/LowLevelTestsRunner/Private/TestStubs/PlatformFileModuleStub.cpp
--- a/Engine/Source/Developer/LowLevelTestsRunner/Private/TestStubs/PlatformFileModuleStub.cpp
+++ b/Engine/Source/Developer/LowLevelTestsRunner/Private/TestStubs/PlatformFileModuleStub.cpp
@@ -48,4 +48,4 @@ public:
 };
 
 
-IMPLEMENT_MODULE(FPlatformFileModuleStub, PlatformFileStub);
\ No newline at end of file
+IMPLEMENT_MODULE(FPlatformFileModuleStub, LowLevelTestsRunner);
\ No newline at end of file
diff --git a/Engine/Source/Developer/PhysicsUtilities/Private/PhysicsAssetUtils.cpp b/Engine/Source/Developer/PhysicsUtilities/Private/PhysicsAssetUtils.cpp
--- a/Engine/Source/Developer/PhysicsUtilities/Private/PhysicsAssetUtils.cpp
+++ b/Engine/Source/Developer/PhysicsUtilities/Private/PhysicsAssetUtils.cpp
@@ -1296,4 +1296,4 @@ void SanitizeRestrictedContent(UPhysicsAsset* PhysAsset)
 
 }; // namespace FPhysicsAssetUtils
 
-IMPLEMENT_MODULE(FDefaultModuleImpl, PhysicsAssetUtils)
+IMPLEMENT_MODULE(FDefaultModuleImpl, PhysicsUtilities)
diff --git a/Engine/Source/Runtime/AudioLink/AudioLinkEngine/Private/AudioLinkModule.cpp b/Engine/Source/Runtime/AudioLink/AudioLinkEngine/Private/AudioLinkModule.cpp
--- a/Engine/Source/Runtime/AudioLink/AudioLinkEngine/Private/AudioLinkModule.cpp
+++ b/Engine/Source/Runtime/AudioLink/AudioLinkEngine/Private/AudioLinkModule.cpp
@@ -18,4 +18,4 @@ public:
 	}
 };
 
-IMPLEMENT_MODULE(FAudioLinkModule, AudioLink);
+IMPLEMENT_MODULE(FAudioLinkModule, AudioLinkEngine);
diff --git a/Engine/Source/Runtime/Core/Public/Modules/ModuleManager.h b/Engine/Source/Runtime/Core/Public/Modules/ModuleManager.h
--- a/Engine/Source/Runtime/Core/Public/Modules/ModuleManager.h
+++ b/Engine/Source/Runtime/Core/Public/Modules/ModuleManager.h
@@ -816,6 +816,28 @@ class FDefaultGameModuleImpl
  *
  * @see IMPLEMENT_GAME_MODULE
  */
+
+namespace UE::Core::Private
+{
+	constexpr bool ModuleNameEquals(const char* Lhs, const char* Rhs)
+	{
+		for (;;)
+		{
+			if (*Lhs != *Rhs)
+			{
+				return false;
+			}
+
+			if (*Lhs == '\0')
+			{
+				return true;
+			}
+
+			++Lhs;
+			++Rhs;
+		}
+	}
+}
 #if IS_MONOLITHIC
 
 	// If we're linking monolithically we assume all modules are linked in with the main binary.
@@ -823,7 +845,7 @@ class FDefaultGameModuleImpl
 		/** Global registrant object for this module when linked statically */ \
 		static FStaticallyLinkedModuleRegistrant< ModuleImplClass > ModuleRegistrant##ModuleName( TEXT(#ModuleName) ); \
 		/* Forced reference to this function is added by the linker to check that each module uses IMPLEMENT_MODULE */ \
-		extern "C" void IMPLEMENT_MODULE_##ModuleName() { } \
+		extern "C" void IMPLEMENT_MODULE_##ModuleName() { UE_STATIC_ASSERT_WARN(UE::Core::Private::ModuleNameEquals(#ModuleName, UE_MODULE_NAME ), "Module name mismatch (" #ModuleName " != " UE_MODULE_NAME "). Please ensure module name passed to IMPLEMENT_MODULE is " UE_MODULE_NAME " to avoid runtime errors in monolithic builds."); } \
 		PER_MODULE_BOILERPLATE_ANYLINK(ModuleImplClass, ModuleName)
 
 #else
@@ -841,7 +863,7 @@ class FDefaultGameModuleImpl
 		} \
 		static FModuleInitializerEntry ModuleName##InitializerEntry(TEXT(#ModuleName), Initialize##ModuleName##Module, TEXT(UE_MODULE_NAME)); \
 		/* Forced reference to this function is added by the linker to check that each module uses IMPLEMENT_MODULE */ \
-		extern "C" void IMPLEMENT_MODULE_##ModuleName() { } \
+		extern "C" void IMPLEMENT_MODULE_##ModuleName() { UE_STATIC_ASSERT_WARN(UE::Core::Private::ModuleNameEquals(#ModuleName, UE_MODULE_NAME ), "Module name mismatch (" #ModuleName " != " UE_MODULE_NAME "). Please ensure module name passed to IMPLEMENT_MODULE is " UE_MODULE_NAME " to avoid runtime errors in monolithic builds."); } \
 		PER_MODULE_BOILERPLATE_ANYLINK(ModuleImplClass, ModuleName)
 
 #endif //IS_MONOLITHIC
