Added a cvar to whether enable or not flushing registered MB thread caches on one thread (enabled by default)

#rnx

diff --git a/Engine/Source/Runtime/Core/Private/HAL/MallocBinnedCommon.cpp b/Engine/Source/Runtime/Core/Private/HAL/MallocBinnedCommon.cpp
--- a/Engine/Source/Runtime/Core/Private/HAL/MallocBinnedCommon.cpp
+++ b/Engine/Source/Runtime/Core/Private/HAL/MallocBinnedCommon.cpp
@@ -530,6 +530,12 @@ static FAutoConsoleVariableRef GMallocBinnedFlushThreadCacheMaxWaitTimeCVar(
 	ECVF_ReadOnly
 );
 
+int32 GMallocBinnedFlushRegisteredThreadCachesOnOneThread = 1;
+static FAutoConsoleVariableRef GMallocBinnedFlushRegisteredThreadCachesOnOneThreadCVar(
+	TEXT("MallocBinned.FlushRegisteredThreadCachesOnOneThread"),
+	GMallocBinnedFlushRegisteredThreadCachesOnOneThread,
+	TEXT("Wether or not to attempt to flush registered thread caches on one thread (enabled by default)."));
+
 #if UE_BINNEDCOMMON_ALLOW_RUNTIME_TWEAKING
 
 int32 GMallocBinnedBundleSize = DEFAULT_GMallocBinnedBundleSize;
diff --git a/Engine/Source/Runtime/Core/Public/HAL/MallocBinnedCommon.h b/Engine/Source/Runtime/Core/Public/HAL/MallocBinnedCommon.h
--- a/Engine/Source/Runtime/Core/Public/HAL/MallocBinnedCommon.h
+++ b/Engine/Source/Runtime/Core/Public/HAL/MallocBinnedCommon.h
@@ -181,6 +181,7 @@ extern CORE_API int32 GMallocBinnedBundleCount;
 #endif
 
 extern CORE_API float GMallocBinnedFlushThreadCacheMaxWaitTime;
+extern CORE_API int32 GMallocBinnedFlushRegisteredThreadCachesOnOneThread;
 
 class FMallocBinnedCommonBase : public FMalloc
 {
diff --git a/Engine/Source/Runtime/Core/Public/HAL/MallocBinnedCommonUtils.h b/Engine/Source/Runtime/Core/Public/HAL/MallocBinnedCommonUtils.h
--- a/Engine/Source/Runtime/Core/Public/HAL/MallocBinnedCommonUtils.h
+++ b/Engine/Source/Runtime/Core/Public/HAL/MallocBinnedCommonUtils.h
@@ -74,7 +74,7 @@ public:
 
 		// Process thread-local memory caches from as many threads as possible without waking them up.
 		// Skip on desktop as we may have too many threads and this could cause some hitches.
-		if (!PLATFORM_DESKTOP)
+		if (!PLATFORM_DESKTOP && GMallocBinnedFlushRegisteredThreadCachesOnOneThread != 0)
 		{
 			FScopeLock Lock(&Allocator.GetMutex());
 			FScopeLock FreeBlockLock(&AllocType::GetFreeBlockListsRegistrationMutex());
