always forcenoinline for ALWAYS_OPEN macro due to inlining removing this feature
*once* SOL-6589 is fixed this can be removed

diff --git a/Engine/Source/Runtime/Core/Private/Internationalization/TextCache.cpp b/Engine/Source/Runtime/Core/Private/Internationalization/TextCache.cpp
--- a/Engine/Source/Runtime/Core/Private/Internationalization/TextCache.cpp
+++ b/Engine/Source/Runtime/Core/Private/Internationalization/TextCache.cpp
@@ -8,6 +8,7 @@
 #include "Misc/CString.h"
 #include "Misc/LazySingleton.h"
 #include "Misc/ScopeLock.h"
+#include "AutoRTFM/AutoRTFM.h"
 
 FTextCache& FTextCache::Get()
 {
@@ -24,8 +25,7 @@ FText FTextCache::FindOrCache(const TCHAR* InTextLiteral, const TCHAR* InNamespa
 	return FindOrCache(InTextLiteral, FTextId(InNamespace, InKey));
 }
 
-// #jira SOL-6589: We have to mark this noinline for now until our compiler gets smarter.
-UE_AUTORTFM_ALWAYS_OPEN FORCENOINLINE FText FTextCache::FindOrCache(const TCHAR* InTextLiteral, const FTextId& InTextId)
+UE_AUTORTFM_ALWAYS_OPEN FText FTextCache::FindOrCache(const TCHAR* InTextLiteral, const FTextId& InTextId)
 {
 	LLM_SCOPE(ELLMTag::Localization);
 
diff --git a/Engine/Source/Runtime/Core/Public/AutoRTFM/AutoRTFM.h b/Engine/Source/Runtime/Core/Public/AutoRTFM/AutoRTFM.h
--- a/Engine/Source/Runtime/Core/Public/AutoRTFM/AutoRTFM.h
+++ b/Engine/Source/Runtime/Core/Public/AutoRTFM/AutoRTFM.h
@@ -31,7 +31,8 @@
 #if UE_AUTORTFM
 #define UE_AUTORTFM_AUTORTFM(F) [[clang::autortfm(F), clang::noinline]]
 #define UE_AUTORTFM_NOAUTORTFM [[clang::noautortfm, clang::noinline]]
-#define UE_AUTORTFM_ALWAYS_OPEN [[clang::autortfm_always_open]]
+// #jira SOL-6589: remove clang::noinline once this JIRA is fixed.
+#define UE_AUTORTFM_ALWAYS_OPEN [[clang::autortfm_always_open, clang::noinline]]
 #else
 #define UE_AUTORTFM_AUTORTFM(F)
 #define UE_AUTORTFM_NOAUTORTFM
diff --git a/Engine/Source/Runtime/GameplayTags/Private/GameplayTagsManager.cpp b/Engine/Source/Runtime/GameplayTags/Private/GameplayTagsManager.cpp
--- a/Engine/Source/Runtime/GameplayTags/Private/GameplayTagsManager.cpp
+++ b/Engine/Source/Runtime/GameplayTags/Private/GameplayTagsManager.cpp
@@ -2216,9 +2216,8 @@ void UGameplayTagsManager::DoneAddingNativeTags()
 	}
 }
 
-// #jira SOL-6589: remove forcenoinline once this JIRA is fixed.
 UE_AUTORTFM_ALWAYS_OPEN
-FORCENOINLINE FGameplayTagContainer UGameplayTagsManager::RequestGameplayTagParents(const FGameplayTag& GameplayTag) const
+FGameplayTagContainer UGameplayTagsManager::RequestGameplayTagParents(const FGameplayTag& GameplayTag) const
 {
 	FScopeLock Lock(&GameplayTagMapCritical);
 
