Log TraceLog control port

When running multiple instances on one the same machine TraceLog searches for an open port to listen to control commands, but there was no way for the user to know which port each process ended up using. Add log entries that print the final listening port.

diff --git a/Engine/Source/Runtime/Core/Private/ProfilingDebugging/TraceAuxiliary.cpp b/Engine/Source/Runtime/Core/Private/ProfilingDebugging/TraceAuxiliary.cpp
--- a/Engine/Source/Runtime/Core/Private/ProfilingDebugging/TraceAuxiliary.cpp
+++ b/Engine/Source/Runtime/Core/Private/ProfilingDebugging/TraceAuxiliary.cpp
@@ -842,14 +842,32 @@ void OnConnectionCallback()
 void OnMessageCallback(const UE::Trace::FMessageEvent& Message)
 {
 	const auto TypeStr = StringCast<TCHAR>(Message.TypeStr);
-	if (Message.Description != nullptr)
+	switch(Message.Type)
 	{
-		const auto Description = StringCast<TCHAR>(Message.Description);
-		UE_LOG(LogTrace, Error, TEXT("%s %s"), TypeStr.Get(), Description.Get());
-	}
-	else
-	{
-		UE_LOG(LogTrace, Error, TEXT("%s"), TypeStr.Get());
+	// Information to the user
+	case UE::Trace::EMessageType::Info:
+		{
+			if (Message.Description != nullptr)
+			{
+				const auto Description = StringCast<TCHAR>(Message.Description);
+				UE_LOG(LogTrace, Display, TEXT("%s"), Description.Get());
+			}
+		}
+		break;
+	// By default treated as errors	
+	default:
+		{
+			if (Message.Description != nullptr)
+			{
+				const auto Description = StringCast<TCHAR>(Message.Description);
+				UE_LOG(LogTrace, Error, TEXT("%s %s"), TypeStr.Get(), Description.Get());
+			}
+			else
+			{
+				UE_LOG(LogTrace, Error, TEXT("%s"), TypeStr.Get());
+			}
+		}
+		break;
 	}
 }
 
diff --git a/Engine/Source/Runtime/TraceLog/Private/Trace/Control.cpp b/Engine/Source/Runtime/TraceLog/Private/Trace/Control.cpp
--- a/Engine/Source/Runtime/TraceLog/Private/Trace/Control.cpp
+++ b/Engine/Source/Runtime/TraceLog/Private/Trace/Control.cpp
@@ -134,6 +134,7 @@ static bool Writer_ControlListen()
 		return false;
 	}
 
+	UE_TRACE_MESSAGE_F(Info, "Control listening on port %u", GControlPort);
 	GControlState = EControlState::Listening;
 	return true;
 }
diff --git a/Engine/Source/Runtime/TraceLog/Public/Trace/Trace.h b/Engine/Source/Runtime/TraceLog/Public/Trace/Trace.h
--- a/Engine/Source/Runtime/TraceLog/Public/Trace/Trace.h
+++ b/Engine/Source/Runtime/TraceLog/Public/Trace/Trace.h
@@ -89,6 +89,7 @@ enum class EMessageType : uint8
 	EstablishError,
 	FileOpenError,
 	WriterError,
+	Info				= 0xff,
 };
 	
 struct FMessageEvent
