Change AutoRTFM language issues (missing functions, atomics in transactional code, etc) into ensures and a BPVM script exception.

#okforversepublic

diff --git a/Engine/Source/Programs/AutoRTFMTests/Private/Abort.cpp b/Engine/Source/Programs/AutoRTFMTests/Private/Abort.cpp
--- a/Engine/Source/Programs/AutoRTFMTests/Private/Abort.cpp
+++ b/Engine/Source/Programs/AutoRTFMTests/Private/Abort.cpp
@@ -476,3 +476,42 @@ TEST_CASE("Abort.OnAbortTiming")
 	REQUIRE(Memory == 666);
 	REQUIRE(bOnAbortRan == true);
 }
+
+TEST_CASE("Abort.Language")
+{
+	bool bTouched = false;
+	std::atomic_bool bThingie = false;
+
+	const AutoRTFM::ETransactionResult Result = AutoRTFM::Transact([&]
+		{
+			bTouched = true;
+			bThingie = true;
+		});
+
+	REQUIRE(AutoRTFM::ETransactionResult::AbortedByLanguage == Result);
+	REQUIRE(false == bTouched);
+}
+
+TEST_CASE("Abort.LanguageThroughOpen")
+{
+	bool bTouched = false;
+	std::atomic_bool bThingie = false;
+
+	const AutoRTFM::ETransactionResult Result = AutoRTFM::Transact([&]
+		{
+			bTouched = true;
+
+			AutoRTFM::Open([&]
+				{
+					const AutoRTFM::EContextStatus Status = AutoRTFM::Close([&]
+						{
+							bThingie = true;
+						});
+
+					REQUIRE(AutoRTFM::EContextStatus::AbortedByLanguage == Status);
+				});
+		});
+
+	REQUIRE(AutoRTFM::ETransactionResult::AbortedByLanguage == Result);
+	REQUIRE(false == bTouched);
+}
diff --git a/Engine/Source/Programs/AutoRTFMTests/Private/AutoRTFMTests.cpp b/Engine/Source/Programs/AutoRTFMTests/Private/AutoRTFMTests.cpp
--- a/Engine/Source/Programs/AutoRTFMTests/Private/AutoRTFMTests.cpp
+++ b/Engine/Source/Programs/AutoRTFMTests/Private/AutoRTFMTests.cpp
@@ -17,7 +17,11 @@ public:
 		GEngineLoop.PreInit(0, nullptr);
 		FModuleManager::Get().StartProcessingNewlyLoadedObjects();
 
-		AutoRTFM::ForTheRuntime::SetAutoRTFMRuntime(AutoRTFM::ForTheRuntime::EAutoRTFMEnabledState::AutoRTFM_Enabled);
+		// Enable all Verse code to run under AutoRTFM (shouldn't affect our tests here, but better safe than sorry).
+		AutoRTFM::ForTheRuntime::SetAutoRTFMRuntime(AutoRTFM::ForTheRuntime::EAutoRTFMEnabledState::AutoRTFM_EnabledForAllVerse);
+
+		// We don't want to trigger ensure's on abort because we are going to test that.
+		AutoRTFM::ForTheRuntime::SetEnsureOnAbortByLanguage(false);
 	}
 
 	void testRunEnded(const Catch::TestRunStats&) override
diff --git a/Engine/Source/Runtime/Core/Private/AutoRTFM/API.cpp b/Engine/Source/Runtime/Core/Private/AutoRTFM/API.cpp
--- a/Engine/Source/Runtime/Core/Private/AutoRTFM/API.cpp
+++ b/Engine/Source/Runtime/Core/Private/AutoRTFM/API.cpp
@@ -14,6 +14,8 @@ namespace
 	int GAutoRTFMRuntimeEnabled = AutoRTFM::EAutoRTFMEnabledState::AutoRTFM_Disabled;
 #endif // UE_AUTORTFM_ENABLED_RUNTIME_BY_DEFAULT
 
+	bool GAutoRTFMEnsureOnAbortByLanguage = true;
+
 	void UpdateAutoRTFMRuntimeCrashData()
 	{
 		FGenericCrashContext::SetGameData(TEXT("IsAutoRTFMRuntimeEnabled"), AutoRTFM::ForTheRuntime::IsAutoRTFMRuntimeEnabled() ? TEXT("true") : TEXT("false"));
@@ -96,6 +98,22 @@ namespace AutoRTFM
 			}
 #else
 			return false;
+#endif
+		}
+
+		void SetEnsureOnAbortByLanguage(bool bEnabled)
+		{
+#if UE_AUTORTFM
+			GAutoRTFMEnsureOnAbortByLanguage = bEnabled;
+#endif
+		}
+
+		bool IsEnsureOnAbortByLanguageEnabled()
+		{
+#if UE_AUTORTFM
+			return GAutoRTFMEnsureOnAbortByLanguage;
+#else
+			return false;
 #endif
 		}
 	}
@@ -242,7 +260,7 @@ extern "C" UE_AUTORTFM_AUTORTFM("RTFM_autortfm_close") autortfm_status autortfm_
 		}
 		else
 		{
-			UE_LOG(LogAutoRTFM, Fatal, TEXT("Could not find function %p '%s' where '%s'."), UninstrumentedWork, *GetFunctionDescription(UninstrumentedWork), ANSI_TO_TCHAR("autortfm_close"));
+			ensureMsgf(!ForTheRuntime::IsEnsureOnAbortByLanguageEnabled(), TEXT("Could not find function %p '%s' where '%s'."), UninstrumentedWork, *GetFunctionDescription(UninstrumentedWork), ANSI_TO_TCHAR("autortfm_close"));
 	        Context->AbortByLanguageAndThrow();
 		}
 	}
@@ -470,7 +488,7 @@ extern "C" UE_AUTORTFM_NOAUTORTFM autortfm_status RTFM_autortfm_close(void (*Uni
     }
 	else
 	{
-		UE_LOG(LogAutoRTFM, Fatal, TEXT("Could not find function %p '%s' where '%s'."), UninstrumentedWork, *GetFunctionDescription(UninstrumentedWork), ANSI_TO_TCHAR("autortfm_close"));
+		ensureMsgf(!ForTheRuntime::IsEnsureOnAbortByLanguageEnabled(), TEXT("Could not find function %p '%s' where '%s'."), UninstrumentedWork, *GetFunctionDescription(UninstrumentedWork), ANSI_TO_TCHAR("autortfm_close"));
 	    Context->AbortByLanguageAndThrow();
 	}
 
diff --git a/Engine/Source/Runtime/Core/Private/AutoRTFM/CompilerInterface.cpp b/Engine/Source/Runtime/Core/Private/AutoRTFM/CompilerInterface.cpp
--- a/Engine/Source/Runtime/Core/Private/AutoRTFM/CompilerInterface.cpp
+++ b/Engine/Source/Runtime/Core/Private/AutoRTFM/CompilerInterface.cpp
@@ -121,11 +121,11 @@ extern "C" UE_AUTORTFM_API void autortfm_llvm_fail(const char* Message)
 {
     if (Message)
     {
-		UE_LOG(LogAutoRTFM, Fatal, TEXT("Transaction failing because of language issue '%s'."), ANSI_TO_TCHAR(Message));
+		ensureMsgf(!ForTheRuntime::IsEnsureOnAbortByLanguageEnabled(), TEXT("Transaction failing because of language issue '%s'."), ANSI_TO_TCHAR(Message));
     }
     else
     {
-		UE_LOG(LogAutoRTFM, Fatal, TEXT("Transaction failing because of language issue."));
+		ensureMsgf(!ForTheRuntime::IsEnsureOnAbortByLanguageEnabled(), TEXT("Transaction failing because of language issue."));
 	}
 
 	FContext* Context = FContext::Get();
@@ -134,7 +134,7 @@ extern "C" UE_AUTORTFM_API void autortfm_llvm_fail(const char* Message)
 
 extern "C" UE_AUTORTFM_API void autortfm_llvm_missing_function()
 {
-	UE_LOG(LogAutoRTFM, Fatal, TEXT("Transaction failing because of missing function."));
+	ensureMsgf(!ForTheRuntime::IsEnsureOnAbortByLanguageEnabled(), TEXT("Transaction failing because of missing function."));
 
 	FContext* Context = FContext::Get();
     Context->AbortByLanguageAndThrow();
diff --git a/Engine/Source/Runtime/Core/Private/AutoRTFM/Context.cpp b/Engine/Source/Runtime/Core/Private/AutoRTFM/Context.cpp
--- a/Engine/Source/Runtime/Core/Private/AutoRTFM/Context.cpp
+++ b/Engine/Source/Runtime/Core/Private/AutoRTFM/Context.cpp
@@ -409,10 +409,14 @@ ETransactionResult FContext::Transact(void (*InstrumentedFunction)(void*), void*
 		ASSERT(CurrentNest != nullptr);
 		ASSERT(CurrentTransaction != nullptr);
 
-		// A cascading abort should cause all transactions to abort!
-		if (ETransactionResult::AbortedByCascade == Result)
+		// Cascading aborts should cause all transactions to abort!
+		switch (Result)
 		{
+		default:
+			break;
+		case ETransactionResult::AbortedByCascade:
 			CurrentTransaction->AbortAndThrow();
+			break;
 		}
 
 		ClearTransactionStatus();
@@ -439,7 +443,6 @@ void FContext::AbortByRequestWithoutThrowing()
 
 void FContext::AbortByLanguageAndThrow()
 {
-	UE_DEBUG_BREAK();
     ASSERT(Status == EContextStatus::OnTrack);
 	GAutoRTFMMetrics.NumTransactionsAbortedByLanguage++;
     Status = EContextStatus::AbortedByLanguage;
diff --git a/Engine/Source/Runtime/Core/Private/AutoRTFM/FunctionMapInlines.h b/Engine/Source/Runtime/Core/Private/AutoRTFM/FunctionMapInlines.h
--- a/Engine/Source/Runtime/Core/Private/AutoRTFM/FunctionMapInlines.h
+++ b/Engine/Source/Runtime/Core/Private/AutoRTFM/FunctionMapInlines.h
@@ -17,11 +17,11 @@ inline void* FunctionMapLookup(void* OldFunction, const char* Where)
     {
 		if (Where)
 		{
-			UE_LOG(LogAutoRTFM, Fatal, TEXT("Could not find function %p '%s' where '%s'."), OldFunction, *GetFunctionDescription(OldFunction), ANSI_TO_TCHAR(Where));
+			ensureMsgf(!ForTheRuntime::IsEnsureOnAbortByLanguageEnabled(), TEXT("Could not find function %p '%s' where '%s'."), OldFunction, *GetFunctionDescription(OldFunction), ANSI_TO_TCHAR(Where));
 		}
 		else
 		{
-			UE_LOG(LogAutoRTFM, Fatal, TEXT("Could not find function %p '%s'."), OldFunction, *GetFunctionDescription(OldFunction));
+			ensureMsgf(!ForTheRuntime::IsEnsureOnAbortByLanguageEnabled(), TEXT("Could not find function %p '%s'."), OldFunction, *GetFunctionDescription(OldFunction));
 		}
 		
 		FContext* Context = FContext::Get();
diff --git a/Engine/Source/Runtime/Core/Public/AutoRTFM/AutoRTFM.h b/Engine/Source/Runtime/Core/Public/AutoRTFM/AutoRTFM.h
--- a/Engine/Source/Runtime/Core/Public/AutoRTFM/AutoRTFM.h
+++ b/Engine/Source/Runtime/Core/Public/AutoRTFM/AutoRTFM.h
@@ -706,6 +706,12 @@ namespace ForTheRuntime
 	// Query whether the AutoRTFM runtime is enabled for *all* Verse code (not just failure contexts).
 	UE_AUTORTFM_API bool IsAutoRTFMRuntimeEnabledForAllVerse();
 
+	// Set whether we should trigger an ensure on an abort-by-language.
+	UE_AUTORTFM_API void SetEnsureOnAbortByLanguage(bool bEnabled);
+
+	// Returns whether the runtime will trigger an ensure on an abort-by-language, or not.
+	UE_AUTORTFM_API bool IsEnsureOnAbortByLanguageEnabled();
+
 	// Manually create a new transaction from open code and push it as a transaction nest.
 	// Can only be called within an already active parent transaction (EG. this cannot start
 	// a transaction nest itself).
diff --git a/Engine/Source/Runtime/CoreUObject/Private/UObject/ScriptCore.cpp b/Engine/Source/Runtime/CoreUObject/Private/UObject/ScriptCore.cpp
--- a/Engine/Source/Runtime/CoreUObject/Private/UObject/ScriptCore.cpp
+++ b/Engine/Source/Runtime/CoreUObject/Private/UObject/ScriptCore.cpp
@@ -4086,6 +4086,21 @@ DEFINE_FUNCTION(UObject::execAutoRtfmTransact)
 
 	P_NATIVE_BEGIN;
 
+	if (UNLIKELY(Result == AutoRTFM::ETransactionResult::AbortedByLanguage))
+	{
+		FBlueprintExceptionInfo AbortedByLanguage(
+			EBlueprintExceptionType::FatalError,
+			LOCTEXT("AbortedByLanguage", "AutoRTFM aborted because of unhandled constructs in the code (atomics, unhandled function calls, etc)")
+		);
+
+		FBlueprintCoreDelegates::ThrowScriptException(Context, Stack, AbortedByLanguage);
+
+		if (AutoRTFM::IsTransactional())
+		{
+			AutoRTFM::CascadingAbortTransaction();
+		}
+	}
+
 	if (Result != AutoRTFM::ETransactionResult::Committed)
 	{
 		// if this transaction didn't commit, move our code pointer to the target
