Tag Trace allocation with correct tag

diff --git a/Engine/Source/Runtime/Core/Private/ProfilingDebugging/TraceAuxiliary.cpp b/Engine/Source/Runtime/Core/Private/ProfilingDebugging/TraceAuxiliary.cpp
--- a/Engine/Source/Runtime/Core/Private/ProfilingDebugging/TraceAuxiliary.cpp
+++ b/Engine/Source/Runtime/Core/Private/ProfilingDebugging/TraceAuxiliary.cpp
@@ -1544,6 +1544,7 @@ static void TraceAuxiliaryAddPostForkCallback(const TCHAR* CommandLine)
 void FTraceAuxiliary::Initialize(const TCHAR* CommandLine)
 {
 	TRACE_CPUPROFILER_EVENT_SCOPE(FTraceAux_Init);
+	UE_MEMSCOPE(TRACE_TAG);
 
 #if UE_TRACE_SERVER_LAUNCH_ENABLED && UE_TRACE_SERVER_CONTROLS_ENABLED
 	// Auto launch Unreal Trace Server for certain configurations
