Allow use of TAG_TRACE when trace is disabled.

diff --git a/Engine/Source/Runtime/Core/Public/ProfilingDebugging/TagTrace.h b/Engine/Source/Runtime/Core/Public/ProfilingDebugging/TagTrace.h
--- a/Engine/Source/Runtime/Core/Public/ProfilingDebugging/TagTrace.h
+++ b/Engine/Source/Runtime/Core/Public/ProfilingDebugging/TagTrace.h
@@ -21,6 +21,8 @@ CORE_API int32	MemoryTrace_AnnounceCustomTag(int32 Tag, int32 ParentTag, const T
 CORE_API int32	MemoryTrace_AnnounceFNameTag(const class FName& TagName);
 CORE_API int32	MemoryTrace_GetActiveTag();
 
+inline constexpr int32 TRACE_TAG = 257;
+
 ////////////////////////////////////////////////////////////////////////////////
 #if !defined(UE_MEMORY_TAGS_TRACE_ENABLED)
 	#define UE_MEMORY_TAGS_TRACE_ENABLED 0
@@ -84,9 +86,6 @@ private:
 	UE::Trace::Private::FScopedLogScope Inner;
 };
 
-////////////////////////////////////////////////////////////////////////////////
-inline constexpr int32 TRACE_TAG = 257;
-
 ////////////////////////////////////////////////////////////////////////////////
 #define UE_MEMSCOPE(InTag)					FMemScope PREPROCESSOR_JOIN(MemScope,__LINE__)(InTag);
 #define UE_MEMSCOPE_PTR(InPtr)				FMemScopePtr PREPROCESSOR_JOIN(MemPtrScope,__LINE__)((uint64)InPtr);
