Include TagTrace directly to fix CIS error

diff --git a/Engine/Source/Runtime/Core/Private/ProfilingDebugging/TraceAuxiliary.cpp b/Engine/Source/Runtime/Core/Private/ProfilingDebugging/TraceAuxiliary.cpp
--- a/Engine/Source/Runtime/Core/Private/ProfilingDebugging/TraceAuxiliary.cpp
+++ b/Engine/Source/Runtime/Core/Private/ProfilingDebugging/TraceAuxiliary.cpp
@@ -1,6 +1,7 @@
 // Copyright Epic Games, Inc. All Rights Reserved.
 
 #include "ProfilingDebugging/TraceAuxiliary.h"
+#include "ProfilingDebugging/TagTrace.h"
 #include "ProfilingDebugging/StringsTrace.h"
 #include "Trace/Trace.h"
 #include "CoreGlobals.h"
