Added UE_ENABLE_LINEAR_VIRTUAL_ALLOCATOR define to be able to disable linear virtual allocator if needed

#rnx

diff --git a/Engine/Source/Runtime/Core/Private/Memory/LinearVirtualMemoryAllocator.cpp b/Engine/Source/Runtime/Core/Private/Memory/LinearVirtualMemoryAllocator.cpp
--- a/Engine/Source/Runtime/Core/Private/Memory/LinearVirtualMemoryAllocator.cpp
+++ b/Engine/Source/Runtime/Core/Private/Memory/LinearVirtualMemoryAllocator.cpp
@@ -15,7 +15,7 @@ static struct FInitializer
 } Initializer;
 
 
-#if PLATFORM_HAS_FPlatformVirtualMemoryBlock
+#if UE_ENABLE_LINEAR_VIRTUAL_ALLOCATOR
 
 #include "CoreGlobals.h"
 #include "Misc/ScopeLock.h"
@@ -138,4 +138,5 @@ FLinearVirtualMemoryAllocator& GetPersistentLinearAllocator()
 	static FLinearVirtualMemoryAllocator GPersistentLinearAllocator(UE_PERSISTENT_ALLOCATOR_RESERVE_SIZE);
 	return GPersistentLinearAllocator;
 }
-#endif //~PLATFORM_HAS_FPlatformVirtualMemoryBlock
\ No newline at end of file
+
+#endif //~UE_ENABLE_LINEAR_VIRTUAL_ALLOCATOR
\ No newline at end of file
diff --git a/Engine/Source/Runtime/Core/Public/Memory/LinearVirtualMemoryAllocator.h b/Engine/Source/Runtime/Core/Public/Memory/LinearVirtualMemoryAllocator.h
--- a/Engine/Source/Runtime/Core/Public/Memory/LinearVirtualMemoryAllocator.h
+++ b/Engine/Source/Runtime/Core/Public/Memory/LinearVirtualMemoryAllocator.h
@@ -5,7 +5,11 @@
 #include "HAL/Platform.h"
 #include "HAL/UnrealMemory.h"
 
-#if PLATFORM_HAS_FPlatformVirtualMemoryBlock
+#ifndef UE_ENABLE_LINEAR_VIRTUAL_ALLOCATOR
+	#define UE_ENABLE_LINEAR_VIRTUAL_ALLOCATOR PLATFORM_HAS_FPlatformVirtualMemoryBlock
+#endif
+
+#if UE_ENABLE_LINEAR_VIRTUAL_ALLOCATOR
 #include <atomic>
 #include "HAL/CriticalSection.h"
 
@@ -87,7 +91,7 @@ struct FLinearVirtualMemoryAllocator
 
 FORCEINLINE FLinearVirtualMemoryAllocator GetPersistentLinearAllocator() { return FLinearVirtualMemoryAllocator(0); }
 
-#endif	//~PLATFORM_HAS_FPlatformVirtualMemoryBlock
+#endif	//~UE_ENABLE_LINEAR_VIRTUAL_ALLOCATOR
 
 struct FLinearVirtualMemoryAllocatorExtends
 {
