Make MB3 default for Mac, move MB2/MB3 defines to common Apple Platform code

diff --git a/Engine/Source/Runtime/Core/Private/Apple/ApplePlatformMemory.cpp b/Engine/Source/Runtime/Core/Private/Apple/ApplePlatformMemory.cpp
--- a/Engine/Source/Runtime/Core/Private/Apple/ApplePlatformMemory.cpp
+++ b/Engine/Source/Runtime/Core/Private/Apple/ApplePlatformMemory.cpp
@@ -288,17 +288,6 @@ void FApplePlatformMemory::Init()
 	
 }
 
-// if USE_MALLOC_BINNED3 isn't configured enable it by default
-#ifndef USE_MALLOC_BINNED3
-	#define USE_MALLOC_BINNED3 1
-#endif // USE_MALLOC_BINNED3
-
-
-// if USE_MALLOC_BINNED2 isn't configured enable it if USE_MALLOC_BINNED3 isn't true
-#ifndef USE_MALLOC_BINNED2
-	#define USE_MALLOC_BINNED2 !USE_MALLOC_BINNED3
-#endif // USE_MALLOC_BINNED2
-
 void FApplePlatformMemory::SetAllocatorToUse()
 {
     // force Ansi allocator in particular cases
diff --git a/Engine/Source/Runtime/Core/Private/Mac/MacPlatformMemory.cpp b/Engine/Source/Runtime/Core/Private/Mac/MacPlatformMemory.cpp
--- a/Engine/Source/Runtime/Core/Private/Mac/MacPlatformMemory.cpp
+++ b/Engine/Source/Runtime/Core/Private/Mac/MacPlatformMemory.cpp
@@ -11,6 +11,7 @@
 #include "HAL/MallocMimalloc.h"
 #include "HAL/MallocBinned.h"
 #include "HAL/MallocBinned2.h"
+#include "HAL/MallocBinned3.h"
 #include "HAL/MallocStomp.h"
 #include "Misc/AssertionMacros.h"
 #include "Misc/CoreStats.h"
@@ -26,9 +27,6 @@ extern "C"
 	#include <crt_externs.h> // Needed for _NSGetArgc & _NSGetArgv
 }
 
-// Set rather to use BinnedMalloc2 for binned malloc, can be overridden below
-#define USE_MALLOC_BINNED2 (1)
-
 #if PLATFORM_MAC_X86
 void* CFNetwork_CFAllocatorOperatorNew_Replacement(unsigned long Size, CFAllocatorRef Alloc)
 {
@@ -97,6 +95,10 @@ FMalloc* FMacPlatformMemory::BaseAllocator()
 	{
 		AllocatorToUse = EMemoryAllocatorToUse::Binned2;
 	}
+	else if (USE_MALLOC_BINNED3)
+	{
+		AllocatorToUse = EMemoryAllocatorToUse::Binned3;
+	}
 	else
 	{
 		AllocatorToUse = EMemoryAllocatorToUse::Binned;
@@ -151,6 +153,10 @@ FMalloc* FMacPlatformMemory::BaseAllocator()
 		Instance = new FMallocBinned2();
 		break;
 
+	case EMemoryAllocatorToUse::Binned3:
+		Instance = new FMallocBinned3();
+		break;
+
 	default:	// intentional fall-through
 	case EMemoryAllocatorToUse::Binned:
 		// [RCL] 2017-03-06 FIXME: perhaps BinnedPageSize should be used here, but leaving this change to the Mac platform owner.
diff --git a/Engine/Source/Runtime/Core/Public/Apple/ApplePlatformMemory.h b/Engine/Source/Runtime/Core/Public/Apple/ApplePlatformMemory.h
--- a/Engine/Source/Runtime/Core/Public/Apple/ApplePlatformMemory.h
+++ b/Engine/Source/Runtime/Core/Public/Apple/ApplePlatformMemory.h
@@ -14,6 +14,20 @@
 
 NS_ASSUME_NONNULL_BEGIN
 
+// if USE_MALLOC_BINNED3 isn't configured enable it by default
+#ifndef USE_MALLOC_BINNED3
+	#define USE_MALLOC_BINNED3 1
+	#if PLATFORM_IOS
+		#define USE_512MB_MAX_MEMORY_PER_BLOCK_SIZE 1
+	#endif
+#endif // USE_MALLOC_BINNED3
+
+
+// if USE_MALLOC_BINNED2 isn't configured enable it if USE_MALLOC_BINNED3 isn't true
+#ifndef USE_MALLOC_BINNED2
+#define USE_MALLOC_BINNED2 !USE_MALLOC_BINNED3
+#endif // USE_MALLOC_BINNED2
+
 /**
  * NSObject subclass that can be used to override the allocation functions to go through UE4's memory allocator.
  * This ensures that memory allocated by custom Objective-C types can be tracked by UE4's tools and 
