[ThreadSingleton]
 - Use thread_local instead of slot allocation for improved performance and to reduce pressure on tls slots on some platforms

diff --git a/Engine/Source/Runtime/Core/Private/HAL/ThreadingBase.cpp b/Engine/Source/Runtime/Core/Private/HAL/ThreadingBase.cpp
--- a/Engine/Source/Runtime/Core/Private/HAL/ThreadingBase.cpp
+++ b/Engine/Source/Runtime/Core/Private/HAL/ThreadingBase.cpp
@@ -1426,6 +1426,7 @@ FQueuedThread::Run()
 /*-----------------------------------------------------------------------------
 	FThreadSingletonInitializer
 -----------------------------------------------------------------------------*/
+PRAGMA_DISABLE_DEPRECATION_WARNINGS
 
 FTlsAutoCleanup* FThreadSingletonInitializer::Get( TFunctionRef<FTlsAutoCleanup*()> CreateInstance, uint32& InOutTlsSlot )
 {
@@ -1494,6 +1495,8 @@ FTlsAutoCleanup* FThreadSingletonInitializer::Inject(FTlsAutoCleanup* Instance,
 	return ThreadSingleton;
 }
 
+PRAGMA_ENABLE_DEPRECATION_WARNINGS
+
 void FTlsAutoCleanup::Register()
 {
 	static thread_local TArray<TUniquePtr<FTlsAutoCleanup>> TlsInstances;
diff --git a/Engine/Source/Runtime/Core/Public/HAL/RunnableThread.h b/Engine/Source/Runtime/Core/Public/HAL/RunnableThread.h
--- a/Engine/Source/Runtime/Core/Public/HAL/RunnableThread.h
+++ b/Engine/Source/Runtime/Core/Public/HAL/RunnableThread.h
@@ -18,7 +18,9 @@ class FTlsAutoCleanup;
  */
 class FRunnableThread
 {
+	PRAGMA_DISABLE_DEPRECATION_WARNINGS
 	friend class FThreadSingletonInitializer;
+	PRAGMA_ENABLE_DEPRECATION_WARNINGS
 	friend class FTlsAutoCleanup;
 	friend class FThreadManager;
 	friend class FForkableThread;
diff --git a/Engine/Source/Runtime/Core/Public/HAL/ThreadSingleton.h b/Engine/Source/Runtime/Core/Public/HAL/ThreadSingleton.h
--- a/Engine/Source/Runtime/Core/Public/HAL/ThreadSingleton.h
+++ b/Engine/Source/Runtime/Core/Public/HAL/ThreadSingleton.h
@@ -12,7 +12,7 @@ template <typename FuncType> class TFunctionRef;
 /**
  * Thread singleton initializer.
  */
-class FThreadSingletonInitializer
+class UE_DEPRECATED(5.5, "This class will be removed.") FThreadSingletonInitializer
 {
 public:
 
@@ -44,9 +44,9 @@ class TThreadSingleton : public FTlsAutoCleanup
 	/**
 	 * @return TLS slot that holds a TThreadSingleton.
 	 */
-	CORE_API static uint32& GetTlsSlot()
+	CORE_API static T*& GetTlsSlot()
 	{
-		static uint32 TlsSlot = FPlatformTLS::InvalidTlsSlot;
+		static thread_local T* TlsSlot = nullptr;
 		return TlsSlot;
 	}
 #else
@@ -56,9 +56,9 @@ class TThreadSingleton : public FTlsAutoCleanup
 #if PLATFORM_CONSOLE_DYNAMIC_LINK
 	FORCENOINLINE
 #endif
-	static uint32& GetTlsSlot()
+	static T*& GetTlsSlot()
 	{
-		static uint32 TlsSlot = FPlatformTLS::InvalidTlsSlot;
+		static thread_local T* TlsSlot = nullptr;
 		return TlsSlot;
 	}
 #endif
@@ -73,10 +73,10 @@ protected:
 	virtual ~TThreadSingleton()
 	{
 		// Clean the dangling pointer from the TLS.
-		check(GetTlsSlot() != FPlatformTLS::InvalidTlsSlot);
-		if(((FTlsAutoCleanup*)FPlatformTLS::GetTlsValue(GetTlsSlot())) == static_cast<FTlsAutoCleanup*>(this))
+		check(GetTlsSlot() != nullptr);
+		if(((FTlsAutoCleanup*)GetTlsSlot()) == static_cast<FTlsAutoCleanup*>(this))
 		{
-			FPlatformTLS::SetTlsValue(GetTlsSlot(), nullptr);
+			GetTlsSlot() =  nullptr;
 		}
 	}
 
@@ -98,7 +98,13 @@ public:
 	 */
 	FORCEINLINE static T& Get()
 	{
-		return *(T*)FThreadSingletonInitializer::Get( [](){ return (FTlsAutoCleanup*)new T(); }, T::GetTlsSlot() ); //-V572
+		T*& TlsSlot = GetTlsSlot();
+		if (TlsSlot == nullptr)
+		{
+			TlsSlot = new T();
+			TlsSlot->Register();
+		}
+		return *TlsSlot;
 	}
 
 	/**
@@ -107,7 +113,13 @@ public:
 	 */
 	FORCEINLINE static T& Get(TFunctionRef<FTlsAutoCleanup*()> CreateInstance)
 	{
-		return *(T*)FThreadSingletonInitializer::Get(CreateInstance, T::GetTlsSlot()); //-V572
+		T*& TlsSlot = GetTlsSlot();
+		if (TlsSlot == nullptr)
+		{
+			TlsSlot = static_cast<T*>(CreateInstance());
+			TlsSlot->Register();
+		}
+		return *TlsSlot;
 	}
 
 	/**
@@ -115,7 +127,7 @@ public:
 	 */
 	FORCEINLINE static T* TryGet()
 	{
-		return (T*)FThreadSingletonInitializer::TryGet( T::GetTlsSlot() );
+		return GetTlsSlot();
 	}
 
 	/**
@@ -123,6 +135,9 @@ public:
 	*/
 	FORCEINLINE static T* Inject(T* Instance)
 	{
-		return (T*)FThreadSingletonInitializer::Inject(Instance, T::GetTlsSlot());
+		T*& TlsSlot = GetTlsSlot();
+		T* OldValue = TlsSlot;
+		TlsSlot = Instance;
+		return OldValue;
 	}
 };
