FResourceArrayInterface changes
- FResourceArrayInterface is explicitly an interface for TResourceArray but it has also been coopted for general buffer uploads which doesn't actually need IsStatic/GetAllowCPUAccess/SetAllowCPUAccess.
- This change seperates the TResourceArray Static/CPUAccess interface from the general uploading interface which now lives in FResourceArrayUploadInterface.
- Changing the ResourceArray member of RHIResourceCreateInfo to be FResourceArrayUploadInterface.
- Adding FResourceArrayUploadArrayView as a utility for simple buffer uploads and replaces a number of ad-hoc implementations.

diff --git a/Engine/Plugins/Runtime/ProceduralMeshComponent/Source/ProceduralMeshComponent/Private/ProceduralMeshComponent.cpp b/Engine/Plugins/Runtime/ProceduralMeshComponent/Source/ProceduralMeshComponent/Private/ProceduralMeshComponent.cpp
--- a/Engine/Plugins/Runtime/ProceduralMeshComponent/Source/ProceduralMeshComponent/Private/ProceduralMeshComponent.cpp
+++ b/Engine/Plugins/Runtime/ProceduralMeshComponent/Source/ProceduralMeshComponent/Private/ProceduralMeshComponent.cpp
@@ -35,28 +35,6 @@ static TAutoConsoleVariable<int32> CVarRayTracingProceduralMesh(
 	1,
 	TEXT("Include procedural meshes in ray tracing effects (default = 1 (procedural meshes enabled in ray tracing))"));
 
-/** Resource array to pass  */
-class FProcMeshVertexResourceArray : public FResourceArrayInterface
-{
-public:
-	FProcMeshVertexResourceArray(void* InData, uint32 InSize)
-		: Data(InData)
-		, Size(InSize)
-	{
-	}
-
-	virtual const void* GetResourceData() const override { return Data; }
-	virtual uint32 GetResourceDataSize() const override { return Size; }
-	virtual void Discard() override { }
-	virtual bool IsStatic() const override { return false; }
-	virtual bool GetAllowCPUAccess() const override { return false; }
-	virtual void SetAllowCPUAccess(bool bInNeedsCPUAccess) override { }
-
-private:
-	void* Data;
-	uint32 Size;
-};
-
 /** Class representing a single section of the proc mesh */
 class FProcMeshProxySection
 {
diff --git a/Engine/Source/Runtime/Core/Public/Containers/DynamicRHIResourceArray.h b/Engine/Source/Runtime/Core/Public/Containers/DynamicRHIResourceArray.h
--- a/Engine/Source/Runtime/Core/Public/Containers/DynamicRHIResourceArray.h
+++ b/Engine/Source/Runtime/Core/Public/Containers/DynamicRHIResourceArray.h
@@ -56,7 +56,7 @@ public:
 	/**
 	* @return A pointer to the resource data.
 	*/
-	virtual const void* GetResourceData() const 
+	virtual const void* GetResourceData() const override
 	{ 
 		return &(*this)[0]; 
 	}
@@ -64,7 +64,7 @@ public:
 	/**
 	* @return size of resource data allocation
 	*/
-	virtual uint32 GetResourceDataSize() const
+	virtual uint32 GetResourceDataSize() const override
 	{
 		if (this->Num() > UINT32_MAX / sizeof(ElementType))
 		{
@@ -79,7 +79,7 @@ public:
 	* Only discard the resource memory on clients, and if the CPU doesn't need access to it.
 	* Non-clients can't discard the data because they may need to serialize it.
 	*/
-	virtual void Discard()
+	virtual void Discard() override
 	{
 		if(!bNeedsCPUAccess && FPlatformProperties::RequiresCookedData() && !IsRunningCommandlet())
 		{
@@ -90,7 +90,7 @@ public:
 	/**
 	* @return true if the resource array is static and shouldn't be modified
 	*/
-	virtual bool IsStatic() const
+	virtual bool IsStatic() const override
 	{
 		return false;
 	}
@@ -98,7 +98,7 @@ public:
 	/**
 	* @return true if the resource keeps a copy of its resource data after the RHI resource has been created
 	*/
-	virtual bool GetAllowCPUAccess() const
+	virtual bool GetAllowCPUAccess() const override
 	{
 		return bNeedsCPUAccess;
 	}
@@ -106,7 +106,7 @@ public:
 	/** 
 	* Sets whether the resource array will be accessed by CPU. 
 	*/
-	virtual void SetAllowCPUAccess(bool bInNeedsCPUAccess)
+	virtual void SetAllowCPUAccess(bool bInNeedsCPUAccess) override
 	{
 		bNeedsCPUAccess = bInNeedsCPUAccess;
 	}
diff --git a/Engine/Source/Runtime/Core/Public/Containers/ResourceArray.h b/Engine/Source/Runtime/Core/Public/Containers/ResourceArray.h
--- a/Engine/Source/Runtime/Core/Public/Containers/ResourceArray.h
+++ b/Engine/Source/Runtime/Core/Public/Containers/ResourceArray.h
@@ -5,28 +5,72 @@
 #include "CoreTypes.h"
 #include "Serialization/MemoryLayout.h"
 
-/**
- * An element type independent interface to the resource array.
- */
-class FResourceArrayInterface
+/** An element type independent interface for uploading an array of resource data. */
+struct FResourceArrayUploadInterface
 {
-	DECLARE_EXPORTED_TYPE_LAYOUT(FResourceArrayInterface, CORE_API, Abstract);
-public:
+	virtual ~FResourceArrayUploadInterface() {}
 
-	virtual ~FResourceArrayInterface() {}
-
-	/**
-	 * @return A pointer to the resource data.
-	 */
+	/** Returns a pointer to the resource data. */
 	virtual const void* GetResourceData() const = 0;
 
-	/**
-	 * @return size of resource data allocation
-	 */
+	/** Returns size of resource data allocation */
 	virtual uint32 GetResourceDataSize() const = 0;
 
 	/** Called on non-UMA systems after the RHI has copied the resource data, and no longer needs the CPU's copy. */
 	virtual void Discard() = 0;
+};
+
+/** Utility to do a simple upload of data from an array managed by the caller. */
+struct FResourceArrayUploadArrayView : public FResourceArrayUploadInterface
+{
+	const void* const Data;
+	const uint32 SizeInBytes;
+
+	FResourceArrayUploadArrayView() = delete;
+
+	FResourceArrayUploadArrayView(const void* InData, uint32 InSizeInBytes)
+		: Data(InData)
+		, SizeInBytes(InSizeInBytes)
+	{
+	}
+
+	template<typename ElementType>
+	FResourceArrayUploadArrayView(TConstArrayView<ElementType> View)
+		: Data(View.GetData())
+		, SizeInBytes(View.Num() * View.GetTypeSize())
+	{
+	}
+
+	template<typename ElementType, typename AllocatorType>
+	FResourceArrayUploadArrayView(TArray<ElementType, AllocatorType> InArray)
+		: FResourceArrayUploadArrayView(MakeArrayView(InArray))
+	{
+	}
+
+	// FResourceArrayUploadInterface
+	virtual const void* GetResourceData() const final
+	{
+		return Data;
+	}
+
+	virtual uint32 GetResourceDataSize() const final
+	{
+		return SizeInBytes;
+	}
+
+	virtual void Discard() final
+	{
+		// do nothing
+	}
+};
+
+/**
+ * An element type independent interface to the resource array.
+ */
+class FResourceArrayInterface : public FResourceArrayUploadInterface
+{
+	DECLARE_EXPORTED_TYPE_LAYOUT(FResourceArrayInterface, CORE_API, Abstract);
+public:
 
 	/**
 	 * @return true if the resource array is static and shouldn't be modified
diff --git a/Engine/Source/Runtime/D3D12RHI/Private/D3D12Buffer.cpp b/Engine/Source/Runtime/D3D12RHI/Private/D3D12Buffer.cpp
--- a/Engine/Source/Runtime/D3D12RHI/Private/D3D12Buffer.cpp
+++ b/Engine/Source/Runtime/D3D12RHI/Private/D3D12Buffer.cpp
@@ -104,7 +104,7 @@ void FD3D12Buffer::UploadResourceData(FRHICommandListBase& InRHICmdList, FRHIGPU
 		});
 }
 
-void FD3D12Buffer::UploadResourceData(FRHICommandListBase& RHICmdList, FResourceArrayInterface* InResourceArray, D3D12_RESOURCE_STATES InDestinationState, const TCHAR* AssetName, const FName& ClassName, const FName& PackageName)
+void FD3D12Buffer::UploadResourceData(FRHICommandListBase& RHICmdList, FResourceArrayUploadInterface* InResourceArray, D3D12_RESOURCE_STATES InDestinationState, const TCHAR* AssetName, const FName& ClassName, const FName& PackageName)
 {
 	UE_TRACE_METADATA_SCOPE_ASSET_FNAME(AssetName, ClassName, PackageName);
 
@@ -130,7 +130,7 @@ void FD3D12Buffer::UploadResourceData(FRHICommandListBase& RHICmdList, FResource
 }
 
 
-FD3D12SyncPointRef FD3D12Buffer::UploadResourceDataViaCopyQueue(FResourceArrayInterface* InResourceArray)
+FD3D12SyncPointRef FD3D12Buffer::UploadResourceDataViaCopyQueue(FResourceArrayUploadInterface* InResourceArray)
 {
 	// assume not dynamic and not on async thread (probably fine but untested)
 	check(IsInRHIThread() || IsInRenderingThread());
diff --git a/Engine/Source/Runtime/D3D12RHI/Private/D3D12Resources.h b/Engine/Source/Runtime/D3D12RHI/Private/D3D12Resources.h
--- a/Engine/Source/Runtime/D3D12RHI/Private/D3D12Resources.h
+++ b/Engine/Source/Runtime/D3D12RHI/Private/D3D12Resources.h
@@ -942,9 +942,9 @@ public:
 
 	static void UploadResourceData(FD3D12CommandContext& CommandContext, D3D12_RESOURCE_STATES InDestinationState, FD3D12ResourceLocation& DestinationResourceLocation, const FD3D12ResourceLocation& SourceResourceLocation, uint32 Size);
 	void UploadResourceData(FRHICommandListBase& InRHICmdList, FRHIGPUMask GPUMask, D3D12_RESOURCE_STATES InDestinationState, const void* SourceData, int32 SourceDataSize);
-	void UploadResourceData(FRHICommandListBase& InRHICmdList, FResourceArrayInterface* InResourceArray, D3D12_RESOURCE_STATES InDestinationState, const TCHAR* AssetName = nullptr, const FName& ClassName = NAME_None, const FName& PackageName = NAME_None);
+	void UploadResourceData(FRHICommandListBase& InRHICmdList, FResourceArrayUploadInterface* InResourceArray, D3D12_RESOURCE_STATES InDestinationState, const TCHAR* AssetName = nullptr, const FName& ClassName = NAME_None, const FName& PackageName = NAME_None);
 
-	FD3D12SyncPointRef UploadResourceDataViaCopyQueue(FResourceArrayInterface* InResourceArray);
+	FD3D12SyncPointRef UploadResourceDataViaCopyQueue(FResourceArrayUploadInterface* InResourceArray);
 
 	// FRHIResource overrides
 #if RHI_ENABLE_RESOURCE_INFO
diff --git a/Engine/Source/Runtime/Engine/Classes/Particles/SubUVAnimation.h b/Engine/Source/Runtime/Engine/Classes/Particles/SubUVAnimation.h
--- a/Engine/Source/Runtime/Engine/Classes/Particles/SubUVAnimation.h
+++ b/Engine/Source/Runtime/Engine/Classes/Particles/SubUVAnimation.h
@@ -78,27 +78,7 @@ public:
 	virtual void ReleaseRHI() override;
 };
 
-/** Resource array to pass  */
-class FSubUVVertexResourceArray : public FResourceArrayInterface
-{
-public:
-	FSubUVVertexResourceArray(void* InData, uint32 InSize)
-		: Data(InData)
-		, Size(InSize)
-	{
-	}
-
-	virtual const void* GetResourceData() const override { return Data; }
-	virtual uint32 GetResourceDataSize() const override { return Size; }
-	virtual void Discard() override { }
-	virtual bool IsStatic() const override { return false; }
-	virtual bool GetAllowCPUAccess() const override { return false; }
-	virtual void SetAllowCPUAccess(bool bInNeedsCPUAccess) override { }
-
-private:
-	void* Data;
-	uint32 Size;
-};
+using FSubUVVertexResourceArray UE_DEPRECATED(5.5, "FResourceArrayUploadArrayView should be used directly.") = FResourceArrayUploadArrayView;
 
 /**
  * SubUV animation asset, which caches bounding geometry for regions in the SubUVTexture with non-zero opacity.
diff --git a/Engine/Source/Runtime/Engine/Private/Particles/SubUVAnimation.cpp b/Engine/Source/Runtime/Engine/Private/Particles/SubUVAnimation.cpp
--- a/Engine/Source/Runtime/Engine/Private/Particles/SubUVAnimation.cpp
+++ b/Engine/Source/Runtime/Engine/Private/Particles/SubUVAnimation.cpp
@@ -63,7 +63,7 @@ void FSubUVBoundingGeometryBuffer::InitRHI(FRHICommandListBase& RHICmdList)
 
 	if (SizeInBytes > 0)
 	{
-		FSubUVVertexResourceArray ResourceArray(Vertices->GetData(), SizeInBytes);
+		FResourceArrayUploadArrayView ResourceArray(Vertices->GetData(), SizeInBytes);
 		FRHIResourceCreateInfo CreateInfo(TEXT("FSubUVBoundingGeometryBuffer"), &ResourceArray);
 		VertexBufferRHI = RHICmdList.CreateVertexBuffer(SizeInBytes, BUF_ShaderResource | BUF_Static, CreateInfo);
 		ShaderResourceView = RHICmdList.CreateShaderResourceView(VertexBufferRHI, sizeof(FVector2f), PF_G32R32F);
diff --git a/Engine/Source/Runtime/MRMesh/Private/MRMeshComponent.cpp b/Engine/Source/Runtime/MRMesh/Private/MRMeshComponent.cpp
--- a/Engine/Source/Runtime/MRMesh/Private/MRMeshComponent.cpp
+++ b/Engine/Source/Runtime/MRMesh/Private/MRMeshComponent.cpp
@@ -50,27 +50,6 @@ static TAutoConsoleVariable<int32> CVarPauseMRMeshBrickCulling(
 	TEXT("MR Mesh brick culling debug state: 0=off, 1=on, 2=paused"));
 #endif //DEBUG_BRICK_CULLING
 
-class FMRMeshVertexResourceArray : public FResourceArrayInterface
-{
-public:
-	FMRMeshVertexResourceArray(const void* InData, uint32 InSize)
-		: Data(InData)
-		, Size(InSize)
-	{
-	}
-
-	virtual const void* GetResourceData() const override { return Data; }
-	virtual uint32 GetResourceDataSize() const override { return Size; }
-	virtual void Discard() override { }
-	virtual bool IsStatic() const override { return false; }
-	virtual bool GetAllowCPUAccess() const override { return false; }
-	virtual void SetAllowCPUAccess(bool bInNeedsCPUAccess) override { }
-
-private:
-	const void* Data;
-	uint32 Size;
-};
-
 /** Support for non-interleaved data streams. */
 template<typename DataType>
 class FMRMeshVertexBuffer : public FVertexBuffer
@@ -83,7 +62,7 @@ public:
 
 		const uint32 SizeInBytes = PerVertexData.Num() * sizeof(DataType);
 
-		FMRMeshVertexResourceArray ResourceArray(PerVertexData.GetData(), SizeInBytes);
+		FResourceArrayUploadArrayView ResourceArray(PerVertexData.GetData(), SizeInBytes);
 		FRHIResourceCreateInfo CreateInfo(TEXT("FMRMeshVertexBuffer"), &ResourceArray);
 		VertexBufferRHI = RHICmdList.CreateVertexBuffer(SizeInBytes, BUF_Static | BUF_ShaderResource, CreateInfo);
 	}
diff --git a/Engine/Source/Runtime/RHI/Public/RHIResources.h b/Engine/Source/Runtime/RHI/Public/RHIResources.h
--- a/Engine/Source/Runtime/RHI/Public/RHIResources.h
+++ b/Engine/Source/Runtime/RHI/Public/RHIResources.h
@@ -31,6 +31,9 @@ class FRHIComputeCommandList;
 class FRHICommandListImmediate;
 class FRHITextureReference;
 class FRHIShaderBindingLayout;
+class FResourceBulkDataInterface;
+class FResourceArrayInterface;
+struct FResourceArrayUploadInterface;
 
 struct FClearValueBinding;
 struct FRHIResourceInfo;
@@ -348,9 +351,6 @@ struct FClearValueBinding
 	static RHI_API const FClearValueBinding DefaultNormal8Bit;
 };
 
-class FResourceBulkDataInterface;
-class FResourceArrayInterface;
-
 struct FRHIResourceCreateInfo
 {
 	FRHIResourceCreateInfo(const TCHAR* InDebugName)
@@ -373,7 +373,7 @@ struct FRHIResourceCreateInfo
 	}
 
 	// for CreateBuffer calls
-	FRHIResourceCreateInfo(const TCHAR* InDebugName, FResourceArrayInterface* InResourceArray)
+	FRHIResourceCreateInfo(const TCHAR* InDebugName, FResourceArrayUploadInterface* InResourceArray)
 		: FRHIResourceCreateInfo(InDebugName)
 	{
 		ResourceArray = InResourceArray;
@@ -397,7 +397,7 @@ struct FRHIResourceCreateInfo
 	FResourceBulkDataInterface* BulkData;
 
 	// for CreateBuffer calls
-	FResourceArrayInterface* ResourceArray;
+	FResourceArrayUploadInterface* ResourceArray;
 
 	// for binding clear colors to render targets.
 	FClearValueBinding ClearValueBinding;
@@ -3152,7 +3152,7 @@ public:
 	TArray<FRayTracingGeometrySegment> Segments;
 
 	// Offline built geometry data. If null, the geometry will be built by the RHI at runtime.
-	FResourceArrayInterface* OfflineData = nullptr;
+	FResourceArrayUploadInterface* OfflineData = nullptr;
 
 	// Pointer to an existing ray tracing geometry which the new geometry is built from.
 	FRHIRayTracingGeometry* SourceGeometry = nullptr;
diff --git a/Engine/Source/Runtime/RHI/Public/RHIUtilities.h b/Engine/Source/Runtime/RHI/Public/RHIUtilities.h
--- a/Engine/Source/Runtime/RHI/Public/RHIUtilities.h
+++ b/Engine/Source/Runtime/RHI/Public/RHIUtilities.h
@@ -218,7 +218,7 @@ struct FRWBuffer
 	}
 
 	// @param AdditionalUsage passed down to RHICreateVertexBuffer(), get combined with "BUF_UnorderedAccess | BUF_ShaderResource" e.g. BUF_Static
-	void Initialize(FRHICommandListBase& RHICmdList, const TCHAR* InDebugName, uint32 BytesPerElement, uint32 NumElements, EPixelFormat Format, ERHIAccess InResourceState, EBufferUsageFlags AdditionalUsage = BUF_None, FResourceArrayInterface *InResourceArray = nullptr)
+	void Initialize(FRHICommandListBase& RHICmdList, const TCHAR* InDebugName, uint32 BytesPerElement, uint32 NumElements, EPixelFormat Format, ERHIAccess InResourceState, EBufferUsageFlags AdditionalUsage = BUF_None, FResourceArrayUploadInterface* InResourceArray = nullptr)
 	{
 		// Provide a debug name if using Fast VRAM so the allocators diagnostics will work
 		ensure(!(EnumHasAnyFlags(AdditionalUsage, BUF_FastVRAM) && !InDebugName));
@@ -232,7 +232,7 @@ struct FRWBuffer
 		SRV = RHICmdList.CreateShaderResourceView(Buffer, BytesPerElement, UE_PIXELFORMAT_TO_UINT8(Format));
 	}
 
-	void Initialize(FRHICommandListBase& RHICmdList, const TCHAR* InDebugName, uint32 BytesPerElement, uint32 NumElements, EPixelFormat Format, EBufferUsageFlags AdditionalUsage = BUF_None, FResourceArrayInterface* InResourceArray = nullptr)
+	void Initialize(FRHICommandListBase& RHICmdList, const TCHAR* InDebugName, uint32 BytesPerElement, uint32 NumElements, EPixelFormat Format, EBufferUsageFlags AdditionalUsage = BUF_None, FResourceArrayUploadInterface* InResourceArray = nullptr)
 	{
 		Initialize(RHICmdList, InDebugName, BytesPerElement, NumElements, Format, ERHIAccess::UAVCompute, AdditionalUsage, InResourceArray);
 	}
@@ -296,7 +296,7 @@ struct FReadBuffer
 
 	FReadBuffer(): NumBytes(0) {}
 
-	void Initialize(FRHICommandListBase& RHICmdList, const TCHAR* InDebugName, uint32 BytesPerElement, uint32 NumElements, EPixelFormat Format, EBufferUsageFlags AdditionalUsage = BUF_None, FResourceArrayInterface* InResourceArray = nullptr)
+	void Initialize(FRHICommandListBase& RHICmdList, const TCHAR* InDebugName, uint32 BytesPerElement, uint32 NumElements, EPixelFormat Format, EBufferUsageFlags AdditionalUsage = BUF_None, FResourceArrayUploadInterface* InResourceArray = nullptr)
 	{
 		NumBytes = BytesPerElement * NumElements;
 		FRHIResourceCreateInfo CreateInfo(InDebugName);
diff --git a/Engine/Source/Runtime/RHICore/Internal/RHICoreResourceCollection.h b/Engine/Source/Runtime/RHICore/Internal/RHICoreResourceCollection.h
--- a/Engine/Source/Runtime/RHICore/Internal/RHICoreResourceCollection.h
+++ b/Engine/Source/Runtime/RHICore/Internal/RHICoreResourceCollection.h
@@ -68,7 +68,7 @@ namespace UE::RHICore
 		return Result;
 	}
 
-	struct FResourceCollectionUpload : public FResourceArrayInterface
+	struct FResourceCollectionUpload : public FResourceArrayUploadInterface
 	{
 		FResourceCollectionUpload(TConstArrayView<FRHIResourceCollectionMember> InMembers)
 			: Memory(CreateResourceCollectionArray(InMembers))
@@ -86,17 +86,6 @@ namespace UE::RHICore
 		{
 			Memory.Reset();
 		}
-		virtual bool IsStatic() const final
-		{
-			return true;
-		}
-		virtual bool GetAllowCPUAccess() const final
-		{
-			return false;
-		}
-		virtual void SetAllowCPUAccess(bool bInNeedsCPUAccess) final
-		{
-		}
 
 		TArray<uint32> Memory;
 	};
diff --git a/Engine/Source/Runtime/RenderCore/Private/UnifiedBuffer.cpp b/Engine/Source/Runtime/RenderCore/Private/UnifiedBuffer.cpp
--- a/Engine/Source/Runtime/RenderCore/Private/UnifiedBuffer.cpp
+++ b/Engine/Source/Runtime/RenderCore/Private/UnifiedBuffer.cpp
@@ -1687,28 +1687,6 @@ void FScatterUploadBuffer::InitPreSized(uint32 NumElements, uint32 InNumBytesPer
 	NumScatters = NumElements;
 }
 
-// Helper type used to initialize the buffer data on creation
-struct FScatterUploadBufferResourceArray : public FResourceArrayInterface
-{
-	const void* const DataPtr;
-	const int32 DataSize;
-
-	FScatterUploadBufferResourceArray(void* InDataPtr, int32 InDataSize)
-		: DataPtr(InDataPtr)
-		, DataSize(InDataSize)
-	{
-	}
-
-	const void* GetResourceData() const override { return DataPtr; }
-	uint32 GetResourceDataSize() const override { return DataSize; }
-
-	// Not necessary for our purposes
-	void Discard() override { }
-	bool IsStatic() const override { return false; }
-	bool GetAllowCPUAccess() const override { return true; }
-	void SetAllowCPUAccess(bool bInNeedsCPUAccess) override { }
-};
-
 template<typename ResourceType>
 void FScatterUploadBuffer::ResourceUploadTo(FRHICommandList& RHICmdList, const ResourceType& DstBuffer, bool bFlush)
 {
@@ -1726,13 +1704,13 @@ void FScatterUploadBuffer::ResourceUploadTo(FRHICommandList& RHICmdList, const R
 		const EBufferUsageFlags Usage = bFloat4Buffer ? BUF_None : BUF_ByteAddressBuffer;
 
 		{
-			FScatterUploadBufferResourceArray ScatterResourceArray(ScatterData, ScatterDataSize);
+			FResourceArrayUploadArrayView ScatterResourceArray(ScatterData, ScatterDataSize);
 			FRHIResourceCreateInfo CreateInfo(TEXT("ScatterResourceArray"), &ScatterResourceArray);
 			ScatterBuffer.Buffer = RHICmdList.CreateStructuredBuffer(sizeof(uint32), ScatterDataSize, BUF_ShaderResource | BUF_Volatile | Usage, CreateInfo);
 			ScatterBuffer.SRV = RHICmdList.CreateShaderResourceView(ScatterBuffer.Buffer);
 		}
 		{
-			FScatterUploadBufferResourceArray UploadResourceArray(UploadData, UploadDataSize);
+			FResourceArrayUploadArrayView UploadResourceArray(UploadData, UploadDataSize);
 			FRHIResourceCreateInfo CreateInfo(TEXT("ScatterUploadBuffer"), &UploadResourceArray);
 			UploadBuffer.Buffer = RHICmdList.CreateStructuredBuffer(TypeSize, UploadDataSize, BUF_ShaderResource | BUF_Volatile | Usage, CreateInfo);
 			UploadBuffer.SRV = RHICmdList.CreateShaderResourceView(UploadBuffer.Buffer);
diff --git a/Engine/Source/Runtime/Renderer/Private/InstanceCulling/InstanceCullingOcclusionQuery.cpp b/Engine/Source/Runtime/Renderer/Private/InstanceCulling/InstanceCullingOcclusionQuery.cpp
--- a/Engine/Source/Runtime/Renderer/Private/InstanceCulling/InstanceCullingOcclusionQuery.cpp
+++ b/Engine/Source/Runtime/Renderer/Private/InstanceCulling/InstanceCullingOcclusionQuery.cpp
@@ -33,31 +33,10 @@ struct FInstanceCullingOcclusionQueryDeferredContext;
 namespace
 {
 
-// YURIY_TODO: Put the helpers somewhere in the common RHI code
-struct FResourceArrayView : public FResourceArrayInterface
-{
-	const void* Data = nullptr;
-	uint32 SizeInBytes = 0;
-
-	template <typename T>
-	FResourceArrayView(TArrayView<T> View)
-		: Data(View.GetData())
-		, SizeInBytes(View.Num() * View.GetTypeSize())
-	{}
-
-	// FResourceArrayInterface
-	virtual const void* GetResourceData() const override final { return Data; }
-	virtual uint32 GetResourceDataSize() const override final { return SizeInBytes; }
-	virtual void Discard() override final {};
-	virtual bool IsStatic() const override final { return true; }
-	virtual bool GetAllowCPUAccess() const override final { return false; };
-	virtual void SetAllowCPUAccess(bool /*bInNeedsCPUAccess*/) override final {};
-};
-
 template <typename T>
 static FBufferRHIRef CreateBufferWithData(FRHICommandListBase& RHICmdList, EBufferUsageFlags UsageFlags, ERHIAccess ResourceState, const TCHAR* Name, TConstArrayView<T> Data)
 {
-	FResourceArrayView DataView(Data);
+	FResourceArrayUploadArrayView DataView(Data);
 	FRHIResourceCreateInfo CreateInfo(Name);
 	CreateInfo.ResourceArray = &DataView;
 	return RHICmdList.CreateBuffer(DataView.SizeInBytes, UsageFlags, Data.GetTypeSize(), ResourceState, CreateInfo);
