Adding missing reference on a TArray argument.

diff --git a/Engine/Source/Runtime/Core/Public/Containers/ResourceArray.h b/Engine/Source/Runtime/Core/Public/Containers/ResourceArray.h
--- a/Engine/Source/Runtime/Core/Public/Containers/ResourceArray.h
+++ b/Engine/Source/Runtime/Core/Public/Containers/ResourceArray.h
@@ -42,7 +42,7 @@ struct FResourceArrayUploadArrayView : public FResourceArrayUploadInterface
 	}
 
 	template<typename ElementType, typename AllocatorType>
-	FResourceArrayUploadArrayView(TArray<ElementType, AllocatorType> InArray)
+	FResourceArrayUploadArrayView(const TArray<ElementType, AllocatorType>& InArray)
 		: FResourceArrayUploadArrayView(MakeArrayView(InArray))
 	{
 	}
