Add a sanitizer-like mode to AutoRTFM to let you retry transactions that would have succeeded, letting us test that all on-abort handlers are written correctly.

I've made use of this in AutoRTFMTests as a starting place.

diff --git a/Engine/Source/Programs/AutoRTFMTests/Private/API.cpp b/Engine/Source/Programs/AutoRTFMTests/Private/API.cpp
--- a/Engine/Source/Programs/AutoRTFMTests/Private/API.cpp
+++ b/Engine/Source/Programs/AutoRTFMTests/Private/API.cpp
@@ -264,14 +264,30 @@ TEST_CASE("API.autortfm_on_commit")
 
 TEST_CASE("API.autortfm_on_abort")
 {
+	// Too hard to get this test working when retrying nested transactions so bail!
+	if (AutoRTFM::ForTheRuntime::ShouldRetryNestedTransactionsToo())
+	{
+		return;
+	}
+
     bool OuterTransaction = false;
     bool InnerTransaction = false;
     bool InnerTransactionWithAbort = false;
     bool InnerOpenNest = false;
-    AutoRTFM::ETransactionResult NestResult;
+    AutoRTFM::ETransactionResult NestResult = AutoRTFM::ETransactionResult::Committed;
 
     REQUIRE(AutoRTFM::ETransactionResult::Committed == AutoRTFM::Transact([&]
     {
+		// If we are retrying transactions, need to reset the test state.
+		AutoRTFM::OnAbort([&]
+		{
+			OuterTransaction = false;
+			InnerTransaction = false;
+			InnerTransactionWithAbort = false;
+			InnerOpenNest = false;
+			NestResult = AutoRTFM::ETransactionResult::Committed;
+		});
+
         autortfm_on_abort([](void* const Arg)
         {
             *static_cast<bool* const>(Arg) = true;
@@ -343,6 +359,12 @@ TEST_CASE("API.autortfm_did_allocate")
 
     AutoRTFM::Commit([&]
     {
+		// If we are retrying transactions, need to reset the test state.
+		AutoRTFM::OnAbort([&]
+		{
+			NextBump = 0;
+		});
+
         for (unsigned I = 0; I < Size; I++)
         {
             unsigned* Data;
@@ -713,14 +735,30 @@ TEST_CASE("API.OnCommit")
 
 TEST_CASE("API.OnAbort")
 {
+	// Too hard to get this test working when retrying nested transactions so bail!
+	if (AutoRTFM::ForTheRuntime::ShouldRetryNestedTransactionsToo())
+	{
+		return;
+	}
+
     bool OuterTransaction = false;
     bool InnerTransaction = false;
     bool InnerTransactionWithAbort = false;
     bool InnerOpenNest = false;
-    AutoRTFM::ETransactionResult NestResult;
+    AutoRTFM::ETransactionResult NestResult = AutoRTFM::ETransactionResult::Committed;
 
     REQUIRE(AutoRTFM::ETransactionResult::Committed == AutoRTFM::Transact([&]
     {
+		// If we are retrying transactions, need to reset the test state.
+		AutoRTFM::OnAbort([&]
+		{
+			OuterTransaction = false;
+			InnerTransaction = false;
+			InnerTransactionWithAbort = false;
+			InnerOpenNest = false;
+			NestResult = AutoRTFM::ETransactionResult::Committed;
+		});
+
         AutoRTFM::OnAbort([&]
         {
             OuterTransaction = true;
@@ -792,6 +830,12 @@ TEST_CASE("API.DidAllocate")
 
     AutoRTFM::Commit([&]
     {
+		// If we are retrying transactions, need to reset the test state.
+		AutoRTFM::OnAbort([&]
+		{
+			NextBump = 0;
+		});
+
         for (unsigned I = 0; I < Size; I++)
         {
             unsigned* Data;
diff --git a/Engine/Source/Programs/AutoRTFMTests/Private/Abort.cpp b/Engine/Source/Programs/AutoRTFMTests/Private/Abort.cpp
--- a/Engine/Source/Programs/AutoRTFMTests/Private/Abort.cpp
+++ b/Engine/Source/Programs/AutoRTFMTests/Private/Abort.cpp
@@ -55,6 +55,12 @@ TEST_CASE("Abort.NestedAbortOrder")
 
 	AutoRTFM::Commit([&]
 	{
+		// If we are retrying transactions, need to reset the test state.
+		AutoRTFM::OnAbort([&]
+			{
+				Orderer = 0;
+			});
+
 		InnerResult = AutoRTFM::Transact([&]
 			{
 				AutoRTFM::OnAbort([&]
@@ -368,15 +374,46 @@ TEST_CASE("Abort.PushOnAbortHandler_Order")
 
 			const AutoRTFM::ETransactionResult Result = AutoRTFM::Transact([&]
 				{
-					AutoRTFM::OnAbort([&Value] { REQUIRE(42 == Value); Value += 1; });
+					AutoRTFM::OnAbort([&Value]
+						{
+							REQUIRE(42 == Value);
+							Value += 1;
+						});
 
 					// Make a child transaction.
 					AutoRTFM::Commit([&]
 						{
-							AutoRTFM::PushOnAbortHandler(UIntToPointer(747), [&Value]() { REQUIRE(40 == Value); Value += 2; });
+							AutoRTFM::PushOnAbortHandler(UIntToPointer(747), [&Value]()
+								{
+									// If we are retrying nested transactions too, we can't check that
+									// the value was something specific before hand!
+									if (!AutoRTFM::ForTheRuntime::ShouldRetryNestedTransactionsToo())
+									{
+										REQUIRE(40 == Value);
+										Value += 2;
+									}
+									else
+									{
+										Value += 1;
+									}
+								});
 						});
 
-					AutoRTFM::OnAbort([&Value] { REQUIRE(37 == Value); Value += 3; });
+					AutoRTFM::OnAbort([&Value]
+						{
+							// If we are retrying nested transactions too, we've ran the on-abort in the
+							// child transaction once, so our value will be larger.
+							if (!AutoRTFM::ForTheRuntime::ShouldRetryNestedTransactionsToo())
+							{
+								REQUIRE(37 == Value);
+							}
+							else
+							{
+								REQUIRE(38 == Value);
+							}
+
+							Value += 3;
+						});
 
 					Value = 99;
 
@@ -398,7 +435,11 @@ TEST_CASE("Abort.PushOnAbortHandler_Order")
 					// Make a child transaction.
 					AutoRTFM::Commit([&]
 						{
-							AutoRTFM::PushOnAbortHandler(UIntToPointer(747), [&Value]() { REQUIRE(false); });
+							AutoRTFM::PushOnAbortHandler(UIntToPointer(747), [&Value]()
+								{
+									// Only if we are retrying on 
+									REQUIRE(AutoRTFM::ForTheRuntime::ShouldRetryNestedTransactionsToo());
+								});
 						});
 
 					AutoRTFM::OnAbort([&Value] { REQUIRE(37 == Value); Value += 3; });
@@ -439,6 +480,8 @@ TEST_CASE("Abort.PushOnAbortHandler_Order")
 					AutoRTFM::OnCommit([&Value] { REQUIRE(38 == Value); Value += 3; });
 
 					Value = 37;
+
+					AutoRTFM::OnAbort([&Value] { Value = 99; });
 				});
 
 			REQUIRE(AutoRTFM::ETransactionResult::AbortedByRequest == Result);
@@ -453,6 +496,14 @@ TEST_CASE("Abort.OnAbortTiming")
 	int Memory = 666;
 	AutoRTFM::Commit([&]
 	{
+		// If we are retrying transactions, need to reset the test state.
+		AutoRTFM::OnAbort([&]
+		{
+			REQUIRE(bOnAbortRan);
+			REQUIRE(Memory == 666);
+			bOnAbortRan = false;
+		});
+
 		REQUIRE(bOnAbortRan == false);
 		REQUIRE(Memory == 666);
 
@@ -469,9 +520,6 @@ TEST_CASE("Abort.OnAbortTiming")
 
 			AutoRTFM::AbortTransaction();
 		});
-
-		REQUIRE(Memory == 666);
-		REQUIRE(bOnAbortRan == true);
 	});
 	REQUIRE(Memory == 666);
 	REQUIRE(bOnAbortRan == true);
diff --git a/Engine/Source/Programs/AutoRTFMTests/Private/AutoRTFMTests.cpp b/Engine/Source/Programs/AutoRTFMTests/Private/AutoRTFMTests.cpp
--- a/Engine/Source/Programs/AutoRTFMTests/Private/AutoRTFMTests.cpp
+++ b/Engine/Source/Programs/AutoRTFMTests/Private/AutoRTFMTests.cpp
@@ -5,33 +5,58 @@
 
 IMPLEMENT_APPLICATION(AutoRTFMTests, "AutoRTFMTests");
 
+#define CATCH_AMALGAMATED_CUSTOM_MAIN
 #include "catch_amalgamated.cpp"
 
-class SetupListener final : public Catch::EventListenerBase
-{
-public:
-	using Catch::EventListenerBase::EventListenerBase;
+int main(int ArgC, const char* ArgV[]) {
+	Catch::Session Session;
 
-	void testRunStarting(const Catch::TestRunInfo&) override
-	{
-		GEngineLoop.PreInit(0, nullptr);
-		FModuleManager::Get().StartProcessingNewlyLoadedObjects();
+	bool NoRetry = false;
+	bool RetryNestedToo = false;
+
+	Session.cli(Session.cli()
+		| Catch::Clara::Opt(NoRetry)["--no-retry"]
+		| Catch::Clara::Opt(RetryNestedToo)["--retry-nested-too"]);
 
-		// Enable all Verse code to run under AutoRTFM (shouldn't affect our tests here, but better safe than sorry).
-		AutoRTFM::ForTheRuntime::SetAutoRTFMRuntime(AutoRTFM::ForTheRuntime::EAutoRTFMEnabledState::AutoRTFM_EnabledForAllVerse);
+	{
+		const int Result = Session.applyCommandLine(ArgC, ArgV);
 
-		// We don't want to trigger ensure's on abort because we are going to test that.
-		AutoRTFM::ForTheRuntime::SetEnsureOnAbortByLanguage(false);
+		if (0 != Result)
+		{
+			return Result;
+		}
 	}
 
-	void testRunEnded(const Catch::TestRunStats&) override
+	if (RetryNestedToo)
 	{
-		FPlatformMisc::RequestExit(false);
-
-		FEngineLoop::AppPreExit();
-		FModuleManager::Get().UnloadModulesAtShutdown();
-		FEngineLoop::AppExit();
+		AutoRTFM::ForTheRuntime::SetRetryTransaction(AutoRTFM::ForTheRuntime::EAutoRTFMRetryTransactionState::RetryNestedToo);
+	}
+	else if (NoRetry)
+	{
+		AutoRTFM::ForTheRuntime::SetRetryTransaction(AutoRTFM::ForTheRuntime::EAutoRTFMRetryTransactionState::NoRetry);
+	}
+	else
+	{
+		// Otherwise default to just retrying the parent transaction.
+		AutoRTFM::ForTheRuntime::SetRetryTransaction(AutoRTFM::ForTheRuntime::EAutoRTFMRetryTransactionState::RetryNonNested);
 	}
-};
 
-CATCH_REGISTER_LISTENER(SetupListener)
+	GEngineLoop.PreInit(0, nullptr);
+	FModuleManager::Get().StartProcessingNewlyLoadedObjects();
+
+	// Enable all Verse code to run under AutoRTFM (shouldn't affect our tests here, but better safe than sorry).
+	AutoRTFM::ForTheRuntime::SetAutoRTFMRuntime(AutoRTFM::ForTheRuntime::EAutoRTFMEnabledState::AutoRTFM_EnabledForAllVerse);
+
+	// We don't want to trigger ensure's on abort because we are going to test that.
+	AutoRTFM::ForTheRuntime::SetEnsureOnAbortByLanguage(false);
+
+	const int Result = Session.run();
+
+	FPlatformMisc::RequestExit(false);
+
+	FEngineLoop::AppPreExit();
+	FModuleManager::Get().UnloadModulesAtShutdown();
+	FEngineLoop::AppExit();
+
+	return Result;
+}
diff --git a/Engine/Source/Programs/AutoRTFMTests/Private/Benchmarks.cpp b/Engine/Source/Programs/AutoRTFMTests/Private/Benchmarks.cpp
--- a/Engine/Source/Programs/AutoRTFMTests/Private/Benchmarks.cpp
+++ b/Engine/Source/Programs/AutoRTFMTests/Private/Benchmarks.cpp
@@ -187,7 +187,7 @@ void AddToGraph(FGraph<unsigned>& Graph)
 	RandomEdge(G, E);
 }
 
-FGraph<unsigned> BuildGraph(const unsigned Total = 10000)
+FGraph<unsigned> BuildGraph(const unsigned Total = 1000)
 {
 	FGraph<unsigned> Result;
 	for (unsigned Count = Total; Count--;)
diff --git a/Engine/Source/Programs/AutoRTFMTests/Private/Tests.cpp b/Engine/Source/Programs/AutoRTFMTests/Private/Tests.cpp
--- a/Engine/Source/Programs/AutoRTFMTests/Private/Tests.cpp
+++ b/Engine/Source/Programs/AutoRTFMTests/Private/Tests.cpp
@@ -1,7 +1,5 @@
 // Copyright Epic Games, Inc. All Rights Reserved.
 
-#pragma autortfm
-
 #include "Catch2Includes.h"
 #include <AutoRTFM/AutoRTFM.h>
 #include <thread>
@@ -34,7 +32,7 @@ TEST_CASE("Tests.UE_LOG")
 // This test ensures that if you have STM and non-STM modifying data that is
 // adjacent in memory, the STM code won't lose modifications to data that
 // happens to fall into the same STM line.
-TEST_CASE("stm.no_trashing_non_stm")
+TEST_CASE("stm.no_trashing_non_stm", "[.multi-threaded-test]")
 {
 	// A hit-count - lets us ensure each thread is launched and running before
 	// we kick off the meat of the test.
@@ -319,3 +317,24 @@ TEST_CASE("libc.isnan(double)")
 	REQUIRE(false == bXIsNaN);
 	REQUIRE(true == bYIsNaN);
 }
+
+TEST_CASE("Tests.RetryNonNested")
+{
+	// We only run this test if we are retrying non-nested transactions (it proves we retried!).
+	if (!AutoRTFM::ForTheRuntime::ShouldRetryNonNestedTransactions())
+	{
+		return;
+	}
+
+	unsigned Count = 0;
+
+	AutoRTFM::Commit([&]
+	{
+		AutoRTFM::Open([&]
+		{
+			Count++;
+		});
+	});
+
+	REQUIRE(2 == Count);
+}
diff --git a/Engine/Source/Runtime/Core/Private/AutoRTFM/API.cpp b/Engine/Source/Runtime/Core/Private/AutoRTFM/API.cpp
--- a/Engine/Source/Runtime/Core/Private/AutoRTFM/API.cpp
+++ b/Engine/Source/Runtime/Core/Private/AutoRTFM/API.cpp
@@ -14,12 +14,27 @@ namespace
 	int GAutoRTFMRuntimeEnabled = AutoRTFM::EAutoRTFMEnabledState::AutoRTFM_Disabled;
 #endif // UE_AUTORTFM_ENABLED_RUNTIME_BY_DEFAULT
 
-	bool GAutoRTFMEnsureOnAbortByLanguage = true;
-
 	void UpdateAutoRTFMRuntimeCrashData()
 	{
 		FGenericCrashContext::SetGameData(TEXT("IsAutoRTFMRuntimeEnabled"), AutoRTFM::ForTheRuntime::IsAutoRTFMRuntimeEnabled() ? TEXT("true") : TEXT("false"));
 	}
+
+	bool GAutoRTFMEnsureOnAbortByLanguage = true;
+
+	int GAutoRTFMRetryTransactions = AutoRTFM::ForTheRuntime::EAutoRTFMRetryTransactionState::NoRetry;
+
+	void UpdateAutoRTFMRetryTransactionsData()
+	{
+		switch (GAutoRTFMRetryTransactions)
+		{
+		case AutoRTFM::ForTheRuntime::EAutoRTFMRetryTransactionState::NoRetry:
+			return FGenericCrashContext::SetGameData(TEXT("AutoRTFMRetryTransactionState"), TEXT("NoRetry"));
+		case AutoRTFM::ForTheRuntime::EAutoRTFMRetryTransactionState::RetryNonNested:
+			return FGenericCrashContext::SetGameData(TEXT("AutoRTFMRetryTransactionState"), TEXT("RetryNonNested"));
+		case AutoRTFM::ForTheRuntime::EAutoRTFMRetryTransactionState::RetryNestedToo:
+			return FGenericCrashContext::SetGameData(TEXT("AutoRTFMRetryTransactionState"), TEXT("RetryNestedToo"));
+		}
+	}
 }
 
 #if UE_AUTORTFM
@@ -31,10 +46,19 @@ static FAutoConsoleVariableRef CVarAutoRTFMRuntimeEnabled(
 	ECVF_Default
 );
 
+static FAutoConsoleVariableRef CVarAutoRTFMRetryTransactions(
+	TEXT("AutoRTFMRetryTransactions"),
+	GAutoRTFMRetryTransactions,
+	TEXT("Enables the AutoRTFM sanitizer-like mode where we can force an abort-and-retry on transactions (useful to test abort codepaths work as intended)"),
+	FConsoleVariableDelegate::CreateLambda([](IConsoleVariable*) { UpdateAutoRTFMRetryTransactionsData(); }),
+	ECVF_Default
+);
+
 static FDelayedAutoRegisterHelper DelayedAutoRegister(EDelayedRegisterRunPhase::EndOfEngineInit, []
-{
-	UpdateAutoRTFMRuntimeCrashData();
-});
+	{
+		UpdateAutoRTFMRuntimeCrashData();
+		UpdateAutoRTFMRetryTransactionsData();
+	});
 #endif
 
 namespace AutoRTFM
@@ -114,6 +138,45 @@ namespace AutoRTFM
 			return GAutoRTFMEnsureOnAbortByLanguage;
 #else
 			return false;
+#endif
+		}
+
+		void SetRetryTransaction(EAutoRTFMRetryTransactionState State)
+		{
+#if UE_AUTORTFM
+			GAutoRTFMRetryTransactions = State;
+			UpdateAutoRTFMRetryTransactionsData();
+#endif
+		}
+
+		bool ShouldRetryNonNestedTransactions()
+		{
+#if UE_AUTORTFM
+			switch (GAutoRTFMRetryTransactions)
+			{
+			default:
+				return false;
+			case EAutoRTFMRetryTransactionState::RetryNonNested:
+			case EAutoRTFMRetryTransactionState::RetryNestedToo:
+				return true;
+			}
+#else
+			return false;
+#endif
+		}
+
+		bool ShouldRetryNestedTransactionsToo()
+		{
+#if UE_AUTORTFM
+			switch (GAutoRTFMRetryTransactions)
+			{
+			default:
+				return false;
+			case EAutoRTFMRetryTransactionState::RetryNestedToo:
+				return true;
+			}
+#else
+			return false;
 #endif
 		}
 	}
diff --git a/Engine/Source/Runtime/Core/Private/AutoRTFM/Context.cpp b/Engine/Source/Runtime/Core/Private/AutoRTFM/Context.cpp
--- a/Engine/Source/Runtime/Core/Private/AutoRTFM/Context.cpp
+++ b/Engine/Source/Runtime/Core/Private/AutoRTFM/Context.cpp
@@ -242,12 +242,9 @@ void FContext::ClearTransactionStatus()
 	case EContextStatus::OnTrack:
 		break;
 	case EContextStatus::AbortedByLanguage:
-		Status = EContextStatus::OnTrack;
-		break;
 	case EContextStatus::AbortedByRequest:
-		Status = EContextStatus::OnTrack;
-		break;
 	case EContextStatus::AbortedByCascade:
+	case EContextStatus::AbortedByFailedLockAcquisition:
 		Status = EContextStatus::OnTrack;
 		break;
 	default:
@@ -311,7 +308,6 @@ ETransactionResult FContext::Transact(void (*InstrumentedFunction)(void*), void*
         return ETransactionResult::AbortedByLanguage;
     }
     
-	//TUniquePtr<FTransaction> NewTransactionUniquePtr(new FTransaction(this));
 	FTransaction* NewTransaction = new FTransaction(this);
 	FCallNest* NewNest = new FCallNest(this);
 
@@ -333,6 +329,8 @@ ETransactionResult FContext::Transact(void (*InstrumentedFunction)(void*), void*
 		PushCallNest(NewNest);
         OuterTransactStackAddress = TransactStackAddress;
 
+		bool bTriedToRunOnce = false;
+
         for (;;)
         {
             Status = EContextStatus::OnTrack;
@@ -347,7 +345,19 @@ ETransactionResult FContext::Transact(void (*InstrumentedFunction)(void*), void*
 				DumpState();
 				UE_LOG(LogAutoRTFM, Verbose, TEXT("Committing..."));
 
-                if (AttemptToCommitTransaction(CurrentTransaction))
+				if (UNLIKELY(!bTriedToRunOnce && AutoRTFM::ForTheRuntime::ShouldRetryNonNestedTransactions()))
+				{
+					// We skip trying to commit this time, and instead re-run the transaction.
+					Status = EContextStatus::AbortedByFailedLockAcquisition;
+					CurrentTransaction->AbortWithoutThrowing();
+					ClearTransactionStatus();
+
+					// We've tried to run at least once if we get here!
+					bTriedToRunOnce = true;
+
+					continue;
+				}
+                else if (AttemptToCommitTransaction(CurrentTransaction))
                 {
                     Result = ETransactionResult::Committed;
                     break;
@@ -395,14 +405,37 @@ ETransactionResult FContext::Transact(void (*InstrumentedFunction)(void*), void*
     {
 		// This transaction is within another transaction
 		ASSERT(Status == EContextStatus::OnTrack);
+
 		PushTransaction(NewTransaction);
 		PushCallNest(NewNest);
 
-		CurrentNest->Try([&]() { InstrumentedFunction(Arg); });
-		ASSERT(CurrentTransaction == NewTransaction);
+		bool bTriedToRunOnce = false;
+
+		for (;;)
+		{
+			CurrentNest->Try([&]() { InstrumentedFunction(Arg); });
+			ASSERT(CurrentTransaction == NewTransaction);
+
+			if (Status == EContextStatus::OnTrack)
+			{
+				if (UNLIKELY(!bTriedToRunOnce && AutoRTFM::ForTheRuntime::ShouldRetryNestedTransactionsToo()))
+				{
+					// We skip trying to commit this time, and instead re-run the transaction.
+					Status = EContextStatus::AbortedByFailedLockAcquisition;
+					NewTransaction->AbortWithoutThrowing();
+					ClearTransactionStatus();
+
+					// We've tried to run at least once if we get here!
+					bTriedToRunOnce = true;
+
+					continue;
+				}
+			}
+
+			Result = ResolveNestedTransaction(NewTransaction);
+			break;
+		}
 
-		Result = ResolveNestedTransaction(NewTransaction);
-		
 		PopCallNest();
 		PopTransaction();
 
diff --git a/Engine/Source/Runtime/Core/Public/AutoRTFM/AutoRTFM.h b/Engine/Source/Runtime/Core/Public/AutoRTFM/AutoRTFM.h
--- a/Engine/Source/Runtime/Core/Public/AutoRTFM/AutoRTFM.h
+++ b/Engine/Source/Runtime/Core/Public/AutoRTFM/AutoRTFM.h
@@ -697,8 +697,25 @@ namespace ForTheRuntime
 		AutoRTFM_EnabledForAllVerse,
 	};
 
-	// Set whether the AutoRTFM runtime is enabled or disabled.
-	UE_AUTORTFM_API bool SetAutoRTFMRuntime(EAutoRTFMEnabledState bEnabled);
+	// An enum to represent whether we should abort and retry transactions (for testing purposes).
+	enum EAutoRTFMRetryTransactionState 
+	{
+		// Do not abort and retry transactions (the default).
+		NoRetry = 0,
+
+		// Abort and retry non-nested transactions (EG. only abort the parent transactional nest).
+		RetryNonNested,
+
+		// Abort and retry nested-transactions too. Will be slower as each nested-transaction will
+		// be aborted and retried at least *twice* (once when the non-nested transaction runs the
+		// first time, and a second time when the non-nested transaction is doing its retry after
+		// aborting).
+		RetryNestedToo,
+	};
+
+	// Set whether the AutoRTFM runtime is enabled or disabled. Returns true when the state was changed
+	// successfully.
+	UE_AUTORTFM_API bool SetAutoRTFMRuntime(EAutoRTFMEnabledState State);
 
 	// Query whether the AutoRTFM runtime is enabled.
 	UE_AUTORTFM_API bool IsAutoRTFMRuntimeEnabled();
@@ -712,6 +729,15 @@ namespace ForTheRuntime
 	// Returns whether the runtime will trigger an ensure on an abort-by-language, or not.
 	UE_AUTORTFM_API bool IsEnsureOnAbortByLanguageEnabled();
 
+	// Set whether we should retry transactions.
+	UE_AUTORTFM_API void SetRetryTransaction(EAutoRTFMRetryTransactionState State);
+
+	// Returns true if we should retry non-nested transactions.
+	UE_AUTORTFM_API bool ShouldRetryNonNestedTransactions();
+
+	// Returns true if we should also retry nested transactions.
+	UE_AUTORTFM_API bool ShouldRetryNestedTransactionsToo();
+
 	// Manually create a new transaction from open code and push it as a transaction nest.
 	// Can only be called within an already active parent transaction (EG. this cannot start
 	// a transaction nest itself).
