Rearrange TraceLog message types and emit fatal error for OOM.

Rearrange the order of messages so that the order goes from informational to increased seriousness and categorize by value. Also add a fatal category which implies catastrophic stop of the tracing application.

diff --git a/Engine/Source/Runtime/Core/Private/ProfilingDebugging/TraceAuxiliary.cpp b/Engine/Source/Runtime/Core/Private/ProfilingDebugging/TraceAuxiliary.cpp
--- a/Engine/Source/Runtime/Core/Private/ProfilingDebugging/TraceAuxiliary.cpp
+++ b/Engine/Source/Runtime/Core/Private/ProfilingDebugging/TraceAuxiliary.cpp
@@ -844,34 +844,47 @@ void OnConnectionCallback()
 ////////////////////////////////////////////////////////////////////////////////
 void OnMessageCallback(const UE::Trace::FMessageEvent& Message)
 {
+	#define LOG_WITH_VERBOSITY(Verbosity, Text) \
+		UE_LOG(LogTrace, Verbosity, TEXT("%s"), Text)
+
 	const auto TypeStr = StringCast<TCHAR>(Message.TypeStr);
+	TStringBuilder<128> LogStr;
+	if (Message.Description != nullptr)
+	{
+		LogStr << Message.TypeStr << TEXT(" ") << Message.Description;
+	}
+	else
+	{
+		LogStr << Message.TypeStr;
+	}
+	
 	switch(Message.Type)
 	{
-	// Information to the user
-	case UE::Trace::EMessageType::Info:
-		{
-			if (Message.Description != nullptr)
-			{
-				const auto Description = StringCast<TCHAR>(Message.Description);
-				UE_LOG(LogTrace, Display, TEXT("%s"), Description.Get());
-			}
-		}
+	case UE::Trace::EMessageType::Log:
+		LOG_WITH_VERBOSITY(Log, LogStr.ToString());
+		break;
+	case UE::Trace::EMessageType::Display:
+		LOG_WITH_VERBOSITY(Display, LogStr.ToString());
 		break;
-	// By default treated as errors	
 	default:
 		{
-			if (Message.Description != nullptr)
+			if (Message.Type > UE::Trace::EMessageType::FatalStart)
 			{
-				const auto Description = StringCast<TCHAR>(Message.Description);
-				UE_LOG(LogTrace, Error, TEXT("%s %s"), TypeStr.Get(), Description.Get());
+				LOG_WITH_VERBOSITY(Fatal, LogStr.ToString());
 			}
-			else
+			else if (Message.Type > UE::Trace::EMessageType::ErrorStart)
+			{
+				LOG_WITH_VERBOSITY(Error, LogStr.ToString());
+			}
+			else if (Message.Type > UE::Trace::EMessageType::WarningStart)
 			{
-				UE_LOG(LogTrace, Error, TEXT("%s"), TypeStr.Get());
+				LOG_WITH_VERBOSITY(Warning, LogStr.ToString());
 			}
 		}
 		break;
 	}
+	
+	#undef LOG_WITH_VERBOSITY
 }
 
 ////////////////////////////////////////////////////////////////////////////////
diff --git a/Engine/Source/Runtime/TraceLog/Private/Trace/Control.cpp b/Engine/Source/Runtime/TraceLog/Private/Trace/Control.cpp
--- a/Engine/Source/Runtime/TraceLog/Private/Trace/Control.cpp
+++ b/Engine/Source/Runtime/TraceLog/Private/Trace/Control.cpp
@@ -134,7 +134,7 @@ static bool Writer_ControlListen()
 		return false;
 	}
 
-	UE_TRACE_MESSAGE_F(Info, "Control listening on port %u", GControlPort);
+	UE_TRACE_MESSAGE_F(Display, "Control listening on port %u", GControlPort);
 	GControlState = EControlState::Listening;
 	return true;
 }
diff --git a/Engine/Source/Runtime/TraceLog/Private/Trace/Writer.cpp b/Engine/Source/Runtime/TraceLog/Private/Trace/Writer.cpp
--- a/Engine/Source/Runtime/TraceLog/Private/Trace/Writer.cpp
+++ b/Engine/Source/Runtime/TraceLog/Private/Trace/Writer.cpp
@@ -251,6 +251,11 @@ void* Writer_MemoryAllocate(SIZE_T Size, uint32 Alignment)
 	}
 #endif // TRACE_PRIVATE_STOMP
 
+	if (Ret == nullptr)
+	{
+		UE_TRACE_MESSAGE_F(OOMFatal, "OOM allocating %llu bytes", uint64(Size));
+	}
+
 #if TRACE_PRIVATE_STATISTICS
 	AtomicAddRelaxed(&GTraceStatistics.MemoryUsed, uint64(Size));
 #endif
diff --git a/Engine/Source/Runtime/TraceLog/Public/Trace/Trace.h b/Engine/Source/Runtime/TraceLog/Public/Trace/Trace.h
--- a/Engine/Source/Runtime/TraceLog/Public/Trace/Trace.h
+++ b/Engine/Source/Runtime/TraceLog/Public/Trace/Trace.h
@@ -81,7 +81,18 @@ using OnConnectFunc = void(void);
 
 enum class EMessageType : uint8
 {
-	Reserved = 0,
+	Reserved			= 0,
+	// Add to log
+	Log,
+	// For backwards compatibility
+	Info = Log,
+	// Display in console or similar
+	Display,
+	// Warnings to notify user
+	WarningStart		= 0x04,
+	// Errors are critical to the user, but application
+	// can continue to run.
+	ErrorStart			= 0x10,
 	WriteError,
 	ReadError,
 	ConnectError,
@@ -89,7 +100,9 @@ enum class EMessageType : uint8
 	EstablishError,
 	FileOpenError,
 	WriterError,
-	Info				= 0xff,
+	// Fatal errors should cause application to stop
+	FatalStart			= 0x40,
+	OOMFatal,
 };
 	
 struct FMessageEvent
