[TSAN]
  - Fix TSAN warning in StallDetector

#rnx

diff --git a/Engine/Source/Runtime/Core/Private/ProfilingDebugging/StallDetector.cpp b/Engine/Source/Runtime/Core/Private/ProfilingDebugging/StallDetector.cpp
--- a/Engine/Source/Runtime/Core/Private/ProfilingDebugging/StallDetector.cpp
+++ b/Engine/Source/Runtime/Core/Private/ProfilingDebugging/StallDetector.cpp
@@ -88,9 +88,9 @@ namespace UE
 			Stop();
 		}
 
-		bool GetStartedThread()
+		bool GetStartedThread() const
 		{
-			return StartedThread;
+			return StartedThread.load(std::memory_order_relaxed);
 		}
 		
 #if STALL_DETECTOR_HEART_BEAT_CLOCK
@@ -101,7 +101,7 @@ namespace UE
 #endif
 
 	private:
-		bool StartedThread;
+		std::atomic<bool> StartedThread;
 		bool StopThread;
 
 #if STALL_DETECTOR_HEART_BEAT_CLOCK
