[Ticker]
 - Fix TSAN warning

#rnx

diff --git a/Engine/Source/Runtime/Core/Private/Containers/Ticker.cpp b/Engine/Source/Runtime/Core/Private/Containers/Ticker.cpp
--- a/Engine/Source/Runtime/Core/Private/Containers/Ticker.cpp
+++ b/Engine/Source/Runtime/Core/Private/Containers/Ticker.cpp
@@ -12,14 +12,14 @@ FTSTicker& FTSTicker::GetCoreTicker()
 
 FTSTicker::FDelegateHandle FTSTicker::AddTicker(const FTickerDelegate& InDelegate, float InDelay)
 {
-	FElementPtr NewElement{ new FElement{ CurrentTime + InDelay, InDelay, InDelegate } };
+	FElementPtr NewElement{ new FElement{ CurrentTime.load(std::memory_order_relaxed) + InDelay, InDelay, InDelegate } };
 	AddedElements.Enqueue(NewElement);
 	return NewElement;
 }
 
 FTSTicker::FDelegateHandle FTSTicker::AddTicker(const TCHAR* InName, float InDelay, TFunction<bool(float)> Function)
 {
-	FElementPtr NewElement{ new FElement{ CurrentTime + InDelay, InDelay, FTickerDelegate::CreateLambda(Function) } };
+	FElementPtr NewElement{ new FElement{ CurrentTime.load(std::memory_order_relaxed) + InDelay, InDelay, FTickerDelegate::CreateLambda(Function) } };
 	AddedElements.Enqueue(NewElement);
 	return NewElement;
 }
@@ -75,7 +75,8 @@ void FTSTicker::Tick(float DeltaTime)
 		return;
 	}
 
-	CurrentTime += DeltaTime;
+	// We can do a relaxed read/store since only the game thread will call the tick function.
+	CurrentTime.store(CurrentTime.load(std::memory_order_relaxed) + DeltaTime, std::memory_order_relaxed);
 
 	TArray<FElementPtr> TickedElements;
 	int32 ElementIdx = 0;
@@ -100,7 +101,7 @@ void FTSTicker::Tick(float DeltaTime)
 				continue;
 			}
 
-			if (Element->FireTime > CurrentTime)
+			if (Element->FireTime > CurrentTime.load(std::memory_order_relaxed))
 			{
 				ClearExecutionFlag(Element);
 				TickedElements.Add(MoveTemp(Element));
@@ -113,7 +114,7 @@ void FTSTicker::Tick(float DeltaTime)
 				if (PrevState != FElement::RemovedState)
 				{
 					checkf(PrevState == FElement::DefaultState, TEXT("Invalid state %u"), PrevState);
-					Element->FireTime = CurrentTime + Element->DelayTime;
+					Element->FireTime = CurrentTime.load(std::memory_order_relaxed) + Element->DelayTime;
 					TickedElements.Add(MoveTemp(Element));
 				}
 			}
diff --git a/Engine/Source/Runtime/Core/Public/Containers/Ticker.h b/Engine/Source/Runtime/Core/Public/Containers/Ticker.h
--- a/Engine/Source/Runtime/Core/Public/Containers/Ticker.h
+++ b/Engine/Source/Runtime/Core/Public/Containers/Ticker.h
@@ -123,7 +123,7 @@ private:
 	TMpscQueue<FElementPtr> AddedElements;
 
 	/** Current time of the ticker **/
-	double CurrentTime{ 0.0 };
+	std::atomic<double> CurrentTime{ 0.0 };
 	/** Future delegates to fire **/
 	TArray<FElementPtr> Elements;
 };
