[ThreadSingleton]
 - Remove check that could trip depending on the tls destruction order

#rnx

diff --git a/Engine/Source/Runtime/Core/Public/HAL/ThreadSingleton.h b/Engine/Source/Runtime/Core/Public/HAL/ThreadSingleton.h
--- a/Engine/Source/Runtime/Core/Public/HAL/ThreadSingleton.h
+++ b/Engine/Source/Runtime/Core/Public/HAL/ThreadSingleton.h
@@ -73,7 +73,6 @@ protected:
 	virtual ~TThreadSingleton()
 	{
 		// Clean the dangling pointer from the TLS.
-		check(GetTlsSlot() != nullptr);
 		if(((FTlsAutoCleanup*)GetTlsSlot()) == static_cast<FTlsAutoCleanup*>(this))
 		{
 			GetTlsSlot() =  nullptr;
