Make the AddToFrontend GFA default to unloading on deactivation ( maintaining default false for existing assets ).

diff --git a/Engine/Source/Runtime/Core/Public/UObject/FortniteMainBranchObjectVersion.h b/Engine/Source/Runtime/Core/Public/UObject/FortniteMainBranchObjectVersion.h
--- a/Engine/Source/Runtime/Core/Public/UObject/FortniteMainBranchObjectVersion.h
+++ b/Engine/Source/Runtime/Core/Public/UObject/FortniteMainBranchObjectVersion.h
@@ -493,6 +493,9 @@ struct FFortniteMainBranchObjectVersion
 
 		// Fix chaos cloth buckling stiffness parameter bug 
 		ChaosClothAllowZeroBucklingStiffness,
+
+		// AddToFrontend GFA now defaults to unload plugin on exit frontend
+		GameFeatureDataActionAddToFrontendDefaultToUnload,
 		
 		// -----<new versions can be added above this line>-------------------------------------------------
 		VersionPlusOne,
