Fork
* Added cvar fork.ParentSleepDurationInSec to give ability to configure the time parent process sleeps when it has no signals to execute
* Added cvar fork.LogMemoryStatsWhenForking to control the memory stats log before every fork (enabled by default)
* Added LogFork category used by fork related logs (previously LogHal)
* Fixed signal timestamp in logs

diff --git a/Engine/Source/Runtime/Core/Private/Unix/UnixPlatformProcess.cpp b/Engine/Source/Runtime/Core/Private/Unix/UnixPlatformProcess.cpp
--- a/Engine/Source/Runtime/Core/Private/Unix/UnixPlatformProcess.cpp
+++ b/Engine/Source/Runtime/Core/Private/Unix/UnixPlatformProcess.cpp
@@ -8,6 +8,7 @@
 #include "Containers/StringConv.h"
 #include "Logging/LogMacros.h"
 #include "HAL/FileManager.h"
+#include "HAL/IConsoleManager.h"
 #include "Misc/Parse.h"
 #include "HAL/Runnable.h"
 #include "HAL/RunnableThread.h"
@@ -33,6 +34,10 @@
 #include "HAL/PlatformTime.h"
 #include "Misc/DateTime.h"
 
+
+
+DECLARE_LOG_CATEGORY_CLASS(LogFork, Log, All);
+
 namespace PlatformProcessLimits
 {
 	enum
@@ -41,6 +46,21 @@ namespace PlatformProcessLimits
 	};
 };
 
+namespace UnixPlatformProcess
+{
+	static float GParentSleepDurationInSec = 10.0f;
+	static FAutoConsoleVariableRef CVarForkParentSleepDurationInSec(TEXT("fork.ParentSleepDurationInSec"), 
+		GParentSleepDurationInSec,
+		TEXT("The time in seconds the parent process will sleep when it has no more signals to process."), 
+		ECVF_Default);
+
+	static bool GLogMemoryStatsWhenForking = true;
+	static FAutoConsoleVariableRef CVarLogMemoryStatsWhenForking(TEXT("fork.LogMemoryStatsWhenForking"),
+		GLogMemoryStatsWhenForking,
+		TEXT("When true the parent process will log memory stats before every fork."),
+		ECVF_Default);
+}
+
 #if IS_MONOLITHIC
 __thread uint32 FUnixTLS::ThreadIdTLS = 0;
 #else
@@ -1407,7 +1427,6 @@ static FDelegateHandle OnEndFrameHandle;
 FGenericPlatformProcess::EWaitAndForkResult FUnixPlatformProcess::WaitAndFork()
 {
 #define WAIT_AND_FORK_QUEUE_LENGTH 4096
-#define WAIT_AND_FORK_PARENT_SLEEP_DURATION 10
 #define WAIT_AND_FORK_CHILD_SPAWN_DELAY 0.125
 #ifndef WAIT_AND_FORK_PARENT_SHUTDOWN_EXIT_CODE
 	#define WAIT_AND_FORK_PARENT_SHUTDOWN_EXIT_CODE 0
@@ -1425,10 +1444,10 @@ FGenericPlatformProcess::EWaitAndForkResult FUnixPlatformProcess::WaitAndFork()
 	struct FForkSignalData
 	{
 		FForkSignalData() = default;
-		FForkSignalData(int32 InSignal, double InTimeSeconds) : SignalValue(InSignal), TimeSeconds(InTimeSeconds) {}
+		FForkSignalData(int32 InSignal, FDateTime&& InTime) : SignalValue(InSignal), Time(MoveTemp(InTime)) {}
 
 		int32 SignalValue = 0;
-		double TimeSeconds = 0.0;
+		FDateTime Time;
 	};
 
 	static TCircularQueue<FForkSignalData> WaitAndForkSignalQueue(WAIT_AND_FORK_QUEUE_LENGTH);
@@ -1441,7 +1460,7 @@ FGenericPlatformProcess::EWaitAndForkResult FUnixPlatformProcess::WaitAndFork()
 	{
 		for (int32 ForkIdx = 0; ForkIdx < NumForks; ++ForkIdx)
 		{
-			WaitAndForkSignalQueue.Enqueue(FForkSignalData(ForkIdx + 1, FPlatformTime::Seconds()));
+			WaitAndForkSignalQueue.Enqueue(FForkSignalData(ForkIdx + 1, FDateTime::UtcNow()));
 		}
 	}
 
@@ -1453,7 +1472,7 @@ FGenericPlatformProcess::EWaitAndForkResult FUnixPlatformProcess::WaitAndFork()
 		bool bDirExists = IFileManager::Get().DirectoryExists(*ChildParametersPath);
 		if (!bDirExists)
 		{
-			UE_LOG(LogHAL, Fatal, TEXT("Path referred to by -WaitAndForkCmdLinePath does not exist: %s"), *ChildParametersPath);
+			UE_LOG(LogFork, Fatal, TEXT("Path referred to by -WaitAndForkCmdLinePath does not exist: %s"), *ChildParametersPath);
 		}
 	}
 
@@ -1464,7 +1483,7 @@ FGenericPlatformProcess::EWaitAndForkResult FUnixPlatformProcess::WaitAndFork()
 	FParse::Value(FCommandLine::Get(), TEXT("-WaitAndForkResponseTimeout="), WaitAndForkResponseTimeout);
 	if (WaitAndForkResponseTimeout > 0.0)
 	{
-		UE_LOG(LogHAL, Log, TEXT("WaitAndFork setting WaitAndForkResponseTimeout to %0.2f seconds."), WaitAndForkResponseTimeout);
+		UE_LOG(LogFork, Log, TEXT("WaitAndFork setting WaitAndForkResponseTimeout to %0.2f seconds."), WaitAndForkResponseTimeout);
 	}
 
 	FCoreDelegates::OnParentBeginFork.Broadcast();
@@ -1478,20 +1497,20 @@ FGenericPlatformProcess::EWaitAndForkResult FUnixPlatformProcess::WaitAndFork()
 		Action.sa_sigaction = [](int32 Signal, siginfo_t* Info, void* Context) {
 			if (Signal == WAIT_AND_FORK_QUEUE_SIGNAL && Info)
 			{
-				WaitAndForkSignalQueue.Enqueue(FForkSignalData(Info->si_value.sival_int, FPlatformTime::Seconds()));
+				WaitAndForkSignalQueue.Enqueue(FForkSignalData(Info->si_value.sival_int, FDateTime::UtcNow()));
 			}
 		};
 		sigaction(WAIT_AND_FORK_QUEUE_SIGNAL, &Action, nullptr);
 	}
 
-	UE_LOG(LogHAL, Log, TEXT("   *** WaitAndFork awaiting signal %d to process pid %d create child processes... ***"), WAIT_AND_FORK_QUEUE_SIGNAL, FPlatformProcess::GetCurrentProcessId());
+	UE_LOG(LogFork, Log, TEXT("   *** WaitAndFork awaiting signal %d to process pid %d create child processes... ***"), WAIT_AND_FORK_QUEUE_SIGNAL, FPlatformProcess::GetCurrentProcessId());
 	GLog->Flush();
 
 	struct FMemoryStatsHolder
 	{
-		double AvailablePhysical;
-		double PeakUsedPhysical;
-		double PeakUsedVirtual;
+		double AvailablePhysical = 0.0;
+		double PeakUsedPhysical = 0.0;
+		double PeakUsedVirtual = 0.0;
 
 		constexpr double ByteToMiB(uint64 InBytes) { return static_cast<double>(InBytes) / (1024.0 * 1024.0); }
 
@@ -1502,7 +1521,12 @@ FGenericPlatformProcess::EWaitAndForkResult FUnixPlatformProcess::WaitAndFork()
 		{ }
 	};
 
-	FMemoryStatsHolder PreviousMasterMemStats(FPlatformMemory::GetStats());
+	TOptional<FMemoryStatsHolder> PreviousMasterMemStats;
+	
+	if (UnixPlatformProcess::GLogMemoryStatsWhenForking)
+	{
+		PreviousMasterMemStats = FMemoryStatsHolder(FPlatformMemory::GetStats());
+	}
 
 	EWaitAndForkResult RetVal = EWaitAndForkResult::Parent;
 	struct FPidAndSignal
@@ -1522,25 +1546,30 @@ FGenericPlatformProcess::EWaitAndForkResult FUnixPlatformProcess::WaitAndFork()
 		FForkSignalData SignalData;
 		if (WaitAndForkSignalQueue.Dequeue(SignalData))
 		{
-			// Sleep for a short while to avoid spamming new processes to the OS all at once
-			FPlatformProcess::Sleep(WAIT_AND_FORK_CHILD_SPAWN_DELAY);
-
 			uint16 Cookie = (SignalData.SignalValue >> 16) & 0xffff;
 			uint16 ChildIdx = SignalData.SignalValue & 0xffff;
 
-			FDateTime SignalReceived = FDateTime::FromUnixTimestamp(FMath::FloorToInt64(SignalData.TimeSeconds));
+			UE_LOG(LogFork, Log, TEXT("[Parent] WaitAndForkSignalQueue has %d signals to execute. Sleeping for %f secs before processing request %04hx-%04hx "), WaitAndForkSignalQueue.Count(), WAIT_AND_FORK_CHILD_SPAWN_DELAY, Cookie, ChildIdx);
+
+			// Sleep for a short while to avoid spamming new processes to the OS all at once
+			FPlatformProcess::Sleep(WAIT_AND_FORK_CHILD_SPAWN_DELAY);
+
+			FDateTime SignalReceived = SignalData.Time;
 
 			FCoreDelegates::OnParentPreFork.Broadcast();
 
-			UE_LOG(LogHAL, Log, TEXT("[Parent] WaitAndFork processing child request %04hx-%04hx received at: %s"), Cookie, ChildIdx, *SignalReceived.ToString());
+			UE_LOG(LogFork, Log, TEXT("[Parent] WaitAndFork processing child request %04hx-%04hx received at: %s"), Cookie, ChildIdx, *SignalReceived.ToString());
 
-			FMemoryStatsHolder CurrentMasterMemStats(FPlatformMemory::GetStats());
-			UE_LOG(LogHAL, Log, TEXT("MemoryStats PreFork: AvailablePhysical: %.02fMiB (%+.02fMiB), PeakPhysical: %.02fMiB, PeakVirtual: %.02fMiB"),
-				CurrentMasterMemStats.AvailablePhysical, (CurrentMasterMemStats.AvailablePhysical - PreviousMasterMemStats.AvailablePhysical),
-				CurrentMasterMemStats.PeakUsedPhysical,
-				CurrentMasterMemStats.PeakUsedVirtual				
-			);
-			PreviousMasterMemStats = CurrentMasterMemStats;
+			if (UnixPlatformProcess::GLogMemoryStatsWhenForking)
+			{
+				FMemoryStatsHolder CurrentMasterMemStats(FPlatformMemory::GetStats());
+				UE_LOG(LogFork, Log, TEXT("[Parent] MemoryStats PreFork: AvailablePhysical: %.02fMiB (%+.02fMiB), PeakPhysical: %.02fMiB, PeakVirtual: %.02fMiB"),
+					CurrentMasterMemStats.AvailablePhysical, (CurrentMasterMemStats.AvailablePhysical - PreviousMasterMemStats->AvailablePhysical),
+					CurrentMasterMemStats.PeakUsedPhysical,
+					CurrentMasterMemStats.PeakUsedVirtual
+				);
+				PreviousMasterMemStats = CurrentMasterMemStats;
+			}
 			
 			// Make sure there are no pending messages in the log.
 			GLog->Flush();
@@ -1558,7 +1587,7 @@ FGenericPlatformProcess::EWaitAndForkResult FUnixPlatformProcess::WaitAndFork()
 				// Fork errors may include hitting process limits or other environmental factors so we will just report the issue since the environmental factor can be
 				// fixed while the process is still running.
 				int ErrNo = errno;
-				UE_LOG(LogHAL, Error, TEXT("WaitAndFork failed to fork! fork() error:%d"), ErrNo);
+				UE_LOG(LogFork, Error, TEXT("WaitAndFork failed to fork! fork() error:%d"), ErrNo);
 			}
 			else if (ChildPID == 0)
 			{
@@ -1590,13 +1619,14 @@ FGenericPlatformProcess::EWaitAndForkResult FUnixPlatformProcess::WaitAndFork()
 					}
 					else
 					{
-						UE_LOG(LogHAL, Error, TEXT("[Child] WaitAndFork child %04hx-%04hx failed to set command line from: %s"), Cookie, ChildIdx, *CmdLineFilename);
+						UE_LOG(LogFork, Error, TEXT("[Child] WaitAndFork child %04hx-%04hx failed to set command line from: %s"), Cookie, ChildIdx, *CmdLineFilename);
 					}
 				}
 
 				// Start up the log again
 				FPlatformOutputDevices::SetupOutputDevices();
 				GLog->SetCurrentThreadAsPrimaryThread();
+				
 
 				// Set the process name, if specified
 				if (ChildIdx > 0)
@@ -1604,7 +1634,7 @@ FGenericPlatformProcess::EWaitAndForkResult FUnixPlatformProcess::WaitAndFork()
 					if (prctl(PR_SET_NAME, TCHAR_TO_UTF8(*FString::Printf(TEXT("DS-%04hx-%04hx"), Cookie, ChildIdx))) != 0)
 					{
 						int ErrNo = errno;
-						UE_LOG(LogHAL, Fatal, TEXT("[Child] WaitAndFork failed to set process name with prctl! error:%d"), ErrNo);
+						UE_LOG(LogFork, Fatal, TEXT("[Child] WaitAndFork failed to set process name with prctl! error:%d"), ErrNo);
 					}
 				}
 
@@ -1629,15 +1659,18 @@ FGenericPlatformProcess::EWaitAndForkResult FUnixPlatformProcess::WaitAndFork()
 
 					const double StartChildWaitSeconds = FPlatformTime::Seconds();
 
-					UE_LOG(LogHAL, Log, TEXT("[Child] WaitAndFork child %04hx-%04hx waiting for signal %d to proceed."), Cookie, ChildIdx, WAIT_AND_FORK_RESPONSE_SIGNAL);
+					UE_LOG(LogFork, Log, TEXT("[Child] WaitAndFork child %04hx-%04hx waiting for signal %d to proceed."), Cookie, ChildIdx, WAIT_AND_FORK_RESPONSE_SIGNAL);
+
 					while (!IsEngineExitRequested() && !bResponseReceived)
 					{
+						UE_LOG(LogFork, Verbose, TEXT("[Child] WaitAndFork child %04hx-%04hx signal %d not received. Sleeping for 1sec."), Cookie, ChildIdx, WAIT_AND_FORK_RESPONSE_SIGNAL);
+
 						FPlatformProcess::Sleep(1);
 
 						// Check to see how long we've been waiting and if we should time out.
 						if ((WaitAndForkResponseTimeout > 0.0) && ((FPlatformTime::Seconds() - StartChildWaitSeconds) > WaitAndForkResponseTimeout))
 						{
-							UE_LOG(LogHAL, Error, TEXT("[Child] WaitAndFork child %04hx-%04hx has exceeded WAIT_AND_FORK_RESPONSE_SIGNAL timeout"), Cookie, ChildIdx);
+							UE_LOG(LogFork, Error, TEXT("[Child] WaitAndFork child %04hx-%04hx has exceeded WAIT_AND_FORK_RESPONSE_SIGNAL timeout"), Cookie, ChildIdx);
 							FPlatformMisc::RequestExitWithStatus(true, WAIT_AND_FORK_RESPONSE_TIMEOUT_EXIT_CODE);
 							break;
 						}
@@ -1647,7 +1680,7 @@ FGenericPlatformProcess::EWaitAndForkResult FUnixPlatformProcess::WaitAndFork()
 					sigaction(WAIT_AND_FORK_RESPONSE_SIGNAL, &Action, nullptr);
 				}
 
-				UE_LOG(LogHAL, Log, TEXT("[Child] WaitAndFork child process %04hx-%04hx has started with pid %d."), Cookie, ChildIdx, GetCurrentProcessId());
+				UE_LOG(LogFork, Log, TEXT("[Child] WaitAndFork child process %04hx-%04hx has started with pid %d."), Cookie, ChildIdx, GetCurrentProcessId());
 				FApp::PrintStartupLogMessages();
 
 				OnEndFrameHandle = FCoreDelegates::OnEndFrame.AddStatic(FUnixPlatformProcess::OnChildEndFramePostFork);
@@ -1674,13 +1707,15 @@ FGenericPlatformProcess::EWaitAndForkResult FUnixPlatformProcess::WaitAndFork()
 
 				FCoreDelegates::OnPostFork.Broadcast(EForkProcessRole::Parent);
 
-				UE_LOG(LogHAL, Log, TEXT("[Parent] WaitAndFork Successfully processed request %04hx-%04hx, made a child with pid %d! Total number of children: %d."), Cookie, ChildIdx, ChildPID, AllChildren.Num());
+				UE_LOG(LogFork, Log, TEXT("[Parent] WaitAndFork Successfully processed request %04hx-%04hx, made a child with pid %d! Total number of children: %d."), Cookie, ChildIdx, ChildPID, AllChildren.Num());
 			}
 		}
 		else
 		{
+			UE_LOG(LogFork, Verbose, TEXT("[Parent] SignalQueue is empty. Sleeping for %f secs"), UnixPlatformProcess::GParentSleepDurationInSec);
+
 			// No signal to process. Sleep for a bit and do some bookkeeping.
-			FPlatformProcess::Sleep(WAIT_AND_FORK_PARENT_SLEEP_DURATION);
+			FPlatformProcess::Sleep(UnixPlatformProcess::GParentSleepDurationInSec);
 
 			// Trim terminated children
 			for (int32 ChildIdx = AllChildren.Num() - 1; ChildIdx >= 0; --ChildIdx)
@@ -1692,24 +1727,24 @@ FGenericPlatformProcess::EWaitAndForkResult FUnixPlatformProcess::WaitAndFork()
 				if (WaitResult == -1)
 				{
 					int32 ErrNo = errno;
-					UE_LOG(LogHAL, Log, TEXT("[Parent] WaitAndFork unknown error while querying existance of child %d. Error:%d"), ChildPidAndSignal.Pid, ErrNo);
+					UE_LOG(LogFork, Log, TEXT("[Parent] WaitAndFork unknown error while querying existance of child %d. Error:%d"), ChildPidAndSignal.Pid, ErrNo);
 				}
 				else if (WaitResult != 0)
 				{
 					int32 ExitCode = WIFEXITED(Status) ? WEXITSTATUS(Status) : 0;
 					if (ExitCode != 0 && ExitCode == WAIT_AND_FORK_PARENT_SHUTDOWN_EXIT_CODE)
 					{
-						UE_LOG(LogHAL, Log, TEXT("[Parent] WaitAndFork child %d exited with return code %d, indicating that the parent process should shut down. Shutting down..."), ChildPidAndSignal.Pid, WAIT_AND_FORK_PARENT_SHUTDOWN_EXIT_CODE);
+						UE_LOG(LogFork, Log, TEXT("[Parent] WaitAndFork child %d exited with return code %d, indicating that the parent process should shut down. Shutting down..."), ChildPidAndSignal.Pid, WAIT_AND_FORK_PARENT_SHUTDOWN_EXIT_CODE);
 						RequestEngineExit(TEXT("Unix Child has exited"));
 					}
 					else if (NumForks > 0 && ChildPidAndSignal.SignalValue > 0 && ChildPidAndSignal.SignalValue <= NumForks)
 					{
-						UE_LOG(LogHAL, Log, TEXT("[Parent] WaitAndFork child %d missing. This was NumForks child %d. Relaunching..."), ChildPidAndSignal.Pid, ChildPidAndSignal.SignalValue);
-						WaitAndForkSignalQueue.Enqueue(FForkSignalData(ChildPidAndSignal.SignalValue, FPlatformTime::Seconds()));
+						UE_LOG(LogFork, Log, TEXT("[Parent] WaitAndFork child %d missing. This was NumForks child %d. Relaunching..."), ChildPidAndSignal.Pid, ChildPidAndSignal.SignalValue);
+						WaitAndForkSignalQueue.Enqueue(FForkSignalData(ChildPidAndSignal.SignalValue, FDateTime::UtcNow()));
 					}
 					else
 					{
-						UE_LOG(LogHAL, Log, TEXT("[Parent] WaitAndFork child %d missing. Removing from children list..."), ChildPidAndSignal.Pid);
+						UE_LOG(LogFork, Log, TEXT("[Parent] WaitAndFork child %d missing. Removing from children list..."), ChildPidAndSignal.Pid);
 					}
 
 					AllChildren.RemoveAt(ChildIdx, EAllowShrinking::No);
