- Replaced some LowLevelOutputDebugStrings during GMalloc startup with NSLog because the LowLevel calls allocate memory while setting up memory...

diff --git a/Engine/Source/Runtime/Core/Private/Apple/ApplePlatformMemory.cpp b/Engine/Source/Runtime/Core/Private/Apple/ApplePlatformMemory.cpp
--- a/Engine/Source/Runtime/Core/Private/Apple/ApplePlatformMemory.cpp
+++ b/Engine/Source/Runtime/Core/Private/Apple/ApplePlatformMemory.cpp
@@ -293,13 +293,13 @@ void FApplePlatformMemory::SetAllocatorToUse()
     // force Ansi allocator in particular cases
     if(getenv("UE4_FORCE_MALLOC_ANSI") != nullptr)
     {
-		FPlatformMisc::LowLevelOutputDebugString(TEXT("UE4_FORCE_MALLOC_ANSI is set, using Ansi allocator.\n"));
+		NSLog(@"UE4_FORCE_MALLOC_ANSI is set, using Ansi allocator.\n");
         AllocatorToUse = EMemoryAllocatorToUse::Ansi;
         return;
     }
     if (FORCE_ANSI_ALLOCATOR)
     {
-		FPlatformMisc::LowLevelOutputDebugString(TEXT("FORCE_ANSI_ALLOCATOR defined, using Ansi allocator.\n"));
+		NSLog(@"FORCE_ANSI_ALLOCATOR defined, using Ansi allocator.\n");
         AllocatorToUse = EMemoryAllocatorToUse::Ansi;
         return;
     }
@@ -308,12 +308,12 @@ void FApplePlatformMemory::SetAllocatorToUse()
 	{
 		if (!CanOverallocateVirtualMemory())
 		{
-			FPlatformMisc::LowLevelOutputDebugString(TEXT("MallocBinned3 requested but com.apple.developer.kernel.extended-virtual-addressing entitlement not found. Check your entitlements. Falling back to Ansi.\n"));
+			NSLog(@"MallocBinned3 requested but com.apple.developer.kernel.extended-virtual-addressing entitlement not found. Check your entitlements. Falling back to Ansi.\n");
 			AllocatorToUse = EMemoryAllocatorToUse::Ansi;
 			return;
 		}
 
-		FPlatformMisc::LowLevelOutputDebugString(TEXT("Using MallocBinned3 allocator.\n"));
+		NSLog(@"Using MallocBinned3 allocator.\n");
 		AllocatorToUse = EMemoryAllocatorToUse::Binned3;
 		return;
 	}
@@ -321,18 +321,18 @@ void FApplePlatformMemory::SetAllocatorToUse()
     {
         if(!CanOverallocateVirtualMemory())
         {
-			FPlatformMisc::LowLevelOutputDebugString(TEXT("MallocBinned2 requested but com.apple.developer.kernel.extended-virtual-addressing entitlement not found. Check your entitlements. Falling back to Ansi.\n"));
+			NSLog(@"MallocBinned2 requested but com.apple.developer.kernel.extended-virtual-addressing entitlement not found. Check your entitlements. Falling back to Ansi.\n");
             AllocatorToUse = EMemoryAllocatorToUse::Ansi;
             return;
         }
 
-		FPlatformMisc::LowLevelOutputDebugString(TEXT("Using MallocBinned2 allocator.\n"));
+		NSLog(@"Using MallocBinned2 allocator.\n");
         AllocatorToUse = EMemoryAllocatorToUse::Binned2;
         return;
     }
     else
     {
-		FPlatformMisc::LowLevelOutputDebugString(TEXT("Defaulting to Ansi allocator.\n"));
+		NSLog(@"Defaulting to Ansi allocator.\n");
         AllocatorToUse = EMemoryAllocatorToUse::Ansi;
         return;
     }
