#UE Insert LevelSequenceUpgradeDynamicBindings_NoOp version to make sure content that was saved after LevelSequenceUpgradeDynamicBindings was placed in the wrong location maintains its version

diff --git a/Engine/Source/Runtime/Core/Public/UObject/FortniteMainBranchObjectVersion.h b/Engine/Source/Runtime/Core/Public/UObject/FortniteMainBranchObjectVersion.h
--- a/Engine/Source/Runtime/Core/Public/UObject/FortniteMainBranchObjectVersion.h
+++ b/Engine/Source/Runtime/Core/Public/UObject/FortniteMainBranchObjectVersion.h
@@ -491,6 +491,9 @@ struct FFortniteMainBranchObjectVersion
 
 		// Fix chaos cloth buckling stiffness parameter bug 
 		ChaosClothAllowZeroBucklingStiffness,
+		
+		// LevelSequenceUpgradeDynamicBindings was removed but was intended for this position. Putting this here to make sure versioning of subsequent assets remains the same
+		LevelSequenceUpgradeDynamicBindings_NoOp,
 
 		// AddToFrontend GFA now defaults to unload plugin on exit frontend
 		GameFeatureDataActionAddToFrontendDefaultToUnload,
