Prune log info when passing test results to test report when test succeed to reduce foot print

The optimization is enabled by default but can be disabled through project config:
[/Script/AutomationController.AutomationControllerSettings]
bPruneLogsOnSuccess=false

Sorting tests by failure type in json output is now disabled by default and can be enabled using this config:
bSortTestsByFailure=true
It resides in the same AutomationControllerSettings section.

diff --git a/Engine/Source/Developer/AutomationController/Private/AutomationControllerManager.cpp b/Engine/Source/Developer/AutomationController/Private/AutomationControllerManager.cpp
--- a/Engine/Source/Developer/AutomationController/Private/AutomationControllerManager.cpp
+++ b/Engine/Source/Developer/AutomationController/Private/AutomationControllerManager.cpp
@@ -167,6 +167,8 @@ FAutomationControllerManager::FAutomationControllerManager()
 	}
 
 	bKeepPIEOpen = Settings->bKeepPIEOpen;
+	bSortTestsByFailure = Settings->bSortTestsByFailure;
+	bPruneLogsOnSuccess = Settings->bPruneLogsOnSuccess;
 	
 	FString DeveloperPath;
 	FParse::Value(FCommandLine::Get(), TEXT("ReportOutputPath="), ReportExportPath, false);
@@ -823,7 +825,7 @@ void FAutomationControllerManager::ExecuteNextTask( int32 ClusterIndex, OUT bool
 							UE_LOG(LogAutomationController, Log, TEXT("Sending RunTest %s to %s"), *NextTest->GetDisplayName(), *DeviceAddress.ToString());
 
 							SendMessage(
-								FMessageEndpoint::MakeMessage<FAutomationWorkerRunTests>(ExecutionCount, AddressIndex, NextTest->GetCommand(), NextTest->GetDisplayName(), NextTest->GetFullTestPath(), bSendAnalytics),
+								FMessageEndpoint::MakeMessage<FAutomationWorkerRunTests>(ExecutionCount, AddressIndex, NextTest->GetCommand(), NextTest->GetDisplayName(), NextTest->GetFullTestPath(), bSendAnalytics, bPruneLogsOnSuccess),
 								FAutomationWorkerRunTests::StaticStruct(),
 								DeviceAddress);
 
@@ -998,7 +1000,7 @@ void FAutomationControllerManager::ProcessResults()
 
 		FAutomatedTestPassResults SerializedPassResults = JsonTestPassResults;
 
-		if (FParse::Param(FCommandLine::Get(), TEXT("SortTestsByFailure")))
+		if (bSortTestsByFailure)
 		{
 			// Sort result by failure to improve readability (disabled by default)
 			SerializedPassResults.Tests.StableSort([](const FAutomatedTestResult& A, const FAutomatedTestResult& B) {
diff --git a/Engine/Source/Developer/AutomationController/Private/AutomationControllerManager.h b/Engine/Source/Developer/AutomationController/Private/AutomationControllerManager.h
--- a/Engine/Source/Developer/AutomationController/Private/AutomationControllerManager.h
+++ b/Engine/Source/Developer/AutomationController/Private/AutomationControllerManager.h
@@ -638,6 +638,12 @@ private:
 	/** If we should send keep the PIE open when test pass end */
 	bool bKeepPIEOpen = false;
 
+	/** If we should sort tests by failure type in json/html report */
+	bool bSortTestsByFailure = false;
+
+	/** If we should prune log events from test report on success */
+	bool bPruneLogsOnSuccess = true;
+
 	/** The list of results generated by our test pass. */
 	FAutomatedTestPassResults JsonTestPassResults;
 
diff --git a/Engine/Source/Developer/AutomationController/Private/AutomationControllerSettings.cpp b/Engine/Source/Developer/AutomationController/Private/AutomationControllerSettings.cpp
--- a/Engine/Source/Developer/AutomationController/Private/AutomationControllerSettings.cpp
+++ b/Engine/Source/Developer/AutomationController/Private/AutomationControllerSettings.cpp
@@ -14,6 +14,8 @@ UAutomationControllerSettings::UAutomationControllerSettings(const FObjectInitia
 	, bElevateLogWarningsToErrors(true)
 	, bKeepPIEOpen(false)
 	, bAutoExpandSingleItemSubgroups(true)
+	, bSortTestsByFailure(false)
+	, bPruneLogsOnSuccess(true)
 	, bTreatLogWarningsAsTestErrors(true)
 	, GameInstanceLostTimerSeconds(300.0f)
 	, bResetTelemetryStorageOnNewSession(false)
diff --git a/Engine/Source/Developer/AutomationController/Public/AutomationControllerSettings.h b/Engine/Source/Developer/AutomationController/Public/AutomationControllerSettings.h
--- a/Engine/Source/Developer/AutomationController/Public/AutomationControllerSettings.h
+++ b/Engine/Source/Developer/AutomationController/Public/AutomationControllerSettings.h
@@ -169,6 +169,14 @@ public:
 	UPROPERTY(Config)
 	bool bAutoExpandSingleItemSubgroups;
 
+	/** Whether to Sort tests by failure type in json/html report */
+	UPROPERTY()
+	bool bSortTestsByFailure;
+
+	/** Whether to prune log events from test report on success */
+	UPROPERTY()
+	bool bPruneLogsOnSuccess;
+
 private:
 	/** Whether to treat log warnings as test errors (default=true) */
 	UPROPERTY(Config, Meta = (DeprecatedProperty, DeprecationMessage = "Use bElevateLogWarningsToErrors instead."))
diff --git a/Engine/Source/Programs/AutomationTool/Gauntlet/Framework/Utils/Gauntlet.HordeReport.cs b/Engine/Source/Programs/AutomationTool/Gauntlet/Framework/Utils/Gauntlet.HordeReport.cs
--- a/Engine/Source/Programs/AutomationTool/Gauntlet/Framework/Utils/Gauntlet.HordeReport.cs
+++ b/Engine/Source/Programs/AutomationTool/Gauntlet/Framework/Utils/Gauntlet.HordeReport.cs
@@ -1196,6 +1196,11 @@ namespace Gauntlet
 		}
 		private static JsonSerializerOptions GetDefaultJsonOptions()
 		{
+			if (CommandUtils.IsBuildMachine)
+			{
+				return new JsonSerializerOptions { };
+			}
+			// Pretty print when running locally
 			return new JsonSerializerOptions
 			{
 				WriteIndented = true
diff --git a/Engine/Source/Runtime/AutomationMessages/Public/AutomationWorkerMessages.h b/Engine/Source/Runtime/AutomationMessages/Public/AutomationWorkerMessages.h
--- a/Engine/Source/Runtime/AutomationMessages/Public/AutomationWorkerMessages.h
+++ b/Engine/Source/Runtime/AutomationMessages/Public/AutomationWorkerMessages.h
@@ -319,17 +319,22 @@ struct FAutomationWorkerRunTests : public FAutomationWorkerMessageBase
 	UPROPERTY()
 	bool bSendAnalytics;
 
+	/** If true, prune log events from test report on success */
+	UPROPERTY()
+	bool bPruneLogsOnSuccess;
+
 	/** Default constructor. */
 	FAutomationWorkerRunTests( ) :ExecutionCount(0), RoleIndex(0), bSendAnalytics(false) { }
 
 	/** Creates and initializes a new instance. */
-	FAutomationWorkerRunTests( uint32 InExecutionCount, int32 InRoleIndex, FString InTestName, FString InBeautifiedTestName, FString InFullTestPath, bool InSendAnalytics)
+	FAutomationWorkerRunTests( uint32 InExecutionCount, int32 InRoleIndex, FString InTestName, FString InBeautifiedTestName, FString InFullTestPath, bool InSendAnalytics, bool InPruneLogsOnSuccess)
 		: ExecutionCount(InExecutionCount)
 		, RoleIndex(InRoleIndex)
 		, TestName(InTestName)
 		, BeautifiedTestName(InBeautifiedTestName)
 		, FullTestPath(InFullTestPath)
 		, bSendAnalytics(InSendAnalytics)
+		, bPruneLogsOnSuccess(InPruneLogsOnSuccess)
 	{ }
 };
 
diff --git a/Engine/Source/Runtime/AutomationWorker/Private/AutomationWorkerModule.cpp b/Engine/Source/Runtime/AutomationWorker/Private/AutomationWorkerModule.cpp
--- a/Engine/Source/Runtime/AutomationWorker/Private/AutomationWorkerModule.cpp
+++ b/Engine/Source/Runtime/AutomationWorker/Private/AutomationWorkerModule.cpp
@@ -208,7 +208,9 @@ void FAutomationWorkerModule::ReportTestComplete()
 			Message->ExecutionCount = ExecutionCount;
 			Message->State = bSuccess ? EAutomationState::Success : EAutomationState::Fail;
 			Message->Duration = ExecutionInfo.Duration;
-			Message->Entries = ExecutionInfo.GetEntries();
+			// Prune log entries if it is a success to reduce foot print.
+			Message->Entries = bPruneLogsOnSuccess && bSuccess ?
+				ExecutionInfo.GetEntries().FilterByPredicate([](const FAutomationExecutionEntry& Entry) { return Entry.Event.Context != "log"; }) : ExecutionInfo.GetEntries();
 			Message->WarningTotal = ExecutionInfo.GetWarningTotal();
 			Message->ErrorTotal = ExecutionInfo.GetErrorTotal();
 
@@ -681,6 +683,7 @@ void FAutomationWorkerModule::HandleRunTestsMessage( const FAutomationWorkerRunT
 	BeautifiedTestName = Message.BeautifiedTestName;
 	FullTestPath = Message.FullTestPath;
 	bSendAnalytics = Message.bSendAnalytics;
+	bPruneLogsOnSuccess = Message.bPruneLogsOnSuccess;
 	TestRequesterAddress = Context->GetSender();
 	ActiveScreenshotComparisonId = FGuid();
 
diff --git a/Engine/Source/Runtime/AutomationWorker/Private/AutomationWorkerModule.h b/Engine/Source/Runtime/AutomationWorker/Private/AutomationWorkerModule.h
--- a/Engine/Source/Runtime/AutomationWorker/Private/AutomationWorkerModule.h
+++ b/Engine/Source/Runtime/AutomationWorker/Private/AutomationWorkerModule.h
@@ -180,6 +180,9 @@ private:
 	/** Whether to send analytics events to the backend - sent from controller */
 	bool bSendAnalytics;
 
+	/** Whether to prune log events from test report on success */
+	bool bPruneLogsOnSuccess;
+
 	/** Whether the controller has requested that the network command should execute */
 	bool bExecuteNextNetworkCommand;
 
diff --git a/Engine/Source/Runtime/Core/Private/Misc/AutomationTest.cpp b/Engine/Source/Runtime/Core/Private/Misc/AutomationTest.cpp
--- a/Engine/Source/Runtime/Core/Private/Misc/AutomationTest.cpp
+++ b/Engine/Source/Runtime/Core/Private/Misc/AutomationTest.cpp
@@ -199,23 +199,22 @@ void FAutomationTestFramework::FAutomationTestOutputDevice::Serialize( const TCH
 		if (CaptureLog)
 		{
 			ELogVerbosity::Type EffectiveVerbosity = AutomationTest::GetAutomationLogLevel(Verbosity, Category, LocalCurTest);
-
-			FString FormattedMsg = FString::Printf(TEXT("%s: %s [log]"), *Category.ToString(), V);
-			
-			// Errors
-			if (EffectiveVerbosity == ELogVerbosity::Error)
-			{
-				LocalCurTest->AddError(FormattedMsg, STACK_OFFSET);
-			}
-			// Warnings
-			else if (EffectiveVerbosity == ELogVerbosity::Warning)
-			{
-				LocalCurTest->AddWarning(FormattedMsg, STACK_OFFSET);
-			}
-			// Display
-			else if (EffectiveVerbosity != ELogVerbosity::NoLogging)
+			if (EffectiveVerbosity != ELogVerbosity::NoLogging)
 			{
-				LocalCurTest->AddInfo(FormattedMsg, STACK_OFFSET);
+				FString FormattedMsg = FString::Printf(TEXT("%s: %s"), *Category.ToString(), V);
+
+				FAutomationEvent Event(EAutomationEventType::Info, FormattedMsg, TEXT("log"));
+				// Errors
+				if (EffectiveVerbosity == ELogVerbosity::Error)
+				{
+					Event.Type = EAutomationEventType::Error;
+				}
+				// Warnings
+				else if (EffectiveVerbosity == ELogVerbosity::Warning)
+				{
+					Event.Type = EAutomationEventType::Warning;
+				}
+				LocalCurTest->AddEvent(Event, STACK_OFFSET);
 			}
 		}
 		// Log...etc
@@ -1417,8 +1416,21 @@ void FAutomationTestBase::SetTelemetryStorage(const FString& StorageName)
 
 void FAutomationTestBase::AddEvent(const FAutomationEvent& InEvent, int32 StackOffset, bool bCaptureStack)
 {
-	FWriteScopeLock Lock(ActionCS);
-	ExecutionInfo.AddEvent(InEvent, StackOffset + 1, bCaptureStack);
+	ELogVerbosity::Type LogType = ELogVerbosity::Display;
+	if (InEvent.Type == EAutomationEventType::Error)
+	{
+		LogType = ELogVerbosity::Error;
+	}
+	else if (InEvent.Type == EAutomationEventType::Warning)
+	{
+		LogType = ELogVerbosity::Warning;
+	}
+
+	if (!IsExpectedMessage(InEvent.Message, LogType))
+	{
+		FWriteScopeLock Lock(ActionCS);
+		ExecutionInfo.AddEvent(InEvent, StackOffset + 1, bCaptureStack);
+	}
 }
 
 bool FAutomationTestBase::HasAnyErrors() const
