Fix the ASan bug in AutoRTFM when under transactional conditions by not calling Malloc if the Realloc count is 0.

diff --git a/Engine/Source/Runtime/Core/Public/HAL/FMemory.inl b/Engine/Source/Runtime/Core/Public/HAL/FMemory.inl
--- a/Engine/Source/Runtime/Core/Public/HAL/FMemory.inl
+++ b/Engine/Source/Runtime/Core/Public/HAL/FMemory.inl
@@ -55,7 +55,7 @@ FMEMORY_INLINE_FUNCTION_DECORATOR void* FMemory::Malloc(SIZE_T Count, uint32 Ali
 
 FMEMORY_INLINE_FUNCTION_DECORATOR void* FMemory::Realloc(void* Original, SIZE_T Count, uint32 Alignment)
 {
-	if(AutoRTFM::IsClosed())
+	if (AutoRTFM::IsClosed())
 	{
 		// AutoRTFM: For transactional code, we have to do a little dance to handle Realloc
 		// properly. We turn realloc into Malloc + Memcpy + Free and never call into the
@@ -71,17 +71,33 @@ FMEMORY_INLINE_FUNCTION_DECORATOR void* FMemory::Realloc(void* Original, SIZE_T
 		// if we new that the Original pointer was allocated in this transaction, we could
 		// call into the underlying realloc - however we would also have to account for the
 		// malloc deferring a call to free, so we would also have to erase that call to free.
-		void* Ptr = Malloc(Count, Alignment);
-		if (!Ptr)
+
+		void* Ptr = nullptr;
+
+		// Depending on the underlying implementation `Malloc` here, even if `Count` is zero,
+		// could do an actual allocation (it is implementation-defined what occurs). So
+		// instead, since we are fine to return null with a `Count` of zero, we check for
+		// that case and skip the `Malloc` call entirely.
+		if (Count > 0)
 		{
-			return nullptr;
+			Ptr = Malloc(Count, Alignment);
+
+			if (!Ptr)
+			{
+				return nullptr;
+			}
 		}
 
 		if (Original)
 		{
-			SIZE_T OriginalCount = GetAllocSize(Original);
-			SIZE_T CopyCount = FGenericPlatformMath::Min(Count, OriginalCount); // handle the case where the new size is smaller
-			Memcpy(Ptr, Original, CopyCount);
+			if (Ptr)
+			{
+				SIZE_T OriginalCount = GetAllocSize(Original);
+				SIZE_T CopyCount = FGenericPlatformMath::Min(Count, OriginalCount); // handle the case where the new size is smaller
+
+				Memcpy(Ptr, Original, CopyCount);
+			}
+			
 			Free(Original);
 		}
 
