Add substring search (Contains) to FTextComparison

#tests Used locally to fix non-english searches requiring specific case. Compiled legacy text individually.

diff --git a/Engine/Source/Runtime/Core/Private/Internationalization/ICUText.cpp b/Engine/Source/Runtime/Core/Private/Internationalization/ICUText.cpp
--- a/Engine/Source/Runtime/Core/Private/Internationalization/ICUText.cpp
+++ b/Engine/Source/Runtime/Core/Private/Internationalization/ICUText.cpp
@@ -19,6 +19,7 @@ THIRD_PARTY_INCLUDES_START
 	#include <unicode/unistr.h>
 	#include <unicode/coll.h>
 	#include <unicode/sortkey.h>
+	#include <unicode/usearch.h>
 	#include <unicode/ubidi.h>
 THIRD_PARTY_INCLUDES_END
 
@@ -140,6 +141,32 @@ bool FTextComparison::EqualToCaseIgnored( const FString& A, const FString& B )
 	return EqualTo(A, B, ETextComparisonLevel::Secondary);
 }
 
+bool FTextComparison::Contains(const FString& TextToSearch, const FString& TextToFind, const ETextComparisonLevel::Type ComparisonLevel)
+{
+	const TSharedRef<const icu::Collator, ESPMode::ThreadSafe> Collator( FInternationalization::Get().GetCurrentLanguage()->Implementation->GetCollator(ComparisonLevel) );
+
+	icu::UnicodeString UnicodeTextToSearch, UnicodeTextToFind;
+	ICUUtilities::ConvertString(TextToSearch, UnicodeTextToSearch);
+	ICUUtilities::ConvertString(TextToFind, UnicodeTextToFind);
+
+	bool bResult = false;
+
+	UErrorCode error = U_ZERO_ERROR;
+	UStringSearch* StringSearch = usearch_openFromCollator(UnicodeTextToFind.getBuffer(), UnicodeTextToFind.length(), UnicodeTextToSearch.getBuffer(), UnicodeTextToSearch.length(), Collator->toUCollator(), nullptr, &error);
+	if (StringSearch)
+	{
+		bResult = usearch_first(StringSearch, &error) != USEARCH_DONE;
+		usearch_close(StringSearch);
+	}
+
+	return bResult;
+}
+
+bool FTextComparison::ContainsCaseIgnored(const FString& TextToSearch, const FString& TextToFind)
+{
+	return Contains(TextToSearch, TextToFind, ETextComparisonLevel::Secondary);
+}
+
 class FText::FSortPredicate::FSortPredicateImplementation
 {
 public:
diff --git a/Engine/Source/Runtime/Core/Private/Internationalization/LegacyText.cpp b/Engine/Source/Runtime/Core/Private/Internationalization/LegacyText.cpp
--- a/Engine/Source/Runtime/Core/Private/Internationalization/LegacyText.cpp
+++ b/Engine/Source/Runtime/Core/Private/Internationalization/LegacyText.cpp
@@ -70,6 +70,16 @@ bool FTextComparison::EqualToCaseIgnored( const FString& A, const FString& B )
 	return A.Equals(B, ESearchCase::IgnoreCase);
 }
 
+bool FTextComparison::Contains(const FString& TextToSearch, const FString& TextToFind, const ETextComparisonLevel::Type ComparisonLevel)
+{
+	return TextToSearch.Contains(TextToFind, ESearchCase::CaseSensitive);
+}
+
+bool FTextComparison::ContainsCaseIgnored(const FString& TextToSearch, const FString& TextToFind)
+{
+	return TextToSearch.Contains(TextToFind, ESearchCase::IgnoreCase);
+}
+
 FText::FSortPredicate::FSortPredicate(const ETextComparisonLevel::Type ComparisonLevel)
 {
 
diff --git a/Engine/Source/Runtime/Core/Public/Internationalization/TextComparison.h b/Engine/Source/Runtime/Core/Public/Internationalization/TextComparison.h
--- a/Engine/Source/Runtime/Core/Public/Internationalization/TextComparison.h
+++ b/Engine/Source/Runtime/Core/Public/Internationalization/TextComparison.h
@@ -30,4 +30,7 @@ public:
 
 	static CORE_API bool EqualTo(const FString& A, const FString& B, const ETextComparisonLevel::Type ComparisonLevel = ETextComparisonLevel::Default);
 	static CORE_API bool EqualToCaseIgnored(const FString& A, const FString& B);
+
+	static CORE_API bool Contains(const FString& TextToSearch, const FString& TextToFind, const ETextComparisonLevel::Type ComparisonLevel = ETextComparisonLevel::Default);
+	static CORE_API bool ContainsCaseIgnored(const FString& TextToSearch, const FString& TextToFind);
 };
