Corrected a comment on r.Mobile.Forward.RenderRectLightsAsSpotLights cvar default value

diff --git a/Engine/Source/Runtime/Core/Private/HAL/ConsoleManager.cpp b/Engine/Source/Runtime/Core/Private/HAL/ConsoleManager.cpp
--- a/Engine/Source/Runtime/Core/Private/HAL/ConsoleManager.cpp
+++ b/Engine/Source/Runtime/Core/Private/HAL/ConsoleManager.cpp
@@ -3775,8 +3775,8 @@ static TAutoConsoleVariable<int32> CVarMobileForwardEnableLocalLights(
 static TAutoConsoleVariable<int32> CVarMobileRenderRectLightsAsSpotLights(
 	TEXT("r.Mobile.Forward.RenderRectLightsAsSpotLights"),
 	0,
-	TEXT("0: Rect Lights are rendered as area lights\n"
-		"1: Rect Lights are rendered as Spot Lights (default)\n"),
+	TEXT("0: Rect Lights are rendered as area lights (default)\n"
+		"1: Rect Lights are rendered as Spot Lights\n"),
 	ECVF_ReadOnly | ECVF_MobileShaderChange);
 
 static TAutoConsoleVariable<int32> CVarMobileForwardEnableParticleLights(
