[Insights] Extended the trace counter API (COUNTERSTRACE_ENABLED):
  * Added bUnchecked template parameter to FCountersTrace::TCounter class.  When enabled, the values are traced even if it does not change.
  * Added UNCHECKED_* variants for all counter macros.
  * Added macros for inline atomic counters (TRACE_ATOMIC_INT_VALUE, TRACE_ATOMIC_MEMORY_VALUE).
  * Added TRACE_DECLARE_COUNTER(CounterType, CounterName, CounterDisplayName, CounterDisplayHint) API.
  * Added support for atomic float counters (FCounterAtomicFloat, TRACE_ATOMIC_FLOAT_VALUE, TRACE_DECLARE_ATOMIC_FLOAT_COUNTER, etc.).

diff --git a/Engine/Source/Runtime/Core/Public/ProfilingDebugging/CountersTrace.h b/Engine/Source/Runtime/Core/Public/ProfilingDebugging/CountersTrace.h
--- a/Engine/Source/Runtime/Core/Public/ProfilingDebugging/CountersTrace.h
+++ b/Engine/Source/Runtime/Core/Public/ProfilingDebugging/CountersTrace.h
@@ -53,7 +53,7 @@ struct FCountersTrace
 	CORE_API static const TCHAR* AllocAndCopyCounterName(const TCHAR* InCounterName);
 	CORE_API static void FreeCounterName(const TCHAR* InCounterName);
 
-	template<typename ValueType, ETraceCounterType CounterType, typename StoredType = ValueType>
+	template<typename ValueType, ETraceCounterType CounterType, typename StoredType = ValueType, bool bUnchecked = false>
 	class TCounter
 	{
 	public:
@@ -110,7 +110,7 @@ struct FCountersTrace
 
 		void Set(ValueType InValue)
 		{
-			if (Value != InValue)
+			if (bUnchecked || Value != InValue)
 			{
 				Value = InValue;
 				if (UE_TRACE_CHANNELEXPR_IS_ENABLED(CountersChannel))
@@ -123,7 +123,7 @@ struct FCountersTrace
 
 		void Add(ValueType InValue)
 		{
-			if (InValue != 0)
+			if (bUnchecked || InValue != 0)
 			{
 				Value += InValue;
 				if (UE_TRACE_CHANNELEXPR_IS_ENABLED(CountersChannel))
@@ -136,7 +136,7 @@ struct FCountersTrace
 
 		void Subtract(ValueType InValue)
 		{
-			if (InValue != 0)
+			if (bUnchecked || InValue != 0)
 			{
 				Value -= InValue;
 				if (UE_TRACE_CHANNELEXPR_IS_ENABLED(CountersChannel))
@@ -177,8 +177,17 @@ struct FCountersTrace
 	using FCounterInt = TCounter<int64, TraceCounterType_Int>;
 	using FCounterAtomicInt = TCounter<int64, TraceCounterType_Int, std::atomic<int64>>;
 	using FCounterFloat = TCounter<double, TraceCounterType_Float>;
+	using FCounterAtomicFloat = TCounter<double, TraceCounterType_Float, std::atomic<double>>;
+
+	using FCounterUncheckedInt = TCounter<int64, TraceCounterType_Int, int64, true>;
+	using FCounterUncheckedAtomicInt = TCounter<int64, TraceCounterType_Int, std::atomic<int64>, true>;
+	using FCounterUncheckedFloat = TCounter<double, TraceCounterType_Float, double, true>;
+	using FCounterUncheckedAtomicFloat = TCounter<double, TraceCounterType_Float, std::atomic<double>, true>;
 };
 
+//////////////////////////////////////////////////
+// Inline Counters
+
 #define __TRACE_CHECK_COUNTER_NAME(CounterDisplayName) \
 	static_assert(std::is_const_v<std::remove_reference_t<decltype(CounterDisplayName)>>, "CounterDisplayName string must be a const TCHAR array."); \
 	static_assert(TIsArrayOrRefOfTypeByPredicate<decltype(CounterDisplayName), TIsCharEncodingCompatibleWithTCHAR>::Value, "CounterDisplayName string must be a TCHAR array.");
@@ -187,53 +196,123 @@ struct FCountersTrace
 	__TRACE_CHECK_COUNTER_NAME(CounterDisplayName) \
 	static FCountersTrace::CounterType PREPROCESSOR_JOIN(__TraceCounter, __LINE__)(CounterDisplayName, CounterDisplayHint);
 
-#define TRACE_INT_VALUE(CounterDisplayName, Value) \
-	__TRACE_DECLARE_INLINE_COUNTER(CounterDisplayName, FCounterInt, TraceCounterDisplayHint_None) \
+#define __TRACE_INLINE_COUNTER_SET(CounterDisplayName, Value, CounterType, CounterDisplayHint) \
+	__TRACE_DECLARE_INLINE_COUNTER(CounterDisplayName, CounterType, CounterDisplayHint) \
 	PREPROCESSOR_JOIN(__TraceCounter, __LINE__).Set(Value);
 
-#define TRACE_FLOAT_VALUE(CounterDisplayName, Value) \
-	__TRACE_DECLARE_INLINE_COUNTER(CounterDisplayName, FCounterFloat, TraceCounterDisplayHint_None) \
-	PREPROCESSOR_JOIN(__TraceCounter, __LINE__).Set(Value);
+#define TRACE_INT_VALUE(CounterDisplayName, Value)                      __TRACE_INLINE_COUNTER_SET(CounterDisplayName, Value, FCounterInt, TraceCounterDisplayHint_None)
+#define TRACE_ATOMIC_INT_VALUE(CounterDisplayName, Value)               __TRACE_INLINE_COUNTER_SET(CounterDisplayName, Value, FCounterAtomicInt, TraceCounterDisplayHint_None)
+#define TRACE_FLOAT_VALUE(CounterDisplayName, Value)                    __TRACE_INLINE_COUNTER_SET(CounterDisplayName, Value, FCounterFloat, TraceCounterDisplayHint_None)
+#define TRACE_ATOMIC_FLOAT_VALUE(CounterDisplayName, Value)             __TRACE_INLINE_COUNTER_SET(CounterDisplayName, Value, FCounterAtomicFloat, TraceCounterDisplayHint_None)
+#define TRACE_MEMORY_VALUE(CounterDisplayName, Value)                   __TRACE_INLINE_COUNTER_SET(CounterDisplayName, Value, FCounterInt, TraceCounterDisplayHint_Memory)
+#define TRACE_ATOMIC_MEMORY_VALUE(CounterDisplayName, Value)            __TRACE_INLINE_COUNTER_SET(CounterDisplayName, Value, FCounterAtomicInt, TraceCounterDisplayHint_Memory)
 
-#define TRACE_MEMORY_VALUE(CounterDisplayName, Value) \
-	__TRACE_DECLARE_INLINE_COUNTER(CounterDisplayName, FCounterInt, TraceCounterDisplayHint_Memory) \
-	PREPROCESSOR_JOIN(__TraceCounter, __LINE__).Set(Value);
+#define TRACE_UNCHECKED_INT_VALUE(CounterDisplayName, Value)            __TRACE_INLINE_COUNTER_SET(CounterDisplayName, Value, FCounterUncheckedInt, TraceCounterDisplayHint_None)
+#define TRACE_UNCHECKED_ATOMIC_INT_VALUE(CounterDisplayName, Value)     __TRACE_INLINE_COUNTER_SET(CounterDisplayName, Value, FCounterUncheckedAtomicInt, TraceCounterDisplayHint_None)
+#define TRACE_UNCHECKED_FLOAT_VALUE(CounterDisplayName, Value)          __TRACE_INLINE_COUNTER_SET(CounterDisplayName, Value, FCounterUncheckedFloat, TraceCounterDisplayHint_None)
+#define TRACE_UNCHECKED_ATOMIC_FLOAT_VALUE(CounterDisplayName, Value)   __TRACE_INLINE_COUNTER_SET(CounterDisplayName, Value, FCounterUncheckedAtomicFloat, TraceCounterDisplayHint_None)
+#define TRACE_UNCHECKED_MEMORY_VALUE(CounterDisplayName, Value)         __TRACE_INLINE_COUNTER_SET(CounterDisplayName, Value, FCounterUncheckedInt, TraceCounterDisplayHint_Memory)
+#define TRACE_UNCHECKED_ATOMIC_MEMORY_VALUE(CounterDisplayName, Value)  __TRACE_INLINE_COUNTER_SET(CounterDisplayName, Value, FCounterUncheckedAtomicInt, TraceCounterDisplayHint_Memory)
 
-#define TRACE_DECLARE_INT_COUNTER(CounterName, CounterDisplayName) \
-	__TRACE_CHECK_COUNTER_NAME(CounterDisplayName) \
-	FCountersTrace::FCounterInt PREPROCESSOR_JOIN(__GTraceCounter, CounterName)(CounterDisplayName, TraceCounterDisplayHint_None);
+//////////////////////////////////////////////////
 
-#define TRACE_DECLARE_ATOMIC_INT_COUNTER(CounterName, CounterDisplayName) \
+#define TRACE_DECLARE_COUNTER(CounterType, CounterName, CounterDisplayName, CounterDisplayHint) \
 	__TRACE_CHECK_COUNTER_NAME(CounterDisplayName) \
-	FCountersTrace::FCounterAtomicInt PREPROCESSOR_JOIN(__GTraceCounter, CounterName)(CounterDisplayName, TraceCounterDisplayHint_None);
+	FCountersTrace::CounterType PREPROCESSOR_JOIN(__GTraceCounter, CounterName)(CounterDisplayName, CounterDisplayHint);
+
+//////////////////////////////////////////////////
+// Declare Int Counters
+
+#define TRACE_DECLARE_INT_COUNTER(CounterName, CounterDisplayName) \
+	TRACE_DECLARE_COUNTER(FCounterInt, CounterName, CounterDisplayName, TraceCounterDisplayHint_None)
+
+#define TRACE_DECLARE_UNCHECKED_INT_COUNTER(CounterName, CounterDisplayName) \
+	TRACE_DECLARE_COUNTER(FCounterUncheckedInt, CounterName, CounterDisplayName, TraceCounterDisplayHint_None)
 
 #define TRACE_DECLARE_INT_COUNTER_EXTERN(CounterName) \
 	extern FCountersTrace::FCounterInt PREPROCESSOR_JOIN(__GTraceCounter, CounterName);
 
+#define TRACE_DECLARE_UNCHECKED_INT_COUNTER_EXTERN(CounterName) \
+	extern FCountersTrace::FCounterUncheckedInt PREPROCESSOR_JOIN(__GTraceCounter, CounterName);
+
+//////////////////////////////////////////////////
+// Declare Atomic Int Counters
+
+#define TRACE_DECLARE_ATOMIC_INT_COUNTER(CounterName, CounterDisplayName) \
+	TRACE_DECLARE_COUNTER(FCounterAtomicInt, CounterName, CounterDisplayName, TraceCounterDisplayHint_None)
+
+#define TRACE_DECLARE_UNCHECKED_ATOMIC_INT_COUNTER(CounterName, CounterDisplayName) \
+	TRACE_DECLARE_COUNTER(FCounterUncheckedAtomicInt, CounterName, CounterDisplayName, TraceCounterDisplayHint_None)
+
 #define TRACE_DECLARE_ATOMIC_INT_COUNTER_EXTERN(CounterName) \
 	extern FCountersTrace::FCounterAtomicInt PREPROCESSOR_JOIN(__GTraceCounter, CounterName);
 
+#define TRACE_DECLARE_UNCHECKED_ATOMIC_INT_COUNTER_EXTERN(CounterName) \
+	extern FCountersTrace::FCounterUncheckedAtomicInt PREPROCESSOR_JOIN(__GTraceCounter, CounterName);
+
+//////////////////////////////////////////////////
+// Declare Float Counters
+
 #define TRACE_DECLARE_FLOAT_COUNTER(CounterName, CounterDisplayName) \
-	__TRACE_CHECK_COUNTER_NAME(CounterDisplayName) \
-	FCountersTrace::FCounterFloat PREPROCESSOR_JOIN(__GTraceCounter, CounterName)(CounterDisplayName, TraceCounterDisplayHint_None);
+	TRACE_DECLARE_COUNTER(FCounterFloat, CounterName, CounterDisplayName, TraceCounterDisplayHint_None)
+
+#define TRACE_DECLARE_UNCHECKED_FLOAT_COUNTER(CounterName, CounterDisplayName) \
+	TRACE_DECLARE_COUNTER(FCounterUncheckedFloat, CounterName, CounterDisplayName, TraceCounterDisplayHint_None)
 
 #define TRACE_DECLARE_FLOAT_COUNTER_EXTERN(CounterName) \
 	extern FCountersTrace::FCounterFloat PREPROCESSOR_JOIN(__GTraceCounter, CounterName);
 
+#define TRACE_DECLARE_UNCHECKED_FLOAT_COUNTER_EXTERN(CounterName) \
+	extern FCountersTrace::FCounterUncheckedFloat PREPROCESSOR_JOIN(__GTraceCounter, CounterName);
+
+//////////////////////////////////////////////////
+// Declare Atomic Float Counters
+
+#define TRACE_DECLARE_ATOMIC_FLOAT_COUNTER(CounterName, CounterDisplayName) \
+	TRACE_DECLARE_COUNTER(FCounterAtomicFloat, CounterName, CounterDisplayName, TraceCounterDisplayHint_None)
+
+#define TRACE_DECLARE_UNCHECKED_ATOMIC_FLOAT_COUNTER(CounterName, CounterDisplayName) \
+	TRACE_DECLARE_COUNTER(FCounterUncheckedAtomicFloat, CounterName, CounterDisplayName, TraceCounterDisplayHint_None)
+
+#define TRACE_DECLARE_ATOMIC_FLOAT_COUNTER_EXTERN(CounterName) \
+	extern FCountersTrace::FCounterAtomicFloat PREPROCESSOR_JOIN(__GTraceCounter, CounterName);
+
+#define TRACE_DECLARE_UNCHECKED_ATOMIC_FLOAT_COUNTER_EXTERN(CounterName) \
+	extern FCountersTrace::FCounterUncheckedAtomicFloat PREPROCESSOR_JOIN(__GTraceCounter, CounterName);
+
+//////////////////////////////////////////////////
+// Declare Memory Counters
+
 #define TRACE_DECLARE_MEMORY_COUNTER(CounterName, CounterDisplayName) \
-	__TRACE_CHECK_COUNTER_NAME(CounterDisplayName) \
-	FCountersTrace::FCounterInt PREPROCESSOR_JOIN(__GTraceCounter, CounterName)(CounterDisplayName, TraceCounterDisplayHint_Memory);
+	TRACE_DECLARE_COUNTER(FCounterInt, CounterName, CounterDisplayName, TraceCounterDisplayHint_Memory)
+
+#define TRACE_DECLARE_UNCHECKED_MEMORY_COUNTER(CounterName, CounterDisplayName) \
+	TRACE_DECLARE_COUNTER(FCounterUncheckedInt, CounterName, CounterDisplayName, TraceCounterDisplayHint_Memory)
 
 #define TRACE_DECLARE_MEMORY_COUNTER_EXTERN(CounterName) \
 	TRACE_DECLARE_INT_COUNTER_EXTERN(CounterName)
 
+#define TRACE_DECLARE_UNCHECKED_MEMORY_COUNTER_EXTERN(CounterName) \
+	TRACE_DECLARE_UNCHECKED_INT_COUNTER_EXTERN(CounterName)
+
+//////////////////////////////////////////////////
+// Declare Atomic Memory Counters
+
 #define TRACE_DECLARE_ATOMIC_MEMORY_COUNTER(CounterName, CounterDisplayName) \
-	__TRACE_CHECK_COUNTER_NAME(CounterDisplayName) \
-	FCountersTrace::FCounterAtomicInt PREPROCESSOR_JOIN(__GTraceCounter, CounterName)(CounterDisplayName, TraceCounterDisplayHint_Memory);
+	TRACE_DECLARE_COUNTER(FCounterAtomicInt, CounterName, CounterDisplayName, TraceCounterDisplayHint_Memory)
+
+#define TRACE_DECLARE_UNCHECKED_ATOMIC_MEMORY_COUNTER(CounterName, CounterDisplayName) \
+	TRACE_DECLARE_COUNTER(FCounterUncheckedAtomicInt, CounterName, CounterDisplayName, TraceCounterDisplayHint_Memory)
 
 #define TRACE_DECLARE_ATOMIC_MEMORY_COUNTER_EXTERN(CounterName) \
 	TRACE_DECLARE_ATOMIC_INT_COUNTER_EXTERN(CounterName)
 
+#define TRACE_DECLARE_UNCHECKED_ATOMIC_MEMORY_COUNTER_EXTERN(CounterName) \
+	TRACE_DECLARE_UNCHECKED_ATOMIC_INT_COUNTER_EXTERN(CounterName)
+
+//////////////////////////////////////////////////
+// Counter Operations
+
 #define TRACE_COUNTER_SET(CounterName, Value) \
 	PREPROCESSOR_JOIN(__GTraceCounter, CounterName).Set(Value);
 
@@ -249,25 +328,54 @@ struct FCountersTrace
 #define TRACE_COUNTER_DECREMENT(CounterName) \
 	PREPROCESSOR_JOIN(__GTraceCounter, CounterName).Decrement();
 
-#else
+#else // COUNTERSTRACE_ENABLED
 
 #define TRACE_INT_VALUE(CounterDisplayName, Value)
+#define TRACE_ATOMIC_INT_VALUE(CounterDisplayName, Value)
 #define TRACE_FLOAT_VALUE(CounterDisplayName, Value)
+#define TRACE_ATOMIC_FLOAT_VALUE(CounterDisplayName, Value)
 #define TRACE_MEMORY_VALUE(CounterDisplayName, Value)
+#define TRACE_ATOMIC_MEMORY_VALUE(CounterDisplayName, Value)
+
+#define TRACE_UNCHECKED_INT_VALUE(CounterDisplayName, Value)
+#define TRACE_UNCHECKED_ATOMIC_INT_VALUE(CounterDisplayName, Value)
+#define TRACE_UNCHECKED_FLOAT_VALUE(CounterDisplayName, Value)
+#define TRACE_UNCHECKED_ATOMIC_FLOAT_VALUE(CounterDisplayName, Value)
+#define TRACE_UNCHECKED_MEMORY_VALUE(CounterDisplayName, Value)
+#define TRACE_UNCHECKED_ATOMIC_MEMORY_VALUE(CounterDisplayName, Value)
+
+#define TRACE_DECLARE_COUNTER(CounterType, CounterName, CounterDisplayName, CounterDisplayHint)
+
 #define TRACE_DECLARE_INT_COUNTER(CounterName, CounterDisplayName)
 #define TRACE_DECLARE_INT_COUNTER_EXTERN(CounterName)
 #define TRACE_DECLARE_ATOMIC_INT_COUNTER(CounterName, CounterDisplayName)
 #define TRACE_DECLARE_ATOMIC_INT_COUNTER_EXTERN(CounterName)
 #define TRACE_DECLARE_FLOAT_COUNTER(CounterName, CounterDisplayName)
 #define TRACE_DECLARE_FLOAT_COUNTER_EXTERN(CounterName)
+#define TRACE_DECLARE_ATOMIC_FLOAT_COUNTER(CounterName, CounterDisplayName)
+#define TRACE_DECLARE_ATOMIC_FLOAT_COUNTER_EXTERN(CounterName)
 #define TRACE_DECLARE_MEMORY_COUNTER(CounterName, CounterDisplayName)
 #define TRACE_DECLARE_MEMORY_COUNTER_EXTERN(CounterName)
 #define TRACE_DECLARE_ATOMIC_MEMORY_COUNTER(CounterName, CounterDisplayName)
 #define TRACE_DECLARE_ATOMIC_MEMORY_COUNTER_EXTERN(CounterName)
+
+#define TRACE_DECLARE_UNCHECKED_INT_COUNTER(CounterName, CounterDisplayName)
+#define TRACE_DECLARE_UNCHECKED_INT_COUNTER_EXTERN(CounterName)
+#define TRACE_DECLARE_UNCHECKED_ATOMIC_INT_COUNTER(CounterName, CounterDisplayName)
+#define TRACE_DECLARE_UNCHECKED_ATOMIC_INT_COUNTER_EXTERN(CounterName)
+#define TRACE_DECLARE_UNCHECKED_FLOAT_COUNTER(CounterName, CounterDisplayName)
+#define TRACE_DECLARE_UNCHECKED_FLOAT_COUNTER_EXTERN(CounterName)
+#define TRACE_DECLARE_UNCHECKED_ATOMIC_FLOAT_COUNTER(CounterName, CounterDisplayName)
+#define TRACE_DECLARE_UNCHECKED_ATOMIC_FLOAT_COUNTER_EXTERN(CounterName)
+#define TRACE_DECLARE_UNCHECKED_MEMORY_COUNTER(CounterName, CounterDisplayName)
+#define TRACE_DECLARE_UNCHECKED_MEMORY_COUNTER_EXTERN(CounterName)
+#define TRACE_DECLARE_UNCHECKED_ATOMIC_MEMORY_COUNTER(CounterName, CounterDisplayName)
+#define TRACE_DECLARE_UNCHECKED_ATOMIC_MEMORY_COUNTER_EXTERN(CounterName)
+
 #define TRACE_COUNTER_SET(CounterName, Value)
 #define TRACE_COUNTER_ADD(CounterName, Value)
 #define TRACE_COUNTER_SUBTRACT(CounterName, Value)
 #define TRACE_COUNTER_INCREMENT(CounterName)
 #define TRACE_COUNTER_DECREMENT(CounterName)
 
-#endif
\ No newline at end of file
+#endif // COUNTERSTRACE_ENABLED
