Added MakeConstArrayView.

diff --git a/Engine/Source/Runtime/Core/Public/Containers/ArrayView.h b/Engine/Source/Runtime/Core/Public/Containers/ArrayView.h
--- a/Engine/Source/Runtime/Core/Public/Containers/ArrayView.h
+++ b/Engine/Source/Runtime/Core/Public/Containers/ArrayView.h
@@ -821,15 +821,46 @@ auto MakeArrayView(OtherRangeType&& Other UE_LIFETIMEBOUND)
 }
 
 template<typename ElementType>
-auto MakeArrayView(ElementType* Pointer UE_LIFETIMEBOUND, int32 Size)
+TArrayView<ElementType> MakeArrayView(ElementType* Pointer UE_LIFETIMEBOUND, int32 Size)
 {
 	return TArrayView<ElementType>(Pointer, Size);
 }
 
-template <typename T>
-TArrayView<const T> MakeArrayView(std::initializer_list<T> List UE_LIFETIMEBOUND)
+template <typename ElementType>
+TArrayView<const ElementType> MakeArrayView(std::initializer_list<ElementType> List UE_LIFETIMEBOUND)
+{
+	return TArrayView<const ElementType>(List.begin(), List.size());
+}
+
+template <
+	typename OtherRangeType,
+	typename CVUnqualifiedOtherRangeType = std::remove_cv_t<std::remove_reference_t<OtherRangeType>>
+	UE_REQUIRES(TIsContiguousContainer<CVUnqualifiedOtherRangeType>::Value && TIsTArrayView_V<CVUnqualifiedOtherRangeType>)
+>
+auto MakeConstArrayView(OtherRangeType&& Other)
+{
+	return TArrayView<const std::remove_pointer_t<decltype(GetData(DeclVal<OtherRangeType&>()))>>(Forward<OtherRangeType>(Other));
+}
+template <
+	typename OtherRangeType,
+	typename CVUnqualifiedOtherRangeType = std::remove_cv_t<std::remove_reference_t<OtherRangeType>>
+	UE_REQUIRES(TIsContiguousContainer<CVUnqualifiedOtherRangeType>::Value && !TIsTArrayView_V<CVUnqualifiedOtherRangeType>)
+>
+auto MakeConstArrayView(OtherRangeType&& Other UE_LIFETIMEBOUND)
+{
+	return TArrayView<const std::remove_pointer_t<decltype(GetData(DeclVal<OtherRangeType&>()))>>(Forward<OtherRangeType>(Other));
+}
+
+template<typename ElementType>
+TArrayView<const ElementType> MakeConstArrayView(const ElementType* Pointer UE_LIFETIMEBOUND, int32 Size)
+{
+	return TArrayView<const ElementType>(Pointer, Size);
+}
+
+template <typename ElementType>
+TArrayView<const ElementType> MakeConstArrayView(std::initializer_list<ElementType> List UE_LIFETIMEBOUND)
 {
-	return TArrayView<const T>(List.begin(), List.size());
+	return TArrayView<const ElementType>(List.begin(), List.size());
 }
 
 //////////////////////////////////////////////////////////////////////////
