Add callStackSymbols to log when error/crash happens

diff --git a/Engine/Source/Runtime/Core/Private/IOS/IOSErrorOutputDevice.cpp b/Engine/Source/Runtime/Core/Private/IOS/IOSErrorOutputDevice.cpp
--- a/Engine/Source/Runtime/Core/Private/IOS/IOSErrorOutputDevice.cpp
+++ b/Engine/Source/Runtime/Core/Private/IOS/IOSErrorOutputDevice.cpp
@@ -49,6 +49,29 @@ void FIOSErrorOutputDevice::HandleError()
     
 	// Dump the error and flush the log.
 #if !NO_LOGGING
+	NSArray<NSString *>* CallStackSymbols = [NSThread callStackSymbols];
+	static_assert(sizeof(GErrorHist[0]) == 2 * sizeof(char));
+	int Pos = 0;	// In unit of TCHAR, aka wide char
+	int NumbersOfLinesToSkip = 5; // First 5 lines of callstacks are just error output stuff
+	for (NSString* Line in CallStackSymbols)
+	{
+		if ([CallStackSymbols indexOfObject:Line] < NumbersOfLinesToSkip)
+		{
+			continue;
+		}
+		if (Pos >= sizeof(GErrorHist))
+		{
+			break;
+		}
+		// NSString does not understand wide CString
+		[Line getCString:(char *)(GErrorHist + Pos)
+			   maxLength:(sizeof(GErrorHist) - Pos) * 2
+				encoding:NSUTF16StringEncoding];
+		Pos += [Line lengthOfBytesUsingEncoding:NSUTF16StringEncoding] / 2;
+		// Append L'\n' instead of L'\0'
+		*(GErrorHist + Pos) = L'\n';
+		Pos += 1;
+	}
 	FDebug::LogFormattedMessageWithCallstack(LogIOS.GetCategoryName(), __FILE__, __LINE__, TEXT("=== Critical error: ==="), GErrorHist, ELogVerbosity::Error);
 #endif
 
