Fix the bug I caused where we'd call _aligned_msize incorrectly when using the ANSI allocator.

I had introduced this bug because AutoRTFM absolutely requires being able to query the allocated size of an allocation. To workaround this, we will not use the MSVC _aligned_* allocation methods when using AutoRTFM. I also did try to use the _aligned_offset_* allocator methods, but those have a known bug under ASAN so are effectively useless.

diff --git a/Engine/Source/Runtime/Core/Private/HAL/MallocAnsi.cpp b/Engine/Source/Runtime/Core/Private/HAL/MallocAnsi.cpp
--- a/Engine/Source/Runtime/Core/Private/HAL/MallocAnsi.cpp
+++ b/Engine/Source/Runtime/Core/Private/HAL/MallocAnsi.cpp
@@ -8,6 +8,7 @@
 #include "Misc/AssertionMacros.h"
 #include "Math/UnrealMathUtility.h"
 #include "Templates/AlignmentTemplates.h"
+#include "AutoRTFM/AutoRTFM.h"
 
 #if PLATFORM_USE_ANSI_POSIX_MALLOC
 	#include <malloc.h>
@@ -21,9 +22,11 @@
 	#include "Windows/WindowsHWrapper.h"
 #endif
 
+#define MALLOC_ANSI_USES__ALIGNED_MALLOC !UE_AUTORTFM && PLATFORM_USES__ALIGNED_MALLOC
+
 void* AnsiMalloc(SIZE_T Size, uint32 Alignment)
 {
-#if PLATFORM_USES__ALIGNED_MALLOC
+#if MALLOC_ANSI_USES__ALIGNED_MALLOC
 	void* Result = _aligned_malloc( Size, Alignment );
 #elif PLATFORM_USE_ANSI_POSIX_MALLOC
 	void* Result;
@@ -49,20 +52,20 @@ void* AnsiMalloc(SIZE_T Size, uint32 Alignment)
 
 static SIZE_T AnsiGetAllocationSize(void* Original)
 {
-#if	PLATFORM_USES__ALIGNED_MALLOC
-	return _aligned_msize(Original, 16, 0); // Assumes alignment of 16
+#if MALLOC_ANSI_USES__ALIGNED_MALLOC
+	return _aligned_msize(Original, 16, 0); // TODO: incorrectly assumes alignment of 16
 #elif PLATFORM_USE_ANSI_POSIX_MALLOC || PLATFORM_USE_ANSI_MEMALIGN
 	return malloc_usable_size(Original);
 #else
 	return *((SIZE_T*)((uint8*)Original - sizeof(void*) - sizeof(SIZE_T)));
-#endif // PLATFORM_USES__ALIGNED_MALLOC
+#endif
 }
 
 void* AnsiRealloc(void* Ptr, SIZE_T NewSize, uint32 Alignment)
 {
 	void* Result;
 
-#if PLATFORM_USES__ALIGNED_MALLOC
+#if MALLOC_ANSI_USES__ALIGNED_MALLOC
 	if (Ptr && NewSize)
 	{
 		Result = _aligned_realloc(Ptr, NewSize, Alignment);
@@ -129,7 +132,7 @@ void* AnsiRealloc(void* Ptr, SIZE_T NewSize, uint32 Alignment)
 
 void AnsiFree(void* Ptr)
 {
-#if PLATFORM_USES__ALIGNED_MALLOC
+#if MALLOC_ANSI_USES__ALIGNED_MALLOC
 	_aligned_free(Ptr);
 #elif PLATFORM_USE_ANSI_POSIX_MALLOC || PLATFORM_USE_ANSI_MEMALIGN
 	free(Ptr);
@@ -214,15 +217,19 @@ void FMallocAnsi::Free( void* Ptr )
 	AnsiFree(Ptr);
 }
 
-bool FMallocAnsi::GetAllocationSize( void *Original, SIZE_T &SizeOut )
+bool FMallocAnsi::GetAllocationSize(void* Original, SIZE_T& SizeOut)
 {
 	if (!Original)
 	{
 		return false;
 	}
 
+#if MALLOC_ANSI_USES__ALIGNED_MALLOC
+	return false;
+#else
 	SizeOut = AnsiGetAllocationSize(Original);
 	return true;
+#endif
 }
 
 bool FMallocAnsi::IsInternallyThreadSafe() const
diff --git a/Engine/Source/Runtime/Core/Public/HAL/FMemory.inl b/Engine/Source/Runtime/Core/Public/HAL/FMemory.inl
--- a/Engine/Source/Runtime/Core/Public/HAL/FMemory.inl
+++ b/Engine/Source/Runtime/Core/Public/HAL/FMemory.inl
@@ -177,7 +177,16 @@ FMEMORY_INLINE_FUNCTION_DECORATOR SIZE_T FMemory::GetAllocSize(void* Original)
 		else
 		{
 			SIZE_T Size = 0;
-			Result = FMEMORY_INLINE_GMalloc->GetAllocationSize(Original, Size) ? Size : 0;
+			const bool bGotSize = FMEMORY_INLINE_GMalloc->GetAllocationSize(Original, Size);
+			Result = bGotSize ? Size : 0;
+
+			// This folds away at compile time so that the check is only ever performed inside transactional
+			// code paths. The check is to ensure that the allocator used will return the correct allocation
+			// size, which is a cornerstone requirement for AutoRTFM to function.
+			if (AutoRTFM::IsClosed())
+			{
+				checkf(bGotSize, TEXT("For AutoRTFM to function it must be able to get the size of an allocation"));
+			}
 		}
 	};
 
