Improve iOS crash log

diff --git a/Engine/Source/Runtime/Core/Private/IOS/IOSErrorOutputDevice.cpp b/Engine/Source/Runtime/Core/Private/IOS/IOSErrorOutputDevice.cpp
--- a/Engine/Source/Runtime/Core/Private/IOS/IOSErrorOutputDevice.cpp
+++ b/Engine/Source/Runtime/Core/Private/IOS/IOSErrorOutputDevice.cpp
@@ -10,6 +10,7 @@
 #include "Misc/OutputDeviceHelper.h"
 #include "Misc/OutputDeviceRedirector.h"
 #include "CoreGlobals.h"
+#include <cxxabi.h>
 
 FIOSErrorOutputDevice::FIOSErrorOutputDevice()
 :	ErrorPos(0)
@@ -50,6 +51,7 @@ void FIOSErrorOutputDevice::HandleError()
 	// Dump the error and flush the log.
 #if !NO_LOGGING
 	NSArray<NSString *>* CallStackSymbols = [NSThread callStackSymbols];
+	int32 Status = 0;
 	static_assert(sizeof(GErrorHist[0]) == 2 * sizeof(char));
 	int Pos = 0;	// In unit of TCHAR, aka wide char
 	int NumbersOfLinesToSkip = 5; // First 5 lines of callstacks are just error output stuff
@@ -63,6 +65,20 @@ void FIOSErrorOutputDevice::HandleError()
 		{
 			break;
 		}
+		// Try to demangle the function name
+		NSMutableArray<NSString *>* SubStrings = [NSMutableArray arrayWithArray:[Line componentsSeparatedByString:@" "]];
+		// last part of the string should look like "_ZN21FIOSErrorOutputDevice11HandleErrorEv + 248"
+		const int NameIndex = 3;	// function name is 3rd from last
+		if (SubStrings.count >= NameIndex && [SubStrings[SubStrings.count - NameIndex + 1] isEqualToString:@"+"])
+		{
+			char* DemangledName = abi::__cxa_demangle([SubStrings[SubStrings.count - NameIndex] cStringUsingEncoding:NSUTF8StringEncoding], NULL, NULL, &Status);
+			if (Status >= 0)
+			{
+				SubStrings[SubStrings.count - NameIndex] = [NSString stringWithCString:DemangledName encoding:NSUTF8StringEncoding];
+			}
+			free(DemangledName);
+		}
+		Line = [SubStrings componentsJoinedByString:@" "];
 		// NSString does not understand wide CString
 		[Line getCString:(char *)(GErrorHist + Pos)
 			   maxLength:(sizeof(GErrorHist) - Pos) * 2
@@ -70,8 +86,9 @@ void FIOSErrorOutputDevice::HandleError()
 		Pos += [Line lengthOfBytesUsingEncoding:NSUTF16StringEncoding] / 2;
 		// Append L'\n' instead of L'\0'
 		*(GErrorHist + Pos) = L'\n';
-		Pos += 1;
+		Pos++;
 	}
+	
 	FDebug::LogFormattedMessageWithCallstack(LogIOS.GetCategoryName(), __FILE__, __LINE__, TEXT("=== Critical error: ==="), GErrorHist, ELogVerbosity::Error);
 #endif
 
