trivial breaking up of the SUBOBJECT_DELIMITER_CHAR macro, so the char may be used in UTF8 contexts.

diff --git a/Engine/Source/Runtime/Core/Public/UObject/NameTypes.h b/Engine/Source/Runtime/Core/Public/UObject/NameTypes.h
--- a/Engine/Source/Runtime/Core/Public/UObject/NameTypes.h
+++ b/Engine/Source/Runtime/Core/Public/UObject/NameTypes.h
@@ -158,7 +158,8 @@ typedef FNameEntryId NAME_INDEX;
 #define SUBOBJECT_DELIMITER				TEXT(SUBOBJECT_DELIMITER_ANSI)
 
 /** this is the character used to separate a subobject root from its subobjects in a path name, as a char */
-#define SUBOBJECT_DELIMITER_CHAR		TEXT(':')
+#define SUBOBJECT_DELIMITER_CHAR_ANSI	':'
+#define SUBOBJECT_DELIMITER_CHAR		TEXT(SUBOBJECT_DELIMITER_CHAR_ANSI)
 
 /** These are the characters that cannot be used in general FNames */
 #define INVALID_NAME_CHARACTERS			TEXT("\"' ,\n\r\t")
