Support for joining tuples, using a new set of `UE::Strings::JoinTuple*` functions.
Added CIterable concept which we are no longer using in Join.h, but can be used to constrain template types to those which provide valid begin and end functions.
Added new test cases.

#tests Successfully compiled and executed the new test cases via FoundationTests.

diff --git a/Engine/Source/Runtime/Core/Public/Concepts/Iterable.h b/Engine/Source/Runtime/Core/Public/Concepts/Iterable.h
new file mode 100644
--- /dev/null
+++ b/Engine/Source/Runtime/Core/Public/Concepts/Iterable.h
@@ -0,0 +1,31 @@
+// Copyright Epic Games, Inc. All Rights Reserved.
+
+#pragma once
+
+#include <iterator>
+
+namespace UE::Core::Iterable::Private
+{
+	using std::begin;
+	using std::end;
+
+	template <typename T>
+	auto Begin(T&& Range) -> decltype(begin((T&&)Range));
+
+	template <typename T>
+	auto End(T&& Range) -> decltype(end((T&&)Range));
+}
+
+/**
+ * Describes a type that can be iterated over using range-based for loops or standard algorithms.
+ * The type must provide valid `begin` and `end` methods, either through member functions or
+ * through overloads of `begin` and `end` found via argument-dependent lookup (ADL).
+ */
+struct CIterable
+{
+	template <typename T>
+	auto Requires(T&& Val) -> decltype(
+		UE::Core::Iterable::Private::Begin((T&&)Val),
+		UE::Core::Iterable::Private::End((T&&)Val)
+	);
+};
diff --git a/Engine/Source/Runtime/Core/Public/String/Join.h b/Engine/Source/Runtime/Core/Public/String/Join.h
--- a/Engine/Source/Runtime/Core/Public/String/Join.h
+++ b/Engine/Source/Runtime/Core/Public/String/Join.h
@@ -7,6 +7,7 @@
 #include "Misc/StringBuilder.h"
 #include "Templates/IdentityFunctor.h"
 #include "Templates/Invoke.h"
+#include "Templates/Tuple.h"
 
 namespace UE::String::Private
 {
@@ -28,6 +29,23 @@ struct TJoinBy
 	DelimiterType&& Delimiter;
 };
 
+template <typename TupleType, typename ProjectionType, typename DelimiterType, typename QuoteType>
+struct TJoinTupleQuotedBy
+{
+	TupleType&& Tuple;
+	ProjectionType Projection;
+	DelimiterType&& Delimiter;
+	QuoteType&& Quote;
+};
+
+template <typename TupleType, typename ProjectionType, typename DelimiterType>
+struct TJoinTupleBy
+{
+	TupleType&& Tuple;
+	ProjectionType Projection;
+	DelimiterType&& Delimiter;
+};
+
 } // UE::String::Private
 
 namespace UE::String
@@ -208,6 +226,183 @@ inline TStringBuilderBase<CharType>& JoinTo(RangeType&& Range, DelimiterType&& D
 	return JoinByTo(Forward<RangeType>(Range), FIdentityFunctor(), Forward<DelimiterType>(Delimiter), Builder);
 }
 
+/**
+ * Create an object that can be appended to a string builder to append every element of the tuple
+ * to the builder, separating the elements by the delimiter and surrounding every element on both
+ * sides with a quote.
+ *
+ * @param Tuple        The tuple of elements to join and append.
+ * @param Projection   The projection to apply to the elements before appending them.
+ * @param Delimiter    The delimiter to append as a separator for the elements.
+ * @param Quote        The quote to append on both sides of each element.
+ *
+ * @return An anonymous object to append to a string builder.
+ */
+template <typename TupleType, typename ProjectionType, typename DelimiterType, typename QuoteType>
+inline auto JoinTupleQuotedBy(TupleType&& Tuple, ProjectionType Projection, DelimiterType&& Delimiter, QuoteType&& Quote)
+	-> Private::TJoinTupleQuotedBy<TupleType, ProjectionType, DelimiterType, QuoteType>
+{
+	return {Forward<TupleType>(Tuple), MoveTemp(Projection), Forward<DelimiterType>(Delimiter), Forward<QuoteType>(Quote)};
+}
+
+/**
+ * Append every element of the tuple to the builder, separating the elements by the delimiter, and
+ * surrounding every element on each side with the given quote.
+ *
+ * @param Tuple        The tuple of elements to join and append.
+ * @param Projection   The projection to apply to the elements before appending them.
+ * @param Delimiter    The delimiter to append as a separator for the elements.
+ * @param Quote        The quote to append on both sides of each element.
+ * @param Builder      The builder to append to.
+ *
+ * @return The builder, to allow additional operations to be composed with this one.
+ */
+template <typename TupleType, typename ProjectionType, typename DelimiterType, typename QuoteType, typename CharType
+	UE_REQUIRES(TIsTuple_V<std::decay_t<TupleType>>)>
+inline TStringBuilderBase<CharType>& JoinTupleQuotedByTo(
+	TupleType&& Tuple,
+	ProjectionType Projection,
+	DelimiterType&& Delimiter,
+	QuoteType&& Quote,
+	TStringBuilderBase<CharType>& Builder)
+{
+	bool bFirst = true;
+	VisitTupleElements([&bFirst, &Projection, &Delimiter, &Quote, &Builder](auto&& Elem)
+	{
+		if (bFirst)
+		{
+			bFirst = false;
+		}
+		else
+		{
+			Builder << Delimiter;
+		}
+		Builder << Quote << Invoke(Projection, Elem) << Quote;
+	}, Tuple);
+	return Builder;
+}
+
+/**
+ * Create an object that can be appended to a string builder to append every element of the tuple
+ * to the builder, separating the elements by the delimiter and surrounding every element on both
+ * sides with a quote.
+ *
+ * @param Tuple       The tuple of elements to join and append.
+ * @param Delimiter   The delimiter to append as a separator for the elements.
+ * @param Quote       The quote to append on both sides of each element.
+ *
+ * @return An anonymous object to append to a string builder.
+ */
+template <typename TupleType, typename DelimiterType, typename QuoteType>
+inline auto JoinTupleQuoted(TupleType&& Tuple, DelimiterType&& Delimiter, QuoteType&& Quote)
+	-> Private::TJoinTupleQuotedBy<TupleType, FIdentityFunctor, DelimiterType, QuoteType>
+{
+	return {Forward<TupleType>(Tuple), FIdentityFunctor(), Forward<DelimiterType>(Delimiter), Forward<QuoteType>(Quote)};
+}
+
+/**
+ * Append every element of the tuple to the builder, separating the elements by the delimiter, and
+ * surrounding every element on each side with the given quote.
+ *
+ * @param Tuple       The tuple of elements to join and append.
+ * @param Delimiter   The delimiter to append as a separator for the elements.
+ * @param Quote       The quote to append on both sides of each element.
+ * @param Builder     The builder to append to.
+ *
+ * @return The builder, to allow additional operations to be composed with this one.
+ */
+template <typename TupleType, typename DelimiterType, typename QuoteType, typename CharType>
+inline TStringBuilderBase<CharType>& JoinTupleQuotedTo(
+	TupleType&& Tuple,
+	DelimiterType&& Delimiter,
+	QuoteType&& Quote,
+	TStringBuilderBase<CharType>& Builder)
+{
+	return JoinTupleQuotedByTo(Forward<TupleType>(Tuple), FIdentityFunctor(), Forward<DelimiterType>(Delimiter), Forward<QuoteType>(Quote), Builder);
+}
+
+/**
+ * Create an object that can be appended to a string builder to append every element of the tuple
+ * to the builder, separating the elements by the delimiter.
+ *
+ * @param Tuple        The tuple of elements to join and append.
+ * @param Projection   The projection to apply to the elements before appending them.
+ * @param Delimiter    The delimiter to append as a separator for the elements.
+ *
+ * @return An anonymous object to append to a string builder.
+ */
+template <typename TupleType, typename ProjectionType, typename DelimiterType>
+inline auto JoinTupleBy(TupleType&& Tuple, ProjectionType Projection, DelimiterType&& Delimiter)
+	-> Private::TJoinTupleBy<TupleType, ProjectionType, DelimiterType>
+{
+	return {Forward<TupleType>(Tuple), MoveTemp(Projection), Forward<DelimiterType>(Delimiter)};
+}
+
+/**
+ * Append every element of the tuple to the builder, separating the elements by the delimiter.
+ *
+ * @param Tuple        The tuple of elements to join and append.
+ * @param Projection   The projection to apply to the elements before appending them.
+ * @param Delimiter    The delimiter to append as a separator for the elements.
+ * @param Builder      The builder to append to.
+ *
+ * @return The builder, to allow additional operations to be composed with this one.
+ */
+template <typename TupleType, typename ProjectionType, typename DelimiterType, typename CharType
+	UE_REQUIRES(TIsTuple_V<std::decay_t<TupleType>>)>
+inline TStringBuilderBase<CharType>& JoinTupleByTo(
+	TupleType&& Tuple,
+	ProjectionType Projection,
+	DelimiterType&& Delimiter,
+	TStringBuilderBase<CharType>& Builder)
+{
+	bool bFirst = true;
+	VisitTupleElements([&bFirst, &Projection, &Delimiter, &Builder](auto&& Elem)
+	{
+		if (bFirst)
+		{
+			bFirst = false;
+		}
+		else
+		{
+			Builder << Delimiter;
+		}
+		Builder << Invoke(Projection, Elem);
+	}, Tuple);
+	return Builder;
+}
+
+/**
+ * Create an object that can be appended to a string builder to append every element of the tuple
+ * to the builder, separating the elements by the delimiter.
+ *
+ * @param Tuple        The tuple of elements to join and append.
+ * @param Delimiter    The delimiter to append as a separator for the elements.
+ *
+ * @return An anonymous object to append to a string builder.
+ */
+template <typename TupleType, typename DelimiterType>
+inline auto JoinTuple(TupleType&& Tuple, DelimiterType&& Delimiter)
+	-> Private::TJoinTupleBy<TupleType, FIdentityFunctor, DelimiterType>
+{
+	return {Forward<TupleType>(Tuple), FIdentityFunctor(), Forward<DelimiterType>(Delimiter)};
+}
+
+/**
+ * Append every element of the tuple to the builder, separating the elements by the delimiter.
+ *
+ * @param Tuple        The tuple of elements to join and append.
+ * @param Delimiter    The delimiter to append as a separator for the elements.
+ * @param Builder      The builder to append to.
+ *
+ * @return The builder, to allow additional operations to be composed with this one.
+ */
+template <typename TupleType, typename DelimiterType, typename CharType>
+inline TStringBuilderBase<CharType>& JoinTupleTo(TupleType&& Tuple, DelimiterType&& Delimiter, TStringBuilderBase<CharType>& Builder)
+{
+	return JoinTupleByTo(Forward<TupleType>(Tuple), FIdentityFunctor(), Forward<DelimiterType>(Delimiter), Builder);
+}
+
 } // UE::String
 
 namespace UE::String::Private
@@ -234,4 +429,25 @@ inline TStringBuilderBase<CharType>& operator<<(
 	return JoinByTo(Forward<RangeType>(Adapter.Range), MoveTemp(Adapter.Projection), Forward<DelimiterType>(Adapter.Delimiter), Builder);
 }
 
+template <typename TupleType, typename ProjectionType, typename DelimiterType, typename QuoteType, typename CharType>
+inline TStringBuilderBase<CharType>& operator<<(
+	TStringBuilderBase<CharType>& Builder,
+	Private::TJoinTupleQuotedBy<TupleType, ProjectionType, DelimiterType, QuoteType>&& Adapter)
+{
+	return JoinTupleQuotedByTo(
+		Forward<TupleType>(Adapter.Tuple),
+		MoveTemp(Adapter.Projection),
+		Forward<DelimiterType>(Adapter.Delimiter),
+		Forward<QuoteType>(Adapter.Quote),
+		Builder);
+}
+
+template <typename TupleType, typename ProjectionType, typename DelimiterType, typename CharType>
+inline TStringBuilderBase<CharType>& operator<<(
+	TStringBuilderBase<CharType>& Builder,
+	Private::TJoinTupleBy<TupleType, ProjectionType, DelimiterType>&& Adapter)
+{
+	return JoinTupleByTo(Forward<TupleType>(Adapter.Tuple), MoveTemp(Adapter.Projection), Forward<DelimiterType>(Adapter.Delimiter), Builder);
+}
+
 } // UE::String::Private
diff --git a/Engine/Source/Runtime/Core/Tests/String/JoinTest.cpp b/Engine/Source/Runtime/Core/Tests/String/JoinTest.cpp
--- a/Engine/Source/Runtime/Core/Tests/String/JoinTest.cpp
+++ b/Engine/Source/Runtime/Core/Tests/String/JoinTest.cpp
@@ -14,26 +14,106 @@ namespace UE::String
 
 TEST_CASE_NAMED(FStringJoinTest, "System::Core::String::Join", "[Core][String][SmokeFilter]")
 {
+	// Helper lambda for calling ToLower if the parameter type is FString.
+	auto ConditionalToLower = []<typename T>(T&& Elem)
+	{
+		if constexpr (std::is_same_v<T, FString>)
+		{
+			return Elem.ToLower();
+		}
+		return Forward<T>(Elem);
+	};
+
 	SECTION("Join")
 	{
+		// TArrayView<T>&&
 		CHECK(WriteToString<128>(Join(MakeArrayView<FStringView>({TEXTVIEW("ABC"), TEXTVIEW("DEF")}), TEXT(", "))) == TEXTVIEW("ABC, DEF"));
 		CHECK(WriteToString<128>(Join(MakeArrayView<FString>({FString(TEXT("ABC")), FString(TEXT("DEF"))}), TEXT(", "))) == TEXTVIEW("ABC, DEF"));
+
+		// const T(&)[N]
+		const FStringView ArrayStringView[] = {TEXTVIEW("ABC"), TEXTVIEW("DEF")};
+		CHECK(WriteToString<128>(Join(ArrayStringView, TEXT(", "))) == TEXTVIEW("ABC, DEF"));
+		const FString ArrayString[] = {FString(TEXT("ABC")), FString(TEXT("DEF"))};
+		CHECK(WriteToString<128>(Join(ArrayString, TEXT(", "))) == TEXTVIEW("ABC, DEF"));
 	}
 
 	SECTION("JoinBy")
 	{
+		// TArrayView<FString>&&
 		CHECK(WriteToString<128>(JoinBy(MakeArrayView<FString>({FString(TEXT("ABC")), FString(TEXT("DEF"))}), UE_PROJECTION_MEMBER(FString, ToLower), TEXT(", "))) == TEXTVIEW("abc, def"));
+
+		// const FString(&)[N]
+		const FString ArrayString[] = { FString(TEXT("ABC")), FString(TEXT("DEF")) };
+		CHECK(WriteToString<128>(JoinBy(ArrayString, UE_PROJECTION_MEMBER(FString, ToLower), TEXT(", "))) == TEXTVIEW("abc, def"));
 	}
 
 	SECTION("JoinQuoted")
 	{
+		// TArrayView<T>&&
 		CHECK(WriteToString<128>(JoinQuoted(MakeArrayView<FStringView>({TEXTVIEW("ABC"), TEXTVIEW("DEF")}), TEXT(", "), TEXT("|"))) == TEXTVIEW("|ABC|, |DEF|"));
 		CHECK(WriteToString<128>(JoinQuoted(MakeArrayView<FString>({FString(TEXT("ABC")), FString(TEXT("DEF"))}), TEXT(", "), TEXT("|"))) == TEXTVIEW("|ABC|, |DEF|"));
+
+		// const T(&)[N]
+		const FStringView ArrayStringView[] = { TEXTVIEW("ABC"), TEXTVIEW("DEF") };
+		CHECK(WriteToString<128>(JoinQuoted(ArrayStringView, TEXT(", "), TEXT("|"))) == TEXTVIEW("|ABC|, |DEF|"));
+		const FString ArrayString[] = { FString(TEXT("ABC")), FString(TEXT("DEF")) };
+		CHECK(WriteToString<128>(JoinQuoted(ArrayString, TEXT(", "), TEXT("|"))) == TEXTVIEW("|ABC|, |DEF|"));
 	}
 
 	SECTION("JoinQuotedBy")
 	{
+		// TArrayView<FString>&&
 		CHECK(WriteToString<128>(JoinQuotedBy(MakeArrayView<FString>({FString(TEXT("ABC")), FString(TEXT("DEF"))}), UE_PROJECTION_MEMBER(FString, ToLower), TEXT(", "), TEXT("|"))) == TEXTVIEW("|abc|, |def|"));
+
+		// const FString(&)[N]
+		const FString ArrayString[] = { FString(TEXT("ABC")), FString(TEXT("DEF")) };
+		CHECK(WriteToString<128>(JoinQuotedBy(ArrayString, UE_PROJECTION_MEMBER(FString, ToLower), TEXT(", "), TEXT("|"))) == TEXTVIEW("|abc|, |def|"));
+	}
+
+	SECTION("JoinTuple")
+	{
+		// TTuple<Ts...>&&
+		CHECK(WriteToString<128>(JoinTuple(MakeTuple(TEXTVIEW("ABC"), 123), TEXT(", "))) == TEXTVIEW("ABC, 123"));
+		CHECK(WriteToString<128>(JoinTuple(MakeTuple(FString(TEXT("ABC")), 123), TEXT(", "))) == TEXTVIEW("ABC, 123"));
+
+		// const TTuple<Ts...>&
+		const TTuple<FStringView, int32> TupleStringInt = MakeTuple(TEXTVIEW("ABC"), 123);
+		CHECK(WriteToString<128>(JoinTuple(TupleStringInt, TEXT(", "))) == TEXTVIEW("ABC, 123"));
+		const TTuple<FString, int32> TupleStringViewInt = MakeTuple(FString(TEXT("ABC")), 123);
+		CHECK(WriteToString<128>(JoinTuple(TupleStringViewInt, TEXT(", "))) == TEXTVIEW("ABC, 123"));
+	}
+
+	SECTION("JoinTupleBy")
+	{
+		// TTuple<FString, int32>&&
+		CHECK(WriteToString<128>(JoinTupleBy(MakeTuple(FString(TEXT("ABC")), 123), ConditionalToLower, TEXT(", "))) == TEXTVIEW("abc, 123"));
+
+		// const TTuple<FString, int32>&
+		const TTuple<FString, int32> TupleStringInt = MakeTuple(FString(TEXT("ABC")), 123);
+		CHECK(WriteToString<128>(JoinTupleBy(TupleStringInt, ConditionalToLower, TEXT(", "))) == TEXTVIEW("abc, 123"));
+	}
+
+	SECTION("JoinTupleQuoted")
+	{
+		// TTuple<Ts...>&&
+		CHECK(WriteToString<128>(JoinTupleQuoted(MakeTuple(TEXTVIEW("ABC"), 123), TEXT(", "), TEXT("|"))) == TEXTVIEW("|ABC|, |123|"));
+		CHECK(WriteToString<128>(JoinTupleQuoted(MakeTuple(FString(TEXT("ABC")), 123), TEXT(", "), TEXT("|"))) == TEXTVIEW("|ABC|, |123|"));
+
+		// const TTuple<Ts...>&
+		const TTuple<FStringView, int32> TupleStringInt = MakeTuple(TEXTVIEW("ABC"), 123);
+		CHECK(WriteToString<128>(JoinTupleQuoted(TupleStringInt, TEXT(", "), TEXT("|"))) == TEXTVIEW("|ABC|, |123|"));
+		const TTuple<FString, int32> TupleStringViewInt = MakeTuple(FString(TEXT("ABC")), 123);
+		CHECK(WriteToString<128>(JoinTupleQuoted(TupleStringViewInt, TEXT(", "), TEXT("|"))) == TEXTVIEW("|ABC|, |123|"));
+	}
+
+	SECTION("JoinTupleQuotedBy")
+	{
+		// TTuple<FString, int32>&&
+		CHECK(WriteToString<128>(JoinTupleQuotedBy(MakeTuple(FString(TEXT("ABC")), 123), ConditionalToLower, TEXT(", "), TEXT("|"))) == TEXTVIEW("|abc|, |123|"));
+
+		// const TTuple<FString, int32>&
+		const TTuple<FString, int32> TupleStringInt = MakeTuple(FString(TEXT("ABC")), 123);
+		CHECK(WriteToString<128>(JoinTupleQuotedBy(TupleStringInt, ConditionalToLower, TEXT(", "), TEXT("|"))) == TEXTVIEW("|abc|, |123|"));
 	}
 }
 
