[BinnedMalloc]
 - Fix very rare race condition caused by an alloc happening during the first call to event wait while the thread-cache lock wasn't held

diff --git a/Engine/Source/Runtime/Core/Public/HAL/MallocBinnedCommon.h b/Engine/Source/Runtime/Core/Public/HAL/MallocBinnedCommon.h
--- a/Engine/Source/Runtime/Core/Public/HAL/MallocBinnedCommon.h
+++ b/Engine/Source/Runtime/Core/Public/HAL/MallocBinnedCommon.h
@@ -457,7 +457,18 @@ protected:
 	{
 		FORCEINLINE static FPerThreadFreeBlockLists* Get() TSAN_SAFE
 		{
-			return FPlatformTLS::IsValidTlsSlot(BinnedTlsSlot) ? (FPerThreadFreeBlockLists*)FPlatformTLS::GetTlsValue(BinnedTlsSlot) : nullptr;
+			FPerThreadFreeBlockLists* ThreadSingleton = FPlatformTLS::IsValidTlsSlot(BinnedTlsSlot) ? (FPerThreadFreeBlockLists*)FPlatformTLS::GetTlsValue(BinnedTlsSlot) : nullptr;
+			// If the current thread doesn't have the Lock, we can't return the TLS cache for being used on the current thread as we risk racing with another thread doing trimming.
+			// This can only happen in such a scenario.
+			//
+			//  FMemory::MarkTLSCachesAsUnusedOnCurrentThread();
+			//  Node->Event->Wait(); <----- UNSAFE to use the TLS cache by its owner thread but can happen when the wait implementation allocates or frees something.
+			//  FMemory::MarkTLSCachesAsUsedOnCurrentThread();
+			if (ThreadSingleton && ThreadSingleton->bLockedByOwnerThread)
+			{
+				return ThreadSingleton;
+			}
+			return nullptr;
 		}
 
 		static void SetTLS()
@@ -472,6 +483,7 @@ protected:
 				TLSMemory.fetch_add(TLSSize, std::memory_order_relaxed);
 #endif
 				verify(ThreadSingleton);
+				ThreadSingleton->bLockedByOwnerThread = true;
 				ThreadSingleton->Lock();
 				FPlatformTLS::SetTlsValue(BinnedTlsSlot, ThreadSingleton);
 				AllocType::RegisterThreadFreeBlockLists(ThreadSingleton);
@@ -483,6 +495,7 @@ protected:
 			FPerThreadFreeBlockLists* ThreadSingleton = (FPerThreadFreeBlockLists*)FPlatformTLS::GetTlsValue(BinnedTlsSlot);
 			if (ThreadSingleton)
 			{
+				ThreadSingleton->bLockedByOwnerThread = false;
 				ThreadSingleton->Unlock();
 			}
 		}
@@ -493,6 +506,7 @@ protected:
 			if (ThreadSingleton)
 			{
 				ThreadSingleton->Lock();
+				ThreadSingleton->bLockedByOwnerThread = true;
 			}
 		}
 
@@ -507,6 +521,7 @@ protected:
 				TLSMemory.fetch_sub(TLSSize, std::memory_order_relaxed);
 #endif
 				AllocType::UnregisterThreadFreeBlockLists(ThreadSingleton);
+				ThreadSingleton->bLockedByOwnerThread = false;
 				ThreadSingleton->Unlock();
 				ThreadSingleton->~FPerThreadFreeBlockLists();
 
@@ -586,6 +601,7 @@ protected:
 		UE::FMutex Mutex;
 		uint64 MemoryTrimEpoch = 0;
 		FFreeBlockList FreeLists[NumSmallPools];
+		bool bLockedByOwnerThread = false;
 	};
 
 	FORCEINLINE SIZE_T QuantizeSizeCommon(SIZE_T Count, uint32 Alignment, const AllocType& Alloc) const
