Mobile: Disable rect lights by default as they require 2 additional samplers.
Enable rect lights in QAGame only so they get tested on mobile platforms

diff --git a/Engine/Source/Runtime/Core/Private/HAL/ConsoleManager.cpp b/Engine/Source/Runtime/Core/Private/HAL/ConsoleManager.cpp
--- a/Engine/Source/Runtime/Core/Private/HAL/ConsoleManager.cpp
+++ b/Engine/Source/Runtime/Core/Private/HAL/ConsoleManager.cpp
@@ -3774,9 +3774,9 @@ static TAutoConsoleVariable<int32> CVarMobileForwardEnableLocalLights(
 
 static TAutoConsoleVariable<int32> CVarMobileRenderRectLightsAsSpotLights(
 	TEXT("r.Mobile.Forward.RenderRectLightsAsSpotLights"),
-	0,
-	TEXT("0: Rect Lights are rendered as area lights (default)\n"
-		"1: Rect Lights are rendered as Spot Lights\n"),
+	1,
+	TEXT("0: Rect Lights are rendered as area lights\n"
+		"1: Rect Lights are rendered as Spot Lights (default)\n"),
 	ECVF_ReadOnly | ECVF_MobileShaderChange);
 
 static TAutoConsoleVariable<int32> CVarMobileForwardEnableParticleLights(
