Fixed incorrect translation on ViewMatrix for shadow views, which could cause precision loss at LWC ranges, and fixed typo in CheckMatrixPrecision.

The ViewMatrix and OverriddenTranslatedViewMatrix were set to the same value, which is a matrix with 0 translation. When the view matrix is uploaded to GPU, it is first re-based to be relative to the ViewOrigin, before being truncated to single-precision. If the origin of the ViewMatrix is always 0, this rebasing process does not work correctly and precision loss can occur. This fix therefore applies a -ViewOrigin translation to the ViewMatrix, so has the correct translation in absolute worldspace.

#tests TM-VRSmoke (at 0,0,0 and 2T), TM-VSM (at 0 and 2T). Tested for changes in shadow view transform, and LOD selection of meshes

diff --git a/Engine/Source/Editor/UnrealEd/Private/EditorViewportClient.cpp b/Engine/Source/Editor/UnrealEd/Private/EditorViewportClient.cpp
--- a/Engine/Source/Editor/UnrealEd/Private/EditorViewportClient.cpp
+++ b/Engine/Source/Editor/UnrealEd/Private/EditorViewportClient.cpp
@@ -260,7 +260,8 @@ FViewportCameraTransform::FViewportCameraTransform()
 
 void FViewportCameraTransform::SetLocation( const FVector& Position )
 {
-	ViewLocation = Position;
+	FVector ClampedPosition = Position.BoundToCube(WORLD_MAX);
+	ViewLocation = ClampedPosition;
 	DesiredLocation = ViewLocation;
 }
 
diff --git a/Engine/Source/Runtime/Core/Private/Math/DoubleFloat.cpp b/Engine/Source/Runtime/Core/Private/Math/DoubleFloat.cpp
--- a/Engine/Source/Runtime/Core/Private/Math/DoubleFloat.cpp
+++ b/Engine/Source/Runtime/Core/Private/Math/DoubleFloat.cpp
@@ -2,13 +2,14 @@
 
 #include "Math/DoubleFloat.h"
 
-#define UE_DF_MIN_PRECISION (2e-4)
+constexpr float UE_DF_MIN_PRECISION = 1.0f/(1 << 2);
 // Max value of a float before it's precision is lower than UE_DF_MIN_PRECISION
 // (there may be 1 more implicit bit available in the significant, but this works as a safe upper bound)
-#define UE_DF_FLOAT_MAX_VALUE (2e23 * UE_DF_MIN_PRECISION - 1)
+constexpr float UE_DF_FLOAT_MAX_VALUE = ((float)(1 << 23) * UE_DF_MIN_PRECISION - 1.0f);
 
 FMatrix CheckMatrixPrecision(const FMatrix& Matrix)
 {
+#if !UE_BUILD_SHIPPING
 	const double OriginMax = UE_DF_FLOAT_MAX_VALUE;
 
 	const FVector Origin = Matrix.GetOrigin();
@@ -16,7 +17,10 @@ FMatrix CheckMatrixPrecision(const FMatrix& Matrix)
 	const double OriginY = FMath::Abs(Origin.Y);
 	const double OriginZ = FMath::Abs(Origin.Z);
 	ensureMsgf(OriginX <= OriginMax && OriginY <= OriginMax && OriginZ <= OriginMax, 
-		TEXT("Found precision loss while converting matrix to GPU format, verify the input transforms."));
+		TEXT("Found precision loss while converting matrix to GPU format, verify the input transforms. ")
+		TEXT("This error usually indicates the view transform is invalid, or the PreViewTranslation/ViewOrigin was not set up correctly."));
+#endif //!UE_BUILD_SHIPPING
+
 	return Matrix;
 }
 
diff --git a/Engine/Source/Runtime/Engine/Private/SceneView.cpp b/Engine/Source/Runtime/Engine/Private/SceneView.cpp
--- a/Engine/Source/Runtime/Engine/Private/SceneView.cpp
+++ b/Engine/Source/Runtime/Engine/Private/SceneView.cpp
@@ -1227,8 +1227,10 @@ FMatrix FViewMatrices::ScreenToClipProjectionMatrix() const
 
 void FViewMatrices::HackOverrideViewMatrixForShadows(const FMatrix& InViewMatrix)
 {
-	OverriddenTranslatedViewMatrix = ViewMatrix = InViewMatrix;
-	OverriddenInvTranslatedViewMatrix = InvViewMatrix = InViewMatrix.Inverse();
+	OverriddenTranslatedViewMatrix = InViewMatrix;
+	ViewMatrix = FTranslationMatrix(-ViewOrigin) * InViewMatrix;
+	InvViewMatrix = ViewMatrix.Inverse();
+	OverriddenInvTranslatedViewMatrix = OverriddenTranslatedViewMatrix.Inverse();
 }
 
 void FSceneView::UpdateViewMatrix()
