Fixed class template argument deduction for TWeakObjectPtr.

diff --git a/Engine/Source/Runtime/Core/Public/UObject/WeakObjectPtrTemplates.h b/Engine/Source/Runtime/Core/Public/UObject/WeakObjectPtrTemplates.h
--- a/Engine/Source/Runtime/Core/Public/UObject/WeakObjectPtrTemplates.h
+++ b/Engine/Source/Runtime/Core/Public/UObject/WeakObjectPtrTemplates.h
@@ -293,6 +293,9 @@ public:
 #endif
 };
 
+template <typename T>
+TWeakObjectPtr(T*) -> TWeakObjectPtr<T>;
+
 #if !PLATFORM_COMPILER_HAS_GENERATED_COMPARISON_OPERATORS
 /**
  * Compare weak pointers for equality.
