Include GPU breadcrumb data in Editor crash report payloads.

This change serializes the GPU breadcrumb data into the shared context so that it can then be retrieved by the crash report client and included in the final payload.

#tests tested locally

diff --git a/Engine/Source/Programs/CrashReportClient/Private/CrashReportClientApp.cpp b/Engine/Source/Programs/CrashReportClient/Private/CrashReportClientApp.cpp
--- a/Engine/Source/Programs/CrashReportClient/Private/CrashReportClientApp.cpp
+++ b/Engine/Source/Programs/CrashReportClient/Private/CrashReportClientApp.cpp
@@ -513,7 +513,8 @@ FPlatformErrorReport CollectErrorReport(FRecoveryService* RecoveryService, uint3
 		SharedCrashContext.SessionContext,
 		SharedCrashContext.EnabledPluginsNum > 0 ? &SharedCrashContext.DynamicData[SharedCrashContext.EnabledPluginsOffset] : nullptr,
 		SharedCrashContext.EngineDataNum > 0 ? &SharedCrashContext.DynamicData[SharedCrashContext.EngineDataOffset] : nullptr,
-		SharedCrashContext.GameDataNum > 0 ? &SharedCrashContext.DynamicData[SharedCrashContext.GameDataOffset] : nullptr
+		SharedCrashContext.GameDataNum > 0 ? &SharedCrashContext.DynamicData[SharedCrashContext.GameDataOffset] : nullptr,
+		&SharedCrashContext.GPUBreadcrumbs
 	);
 	// Next create a crash context for the crashed process.
 	FPlatformCrashContext CrashContext(SharedCrashContext.CrashType, SharedCrashContext.ErrorMessage);
diff --git a/Engine/Source/Runtime/Core/Private/GenericPlatform/GenericPlatformCrashContext.cpp b/Engine/Source/Runtime/Core/Private/GenericPlatform/GenericPlatformCrashContext.cpp
--- a/Engine/Source/Runtime/Core/Private/GenericPlatform/GenericPlatformCrashContext.cpp
+++ b/Engine/Source/Runtime/Core/Private/GenericPlatform/GenericPlatformCrashContext.cpp
@@ -92,8 +92,6 @@ const TCHAR* AttendedStatusToString(const EUnattendedStatus Status)
 	}
 }
 
-#if WITH_ADDITIONAL_CRASH_CONTEXTS
-
 // Sanitize the event name string to remove characters that are used as delimiters for parsing.
 FString FGPUBreadcrumbCrashData::FSerializer::Sanitize(FString const& Name)
 {
@@ -174,7 +172,55 @@ CORE_API FGPUBreadcrumbCrashData::FQueueData FGPUBreadcrumbCrashData::FSerialize
 	return Result;
 }
 
-#endif // WITH_ADDITIONAL_CRASH_CONTEXTS
+static FGPUBreadcrumbCrashData GPUBreadcrumbsFromSharedContext(const FGPUBreadcrumbsSharedContext& Context)
+{
+	FGPUBreadcrumbCrashData DstData(Context.SourceName, Context.Version);
+
+	for (uint32 QueueIdx = 0; QueueIdx < Context.NumQueues; ++QueueIdx)
+	{
+		const FGPUBreadcrumbsSharedContext::FQueueData& SrcQueue = Context.Queues[QueueIdx];
+
+		FGPUBreadcrumbCrashData::FQueueData DstQueue;
+		DstQueue.ActiveHash.FromString(SrcQueue.ActiveHash);
+		DstQueue.FullHash.FromString(SrcQueue.FullHash);
+		DstQueue.BreadcrumbString = SrcQueue.Breadcrumbs;
+
+		DstData.Queues.Emplace(SrcQueue.QueueName, MoveTemp(DstQueue));
+	}
+	
+	return DstData;
+}
+
+static void GPUBreadcrumbsToSharedContext(const FGPUBreadcrumbCrashData& GPUBreadcrumbs, FGPUBreadcrumbsSharedContext& OutSharedContext)
+{
+	FCString::Strncpy(OutSharedContext.Version, *GPUBreadcrumbs.Version, CR_MAX_GENERIC_FIELD_CHARS);
+	FCString::Strncpy(OutSharedContext.SourceName, *GPUBreadcrumbs.SourceName, CR_MAX_GENERIC_FIELD_CHARS);
+
+	OutSharedContext.NumQueues = 0;
+	for (const TPair<FString, FGPUBreadcrumbCrashData::FQueueData>& SrcData : GPUBreadcrumbs.Queues)
+	{
+		const FGPUBreadcrumbCrashData::FQueueData& SrcBreadcrumbs = SrcData.Value;
+
+		// Skip queues with no breadcrumb data or with too many breadcrumbs.
+		if (SrcBreadcrumbs.BreadcrumbString.IsEmpty() || SrcBreadcrumbs.BreadcrumbString.Len() >= CR_MAX_GPU_BREADCRUMBS_STRING_CHARS)
+		{
+			continue;
+		}
+
+		FGPUBreadcrumbsSharedContext::FQueueData& DstQueue = OutSharedContext.Queues[OutSharedContext.NumQueues];
+		FCString::Strncpy(DstQueue.QueueName, *SrcData.Key, CR_MAX_GENERIC_FIELD_CHARS);
+		FCString::Strncpy(DstQueue.FullHash, *SrcBreadcrumbs.FullHash.ToString(), CR_MAX_GENERIC_FIELD_CHARS);
+		FCString::Strncpy(DstQueue.ActiveHash, *SrcBreadcrumbs.ActiveHash.ToString(), CR_MAX_GENERIC_FIELD_CHARS);
+		FCString::Strncpy(DstQueue.Breadcrumbs, *SrcBreadcrumbs.BreadcrumbString, CR_MAX_GPU_BREADCRUMBS_STRING_CHARS);
+
+		OutSharedContext.NumQueues++;
+		if (OutSharedContext.NumQueues >= CR_MAX_GPU_BREADCRUMBS_QUEUES)
+		{
+			break;
+		}
+	}
+}
+
 
 /*-----------------------------------------------------------------------------
 	FGenericCrashContext
@@ -471,7 +517,7 @@ void FGenericCrashContext::Initialize()
 const TCHAR* CR_PAIR_DELIM = TEXT("\x01");
 const TCHAR* CR_PAIR_EQ = TEXT("\x02");
 
-void FGenericCrashContext::InitializeFromContext(const FSessionContext& Session, const TCHAR* EnabledPluginsStr, const TCHAR* EngineDataStr, const TCHAR* GameDataStr)
+void FGenericCrashContext::InitializeFromContext(const FSessionContext& Session, const TCHAR* EnabledPluginsStr, const TCHAR* EngineDataStr, const TCHAR* GameDataStr, const FGPUBreadcrumbsSharedContext* GPUBreadcrumbs)
 {
 	static const TCHAR* TokenDelim[] = { CR_PAIR_DELIM, CR_PAIR_EQ };
 
@@ -521,6 +567,11 @@ void FGenericCrashContext::InitializeFromContext(const FSessionContext& Session,
 		}
 	}
 
+	if (GPUBreadcrumbs && GPUBreadcrumbs->NumQueues > 0)
+	{
+		SetGPUBreadcrumbs(GPUBreadcrumbsFromSharedContext(*GPUBreadcrumbs));
+	}
+
 	SerializeTempCrashContextToFile();
 
 	bIsInitialized = true;
@@ -587,6 +638,12 @@ void FGenericCrashContext::CopySharedCrashContext(FSharedCrashContext& Dst)
 	DynamicDataPtr += FCString::Strlen(DynamicDataPtr) + 1;
 
 	#undef CR_DYNAMIC_BUFFER_REMAIN
+
+	// Copy GPU breadcrumbs.
+	if (NCached::GPUBreadcrumbs.IsSet())
+	{
+		GPUBreadcrumbsToSharedContext(*NCached::GPUBreadcrumbs, Dst.GPUBreadcrumbs);
+	}
 }
 
 void FGenericCrashContext::SetMemoryStats(const FPlatformMemoryStats& InMemoryStats)
@@ -1165,7 +1222,7 @@ void FGenericCrashContext::AddGPUBreadcrumbs() const
 	// We use a version indicator for the format used by the breadcrumbs
 	// string, so that parsers can know what to expect and don't break
 	// if changes are made in the format exported by the engine.
-	AddCrashProperty(TEXT("FormatVersion"), FGPUBreadcrumbCrashData::Version);
+	AddCrashProperty(TEXT("FormatVersion"), NCached::GPUBreadcrumbs->Version);
 	AddCrashProperty(TEXT("Source"), NCached::GPUBreadcrumbs->SourceName);
 
 	for (auto& [Queue, Breadcrumbs] : NCached::GPUBreadcrumbs->Queues)
diff --git a/Engine/Source/Runtime/Core/Public/GenericPlatform/GenericPlatformCrashContext.h b/Engine/Source/Runtime/Core/Public/GenericPlatform/GenericPlatformCrashContext.h
--- a/Engine/Source/Runtime/Core/Public/GenericPlatform/GenericPlatformCrashContext.h
+++ b/Engine/Source/Runtime/Core/Public/GenericPlatform/GenericPlatformCrashContext.h
@@ -153,6 +153,8 @@ CORE_API const TCHAR* AttendedStatusToString(const EUnattendedStatus Status);
 #define CR_MAX_COMMANDLINE_CHARS 1024
 #define CR_MAX_RICHTEXT_FIELD_CHARS 512
 #define CR_MAX_DYNAMIC_BUFFER_CHARS 1024*32
+#define CR_MAX_GPU_BREADCRUMBS_QUEUES 2
+#define CR_MAX_GPU_BREADCRUMBS_STRING_CHARS 1024*8
 
 /**
  * Fixed size structure that holds session specific state.
@@ -222,6 +224,23 @@ struct FUserSettingsContext
 	TCHAR					LogFilePath[CR_MAX_DIRECTORY_CHARS];
 };
 
+/** Fixed size structure holding GPU breadcrumbs information, to be communicated to the crash reporting client. */
+struct FGPUBreadcrumbsSharedContext
+{
+	struct FQueueData {
+		TCHAR QueueName[CR_MAX_GENERIC_FIELD_CHARS];
+		TCHAR FullHash[CR_MAX_GENERIC_FIELD_CHARS];
+		TCHAR ActiveHash[CR_MAX_GENERIC_FIELD_CHARS];
+		TCHAR Breadcrumbs[CR_MAX_GPU_BREADCRUMBS_STRING_CHARS];
+	};
+
+	TCHAR SourceName[CR_MAX_GENERIC_FIELD_CHARS];
+	TCHAR Version[CR_MAX_GENERIC_FIELD_CHARS];
+
+	uint32 NumQueues = 0;
+	FQueueData Queues[CR_MAX_GPU_BREADCRUMBS_QUEUES];
+};
+
 /**
  * Fixed size struct holds crash information and session specific state. It is designed
  * to shared between processes (e.g. Game and CrashReporterClient).
@@ -262,6 +281,9 @@ struct FSharedCrashContext
 
 	// Instruction address where the exception was raised that initiated crash reporting
 	void*					ExceptionProgramCounter;
+
+	// GPU breadcrumbs.
+	FGPUBreadcrumbsSharedContext  GPUBreadcrumbs;
 };
 
 #if WITH_ADDITIONAL_CRASH_CONTEXTS
@@ -305,7 +327,7 @@ struct FGPUBreadcrumbCrashData
 	 * changes, in order to help parsers in dealing with strings from multiple
 	 * versions.
 	 */
-	static constexpr TCHAR const Version[] = TEXT("526BDA74-7A81-44C3-B0FD-9DBF80973C25");
+	static constexpr TCHAR const CurrentVersion[] = TEXT("526BDA74-7A81-44C3-B0FD-9DBF80973C25");
 
 	enum class EState : uint8
 	{
@@ -314,7 +336,7 @@ struct FGPUBreadcrumbCrashData
 		Finished   = 2
 	};
 
-	// These are serialized. Do not change this array without bumping the Version.
+	// These are serialized. Do not change this array without bumping the CurrentVersion.
 	TCHAR const static constexpr StateChars[] =
 	{
 		TEXT('N'),
@@ -332,6 +354,7 @@ struct FGPUBreadcrumbCrashData
 	};
 	TMap<FString, FQueueData> Queues;
 	FString SourceName;
+	FString Version;
 
 	class FSerializer
 	{
@@ -352,8 +375,8 @@ struct FGPUBreadcrumbCrashData
 		CORE_API FQueueData GetResult();
 	};
 
-	FGPUBreadcrumbCrashData(TCHAR const* SourceName)
-		: SourceName(SourceName)
+	FGPUBreadcrumbCrashData(TCHAR const* InSourceName, TCHAR const* InVersion = CurrentVersion)
+		: SourceName(InSourceName), Version(InVersion)
 	{}
 };
 
@@ -431,7 +454,7 @@ public:
 	CORE_API static void Initialize();
 
 	/** Initialized crash context, using a crash context (e.g. shared from another process). */
-	CORE_API static void InitializeFromContext(const FSessionContext& Context, const TCHAR* EnabledPlugins, const TCHAR* EngineData, const TCHAR* GameData);
+	CORE_API static void InitializeFromContext(const FSessionContext& Context, const TCHAR* EnabledPlugins, const TCHAR* EngineData, const TCHAR* GameData, const FGPUBreadcrumbsSharedContext* GPUBreadcrumbs);
 
 	/** Get the current cached session context */
 	CORE_API static const FSessionContext& GetCachedSessionContext();
