String: Added string case conversion functions under UE::String.

#tests All unit tests pass. Preflight passed.

diff --git a/Engine/Source/Runtime/Core/Private/String/CaseConversion.cpp b/Engine/Source/Runtime/Core/Private/String/CaseConversion.cpp
new file mode 100644
--- /dev/null
+++ b/Engine/Source/Runtime/Core/Private/String/CaseConversion.cpp
@@ -0,0 +1,66 @@
+// Copyright Epic Games, Inc. All Rights Reserved.
+
+#include "String/CaseConversion.h"
+#include "Misc/StringBuilder.h"
+
+namespace UE::String::Private
+{
+
+template <typename CharType>
+static inline void UpperCaseTo(TStringView<CharType> Input, TStringBuilderBase<CharType>& Output)
+{
+	const int32 Offset = Output.AddUninitialized(Input.Len());
+	CharType* OutputIterator = Output.GetData() + Offset;
+	for (CharType Char : Input)
+	{
+		*OutputIterator++ = TChar<CharType>::ToUpper(Char);
+	}
+}
+
+template <typename CharType>
+static inline void LowerCaseTo(TStringView<CharType> Input, TStringBuilderBase<CharType>& Output)
+{
+	const int32 Offset = Output.AddUninitialized(Input.Len());
+	CharType* OutputIterator = Output.GetData() + Offset;
+	for (CharType Char : Input)
+	{
+		*OutputIterator++ = TChar<CharType>::ToLower(Char);
+	}
+}
+
+} // UE::String::Private
+
+namespace UE::String
+{
+
+void UpperCaseTo(FAnsiStringView Input, FAnsiStringBuilderBase& Output)
+{
+	Private::UpperCaseTo(Input, Output);
+}
+
+void UpperCaseTo(FUtf8StringView Input, FUtf8StringBuilderBase& Output)
+{
+	Private::UpperCaseTo(Input, Output);
+}
+
+void UpperCaseTo(FWideStringView Input, FWideStringBuilderBase& Output)
+{
+	Private::UpperCaseTo(Input, Output);
+}
+
+void LowerCaseTo(FAnsiStringView Input, FAnsiStringBuilderBase& Output)
+{
+	Private::LowerCaseTo(Input, Output);
+}
+
+void LowerCaseTo(FUtf8StringView Input, FUtf8StringBuilderBase& Output)
+{
+	Private::LowerCaseTo(Input, Output);
+}
+
+void LowerCaseTo(FWideStringView Input, FWideStringBuilderBase& Output)
+{
+	Private::LowerCaseTo(Input, Output);
+}
+
+} // UE::String
diff --git a/Engine/Source/Runtime/Core/Public/String/CaseConversion.h b/Engine/Source/Runtime/Core/Public/String/CaseConversion.h
new file mode 100644
--- /dev/null
+++ b/Engine/Source/Runtime/Core/Public/String/CaseConversion.h
@@ -0,0 +1,120 @@
+// Copyright Epic Games, Inc. All Rights Reserved.
+
+#pragma once
+
+#include "Containers/StringFwd.h"
+#include "Containers/StringView.h"
+#include "CoreTypes.h"
+
+namespace UE::String::Private
+{
+
+template <typename CharType>
+struct TUpperCase { TStringView<CharType> Input; };
+
+template <typename CharType>
+struct TLowerCase { TStringView<CharType> Input; };
+
+} // UE::String::Private
+
+namespace UE::String
+{
+
+/**
+ * Convert the string to uppercase and append to the string builder.
+ * 
+ * @note Only ASCII characters will be converted, similar to CRT to[w]upper() with the standard C locale.
+ * @param Input   The string to convert to uppercase.
+ * @param Output  The string builder to append the converted string to.
+ * 
+ * @code
+ *	UpperCaseTo(TEXTVIEW("example"), Output); // Output now contains "EXAMPLE"
+ * @endcode
+ */
+CORE_API void UpperCaseTo(FAnsiStringView Input, FAnsiStringBuilderBase& Output);
+CORE_API void UpperCaseTo(FUtf8StringView Input, FUtf8StringBuilderBase& Output);
+CORE_API void UpperCaseTo(FWideStringView Input, FWideStringBuilderBase& Output);
+
+/**
+ * Convert the string to lowercase and append to the string builder.
+ *
+ * @note Only ASCII characters will be converted, similar to CRT to[w]upper() with the standard C locale.
+ * @param Input   The string to convert to lowercase.
+ * @param Output  The string builder to append the converted string to.
+ * 
+ * @code
+ *	LowerCaseTo(TEXTVIEW("EXAMPLE"), Output); // Output now contains "example"
+ * @endcode
+ */
+CORE_API void LowerCaseTo(FAnsiStringView Input, FAnsiStringBuilderBase& Output);
+CORE_API void LowerCaseTo(FUtf8StringView Input, FUtf8StringBuilderBase& Output);
+CORE_API void LowerCaseTo(FWideStringView Input, FWideStringBuilderBase& Output);
+
+/**
+ * Convert the string to uppercase when appended to a string builder.
+ *
+ * @note Only ASCII characters will be converted, similar to CRT to[w]upper() with the standard C locale.
+ * @param Input   The string to convert to uppercase.
+ * @return An anonymous object to append to a string builder.
+ * 
+ * @code
+ *	Builder << String::UpperCase(TEXTVIEW("example")); // Builder now contains "EXAMPLE"
+ * @endcode
+ */
+inline Private::TUpperCase<ANSICHAR> UpperCase(FAnsiStringView Input)
+{
+	return {Input};
+}
+inline Private::TUpperCase<UTF8CHAR> UpperCase(FUtf8StringView Input)
+{
+	return {Input};
+}
+inline Private::TUpperCase<WIDECHAR> UpperCase(FWideStringView Input)
+{
+	return {Input};
+}
+
+/**
+ * Convert the string to lowercase when appended to a string builder.
+ *
+ * @note Only ASCII characters will be converted, similar to CRT to[w]upper() with the standard C locale.
+ * @param Input   The string to convert to lowercase.
+ * @return An anonymous object to append to a string builder.
+ * 
+ * @code
+ *	Builder << String::LowerCase(TEXTVIEW("EXAMPLE")); // Builder now contains "example"
+ * @endcode
+ */
+inline Private::TLowerCase<ANSICHAR> LowerCase(FAnsiStringView Input)
+{
+	return {Input};
+}
+inline Private::TLowerCase<UTF8CHAR> LowerCase(FUtf8StringView Input)
+{
+	return {Input};
+}
+inline Private::TLowerCase<WIDECHAR> LowerCase(FWideStringView Input)
+{
+	return {Input};
+}
+
+} // UE::String
+
+namespace UE::String::Private
+{
+
+template <typename CharType>
+inline TStringBuilderBase<CharType>& operator<<(TStringBuilderBase<CharType>& Builder, const TUpperCase<CharType>& Adapter)
+{
+	UpperCaseTo(Adapter.Input, Builder);
+	return Builder;
+}
+
+template <typename CharType>
+inline TStringBuilderBase<CharType>& operator<<(TStringBuilderBase<CharType>& Builder, const TLowerCase<CharType>& Adapter)
+{
+	LowerCaseTo(Adapter.Input, Builder);
+	return Builder;
+}
+
+} // UE::String::Private
diff --git a/Engine/Source/Runtime/Core/Tests/String/CaseConversionTest.cpp b/Engine/Source/Runtime/Core/Tests/String/CaseConversionTest.cpp
new file mode 100644
--- /dev/null
+++ b/Engine/Source/Runtime/Core/Tests/String/CaseConversionTest.cpp
@@ -0,0 +1,93 @@
+// Copyright Epic Games, Inc. All Rights Reserved.
+
+#if WITH_TESTS
+
+#include "String/CaseConversion.h"
+#include "Containers/StringView.h"
+#include "Misc/StringBuilder.h"
+#include "Tests/TestHarnessAdapter.h"
+
+namespace UE::String
+{
+
+TEST_CASE_NAMED(FStringUpperCaseTest, "System::Core::String::UpperCase", "[Core][String][SmokeFilter]")
+{
+	SECTION("UpperCase with ANSICHAR")
+	{
+		CHECK(WriteToAnsiString<128>(UpperCase(ANSITEXTVIEW("no uppercase"))) == ANSITEXTVIEW("NO UPPERCASE"));
+		CHECK(WriteToAnsiString<128>(UpperCase(ANSITEXTVIEW("Mixed CASE"))) == ANSITEXTVIEW("MIXED CASE"));
+		CHECK(WriteToAnsiString<128>(UpperCase(ANSITEXTVIEW("ALL UPPERCASE"))) == ANSITEXTVIEW("ALL UPPERCASE"));
+		CHECK(WriteToAnsiString<128>(UpperCase(ANSITEXTVIEW("  spaces   and\ttabs\n"))) == ANSITEXTVIEW("  SPACES   AND\tTABS\n"));
+		CHECK(WriteToAnsiString<128>(UpperCase(ANSITEXTVIEW(""))) == ANSITEXTVIEW(""));
+		CHECK(WriteToAnsiString<128>(UpperCase(ANSITEXTVIEW("a"))) == ANSITEXTVIEW("A"));
+		CHECK(WriteToAnsiString<128>(UpperCase(ANSITEXTVIEW("Z"))) == ANSITEXTVIEW("Z"));
+		CHECK(WriteToAnsiString<128>(UpperCase(ANSITEXTVIEW("1234!@#$"))) == ANSITEXTVIEW("1234!@#$"));
+	}
+
+	SECTION("UpperCase with WIDECHAR")
+	{
+		CHECK(WriteToWideString<128>(UpperCase(WIDETEXTVIEW("no uppercase"))) == WIDETEXTVIEW("NO UPPERCASE"));
+		CHECK(WriteToWideString<128>(UpperCase(WIDETEXTVIEW("Mixed CASE"))) == WIDETEXTVIEW("MIXED CASE"));
+		CHECK(WriteToWideString<128>(UpperCase(WIDETEXTVIEW("ALL UPPERCASE"))) == WIDETEXTVIEW("ALL UPPERCASE"));
+		CHECK(WriteToWideString<128>(UpperCase(WIDETEXTVIEW("  spaces   and\ttabs\n"))) == WIDETEXTVIEW("  SPACES   AND\tTABS\n"));
+		CHECK(WriteToWideString<128>(UpperCase(WIDETEXTVIEW(""))) == WIDETEXTVIEW(""));
+		CHECK(WriteToWideString<128>(UpperCase(WIDETEXTVIEW("a"))) == WIDETEXTVIEW("A"));
+		CHECK(WriteToWideString<128>(UpperCase(WIDETEXTVIEW("Z"))) == WIDETEXTVIEW("Z"));
+		CHECK(WriteToWideString<128>(UpperCase(WIDETEXTVIEW("1234!@#$"))) == WIDETEXTVIEW("1234!@#$"));
+	}
+
+	SECTION("UpperCase with UTF8CHAR")
+	{
+		CHECK(WriteToUtf8String<128>(UpperCase(UTF8TEXTVIEW("no uppercase"))) == UTF8TEXTVIEW("NO UPPERCASE"));
+		CHECK(WriteToUtf8String<128>(UpperCase(UTF8TEXTVIEW("Mixed CASE"))) == UTF8TEXTVIEW("MIXED CASE"));
+		CHECK(WriteToUtf8String<128>(UpperCase(UTF8TEXTVIEW("ALL UPPERCASE"))) == UTF8TEXTVIEW("ALL UPPERCASE"));
+		CHECK(WriteToUtf8String<128>(UpperCase(UTF8TEXTVIEW("  spaces   and\ttabs\n"))) == UTF8TEXTVIEW("  SPACES   AND\tTABS\n"));
+		CHECK(WriteToUtf8String<128>(UpperCase(UTF8TEXTVIEW(""))) == UTF8TEXTVIEW(""));
+		CHECK(WriteToUtf8String<128>(UpperCase(UTF8TEXTVIEW("a"))) == UTF8TEXTVIEW("A"));
+		CHECK(WriteToUtf8String<128>(UpperCase(UTF8TEXTVIEW("Z"))) == UTF8TEXTVIEW("Z"));
+		CHECK(WriteToUtf8String<128>(UpperCase(UTF8TEXTVIEW("1234!@#$"))) == UTF8TEXTVIEW("1234!@#$"));
+	}
+}
+
+TEST_CASE_NAMED(FStringLowerCaseTest, "System::Core::String::LowerCase", "[Core][String][SmokeFilter]")
+{
+	SECTION("LowerCase with ANSICHAR")
+	{
+		CHECK(WriteToAnsiString<128>(LowerCase(ANSITEXTVIEW("NO LOWERCASE"))) == ANSITEXTVIEW("no lowercase"));
+		CHECK(WriteToAnsiString<128>(LowerCase(ANSITEXTVIEW("Mixed CASE"))) == ANSITEXTVIEW("mixed case"));
+		CHECK(WriteToAnsiString<128>(LowerCase(ANSITEXTVIEW("already lowercase"))) == ANSITEXTVIEW("already lowercase"));
+		CHECK(WriteToAnsiString<128>(LowerCase(ANSITEXTVIEW("  SPACES   AND\tTABS\n"))) == ANSITEXTVIEW("  spaces   and\ttabs\n"));
+		CHECK(WriteToAnsiString<128>(LowerCase(ANSITEXTVIEW(""))) == ANSITEXTVIEW(""));
+		CHECK(WriteToAnsiString<128>(LowerCase(ANSITEXTVIEW("A"))) == ANSITEXTVIEW("a"));
+		CHECK(WriteToAnsiString<128>(LowerCase(ANSITEXTVIEW("z"))) == ANSITEXTVIEW("z"));
+		CHECK(WriteToAnsiString<128>(LowerCase(ANSITEXTVIEW("1234!@#$"))) == ANSITEXTVIEW("1234!@#$"));
+	}
+
+	SECTION("LowerCase with WIDECHAR")
+	{
+		CHECK(WriteToWideString<128>(LowerCase(WIDETEXTVIEW("NO LOWERCASE"))) == WIDETEXTVIEW("no lowercase"));
+		CHECK(WriteToWideString<128>(LowerCase(WIDETEXTVIEW("Mixed CASE"))) == WIDETEXTVIEW("mixed case"));
+		CHECK(WriteToWideString<128>(LowerCase(WIDETEXTVIEW("already lowercase"))) == WIDETEXTVIEW("already lowercase"));
+		CHECK(WriteToWideString<128>(LowerCase(WIDETEXTVIEW("  SPACES   AND\tTABS\n"))) == WIDETEXTVIEW("  spaces   and\ttabs\n"));
+		CHECK(WriteToWideString<128>(LowerCase(WIDETEXTVIEW(""))) == WIDETEXTVIEW(""));
+		CHECK(WriteToWideString<128>(LowerCase(WIDETEXTVIEW("A"))) == WIDETEXTVIEW("a"));
+		CHECK(WriteToWideString<128>(LowerCase(WIDETEXTVIEW("z"))) == WIDETEXTVIEW("z"));
+		CHECK(WriteToWideString<128>(LowerCase(WIDETEXTVIEW("1234!@#$"))) == WIDETEXTVIEW("1234!@#$"));
+	}
+
+	SECTION("LowerCase with UTF8CHAR")
+	{
+		CHECK(WriteToUtf8String<128>(LowerCase(UTF8TEXTVIEW("NO LOWERCASE"))) == UTF8TEXTVIEW("no lowercase"));
+		CHECK(WriteToUtf8String<128>(LowerCase(UTF8TEXTVIEW("Mixed CASE"))) == UTF8TEXTVIEW("mixed case"));
+		CHECK(WriteToUtf8String<128>(LowerCase(UTF8TEXTVIEW("already lowercase"))) == UTF8TEXTVIEW("already lowercase"));
+		CHECK(WriteToUtf8String<128>(LowerCase(UTF8TEXTVIEW("  SPACES   AND\tTABS\n"))) == UTF8TEXTVIEW("  spaces   and\ttabs\n"));
+		CHECK(WriteToUtf8String<128>(LowerCase(UTF8TEXTVIEW(""))) == UTF8TEXTVIEW(""));
+		CHECK(WriteToUtf8String<128>(LowerCase(UTF8TEXTVIEW("A"))) == UTF8TEXTVIEW("a"));
+		CHECK(WriteToUtf8String<128>(LowerCase(UTF8TEXTVIEW("z"))) == UTF8TEXTVIEW("z"));
+		CHECK(WriteToUtf8String<128>(LowerCase(UTF8TEXTVIEW("1234!@#$"))) == UTF8TEXTVIEW("1234!@#$"));
+	}
+}
+
+} // UE::String
+
+#endif //WITH_TESTS
