Add a cvar LLM.CsvFlushEveryRow to determine if LLM CSV should be flushed to disk every row. If the cvar is disabled, we flush only on a crash.
Note: Flushing every row is the pre-existing behavior and remains default. On some platforms, flushing the CSV this can be very slow, so flushing only on a crash may be preferable

diff --git a/Engine/Source/Runtime/Core/Private/HAL/LowLevelMemTracker.cpp b/Engine/Source/Runtime/Core/Private/HAL/LowLevelMemTracker.cpp
--- a/Engine/Source/Runtime/Core/Private/HAL/LowLevelMemTracker.cpp
+++ b/Engine/Source/Runtime/Core/Private/HAL/LowLevelMemTracker.cpp
@@ -951,6 +951,9 @@ TAutoConsoleVariable<int32> CVarLLMWriteInterval(
 	TEXT("The number of seconds between each line in the LLM csv (zero to write every frame)")
 );
 
+TAutoConsoleVariable<int32> CVarLLMCsvFlushEveryRow(TEXT("LLM.CsvFlushEveryRow"), 1,
+	TEXT("Whether to flush the CSV with every row written. If disabled we only flush on a crash"));
+
 TAutoConsoleVariable<int32> CVarLLMHeaderMaxSize(
 	TEXT("LLM.LLMHeaderMaxSize"),
 #if LLM_ALLOW_ASSETS_TAGS
@@ -1278,6 +1281,8 @@ public:
 
 	void OnPreFork();
 
+	void FlushOnCrash();
+
 private:
 	void Write(FStringView Text);
 	static const TCHAR* GetTrackerCsvName(ELLMTracker InTracker);
@@ -1295,6 +1300,7 @@ private:
 	double LastWriteTime;
 	int32 WriteCount;
 	ELLMTracker Tracker;
+	bool bRegisteredCrashDelegate;
 };
 
 /** Outputs the LLM tags and sizes to TraceLog events. */
@@ -1892,6 +1898,7 @@ void FLowLevelMemTracker::OnPreFork()
 
 void FLowLevelMemTracker::UpdateStatsPerFrame(const TCHAR* LogName)
 {
+	CSV_SCOPED_TIMING_STAT_EXCLUSIVE(LLM);
 #if UE_ENABLE_ARRAY_SLACK_TRACKING
 	// Slack tracking, when compiled in, can run even when regular LLM tracking is disabled
 	LlmTrackArrayTick();
@@ -5968,6 +5975,7 @@ FLLMCsvWriter::FLLMCsvWriter()
 	: Archive(nullptr)
 	, LastWriteTime(FPlatformTime::Seconds())
 	, WriteCount(0)
+	, bRegisteredCrashDelegate(false)
 {
 }
 
@@ -5992,6 +6000,15 @@ void FLLMCsvWriter::OnPreFork()
 	}
 }
 
+void FLLMCsvWriter::FlushOnCrash()
+{
+	if (Archive)
+	{
+		FPlatformMisc::LowLevelOutputDebugStringf(TEXT("Flushing LLM CSV on crash\n"));
+		Archive->Flush();
+	}
+}
+
 void FLLMCsvWriter::SetTracker(ELLMTracker InTracker)
 {
 	Tracker = InTracker;
@@ -6195,7 +6212,19 @@ void FLLMCsvWriter::AddRow(FLowLevelMemTracker& LLMRef, const FTrackerTagSizeMap
 		UE_LOG(LogHAL, Log, TEXT("Wrote LLM csv line %d"), WriteCount);
 	}
 
-	Archive->Flush();
+	if (CVarLLMCsvFlushEveryRow.GetValueOnAnyThread())
+	{
+		Archive->Flush();
+	}
+	else if (!bRegisteredCrashDelegate)
+	{
+		// If we're not flushing every row, lazily register crash delegates to ensure we flush on a crash
+		// Note: we intentionally leak this since we can't clean it up safely 
+		FCoreDelegates::OnHandleSystemError.AddRaw(this, &FLLMCsvWriter::FlushOnCrash);
+		FCoreDelegates::GetOutOfMemoryDelegate().AddRaw(this, &FLLMCsvWriter::FlushOnCrash);
+		bRegisteredCrashDelegate = true;
+	}
+
 }
 
 // FLLMTraceWriter implementation.
