Added another TWeakObjectPtr deduction guide to fix MSVC 19.36.

diff --git a/Engine/Source/Runtime/Core/Public/UObject/WeakObjectPtrTemplates.h b/Engine/Source/Runtime/Core/Public/UObject/WeakObjectPtrTemplates.h
--- a/Engine/Source/Runtime/Core/Public/UObject/WeakObjectPtrTemplates.h
+++ b/Engine/Source/Runtime/Core/Public/UObject/WeakObjectPtrTemplates.h
@@ -296,6 +296,9 @@ public:
 template <typename T>
 TWeakObjectPtr(T*) -> TWeakObjectPtr<T>;
 
+template <typename T>
+TWeakObjectPtr(const TWeakObjectPtr<T>&) -> TWeakObjectPtr<T>;
+
 #if !PLATFORM_COMPILER_HAS_GENERATED_COMPARISON_OPERATORS
 /**
  * Compare weak pointers for equality.
