Add FFileHandleRegistryReadTracker utility class to make it easier to track FileHandleRegistry use when reading files.

#rnx

- Some platform file implementations want to track files open for read via the file handle registry but not track files open for write but file handles that are open for writing CAN also be used for reading. This then complicates the file reading implementation as we have to keep track of if the file handle is part of the file handle registry or not when reading and if it is being tracked for the read or not.
- FFileHandleRegistryReadTracker helps simply this code as it can be told in the constructor if it should actually track the file handle or not so we don't need as many branches in the calling code and it becomes easier to read.

diff --git a/Engine/Source/Runtime/Core/Public/HAL/PlatformFileCommon.h b/Engine/Source/Runtime/Core/Public/HAL/PlatformFileCommon.h
--- a/Engine/Source/Runtime/Core/Public/HAL/PlatformFileCommon.h
+++ b/Engine/Source/Runtime/Core/Public/HAL/PlatformFileCommon.h
@@ -232,3 +232,45 @@ private:
 	FThreadSafeCounter HandlesCurrentlyInUse;
 };
 
+class FFileHandleRegistryReadTracker
+{
+public:
+	FFileHandleRegistryReadTracker(FFileHandleRegistry& InFileRegistry, FRegisteredFileHandle& InHandle, bool bTrackRead = true)
+		: FileRegistry(InFileRegistry)
+		, Handle(InHandle)
+	{
+		if (bTrackRead)
+		{
+			State = FileRegistry.TrackStartRead(&Handle)	? ETrackingState::Success
+															: ETrackingState::Failure;
+		}
+		else
+		{
+			State = ETrackingState::Skip;
+		}
+	}
+
+	~FFileHandleRegistryReadTracker()
+	{
+		if (State == ETrackingState::Success)
+		{
+			FileRegistry.TrackEndRead(&Handle);
+		}
+	}
+
+	bool IsValid() const
+	{
+		return State != ETrackingState::Failure;
+	}
+
+private:
+	FFileHandleRegistry& FileRegistry;
+	FRegisteredFileHandle& Handle;
+
+	enum class ETrackingState : uint8
+	{
+		Failure,
+		Skip,
+		Success
+	} State;
+};
