EAutomationTestFlags is now an enum class of almost entirely pure flags, with the aim of making the flags visible in the debugger.  ApplicationContextMask still needs fixing up everywhere.

diff --git a/Engine/Plugins/Animation/ControlRig/Source/ControlRig/Public/Units/RigUnitTest.h b/Engine/Plugins/Animation/ControlRig/Source/ControlRig/Public/Units/RigUnitTest.h
--- a/Engine/Plugins/Animation/ControlRig/Source/ControlRig/Public/Units/RigUnitTest.h
+++ b/Engine/Plugins/Animation/ControlRig/Source/ControlRig/Public/Units/RigUnitTest.h
@@ -55,7 +55,7 @@ public:
 		TUnitStruct##Test( const FString& InName ) \
 		:FControlRigUnitTestBase( InName, false ) {\
 		} \
-		virtual uint32 GetTestFlags() const override { return EAutomationTestFlags::EditorContext | EAutomationTestFlags::EngineFilter; } \
+		virtual EAutomationTestFlags GetTestFlags() const override { return EAutomationTestFlags::EditorContext | EAutomationTestFlags::EngineFilter; } \
 		virtual bool IsStressTest() const { return false; } \
 		virtual uint32 GetRequiredDeviceNum() const override { return 1; } \
 		virtual FString GetTestSourceFileName() const override { return __FILE__; } \
diff --git a/Engine/Plugins/Developer/Concert/ConcertSync/ConcertSyncClient/Source/ConcertSyncClient/Test/ConcertClientLiveTransactionAuthorsTests.cpp b/Engine/Plugins/Developer/Concert/ConcertSync/ConcertSyncClient/Source/ConcertSyncClient/Test/ConcertClientLiveTransactionAuthorsTests.cpp
--- a/Engine/Plugins/Developer/Concert/ConcertSync/ConcertSyncClient/Source/ConcertSyncClient/Test/ConcertClientLiveTransactionAuthorsTests.cpp
+++ b/Engine/Plugins/Developer/Concert/ConcertSync/ConcertSyncClient/Source/ConcertSyncClient/Test/ConcertClientLiveTransactionAuthorsTests.cpp
@@ -13,7 +13,7 @@
 #include "Scratchpad/ConcertScratchpad.h"
 
 /** Flags used for the tests. */
-static const int ConcertClientLiveTransactionAuthorsTestFlags = EAutomationTestFlags::EditorContext | EAutomationTestFlags::EngineFilter;
+static const EAutomationTestFlags ConcertClientLiveTransactionAuthorsTestFlags = EAutomationTestFlags::EditorContext | EAutomationTestFlags::EngineFilter;
 
 namespace ConcertLiveTransactionAuthorsTestUtils
 {
diff --git a/Engine/Plugins/Developer/Concert/ConcertSync/ConcertSyncTest/Source/ConcertSyncTest/Private/ConcertDataStoreTests.cpp b/Engine/Plugins/Developer/Concert/ConcertSync/ConcertSyncTest/Source/ConcertSyncTest/Private/ConcertDataStoreTests.cpp
--- a/Engine/Plugins/Developer/Concert/ConcertSync/ConcertSyncTest/Source/ConcertSyncTest/Private/ConcertDataStoreTests.cpp
+++ b/Engine/Plugins/Developer/Concert/ConcertSync/ConcertSyncTest/Source/ConcertSyncTest/Private/ConcertDataStoreTests.cpp
@@ -15,7 +15,7 @@
 #define LOCTEXT_NAMESPACE "ConcertDataStoreTests" 
 
 /** Flags used for the Concert data store tests. */
-static const int ConcertDataStoreTestFlags = EAutomationTestFlags::EditorContext | EAutomationTestFlags::EngineFilter; // | EAutomationTestFlags::SmokeFilter;
+static const EAutomationTestFlags ConcertDataStoreTestFlags = EAutomationTestFlags::EditorContext | EAutomationTestFlags::EngineFilter; // | EAutomationTestFlags::SmokeFilter;
 
 namespace ConcertDataStoreTestUtils
 {
diff --git a/Engine/Plugins/Developer/TraceSourceFiltering/Source/SourceFilteringTrace/Public/TraceSourceFilteringTestBase.h b/Engine/Plugins/Developer/TraceSourceFiltering/Source/SourceFilteringTrace/Public/TraceSourceFilteringTestBase.h
--- a/Engine/Plugins/Developer/TraceSourceFiltering/Source/SourceFilteringTrace/Public/TraceSourceFilteringTestBase.h
+++ b/Engine/Plugins/Developer/TraceSourceFiltering/Source/SourceFilteringTrace/Public/TraceSourceFilteringTestBase.h
@@ -23,7 +23,7 @@ public:
 	virtual ~FTraceSourceFilteringTestBase() {};
 
 	/** Begin FAutomationTestBase overrides */
-	virtual uint32 GetTestFlags() const override { return EAutomationTestFlags::ApplicationContextMask | EAutomationTestFlags::EngineFilter; }
+	virtual EAutomationTestFlags GetTestFlags() const override { return EAutomationTestFlags::ApplicationContextMask | EAutomationTestFlags::EngineFilter; }
 	virtual bool IsStressTest() const {	return false; }
 	virtual uint32 GetRequiredDeviceNum() const override { return 1; }
 	virtual bool RunTest(const FString& Parameters) override;
diff --git a/Engine/Plugins/PCG/Source/PCG/Public/Tests/PCGTestsCommon.h b/Engine/Plugins/PCG/Source/PCG/Public/Tests/PCGTestsCommon.h
--- a/Engine/Plugins/PCG/Source/PCG/Public/Tests/PCGTestsCommon.h
+++ b/Engine/Plugins/PCG/Source/PCG/Public/Tests/PCGTestsCommon.h
@@ -23,7 +23,7 @@ struct FPCGPoint;
 
 namespace PCGTestsCommon
 {
-	constexpr int TestFlags = EAutomationTestFlags::EditorContext | EAutomationTestFlags::EngineFilter;
+	constexpr EAutomationTestFlags TestFlags = EAutomationTestFlags::EditorContext | EAutomationTestFlags::EngineFilter;
 
 	/** If you are using a FTestData, you can use FTestData::InitializeTestContext. Initialize a context and set the number of tasks available to 1. */
 	TUniquePtr<FPCGContext> InitializeTestContext(IPCGElement* InElement, const FPCGDataCollection& InputData, UPCGComponent* InSourceComponent, const UPCGNode* InNode);
diff --git a/Engine/Plugins/Runtime/GameplayAbilities/Source/GameplayAbilities/Private/Tests/AbilitySystemComponentTests.cpp b/Engine/Plugins/Runtime/GameplayAbilities/Source/GameplayAbilities/Private/Tests/AbilitySystemComponentTests.cpp
--- a/Engine/Plugins/Runtime/GameplayAbilities/Source/GameplayAbilities/Private/Tests/AbilitySystemComponentTests.cpp
+++ b/Engine/Plugins/Runtime/GameplayAbilities/Source/GameplayAbilities/Private/Tests/AbilitySystemComponentTests.cpp
@@ -259,7 +259,7 @@ public:
 		AddTest(&AbilitySystemComponentTestSuite::Test_FailedAbilityFlow, TEXT("FailedAbilityFlow"));
 	}
 
-	virtual uint32 GetTestFlags() const override { return EAutomationTestFlags::EditorContext | EAutomationTestFlags::EngineFilter; }
+	virtual EAutomationTestFlags GetTestFlags() const override { return EAutomationTestFlags::EditorContext | EAutomationTestFlags::EngineFilter; }
 	virtual bool IsStressTest() const { return false; }
 	virtual uint32 GetRequiredDeviceNum() const override { return 1; }
 
diff --git a/Engine/Plugins/Runtime/GameplayAbilities/Source/GameplayAbilities/Private/Tests/GameplayEffectTests.cpp b/Engine/Plugins/Runtime/GameplayAbilities/Source/GameplayAbilities/Private/Tests/GameplayEffectTests.cpp
--- a/Engine/Plugins/Runtime/GameplayAbilities/Source/GameplayAbilities/Private/Tests/GameplayEffectTests.cpp
+++ b/Engine/Plugins/Runtime/GameplayAbilities/Source/GameplayAbilities/Private/Tests/GameplayEffectTests.cpp
@@ -718,7 +718,7 @@ public:
 		ADD_TEST(Test_GameplayCues);
 	}
 
-	virtual uint32 GetTestFlags() const override { return EAutomationTestFlags::EditorContext | EAutomationTestFlags::EngineFilter; }
+	virtual EAutomationTestFlags GetTestFlags() const override { return EAutomationTestFlags::EditorContext | EAutomationTestFlags::EngineFilter; }
 	virtual bool IsStressTest() const { return false; }
 	virtual uint32 GetRequiredDeviceNum() const override { return 1; }
 
diff --git a/Engine/Plugins/Runtime/RigVM/Source/RigVM/Public/RigVMCore/RigVMStructTest.h b/Engine/Plugins/Runtime/RigVM/Source/RigVM/Public/RigVMCore/RigVMStructTest.h
--- a/Engine/Plugins/Runtime/RigVM/Source/RigVM/Public/RigVMCore/RigVMStructTest.h
+++ b/Engine/Plugins/Runtime/RigVM/Source/RigVM/Public/RigVMCore/RigVMStructTest.h
@@ -44,7 +44,7 @@ private:
 		TRigVMStruct##Test( const FString& InName ) \
 		:FRigVMStructTestBase<FRigVMExecuteContext>( InName, false ) {\
 		} \
-		virtual uint32 GetTestFlags() const override { return EAutomationTestFlags::EditorContext | EAutomationTestFlags::EngineFilter; } \
+		virtual EAutomationTestFlags GetTestFlags() const override { return EAutomationTestFlags::EditorContext | EAutomationTestFlags::EngineFilter; } \
 		virtual bool IsStressTest() const { return false; } \
 		virtual uint32 GetRequiredDeviceNum() const override { return 1; } \
 		virtual FString GetTestSourceFileName() const override { return __FILE__; } \
diff --git a/Engine/Plugins/Tests/CQTest/Source/CQTest/Public/CQTest.h b/Engine/Plugins/Tests/CQTest/Source/CQTest/Public/CQTest.h
--- a/Engine/Plugins/Tests/CQTest/Source/CQTest/Public/CQTest.h
+++ b/Engine/Plugins/Tests/CQTest/Source/CQTest/Public/CQTest.h
@@ -65,7 +65,7 @@ enum class ECQTestSuppressLogBehavior
 };
 
 static const FString GenerateTestDirectory = TEXT("");
-static constexpr uint32 DefaultFlags = EAutomationTestFlags::ApplicationContextMask | EAutomationTestFlags::ProductFilter;
+static constexpr EAutomationTestFlags DefaultFlags = EAutomationTestFlags::ApplicationContextMask | EAutomationTestFlags::ProductFilter;
 
 template <typename T>
 concept HasBeforeAll = requires(T t) { { T::BeforeAll(FString()) }; };
@@ -143,10 +143,10 @@ using TTestInstanceGenerator = TUniquePtr<TBaseTest<AsserterType>>(*)(TTestRunne
 template <typename AsserterType>
 struct TTestRunner : public FAutomationTestBase
 {
-	TTestRunner(FString Name, int32 LineNumber, const char* FileName, FString TestDir, uint32 TestFlags, TTestInstanceGenerator<AsserterType> Factory);
+	TTestRunner(FString Name, int32 LineNumber, const char* FileName, FString TestDir, EAutomationTestFlags TestFlags, TTestInstanceGenerator<AsserterType> Factory);
 	~TTestRunner();
 
-	uint32 GetTestFlags() const override;
+	EAutomationTestFlags GetTestFlags() const override;
 	FString GetTestSourceFileName() const override;
 	int32 GetTestSourceFileLine() const override;
 	int32 GetTestSourceFileLine(const FString& Name) const override;
@@ -162,7 +162,7 @@ struct TTestRunner : public FAutomationTestBase
 	int32 LineNumber = 0;
 	FString FileName;
 	FString TestDir;
-	uint32 TestFlags{ DefaultFlags };
+	EAutomationTestFlags TestFlags{ DefaultFlags };
 	bool bInitializing{ true };
 
 	TUniquePtr<TBaseTest<AsserterType>> CurrentTestPtr;
diff --git a/Engine/Plugins/Tests/CQTest/Source/CQTest/Public/Components/PIENetworkComponent.h b/Engine/Plugins/Tests/CQTest/Source/CQTest/Public/Components/PIENetworkComponent.h
--- a/Engine/Plugins/Tests/CQTest/Source/CQTest/Public/Components/PIENetworkComponent.h
+++ b/Engine/Plugins/Tests/CQTest/Source/CQTest/Public/Components/PIENetworkComponent.h
@@ -114,7 +114,7 @@ private:
 	bool bIsDedicatedServer = true;
 };
 
-constexpr static uint32 NetworkTestContext = EAutomationTestFlags::EditorContext | EAutomationTestFlags::ProductFilter;
+constexpr static EAutomationTestFlags NetworkTestContext = EAutomationTestFlags::EditorContext | EAutomationTestFlags::ProductFilter;
 
 #define NETWORK_TEST_CLASS(_ClassName, _TestDir) TEST_CLASS_WITH_FLAGS(_ClassName, _TestDir, NetworkTestContext)
 
diff --git a/Engine/Plugins/Tests/CQTest/Source/CQTest/Public/Impl/CQTest.inl b/Engine/Plugins/Tests/CQTest/Source/CQTest/Public/Impl/CQTest.inl
--- a/Engine/Plugins/Tests/CQTest/Source/CQTest/Public/Impl/CQTest.inl
+++ b/Engine/Plugins/Tests/CQTest/Source/CQTest/Public/Impl/CQTest.inl
@@ -129,7 +129,7 @@ namespace
 
 
 template <typename AsserterType>
-inline TTestRunner<AsserterType>::TTestRunner(FString InName, int32 InLineNumber, const char* InFileName, FString InTestDir, uint32 InTestFlags, TTestInstanceGenerator<AsserterType> InFactory)
+inline TTestRunner<AsserterType>::TTestRunner(FString InName, int32 InLineNumber, const char* InFileName, FString InTestDir, EAutomationTestFlags InTestFlags, TTestInstanceGenerator<AsserterType> InFactory)
 	: FAutomationTestBase(InName, true)
 	, LineNumber(InLineNumber)
 	, FileName(FString(InFileName))
@@ -221,7 +221,7 @@ inline uint32 TTestRunner<AsserterType>::GetRequiredDeviceNum() const
 }
 
 template <typename AsserterType>
-inline uint32 TTestRunner<AsserterType>::GetTestFlags() const
+inline EAutomationTestFlags TTestRunner<AsserterType>::GetTestFlags() const
 {
 	return TestFlags;
 }
diff --git a/Engine/Plugins/Tests/EditorTests/Source/EditorTests/Private/UnrealEd/CollisionAutomationTests.cpp b/Engine/Plugins/Tests/EditorTests/Source/EditorTests/Private/UnrealEd/CollisionAutomationTests.cpp
--- a/Engine/Plugins/Tests/EditorTests/Source/EditorTests/Private/UnrealEd/CollisionAutomationTests.cpp
+++ b/Engine/Plugins/Tests/EditorTests/Source/EditorTests/Private/UnrealEd/CollisionAutomationTests.cpp
@@ -30,10 +30,10 @@ namespace CollisionAutomationTests
 {
 	FAutomationTestBase* TestBase;
 	// Return the currently active world
-	UWorld* GetAutomationWorld(const int32 TestFlags)
+	UWorld* GetAutomationWorld(const EAutomationTestFlags TestFlags)
  	{
 		UWorld* World = nullptr;
-		if( TestFlags & EAutomationTestFlags::ClientContext)
+		if (!!(TestFlags & EAutomationTestFlags::ClientContext))
 		{
 			check(GEngine->GetWorldContexts().Num() == 1);
 			World = GEngine->GetWorldContexts()[0].World();
diff --git a/Engine/Source/Developer/AutomationController/Private/AutomationCommandline.cpp b/Engine/Source/Developer/AutomationController/Private/AutomationCommandline.cpp
--- a/Engine/Source/Developer/AutomationController/Private/AutomationCommandline.cpp
+++ b/Engine/Source/Developer/AutomationController/Private/AutomationCommandline.cpp
@@ -98,7 +98,7 @@ public:
 		FilterMaps.Add("Product", EAutomationTestFlags::ProductFilter);
 		FilterMaps.Add("Standard", EAutomationTestFlags::SmokeFilter | EAutomationTestFlags::EngineFilter | EAutomationTestFlags::ProductFilter | EAutomationTestFlags::PerfFilter);
 		FilterMaps.Add("Negative", EAutomationTestFlags::NegativeFilter);
-		FilterMaps.Add("All", EAutomationTestFlags::FilterMask);
+		FilterMaps.Add("All", EAutomationTestFlags_FilterMask);
 	}
 	
 	void Shutdown()
@@ -583,15 +583,15 @@ protected:
 					Ar.Logf(TEXT("Automation: Setting minimum priority of cases to run to: %s"), *FlagToUse);
 					if (FlagToUse.Contains(TEXT("Low")))
 					{
-						AutomationController->SetRequestedTestFlags(EAutomationTestFlags::PriorityMask);
+						AutomationController->SetRequestedTestFlags(EAutomationTestFlags_PriorityMask);
 					}
 					else if (FlagToUse.Contains(TEXT("Medium")))
 					{
-						AutomationController->SetRequestedTestFlags(EAutomationTestFlags::MediumPriorityAndAbove);
+						AutomationController->SetRequestedTestFlags(EAutomationTestFlags_MediumPriorityAndAbove);
 					}
 					else if (FlagToUse.Contains(TEXT("High")))
 					{
-						AutomationController->SetRequestedTestFlags(EAutomationTestFlags::HighPriorityAndAbove);
+						AutomationController->SetRequestedTestFlags(EAutomationTestFlags_HighPriorityAndAbove);
 					}
 					else if (FlagToUse.Contains(TEXT("Critical")))
 					{
@@ -599,7 +599,7 @@ protected:
 					}
 					else if (FlagToUse.Contains(TEXT("None")))
 					{
-						AutomationController->SetRequestedTestFlags(0);
+						AutomationController->SetRequestedTestFlags(EAutomationTestFlags::None);
 					}
 					else
 					{
@@ -628,7 +628,7 @@ protected:
 					}
 					else if (FlagToUse.Contains(TEXT("None")))
 					{
-						AutomationController->SetRequestedTestFlags(0);
+						AutomationController->SetRequestedTestFlags(EAutomationTestFlags::None);
 					}
 
 					else
@@ -786,7 +786,7 @@ private:
 	int32 TestCount;
 
 	//Dictionary that maps flag names to flag values.
-	TMap<FString, int32> FilterMaps;
+	TMap<FString, EAutomationTestFlags> FilterMaps;
 
 	// Any that we encountered during processing. Used in 'Quit' to determine error code
 	TArray<FString> Errors;
diff --git a/Engine/Source/Developer/AutomationController/Private/AutomationControllerManager.cpp b/Engine/Source/Developer/AutomationController/Private/AutomationControllerManager.cpp
--- a/Engine/Source/Developer/AutomationController/Private/AutomationControllerManager.cpp
+++ b/Engine/Source/Developer/AutomationController/Private/AutomationControllerManager.cpp
@@ -247,7 +247,7 @@ void FAutomationControllerManager::RequestAvailableWorkers(const FGuid& SessionI
 	//store off active session ID to reject messages that come in from different sessions
 	ActiveSessionId = SessionId;
 
-	//EAutomationTestFlags::FilterMask
+	//EAutomationTestFlags_FilterMask
 
 	//TODO AUTOMATION - include change list, game, etc, or remove when launcher is integrated
 	int32 ChangelistNumber = 10000;
diff --git a/Engine/Source/Developer/AutomationController/Private/AutomationControllerManager.h b/Engine/Source/Developer/AutomationController/Private/AutomationControllerManager.h
--- a/Engine/Source/Developer/AutomationController/Private/AutomationControllerManager.h
+++ b/Engine/Source/Developer/AutomationController/Private/AutomationControllerManager.h
@@ -337,7 +337,7 @@ public:
 		return bDeveloperDirectoryIncluded;
 	}
 
-	virtual void SetRequestedTestFlags(const uint32 InRequestedTestFlags) override
+	virtual void SetRequestedTestFlags(const EAutomationTestFlags InRequestedTestFlags) override
 	{
 		RequestedTestFlags = InRequestedTestFlags;
 		RequestTests();
@@ -571,7 +571,7 @@ private:
 	bool bTestResultsAvailable = false;
 
 	/** Which sets of tests to consider */
-	uint32 RequestedTestFlags = 0;
+	EAutomationTestFlags RequestedTestFlags = EAutomationTestFlags::None;
 
 	/** Timer to keep track of the last time tests were updated */
 	double CheckTestTimer;
diff --git a/Engine/Source/Developer/AutomationController/Private/AutomationReport.cpp b/Engine/Source/Developer/AutomationController/Private/AutomationReport.cpp
--- a/Engine/Source/Developer/AutomationController/Private/AutomationReport.cpp
+++ b/Engine/Source/Developer/AutomationController/Private/AutomationReport.cpp
@@ -258,7 +258,7 @@ bool FAutomationReport::IsSupported(const int32 ClusterIndex) const
 }
 
 
-uint32 FAutomationReport::GetTestFlags( ) const
+EAutomationTestFlags FAutomationReport::GetTestFlags( ) const
 {
 	return TestInfo.GetTestFlags();
 }
@@ -273,7 +273,7 @@ int32 FAutomationReport::GetSourceFileLine() const
 	return TestInfo.GetSourceFileLine();
 }
 
-void FAutomationReport::SetTestFlags( const uint32 InTestFlags)
+void FAutomationReport::SetTestFlags( const EAutomationTestFlags InTestFlags)
 {
 	TestInfo.AddTestFlags( InTestFlags );
 
@@ -288,9 +288,9 @@ const bool FAutomationReport::IsParent()
 	return bIsParent;
 }
 
-const bool FAutomationReport::IsSmokeTest( )
+const bool FAutomationReport::IsSmokeTest()
 {
-	return GetTestFlags( ) & EAutomationTestFlags::SmokeFilter ? true : false;
+	return !!(GetTestFlags() & EAutomationTestFlags::SmokeFilter);
 }
 
 bool FAutomationReport::SetFilter( TSharedPtr< AutomationFilterCollection > InFilter, const bool ParentPassedFilter )
diff --git a/Engine/Source/Developer/AutomationController/Private/AutomationReport.h b/Engine/Source/Developer/AutomationController/Private/AutomationReport.h
--- a/Engine/Source/Developer/AutomationController/Private/AutomationReport.h
+++ b/Engine/Source/Developer/AutomationController/Private/AutomationReport.h
@@ -38,8 +38,8 @@ public:
 	virtual void SetEnabled(bool bShouldBeEnabled) override;
 	virtual void SetSupport(const int32 ClusterIndex) override;
 	virtual bool IsSupported(const int32 ClusterIndex) const override;
-	virtual void SetTestFlags(const uint32 TestFlags) override;
-	virtual uint32 GetTestFlags( ) const override;
+	virtual void SetTestFlags(const EAutomationTestFlags TestFlags) override;
+	virtual EAutomationTestFlags GetTestFlags( ) const override;
 	virtual FString GetSourceFile() const override;
 	virtual int32 GetSourceFileLine() const override;
 	virtual const bool IsParent( ) override;
diff --git a/Engine/Source/Developer/AutomationController/Public/IAutomationControllerManager.h b/Engine/Source/Developer/AutomationController/Public/IAutomationControllerManager.h
--- a/Engine/Source/Developer/AutomationController/Public/IAutomationControllerManager.h
+++ b/Engine/Source/Developer/AutomationController/Public/IAutomationControllerManager.h
@@ -350,7 +350,7 @@ public:
 	/**
 	 * Sets what flags we're considering automation tests from
 	 */
-	virtual void SetRequestedTestFlags(const uint32 InRequestedTestFlags) = 0;
+	virtual void SetRequestedTestFlags(const EAutomationTestFlags InRequestedTestFlags) = 0;
 
 	/**
 	 * Check if the automation tests have completed.
diff --git a/Engine/Source/Developer/AutomationController/Public/IAutomationReport.h b/Engine/Source/Developer/AutomationController/Public/IAutomationReport.h
--- a/Engine/Source/Developer/AutomationController/Public/IAutomationReport.h
+++ b/Engine/Source/Developer/AutomationController/Public/IAutomationReport.h
@@ -248,7 +248,7 @@ public:
 	 * 
 	 * @return the test type.
 	 */
-	virtual uint32 GetTestFlags() const = 0;
+	virtual EAutomationTestFlags GetTestFlags() const = 0;
 
 	/** Gets the source file the test was defined on. */
 	virtual FString GetSourceFile() const = 0;
@@ -302,7 +302,7 @@ public:
 	 * 
 	 * @param InTestFlags The EAutomationTestFlags of the test.
 	 */
-	virtual void SetTestFlags(const uint32 InTestFlags) = 0;
+	virtual void SetTestFlags(const EAutomationTestFlags InTestFlags) = 0;
 
 	/** Returns if a particular platform is supported */
 	virtual bool IsSupported(const int32 ClusterIndex) const = 0;
diff --git a/Engine/Source/Developer/AutomationWindow/Private/SAutomationTestItem.cpp b/Engine/Source/Developer/AutomationWindow/Private/SAutomationTestItem.cpp
--- a/Engine/Source/Developer/AutomationWindow/Private/SAutomationTestItem.cpp
+++ b/Engine/Source/Developer/AutomationWindow/Private/SAutomationTestItem.cpp
@@ -435,7 +435,7 @@ END_SLATE_FUNCTION_BUILD_OPTIMIZATION
 const FSlateBrush* SAutomationTestItem::GetSmokeTestImage() const
 {
 	const FSlateBrush* ImageToUse = nullptr;
-	if ( TestStatus->GetTestFlags() & EAutomationTestFlags::SmokeFilter )
+	if (!!(TestStatus->GetTestFlags() & EAutomationTestFlags::SmokeFilter))
 	{
 		if ( TestStatus->IsParent() )
 		{
diff --git a/Engine/Source/Developer/AutomationWindow/Private/SAutomationWindow.cpp b/Engine/Source/Developer/AutomationWindow/Private/SAutomationWindow.cpp
--- a/Engine/Source/Developer/AutomationWindow/Private/SAutomationWindow.cpp
+++ b/Engine/Source/Developer/AutomationWindow/Private/SAutomationWindow.cpp
@@ -1115,12 +1115,12 @@ void SAutomationWindow::HandlePresetTextCommited( const FText& CommittedText, ET
 void SAutomationWindow::HandleRequesteFilterChanged(TSharedPtr<FString> Item, ESelectInfo::Type SelectInfo)
 {
 	const int32 EntryIndex = RequestedFilterComboList.Find(Item);
-	uint32 NewRequestedFlags = EAutomationTestFlags::SmokeFilter;
+	EAutomationTestFlags NewRequestedFlags = EAutomationTestFlags::SmokeFilter;
 
 	switch (EntryIndex)
 	{
 		case 0:	//	"All Tests"
-			NewRequestedFlags = EAutomationTestFlags::FilterMask;
+			NewRequestedFlags = EAutomationTestFlags_FilterMask;
 			break;
 		case 1:	//	"Smoke Tests"
 			NewRequestedFlags = EAutomationTestFlags::SmokeFilter;
diff --git a/Engine/Source/Runtime/AutomationMessages/Public/AutomationWorkerMessages.h b/Engine/Source/Runtime/AutomationMessages/Public/AutomationWorkerMessages.h
--- a/Engine/Source/Runtime/AutomationMessages/Public/AutomationWorkerMessages.h
+++ b/Engine/Source/Runtime/AutomationMessages/Public/AutomationWorkerMessages.h
@@ -184,20 +184,21 @@ struct FAutomationWorkerRequestTests : public FAutomationWorkerMessageBase
 
 	/** Holds a flag indicating whether the developer directory should be included. */
 	UPROPERTY(EditAnywhere, Category="Message")
-	bool DeveloperDirectoryIncluded;
+	bool DeveloperDirectoryIncluded = false;
 
 	/** Holds a flag indicating which tests we'd like to request. */
 	UPROPERTY(EditAnywhere, Category="Message")
-	uint32 RequestedTestFlags;
+	uint32 RequestedTestFlags = 0;
 
 	/** Default constructor. */
-	FAutomationWorkerRequestTests() : DeveloperDirectoryIncluded(false), RequestedTestFlags(0) { }
+	FAutomationWorkerRequestTests() = default;
 
 	/** Creates and initializes a new instance. */
-	FAutomationWorkerRequestTests(bool InDeveloperDirectoryIncluded, uint32 InRequestedTestFlags)
+	FAutomationWorkerRequestTests(bool InDeveloperDirectoryIncluded, EAutomationTestFlags InRequestedTestFlags)
 		: DeveloperDirectoryIncluded(InDeveloperDirectoryIncluded)
-		, RequestedTestFlags(InRequestedTestFlags)
-	{ }
+		, RequestedTestFlags((uint32)InRequestedTestFlags)
+	{
+	}
 };
 
 
@@ -225,7 +226,7 @@ struct FAutomationWorkerSingleTestReply : public FAutomationWorkerMessageBase
 	FString SourceFile;
 
 	UPROPERTY(EditAnywhere, Category="Message")
-	int32 SourceFileLine;
+	int32 SourceFileLine = 0;
 
 	UPROPERTY(EditAnywhere, Category="Message")
 	FString AssetPath;
@@ -234,13 +235,13 @@ struct FAutomationWorkerSingleTestReply : public FAutomationWorkerMessageBase
 	FString OpenCommand;
 
 	UPROPERTY(EditAnywhere, Category="Message")
-	uint32 TestFlags;
+	uint32 TestFlags = 0;
 
 	UPROPERTY(EditAnywhere, Category="Message")
-	uint32 NumParticipantsRequired;
+	uint32 NumParticipantsRequired = 0;
 
 	/** Default constructor. */
-	FAutomationWorkerSingleTestReply() : SourceFileLine(0), TestFlags(0), NumParticipantsRequired(0) { }
+	FAutomationWorkerSingleTestReply() = default;
 
 	/** Creates and initializes a new instance. */
 	FAutomationWorkerSingleTestReply(const FAutomationTestInfo& InTestInfo)
@@ -253,7 +254,7 @@ struct FAutomationWorkerSingleTestReply : public FAutomationWorkerMessageBase
 		SourceFileLine = InTestInfo.GetSourceFileLine();
 		AssetPath = InTestInfo.GetAssetPath();
 		OpenCommand = InTestInfo.GetOpenCommand();
-		TestFlags = InTestInfo.GetTestFlags();
+		TestFlags = (uint32)InTestInfo.GetTestFlags();
 		NumParticipantsRequired = InTestInfo.GetNumParticipantsRequired();
 	}
 
@@ -263,7 +264,7 @@ struct FAutomationWorkerSingleTestReply : public FAutomationWorkerMessageBase
 			DisplayName,
 			FullTestPath,
 			TestName,
-			TestFlags,
+			(EAutomationTestFlags)TestFlags,
 			NumParticipantsRequired,
 			TestParameter,
 			SourceFile,
diff --git a/Engine/Source/Runtime/AutomationWorker/Private/AutomationWorkerModule.cpp b/Engine/Source/Runtime/AutomationWorker/Private/AutomationWorkerModule.cpp
--- a/Engine/Source/Runtime/AutomationWorker/Private/AutomationWorkerModule.cpp
+++ b/Engine/Source/Runtime/AutomationWorker/Private/AutomationWorkerModule.cpp
@@ -384,7 +384,7 @@ void FAutomationWorkerModule::HandleRequestTestsMessage( const FAutomationWorker
 
 	FAutomationTestFramework::Get().LoadTestModules();
 	FAutomationTestFramework::Get().SetDeveloperDirectoryIncluded(Message.DeveloperDirectoryIncluded);
-	FAutomationTestFramework::Get().SetRequestedTestFilter(Message.RequestedTestFlags);
+	FAutomationTestFramework::Get().SetRequestedTestFilter((EAutomationTestFlags)Message.RequestedTestFlags);
 	FAutomationTestFramework::Get().GetValidTestNames( TestInfo );
 
 	SendTests(Context->GetSender());
diff --git a/Engine/Source/Runtime/Core/Private/Misc/AutomationTest.cpp b/Engine/Source/Runtime/Core/Private/Misc/AutomationTest.cpp
--- a/Engine/Source/Runtime/Core/Private/Misc/AutomationTest.cpp
+++ b/Engine/Source/Runtime/Core/Private/Misc/AutomationTest.cpp
@@ -145,35 +145,35 @@ bool FAutomationTestBase::bSuppressLogErrors = false;
 bool FAutomationTestBase::bElevateLogWarningsToErrors = false;
 TArray<FString> FAutomationTestBase::SuppressedLogCategories;
 
-CORE_API const TMap<FString, EAutomationTestFlags::Type>& EAutomationTestFlags::GetTestFlagsMap()
+CORE_API const TMap<FString, EAutomationTestFlags>& EAutomationTestFlags_GetTestFlagsMap()
 {
 	LLM_SCOPE_BYNAME(TEXT("AutomationTest/Framework"));
 	/** String to EAutomationTestFlags map */
-	static const TMap<FString, Type> FlagsMap = {
-		{ TEXT("EditorContext"), Type::EditorContext},
-		{ TEXT("ClientContext"), Type::ClientContext},
-		{ TEXT("ServerContext"), Type::ServerContext},
-		{ TEXT("CommandletContext"), Type::CommandletContext},
-		{ TEXT("ProgramContext"), Type::ProgramContext},
-		{ TEXT("ApplicationContextMask"), Type::ApplicationContextMask},
-		{ TEXT("NonNullRHI"), Type::NonNullRHI},
-		{ TEXT("RequiresUser"), Type::RequiresUser},
-		{ TEXT("FeatureMask"), Type::FeatureMask},
-		{ TEXT("Disabled"), Type::Disabled},
-		{ TEXT("CriticalPriority"), Type::CriticalPriority},
-		{ TEXT("HighPriority"), Type::HighPriority},
-		{ TEXT("HighPriorityAndAbove"), Type::HighPriorityAndAbove},
-		{ TEXT("MediumPriority"), Type::MediumPriority},
-		{ TEXT("MediumPriorityAndAbove"), Type::MediumPriorityAndAbove},
-		{ TEXT("LowPriority"), Type::LowPriority},
-		{ TEXT("PriorityMask"), Type::PriorityMask},
-		{ TEXT("SmokeFilter"), Type::SmokeFilter},
-		{ TEXT("EngineFilter"), Type::EngineFilter},
-		{ TEXT("ProductFilter"), Type::ProductFilter},
-		{ TEXT("PerfFilter"), Type::PerfFilter},
-		{ TEXT("StressFilter"), Type::StressFilter},
-		{ TEXT("NegativeFilter"), Type::NegativeFilter},
-		{ TEXT("FilterMask"), Type::FilterMask}
+	static const TMap<FString, EAutomationTestFlags> FlagsMap = {
+		{ TEXT("EditorContext"),          EAutomationTestFlags::EditorContext},
+		{ TEXT("ClientContext"),          EAutomationTestFlags::ClientContext},
+		{ TEXT("ServerContext"),          EAutomationTestFlags::ServerContext},
+		{ TEXT("CommandletContext"),      EAutomationTestFlags::CommandletContext},
+		{ TEXT("ProgramContext"),         EAutomationTestFlags::ProgramContext},
+		{ TEXT("ApplicationContextMask"), EAutomationTestFlags_ApplicationContextMask},
+		{ TEXT("NonNullRHI"),             EAutomationTestFlags::NonNullRHI},
+		{ TEXT("RequiresUser"),           EAutomationTestFlags::RequiresUser},
+		{ TEXT("FeatureMask"),            EAutomationTestFlags_FeatureMask},
+		{ TEXT("Disabled"),               EAutomationTestFlags::Disabled},
+		{ TEXT("CriticalPriority"),       EAutomationTestFlags::CriticalPriority},
+		{ TEXT("HighPriority"),           EAutomationTestFlags::HighPriority},
+		{ TEXT("HighPriorityAndAbove"),   EAutomationTestFlags_HighPriorityAndAbove},
+		{ TEXT("MediumPriority"),         EAutomationTestFlags::MediumPriority},
+		{ TEXT("MediumPriorityAndAbove"), EAutomationTestFlags_MediumPriorityAndAbove},
+		{ TEXT("LowPriority"),            EAutomationTestFlags::LowPriority},
+		{ TEXT("PriorityMask"),           EAutomationTestFlags_PriorityMask},
+		{ TEXT("SmokeFilter"),            EAutomationTestFlags::SmokeFilter},
+		{ TEXT("EngineFilter"),           EAutomationTestFlags::EngineFilter},
+		{ TEXT("ProductFilter"),          EAutomationTestFlags::ProductFilter},
+		{ TEXT("PerfFilter"),             EAutomationTestFlags::PerfFilter},
+		{ TEXT("StressFilter"),           EAutomationTestFlags::StressFilter},
+		{ TEXT("NegativeFilter"),         EAutomationTestFlags::NegativeFilter},
+		{ TEXT("FilterMask"),             EAutomationTestFlags_FilterMask}
 	};
 	return FlagsMap;
 };
@@ -347,7 +347,7 @@ bool FAutomationTestFramework::UnregisterAutomationTest( const FString& InTestNa
 void FAutomationTestFramework::EnqueueLatentCommand(TSharedPtr<IAutomationLatentCommand> NewCommand)
 {
 	//ensure latent commands are never used within smoke tests - will only catch when smokes are exclusively requested
-	check((RequestedTestFilter & EAutomationTestFlags::FilterMask) != EAutomationTestFlags::SmokeFilter);
+	check((RequestedTestFilter & EAutomationTestFlags_FilterMask) != EAutomationTestFlags::SmokeFilter);
 
 	//ensure we are currently "running a test"
 	check(GIsAutomationTesting);
@@ -358,7 +358,7 @@ void FAutomationTestFramework::EnqueueLatentCommand(TSharedPtr<IAutomationLatent
 void FAutomationTestFramework::EnqueueNetworkCommand(TSharedPtr<IAutomationNetworkCommand> NewCommand)
 {
 	//ensure latent commands are never used within smoke tests
-	check((RequestedTestFilter & EAutomationTestFlags::FilterMask) != EAutomationTestFlags::SmokeFilter);
+	check((RequestedTestFilter & EAutomationTestFlags_FilterMask) != EAutomationTestFlags::SmokeFilter);
 
 	//ensure we are currently "running a test"
 	check(GIsAutomationTesting);
@@ -405,9 +405,8 @@ bool FAutomationTestFramework::RunSmokeTests()
 
 	bool bAllSuccessful = true;
 
-	uint32 PreviousRequestedTestFilter = RequestedTestFilter;
 	//so extra log spam isn't generated
-	RequestedTestFilter = EAutomationTestFlags::SmokeFilter;
+	TGuardValue<EAutomationTestFlags> GuardRequestedTestFilter(RequestedTestFilter, EAutomationTestFlags::SmokeFilter);
 	
 	// Skip running on cooked platforms like mobile
 	//@todo - better determination of whether to run than requires cooked data
@@ -433,7 +432,7 @@ bool FAutomationTestFramework::RunSmokeTests()
 			for ( int TestIndex = 0; TestIndex < TestInfo.Num(); ++TestIndex )
 			{
 				SlowTask.EnterProgressFrame(1);
-				if (TestInfo[TestIndex].GetTestFlags() & EAutomationTestFlags::SmokeFilter )
+				if (!!(TestInfo[TestIndex].GetTestFlags() & EAutomationTestFlags::SmokeFilter))
 				{
 					FString TestCommand = TestInfo[TestIndex].GetTestName();
 					FAutomationTestExecutionInfo& CurExecutionInfo = OutExecutionInfoMap.Add( TestCommand, FAutomationTestExecutionInfo() );
@@ -474,9 +473,6 @@ bool FAutomationTestFramework::RunSmokeTests()
 		bAllSuccessful = false;
 	}
 
-	//revert to allowing all logs
-	RequestedTestFilter = PreviousRequestedTestFilter;
-
 	return bAllSuccessful;
 }
 
@@ -627,7 +623,7 @@ void FAutomationTestFramework::LoadTestModules( )
 {
 	const bool bRunningEditor = GIsEditor && !IsRunningCommandlet();
 
-	bool bRunningSmokeTests = ((RequestedTestFilter & EAutomationTestFlags::FilterMask) == EAutomationTestFlags::SmokeFilter);
+	bool bRunningSmokeTests = ((RequestedTestFilter & EAutomationTestFlags_FilterMask) == EAutomationTestFlags::SmokeFilter);
 	if( !bRunningSmokeTests )
 	{
 		TArray<FString> EngineTestModules;
@@ -685,7 +681,7 @@ void FAutomationTestFramework::GetValidTestNames( TArray<FAutomationTestInfo>& T
 	const bool bRunningProgram = !GIsEditor && FPlatformProperties::IsProgram();
 
 	//application flags
-	uint32 ApplicationSupportFlags = 0;
+	EAutomationTestFlags ApplicationSupportFlags = EAutomationTestFlags::None;
 	if ( bRunningEditor )
 	{
 		ApplicationSupportFlags |= EAutomationTestFlags::EditorContext;
@@ -708,7 +704,7 @@ void FAutomationTestFramework::GetValidTestNames( TArray<FAutomationTestInfo>& T
 	}
 
 	//Feature support - assume valid RHI until told otherwise
-	uint32 FeatureSupportFlags = EAutomationTestFlags::FeatureMask;
+	EAutomationTestFlags FeatureSupportFlags = EAutomationTestFlags_FeatureMask;
 	// @todo: Handle this correctly. GIsUsingNullRHI is defined at Engine-level, so it can't be used directly here in Core.
 	// For now, assume Null RHI is only used for commandlets, servers, and when the command line specifies to use it.
 	if (FPlatformProperties::SupportsWindowedMode())
@@ -729,20 +725,20 @@ void FAutomationTestFramework::GetValidTestNames( TArray<FAutomationTestInfo>& T
 		const FAutomationTestBase* CurTest = TestIter.Value();
 		check( CurTest );
 
-		uint32 CurTestFlags = CurTest->GetTestFlags();
+		EAutomationTestFlags CurTestFlags = CurTest->GetTestFlags();
 
 		//filter out full tests when running smoke tests
-		const bool bPassesFilterRequirement = ((CurTestFlags & RequestedTestFilter) != 0);
+		const bool bPassesFilterRequirement = !!(CurTestFlags & RequestedTestFilter);
 
 		//Application Tests
-		uint32 CurTestApplicationFlags = (CurTestFlags & EAutomationTestFlags::ApplicationContextMask);
-		const bool bPassesApplicationRequirements = (CurTestApplicationFlags == 0) || (CurTestApplicationFlags & ApplicationSupportFlags);
+		EAutomationTestFlags CurTestApplicationFlags = (CurTestFlags & EAutomationTestFlags_ApplicationContextMask);
+		const bool bPassesApplicationRequirements = !CurTestApplicationFlags || !!(CurTestApplicationFlags & ApplicationSupportFlags);
 		
 		//Feature Tests
-		uint32 CurTestFeatureFlags = (CurTestFlags & EAutomationTestFlags::FeatureMask);
-		const bool bPassesFeatureRequirements = (CurTestFeatureFlags == 0) || (CurTestFeatureFlags & FeatureSupportFlags);
+		EAutomationTestFlags CurTestFeatureFlags = (CurTestFlags & EAutomationTestFlags_FeatureMask);
+		const bool bPassesFeatureRequirements = !CurTestFeatureFlags || !!(CurTestFeatureFlags & FeatureSupportFlags);
 
-		const bool bEnabled = (CurTestFlags & EAutomationTestFlags::Disabled) == 0;
+		const bool bEnabled = !(CurTestFlags & EAutomationTestFlags::Disabled);
 
 		const double GenerateTestNamesStartTime = FPlatformTime::Seconds();
 		
@@ -750,7 +746,7 @@ void FAutomationTestFramework::GetValidTestNames( TArray<FAutomationTestInfo>& T
 		{
 			TArray<FAutomationTestInfo> TestsToAdd;
 			CurTest->GenerateTestNames(TestsToAdd);
-			TestInfo.Append(TestsToAdd);			
+			TestInfo.Append(TestsToAdd);
 		}
 
 		// Make sure people are not writing complex tests that take forever to return the names of the tests
@@ -797,7 +793,7 @@ void FAutomationTestFramework::SetDeveloperDirectoryIncluded(const bool bInDevel
 	bDeveloperDirectoryIncluded = bInDeveloperDirectoryIncluded;
 }
 
-void FAutomationTestFramework::SetRequestedTestFilter(const uint32 InRequestedTestFlags)
+void FAutomationTestFramework::SetRequestedTestFilter(const EAutomationTestFlags InRequestedTestFlags)
 {
 	RequestedTestFilter = InRequestedTestFlags;
 }
@@ -978,8 +974,8 @@ void FAutomationTestFramework::InternalStartTest( const FString& InTestToRun, co
 		CurrentTestFullPath = InFullTestPath;
 
 		// If not a smoke test, log the test has started.
-		uint32 NonSmokeTestFlags = (EAutomationTestFlags::FilterMask & (~EAutomationTestFlags::SmokeFilter));
-		if (RequestedTestFilter & NonSmokeTestFlags)
+		EAutomationTestFlags NonSmokeTestFlags = (EAutomationTestFlags_FilterMask & (~EAutomationTestFlags::SmokeFilter));
+		if (!!(RequestedTestFilter & NonSmokeTestFlags))
 		{
 			if (AutomationTest::bLogTestStateTrace)
 			{
@@ -1021,8 +1017,8 @@ bool FAutomationTestFramework::InternalStopTest(FAutomationTestExecutionInfo& Ou
 
 	double EndTime = FPlatformTime::Seconds();
 	double TimeForTest = static_cast<float>(EndTime - StartTime);
-	uint32 NonSmokeTestFlags = (EAutomationTestFlags::FilterMask & (~EAutomationTestFlags::SmokeFilter));
-	if (RequestedTestFilter & NonSmokeTestFlags)
+	EAutomationTestFlags NonSmokeTestFlags = (EAutomationTestFlags_FilterMask & (~EAutomationTestFlags::SmokeFilter));
+	if (!!(RequestedTestFilter & NonSmokeTestFlags))
 	{
 		UE_LOG(LogAutomationTest, Log, TEXT("%s %s ran in %f"), *CurrentTest->GetBeautifiedTestName(), *Parameters, TimeForTest);
 		if (AutomationTest::bLogTestStateTrace)
@@ -1596,21 +1592,17 @@ void FAutomationTestBase::AddExpectedErrorPlain(
 	AddExpectedMessagePlain(MoveTemp(ExpectedString), ELogVerbosity::Warning, static_cast<EAutomationExpectedMessageFlags::MatchType>(CompareType), Occurrences);
 }
 
-uint32 FAutomationTestBase::ExtractAutomationTestFlags(FString InTagNotation)
+EAutomationTestFlags FAutomationTestBase::ExtractAutomationTestFlags(FString InTagNotation)
 {
-	uint32 Result = 0;
+	EAutomationTestFlags Result = EAutomationTestFlags::None;
 	TArray<FString> OutputParts;
 	InTagNotation
 		.Replace(TEXT("["), TEXT(""))
 		.Replace(TEXT("]"), TEXT(";"))
 		.ParseIntoArray(OutputParts, TEXT(";"), true);
-	for (auto it = OutputParts.begin(); it != OutputParts.end(); ++it)
+	for (const FString& Part : OutputParts)
 	{
-		auto Value = EAutomationTestFlags::FromString(*it);
-		if (Value != EAutomationTestFlags::None)
-		{
-			Result |= Value;
-		}
+		Result |= EAutomationTestFlags_GetTestFlagsMap().FindRef(Part, EAutomationTestFlags::None);
 	}
 	return Result;
 }
diff --git a/Engine/Source/Runtime/Core/Private/Tests/Serialization/CompactBinaryTest.cpp b/Engine/Source/Runtime/Core/Private/Tests/Serialization/CompactBinaryTest.cpp
--- a/Engine/Source/Runtime/Core/Private/Tests/Serialization/CompactBinaryTest.cpp
+++ b/Engine/Source/Runtime/Core/Private/Tests/Serialization/CompactBinaryTest.cpp
@@ -16,8 +16,8 @@
 
 ///////////////////////////////////////////////////////////////////////////////////////////////////////////////////////
 
-static constexpr EAutomationTestFlags::Type CompactBinaryTestFlags =
-	EAutomationTestFlags::Type(EAutomationTestFlags::ApplicationContextMask | EAutomationTestFlags::SmokeFilter);
+static constexpr EAutomationTestFlags CompactBinaryTestFlags =
+	EAutomationTestFlags::ApplicationContextMask | EAutomationTestFlags::SmokeFilter;
 
 ///////////////////////////////////////////////////////////////////////////////////////////////////////////////////////
 
diff --git a/Engine/Source/Runtime/Core/Public/Misc/AutomationTest.h b/Engine/Source/Runtime/Core/Public/Misc/AutomationTest.h
--- a/Engine/Source/Runtime/Core/Public/Misc/AutomationTest.h
+++ b/Engine/Source/Runtime/Core/Public/Misc/AutomationTest.h
@@ -75,78 +75,69 @@ class FAutomationTestBase;
 * Flags for specifying automation test requirements/behavior
 * Update GetTestFlagsMap when updating this enum.
 */
-struct EAutomationTestFlags
+enum class EAutomationTestFlags
 {
-	enum Type
-	{
-		None = 0x00000000,
-		//~ Application context required for the test
-		// Test is suitable for running within the editor
-		EditorContext = 0x00000001,
-		// Test is suitable for running within the client
-		ClientContext = 0x00000002,
-		// Test is suitable for running within the server
-		ServerContext = 0x00000004,
-		// Test is suitable for running within a commandlet
-		CommandletContext = 0x00000008,
-		// Test is suitable for running within program application (not editor, nor game)
-		ProgramContext = 0x00000010,
-		ApplicationContextMask = EditorContext | ClientContext | ServerContext | CommandletContext | ProgramContext,
-
-		//~ Features required for the test - not specifying means it is valid for any feature combination
-		// Test requires a non-null RHI to run correctly
-		NonNullRHI = 0x00000100,
-		// Test requires a user instigated session
-		RequiresUser = 0x00000200,
-		FeatureMask = NonNullRHI | RequiresUser,
-
-		//~ One-off flag to allow for fast disabling of tests without commenting code out
-		// Temp disabled and never returns for a filter
-		Disabled = 0x00010000,
-
-		//~ Priority of the test
-		// The highest priority possible. Showstopper/blocker.
-		CriticalPriority			= 0x00100000,
-		// High priority. Major feature functionality etc. 
-		HighPriority				= 0x00200000,
-		// Mask for High on SetMinimumPriority
-		HighPriorityAndAbove		= CriticalPriority | HighPriority,
-		// Medium Priority. Minor feature functionality, major generic content issues.
-		MediumPriority				= 0x00400000,
-		// Mask for Medium on SetMinimumPriority
-		MediumPriorityAndAbove		= CriticalPriority | HighPriority | MediumPriority,
-		// Low Priority. Minor content bugs. String errors. Etc.
-		LowPriority					= 0x00800000,
-		PriorityMask = CriticalPriority | HighPriority | MediumPriority | LowPriority,
-
-		//~ Speed of the test
-		//Super Fast Filter
-		SmokeFilter					= 0x01000000,
-		//Engine Level Test
-		EngineFilter				= 0x02000000,
-		//Product Level Test
-		ProductFilter				= 0x04000000,
-		//Performance Test
-		PerfFilter					= 0x08000000,
-		//Stress Test
-		StressFilter				= 0x10000000,
-		//Negative Test. For tests whose correct expected outcome is failure.
-		NegativeFilter				= 0x20000000,
-		FilterMask = SmokeFilter | EngineFilter | ProductFilter | PerfFilter | StressFilter | NegativeFilter
-	};
+	None = 0x00000000,
+	//~ Application context required for the test
+	// Test is suitable for running within the editor
+	EditorContext = 0x00000001,
+	// Test is suitable for running within the client
+	ClientContext = 0x00000002,
+	// Test is suitable for running within the server
+	ServerContext = 0x00000004,
+	// Test is suitable for running within a commandlet
+	CommandletContext = 0x00000008,
+	// Test is suitable for running within program application (not editor, nor game)
+	ProgramContext = 0x00000010,
+
+	//~ Features required for the test - not specifying means it is valid for any feature combination
+	// Test requires a non-null RHI to run correctly
+	NonNullRHI = 0x00000100,
+	// Test requires a user instigated session
+	RequiresUser = 0x00000200,
+
+	//~ One-off flag to allow for fast disabling of tests without commenting code out
+	// Temp disabled and never returns for a filter
+	Disabled = 0x00010000,
+
+	//~ Priority of the test
+	// The highest priority possible. Showstopper/blocker.
+	CriticalPriority			= 0x00100000,
+	// High priority. Major feature functionality etc. 
+	HighPriority				= 0x00200000,
+	// Medium Priority. Minor feature functionality, major generic content issues.
+	MediumPriority				= 0x00400000,
+	// Low Priority. Minor content bugs. String errors. Etc.
+	LowPriority					= 0x00800000,
+
+	//~ Speed of the test
+	//Super Fast Filter
+	SmokeFilter					= 0x01000000,
+	//Engine Level Test
+	EngineFilter				= 0x02000000,
+	//Product Level Test
+	ProductFilter				= 0x04000000,
+	//Performance Test
+	PerfFilter					= 0x08000000,
+	//Stress Test
+	StressFilter				= 0x10000000,
+	//Negative Test. For tests whose correct expected outcome is failure.
+	NegativeFilter				= 0x20000000,
+
+	// We need to deprecate and remove this once all usage has been replaced
+	ApplicationContextMask = EditorContext | ClientContext | ServerContext | CommandletContext | ProgramContext,
+};
 
-	static CORE_API const TMap<FString, Type>& GetTestFlagsMap();
+ENUM_CLASS_FLAGS(EAutomationTestFlags)
 
-	static const Type FromString(FString Name)
-	{
-		static auto FlagMap = GetTestFlagsMap();
-		if (FlagMap.Contains(Name))
-		{
-			return FlagMap[Name];
-		}
-		return Type::None;
-	}
-};
+constexpr EAutomationTestFlags EAutomationTestFlags_ApplicationContextMask = EAutomationTestFlags::EditorContext | EAutomationTestFlags::ClientContext | EAutomationTestFlags::ServerContext | EAutomationTestFlags::CommandletContext | EAutomationTestFlags::ProgramContext;
+constexpr EAutomationTestFlags EAutomationTestFlags_FeatureMask            = EAutomationTestFlags::NonNullRHI | EAutomationTestFlags::RequiresUser;
+constexpr EAutomationTestFlags EAutomationTestFlags_HighPriorityAndAbove   = EAutomationTestFlags::CriticalPriority | EAutomationTestFlags::HighPriority;
+constexpr EAutomationTestFlags EAutomationTestFlags_MediumPriorityAndAbove = EAutomationTestFlags::CriticalPriority | EAutomationTestFlags::HighPriority | EAutomationTestFlags::MediumPriority;
+constexpr EAutomationTestFlags EAutomationTestFlags_PriorityMask           = EAutomationTestFlags::CriticalPriority | EAutomationTestFlags::HighPriority | EAutomationTestFlags::MediumPriority | EAutomationTestFlags::LowPriority;
+constexpr EAutomationTestFlags EAutomationTestFlags_FilterMask             = EAutomationTestFlags::SmokeFilter | EAutomationTestFlags::EngineFilter | EAutomationTestFlags::ProductFilter | EAutomationTestFlags::PerfFilter | EAutomationTestFlags::StressFilter | EAutomationTestFlags::NegativeFilter;
+
+CORE_API const TMap<FString, EAutomationTestFlags>& EAutomationTestFlags_GetTestFlagsMap();
 
 /** Flags for indicating the matching type to use for an expected message */
 namespace EAutomationExpectedMessageFlags
@@ -277,11 +268,7 @@ class FAutomationTestInfo
 public:
 
 	// Default constructor
-	FAutomationTestInfo( )
-		: TestFlags( 0 )
-		, NumParticipantsRequired( 0 )
-		, NumDevicesCurrentlyRunningTest( 0 )
-	{}
+	FAutomationTestInfo() = default;
 
 	/**
 	 * Constructor
@@ -291,7 +278,7 @@ public:
 	 * @param	InTestFlag - Test flags
 	 * @param	InParameterName - optional parameter. e.g. asset name
 	 */
-	FAutomationTestInfo(const FString& InDisplayName, const FString& InFullTestPath, const FString& InTestName, const uint32 InTestFlags, const int32 InNumParticipantsRequired, const FString& InParameterName = FString(), const FString& InSourceFile = FString(), int32 InSourceFileLine = 0, const FString& InAssetPath = FString(), const FString& InOpenCommand = FString())
+	FAutomationTestInfo(const FString& InDisplayName, const FString& InFullTestPath, const FString& InTestName, const EAutomationTestFlags InTestFlags, const int32 InNumParticipantsRequired, const FString& InParameterName = FString(), const FString& InSourceFile = FString(), int32 InSourceFileLine = 0, const FString& InAssetPath = FString(), const FString& InOpenCommand = FString())
 		: DisplayName( InDisplayName )
 		, FullTestPath( InFullTestPath )
 		, TestName( InTestName )
@@ -303,7 +290,8 @@ public:
 		, TestFlags( InTestFlags )
 		, NumParticipantsRequired( InNumParticipantsRequired )
 		, NumDevicesCurrentlyRunningTest( 0 )
-	{}
+	{
+	}
 
 public:
 
@@ -312,7 +300,7 @@ public:
 	 *
 	 * @Param InTestFlags - the child test flag to add.
 	 */
-	void AddTestFlags( const uint32 InTestFlags)
+	void AddTestFlags( const EAutomationTestFlags InTestFlags)
 	{
 		TestFlags |= InTestFlags;
 	}
@@ -402,7 +390,7 @@ public:
 	 *
 	 * @return the test type.
 	 */
-	const uint32 GetTestFlags() const
+	const EAutomationTestFlags GetTestFlags() const
 	{
 		return TestFlags;
 	}
@@ -489,13 +477,13 @@ private:
 	FString OpenCommand;
 
 	/** The test flags. */
-	uint32 TestFlags;
+	EAutomationTestFlags TestFlags = EAutomationTestFlags::None;
 
 	/** The number of participants this test requires */
-	uint32 NumParticipantsRequired;
+	uint32 NumParticipantsRequired = 0;
 
 	/** The number of devices which have been given this test to run */
-	uint32 NumDevicesCurrentlyRunningTest;
+	uint32 NumDevicesCurrentlyRunningTest = 0;
 };
 
 
@@ -1101,7 +1089,7 @@ public:
 	/**
 	* Sets which set of tests to pull from.
 	*/
-	CORE_API void SetRequestedTestFilter(const uint32 InRequestedTestFlags);
+	CORE_API void SetRequestedTestFilter(const EAutomationTestFlags InRequestedTestFlags);
 	
 
 	/**
@@ -1376,7 +1364,7 @@ private:
 	TQueue< TSharedPtr<IAutomationNetworkCommand> > NetworkCommands;
 
 	/** Whether we are currently executing smoke tests for startup/commandlet to minimize log spam */
-	uint32 RequestedTestFilter;
+	EAutomationTestFlags RequestedTestFilter;
 
 	/** Time when the test began executing */
 	double StartTime;
@@ -1450,7 +1438,7 @@ public:
 	 *
 	 * @return	Automation test flags associated with the test
 	 */
-	virtual uint32 GetTestFlags() const = 0;
+	virtual EAutomationTestFlags GetTestFlags() const = 0;
 
 	/** Gets the C++ name of the test. */
 	FString GetTestName() const { return TestName; }
@@ -2435,7 +2423,7 @@ protected:
 	virtual void SetTestContext(FString Context) { TestParameterContext = Context; }
 
 	/** Extracts a combined EAutomationTestFlags value from a string representation using tag notation "[Filter_1]...[Filter_n][Tag_1]...[Tag_m]" */
-	CORE_API uint32 ExtractAutomationTestFlags(FString InTagNotation);
+	CORE_API EAutomationTestFlags ExtractAutomationTestFlags(FString InTagNotation);
 
 protected:
 
@@ -3858,16 +3846,16 @@ public: \
 	public: \
 		TClass( const FString& InName ) \
 		:TBaseClass( InName, false ) {\
-			static_assert((TFlags)&EAutomationTestFlags::ApplicationContextMask, "AutomationTest has no application flag.  It shouldn't run.  See AutomationTest.h."); \
-			static_assert(	(((TFlags)&EAutomationTestFlags::FilterMask) == EAutomationTestFlags::SmokeFilter) || \
-							(((TFlags)&EAutomationTestFlags::FilterMask) == EAutomationTestFlags::EngineFilter) || \
-							(((TFlags)&EAutomationTestFlags::FilterMask) == EAutomationTestFlags::ProductFilter) || \
-							(((TFlags)&EAutomationTestFlags::FilterMask) == EAutomationTestFlags::PerfFilter) || \
-							(((TFlags)&EAutomationTestFlags::FilterMask) == EAutomationTestFlags::StressFilter) || \
-							(((TFlags)&EAutomationTestFlags::FilterMask) == EAutomationTestFlags::NegativeFilter), \
+			static_assert(!!((TFlags) & EAutomationTestFlags_ApplicationContextMask), "AutomationTest has no application flag.  It shouldn't run.  See AutomationTest.h."); \
+			static_assert(	!!(((TFlags) & EAutomationTestFlags_FilterMask) == EAutomationTestFlags::SmokeFilter) || \
+							!!(((TFlags) & EAutomationTestFlags_FilterMask) == EAutomationTestFlags::EngineFilter) || \
+							!!(((TFlags) & EAutomationTestFlags_FilterMask) == EAutomationTestFlags::ProductFilter) || \
+							!!(((TFlags) & EAutomationTestFlags_FilterMask) == EAutomationTestFlags::PerfFilter) || \
+							!!(((TFlags) & EAutomationTestFlags_FilterMask) == EAutomationTestFlags::StressFilter) || \
+							!!(((TFlags) & EAutomationTestFlags_FilterMask) == EAutomationTestFlags::NegativeFilter), \
 							"All AutomationTests must have exactly 1 filter type specified.  See AutomationTest.h."); \
 		} \
-		virtual uint32 GetTestFlags() const override { return TFlags; } \
+		virtual EAutomationTestFlags GetTestFlags() const override { return TFlags; } \
 		virtual bool IsStressTest() const { return false; } \
 		virtual uint32 GetRequiredDeviceNum() const override { return 1; } \
 		virtual FString GetTestSourceFileName() const override { return FileName; } \
@@ -3890,16 +3878,16 @@ public: \
 	public: \
 		TClass( const FString& InName ) \
 		:TBaseClass( InName, true ) { \
-			static_assert((TFlags)&EAutomationTestFlags::ApplicationContextMask, "AutomationTest has no application flag.  It shouldn't run.  See AutomationTest.h."); \
-			static_assert(	(((TFlags)&EAutomationTestFlags::FilterMask) == EAutomationTestFlags::SmokeFilter) || \
-							(((TFlags)&EAutomationTestFlags::FilterMask) == EAutomationTestFlags::EngineFilter) || \
-							(((TFlags)&EAutomationTestFlags::FilterMask) == EAutomationTestFlags::ProductFilter) || \
-							(((TFlags)&EAutomationTestFlags::FilterMask) == EAutomationTestFlags::PerfFilter) || \
-							(((TFlags)&EAutomationTestFlags::FilterMask) == EAutomationTestFlags::StressFilter) || \
-							(((TFlags)&EAutomationTestFlags::FilterMask) == EAutomationTestFlags::NegativeFilter), \
+			static_assert(!!((TFlags) & EAutomationTestFlags_ApplicationContextMask), "AutomationTest has no application flag.  It shouldn't run.  See AutomationTest.h."); \
+			static_assert(	!!(((TFlags) & EAutomationTestFlags_FilterMask) == EAutomationTestFlags::SmokeFilter) || \
+							!!(((TFlags) & EAutomationTestFlags_FilterMask) == EAutomationTestFlags::EngineFilter) || \
+							!!(((TFlags) & EAutomationTestFlags_FilterMask) == EAutomationTestFlags::ProductFilter) || \
+							!!(((TFlags) & EAutomationTestFlags_FilterMask) == EAutomationTestFlags::PerfFilter) || \
+							!!(((TFlags) & EAutomationTestFlags_FilterMask) == EAutomationTestFlags::StressFilter) || \
+							!!(((TFlags) & EAutomationTestFlags_FilterMask) == EAutomationTestFlags::NegativeFilter), \
 							"All AutomationTests must have exactly 1 filter type specified.  See AutomationTest.h."); \
 		} \
-		virtual uint32 GetTestFlags() const override { return ((TFlags) & ~(EAutomationTestFlags::SmokeFilter)); } \
+		virtual EAutomationTestFlags GetTestFlags() const override { return ((TFlags) & ~(EAutomationTestFlags::SmokeFilter)); } \
 		virtual bool IsStressTest() const { return true; } \
 		virtual uint32 GetRequiredDeviceNum() const override { return 1; } \
 		virtual FString GetTestSourceFileName() const override { return FileName; } \
@@ -3916,16 +3904,16 @@ public: \
 	public: \
 		TClass( const FString& InName ) \
 		:TBaseClass( InName, false ) { \
-			static_assert((TFlags)&EAutomationTestFlags::ApplicationContextMask, "AutomationTest has no application flag.  It shouldn't run.  See AutomationTest.h."); \
-			static_assert(	(((TFlags)&EAutomationTestFlags::FilterMask) == EAutomationTestFlags::SmokeFilter) || \
-							(((TFlags)&EAutomationTestFlags::FilterMask) == EAutomationTestFlags::EngineFilter) || \
-							(((TFlags)&EAutomationTestFlags::FilterMask) == EAutomationTestFlags::ProductFilter) || \
-							(((TFlags)&EAutomationTestFlags::FilterMask) == EAutomationTestFlags::PerfFilter) || \
-							(((TFlags)&EAutomationTestFlags::FilterMask) == EAutomationTestFlags::StressFilter) || \
-							(((TFlags)&EAutomationTestFlags::FilterMask) == EAutomationTestFlags::NegativeFilter), \
+			static_assert(!!((TFlags) & EAutomationTestFlags_ApplicationContextMask), "AutomationTest has no application flag.  It shouldn't run.  See AutomationTest.h."); \
+			static_assert(	!!(((TFlags) & EAutomationTestFlags_FilterMask) == EAutomationTestFlags::SmokeFilter) || \
+							!!(((TFlags) & EAutomationTestFlags_FilterMask) == EAutomationTestFlags::EngineFilter) || \
+							!!(((TFlags) & EAutomationTestFlags_FilterMask) == EAutomationTestFlags::ProductFilter) || \
+							!!(((TFlags) & EAutomationTestFlags_FilterMask) == EAutomationTestFlags::PerfFilter) || \
+							!!(((TFlags) & EAutomationTestFlags_FilterMask) == EAutomationTestFlags::StressFilter) || \
+							!!(((TFlags) & EAutomationTestFlags_FilterMask) == EAutomationTestFlags::NegativeFilter), \
 							"All AutomationTests must have exactly 1 filter type specified.  See AutomationTest.h."); \
 		} \
-		virtual uint32 GetTestFlags() const override { return ((TFlags) & ~(EAutomationTestFlags::EditorContext | EAutomationTestFlags::CommandletContext | EAutomationTestFlags::SmokeFilter)); } \
+		virtual EAutomationTestFlags GetTestFlags() const override { return ((TFlags) & ~(EAutomationTestFlags::EditorContext | EAutomationTestFlags::CommandletContext | EAutomationTestFlags::SmokeFilter)); } \
 		virtual uint32 GetRequiredDeviceNum() const override { return NumParticipants; } \
 		virtual FString GetTestSourceFileName() const override { return FileName; } \
 		virtual int32 GetTestSourceFileLine() const override { return LineNumber; } \
@@ -3945,16 +3933,16 @@ public: \
 	public: \
 		TClass( const FString& InName ) \
 		:FBDDAutomationTestBase( InName, false ) {\
-			static_assert((TFlags)&EAutomationTestFlags::ApplicationContextMask, "AutomationTest has no application flag.  It shouldn't run.  See AutomationTest.h."); \
-			static_assert(	(((TFlags)&EAutomationTestFlags::FilterMask) == EAutomationTestFlags::SmokeFilter) || \
-							(((TFlags)&EAutomationTestFlags::FilterMask) == EAutomationTestFlags::EngineFilter) || \
-							(((TFlags)&EAutomationTestFlags::FilterMask) == EAutomationTestFlags::ProductFilter) || \
-							(((TFlags)&EAutomationTestFlags::FilterMask) == EAutomationTestFlags::PerfFilter) || \
-							(((TFlags)&EAutomationTestFlags::FilterMask) == EAutomationTestFlags::StressFilter) || \
-							(((TFlags)&EAutomationTestFlags::FilterMask) == EAutomationTestFlags::NegativeFilter), \
+			static_assert(!!((TFlags) & EAutomationTestFlags_ApplicationContextMask), "AutomationTest has no application flag.  It shouldn't run.  See AutomationTest.h."); \
+			static_assert(	!!((((TFlags) & EAutomationTestFlags_FilterMask) == EAutomationTestFlags::SmokeFilter) || \
+							!!((((TFlags) & EAutomationTestFlags_FilterMask) == EAutomationTestFlags::EngineFilter) || \
+							!!((((TFlags) & EAutomationTestFlags_FilterMask) == EAutomationTestFlags::ProductFilter) || \
+							!!((((TFlags) & EAutomationTestFlags_FilterMask) == EAutomationTestFlags::PerfFilter) || \
+							!!((((TFlags) & EAutomationTestFlags_FilterMask) == EAutomationTestFlags::StressFilter) || \
+							!!((((TFlags) & EAutomationTestFlags_FilterMask) == EAutomationTestFlags::NegativeFilter), \
 							"All AutomationTests must have exactly 1 filter type specified.  See AutomationTest.h."); \
 		} \
-		virtual uint32 GetTestFlags() const override { return TFlags; } \
+		virtual EAutomationTestFlags GetTestFlags() const override { return TFlags; } \
 		virtual bool IsStressTest() const { return false; } \
 		virtual uint32 GetRequiredDeviceNum() const override { return 1; } \
 		virtual FString GetTestSourceFileName() const override { return FileName; } \
@@ -3972,16 +3960,16 @@ public: \
 	public: \
 		TClass( const FString& InName ) \
 		: FAutomationSpecBase( InName, false ) {\
-			static_assert((TFlags)&EAutomationTestFlags::ApplicationContextMask, "AutomationTest has no application flag.  It shouldn't run.  See AutomationTest.h."); \
-			static_assert(	(((TFlags)&EAutomationTestFlags::FilterMask) == EAutomationTestFlags::SmokeFilter) || \
-							(((TFlags)&EAutomationTestFlags::FilterMask) == EAutomationTestFlags::EngineFilter) || \
-							(((TFlags)&EAutomationTestFlags::FilterMask) == EAutomationTestFlags::ProductFilter) || \
-							(((TFlags)&EAutomationTestFlags::FilterMask) == EAutomationTestFlags::PerfFilter) || \
-							(((TFlags)&EAutomationTestFlags::FilterMask) == EAutomationTestFlags::StressFilter) || \
-							(((TFlags)&EAutomationTestFlags::FilterMask) == EAutomationTestFlags::NegativeFilter), \
+			static_assert(!!((TFlags) & EAutomationTestFlags_ApplicationContextMask), "AutomationTest has no application flag.  It shouldn't run.  See AutomationTest.h."); \
+			static_assert(	!!(((TFlags) & EAutomationTestFlags_FilterMask) == EAutomationTestFlags::SmokeFilter) || \
+							!!(((TFlags) & EAutomationTestFlags_FilterMask) == EAutomationTestFlags::EngineFilter) || \
+							!!(((TFlags) & EAutomationTestFlags_FilterMask) == EAutomationTestFlags::ProductFilter) || \
+							!!(((TFlags) & EAutomationTestFlags_FilterMask) == EAutomationTestFlags::PerfFilter) || \
+							!!(((TFlags) & EAutomationTestFlags_FilterMask) == EAutomationTestFlags::StressFilter) || \
+							!!(((TFlags) & EAutomationTestFlags_FilterMask) == EAutomationTestFlags::NegativeFilter), \
 							"All AutomationTests must have exactly 1 filter type specified.  See AutomationTest.h."); \
 		} \
-		virtual uint32 GetTestFlags() const override { return TFlags; } \
+		virtual EAutomationTestFlags GetTestFlags() const override { return TFlags; } \
         using FAutomationSpecBase::GetTestSourceFileName; \
 		virtual FString GetTestSourceFileName() const override { return FileName; } \
         using FAutomationSpecBase::GetTestSourceFileLine; \
@@ -3999,16 +3987,16 @@ public: \
 	public: \
 		TClass( const FString& InName ) \
 		: FAutomationSpecBase( InName, false ) {\
-			static_assert((TFlags)&EAutomationTestFlags::ApplicationContextMask, "AutomationTest has no application flag.  It shouldn't run.  See AutomationTest.h."); \
-			static_assert(	(((TFlags)&EAutomationTestFlags::FilterMask) == EAutomationTestFlags::SmokeFilter) || \
-							(((TFlags)&EAutomationTestFlags::FilterMask) == EAutomationTestFlags::EngineFilter) || \
-							(((TFlags)&EAutomationTestFlags::FilterMask) == EAutomationTestFlags::ProductFilter) || \
-							(((TFlags)&EAutomationTestFlags::FilterMask) == EAutomationTestFlags::PerfFilter) || \
-							(((TFlags)&EAutomationTestFlags::FilterMask) == EAutomationTestFlags::StressFilter) || \
-							(((TFlags)&EAutomationTestFlags::FilterMask) == EAutomationTestFlags::NegativeFilter), \
+			static_assert(!!((TFlags) & EAutomationTestFlags_ApplicationContextMask), "AutomationTest has no application flag.  It shouldn't run.  See AutomationTest.h."); \
+			static_assert(	!!(((TFlags) & EAutomationTestFlags_FilterMask) == EAutomationTestFlags::SmokeFilter) || \
+							!!(((TFlags) & EAutomationTestFlags_FilterMask) == EAutomationTestFlags::EngineFilter) || \
+							!!(((TFlags) & EAutomationTestFlags_FilterMask) == EAutomationTestFlags::ProductFilter) || \
+							!!(((TFlags) & EAutomationTestFlags_FilterMask) == EAutomationTestFlags::PerfFilter) || \
+							!!(((TFlags) & EAutomationTestFlags_FilterMask) == EAutomationTestFlags::StressFilter) || \
+							!!(((TFlags) & EAutomationTestFlags_FilterMask) == EAutomationTestFlags::NegativeFilter), \
 							"All AutomationTests must have exactly 1 filter type specified.  See AutomationTest.h."); \
 		} \
-		virtual uint32 GetTestFlags() const override { return TFlags; } \
+		virtual EAutomationTestFlags GetTestFlags() const override { return TFlags; } \
 		using FAutomationSpecBase::GetTestSourceFileName; \
 		virtual FString GetTestSourceFileName() const override { return FileName; } \
 		using FAutomationSpecBase::GetTestSourceFileLine; \
diff --git a/Engine/Source/Runtime/Core/Public/Misc/LowLevelTestAdapter.h b/Engine/Source/Runtime/Core/Public/Misc/LowLevelTestAdapter.h
--- a/Engine/Source/Runtime/Core/Public/Misc/LowLevelTestAdapter.h
+++ b/Engine/Source/Runtime/Core/Public/Misc/LowLevelTestAdapter.h
@@ -57,16 +57,16 @@ FString CaptureExpressionsAndValues(const FString& InExpressions, ArgTypes&&...
 			{ \
 				TestFlags = ExtractAutomationTestFlags(TFlags); \
 				PrettyNameDotNotation = FString(PrettyName).Replace(TEXT("::"), TEXT(".")); \
-				if (!(TestFlags & EAutomationTestFlags::ApplicationContextMask)) \
+				if (!(TestFlags & EAutomationTestFlags_ApplicationContextMask)) \
 				{ \
-					TestFlags |= EAutomationTestFlags::ApplicationContextMask; \
+					TestFlags |= EAutomationTestFlags_ApplicationContextMask; \
 				} \
-				if (!(TestFlags & EAutomationTestFlags::FilterMask)) \
+				if (!(TestFlags & EAutomationTestFlags_FilterMask)) \
 				{ \
 					TestFlags |= EAutomationTestFlags::EngineFilter; \
 				} \
 			} \
-			virtual uint32 GetTestFlags() const override { return TestFlags; } \
+			virtual EAutomationTestFlags GetTestFlags() const override { return TestFlags; } \
 			virtual bool IsStressTest() const { return false; } \
 			virtual uint32 GetRequiredDeviceNum() const override { return 1; } \
 			virtual FString GetTestSourceFileName() const override { return FileName; } \
@@ -84,7 +84,7 @@ FString CaptureExpressionsAndValues(const FString& InExpressions, ArgTypes&&...
 			} \
 			virtual FString GetBeautifiedTestName() const override { return PrettyNameDotNotation; } \
 		private:\
-			uint32 TestFlags; \
+			EAutomationTestFlags TestFlags; \
 			FString PrettyNameDotNotation; \
 		};
 
diff --git a/Engine/Source/Runtime/Engine/Private/Tests/AnalyticsTest.cpp b/Engine/Source/Runtime/Engine/Private/Tests/AnalyticsTest.cpp
--- a/Engine/Source/Runtime/Engine/Private/Tests/AnalyticsTest.cpp
+++ b/Engine/Source/Runtime/Engine/Private/Tests/AnalyticsTest.cpp
@@ -10,7 +10,7 @@
 /**
 * Artificial Record Event for analytics - Simulates the engine startup simulation.  
 */
-IMPLEMENT_SIMPLE_AUTOMATION_TEST(FAnalyticStartUpSimTest, "System.Engine.Analytic.Record Event - Simulate Program Start", EAutomationTestFlags::FeatureMask | EAutomationTestFlags::ClientContext | EAutomationTestFlags::EngineFilter)
+IMPLEMENT_SIMPLE_AUTOMATION_TEST(FAnalyticStartUpSimTest, "System.Engine.Analytic.Record Event - Simulate Program Start", EAutomationTestFlags_FeatureMask | EAutomationTestFlags::ClientContext | EAutomationTestFlags::EngineFilter)
 
 bool FAnalyticStartUpSimTest::RunTest(const FString& Parameters)
 {
@@ -72,7 +72,7 @@ bool FAnalyticStartUpSimTest::RunTest(const FString& Parameters)
 /**
 * FAnalyticsEventAttribute Unit Test.
 */
-IMPLEMENT_SIMPLE_AUTOMATION_TEST(FAnalyticsEventAttributeUnitTest, "System.Engine.Analytic.EventAttribute Struct Unit Test", EAutomationTestFlags::FeatureMask | EAutomationTestFlags::ClientContext | EAutomationTestFlags::EngineFilter)
+IMPLEMENT_SIMPLE_AUTOMATION_TEST(FAnalyticsEventAttributeUnitTest, "System.Engine.Analytic.EventAttribute Struct Unit Test", EAutomationTestFlags_FeatureMask | EAutomationTestFlags::ClientContext | EAutomationTestFlags::EngineFilter)
 
 bool FAnalyticsEventAttributeUnitTest::RunTest(const FString& Parameters)
 {
diff --git a/Engine/Source/Runtime/Engine/Private/Tests/EngineAutomationTests.cpp b/Engine/Source/Runtime/Engine/Private/Tests/EngineAutomationTests.cpp
--- a/Engine/Source/Runtime/Engine/Private/Tests/EngineAutomationTests.cpp
+++ b/Engine/Source/Runtime/Engine/Private/Tests/EngineAutomationTests.cpp
@@ -23,10 +23,10 @@
 
 namespace
 {
-	UWorld* GetSimpleEngineAutomationTestGameWorld(const int32 TestFlags)
+	UWorld* GetSimpleEngineAutomationTestGameWorld(const EAutomationTestFlags TestFlags)
 	{
 		// Accessing the game world is only valid for game-only 
-		check((TestFlags & EAutomationTestFlags::ApplicationContextMask) == EAutomationTestFlags::ClientContext);
+		check((TestFlags & EAutomationTestFlags_ApplicationContextMask) == EAutomationTestFlags::ClientContext);
 		check(GEngine->GetWorldContexts().Num() == 1);
 		check(GEngine->GetWorldContexts()[0].WorldType == EWorldType::Game);
 
diff --git a/Engine/Source/Runtime/Engine/Private/Tests/PropertyPairsMapTests.cpp b/Engine/Source/Runtime/Engine/Private/Tests/PropertyPairsMapTests.cpp
--- a/Engine/Source/Runtime/Engine/Private/Tests/PropertyPairsMapTests.cpp
+++ b/Engine/Source/Runtime/Engine/Private/Tests/PropertyPairsMapTests.cpp
@@ -9,7 +9,7 @@
 
 namespace PropertyPairsMapTests
 {
-	constexpr const uint32 TestFlags = EAutomationTestFlags::ApplicationContextMask | EAutomationTestFlags::EngineFilter;
+	constexpr const EAutomationTestFlags TestFlags = EAutomationTestFlags::ApplicationContextMask | EAutomationTestFlags::EngineFilter;
 
 	IMPLEMENT_SIMPLE_AUTOMATION_TEST(FPropertyPairsMapTests, TEST_NAME_ROOT, TestFlags)
 	bool FPropertyPairsMapTests::RunTest(const FString& Parameters)
diff --git a/Engine/Source/Runtime/Engine/Private/Tests/ReferenceClusterTest.cpp b/Engine/Source/Runtime/Engine/Private/Tests/ReferenceClusterTest.cpp
--- a/Engine/Source/Runtime/Engine/Private/Tests/ReferenceClusterTest.cpp
+++ b/Engine/Source/Runtime/Engine/Private/Tests/ReferenceClusterTest.cpp
@@ -9,7 +9,7 @@
 
 namespace ReferenceClusterTests
 {
-	constexpr const uint32 TestFlags = EAutomationTestFlags::ApplicationContextMask | EAutomationTestFlags::EngineFilter;
+	constexpr const EAutomationTestFlags TestFlags = EAutomationTestFlags::ApplicationContextMask | EAutomationTestFlags::EngineFilter;
 
 	IMPLEMENT_SIMPLE_AUTOMATION_TEST(FReferenceClusterTest, TEST_NAME_ROOT, TestFlags)
 
diff --git a/Engine/Source/Runtime/Engine/Private/Tests/Texture2DTests.cpp b/Engine/Source/Runtime/Engine/Private/Tests/Texture2DTests.cpp
--- a/Engine/Source/Runtime/Engine/Private/Tests/Texture2DTests.cpp
+++ b/Engine/Source/Runtime/Engine/Private/Tests/Texture2DTests.cpp
@@ -9,7 +9,7 @@
 namespace Texture2DTest
 {
 
-constexpr const uint32 TestFlags = EAutomationTestFlags::ApplicationContextMask | EAutomationTestFlags::EngineFilter;
+constexpr const EAutomationTestFlags TestFlags = EAutomationTestFlags::ApplicationContextMask | EAutomationTestFlags::EngineFilter;
 
 // A simple test to make sure that basic functionality in UTexture2D::CreateTransient works as it seems to be a 
 // fairly uncommon code path in our samples/test games etc.
diff --git a/Engine/Source/Runtime/Engine/Private/Tests/TransactionDiffingTests.cpp b/Engine/Source/Runtime/Engine/Private/Tests/TransactionDiffingTests.cpp
--- a/Engine/Source/Runtime/Engine/Private/Tests/TransactionDiffingTests.cpp
+++ b/Engine/Source/Runtime/Engine/Private/Tests/TransactionDiffingTests.cpp
@@ -17,7 +17,7 @@ void UTransactionDiffingTestObject::Serialize(FStructuredArchive::FRecord Record
 namespace TransactionDiffingTests
 {
 
-constexpr const uint32 TestFlags = EAutomationTestFlags::ApplicationContextMask | EAutomationTestFlags::EngineFilter;
+constexpr const EAutomationTestFlags TestFlags = EAutomationTestFlags::ApplicationContextMask | EAutomationTestFlags::EngineFilter;
 
 IMPLEMENT_SIMPLE_AUTOMATION_TEST(FEditPropertyDataTest, "System.Engine.Transactions.EditPropertyData", TestFlags)
 bool FEditPropertyDataTest::RunTest(const FString& Parameters)
diff --git a/Engine/Source/Runtime/Engine/Private/VisualLogger/VisualLoggerAutomationTests.cpp b/Engine/Source/Runtime/Engine/Private/VisualLogger/VisualLoggerAutomationTests.cpp
--- a/Engine/Source/Runtime/Engine/Private/VisualLogger/VisualLoggerAutomationTests.cpp
+++ b/Engine/Source/Runtime/Engine/Private/VisualLogger/VisualLoggerAutomationTests.cpp
@@ -12,10 +12,10 @@
 
 namespace
 {
-	UWorld* GetSimpleEngineAutomationTestWorld(const int32 TestFlags)
+	UWorld* GetSimpleEngineAutomationTestWorld(const EAutomationTestFlags TestFlags)
 	{
-		// Accessing the game world is only valid for game-only 
-		if (((TestFlags & EAutomationTestFlags::EditorContext) || (TestFlags & EAutomationTestFlags::ClientContext)) == false)
+		// Accessing the game world is only valid for game-only
+		if (!(TestFlags & EAutomationTestFlags::EditorContext) && !(TestFlags & EAutomationTestFlags::ClientContext))
 		{
 			return nullptr;
 		}
diff --git a/Engine/Source/Runtime/MeshDescription/Private/Tests/MeshElementArrayTests.cpp b/Engine/Source/Runtime/MeshDescription/Private/Tests/MeshElementArrayTests.cpp
--- a/Engine/Source/Runtime/MeshDescription/Private/Tests/MeshElementArrayTests.cpp
+++ b/Engine/Source/Runtime/MeshDescription/Private/Tests/MeshElementArrayTests.cpp
@@ -9,7 +9,7 @@
 
 #if WITH_DEV_AUTOMATION_TESTS
 
-constexpr const uint32 TestFlags = EAutomationTestFlags::ApplicationContextMask | EAutomationTestFlags::EngineFilter;
+constexpr const EAutomationTestFlags TestFlags = EAutomationTestFlags::ApplicationContextMask | EAutomationTestFlags::EngineFilter;
 
 /**
  * A very basic type that can be used as the ElementIDType with TMeshElementArray in our tests
