[QueuedThreadPool]
 - Remove busywait logic

#rnx

diff --git a/Engine/Source/Runtime/Core/Public/Misc/QueuedThreadPoolWrapper.h b/Engine/Source/Runtime/Core/Public/Misc/QueuedThreadPoolWrapper.h
--- a/Engine/Source/Runtime/Core/Public/Misc/QueuedThreadPoolWrapper.h
+++ b/Engine/Source/Runtime/Core/Public/Misc/QueuedThreadPoolWrapper.h
@@ -4,6 +4,7 @@
 
 #include "Async/Fundamental/Scheduler.h"
 #include "Async/Fundamental/Task.h"
+#include "Async/ManualResetEvent.h"
 #include "Async/TaskGraphInterfaces.h"
 #include "Containers/Array.h"
 #include "Containers/Map.h"
@@ -315,7 +316,10 @@ private:
 
 	void OnTaskCompleted(IQueuedWork* InQueuedWork)
 	{
-		--TaskCount;
+		if (--TaskCount == 0 && bIsExiting)
+		{
+			Finished.Notify();
+		}
 	}
 
 	int32 GetNumThreads() const override
@@ -346,22 +350,16 @@ protected:
 	{
 		bIsExiting = true;
 
-		if (LowLevelTasks::FScheduler::Get().IsWorkerThread())
+		if (TaskCount != 0)
 		{
-			LowLevelTasks::BusyWaitUntil([this]() { return TaskCount == 0; });
-		}
-		else
-		{
-			while (TaskCount != 0)
-			{
-				FPlatformProcess::Sleep(0.01f);
-			}
+			Finished.Wait();
 		}
 	}
 private:
 	TFunction<ENamedThreads::Type (EQueuedWorkPriority)> PriorityMapper;
-	TAtomic<uint32> TaskCount;
-	TAtomic<bool> bIsExiting;
+	TAtomic<uint32>       TaskCount;
+	TAtomic<bool>         bIsExiting;
+	UE::FManualResetEvent Finished;
 };
 
 /** ThreadPool wrapper implementation allowing to schedule thread-pool tasks on the the low level backend which is also used by the taskgraph.
@@ -457,9 +455,15 @@ private:
 
 	void FinalizeExecution()
 	{
-		TaskCount.fetch_sub(1, std::memory_order_release);
-		bool bWakeUpWorker = false;
-		ScheduleTasks(bWakeUpWorker);
+		if (TaskCount.fetch_sub(1, std::memory_order_release) == 1 && bIsExiting)
+		{
+			Finished.Notify();
+		}
+		else
+		{
+			bool bWakeUpWorker = false;
+			ScheduleTasks(bWakeUpWorker);
+		}
 	}
 
 	void AddQueuedWork(IQueuedWork* InQueuedWork, EQueuedWorkPriority InPriority = EQueuedWorkPriority::Normal) override
@@ -533,16 +537,9 @@ protected:
 			verifySlow(Scheduler->TryLaunch(QueuedWork->Task, LowLevelTasks::EQueuePreference::GlobalQueuePreference));
 		}
 
-		if (Scheduler->IsWorkerThread())
+		if (TaskCount != 0)
 		{
-			Scheduler->BusyWaitUntil([this]() { return TaskCount == 0; });
-		}
-		else
-		{
-			while (TaskCount != 0)
-			{
-				FPlatformProcess::Sleep(0.01f);
-			}
+			Finished.Wait();
 		}
 	}
 
@@ -573,4 +570,5 @@ private:
 	std::atomic_uint TaskCount{0};
 	std::atomic_bool bIsExiting{false};
 	std::atomic_bool bIsPaused{false};
+	UE::FManualResetEvent Finished;
 };
