Add FAutoConsoleCommand constructor overload for FConsoleCommandWithArgsAndOutputDeviceDelegate

diff --git a/Engine/Source/Runtime/Core/Public/HAL/IConsoleManager.h b/Engine/Source/Runtime/Core/Public/HAL/IConsoleManager.h
--- a/Engine/Source/Runtime/Core/Public/HAL/IConsoleManager.h
+++ b/Engine/Source/Runtime/Core/Public/HAL/IConsoleManager.h
@@ -1897,6 +1897,19 @@ public:
 	{
 	}
 
+    /**
+    * Register a console command that takes an output device
+    *
+    * @param	Name		The name of this command (must not be nullptr)
+    * @param	Help		Help text for this command
+    * @param	Command		The user function to call when this command is executed
+    * @param	Flags		Optional flags bitmask
+    */
+    FAutoConsoleCommand(const TCHAR* Name, const TCHAR* Help, const FConsoleCommandWithArgsAndOutputDeviceDelegate& Command, uint32 Flags = ECVF_Default)
+        : FAutoConsoleObject(IConsoleManager::Get().RegisterConsoleCommand(Name, Help, Command, Flags))
+    {
+    }
+
 	/**
 	* Register a console command that takes a world argument
 	*
@@ -1935,6 +1948,18 @@ public:
 	{
 	}
 
+	FAutoConsoleCommand(const TCHAR* Name, const TCHAR* Help, const FConsoleCommandWithOutputDeviceDelegate& Command, uint32 Flags = ECVF_Default)
+	{
+	}
+
+	FAutoConsoleCommand(const TCHAR* Name, const TCHAR* Help, const FConsoleCommandWithArgsAndOutputDeviceDelegate& Command, uint32 Flags = ECVF_Default)
+	{
+	}
+
+	FAutoConsoleCommand(const TCHAR* Name, const TCHAR* Help, const FConsoleCommandWithWorldDelegate& Command, uint32 Flags = ECVF_Default)
+	{
+	}
+
 	FAutoConsoleCommand(const TCHAR* Name, const TCHAR* Help, const FConsoleCommandWithWorldArgsAndOutputDeviceDelegate& Command, uint32 Flags = ECVF_Default)
 	{
 	}
