Removed rehashing from key iteration, and mutability from TSet members, making TSet thread-safe for concurrent reads.

diff --git a/Engine/Source/Runtime/Core/Public/Containers/Set.h b/Engine/Source/Runtime/Core/Public/Containers/Set.h
--- a/Engine/Source/Runtime/Core/Public/Containers/Set.h
+++ b/Engine/Source/Runtime/Core/Public/Containers/Set.h
@@ -1390,8 +1390,8 @@ private:
 
 	ElementArrayType Elements;
 
-	mutable HashType Hash;
-	mutable int32	 HashSize;
+	HashType Hash;
+	int32	 HashSize;
 
 public:
 	void WriteMemoryImage(FMemoryImageWriter& Writer) const
@@ -1500,7 +1500,7 @@ private:
 	 * @param AllowShrinking - If the hash is allowed to shrink.
 	 * @return true if the set was rehashed.
 	 */
-	bool ConditionalRehash(int32 NumHashedElements, EAllowShrinking AllowShrinking) const
+	bool ConditionalRehash(int32 NumHashedElements, EAllowShrinking AllowShrinking)
 	{
 		// Calculate the desired hash size for the specified number of elements.
 		const int32 DesiredHashSize = Allocator::GetNumberOfHashBuckets(NumHashedElements);
@@ -1516,7 +1516,7 @@ private:
 	}
 
 	/** Resizes the hash. */
-	void Rehash() const
+	void Rehash()
 	{
 		// Free the old hash.
 		Hash.ResizeAllocation(0,0,sizeof(FSetElementId));
@@ -1622,8 +1622,6 @@ private:
 			, Key  (InKey) //-V1041
 			, Index(INDEX_NONE)
 		{
-			// The set's hash needs to be initialized to find the elements with the specified key.
-			Set.ConditionalRehash(Set.Elements.Num(), EAllowShrinking::No);
 			if (Set.HashSize)
 			{
 				NextIndex = Set.GetTypedHash(KeyFuncs::GetKeyHash(Key)).Index;
@@ -2104,8 +2102,8 @@ private:
 	typedef typename Allocator::HashAllocator::template ForElementType<FSetElementId> HashType;
 
 	ElementArrayType Elements;
-	mutable HashType Hash;
-	mutable int32    HashSize;
+	HashType         Hash;
+	int32            HashSize;
 
 	FORCEINLINE FSetElementId& GetTypedHash(int32 HashIndex) const
 	{
