Reverted some changes that were accidentally checked in with the EAutomationTestFlags_ApplicationContextMask change.

diff --git a/Engine/Source/Runtime/Core/Private/Containers/StringTest.cpp b/Engine/Source/Runtime/Core/Private/Containers/StringTest.cpp
--- a/Engine/Source/Runtime/Core/Private/Containers/StringTest.cpp
+++ b/Engine/Source/Runtime/Core/Private/Containers/StringTest.cpp
@@ -440,10 +440,10 @@ bool FStringConstructorWithSlackTest::RunTest(const FString& Parameters)
 
 		const SIZE_T ExpectedCapacity = FMemory::QuantizeSize(NumElements * sizeof(TCHAR));
 
-		FString StringFromTChar = FString::ConstructWithSlack(TestString, ExtraSlack);
+		FString StringFromTChar(TestString, ExtraSlack);
 		TestEqual(TEXT("(TCHAR: Valid string with valid slack) resulting capacity"), StringFromTChar.GetAllocatedSize(), ExpectedCapacity);
 
-		FString StringFromAscii = FString::ConstructWithSlack(TestAsciiString, ExtraSlack);
+		FString StringFromAscii(TestAsciiString, ExtraSlack);
 		TestEqual(TEXT("(ASCII: Valid string with valid slack) resulting capacity"), StringFromAscii.GetAllocatedSize(), ExpectedCapacity);
 
 		FString StringFromFStringView(FStringView(TestString), ExtraSlack);
@@ -463,10 +463,10 @@ bool FStringConstructorWithSlackTest::RunTest(const FString& Parameters)
 
 		const SIZE_T ExpectedCapacity = FMemory::QuantizeSize(NumElements * sizeof(TCHAR));
 
-		FString StringFromTChar = FString::ConstructWithSlack(TestString, ExtraSlack);
+		FString StringFromTChar(TestString, ExtraSlack);
 		TestEqual(TEXT("(TCHAR: Valid string with zero slack) resulting capacity"), StringFromTChar.GetAllocatedSize(), ExpectedCapacity);
 
-		FString StringFromAscii = FString::ConstructWithSlack(TestAsciiString, ExtraSlack);
+		FString StringFromAscii(TestAsciiString, ExtraSlack);
 		TestEqual(TEXT("(ASCII: Valid string with zero slack) resulting capacity"), StringFromAscii.GetAllocatedSize(), ExpectedCapacity);
 
 		FString StringFromFStringView(FStringView(TestString), ExtraSlack);
@@ -486,10 +486,10 @@ bool FStringConstructorWithSlackTest::RunTest(const FString& Parameters)
 
 		const SIZE_T ExpectedCapacity = FMemory::QuantizeSize(NumElements * sizeof(TCHAR));
 
-		FString StringFromTChar = FString::ConstructWithSlack(TestString, ExtraSlack);
+		FString StringFromTChar(TestString, ExtraSlack);
 		TestEqual(TEXT("(TCHAR: Empty string with slack) resulting capacity"), StringFromTChar.GetAllocatedSize(), ExpectedCapacity);
 
-		FString StringFromAscii = FString::ConstructWithSlack(TestAsciiString, ExtraSlack);
+		FString StringFromAscii(TestAsciiString, ExtraSlack);
 		TestEqual(TEXT("(ASCII: Empty string with slack) resulting capacity"), StringFromAscii.GetAllocatedSize(), ExpectedCapacity);
 
 		FString StringFromFStringView(FStringView(TestString), ExtraSlack);
@@ -507,10 +507,10 @@ bool FStringConstructorWithSlackTest::RunTest(const FString& Parameters)
 
 		const SIZE_T ExpectedCapacity = 0u;
 
-		FString StringFromTChar = FString::ConstructWithSlack(TestString, ExtraSlack);
+		FString StringFromTChar(TestString, ExtraSlack);
 		TestEqual(TEXT("(TCHAR: Empty string with zero slack) resulting capacity"), StringFromTChar.GetAllocatedSize(), ExpectedCapacity);
 
-		FString StringFromAscii = FString::ConstructWithSlack(TestAsciiString, ExtraSlack);
+		FString StringFromAscii(TestAsciiString, ExtraSlack);
 		TestEqual(TEXT("(ASCII: Empty string with zero slack) resulting capacity"), StringFromAscii.GetAllocatedSize(), ExpectedCapacity);
 
 		FString StringFromFStringView(FStringView(TestString), ExtraSlack);
@@ -838,7 +838,7 @@ bool FStringConstructorWithLengthTest::RunTest(const FString& Parameters)
 {
 	auto DoTest = [this](const TCHAR* Ptr, int32 Size, const TArray<TCHAR>& Expected)
 	{
-		FString Str = FString::ConstructFromPtrSize(Ptr, Size);
+		FString Str(Size, Ptr);
 
 		const TArray<TCHAR>& StrArr = Str.GetCharArray();
 		if (StrArr != Expected)
