[LocalWorkQueue]
 - Replace busywait by a better approach that doesn't end up wasting cpu cycles

#rnx

diff --git a/Engine/Source/Runtime/Core/Public/Async/LocalWorkQueue.h b/Engine/Source/Runtime/Core/Public/Async/LocalWorkQueue.h
--- a/Engine/Source/Runtime/Core/Public/Async/LocalWorkQueue.h
+++ b/Engine/Source/Runtime/Core/Public/Async/LocalWorkQueue.h
@@ -3,6 +3,7 @@
 #pragma once
 #include "Async/Fundamental/Task.h"
 #include "Async/Fundamental/Scheduler.h"
+#include "Async/EventCount.h"
 #include "Async/TaskGraphInterfaces.h"
 #include "Experimental/Containers/FAAArrayQueue.h"
 #include "Experimental/ConcurrentLinearAllocator.h"
@@ -47,9 +48,16 @@ class TLocalWorkQueue
 
 	struct FInternalData : public TConcurrentLinearObject<FInternalData, FTaskGraphBlockAllocationTag>, public FThreadSafeRefCountedObject
 	{
-		FAAArrayQueue<TaskType> TaskQueue;	
-		std::atomic_int ActiveWorkers {0};
-		std::atomic_bool CheckDone {false};
+		FAAArrayQueue<TaskType> TaskQueue;
+		std::atomic_int         ActiveWorkers {0};
+		std::atomic_bool        CheckDone {false};
+		UE::FEventCount         FinishedEvent;
+
+		~FInternalData()
+		{
+			check(ActiveWorkers == 0);
+			check(TaskQueue.dequeue() == nullptr);
+		}
 	};
 
 	TRefCountPtr<FInternalData> InternalData;
@@ -97,7 +105,7 @@ public:
 		check(!InternalData->CheckDone.load(std::memory_order_relaxed));
 		check(DoWork != nullptr);
 
-		for(uint16 i = 0; i < NumWorkers; i++)
+		for (uint16 Index = 0; Index < NumWorkers; Index++)
 		{
 			using FTaskHandle = TSharedPtr<LowLevelTasks::FTask, ESPMode::ThreadSafe>;
 			FTaskHandle TaskHandle = MakeShared<LowLevelTasks::FTask, ESPMode::ThreadSafe>();
@@ -107,17 +115,15 @@ public:
 			{
 				TRACE_CPUPROFILER_EVENT_SCOPE(TLocalWorkQueue::AddWorkers);
 				InternalData->ActiveWorkers.fetch_add(1, std::memory_order_acquire);
-				while(true)
+				while (TaskType* Work = InternalData->TaskQueue.dequeue())
 				{
-					TaskType* Work = InternalData->TaskQueue.dequeue();
-					if (Work == nullptr)
-					{
-						break;
-					}
 					check(!InternalData->CheckDone.load(std::memory_order_relaxed));
-					(*LocalDoWork)(Work);				
-				}		
-				InternalData->ActiveWorkers.fetch_sub(1, std::memory_order_release);
+					(*LocalDoWork)(Work);
+				}
+				if (InternalData->ActiveWorkers.fetch_sub(1, std::memory_order_release) == 1)
+				{
+					InternalData->FinishedEvent.Notify();
+				}
 			});
 			verify(TryLaunch(*TaskHandle, LowLevelTasks::EQueuePreference::GlobalQueuePreference));
 		}
@@ -126,23 +132,32 @@ public:
 	inline void Run(TFunctionRef<void(TaskType*)> InDoWork)
 	{
 		DoWork = &InDoWork;
-		LowLevelTasks::BusyWaitUntil([&InDoWork, InternalData = InternalData]()
+
+		TRACE_CPUPROFILER_EVENT_SCOPE(TLocalWorkQueue::Run);
+
+		while (true)
 		{
-			TRACE_CPUPROFILER_EVENT_SCOPE(TLocalWorkQueue::Run);
-			bool Completed = false;
-			while(true)
+			bool bNoActiveWorkers = InternalData->ActiveWorkers.load(std::memory_order_acquire) == 0;
+			if (TaskType* Work = InternalData->TaskQueue.dequeue())
 			{
-				Completed = InternalData->ActiveWorkers.load(std::memory_order_acquire) == 0;
-				TaskType* Work = InternalData->TaskQueue.dequeue();
-				if (Work == nullptr)
+				InDoWork(Work);
+			}
+			else if (bNoActiveWorkers && InternalData->ActiveWorkers.load(std::memory_order_acquire) == 0)
+			{
+				break;
+			}
+			else
+			{
+				auto Token = InternalData->FinishedEvent.PrepareWait();
+				if (InternalData->ActiveWorkers.load(std::memory_order_acquire) == 0)
 				{
-					Completed = Completed && InternalData->ActiveWorkers.load(std::memory_order_acquire) == 0;
-					break;
-				}			
-				InDoWork(Work);	
+					continue;
+				}
+
+				TRACE_CPUPROFILER_EVENT_SCOPE(TLocalWorkQueue::WaitingForWorkers);
+				InternalData->FinishedEvent.Wait(Token);
 			}
-			return Completed; 
-		});
+		}
 
 		InternalData->CheckDone.store(true);
 		check(InternalData->TaskQueue.dequeue() == nullptr);
