Moved TSharedString from DDC into Core

diff --git a/Engine/Source/Developer/DerivedDataCache/Private/DerivedDataSharedString.cpp b/Engine/Source/Developer/DerivedDataCache/Private/DerivedDataSharedString.cpp
deleted file mode 100644
diff --git a/Engine/Source/Developer/DerivedDataCache/Public/DerivedDataSharedString.h b/Engine/Source/Developer/DerivedDataCache/Public/DerivedDataSharedString.h
--- a/Engine/Source/Developer/DerivedDataCache/Public/DerivedDataSharedString.h
+++ b/Engine/Source/Developer/DerivedDataCache/Public/DerivedDataSharedString.h
@@ -2,189 +2,6 @@
 
 #pragma once
 
-#include "Containers/StringView.h"
-#include "CoreTypes.h"
+#include "Containers/SharedString.h"
 #include "DerivedDataSharedStringFwd.h"
-#include "HAL/MemoryBase.h"
-#include "HAL/UnrealMemory.h"
-#include "Misc/CString.h"
-
-#include <atomic>
-
-#define UE_API DERIVEDDATACACHE_API
-
-class FCbFieldView;
-class FCbWriter;
-
-namespace UE::DerivedData
-{
-
-/**
- * A reference to an immutable, shared, reference-counted string.
- *
- * Prefer TStringView when there is a clear single point of ownership with a longer lifetime than
- * the references to the string. This type is meant for cases where a shared string does not have
- * an obvious owner or where the lifetime is not easy to manage.
- *
- * The string is stored as a pointer to the start of the null-terminated string which is preceded
- * by a 4-byte reference count and a 4-byte size.
- */
-template <typename CharType>
-class TSharedString
-{
-public:
-	TSharedString() = default;
-	TSharedString(TSharedString&& String);
-	TSharedString(const TSharedString& String);
-	TSharedString& operator=(TSharedString&& String);
-	TSharedString& operator=(const TSharedString& String);
-
-	TSharedString(TStringView<CharType> String);
-	TSharedString& operator=(TStringView<CharType> String);
-
-	~TSharedString() { Reset(); }
-
-	inline void Reset() { Release(Chars); Chars = nullptr; }
-
-	[[nodiscard]] inline bool IsEmpty() const { return !Chars; }
-	[[nodiscard]] inline int32 Len() const { return Chars ? reinterpret_cast<const int32*>(Chars)[-1] : 0; }
-	[[nodiscard]] inline const CharType* operator*() const { return Chars ? Chars : &NullChar; }
-
-	static const TSharedString Empty;
-
-private:
-	static inline void AddRef(CharType* Chars);
-	static inline void Release(CharType* Chars);
-
-	CharType* Chars = nullptr;
-
-	static constexpr inline CharType NullChar{};
-};
-
-template <typename CharType>
-inline const TSharedString<CharType> TSharedString<CharType>::Empty;
-
-template <typename CharType>
-constexpr inline const CharType* GetData(const TSharedString<CharType> String)
-{
-	return *String;
-}
-
-template <typename CharType>
-constexpr inline auto GetNum(const TSharedString<CharType> String)
-{
-	return String.Len();
-}
-
-template <typename CharType>
-inline TSharedString<CharType>::TSharedString(TSharedString&& String)
-	: Chars(String.Chars)
-{
-	String.Chars = nullptr;
-}
-
-template <typename CharType>
-TSharedString<CharType>::TSharedString(const TSharedString& String)
-	: Chars(String.Chars)
-{
-	AddRef(Chars);
-}
-
-template <typename CharType>
-TSharedString<CharType>& TSharedString<CharType>::operator=(TSharedString&& String)
-{
-	if (this != &String)
-	{
-		Release(Chars);
-		Chars = String.Chars;
-		String.Chars = nullptr;
-	}
-	return *this;
-}
-
-template <typename CharType>
-TSharedString<CharType>& TSharedString<CharType>::operator=(const TSharedString& String)
-{
-	CharType* OldChars = Chars;
-	CharType* NewChars = String.Chars;
-	AddRef(NewChars);
-	Release(OldChars);
-	Chars = NewChars;
-	return *this;
-}
-
-template <typename CharType>
-TSharedString<CharType>::TSharedString(const TStringView<CharType> String)
-{
-	if (const int32 Length = String.Len())
-	{
-		static_assert(alignof(int32) <= MIN_ALIGNMENT);
-		static_assert(alignof(int32) >= alignof(CharType));
-		static_assert(sizeof(int32) == sizeof(std::atomic<int32>));
-		const SIZE_T Size = sizeof(int32) + sizeof(int32) + sizeof(CharType) + sizeof(CharType) * Length;
-		int32* const Header = static_cast<int32*>(FMemory::Malloc(Size));
-		new(&Header[0]) std::atomic<int32>(1); // Ref Count
-		new(&Header[1]) int32(Length);
-		Chars = reinterpret_cast<CharType*>(&Header[2]);
-		String.CopyString(Chars, Length);
-		Chars[Length] = CharType(0);
-	}
-}
-
-template <typename CharType>
-TSharedString<CharType>& TSharedString<CharType>::operator=(const TStringView<CharType> String)
-{
-	return *this = TSharedString(String);
-}
-
-template <typename CharType>
-void TSharedString<CharType>::AddRef(CharType* const MaybeNullChars)
-{
-	if (MaybeNullChars)
-	{
-		static_assert(sizeof(int32) == sizeof(std::atomic<int32>));
-		int32* const Header = reinterpret_cast<int32*>(MaybeNullChars) - 2;
-		std::atomic<int32>& RefCount = reinterpret_cast<std::atomic<int32>&>(Header[0]);
-		RefCount.fetch_add(1, std::memory_order_acquire);
-	}
-}
-
-template <typename CharType>
-void TSharedString<CharType>::Release(CharType* const MaybeNullChars)
-{
-	if (MaybeNullChars)
-	{
-		static_assert(sizeof(int32) == sizeof(std::atomic<int32>));
-		int32* const Header = reinterpret_cast<int32*>(MaybeNullChars) - 2;
-		std::atomic<int32>& RefCount = reinterpret_cast<std::atomic<int32>&>(Header[0]);
-		if (RefCount.fetch_sub(1, std::memory_order_acq_rel) == 1)
-		{
-			FMemory::Free(Header);
-		}
-	}
-}
-
-template <typename CharType>
-inline uint32 GetTypeHash(const TSharedString<CharType>& String)
-{
-	return GetTypeHash(MakeStringView(String));
-}
-
-template <typename CharType>
-inline bool operator==(const TSharedString<CharType>& Lhs, const TSharedString<CharType>& Rhs)
-{
-	return MakeStringView(Lhs).Equals(MakeStringView(Rhs), ESearchCase::IgnoreCase);
-}
-
-template <typename CharType>
-inline bool operator<(const TSharedString<CharType>& Lhs, const TSharedString<CharType>& Rhs)
-{
-	return MakeStringView(Lhs).Compare(MakeStringView(Rhs), ESearchCase::IgnoreCase) < 0;
-}
-
-UE_API bool LoadFromCompactBinary(FCbFieldView Field, FUtf8SharedString& OutString);
-UE_API bool LoadFromCompactBinary(FCbFieldView Field, FWideSharedString& OutString);
-
-} // UE::DerivedData
-
-#undef UE_API
+#include "Serialization/CompactBinarySerialization.h"
diff --git a/Engine/Source/Developer/DerivedDataCache/Public/DerivedDataSharedStringFwd.h b/Engine/Source/Developer/DerivedDataCache/Public/DerivedDataSharedStringFwd.h
--- a/Engine/Source/Developer/DerivedDataCache/Public/DerivedDataSharedStringFwd.h
+++ b/Engine/Source/Developer/DerivedDataCache/Public/DerivedDataSharedStringFwd.h
@@ -2,30 +2,16 @@
 
 #pragma once
 
-#include "CoreTypes.h"
-#include "Traits/ElementType.h"
-#include "Traits/IsContiguousContainer.h"
+#include "Containers/StringFwd.h"
 
 namespace UE::DerivedData
 {
 
-template <typename CharType> class TSharedString;
+template <typename CharType> using TSharedString = UE::TSharedString<CharType>;
 
-using FSharedString = TSharedString<TCHAR>;
-using FAnsiSharedString = TSharedString<ANSICHAR>;
-using FWideSharedString = TSharedString<WIDECHAR>;
-using FUtf8SharedString = TSharedString<UTF8CHAR>;
+using FSharedString = UE::FSharedString;
+using FAnsiSharedString = UE::FAnsiSharedString;
+using FWideSharedString = UE::FWideSharedString;
+using FUtf8SharedString = UE::FUtf8SharedString;
 
 } // UE::DerivedData
-
-template <typename CharType>
-struct TIsContiguousContainer<UE::DerivedData::TSharedString<CharType>>
-{
-	static constexpr bool Value = true;
-};
-
-template <typename CharType>
-struct TElementType<UE::DerivedData::TSharedString<CharType>>
-{
-	using Type = CharType;
-};
diff --git a/Engine/Source/Runtime/Core/Private/Serialization/CompactBinarySerialization.cpp b/Engine/Source/Runtime/Core/Private/Serialization/CompactBinarySerialization.cpp
--- a/Engine/Source/Runtime/Core/Private/Serialization/CompactBinarySerialization.cpp
+++ b/Engine/Source/Runtime/Core/Private/Serialization/CompactBinarySerialization.cpp
@@ -4,6 +4,7 @@
 
 #include "Containers/Array.h"
 #include "Containers/ContainerAllocationPolicies.h"
+#include "Containers/SharedString.h"
 #include "Containers/StringView.h"
 #include "Containers/UnrealString.h"
 #include "CoreGlobals.h"
@@ -301,6 +302,24 @@ bool LoadFromCompactBinary(FCbFieldView Field, FWideStringBuilderBase& OutValue)
 	return !Field.HasError();
 }
 
+bool LoadFromCompactBinary(FCbFieldView Field, UE::FUtf8SharedString& OutString)
+{
+	OutString = Field.AsString();
+	return !Field.HasError();
+}
+
+bool LoadFromCompactBinary(FCbFieldView Field, UE::FWideSharedString& OutString)
+{
+	TWideStringBuilder<512> String;
+	if (LoadFromCompactBinary(Field, String))
+	{
+		OutString = String;
+		return true;
+	}
+	OutString.Reset();
+	return false;
+}
+
 bool LoadFromCompactBinary(FCbFieldView Field, FString& OutValue)
 {
 	OutValue = FString(Field.AsString());
diff --git a/Engine/Source/Runtime/Core/Public/Containers/SharedString.h b/Engine/Source/Runtime/Core/Public/Containers/SharedString.h
new file mode 100644
--- /dev/null
+++ b/Engine/Source/Runtime/Core/Public/Containers/SharedString.h
@@ -0,0 +1,203 @@
+// Copyright Epic Games, Inc. All Rights Reserved.
+
+#pragma once
+
+#include "Containers/StringFwd.h"
+#include "Containers/StringView.h"
+#include "CoreTypes.h"
+#include "HAL/MemoryBase.h"
+#include "HAL/UnrealMemory.h"
+#include "Misc/CString.h"
+
+#include <atomic>
+
+namespace UE
+{
+
+/**
+* A reference to an immutable, shared, reference-counted string.
+*
+* Prefer TStringView when there is a clear single point of ownership with a longer lifetime than
+* the references to the string. This type is meant for cases where a shared string does not have
+* an obvious owner or where the lifetime is not easy to manage.
+*
+* The string is stored as a pointer to the start of the null-terminated string which is preceded
+* by a 4-byte reference count and a 4-byte size.
+*/
+template <typename CharType>
+class TSharedString
+{
+public:
+	using ElementType = CharType;
+
+	TSharedString() = default;
+	TSharedString(TSharedString&& String);
+	TSharedString(const TSharedString& String);
+	TSharedString& operator=(TSharedString&& String);
+	TSharedString& operator=(const TSharedString& String);
+
+	/** Allocates a copy of the string and constructs this as a reference to it. */
+	TSharedString(TStringView<CharType> String);
+
+	/** Allocates a copy of the string and assigns this as a reference to it. */
+	TSharedString& operator=(TStringView<CharType> String);
+
+	inline ~TSharedString()
+	{
+		Release(Chars);
+	}
+
+	/** Resets this to reference the empty string. */
+	inline void Reset()
+	{
+		Release(Chars);
+		Chars = nullptr;
+	}
+
+	/** Returns true if the referenced string is empty. */
+	[[nodiscard]] inline bool IsEmpty() const
+	{
+		return !Chars;
+	}
+
+	/** Returns the length of the referenced string excluding the null terminator. */
+	[[nodiscard]] inline int32 Len() const
+	{
+		return Chars ? reinterpret_cast<const int32*>(Chars)[-1] : 0;
+	}
+
+	/** Returns a pointer to the start of the referenced null-terminated string. */
+	[[nodiscard]] inline const CharType* operator*() const
+	{
+		return Chars ? Chars : &NullChar;
+	}
+
+	/** An empty string provided mainly for returning a reference to an empty TSharedString. */
+	static const TSharedString Empty;
+
+private:
+	friend constexpr inline const CharType* GetData(const TSharedString String)
+	{
+		return *String;
+	}
+
+	friend constexpr inline auto GetNum(const TSharedString String)
+	{
+		return String.Len();
+	}
+
+	friend inline uint32 GetTypeHash(const TSharedString& String)
+	{
+		return GetTypeHash(MakeStringView(String));
+	}
+
+	friend inline bool operator==(const TSharedString& Lhs, const TSharedString& Rhs)
+	{
+		return MakeStringView(Lhs).Equals(MakeStringView(Rhs), ESearchCase::IgnoreCase);
+	}
+
+	friend inline bool operator<(const TSharedString& Lhs, const TSharedString& Rhs)
+	{
+		return MakeStringView(Lhs).Compare(MakeStringView(Rhs), ESearchCase::IgnoreCase) < 0;
+	}
+
+	static inline void AddRef(CharType* Chars);
+	static inline void Release(CharType* Chars);
+
+	CharType* Chars = nullptr;
+
+	static constexpr inline CharType NullChar{};
+};
+
+template <typename CharType>
+inline const TSharedString<CharType> TSharedString<CharType>::Empty;
+
+template <typename CharType>
+inline TSharedString<CharType>::TSharedString(TSharedString&& String)
+	: Chars(String.Chars)
+{
+	String.Chars = nullptr;
+}
+
+template <typename CharType>
+TSharedString<CharType>::TSharedString(const TSharedString& String)
+	: Chars(String.Chars)
+{
+	AddRef(Chars);
+}
+
+template <typename CharType>
+TSharedString<CharType>& TSharedString<CharType>::operator=(TSharedString&& String)
+{
+	if (this != &String)
+	{
+		Release(Chars);
+		Chars = String.Chars;
+		String.Chars = nullptr;
+	}
+	return *this;
+}
+
+template <typename CharType>
+TSharedString<CharType>& TSharedString<CharType>::operator=(const TSharedString& String)
+{
+	CharType* OldChars = Chars;
+	CharType* NewChars = String.Chars;
+	AddRef(NewChars);
+	Release(OldChars);
+	Chars = NewChars;
+	return *this;
+}
+
+template <typename CharType>
+TSharedString<CharType>::TSharedString(const TStringView<CharType> String)
+{
+	if (const int32 Length = String.Len())
+	{
+		static_assert(alignof(int32) <= MIN_ALIGNMENT);
+		static_assert(alignof(int32) >= alignof(CharType));
+		static_assert(sizeof(int32) == sizeof(std::atomic<int32>));
+		const SIZE_T Size = sizeof(int32) + sizeof(int32) + sizeof(CharType) + sizeof(CharType) * Length;
+		int32* const Header = static_cast<int32*>(FMemory::Malloc(Size));
+		new(&Header[0]) std::atomic<int32>(1); // Ref Count
+		new(&Header[1]) int32(Length);
+		Chars = reinterpret_cast<CharType*>(&Header[2]);
+		String.CopyString(Chars, Length);
+		Chars[Length] = CharType(0);
+	}
+}
+
+template <typename CharType>
+TSharedString<CharType>& TSharedString<CharType>::operator=(const TStringView<CharType> String)
+{
+	return *this = TSharedString(String);
+}
+
+template <typename CharType>
+void TSharedString<CharType>::AddRef(CharType* const MaybeNullChars)
+{
+	if (MaybeNullChars)
+	{
+		static_assert(sizeof(int32) == sizeof(std::atomic<int32>));
+		int32* const Header = reinterpret_cast<int32*>(MaybeNullChars) - 2;
+		std::atomic<int32>& RefCount = reinterpret_cast<std::atomic<int32>&>(Header[0]);
+		RefCount.fetch_add(1, std::memory_order_relaxed);
+	}
+}
+
+template <typename CharType>
+void TSharedString<CharType>::Release(CharType* const MaybeNullChars)
+{
+	if (MaybeNullChars)
+	{
+		static_assert(sizeof(int32) == sizeof(std::atomic<int32>));
+		int32* const Header = reinterpret_cast<int32*>(MaybeNullChars) - 2;
+		std::atomic<int32>& RefCount = reinterpret_cast<std::atomic<int32>&>(Header[0]);
+		if (RefCount.fetch_sub(1, std::memory_order_acq_rel) == 1)
+		{
+			FMemory::Free(Header);
+		}
+	}
+}
+
+} // UE
diff --git a/Engine/Source/Runtime/Core/Public/Containers/StringFwd.h b/Engine/Source/Runtime/Core/Public/Containers/StringFwd.h
--- a/Engine/Source/Runtime/Core/Public/Containers/StringFwd.h
+++ b/Engine/Source/Runtime/Core/Public/Containers/StringFwd.h
@@ -3,6 +3,7 @@
 #pragma once
 
 #include "CoreTypes.h"
+#include "Traits/ElementType.h"
 #include "Traits/IsContiguousContainer.h"
 
 // String Builder
@@ -41,3 +42,23 @@ using FStringView = TStringView<TCHAR>;
 using FAnsiStringView = TStringView<ANSICHAR>;
 using FWideStringView = TStringView<WIDECHAR>;
 using FUtf8StringView = TStringView<UTF8CHAR>;
+
+// Shared String
+
+namespace UE
+{
+
+template <typename CharType> class TSharedString;
+
+using FSharedString = TSharedString<TCHAR>;
+using FAnsiSharedString = TSharedString<ANSICHAR>;
+using FWideSharedString = TSharedString<WIDECHAR>;
+using FUtf8SharedString = TSharedString<UTF8CHAR>;
+
+} // UE
+
+template <typename CharType>
+struct TIsContiguousContainer<UE::TSharedString<CharType>>
+{
+	static constexpr bool Value = true;
+};
diff --git a/Engine/Source/Runtime/Core/Public/Serialization/CompactBinarySerialization.h b/Engine/Source/Runtime/Core/Public/Serialization/CompactBinarySerialization.h
--- a/Engine/Source/Runtime/Core/Public/Serialization/CompactBinarySerialization.h
+++ b/Engine/Source/Runtime/Core/Public/Serialization/CompactBinarySerialization.h
@@ -105,6 +105,8 @@ CORE_API FArchive& operator<<(FArchive& Ar, FCbObject& Object);
 
 CORE_API bool LoadFromCompactBinary(FCbFieldView Field, FUtf8StringBuilderBase& OutValue);
 CORE_API bool LoadFromCompactBinary(FCbFieldView Field, FWideStringBuilderBase& OutValue);
+CORE_API bool LoadFromCompactBinary(FCbFieldView Field, UE::FUtf8SharedString& OutString);
+CORE_API bool LoadFromCompactBinary(FCbFieldView Field, UE::FWideSharedString& OutString);
 CORE_API bool LoadFromCompactBinary(FCbFieldView Field, FString& OutValue);
 CORE_API bool LoadFromCompactBinary(FCbFieldView Field, FName& OutValue);
 
diff --git a/Engine/Source/Runtime/CoreUObject/Private/Serialization/DerivedData.cpp b/Engine/Source/Runtime/CoreUObject/Private/Serialization/DerivedData.cpp
--- a/Engine/Source/Runtime/CoreUObject/Private/Serialization/DerivedData.cpp
+++ b/Engine/Source/Runtime/CoreUObject/Private/Serialization/DerivedData.cpp
@@ -373,28 +373,28 @@ void FDerivedData::Serialize(FArchive& Ar, UObject* Owner)
 
 #if WITH_EDITORONLY_DATA
 
-FDerivedData::FDerivedData(const DerivedData::FSharedString& Name, const FSharedBuffer& Data)
+FDerivedData::FDerivedData(const FSharedString& Name, const FSharedBuffer& Data)
 	: EditorData(MakePimpl<DerivedData::Private::FEditorData, EPimplPtrMode::DeepCopy>(
 		Name, TInPlaceType<DerivedData::Private::FCompositeBufferWithHash>(), Data))
 {
 	CookedData.Flags = EDerivedDataFlags::Required;
 }
 
-FDerivedData::FDerivedData(const DerivedData::FSharedString& Name, const FCompositeBuffer& Data)
+FDerivedData::FDerivedData(const FSharedString& Name, const FCompositeBuffer& Data)
 	: EditorData(MakePimpl<DerivedData::Private::FEditorData, EPimplPtrMode::DeepCopy>(
 		Name, TInPlaceType<DerivedData::Private::FCompositeBufferWithHash>(), Data))
 {
 	CookedData.Flags = EDerivedDataFlags::Required;
 }
 
-FDerivedData::FDerivedData(const DerivedData::FSharedString& Name, const FCompressedBuffer& Data)
+FDerivedData::FDerivedData(const FSharedString& Name, const FCompressedBuffer& Data)
 	: EditorData(MakePimpl<DerivedData::Private::FEditorData, EPimplPtrMode::DeepCopy>(
 		Name, TInPlaceType<FCompressedBuffer>(), Data))
 {
 	CookedData.Flags = EDerivedDataFlags::Required;
 }
 
-FDerivedData::FDerivedData(const DerivedData::FSharedString& Name, const DerivedData::FCacheKey& Key)
+FDerivedData::FDerivedData(const FSharedString& Name, const DerivedData::FCacheKey& Key)
 	: EditorData(MakePimpl<DerivedData::Private::FEditorData, EPimplPtrMode::DeepCopy>(
 		Name, TInPlaceType<DerivedData::Private::FCacheKeyWithId>(), Key, DerivedData::FValueId::Null))
 {
@@ -402,7 +402,7 @@ FDerivedData::FDerivedData(const DerivedData::FSharedString& Name, const Derived
 }
 
 FDerivedData::FDerivedData(
-	const DerivedData::FSharedString& Name,
+	const FSharedString& Name,
 	const DerivedData::FCacheKey& Key,
 	const DerivedData::FValueId& ValueId)
 	: EditorData(MakePimpl<DerivedData::Private::FEditorData, EPimplPtrMode::DeepCopy>(
@@ -411,9 +411,9 @@ FDerivedData::FDerivedData(
 	CookedData.Flags = EDerivedDataFlags::Required;
 }
 
-const DerivedData::FSharedString& FDerivedData::GetName() const
+const FSharedString& FDerivedData::GetName() const
 {
-	return EditorData ? EditorData->GetName() : DerivedData::FSharedString::Empty;
+	return EditorData ? EditorData->GetName() : FSharedString::Empty;
 }
 
 void FDerivedData::SetFlags(EDerivedDataFlags InFlags)
diff --git a/Engine/Source/Runtime/CoreUObject/Public/Serialization/DerivedData.h b/Engine/Source/Runtime/CoreUObject/Public/Serialization/DerivedData.h
--- a/Engine/Source/Runtime/CoreUObject/Public/Serialization/DerivedData.h
+++ b/Engine/Source/Runtime/CoreUObject/Public/Serialization/DerivedData.h
@@ -20,8 +20,6 @@ template <typename FuncType> class TUniqueFunction;
 
 namespace UE::DerivedData { struct FCacheKey; }
 namespace UE::DerivedData { struct FValueId; }
-namespace UE::DerivedData { template <typename CharType> class TSharedString; }
-namespace UE::DerivedData { using FSharedString = TSharedString<TCHAR>; }
 namespace UE::DerivedData::Private { class FCookedData; }
 namespace UE::DerivedData::Private { class FEditorData; }
 namespace UE::DerivedData::Private { class FIoResponse; }
@@ -122,17 +120,17 @@ public:
 
 #if WITH_EDITORONLY_DATA
 	/** References a value that is stored in a buffer. */
-	UE_API FDerivedData(const DerivedData::FSharedString& Name, const FSharedBuffer& Data);
-	UE_API FDerivedData(const DerivedData::FSharedString& Name, const FCompositeBuffer& Data);
-	UE_API FDerivedData(const DerivedData::FSharedString& Name, const FCompressedBuffer& Data);
+	UE_API FDerivedData(const FSharedString& Name, const FSharedBuffer& Data);
+	UE_API FDerivedData(const FSharedString& Name, const FCompositeBuffer& Data);
+	UE_API FDerivedData(const FSharedString& Name, const FCompressedBuffer& Data);
 
 	/** References a value that was saved using ICache::PutValue. */
-	UE_API FDerivedData(const DerivedData::FSharedString& Name, const DerivedData::FCacheKey& Key);
+	UE_API FDerivedData(const FSharedString& Name, const DerivedData::FCacheKey& Key);
 	/** References a value in a record that was saved using ICache::Put. */
-	UE_API FDerivedData(const DerivedData::FSharedString& Name, const DerivedData::FCacheKey& Key, const DerivedData::FValueId& ValueId);
+	UE_API FDerivedData(const FSharedString& Name, const DerivedData::FCacheKey& Key, const DerivedData::FValueId& ValueId);
 
 	/** Returns the name of the reference if available. */
-	UE_API const DerivedData::FSharedString& GetName() const;
+	UE_API const FSharedString& GetName() const;
 
 	/** Overwrites the existing flags. */
 	UE_API void SetFlags(EDerivedDataFlags Flags);
diff --git a/Engine/Source/Runtime/Engine/Private/Animation/AnimationCompressionDerivedData.h b/Engine/Source/Runtime/Engine/Private/Animation/AnimationCompressionDerivedData.h
--- a/Engine/Source/Runtime/Engine/Private/Animation/AnimationCompressionDerivedData.h
+++ b/Engine/Source/Runtime/Engine/Private/Animation/AnimationCompressionDerivedData.h
@@ -6,6 +6,7 @@
 
 #if WITH_EDITORONLY_DATA
 #include "DerivedDataRequestOwner.h"
+#include "DerivedDataSharedStringFwd.h"
 #include "Experimental/Misc/ExecutionResource.h"
 #include "Async/AsyncWork.h"
 #include "Animation/AnimCompressionTypes.h"
@@ -20,9 +21,6 @@ namespace UE::DerivedData
 {
 	struct FCacheGetValueResponse;
 	struct FCacheKey;
-
-	template <typename CharType> class TSharedString;
-	using FSharedString = TSharedString<TCHAR>;
 }
 
 namespace UE::Anim
