[Scheduler/TaskGraph]
 - Deprecate BusyWait APIs

diff --git a/Engine/Source/Runtime/Core/Private/Async/Fundamental/UnitTests.cpp b/Engine/Source/Runtime/Core/Private/Async/Fundamental/UnitTests.cpp
--- a/Engine/Source/Runtime/Core/Private/Async/Fundamental/UnitTests.cpp
+++ b/Engine/Source/Runtime/Core/Private/Async/Fundamental/UnitTests.cpp
@@ -19,6 +19,51 @@ namespace Tasks2Tests
 {
 	using namespace LowLevelTasks;
 
+	// There is no event that can be waited on for low level tasks so
+	// this is acceptable for testing purpose only to spin until completion.
+	inline void WaitForTask(FTask& Task)
+	{
+		// Task has already started, wait
+		while (!Task.IsCompleted())
+		{
+			FPlatformProcess::YieldThread();
+		}
+	}
+
+	inline void WaitForTasks(TArrayView<FTask> Tasks)
+	{
+		for (FTask& Task : Tasks)
+		{
+			WaitForTask(Task);
+		}
+	}
+
+	// For highly recursive workload that waits, we need the waiters
+	// to contribute to the work otherwise we'll exhaust all threads
+	// and deadlock. So try to expedite to join the work or wait
+	// if the work is already started. Only operates on TSharedPtr
+	// since we're not waiting for task completion during expedite
+	// and may exit scope before the task is scheduled.
+	inline void ExpediteOrWaitForTask(TSharedPtr<FTask>& Task)
+	{
+		if (Task->TryExpedite())
+		{
+			return;
+		}
+
+		// Task has already started, wait
+		WaitForTask(*Task);
+	}
+
+	inline void ExpediteOrWaitForTasks(TArrayView<TSharedPtr<FTask>> Tasks)
+	{
+		for (TSharedPtr<FTask>& Task : Tasks)
+		{
+			ExpediteOrWaitForTask(Task);
+		}
+	}
+
+
 	template<int NumTasks>
 	void EmptyPerfTest(bool last)
 	{
@@ -47,7 +92,7 @@ namespace Tasks2Tests
 				TryLaunch(Tasks[i]);
 			}
 
-			BusyWaitForTasks<FTask>(Tasks);
+			WaitForTasks(Tasks);
 
 			double EndTime = FPlatformTime::Seconds();
 			double Total = EndTime - StartTime;
@@ -92,8 +137,8 @@ namespace Tasks2Tests
 			double SpawnedTime = FPlatformTime::Seconds();
 
 			TryLaunch(SignalTask);
-			BusyWaitForTasks<FTask>(Tasks);
-			BusyWaitForTask(SignalTask);
+			WaitForTasks(Tasks);
+			WaitForTask(SignalTask);
 
 			double EndTime = FPlatformTime::Seconds();
 			double Total = EndTime - StartTime;
@@ -191,33 +236,34 @@ namespace Tasks2Tests
 
 		{
 			using namespace LowLevelTasks;
-			TArray<FTask> Tasks;
+			TArray<TSharedPtr<FTask>> Tasks;
 			Tasks.AddDefaulted(NumTasks);
 
 			double StartTime = FPlatformTime::Seconds();
 
-			for (int i = 0; i < NumTasks; i++)
+			for (TSharedPtr<FTask>& Task : Tasks)
 			{
-				Tasks[i].Init(TEXT("Perf Test"), [&Garbage, &TestVal]()
+				Task = MakeShared<FTask>();
+				Task->Init(TEXT("Perf Test"), [&Garbage, &TestVal, Task]()
 				{
 					Garbage = Test::Fibonacci(TestVal);
 				});
 			}
 
-			FTask SignalTask;
-			SignalTask.Init(TEXT("Signal Task"), [&Tasks]()
+			TSharedPtr<FTask> SignalTask = MakeShared<FTask>();
+			SignalTask->Init(TEXT("Signal Task"), [&Tasks]()
 			{
 				for (int i = 0; i < NumTasks; i++)
 				{
-					TryLaunch(Tasks[i]);
+					TryLaunch(*Tasks[i]);
 				}
 			});
 
 			double SpawnedTime = FPlatformTime::Seconds();
 
-			TryLaunch(SignalTask);
-			BusyWaitForTasks<FTask>(Tasks);
-			BusyWaitForTask(SignalTask);
+			TryLaunch(*SignalTask);
+			ExpediteOrWaitForTasks(Tasks);
+			ExpediteOrWaitForTask(SignalTask);
 
 			double EndTime = FPlatformTime::Seconds();
 			double Total = EndTime - StartTime;
@@ -317,22 +363,25 @@ namespace Tasks2Tests
 				if (N > STheshold)
 				{
 					int64 Res1, Res2;
-					FTask Task1, Task2;
 
-					Task1.Init(TEXT("Basic Fibonacci Test1"), ETaskPriority::Normal, [&Res1, N]()
+					// autodelete will allow us to expedite the task without regard to lifetime
+					TSharedPtr<FTask> Task1 = MakeShared<FTask>();
+					TSharedPtr<FTask> Task2 = MakeShared<FTask>();
+
+					Task1->Init(TEXT("Basic Fibonacci Test1"), ETaskPriority::Normal, [&Res1, Task1, N]()
 					{
 						Res1 = Test::FibonacciTasksBasic(N - 1);
 					});
-					TryLaunch(Task1);
+					TryLaunch(*Task1);
 
-					Task2.Init(TEXT("Basic Fibonacci Test2"), ETaskPriority::Normal, [&Res2, N]()
+					Task2->Init(TEXT("Basic Fibonacci Test2"), ETaskPriority::Normal, [&Res2, Task2, N]()
 					{
 						Res2 = Test::FibonacciTasksBasic(N - 2);
 					});
-					TryLaunch(Task2);
+					TryLaunch(*Task2);
 
-					BusyWaitForTask(Task2);
-					BusyWaitForTask(Task1);
+					ExpediteOrWaitForTask(Task2);
+					ExpediteOrWaitForTask(Task1);
 					return Res1 + Res2;
 				}
 				return Fibonacci(N);
@@ -367,33 +416,34 @@ namespace Tasks2Tests
 		volatile int64& TestVal = STestVal;
 
 		{
-			TArray<FTask> Tasks;
+			TArray<TSharedPtr<FTask>> Tasks;
 			Tasks.AddDefaulted(NumTasks);
 
 			double StartTime = FPlatformTime::Seconds();
 
-			for (int i = 0; i < NumTasks; i++)
+			for (TSharedPtr<FTask>& Task : Tasks)
 			{
-				Tasks[i].Init(TEXT("Perf Test Stealing"), ETaskPriority::Normal, [&Garbage, &TestVal]()
+				Task = MakeShared<FTask>();
+				Task->Init(TEXT("Perf Test Stealing"), ETaskPriority::Normal, [&Garbage, &TestVal, Task]()
 				{
 					Garbage = Test::FibonacciTasksBasic(TestVal);
 				});
 			}
 
-			FTask SignalTask;
-			SignalTask.Init(TEXT("Signal Task"), [&Tasks]()
+			TSharedPtr<FTask> SignalTask = MakeShared<FTask>();
+			SignalTask->Init(TEXT("Signal Task"), [&Tasks]()
 			{
-				for (int i = 0; i < NumTasks; i++)
+				for (int32 Index = 0; Index < NumTasks; Index++)
 				{
-					TryLaunch(Tasks[i]);
+					TryLaunch(*Tasks[Index]);
 				}
 			});
 
 			double SpawnedTime = FPlatformTime::Seconds();
 
-			TryLaunch(SignalTask);
-			BusyWaitForTasks<FTask>(Tasks);
-			BusyWaitForTask(SignalTask);
+			TryLaunch(*SignalTask);
+			ExpediteOrWaitForTasks(Tasks);
+			ExpediteOrWaitForTask(SignalTask);
 
 			double EndTime = FPlatformTime::Seconds();
 			double Total = EndTime - StartTime;
@@ -557,8 +607,8 @@ namespace Tasks2Tests
 				if(WasCanceled)
 				{
 					WasCanceled = !Task.TryRevive();
-				}	
-				BusyWaitForTask(Task);
+				}
+				WaitForTask(Task);
 
 				if (WasCanceled)
 				{
@@ -584,7 +634,7 @@ namespace Tasks2Tests
 				TryLaunch(*TaskHandle);
 
 				bool WasCanceled = TestCancel && TaskHandle->TryCancel();
-				BusyWaitForTask(*TaskHandle);
+				WaitForTask(*TaskHandle);
 
 				if (WasCanceled)
 				{
@@ -610,7 +660,7 @@ namespace Tasks2Tests
 				Task.TryExpedite();
 
 				verify(!Task.TryCancel());
-				BusyWaitForTask(Task);
+				WaitForTask(Task);
 
 				verify(TestValue == 42);
 			}
@@ -660,7 +710,7 @@ namespace Tasks2Tests
 				return &TaskB;
 			});
 			verify(TryLaunch(TaskA));
-			BusyWaitForTask(TaskB);
+			WaitForTask(TaskB);
 
 			verify(TestValue == 42);
 		}
diff --git a/Engine/Source/Runtime/Core/Public/Async/Fundamental/Scheduler.h b/Engine/Source/Runtime/Core/Public/Async/Fundamental/Scheduler.h
--- a/Engine/Source/Runtime/Core/Public/Async/Fundamental/Scheduler.h
+++ b/Engine/Source/Runtime/Core/Public/Async/Fundamental/Scheduler.h
@@ -186,15 +186,18 @@ namespace LowLevelTasks
 
 		//tries to do some work until the Task is completed
 		template<typename TaskType>
+		UE_DEPRECATED(5.5, "This method will be removed.")
 		inline void BusyWait(const TaskType& Task, bool ForceAllowBackgroundWork = false);
 
 		//tries to do some work until the Conditional return true
 		template<typename Conditional>
+		UE_DEPRECATED(5.5, "This method will be removed.")
 		inline void BusyWaitUntil(Conditional&& Cond, bool ForceAllowBackgroundWork = false);
 
 		//tries to do some work until all the Tasks are completed
 		//the template parameter can be any Type that has a const conversion operator to FTask
 		template<typename TaskType>
+		UE_DEPRECATED(5.5, "This method will be removed.")
 		inline void BusyWait(const TArrayView<const TaskType>& Tasks, bool ForceAllowBackgroundWork = false);
 
 		//number of instantiated workers
@@ -321,21 +324,30 @@ namespace LowLevelTasks
 		return FScheduler::Get().TryLaunch(Task, QueuePreference, bWakeUpWorker);
 	}
 
+	UE_DEPRECATED(5.5, "This method will be removed.")
 	FORCEINLINE_DEBUGGABLE void BusyWaitForTask(const FTask& Task, bool ForceAllowBackgroundWork = false)
 	{
+		PRAGMA_DISABLE_DEPRECATION_WARNINGS
 		FScheduler::Get().BusyWait(Task, ForceAllowBackgroundWork);
+		PRAGMA_ENABLE_DEPRECATION_WARNINGS
 	}
 
 	template<typename Conditional>
+	UE_DEPRECATED(5.5, "This method will be removed.")
 	FORCEINLINE_DEBUGGABLE void BusyWaitUntil(Conditional&& Cond, bool ForceAllowBackgroundWork = false)
 	{
+		PRAGMA_DISABLE_DEPRECATION_WARNINGS
 		FScheduler::Get().BusyWaitUntil<Conditional>(Forward<Conditional>(Cond), ForceAllowBackgroundWork);
+		PRAGMA_ENABLE_DEPRECATION_WARNINGS
 	}
 
 	template<typename TaskType>
+	UE_DEPRECATED(5.5, "This method will be removed.")
 	FORCEINLINE_DEBUGGABLE void BusyWaitForTasks(const TArrayView<const TaskType>& Tasks, bool ForceAllowBackgroundWork = false)
 	{
+		PRAGMA_DISABLE_DEPRECATION_WARNINGS
 		FScheduler::Get().BusyWait<TaskType>(Tasks, ForceAllowBackgroundWork);
+		PRAGMA_ENABLE_DEPRECATION_WARNINGS
 	}
 
 	/******************
diff --git a/Engine/Source/Runtime/Core/Public/Tasks/Task.h b/Engine/Source/Runtime/Core/Public/Tasks/Task.h
--- a/Engine/Source/Runtime/Core/Public/Tasks/Task.h
+++ b/Engine/Source/Runtime/Core/Public/Tasks/Task.h
@@ -19,6 +19,7 @@ namespace UE::Tasks
 	bool Wait(const TaskCollectionType& Tasks, FTimespan InTimeout = FTimespan::MaxValue());
 
 	template<typename TaskCollectionType>
+	UE_DEPRECATED(5.5, "Use Wait instead.")
 	bool BusyWait(const TaskCollectionType& Tasks, FTimespan InTimeout = FTimespan::MaxValue());
 
 	template<typename TaskType> void AddNested(const TaskType& Nested);
@@ -102,18 +103,24 @@ namespace UE::Tasks
 			// waits for task's completion for at least the specified amount of time, while executing other tasks.
 			// the call can return much later than the given timeout
 			// @return true if the task is completed
+			UE_DEPRECATED(5.5, "Use Wait instead.")
 			bool BusyWait(FTimespan Timeout = FTimespan::MaxValue()) const
 			{
+				PRAGMA_DISABLE_DEPRECATION_WARNINGS
 				return !IsValid() || Pimpl->BusyWait(FTimeout{ Timeout });
+				PRAGMA_ENABLE_DEPRECATION_WARNINGS
 			}
 
 			// waits for task's completion or the given condition becomes true, while executing other tasks.
 			// the call can return much later than the given condition became true
 			// @return true if the task is completed
 			template<typename ConditionType>
+			UE_DEPRECATED(5.5, "This method will be removed.")
 			bool BusyWait(ConditionType&& Condition) const
 			{
+				PRAGMA_DISABLE_DEPRECATION_WARNINGS
 				return !IsValid() || Pimpl->BusyWait(Forward<ConditionType>(Condition));
+				PRAGMA_ENABLE_DEPRECATION_WARNINGS
 			}
 
 			// launches a task for asynchronous execution
@@ -370,8 +377,10 @@ namespace UE::Tasks
 
 	// wait for multiple tasks while executing other tasks
 	template<typename TaskCollectionType>
+	UE_DEPRECATED(5.5, "Use Wait instead.")
 	bool BusyWait(const TaskCollectionType& Tasks, FTimespan InTimeout/* = FTimespan::MaxValue()*/)
 	{
+		PRAGMA_DISABLE_DEPRECATION_WARNINGS
 		TaskTrace::FWaitingScope WaitingScope(Private::GetTraceIds(Tasks));
 		TRACE_CPUPROFILER_EVENT_SCOPE(Tasks::BusyWait);
 
@@ -389,6 +398,7 @@ namespace UE::Tasks
 		}
 
 		return true;
+		PRAGMA_ENABLE_DEPRECATION_WARNINGS
 	}
 
 	using FTask = Private::FTaskHandle;
diff --git a/Engine/Source/Runtime/Core/Public/Tasks/TaskPrivate.h b/Engine/Source/Runtime/Core/Public/Tasks/TaskPrivate.h
--- a/Engine/Source/Runtime/Core/Public/Tasks/TaskPrivate.h
+++ b/Engine/Source/Runtime/Core/Public/Tasks/TaskPrivate.h
@@ -439,6 +439,7 @@ public:
 			CORE_API void WaitWithNamedThreadsSupport();
 
 			// waits until the task is completed or waiting timed out, while executing other tasks
+			UE_DEPRECATED(5.5, "Use Wait instead")
 			bool BusyWait(FTimeout Timeout)
 			{
 				return Wait(Timeout);
@@ -446,6 +447,7 @@ public:
 
 			// waits until the task is completed or the condition returns true, while executing other tasks
 			template<typename ConditionType>
+			UE_DEPRECATED(5.5, "This method will be removed.")
 			bool BusyWait(ConditionType&& Condition)
 			{
 				TaskTrace::FWaitingScope WaitingScope(GetTraceId());
diff --git a/Engine/Source/Runtime/Core/Tests/Tasks/TasksTest.cpp b/Engine/Source/Runtime/Core/Tests/Tasks/TasksTest.cpp
--- a/Engine/Source/Runtime/Core/Tests/Tasks/TasksTest.cpp
+++ b/Engine/Source/Runtime/Core/Tests/Tasks/TasksTest.cpp
@@ -43,7 +43,6 @@ namespace UE { namespace TasksTests
 
 		{	// launch a task and wait till it's executed
 			Launch(UE_SOURCE_LOCATION, [] {}).Wait();
-			Launch(UE_SOURCE_LOCATION, [] {}).BusyWait();
 		}
 
 		{	// FTaskEvent asserts on destruction if it wasn't triggered. uncomment the code to verify
@@ -62,7 +61,6 @@ namespace UE { namespace TasksTests
 			Event.Trigger();
 			check(Event.IsCompleted());
 			verify(Event.Wait(FTimespan::Zero()));
-			verify(Event.BusyWait(FTimespan::Zero()));
 		}
 
 		{	// FTaskEvent can be triggered multiple times
@@ -79,32 +77,22 @@ namespace UE { namespace TasksTests
 			check(!Event.IsCompleted());
 
 			// check that waiting blocks
-			FTask Task = Launch(UE_SOURCE_LOCATION, [Event]() mutable { Event.BusyWait(); });
+			FTask Task = Launch(UE_SOURCE_LOCATION, [Event]() mutable { Event.Wait(); });
 			FPlatformProcess::Sleep(0.1f);
 			check(!Task.IsCompleted());
 
 			Event.Trigger();
 			check(Event.IsCompleted());
 			verify(Event.Wait(FTimespan::Zero()));
-			verify(Event.BusyWait(FTimespan::Zero()));
 		}
 
-		{	// busy-waiting for multiple tasks
+		{	// waiting for multiple tasks
 			TArray<FTask> Tasks
 			{
 				Launch(UE_SOURCE_LOCATION,[] {}),
 				Launch(UE_SOURCE_LOCATION,[] {})
 			};
-			BusyWait(Tasks);
-		}
-
-		{	// busy-waiting for multiple tasks
-			TArray<FTask> Tasks
-			{
-				Launch(UE_SOURCE_LOCATION,[] {}),
-				Launch(UE_SOURCE_LOCATION,[] {})
-			};
-			BusyWait(Tasks, FTimespan::FromMilliseconds(10));
+			Wait(Tasks);
 		}
 
 		{	// basic use-case, postpone waiting so the task is executed first
@@ -118,17 +106,6 @@ namespace UE { namespace TasksTests
 			check(Done);
 		}
 
-		{	// basic use-case, postpone busy-waiting so the task is executed first
-			std::atomic<bool> Done{ false };
-			FTask Task = Launch(UE_SOURCE_LOCATION, [&Done] { Done = true; });
-			while (!Task.IsCompleted())
-			{
-				FPlatformProcess::Yield();
-			}
-			Task.BusyWait();
-			check(Done);
-		}
-
 		{	// basic use-case with result, postpone execution so waiting kicks in first
 			TTask<int> Task = Launch(UE_SOURCE_LOCATION, [] { FPlatformProcess::Sleep(0.1f);  return 42; });
 			verify(Task.GetResult() == 42);
@@ -969,35 +946,32 @@ namespace UE { namespace TasksTests
 		UE_BENCHMARK(5, DependenciesPerfTest<150, 150>);
 	}
 
-	// blocks all workers (except reserve workers) until given event is triggered. Returns blocking tasks.
-	TArray<LowLevelTasks::FTask> BlockWorkers(FTaskEvent& ResumeEvent, uint32 NumWorkers = LowLevelTasks::FScheduler::Get().GetNumWorkers())
+	// blocks all workers until given event is triggered. Returns blocking tasks.
+	TArray<FTask> BlockWorkers(FTaskEvent& ResumeEvent, uint32 NumWorkers = LowLevelTasks::FScheduler::Get().GetNumWorkers())
 	{
 		TRACE_CPUPROFILER_EVENT_SCOPE(BlockWorkers);
 
-		TArray<LowLevelTasks::FTask> WorkerBlockers; // tasks that block worker threads
+		TArray<FTask> WorkerBlockers; // tasks that block worker threads
 		WorkerBlockers.Reserve(NumWorkers);
 
 		std::atomic<uint32>   NumWorkersBlocked{ 0 };
 		UE::FManualResetEvent AllWorkersBlocked;
 
-		for (int i = 0; i != NumWorkers; ++i)
+		for (int32 Index = 0; Index != NumWorkers; ++Index)
 		{
-			WorkerBlockers.Emplace();
-			LowLevelTasks::FTask& Task = WorkerBlockers.Last();
-			Task.Init(TEXT("WorkerBlocker"),
-				[&NumWorkersBlocked, &NumWorkers, &ResumeEvent, &AllWorkersBlocked]
-				{
-					checkf(LowLevelTasks::FScheduler::Get().IsWorkerThread(), TEXT("No reserve workers are expected to get blocked"));
-					if (++NumWorkersBlocked == NumWorkers)
+			WorkerBlockers.Emplace(
+				UE::Tasks::Launch(
+					TEXT("WorkerBlocker"),
+					[&NumWorkersBlocked, &NumWorkers, &ResumeEvent, &AllWorkersBlocked]
 					{
-						AllWorkersBlocked.Notify();
+						if (++NumWorkersBlocked == NumWorkers)
+						{
+							AllWorkersBlocked.Notify();
+						}
+						TRACE_CPUPROFILER_EVENT_SCOPE(BlockWorkers_Blocked);
+						ResumeEvent.Wait();
 					}
-					TRACE_CPUPROFILER_EVENT_SCOPE(BlockWorkers_Blocked);
-					ResumeEvent.Wait();
-				},
-				LowLevelTasks::ETaskFlags::AllowNothing
-			);
-			LowLevelTasks::FScheduler::Get().TryLaunch(Task);
+			));
 		}
 
 		TRACE_CPUPROFILER_EVENT_SCOPE(WaitingUntilAllWorkersBlocked);
@@ -1035,10 +1009,10 @@ namespace UE { namespace TasksTests
 
 	TEST_CASE_NAMED(FTasksDeepRetractionTest, "System::Core::Tasks::DeepRetraction", "[.][ApplicationContextMask][EngineFilter]")
 	{
-		FPlatformProcess::Sleep(0.1f); // give workers time to fall asleep, to avoid any reserve worker messing around
+		FPlatformProcess::Sleep(0.1f); // give workers time to fall asleep
 
 		FTaskEvent ResumeEvent{ UE_SOURCE_LOCATION };
-		TArray<LowLevelTasks::FTask> WorkerBlockers = BlockWorkers(ResumeEvent);
+		TArray<FTask> WorkerBlockers = BlockWorkers(ResumeEvent);
 
 		{	// basic retraction, no dependencies
 			bool bDone = false;
@@ -1127,7 +1101,7 @@ namespace UE { namespace TasksTests
 		//}
 
 		ResumeEvent.Trigger();
-		LowLevelTasks::BusyWaitForTasks<LowLevelTasks::FTask>(WorkerBlockers);
+		Wait(WorkerBlockers);
 	}
 
 	template<uint32 Num>
@@ -1673,10 +1647,10 @@ namespace UE { namespace TasksTests
 		for (int Index = 0; Index < 1000; ++Index)
 		{
 			FTaskEvent ResumeEvent{ UE_SOURCE_LOCATION };
-			TArray<LowLevelTasks::FTask> WorkerBlockers = BlockWorkers(ResumeEvent, NumWorkers);
+			TArray<FTask> WorkerBlockers = BlockWorkers(ResumeEvent, NumWorkers);
 
 			ResumeEvent.Trigger();
-			LowLevelTasks::BusyWaitForTasks<LowLevelTasks::FTask>(WorkerBlockers);
+			Wait(WorkerBlockers);
 		}
 	}
 
