Load game localization data in cooked editors
#rnx

diff --git a/Engine/Source/Runtime/Core/Private/Internationalization/LocalizationResourceTextSource.cpp b/Engine/Source/Runtime/Core/Private/Internationalization/LocalizationResourceTextSource.cpp
--- a/Engine/Source/Runtime/Core/Private/Internationalization/LocalizationResourceTextSource.cpp
+++ b/Engine/Source/Runtime/Core/Private/Internationalization/LocalizationResourceTextSource.cpp
@@ -103,7 +103,7 @@ void FLocalizationResourceTextSource::LoadLocalizedResources(const ELocalization
 	{
 		EditorLocalizationPaths += FPaths::GetEditorLocalizationPaths();
 #if UE_IS_COOKED_EDITOR
-		EditorLocalizationPaths += GetCookedEditorLocalizationPaths();
+		EditorLocalizationPaths += GetGameLocalizationPaths();
 #endif
 		EditorLocalizationPaths += FPaths::GetToolTipLocalizationPaths();
 
diff --git a/Engine/Source/Runtime/Projects/Private/LocalizationDescriptor.cpp b/Engine/Source/Runtime/Projects/Private/LocalizationDescriptor.cpp
--- a/Engine/Source/Runtime/Projects/Private/LocalizationDescriptor.cpp
+++ b/Engine/Source/Runtime/Projects/Private/LocalizationDescriptor.cpp
@@ -256,7 +256,11 @@ bool FLocalizationTargetDescriptor::ShouldLoadLocalizationTarget() const
 		return WITH_EDITOR;
 
 	case ELocalizationTargetDescriptorLoadingPolicy::Game:
+#if UE_IS_COOKED_EDITOR
+		return true;
+#else
 		return FApp::IsGame();
+#endif
 
 	case ELocalizationTargetDescriptorLoadingPolicy::PropertyNames:
 #if WITH_EDITOR
