Replace private FScopedReadTracker with a generic version FFileHandleRegistryReadTracker that was recently added to PlatformFileCommon.h

#rnx

diff --git a/Engine/Source/Runtime/Core/Private/Unix/UnixPlatformFile.cpp b/Engine/Source/Runtime/Core/Private/Unix/UnixPlatformFile.cpp
--- a/Engine/Source/Runtime/Core/Private/Unix/UnixPlatformFile.cpp
+++ b/Engine/Source/Runtime/Core/Private/Unix/UnixPlatformFile.cpp
@@ -189,29 +189,12 @@ public:
 
 	virtual bool Read(uint8* Destination, int64 BytesToRead) override
 	{
-		struct FScopedReadTracker
-		{
-			FScopedReadTracker(FFileHandleUnix& InHandle) : Handle(InHandle) 
-			{ 
-				bSuccess = GFileRegistry.TrackStartRead(&Handle);
-			}
-			~FScopedReadTracker() 
-			{
-				if (bSuccess)
-				{
-					GFileRegistry.TrackEndRead(&Handle); 
-				}
-			}
-			FFileHandleUnix& Handle;
-			bool bSuccess = false;
-		};
-
 		check(IsValid());
 		if (!FileOpenAsWrite)
 		{
 			// Handle virtual file handles (only in read mode, write mode doesn't use the file handle registry)
-			FScopedReadTracker ScopedReadTracker(*this);
-			if (!ScopedReadTracker.bSuccess)
+			FFileHandleRegistryReadTracker TrackRead(GFileRegistry, *this, true);
+			if (!TrackRead.IsValid())
 			{
 				return false;
 			}
