Added basic tests for TSharedString and fixed its use of placement new

#rnx

diff --git a/Engine/Source/Runtime/Core/Public/Containers/SharedString.h b/Engine/Source/Runtime/Core/Public/Containers/SharedString.h
--- a/Engine/Source/Runtime/Core/Public/Containers/SharedString.h
+++ b/Engine/Source/Runtime/Core/Public/Containers/SharedString.h
@@ -159,8 +159,8 @@ TSharedString<CharType>::TSharedString(const TStringView<CharType> String)
 		static_assert(sizeof(int32) == sizeof(std::atomic<int32>));
 		const SIZE_T Size = sizeof(int32) + sizeof(int32) + sizeof(CharType) + sizeof(CharType) * Length;
 		int32* const Header = static_cast<int32*>(FMemory::Malloc(Size));
-		new(&Header[0]) std::atomic<int32>(1); // Ref Count
-		new(&Header[1]) int32(Length);
+		new((void*)&Header[0]) std::atomic<int32>(1); // Ref Count
+		new((void*)&Header[1]) int32(Length);
 		Chars = reinterpret_cast<CharType*>(&Header[2]);
 		String.CopyString(Chars, Length);
 		Chars[Length] = CharType(0);
diff --git a/Engine/Source/Runtime/Core/Tests/Containers/SharedStringTest.cpp b/Engine/Source/Runtime/Core/Tests/Containers/SharedStringTest.cpp
new file mode 100644
--- /dev/null
+++ b/Engine/Source/Runtime/Core/Tests/Containers/SharedStringTest.cpp
@@ -0,0 +1,63 @@
+// Copyright Epic Games, Inc. All Rights Reserved.
+
+#if WITH_TESTS
+
+#include "Containers/SharedString.h"
+
+#include "Tests/TestHarnessAdapter.h"
+
+namespace UE
+{
+
+TEST_CASE_NAMED(FSharedStringTest, "System::Core::String::SharedString", "[Core][String][SmokeFilter]")
+{
+	SECTION("Empty")
+	{
+		const FSharedString String;
+		CHECK(String.IsEmpty());
+		CHECK(String.Len() == 0);
+		CHECK(**String == TEXT('\0'));
+		CHECK(String == FSharedString::Empty);
+	}
+
+	SECTION("Construct")
+	{
+		const FStringView View = TEXTVIEW("ABC");
+		const FSharedString String = View;
+		CHECK(!String.IsEmpty());
+		CHECK(String.Len() == View.Len());
+		CHECK(**String == View[0]);
+		CHECK(*String != View.GetData());
+		CHECK(String == View);
+		CHECK_FALSE(String == FSharedString::Empty);
+	}
+
+	SECTION("Copy/Move")
+	{
+		const FStringView View = TEXTVIEW("ABC");
+		FSharedString String = View;
+		const TCHAR* StringData = *String;
+
+		FSharedString StringCopy = String;
+		CHECK(*StringCopy == StringData);
+
+		String.Reset();
+		CHECK(String.IsEmpty());
+		CHECK(StringCopy == View);
+
+		String = StringCopy;
+		CHECK(*String == StringData);
+
+		String = MoveTemp(StringCopy);
+		CHECK(*String == StringData);
+		CHECK(StringCopy.IsEmpty());
+
+		StringCopy = MoveTemp(String);
+		CHECK(*StringCopy == StringData);
+		CHECK(String.IsEmpty());
+	}
+}
+
+} // UE
+
+#endif // WITH_TESTS
