Removed the now-replaced EAutomationTestFlags::ApplicationContextMask.

diff --git a/Engine/Source/Runtime/Core/Public/Misc/AutomationTest.h b/Engine/Source/Runtime/Core/Public/Misc/AutomationTest.h
--- a/Engine/Source/Runtime/Core/Public/Misc/AutomationTest.h
+++ b/Engine/Source/Runtime/Core/Public/Misc/AutomationTest.h
@@ -123,9 +123,6 @@ enum class EAutomationTestFlags
 	StressFilter				= 0x10000000,
 	//Negative Test. For tests whose correct expected outcome is failure.
 	NegativeFilter				= 0x20000000,
-
-	// We need to deprecate and remove this once all usage has been replaced
-	ApplicationContextMask = EditorContext | ClientContext | ServerContext | CommandletContext | ProgramContext,
 };
 
 ENUM_CLASS_FLAGS(EAutomationTestFlags)
