Add CSV_PROFILER_MINIMAL define for removing almost all instrumentation overhead. Only a handful of essential stats are captured; the rest are stripped out at compile time. In this mode, the only function which emits stats to the CSV is RecordCustomStatMinimal (with corresponding macros CSV_CUSTOM_STAT_MINIMAL and CSV_CUSTOM_STAT_MINIMAL_GLOBAL).
#tests Tested locally with and without the new cvar

diff --git a/Engine/Source/Runtime/Core/Private/ProfilingDebugging/CsvProfiler.cpp b/Engine/Source/Runtime/Core/Private/ProfilingDebugging/CsvProfiler.cpp
--- a/Engine/Source/Runtime/Core/Private/ProfilingDebugging/CsvProfiler.cpp
+++ b/Engine/Source/Runtime/Core/Private/ProfilingDebugging/CsvProfiler.cpp
@@ -69,10 +69,6 @@ CSV_DEFINE_CATEGORY(CsvProfiler, true);
 	void CSVTest();
 #endif
 
-CSV_DEFINE_STAT_GLOBAL(FrameTime);
-
-#define RECORD_TIMESTAMPS 1 
-
 #define LIST_VALIDATION (DO_CHECK && 0)
 
 DEFINE_LOG_CATEGORY_STATIC(LogCsvProfiler, Log, All);
@@ -3142,7 +3138,7 @@ void FCsvProfiler::EndFrame()
 		uint64 CurrentTimeStamp = FPlatformTime::Cycles64();
 		uint64 ElapsedCycles = CurrentTimeStamp - LastEndFrameTimestamp;
 		float ElapsedMs = (float)FPlatformTime::ToMilliseconds64(ElapsedCycles);
-		CSV_CUSTOM_STAT_DEFINED(FrameTime, ElapsedMs, ECsvCustomStatOp::Set);
+		CSV_CUSTOM_STAT_MINIMAL_GLOBAL(FrameTime, ElapsedMs, ECsvCustomStatOp::Set);
 
 		FPlatformMemoryStats MemoryStats = FPlatformMemory::GetStats();
 
@@ -3161,11 +3157,11 @@ void FCsvProfiler::EndFrame()
 		TotalSystemMB -= float(FPlatformMemory::GetExtraDevelopmentMemorySize() / 1024ull / 1024ull);
 #endif
 		
-		CSV_CUSTOM_STAT_GLOBAL(MemoryFreeMB, PhysicalMBFree, ECsvCustomStatOp::Set);
-		CSV_CUSTOM_STAT_GLOBAL(PhysicalUsedMB, PhysicalMBUsed, ECsvCustomStatOp::Set);
-		CSV_CUSTOM_STAT_GLOBAL(VirtualUsedMB, VirtualMBUsed, ECsvCustomStatOp::Set);
-		CSV_CUSTOM_STAT_GLOBAL(ExtendedUsedMB, UsedExtendedMB, ECsvCustomStatOp::Set);
-		CSV_CUSTOM_STAT_GLOBAL(SystemMaxMB, TotalSystemMB, ECsvCustomStatOp::Set);
+		CSV_CUSTOM_STAT_MINIMAL_GLOBAL(MemoryFreeMB, PhysicalMBFree, ECsvCustomStatOp::Set);
+		CSV_CUSTOM_STAT_MINIMAL_GLOBAL(PhysicalUsedMB, PhysicalMBUsed, ECsvCustomStatOp::Set);
+		CSV_CUSTOM_STAT_MINIMAL_GLOBAL(VirtualUsedMB, VirtualMBUsed, ECsvCustomStatOp::Set);
+		CSV_CUSTOM_STAT_MINIMAL_GLOBAL(ExtendedUsedMB, UsedExtendedMB, ECsvCustomStatOp::Set);
+		CSV_CUSTOM_STAT_MINIMAL_GLOBAL(SystemMaxMB, TotalSystemMB, ECsvCustomStatOp::Set);
 
 		MemoryStats.SetEndFrameCsvStats();
 
@@ -3461,7 +3457,7 @@ void FCsvProfiler::SetDeviceProfileName(FString InDeviceProfileName)
 /** Push/pop events */
 void FCsvProfiler::BeginStat(const char * StatName, uint32 CategoryIndex, const char * NamedEventName)
 {
-#if RECORD_TIMESTAMPS
+#if !CSV_PROFILER_MINIMAL
 	if (GCsvProfilerIsCapturing && GCsvCategoriesEnabled[CategoryIndex])
 	{
 #if CSV_PROFILER_SUPPORT_NAMED_EVENTS
@@ -3477,7 +3473,7 @@ void FCsvProfiler::BeginStat(const char * StatName, uint32 CategoryIndex, const
 
 void FCsvProfiler::BeginStat(const FName& StatName, uint32 CategoryIndex)
 {
-#if RECORD_TIMESTAMPS
+#if !CSV_PROFILER_MINIMAL
 	if (GCsvProfilerIsCapturing && GCsvCategoriesEnabled[CategoryIndex])
 	{
 #if CSV_PROFILER_SUPPORT_NAMED_EVENTS
@@ -3493,7 +3489,7 @@ void FCsvProfiler::BeginStat(const FName& StatName, uint32 CategoryIndex)
 
 void FCsvProfiler::EndStat(const char * StatName, uint32 CategoryIndex)
 {
-#if RECORD_TIMESTAMPS
+#if !CSV_PROFILER_MINIMAL
 	if (GCsvProfilerIsCapturing && GCsvCategoriesEnabled[CategoryIndex])
 	{
 		FCsvProfilerThreadData::Get().AddTimestampEnd(StatName, CategoryIndex);
@@ -3509,7 +3505,7 @@ void FCsvProfiler::EndStat(const char * StatName, uint32 CategoryIndex)
 
 void FCsvProfiler::EndStat(const FName& StatName, uint32 CategoryIndex)
 {
-#if RECORD_TIMESTAMPS
+#if !CSV_PROFILER_MINIMAL
 	if (GCsvProfilerIsCapturing && GCsvCategoriesEnabled[CategoryIndex])
 	{
 		FCsvProfilerThreadData::Get().AddTimestampEnd(StatName, CategoryIndex);
@@ -3525,7 +3521,7 @@ void FCsvProfiler::EndStat(const FName& StatName, uint32 CategoryIndex)
 
 void FCsvProfiler::BeginExclusiveStat(const char * StatName, const char * NamedEventName)
 {
-#if RECORD_TIMESTAMPS
+#if !CSV_PROFILER_MINIMAL
 	if (GCsvProfilerIsCapturing && GCsvCategoriesEnabled[CSV_CATEGORY_INDEX(Exclusive)])
 	{
 #if CSV_PROFILER_SUPPORT_NAMED_EVENTS
@@ -3541,7 +3537,7 @@ void FCsvProfiler::BeginExclusiveStat(const char * StatName, const char * NamedE
 
 void FCsvProfiler::EndExclusiveStat(const char * StatName)
 {
-#if RECORD_TIMESTAMPS
+#if !CSV_PROFILER_MINIMAL
 	if (GCsvProfilerIsCapturing && GCsvCategoriesEnabled[CSV_CATEGORY_INDEX(Exclusive)])
 	{
 		FCsvProfilerThreadData::Get().AddTimestampExclusiveEnd(StatName);
@@ -3558,7 +3554,7 @@ void FCsvProfiler::EndExclusiveStat(const char * StatName)
 
 void FCsvProfiler::BeginSetWaitStat(const char * StatName)
 {
-#if RECORD_TIMESTAMPS
+#if !CSV_PROFILER_MINIMAL
 	if (GCsvProfilerIsCapturing && GCsvCategoriesEnabled[CSV_CATEGORY_INDEX(Exclusive)])
 	{
 #if CSV_PROFILER_SUPPORT_NAMED_EVENTS
@@ -3574,7 +3570,7 @@ void FCsvProfiler::BeginSetWaitStat(const char * StatName)
 
 void FCsvProfiler::EndSetWaitStat()
 {
-#if RECORD_TIMESTAMPS
+#if !CSV_PROFILER_MINIMAL
 	if (GCsvProfilerIsCapturing && GCsvCategoriesEnabled[CSV_CATEGORY_INDEX(Exclusive)])
 	{
 #if CSV_PROFILER_SUPPORT_NAMED_EVENTS
@@ -3590,7 +3586,7 @@ void FCsvProfiler::EndSetWaitStat()
 
 void FCsvProfiler::BeginWait()
 {
-#if RECORD_TIMESTAMPS
+#if !CSV_PROFILER_MINIMAL
 	if (GCsvProfilerIsCapturing && GCsvCategoriesEnabled[CSV_CATEGORY_INDEX(Exclusive)])
 	{
 		const char* WaitStatName = FCsvProfilerThreadData::Get().GetWaitStatName();
@@ -3617,7 +3613,7 @@ void FCsvProfiler::BeginWait()
 
 void FCsvProfiler::EndWait()
 {
-#if RECORD_TIMESTAMPS
+#if !CSV_PROFILER_MINIMAL
 	if (GCsvProfilerIsCapturing && GCsvCategoriesEnabled[CSV_CATEGORY_INDEX(Exclusive)])
 	{
 		const char* WaitStatName = FCsvProfilerThreadData::Get().GetWaitStatName();
@@ -3757,49 +3753,69 @@ void FCsvProfiler::RecordEventAtTimestamp(int32 CategoryIndex, const FString& Ev
 	}
 }
 
+void FCsvProfiler::RecordCustomStatMinimal(const char* StatName, uint32 CategoryIndex, float Value, const ECsvCustomStatOp CustomStatOp)
+{
+	if (GCsvProfilerIsCapturing && GCsvCategoriesEnabled[CategoryIndex])
+	{
+		FCsvProfilerThreadData::Get().AddCustomStat(StatName, CategoryIndex, Value, CustomStatOp);
+	}
+}
+
 void FCsvProfiler::RecordCustomStat(const char * StatName, uint32 CategoryIndex, float Value, const ECsvCustomStatOp CustomStatOp)
 {
+#if !CSV_PROFILER_MINIMAL
 	if (GCsvProfilerIsCapturing && GCsvCategoriesEnabled[CategoryIndex])
 	{
 		FCsvProfilerThreadData::Get().AddCustomStat(StatName, CategoryIndex, Value, CustomStatOp);
 	}
+#endif
 }
 
 void FCsvProfiler::RecordCustomStat(const char* StatName, uint32 CategoryIndex, double Value, const ECsvCustomStatOp CustomStatOp)
 {
+#if !CSV_PROFILER_MINIMAL
 	// LWC_TODO: Double support for FCsvProfiler::RecordCustomStat
 	RecordCustomStat(StatName, CategoryIndex, (float)Value, CustomStatOp);
+#endif
 }
 
 void FCsvProfiler::RecordCustomStat(const FName& StatName, uint32 CategoryIndex, float Value, const ECsvCustomStatOp CustomStatOp)
 {
+#if !CSV_PROFILER_MINIMAL
 	if (GCsvProfilerIsCapturing && GCsvCategoriesEnabled[CategoryIndex])
 	{
 		FCsvProfilerThreadData::Get().AddCustomStat(StatName, CategoryIndex, Value, CustomStatOp);
 	}
+#endif
 }
 
 void FCsvProfiler::RecordCustomStat(const FName& StatName, uint32 CategoryIndex, double Value, const ECsvCustomStatOp CustomStatOp)
 {
+#if !CSV_PROFILER_MINIMAL
 	// LWC_TODO: Double support for FCsvProfiler::RecordCustomStat
 	RecordCustomStat(StatName, CategoryIndex, (float)Value, CustomStatOp);
+#endif
 }
 
 
 void FCsvProfiler::RecordCustomStat(const char * StatName, uint32 CategoryIndex, int32 Value, const ECsvCustomStatOp CustomStatOp)
 {
+#if !CSV_PROFILER_MINIMAL
 	if (GCsvProfilerIsCapturing && GCsvCategoriesEnabled[CategoryIndex])
 	{
 		FCsvProfilerThreadData::Get().AddCustomStat(StatName, CategoryIndex, Value, CustomStatOp);
 	}
+#endif
 }
 
 void FCsvProfiler::RecordCustomStat(const FName& StatName, uint32 CategoryIndex, int32 Value, const ECsvCustomStatOp CustomStatOp)
 {
+#if !CSV_PROFILER_MINIMAL
 	if (GCsvProfilerIsCapturing && GCsvCategoriesEnabled[CategoryIndex])
 	{
 		FCsvProfilerThreadData::Get().AddCustomStat(StatName, CategoryIndex, Value, CustomStatOp);
 	}
+#endif
 }
 
 void FCsvProfiler::Init()
diff --git a/Engine/Source/Runtime/Core/Public/ProfilingDebugging/CsvProfiler.h b/Engine/Source/Runtime/Core/Public/ProfilingDebugging/CsvProfiler.h
--- a/Engine/Source/Runtime/Core/Public/ProfilingDebugging/CsvProfiler.h
+++ b/Engine/Source/Runtime/Core/Public/ProfilingDebugging/CsvProfiler.h
@@ -46,105 +46,119 @@ class FScopedCsvStat;
 class FScopedCsvStatExclusive;
 struct FCsvDeclaredStat;
 
+
 #if CSV_PROFILER
+	// Helpers
+	#define CSV_CATEGORY_INDEX(CategoryName)									(_GCsvCategory_##CategoryName.Index)
+	#define CSV_CATEGORY_INDEX_GLOBAL											(0)
+
+	// Categories
+	#define CSV_DEFINE_CATEGORY(CategoryName,bDefaultValue)						FCsvCategory _GCsvCategory_##CategoryName(TEXT(#CategoryName),bDefaultValue)
+	#define CSV_DEFINE_CATEGORY_MODULE(Module_API,CategoryName,bDefaultValue)	FCsvCategory Module_API _GCsvCategory_##CategoryName(TEXT(#CategoryName),bDefaultValue)
+	#define CSV_DECLARE_CATEGORY_EXTERN(CategoryName)							extern FCsvCategory _GCsvCategory_##CategoryName
+	#define CSV_DECLARE_CATEGORY_MODULE_EXTERN(Module_API,CategoryName)			extern Module_API FCsvCategory _GCsvCategory_##CategoryName
+
+	// Events
+	#define CSV_EVENT(Category, Format, ...) \
+		FCsvProfiler::RecordEventf( CSV_CATEGORY_INDEX(Category), Format, ##__VA_ARGS__ ); \
+		TRACE_BOOKMARK(TEXT(PREPROCESSOR_TO_STRING(Category)) TEXT("/") Format, ##__VA_ARGS__)
+
+	#define CSV_EVENT_GLOBAL(Format, ...) \
+		FCsvProfiler::RecordEventf( CSV_CATEGORY_INDEX_GLOBAL, Format, ##__VA_ARGS__ ); \
+		TRACE_BOOKMARK(Format, ##__VA_ARGS__)
+
+	// Metadata
+	#define CSV_METADATA(Key,Value)									FCsvProfiler::SetMetadata( Key, Value )
+	#define CSV_NON_PERSISTENT_METADATA(Key,Value)					FCsvProfiler::SetNonPersistentMetadata( Key, Value )
+
+	// Minimal stat macros (emitted even when CSV_PROFILER_MINIMAL is 1)
+	// IMPORTANT: Only a handful of key stats are expected to use these. Do not add more without very good reason!
+	#define CSV_CUSTOM_STAT_MINIMAL(Category,StatName,Value,Op) 	FCsvProfiler::RecordCustomStatMinimal(#StatName, CSV_CATEGORY_INDEX(Category), Value, Op);
+	#define CSV_CUSTOM_STAT_MINIMAL_GLOBAL(StatName,Value,Op) 		FCsvProfiler::RecordCustomStatMinimal(#StatName, CSV_CATEGORY_INDEX_GLOBAL, Value, Op); 
 
-// Helpers
-#define CSV_CATEGORY_INDEX(CategoryName)						(_GCsvCategory_##CategoryName.Index)
-#define CSV_CATEGORY_INDEX_GLOBAL								(0)
-#define CSV_STAT_FNAME(StatName)								(_GCsvStat_##StatName.Name)
-
-// Inline stats (no up front definition)
-#define CSV_SCOPED_TIMING_STAT(Category,StatName) \
-	TRACE_CSV_PROFILER_INLINE_STAT(#StatName, CSV_CATEGORY_INDEX(Category)); \
-	FScopedCsvStat _ScopedCsvStat_ ## StatName (#StatName, CSV_CATEGORY_INDEX(Category), "CSV_"#StatName);
-#define CSV_SCOPED_TIMING_STAT_GLOBAL(StatName) \
-	TRACE_CSV_PROFILER_INLINE_STAT(#StatName, CSV_CATEGORY_INDEX_GLOBAL); \
-	FScopedCsvStat _ScopedCsvStat_ ## StatName (#StatName, CSV_CATEGORY_INDEX_GLOBAL, "CSV_"#StatName);
-#define CSV_SCOPED_TIMING_STAT_EXCLUSIVE(StatName) \
-	TRACE_CSV_PROFILER_INLINE_STAT_EXCLUSIVE(#StatName); \
-	FScopedCsvStatExclusive _ScopedCsvStatExclusive_ ## StatName (#StatName, "CSV_"#StatName);
-#define CSV_SCOPED_TIMING_STAT_EXCLUSIVE_CONDITIONAL(StatName,Condition) \
-	TRACE_CSV_PROFILER_INLINE_STAT_EXCLUSIVE(#StatName); \
-	FScopedCsvStatExclusiveConditional _ScopedCsvStatExclusive_ ## StatName (#StatName,Condition, "CSV_"#StatName);
-#define CSV_SCOPED_TIMING_STAT_RECURSIVE(Category,StatName) \
-	TRACE_CSV_PROFILER_INLINE_STAT(#StatName, CSV_CATEGORY_INDEX(Category)); \
-	static thread_local int32 _ScopedCsvStatRecursive_EntryCount_ ## StatName = 0; \
-	FScopedCsvStatRecursive _ScopedCsvStatRecursive_ ## StatName (_ScopedCsvStatRecursive_EntryCount_ ## StatName, #StatName, CSV_CATEGORY_INDEX(Category), "CSV_"#StatName);
-#define CSV_SCOPED_TIMING_STAT_RECURSIVE_CONDITIONAL(Category,StatName,Condition) \
-	TRACE_CSV_PROFILER_INLINE_STAT(#StatName, CSV_CATEGORY_INDEX(Category)); \
-	static thread_local int32 _ScopedCsvStatRecursive_EntryCount_ ## StatName = 0; \
-	FScopedCsvStatRecursiveConditional _ScopedCsvStatRecursive_ ## StatName (_ScopedCsvStatRecursive_EntryCount_ ## StatName, #StatName, CSV_CATEGORY_INDEX(Category), Condition, "CSV_"#StatName);
-
-#define CSV_SCOPED_WAIT(WaitTime)							FScopedCsvWaitConditional _ScopedCsvWait(WaitTime>0 && FCsvProfiler::IsWaitTrackingEnabledOnCurrentThread());
-#define CSV_SCOPED_WAIT_CONDITIONAL(Condition)				FScopedCsvWaitConditional _ScopedCsvWait(Condition);
-
-#define CSV_SCOPED_SET_WAIT_STAT(StatName) \
-	TRACE_CSV_PROFILER_INLINE_STAT_EXCLUSIVE("EventWait/"#StatName); \
-	FScopedCsvSetWaitStat _ScopedCsvSetWaitStat ## StatName("EventWait/"#StatName);
-
-#define CSV_SCOPED_SET_WAIT_STAT_IGNORE()						FScopedCsvSetWaitStat _ScopedCsvSetWaitStat ## StatName();
-
-#define CSV_CUSTOM_STAT(Category,StatName,Value,Op) \
-	TRACE_CSV_PROFILER_INLINE_STAT(#StatName, CSV_CATEGORY_INDEX(Category)); \
-	FCsvProfiler::RecordCustomStat(#StatName, CSV_CATEGORY_INDEX(Category), Value, Op);
-#define CSV_CUSTOM_STAT_GLOBAL(StatName,Value,Op) \
-	TRACE_CSV_PROFILER_INLINE_STAT(#StatName, CSV_CATEGORY_INDEX_GLOBAL); \
-	FCsvProfiler::RecordCustomStat(#StatName, CSV_CATEGORY_INDEX_GLOBAL, Value, Op); 
-
-// Stats declared up front
-#define CSV_DEFINE_STAT(Category,StatName)						FCsvDeclaredStat _GCsvStat_##StatName((TCHAR*)TEXT(#StatName), CSV_CATEGORY_INDEX(Category));
-#define CSV_DEFINE_STAT_GLOBAL(StatName)						FCsvDeclaredStat _GCsvStat_##StatName((TCHAR*)TEXT(#StatName), CSV_CATEGORY_INDEX_GLOBAL);
-#define CSV_DECLARE_STAT_EXTERN(Category,StatName)				extern FCsvDeclaredStat _GCsvStat_##StatName
-#define CSV_CUSTOM_STAT_DEFINED(StatName,Value,Op)				FCsvProfiler::RecordCustomStat(_GCsvStat_##StatName.Name, _GCsvStat_##StatName.CategoryIndex, Value, Op);
-
-// Categories
-#define CSV_DEFINE_CATEGORY(CategoryName,bDefaultValue)			FCsvCategory _GCsvCategory_##CategoryName(TEXT(#CategoryName),bDefaultValue)
-#define CSV_DECLARE_CATEGORY_EXTERN(CategoryName)				extern FCsvCategory _GCsvCategory_##CategoryName
-
-#define CSV_DEFINE_CATEGORY_MODULE(Module_API,CategoryName,bDefaultValue)	FCsvCategory Module_API _GCsvCategory_##CategoryName(TEXT(#CategoryName),bDefaultValue)
-#define CSV_DECLARE_CATEGORY_MODULE_EXTERN(Module_API,CategoryName)			extern Module_API FCsvCategory _GCsvCategory_##CategoryName
-
-// Events
-#define CSV_EVENT(Category, Format, ...) \
-	FCsvProfiler::RecordEventf( CSV_CATEGORY_INDEX(Category), Format, ##__VA_ARGS__ ); \
-	TRACE_BOOKMARK(TEXT(PREPROCESSOR_TO_STRING(Category)) TEXT("/") Format, ##__VA_ARGS__)
-
-#define CSV_EVENT_GLOBAL(Format, ...) \
-	FCsvProfiler::RecordEventf( CSV_CATEGORY_INDEX_GLOBAL, Format, ##__VA_ARGS__ ); \
-	TRACE_BOOKMARK(Format, ##__VA_ARGS__)
-
-// Metadata
-#define CSV_METADATA(Key,Value)									FCsvProfiler::SetMetadata( Key, Value )
-#define CSV_NON_PERSISTENT_METADATA(Key,Value)					FCsvProfiler::SetNonPersistentMetadata( Key, Value )
+#else
+	#define CSV_CATEGORY_INDEX(CategoryName)	(0)
+	#define CSV_CATEGORY_INDEX_GLOBAL			(0)				
+	#define CSV_DEFINE_CATEGORY(CategoryName,bDefaultValue)			
+	#define CSV_DEFINE_CATEGORY_MODULE(Module_API,CategoryName,bDefaultValue)	
+	#define CSV_DECLARE_CATEGORY_EXTERN(CategoryName)				
+	#define CSV_DECLARE_CATEGORY_MODULE_EXTERN(Module_API,CategoryName)			
+	#define CSV_EVENT(Category, Format, ...) 						
+	#define CSV_EVENT_GLOBAL(Format, ...)
+	#define CSV_METADATA(Key,Value)
+	#define CSV_NON_PERSISTENT_METADATA(Key,Value)
+	#define CSV_CUSTOM_STAT_MINIMAL(Category,StatName,Value,Op)
+	#define CSV_CUSTOM_STAT_MINIMAL_GLOBAL(StatName,Value,Op)
+#endif
 
+// Stat macros. These are disabled if CSV_PROFILER_MINIMAL is 1 
+#if CSV_PROFILER && !CSV_PROFILER_MINIMAL
+	// Helpers
+	#define CSV_STAT_FNAME(StatName)								(_GCsvStat_##StatName.Name)
+
+	// Timing stats - no up-front definition
+	#define CSV_SCOPED_TIMING_STAT(Category,StatName) \
+		TRACE_CSV_PROFILER_INLINE_STAT(#StatName, CSV_CATEGORY_INDEX(Category)); \
+		FScopedCsvStat _ScopedCsvStat_ ## StatName (#StatName, CSV_CATEGORY_INDEX(Category), "CSV_"#StatName);
+	#define CSV_SCOPED_TIMING_STAT_GLOBAL(StatName) \
+		TRACE_CSV_PROFILER_INLINE_STAT(#StatName, CSV_CATEGORY_INDEX_GLOBAL); \
+		FScopedCsvStat _ScopedCsvStat_ ## StatName (#StatName, CSV_CATEGORY_INDEX_GLOBAL, "CSV_"#StatName);
+	#define CSV_SCOPED_TIMING_STAT_EXCLUSIVE(StatName) \
+		TRACE_CSV_PROFILER_INLINE_STAT_EXCLUSIVE(#StatName); \
+		FScopedCsvStatExclusive _ScopedCsvStatExclusive_ ## StatName (#StatName, "CSV_"#StatName);
+	#define CSV_SCOPED_TIMING_STAT_EXCLUSIVE_CONDITIONAL(StatName,Condition) \
+		TRACE_CSV_PROFILER_INLINE_STAT_EXCLUSIVE(#StatName); \
+		FScopedCsvStatExclusiveConditional _ScopedCsvStatExclusive_ ## StatName (#StatName,Condition, "CSV_"#StatName);
+	#define CSV_SCOPED_TIMING_STAT_RECURSIVE(Category,StatName) \
+		TRACE_CSV_PROFILER_INLINE_STAT(#StatName, CSV_CATEGORY_INDEX(Category)); \
+		static thread_local int32 _ScopedCsvStatRecursive_EntryCount_ ## StatName = 0; \
+		FScopedCsvStatRecursive _ScopedCsvStatRecursive_ ## StatName (_ScopedCsvStatRecursive_EntryCount_ ## StatName, #StatName, CSV_CATEGORY_INDEX(Category), "CSV_"#StatName);
+	#define CSV_SCOPED_TIMING_STAT_RECURSIVE_CONDITIONAL(Category,StatName,Condition) \
+		TRACE_CSV_PROFILER_INLINE_STAT(#StatName, CSV_CATEGORY_INDEX(Category)); \
+		static thread_local int32 _ScopedCsvStatRecursive_EntryCount_ ## StatName = 0; \
+		FScopedCsvStatRecursiveConditional _ScopedCsvStatRecursive_ ## StatName (_ScopedCsvStatRecursive_EntryCount_ ## StatName, #StatName, CSV_CATEGORY_INDEX(Category), Condition, "CSV_"#StatName);
+
+	// Waits
+	#define CSV_SCOPED_WAIT(WaitTime)							FScopedCsvWaitConditional _ScopedCsvWait(WaitTime>0 && FCsvProfiler::IsWaitTrackingEnabledOnCurrentThread());
+	#define CSV_SCOPED_WAIT_CONDITIONAL(Condition)				FScopedCsvWaitConditional _ScopedCsvWait(Condition);
+
+	#define CSV_SCOPED_SET_WAIT_STAT(StatName) \
+		TRACE_CSV_PROFILER_INLINE_STAT_EXCLUSIVE("EventWait/"#StatName); \
+		FScopedCsvSetWaitStat _ScopedCsvSetWaitStat ## StatName("EventWait/"#StatName);
+
+	#define CSV_SCOPED_SET_WAIT_STAT_IGNORE()						FScopedCsvSetWaitStat _ScopedCsvSetWaitStat ## StatName();
+
+	// Custom Stats
+	#define CSV_CUSTOM_STAT(Category,StatName,Value,Op) \
+		TRACE_CSV_PROFILER_INLINE_STAT(#StatName, CSV_CATEGORY_INDEX(Category)); \
+		FCsvProfiler::RecordCustomStat(#StatName, CSV_CATEGORY_INDEX(Category), Value, Op);
+	#define CSV_CUSTOM_STAT_GLOBAL(StatName,Value,Op) \
+		TRACE_CSV_PROFILER_INLINE_STAT(#StatName, CSV_CATEGORY_INDEX_GLOBAL); \
+		FCsvProfiler::RecordCustomStat(#StatName, CSV_CATEGORY_INDEX_GLOBAL, Value, Op); 
+
+	// Stats declarations/definitions (for FName stats)
+	#define CSV_DEFINE_STAT(Category,StatName)						FCsvDeclaredStat _GCsvStat_##StatName((TCHAR*)TEXT(#StatName), CSV_CATEGORY_INDEX(Category));
+	#define CSV_DEFINE_STAT_GLOBAL(StatName)						FCsvDeclaredStat _GCsvStat_##StatName((TCHAR*)TEXT(#StatName), CSV_CATEGORY_INDEX_GLOBAL);
+	#define CSV_DECLARE_STAT_EXTERN(Category,StatName)				extern FCsvDeclaredStat _GCsvStat_##StatName
+	#define CSV_CUSTOM_STAT_DEFINED(StatName,Value,Op)				FCsvProfiler::RecordCustomStat(_GCsvStat_##StatName.Name, _GCsvStat_##StatName.CategoryIndex, Value, Op);
 #else
-  #define CSV_CATEGORY_INDEX(CategoryName)						
-  #define CSV_CATEGORY_INDEX_GLOBAL								
-  #define CSV_STAT_FNAME(StatName)	NAME_None							
-  #define CSV_SCOPED_TIMING_STAT(Category,StatName)				
-  #define CSV_SCOPED_TIMING_STAT_GLOBAL(StatName)					
-  #define CSV_SCOPED_TIMING_STAT_EXCLUSIVE(StatName)
-  #define CSV_SCOPED_TIMING_STAT_EXCLUSIVE_CONDITIONAL(StatName,Condition)
-  #define CSV_SCOPED_TIMING_STAT_RECURSIVE(Category,StatName)
-  #define CSV_SCOPED_TIMING_STAT_RECURSIVE_CONDITIONAL(Category,StatName,Condition)
-  #define CSV_SCOPED_WAIT(WaitTime)
-  #define CSV_SCOPED_WAIT_CONDITIONAL(Condition)
-  #define CSV_SCOPED_SET_WAIT_STAT(StatName)
-  #define CSV_SCOPED_SET_WAIT_STAT_IGNORE()
-  #define CSV_CUSTOM_STAT(Category,StatName,Value,Op)				
-  #define CSV_CUSTOM_STAT_GLOBAL(StatName,Value,Op) 				
-  #define CSV_DEFINE_STAT(Category,StatName)						
-  #define CSV_DEFINE_STAT_GLOBAL(StatName)						
-  #define CSV_DECLARE_STAT_EXTERN(Category,StatName)				
-  #define CSV_CUSTOM_STAT_DEFINED(StatName,Value,Op)				
-  #define CSV_DEFINE_CATEGORY(CategoryName,bDefaultValue)			
-  #define CSV_DECLARE_CATEGORY_EXTERN(CategoryName)				
-  #define CSV_DEFINE_CATEGORY_MODULE(Module_API,CategoryName,bDefaultValue)	
-  #define CSV_DECLARE_CATEGORY_MODULE_EXTERN(Module_API,CategoryName)			
-  #define CSV_EVENT(Category, Format, ...) 						
-  #define CSV_EVENT_GLOBAL(Format, ...)
-  #define CSV_METADATA(Key,Value)
-  #define CSV_NON_PERSISTENT_METADATA(Key,Value)
+	#define CSV_STAT_FNAME(StatName)			NAME_None							
+	#define CSV_SCOPED_TIMING_STAT(Category,StatName)				
+	#define CSV_SCOPED_TIMING_STAT_GLOBAL(StatName)					
+	#define CSV_SCOPED_TIMING_STAT_EXCLUSIVE(StatName)
+	#define CSV_SCOPED_TIMING_STAT_EXCLUSIVE_CONDITIONAL(StatName,Condition)
+	#define CSV_SCOPED_TIMING_STAT_RECURSIVE(Category,StatName)
+	#define CSV_SCOPED_TIMING_STAT_RECURSIVE_CONDITIONAL(Category,StatName,Condition)
+	#define CSV_SCOPED_WAIT(WaitTime)
+	#define CSV_SCOPED_WAIT_CONDITIONAL(Condition)
+	#define CSV_SCOPED_SET_WAIT_STAT(StatName)
+	#define CSV_SCOPED_SET_WAIT_STAT_IGNORE()
+	#define CSV_CUSTOM_STAT(Category,StatName,Value,Op)				
+	#define CSV_CUSTOM_STAT_GLOBAL(StatName,Value,Op) 				
+	#define CSV_DEFINE_STAT(Category,StatName)						
+	#define CSV_DEFINE_STAT_GLOBAL(StatName)						
+	#define CSV_DECLARE_STAT_EXTERN(Category,StatName)				
+	#define CSV_CUSTOM_STAT_DEFINED(StatName,Value,Op)				
 #endif
 
 
@@ -369,6 +383,9 @@ public:
 	CORE_API static void RecordCustomStat(const char * StatName, uint32 CategoryIndex, int32 Value, const ECsvCustomStatOp CustomStatOp);
 	CORE_API static void RecordCustomStat(const FName& StatName, uint32 CategoryIndex, int32 Value, const ECsvCustomStatOp CustomStatOp);
 
+	// Warning: Don't use this unless the stat is one of a handful of key stats that we want to capture in CSV_PROFILER_MINIMAL mode
+	CORE_API static void RecordCustomStatMinimal(const char* StatName, uint32 CategoryIndex, float Value, const ECsvCustomStatOp CustomStatOp);
+
 	CORE_API static void RecordEvent(int32 CategoryIndex, const FString& EventText);
 	CORE_API static void RecordEventAtFrameStart(int32 CategoryIndex, const FString& EventText);
 	CORE_API static void RecordEventAtTimestamp(int32 CategoryIndex, const FString& EventText, uint64 Cycles64);
diff --git a/Engine/Source/Runtime/Core/Public/ProfilingDebugging/CsvProfilerConfig.h b/Engine/Source/Runtime/Core/Public/ProfilingDebugging/CsvProfilerConfig.h
--- a/Engine/Source/Runtime/Core/Public/ProfilingDebugging/CsvProfilerConfig.h
+++ b/Engine/Source/Runtime/Core/Public/ProfilingDebugging/CsvProfilerConfig.h
@@ -18,8 +18,16 @@
 #define CSV_PROFILER_USE_CUSTOM_FRAME_TIMINGS 0
 #endif
 
+// Enable this to run the CSV profiler with near-zero instrumentation overhead. Only a handful of key stats will be captured
+#ifndef CSV_PROFILER_MINIMAL
+#define CSV_PROFILER_MINIMAL 0
+#endif
+
+
 // CSV_PROFILER default enabling rules, if not specified explicitly in <Program>.Target.cs GlobalDefinitions
 #ifndef CSV_PROFILER
 	#define CSV_PROFILER (WITH_ENGINE && (!UE_BUILD_SHIPPING || CSV_PROFILER_ENABLE_IN_SHIPPING))
 #endif
 
+// CSV_PROFILER_STATS is disabled if CSV_PROFILER_MINIMAL is enabled. Profiling subsystems should be predicated on this rather than CSV_PROFILER
+#define CSV_PROFILER_STATS (CSV_PROFILER && !CSV_PROFILER_MINIMAL)
\ No newline at end of file
diff --git a/Engine/Source/Runtime/Launch/Private/LaunchEngineLoop.cpp b/Engine/Source/Runtime/Launch/Private/LaunchEngineLoop.cpp
--- a/Engine/Source/Runtime/Launch/Private/LaunchEngineLoop.cpp
+++ b/Engine/Source/Runtime/Launch/Private/LaunchEngineLoop.cpp
@@ -1727,18 +1727,18 @@ static void UpdateCoreCsvStats_EndFrame()
 {
 	if (!IsRunningDedicatedServer())
 	{
-	    CSV_CUSTOM_STAT_GLOBAL(RenderThreadTime, FPlatformTime::ToMilliseconds(GRenderThreadTime), ECsvCustomStatOp::Set);
-	    CSV_CUSTOM_STAT_GLOBAL(GameThreadTime, FPlatformTime::ToMilliseconds(GGameThreadTime), ECsvCustomStatOp::Set);
-	    CSV_CUSTOM_STAT_GLOBAL(GPUTime, FPlatformTime::ToMilliseconds(GGPUFrameTime), ECsvCustomStatOp::Set);
-		CSV_CUSTOM_STAT_GLOBAL(RenderThreadTime_CriticalPath, FPlatformTime::ToMilliseconds(GRenderThreadTimeCriticalPath), ECsvCustomStatOp::Set);
-		CSV_CUSTOM_STAT_GLOBAL(GameThreadTime_CriticalPath, FPlatformTime::ToMilliseconds(GGameThreadTimeCriticalPath), ECsvCustomStatOp::Set);
+		CSV_CUSTOM_STAT_MINIMAL_GLOBAL(RenderThreadTime, FPlatformTime::ToMilliseconds(GRenderThreadTime), ECsvCustomStatOp::Set);
+		CSV_CUSTOM_STAT_MINIMAL_GLOBAL(GameThreadTime, FPlatformTime::ToMilliseconds(GGameThreadTime), ECsvCustomStatOp::Set);
+		CSV_CUSTOM_STAT_MINIMAL_GLOBAL(GPUTime, FPlatformTime::ToMilliseconds(GGPUFrameTime), ECsvCustomStatOp::Set);
+		CSV_CUSTOM_STAT_MINIMAL_GLOBAL(RenderThreadTime_CriticalPath, FPlatformTime::ToMilliseconds(GRenderThreadTimeCriticalPath), ECsvCustomStatOp::Set);
+		CSV_CUSTOM_STAT_MINIMAL_GLOBAL(GameThreadTime_CriticalPath, FPlatformTime::ToMilliseconds(GGameThreadTimeCriticalPath), ECsvCustomStatOp::Set);
 		if (IsRunningRHIInSeparateThread())
 	    {
-		    CSV_CUSTOM_STAT_GLOBAL(RHIThreadTime, FPlatformTime::ToMilliseconds(GRHIThreadTime), ECsvCustomStatOp::Set);
+			CSV_CUSTOM_STAT_MINIMAL_GLOBAL(RHIThreadTime, FPlatformTime::ToMilliseconds(GRHIThreadTime), ECsvCustomStatOp::Set);
 	    }
 	    if (GInputLatencyTime > 0)
 	    {
-		    CSV_CUSTOM_STAT_GLOBAL(InputLatencyTime, FPlatformTime::ToMilliseconds64(GInputLatencyTime), ECsvCustomStatOp::Set);
+			CSV_CUSTOM_STAT_MINIMAL_GLOBAL(InputLatencyTime, FPlatformTime::ToMilliseconds64(GInputLatencyTime), ECsvCustomStatOp::Set);
 	    }
 	    FPlatformMemoryStats MemoryStats = PlatformMemoryHelpers::GetFrameMemoryStats();
 	    float PhysicalMBFree = float(MemoryStats.AvailablePhysical / 1024) / 1024.0f;
@@ -1746,9 +1746,9 @@ static void UpdateCoreCsvStats_EndFrame()
 	    // Subtract any extra development memory from physical free. This can result in negative values in cases where we would have crashed OOM
 	    PhysicalMBFree -= float(FPlatformMemory::GetExtraDevelopmentMemorySize() / 1024ull / 1024ull);
 #endif
-	    CSV_CUSTOM_STAT_GLOBAL(MemoryFreeMB, PhysicalMBFree, ECsvCustomStatOp::Set);
+		CSV_CUSTOM_STAT_MINIMAL_GLOBAL(MemoryFreeMB, PhysicalMBFree, ECsvCustomStatOp::Set);
 
-#if !UE_BUILD_SHIPPING
+#if !UE_BUILD_SHIPPING && !CSV_PROFILER_MINIMAL
 	    float TargetFPS = 30.0f;
 	    static IConsoleVariable* MaxFPSCVar = IConsoleManager::Get().FindConsoleVariable(TEXT("t.MaxFPS"));
 	    if (MaxFPSCVar && MaxFPSCVar->GetFloat() > 0)
@@ -5499,7 +5499,7 @@ static inline void BeginFrameRenderThread(FRHICommandListImmediate& RHICmdList,
 		UEngine::SetRenderSubmitLatencyMarkerStart(CurrentFrameCounter);
 	});
 
-#if CSV_PROFILER
+#if CSV_PROFILER_STATS
 	FCsvProfiler::BeginExclusiveStat("RenderThreadOther");
 #endif
 
@@ -5518,7 +5518,7 @@ static inline void EndFrameRenderThread(FRHICommandListImmediate& RHICmdList, ui
 		return;
 	}
 
-#if CSV_PROFILER
+#if CSV_PROFILER_STATS
 	FCsvProfiler::EndExclusiveStat("RenderThreadOther");
 #endif
 
