CsvProfiler: add float and int overloads for RecordCustomStatMinimal

diff --git a/Engine/Source/Runtime/Core/Private/ProfilingDebugging/CsvProfiler.cpp b/Engine/Source/Runtime/Core/Private/ProfilingDebugging/CsvProfiler.cpp
--- a/Engine/Source/Runtime/Core/Private/ProfilingDebugging/CsvProfiler.cpp
+++ b/Engine/Source/Runtime/Core/Private/ProfilingDebugging/CsvProfiler.cpp
@@ -3761,6 +3761,21 @@ void FCsvProfiler::RecordCustomStatMinimal(const char* StatName, uint32 Category
 	}
 }
 
+void FCsvProfiler::RecordCustomStatMinimal(const char* StatName, uint32 CategoryIndex, int32 Value, const ECsvCustomStatOp CustomStatOp)
+{
+	if (GCsvProfilerIsCapturing && GCsvCategoriesEnabled[CategoryIndex])
+	{
+		FCsvProfilerThreadData::Get().AddCustomStat(StatName, CategoryIndex, Value, CustomStatOp);
+	}
+}
+
+void FCsvProfiler::RecordCustomStatMinimal(const char* StatName, uint32 CategoryIndex, double Value, const ECsvCustomStatOp CustomStatOp)
+{
+	// LWC_TODO: Double support for FCsvProfiler::RecordCustomStat
+	RecordCustomStatMinimal(StatName, CategoryIndex, (float)Value, CustomStatOp);
+}
+
+
 void FCsvProfiler::RecordCustomStat(const char * StatName, uint32 CategoryIndex, float Value, const ECsvCustomStatOp CustomStatOp)
 {
 #if !CSV_PROFILER_MINIMAL
diff --git a/Engine/Source/Runtime/Core/Public/ProfilingDebugging/CsvProfiler.h b/Engine/Source/Runtime/Core/Public/ProfilingDebugging/CsvProfiler.h
--- a/Engine/Source/Runtime/Core/Public/ProfilingDebugging/CsvProfiler.h
+++ b/Engine/Source/Runtime/Core/Public/ProfilingDebugging/CsvProfiler.h
@@ -385,6 +385,8 @@ public:
 
 	// Warning: Don't use this unless the stat is one of a handful of key stats that we want to capture in CSV_PROFILER_MINIMAL mode
 	CORE_API static void RecordCustomStatMinimal(const char* StatName, uint32 CategoryIndex, float Value, const ECsvCustomStatOp CustomStatOp);
+	CORE_API static void RecordCustomStatMinimal(const char* StatName, uint32 CategoryIndex, int32 Value, const ECsvCustomStatOp CustomStatOp);
+	CORE_API static void RecordCustomStatMinimal(const char* StatName, uint32 CategoryIndex, double Value, const ECsvCustomStatOp CustomStatOp);
 
 	CORE_API static void RecordEvent(int32 CategoryIndex, const FString& EventText);
 	CORE_API static void RecordEventAtFrameStart(int32 CategoryIndex, const FString& EventText);
