[TaskGraph]
 - Fix FTaskEvent trigger not waking up workers by default

diff --git a/Engine/Source/Runtime/Core/Public/Tasks/TaskPrivate.h b/Engine/Source/Runtime/Core/Public/Tasks/TaskPrivate.h
--- a/Engine/Source/Runtime/Core/Public/Tasks/TaskPrivate.h
+++ b/Engine/Source/Runtime/Core/Public/Tasks/TaskPrivate.h
@@ -538,7 +538,8 @@ public:
 
 				// Push the first subsequent to the local queue so we pick it up directly as our next task.
 				// This saves us the cost of going to the global queue and performing a wake-up.
-				bool bWakeUpWorker = false;
+				// But if we're a task event, always wake up new workers because the current task could continue executing for a long time after the trigger.
+				bool bWakeUpWorker = ExtendedPriority == EExtendedTaskPriority::TaskEvent;
 
 				for (FTaskBase* Subsequent : Subsequents.Close())
 				{
