Removed deprecated functions from FAutomationTestBase.

diff --git a/Engine/Source/Runtime/Core/Public/Misc/AutomationTest.h b/Engine/Source/Runtime/Core/Public/Misc/AutomationTest.h
--- a/Engine/Source/Runtime/Core/Public/Misc/AutomationTest.h
+++ b/Engine/Source/Runtime/Core/Public/Misc/AutomationTest.h
@@ -1584,18 +1584,6 @@ public:
 	 */
 	CORE_API bool GetLastExecutionSuccessState();
 
-	/**
-	 * [Deprecated] Use AddError(msg) instead to change the state of the test to a failure
-	 */
-	UE_DEPRECATED(5.1, "Use AddError(msg) instead to change the state of the test to a failure.")
-	void SetSuccessState(bool bSuccessful) { }
-
-	/**
-	 * [Deprecated] Return the last success state for this test
-	 */
-	UE_DEPRECATED(5.1, "Use GetLastExecutionSuccessState instead.")
-	bool GetSuccessState() { return GetLastExecutionSuccessState(); }
-
 	/**
 	 * Populate the provided execution info object with the execution info contained within the test. Not particularly efficient,
 	 * but providing direct access to the test's private execution info could result in errors.
