Add an option to only use the platform tracker when using LLM.
This is useful when wanting to test low memory situations with an external memory profiler in constrained memory environments.

diff --git a/Engine/Source/Runtime/Core/Private/HAL/LowLevelMemTracker.cpp b/Engine/Source/Runtime/Core/Private/HAL/LowLevelMemTracker.cpp
--- a/Engine/Source/Runtime/Core/Private/HAL/LowLevelMemTracker.cpp
+++ b/Engine/Source/Runtime/Core/Private/HAL/LowLevelMemTracker.cpp
@@ -27,6 +27,9 @@
 #include "Templates/Atomic.h"
 #include "Trace/Trace.inl"
 
+// Enable to allow memory tracking to be passed through to external tools without the extra overhead of LLM also tracking allocations.
+#define UE_ONLY_USE_PLATFORM_TRACKER 0
+
 #if UE_ENABLE_ARRAY_SLACK_TRACKING
 
 // Specifies whether to generate the whole log file in memory before writing.  Switch uses an async thread for file writing,
@@ -1574,7 +1577,9 @@ protected:
 	// The total tracked memory when the snapshot was taken
 	int64 TrackedTotalInSnapshot;
 
+#if !UE_ONLY_USE_PLATFORM_TRACKER
 	FLLMAllocMap AllocationMap;
+#endif
 
 	FTrackerTagSizeMap TagSizes;
 
@@ -4714,7 +4719,9 @@ void FLLMTracker::Initialise(
 	TraceWriter.SetTracker(InTracker);
 	CsvProfilerWriter.SetTracker(InTracker);
 
+#if !UE_ONLY_USE_PLATFORM_TRACKER
 	AllocationMap.SetAllocator(InAllocator);
+#endif
 }
 
 FLLMThreadState* FLLMTracker::GetOrCreateState()
@@ -4853,6 +4860,7 @@ void FLLMTracker::TrackAllocation(const void* Ptr, int64 Size, const FTagData* A
 	// track on the thread state
 	State->TrackAllocation(Ptr, Size, Tracker, AllocType, ActiveTagData, AssetTagData, AssetClassTagData, bTrackInMemPro);
 
+#if !UE_ONLY_USE_PLATFORM_TRACKER
 	// tracking a nullptr with a Size is allowed, but we don't need to remember it, since we can't free it ever.
 	if (Ptr != nullptr)
 	{
@@ -4869,10 +4877,12 @@ void FLLMTracker::TrackAllocation(const void* Ptr, int64 Size, const FTagData* A
 		PointerKey Key(Ptr, SizeHigh);
 		AllocationMap.Add(Key, SizeLow, AllocInfo);
 	}
+#endif
 }
 
 void FLLMTracker::TrackFree(const void* Ptr, ELLMAllocType AllocType, bool bTrackInMemPro)
 {
+#if !UE_ONLY_USE_PLATFORM_TRACKER
 	// look up the pointer in the tracking map
 	FLLMAllocMap::Values Values;
 	{
@@ -4906,12 +4916,20 @@ void FLLMTracker::TrackFree(const void* Ptr, ELLMAllocType AllocType, bool bTrac
 	const FTagData* AssetTagData = nullptr;
 	const FTagData* AssetClassTagData = nullptr;
 #endif
+#else
+	FLLMThreadState* State = GetOrCreateState();
+	int64 Size = 0;
+	const FTagData* TagData = nullptr;
+	const FTagData* AssetTagData = nullptr;
+	const FTagData* AssetClassTagData = nullptr;
+#endif
 
 	State->TrackFree(Ptr, Size, Tracker, AllocType, TagData, AssetTagData, AssetClassTagData, bTrackInMemPro);
 }
 
 void FLLMTracker::OnAllocMoved(const void* Dest, const void* Source, ELLMAllocType AllocType)
 {
+#if !UE_ONLY_USE_PLATFORM_TRACKER
 	FLLMAllocMap::Values Values;
 	{
 		if (!AllocationMap.Remove(PointerKey(Source), Values))
@@ -4934,7 +4952,11 @@ void FLLMTracker::OnAllocMoved(const void* Dest, const void* Source, ELLMAllocTy
 	int64 Size = (SizeHigh << 32ull) | SizeLow;
 	const FLLMTracker::FLowLevelAllocInfo& AllocInfo = Values.Value2;
 	const FTagData* TagData = AllocInfo.GetTag(LLMRef);
-
+#else
+	// TODO: This loses the memory size of an allocation moved within a larger allocation.
+	int64 Size = 0;
+	const FTagData* TagData = nullptr;
+#endif
 	FLLMThreadState* State = GetOrCreateState();
 	State->TrackMoved(Dest, Source, Size, Tracker, TagData);
 }
@@ -5060,7 +5082,9 @@ void FLLMTracker::Clear()
 	}
 	ThreadStates.Empty();
 
+#if !UE_ONLY_USE_PLATFORM_TRACKER
 	AllocationMap.Clear();
+#endif
 
 	CsvWriter.Clear();
 	TraceWriter.Clear();
@@ -5118,7 +5142,9 @@ void FLLMTracker::Update()
 	{
 		LastTrimTime = CurrentTime;
 		{
+#if !UE_ONLY_USE_PLATFORM_TRACKER
 			AllocationMap.Trim();
+#endif
 		}
 	}
 
@@ -5280,6 +5306,7 @@ bool FLLMTracker::DumpForkedAllocationInfo()
 	}
 
 	CountsPerTag.Reserve(NumTags);
+#if !UE_ONLY_USE_PLATFORM_TRACKER
 	AllocationMap.LockAll();
 	for (const FLLMAllocMap::FTuple& Tuple : AllocationMap)
 	{
@@ -5408,6 +5435,7 @@ bool FLLMTracker::DumpForkedAllocationInfo()
 		}
 	}
 	AllocationMap.UnlockAll();
+#endif // !UE_ONLY_USE_PLATFORM_TRACKER
 
 	CountsPerTag.ValueSort([](const FCounts& A, const FCounts& B)
 	{
@@ -5498,7 +5526,7 @@ void FLLMTracker::PublishCsvProfiler(UE::LLM::ESizeParams SizeParams)
 
 void FLLMTracker::OnTagsResorted(FTagDataArray& OldTagDatas)
 {
-#if LLM_ENABLED_FULL_TAGS
+#if LLM_ENABLED_FULL_TAGS && !UE_ONLY_USE_PLATFORM_TRACKER
 	{
 		// Each allocation references the tag by its index, which we have just remapped.
 		// Remap each allocation's tag index to the new index for the tag.
@@ -5608,7 +5636,7 @@ void FLLMTracker::GetTagsNamesWithAmountFiltered(TMap<FName, uint64>& OutTagsNam
 			}
 		}
 	}
-
+#if !UE_ONLY_USE_PLATFORM_TRACKER
 	AllocationMap.LockAll();
 	for (const FLLMAllocMap::FTuple& Tuple : AllocationMap)
 	{
@@ -5645,6 +5673,8 @@ void FLLMTracker::GetTagsNamesWithAmountFiltered(TMap<FName, uint64>& OutTagsNam
 		}
 	}
 	AllocationMap.UnlockAll();
+#endif // !UE_ONLY_USE_PLATFORM_TRACKER
+
 #if LLM_ENABLED_FULL_TAGS
 	for (TPair<int32, FLocalTagData>& Pair : CompressedTagToTagData)
 #else
@@ -5660,6 +5690,7 @@ void FLLMTracker::GetTagsNamesWithAmountFiltered(TMap<FName, uint64>& OutTagsNam
 
 bool FLLMTracker::FindTagsForPtr(void* InPtr, TArray<const FTagData *, TInlineAllocator<static_cast<int32>(ELLMTagSet::Max)>>& OutTags) const
 {
+#if !UE_ONLY_USE_PLATFORM_TRACKER
 	uint32 Size;
 	FLowLevelAllocInfo AllocInfoPtr;
 	PointerKey FoundKey = AllocationMap.Find(PointerKey(InPtr), Size, AllocInfoPtr);
@@ -5679,6 +5710,9 @@ bool FLLMTracker::FindTagsForPtr(void* InPtr, TArray<const FTagData *, TInlineAl
 #endif
 
 	return true;
+#else
+	return false;
+#endif // !UE_ONLY_USE_PLATFORM_TRACKER
 }
 
 int64 FLLMTracker::GetTagAmount(const FTagData* TagData, UE::LLM::ESizeParams SizeParams) const
@@ -5777,7 +5811,10 @@ void FLLMThreadState::TrackAllocation(const void* Ptr, int64 Size, ELLMTracker T
 
 	AllocTypeAmounts[static_cast<int32>(AllocType)] += Size;
 
-	IncrTag(TagData, Size);
+	if (TagData)
+	{
+		IncrTag(TagData, Size);
+	}
 #if LLM_ALLOW_ASSETS_TAGS
 	if (AssetTagData)
 	{
@@ -5789,7 +5826,7 @@ void FLLMThreadState::TrackAllocation(const void* Ptr, int64 Size, ELLMTracker T
 	}
 #endif
 
-	ELLMTag EnumTag = TagData->GetContainingEnum();
+	ELLMTag EnumTag = TagData ? TagData->GetContainingEnum() : ELLMTag(0);
 	if (Tracker == ELLMTracker::Default)
 	{
 		FPlatformMemory::OnLowLevelMemory_Alloc(Ptr, static_cast<uint64>(Size), static_cast<uint64>(EnumTag));
@@ -5809,8 +5846,10 @@ void FLLMThreadState::TrackFree(const void* Ptr, int64 Size, ELLMTracker Tracker
 	FScopeLock Lock(&TagSection);
 
 	AllocTypeAmounts[static_cast<int32>(AllocType)] -= Size;
-
-	IncrTag(TagData, -Size);
+	if (TagData)
+	{
+		IncrTag(TagData, -Size);
+	}
 #if LLM_ALLOW_ASSETS_TAGS
 	if (AssetTagData)
 	{
@@ -5821,7 +5860,8 @@ void FLLMThreadState::TrackFree(const void* Ptr, int64 Size, ELLMTracker Tracker
 		IncrTag(AssetClassTagData, -Size);
 	}
 #endif
-	ELLMTag EnumTag = TagData->GetContainingEnum();
+
+	ELLMTag EnumTag = TagData ? TagData->GetContainingEnum() : ELLMTag(0);
 	if (Tracker == ELLMTracker::Default)
 	{
 		FPlatformMemory::OnLowLevelMemory_Free(Ptr, static_cast<uint64>(Size), static_cast<uint64>(EnumTag));
@@ -5848,7 +5888,7 @@ void FLLMThreadState::TrackMoved(const void* Dest, const void* Source, int64 Siz
 	const FTagData* TagData)
 {
 	// Update external memory trackers (ideally would want a proper 'move' option on these).
-	ELLMTag EnumTag = TagData->GetContainingEnum();
+	ELLMTag EnumTag = TagData ? TagData->GetContainingEnum() : ELLMTag(0);
 	if (Tracker == ELLMTracker::Default)
 	{
 		FPlatformMemory::OnLowLevelMemory_Free(Source, static_cast<uint64>(Size), static_cast<uint64>(EnumTag));
