Load game localization data in cooked editors, replacing the CookedLocalizationPaths setting
[RN] Game localization data is now loaded in cooked editors, replacing the CookedLocalizationPaths setting

diff --git a/Engine/Source/Runtime/Core/Private/Internationalization/LocalizationResourceTextSource.cpp b/Engine/Source/Runtime/Core/Private/Internationalization/LocalizationResourceTextSource.cpp
--- a/Engine/Source/Runtime/Core/Private/Internationalization/LocalizationResourceTextSource.cpp
+++ b/Engine/Source/Runtime/Core/Private/Internationalization/LocalizationResourceTextSource.cpp
@@ -78,13 +78,6 @@ void FLocalizationResourceTextSource::LoadLocalizedResources(const ELocalization
 		return LocalizationPaths;
 	};
 
-	auto GetCookedEditorLocalizationPaths = [&AppendChunkedLocalizationPaths]()
-	{
-		TArray<FString> LocalizationPaths = FPaths::GetCookedEditorLocalizationPaths();
-		AppendChunkedLocalizationPaths(LocalizationPaths);
-		return LocalizationPaths;
-	};
-
 	// Collect the localization paths to load from.
 	TArray<FString> GameNativePaths;
 	TArray<FString> GameLocalizationPaths;
diff --git a/Engine/Source/Runtime/Core/Private/Internationalization/TextLocalizationManager.cpp b/Engine/Source/Runtime/Core/Private/Internationalization/TextLocalizationManager.cpp
--- a/Engine/Source/Runtime/Core/Private/Internationalization/TextLocalizationManager.cpp
+++ b/Engine/Source/Runtime/Core/Private/Internationalization/TextLocalizationManager.cpp
@@ -1514,13 +1514,6 @@ void FTextLocalizationManager::LoadChunkedLocalizationResources_Sync(TArrayView<
 
 	TArray<FString> GameLocalizationPaths;
 	GameLocalizationPaths += FPaths::GetGameLocalizationPaths();
-#if UE_IS_COOKED_EDITOR
-	if (GIsEditor)
-	{
-		// Cooked editors may also load game localization targets
-		GameLocalizationPaths += FPaths::GetCookedEditorLocalizationPaths();
-	}
-#endif
 
 	// Note: We only allow game localization targets to be chunked, and the layout is assumed to follow our standard pattern (as used by the localization dashboard and FLocTextHelper)
 	TArray<FString> ChunkedLocalizationTargets = FLocalizationResourceTextSource::GetChunkedLocalizationTargets();
diff --git a/Engine/Source/Runtime/Core/Private/Misc/Paths.cpp b/Engine/Source/Runtime/Core/Private/Misc/Paths.cpp
--- a/Engine/Source/Runtime/Core/Private/Misc/Paths.cpp
+++ b/Engine/Source/Runtime/Core/Private/Misc/Paths.cpp
@@ -693,22 +693,6 @@ const TArray<FString>& FPaths::GetEditorLocalizationPaths()
 	return StaticData.EditorLocalizationPaths;
 }
 
-const TArray<FString>& FPaths::GetCookedEditorLocalizationPaths()
-{
-	FStaticData& StaticData = TLazySingleton<FStaticData>::Get();
-
-	if (!StaticData.bCookedEditorLocalizationPathsInitialized)
-	{
-		if (GConfig && GConfig->IsReadyForUse())
-		{
-			GConfig->GetArray(TEXT("Internationalization"), TEXT("CookedLocalizationPaths"), StaticData.CookedEditorLocalizationPaths, GEditorIni);
-			StaticData.bCookedEditorLocalizationPathsInitialized = true;
-		}
-	}
-
-	return StaticData.CookedEditorLocalizationPaths;
-}
-
 const TArray<FString>& FPaths::GetPropertyNameLocalizationPaths()
 {
 	FStaticData& StaticData = TLazySingleton<FStaticData>::Get();
diff --git a/Engine/Source/Runtime/Core/Public/Misc/Paths.h b/Engine/Source/Runtime/Core/Public/Misc/Paths.h
--- a/Engine/Source/Runtime/Core/Public/Misc/Paths.h
+++ b/Engine/Source/Runtime/Core/Public/Misc/Paths.h
@@ -424,7 +424,8 @@ public:
 	/**
 	 * Returns a list of cooked editor-specific localization paths
 	 */
-	static CORE_API const TArray<FString>& GetCookedEditorLocalizationPaths();
+	UE_DEPRECATED(5.5, "Cooked editor localization path no longer exist. Cooked editors now load game localization paths.")
+	static const TArray<FString>& GetCookedEditorLocalizationPaths() {}
 
 	/** 
 	 * Returns a list of property name localization paths
