CIS fix

#rnx

diff --git a/Engine/Source/Runtime/Core/Public/Misc/Paths.h b/Engine/Source/Runtime/Core/Public/Misc/Paths.h
--- a/Engine/Source/Runtime/Core/Public/Misc/Paths.h
+++ b/Engine/Source/Runtime/Core/Public/Misc/Paths.h
@@ -425,7 +425,7 @@ public:
 	 * Returns a list of cooked editor-specific localization paths
 	 */
 	UE_DEPRECATED(5.5, "Cooked editor localization path no longer exist. Cooked editors now load game localization paths.")
-	static const TArray<FString>& GetCookedEditorLocalizationPaths() {}
+	static TArray<FString> GetCookedEditorLocalizationPaths() { return TArray<FString>(); }
 
 	/** 
 	 * Returns a list of property name localization paths
