[CsvProfiler] Optimized overhead by using the 1/frame cached memory stats, and reducing platform API calls.

diff --git a/Engine/Source/Runtime/Core/Private/ProfilingDebugging/CsvProfiler.cpp b/Engine/Source/Runtime/Core/Private/ProfilingDebugging/CsvProfiler.cpp
--- a/Engine/Source/Runtime/Core/Private/ProfilingDebugging/CsvProfiler.cpp
+++ b/Engine/Source/Runtime/Core/Private/ProfilingDebugging/CsvProfiler.cpp
@@ -35,6 +35,7 @@
 #include "Misc/Guid.h"
 #include "Misc/WildcardString.h"
 #include "Modules/ModuleManager.h"
+#include "HAL/PlatformMemoryHelpers.h"
 
 #include "HAL/PlatformMisc.h"
 
@@ -3140,7 +3141,7 @@ void FCsvProfiler::EndFrame()
 		float ElapsedMs = (float)FPlatformTime::ToMilliseconds64(ElapsedCycles);
 		CSV_CUSTOM_STAT_MINIMAL_GLOBAL(FrameTime, ElapsedMs, ECsvCustomStatOp::Set);
 
-		FPlatformMemoryStats MemoryStats = FPlatformMemory::GetStats();
+		FPlatformMemoryStats MemoryStats = PlatformMemoryHelpers::GetFrameMemoryStats();
 
 		float PhysicalMBFree = float(MemoryStats.AvailablePhysical) / (1024.0f * 1024.0f);
 		float UsedExtendedMB = 0;
