Centralized the logic for forcing game localization data to load in an editor

This is now controlled by FTextLocalizationManager::ShouldForceLoadGameLocalization, which will return true in the following cases:
  * A cooked editor.
  * An uncooked editor with the game localization preview enabled, or with the CVar Localization.ForceLoadGameLocalizationInEditor set to true.

This is queried to add the ELocalizationLoadFlags::ForceLocalizedGame flag when loading localization data, and also by FLocalizationTargetDescriptor::ShouldLoadLocalizationTarget when deciding whether to load a game localization target for a plugin (which also allows plugins to work with the game localization preview).

diff --git a/Engine/Source/Runtime/Core/Private/Internationalization/LocalizationResourceTextSource.cpp b/Engine/Source/Runtime/Core/Private/Internationalization/LocalizationResourceTextSource.cpp
--- a/Engine/Source/Runtime/Core/Private/Internationalization/LocalizationResourceTextSource.cpp
+++ b/Engine/Source/Runtime/Core/Private/Internationalization/LocalizationResourceTextSource.cpp
@@ -95,9 +95,6 @@ void FLocalizationResourceTextSource::LoadLocalizedResources(const ELocalization
 	if (ShouldLoadEditor(InLoadFlags))
 	{
 		EditorLocalizationPaths += FPaths::GetEditorLocalizationPaths();
-#if UE_IS_COOKED_EDITOR
-		EditorLocalizationPaths += GetGameLocalizationPaths();
-#endif
 		EditorLocalizationPaths += FPaths::GetToolTipLocalizationPaths();
 
 		bool bShouldUseLocalizedPropertyNames = false;
diff --git a/Engine/Source/Runtime/Core/Private/Internationalization/TextLocalizationManager.cpp b/Engine/Source/Runtime/Core/Private/Internationalization/TextLocalizationManager.cpp
--- a/Engine/Source/Runtime/Core/Private/Internationalization/TextLocalizationManager.cpp
+++ b/Engine/Source/Runtime/Core/Private/Internationalization/TextLocalizationManager.cpp
@@ -63,6 +63,11 @@ static FAutoConsoleVariableRef CVarAsyncLoadLocalizationData(TEXT("Localization.
 static bool AsyncLoadLocalizationDataOnLanguageChange = false;
 static FAutoConsoleVariableRef CVarAsyncLoadLocalizationDataOnLanguageChange(TEXT("Localization.AsyncLoadLocalizationDataOnLanguageChange"), AsyncLoadLocalizationDataOnLanguageChange, TEXT("True to load localization data asynchronously (non-blocking) when the language changes, or False to load it synchronously (blocking)"));
 
+#if WITH_EDITOR
+static bool ForceLoadGameLocalizationInEditor = false;
+static FAutoConsoleVariableRef CVarForceLoadGameLocalizationInEditor(TEXT("Localization.ForceLoadGameLocalizationInEditor"), ForceLoadGameLocalizationInEditor, TEXT("True to force load game localization data in an editor"));
+#endif
+
 #if ENABLE_LOC_TESTING
 static FAutoConsoleCommand CmdDumpLiveTable(
 	TEXT("Localization.DumpLiveTable"), 
@@ -1343,11 +1348,10 @@ void FTextLocalizationManager::LoadLocalizationResourcesForPrioritizedCultures_S
 	}
 
 	// Leet-ify always needs the native text to operate on, so force native data if we're loading for LEET
-	// The keys culture also uses native so that we have a good way to restore the orginal text from the keys state 
-	ELocalizationLoadFlags FinalLocLoadFlags = LocLoadFlags;
+	// The keys culture also uses native so that we have a good way to restore the original text from the keys state 
+	ELocalizationLoadFlags FinalLocLoadFlags = LocLoadFlags | (ShouldForceLoadGameLocalization() ? ELocalizationLoadFlags::ForceLocalizedGame : ELocalizationLoadFlags::None);
 #if ENABLE_LOC_TESTING
-	bool bShouldForceLoadNative = (PrioritizedCultureNames[0] == FLeetCulture::StaticGetName()) || (PrioritizedCultureNames[0] == FKeysCulture::StaticGetName());
-	if (bShouldForceLoadNative)
+	if ((PrioritizedCultureNames[0] == FLeetCulture::StaticGetName()) || (PrioritizedCultureNames[0] == FKeysCulture::StaticGetName()))
 	{
 		FinalLocLoadFlags |= ELocalizationLoadFlags::Native;
 	}
@@ -1930,7 +1934,7 @@ void FTextLocalizationManager::EnableGameLocalizationPreview(const FString& Cult
 		PrioritizedCultureNames.Add(PreviewCulture);
 	}
 
-	ELocalizationLoadFlags LocLoadFlags = ELocalizationLoadFlags::Game | ELocalizationLoadFlags::ForceLocalizedGame;
+	ELocalizationLoadFlags LocLoadFlags = ELocalizationLoadFlags::Game | ELocalizationLoadFlags::Additional;
 	LocLoadFlags |= (bIsGameLocalizationPreviewEnabled ? ELocalizationLoadFlags::Native : ELocalizationLoadFlags::None);
 
 	LoadLocalizationResourcesForPrioritizedCultures_Async(PrioritizedCultureNames, LocLoadFlags);
@@ -1979,6 +1983,23 @@ bool FTextLocalizationManager::IsLocalizationLocked() const
 }
 #endif
 
+bool FTextLocalizationManager::ShouldForceLoadGameLocalization() const
+{
+#if WITH_EDITOR
+	if (GIsEditor)
+	{
+#if UE_IS_COOKED_EDITOR
+		return true;
+#else	// UE_IS_COOKED_EDITOR
+		return bIsGameLocalizationPreviewEnabled || TextLocalizationManager::ForceLoadGameLocalizationInEditor;
+#endif	// UE_IS_COOKED_EDITOR
+	}
+	return false;
+#else	// WITH_EDITOR
+	return false;
+#endif	// WITH_EDITOR
+}
+
 #if ENABLE_LOC_TESTING
 void FTextLocalizationManager::LeetifyAllDisplayStrings()
 {
diff --git a/Engine/Source/Runtime/Core/Public/Internationalization/TextLocalizationManager.h b/Engine/Source/Runtime/Core/Public/Internationalization/TextLocalizationManager.h
--- a/Engine/Source/Runtime/Core/Public/Internationalization/TextLocalizationManager.h
+++ b/Engine/Source/Runtime/Core/Public/Internationalization/TextLocalizationManager.h
@@ -335,6 +335,11 @@ public:
 	CORE_API bool IsLocalizationLocked() const;
 #endif
 
+	/**
+	 * True if we should force load game localization data.
+	 */
+	CORE_API bool ShouldForceLoadGameLocalization() const;
+
 	/** Event type for immediately reacting to changes in display strings for text. */
 	DECLARE_EVENT(FTextLocalizationManager, FTextRevisionChangedEvent)
 	FTextRevisionChangedEvent OnTextRevisionChangedEvent;
diff --git a/Engine/Source/Runtime/Projects/Private/LocalizationDescriptor.cpp b/Engine/Source/Runtime/Projects/Private/LocalizationDescriptor.cpp
--- a/Engine/Source/Runtime/Projects/Private/LocalizationDescriptor.cpp
+++ b/Engine/Source/Runtime/Projects/Private/LocalizationDescriptor.cpp
@@ -256,11 +256,7 @@ bool FLocalizationTargetDescriptor::ShouldLoadLocalizationTarget() const
 		return WITH_EDITOR;
 
 	case ELocalizationTargetDescriptorLoadingPolicy::Game:
-#if UE_IS_COOKED_EDITOR
-		return true;
-#else
-		return FApp::IsGame();
-#endif
+		return FApp::IsGame() || FTextLocalizationManager::Get().ShouldForceLoadGameLocalization();
 
 	case ELocalizationTargetDescriptorLoadingPolicy::PropertyNames:
 #if WITH_EDITOR
