Add ProgramSizeMB metadata so we can capture this even when LLM isn't running

diff --git a/Engine/Source/Runtime/Core/Private/Android/AndroidPlatformMemory.cpp b/Engine/Source/Runtime/Core/Private/Android/AndroidPlatformMemory.cpp
--- a/Engine/Source/Runtime/Core/Private/Android/AndroidPlatformMemory.cpp
+++ b/Engine/Source/Runtime/Core/Private/Android/AndroidPlatformMemory.cpp
@@ -444,11 +444,12 @@ FMalloc* FAndroidPlatformMemory::BaseAllocator()
 		return Instance;
 	}
 
+	FPlatformMemoryStats Stats = FAndroidPlatformMemory::GetStats();
 #if ENABLE_LOW_LEVEL_MEM_TRACKER
 	// make sure LLM is using UsedPhysical for program size, instead of Available-Free
-	FPlatformMemoryStats Stats = FAndroidPlatformMemory::GetStats();
 	FLowLevelMemTracker::Get().SetProgramSize(Stats.UsedPhysical);
 #endif
+	FPlatformMemory::ProgramSize = Stats.UsedPhysical;
 
 #if FORCE_ANSI_ALLOCATOR
 	return (Instance = new FMallocAnsi());
diff --git a/Engine/Source/Runtime/Core/Private/Apple/ApplePlatformMemory.cpp b/Engine/Source/Runtime/Core/Private/Apple/ApplePlatformMemory.cpp
--- a/Engine/Source/Runtime/Core/Private/Apple/ApplePlatformMemory.cpp
+++ b/Engine/Source/Runtime/Core/Private/Apple/ApplePlatformMemory.cpp
@@ -346,11 +346,12 @@ FMalloc* FApplePlatformMemory::BaseAllocator()
 		return Instance;
 	}
 
-#if ENABLE_LOW_LEVEL_MEM_TRACKER
 	FPlatformMemoryStats MemStats = FApplePlatformMemory::GetStats();
+#if ENABLE_LOW_LEVEL_MEM_TRACKER
 	FLowLevelMemTracker::Get().SetProgramSize(MemStats.UsedPhysical);
 #endif
-    
+	FPlatformMemory::ProgramSize = Stats.UsedPhysical;
+
     SetAllocatorToUse();
     
     switch (AllocatorToUse)
diff --git a/Engine/Source/Runtime/Core/Private/GenericPlatform/GenericPlatformMemory.cpp b/Engine/Source/Runtime/Core/Private/GenericPlatform/GenericPlatformMemory.cpp
--- a/Engine/Source/Runtime/Core/Private/GenericPlatform/GenericPlatformMemory.cpp
+++ b/Engine/Source/Runtime/Core/Private/GenericPlatform/GenericPlatformMemory.cpp
@@ -164,7 +164,7 @@ uint64 FGenericPlatformMemory::OOMAllocationSize = 0;
 uint32 FGenericPlatformMemory::OOMAllocationAlignment = 0;
 FGenericPlatformMemory::EMemoryAllocatorToUse FGenericPlatformMemory::AllocatorToUse = Platform;
 void* FGenericPlatformMemory::BackupOOMMemoryPool = nullptr;
-
+uint64 FGenericPlatformMemory::ProgramSize = 0;
 
 void FGenericPlatformMemory::SetupMemoryPools()
 {
diff --git a/Engine/Source/Runtime/Core/Private/ProfilingDebugging/CsvProfiler.cpp b/Engine/Source/Runtime/Core/Private/ProfilingDebugging/CsvProfiler.cpp
--- a/Engine/Source/Runtime/Core/Private/ProfilingDebugging/CsvProfiler.cpp
+++ b/Engine/Source/Runtime/Core/Private/ProfilingDebugging/CsvProfiler.cpp
@@ -3067,6 +3067,13 @@ void FCsvProfiler::BeginFrame()
 					SetMetadataInternal(TEXT("TargetFramerate"), *FString::FromInt(TargetFPS));
 					SetMetadataInternal(TEXT("StartTimestamp"), *FString::Printf(TEXT("%lld"), FDateTime::UtcNow().ToUnixTimestamp()));
 					SetMetadataInternal(TEXT("NamedEvents"), (GCycleStatsShouldEmitNamedEvents > 0) ? TEXT("1") : TEXT("0"));
+
+					if (FPlatformMemory::GetProgramSize() > 0)
+					{
+						// Some platforms adjust program size at runtime (based on DLL initialization), so with do this on start capture rather than in CsvProfiler::Init()
+						SetMetadataInternal(TEXT("ProgramSizeMB"), *FString::SanitizeFloat((float)FPlatformMemory::GetProgramSize() / 1024.0f / 1024.0f));
+					}
+
 					bNamedEventsWasEnabled = (GCycleStatsShouldEmitNamedEvents > 0);
 
 					GCsvStatCounts = !!CVarCsvStatCounts.GetValueOnGameThread();
diff --git a/Engine/Source/Runtime/Core/Public/GenericPlatform/GenericPlatformMemory.h b/Engine/Source/Runtime/Core/Public/GenericPlatform/GenericPlatformMemory.h
--- a/Engine/Source/Runtime/Core/Public/GenericPlatform/GenericPlatformMemory.h
+++ b/Engine/Source/Runtime/Core/Public/GenericPlatform/GenericPlatformMemory.h
@@ -756,6 +756,22 @@ public:
 	*/
 	static CORE_API uint64 GetExtraDevelopmentMemorySize();
 
+	/**
+	* Returns initial program size or 0 if the platform doesn't track initial program size
+	*/
+	static CORE_API uint64 GetProgramSize()
+	{
+		return ProgramSize;
+	}
+
+	/**
+	* Sets the initial program size
+	*/
+	static CORE_API void SetProgramSize(uint64 InProgramSize)
+	{
+		ProgramSize = InProgramSize;
+	}
+
 	/**
 	* This function sets AllocFunction and FreeFunction and returns true, or just returns false.
 	* These functions are the platform dependant low low low level functions that LLM uses to allocate memory.
@@ -834,4 +850,7 @@ protected:
 
 	/** Updates platform specific stats. This method is called through FGenericStatsUpdater from the task graph thread. */
 	static CORE_API void InternalUpdateStats( const FPlatformMemoryStats& MemoryStats );
+
+	/** Program memory allocation in bytes. */
+	static CORE_API uint64 ProgramSize;
 };
