Fix IOS

diff --git a/Engine/Source/Runtime/Core/Private/Apple/ApplePlatformMemory.cpp b/Engine/Source/Runtime/Core/Private/Apple/ApplePlatformMemory.cpp
--- a/Engine/Source/Runtime/Core/Private/Apple/ApplePlatformMemory.cpp
+++ b/Engine/Source/Runtime/Core/Private/Apple/ApplePlatformMemory.cpp
@@ -350,7 +350,7 @@ FMalloc* FApplePlatformMemory::BaseAllocator()
 #if ENABLE_LOW_LEVEL_MEM_TRACKER
 	FLowLevelMemTracker::Get().SetProgramSize(MemStats.UsedPhysical);
 #endif
-	FPlatformMemory::ProgramSize = Stats.UsedPhysical;
+	FPlatformMemory::ProgramSize = MemStats.UsedPhysical;
 
     SetAllocatorToUse();
     
