Refactor CsvProfiler begin/end capture - move capture start/end logic into dedicated functions to make the logic easier to maintain
- Add metadata for capture duration

diff --git a/Engine/Source/Runtime/Core/Private/ProfilingDebugging/CsvProfiler.cpp b/Engine/Source/Runtime/Core/Private/ProfilingDebugging/CsvProfiler.cpp
--- a/Engine/Source/Runtime/Core/Private/ProfilingDebugging/CsvProfiler.cpp
+++ b/Engine/Source/Runtime/Core/Private/ProfilingDebugging/CsvProfiler.cpp
@@ -2846,6 +2846,7 @@ FCsvProfiler::FCsvProfiler()
 	, bNamedEventsWasEnabled(false)
 	, LastEndFrameTimestamp(0)
 	, CaptureEndFrameCount(0)
+	, CaptureStartTime(0.0)
 	, ProcessingThread(nullptr)
 	, FileWriteBlockingEvent(FPlatformProcess::GetSynchEventFromPool())
 {
@@ -2958,132 +2959,7 @@ void FCsvProfiler::BeginFrame()
 		if (CommandQueue.Peek(CurrentCommand) && CurrentCommand.CommandType == ECsvCommandType::Start)
 		{
 			CommandQueue.Dequeue(CurrentCommand);
-			if (GCsvProfilerIsCapturing)
-			{
-				UE_LOG(LogCsvProfiler, Warning, TEXT("Capture start requested, but a capture was already running"));
-			}
-			else
-			{
-				UE_LOG(LogCsvProfiler, Display, TEXT("Capture Starting"));
-
-				if (GConfig)
-				{
-					// Update categories from the config. The config may have changed if there were hotfixes
-					FCsvCategoryData::Get()->UpdateCategoriesFromConfig();
-				}
-
-				// signal external profiler that we are capturing
-				OnCSVProfileStartDelegate.Broadcast();
-
-				// Latch the cvars when we start a capture
-				int32 BufferSize = FMath::Max(CVarCsvWriteBufferSize.GetValueOnAnyThread(), 0);
-				bool bContinuousWrites = IsContinuousWriteEnabled(true);
-
-				// Allow overriding of compression based on the "csv.CompressionMode" CVar
-				bool bCompressOutput;
-				switch (CVarCsvCompressionMode.GetValueOnGameThread())
-				{
-				case 0:
-					bCompressOutput = false;
-					break;
-
-				case 1:
-					bCompressOutput = (BufferSize > 0); 
-					break;
-
-				default:
-					bCompressOutput = EnumHasAnyFlags(CurrentCommand.Flags, ECsvProfilerFlags::CompressOutput) && (BufferSize > 0);
-					break;
-				}
-
-				const TCHAR* CsvExtension = bCompressOutput ? TEXT(".csv.gz") : TEXT(".csv");
-
-				// Determine the output path and filename based on override params
-				FString DestinationFolder = CurrentCommand.DestinationFolder.IsEmpty() ? FPaths::ProfilingDir() + TEXT("CSV/") : CurrentCommand.DestinationFolder + TEXT("/");
-				FString Filename = CurrentCommand.Filename.IsEmpty() ? FString::Printf(TEXT("Profile(%s)%s"), *FDateTime::Now().ToString(TEXT("%Y%m%d_%H%M%S")), CsvExtension) : CurrentCommand.Filename;
-				OutputFilename = DestinationFolder + Filename;
-
-				TSharedPtr<FArchive> OutputFile = MakeShareable(IFileManager::Get().CreateFileWriter(*OutputFilename));
-				if (!OutputFile)
-				{
-					UE_LOG(LogCsvProfiler, Error, TEXT("Failed to create CSV file \"%s\". Capture will not start."), *OutputFilename);
-				}
-				else
-				{
-					int64 NumFramesToBuffer = CVarCsvStreamFramesToBuffer.GetValueOnAnyThread();
-					CsvWriter = new FCsvStreamWriter(OutputFile.ToSharedRef(), bContinuousWrites, BufferSize, NumFramesToBuffer, bCompressOutput, RenderThreadId, RHIThreadId);
-
-					NumFramesToCapture = CurrentCommand.Value;
-					GCsvRepeatFrameCount = NumFramesToCapture;
-					CaptureFrameNumber = 0;
-					CaptureFrameNumberRT = 0;
-					LastEndFrameTimestamp = FPlatformTime::Cycles64();
-					CurrentFlags = CurrentCommand.Flags;
-
-					if (GCsvUseProcessingThread && ProcessingThread == nullptr)
-					{
-						// Lazily create the CSV processing thread
-						ProcessingThread = new FCsvProfilerProcessingThread(*this);
-						if (ProcessingThread->IsValid() == false)
-						{
-							UE_LOG(LogCsvProfiler, Error, TEXT("CSV Processing Thread could not be created due to being in a single-thread environment "));
-							delete ProcessingThread;
-							ProcessingThread = nullptr;
-							GCsvUseProcessingThread = false;
-						}
-					}
-					 
-					// Set the CSV ID and mirror it to the log
-					FString CsvId = FGuid::NewGuid().ToString();
-					SetMetadataInternal(TEXT("CsvID"), *CsvId);
-					UE_LOG(LogCsvProfiler, Display, TEXT("Capture started. CSV ID: %s"), *CsvId);
-
-					int32 TargetFPS = FPlatformMisc::GetMaxRefreshRate();
-					static IConsoleVariable* CsvTargetFrameRateCVar = IConsoleManager::Get().FindConsoleVariable(TEXT("csv.TargetFrameRateOverride"));
-					static IConsoleVariable* MaxFPSCVar = IConsoleManager::Get().FindConsoleVariable(TEXT("t.MaxFPS"));
-					static IConsoleVariable* SyncIntervalCVar = IConsoleManager::Get().FindConsoleVariable(TEXT("rhi.SyncInterval"));
-					int32 CmdLineTargetFPS = TargetFPS;
-					if (CsvTargetFrameRateCVar && CsvTargetFrameRateCVar->GetInt() > 0)
-					{
-						TargetFPS = CsvTargetFrameRateCVar->GetInt();
-					}
-					else if (FParse::Value(FCommandLine::Get(), TEXT("csv.TargetFrameRateOverride"), CmdLineTargetFPS)) // Too early to set CsvTargetFrameRateCVar with execcmds
-					{
-						TargetFPS = CmdLineTargetFPS;
-					}
-					else
-					{
-						// Figure out the target framerate
-						if (MaxFPSCVar && MaxFPSCVar->GetInt() > 0)
-						{
-							TargetFPS = MaxFPSCVar->GetInt();
-						}
-						if (SyncIntervalCVar && SyncIntervalCVar->GetInt() > 0)
-						{
-							TargetFPS = FMath::Min(TargetFPS, FPlatformMisc::GetMaxRefreshRate() / SyncIntervalCVar->GetInt());
-						}
-					}
-
-					SetMetadataInternal(TEXT("TargetFramerate"), *FString::FromInt(TargetFPS));
-					SetMetadataInternal(TEXT("StartTimestamp"), *FString::Printf(TEXT("%lld"), FDateTime::UtcNow().ToUnixTimestamp()));
-					SetMetadataInternal(TEXT("NamedEvents"), (GCycleStatsShouldEmitNamedEvents > 0) ? TEXT("1") : TEXT("0"));
-
-					if (FPlatformMemory::GetProgramSize() > 0)
-					{
-						// Some platforms adjust program size at runtime (based on DLL initialization), so with do this on start capture rather than in CsvProfiler::Init()
-						SetMetadataInternal(TEXT("ProgramSizeMB"), *FString::SanitizeFloat((float)FPlatformMemory::GetProgramSize() / 1024.0f / 1024.0f));
-					}
-
-					bNamedEventsWasEnabled = (GCycleStatsShouldEmitNamedEvents > 0);
-
-					GCsvStatCounts = !!CVarCsvStatCounts.GetValueOnGameThread();
-
-					// Check TLS is initialized before starting the capture. This should have happened in BeginCapture
-					check(FCsvProfilerThreadData::IsTlsSlotInitialized());
-					TRACE_CSV_PROFILER_BEGIN_CAPTURE(*Filename, RenderThreadId, RHIThreadId, GDefaultWaitStatName, GCsvStatCounts);
-					GCsvProfilerIsCapturing = true;
-				}
-			}
+			BeginCaptureInternal(CurrentCommand);
 		}
 
 		if (GCsvProfilerIsCapturing)
@@ -3119,6 +2995,139 @@ void FCsvProfiler::BeginFrame()
 #endif // CSV_PROFILER_ALLOW_DEBUG_FEATURES
 }
 
+
+void FCsvProfiler::BeginCaptureInternal(const FCsvCaptureCommand& CurrentCommand)
+{
+	if (GCsvProfilerIsCapturing)
+	{
+		UE_LOG(LogCsvProfiler, Warning, TEXT("Capture start requested, but a capture was already running"));
+		return;
+	}
+
+	UE_LOG(LogCsvProfiler, Display, TEXT("Capture Starting"));
+	if (GConfig)
+	{
+		// Update categories from the config. The config may have changed if there were hotfixes
+		FCsvCategoryData::Get()->UpdateCategoriesFromConfig();
+	}
+
+	// signal external profiler that we are capturing
+	OnCSVProfileStartDelegate.Broadcast();
+
+	// Latch the cvars when we start a capture
+	int32 BufferSize = FMath::Max(CVarCsvWriteBufferSize.GetValueOnAnyThread(), 0);
+	bool bContinuousWrites = IsContinuousWriteEnabled(true);
+
+	// Allow overriding of compression based on the "csv.CompressionMode" CVar
+	bool bCompressOutput;
+	switch (CVarCsvCompressionMode.GetValueOnGameThread())
+	{
+	case 0:
+		bCompressOutput = false;
+		break;
+
+	case 1:
+		bCompressOutput = (BufferSize > 0);
+		break;
+
+	default:
+		bCompressOutput = EnumHasAnyFlags(CurrentCommand.Flags, ECsvProfilerFlags::CompressOutput) && (BufferSize > 0);
+		break;
+	}
+
+	const TCHAR* CsvExtension = bCompressOutput ? TEXT(".csv.gz") : TEXT(".csv");
+
+	// Determine the output path and filename based on override params
+	FString DestinationFolder = CurrentCommand.DestinationFolder.IsEmpty() ? FPaths::ProfilingDir() + TEXT("CSV/") : CurrentCommand.DestinationFolder + TEXT("/");
+	FString Filename = CurrentCommand.Filename.IsEmpty() ? FString::Printf(TEXT("Profile(%s)%s"), *FDateTime::Now().ToString(TEXT("%Y%m%d_%H%M%S")), CsvExtension) : CurrentCommand.Filename;
+	OutputFilename = DestinationFolder + Filename;
+
+	TSharedPtr<FArchive> OutputFile = MakeShareable(IFileManager::Get().CreateFileWriter(*OutputFilename));
+	if (!OutputFile)
+	{
+		UE_LOG(LogCsvProfiler, Error, TEXT("Failed to create CSV file \"%s\". Capture will not start."), *OutputFilename);
+		return;
+	}
+
+	// Actually start the capture
+	int64 NumFramesToBuffer = CVarCsvStreamFramesToBuffer.GetValueOnAnyThread();
+	CsvWriter = new FCsvStreamWriter(OutputFile.ToSharedRef(), bContinuousWrites, BufferSize, NumFramesToBuffer, bCompressOutput, RenderThreadId, RHIThreadId);
+
+	NumFramesToCapture = CurrentCommand.Value;
+	GCsvRepeatFrameCount = NumFramesToCapture;
+	CaptureFrameNumber = 0;
+	CaptureFrameNumberRT = 0;
+	LastEndFrameTimestamp = FPlatformTime::Cycles64();
+	CurrentFlags = CurrentCommand.Flags;
+
+	if (GCsvUseProcessingThread && ProcessingThread == nullptr)
+	{
+		// Lazily create the CSV processing thread
+		ProcessingThread = new FCsvProfilerProcessingThread(*this);
+		if (ProcessingThread->IsValid() == false)
+		{
+			UE_LOG(LogCsvProfiler, Error, TEXT("CSV Processing Thread could not be created due to being in a single-thread environment "));
+			delete ProcessingThread;
+			ProcessingThread = nullptr;
+			GCsvUseProcessingThread = false;
+		}
+	}
+
+	// Set the CSV ID and mirror it to the log
+	FString CsvId = FGuid::NewGuid().ToString();
+	SetMetadataInternal(TEXT("CsvID"), *CsvId);
+	UE_LOG(LogCsvProfiler, Display, TEXT("Capture started. CSV ID: %s"), *CsvId);
+
+	int32 TargetFPS = FPlatformMisc::GetMaxRefreshRate();
+	static IConsoleVariable* CsvTargetFrameRateCVar = IConsoleManager::Get().FindConsoleVariable(TEXT("csv.TargetFrameRateOverride"));
+	static IConsoleVariable* MaxFPSCVar = IConsoleManager::Get().FindConsoleVariable(TEXT("t.MaxFPS"));
+	static IConsoleVariable* SyncIntervalCVar = IConsoleManager::Get().FindConsoleVariable(TEXT("rhi.SyncInterval"));
+	int32 CmdLineTargetFPS = TargetFPS;
+	if (CsvTargetFrameRateCVar && CsvTargetFrameRateCVar->GetInt() > 0)
+	{
+		TargetFPS = CsvTargetFrameRateCVar->GetInt();
+	}
+	else if (FParse::Value(FCommandLine::Get(), TEXT("csv.TargetFrameRateOverride"), CmdLineTargetFPS)) // Too early to set CsvTargetFrameRateCVar with execcmds
+	{
+		TargetFPS = CmdLineTargetFPS;
+	}
+	else
+	{
+		// Figure out the target framerate
+		if (MaxFPSCVar && MaxFPSCVar->GetInt() > 0)
+		{
+			TargetFPS = MaxFPSCVar->GetInt();
+		}
+		if (SyncIntervalCVar && SyncIntervalCVar->GetInt() > 0)
+		{
+			TargetFPS = FMath::Min(TargetFPS, FPlatformMisc::GetMaxRefreshRate() / SyncIntervalCVar->GetInt());
+		}
+	}
+
+	SetMetadataInternal(TEXT("TargetFramerate"), *FString::FromInt(TargetFPS));
+	SetMetadataInternal(TEXT("StartTimestamp"), *FString::Printf(TEXT("%lld"), FDateTime::UtcNow().ToUnixTimestamp()));
+	SetMetadataInternal(TEXT("NamedEvents"), (GCycleStatsShouldEmitNamedEvents > 0) ? TEXT("1") : TEXT("0"));
+
+	if (FPlatformMemory::GetProgramSize() > 0)
+	{
+		// Some platforms adjust program size at runtime (based on DLL initialization), so with do this on start capture rather than in CsvProfiler::Init()
+		SetMetadataInternal(TEXT("ProgramSizeMB"), *FString::SanitizeFloat((float)FPlatformMemory::GetProgramSize() / 1024.0f / 1024.0f));
+	}
+
+	bNamedEventsWasEnabled = (GCycleStatsShouldEmitNamedEvents > 0);
+
+	GCsvStatCounts = !!CVarCsvStatCounts.GetValueOnGameThread();
+
+	// Check TLS is initialized before starting the capture. This should have happened in BeginCapture
+	check(FCsvProfilerThreadData::IsTlsSlotInitialized());
+	TRACE_CSV_PROFILER_BEGIN_CAPTURE(*Filename, RenderThreadId, RHIThreadId, GDefaultWaitStatName, GCsvStatCounts);
+	GCsvProfilerIsCapturing = true;
+	CaptureStartTime = FPlatformTime::Seconds();
+}
+
+
+
+
 void FCsvProfiler::EndFrame()
 {
 	LLM_SCOPE(ELLMTag::CsvProfiler);
@@ -3188,106 +3197,108 @@ void FCsvProfiler::EndFrame()
 	if (CommandQueue.Peek(CurrentCommand) && CurrentCommand.CommandType == ECsvCommandType::Stop)
 	{
 		QUICK_SCOPE_CYCLE_COUNTER(STAT_FCsvProfiler_EndFrame_Stop);
-		bool bCaptureComplete = false;
-
-		if (!GCsvProfilerIsCapturing && !GCsvProfilerIsWritingFile)
+		if ( TryEndCaptureInternal(CurrentCommand) )
 		{
-			bCaptureComplete = true;
+			// Pop the 'stop' command now that the capture has ended (or we weren't capturing anyway).
+			CommandQueue.Dequeue(CurrentCommand);
 		}
-		else
+	}
+
+	GCsvProfilerFrameNumber++;
+}
+
+bool FCsvProfiler::TryEndCaptureInternal(const FCsvCaptureCommand& CurrentCommand)
+{
+	if (GCsvProfilerIsCapturing || GCsvProfilerIsWritingFile)
+	{
+		// Delay end capture by a frame to allow RT stats to catch up
+		if (CurrentCommand.FrameRequested == GCsvProfilerFrameNumber)
 		{
-			// Delay end capture by a frame to allow RT stats to catch up
-			if (CurrentCommand.FrameRequested == GCsvProfilerFrameNumber)
-			{
-				CaptureEndFrameCount = CaptureFrameNumber;
-			}
-			else
-			{
-				UE_LOG(LogCsvProfiler, Display, TEXT("Capture Stop requested"));
+			CaptureEndFrameCount = CaptureFrameNumber;
+			return false;
+		}
 
-				// signal external profiler that we are done
-				OnCSVProfileEndDelegate.Broadcast();
+		UE_LOG(LogCsvProfiler, Display, TEXT("Capture Stop requested"));
 
-				// Signal to the processing thread to write the file out (if we have one).
-				GCsvProfilerIsWritingFile = true;
-				GCsvProfilerIsCapturing = false;
+		// signal external profiler that we are done
+		OnCSVProfileEndDelegate.Broadcast();
 
-				TRACE_CSV_PROFILER_END_CAPTURE();
+		// Signal to the processing thread to write the file out (if we have one).
+		GCsvProfilerIsWritingFile = true;
+		GCsvProfilerIsCapturing = false;
 
-				if (!ProcessingThread)
-				{
-					// Suspend the hang and hitch heartbeats, as this is a long running task.
-					FSlowHeartBeatScope SuspendHeartBeat;
-					FDisableHitchDetectorScope SuspendGameThreadHitch;
+		TRACE_CSV_PROFILER_END_CAPTURE();
 
-					// No processing thread, block and write the file out on the game thread.
-					FinalizeCsvFile();
-					bCaptureComplete = true;
-				}
-				else if (CVarCsvBlockOnCaptureEnd.GetValueOnGameThread() == 1)
-				{
-					// Suspend the hang and hitch heartbeats, as this is a long running task.
-					FSlowHeartBeatScope SuspendHeartBeat;
-					FDisableHitchDetectorScope SuspendGameThreadHitch;
+		if (!ProcessingThread)
+		{
+			// Suspend the hang and hitch heartbeats, as this is a long running task.
+			FSlowHeartBeatScope SuspendHeartBeat;
+			FDisableHitchDetectorScope SuspendGameThreadHitch;
 
-					// Block the game thread here whilst the result file is written out.
-					FileWriteBlockingEvent->Wait();
-				}
+			// No processing thread, block and write the file out on the game thread. We're done.
+			FinalizeCsvFile();
+		}
+		else
+		{
+			if (CVarCsvBlockOnCaptureEnd.GetValueOnGameThread() == 1)
+			{
+				// Suspend the hang and hitch heartbeats, as this is a long running task.
+				FSlowHeartBeatScope SuspendHeartBeat;
+				FDisableHitchDetectorScope SuspendGameThreadHitch;
+
+				// Block the game thread here whilst the result file is written out.
+				FileWriteBlockingEvent->Wait();
 			}
+			// Not done yet...
+			return false;
 		}
+	}
 
-		if (bCaptureComplete)
-		{
-			check(!GCsvProfilerIsCapturing && !GCsvProfilerIsWritingFile);
+	// If we get here then we're actually done
+	check(!GCsvProfilerIsCapturing && !GCsvProfilerIsWritingFile);
 
-			// Pop the 'stop' command now that the capture has ended (or we weren't capturing anyway).
-			CommandQueue.Dequeue(CurrentCommand);
+	// Signal the async completion callback, if one was provided when the capture was stopped.
+	if (CurrentCommand.Completion)
+	{
+		CurrentCommand.Completion->SetValue(OutputFilename);
+		delete CurrentCommand.Completion;
+	}
 
-			// Signal the async completion callback, if one was provided when the capture was stopped.
-			if (CurrentCommand.Completion)
-			{
-				CurrentCommand.Completion->SetValue(OutputFilename);
-				delete CurrentCommand.Completion;
-			}
+	FileWriteBlockingEvent->Reset();
 
-			FileWriteBlockingEvent->Reset();
+	// No output filename means we weren't running a capture.
+	bool bCaptureEnded = true;
+	if (OutputFilename.IsEmpty())
+	{
+		UE_LOG(LogCsvProfiler, Warning, TEXT("Capture Stop requested, but no capture was running!"));
+	}
+	else
+	{
+		OutputFilename.Reset();
 
-			// No output filename means we weren't running a capture.
-			bool bCaptureEnded = true;
-			if (OutputFilename.IsEmpty())
+		// Handle repeats
+		if (GCsvRepeatCount != 0 && GCsvRepeatFrameCount > 0)
+		{
+			if (GCsvRepeatCount > 0)
 			{
-				UE_LOG(LogCsvProfiler, Warning, TEXT("Capture Stop requested, but no capture was running!"));
+				GCsvRepeatCount--;
 			}
-			else
+			if (GCsvRepeatCount != 0)
 			{
-				OutputFilename.Reset();
+				bCaptureEnded = false;
 
-				// Handle repeats
-				if (GCsvRepeatCount != 0 && GCsvRepeatFrameCount > 0)
-				{
-					if (GCsvRepeatCount > 0)
-					{
-						GCsvRepeatCount--;
-					}
-					if (GCsvRepeatCount != 0)
-					{
-						bCaptureEnded = false;
-
-						// TODO: support directories
-						BeginCapture(GCsvRepeatFrameCount);
-					}
-				}
-			}
-
-			if (bCaptureEnded && (GCsvExitOnCompletion || FParse::Param(FCommandLine::Get(), TEXT("ExitAfterCsvProfiling"))))
-			{
-				bool bForceExit = !!CVarCsvForceExit.GetValueOnGameThread();
-				FPlatformMisc::RequestExit(bForceExit, TEXT("CsvProfiler.ExitAfterCsvProfiling"));
+				// TODO: support directories
+				BeginCapture(GCsvRepeatFrameCount);
 			}
 		}
 	}
 
-	GCsvProfilerFrameNumber++;
+	if (bCaptureEnded && (GCsvExitOnCompletion || FParse::Param(FCommandLine::Get(), TEXT("ExitAfterCsvProfiling"))))
+	{
+		bool bForceExit = !!CVarCsvForceExit.GetValueOnGameThread();
+		FPlatformMisc::RequestExit(bForceExit, TEXT("CsvProfiler.ExitAfterCsvProfiling"));
+	}
+	return true;
 }
 
 void FCsvProfiler::OnEndFramePostFork()
@@ -3392,6 +3403,7 @@ TSharedFuture<FString> FCsvProfiler::EndCapture(FGraphEventRef EventToSignal)
 	OnCSVProfileEndRequestedDelegate.Broadcast();
 
 	SetNonPersistentMetadata(TEXT("EndTimestamp"), *FString::Printf(TEXT("%lld"), FDateTime::UtcNow().ToUnixTimestamp()));
+	SetNonPersistentMetadata(TEXT("CaptureDuration"), *FString::SanitizeFloat(FPlatformTime::Seconds()-CaptureStartTime));
 
 	TPromise<FString>* Completion = new TPromise<FString>([EventToSignal]()
 	{
@@ -4218,4 +4230,4 @@ void CSVTest()
 
 #endif // CSV_PROFILER_ALLOW_DEBUG_FEATURES
 
-#endif // CSV_PROFILER
+#endif // CSV_PROFILER
\ No newline at end of file
diff --git a/Engine/Source/Runtime/Core/Public/ProfilingDebugging/CsvProfiler.h b/Engine/Source/Runtime/Core/Public/ProfilingDebugging/CsvProfiler.h
--- a/Engine/Source/Runtime/Core/Public/ProfilingDebugging/CsvProfiler.h
+++ b/Engine/Source/Runtime/Core/Public/ProfilingDebugging/CsvProfiler.h
@@ -513,6 +513,10 @@ private:
 
 	float ProcessStatData();
 
+	void BeginCaptureInternal(const FCsvCaptureCommand& CurrentCommand);
+
+	bool TryEndCaptureInternal(const FCsvCaptureCommand& CurrentCommand);
+
 	int32 NumFramesToCapture;
 	int32 CaptureFrameNumber;
 	int32 CaptureFrameNumberRT;
@@ -524,6 +528,8 @@ private:
 	uint64 LastEndFrameTimestamp;
 	uint32 CaptureEndFrameCount;
 
+	double CaptureStartTime;
+
 	FString OutputFilename;
 	TQueue<FCsvCaptureCommand> CommandQueue;
 	FCsvProfilerProcessingThread* ProcessingThread;
