PR #11262: Fix copy/paste mistake in code documentation.

#rnx

diff --git a/Engine/Source/Runtime/Core/Public/Misc/IEngineCrypto.h b/Engine/Source/Runtime/Core/Public/Misc/IEngineCrypto.h
--- a/Engine/Source/Runtime/Core/Public/Misc/IEngineCrypto.h
+++ b/Engine/Source/Runtime/Core/Public/Misc/IEngineCrypto.h
@@ -61,7 +61,7 @@ struct IEngineCrypto : public IModularFeature
 	virtual int32 DecryptPublic(const TArrayView<const uint8> InSource, TArray<uint8>& OutDestination, FRSAKeyHandle InKey) = 0;
 
 	/**
-	 * Encrypt the supplied byte data using the given private key
+	 * Decrypt the supplied byte data using the given private key
 	 */
 	virtual int32 DecryptPrivate(const TArrayView<const uint8> InSource, TArray<uint8>& OutDestination, FRSAKeyHandle InKey) = 0;
 };
\ No newline at end of file
diff --git a/Engine/Source/Runtime/RSA/Public/RSA.h b/Engine/Source/Runtime/RSA/Public/RSA.h
--- a/Engine/Source/Runtime/RSA/Public/RSA.h
+++ b/Engine/Source/Runtime/RSA/Public/RSA.h
@@ -46,7 +46,7 @@ struct FRSA
 	static RSA_API int32 DecryptPublic(const TArrayView<const uint8> InSource, TArray<uint8>& OutDestination, const FRSAKeyHandle InKey);
 
 	/**
-	 * Encrypt the supplied byte data using the given private key
+	 * Decrypt the supplied byte data using the given private key
 	 */
 	static RSA_API int32 DecryptPrivate(const TArrayView<const uint8> InSource, TArray<uint8>& OutDestination, const FRSAKeyHandle InKey);
 };
