Fix std::is_default_constructible<TPair<K,V>>

#rnx

diff --git a/Engine/Plugins/Importers/USDImporter/Source/USDStage/Private/USDLevelSequenceHelper.cpp b/Engine/Plugins/Importers/USDImporter/Source/USDStage/Private/USDLevelSequenceHelper.cpp
--- a/Engine/Plugins/Importers/USDImporter/Source/USDStage/Private/USDLevelSequenceHelper.cpp
+++ b/Engine/Plugins/Importers/USDImporter/Source/USDStage/Private/USDLevelSequenceHelper.cpp
@@ -738,6 +738,10 @@ private:
 
 	struct FPrimTwinBindings
 	{
+		// clang fix for std::is_default_constructible_v 
+		// returning false in inlined code of outer class
+		FPrimTwinBindings() {}
+
 		ULevelSequence* Sequence = nullptr;
 
 		// For now we support one binding per component type (mostly so we can fit a binding to a scene component and
diff --git a/Engine/Source/Editor/UnrealEd/Classes/Editor/Transactor.h b/Engine/Source/Editor/UnrealEd/Classes/Editor/Transactor.h
--- a/Engine/Source/Editor/UnrealEd/Classes/Editor/Transactor.h
+++ b/Engine/Source/Editor/UnrealEd/Classes/Editor/Transactor.h
@@ -220,6 +220,10 @@ protected:
 
 	struct FObjectRecords
 	{
+		// clang fix for std::is_default_constructible_v 
+		// returning false in inlined code of outer class
+		FObjectRecords() {}
+
 		friend FArchive& operator<<(FArchive& Ar, FObjectRecords& ObjectRecords)
 		{
 			Ar << ObjectRecords.SaveCount;
@@ -239,6 +243,10 @@ protected:
 
 	struct FPackageRecord
 	{
+		// clang fix for std::is_default_constructible_v 
+		// returning false in inlined code of outer class
+		FPackageRecord() {}
+
 		friend FArchive& operator<<(FArchive& Ar, FPackageRecord& PackageRecord)
 		{
 			Ar << PackageRecord.DirtyFenceCount;
diff --git a/Engine/Source/Runtime/Core/Public/Containers/Set.h b/Engine/Source/Runtime/Core/Public/Containers/Set.h
--- a/Engine/Source/Runtime/Core/Public/Containers/Set.h
+++ b/Engine/Source/Runtime/Core/Public/Containers/Set.h
@@ -160,7 +160,7 @@ class TSetElementBase
 public:
 	typedef InElementType ElementType;
 
-	FORCEINLINE TSetElementBase() {}
+	TSetElementBase() = default;
 
 	/** Initialization constructor. */
 	template <
@@ -194,7 +194,7 @@ class TSetElementBase<InElementType, true>
 public:
 	typedef InElementType ElementType;
 
-	FORCEINLINE TSetElementBase() {}
+	TSetElementBase() = default;
 
 	/** Initialization constructor. */
 	template <
@@ -228,8 +228,7 @@ class TSetElement : public TSetElementBase<InElementType, THasTypeLayout<InEleme
 	using Super = TSetElementBase<InElementType, THasTypeLayout<InElementType>::Value>;
 public:
 	/** Default constructor. */
-	FORCEINLINE TSetElement()
-	{}
+	TSetElement() = default;
 
 	/** Initialization constructor. */
 	template <
diff --git a/Engine/Source/Runtime/Core/Public/Templates/Tuple.h b/Engine/Source/Runtime/Core/Public/Templates/Tuple.h
--- a/Engine/Source/Runtime/Core/Public/Templates/Tuple.h
+++ b/Engine/Source/Runtime/Core/Public/Templates/Tuple.h
@@ -11,6 +11,7 @@
 #include "Templates/Requires.h"
 #include "Templates/TypeHash.h"
 #include <tuple>
+#include <type_traits>
 
 // This workaround exists because Visual Studio causes false positives for code like this during static analysis:
 //
@@ -72,6 +73,9 @@ namespace UE::Core::Private::Tuple
 		}
 
 		TTupleBaseElement()
+#ifdef __cpp_concepts // A C++17 fix was deemed too complex
+			requires (std::is_default_constructible_v<T>)
+#endif
 			: Value()
 		{
 		}
@@ -97,6 +101,9 @@ namespace UE::Core::Private::Tuple
 		}
 
 		TTupleBaseElement()
+#ifdef __cpp_concepts // A C++17 fix was deemed too complex
+			requires (std::is_default_constructible_v<T>)
+#endif
 			: Key()
 		{
 		}
diff --git a/Engine/Source/Runtime/Interchange/Core/Public/Types/AttributeStorage.h b/Engine/Source/Runtime/Interchange/Core/Public/Types/AttributeStorage.h
--- a/Engine/Source/Runtime/Interchange/Core/Public/Types/AttributeStorage.h
+++ b/Engine/Source/Runtime/Interchange/Core/Public/Types/AttributeStorage.h
@@ -700,6 +700,10 @@ namespace UE
 			/** Structure used to hold the attribute information stored in the attribute allocation table. */
 			struct FAttributeAllocationInfo
 			{
+				// clang fix for std::is_default_constructible_v 
+				// returning false in inlined code of outer class
+				FAttributeAllocationInfo() {}
+				
         		//The offset in the storage
         		uint64 Offset = 0;
         		//The size of the data in the storage
