Adding utf8 strings to test suite

This allows SoftObjectPath.cpp to compile when subpath is changed from FString to FUtf8String

diff --git a/Engine/Source/Runtime/Core/Private/Misc/AutomationTest.cpp b/Engine/Source/Runtime/Core/Private/Misc/AutomationTest.cpp
--- a/Engine/Source/Runtime/Core/Private/Misc/AutomationTest.cpp
+++ b/Engine/Source/Runtime/Core/Private/Misc/AutomationTest.cpp
@@ -1778,6 +1778,57 @@ bool FAutomationTestBase::TestEqual(const TCHAR* What, const TCHAR* Actual, cons
 	return true;
 }
 
+bool FAutomationTestBase::TestEqual(const TCHAR* What, FUtf8StringView Actual, FUtf8StringView Expected)
+{
+	if (Actual.Compare(Expected, ESearchCase::CaseSensitive))
+	{
+		AddError(FString::Printf(TEXT("Expected '%s' to be \"%.*hs\", but it was \"%.*hs\"."), What, Expected.Len(), Expected.GetData(), Actual.Len(), Actual.GetData()), 1);
+		return false;
+	}
+	return true;
+}
+
+bool FAutomationTestBase::TestEqual(const TCHAR* What, FStringView Actual, FStringView Expected)
+{
+	if (Actual.Compare(Expected, ESearchCase::CaseSensitive))
+	{
+		AddError(FString::Printf(TEXT("Expected '%s' to be \"%.*s\", but it was \"%.*s\"."), What, Expected.Len(), Expected.GetData(), Actual.Len(), Actual.GetData()), 1);
+		return false;
+	}
+	return true;
+}
+
+bool FAutomationTestBase::TestNotEqual(const TCHAR* What, const TCHAR* Actual, const TCHAR* Expected)
+{
+	if (FCString::Strcmp(Actual, Expected) == 0)
+	{
+		AddError(FString::Printf(TEXT("Expected '%s' not to be \"%s\", but it was \"%s\"."), What, Expected, Actual), 1);
+		return false;
+	}
+	return true;
+}
+
+bool FAutomationTestBase::TestNotEqual(const TCHAR* What, FUtf8StringView Actual, FUtf8StringView Expected)
+{
+	if (Actual.Compare(Expected, ESearchCase::CaseSensitive) == 0)
+	{
+		AddError(FString::Printf(TEXT("Expected '%s' not to be \"%.*hs\", but it was \"%.*hs\"."), What, Expected.Len(), Expected.GetData(), Actual.Len(), Actual.GetData()), 1);
+		return false;
+	}
+	return true;
+}
+
+bool FAutomationTestBase::TestNotEqual(const TCHAR* What, FStringView Actual, FStringView Expected)
+{
+	if (Actual.Compare(Expected, ESearchCase::CaseSensitive) == 0)
+	{
+		AddError(FString::Printf(TEXT("Expected '%s' not to be \"%.*s\", but it was \"%.*s\"."), What, Expected.Len(), Expected.GetData(), Actual.Len(), Actual.GetData()), 1);
+		return false;
+	}
+	return true;
+}
+
+
 bool FAutomationTestBase::TestNotEqual(const TCHAR* What, const float Actual, const float Expected, float Tolerance)
 {
 	if (FMath::IsNearlyEqual(Actual, Expected, Tolerance))
@@ -1808,6 +1859,26 @@ bool FAutomationTestBase::TestEqualInsensitive(const TCHAR* What, const TCHAR* A
 	return true;
 }
 
+bool FAutomationTestBase::TestEqualInsensitive(const TCHAR* What, FStringView Actual, FStringView Expected)
+{
+	if (Actual.Compare(Expected, ESearchCase::IgnoreCase) != 0)
+	{
+		AddError(FString::Printf(TEXT("Expected '%s' to be \"%.*s\", but it was \"%.*s\"."), What, Expected.Len(), Expected.GetData(), Actual.Len(), Actual.GetData()), 1);
+		return false;
+	}
+	return true;
+}
+
+bool FAutomationTestBase::TestEqualInsensitive(const TCHAR* What, FUtf8StringView Actual, FUtf8StringView Expected)
+{
+	if (Actual.Compare(Expected, ESearchCase::IgnoreCase) != 0)
+	{
+		AddError(FString::Printf(TEXT("Expected '%s' to be \"%.*hs\", but it was \"%.*hs\"."), What, Expected.Len(), Expected.GetData(), Actual.Len(), Actual.GetData()), 1);
+		return false;
+	}
+	return true;
+}
+
 bool FAutomationTestBase::TestNotEqualInsensitive(const TCHAR* What, const TCHAR* Actual, const TCHAR* Expected)
 {
 	if (Actual && Expected)
@@ -1829,6 +1900,48 @@ bool FAutomationTestBase::TestNotEqualInsensitive(const TCHAR* What, const TCHAR
 	return true;
 }
 
+bool FAutomationTestBase::TestNotEqualInsensitive(const TCHAR* What, FStringView Actual, FStringView Expected)
+{
+	if (Actual.GetData() && Expected.GetData())
+	{
+		if (Actual.Compare(Expected, ESearchCase::IgnoreCase) == 0)
+		{
+			AddError(FString::Printf(TEXT("Expected '%s' to differ from \"%.*s\", but it was \"%.*s\"."), What, Expected.Len(), Expected.GetData(), Actual.Len(), Actual.GetData()), 1);
+			return false;
+		}
+	}
+	else // null exists
+	{
+		if (Actual == Expected)
+		{
+			AddError(FString::Printf(TEXT("Expected '%s' to differ but both values were unexpectedly null"), What), 1);
+			return false;
+		}
+	}
+	return true;
+}
+
+bool FAutomationTestBase::TestNotEqualInsensitive(const TCHAR* What, FUtf8StringView Actual, FUtf8StringView Expected)
+{
+	if (Actual.GetData() && Expected.GetData())
+	{
+		if (Actual.Compare(Expected, ESearchCase::IgnoreCase) == 0)
+		{
+			AddError(FString::Printf(TEXT("Expected '%s' to differ from \"%.*hs\", but it was \"%.*hs\"."), What, Expected.Len(), Expected.GetData(), Actual.Len(), Actual.GetData()), 1);
+			return false;
+		}
+	}
+	else // null exists
+	{
+		if (Actual == Expected)
+		{
+			AddError(FString::Printf(TEXT("Expected '%s' to differ but both values were unexpectedly null"), What), 1);
+			return false;
+		}
+	}
+	return true;
+}
+
 bool FAutomationTestBase::TestNearlyEqual(const TCHAR* What, const float Actual, const float Expected, float Tolerance)
 {
 	return TestEqual(What, Actual, Expected, Tolerance);
diff --git a/Engine/Source/Runtime/Core/Public/Misc/AutomationTest.h b/Engine/Source/Runtime/Core/Public/Misc/AutomationTest.h
--- a/Engine/Source/Runtime/Core/Public/Misc/AutomationTest.h
+++ b/Engine/Source/Runtime/Core/Public/Misc/AutomationTest.h
@@ -9,6 +9,7 @@
 #include "Containers/Queue.h"
 #include "Containers/Set.h"
 #include "Containers/UnrealString.h"
+#include "Containers/StringView.h"
 #include "CoreTypes.h"
 #include "Delegates/Delegate.h"
 #include "Delegates/DelegateBase.h"
@@ -1818,8 +1819,17 @@ public:
 	CORE_API bool TestEqual(const TCHAR* What, const FColor Actual, const FColor Expected);
 	CORE_API bool TestEqual(const TCHAR* What, const FLinearColor Actual, const FLinearColor Expected);
 	CORE_API bool TestEqual(const TCHAR* What, const TCHAR* Actual, const TCHAR* Expected);
+	CORE_API bool TestEqual(const TCHAR* What, FStringView Actual, FStringView Expected);
+	CORE_API bool TestEqual(const TCHAR* What, FUtf8StringView Actual, FUtf8StringView Expected);
+	CORE_API bool TestNotEqual(const TCHAR* What, const TCHAR* Actual, const TCHAR* Expected);
+	CORE_API bool TestNotEqual(const TCHAR* What, FStringView Actual, FStringView Expected);
+	CORE_API bool TestNotEqual(const TCHAR* What, FUtf8StringView Actual, FUtf8StringView Expected);
 	CORE_API bool TestEqualInsensitive(const TCHAR* What, const TCHAR* Actual, const TCHAR* Expected);
+	CORE_API bool TestEqualInsensitive(const TCHAR* What, FStringView Actual, FStringView Expected);
+	CORE_API bool TestEqualInsensitive(const TCHAR* What, FUtf8StringView Actual, FUtf8StringView Expected);
 	CORE_API bool TestNotEqualInsensitive(const TCHAR* What, const TCHAR* Actual, const TCHAR* Expected);
+	CORE_API bool TestNotEqualInsensitive(const TCHAR* What, FStringView Actual, FStringView Expected);
+	CORE_API bool TestNotEqualInsensitive(const TCHAR* What, FUtf8StringView Actual, FUtf8StringView Expected);
 
 	bool TestEqual(const FString& What, const int32 Actual, const int32 Expected)
 	{
@@ -1861,94 +1871,21 @@ public:
 		return TestEqual(*What, Actual, Expected);
 	}
 
-	bool TestEqual(const TCHAR* What, const FString& Actual, const TCHAR* Expected)
+	bool TestEqual(const FString& What, FUtf8StringView Actual, FUtf8StringView  Expected)
 	{
-		return TestEqualInsensitive(What, *Actual, Expected);
+		return TestEqualInsensitive(*What, Actual, Expected);
 	}
 
-	bool TestEqual(const FString& What, const FString& Actual, const TCHAR* Expected)
+	bool TestEqual(const FString& What, FStringView Actual, FStringView Expected)
 	{
-		return TestEqualInsensitive(*What, *Actual, Expected);
+		return TestEqualInsensitive(*What, Actual, Expected);
 	}
 
-	bool TestEqual(const TCHAR* What, const TCHAR* Actual, const FString& Expected)
+	bool TestEqual(const FString& What, FStringView Actual, FUtf8StringView Expected)
 	{
-		return TestEqualInsensitive(What, Actual, *Expected);
-	}
-
-	bool TestEqual(const FString& What, const TCHAR* Actual, const FString& Expected)
-	{
-		return TestEqualInsensitive(*What, Actual, *Expected);
-	}
-
-	bool TestEqual(const TCHAR* What, const FString& Actual, const FString& Expected)
-	{
-		return TestEqualInsensitive(What, *Actual, *Expected);
-	}
-
-	bool TestEqual(const FString& What, const FString& Actual, const FString& Expected)
-	{
-		return TestEqualInsensitive(*What, *Actual, *Expected);
-	}
-
-	bool TestEqualInsensitive(const TCHAR* What, const FString& Actual, const TCHAR* Expected)
-	{
-		return TestEqualInsensitive(What, *Actual, Expected);
-	}
-
-	bool TestEqualInsensitive(const FString& What, const FString& Actual, const TCHAR* Expected)
-	{
-		return TestEqualInsensitive(*What, *Actual, Expected);
-	}
-
-	bool TestEqualInsensitive(const TCHAR* What, const TCHAR* Actual, const FString& Expected)
-	{
-		return TestEqualInsensitive(What, Actual, *Expected);
-	}
-
-	bool TestEqualInsensitive(const FString& What, const TCHAR* Actual, const FString& Expected)
-	{
-		return TestEqualInsensitive(*What, Actual, *Expected);
-	}
-
-	bool TestEqualInsensitive(const TCHAR* What, const FString& Actual, const FString& Expected)
-	{
-		return TestEqualInsensitive(What, *Actual, *Expected);
-	}
-
-	bool TestEqualInsensitive(const FString& What, const FString& Actual, const FString& Expected)
-	{
-		return TestEqualInsensitive(*What, *Actual, *Expected);
-	}
-
-	bool TestNotEqualInsensitive(const TCHAR* What, const FString& Actual, const TCHAR* Expected)
-	{
-		return TestNotEqualInsensitive(What, *Actual, Expected);
-	}
-
-	bool TestNotEqualInsensitive(const FString& What, const FString& Actual, const TCHAR* Expected)
-	{
-		return TestNotEqualInsensitive(*What, *Actual, Expected);
-	}
-
-	bool TestNotEqualInsensitive(const TCHAR* What, const TCHAR* Actual, const FString& Expected)
-	{
-		return TestNotEqualInsensitive(What, Actual, *Expected);
-	}
-
-	bool TestNotEqualInsensitive(const FString& What, const TCHAR* Actual, const FString& Expected)
-	{
-		return TestNotEqualInsensitive(*What, Actual, *Expected);
-	}
-
-	bool TestNotEqualInsensitive(const TCHAR* What, const FString& Actual, const FString& Expected)
-	{
-		return TestNotEqualInsensitive(What, *Actual, *Expected);
-	}
-
-	bool TestNotEqualInsensitive(const FString& What, const FString& Actual, const FString& Expected)
-	{
-		return TestNotEqualInsensitive(*What, *Actual, *Expected);
+		// This overload is here because there are some tests that have raw string literals as the expected value.
+		FString Tmp(Expected);
+		return TestEqualInsensitive(*What, Actual, Tmp);
 	}
 
 	/**
