Build fix for issue 501491

diff --git a/Engine/Source/Runtime/Core/Public/Misc/AutomationTest.h b/Engine/Source/Runtime/Core/Public/Misc/AutomationTest.h
--- a/Engine/Source/Runtime/Core/Public/Misc/AutomationTest.h
+++ b/Engine/Source/Runtime/Core/Public/Misc/AutomationTest.h
@@ -1888,6 +1888,13 @@ public:
 		return TestEqualInsensitive(*What, Actual, Tmp);
 	}
 
+	bool TestEqual(const FString& What, FUtf8StringView Actual, FStringView Expected)
+	{
+		// This overload is here because there are some tests that have raw string literals as the Actual value.
+		FString Tmp(Actual);
+		return TestEqualInsensitive(*What, Tmp, Expected);
+	}
+
 	/**
 	 * Logs an error if the two values are not equal.
 	 *
