VeryLargePageAllocator page commit limits and general refactor/cleanup:
- Add cvars VeryLargePageAllocator.MaxCommittedPageCountDefault and VeryLargePageAllocator.MaxCommittedPageCountSmallPool to limit the number of large pages committed for each pool. Since VLPA very rarely releases pages, this avoids the situation where VLPA permanently holds too much memory, leaving less for other large allocations or rendering etc.

Cleanup/refactor:
- Strip out unused code paths and corresponding defines. UE_VERYLARGEPAGEALLOCATOR_TAKEONALL64KBALLOCATIONS, UE_USE_VERYLARGEPAGEALLOCATOR_FALLBACKPATH and LARGEPAGEALLOCATOR_SORT_OnAddress were always enabled.
- Add helper functions GetOrAllocLargePage and AllocNewLargePage to simplify allocation logic. This also avoids some duplication of code in preallocation and allocate()
- Strip the VeryLargePageAllocator prefix from various cvars and globals to improve readability (making sure they're all static scope)
- Rename the function IsPartOf to IsSmallBlockAllocation to be less misleading
- Replace the cvar VeryLargePageAllocator.LeavePageCachingEnabledWhenOOMHappened with VeryLargePageAllocator.DisablePageCachingOnOOM (which does the opposite) for clarity

#tests tested on a devkit before/after and compared VLPA page allocation stats

diff --git a/Engine/Source/Runtime/Core/Private/HAL/Allocators/CachedOSVeryLargePageAllocator.cpp b/Engine/Source/Runtime/Core/Private/HAL/Allocators/CachedOSVeryLargePageAllocator.cpp
--- a/Engine/Source/Runtime/Core/Private/HAL/Allocators/CachedOSVeryLargePageAllocator.cpp
+++ b/Engine/Source/Runtime/Core/Private/HAL/Allocators/CachedOSVeryLargePageAllocator.cpp
@@ -13,36 +13,57 @@
 #if UE_USE_VERYLARGEPAGEALLOCATOR
 
 CORE_API bool GEnableVeryLargePageAllocator = true;
-static int32 GVeryLargePageAllocatorMaxEmptyBackStoreCount[FMemory::AllocationHints::Max] = {
-	0,	// FMemory::AllocationHints::Default
-	0,	// FMemory::AllocationHints::Temporary
-	0	// FMemory::AllocationHints::SmallPool
-};
 
-static_assert(int32(FMemory::AllocationHints::Max) == 3); // ensure FMemory::AllocationHints has 3 types of hint so GVeryLargePageAllocatorMaxEmptyBackStoreCount has needed values
+static_assert(int32(FMemory::AllocationHints::Max) == 3); // ensure FMemory::AllocationHints has 3 types of hint so GMaxEmptyBackStoreCount has needed values
 
-bool GVeryLargePageAllocatorPreAllocatePools = true;
-static FAutoConsoleVariableRef CVarVeryLargePageAllocatorPreAllocatePools(
+static bool GPreAllocatePools = true;
+static FAutoConsoleVariableRef CVarPreAllocatePools(
 	TEXT("VeryLargePageAllocator.PreAllocatePools"),
-	GVeryLargePageAllocatorPreAllocatePools,
+	GPreAllocatePools,
 	TEXT("Having pages preallocated and cached during the life of the title help to avoid defragmentation of physical memory.\n")
-	TEXT("Preallocation is disabled when system reaches OOM."));
+	TEXT("Preallocation may be disabled when system reaches OOM (see VeryLargePageAllocator.DisablePageCachingOnOOM)"));
 
-bool GVeryLargePageAllocatorLeavePageCachingEnabledWhenOOMHappened = true;
-static FAutoConsoleVariableRef CVarVeryLargePageAllocatorLeavePageCachingEnabledWhenOOMHappened(
-	TEXT("VeryLargePageAllocator.LeavePageCachingEnabledWhenOOMHappened"),
-	GVeryLargePageAllocatorLeavePageCachingEnabledWhenOOMHappened,
-	TEXT("Leave page caching enabled when a OOM has happened and all unusued pages have been freed (so new allocated pages gets cached again)"));
+static bool GDisablePageCachingOnOOM = false;
+static FAutoConsoleVariableRef CVarDisablePageCachingOnOOM(
+	TEXT("VeryLargePageAllocator.DisablePageCachingOnOOM"),
+	GDisablePageCachingOnOOM,
+	TEXT("If enabled, permanently disable page caching when a OOM happens and all unused pages have been freed (so new allocated pages gets cached again)\n")
+	TEXT("This can lead to unpredictable performance on some platforms."));
 
-static FAutoConsoleVariableRef CVarVeryLargePageAllocatorMaxEmptyBackstoreDefault(
+static int32 GMaxEmptyBackStoreCount[FMemory::AllocationHints::Max] = {
+	0,	// FMemory::AllocationHints::Default
+	0,	// FMemory::AllocationHints::Temporary
+	0	// FMemory::AllocationHints::SmallPool
+};
+static FAutoConsoleVariableRef CVarMaxEmptyBackstoreDefault(
 	TEXT("VeryLargePageAllocator.MaxEmptyBackstoreDefault"),
-	GVeryLargePageAllocatorMaxEmptyBackStoreCount[FMemory::AllocationHints::Default],
-	TEXT("Number of free pages (2MB each) to cache (not decommited) for allocation hint DEFAULT"));
+	GMaxEmptyBackStoreCount[FMemory::AllocationHints::Default],
+	TEXT("Number of free pages (2MB each) to cache (not decommitted) for allocation hint DEFAULT"));
 
-static FAutoConsoleVariableRef CVarVeryLargePageAllocatorMaxEmptyBackstoreSmallPool(
+static FAutoConsoleVariableRef CVarMaxEmptyBackstoreSmallPool(
 	TEXT("VeryLargePageAllocator.MaxEmptyBackstoreSmallPool"),
-	GVeryLargePageAllocatorMaxEmptyBackStoreCount[FMemory::AllocationHints::SmallPool],
-	TEXT("Number of free pages (2MB each) to cache (not decommited) for allocation hint SMALL POOL"));
+	GMaxEmptyBackStoreCount[FMemory::AllocationHints::SmallPool],
+	TEXT("Number of free pages (2MB each) to cache (not decommitted) for allocation hint SMALL POOL"));
+
+
+static int32 GMaxCommittedPageCount[FMemory::AllocationHints::Max] = {
+	0,	// FMemory::AllocationHints::Default
+	0,	// FMemory::AllocationHints::Temporary
+	0	// FMemory::AllocationHints::SmallPool
+};
+
+static FAutoConsoleVariableRef CVarMaxPageCountDefault(
+	TEXT("VeryLargePageAllocator.MaxCommittedPageCountDefault"),
+	GMaxCommittedPageCount[FMemory::AllocationHints::Default],
+	TEXT("Maximum number of pages (2MB each) to use for allocation hint DEFAULT")
+	TEXT("Adjustable at runtime. If we run out of pages we'll fall back to the standard allocator"));
+
+static FAutoConsoleVariableRef CVarMaxPageCountSmallPool(
+	TEXT("VeryLargePageAllocator.MaxCommittedPageCountSmallPool"),
+	GMaxCommittedPageCount[FMemory::AllocationHints::SmallPool],
+	TEXT("Maximum number of pages (2MB each) to use for allocation hint SMALL POOL\n")
+	TEXT("Adjustable at runtime. If we run out of pages we'll fall back to the standard allocator"));
+
 
 #if CSV_PROFILER_STATS
 CSV_DECLARE_CATEGORY_MODULE_EXTERN(CORE_API, FMemory);
@@ -56,12 +77,7 @@ void FCachedOSVeryLargePageAllocator::Init()
 	Block = FPlatformMemory::FPlatformVirtualMemoryBlock::AllocateVirtual(AddressSpaceToReserve);
 	AddressSpaceReserved = (uintptr_t)Block.GetVirtualPointer();
 	AddressSpaceReservedEnd = AddressSpaceReserved + AddressSpaceToReserve;
-#if UE_VERYLARGEPAGEALLOCATOR_TAKEONALL64KBALLOCATIONS
 	AddressSpaceReservedEndSmallPool = AddressSpaceReserved + (AddressSpaceToReserve / 2);
-#else
-	AddressSpaceReservedEndSmallPool = AddressSpaceReservedEnd;
-#endif
-
 
 	for (int i = 0; i < FMemory::AllocationHints::Max; i++)
 	{
@@ -70,26 +86,19 @@ void FCachedOSVeryLargePageAllocator::Init()
 		UsedLargePagesHead[i] = nullptr;
 		EmptyButAvailableLargePagesHead[i] = nullptr;
 		EmptyBackStoreCount[i] = 0;
-		CommitedLargePagesCount[i] = 0;
+		CommittedLargePagesCount[i] = 0;
 	}
-
-#if UE_VERYLARGEPAGEALLOCATOR_TAKEONALL64KBALLOCATIONS
 	for (int i = 0; i < NumberOfLargePages / 2; i++)
-#else
-	for (int i = 0; i < NumberOfLargePages; i++)
-#endif
 	{
 		LargePagesArray[i].Init((void*)((uintptr_t)AddressSpaceReserved + (i * SizeOfLargePage)));
 		LargePagesArray[i].LinkHead(FreeLargePagesHead[FMemory::AllocationHints::SmallPool]);
 	}
 
-#if UE_VERYLARGEPAGEALLOCATOR_TAKEONALL64KBALLOCATIONS
 	for (int i = NumberOfLargePages / 2; i < NumberOfLargePages; i++)
 	{
 		LargePagesArray[i].Init((void*)((uintptr_t)AddressSpaceReserved + (i * SizeOfLargePage)));
 		LargePagesArray[i].LinkHead(FreeLargePagesHead[FMemory::AllocationHints::Default]);
 	}
-#endif
 
 	if (!GEnableVeryLargePageAllocator)
 	{
@@ -129,32 +138,26 @@ void FCachedOSVeryLargePageAllocator::Refresh()
 		{
 			FMemory::AllocationHints AllocationHint = FMemory::AllocationHints(i);
 
-			int32 LargePageCount = CommitedLargePagesCount[AllocationHint] + EmptyBackStoreCount[AllocationHint];
+			int32 LargePageCount = CommittedLargePagesCount[AllocationHint] + EmptyBackStoreCount[AllocationHint];
 
-			if (LargePageCount < GVeryLargePageAllocatorMaxEmptyBackStoreCount[AllocationHint] && GVeryLargePageAllocatorPreAllocatePools)
+			if (LargePageCount < GMaxEmptyBackStoreCount[AllocationHint] && GPreAllocatePools)
 			{
 				// Preallocate large pages
-				for (; LargePageCount < GVeryLargePageAllocatorMaxEmptyBackStoreCount[AllocationHint]; LargePageCount++)
+				for (; LargePageCount < GMaxEmptyBackStoreCount[AllocationHint]; LargePageCount++)
 				{
-					LLM_PLATFORM_SCOPE(ELLMTag::FMalloc);
-
-					FLargePage* LargePage = FreeLargePagesHead[AllocationHint];
-
-					check(LargePage != nullptr); // Can't happen
-					LargePage->AllocationHint = AllocationHint;
-					LargePage->Unlink();
-
-					if (!Block.Commit(LargePage->BaseAddress - AddressSpaceReserved, SizeOfLargePage, false))
+					bool bCommitFailure = false;
+					FLargePage* LargePage = AllocNewLargePage(AllocationHint, nullptr, bCommitFailure);
+					if (!LargePage)
 					{
-						// Cant commit 2MB pages, stop preallocating and return page to FreeLargePagesHead list
-						LargePage->LinkHead(FreeLargePagesHead[AllocationHint]);
-						GVeryLargePageAllocatorPreAllocatePools = false;
+						if (bCommitFailure && GDisablePageCachingOnOOM)
+						{
+							// If we failed to preallocate due to a commit failure then disable preallocation in future
+							UE_LOG(LogMemory, Warning, TEXT("LargePageAllocator - Preallocation permanently disabled due to failure to commit pages"));
+							GPreAllocatePools = false;
+						}
 						break;
 					}
-
 					LargePage->LinkHead(EmptyButAvailableLargePagesHead[AllocationHint]);
-					CachedFree += SizeOfLargePage;
-					CommitedLargePagesCount[AllocationHint] += 1;
 					EmptyBackStoreCount[AllocationHint] += 1;
 				}
 			}
@@ -173,98 +176,26 @@ void* FCachedOSVeryLargePageAllocator::Allocate(SIZE_T Size, uint32 AllocationHi
 
 	if (bEnabled && Size == SizeOfSubPage)
 	{
-#if !UE_VERYLARGEPAGEALLOCATOR_TAKEONALL64KBALLOCATIONS
-		if (AllocationHint == FMemory::AllocationHints::SmallPool)
-#endif
-		{
-			bool bLinkToUsedLargePagesWithSpaceHead = false;
-			FLargePage* LargePage = UsedLargePagesWithSpaceHead[AllocationHint];
+		FLargePage* LargePage = GetOrAllocLargePage(AllocationHint, Mutex);
 
-			if (LargePage == nullptr)
+		if (LargePage)
+		{
+			ret = LargePage->Allocate();
+			if (ret)
 			{
-				bLinkToUsedLargePagesWithSpaceHead = true;
-				LargePage = EmptyButAvailableLargePagesHead[AllocationHint];
-
-				if (LargePage != nullptr)
+				// If the page is full then move it to the full list
+				if (LargePage->NumberOfFreeSubPages == 0) 
 				{
-					LargePage->AllocationHint = AllocationHint;
 					LargePage->Unlink();
-					EmptyBackStoreCount[AllocationHint] -= 1;
-				}
-				else
-				{
-					LargePage = FreeLargePagesHead[AllocationHint];
-
-					if (LargePage != nullptr)
-					{
-						LargePage->AllocationHint = AllocationHint;
-						LargePage->Unlink();
-						{
-#if UE_ALLOW_OSMEMORYLOCKFREE
-							FScopeUnlock ScopeUnlock(Mutex);
-#endif
-							LLM_PLATFORM_SCOPE(ELLMTag::FMalloc);
-#if UE_USE_VERYLARGEPAGEALLOCATOR_FALLBACKPATH
-							if (!Block.Commit(LargePage->BaseAddress - AddressSpaceReserved, SizeOfLargePage, false))
-							{
-#if UE_ALLOW_OSMEMORYLOCKFREE
-								if (Mutex != nullptr)
-								{
-									FScopeLock Lock(Mutex);
-									LargePage->LinkHead(FreeLargePagesHead[AllocationHint]);
-								}
-								else
-#endif
-								{
-									LargePage->LinkHead(FreeLargePagesHead[AllocationHint]);
-								}
-
-								// Fallback to regular allocator
-								LargePage = nullptr;
-							}
-							else
-							{
-								// A new large page has been created. Add it to CachedFree counter
-								CachedFree += SizeOfLargePage;
-								CommitedLargePagesCount[AllocationHint] += 1;
-#if CSV_PROFILER_STATS
-								FPlatformAtomics::InterlockedIncrement(&GLargePageAllocatorCommitCount);
-#endif
-							}
-						}
-					}
-#else 
-							Block.Commit(LargePage->BaseAddress - AddressSpaceReserved, SizeOfLargePage);
-							CachedFree += SizeOfLargePage;
-							CommitedLargePagesCount[AllocationHint] += 1;
-						}
-					}
-#endif // UE_USE_VERYLARGEPAGEALLOCATOR_FALLBACKPATH
+					LargePage->LinkHead(UsedLargePagesHead[AllocationHint]);
 				}
+				CachedFree -= SizeOfSubPage;
 			}
-			if (LargePage != nullptr)
+			else
 			{
-				if (bLinkToUsedLargePagesWithSpaceHead)
-				{
-					LargePage->LinkHead(UsedLargePagesWithSpaceHead[AllocationHint]);
-				}
-
-				ret = LargePage->Allocate();
-				if (ret)
+				if (AllocationHint == FMemory::AllocationHints::SmallPool)
 				{
-					if (LargePage->NumberOfFreeSubPages == 0)
-					{
-						LargePage->Unlink();
-						LargePage->LinkHead(UsedLargePagesHead[AllocationHint]);
-					}
-					CachedFree -= SizeOfSubPage;
-				}
-				else
-				{
-					if (AllocationHint == FMemory::AllocationHints::SmallPool)
-					{
-						UE_CLOG(!ret, LogMemory, Fatal, TEXT("The FCachedOSVeryLargePageAllocator has run out of address space for SmallPool allocations, increase UE_VERYLARGEPAGEALLOCATOR_RESERVED_SIZE_IN_GB for your platform!"));
-					}
+					UE_CLOG(!ret, LogMemory, Fatal, TEXT("The FCachedOSVeryLargePageAllocator has run out of address space for SmallPool allocations, increase UE_VERYLARGEPAGEALLOCATOR_RESERVED_SIZE_IN_GB for your platform!"));
 				}
 			}
 		}
@@ -277,7 +208,84 @@ void* FCachedOSVeryLargePageAllocator::Allocate(SIZE_T Size, uint32 AllocationHi
 	return ret;
 }
 
-#define LARGEPAGEALLOCATOR_SORT_OnAddress 1
+
+FCachedOSVeryLargePageAllocator::FLargePage* FCachedOSVeryLargePageAllocator::GetOrAllocLargePage(uint32 AllocationHint, FCriticalSection* Mutex)
+{
+	// Use an existing page with space if one is available
+	FLargePage* LargePage = UsedLargePagesWithSpaceHead[AllocationHint];
+	if (LargePage)
+	{
+		return LargePage;
+	}
+
+	// Attempt to allocate an empty already-committed page
+	LargePage = EmptyButAvailableLargePagesHead[AllocationHint];
+	if (LargePage)
+	{
+		LargePage->AllocationHint = AllocationHint;
+		LargePage->Unlink();
+		EmptyBackStoreCount[AllocationHint] -= 1;
+	}
+	else
+	{
+		// If all else fails, allocate and commit a new page
+		bool bCommitFailureUnused;
+		LargePage = AllocNewLargePage(AllocationHint, Mutex, bCommitFailureUnused);
+	}
+
+	// Move the page to the active list
+	if (LargePage)
+	{
+		LargePage->LinkHead(UsedLargePagesWithSpaceHead[AllocationHint]);
+	}
+	return LargePage;
+}
+
+FCachedOSVeryLargePageAllocator::FLargePage* FCachedOSVeryLargePageAllocator::AllocNewLargePage(uint32 AllocationHint, FCriticalSection* Mutex, bool& bOutCommitFailure)
+{
+	// If there's a limit for this allocation type then check first
+	bOutCommitFailure = false;
+	if (GMaxCommittedPageCount[AllocationHint] > 0 && CommittedLargePagesCount[AllocationHint] >= GMaxCommittedPageCount[AllocationHint])
+	{
+		return nullptr;
+	}
+
+	FLargePage* LargePage = FreeLargePagesHead[AllocationHint];
+	if (LargePage != nullptr)
+	{
+		LargePage->AllocationHint = AllocationHint;
+		LargePage->Unlink();
+		{
+#if UE_ALLOW_OSMEMORYLOCKFREE
+			FScopeUnlock ScopeUnlock(Mutex);
+#endif
+			LLM_PLATFORM_SCOPE(ELLMTag::FMalloc);
+			if (!Block.Commit(LargePage->BaseAddress - AddressSpaceReserved, SizeOfLargePage, false))
+			{
+#if UE_ALLOW_OSMEMORYLOCKFREE
+				if (Mutex != nullptr)
+				{
+					FScopeLock Lock(Mutex);
+					LargePage->LinkHead(FreeLargePagesHead[AllocationHint]);
+				}
+				else
+#endif
+				{
+					LargePage->LinkHead(FreeLargePagesHead[AllocationHint]);
+				}
+				bOutCommitFailure = true;
+				return nullptr;
+			}
+			// A new large page has been created. Add it to CachedFree counter
+			CachedFree += SizeOfLargePage;
+			CommittedLargePagesCount[AllocationHint] += 1;
+#if CSV_PROFILER_STATS
+			FPlatformAtomics::InterlockedIncrement(&GLargePageAllocatorCommitCount);
+#endif
+		}
+	}
+	return LargePage;
+}
 
 void FCachedOSVeryLargePageAllocator::Free(void* Ptr, SIZE_T Size, FCriticalSection* Mutex, bool ThreadIsTimeCritical)
 {
@@ -296,7 +304,7 @@ void FCachedOSVeryLargePageAllocator::Free(void* Ptr, SIZE_T Size, FCriticalSect
 			LargePage->Unlink();
 
 			// move it to EmptyButAvailableLargePagesHead if that pool of backstore is not full yet
-			if (EmptyBackStoreCount[LargePage->AllocationHint] < GVeryLargePageAllocatorMaxEmptyBackStoreCount[LargePage->AllocationHint])
+			if (EmptyBackStoreCount[LargePage->AllocationHint] < GMaxEmptyBackStoreCount[LargePage->AllocationHint])
 			{
 				LargePage->LinkHead(EmptyButAvailableLargePagesHead[LargePage->AllocationHint]);
 				EmptyBackStoreCount[LargePage->AllocationHint] += 1;
@@ -311,7 +319,7 @@ void FCachedOSVeryLargePageAllocator::Free(void* Ptr, SIZE_T Size, FCriticalSect
 					Block.Decommit(LargePage->BaseAddress - AddressSpaceReserved, SizeOfLargePage);
 				}
 
-				CommitedLargePagesCount[LargePage->AllocationHint] -= 1;
+				CommittedLargePagesCount[LargePage->AllocationHint] -= 1;
 #if CSV_PROFILER_STATS
 				FPlatformAtomics::InterlockedIncrement(&GLargePageAllocatorDecommitCount);
 #endif
@@ -323,11 +331,11 @@ void FCachedOSVeryLargePageAllocator::Free(void* Ptr, SIZE_T Size, FCriticalSect
 		else if (LargePage->NumberOfFreeSubPages == 1)
 		{
 			LargePage->Unlink();
-#if LARGEPAGEALLOCATOR_SORT_OnAddress
+			// Sort on address
 			FLargePage* InsertPoint = UsedLargePagesWithSpaceHead[LargePage->AllocationHint];
 			while (InsertPoint != nullptr)
 			{
-				if (LargePage->BaseAddress < InsertPoint->BaseAddress)	// sort on address
+				if (LargePage->BaseAddress < InsertPoint->BaseAddress)
 				{
 					break;
 				}
@@ -341,38 +349,6 @@ void FCachedOSVeryLargePageAllocator::Free(void* Ptr, SIZE_T Size, FCriticalSect
 			{
 				LargePage->LinkBefore(InsertPoint);
 			}
-#else
-			LargePage->LinkHead(UsedLargePagesWithSpaceHead[LargePage->AllocationHint]);
-#endif
-		}
-		else
-		{
-#if !LARGEPAGEALLOCATOR_SORT_OnAddress
-			FLargePage* InsertPoint = LargePage->Next();
-			FLargePage* LastInsertPoint = nullptr;
-
-			if ((InsertPoint != nullptr) && LargePage->NumberOfFreeSubPages > InsertPoint->NumberOfFreeSubPages)
-			{
-				LastInsertPoint = InsertPoint;
-				LargePage->Unlink();
-				while (InsertPoint != nullptr)
-				{
-					if (LargePage->NumberOfFreeSubPages <= InsertPoint->NumberOfFreeSubPages)	// sort on number of free sub pages
-					{
-						break;
-					}
-					InsertPoint = InsertPoint->Next();
-				}
-				if (InsertPoint != nullptr)
-				{
-					LargePage->LinkBefore(InsertPoint);
-				}
-				else
-				{
-					LargePage->LinkAfter(LastInsertPoint);
-				}
-			}
-#endif
 		}
 	}
 	else
@@ -397,7 +373,7 @@ void FCachedOSVeryLargePageAllocator::ShrinkEmptyBackStore(int32 NewEmptyBackSto
 		LargePage->LinkHead(FreeLargePagesHead[LargePage->AllocationHint]);
 		CachedFree -= SizeOfLargePage;
 		EmptyBackStoreCount[AllocationHint] -= 1;
-		CommitedLargePagesCount[LargePage->AllocationHint] -= 1;
+		CommittedLargePagesCount[LargePage->AllocationHint] -= 1;
 	}
 }
 
@@ -410,10 +386,10 @@ void FCachedOSVeryLargePageAllocator::FreeAll(FCriticalSection* Mutex)
 		ShrinkEmptyBackStore(0, AllocationHint);
 	}
 
-	// Stop preallocating system since allocator reached a OOM
-	if (!GVeryLargePageAllocatorLeavePageCachingEnabledWhenOOMHappened)
+	if (GDisablePageCachingOnOOM)
 	{
-		GVeryLargePageAllocatorPreAllocatePools = false;
+		// Stop preallocating system since allocator reached a OOM
+		GPreAllocatePools = false;
 	}
 
 	// Free empty cached pages of CachedOSPageAllocator
@@ -427,8 +403,8 @@ void FCachedOSVeryLargePageAllocator::UpdateStats()
 	CSV_CUSTOM_STAT(FMemory, LargeAllocatorDecommitCount, GLargePageAllocatorDecommitCount, ECsvCustomStatOp::Set);
 	CSV_CUSTOM_STAT(FMemory, LargeAllocatorBackingStoreCountSmall, EmptyBackStoreCount[FMemory::AllocationHints::SmallPool], ECsvCustomStatOp::Set);
 	CSV_CUSTOM_STAT(FMemory, LargeAllocatorBackingStoreCountDefault, EmptyBackStoreCount[FMemory::AllocationHints::Default], ECsvCustomStatOp::Set);
-	CSV_CUSTOM_STAT(FMemory, LargeAllocatorPageCountSmall, CommitedLargePagesCount[FMemory::AllocationHints::SmallPool], ECsvCustomStatOp::Set);
-	CSV_CUSTOM_STAT(FMemory, LargeAllocatorPageCountDefault, CommitedLargePagesCount[FMemory::AllocationHints::Default], ECsvCustomStatOp::Set);
+	CSV_CUSTOM_STAT(FMemory, LargeAllocatorPageCountSmall, CommittedLargePagesCount[FMemory::AllocationHints::SmallPool], ECsvCustomStatOp::Set);
+	CSV_CUSTOM_STAT(FMemory, LargeAllocatorPageCountDefault, CommittedLargePagesCount[FMemory::AllocationHints::Default], ECsvCustomStatOp::Set);
 
 	GLargePageAllocatorCommitCount = 0;
 	GLargePageAllocatorDecommitCount = 0;
diff --git a/Engine/Source/Runtime/Core/Public/HAL/Allocators/CachedOSVeryLargePageAllocator.h b/Engine/Source/Runtime/Core/Public/HAL/Allocators/CachedOSVeryLargePageAllocator.h
--- a/Engine/Source/Runtime/Core/Public/HAL/Allocators/CachedOSVeryLargePageAllocator.h
+++ b/Engine/Source/Runtime/Core/Public/HAL/Allocators/CachedOSVeryLargePageAllocator.h
@@ -9,10 +9,6 @@
 #include "HAL/UnrealMemory.h"
 
 
-#ifndef UE_USE_VERYLARGEPAGEALLOCATOR_FALLBACKPATH
-#define UE_USE_VERYLARGEPAGEALLOCATOR_FALLBACKPATH UE_USE_VERYLARGEPAGEALLOCATOR
-#endif // UE_USE_VERYLARGEPAGEALLOCATOR_FALLBACKPATH
-
 #if UE_USE_VERYLARGEPAGEALLOCATOR
 
 #if PLATFORM_64BITS
@@ -23,29 +19,19 @@
 //#define CACHEDOSVERYLARGEPAGEALLOCATOR_MAX_CACHED_OS_FREES (CACHEDOSVERYLARGEPAGEALLOCATOR_BYTE_LIMIT/(1024*64))
 #define CACHEDOSVERYLARGEPAGEALLOCATOR_MAX_CACHED_OS_FREES (256)
 
-
-#ifndef UE_VERYLARGEPAGEALLOCATOR_TAKEONALL64KBALLOCATIONS
-#define UE_VERYLARGEPAGEALLOCATOR_TAKEONALL64KBALLOCATIONS 1
-#endif
-
 #ifndef UE_VERYLARGEPAGEALLOCATOR_RESERVED_SIZE_IN_GB
-#define UE_VERYLARGEPAGEALLOCATOR_RESERVED_SIZE_IN_GB 2	//default to 2GB
+#define UE_VERYLARGEPAGEALLOCATOR_RESERVED_SIZE_IN_GB 4	//default to 4GB
 #endif
 
-
 #ifndef UE_VERYLARGEPAGEALLOCATOR_PAGESIZE_KB
 #define UE_VERYLARGEPAGEALLOCATOR_PAGESIZE_KB 2048	//default to 2MB
 #endif
+
 class FCachedOSVeryLargePageAllocator
 {
-	// we make the address space twice as big as we need and use the 1st have for small pool allocations, the 2nd half is used for other allocations that are still == SizeOfSubPage
-#if UE_VERYLARGEPAGEALLOCATOR_TAKEONALL64KBALLOCATIONS
-	static constexpr uint64 AddressSpaceToReserve = ((1024LL * 1024LL * 1024LL) * UE_VERYLARGEPAGEALLOCATOR_RESERVED_SIZE_IN_GB * 2LL);
+	static constexpr uint64 AddressSpaceToReserve = ((1024LL * 1024LL * 1024LL) * UE_VERYLARGEPAGEALLOCATOR_RESERVED_SIZE_IN_GB);
 	static constexpr uint64 AddressSpaceToReserveForSmallPool = AddressSpaceToReserve / 2;
-#else
-	static constexpr uint64 AddressSpaceToReserve = ((1024 * 1024 * 1024LL) * UE_VERYLARGEPAGEALLOCATOR_RESERVED_SIZE_IN_GB);
-	static constexpr uint64 AddressSpaceToReserveForSmallPool = AddressSpaceToReserve;
-#endif
+
 	static constexpr uint64 SizeOfLargePage = (UE_VERYLARGEPAGEALLOCATOR_PAGESIZE_KB * 1024);
 	static constexpr uint64 SizeOfSubPage = (1024 * 64);
 	static constexpr uint64 NumberOfLargePages = (AddressSpaceToReserve / SizeOfLargePage);
@@ -80,7 +66,7 @@ public:
 		return CachedFree + CachedOSPageAllocator.GetCachedFreeTotal();
 	}
 
-	FORCEINLINE bool IsPartOf(const void* Ptr)
+	FORCEINLINE bool IsSmallBlockAllocation(const void* Ptr)
 	{
 		if (((uintptr_t)Ptr - AddressSpaceReserved) < AddressSpaceToReserveForSmallPool)
 		{
@@ -90,20 +76,6 @@ public:
 	}
 
 private:
-
-	void Init();
-	void ShrinkEmptyBackStore(int32 NewEmptyBackStoreSize, FMemory::AllocationHints AllocationHint);
-
-	bool bEnabled;
-	uintptr_t	AddressSpaceReserved;
-	uintptr_t	AddressSpaceReservedEndSmallPool;
-	uintptr_t	AddressSpaceReservedEnd;
-	uint64		CachedFree;
-	int32		EmptyBackStoreCount[FMemory::AllocationHints::Max];
-	int32		CommitedLargePagesCount[FMemory::AllocationHints::Max];
-
-	FPlatformMemory::FPlatformVirtualMemoryBlock Block;
-
 	struct FLargePage : public TIntrusiveLinkedList<FLargePage>
 	{
 		uintptr_t	FreeSubPages[NumberOfSubPagesPerLargePage];
@@ -140,6 +112,21 @@ private:
 		}
 	};
 
+	void Init();
+	void ShrinkEmptyBackStore(int32 NewEmptyBackStoreSize, FMemory::AllocationHints AllocationHint);
+	FLargePage* GetOrAllocLargePage(uint32 AllocationHint, FCriticalSection* Mutex);
+	FLargePage* AllocNewLargePage(uint32 AllocationHint, FCriticalSection* Mutex, bool& bOutCommitFailure);
+
+	bool bEnabled;
+	uintptr_t	AddressSpaceReserved;
+	uintptr_t	AddressSpaceReservedEndSmallPool;
+	uintptr_t	AddressSpaceReservedEnd;
+	uint64		CachedFree;
+	int32		EmptyBackStoreCount[FMemory::AllocationHints::Max];
+	int32		CommittedLargePagesCount[FMemory::AllocationHints::Max];
+
+	FPlatformMemory::FPlatformVirtualMemoryBlock Block;
+
 	FLargePage* FreeLargePagesHead[FMemory::AllocationHints::Max];				// no backing store
 
 	FLargePage* UsedLargePagesHead[FMemory::AllocationHints::Max];				// has backing store and is full
diff --git a/Engine/Source/Runtime/Core/Public/HAL/MallocBinned2.h b/Engine/Source/Runtime/Core/Public/HAL/MallocBinned2.h
--- a/Engine/Source/Runtime/Core/Public/HAL/MallocBinned2.h
+++ b/Engine/Source/Runtime/Core/Public/HAL/MallocBinned2.h
@@ -225,7 +225,7 @@ class FMallocBinned2 : public TMallocBinnedCommon<FMallocBinned2, BINNED2_MINIMU
 	FORCEINLINE bool IsOSAllocation(const void* Ptr)
 	{
 #if UE_USE_VERYLARGEPAGEALLOCATOR && !PLATFORM_UNIX && !PLATFORM_ANDROID
-		return !CachedOSPageAllocator.IsPartOf(Ptr) && IsAligned(Ptr, BINNED2_LARGE_ALLOC);
+		return !CachedOSPageAllocator.IsSmallBlockAllocation(Ptr) && IsAligned(Ptr, BINNED2_LARGE_ALLOC);
 #else
 		return IsAligned(Ptr, BINNED2_LARGE_ALLOC);
 #endif
