Added new UE_ preprocessor metaprogramming macros to replace the PREPROCESSOR_ macros, to aid discoverability.

diff --git a/Engine/Plugins/Runtime/GameFeatures/Source/GameFeatures/Public/GameFeaturesSubsystem.h b/Engine/Plugins/Runtime/GameFeatures/Source/GameFeatures/Public/GameFeaturesSubsystem.h
--- a/Engine/Plugins/Runtime/GameFeatures/Source/GameFeatures/Public/GameFeaturesSubsystem.h
+++ b/Engine/Plugins/Runtime/GameFeatures/Source/GameFeatures/Public/GameFeaturesSubsystem.h
@@ -203,7 +203,7 @@ using FGameFeaturePluginUninstallComplete = FGameFeaturePluginChangeStateComplet
 using FGameFeaturePluginTerminateComplete = FGameFeaturePluginChangeStateComplete;
 using FGameFeaturePluginUpdateProtocolComplete = FGameFeaturePluginChangeStateComplete;
 
-DECLARE_DELEGATE_OneParam(FMultipleGameFeaturePluginChangeStateComplete, PREPROCESSOR_COMMA_SEPARATED(const TMap<FString, UE::GameFeatures::FResult>& /*Results*/));
+using FMultipleGameFeaturePluginChangeStateComplete = TDelegate<void(const TMap<FString, UE::GameFeatures::FResult>& Results)>;
 
 using FBuiltInGameFeaturePluginsLoaded = FMultipleGameFeaturePluginChangeStateComplete;
 using FMultipleGameFeaturePluginsLoaded = FMultipleGameFeaturePluginChangeStateComplete;
diff --git a/Engine/Source/Runtime/Core/Public/HAL/PreprocessorHelpers.h b/Engine/Source/Runtime/Core/Public/HAL/PreprocessorHelpers.h
--- a/Engine/Source/Runtime/Core/Public/HAL/PreprocessorHelpers.h
+++ b/Engine/Source/Runtime/Core/Public/HAL/PreprocessorHelpers.h
@@ -3,63 +3,61 @@
 #pragma once
 
 // Turns an preprocessor token into a real string (see UBT_COMPILED_PLATFORM)
-#define PREPROCESSOR_TO_STRING(x) PREPROCESSOR_TO_STRING_INNER(x)
-#define PREPROCESSOR_TO_STRING_INNER(x) #x
+#define UE_STRINGIZE(Token) UE_PRIVATE_STRINGIZE(Token)
+#define UE_PRIVATE_STRINGIZE(Token) #Token
 
 // Concatenates two preprocessor tokens, performing macro expansion on them first
-#define PREPROCESSOR_JOIN(x, y) PREPROCESSOR_JOIN_INNER(x, y)
-#define PREPROCESSOR_JOIN_INNER(x, y) x##y
+#define UE_JOIN(TokenA, TokenB) UE_PRIVATE_JOIN(TokenA, TokenB)
+#define UE_PRIVATE_JOIN(TokenA, TokenB) TokenA##TokenB
 
 // Concatenates the first two preprocessor tokens of a variadic list, after performing macro expansion on them
-#define PREPROCESSOR_JOIN_FIRST(x, ...) PREPROCESSOR_JOIN_FIRST_INNER(x, __VA_ARGS__)
-#define PREPROCESSOR_JOIN_FIRST_INNER(x, ...) x##__VA_ARGS__
+#define UE_JOIN_FIRST(Token, ...) UE_PRIVATE_JOIN_FIRST(Token, __VA_ARGS__)
+#define UE_PRIVATE_JOIN_FIRST(Token, ...) Token##__VA_ARGS__
 
 // Expands to the second argument or the third argument if the first argument is 1 or 0 respectively
-#define PREPROCESSOR_IF(cond, x, y) PREPROCESSOR_JOIN(PREPROCESSOR_IF_INNER_, cond)(x, y)
-#define PREPROCESSOR_IF_INNER_1(x, y) x
-#define PREPROCESSOR_IF_INNER_0(x, y) y
+#define UE_IF(OneOrZero, Token1, Token0) UE_JOIN(UE_PRIVATE_IF_, OneOrZero)(Token1, Token0)
+#define UE_PRIVATE_IF_1(Token1, Token0) Token1
+#define UE_PRIVATE_IF_0(Token1, Token0) Token0
 
 // Expands to the parameter list of the macro - used to pass a *potentially* comma-separated identifier to another macro as a single parameter
-#define PREPROCESSOR_COMMA_SEPARATED(first, ...) first, ##__VA_ARGS__
+#define UE_COMMA_SEPARATED(First, ...) First, ##__VA_ARGS__
 
 // Expands to a number which is the count of variadic arguments passed to it.
-#define PREPROCESSOR_VA_ARG_COUNT(...) PREPROCESSOR_APPEND_VA_ARG_COUNT(, ##__VA_ARGS__)
+#define UE_VA_ARG_COUNT(...) UE_APPEND_VA_ARG_COUNT(, ##__VA_ARGS__)
 
 // Expands to a token of Prefix##<count>, where <count> is the number of variadic arguments.
 //
 // Example:
-//   PREPROCESSOR_APPEND_VA_ARG_COUNT(SOME_MACRO_)          => SOME_MACRO_0
-//   PREPROCESSOR_APPEND_VA_ARG_COUNT(SOME_MACRO_, a, b, c) => SOME_MACRO_3
+//   UE_APPEND_VA_ARG_COUNT(SOME_MACRO_)          => SOME_MACRO_0
+//   UE_APPEND_VA_ARG_COUNT(SOME_MACRO_, a, b, c) => SOME_MACRO_3
 #if !defined(_MSVC_TRADITIONAL) || !_MSVC_TRADITIONAL
-	#define PREPROCESSOR_APPEND_VA_ARG_COUNT(Prefix, ...) PREPROCESSOR_APPEND_VA_ARG_COUNT_INNER(Prefix, ##__VA_ARGS__, 26, 25, 24, 23, 22, 21, 20, 19, 18, 17, 16, 15, 14, 13, 12, 11, 10, 9, 8, 7, 6, 5, 4, 3, 2, 1, 0)
+	#define UE_APPEND_VA_ARG_COUNT(Prefix, ...) UE_PRIVATE_APPEND_VA_ARG_COUNT(Prefix, ##__VA_ARGS__, 26, 25, 24, 23, 22, 21, 20, 19, 18, 17, 16, 15, 14, 13, 12, 11, 10, 9, 8, 7, 6, 5, 4, 3, 2, 1, 0)
 #else
-	#define PREPROCESSOR_APPEND_VA_ARG_COUNT(Prefix, ...) PREPROCESSOR_APPEND_VA_ARG_COUNT_INVOKE(PREPROCESSOR_APPEND_VA_ARG_COUNT_INNER, (Prefix, ##__VA_ARGS__, 26, 25, 24, 23, 22, 21, 20, 19, 18, 17, 16, 15, 14, 13, 12, 11, 10, 9, 8, 7, 6, 5, 4, 3, 2, 1, 0))
+	#define UE_APPEND_VA_ARG_COUNT(Prefix, ...) UE_PRIVATE_APPEND_VA_ARG_COUNT_INVOKE(UE_PRIVATE_APPEND_VA_ARG_COUNT, (Prefix, ##__VA_ARGS__, 26, 25, 24, 23, 22, 21, 20, 19, 18, 17, 16, 15, 14, 13, 12, 11, 10, 9, 8, 7, 6, 5, 4, 3, 2, 1, 0))
 
 	// MSVC's traditional preprocessor doesn't handle the zero-argument case correctly, so we use a workaround.
 	// The workaround uses token pasting of Macro##ArgsInParens, which the conformant preprocessor doesn't like and emits C5103.
-	#define PREPROCESSOR_APPEND_VA_ARG_COUNT_INVOKE(Macro, ArgsInParens) PREPROCESSOR_APPEND_VA_ARG_COUNT_EXPAND(Macro##ArgsInParens)
-	#define PREPROCESSOR_APPEND_VA_ARG_COUNT_EXPAND(Arg) Arg
+	#define UE_PRIVATE_APPEND_VA_ARG_COUNT_INVOKE(Macro, ArgsInParens) UE_PRIVATE_APPEND_VA_ARG_COUNT_EXPAND(Macro##ArgsInParens)
+	#define UE_PRIVATE_APPEND_VA_ARG_COUNT_EXPAND(Arg) Arg
 #endif
-#define PREPROCESSOR_APPEND_VA_ARG_COUNT_INNER(Prefix,a,b,c,d,e,f,g,h,i,j,k,l,m,n,o,p,q,r,s,t,u,v,w,x,y,z,count,...) Prefix##count
+#define UE_PRIVATE_APPEND_VA_ARG_COUNT(Prefix,A,B,C,D,E,F,G,H,I,J,K,L,M,N,O,P,Q,R,S,T,U,V,W,X,Y,Z,Count,...) Prefix##Count
 
 // Expands to nothing - used as a placeholder
-#define PREPROCESSOR_NOTHING
+#define UE_EMPTY
 
 // Expands to nothing when used as a function - used as a placeholder
-#define PREPROCESSOR_NOTHING_FUNCTION(...)
-
-#define UE_SOURCE_LOCATION TEXT(__FILE__ "(" PREPROCESSOR_TO_STRING(__LINE__) ")")
+#define UE_EMPTY_FUNCTION(...)
 
 // Removes a single layer of parentheses from a macro argument if they are present - used to allow
 // brackets to be optionally added when the argument contains commas, e.g.:
 //
-// #define DEFINE_VARIABLE(Type, Name) PREPROCESSOR_REMOVE_OPTIONAL_PARENS(Type) Name;
+// #define DEFINE_VARIABLE(Type, Name) UE_REMOVE_OPTIONAL_PARENS(Type) Name;
 //
 // DEFINE_VARIABLE(int, IntVar)                  // expands to: int IntVar;
 // DEFINE_VARIABLE((TPair<int, float>), PairVar) // expands to: TPair<int, float> PairVar;
-#define PREPROCESSOR_REMOVE_OPTIONAL_PARENS(...) PREPROCESSOR_JOIN_FIRST(PREPROCESSOR_REMOVE_OPTIONAL_PARENS_IMPL,PREPROCESSOR_REMOVE_OPTIONAL_PARENS_IMPL __VA_ARGS__)
-#define PREPROCESSOR_REMOVE_OPTIONAL_PARENS_IMPL(...) PREPROCESSOR_REMOVE_OPTIONAL_PARENS_IMPL __VA_ARGS__
-#define PREPROCESSOR_REMOVE_OPTIONAL_PARENS_IMPLPREPROCESSOR_REMOVE_OPTIONAL_PARENS_IMPL
+#define UE_REMOVE_OPTIONAL_PARENS(...) UE_JOIN_FIRST(UE_PRIVATE_PREPROCESSOR_REMOVE_OPTIONAL_PARENS,UE_PRIVATE_PREPROCESSOR_REMOVE_OPTIONAL_PARENS __VA_ARGS__)
+#define UE_PRIVATE_PREPROCESSOR_REMOVE_OPTIONAL_PARENS(...) UE_PRIVATE_PREPROCESSOR_REMOVE_OPTIONAL_PARENS __VA_ARGS__
+#define UE_PRIVATE_PREPROCESSOR_REMOVE_OPTIONAL_PARENSUE_PRIVATE_PREPROCESSOR_REMOVE_OPTIONAL_PARENS
 
 // setup standardized way of including platform headers from the "uber-platform" headers like PlatformFile.h
 #ifdef OVERRIDE_PLATFORM_HEADER_NAME
@@ -70,6 +68,8 @@
 #define PLATFORM_HEADER_NAME UBT_COMPILED_PLATFORM
 #endif
 
+#define UE_SOURCE_LOCATION TEXT(__FILE__ "(" UE_STRINGIZE(__LINE__) ")")
+
 #ifndef PLATFORM_IS_EXTENSION
 #define PLATFORM_IS_EXTENSION 0
 #endif
@@ -89,3 +89,15 @@
 // Creates a string that can be used to include a header with the platform in its name, like "Pre/Fix/PlatformName/PlatformNameSuffix.h"
 #define COMPILED_PLATFORM_HEADER_WITH_PREFIX(Prefix, Suffix) PREPROCESSOR_TO_STRING(Prefix/PLATFORM_HEADER_NAME/PREPROCESSOR_JOIN(PLATFORM_HEADER_NAME, Suffix))
 #endif
+
+// These macros should be regarded as deprecated - use the UE_ macros they map to instead.
+#define PREPROCESSOR_TO_STRING(Token)                 UE_STRINGIZE(Token)
+#define PREPROCESSOR_JOIN(TokenA, TokenB)             UE_JOIN(TokenA, TokenB)
+#define PREPROCESSOR_JOIN_FIRST(Token, ...)           UE_JOIN_FIRST(Token, ##__VA_ARGS__)
+#define PREPROCESSOR_IF(OneOrZero, Token1, Token0)    UE_IF(OneOrZero, Token1, Token0)
+#define PREPROCESSOR_COMMA_SEPARATED(First, ...)      UE_COMMA_SEPARATED(First, ##__VA_ARGS__)
+#define PREPROCESSOR_VA_ARG_COUNT(...)                UE_VA_ARG_COUNT(__VA_ARGS__)
+#define PREPROCESSOR_APPEND_VA_ARG_COUNT(Prefix, ...) UE_APPEND_VA_ARG_COUNT(Prefix, ##__VA_ARGS__)
+#define PREPROCESSOR_NOTHING                          UE_EMPTY
+#define PREPROCESSOR_NOTHING_FUNCTION(...)            UE_EMPTY_FUNCTION(__VA_ARGS__)
+#define PREPROCESSOR_REMOVE_OPTIONAL_PARENS(...)      UE_REMOVE_OPTIONAL_PARENS(__VA_ARGS__)
diff --git a/Engine/Source/Runtime/Core/Tests/HAL/PreprocessorHelperTests.cpp b/Engine/Source/Runtime/Core/Tests/HAL/PreprocessorHelperTests.cpp
--- a/Engine/Source/Runtime/Core/Tests/HAL/PreprocessorHelperTests.cpp
+++ b/Engine/Source/Runtime/Core/Tests/HAL/PreprocessorHelperTests.cpp
@@ -6,22 +6,22 @@
 
 #if WITH_TESTS
 
-TEST_CASE_NAMED(FPreprocessorHelperArgCountTest, "System::Core::HAL::PREPROCESSOR_VA_ARG_COUNT", "[ApplicationContextMask][EngineFilter]")
+TEST_CASE_NAMED(FPreprocessorHelperArgCountTest, "System::Core::HAL::UE_VA_ARG_COUNT", "[ApplicationContextMask][EngineFilter]")
 {
-	STATIC_CHECK(PREPROCESSOR_VA_ARG_COUNT() == 0);
-	STATIC_CHECK(PREPROCESSOR_VA_ARG_COUNT(a) == 1);
-	STATIC_CHECK(PREPROCESSOR_VA_ARG_COUNT(a, b) == 2);
-	STATIC_CHECK(PREPROCESSOR_VA_ARG_COUNT(a, b, c) == 3);
-	STATIC_CHECK(PREPROCESSOR_VA_ARG_COUNT(a, b, c, d) == 4);
-	STATIC_CHECK(PREPROCESSOR_VA_ARG_COUNT(a, b, c, d, e) == 5);
-	STATIC_CHECK(PREPROCESSOR_VA_ARG_COUNT(a, b, c, d, e, f) == 6);
-	STATIC_CHECK(PREPROCESSOR_VA_ARG_COUNT(a, b, c, d, e, f, g) == 7);
-	STATIC_CHECK(PREPROCESSOR_VA_ARG_COUNT(a, b, c, d, e, f, g, h) == 8);
-	STATIC_CHECK(PREPROCESSOR_VA_ARG_COUNT(a, b, c, d, e, f, g, h, i) == 9);
-	STATIC_CHECK(PREPROCESSOR_VA_ARG_COUNT(a, b, c, d, e, f, g, h, i, j) == 10);
+	STATIC_CHECK(UE_VA_ARG_COUNT() == 0);
+	STATIC_CHECK(UE_VA_ARG_COUNT(a) == 1);
+	STATIC_CHECK(UE_VA_ARG_COUNT(a, b) == 2);
+	STATIC_CHECK(UE_VA_ARG_COUNT(a, b, c) == 3);
+	STATIC_CHECK(UE_VA_ARG_COUNT(a, b, c, d) == 4);
+	STATIC_CHECK(UE_VA_ARG_COUNT(a, b, c, d, e) == 5);
+	STATIC_CHECK(UE_VA_ARG_COUNT(a, b, c, d, e, f) == 6);
+	STATIC_CHECK(UE_VA_ARG_COUNT(a, b, c, d, e, f, g) == 7);
+	STATIC_CHECK(UE_VA_ARG_COUNT(a, b, c, d, e, f, g, h) == 8);
+	STATIC_CHECK(UE_VA_ARG_COUNT(a, b, c, d, e, f, g, h, i) == 9);
+	STATIC_CHECK(UE_VA_ARG_COUNT(a, b, c, d, e, f, g, h, i, j) == 10);
 }
 
-TEST_CASE_NAMED(FPreprocessorHelperArgAppendCountTest, "System::Core::HAL::PREPROCESSOR_APPEND_VA_ARG_COUNT", "[ApplicationContextMask][EngineFilter]")
+TEST_CASE_NAMED(FPreprocessorHelperArgAppendCountTest, "System::Core::HAL::UE_APPEND_VA_ARG_COUNT", "[ApplicationContextMask][EngineFilter]")
 {
 	#define TESTPREFIX_0 22
 	#define TESTPREFIX_1 32
@@ -35,17 +35,17 @@ TEST_CASE_NAMED(FPreprocessorHelperArgAppendCountTest, "System::Core::HAL::PREPR
 	#define TESTPREFIX_9 90
 	#define TESTPREFIX_10 96
 
-	STATIC_CHECK(PREPROCESSOR_APPEND_VA_ARG_COUNT(TESTPREFIX_) == TESTPREFIX_0);
-	STATIC_CHECK(PREPROCESSOR_APPEND_VA_ARG_COUNT(TESTPREFIX_, a) == TESTPREFIX_1);
-	STATIC_CHECK(PREPROCESSOR_APPEND_VA_ARG_COUNT(TESTPREFIX_, a, b) == TESTPREFIX_2);
-	STATIC_CHECK(PREPROCESSOR_APPEND_VA_ARG_COUNT(TESTPREFIX_, a, b, c) == TESTPREFIX_3);
-	STATIC_CHECK(PREPROCESSOR_APPEND_VA_ARG_COUNT(TESTPREFIX_, a, b, c, d) == TESTPREFIX_4);
-	STATIC_CHECK(PREPROCESSOR_APPEND_VA_ARG_COUNT(TESTPREFIX_, a, b, c, d, e) == TESTPREFIX_5);
-	STATIC_CHECK(PREPROCESSOR_APPEND_VA_ARG_COUNT(TESTPREFIX_, a, b, c, d, e, f) == TESTPREFIX_6);
-	STATIC_CHECK(PREPROCESSOR_APPEND_VA_ARG_COUNT(TESTPREFIX_, a, b, c, d, e, f, g) == TESTPREFIX_7);
-	STATIC_CHECK(PREPROCESSOR_APPEND_VA_ARG_COUNT(TESTPREFIX_, a, b, c, d, e, f, g, h) == TESTPREFIX_8);
-	STATIC_CHECK(PREPROCESSOR_APPEND_VA_ARG_COUNT(TESTPREFIX_, a, b, c, d, e, f, g, h, i) == TESTPREFIX_9);
-	STATIC_CHECK(PREPROCESSOR_APPEND_VA_ARG_COUNT(TESTPREFIX_, a, b, c, d, e, f, g, h, i, j) == TESTPREFIX_10);
+	STATIC_CHECK(UE_APPEND_VA_ARG_COUNT(TESTPREFIX_) == TESTPREFIX_0);
+	STATIC_CHECK(UE_APPEND_VA_ARG_COUNT(TESTPREFIX_, a) == TESTPREFIX_1);
+	STATIC_CHECK(UE_APPEND_VA_ARG_COUNT(TESTPREFIX_, a, b) == TESTPREFIX_2);
+	STATIC_CHECK(UE_APPEND_VA_ARG_COUNT(TESTPREFIX_, a, b, c) == TESTPREFIX_3);
+	STATIC_CHECK(UE_APPEND_VA_ARG_COUNT(TESTPREFIX_, a, b, c, d) == TESTPREFIX_4);
+	STATIC_CHECK(UE_APPEND_VA_ARG_COUNT(TESTPREFIX_, a, b, c, d, e) == TESTPREFIX_5);
+	STATIC_CHECK(UE_APPEND_VA_ARG_COUNT(TESTPREFIX_, a, b, c, d, e, f) == TESTPREFIX_6);
+	STATIC_CHECK(UE_APPEND_VA_ARG_COUNT(TESTPREFIX_, a, b, c, d, e, f, g) == TESTPREFIX_7);
+	STATIC_CHECK(UE_APPEND_VA_ARG_COUNT(TESTPREFIX_, a, b, c, d, e, f, g, h) == TESTPREFIX_8);
+	STATIC_CHECK(UE_APPEND_VA_ARG_COUNT(TESTPREFIX_, a, b, c, d, e, f, g, h, i) == TESTPREFIX_9);
+	STATIC_CHECK(UE_APPEND_VA_ARG_COUNT(TESTPREFIX_, a, b, c, d, e, f, g, h, i, j) == TESTPREFIX_10);
 }
 
 #endif //WITH_TESTS
