Fix for TestFrameworkTests due to UTF8 work.
Added null checking to the TestNotEquals for TCHAR* pairs.

UTF8 work added TestNotEquals for String views and TCHAR* for completeness.
This was then used instead of the templated TestNotEquals [REDACTED_HANDLE] in AutomationTest.h
The templated version would have only done a pointer comparison, and would not have actually checked the values.
But it did handle the `null` pointer, which is a case from the AutomatedTestFrameworkTests.

diff --git a/Engine/Source/Runtime/Core/Private/Misc/AutomationTest.cpp b/Engine/Source/Runtime/Core/Private/Misc/AutomationTest.cpp
--- a/Engine/Source/Runtime/Core/Private/Misc/AutomationTest.cpp
+++ b/Engine/Source/Runtime/Core/Private/Misc/AutomationTest.cpp
@@ -1800,10 +1800,21 @@ bool FAutomationTestBase::TestEqual(const TCHAR* What, FStringView Actual, FStri
 
 bool FAutomationTestBase::TestNotEqual(const TCHAR* What, const TCHAR* Actual, const TCHAR* Expected)
 {
-	if (FCString::Strcmp(Actual, Expected) == 0)
+	if (Actual && Expected)
 	{
-		AddError(FString::Printf(TEXT("Expected '%s' not to be \"%s\", but it was \"%s\"."), What, Expected, Actual), 1);
-		return false;
+		if (FCString::Strcmp(Actual, Expected) == 0)
+		{
+			AddError(FString::Printf(TEXT("Expected '%s' not to be\"%s\", but it was \"%s\"."), What, Expected, Actual), 1);
+			return false;
+		}
+	}
+	else // null exists
+	{
+		if (Actual == Expected)
+		{
+			AddError(FString::Printf(TEXT("Expected '%s' to differ but both values were unexpectedly null"), What), 1);
+			return false;
+		}
 	}
 	return true;
 }
