Updated more StringTable API to use FTextKey for its namespace and keys

diff --git a/Engine/Source/Runtime/Core/Private/Internationalization/StringTableCore.cpp b/Engine/Source/Runtime/Core/Private/Internationalization/StringTableCore.cpp
--- a/Engine/Source/Runtime/Core/Private/Internationalization/StringTableCore.cpp
+++ b/Engine/Source/Runtime/Core/Private/Internationalization/StringTableCore.cpp
@@ -125,11 +125,11 @@ FString FStringTable::GetNamespace() const
 	return TableNamespace.GetChars();
 }
 
-void FStringTable::SetNamespace(const FString& InNamespace)
+void FStringTable::SetNamespace(const FTextKey& InNamespace)
 {
 	FScopeLock KeyMappingLock(&KeyMappingCS);
 
-	if (FCString::Strcmp(TableNamespace.GetChars(), *InNamespace) != 0)
+	if (TableNamespace != InNamespace)
 	{
 		TableNamespace = InNamespace;
 
diff --git a/Engine/Source/Runtime/Core/Private/Internationalization/StringTableRegistry.cpp b/Engine/Source/Runtime/Core/Private/Internationalization/StringTableRegistry.cpp
--- a/Engine/Source/Runtime/Core/Private/Internationalization/StringTableRegistry.cpp
+++ b/Engine/Source/Runtime/Core/Private/Internationalization/StringTableRegistry.cpp
@@ -183,7 +183,7 @@ void FStringTableRegistry::OnDirectoryChanged(const TArray<FFileChangeData>& InF
 }
 #endif // WITH_EDITOR
 
-void FStringTableRegistry::Internal_NewLocTable(const FName InTableId, const FString& InNamespace)
+void FStringTableRegistry::Internal_NewLocTable(const FName InTableId, const FTextKey& InNamespace)
 {
 	FStringTableRef StringTable = FStringTable::NewStringTable();
 	StringTable->SetNamespace(InNamespace);
@@ -191,7 +191,7 @@ void FStringTableRegistry::Internal_NewLocTable(const FName InTableId, const FSt
 	RegisterStringTable(InTableId, StringTable);
 }
 
-void FStringTableRegistry::Internal_LocTableFromFile(const FName InTableId, const FString& InNamespace, const FString& InFilePath, const FString& InRootPath)
+void FStringTableRegistry::Internal_LocTableFromFile(const FName InTableId, const FTextKey& InNamespace, const FString& InFilePath, const FString& InRootPath)
 {
 	FStringTableRef StringTable = FStringTable::NewStringTable();
 	StringTable->SetNamespace(InNamespace);
@@ -212,7 +212,7 @@ void FStringTableRegistry::Internal_LocTableFromFile(const FName InTableId, cons
 #endif // WITH_EDITOR
 }
 
-void FStringTableRegistry::Internal_SetLocTableEntry(const FName InTableId, const FString& InKey, const FString& InSourceString)
+void FStringTableRegistry::Internal_SetLocTableEntry(const FName InTableId, const FTextKey& InKey, const FString& InSourceString)
 {
 	FStringTablePtr StringTable = FindMutableStringTable(InTableId);
 	checkf(StringTable.IsValid(), TEXT("Attempting to add a string table entry to the unknown string table '%s'"), *InTableId.ToString());
@@ -220,7 +220,7 @@ void FStringTableRegistry::Internal_SetLocTableEntry(const FName InTableId, cons
 	StringTable->SetSourceString(InKey, InSourceString);
 }
 
-void FStringTableRegistry::Internal_SetLocTableEntryMetaData(const FName InTableId, const FString& InKey, const FName InMetaDataId, const FString& InMetaData)
+void FStringTableRegistry::Internal_SetLocTableEntryMetaData(const FName InTableId, const FTextKey& InKey, const FName InMetaDataId, const FString& InMetaData)
 {
 	FStringTablePtr StringTable = FindMutableStringTable(InTableId);
 	checkf(StringTable.IsValid(), TEXT("Attempting to add string table entry meta-data to the unknown string table '%s'"), *InTableId.ToString());
@@ -228,7 +228,7 @@ void FStringTableRegistry::Internal_SetLocTableEntryMetaData(const FName InTable
 	StringTable->SetMetaData(InKey, InMetaDataId, InMetaData);
 }
 
-FText FStringTableRegistry::Internal_FindLocTableEntry(const FName InTableId, const FString& InKey, const EStringTableLoadingPolicy InLoadingPolicy) const
+FText FStringTableRegistry::Internal_FindLocTableEntry(const FName InTableId, const FTextKey& InKey, const EStringTableLoadingPolicy InLoadingPolicy) const
 {
 	return FText(InTableId, InKey, InLoadingPolicy);
 }
diff --git a/Engine/Source/Runtime/Core/Private/Internationalization/Text.cpp b/Engine/Source/Runtime/Core/Private/Internationalization/Text.cpp
--- a/Engine/Source/Runtime/Core/Private/Internationalization/Text.cpp
+++ b/Engine/Source/Runtime/Core/Private/Internationalization/Text.cpp
@@ -256,8 +256,8 @@ FText::FText( FString&& InSourceString )
 {
 }
 
-FText::FText( FName InTableId, FString InKey, const EStringTableLoadingPolicy InLoadingPolicy )
-	: TextData(MakeRefCount<FTextHistory_StringTableEntry>(InTableId, MoveTemp(InKey), InLoadingPolicy))
+FText::FText( FName InTableId, const FTextKey& InKey, const EStringTableLoadingPolicy InLoadingPolicy )
+	: TextData(MakeRefCount<FTextHistory_StringTableEntry>(InTableId, InKey, InLoadingPolicy))
 	, Flags(0)
 {
 }
@@ -1047,7 +1047,7 @@ FText FText::ChangeKey( const FTextKey& Namespace, const FTextKey& Key, const FT
 }
 #endif
 
-FText FText::FromStringTable(const FName InTableId, const FString& InKey, const EStringTableLoadingPolicy InLoadingPolicy)
+FText FText::FromStringTable(const FName InTableId, const FTextKey& InKey, const EStringTableLoadingPolicy InLoadingPolicy)
 {
 	return FStringTableRegistry::Get().Internal_FindLocTableEntry(InTableId, InKey, InLoadingPolicy);
 }
diff --git a/Engine/Source/Runtime/Core/Private/Internationalization/TextHistory.cpp b/Engine/Source/Runtime/Core/Private/Internationalization/TextHistory.cpp
--- a/Engine/Source/Runtime/Core/Private/Internationalization/TextHistory.cpp
+++ b/Engine/Source/Runtime/Core/Private/Internationalization/TextHistory.cpp
@@ -2330,10 +2330,10 @@ bool FTextHistory_Transform::GetHistoricNumericData(const FText& InText, FHistor
 ///////////////////////////////////////
 // FTextHistory_StringTableEntry
 
-FTextHistory_StringTableEntry::FTextHistory_StringTableEntry(FName InTableId, FString&& InKey, const EStringTableLoadingPolicy InLoadingPolicy)
+FTextHistory_StringTableEntry::FTextHistory_StringTableEntry(FName InTableId, const FTextKey& InKey, const EStringTableLoadingPolicy InLoadingPolicy)
 	: StringTableReferenceData(MakeShared<FStringTableReferenceData, ESPMode::ThreadSafe>())
 {
-	StringTableReferenceData->Initialize(InTableId, MoveTemp(InKey), InLoadingPolicy);
+	StringTableReferenceData->Initialize(InTableId, InKey, InLoadingPolicy);
 	MarkDisplayStringUpToDate();
 }
 
@@ -2508,10 +2508,10 @@ bool FTextHistory_StringTableEntry::WriteToBuffer(FString& Buffer, const bool bS
 	return false;
 }
 
-void FTextHistory_StringTableEntry::FStringTableReferenceData::Initialize(FName InTableId, FTextKey InKey, const EStringTableLoadingPolicy InLoadingPolicy)
+void FTextHistory_StringTableEntry::FStringTableReferenceData::Initialize(FName InTableId, const FTextKey& InKey, const EStringTableLoadingPolicy InLoadingPolicy)
 {
 	TableId = InTableId;
-	Key = MoveTemp(InKey);
+	Key = InKey;
 	FStringTableRedirects::RedirectTableIdAndKey(TableId, Key);
 
 	if (InLoadingPolicy == EStringTableLoadingPolicy::Find)
diff --git a/Engine/Source/Runtime/Core/Private/Internationalization/TextHistory.h b/Engine/Source/Runtime/Core/Private/Internationalization/TextHistory.h
--- a/Engine/Source/Runtime/Core/Private/Internationalization/TextHistory.h
+++ b/Engine/Source/Runtime/Core/Private/Internationalization/TextHistory.h
@@ -570,7 +570,7 @@ class FTextHistory_StringTableEntry : public FTextHistory
 {
 public:
 	FTextHistory_StringTableEntry() = default;
-	FTextHistory_StringTableEntry(FName InTableId, FString&& InKey, const EStringTableLoadingPolicy InLoadingPolicy);
+	FTextHistory_StringTableEntry(FName InTableId, const FTextKey& InKey, const EStringTableLoadingPolicy InLoadingPolicy);
 
 	//~ Begin FTextHistory Interface
 	OVERRIDE_TEXT_HISTORY_STRINGIFICATION;
@@ -603,7 +603,7 @@ private:
 	{
 	public:
 		/** Initialize this data, immediately starting an asset load if required and possible */
-		void Initialize(FName InTableId, FTextKey InKey, const EStringTableLoadingPolicy InLoadingPolicy);
+		void Initialize(FName InTableId, const FTextKey& InKey, const EStringTableLoadingPolicy InLoadingPolicy);
 
 		/** Check whether this instance is considered identical to the other instance */
 		bool IsIdentical(const FStringTableReferenceData& Other) const;
diff --git a/Engine/Source/Runtime/Core/Public/Internationalization/StringTableCore.h b/Engine/Source/Runtime/Core/Public/Internationalization/StringTableCore.h
--- a/Engine/Source/Runtime/Core/Public/Internationalization/StringTableCore.h
+++ b/Engine/Source/Runtime/Core/Public/Internationalization/StringTableCore.h
@@ -108,7 +108,7 @@ public:
 	CORE_API FString GetNamespace() const;
 
 	/** Set the namespace used by all entries in this string table */
-	CORE_API void SetNamespace(const FString& InNamespace);
+	CORE_API void SetNamespace(const FTextKey& InNamespace);
 
 	/** Get the source string used by the given entry (if any) */
 	CORE_API bool GetSourceString(const FTextKey& InKey, FString& OutSourceString) const;
diff --git a/Engine/Source/Runtime/Core/Public/Internationalization/StringTableRegistry.h b/Engine/Source/Runtime/Core/Public/Internationalization/StringTableRegistry.h
--- a/Engine/Source/Runtime/Core/Public/Internationalization/StringTableRegistry.h
+++ b/Engine/Source/Runtime/Core/Public/Internationalization/StringTableRegistry.h
@@ -52,19 +52,19 @@ public:
 	CORE_API void LogMissingStringTableEntry(const FName InTableId, const FTextKey& InKey);
 
 	/** Internal function called by LOCTABLE_NEW to create and register a new FStringTable instance */
-	CORE_API void Internal_NewLocTable(const FName InTableId, const FString& InNamespace);
+	CORE_API void Internal_NewLocTable(const FName InTableId, const FTextKey& InNamespace);
 
 	/** Internal function called by LOCTABLE_FROMFILE_X to create and register a FStringTable instance that has been populated from a file */
-	CORE_API void Internal_LocTableFromFile(const FName InTableId, const FString& InNamespace, const FString& InFilePath, const FString& InRootPath);
+	CORE_API void Internal_LocTableFromFile(const FName InTableId, const FTextKey& InNamespace, const FString& InFilePath, const FString& InRootPath);
 
 	/** Internal function called by LOCTABLE_SETSTRING to set the entry denoted by the given key to the given source string, within the given string table (table must have been registered already) */
-	CORE_API void Internal_SetLocTableEntry(const FName InTableId, const FString& InKey, const FString& InSourceString);
+	CORE_API void Internal_SetLocTableEntry(const FName InTableId, const FTextKey& InKey, const FString& InSourceString);
 
 	/** Internal function called by LOCTABLE_SETMETA to set meta-data for the entry denoted by the given key, within the given string table (table must have been registered already) */
-	CORE_API void Internal_SetLocTableEntryMetaData(const FName InTableId, const FString& InKey, const FName InMetaDataId, const FString& InMetaData);
+	CORE_API void Internal_SetLocTableEntryMetaData(const FName InTableId, const FTextKey& InKey, const FName InMetaDataId, const FString& InMetaData);
 
 	/** Internal function called by LOCTABLE to find the entry with by the given key within the given string table (redirects, will load assets if needed, and returns a dummy FText if not found) */
-	CORE_API FText Internal_FindLocTableEntry(const FName InTableId, const FString& InKey, const EStringTableLoadingPolicy InLoadingPolicy) const;
+	CORE_API FText Internal_FindLocTableEntry(const FName InTableId, const FTextKey& InKey, const EStringTableLoadingPolicy InLoadingPolicy) const;
 
 private:
 	/** Private constructor - use singleton accessor */
diff --git a/Engine/Source/Runtime/Core/Public/Internationalization/Text.h b/Engine/Source/Runtime/Core/Public/Internationalization/Text.h
--- a/Engine/Source/Runtime/Core/Public/Internationalization/Text.h
+++ b/Engine/Source/Runtime/Core/Public/Internationalization/Text.h
@@ -469,7 +469,7 @@ public:
 	 * Attempts to create an FText instance from a string table ID and key (this is the same as the LOCTABLE macro, except this can also work with non-literal string values).
 	 * @return The found text, or a dummy FText if not found.
 	 */
-	static CORE_API FText FromStringTable(const FName InTableId, const FString& InKey, const EStringTableLoadingPolicy InLoadingPolicy = EStringTableLoadingPolicy::FindOrLoad);
+	static CORE_API FText FromStringTable(const FName InTableId, const FTextKey& InKey, const EStringTableLoadingPolicy InLoadingPolicy = EStringTableLoadingPolicy::FindOrLoad);
 
 	/**
 	 * Generate an FText representing the pass name
@@ -799,7 +799,7 @@ private:
 
 	CORE_API explicit FText( FString&& InSourceString );
 
-	CORE_API FText( FName InTableId, FString InKey, const EStringTableLoadingPolicy InLoadingPolicy );
+	CORE_API FText( FName InTableId, const FTextKey& InKey, const EStringTableLoadingPolicy InLoadingPolicy );
 
 	CORE_API FText( FString&& InSourceString, const FTextKey& InNamespace, const FTextKey& InKey, uint32 InFlags=0 );
 
