Stop the double stop (har de har) in FScopeCycleCounter.

#rnx

diff --git a/Engine/Source/Runtime/Core/Public/Stats/Stats.h b/Engine/Source/Runtime/Core/Public/Stats/Stats.h
--- a/Engine/Source/Runtime/Core/Public/Stats/Stats.h
+++ b/Engine/Source/Runtime/Core/Public/Stats/Stats.h
@@ -74,7 +74,6 @@ public:
 	FORCEINLINE_STATS FScopeCycleCounter(TStatId StatId, bool bAlways = false)
 		: FScopeCycleCounter(StatId, EStatFlags::None, bAlways)
 	{
-		AutoRTFM::PushOnAbortHandler(this, [this](){ this->Stop(); });
 	}
 
 	/**
