Fixed potential for integer overflow inside FFrameTime arithmetic operators

diff --git a/Engine/Source/Runtime/Core/Public/Misc/FrameTime.h b/Engine/Source/Runtime/Core/Public/Misc/FrameTime.h
--- a/Engine/Source/Runtime/Core/Public/Misc/FrameTime.h
+++ b/Engine/Source/Runtime/Core/Public/Misc/FrameTime.h
@@ -153,9 +153,10 @@ public:
 
 	friend FORCEINLINE_DEBUGGABLE FFrameTime& operator+=(FFrameTime& LHS, FFrameTime RHS)
 	{
-		float NewSubFrame = LHS.SubFrame + RHS.SubFrame;
+		float       NewSubFrame    = LHS.SubFrame + RHS.SubFrame;
+		const int64 NewFrameNumber = int64(LHS.FrameNumber.Value) + int64(RHS.FrameNumber.Value) + FMath::FloorToInt(NewSubFrame);
 
-		LHS.FrameNumber = LHS.FrameNumber + RHS.FrameNumber + FFrameNumber(FMath::FloorToInt(NewSubFrame));
+		LHS.FrameNumber = static_cast<int32>(NewFrameNumber);
 		LHS.SubFrame    = FMath::Frac(NewSubFrame);
 
 		return LHS;
@@ -165,9 +166,9 @@ public:
 	friend FORCEINLINE_DEBUGGABLE FFrameTime operator+(FFrameTime A, FFrameTime B)
 	{
 		const float        NewSubFrame    = A.SubFrame + B.SubFrame;
-		const FFrameNumber NewFrameNumber = A.FrameNumber + B.FrameNumber + FFrameNumber(FMath::FloorToInt(NewSubFrame));
+		const int64        NewFrameNumber = int64(A.FrameNumber.Value) + int64(B.FrameNumber.Value) + FMath::FloorToInt(NewSubFrame);
 
-		return FFrameTime(NewFrameNumber, FMath::Frac(NewSubFrame));
+		return FFrameTime(static_cast<int32>(NewFrameNumber), FMath::Frac(NewSubFrame));
 	}
 
 
@@ -177,7 +178,9 @@ public:
 		// Note that the difference between frame -1.5 and 1.5 is 2, not 3, since sub frame positions are always positive
 		const float        NewSubFrame     = LHS.SubFrame - RHS.SubFrame;
 		const float        FlooredSubFrame = FMath::FloorToFloat(NewSubFrame);
-		LHS.FrameNumber  = LHS.FrameNumber - RHS.FrameNumber + FFrameNumber(FMath::TruncToInt(FlooredSubFrame));
+		const int64        NewFrameNumber  = int64(LHS.FrameNumber.Value) - int64(RHS.FrameNumber.Value) + FMath::TruncToInt(FlooredSubFrame);
+
+		LHS.FrameNumber.Value  = static_cast<int32>(NewFrameNumber);
 		LHS.SubFrame = NewSubFrame - FlooredSubFrame;
 
 		return LHS;
@@ -190,9 +193,9 @@ public:
 		// Note that the difference between frame -1.5 and 1.5 is 2, not 3, since sub frame positions are always positive
 		const float        NewSubFrame     = A.SubFrame - B.SubFrame;
 		const float        FlooredSubFrame = FMath::FloorToFloat(NewSubFrame);
-		const FFrameNumber NewFrameNumber  = A.FrameNumber - B.FrameNumber + FFrameNumber(FMath::TruncToInt(FlooredSubFrame));
+		const int64        NewFrameNumber  = int64(A.FrameNumber.Value) - int64(B.FrameNumber.Value) + FMath::TruncToInt(FlooredSubFrame);
 
-		return FFrameTime(NewFrameNumber, NewSubFrame - FlooredSubFrame);
+		return FFrameTime(static_cast<int32>(NewFrameNumber), NewSubFrame - FlooredSubFrame);
 	}
 
 
@@ -224,17 +227,17 @@ public:
 	}
 
 
-	friend FORCEINLINE FFrameTime operator*(FFrameTime A, float Scalar)
+	friend FORCEINLINE FFrameTime operator*(FFrameTime A, double Scalar)
 	{
 		return FFrameTime::FromDecimal(A.AsDecimal() * Scalar);
 	}
 
-	friend FORCEINLINE FFrameTime operator*(float Scalar, FFrameTime A)
+	friend FORCEINLINE FFrameTime operator*(double Scalar, FFrameTime A)
 	{
 		return FFrameTime::FromDecimal(A.AsDecimal() * Scalar);
 	}
 
-	friend FORCEINLINE FFrameTime operator/(FFrameTime A, float Scalar)
+	friend FORCEINLINE FFrameTime operator/(FFrameTime A, double Scalar)
 	{
 		return FFrameTime::FromDecimal(A.AsDecimal() / Scalar);
 	}
