[TSTicker]
 - Add support for TUniqueFunction in FTSTicker along with ExecuteFromGameThread

#rnx

diff --git a/Engine/Source/Runtime/Core/Private/Containers/Ticker.cpp b/Engine/Source/Runtime/Core/Private/Containers/Ticker.cpp
--- a/Engine/Source/Runtime/Core/Private/Containers/Ticker.cpp
+++ b/Engine/Source/Runtime/Core/Private/Containers/Ticker.cpp
@@ -17,9 +17,9 @@ FTSTicker::FDelegateHandle FTSTicker::AddTicker(const FTickerDelegate& InDelegat
 	return NewElement;
 }
 
-FTSTicker::FDelegateHandle FTSTicker::AddTicker(const TCHAR* InName, float InDelay, TFunction<bool(float)> Function)
+FTSTicker::FDelegateHandle FTSTicker::AddTicker(const TCHAR* InName, float InDelay, TUniqueFunction<bool(float)>&& InFunction)
 {
-	FElementPtr NewElement{ new FElement{ CurrentTime.load(std::memory_order_relaxed) + InDelay, InDelay, FTickerDelegate::CreateLambda(Function) } };
+	FElementPtr NewElement{ new FElement{ CurrentTime.load(std::memory_order_relaxed) + InDelay, InDelay, MoveTemp(InFunction) } };
 	AddedElements.Enqueue(NewElement);
 	return NewElement;
 }
@@ -150,12 +150,23 @@ FTSTicker::FElement::FElement()
 FTSTicker::FElement::FElement(double InFireTime, float InDelayTime, const FTickerDelegate& InDelegate)
 	: FireTime(InFireTime)
 	, DelayTime(InDelayTime)
-	, Delegate(InDelegate)
+	, Function(
+		[InDelegate](float DeltaTime)
+		{ 
+			return InDelegate.IsBound() && InDelegate.Execute(DeltaTime);
+		}
+	)
+{}
+
+FTSTicker::FElement::FElement(double InFireTime, float InDelayTime, TUniqueFunction<bool(float)>&& InFunction)
+	: FireTime(InFireTime)
+	, DelayTime(InDelayTime)
+	, Function(MoveTemp(InFunction))
 {}
 
 bool FTSTicker::FElement::Fire(float DeltaTime)
 {
-	return Delegate.IsBound() && Delegate.Execute(DeltaTime);
+	return Function(DeltaTime);
 }
 
 FTSTickerObjectBase::FTSTickerObjectBase(float InDelay, FTSTicker& InTicker)
diff --git a/Engine/Source/Runtime/Core/Public/Containers/Ticker.h b/Engine/Source/Runtime/Core/Public/Containers/Ticker.h
--- a/Engine/Source/Runtime/Core/Public/Containers/Ticker.h
+++ b/Engine/Source/Runtime/Core/Public/Containers/Ticker.h
@@ -53,7 +53,7 @@ public:
 	* @param InDelay Delay until next fire; 0 means "next frame"
 	* @param Function Function to execute. Should return true to fire after another InDelay time
 	*/
-	CORE_API FDelegateHandle AddTicker(const TCHAR * InName, float InDelay, TFunction<bool(float)> Function);
+	CORE_API FDelegateHandle AddTicker(const TCHAR* InName, float InDelay, TUniqueFunction<bool(float)>&& InFunction);
 
 	/**
 	 * Removes a previously added ticker delegate.
@@ -93,8 +93,8 @@ private:
 		double FireTime;
 		/** Delay that this delegate was scheduled with. Kept here so that if the delegate returns true, we will reschedule it. **/
 		float DelayTime;
-		/** Delegate to fire **/
-		FTickerDelegate Delegate;
+		/** Function to fire **/
+		TUniqueFunction<bool(float)> Function;
 
 		static constexpr uint64 DefaultState = 0;
 		static constexpr uint64 RemovedState = 1;
@@ -111,6 +111,8 @@ private:
 		CORE_API FElement();
 		/** This is the ctor that the code will generally use. */
 		CORE_API FElement(double InFireTime, float InDelayTime, const FTickerDelegate& InDelegate);
+		/** This is the ctor that the code will generally use. */
+		CORE_API FElement(double InFireTime, float InDelayTime, TUniqueFunction<bool(float)>&& InFunction);
 
 		/** Fire the delegate if it is fireable **/
 		CORE_API bool Fire(float DeltaTime);
