UE_COLD: Declare function unlikely to be called

diff --git a/Engine/Source/Runtime/Core/Private/Misc/AssertionMacros.cpp b/Engine/Source/Runtime/Core/Private/Misc/AssertionMacros.cpp
--- a/Engine/Source/Runtime/Core/Private/Misc/AssertionMacros.cpp
+++ b/Engine/Source/Runtime/Core/Private/Misc/AssertionMacros.cpp
@@ -717,7 +717,7 @@ FORCENOINLINE bool FDebug::OptionallyLogFormattedEnsureMessageReturningFalseImpl
 }
 #endif
 
-FORCENOINLINE void UE_DEBUG_SECTION VARARGS LowLevelFatalErrorHandler(const ANSICHAR* File, int32 Line, const TCHAR* Format, ...)
+void UE_COLD UE_DEBUG_SECTION VARARGS LowLevelFatalErrorHandler(const ANSICHAR* File, int32 Line, const TCHAR* Format, ...)
 {
 	va_list Args;
 	va_start(Args, Format);
@@ -762,7 +762,7 @@ FORCENOINLINE void FDebug::DumpStackTraceToLog(const TCHAR* Heading, const ELogV
 }
 
 #if DO_ENSURE && !USING_CODE_ANALYSIS
-bool UE_DEBUG_SECTION VARARGS CheckVerifyImpl(std::atomic<bool>& bExecuted, bool bAlways, const ANSICHAR* File, int32 Line, void* ProgramCounter, const ANSICHAR* Expr, const TCHAR* Format, va_list Args)
+bool UE_COLD UE_DEBUG_SECTION VARARGS CheckVerifyImpl(std::atomic<bool>& bExecuted, bool bAlways, const ANSICHAR* File, int32 Line, void* ProgramCounter, const ANSICHAR* Expr, const TCHAR* Format, va_list Args)
 {
 	FDebug::OptionallyLogFormattedEnsureMessageReturningFalse(true, Expr, File, Line, ProgramCounter, Format, Args);
 
@@ -780,7 +780,7 @@ bool UE_DEBUG_SECTION VARARGS CheckVerifyImpl(std::atomic<bool>& bExecuted, bool
 }
 
 UE_AUTORTFM_ALWAYS_OPEN
-bool UE_DEBUG_SECTION UE::Assert::Private::ExecCheckImplInternal(std::atomic<bool>& bExecuted, bool bAlways, const ANSICHAR* File, int32 Line, const ANSICHAR* Expr)
+bool UE_COLD UE_DEBUG_SECTION UE::Assert::Private::ExecCheckImplInternal(std::atomic<bool>& bExecuted, bool bAlways, const ANSICHAR* File, int32 Line, const ANSICHAR* Expr)
 {
 	if (((bAlways && GEnsureAlwaysEnabled) || !bExecuted.load(std::memory_order_relaxed)) && FPlatformMisc::IsEnsureAllowed())
 	{
@@ -797,7 +797,7 @@ bool UE_DEBUG_SECTION UE::Assert::Private::ExecCheckImplInternal(std::atomic<boo
 }
 
 UE_AUTORTFM_ALWAYS_OPEN
-bool UE_DEBUG_SECTION VARARGS UE::Assert::Private::EnsureFailed(std::atomic<bool>& bExecuted, const FStaticEnsureRecord* Ensure, ...)
+bool UE_COLD UE_DEBUG_SECTION VARARGS UE::Assert::Private::EnsureFailed(std::atomic<bool>& bExecuted, const FStaticEnsureRecord* Ensure, ...)
 {
 	if (bExecuted.exchange(true, std::memory_order_release) && !(Ensure->bAlways && GEnsureAlwaysEnabled))
 	{
diff --git a/Engine/Source/Runtime/Core/Public/HAL/Platform.h b/Engine/Source/Runtime/Core/Public/HAL/Platform.h
--- a/Engine/Source/Runtime/Core/Public/HAL/Platform.h
+++ b/Engine/Source/Runtime/Core/Public/HAL/Platform.h
@@ -892,6 +892,20 @@
 	#define PLATFORM_CODE_SECTION(Name)
 #endif
 
+// Mark a function or type as unlikely to be used so the compiler can store it away from the warm code paths & optimize it for size rather than speed.
+// Examples:
+//	void UE_COLD HandleError();
+//	auto ErrorHandler = []() UE_COLD { abort(); };
+#if !defined(UE_COLD)
+	#if defined(_MSC_VER)
+		#define UE_COLD __declspec(noinline)
+	#elif defined(__GNUC__) || defined(__clang__) || defined(__llvm__)
+		#define UE_COLD __attribute__((cold))
+	#else
+		#define UE_COLD
+	#endif
+#endif
+
 // These have to be forced inline on some OSes so the dynamic loader will not
 // resolve to our allocators for the system libraries.
 #ifndef OPERATOR_NEW_INLINE
diff --git a/Engine/Source/Runtime/Core/Public/Logging/StructuredLog.h b/Engine/Source/Runtime/Core/Public/Logging/StructuredLog.h
--- a/Engine/Source/Runtime/Core/Public/Logging/StructuredLog.h
+++ b/Engine/Source/Runtime/Core/Public/Logging/StructuredLog.h
@@ -409,7 +409,7 @@ struct TFieldArgType<TLogFieldName<NameType>>
 
 /** Log with fields created from the arguments, which may be values or pairs of name/value. */
 template <typename... FieldArgTypes, typename LogType>
-FORCENOINLINE UE_DEBUG_SECTION void LogWithFields(const FLogCategoryBase& Category, const LogType& Log, typename TFieldArgType<FieldArgTypes>::Type... FieldArgs)
+UE_COLD UE_DEBUG_SECTION void LogWithFields(const FLogCategoryBase& Category, const LogType& Log, typename TFieldArgType<FieldArgTypes>::Type... FieldArgs)
 {
 	constexpr int32 FieldCount = FLogFieldCreator::template GetCount<FieldArgTypes...>();
 	static_assert(FieldCount > 0);
@@ -420,7 +420,7 @@ FORCENOINLINE UE_DEBUG_SECTION void LogWithFields(const FLogCategoryBase& Catego
 
 /** Fatal log with fields created from the arguments, which may be values or pairs of name/value. */
 template <typename... FieldArgTypes, typename LogType>
-inline void FatalLogWithFields(const FLogCategoryBase& Category, const LogType& Log, typename TFieldArgType<FieldArgTypes>::Type... FieldArgs)
+inline void UE_COLD FatalLogWithFields(const FLogCategoryBase& Category, const LogType& Log, typename TFieldArgType<FieldArgTypes>::Type... FieldArgs)
 {
 	constexpr int32 FieldCount = FLogFieldCreator::template GetCount<FieldArgTypes...>();
 	static_assert(FieldCount > 0);
diff --git a/Engine/Source/Runtime/Core/Public/Misc/AssertionMacros.h b/Engine/Source/Runtime/Core/Public/Misc/AssertionMacros.h
--- a/Engine/Source/Runtime/Core/Public/Misc/AssertionMacros.h
+++ b/Engine/Source/Runtime/Core/Public/Misc/AssertionMacros.h
@@ -21,12 +21,14 @@
 	// executable. This code should never execute so using a separate section keeps
 	// it well off the hot path and hopefully out of the instruction cache. It also
 	// facilitates reasoning about the makeup of a compiled/linked binary.
+	// Also see UE_COLD.
 	#define UE_DEBUG_SECTION PLATFORM_CODE_SECTION(".uedbg")
 #else
 	// On ARM we can't do this because the executable will require jumps larger
 	// than the branch instruction can handle. Clang will only generate
 	// the trampolines in the .text segment of the binary. If the uedbg segment
 	// is present it will generate code that it cannot link.
+	// Consider using UE_COLD instead.
 	#define UE_DEBUG_SECTION
 #endif // DO_CHECK || DO_GUARD_SLOW
 #endif
@@ -194,7 +196,7 @@ public:
 // lambdas. This can be worked around by calling the lambda from inside this
 // templated (and correctly non-inlined) function.
 template <typename RetType=void, class InnerType, typename... ArgTypes>
-RetType FORCENOINLINE UE_DEBUG_SECTION DispatchCheckVerify(InnerType&& Inner, ArgTypes const&... Args)
+RetType UE_COLD UE_DEBUG_SECTION DispatchCheckVerify (InnerType&& Inner, ArgTypes const&... Args)
 {
 	return Inner(Args...);
 }
@@ -387,9 +389,9 @@ RetType FORCENOINLINE UE_DEBUG_SECTION DispatchCheckVerify(InnerType&& Inner, Ar
 		}
 	};
 
-	CORE_API bool UE_DEBUG_SECTION VARARGS EnsureFailed(std::atomic<bool>& bExecuted, const FStaticEnsureRecord* Ensure, ...);
-	
-	CORE_API bool UE_DEBUG_SECTION ExecCheckImplInternal(std::atomic<bool>& bExecuted, bool bAlways, const ANSICHAR* File, int32 Line, const ANSICHAR* Expr);
+	CORE_API bool UE_COLD UE_DEBUG_SECTION VARARGS EnsureFailed(std::atomic<bool>& bExecuted, const FStaticEnsureRecord* Ensure, ...);
+
+	CORE_API bool UE_COLD UE_DEBUG_SECTION ExecCheckImplInternal(std::atomic<bool>& bExecuted, bool bAlways, const ANSICHAR* File, int32 Line, const ANSICHAR* Expr);
 
 	} // UE::Assert::Private
 
@@ -404,11 +406,11 @@ RetType FORCENOINLINE UE_DEBUG_SECTION DispatchCheckVerify(InnerType&& Inner, Ar
 
 	#define UE_ENSURE_IMPL(Always, InExpression) \
 		(LIKELY(!!(InExpression)) \
-			|| (::UE::Assert::Private::ExecCheckImplInternal([]() UE_DEBUG_SECTION -> std::atomic<bool>& { static std::atomic<bool> ENSURE_bExecuted = false; return ENSURE_bExecuted; } (), Always, __FILE__, __LINE__, #InExpression) \
+			|| (::UE::Assert::Private::ExecCheckImplInternal([]() UE_COLD UE_DEBUG_SECTION -> std::atomic<bool>& { static std::atomic<bool> ENSURE_bExecuted = false; return ENSURE_bExecuted; } (), Always, __FILE__, __LINE__, #InExpression) \
 			&& [] () { PLATFORM_BREAK(); return false; } ()))
 
 	#define UE_ENSURE_IMPL2(Capture, Always, InExpression, InFormat, ...) \
-		(LIKELY(!!(InExpression)) || ([Capture] () UE_DEBUG_SECTION \
+		(LIKELY(!!(InExpression)) || ([Capture] () UE_COLD UE_DEBUG_SECTION \
 		{ \
 			UE_VALIDATE_FORMAT_STRING(InFormat, ##__VA_ARGS__); \
 			static std::atomic<bool> ENSURE_bExecuted = false; \
@@ -507,7 +509,7 @@ namespace UEAsserts_Private
 ----------------------------------------------------------------------------*/
 
 /** low level fatal error handler. */
-CORE_API void UE_DEBUG_SECTION VARARGS LowLevelFatalErrorHandler(const ANSICHAR* File, int32 Line, const TCHAR* Format=TEXT(""), ... );
+CORE_API void UE_COLD UE_DEBUG_SECTION VARARGS LowLevelFatalErrorHandler (const ANSICHAR* File, int32 Line, const TCHAR* Format=TEXT(""), ... );
 
 #define LowLevelFatalError(Format, ...) \
 	{ \
