CmdModifyString - Console command to replace a DisplayString in the LocalizationManager StringTable, to preview a string in the running editor

diff --git a/Engine/Source/Runtime/Core/Private/Internationalization/TextLocalizationManager.cpp b/Engine/Source/Runtime/Core/Private/Internationalization/TextLocalizationManager.cpp
--- a/Engine/Source/Runtime/Core/Private/Internationalization/TextLocalizationManager.cpp
+++ b/Engine/Source/Runtime/Core/Private/Internationalization/TextLocalizationManager.cpp
@@ -112,6 +112,54 @@ static FAutoConsoleCommand CmdDumpLiveTable(
 #endif
 		}
 	}));
+
+static FAutoConsoleCommand CmdModifyString(
+	TEXT("Localization.CmdModifyString"),
+	TEXT("Replaces DisplayString in the live table given required arguments: 'Namespace', 'Key', and 'DisplayString'"),
+	FConsoleCommandWithArgsDelegate::CreateStatic([](const TArray<FString>& Args)
+		{
+			// Rebuild the full string of arguments, since values within quotes may have been split on space(s)
+			FString Arguments;
+			for (const FString& Arg : Args)
+			{
+				Arguments.Append(Arg);
+				Arguments.Append(" ");
+			}
+
+			TOptional<FString> Namespace;
+			TOptional<FString> Key;
+			TOptional<FString> DisplayString;
+
+			FString Tmp;
+			if (FParse::Value(*Arguments, TEXT("Namespace="), Tmp))
+			{
+				Namespace = MoveTemp(Tmp);
+			}
+
+			if (FParse::Value(*Arguments, TEXT("Key="), Tmp))
+			{
+				Key = MoveTemp(Tmp);
+			}
+
+			if (FParse::Value(*Arguments, TEXT("DisplayString="), Tmp))
+			{
+				DisplayString = MoveTemp(Tmp);
+			}
+
+			if (!Namespace.IsSet() || !Key.IsSet() || !DisplayString.IsSet())
+			{
+				UE_LOG(LogLocalization, Warning, TEXT("Missing argument(s): Namespace, Key, and/or DisplayString"));
+				return;
+			}
+			// An empty DisplayString is allowed, but the argument for it must be provided: -DisplayString=""
+			if (Namespace.GetValue().IsEmpty() || Key.GetValue().IsEmpty())
+			{
+				UE_LOG(LogLocalization, Warning, TEXT("Empty argument(s): Namespace, Key, and/or DisplayString"));
+				return;
+			}
+
+			FTextLocalizationManager::Get().ReplaceStringInLiveTable(Namespace.GetPtrOrNull(), Key.GetPtrOrNull(), DisplayString.GetPtrOrNull());
+		}));
 #endif
 
 FString KeyifyTextId(const FTextId& TextId)
@@ -742,6 +790,44 @@ void FTextLocalizationManager::DumpLiveTable(const FString& OutputFilename, cons
 
 	FFileHelper::SaveStringToFile(DumpString, *OutputFilename, FFileHelper::EEncodingOptions::ForceUTF8WithoutBOM);
 }
+
+void FTextLocalizationManager::ReplaceStringInLiveTable(const FString* Namespace, const FString* Key, const FString* DisplayString)
+{
+	FSlowHeartBeatScope SuspendHeartBeat;
+
+	// Lock while updating the table
+	FScopeLock ScopeLock(&DisplayStringTableCS);
+
+	const FTextId TextId(*Namespace, *Key);
+
+	FDisplayStringEntry* LiveEntry = DisplayStringLookupTable.Find(TextId);
+
+	if (LiveEntry)
+	{
+		auto GetSourceStringRef = [DisplayString]() -> const FString&
+			{
+				if (DisplayString)
+				{
+					return *DisplayString;
+				}
+
+				static const FString EmptyString;
+				return EmptyString;
+			};
+
+		const FString& SourceString = GetSourceStringRef();
+
+		LiveEntry->DisplayString = MakeTextDisplayString(CopyTemp(SourceString));
+		DirtyLocalRevisionForTextId(TextId);
+
+		UE_LOG(LogConsoleResponse, Display, TEXT("Updated string for Namespace='%s', Key='%s' to DisplayString='%s'"), **Namespace, **Key, **DisplayString);
+	}
+	else
+	{
+		UE_LOG(LogConsoleResponse, Warning, TEXT("String not found for Namespace='%s', Key='%s'"), **Namespace, **Key);
+	}
+}
+
 #endif
 
 FString FTextLocalizationManager::GetRequestedLanguageName() const
diff --git a/Engine/Source/Runtime/Core/Public/Internationalization/TextLocalizationManager.h b/Engine/Source/Runtime/Core/Public/Internationalization/TextLocalizationManager.h
--- a/Engine/Source/Runtime/Core/Public/Internationalization/TextLocalizationManager.h
+++ b/Engine/Source/Runtime/Core/Public/Internationalization/TextLocalizationManager.h
@@ -164,6 +164,7 @@ private:
 public:
 	CORE_API void DumpLiveTable(const FString* NamespaceFilter = nullptr, const FString* KeyFilter = nullptr, const FString* DisplayStringFilter = nullptr, const FLogCategoryBase* CategoryOverride = nullptr) const;
 	CORE_API void DumpLiveTable(const FString& OutputFilename, const FString* NamespaceFilter = nullptr, const FString* KeyFilter = nullptr, const FString* DisplayStringFilter = nullptr) const;
+	CORE_API void ReplaceStringInLiveTable(const FString* Namespace, const FString* Key, const FString* DisplayString);
 #endif
 
 	/**
diff --git a/Engine/Source/Runtime/Engine/Private/UserInterface/Console.cpp b/Engine/Source/Runtime/Engine/Private/UserInterface/Console.cpp
--- a/Engine/Source/Runtime/Engine/Private/UserInterface/Console.cpp
+++ b/Engine/Source/Runtime/Engine/Private/UserInterface/Console.cpp
@@ -797,7 +797,8 @@ void UConsole::AppendInputText(const FString& Text)
 		int32 Character = **TextMod.Left(1);
 		TextMod.MidInline(1, MAX_int32, EAllowShrinking::No);
 
-		if (Character >= 0x20 && Character < 0x100)
+		// Discard control characters, and allow more than extended ASCII (0xff) for Arabic etc
+		if (Character >= 0x20)
 		{
 			TCHAR Temp[2];
 			Temp[0] = Character;
