Use named FString::ConstructWithSlack constructor instead of ambiguous constructor.

diff --git a/Engine/Source/Runtime/Core/Private/Containers/StringTest.cpp b/Engine/Source/Runtime/Core/Private/Containers/StringTest.cpp
--- a/Engine/Source/Runtime/Core/Private/Containers/StringTest.cpp
+++ b/Engine/Source/Runtime/Core/Private/Containers/StringTest.cpp
@@ -440,10 +440,10 @@ bool FStringConstructorWithSlackTest::RunTest(const FString& Parameters)
 
 		const SIZE_T ExpectedCapacity = FMemory::QuantizeSize(NumElements * sizeof(TCHAR));
 
-		FString StringFromTChar(TestString, ExtraSlack);
+		FString StringFromTChar = FString::ConstructWithSlack(TestString, ExtraSlack);
 		TestEqual(TEXT("(TCHAR: Valid string with valid slack) resulting capacity"), StringFromTChar.GetAllocatedSize(), ExpectedCapacity);
 
-		FString StringFromAscii(TestAsciiString, ExtraSlack);
+		FString StringFromAscii = FString::ConstructWithSlack(TestAsciiString, ExtraSlack);
 		TestEqual(TEXT("(ASCII: Valid string with valid slack) resulting capacity"), StringFromAscii.GetAllocatedSize(), ExpectedCapacity);
 
 		FString StringFromFStringView(FStringView(TestString), ExtraSlack);
@@ -463,10 +463,10 @@ bool FStringConstructorWithSlackTest::RunTest(const FString& Parameters)
 
 		const SIZE_T ExpectedCapacity = FMemory::QuantizeSize(NumElements * sizeof(TCHAR));
 
-		FString StringFromTChar(TestString, ExtraSlack);
+		FString StringFromTChar = FString::ConstructWithSlack(TestString, ExtraSlack);
 		TestEqual(TEXT("(TCHAR: Valid string with zero slack) resulting capacity"), StringFromTChar.GetAllocatedSize(), ExpectedCapacity);
 
-		FString StringFromAscii(TestAsciiString, ExtraSlack);
+		FString StringFromAscii = FString::ConstructWithSlack(TestAsciiString, ExtraSlack);
 		TestEqual(TEXT("(ASCII: Valid string with zero slack) resulting capacity"), StringFromAscii.GetAllocatedSize(), ExpectedCapacity);
 
 		FString StringFromFStringView(FStringView(TestString), ExtraSlack);
@@ -486,10 +486,10 @@ bool FStringConstructorWithSlackTest::RunTest(const FString& Parameters)
 
 		const SIZE_T ExpectedCapacity = FMemory::QuantizeSize(NumElements * sizeof(TCHAR));
 
-		FString StringFromTChar(TestString, ExtraSlack);
+		FString StringFromTChar = FString::ConstructWithSlack(TestString, ExtraSlack);
 		TestEqual(TEXT("(TCHAR: Empty string with slack) resulting capacity"), StringFromTChar.GetAllocatedSize(), ExpectedCapacity);
 
-		FString StringFromAscii(TestAsciiString, ExtraSlack);
+		FString StringFromAscii = FString::ConstructWithSlack(TestAsciiString, ExtraSlack);
 		TestEqual(TEXT("(ASCII: Empty string with slack) resulting capacity"), StringFromAscii.GetAllocatedSize(), ExpectedCapacity);
 
 		FString StringFromFStringView(FStringView(TestString), ExtraSlack);
@@ -507,10 +507,10 @@ bool FStringConstructorWithSlackTest::RunTest(const FString& Parameters)
 
 		const SIZE_T ExpectedCapacity = 0u;
 
-		FString StringFromTChar(TestString, ExtraSlack);
+		FString StringFromTChar = FString::ConstructWithSlack(TestString, ExtraSlack);
 		TestEqual(TEXT("(TCHAR: Empty string with zero slack) resulting capacity"), StringFromTChar.GetAllocatedSize(), ExpectedCapacity);
 
-		FString StringFromAscii(TestAsciiString, ExtraSlack);
+		FString StringFromAscii = FString::ConstructWithSlack(TestAsciiString, ExtraSlack);
 		TestEqual(TEXT("(ASCII: Empty string with zero slack) resulting capacity"), StringFromAscii.GetAllocatedSize(), ExpectedCapacity);
 
 		FString StringFromFStringView(FStringView(TestString), ExtraSlack);
@@ -833,38 +833,6 @@ bool FStringFindAndContainsTest::RunTest(const FString& Parameters)
 	return true;
 }
 
-IMPLEMENT_SIMPLE_AUTOMATION_TEST(FStringConstructorWithLengthTest, "System.Core.String.ConstructorWithLength", EAutomationTestFlags_ApplicationContextMask | EAutomationTestFlags::SmokeFilter)
-bool FStringConstructorWithLengthTest::RunTest(const FString& Parameters)
-{
-	auto DoTest = [this](const TCHAR* Ptr, int32 Size, const TArray<TCHAR>& Expected)
-	{
-		FString Str(Size, Ptr);
-
-		const TArray<TCHAR>& StrArr = Str.GetCharArray();
-		if (StrArr != Expected)
-		{
-			AddError(
-				FString::Printf(
-					TEXT("FString(%d, TEXT(\"%s\")) failure: result '%s' (expected '%.*s')"),
-					Size,
-					Ptr,
-					*Str,
-					Expected.Num(),
-					Expected.GetData()
-				)
-			);
-		}
-	};
-
-	DoTest(TEXT("\0abc"),    4, {});
-	DoTest(TEXT("abc\0def"), 3, { TEXT('a'), TEXT('b'), TEXT('c'),                                                          TEXT('\0') });
-	DoTest(TEXT("abc\0def"), 4, { TEXT('a'), TEXT('b'), TEXT('c'), TEXT('\0'),                                              TEXT('\0') });
-	DoTest(TEXT("abc\0def"), 7, { TEXT('a'), TEXT('b'), TEXT('c'), TEXT('\0'), TEXT('d'), TEXT('e'), TEXT('f'),             TEXT('\0') });
-	DoTest(TEXT("abc\0def"), 8, { TEXT('a'), TEXT('b'), TEXT('c'), TEXT('\0'), TEXT('d'), TEXT('e'), TEXT('f'), TEXT('\0'), TEXT('\0') });
-
-	return true;
-}
-
 IMPLEMENT_SIMPLE_AUTOMATION_TEST(FStringConstructWithSlackTest, "System.Core.String.ConstructWithSlack", EAutomationTestFlags_ApplicationContextMask | EAutomationTestFlags::SmokeFilter)
 bool FStringConstructWithSlackTest::RunTest(const FString& Parameters)
 {
