Fixed some legacy trait usage in Crc.h.

diff --git a/Engine/Source/Runtime/Core/Public/Misc/Crc.h b/Engine/Source/Runtime/Core/Public/Misc/Crc.h
--- a/Engine/Source/Runtime/Core/Public/Misc/Crc.h
+++ b/Engine/Source/Runtime/Core/Public/Misc/Crc.h
@@ -7,7 +7,6 @@
 #include "Misc/AssertionMacros.h"
 #include "Misc/CString.h"
 #include "Misc/Char.h"
-#include "Templates/EnableIf.h"
 #include "Templates/UnrealTypeTraits.h"
 #include "Traits/IsCharType.h"
 
@@ -43,40 +42,41 @@ struct FCrc
 
 	/** String CRC. */
 	template <typename CharType>
-	[[nodiscard]] static typename TEnableIf<sizeof(CharType) != 1, uint32>::Type StrCrc32(const CharType* Data, uint32 CRC = 0)
+	[[nodiscard]] static uint32 StrCrc32(const CharType* Data, uint32 CRC = 0)
 	{
 		// We ensure that we never try to do a StrCrc32 with a CharType of more than 4 bytes.  This is because
 		// we always want to treat every CRC as if it was based on 4 byte chars, even if it's less, because we
 		// want consistency between equivalent strings with different character types.
 		static_assert(sizeof(CharType) <= 4, "StrCrc32 only works with CharType up to 32 bits.");
 
-		CRC = ~CRC;
-		while (CharType Ch = *Data++)
+		if constexpr (sizeof(CharType) != 1)
 		{
-			CRC = (CRC >> 8) ^ CRCTablesSB8[0][(CRC ^ Ch) & 0xFF];
-			Ch >>= 8;
-			CRC = (CRC >> 8) ^ CRCTablesSB8[0][(CRC ^ Ch) & 0xFF];
-			Ch >>= 8;
-			CRC = (CRC >> 8) ^ CRCTablesSB8[0][(CRC ^ Ch) & 0xFF];
-			Ch >>= 8;
-			CRC = (CRC >> 8) ^ CRCTablesSB8[0][(CRC ^ Ch) & 0xFF];
+			CRC = ~CRC;
+			while (CharType Ch = *Data++)
+			{
+				CRC = (CRC >> 8) ^ CRCTablesSB8[0][(CRC ^ Ch) & 0xFF];
+				Ch >>= 8;
+				CRC = (CRC >> 8) ^ CRCTablesSB8[0][(CRC ^ Ch) & 0xFF];
+				Ch >>= 8;
+				CRC = (CRC >> 8) ^ CRCTablesSB8[0][(CRC ^ Ch) & 0xFF];
+				Ch >>= 8;
+				CRC = (CRC >> 8) ^ CRCTablesSB8[0][(CRC ^ Ch) & 0xFF];
+			}
+			return ~CRC;
 		}
-		return ~CRC;
-	}
-
-	template <typename CharType>
-	[[nodiscard]] static typename TEnableIf<sizeof(CharType) == 1, uint32>::Type StrCrc32(const CharType* Data, uint32 CRC = 0)
-	{
-		/* Overload for when CharType is a byte, which causes warnings when right-shifting by 8 */
-		CRC = ~CRC;
-		while (CharType Ch = *Data++)
+		else
 		{
-			CRC = (CRC >> 8) ^ CRCTablesSB8[0][(CRC ^ Ch) & 0xFF];
-			CRC = (CRC >> 8) ^ CRCTablesSB8[0][(CRC     ) & 0xFF];
-			CRC = (CRC >> 8) ^ CRCTablesSB8[0][(CRC     ) & 0xFF];
-			CRC = (CRC >> 8) ^ CRCTablesSB8[0][(CRC     ) & 0xFF];
+			/* Special case for when CharType is a byte, which causes warnings when right-shifting by 8 */
+			CRC = ~CRC;
+			while (CharType Ch = *Data++)
+			{
+				CRC = (CRC >> 8) ^ CRCTablesSB8[0][(CRC ^ Ch) & 0xFF];
+				CRC = (CRC >> 8) ^ CRCTablesSB8[0][(CRC     ) & 0xFF];
+				CRC = (CRC >> 8) ^ CRCTablesSB8[0][(CRC     ) & 0xFF];
+				CRC = (CRC >> 8) ^ CRCTablesSB8[0][(CRC     ) & 0xFF];
+			}
+			return ~CRC;
 		}
-		return ~CRC;
 	}
 
 	/**
@@ -269,7 +269,7 @@ FORCEINLINE const TCHAR* ToCStr(const TCHAR* Ptr)
 template <typename InKeyType, bool bInAllowDuplicateKeys = false>
 struct TStringPointerSetKeyFuncs_DEPRECATED
 {
-	static_assert(TIsCharType<std::remove_pointer_t<decltype(ToCStr(std::declval<InKeyType>()))>>::Value, "TStringPointerSetKeyFuncs_DEPRECATED should only be used with keys which character types");
+	static_assert(TIsCharType_V<std::remove_pointer_t<decltype(ToCStr(std::declval<InKeyType>()))>>, "TStringPointerSetKeyFuncs_DEPRECATED should only be used with keys which character types");
 
 	using KeyType         = InKeyType;
 	using KeyInitType     = typename TTypeTraits<InKeyType>::ConstPointerType;
@@ -301,7 +301,7 @@ struct TStringPointerSetKeyFuncs_DEPRECATED
 template <typename InKeyType, typename InValueType, bool bInAllowDuplicateKeys = false>
 struct TStringPointerMapKeyFuncs_DEPRECATED
 {
-	static_assert(TIsCharType<std::remove_pointer_t<decltype(ToCStr(std::declval<InKeyType>()))>>::Value, "TStringPointerMapKeyFuncs_DEPRECATED should only be used with keys which character types");
+	static_assert(TIsCharType_V<std::remove_pointer_t<decltype(ToCStr(std::declval<InKeyType>()))>>, "TStringPointerMapKeyFuncs_DEPRECATED should only be used with keys which character types");
 
 	using KeyType         = InKeyType;
 	using KeyInitType     = typename TTypeTraits<InKeyType>::ConstPointerType;
@@ -326,3 +326,7 @@ struct TStringPointerMapKeyFuncs_DEPRECATED
 		return FCrc::Strihash_DEPRECATED(ToCStr(Key));
 	}
 };
+
+#if UE_ENABLE_INCLUDE_ORDER_DEPRECATED_IN_5_5
+#include "Templates/EnableIf.h"
+#endif
