Making Parse::Value ( <number types> ) quote tolerant.
This is the first in a 2 step for the Jira..
Before sanitized CLI lines can have values wrapped in quotes, I need to make the consuming parse able to handle them.

diff --git a/Engine/Source/Runtime/Core/Private/Misc/Parse.cpp b/Engine/Source/Runtime/Core/Private/Misc/Parse.cpp
--- a/Engine/Source/Runtime/Core/Private/Misc/Parse.cpp
+++ b/Engine/Source/Runtime/Core/Private/Misc/Parse.cpp
@@ -580,11 +580,14 @@ bool FParse::Value(	const TCHAR* Stream, const TCHAR* Match, FName& Name )
 //
 bool FParse::Value( const TCHAR* Stream, const TCHAR* Match, uint32& Value )
 {
-	const TCHAR* Temp = FCString::Strifind(Stream,Match);
-	TCHAR* End;
-	if( Temp==NULL )
+	TCHAR Temp[256];
+	if (!FParse::Value(Stream, Match, Temp, UE_ARRAY_COUNT(Temp)))
+	{
 		return false;
-	Value = FCString::Strtoi( Temp + FCString::Strlen(Match), &End, 10 );
+	}
+	TCHAR* End_NotUsed;
+
+	Value = FCString::Strtoi(Temp, &End_NotUsed, 10 );
 
 	return true;
 }
@@ -594,10 +597,12 @@ bool FParse::Value( const TCHAR* Stream, const TCHAR* Match, uint32& Value )
 //
 bool FParse::Value( const TCHAR* Stream, const TCHAR* Match, uint8& Value )
 {
-	const TCHAR* Temp = FCString::Strifind(Stream,Match);
-	if( Temp==NULL )
+	TCHAR Temp[256];
+	if (!FParse::Value(Stream, Match, Temp, UE_ARRAY_COUNT(Temp)))
+	{
 		return false;
-	Temp += FCString::Strlen( Match );
+	}
+
 	Value = (uint8)FCString::Atoi( Temp );
 	return Value!=0 || FChar::IsDigit(Temp[0]);
 }
@@ -607,10 +612,12 @@ bool FParse::Value( const TCHAR* Stream, const TCHAR* Match, uint8& Value )
 //
 bool FParse::Value( const TCHAR* Stream, const TCHAR* Match, int8& Value )
 {
-	const TCHAR* Temp = FCString::Strifind(Stream,Match);
-	if( Temp==NULL )
+	TCHAR Temp[256];
+	if (!FParse::Value(Stream, Match, Temp, UE_ARRAY_COUNT(Temp)))
+	{
 		return false;
-	Temp += FCString::Strlen( Match );
+	}
+
 	Value = (int8)FCString::Atoi( Temp );
 	return Value!=0 || FChar::IsDigit(Temp[0]);
 }
@@ -620,10 +627,12 @@ bool FParse::Value( const TCHAR* Stream, const TCHAR* Match, int8& Value )
 //
 bool FParse::Value( const TCHAR* Stream, const TCHAR* Match, uint16& Value )
 {
-	const TCHAR* Temp = FCString::Strifind( Stream, Match );
-	if( Temp==NULL )
+	TCHAR Temp[256];
+	if (!FParse::Value(Stream, Match, Temp, UE_ARRAY_COUNT(Temp)))
+	{
 		return false;
-	Temp += FCString::Strlen( Match );
+	}
+
 	Value = (uint16)FCString::Atoi( Temp );
 	return Value!=0 || FChar::IsDigit(Temp[0]);
 }
@@ -633,10 +642,12 @@ bool FParse::Value( const TCHAR* Stream, const TCHAR* Match, uint16& Value )
 //
 bool FParse::Value( const TCHAR* Stream, const TCHAR* Match, int16& Value )
 {
-	const TCHAR* Temp = FCString::Strifind( Stream, Match );
-	if( Temp==NULL )
+	TCHAR Temp[256];
+	if (!FParse::Value(Stream, Match, Temp, UE_ARRAY_COUNT(Temp)))
+	{
 		return false;
-	Temp += FCString::Strlen( Match );
+	}
+
 	Value = (int16)FCString::Atoi( Temp );
 	return Value!=0 || FChar::IsDigit(Temp[0]);
 }
@@ -646,10 +657,13 @@ bool FParse::Value( const TCHAR* Stream, const TCHAR* Match, int16& Value )
 //
 bool FParse::Value( const TCHAR* Stream, const TCHAR* Match, float& Value )
 {
-	const TCHAR* Temp = FCString::Strifind( Stream, Match );
-	if( Temp==NULL )
+	TCHAR Temp[256];
+	if (!FParse::Value(Stream, Match, Temp, UE_ARRAY_COUNT(Temp)))
+	{
 		return false;
-	Value = FCString::Atof( Temp+FCString::Strlen(Match) );
+	}
+
+	Value = FCString::Atof( Temp );
 	return true;
 }
 
@@ -658,10 +672,13 @@ bool FParse::Value( const TCHAR* Stream, const TCHAR* Match, float& Value )
 //
 bool FParse::Value(const TCHAR* Stream, const TCHAR* Match, double& Value)
 {
-	const TCHAR* Temp = FCString::Strifind(Stream, Match);
-	if (Temp == NULL)
+	TCHAR Temp[256];
+	if (!FParse::Value(Stream, Match, Temp, UE_ARRAY_COUNT(Temp)))
+	{
 		return false;
-	Value = FCString::Atod(Temp + FCString::Strlen(Match));
+	}
+
+	Value = FCString::Atod(Temp);
 	return true;
 }
 
@@ -671,20 +688,23 @@ bool FParse::Value(const TCHAR* Stream, const TCHAR* Match, double& Value)
 //
 bool FParse::Value( const TCHAR* Stream, const TCHAR* Match, int32& Value )
 {
-	const TCHAR* Temp = FCString::Strifind( Stream, Match );
-	if( Temp==NULL )
+	TCHAR Temp[256];
+	if (!FParse::Value(Stream, Match, Temp, UE_ARRAY_COUNT(Temp)))
+	{
 		return false;
-	Value = FCString::Atoi( Temp + FCString::Strlen(Match) );
+	}
+
+	Value = FCString::Atoi( Temp );
 	return true;
 }
 
 //
 // Get a boolean value.
 //
-bool FParse::Bool( const TCHAR* Stream, const TCHAR* Match, bool& OnOff )
+bool FParse::Bool(const TCHAR* Stream, const TCHAR* Match, bool& OnOff)
 {
 	TCHAR TempStr[16];
-	if( FParse::Value( Stream, Match, TempStr, 16 ) )
+	if (FParse::Value(Stream, Match, TempStr, UE_ARRAY_COUNT(TempStr)))
 	{
 		OnOff = FCString::ToBool(TempStr);
 		return true;
@@ -701,8 +721,10 @@ bool FParse::Bool( const TCHAR* Stream, const TCHAR* Match, bool& OnOff )
 bool FParse::Value( const TCHAR* Stream, const TCHAR* Match, struct FGuid& Guid )
 {
 	TCHAR Temp[256];
-	if( !FParse::Value( Stream, Match, Temp, UE_ARRAY_COUNT(Temp) ) )
+	if (!FParse::Value(Stream, Match, Temp, UE_ARRAY_COUNT(Temp)))
+	{
 		return false;
+	}
 
 	Guid.A = Guid.B = Guid.C = Guid.D = 0;
 	if( FCString::Strlen(Temp)==32 )
diff --git a/Engine/Source/Runtime/Core/Tests/Misc/ParseTests.cpp b/Engine/Source/Runtime/Core/Tests/Misc/ParseTests.cpp
--- a/Engine/Source/Runtime/Core/Tests/Misc/ParseTests.cpp
+++ b/Engine/Source/Runtime/Core/Tests/Misc/ParseTests.cpp
@@ -9,6 +9,7 @@
 #include "Tests/TestHarnessAdapter.h"
 
 #include <catch2/generators/catch_generators.hpp>
+#include <catch2/catch_approx.hpp>
 
 TEST_CASE("Parse::Value::ToBuffer", "[Parse][Smoke]")
 {
@@ -625,5 +626,281 @@ TEST_CASE("Parse::Token", "[Parse][Token][Smoke]")
 		}
 	}
 }
+
+
+TEST_CASE("Parse::Value::Numbers", "[Smoke]")
+{
+	SECTION("Int8")
+	{
+		auto [Input, ExpectedResults] = GENERATE_COPY(table<const TCHAR*, int8>(
+			{
+				// raw
+				{ TEXT("a=0"), (int8)0 },
+				{ TEXT("a=42"), (int8)42 },
+				{ TEXT("a=127"), (int8)127 },
+				{ TEXT("a=-1"), (int8)-1 },
+				{ TEXT("a=-128"), (int8)-128 },
+				// quoted
+				{ TEXT("a=\"0\""), (int8)0 },
+				{ TEXT("a=\"42\""), (int8)42 },
+				{ TEXT("a=\"127\""), (int8)127 },
+				{ TEXT("a=\"-1\""), (int8)-1 },
+				{ TEXT("a=\"-128\""), (int8)-128 },
+				// broken quotes takes the first number found
+				{ TEXT("a=123\"456\""), (int8)123 },
+				{ TEXT("a=\"123\"456"), (int8)123 },
+			}));
+
+		int8 Result;
+		CHECK(FParse::Value(Input, TEXT("a="), Result));
+		CHECK(Result == ExpectedResults);
+	}
+
+	SECTION("UInt8")
+	{
+		auto [Input, ExpectedResults] = GENERATE_COPY(table<const TCHAR*, uint8>(
+			{
+				// raw
+				{ TEXT("a=0"), (uint8)0 },
+				{ TEXT("a=42"), (uint8)42 },
+				{ TEXT("a=127"), (uint8)127 },
+				{ TEXT("a=255"), (uint8)255 },
+				// quoted
+				{ TEXT("a=\"0\""), (uint8)0 },
+				{ TEXT("a=\"42\""), (uint8)42 },
+				{ TEXT("a=\"127\""), (uint8)127 },
+				{ TEXT("a=\"255\""), (uint8)255 },
+				// broken quotes takes the first number found
+				{ TEXT("a=123\"456\""), (uint8)123 },
+				{ TEXT("a=\"123\"456"), (uint8)123 },
+			}));
+
+		uint8 Result;
+		CHECK(FParse::Value(Input, TEXT("a="), Result));
+		CHECK(Result == ExpectedResults);
+	}
+	
+	SECTION("Int16")
+	{
+		auto [Input, ExpectedResults] = GENERATE_COPY(table<const TCHAR*, int16>(
+			{
+				// raw
+				{ TEXT("a=0"), (int16)0 },
+				{ TEXT("a=42"), (int16)42 },
+				{ TEXT("a=32767"), (int16)32767 },
+				{ TEXT("a=-1"), (int16)-1 },
+				{ TEXT("a=-32768"), (int16)-32768 },
+				// quoted
+				{ TEXT("a=\"0\""), (int16)0 },
+				{ TEXT("a=\"42\""), (int16)42 },
+				{ TEXT("a=\"32767\""), (int16)32767 },
+				{ TEXT("a=\"-1\""), (int16)-1 },
+				{ TEXT("a=\"-32768\""), (int16)-32768 },
+				// broken quotes takes the first number found
+				{ TEXT("a=123\"456\""), (int16)123 },
+				{ TEXT("a=\"123\"456"), (int16)123 },
+			}));
+
+		int16 Result;
+		CHECK(FParse::Value(Input, TEXT("a="), Result));
+		CHECK(Result == ExpectedResults);
+	}
+
+	SECTION("UInt16")
+	{
+		auto [Input, ExpectedResults] = GENERATE_COPY(table<const TCHAR*, uint16>(
+			{
+				// raw
+				{ TEXT("a=0"), (uint16)0 },
+				{ TEXT("a=42"), (uint16)42 },
+				{ TEXT("a=32767"), (uint16)32767 },
+				{ TEXT("a=65535"), (uint16)65535 },
+				// quoted
+				{ TEXT("a=\"0\""), (uint16)0 },
+				{ TEXT("a=\"42\""), (uint16)42 },
+				{ TEXT("a=\"32767\""), (uint16)32767 },
+				{ TEXT("a=\"65535\""), (uint16)65535 },
+				// broken quotes takes the first number found
+				{ TEXT("a=123\"456\""), (uint16)123 },
+				{ TEXT("a=\"123\"456"), (uint16)123 },
+			}));
+
+		uint16 Result;
+		CHECK(FParse::Value(Input, TEXT("a="), Result));
+		CHECK(Result == ExpectedResults);
+	}
+
+	SECTION("Int32")
+	{
+		auto [Input, ExpectedResults] = GENERATE_COPY(table<const TCHAR*, int32>(
+			{
+				// raw
+				{ TEXT("a=0"), (int32)0 },
+				{ TEXT("a=42"), (int32)42 },
+				{ TEXT("a=2147483647"), (int32)2147483647 },
+				{ TEXT("a=-1"), (int32)-1 },
+				{ TEXT("a=-2147483648"), (int32)-2147483648 },
+				// quoted
+				{ TEXT("a=\"0\""), (int32)0 },
+				{ TEXT("a=\"42\""), (int32)42 },
+				{ TEXT("a=\"2147483647\""), (int32)2147483647 },
+				{ TEXT("a=\"-1\""), (int32)-1 },
+				{ TEXT("a=\"-2147483648\""), (int32)-2147483648 },
+				// broken quotes takes the first number found
+				{ TEXT("a=123\"456\""), (int32)123 },
+				{ TEXT("a=\"123\"456"), (int32)123 },
+			}));
+
+		int32 Result;
+		CHECK(FParse::Value(Input, TEXT("a="), Result));
+		CHECK(Result == ExpectedResults);
+	}
+
+	SECTION("UInt32")
+	{
+		auto [Input, ExpectedResults] = GENERATE_COPY(table<const TCHAR*, uint32>(
+			{
+				// raw
+				{ TEXT("a=0"), (uint32)0 },
+				{ TEXT("a=42"), (uint32)42 },
+				{ TEXT("a=2147483647"), (uint32)2147483647 },
+				{ TEXT("a=4294967295"), (uint32)4294967295 },
+				// quoted
+				{ TEXT("a=\"0\""), (uint32)0 },
+				{ TEXT("a=\"42\""), (uint32)42 },
+				{ TEXT("a=\"2147483647\""), (uint32)2147483647 },
+				{ TEXT("a=\"4294967295\""), (uint32)4294967295 },
+				// broken quotes takes the first number found
+				{ TEXT("a=123\"456\""), (uint32)123 },
+				{ TEXT("a=\"123\"456"), (uint32)123 },
+			}));
+
+		uint32 Result;
+		CHECK(FParse::Value(Input, TEXT("a="), Result));
+		CHECK(Result == ExpectedResults);
+	}
+
+	SECTION("Int64")
+	{
+		auto [Input, ExpectedResults] = GENERATE_COPY(table<const TCHAR*, int64>(
+			{
+				// raw
+				{ TEXT("a=0"), 0ll },
+				{ TEXT("a=42"), 42ll },
+				{ TEXT("a=9223372036854775807"), 9223372036854775807ll },
+				{ TEXT("a=-1"), -1 },
+				{ TEXT("a=-9223372036854775807"), -9223372036854775807ll },
+				// quoted
+				{ TEXT("a=\"0\""), 0ll },
+				{ TEXT("a=\"42\""), 42ll },
+				{ TEXT("a=\"9223372036854775807\""), 9223372036854775807ll },
+				{ TEXT("a=\"-1\""), -1ll },
+				{ TEXT("a=\"-9223372036854775807\""), -9223372036854775807ll },
+				// broken quotes takes the first number found
+				{ TEXT("a=123\"456\""), 123ll },
+				{ TEXT("a=\"123\"456"), 123ll },
+			}));
+
+		int64 Result;
+		CHECK(FParse::Value(Input, TEXT("a="), Result));
+		CHECK(Result == ExpectedResults);
+	}
+
+	SECTION("UInt64")
+	{
+		auto [Input, ExpectedResults] = GENERATE_COPY(table<const TCHAR*, uint64>(
+			{
+				// raw
+				{ TEXT("a=0"), 0ull },
+				{ TEXT("a=42"), 42ull },
+				{ TEXT("a=9223372036854775807"),   9223372036854775807ull },
+				{ TEXT("a=18446744073709551615"), 18446744073709551615ull },
+				// quoted
+				{ TEXT("a=\"0\""), 0ull },
+				{ TEXT("a=\"42\""), 42ull },
+				{ TEXT("a=\"9223372036854775807\""), 9223372036854775807ull },
+				{ TEXT("a=\"18446744073709551615\""), 18446744073709551615ull },
+				// broken quotes takes the first number found
+				{ TEXT("a=123\"456\""), 123ull },
+				{ TEXT("a=\"123\"456"), 123ull },
+			}));
+
+		uint64 Result;
+		CHECK(FParse::Value(Input, TEXT("a="), Result));
+		CHECK(Result == ExpectedResults);
+	}
+
+	SECTION("float")
+	{
+		using Catch::Approx;
+
+		auto [Input, ExpectedResults] = GENERATE_COPY(table<const TCHAR*, float>(
+			{
+				// raw
+				{ TEXT("a=0.0"), 0.0f },
+				{ TEXT("a=0.5"), 0.5f },
+				{ TEXT("a=1.0"), 1.0f },
+				{ TEXT("a=42"), 42.0f },
+				{ TEXT("a=3.1415"), 3.1415f },
+				{ TEXT("a=-3.1415"), -3.1415f },
+				{ TEXT("a=340282346638528859811704183484516925440.0"), 340282346638528859811704183484516925440.0f },
+				{ TEXT("a=-340282346638528859811704183484516925440.0"), -340282346638528859811704183484516925440.0f },
+				// quoted
+				{ TEXT("a=\"0.0\""), 0.0f },
+				{ TEXT("a=\"0.5\""), 0.5f },
+				{ TEXT("a=\"1.0\""), 1.0f },
+				{ TEXT("a=\"42\""), 42.0f },
+				{ TEXT("a=\"3.1415\""), 3.1415f },
+				{ TEXT("a=\"-3.1415\""), -3.1415f },
+				{ TEXT("a=\"340282346638528859811704183484516925440.0\""), 340282346638528859811704183484516925440.0f },
+				{ TEXT("a=\"-340282346638528859811704183484516925440.0\""), -340282346638528859811704183484516925440.0f },
+				// broken quotes takes the first number found
+				{ TEXT("a=123\"456\""), 123.0f },
+				{ TEXT("a=\"123\"456"), 123.0f },
+			}));
+
+		float Result;
+		CHECK(FParse::Value(Input, TEXT("a="), Result));
+		CHECK(Result == Approx(ExpectedResults).margin(0.0001f));
+	}
+
+	SECTION("double")
+	{
+		using Catch::Approx;
+
+		auto [Input, ExpectedResults] = GENERATE_COPY(table<const TCHAR*, double>(
+			{
+				// raw
+				{ TEXT("a=0.0"), 0.0 },
+				{ TEXT("a=0.5"), 0.5 },
+				{ TEXT("a=1.0"), 1.0 },
+				{ TEXT("a=42"), 42.0 },
+				{ TEXT("a=3.1415"), 3.1415 },
+				{ TEXT("a=-3.1415"), -3.1415 },
+				{ TEXT("a=340282346638528859811704183484516925440.0"), 340282346638528859811704183484516925440.0 },
+				{ TEXT("a=-340282346638528859811704183484516925440.0"), -340282346638528859811704183484516925440.0 },
+				{ TEXT("a=" PREPROCESSOR_TO_STRING(DBL_MAX)), DBL_MAX },
+				{ TEXT("a=-" PREPROCESSOR_TO_STRING(DBL_MAX)), -DBL_MAX },
+				// quoted
+				{ TEXT("a=\"0.0\""), 0.0 },
+				{ TEXT("a=\"0.5\""), 0.5 },
+				{ TEXT("a=\"1.0\""), 1.0 },
+				{ TEXT("a=\"42\""), 42.0 },
+				{ TEXT("a=\"3.1415\""), 3.1415 },
+				{ TEXT("a=\"-3.1415\""), -3.1415 },
+				{ TEXT("a=\"340282346638528859811704183484516925440.0\""), 340282346638528859811704183484516925440.0 },
+				{ TEXT("a=\"-340282346638528859811704183484516925440.0\""), -340282346638528859811704183484516925440.0 },
+				{ TEXT("a=\"" PREPROCESSOR_TO_STRING(DBL_MAX) "\""), DBL_MAX},
+				{ TEXT("a=\"-" PREPROCESSOR_TO_STRING(DBL_MAX) "\""), -DBL_MAX },
+				// broken quotes takes the first number found
+				{ TEXT("a=123\"456\""), 123.0 },
+				{ TEXT("a=\"123\"456"), 123.0 },
+			}));
+
+		double Result;
+		CHECK(FParse::Value(Input, TEXT("a="), Result));
+		CHECK(Result == Approx(ExpectedResults).margin(0.0001f));
+	}
+}
 #endif
 	
