Add FCrc::StrCrc32Len; i.e. FCrc::StrCrc32 that takes a length param

diff --git a/Engine/Source/Runtime/Core/Public/Misc/Crc.h b/Engine/Source/Runtime/Core/Public/Misc/Crc.h
--- a/Engine/Source/Runtime/Core/Public/Misc/Crc.h
+++ b/Engine/Source/Runtime/Core/Public/Misc/Crc.h
@@ -79,6 +79,47 @@ struct FCrc
 		}
 	}
 
+	/** String CRC. */
+	template <typename CharType>
+	[[nodiscard]] static uint32 StrCrc32Len(const CharType* Data, int32 Length, uint32 CRC = 0)
+	{
+		// We ensure that we never try to do a StrCrc32 with a CharType of more than 4 bytes.  This is because
+		// we always want to treat every CRC as if it was based on 4 byte chars, even if it's less, because we
+		// want consistency between equivalent strings with different character types.
+		static_assert(sizeof(CharType) <= 4, "StrCrc32Len only works with CharType up to 32 bits.");
+
+		if constexpr (sizeof(CharType) != 1)
+		{
+			CRC = ~CRC;
+			for (int32 Idx = 0; Idx < Length; ++Idx)
+			{
+				CharType Ch = Data[Idx];
+				CRC = (CRC >> 8) ^ CRCTablesSB8[0][(CRC ^ Ch) & 0xFF];
+				Ch >>= 8;
+				CRC = (CRC >> 8) ^ CRCTablesSB8[0][(CRC ^ Ch) & 0xFF];
+				Ch >>= 8;
+				CRC = (CRC >> 8) ^ CRCTablesSB8[0][(CRC ^ Ch) & 0xFF];
+				Ch >>= 8;
+				CRC = (CRC >> 8) ^ CRCTablesSB8[0][(CRC ^ Ch) & 0xFF];
+			}
+			return ~CRC;
+		}
+		else
+		{
+			/* Special case for when CharType is a byte, which causes warnings when right-shifting by 8 */
+			CRC = ~CRC;
+			for (int32 Idx = 0; Idx < Length; ++Idx)
+			{
+				CharType Ch = Data[Idx];
+				CRC = (CRC >> 8) ^ CRCTablesSB8[0][(CRC ^ Ch) & 0xFF];
+				CRC = (CRC >> 8) ^ CRCTablesSB8[0][(CRC     ) & 0xFF];
+				CRC = (CRC >> 8) ^ CRCTablesSB8[0][(CRC     ) & 0xFF];
+				CRC = (CRC >> 8) ^ CRCTablesSB8[0][(CRC     ) & 0xFF];
+			}
+			return ~CRC;
+		}
+	}
+
 	/**
 	 * DEPRECATED
 	 * These tables and functions are deprecated because they're using tables and implementations
