fix for parent server hang issue in live

diff --git a/Engine/Source/Runtime/Core/Private/Unix/UnixPlatformProcess.cpp b/Engine/Source/Runtime/Core/Private/Unix/UnixPlatformProcess.cpp
--- a/Engine/Source/Runtime/Core/Private/Unix/UnixPlatformProcess.cpp
+++ b/Engine/Source/Runtime/Core/Private/Unix/UnixPlatformProcess.cpp
@@ -1460,7 +1460,7 @@ FGenericPlatformProcess::EWaitAndForkResult FUnixPlatformProcess::WaitAndFork()
 	{
 		for (int32 ForkIdx = 0; ForkIdx < NumForks; ++ForkIdx)
 		{
-			WaitAndForkSignalQueue.Enqueue(FForkSignalData(ForkIdx + 1, FDateTime::UtcNow()));
+			WaitAndForkSignalQueue.Enqueue(FForkSignalData(ForkIdx + 1, FDateTime()));
 		}
 	}
 
@@ -1497,7 +1497,7 @@ FGenericPlatformProcess::EWaitAndForkResult FUnixPlatformProcess::WaitAndFork()
 		Action.sa_sigaction = [](int32 Signal, siginfo_t* Info, void* Context) {
 			if (Signal == WAIT_AND_FORK_QUEUE_SIGNAL && Info)
 			{
-				WaitAndForkSignalQueue.Enqueue(FForkSignalData(Info->si_value.sival_int, FDateTime::UtcNow()));
+				WaitAndForkSignalQueue.Enqueue(FForkSignalData(Info->si_value.sival_int, FDateTime()));
 			}
 		};
 		sigaction(WAIT_AND_FORK_QUEUE_SIGNAL, &Action, nullptr);
@@ -1740,7 +1740,7 @@ FGenericPlatformProcess::EWaitAndForkResult FUnixPlatformProcess::WaitAndFork()
 					else if (NumForks > 0 && ChildPidAndSignal.SignalValue > 0 && ChildPidAndSignal.SignalValue <= NumForks)
 					{
 						UE_LOG(LogFork, Log, TEXT("[Parent] WaitAndFork child %d missing. This was NumForks child %d. Relaunching..."), ChildPidAndSignal.Pid, ChildPidAndSignal.SignalValue);
-						WaitAndForkSignalQueue.Enqueue(FForkSignalData(ChildPidAndSignal.SignalValue, FDateTime::UtcNow()));
+						WaitAndForkSignalQueue.Enqueue(FForkSignalData(ChildPidAndSignal.SignalValue, FDateTime()));
 					}
 					else
 					{
