SetDisplayString - Console command to replace a DisplayString in the LocalizationManager StringTable, to preview a string in the running editor

diff --git a/Engine/Source/Runtime/Core/Private/Internationalization/TextLocalizationManager.cpp b/Engine/Source/Runtime/Core/Private/Internationalization/TextLocalizationManager.cpp
--- a/Engine/Source/Runtime/Core/Private/Internationalization/TextLocalizationManager.cpp
+++ b/Engine/Source/Runtime/Core/Private/Internationalization/TextLocalizationManager.cpp
@@ -15,7 +15,6 @@
 #include "Misc/FileHelper.h"
 #include "Misc/Parse.h"
 #include "Misc/ScopeLock.h"
-#include "Containers/Ticker.h"
 #include "Misc/ScopeRWLock.h"
 #include "Misc/CommandLine.h"
 #include "Misc/LazySingleton.h"
@@ -74,6 +73,9 @@ static FAutoConsoleCommand CmdDumpLiveTable(
 	TEXT("Dumps the current live table state to the log, optionally filtering it based on wildcard arguments for 'Namespace', 'Key', or 'DisplayString', eg) -Key=Foo, or -DisplayString=\"This is some text\", or -Key=Bar*Baz -DisplayString=\"This is some other text\""), 
 	FConsoleCommandWithArgsDelegate::CreateStatic([](const TArray<FString>& Args)
 	{
+		// Rebuild the full string of arguments, since values within quotes may have been split on spaces
+		const FString Arguments = FString::Join(Args, TEXT(" "));
+
 		auto ParseOptionalStringArg = [](const TCHAR* Arg, const TCHAR* TokenName, TOptional<FString>& OutResult)
 		{
 			FString TmpResult;
@@ -90,15 +92,19 @@ static FAutoConsoleCommand CmdDumpLiveTable(
 		TOptional<FString> DisplayStringFilter;
 		TOptional<FString> DumpFile;
 
-		for (const FString& Arg : Args)
+		if (!ParseOptionalStringArg(*Arguments, TEXT("Namespace="), NamespaceFilter) &&
+			!ParseOptionalStringArg(*Arguments, TEXT("Key="), KeyFilter) &&
+			!ParseOptionalStringArg(*Arguments, TEXT("DisplayString="), DisplayStringFilter) && 
+			!ParseOptionalStringArg(*Arguments, TEXT("DumpFile="), DumpFile))
 		{
-			if (!ParseOptionalStringArg(*Arg, TEXT("Namespace="), NamespaceFilter) &&
-				!ParseOptionalStringArg(*Arg, TEXT("Key="), KeyFilter) &&
-				!ParseOptionalStringArg(*Arg, TEXT("DisplayString="), DisplayStringFilter) && 
-				!ParseOptionalStringArg(*Arg, TEXT("DumpFile="), DumpFile))
-			{
-				UE_LOG(LogLocalization, Warning, TEXT("Unknown argument '%s' passed to Localization.DumpLiveTable!"), *Arg);
-			}
+			UE_LOG(LogLocalization, Warning, TEXT("Unknown argument passed to Localization.DumpLiveTable!"));
+		}
+
+		// Block dumping all 500k strings which may crash the editor
+		if (!NamespaceFilter.IsSet() && !KeyFilter.IsSet() && !DisplayStringFilter.IsSet() && !DumpFile.IsSet())
+		{
+			UE_LOG(LogLocalization, Display, TEXT("No arguments provided, this would dump every string. Consider dumping to a file instead, or providing filter(s) for Namespace, Key, or DisplayString"));
+			return;
 		}
 
 		if (DumpFile.IsSet())
@@ -112,6 +118,50 @@ static FAutoConsoleCommand CmdDumpLiveTable(
 #endif
 		}
 	}));
+
+static FAutoConsoleCommand SetDisplayString(
+	TEXT("Localization.SetDisplayString"),
+	TEXT("Replaces DisplayString in the live table given required arguments: 'Namespace', 'Key', and 'DisplayString'"),
+	FConsoleCommandWithArgsDelegate::CreateStatic([](const TArray<FString>& Args)
+	{
+		// Rebuild the full string of arguments, since values within quotes may have been split on spaces
+		const FString Arguments = FString::Join(Args, TEXT(" "));
+
+		FString Namespace;
+		TOptional<FString> Key;
+		TOptional<FString> DisplayString;
+
+		FString Tmp;
+		if (FParse::Value(*Arguments, TEXT("Namespace="), Tmp))
+		{
+			Namespace = MoveTemp(Tmp);
+		}
+
+		if (FParse::Value(*Arguments, TEXT("Key="), Tmp))
+		{
+			Key = MoveTemp(Tmp);
+		}
+
+		if (FParse::Value(*Arguments, TEXT("DisplayString="), Tmp))
+		{
+			DisplayString = MoveTemp(Tmp);
+		}
+
+		// Namespace is optional, assumed to be empty if not provided
+		if (!Key.IsSet() || !DisplayString.IsSet())
+		{
+			UE_LOG(LogLocalization, Display, TEXT("Missing argument(s): Key and/or DisplayString"));
+			return;
+		}
+		// An empty DisplayString is allowed, but the argument for it must be provided: -DisplayString=""
+		if (Key.GetValue().IsEmpty())
+		{
+			UE_LOG(LogLocalization, Display, TEXT("Empty argument: Key"));
+			return;
+		}
+
+		FTextLocalizationManager::Get().ReplaceStringInLiveTable(Namespace, Key.GetValue(), DisplayString.GetValue());
+	}));
 #endif
 
 FString KeyifyTextId(const FTextId& TextId)
@@ -742,6 +792,29 @@ void FTextLocalizationManager::DumpLiveTable(const FString& OutputFilename, cons
 
 	FFileHelper::SaveStringToFile(DumpString, *OutputFilename, FFileHelper::EEncodingOptions::ForceUTF8WithoutBOM);
 }
+
+void FTextLocalizationManager::ReplaceStringInLiveTable(const FString& Namespace, const FString& Key, const FString& DisplayString)
+{
+	// Lock while updating the table
+	FScopeLock ScopeLock(&DisplayStringTableCS);
+
+	const FTextId TextId(Namespace, Key);
+
+	FDisplayStringEntry* LiveEntry = DisplayStringLookupTable.Find(TextId);
+
+	if (LiveEntry)
+	{
+		LiveEntry->DisplayString = MakeTextDisplayString(CopyTemp(DisplayString));
+		DirtyLocalRevisionForTextId(TextId);
+
+		UE_LOG(LogConsoleResponse, Display, TEXT("Updated string for Namespace='%s', Key='%s' to DisplayString='%s'"), *Namespace, *Key, *DisplayString);
+	}
+	else
+	{
+		UE_LOG(LogConsoleResponse, Display, TEXT("String not found for Namespace='%s', Key='%s'"), *Namespace, *Key);
+	}
+}
+
 #endif
 
 FString FTextLocalizationManager::GetRequestedLanguageName() const
@@ -1892,7 +1965,7 @@ void FTextLocalizationManager::DirtyTextRevision()
 	}
 	else
 	{
-		ExecuteOnGameThread(TEXT("OnTextRevisionChangedEventBroadcastGT"), []()
+		AsyncTask(ENamedThreads::GameThread, []()
 		{
 			FTextLocalizationManager::Get().OnTextRevisionChangedEvent.Broadcast();
 		});
diff --git a/Engine/Source/Runtime/Core/Public/Internationalization/TextLocalizationManager.h b/Engine/Source/Runtime/Core/Public/Internationalization/TextLocalizationManager.h
--- a/Engine/Source/Runtime/Core/Public/Internationalization/TextLocalizationManager.h
+++ b/Engine/Source/Runtime/Core/Public/Internationalization/TextLocalizationManager.h
@@ -164,6 +164,7 @@ private:
 public:
 	CORE_API void DumpLiveTable(const FString* NamespaceFilter = nullptr, const FString* KeyFilter = nullptr, const FString* DisplayStringFilter = nullptr, const FLogCategoryBase* CategoryOverride = nullptr) const;
 	CORE_API void DumpLiveTable(const FString& OutputFilename, const FString* NamespaceFilter = nullptr, const FString* KeyFilter = nullptr, const FString* DisplayStringFilter = nullptr) const;
+	CORE_API void ReplaceStringInLiveTable(const FString& Namespace, const FString& Key, const FString& DisplayString);
 #endif
 
 	/**
diff --git a/Engine/Source/Runtime/Engine/Private/UserInterface/Console.cpp b/Engine/Source/Runtime/Engine/Private/UserInterface/Console.cpp
--- a/Engine/Source/Runtime/Engine/Private/UserInterface/Console.cpp
+++ b/Engine/Source/Runtime/Engine/Private/UserInterface/Console.cpp
@@ -797,7 +797,8 @@ void UConsole::AppendInputText(const FString& Text)
 		int32 Character = **TextMod.Left(1);
 		TextMod.MidInline(1, MAX_int32, EAllowShrinking::No);
 
-		if (Character >= 0x20 && Character < 0x100)
+		// Discard control characters, and allow more than extended ASCII (0xff) for Arabic etc
+		if (Character >= 0x20)
 		{
 			TCHAR Temp[2];
 			Temp[0] = Character;
