FPackageId: Add SerializeForLog function.

diff --git a/Engine/Source/Runtime/Core/Private/IO/PackageId.cpp b/Engine/Source/Runtime/Core/Private/IO/PackageId.cpp
--- a/Engine/Source/Runtime/Core/Private/IO/PackageId.cpp
+++ b/Engine/Source/Runtime/Core/Private/IO/PackageId.cpp
@@ -5,6 +5,7 @@
 #include "Hash/CityHash.h"
 #include "Misc/Char.h"
 #include "HAL/LowLevelMemTracker.h"
+#include "Serialization/CompactBinaryWriter.h"
 #include "Serialization/StructuredArchiveAdapters.h"
 #include "Serialization/StructuredArchiveSlots.h"
 #include "Misc/ScopeRWLock.h"
@@ -60,3 +61,26 @@ void operator<<(FStructuredArchiveSlot Slot, FPackageId& Value)
 {
 	Slot << Value.Id;
 }
+
+void SerializeForLog(FCbWriter& Writer, const FPackageId& Value)
+{
+	Writer.BeginObject();
+	Writer.AddString(ANSITEXTVIEW("$type"), ANSITEXTVIEW("PackageId"));
+
+	TUtf8StringBuilder<FName::StringBufferSize> TextBuilder;
+	TextBuilder.Appendf("0x%llX", Value.Id);
+
+#if WITH_PACKAGEID_NAME_MAP
+	const FName Name = Value.GetName();
+	TextBuilder << " (" << Name << ")";
+#endif // WITH_PACKAGEID_NAME_MAP
+
+	Writer.AddString(ANSITEXTVIEW("$text"), TextBuilder);
+	Writer.AddInteger(ANSITEXTVIEW("Id"), Value.Id);
+
+#if WITH_PACKAGEID_NAME_MAP
+	Writer.AddString(ANSITEXTVIEW("Name"), WriteToUtf8String<FName::StringBufferSize>(Name));
+#endif // WITH_PACKAGEID_NAME_MAP
+
+	Writer.EndObject();
+}
diff --git a/Engine/Source/Runtime/Core/Public/IO/PackageId.h b/Engine/Source/Runtime/Core/Public/IO/PackageId.h
--- a/Engine/Source/Runtime/Core/Public/IO/PackageId.h
+++ b/Engine/Source/Runtime/Core/Public/IO/PackageId.h
@@ -71,6 +71,8 @@ public:
 
 	CORE_API friend void operator<<(FStructuredArchiveSlot Slot, FPackageId& Value);
 
+	CORE_API friend void SerializeForLog(FCbWriter& Writer, const FPackageId& Value);
+
 #if WITH_PACKAGEID_NAME_MAP
 	CORE_API FName GetName() const;
 #endif
