Replace the CSV_THREAD_HIGH_PRI define with a -csvProfilerHighPriority commandline option so we can easily enable for testing purposes

diff --git a/Engine/Source/Runtime/Core/Private/ProfilingDebugging/CsvProfiler.cpp b/Engine/Source/Runtime/Core/Private/ProfilingDebugging/CsvProfiler.cpp
--- a/Engine/Source/Runtime/Core/Private/ProfilingDebugging/CsvProfiler.cpp
+++ b/Engine/Source/Runtime/Core/Private/ProfilingDebugging/CsvProfiler.cpp
@@ -49,7 +49,6 @@
 #endif
 
 #define REPAIR_MARKER_STACKS 1
-#define CSV_THREAD_HIGH_PRI 0
 
 // Global CSV category (no prefix)
 FCsvCategory GGlobalCsvCategory(TEXT("GLOBAL"), true, true);
@@ -2543,11 +2542,17 @@ public:
 	FCsvProfilerProcessingThread(FCsvProfiler& InCsvProfiler)
 		: CsvProfiler(InCsvProfiler)
 	{
-#if CSV_THREAD_HIGH_PRI
-		Thread = FForkProcessHelper::CreateForkableThread(this, TEXT("CSVProfiler"), 0, TPri_Highest, FPlatformAffinity::GetTaskGraphThreadMask());
-#else
-		Thread = FForkProcessHelper::CreateForkableThread(this, TEXT("CSVProfiler"), 0, TPri_Lowest, FPlatformAffinity::GetTaskGraphBackgroundTaskMask());
+#if CSV_PROFILER_ALLOW_DEBUG_FEATURES
+		if (FParse::Param(FCommandLine::Get(), TEXT("csvProfilerHighPriority")))
+		{
+			Thread = FForkProcessHelper::CreateForkableThread(this, TEXT("CSVProfiler"), 0, TPri_Highest, FPlatformAffinity::GetTaskGraphThreadMask());
+		}
+		else
 #endif
+		{
+			Thread = FForkProcessHelper::CreateForkableThread(this, TEXT("CSVProfiler"), 0, TPri_Lowest, FPlatformAffinity::GetTaskGraphBackgroundTaskMask());
+		}
+	
 	}
 
 	virtual ~FCsvProfilerProcessingThread()
