Dont use defined relative base dir if the folder doesnt exists

diff --git a/Engine/Source/Runtime/Core/Private/Linux/LinuxPlatformProcess.cpp b/Engine/Source/Runtime/Core/Private/Linux/LinuxPlatformProcess.cpp
--- a/Engine/Source/Runtime/Core/Private/Linux/LinuxPlatformProcess.cpp
+++ b/Engine/Source/Runtime/Core/Private/Linux/LinuxPlatformProcess.cpp
@@ -30,7 +30,14 @@ const TCHAR* FLinuxPlatformProcess::BaseDir()
 
 #ifdef UE_RELATIVE_BASE_DIR
 		FString CollapseResult(CachedResult);
-		CollapseResult /= UE_RELATIVE_BASE_DIR;
+
+		// this may have been defined at compile time because we are in Restricted, but then we have been staged as a program, and then remapped out of Restricted
+		// so if we are already in a Binaries/Linux directory
+		if (IFileManager::Get().DirectoryExists(*FPaths::Combine(CollapseResult, UE_RELATIVE_BASE_DIR)))
+		{
+			CollapseResult /= UE_RELATIVE_BASE_DIR;
+		}
+
 		FPaths::CollapseRelativeDirectories(CollapseResult);
 		FCString::Strcpy(CachedResult, UNIX_MAX_PATH, *CollapseResult);
 #endif
diff --git a/Engine/Source/Runtime/Core/Private/Unix/UnixPlatformFile.cpp b/Engine/Source/Runtime/Core/Private/Unix/UnixPlatformFile.cpp
--- a/Engine/Source/Runtime/Core/Private/Unix/UnixPlatformFile.cpp
+++ b/Engine/Source/Runtime/Core/Private/Unix/UnixPlatformFile.cpp
@@ -816,15 +816,21 @@ FUnixMappedFileRegion::~FUnixMappedFileRegion()
 FString FUnixPlatformFile::NormalizeFilename(const TCHAR* Filename, bool bIsForWriting)
 {
 	FString Result(Filename);
+
+	// If we are already absolute return
+	if (!FPaths::IsRelative(Result))
+	{
+		return Result;
+	}
+
 	FPaths::NormalizeFilename(Result);
 	return FPaths::ConvertRelativePathToFull(Result);
 }
 
 FString FUnixPlatformFile::NormalizeDirectory(const TCHAR* Directory, bool bIsForWriting)
 {
-	FString Result(Directory);
-	FPaths::NormalizeDirectoryName(Result);
-	return FPaths::ConvertRelativePathToFull(Result);
+	// Both these functions do the same thing
+	return NormalizeFilename(Directory, bIsForWriting);
 }
 
 bool FUnixPlatformFile::FileExists(const TCHAR* Filename)
