Only treat errors as warnings when no critical error has happened

We are currently seeing a lot of out-of-memory crashes logged as warnings instead of errors. The root of the issue is that our validation code sets TreatErrorsAsWarnings while running on the game thread. When a background thread now runs out of memory, their errors are turned into warnings. In general, none of the "TreatErrorsAsWarnings" code should affect fatal errors, but there are codepaths for handling fatal errors that end up writing the failure out as an error

#rnx

diff --git a/Engine/Source/Runtime/Core/Private/Misc/FeedbackContext.cpp b/Engine/Source/Runtime/Core/Private/Misc/FeedbackContext.cpp
--- a/Engine/Source/Runtime/Core/Private/Misc/FeedbackContext.cpp
+++ b/Engine/Source/Runtime/Core/Private/Misc/FeedbackContext.cpp
@@ -91,7 +91,14 @@ ELogVerbosity::Type FFeedbackContext::ResolveVerbosity(ELogVerbosity::Type Verbo
 {
 	if (Verbosity == ELogVerbosity::Error && TreatErrorsAsWarnings)
 	{
-		return ELogVerbosity::Warning;
+		// Only override the verbosity of errors if we are not currently handling a fatal failure.
+		// Some of the codepaths invoked when handling critical/fatal errors might involve logging
+		// the failure as an error along with some additional information, and we don't want to
+		// suppress fatal failures.
+		if (!GIsCriticalError)
+		{
+			return ELogVerbosity::Warning;
+		}
 	}
 	if (Verbosity == ELogVerbosity::Warning && TreatWarningsAsErrors)
 	{
