Add missing header

diff --git a/Engine/Source/Runtime/Core/Private/Linux/LinuxPlatformProcess.cpp b/Engine/Source/Runtime/Core/Private/Linux/LinuxPlatformProcess.cpp
--- a/Engine/Source/Runtime/Core/Private/Linux/LinuxPlatformProcess.cpp
+++ b/Engine/Source/Runtime/Core/Private/Linux/LinuxPlatformProcess.cpp
@@ -6,6 +6,7 @@
 #include "CoreGlobals.h"
 #include "Logging/LogMacros.h"
 #include "Misc/Paths.h"
+#include "HAL/FileManager.h"
 
 const TCHAR* FLinuxPlatformProcess::BaseDir()
 {
