Speculative fix to stop TSAN warning about intentional debug stats race

#rnx

diff --git a/Engine/Source/Runtime/Core/Private/UObject/UnrealNames.cpp b/Engine/Source/Runtime/Core/Private/UObject/UnrealNames.cpp
--- a/Engine/Source/Runtime/Core/Private/UObject/UnrealNames.cpp
+++ b/Engine/Source/Runtime/Core/Private/UObject/UnrealNames.cpp
@@ -1016,9 +1016,9 @@ public:
 	}
 
 	uint32 Capacity() const	{ return CapacityMask + 1; }
-	uint32 NumCreated() const { return NumCreatedEntries; }
-	uint32 NumCreatedWide() const { return NumCreatedWideEntries; }
-	uint32 NumCreatedWithNumber() const { return NumCreatedWithNumberEntries; }
+	uint32 NumCreated() const { return NumCreatedEntries.load(std::memory_order_relaxed); }
+	uint32 NumCreatedWide() const { return NumCreatedWideEntries.load(std::memory_order_relaxed); }
+	uint32 NumCreatedWithNumber() const { return NumCreatedWithNumberEntries.load(std::memory_order_relaxed); }
 
 protected:
 	enum { LoadFactorQuotient = 9, LoadFactorDivisor = 10 }; // I.e. realloc slots when 90% full
@@ -1028,9 +1028,9 @@ protected:
 	uint32 CapacityMask = 0;
 	FNameSlot* Slots = nullptr;
 	FNameEntryAllocator* Entries = nullptr;
-	uint32 NumCreatedEntries = 0;
-	uint32 NumCreatedWideEntries = 0;
-	uint32 NumCreatedWithNumberEntries = 0;
+	std::atomic<uint32> NumCreatedEntries{0};
+	std::atomic<uint32> NumCreatedWideEntries{0};
+	std::atomic<uint32> NumCreatedWithNumberEntries{0};
 
 
 	template<ENameCase Sensitivity>
@@ -1370,8 +1370,8 @@ private:
 
 		ClaimSlot(Slot, FNameSlot(NewEntryId, Value.Hash.SlotProbeHash));
 
-		++NumCreatedEntries;
-		NumCreatedWideEntries += Value.Name.bIsWide;
+		NumCreatedEntries.fetch_add(1, std::memory_order_relaxed);
+		NumCreatedWideEntries.fetch_add(Value.Name.bIsWide, std::memory_order_relaxed);
 		
 		return NewEntryId;
 	}
@@ -1384,7 +1384,7 @@ private:
 
 		ClaimSlot(Slot, FNameSlot(NewEntryId, Value.Hash.SlotProbeHash));
 
-		++NumCreatedWithNumberEntries;
+		NumCreatedWithNumberEntries.fetch_add(1, std::memory_order_relaxed);;
 
 		return NewEntryId;
 	}
