Add overload of IPlatformFile::OpenRead that can return error messages and takes it's parameters in the form of a bitfield enum.

### Generic Platform File
- This overload returns a TValueOrError (aliased to FFileOpenResult for easier reading) with the value being a IFileHandle if the file was opened successfully and an error message if not.
- Instead of returning a raw pointer to a IFileHandle, this overload returns it wrapped in a TUniquePtr to better indicate to the caller how it is supposed to be disposed of.
- The error message is wrapped in a new class 'FFileSystemError' which can contain a custom error message provided by the engine (FString) and an optional system error code (int32) that can be provided if it was a system call that failed. We only resolve the system error code to a human readable string if the caller requests the error message so as not to inccur the cost if it's not going to be used.
- Instead of a list of bools, this overload takes a bitfield (EReadFlags) as a parameter which can be used to customize behavior.
-- In addition to the existing functionality (allowing other handles to open for write) we now provide the option (AllowDelete) for the file to be deleted or renamed and have the read handle still function until closed.

### WindowsPlatformFile
- Changed the original version of OpenRead to be an implementation for the new OpenRead overload and the old overload just calls this one and discards the additional info for now.

diff --git a/Engine/Source/Runtime/Core/Private/GenericPlatform/GenericPlatformFile.cpp b/Engine/Source/Runtime/Core/Private/GenericPlatform/GenericPlatformFile.cpp
--- a/Engine/Source/Runtime/Core/Private/GenericPlatform/GenericPlatformFile.cpp
+++ b/Engine/Source/Runtime/Core/Private/GenericPlatform/GenericPlatformFile.cpp
@@ -489,6 +489,26 @@ bool FGenericReadRequest::CheckForPrecache()
 	return false;
 }
 
+FFileOpenResult IPlatformFile::OpenRead(const TCHAR* Filename, IPlatformFile::EOpenReadFlags Flags)
+{
+	if (EnumHasAnyFlags(Flags, EOpenReadFlags::AllowDelete))
+	{
+		return MakeError(TEXTVIEW("The current IPlatformFile implementation does not support the EReadFlags::AllowDelete flag"));
+	}
+
+	const bool bWrite = EnumHasAnyFlags(Flags, EOpenReadFlags::AllowWrite);
+
+	IFileHandle* Handle = OpenRead(Filename, bWrite);
+	if (Handle != nullptr)
+	{
+		return MakeValue(Handle);
+	}
+	else
+	{
+		return MakeError(TEXTVIEW("[Unknown Error] Current IPlatformFile implementation does not support returning error messages"));
+	}
+}
+
 IAsyncReadFileHandle* IPlatformFile::OpenAsyncRead(const TCHAR* Filename)
 {
 	return new FGenericAsyncReadFileHandle(this, Filename);
diff --git a/Engine/Source/Runtime/Core/Private/Windows/WindowsPlatformFile.cpp b/Engine/Source/Runtime/Core/Private/Windows/WindowsPlatformFile.cpp
--- a/Engine/Source/Runtime/Core/Private/Windows/WindowsPlatformFile.cpp
+++ b/Engine/Source/Runtime/Core/Private/Windows/WindowsPlatformFile.cpp
@@ -1257,11 +1257,14 @@ IAsyncReadFileHandle* FWindowsPlatformFile::OpenAsyncRead(const TCHAR* Filename)
 }
 #endif
 
-IFileHandle* FWindowsPlatformFile::OpenRead(const TCHAR* Filename, bool bAllowWrite)
+FFileOpenResult FWindowsPlatformFile::OpenRead(const TCHAR* Filename, EOpenReadFlags Flags)
 {
-	uint32  Access    = GENERIC_READ;
-	uint32  WinFlags  = FILE_SHARE_READ | (bAllowWrite ? FILE_SHARE_WRITE : 0);
-	uint32  Create    = OPEN_EXISTING;
+	const uint32  Access	= GENERIC_READ;
+	const uint32  WinFlags  =	FILE_SHARE_READ |
+								(EnumHasAnyFlags(Flags, EOpenReadFlags::AllowWrite) ? FILE_SHARE_WRITE : 0) |
+								(EnumHasAnyFlags(Flags, EOpenReadFlags::AllowDelete) ? FILE_SHARE_DELETE : 0);
+	const uint32  Create	= OPEN_EXISTING;
+
 #define USE_OVERLAPPED_IO (!IS_PROGRAM && !WITH_EDITOR)		// Use straightforward synchronous I/O in cooker/editor
 
 	TRACE_PLATFORMFILE_BEGIN_OPEN(Filename);
@@ -1270,7 +1273,8 @@ IFileHandle* FWindowsPlatformFile::OpenRead(const TCHAR* Filename, bool bAllowWr
 	if (Handle != INVALID_HANDLE_VALUE)
 	{
 		TRACE_PLATFORMFILE_END_OPEN(Handle);
-		return new FAsyncBufferedFileReaderWindows(Handle, Access, WinFlags, FILE_FLAG_OVERLAPPED);
+		FAsyncBufferedFileReaderWindows* FileHandle = new FAsyncBufferedFileReaderWindows(Handle, Access, WinFlags, FILE_FLAG_OVERLAPPED);
+		return MakeValue(FileHandle);
 	}
 #else
 	HANDLE Handle = CreateFileW(*FNormalizedFilename(Filename), Access, WinFlags, NULL, Create, FILE_ATTRIBUTE_NORMAL, NULL);
@@ -1284,19 +1288,33 @@ IFileHandle* FWindowsPlatformFile::OpenRead(const TCHAR* Filename, bool bAllowWr
 		// double check that the handle is valid before returning it
 		if (FileHandle->IsValid())
 		{
-			return FileHandle;
+			return MakeValue(FileHandle);
 		}
 		else
 		{
 			delete FileHandle;
-
-			return nullptr;
+			return MakeError(TEXTVIEW("Failed to create FFileHandleWindows"));
 		}
 	}
-#endif
+#endif //USE_OVERLAPPED_IO
 	else
 	{
 		TRACE_PLATFORMFILE_FAIL_OPEN(Filename);
+		return MakeError(TEXTVIEW("Failed to open file for reading"), FPlatformMisc::GetLastError());
+	}
+}
+
+IFileHandle* FWindowsPlatformFile::OpenRead(const TCHAR* Filename, bool bAllowWrite)
+{
+	const EOpenReadFlags Flags = bAllowWrite ? EOpenReadFlags::AllowWrite : EOpenReadFlags::None;
+
+	FFileOpenResult Result = OpenRead(Filename, Flags);
+	if (Result.IsValid())
+	{
+		return Result.StealValue().Release();
+	}
+	else
+	{
 		return nullptr;
 	}
 }
diff --git a/Engine/Source/Runtime/Core/Private/Windows/WindowsPlatformFilePrivate.h b/Engine/Source/Runtime/Core/Private/Windows/WindowsPlatformFilePrivate.h
--- a/Engine/Source/Runtime/Core/Private/Windows/WindowsPlatformFilePrivate.h
+++ b/Engine/Source/Runtime/Core/Private/Windows/WindowsPlatformFilePrivate.h
@@ -64,6 +64,7 @@ public:
 	virtual IAsyncReadFileHandle* OpenAsyncRead(const TCHAR* Filename) override;
 #endif
 
+	virtual FFileOpenResult OpenRead(const TCHAR* Filename, EOpenReadFlags Flags) override;
 	virtual IFileHandle* OpenRead(const TCHAR* Filename, bool bAllowWrite = false) override;
 	virtual IFileHandle* OpenReadNoBuffering(const TCHAR* Filename, bool bAllowWrite = false) override;
 	virtual IFileHandle* OpenWrite(const TCHAR* Filename, bool bAppend = false, bool bAllowRead = false) override;
diff --git a/Engine/Source/Runtime/Core/Public/GenericPlatform/GenericPlatformFile.h b/Engine/Source/Runtime/Core/Public/GenericPlatform/GenericPlatformFile.h
--- a/Engine/Source/Runtime/Core/Public/GenericPlatform/GenericPlatformFile.h
+++ b/Engine/Source/Runtime/Core/Public/GenericPlatform/GenericPlatformFile.h
@@ -17,6 +17,7 @@
 #include "Misc/DateTime.h"
 #include "Misc/EnumClassFlags.h"
 #include "Templates/Function.h"
+#include "Templates/ValueOrError.h"
 
 class FArchive;
 class IAsyncReadFileHandle;
@@ -271,6 +272,60 @@ struct FFileJournalData
 	bool bIsDirectory : 1;
 };
 
+/** Stores custom error messages from the engine along with an optional system error code that can provider more detailed infomation */
+class FFileSystemError
+{
+public:
+	FFileSystemError() = delete;
+	UE_NONCOPYABLE(FFileSystemError);
+
+	explicit FFileSystemError(FString&& InErrorMessage, int32 InSystemErrorCode = 0)
+		: ErrorMessage(MoveTemp(InErrorMessage))
+		, SystemErrorCode(InSystemErrorCode)
+	{
+	}
+
+	explicit FFileSystemError(FStringView InErrorMessage, int32 InSystemErrorCode = 0)
+		: ErrorMessage(InErrorMessage)
+		, SystemErrorCode(InSystemErrorCode)
+	{
+	}
+
+	/** Return the error message, if a valid system error code was provided then this will be appended to the end of the message */
+	FString GetMessage()
+	{
+		if (SystemErrorCode != 0)
+		{
+			TCHAR FormattedErrorMsg[MAX_SPRINTF] = { 0 };
+			FPlatformMisc::GetSystemErrorMessage(FormattedErrorMsg, UE_ARRAY_COUNT(FormattedErrorMsg), SystemErrorCode);
+
+			return FString::Printf(TEXT("%s [%s (%d)]"), *ErrorMessage, FormattedErrorMsg, SystemErrorCode);
+		}
+		else
+		{
+			return ErrorMessage;
+		}
+	}
+
+	/** Return the error message, if a valid system error code was provided then this will be appended to the end of the message */
+	template <typename CharType>
+	friend inline TStringBuilderBase<CharType>& operator<<(TStringBuilderBase<CharType>& Builder, const FFileSystemError& Error)
+	{
+		TCHAR FormattedErrorMsg[MAX_SPRINTF] = { 0 };
+		FPlatformMisc::GetSystemErrorMessage(FormattedErrorMsg, UE_ARRAY_COUNT(FormattedErrorMsg), Error.SystemErrorCode);
+
+		Builder << Error.ErrorMessage << TEXT(" [") << FormattedErrorMsg << TEXT(" (") << Error.SystemErrorCode << TEXT(")]");
+		return Builder;
+	}
+
+private:
+	FString ErrorMessage;
+	int32 SystemErrorCode;
+};
+
+/** Data structure returned by calls to IPlatformFile::OpenRead */
+using FFileOpenResult = TValueOrError<TUniquePtr<IFileHandle>, FFileSystemError>;
+
 /**
 * File I/O Interface
 **/
@@ -426,12 +481,34 @@ public:
 	 */
 	virtual bool SetMarkOfTheWeb(FStringView Filename, bool bNewStatus, const FString* InSourceURL = nullptr) { return false; }
 
-	/** Attempt to open a file for reading.
+	/** Flags to be used when opening a file for reading via IPlatformFile::OpenRead */
+	enum class EOpenReadFlags : uint8
+	{
+		None = 0,
+		/** Allow other handles/processes to write to this file. This flag is needed to open files that are currently being written to as well */
+		AllowWrite	= 1 << 0,
+		/** Allow the file to be deleted or renamed while keeping this handle valid for reading. */
+		AllowDelete = 1 << 1
+	};
+
+	/**
+	 * Open a file handle for reading.
+	 * 
+	 * @param Filename	The file to be opened
+	 * @param Flags		Allows specialization of the open operation, @see EReadFlags
+	 * 
+	 * @return	The return value will either contain the valid file handle or an error message. @see FFileOpenResult
+	 */
+	[[nodiscard]] CORE_API virtual FFileOpenResult OpenRead(const TCHAR* Filename, EOpenReadFlags Flags);
+
+	/**
+	 * Attempt to open a file for reading.
+	 * Please consider using the new overload that takes EReadFlags instead of bools as parameters.
 	 *
 	 * @param Filename file to be opened
 	 * @param bAllowWrite (applies to certain platforms only) whether this file is allowed to be written to by other processes. This flag is needed to open files that are currently being written to as well.
 	 *
-	 * @return If successful will return a non-nullptr pointer. Close the file by delete'ing the handle.
+	 * @return If successful will return a non-nullptr pointer. Close the file by deleting the handle.
 	 */
 	virtual IFileHandle*	OpenRead(const TCHAR* Filename, bool bAllowWrite = false) = 0;
 
@@ -440,8 +517,7 @@ public:
 		return OpenRead(Filename, bAllowWrite);
 	}
 
-
-	/** Attempt to open a file for writing. If successful will return a non-nullptr pointer. Close the file by delete'ing the handle. **/
+	/** Attempt to open a file for writing. If successful will return a non-nullptr pointer. Close the file by deleting the handle. **/
 	virtual IFileHandle*	OpenWrite(const TCHAR* Filename, bool bAppend = false, bool bAllowRead = false) = 0;
 
 	/** Return true if the directory exists. **/
@@ -917,3 +993,6 @@ inline FFileJournalData::FFileJournalData()
 	, bIsDirectory(false)
 {
 }
+
+ENUM_CLASS_FLAGS(IPlatformFile::EOpenReadFlags);
+
