Remove operator bool from TNonNullPtr, as it will ensure if false
Use IsInitialized when needed

#rnx

diff --git a/Engine/Source/Runtime/Core/Public/Templates/NonNullPointer.h b/Engine/Source/Runtime/Core/Public/Templates/NonNullPointer.h
--- a/Engine/Source/Runtime/Core/Public/Templates/NonNullPointer.h
+++ b/Engine/Source/Runtime/Core/Public/Templates/NonNullPointer.h
@@ -190,6 +190,11 @@ public:
 	{
 		return Object != nullptr;
 	}
+	
+	/**
+	 * Use IsInitialized if needed
+	 */
+	explicit operator bool() const = delete;
 
 	////////////////////////////////////////////////////
 	// Start - intrusive TOptional<TNonNullPtr> state //
