Adding _CONDITIONAL variants to TRACE_CPUPROFILER and UE_TRACE macros.

diff --git a/Engine/Source/Runtime/Core/Public/ProfilingDebugging/CpuProfilerTrace.h b/Engine/Source/Runtime/Core/Public/ProfilingDebugging/CpuProfilerTrace.h
--- a/Engine/Source/Runtime/Core/Public/ProfilingDebugging/CpuProfilerTrace.h
+++ b/Engine/Source/Runtime/Core/Public/ProfilingDebugging/CpuProfilerTrace.h
@@ -281,20 +281,41 @@ struct FCpuProfilerTrace
 #define TRACE_CPUPROFILER_EVENT_MANUAL_IS_ENABLED() \
 	bool(CpuChannel)
 
+// Conditionally trace a scoped cpu timing event providing a static string (const ANSICHAR* or const TCHAR*)
+// as the scope name and a condition under which to create the trace. It will use the Cpu trace channel.
+// Example: TRACE_CPUPROFILER_EVENT_SCOPE_STR_CONDITIONAL("My Scoped Timer A", Condition)
+#define TRACE_CPUPROFILER_EVENT_SCOPE_STR_CONDITIONAL(NameStr, Condition) \
+	TRACE_CPUPROFILER_EVENT_DECLARE(PREPROCESSOR_JOIN(__CpuProfilerEventSpecId, __LINE__)); \
+	TRACE_CPUPROFILER_EVENT_SCOPE_USE(PREPROCESSOR_JOIN(__CpuProfilerEventSpecId, __LINE__), NameStr, PREPROCESSOR_JOIN(__CpuProfilerEventScope, __LINE__), (Condition)); 
+
 // Trace a scoped cpu timing event providing a static string (const ANSICHAR* or const TCHAR*)
 // as the scope name. It will use the Cpu trace channel.
 // Example: TRACE_CPUPROFILER_EVENT_SCOPE_STR("My Scoped Timer A")
 #define TRACE_CPUPROFILER_EVENT_SCOPE_STR(NameStr) \
+	TRACE_CPUPROFILER_EVENT_SCOPE_STR_CONDITIONAL(NameStr, true)
+
+// Conditionally trace a scoped cpu timing event providing a static string (const ANSICHAR* or const TCHAR*)
+// as the scope name and a trace channel and a condition.
+// Example: TRACE_CPUPROFILER_EVENT_SCOPE_ON_CHANNEL_STR_CONDITIONAL("My Scoped Timer A", CpuChannel, Condition)
+// Note: The event will be emitted only if both the given channel and CpuChannel is enabled.
+#define TRACE_CPUPROFILER_EVENT_SCOPE_ON_CHANNEL_STR_CONDITIONAL(NameStr, Channel, Condition) \
 	TRACE_CPUPROFILER_EVENT_DECLARE(PREPROCESSOR_JOIN(__CpuProfilerEventSpecId, __LINE__)); \
-	TRACE_CPUPROFILER_EVENT_SCOPE_USE(PREPROCESSOR_JOIN(__CpuProfilerEventSpecId, __LINE__), NameStr, PREPROCESSOR_JOIN(__CpuProfilerEventScope, __LINE__), true); 
+	TRACE_CPUPROFILER_EVENT_SCOPE_USE_ON_CHANNEL(PREPROCESSOR_JOIN(__CpuProfilerEventSpecId, __LINE__), NameStr, PREPROCESSOR_JOIN(__CpuProfilerEventScope, __LINE__), Channel, (Condition)); 
 
 // Trace a scoped cpu timing event providing a static string (const ANSICHAR* or const TCHAR*)
 // as the scope name and a trace channel.
 // Example: TRACE_CPUPROFILER_EVENT_SCOPE_ON_CHANNEL_STR("My Scoped Timer A", CpuChannel)
 // Note: The event will be emitted only if both the given channel and CpuChannel is enabled.
 #define TRACE_CPUPROFILER_EVENT_SCOPE_ON_CHANNEL_STR(NameStr, Channel) \
-	TRACE_CPUPROFILER_EVENT_DECLARE(PREPROCESSOR_JOIN(__CpuProfilerEventSpecId, __LINE__)); \
-	TRACE_CPUPROFILER_EVENT_SCOPE_USE_ON_CHANNEL(PREPROCESSOR_JOIN(__CpuProfilerEventSpecId, __LINE__), NameStr, PREPROCESSOR_JOIN(__CpuProfilerEventScope, __LINE__), Channel, true); 
+	TRACE_CPUPROFILER_EVENT_SCOPE_ON_CHANNEL_STR_CONDITIONAL(NameStr, Channel, true)
+
+// Conditionally trace a scoped cpu timing event providing a scope name (plain text) and a condition.
+// It will use the Cpu trace channel.
+// Example: TRACE_CPUPROFILER_EVENT_SCOPE_CONDITIONAL(MyScopedTimer::A, Condition)
+// Note: Do not use this macro with a static string because, in that case, additional quotes will
+//       be added around the event scope name.
+#define TRACE_CPUPROFILER_EVENT_SCOPE_CONDITIONAL(Name, Condition) \
+	TRACE_CPUPROFILER_EVENT_SCOPE_STR_CONDITIONAL(#Name, (Condition))
 
 // Trace a scoped cpu timing event providing a scope name (plain text).
 // It will use the Cpu trace channel.
@@ -302,7 +323,15 @@ struct FCpuProfilerTrace
 // Note: Do not use this macro with a static string because, in that case, additional quotes will
 //       be added around the event scope name.
 #define TRACE_CPUPROFILER_EVENT_SCOPE(Name) \
-	TRACE_CPUPROFILER_EVENT_SCOPE_STR(#Name)
+	TRACE_CPUPROFILER_EVENT_SCOPE_CONDITIONAL(Name, true)
+
+// Conditionally trace a scoped cpu timing event providing a scope name (plain text), a trace channel, and a condition.
+// Example: TRACE_CPUPROFILER_EVENT_SCOPE_ON_CHANNEL_CONDITIONAL(MyScopedTimer::A, CpuChannel, Condition)
+// Note: Do not use this macro with a static string because, in that case, additional quotes will
+//       be added around the event scope name.
+// Note: The event will be emitted only if both the given channel and CpuChannel is enabled.
+#define TRACE_CPUPROFILER_EVENT_SCOPE_ON_CHANNEL_CONDITIONAL(Name, Channel, Condition) \
+	TRACE_CPUPROFILER_EVENT_SCOPE_ON_CHANNEL_STR_CONDITIONAL(#Name, Channel, (Condition))
 
 // Trace a scoped cpu timing event providing a scope name (plain text) and a trace channel.
 // Example: TRACE_CPUPROFILER_EVENT_SCOPE_ON_CHANNEL(MyScopedTimer::A, CpuChannel)
@@ -310,7 +339,15 @@ struct FCpuProfilerTrace
 //       be added around the event scope name.
 // Note: The event will be emitted only if both the given channel and CpuChannel is enabled.
 #define TRACE_CPUPROFILER_EVENT_SCOPE_ON_CHANNEL(Name, Channel) \
-	TRACE_CPUPROFILER_EVENT_SCOPE_ON_CHANNEL_STR(#Name, Channel)
+	TRACE_CPUPROFILER_EVENT_SCOPE_ON_CHANNEL_CONDITIONAL(Name, Channel, true)
+
+// Conditionally trace a scoped cpu timing event providing a dynamic string (const ANSICHAR* or const TCHAR*)
+// as the scope name and a trace channel.
+// Example: TRACE_CPUPROFILER_EVENT_SCOPE_TEXT_ON_CHANNEL_CONDITIONAL(*MyScopedTimerNameString, CpuChannel, Condition)
+// Note: This macro has a larger overhead compared to macro that accepts a plain text name
+//       or a static string. Use it only if scope name really needs to be a dynamic string.
+#define TRACE_CPUPROFILER_EVENT_SCOPE_TEXT_ON_CHANNEL_CONDITIONAL(Name, Channel, Condition) \
+	FCpuProfilerTrace::FDynamicEventScope PREPROCESSOR_JOIN(__CpuProfilerEventScope, __LINE__)(Name, Channel, (Condition), __FILE__, __LINE__);
 
 // Trace a scoped cpu timing event providing a dynamic string (const ANSICHAR* or const TCHAR*)
 // as the scope name and a trace channel.
@@ -318,7 +355,15 @@ struct FCpuProfilerTrace
 // Note: This macro has a larger overhead compared to macro that accepts a plain text name
 //       or a static string. Use it only if scope name really needs to be a dynamic string.
 #define TRACE_CPUPROFILER_EVENT_SCOPE_TEXT_ON_CHANNEL(Name, Channel) \
-	FCpuProfilerTrace::FDynamicEventScope PREPROCESSOR_JOIN(__CpuProfilerEventScope, __LINE__)(Name, Channel, true, __FILE__, __LINE__);
+	TRACE_CPUPROFILER_EVENT_SCOPE_TEXT_ON_CHANNEL_CONDITIONAL(Name, Channel, true)
+
+// Conditionally trace a scoped cpu timing event providing a dynamic string (const ANSICHAR* or const TCHAR*)
+// as the scope name. It will use the Cpu trace channel.
+// Example: TRACE_CPUPROFILER_EVENT_SCOPE_TEXT_CONDITIONAL(*MyScopedTimerNameString, Condition)
+// Note: This macro has a larger overhead compared to macro that accepts a plain text name
+//       or a static string. Use it only if scope name really needs to be a dynamic string.
+#define TRACE_CPUPROFILER_EVENT_SCOPE_TEXT_CONDITIONAL(Name, Condition) \
+	FCpuProfilerTrace::FDynamicEventScope PREPROCESSOR_JOIN(__CpuProfilerEventScope, __LINE__)(Name, (Condition), __FILE__, __LINE__);
 
 // Trace a scoped cpu timing event providing a dynamic string (const ANSICHAR* or const TCHAR*)
 // as the scope name. It will use the Cpu trace channel.
@@ -326,10 +371,7 @@ struct FCpuProfilerTrace
 // Note: This macro has a larger overhead compared to macro that accepts a plain text name
 //       or a static string. Use it only if scope name really needs to be a dynamic string.
 #define TRACE_CPUPROFILER_EVENT_SCOPE_TEXT(Name) \
-	FCpuProfilerTrace::FDynamicEventScope PREPROCESSOR_JOIN(__CpuProfilerEventScope, __LINE__)(Name, true, __FILE__, __LINE__);
-
-#define TRACE_CPUPROFILER_EVENT_SCOPE_TEXT_CONDITIONAL(Name, Condition) \
-	FCpuProfilerTrace::FDynamicEventScope PREPROCESSOR_JOIN(__CpuProfilerEventScope, __LINE__)(Name, (Condition), __FILE__, __LINE__);
+	TRACE_CPUPROFILER_EVENT_SCOPE_TEXT_CONDITIONAL(Name, true)
 
 // Make sure all thread data has reached the destination.
 // Note: Can be useful to call this before entering a wait condition that might take a while.
@@ -344,10 +386,15 @@ struct FCpuProfilerTrace
 #define TRACE_CPUPROFILER_EVENT_MANUAL_START(EventNameStr)
 #define TRACE_CPUPROFILER_EVENT_MANUAL_IS_ENABLED() false
 #define TRACE_CPUPROFILER_EVENT_SCOPE_STR(NameStr)
+#define TRACE_CPUPROFILER_EVENT_SCOPE_STR_CONDITIONAL(NameStr, Condition)
 #define TRACE_CPUPROFILER_EVENT_SCOPE_ON_CHANNEL_STR(NameStr, Channel)
+#define TRACE_CPUPROFILER_EVENT_SCOPE_ON_CHANNEL_STR_CONDITIONAL(NameStr, Channel, Condition)
 #define TRACE_CPUPROFILER_EVENT_SCOPE(Name)
+#define TRACE_CPUPROFILER_EVENT_SCOPE_CONDITIONAL(Name, Condition)
 #define TRACE_CPUPROFILER_EVENT_SCOPE_ON_CHANNEL(Name, Channel)
+#define TRACE_CPUPROFILER_EVENT_SCOPE_ON_CHANNEL_CONDITIONAL(Name, Channel, Condition)
 #define TRACE_CPUPROFILER_EVENT_SCOPE_TEXT_ON_CHANNEL(Name, Channel)
+#define TRACE_CPUPROFILER_EVENT_SCOPE_TEXT_ON_CHANNEL_CONDITIONAL(Name, Channel, Condition)
 #define TRACE_CPUPROFILER_EVENT_SCOPE_TEXT(Name)
 #define TRACE_CPUPROFILER_EVENT_SCOPE_TEXT_CONDITIONAL(Name, Condition)
 #define TRACE_CPUPROFILER_EVENT_FLUSH()
diff --git a/Engine/Source/Runtime/TraceLog/Public/Trace/Detail/Trace.h b/Engine/Source/Runtime/TraceLog/Public/Trace/Detail/Trace.h
--- a/Engine/Source/Runtime/TraceLog/Public/Trace/Detail/Trace.h
+++ b/Engine/Source/Runtime/TraceLog/Public/Trace/Detail/Trace.h
@@ -154,6 +154,10 @@ class FChannel;
 			if (const auto& __restrict EventName = *UE_LAUNDER((F##LoggerName##EventName##Fields*)(&LogScope))) \
 				((void)EventName),
 
+#define TRACE_PRIVATE_LOG_PRELUDE_CONDITIONAL(EnterFunc, LoggerName, EventName, ChannelsExpr, Condition, ...) \
+	if (Condition) \
+		TRACE_PRIVATE_LOG_PRELUDE(EnterFunc, LoggerName, EventName, ChannelsExpr, __VA_ARGS__) 
+
 #define TRACE_PRIVATE_LOG_EPILOG() \
 				LogScope += LogScope
 
@@ -167,12 +171,24 @@ class FChannel;
 		PREPROCESSOR_JOIN(TheScope, __LINE__).SetActive(), \
 		TRACE_PRIVATE_LOG_EPILOG()
 
+#define TRACE_PRIVATE_LOG_SCOPED_CONDITIONAL(LoggerName, EventName, ChannelsExpr, Condition, ...) \
+	UE::Trace::Private::FScopedLogScope PREPROCESSOR_JOIN(TheScope, __LINE__); \
+	TRACE_PRIVATE_LOG_PRELUDE_CONDITIONAL(ScopedEnter, LoggerName, EventName, ChannelsExpr, Condition, ##__VA_ARGS__) \
+		PREPROCESSOR_JOIN(TheScope, __LINE__).SetActive(), \
+		TRACE_PRIVATE_LOG_EPILOG()
+
 #define TRACE_PRIVATE_LOG_SCOPED_T(LoggerName, EventName, ChannelsExpr, ...) \
 	UE::Trace::Private::FScopedStampedLogScope PREPROCESSOR_JOIN(TheScope, __LINE__); \
 	TRACE_PRIVATE_LOG_PRELUDE(ScopedStampedEnter, LoggerName, EventName, ChannelsExpr, ##__VA_ARGS__) \
 		PREPROCESSOR_JOIN(TheScope, __LINE__).SetActive(), \
 		TRACE_PRIVATE_LOG_EPILOG()
 
+#define TRACE_PRIVATE_LOG_SCOPED_T_CONDITIONAL(LoggerName, EventName, ChannelsExpr, Condition, ...) \
+	UE::Trace::Private::FScopedStampedLogScope PREPROCESSOR_JOIN(TheScope, __LINE__); \
+	TRACE_PRIVATE_LOG_PRELUDE_CONDITIONAL(ScopedStampedEnter, LoggerName, EventName, ChannelsExpr, Condition, ##__VA_ARGS__) \
+		PREPROCESSOR_JOIN(TheScope, __LINE__).SetActive(), \
+		TRACE_PRIVATE_LOG_EPILOG()
+
 #define TRACE_PRIVATE_GET_DEFINITION_TYPE_ID(LoggerName, EventName) \
 	F##LoggerName##EventName##Fields::GetUid()
 
@@ -232,10 +248,18 @@ class FChannel;
 	if (const auto& EventName = *(F##LoggerName##EventName##Dummy*)1) \
 		EventName
 
+#define TRACE_PRIVATE_LOG_SCOPED_CONDITIONAL(LoggerName, EventName, Condition, ...) \
+	if (const auto& EventName = *(F##LoggerName##EventName##Dummy*)1) \
+		EventName
+
 #define TRACE_PRIVATE_LOG_SCOPED_T(LoggerName, EventName, ...) \
 	if (const auto& EventName = *(F##LoggerName##EventName##Dummy*)1) \
 		EventName
 
+#define TRACE_PRIVATE_LOG_SCOPED_T_CONDITIONAL(LoggerName, EventName, Condition, ...) \
+	if (const auto& EventName = *(F##LoggerName##EventName##Dummy*)1) \
+		EventName
+
 #define TRACE_PRIVATE_GET_DEFINITION_TYPE_ID(LoggerName, EventName) \
 	0
 
diff --git a/Engine/Source/Runtime/TraceLog/Public/Trace/Trace.h b/Engine/Source/Runtime/TraceLog/Public/Trace/Trace.h
--- a/Engine/Source/Runtime/TraceLog/Public/Trace/Trace.h
+++ b/Engine/Source/Runtime/TraceLog/Public/Trace/Trace.h
@@ -194,17 +194,19 @@ UE_TRACE_API FChannel* FindChannel(FChannelId ChannelId) UE_TRACE_IMPL(nullptr);
 } // namespace UE
 
 ////////////////////////////////////////////////////////////////////////////////
-#define UE_TRACE_EVENT_DEFINE(LoggerName, EventName)					TRACE_PRIVATE_EVENT_DEFINE(LoggerName, EventName)
-#define UE_TRACE_EVENT_BEGIN(LoggerName, EventName, ...)				TRACE_PRIVATE_EVENT_BEGIN(LoggerName, EventName, ##__VA_ARGS__)
-#define UE_TRACE_EVENT_BEGIN_EXTERN(LoggerName, EventName, ...)			TRACE_PRIVATE_EVENT_BEGIN_EXTERN(LoggerName, EventName, ##__VA_ARGS__)
-#define UE_TRACE_EVENT_FIELD(FieldType, FieldName)						TRACE_PRIVATE_EVENT_FIELD(FieldType, FieldName)
-#define UE_TRACE_EVENT_REFERENCE_FIELD(RefLogger, RefEvent, FieldName)	TRACE_PRIVATE_EVENT_REFFIELD(RefLogger, RefEvent, FieldName)
-#define UE_TRACE_EVENT_END()											TRACE_PRIVATE_EVENT_END()
-#define UE_TRACE_LOG(LoggerName, EventName, ChannelsExpr, ...)			TRACE_PRIVATE_LOG(LoggerName, EventName, ChannelsExpr, ##__VA_ARGS__)
-#define UE_TRACE_LOG_SCOPED(LoggerName, EventName, ChannelsExpr, ...)	TRACE_PRIVATE_LOG_SCOPED(LoggerName, EventName, ChannelsExpr, ##__VA_ARGS__)
-#define UE_TRACE_LOG_SCOPED_T(LoggerName, EventName, ChannelsExpr, ...)	TRACE_PRIVATE_LOG_SCOPED_T(LoggerName, EventName, ChannelsExpr, ##__VA_ARGS__)
-#define UE_TRACE_GET_DEFINITION_TYPE_ID(LoggerName, EventName)			TRACE_PRIVATE_GET_DEFINITION_TYPE_ID(LoggerName, EventName)
-#define UE_TRACE_LOG_DEFINITION(LoggerName, EventName, Id, ChannelsExpr, ...) TRACE_PRIVATE_LOG_DEFINITION(LoggerName, EventName, Id, ChannelsExpr, ##__VA_ARGS__)
+#define UE_TRACE_EVENT_DEFINE(LoggerName, EventName)											TRACE_PRIVATE_EVENT_DEFINE(LoggerName, EventName)
+#define UE_TRACE_EVENT_BEGIN(LoggerName, EventName, ...)										TRACE_PRIVATE_EVENT_BEGIN(LoggerName, EventName, ##__VA_ARGS__)
+#define UE_TRACE_EVENT_BEGIN_EXTERN(LoggerName, EventName, ...)									TRACE_PRIVATE_EVENT_BEGIN_EXTERN(LoggerName, EventName, ##__VA_ARGS__)
+#define UE_TRACE_EVENT_FIELD(FieldType, FieldName)												TRACE_PRIVATE_EVENT_FIELD(FieldType, FieldName)
+#define UE_TRACE_EVENT_REFERENCE_FIELD(RefLogger, RefEvent, FieldName)							TRACE_PRIVATE_EVENT_REFFIELD(RefLogger, RefEvent, FieldName)
+#define UE_TRACE_EVENT_END()																	TRACE_PRIVATE_EVENT_END()
+#define UE_TRACE_LOG(LoggerName, EventName, ChannelsExpr, ...)									TRACE_PRIVATE_LOG(LoggerName, EventName, ChannelsExpr, ##__VA_ARGS__)
+#define UE_TRACE_LOG_SCOPED(LoggerName, EventName, ChannelsExpr, ...)							TRACE_PRIVATE_LOG_SCOPED(LoggerName, EventName, ChannelsExpr, ##__VA_ARGS__)
+#define UE_TRACE_LOG_SCOPED_CONDITIONAL(LoggerName, EventName, ChannelsExpr, Condition, ...)	TRACE_PRIVATE_LOG_SCOPED_CONDITIONAL(LoggerName, EventName, ChannelsExpr, Condition, ##__VA_ARGS__)
+#define UE_TRACE_LOG_SCOPED_T(LoggerName, EventName, ChannelsExpr, ...)							TRACE_PRIVATE_LOG_SCOPED_T(LoggerName, EventName, ChannelsExpr, ##__VA_ARGS__)
+#define UE_TRACE_LOG_SCOPED_T_CONDITIONAL(LoggerName, EventName, ChannelsExpr, Condition, ...)	TRACE_PRIVATE_LOG_SCOPED_T_CONDITIONAL(LoggerName, EventName, ChannelsExpr, Condition, ##__VA_ARGS__)
+#define UE_TRACE_GET_DEFINITION_TYPE_ID(LoggerName, EventName)									TRACE_PRIVATE_GET_DEFINITION_TYPE_ID(LoggerName, EventName)
+#define UE_TRACE_LOG_DEFINITION(LoggerName, EventName, Id, ChannelsExpr, ...)					TRACE_PRIVATE_LOG_DEFINITION(LoggerName, EventName, Id, ChannelsExpr, ##__VA_ARGS__)
 
 ////////////////////////////////////////////////////////////////////////////////
 #define UE_TRACE_CHANNEL(ChannelName, ...)				TRACE_PRIVATE_CHANNEL(ChannelName, ##__VA_ARGS__)
