Rename the internal non virtual overload of OpenRead to OpenReadInternal to avoid 'Woverloaded-virtual' warnings.

#rnx

- When 34707428 introduced a second virtual overload of OpenRead in the base class this clashed with the non-virtual overload in this one.
- Changed some true/false parameters to be local constants so that it's easier for the reader to see what is being passed into calls.

diff --git a/Engine/Source/Runtime/Core/Private/Android/AndroidPlatformFile.cpp b/Engine/Source/Runtime/Core/Private/Android/AndroidPlatformFile.cpp
--- a/Engine/Source/Runtime/Core/Private/Android/AndroidPlatformFile.cpp
+++ b/Engine/Source/Runtime/Core/Private/Android/AndroidPlatformFile.cpp
@@ -1785,17 +1785,19 @@ public:
 
 	virtual IFileHandle* OpenRead(const TCHAR* Filename, bool bAllowWrite = false) override
 	{
-		return OpenRead(Filename, false, bAllowWrite);
+		const bool bAllowLocal = false;
+
+		return OpenReadInternal(Filename, bAllowLocal, bAllowWrite);
 	}
 
-	IFileHandle* OpenRead(const TCHAR* Filename, bool AllowLocal, bool bAllowWrite)
+	IFileHandle* OpenReadInternal(const TCHAR* Filename, bool bAllowLocal, bool bAllowWrite)
 	{
 #if LOG_ANDROID_FILE
 		FPlatformMisc::LowLevelOutputDebugStringf(TEXT("FAndroidPlatformFile::OpenRead('%s')"), Filename);
 #endif
 		FString LocalPath;
 		FString AssetPath;
-		PathToAndroidPaths(LocalPath, AssetPath, Filename, AllowLocal);
+		PathToAndroidPaths(LocalPath, AssetPath, Filename, bAllowLocal);
 
 		if (IsLocal(LocalPath))
 		{
@@ -2379,12 +2381,13 @@ private:
 
 	void MountOBB(const TCHAR* Filename)
 	{
-		FFileHandleAndroid* File
-			= static_cast<FFileHandleAndroid*>(OpenRead(Filename, true, false));
+		const bool bAllowLocal = true;
+		const bool bAllowWrite = false;
+
+		FFileHandleAndroid* File = static_cast<FFileHandleAndroid*>(OpenReadInternal(Filename, bAllowLocal, bAllowWrite));
 		check(nullptr != File);
 		ZipResource.AddPatchFile(MakeShareable(File));
-		FPlatformMisc::LowLevelOutputDebugStringf(
-			TEXT("Mounted OBB '%s'"), Filename);
+		FPlatformMisc::LowLevelOutputDebugStringf(TEXT("Mounted OBB '%s'"), Filename);
 	}
 
 	AAssetManager* AssetMgr;
