[ConcurrentLinearAllocator]
 - Read member before calling the destructor to avoid undefined behavior

#rnx

diff --git a/Engine/Source/Runtime/Core/Public/Experimental/ConcurrentLinearAllocator.h b/Engine/Source/Runtime/Core/Public/Experimental/ConcurrentLinearAllocator.h
--- a/Engine/Source/Runtime/Core/Public/Experimental/ConcurrentLinearAllocator.h
+++ b/Engine/Source/Runtime/Core/Public/Experimental/ConcurrentLinearAllocator.h
@@ -456,9 +456,10 @@ public:
 				// this deletes complete blocks when the last allocation is freed
 				if (Header->NumAllocations.fetch_sub(1, std::memory_order_acq_rel) == 1)
 				{
+					const uintptr_t NextAllocationPtr = Header->NextAllocationPtr;
 					Header->~FBlockHeader();
 					MemoryTrace_UnmarkAllocAsHeap(uint64(Header), EMemoryTraceRootHeap::SystemMemory);
-					BlockAllocationTag::Allocator::Free(Header, Header->NextAllocationPtr - uintptr_t(Header));
+					BlockAllocationTag::Allocator::Free(Header, NextAllocationPtr - uintptr_t(Header));
 				}
 			}
 			else
@@ -473,10 +474,11 @@ public:
 				// this deletes complete blocks when the last allocation is freed
 				if (Header->NumAllocations.fetch_sub(1, std::memory_order_acq_rel) == 1)
 				{
+					const uintptr_t NextAllocationPtr = Header->NextAllocationPtr;
 					Header->~FBlockHeader();
-					ASAN_UNPOISON_MEMORY_REGION( Header, Header->NextAllocationPtr - uintptr_t(Header) );
+					ASAN_UNPOISON_MEMORY_REGION( Header, NextAllocationPtr - uintptr_t(Header) );
 					MemoryTrace_UnmarkAllocAsHeap(uint64(Header), EMemoryTraceRootHeap::SystemMemory);
-					BlockAllocationTag::Allocator::Free(Header, Header->NextAllocationPtr - uintptr_t(Header));
+					BlockAllocationTag::Allocator::Free(Header, NextAllocationPtr - uintptr_t(Header));
 				}
 			}
 		}
