SetDisplayString command: force widget invalidation to update text onscreen in live game

diff --git a/Engine/Source/Runtime/Core/Private/Internationalization/TextLocalizationManager.cpp b/Engine/Source/Runtime/Core/Private/Internationalization/TextLocalizationManager.cpp
--- a/Engine/Source/Runtime/Core/Private/Internationalization/TextLocalizationManager.cpp
+++ b/Engine/Source/Runtime/Core/Private/Internationalization/TextLocalizationManager.cpp
@@ -18,6 +18,7 @@
 #include "Misc/ScopeRWLock.h"
 #include "Misc/CommandLine.h"
 #include "Misc/LazySingleton.h"
+#include "Misc/OutputDeviceRedirector.h"
 #include "Internationalization/Culture.h"
 #include "Internationalization/Internationalization.h"
 #include "Internationalization/StringTableCore.h"
@@ -161,6 +162,12 @@ static FAutoConsoleCommand SetDisplayString(
 		}
 
 		FTextLocalizationManager::Get().ReplaceStringInLiveTable(Namespace, Key.GetValue(), DisplayString.GetValue());
+
+		// For live game, force widget invalidation to update the text onscreen. This isn't needed in the editor.
+		if (IConsoleObject* CObj = IConsoleManager::Get().FindConsoleObject(TEXT("Slate.TriggerInvalidate")))
+		{
+			CObj->AsCommand()->Execute(/*Args=*/TArray<FString>(), /*InWorld=*/nullptr, *GLog);
+		}
 	}));
 #endif
 
