Add IsNumeric string view utilities under UE::String

diff --git a/Engine/Source/Runtime/Core/Private/String/Numeric.cpp b/Engine/Source/Runtime/Core/Private/String/Numeric.cpp
new file mode 100644
--- /dev/null
+++ b/Engine/Source/Runtime/Core/Private/String/Numeric.cpp
@@ -0,0 +1,97 @@
+// Copyright Epic Games, Inc. All Rights Reserved.
+
+#include "String/Numeric.h"
+#include "Containers/StringView.h"
+
+namespace UE::String::Private
+{
+
+template <typename CharType>
+static bool IsNumeric(TStringView<CharType> View)
+{
+	if (View.IsEmpty())
+	{
+		return false;
+	}
+
+	const CharType* CurPtr = View.GetData();
+	const CharType* EndPtr = CurPtr + View.Len();
+
+	if (*CurPtr == '-' || *CurPtr == '+')
+	{
+		++CurPtr;
+	}
+
+	bool bHasDot = false;
+	while (CurPtr < EndPtr)
+	{
+		if (*CurPtr == '.')
+		{
+			if (bHasDot)
+			{
+				return false;
+			}
+
+			bHasDot = true;
+		}
+		else if (!FChar::IsDigit(*CurPtr))
+		{
+			return false;
+		}
+
+		++CurPtr;
+	}
+
+	return true;
+}
+
+template <typename CharType>
+static bool IsNumericOnlyDigits(TStringView<CharType> View)
+{
+	if (View.IsEmpty())
+	{
+		return false;
+	}
+
+	const CharType* CurPtr = View.GetData();
+	const CharType* EndPtr = CurPtr + View.Len();
+
+	while (CurPtr < EndPtr)
+	{
+		if (!FChar::IsDigit(*CurPtr))
+		{
+			return false;
+		}
+
+		++CurPtr;
+	}
+
+	return true;
+}
+
+} // UE::String::Private
+
+namespace UE::String
+{
+
+bool IsNumeric(FWideStringView View)
+{
+	return Private::IsNumeric(View);
+}
+
+bool IsNumeric(FUtf8StringView View)
+{
+	return Private::IsNumeric(View);
+}
+
+bool IsNumericOnlyDigits(FWideStringView View)
+{
+	return Private::IsNumericOnlyDigits(View);
+}
+
+bool IsNumericOnlyDigits(FUtf8StringView View)
+{
+	return Private::IsNumericOnlyDigits(View);
+}
+
+} // UE::String
diff --git a/Engine/Source/Runtime/Core/Public/String/Numeric.h b/Engine/Source/Runtime/Core/Public/String/Numeric.h
new file mode 100644
--- /dev/null
+++ b/Engine/Source/Runtime/Core/Public/String/Numeric.h
@@ -0,0 +1,30 @@
+// Copyright Epic Games, Inc. All Rights Reserved.
+
+#pragma once
+
+#include "Containers/StringFwd.h"
+
+namespace UE::String
+{
+
+/**
+ * Checks if the passed in string contains only digits, at most one dot, and optionally a +/- sign at the start.
+ *
+ * @param View The string to check.
+ *
+ * @return True if the string looks like a valid floating point number.
+ */
+CORE_API bool IsNumeric(FWideStringView View);
+CORE_API bool IsNumeric(FUtf8StringView View);
+
+/**
+ * Checks if the passed in string contains only digits.
+ *
+ * @param View The string to check.
+ *
+ * @return True if the string only contains numeric characters.
+ */
+CORE_API bool IsNumericOnlyDigits(FWideStringView View);
+CORE_API bool IsNumericOnlyDigits(FUtf8StringView View);
+
+}
