Fixed legacy TArray private namespace and updated typedefs to more readable aliases.

diff --git a/Engine/Source/Runtime/Core/Public/Containers/Array.h b/Engine/Source/Runtime/Core/Public/Containers/Array.h
--- a/Engine/Source/Runtime/Core/Public/Containers/Array.h
+++ b/Engine/Source/Runtime/Core/Public/Containers/Array.h
@@ -291,7 +291,7 @@ private:
 	IteratorType Iter;
 };
 
-namespace UE4Array_Private
+namespace UE::Core::Private
 {
 	// Simply forwards to an unqualified GetData(), but can be called from within a container or view
 	// where GetData() is already a member and so hides any others.
@@ -304,10 +304,10 @@ namespace UE4Array_Private
 	template <typename FromArrayType, typename ToArrayType>
 	constexpr bool CanMoveTArrayPointersBetweenArrayTypes()
 	{
-		typedef typename FromArrayType::AllocatorType FromAllocatorType;
-		typedef typename ToArrayType::AllocatorType ToAllocatorType;
-		typedef typename FromArrayType::ElementType   FromElementType;
-		typedef typename ToArrayType::ElementType   ToElementType;
+		using FromAllocatorType = typename FromArrayType::AllocatorType;
+		using ToAllocatorType   = typename ToArrayType::AllocatorType;
+		using FromElementType   = typename FromArrayType::ElementType;
+		using ToElementType     = typename ToArrayType::ElementType;
 
 		// Allocators must be equal or move-compatible...
 		if constexpr (std::is_same_v<FromAllocatorType, ToAllocatorType> || TCanMoveBetweenAllocators<FromAllocatorType, ToAllocatorType>::Value)
@@ -341,10 +341,7 @@ namespace UE4Array_Private
 
 	template <typename T>
 	constexpr bool TIsTArrayOrDerivedFromTArray_V = sizeof(ResolveIsTArrayPtr((T*)nullptr)) == 2;
-}
 
-namespace UE::Core::Private
-{
 	[[noreturn]] CORE_API void OnInvalidArrayNum(unsigned long long NewNum);
 }
 
@@ -367,9 +364,9 @@ class TArray
 	friend class TArray;
 
 public:
-	typedef typename InAllocatorType::SizeType SizeType;
-	typedef InElementType ElementType;
-	typedef InAllocatorType AllocatorType;
+	using SizeType      = typename InAllocatorType::SizeType ;
+	using ElementType   = InElementType;
+	using AllocatorType = InAllocatorType;
 
 private:
 	using USizeType = typename std::make_unsigned_t<SizeType>;
@@ -433,7 +430,7 @@ public:
 	template <
 		typename OtherElementType,
 		typename OtherAllocator
-		UE_REQUIRES(UE4Array_Private::TArrayElementsAreCompatible_V<ElementType, const OtherElementType&>)
+		UE_REQUIRES(UE::Core::Private::TArrayElementsAreCompatible_V<ElementType, const OtherElementType&>)
 	>
 	FORCEINLINE explicit TArray(const TArray<OtherElementType, OtherAllocator>& Other)
 	{
@@ -535,7 +532,7 @@ private:
 	template <typename FromArrayType, typename ToArrayType>
 	static FORCEINLINE void MoveOrCopy(ToArrayType& ToArray, FromArrayType& FromArray, SizeType PrevMax)
 	{
-		if constexpr (UE4Array_Private::CanMoveTArrayPointersBetweenArrayTypes<FromArrayType, ToArrayType>())
+		if constexpr (UE::Core::Private::CanMoveTArrayPointersBetweenArrayTypes<FromArrayType, ToArrayType>())
 		{
 			// Move
 
@@ -592,7 +589,7 @@ private:
 	template <typename FromArrayType, typename ToArrayType>
 	static FORCEINLINE void MoveOrCopyWithSlack(ToArrayType& ToArray, FromArrayType& FromArray, SizeType PrevMax, SizeType ExtraSlack)
 	{
-		if constexpr (UE4Array_Private::CanMoveTArrayPointersBetweenArrayTypes<FromArrayType, ToArrayType>())
+		if constexpr (UE::Core::Private::CanMoveTArrayPointersBetweenArrayTypes<FromArrayType, ToArrayType>())
 		{
 			// Move
 
@@ -636,7 +633,7 @@ public:
 	template <
 		typename OtherElementType,
 		typename OtherAllocator
-		UE_REQUIRES(UE4Array_Private::TArrayElementsAreCompatible_V<ElementType, OtherElementType&&>)
+		UE_REQUIRES(UE::Core::Private::TArrayElementsAreCompatible_V<ElementType, OtherElementType&&>)
 	>
 	FORCEINLINE explicit TArray(TArray<OtherElementType, OtherAllocator>&& Other)
 	{
@@ -652,7 +649,7 @@ public:
 	 */
 	template <
 		typename OtherElementType
-		UE_REQUIRES(UE4Array_Private::TArrayElementsAreCompatible_V<ElementType, OtherElementType&&>)
+		UE_REQUIRES(UE::Core::Private::TArrayElementsAreCompatible_V<ElementType, OtherElementType&&>)
 	>
 	TArray(TArray<OtherElementType, AllocatorType>&& Other, SizeType ExtraSlack)
 	{
@@ -2115,8 +2112,8 @@ public:
 		typename RangeType
 		UE_REQUIRES(
 			TIsContiguousContainer<RangeType>::Value &&
-			!UE4Array_Private::TIsTArrayOrDerivedFromTArray_V<std::remove_reference_t<RangeType>> &&
-			UE4Array_Private::TArrayElementsAreCompatible_V<ElementType, TElementType_T<RangeType>>
+			!UE::Core::Private::TIsTArrayOrDerivedFromTArray_V<std::remove_reference_t<RangeType>> &&
+			UE::Core::Private::TArrayElementsAreCompatible_V<ElementType, TElementType_T<RangeType>>
 		)
 	>
 	void Append(RangeType&& Source)
@@ -2134,7 +2131,7 @@ public:
 
 		// Allocate memory for the new elements.
 		SizeType Pos = AddUninitialized(SourceCount);
-		ConstructItems<ElementType>(GetData() + Pos, UE4Array_Private::GetDataHelper(Source), SourceCount);
+		ConstructItems<ElementType>(GetData() + Pos, UE::Core::Private::GetDataHelper(Source), SourceCount);
 	}
 
 	/**
@@ -2922,8 +2919,8 @@ public:
 	}
 
 	// Iterators
-	typedef TIndexedContainerIterator<      TArray,       ElementType, SizeType> TIterator;
-	typedef TIndexedContainerIterator<const TArray, const ElementType, SizeType> TConstIterator;
+	using TIterator      = TIndexedContainerIterator<      TArray,       ElementType, SizeType>;
+	using TConstIterator = TIndexedContainerIterator<const TArray, const ElementType, SizeType>;
 
 	/**
 	 * Creates an iterator for the contents of this array
@@ -2946,15 +2943,15 @@ public:
 	}
 
 	#if TARRAY_RANGED_FOR_CHECKS
-		typedef TCheckedPointerIterator<      ElementType, SizeType, false> RangedForIteratorType;
-		typedef TCheckedPointerIterator<const ElementType, SizeType, false> RangedForConstIteratorType;
-		typedef TCheckedPointerIterator<      ElementType, SizeType, true>  RangedForReverseIteratorType;
-		typedef TCheckedPointerIterator<const ElementType, SizeType, true>  RangedForConstReverseIteratorType;
+		using RangedForIteratorType             = TCheckedPointerIterator<      ElementType, SizeType, false>;
+		using RangedForConstIteratorType        = TCheckedPointerIterator<const ElementType, SizeType, false>;
+		using RangedForReverseIteratorType      = TCheckedPointerIterator<      ElementType, SizeType, true>;
+		using RangedForConstReverseIteratorType = TCheckedPointerIterator<const ElementType, SizeType, true>;
 	#else
-		typedef                               ElementType* RangedForIteratorType;
-		typedef                         const ElementType* RangedForConstIteratorType;
-		typedef TReversePointerIterator<      ElementType> RangedForReverseIteratorType;
-		typedef TReversePointerIterator<const ElementType> RangedForConstReverseIteratorType;
+		using RangedForIteratorType             =                               ElementType*;
+		using RangedForConstIteratorType        =                         const ElementType*;
+		using RangedForReverseIteratorType      = TReversePointerIterator<      ElementType>;
+		using RangedForConstReverseIteratorType = TReversePointerIterator<const ElementType>;
 	#endif
 
 public:
@@ -3703,7 +3700,7 @@ struct TArrayPrivateFriend
 		A.CountBytes(Ar);
 
 		// For net archives, limit serialization to 16MB, to protect against excessive allocation
-		typedef typename AllocatorType::SizeType SizeType;
+		using SizeType = typename AllocatorType::SizeType;
 		constexpr SizeType MaxNetArraySerialize = (16 * 1024 * 1024) / sizeof(ElementType);
 		SizeType SerializeNum = Ar.IsLoading() ? 0 : A.ArrayNum;
 
