Fix concurrent access to PendingModificationPlugins

#rnx

diff --git a/Engine/Source/Runtime/Core/Private/Misc/ConfigCacheIni.cpp b/Engine/Source/Runtime/Core/Private/Misc/ConfigCacheIni.cpp
--- a/Engine/Source/Runtime/Core/Private/Misc/ConfigCacheIni.cpp
+++ b/Engine/Source/Runtime/Core/Private/Misc/ConfigCacheIni.cpp
@@ -5890,6 +5890,7 @@ void FConfigCacheIni::AddPluginToAllBranches(FName PluginName, FConfigModificati
 	GConfig->AddPluginToBranches(PluginName, ModificationTracker);
 	
 #if ALLOW_OTHER_PLATFORM_CONFIG
+	FScopeLock Lock(&ConfigForPlatformLock);
 	// need to walk over the other platforms without calling ForPlatform because that could end up loading pending plugins
 	for (auto Pair : ConfigForPlatform)
 	{
