MallocBinned2/3 internal renaming necessary to remove all the confusion around names that don't reflect the purpose. It also adds more const correctness to the code and fixes a bunch of defines to be in line with UE naming convention

diff --git a/Engine/Source/Runtime/Core/Private/HAL/MallocBinned2.cpp b/Engine/Source/Runtime/Core/Private/HAL/MallocBinned2.cpp
--- a/Engine/Source/Runtime/Core/Private/HAL/MallocBinned2.cpp
+++ b/Engine/Source/Runtime/Core/Private/HAL/MallocBinned2.cpp
@@ -45,67 +45,66 @@ private:
 	bool bPop;
 	const ANSICHAR* StatString;
 };
-#define NOALLOC_SCOPE_CYCLE_COUNTER(Stat) FNoAllocScopeCycleCounter NoAllocCycleCounter_##Stat(#Stat)
+#	define NOALLOC_SCOPE_CYCLE_COUNTER(Stat) FNoAllocScopeCycleCounter NoAllocCycleCounter_##Stat(#Stat)
 #else
-#define NOALLOC_SCOPE_CYCLE_COUNTER(Stat)
+#	define NOALLOC_SCOPE_CYCLE_COUNTER(Stat)
 #endif // FRAMEPRO_ENABLED
 
 
 PRAGMA_DISABLE_UNSAFE_TYPECAST_WARNINGS
 
-#if BINNED2_ALLOW_RUNTIME_TWEAKING
-
-int32 GMallocBinned2PerThreadCaches = DEFAULT_GMallocBinned2PerThreadCaches;
-static FAutoConsoleVariableRef GMallocBinned2PerThreadCachesCVar(
-	TEXT("MallocBinned2.PerThreadCaches"),
-	GMallocBinned2PerThreadCaches,
-	TEXT("Enables per-thread caches of small (<= 32768 byte) allocations from FMallocBinned2")
-	);
-
-int32 GMallocBinned2MaxBundlesBeforeRecycle = BINNED2_MAX_GMallocBinned2MaxBundlesBeforeRecycle;
-static FAutoConsoleVariableRef GMallocBinned2MaxBundlesBeforeRecycleCVar(
-	TEXT("MallocBinned2.BundleRecycleCount"),
-	GMallocBinned2MaxBundlesBeforeRecycle,
-	TEXT("Number of freed bundles in the global recycler before it returns them to the system, per-block size. Limited by BINNED2_MAX_GMallocBinned2MaxBundlesBeforeRecycle (currently 4)")
+#if UE_MB2_ALLOW_RUNTIME_TWEAKING
+	int32 GBinned2PerThreadCaches = UE_DEFAULT_GBinned2PerThreadCaches;
+	static FAutoConsoleVariableRef GMallocBinned2PerThreadCachesCVar(
+		TEXT("MallocBinned2.PerThreadCaches"),
+		GBinned2PerThreadCaches,
+		TEXT("Enables per-thread caches of small (<= 32768 byte) allocations from FMallocBinned2")
+		);
+
+	int32 GBinned2MaxBundlesBeforeRecycle = UE_DEFAULT_GBinned2MaxBundlesBeforeRecycle;
+	static FAutoConsoleVariableRef GMallocBinned2MaxBundlesBeforeRecycleCVar(
+		TEXT("MallocBinned2.BundleRecycleCount"),
+		GBinned2MaxBundlesBeforeRecycle,
+		TEXT("Number of freed bundles in the global recycler before it returns them to the system, per-block size. Limited by UE_DEFAULT_GBinned2MaxBundlesBeforeRecycle (currently 4)")
+		);
+
+	int32 GBinned2AllocExtra = UE_DEFAULT_GBinned2AllocExtra;
+	static FAutoConsoleVariableRef GMallocBinned2AllocExtraCVar(
+		TEXT("MallocBinned2.AllocExtra"),
+		GBinned2AllocExtra,
+		TEXT("When we do acquire the lock, how many bins cached in TLS caches. In no case will we grab more than a page.")
+		);
+
+	int32 GBinned2MoveOSFreesOffTimeCriticalThreads = 1;
+	static FAutoConsoleVariableRef GGMallocBinned2MoveOSFreesOffTimeCriticalThreadsCVar(
+		TEXT("MallocBinned2.MoveOSFreesOffTimeCriticalThreads"),
+		GBinned2MoveOSFreesOffTimeCriticalThreads,
+		TEXT("When the OS needs to free memory hint to the underlying cache that we are on a time critical thread, it may decide to delay the free for a non time critical thread")
 	);
-
-int32 GMallocBinned2AllocExtra = DEFAULT_GMallocBinned2AllocExtra;
-static FAutoConsoleVariableRef GMallocBinned2AllocExtraCVar(
-	TEXT("MallocBinned2.AllocExtra"),
-	GMallocBinned2AllocExtra,
-	TEXT("When we do acquire the lock, how many blocks cached in TLS caches. In no case will we grab more than a page.")
-	);
-
-int32 GMallocBinned2MoveOSFreesOffTimeCriticalThreads = 1;
-static FAutoConsoleVariableRef GGMallocBinned2MoveOSFreesOffTimeCriticalThreadsCVar(
-	TEXT("MallocBinned2.MoveOSFreesOffTimeCriticalThreads"),
-	GMallocBinned2MoveOSFreesOffTimeCriticalThreads,
-	TEXT("When the OS needs to free memory hint to the underlying cache that we are on a time critical thread, it may decide to delay the free for a non time critical thread")
-);
 #endif
 
 
-#if BINNED2_ALLOCATOR_STATS
-TAtomic<int64> AllocatedSmallPoolMemory(0); // memory that's requested to be allocated by the game
-TAtomic<int64> AllocatedOSSmallPoolMemory(0);
+#if UE_MB2_ALLOCATOR_STATS
+	TAtomic<int64> AllocatedSmallPoolMemory(0); // memory that's requested to be allocated by the game
+	TAtomic<int64> AllocatedOSSmallPoolMemory(0);
 
-TAtomic<int64> AllocatedLargePoolMemory(0); // memory requests to the OS which don't fit in the small pool
-TAtomic<int64> AllocatedLargePoolMemoryWAlignment(0); // when we allocate at OS level we need to align to a size
+	TAtomic<int64> AllocatedLargePoolMemory(0); // memory requests to the OS which don't fit in the small pool
+	TAtomic<int64> AllocatedLargePoolMemoryWAlignment(0); // when we allocate at OS level we need to align to a size
 
-int64 Binned2PoolInfoMemory = 0;
-int64 Binned2HashMemory = 0;
+	int64 Binned2PoolInfoMemory = 0;
+	int64 Binned2HashMemory = 0;
 #endif
 
-#if BINNED2_ALLOCATOR_STATS_VALIDATION
-int64 AllocatedSmallPoolMemoryValidation = 0;
-FCriticalSection ValidationCriticalSection;
-int32 RecursionCounter = 0;
+#if UE_MB2_ALLOCATOR_STATS_VALIDATION
+	int64 AllocatedSmallPoolMemoryValidation = 0;
+	FCriticalSection ValidationCriticalSection;
+	int32 RecursionCounter = 0;
 #endif
 
-// Block sizes are based around getting the maximum amount of allocations per pool, with as little alignment waste as possible.
-// Block sizes should be close to even divisors of the system page size, and well distributed.
+// Bin sizes are based around getting the maximum amount of allocations per pool, with as little alignment waste as possible.
+// Bin sizes should be close to even divisors of the system page size, and well distributed.
 // They must be 16-byte aligned as well.
-static constexpr uint16 SmallBlockSizes[] =
+static constexpr uint16 SmallBinSizes[] =
 {
 	16, 32, 48, 64, 80, 96, 128,
 	160, 192, 224, 256, 288, 320, 384, 448,
@@ -116,29 +115,20 @@ static constexpr uint16 SmallBlockSizes[] =
 };
 
 MS_ALIGN(PLATFORM_CACHE_LINE_SIZE) static uint8 UnusedAlignPadding[PLATFORM_CACHE_LINE_SIZE] GCC_ALIGN(PLATFORM_CACHE_LINE_SIZE) = { 0 };
-uint16 FMallocBinned2::SmallBlockSizesReversed[BINNED2_SMALL_POOL_COUNT] = { 0 };
+uint16 FMallocBinned2::SmallBinSizesReversed[UE_MB2_SMALL_POOL_COUNT] = { 0 };
 uint32 FMallocBinned2::OsAllocationGranularity = 0;
 uint32 FMallocBinned2::PageSize = 0;
 FMallocBinned2* FMallocBinned2::MallocBinned2 = nullptr;
 // Mapping of sizes to small table indices
-uint8 FMallocBinned2::MemSizeToIndex[1 + (BINNED2_MAX_SMALL_POOL_SIZE >> BINNED2_MINIMUM_ALIGNMENT_SHIFT)] = { 0 };
-
-FMallocBinned2::FPoolList::FPoolList()
-	: Front(nullptr)
-{
-}
+uint8 FMallocBinned2::MemSizeToPoolIndex[1 + (UE_MB2_MAX_SMALL_POOL_SIZE >> UE_MB2_MINIMUM_ALIGNMENT_SHIFT)] = { 0 };
 
-FMallocBinned2::FPoolTable::FPoolTable()
-	: BlockSize(0)
-{
-}
 
 struct FMallocBinned2::FPoolInfo
 {
 	static FORCEINLINE bool IsSupportedSize(SIZE_T Size)
 	{
-		bool bResult = IsAligned(Size, BINNED2_MINIMUM_ALIGNMENT);
-		bResult = bResult && (Size >> BINNED2_MINIMUM_ALIGNMENT_SHIFT) <= SIZE_T(MAX_uint32);
+		bool bResult = IsAligned(Size, UE_MB2_MINIMUM_ALIGNMENT);
+		bResult = bResult && (Size >> UE_MB2_MINIMUM_ALIGNMENT_SHIFT) <= SIZE_T(MAX_uint32);
 		return bResult;
 	}
 
@@ -156,11 +146,6 @@ struct FMallocBinned2::FPoolInfo
 	FPoolInfo*  Next;           // Pointer to next pool
 	FPoolInfo** PtrToPrevNext;  // Pointer to whichever pointer points to this pool
 
-#if PLATFORM_32BITS
-/** Explicit padding for 32 bit builds */
-private: uint8 Padding[12]; // 32
-#endif
-
 public:
 	FPoolInfo()
 		: Taken(0)
@@ -171,6 +156,7 @@ public:
 		, PtrToPrevNext(nullptr)
 	{
 	}
+
 	void CheckCanary(ECanary ShouldBe) const
 	{
 		if (Canary != ShouldBe)
@@ -178,6 +164,7 @@ public:
 			UE_LOG(LogMemory, Fatal, TEXT("MallocBinned2 Corruption Canary was 0x%x, should be 0x%x"), int32(Canary), int32(ShouldBe));
 		}
 	}
+
 	void SetCanary(ECanary ShouldBe, bool bPreexisting, bool bGuaranteedToBeNew)
 	{
 		if (bPreexisting)
@@ -214,24 +201,25 @@ public:
 		}
 		Canary = ShouldBe;
 	}
-	bool HasFreeRegularBlock() const
+
+	bool HasFreeBin() const
 	{
 		CheckCanary(ECanary::FirstFreeBlockIsPtr);
-		return FirstFreeBlock && FirstFreeBlock->GetNumFreeRegularBlocks() != 0;
+		return FirstFreeBlock && FirstFreeBlock->GetNumFreeBins() != 0;
 	}
 
-	void* AllocateRegularBlock()
+	void* AllocateBin()
 	{
-		check(HasFreeRegularBlock());
+		check(HasFreeBin());
 		++Taken;
-		void* Result = FirstFreeBlock->AllocateRegularBlock();
+		void* Result = FirstFreeBlock->AllocateBin();
 		ExhaustPoolIfNecessary();
 		return Result;
 	}
 
 	SIZE_T GetOSRequestedBytes() const
 	{
-		return SIZE_T(AllocSize) << BINNED2_MINIMUM_ALIGNMENT_SHIFT;
+		return SIZE_T(AllocSize) << UE_MB2_MINIMUM_ALIGNMENT_SHIFT;
 	}
 
 	UPTRINT GetOsAllocatedBytes() const
@@ -247,7 +235,7 @@ public:
 		checkSlow(InAllocatedBytes >= InRequestedBytes); // We must be allocating at least as much as we requested
 		checkSlow(IsSupportedSize(InRequestedBytes));    // We must be allocating a size we can store
 
-		AllocSize      = uint32(InRequestedBytes >> BINNED2_MINIMUM_ALIGNMENT_SHIFT);
+		AllocSize      = uint32(InRequestedBytes >> UE_MB2_MINIMUM_ALIGNMENT_SHIFT);
 		FirstFreeBlock = (FFreeBlock*)InAllocatedBytes;
 	}
 
@@ -274,11 +262,11 @@ public:
 private:
 	void ExhaustPoolIfNecessary()
 	{
-		if (FirstFreeBlock->GetNumFreeRegularBlocks() == 0)// && (FirstFreeBlock->GetShortBlockSize() == 0 || !FirstFreeBlock->HasFreeShortBlock()))
+		if (FirstFreeBlock->GetNumFreeBins() == 0)
 		{
-			FirstFreeBlock = (FFreeBlock*)FirstFreeBlock->NextFreeBlock;
+			FirstFreeBlock = FirstFreeBlock->NextFreeBlock;
 		}
-		check(!FirstFreeBlock || FirstFreeBlock->GetNumFreeRegularBlocks() != 0);
+		check(!FirstFreeBlock || FirstFreeBlock->GetNumFreeBins() != 0);
 	}
 };
 
@@ -302,13 +290,13 @@ struct FMallocBinned2::Private
 		 */
 		auto CreatePoolArray = [&Allocator](uint64 NumPools)
 		{
-			uint64 PoolArraySize = NumPools * sizeof(FPoolInfo);
+			const uint64 PoolArraySize = NumPools * sizeof(FPoolInfo);
 
 			void* Result;
 			{
 				LLM_PLATFORM_SCOPE(ELLMTag::FMalloc);
 				Result = FPlatformMemory::BinnedAllocFromOS(PoolArraySize);
-#if BINNED2_ALLOCATOR_STATS
+#if UE_MB2_ALLOCATOR_STATS
 				Binned2PoolInfoMemory += PoolArraySize;
 #endif
 			}
@@ -356,7 +344,7 @@ struct FMallocBinned2::Private
 			{
 				LLM_PLATFORM_SCOPE(ELLMTag::FMalloc);
 				Allocator.HashBucketFreeList = (PoolHashBucket*)FPlatformMemory::BinnedAllocFromOS(FMallocBinned2::PageSize);
-#if BINNED2_ALLOCATOR_STATS
+#if UE_MB2_ALLOCATOR_STATS
 				Binned2HashMemory += FMallocBinned2::PageSize;
 #endif
 			}
@@ -423,7 +411,7 @@ struct FMallocBinned2::Private
 	{
 		bool PushBundle(uint32 InPoolIndex, FBundleNode* InBundle)
 		{
-			uint32 NumCachedBundles = FMath::Min<uint32>(GMallocBinned2MaxBundlesBeforeRecycle, BINNED2_MAX_GMallocBinned2MaxBundlesBeforeRecycle);
+			const uint32 NumCachedBundles = FMath::Min<uint32>(GBinned2MaxBundlesBeforeRecycle, UE_DEFAULT_GBinned2MaxBundlesBeforeRecycle);
 			for (uint32 Slot = 0; Slot < NumCachedBundles; Slot++)
 			{
 				if (!Bundles[InPoolIndex].FreeBundles[Slot])
@@ -439,7 +427,7 @@ struct FMallocBinned2::Private
 
 		FBundleNode* PopBundle(uint32 InPoolIndex)
 		{
-			uint32 NumCachedBundles = FMath::Min<uint32>(GMallocBinned2MaxBundlesBeforeRecycle, BINNED2_MAX_GMallocBinned2MaxBundlesBeforeRecycle);
+			const uint32 NumCachedBundles = FMath::Min<uint32>(GBinned2MaxBundlesBeforeRecycle, UE_DEFAULT_GBinned2MaxBundlesBeforeRecycle);
 			for (uint32 Slot = 0; Slot < NumCachedBundles; Slot++)
 			{
 				FBundleNode* Result = Bundles[InPoolIndex].FreeBundles[Slot];
@@ -457,23 +445,23 @@ struct FMallocBinned2::Private
 	private:
 		struct FPaddedBundlePointer
 		{
-			FBundleNode* FreeBundles[BINNED2_MAX_GMallocBinned2MaxBundlesBeforeRecycle];
-#define BUNDLE_PADDING (PLATFORM_CACHE_LINE_SIZE - sizeof(FBundleNode*) * BINNED2_MAX_GMallocBinned2MaxBundlesBeforeRecycle)
-#if (4 + (4 * PLATFORM_64BITS)) * BINNED2_MAX_GMallocBinned2MaxBundlesBeforeRecycle < PLATFORM_CACHE_LINE_SIZE
-			uint8 Padding[BUNDLE_PADDING];
+			FBundleNode* FreeBundles[UE_DEFAULT_GBinned2MaxBundlesBeforeRecycle];
+#if (4 + (4 * PLATFORM_64BITS)) * UE_DEFAULT_GBinned2MaxBundlesBeforeRecycle < PLATFORM_CACHE_LINE_SIZE
+#	define UE_MB2_BUNDLE_PADDING (PLATFORM_CACHE_LINE_SIZE - sizeof(FBundleNode*) * UE_DEFAULT_GBinned2MaxBundlesBeforeRecycle)
+			uint8 Padding[UE_MB2_BUNDLE_PADDING];
 #endif
 			FPaddedBundlePointer()
 			{
-				DefaultConstructItems<FBundleNode*>(FreeBundles, BINNED2_MAX_GMallocBinned2MaxBundlesBeforeRecycle);
+				DefaultConstructItems<FBundleNode*>(FreeBundles, UE_DEFAULT_GBinned2MaxBundlesBeforeRecycle);
 			}
 		};
 		static_assert(sizeof(FPaddedBundlePointer) == PLATFORM_CACHE_LINE_SIZE, "FPaddedBundlePointer should be the same size as a cache line");
-		MS_ALIGN(PLATFORM_CACHE_LINE_SIZE) FPaddedBundlePointer Bundles[BINNED2_SMALL_POOL_COUNT] GCC_ALIGN(PLATFORM_CACHE_LINE_SIZE);
+		MS_ALIGN(PLATFORM_CACHE_LINE_SIZE) FPaddedBundlePointer Bundles[UE_MB2_SMALL_POOL_COUNT] GCC_ALIGN(PLATFORM_CACHE_LINE_SIZE);
 	};
 
 	static FGlobalRecycler GGlobalRecycler;
 
-	static void FreeBundles(FMallocBinned2& Allocator, FBundleNode* BundlesToRecycle, uint32 InBlockSize, uint32 InPoolIndex)
+	static void FreeBundles(FMallocBinned2& Allocator, FBundleNode* BundlesToRecycle, uint32 InBinSize, uint32 InPoolIndex)
 	{
 		FPoolTable& Table = Allocator.SmallPoolTables[InPoolIndex];
 
@@ -506,9 +494,9 @@ struct FMallocBinned2::Private
 
 				// Free a pooled allocation.
 				FFreeBlock* Free = (FFreeBlock*)Node;
-				Free->NumFreeBlocks = 1;
+				Free->NumFreeBins = 1;
 				Free->NextFreeBlock = NodePool->FirstFreeBlock;
-				Free->BlockSize     = InBlockSize;
+				Free->BinSize = InBinSize;
 				Free->CanaryAndForkState = Allocator.CurrentCanary;
 				Free->PoolIndex = InPoolIndex;
 				NodePool->FirstFreeBlock   = Free;
@@ -523,7 +511,7 @@ struct FMallocBinned2::Private
 					// Free the OS memory.
 					NodePool->Unlink();
 					Allocator.CachedOSPageAllocator.Free(BasePtrOfNode, Allocator.PageSize);
-#if BINNED2_ALLOCATOR_STATS
+#if UE_MB2_ALLOCATOR_STATS
 					AllocatedOSSmallPoolMemory -= ((int64)Allocator.PageSize);
 #endif
 				}
@@ -541,35 +529,38 @@ struct FMallocBinned2::Private
 		static FCriticalSection FreeBlockListsRegistrationMutex;
 		return FreeBlockListsRegistrationMutex;
 	}
+
 	static TArray<FPerThreadFreeBlockLists*>& GetRegisteredFreeBlockLists()
 	{
 		static TArray<FPerThreadFreeBlockLists*> RegisteredFreeBlockLists;
 		return RegisteredFreeBlockLists;
 	}
-	static void RegisterThreadFreeBlockLists( FPerThreadFreeBlockLists* FreeBlockLists )
+
+	static void RegisterThreadFreeBlockLists(FPerThreadFreeBlockLists* FreeBlockLists)
 	{
 		NOALLOC_SCOPE_CYCLE_COUNTER(STAT_FMallocBinned2_RegisterThreadFreeBlockLists);
 		FScopeLock Lock(&GetFreeBlockListsRegistrationMutex());
-#if BINNED2_ALLOCATOR_STATS_VALIDATION
+#if UE_MB2_ALLOCATOR_STATS_VALIDATION
 		++RecursionCounter;
 #endif
 		GetRegisteredFreeBlockLists().Add(FreeBlockLists);
-#if BINNED2_ALLOCATOR_STATS_VALIDATION
+#if UE_MB2_ALLOCATOR_STATS_VALIDATION
 		--RecursionCounter;
 #endif
 	}
-	static void UnregisterThreadFreeBlockLists( FPerThreadFreeBlockLists* FreeBlockLists )
+
+	static void UnregisterThreadFreeBlockLists(FPerThreadFreeBlockLists* FreeBlockLists)
 	{
 		NOALLOC_SCOPE_CYCLE_COUNTER(STAT_FMallocBinned2_UnregisterThreadFreeBlockLists);
 		FScopeLock Lock(&GetFreeBlockListsRegistrationMutex());
-#if BINNED2_ALLOCATOR_STATS_VALIDATION
+#if UE_MB2_ALLOCATOR_STATS_VALIDATION
 		++RecursionCounter;
 #endif
 		GetRegisteredFreeBlockLists().Remove(FreeBlockLists);
-#if BINNED2_ALLOCATOR_STATS_VALIDATION
+#if UE_MB2_ALLOCATOR_STATS_VALIDATION
 		--RecursionCounter;
 #endif
-#if BINNED2_ALLOCATOR_STATS
+#if UE_MB2_ALLOCATOR_STATS
 		ConsolidatedMemory.fetch_add(FreeBlockLists->AllocatedMemory, std::memory_order_relaxed);
 #endif
 	}
@@ -577,7 +568,7 @@ struct FMallocBinned2::Private
 	static void CheckThreadFreeBlockListsForFork()
 	{
 #if BINNED2_FORK_SUPPORT
-		if(GMallocBinned2PerThreadCaches)
+		if(GBinned2PerThreadCaches)
 		{
 			FScopeLock Lock(&GetFreeBlockListsRegistrationMutex());
 			TArray<FPerThreadFreeBlockLists*>& List = GetRegisteredFreeBlockLists();
@@ -628,7 +619,7 @@ void FMallocBinned2::FPoolList::LinkToFront(FPoolInfo* Pool)
 	Pool->Link(Front);
 }
 
-FMallocBinned2::FPoolInfo& FMallocBinned2::FPoolList::PushNewPoolToFront(FMallocBinned2& Allocator, uint32 InBlockSize, uint32 InPoolIndex)
+FMallocBinned2::FPoolInfo& FMallocBinned2::FPoolList::PushNewPoolToFront(FMallocBinned2& Allocator, uint32 InBinSize, uint32 InPoolIndex)
 {
 	const uint32 LocalPageSize = Allocator.PageSize;
 
@@ -640,14 +631,16 @@ FMallocBinned2::FPoolInfo& FMallocBinned2::FPoolList::PushNewPoolToFront(FMalloc
 		FScopeUnlock TempUnlock(&Allocator.Mutex);
 		Private::OutOfMemory(LocalPageSize);
 	}
-#if !UE_USE_VERYLARGEPAGEALLOCATOR || !BINNED2_BOOKKEEPING_AT_THE_END_OF_LARGEBLOCK
-	FFreeBlock* Free = new (FreePtr) FFreeBlock(LocalPageSize, InBlockSize, InPoolIndex, Allocator.CurrentCanary);
+
+#if !UE_USE_VERYLARGEPAGEALLOCATOR || !UE_MB2_BOOKKEEPING_AT_THE_END_OF_LARGEBLOCK
+	FFreeBlock* Free = new (FreePtr) FFreeBlock(LocalPageSize, InBinSize, InPoolIndex, Allocator.CurrentCanary);
 	check(IsAligned(Free, LocalPageSize));
 #else
 	FFreeBlock* FreeBlockPtr = GetPoolHeaderFromPointer(FreePtr);
-	FFreeBlock* Free = new (FreeBlockPtr) FFreeBlock(LocalPageSize, InBlockSize, InPoolIndex, Allocator.CurrentCanary);
+	FFreeBlock* Free = new (FreeBlockPtr) FFreeBlock(LocalPageSize, InBinSize, InPoolIndex, Allocator.CurrentCanary);
 #endif
-#if BINNED2_ALLOCATOR_STATS
+
+#if UE_MB2_ALLOCATOR_STATS
 	AllocatedOSSmallPoolMemory += (int64)LocalPageSize;
 #endif
 	// Create pool
@@ -666,15 +659,15 @@ FMallocBinned2::FMallocBinned2()
 	check(!bOnce); // this is now a singleton-like thing and you cannot make multiple copies
 	bOnce = true;
 	
-	static_assert(sizeof(FFreeBlock) <= SmallBlockSizes[0], "sizeof(FFreeBlock)  must be fit in smallest allocation size handled by FMallocBinned2.");
-	static_assert(sizeof(FMallocBinned2::FBundleNode) <= SmallBlockSizes[0], "sizeof(FBundleNode) must fit in smallest allocation size handled by FMallocBinned2.");
+	static_assert(sizeof(FFreeBlock) <= SmallBinSizes[0], "sizeof(FFreeBlock)  must be fit in smallest allocation size handled by FMallocBinned2.");
+	static_assert(sizeof(FMallocBinned2::FBundleNode) <= SmallBinSizes[0], "sizeof(FBundleNode) must fit in smallest allocation size handled by FMallocBinned2.");
 
-	for (uint32 Index = 0; Index != BINNED2_SMALL_POOL_COUNT; ++Index)
+	for (uint32 Index = 0; Index != UE_MB2_SMALL_POOL_COUNT; ++Index)
 	{
-		uint32 Partner = BINNED2_SMALL_POOL_COUNT - Index - 1;
-		SmallBlockSizesReversed[Index] = SmallBlockSizes[Partner];
+		const uint32 Partner = UE_MB2_SMALL_POOL_COUNT - Index - 1;
+		SmallBinSizesReversed[Index] = SmallBinSizes[Partner];
 	}
-	FGenericPlatformMemoryConstants Constants = FPlatformMemory::GetConstants();
+	const FGenericPlatformMemoryConstants Constants = FPlatformMemory::GetConstants();
 	PageSize = Constants.BinnedPageSize;
 	OsAllocationGranularity = Constants.BinnedAllocationGranularity ? Constants.BinnedAllocationGranularity : PageSize;
 	NumPoolsPerPage = PageSize / sizeof(FPoolInfo);
@@ -682,54 +675,52 @@ FMallocBinned2::FMallocBinned2()
 
 	checkf(FMath::IsPowerOfTwo(PageSize), TEXT("OS page size must be a power of two"));
 	checkf(Constants.AddressLimit > PageSize, TEXT("OS address limit must be greater than the page size")); // Check to catch 32 bit overflow in AddressLimit
-	checkf(PageSize % BINNED2_LARGE_ALLOC == 0, TEXT("OS page size must be a multiple of BINNED2_LARGE_ALLOC"));
+	checkf(PageSize % UE_MB2_LARGE_ALLOC == 0, TEXT("OS page size must be a multiple of UE_MB2_LARGE_ALLOC"));
 
-	static_assert(SmallBlockSizes[BINNED2_SMALL_POOL_COUNT - 1] == BINNED2_MAX_SMALL_POOL_SIZE, "BINNED2_MAX_SMALL_POOL_SIZE must equal the smallest block size");
-	static_assert(sizeof(FMallocBinned2::FFreeBlock) <= SmallBlockSizes[0], "Pool header must be able to fit into the smallest block");
-	static_assert(UE_ARRAY_COUNT(SmallBlockSizes) == BINNED2_SMALL_POOL_COUNT, "Small block size array size must match BINNED2_SMALL_POOL_COUNT");
-	static_assert(UE_ARRAY_COUNT(SmallBlockSizes) <= 256, "Small block size array size must fit in a byte");
-	static_assert(sizeof(FFreeBlock) <= BINNED2_MINIMUM_ALIGNMENT, "Free block struct must be small enough to fit into a block.");
+	static_assert(SmallBinSizes[UE_MB2_SMALL_POOL_COUNT - 1] == UE_MB2_MAX_SMALL_POOL_SIZE, "UE_MB2_MAX_SMALL_POOL_SIZE must equal the largest bin size");
+	static_assert(sizeof(FMallocBinned2::FFreeBlock) <= SmallBinSizes[0], "Pool header must be able to fit into the smallest bin");
+	static_assert(UE_ARRAY_COUNT(SmallBinSizes) == UE_MB2_SMALL_POOL_COUNT, "Number of bins in SmallBinSizes must match UE_MB2_SMALL_POOL_COUNT");
+	static_assert(UE_ARRAY_COUNT(SmallBinSizes) <= 256, "Number of bins in SmallBinSizes fit in one byte");
+	static_assert(sizeof(FFreeBlock) <= UE_MB2_MINIMUM_ALIGNMENT, "Free block struct must be small enough to fit into the smallest bin.");
 
 	// Init pool tables.
-	for (uint32 Index = 0; Index != BINNED2_SMALL_POOL_COUNT; ++Index)
+	for (uint32 Index = 0; Index != UE_MB2_SMALL_POOL_COUNT; ++Index)
 	{
-		checkf(Index == 0 || SmallBlockSizes[Index - 1] < SmallBlockSizes[Index], TEXT("Small block sizes must be strictly increasing"));
-		checkf(SmallBlockSizes[Index] <= PageSize, TEXT("Small block size must be small enough to fit into a page"));
-		checkf(SmallBlockSizes[Index] % BINNED2_MINIMUM_ALIGNMENT == 0, TEXT("Small block size must be a multiple of BINNED2_MINIMUM_ALIGNMENT"));
+		checkf(Index == 0 || SmallBinSizes[Index - 1] < SmallBinSizes[Index], TEXT("Small bin sizes must be strictly increasing"));
+		checkf(SmallBinSizes[Index] <= PageSize, TEXT("Small bin size must be small enough to fit into a page"));
+		checkf(SmallBinSizes[Index] % UE_MB2_MINIMUM_ALIGNMENT == 0, TEXT("Small bin size must be a multiple of UE_MB2_MINIMUM_ALIGNMENT"));
 
-		SmallPoolTables[Index].BlockSize = SmallBlockSizes[Index];
+		SmallPoolTables[Index].BinSize = SmallBinSizes[Index];
 	}
 
 	// Set up pool mappings
-	uint8* IndexEntry = MemSizeToIndex;
-	uint32  PoolIndex  = 0;
-	for (uint32 Index = 0; Index != 1 + (BINNED2_MAX_SMALL_POOL_SIZE >> BINNED2_MINIMUM_ALIGNMENT_SHIFT); ++Index)
+	uint8* IndexEntry = MemSizeToPoolIndex;
+	uint32 PoolIndex  = 0;
+	for (uint32 Index = 0; Index != 1 + (UE_MB2_MAX_SMALL_POOL_SIZE >> UE_MB2_MINIMUM_ALIGNMENT_SHIFT); ++Index)
 	{
-		
-		uint32 BlockSize = Index << BINNED2_MINIMUM_ALIGNMENT_SHIFT; // inverse of int32 Index = int32((Size >> BINNED2_MINIMUM_ALIGNMENT_SHIFT));
-		while (SmallBlockSizes[PoolIndex] < BlockSize)
+		const uint32 BinSize = Index << UE_MB2_MINIMUM_ALIGNMENT_SHIFT; // inverse of int32 Index = int32((Size >> UE_MB2_MINIMUM_ALIGNMENT_SHIFT));
+		while (SmallBinSizes[PoolIndex] < BinSize)
 		{
 			++PoolIndex;
-			check(PoolIndex != BINNED2_SMALL_POOL_COUNT);
+			check(PoolIndex != UE_MB2_SMALL_POOL_COUNT);
 		}
 		check(PoolIndex < 256);
 		*IndexEntry++ = uint8(PoolIndex);
 	}
-	// now reverse the pool sizes for cache coherency
 
-	for (uint32 Index = 0; Index != BINNED2_SMALL_POOL_COUNT; ++Index)
+	// now reverse the pool sizes for cache coherency
+	for (uint32 Index = 0; Index != UE_MB2_SMALL_POOL_COUNT; ++Index)
 	{
-		uint32 Partner = BINNED2_SMALL_POOL_COUNT - Index - 1;
-		SmallBlockSizesReversed[Index] = SmallBlockSizes[Partner];
+		const uint32 Partner = UE_MB2_SMALL_POOL_COUNT - Index - 1;
+		SmallBinSizesReversed[Index] = SmallBinSizes[Partner];
 	}
 
-	uint64 MaxHashBuckets = PtrToPoolMapping.GetMaxHashBuckets();
-
+	const uint64 MaxHashBuckets = PtrToPoolMapping.GetMaxHashBuckets();
 	{
 		LLM_PLATFORM_SCOPE(ELLMTag::FMalloc);
 		size_t AllocationSize = Align(MaxHashBuckets * sizeof(PoolHashBucket), OsAllocationGranularity);
 		HashBuckets = (PoolHashBucket*)FPlatformMemory::BinnedAllocFromOS(AllocationSize);
-#if BINNED2_ALLOCATOR_STATS
+#if UE_MB2_ALLOCATOR_STATS
 		Binned2HashMemory += AllocationSize;
 #endif
 	}
@@ -767,7 +758,7 @@ void FMallocBinned2::OnPreFork()
 {
 #if BINNED2_FORK_SUPPORT
 	// Trim caches so we don't use them in the child process and cause pages to be copied
-	if (GMallocBinned2PerThreadCaches)
+	if (GBinned2PerThreadCaches)
 	{
 		FMallocBinnedCommonUtils::FlushCurrentThreadCache(*this);
 		FMallocBinned2::Private::CheckThreadFreeBlockListsForFork();
@@ -778,7 +769,7 @@ void FMallocBinned2::OnPreFork()
 		while (FBundleNode* Node = FMallocBinned2::Private::GGlobalRecycler.PopBundle(PoolIndex))
 		{
 			Node->NextBundle = nullptr; // We need to override this because it's in a union with Count which is not needed for the freeing work
-			FMallocBinned2::Private::FreeBundles(*this, Node, SmallPoolTables[PoolIndex].BlockSize, PoolIndex);
+			FMallocBinned2::Private::FreeBundles(*this, Node, SmallPoolTables[PoolIndex].BinSize, PoolIndex);
 		}
 	}
 
@@ -793,7 +784,7 @@ void FMallocBinned2::OnPreFork()
 void FMallocBinned2::OnPostFork()
 {
 #if BINNED2_FORK_SUPPORT
-	if (GMallocBinned2PerThreadCaches)
+	if (GBinned2PerThreadCaches)
 	{
 		FMallocBinnedCommonUtils::FlushCurrentThreadCache(*this);
 		FMallocBinned2::Private::CheckThreadFreeBlockListsForFork();
@@ -837,17 +828,17 @@ void* FMallocBinned2::MallocSelect(SIZE_T Size, uint32 Alignment, bool bUseSmall
 
 void* FMallocBinned2::MallocExternalSmall(SIZE_T Size, uint32 Alignment)
 {
-	uint32 PoolIndex = BoundSizeToPoolIndex(Size, MemSizeToIndex);
-	FPerThreadFreeBlockLists* Lists = GMallocBinned2PerThreadCaches ? FPerThreadFreeBlockLists::Get() : nullptr;
+	const uint32 PoolIndex = BoundSizeToPoolIndex(Size, MemSizeToPoolIndex);
+	FPerThreadFreeBlockLists* Lists = GBinned2PerThreadCaches ? FPerThreadFreeBlockLists::Get() : nullptr;
 	if (Lists)
 	{
 		if (Lists->ObtainRecycledPartial(PoolIndex, Private::GGlobalRecycler))
 		{
 			if (void* Result = Lists->Malloc(PoolIndex))
 			{
-#if BINNED2_ALLOCATOR_STATS
-				uint32 BlockSize = PoolIndexToBlockSize(PoolIndex);
-				Lists->AllocatedMemory += BlockSize;
+#if UE_MB2_ALLOCATOR_STATS
+				const uint32 BinSize = PoolIndexToBinSize(PoolIndex);
+				Lists->AllocatedMemory += BinSize;
 #endif
 				return Result;
 			}
@@ -868,29 +859,30 @@ void* FMallocBinned2::MallocExternalSmall(SIZE_T Size, uint32 Alignment)
 	}
 	else
 	{
-		Pool = &Table.ActivePools.PushNewPoolToFront(*this, Table.BlockSize, PoolIndex);
+		Pool = &Table.ActivePools.PushNewPoolToFront(*this, Table.BinSize, PoolIndex);
 	}
 
-	void* Result = Pool->AllocateRegularBlock();
-#if BINNED2_ALLOCATOR_STATS
-	AllocatedSmallPoolMemory += PoolIndexToBlockSize(PoolIndex);
-#endif // BINNED2_ALLOCATOR_STATS
-	if (GMallocBinned2AllocExtra)
+	void* Result = Pool->AllocateBin();
+#if UE_MB2_ALLOCATOR_STATS
+	AllocatedSmallPoolMemory += PoolIndexToBinSize(PoolIndex);
+#endif
+
+	if (GBinned2AllocExtra)
 	{
 		if (Lists)
 		{
 			// prefill the free list with some allocations so we are less likely to hit this slow path with the mutex 
-			for (int32 Index = 0; Index < GMallocBinned2AllocExtra && Pool->HasFreeRegularBlock(); Index++)
+			for (int32 Index = 0; Index < GBinned2AllocExtra && Pool->HasFreeBin(); Index++)
 			{
-				if (!Lists->Free(Result, PoolIndex, Table.BlockSize))
+				if (!Lists->Free(Result, PoolIndex, Table.BinSize))
 				{
 					break;
 				}
-				Result = Pool->AllocateRegularBlock();
+				Result = Pool->AllocateBin();
 			}
 		}
 	}
-	if (!Pool->HasFreeRegularBlock())
+	if (!Pool->HasFreeBin())
 	{
 		Table.ExhaustedPools.LinkToFront(Pool);
 	}
@@ -900,16 +892,16 @@ void* FMallocBinned2::MallocExternalSmall(SIZE_T Size, uint32 Alignment)
 
 void* FMallocBinned2::MallocExternalLarge(SIZE_T Size, uint32 Alignment)
 {
-	static_assert(DEFAULT_ALIGNMENT <= BINNED2_MINIMUM_ALIGNMENT, "DEFAULT_ALIGNMENT is assumed to be zero"); // used below
+	static_assert(DEFAULT_ALIGNMENT <= UE_MB2_MINIMUM_ALIGNMENT, "DEFAULT_ALIGNMENT is assumed to be zero"); // used below
 
-	Alignment = FMath::Max<uint32>(Alignment, BINNED2_MINIMUM_ALIGNMENT);
+	Alignment = FMath::Max<uint32>(Alignment, UE_MB2_MINIMUM_ALIGNMENT);
 	Size = Align(FMath::Max((SIZE_T)1, Size), Alignment);
 
 	checkf(FMath::IsPowerOfTwo(Alignment), TEXT("Invalid Malloc alignment: '%u' is not a power of two"), Alignment);
 	checkf(Alignment <= PageSize, TEXT("Invalid Malloc alignment: '%u' is greater than the page size '%u'"), Alignment, PageSize);
 	checkf(FMallocBinned2::FPoolInfo::IsSupportedSize(Size), TEXT("Invalid Malloc size: '%" SIZE_T_FMT "'"), Size);
 
-	UPTRINT AlignedSize = Align(Size, OsAllocationGranularity);
+	const SIZE_T AlignedSize = Align(Size, OsAllocationGranularity);
 
 	LLM_PLATFORM_SCOPE(ELLMTag::FMalloc);
 
@@ -934,8 +926,8 @@ void* FMallocBinned2::MallocExternalLarge(SIZE_T Size, uint32 Alignment)
 	UE_CLOG(!IsAligned(Result, Alignment) ,LogMemory, Fatal, TEXT("FMallocBinned2 alignment was too large for OS. Alignment=%d   Ptr=%p"), Alignment, Result);
 	check(IsAligned(Result, PageSize) && IsOSAllocation(Result));
 
-#if BINNED2_ALLOCATOR_STATS
-	AllocatedLargePoolMemory += (int64)Size;
+#if UE_MB2_ALLOCATOR_STATS
+	AllocatedLargePoolMemory += Size;
 	AllocatedLargePoolMemoryWAlignment += AlignedSize;
 #endif
 
@@ -954,7 +946,7 @@ void* FMallocBinned2::ReallocExternal(void* Ptr, SIZE_T NewSize, uint32 Alignmen
 		return nullptr;
 	}
 
-	static_assert(DEFAULT_ALIGNMENT <= BINNED2_MINIMUM_ALIGNMENT, "DEFAULT_ALIGNMENT is assumed to be zero"); // used below
+	static_assert(DEFAULT_ALIGNMENT <= UE_MB2_MINIMUM_ALIGNMENT, "DEFAULT_ALIGNMENT is assumed to be zero"); // used below
 
 	if (!IsOSAllocation(Ptr))
 	{
@@ -962,23 +954,23 @@ void* FMallocBinned2::ReallocExternal(void* Ptr, SIZE_T NewSize, uint32 Alignmen
 		// Reallocate to a smaller/bigger pool if necessary
 		FFreeBlock* Free = GetPoolHeaderFromPointer(Ptr);
 		CanaryTest(Free);
-		uint32 BlockSize = Free->BlockSize;
-		uint32 PoolIndex = Free->PoolIndex;
+		const uint32 BinSize = Free->BinSize;
+		const uint32 PoolIndex = Free->PoolIndex;
 #if BINNED2_FORK_SUPPORT
 		// If the canary is the pre-fork one, we should not allow this allocation to grow in-place to avoid copying a page from the parent process.
 		if (Free->CanaryAndForkState == CurrentCanary && 
 #else
 		if( 
 #endif
-			((NewSize <= BlockSize) & (IsAligned(BlockSize, Alignment))) && // one branch, not two
-			(PoolIndex == 0 || NewSize > PoolIndexToBlockSize(PoolIndex - 1)))
+			((NewSize <= BinSize) & (IsAligned(BinSize, Alignment))) && // one branch, not two
+			(PoolIndex == 0 || NewSize > PoolIndexToBinSize(PoolIndex - 1)))
 		{
 			return Ptr;
 		}
 
 		// Reallocate and copy the data across
 		void* Result = FMallocBinned2::MallocSelect(NewSize, Alignment);
-		FMemory::Memcpy(Result, Ptr, FMath::Min<SIZE_T>(NewSize, BlockSize));
+		FMemory::Memcpy(Result, Ptr, FMath::Min<SIZE_T>(NewSize, BinSize));
 		FMallocBinned2::FreeExternal(Ptr);
 		return Result;
 	}
@@ -993,18 +985,19 @@ void* FMallocBinned2::ReallocExternal(void* Ptr, SIZE_T NewSize, uint32 Alignmen
 	// Allocated from OS.
 	Mutex.Lock();
 	FPoolInfo* Pool = Private::FindPoolInfo(*this, Ptr);
-	UE_CLOG(!Pool, LogMemory, Fatal, TEXT("FMallocBinned2 Attempt to realloc an unrecognized block %p"), Ptr);
+	UE_CLOG(!Pool, LogMemory, Fatal, TEXT("FMallocBinned2 Attempt to realloc an unrecognized pointer %p"), Ptr);
 
-	UPTRINT PoolOsBytes = Pool->GetOsAllocatedBytes();
-	SIZE_T PoolOSRequestedBytes = Pool->GetOSRequestedBytes();
+	const UPTRINT PoolOsBytes = Pool->GetOsAllocatedBytes();
+	const SIZE_T PoolOSRequestedBytes = Pool->GetOSRequestedBytes();
 	checkf(PoolOSRequestedBytes <= PoolOsBytes, TEXT("FMallocBinned2::ReallocExternal %d %d"), int32(PoolOSRequestedBytes), int32(PoolOsBytes));
 	bool bUseSmallMalloc = UseSmallAlloc(NewSize, Alignment);
 	if (!bUseSmallMalloc)
 	{
 		bUseSmallMalloc = PromoteToLargerBin(NewSize, Alignment, *this);
 	}
+
 	if (NewSize > PoolOsBytes || // can't fit in the old block
-		bUseSmallMalloc || // can switch to the small block allocator
+		bUseSmallMalloc || // can switch to the small bin allocator
 		Align(NewSize, OsAllocationGranularity) < PoolOsBytes) // we can get some pages back
 	{
 		// Grow or shrink.
@@ -1029,14 +1022,14 @@ void* FMallocBinned2::ReallocExternal(void* Ptr, SIZE_T NewSize, uint32 Alignmen
 
 	Mutex.Unlock();
 
-	Alignment = FMath::Max<uint32>(Alignment, BINNED2_MINIMUM_ALIGNMENT);
+	Alignment = FMath::Max<uint32>(Alignment, UE_MB2_MINIMUM_ALIGNMENT);
 	NewSize = Align(FMath::Max((SIZE_T)1, NewSize), Alignment);
 
 	checkf(FMath::IsPowerOfTwo(Alignment), TEXT("Invalid Realloc alignment: '%u' is not a power of two"), Alignment);
 	checkf(Alignment <= PageSize, TEXT("Invalid Realloc alignment: '%u' is greater than the page size '%u'"), Alignment, PageSize);
 	checkf(FMallocBinned2::FPoolInfo::IsSupportedSize(NewSize), TEXT("Invalid Realloc size: '%" SIZE_T_FMT "'"), NewSize);
 
-#if BINNED2_ALLOCATOR_STATS
+#if UE_MB2_ALLOCATOR_STATS
 	AllocatedLargePoolMemory += ((int64)NewSize) - ((int64)PoolOSRequestedBytes);
 	// don't need to change the AllocatedLargePoolMemoryWAlignment because we didn't reallocate so it's the same size
 #endif
@@ -1063,18 +1056,18 @@ void FMallocBinned2::FreeExternal(void* Ptr)
 			return;
 		}
 #endif
-		uint32 BlockSize = BasePtr->BlockSize;
-		uint32 PoolIndex = BasePtr->PoolIndex;
+		const uint32 BinSize = BasePtr->BinSize;
+		const uint32 PoolIndex = BasePtr->PoolIndex;
 
 		FBundleNode* BundlesToRecycle = nullptr;
-		FPerThreadFreeBlockLists* Lists = GMallocBinned2PerThreadCaches ? FPerThreadFreeBlockLists::Get() : nullptr;
+		FPerThreadFreeBlockLists* Lists = GBinned2PerThreadCaches ? FPerThreadFreeBlockLists::Get() : nullptr;
 		if (Lists)
 		{
 			BundlesToRecycle = Lists->RecycleFullBundle(BasePtr->PoolIndex, Private::GGlobalRecycler);
-			bool bPushed = Lists->Free(Ptr, PoolIndex, BlockSize);
+			const bool bPushed = Lists->Free(Ptr, PoolIndex, BinSize);
 			check(bPushed);
-#if BINNED2_ALLOCATOR_STATS
-			Lists->AllocatedMemory -= BlockSize;
+#if UE_MB2_ALLOCATOR_STATS
+			Lists->AllocatedMemory -= BinSize;
 #endif
 		}
 		else
@@ -1082,16 +1075,17 @@ void FMallocBinned2::FreeExternal(void* Ptr)
 			BundlesToRecycle = (FBundleNode*)Ptr;
 			BundlesToRecycle->NextNodeInCurrentBundle = nullptr;
 		}
+
 		if (BundlesToRecycle)
 		{
 			BundlesToRecycle->NextBundle = nullptr;
 			FScopeLock Lock(&Mutex);
-			Private::FreeBundles(*this, BundlesToRecycle, BlockSize, PoolIndex);
-#if BINNED2_ALLOCATOR_STATS
+			Private::FreeBundles(*this, BundlesToRecycle, BinSize, PoolIndex);
+#if UE_MB2_ALLOCATOR_STATS
 			if (!Lists)
 			{
 				// lists track their own stat track them instead in the global stat if we don't have lists
-				AllocatedSmallPoolMemory -= ((int64)(BlockSize));
+				AllocatedSmallPoolMemory -= ((int64)(BinSize));
 			}
 #endif
 		}
@@ -1102,12 +1096,12 @@ void FMallocBinned2::FreeExternal(void* Ptr)
 		FPoolInfo* Pool = Private::FindPoolInfo(*this, Ptr);
 		if (!Pool)
 		{
-			UE_LOG(LogMemory, Fatal, TEXT("FMallocBinned2 Attempt to free an unrecognized block %p"), Ptr);
+			UE_LOG(LogMemory, Fatal, TEXT("FMallocBinned2 Attempt to free an unrecognized pointer %p"), Ptr);
 		}
-		UPTRINT PoolOsBytes = Pool->GetOsAllocatedBytes();
-		SIZE_T PoolOSRequestedBytes = Pool->GetOSRequestedBytes();
+		const UPTRINT PoolOsBytes = Pool->GetOsAllocatedBytes();
+		const SIZE_T PoolOSRequestedBytes = Pool->GetOSRequestedBytes();
 
-#if BINNED2_ALLOCATOR_STATS
+#if UE_MB2_ALLOCATOR_STATS
 		AllocatedLargePoolMemory -= ((int64)PoolOSRequestedBytes);
 		AllocatedLargePoolMemoryWAlignment -= ((int64)PoolOsBytes);
 #endif
@@ -1115,7 +1109,7 @@ void FMallocBinned2::FreeExternal(void* Ptr)
 		checkf(PoolOSRequestedBytes <= PoolOsBytes, TEXT("FMallocBinned2::FreeExternal %d %d"), int32(PoolOSRequestedBytes), int32(PoolOsBytes));
 		Pool->SetCanary(FPoolInfo::ECanary::Unassigned, true, false);
 		// Free an OS allocation.
-		CachedOSPageAllocator.Free(Ptr, PoolOsBytes, &Mutex, FPerThreadFreeBlockLists::Get() != nullptr && GMallocBinned2MoveOSFreesOffTimeCriticalThreads != 0);
+		CachedOSPageAllocator.Free(Ptr, PoolOsBytes, &Mutex, FPerThreadFreeBlockLists::Get() != nullptr && GBinned2MoveOSFreesOffTimeCriticalThreads != 0);
 	}
 }
 
@@ -1126,8 +1120,7 @@ bool FMallocBinned2::GetAllocationSizeExternal(void* Ptr, SIZE_T& SizeOut)
 		check(Ptr); // null is 64k aligned so we should not be here
 		const FFreeBlock* Free = GetPoolHeaderFromPointer(Ptr);
 		CanaryTest(Free);
-		uint32 BlockSize = Free->BlockSize;
-		SizeOut = BlockSize;
+		SizeOut = Free->BinSize;
 		return true;
 	}
 	if (!Ptr)
@@ -1144,34 +1137,34 @@ bool FMallocBinned2::GetAllocationSizeExternal(void* Ptr, SIZE_T& SizeOut)
 
 	if (!Pool)
 	{
-		UE_LOG(LogMemory, Fatal, TEXT("FMallocBinned2 Attempt to GetAllocationSizeExternal an unrecognized block %p"), Ptr);
+		UE_LOG(LogMemory, Fatal, TEXT("FMallocBinned2 Attempt to GetAllocationSizeExternal an unrecognized pointer %p"), Ptr);
 	}
-	UPTRINT PoolOsBytes = Pool->GetOsAllocatedBytes();
-	SIZE_T PoolOSRequestedBytes = Pool->GetOSRequestedBytes();
+	const UPTRINT PoolOsBytes = Pool->GetOsAllocatedBytes();
+	const SIZE_T PoolOSRequestedBytes = Pool->GetOSRequestedBytes();
 	checkf(PoolOSRequestedBytes <= PoolOsBytes, TEXT("FMallocBinned2::GetAllocationSizeExternal %d %d"), int32(PoolOSRequestedBytes), int32(PoolOsBytes));
 	SizeOut = PoolOsBytes;
 	return true;
 }
 
-void FMallocBinned2::FPoolList::ValidateActivePools()
+void FMallocBinned2::FPoolList::ValidateActivePools() const
 {
-	for (FPoolInfo** PoolPtr = &Front; *PoolPtr; PoolPtr = &(*PoolPtr)->Next)
+	for (const FPoolInfo* const* PoolPtr = &Front; *PoolPtr; PoolPtr = &(*PoolPtr)->Next)
 	{
-		FPoolInfo* Pool = *PoolPtr;
+		const FPoolInfo* Pool = *PoolPtr;
 		check(Pool->PtrToPrevNext == PoolPtr);
 		check(Pool->FirstFreeBlock);
-		for (FFreeBlock* Free = Pool->FirstFreeBlock; Free; Free = (FFreeBlock*)Free->NextFreeBlock)
+		for (FFreeBlock* Free = Pool->FirstFreeBlock; Free; Free = Free->NextFreeBlock)
 		{
-			check(Free->GetNumFreeRegularBlocks() > 0);
+			check(Free->GetNumFreeBins() > 0);
 		}
 	}
 }
 
-void FMallocBinned2::FPoolList::ValidateExhaustedPools()
+void FMallocBinned2::FPoolList::ValidateExhaustedPools() const
 {
-	for (FPoolInfo** PoolPtr = &Front; *PoolPtr; PoolPtr = &(*PoolPtr)->Next)
+	for (const FPoolInfo* const* PoolPtr = &Front; *PoolPtr; PoolPtr = &(*PoolPtr)->Next)
 	{
-		FPoolInfo* Pool = *PoolPtr;
+		const FPoolInfo* Pool = *PoolPtr;
 		check(Pool->PtrToPrevNext == PoolPtr);
 		check(!Pool->FirstFreeBlock);
 	}
@@ -1198,7 +1191,7 @@ const TCHAR* FMallocBinned2::GetDescriptiveName()
 
 void FMallocBinned2::FreeBundles(FBundleNode* Bundles, uint32 PoolIndex)
 {
-	Private::FreeBundles(*this, Bundles, PoolIndexToBlockSize(PoolIndex), PoolIndex);
+	Private::FreeBundles(*this, Bundles, PoolIndexToBinSize(PoolIndex), PoolIndex);
 }
 
 FCriticalSection& FMallocBinned2::GetFreeBlockListsRegistrationMutex()
@@ -1213,7 +1206,7 @@ TArray<FMallocBinned2::FPerThreadFreeBlockLists*>& FMallocBinned2::GetRegistered
 
 void FMallocBinned2::Trim(bool bTrimThreadCaches)
 {
-	if (GMallocBinned2PerThreadCaches && bTrimThreadCaches)
+	if (GBinned2PerThreadCaches && bTrimThreadCaches)
 	{
 		FMallocBinnedCommonUtils::Trim(*this);
 
@@ -1230,7 +1223,7 @@ void FMallocBinned2::SetupTLSCachesOnCurrentThread()
 {
 	NOALLOC_SCOPE_CYCLE_COUNTER(STAT_FMallocBinned2_SetupTLSCachesOnCurrentThread);
 
-	if (!BINNED2_ALLOW_RUNTIME_TWEAKING && !GMallocBinned2PerThreadCaches)
+	if (!UE_MB2_ALLOW_RUNTIME_TWEAKING && !GBinned2PerThreadCaches)
 	{
 		return;
 	}
@@ -1288,7 +1281,7 @@ void FMallocBinned2::CanaryFail(const FFreeBlock* Block) const
 #endif
 }
 
-#if BINNED2_ALLOCATOR_STATS
+#if UE_MB2_ALLOCATOR_STATS
 int64 FMallocBinned2::GetTotalAllocatedSmallPoolMemory() const
 {
 	NOALLOC_SCOPE_CYCLE_COUNTER(STAT_FMallocBinned2_GetTotalAllocatedSmallPoolMemory);
@@ -1309,22 +1302,22 @@ int64 FMallocBinned2::GetTotalAllocatedSmallPoolMemory() const
 void FMallocBinned2::GetAllocatorStats( FGenericMemoryStats& OutStats )
 {
 	// Even if we have MB2 stats off, cached slack needs to be included as it might be needed by some platforms
-	uint64 OSPageAllocatorCachedFreeSize = CachedOSPageAllocator.GetCachedFreeTotal();
+	const uint64 OSPageAllocatorCachedFreeSize = CachedOSPageAllocator.GetCachedFreeTotal();
 	OutStats.Add(TEXT("PageAllocatorFreeCacheSize"), OSPageAllocatorCachedFreeSize);
 
-#if BINNED2_ALLOCATOR_STATS
-	int64  TotalAllocatedSmallPoolMemory           = GetTotalAllocatedSmallPoolMemory();
-	int64  LocalAllocatedOSSmallPoolMemory         = AllocatedOSSmallPoolMemory.Load(EMemoryOrder::Relaxed);
-	int64  LocalAllocatedLargePoolMemory           = AllocatedLargePoolMemory.Load(EMemoryOrder::Relaxed);
-	int64  LocalAllocatedLargePoolMemoryWAlignment = AllocatedLargePoolMemoryWAlignment.Load(EMemoryOrder::Relaxed);
+#if UE_MB2_ALLOCATOR_STATS
+	const int64  TotalAllocatedSmallPoolMemory           = GetTotalAllocatedSmallPoolMemory();
+	const int64  LocalAllocatedOSSmallPoolMemory         = AllocatedOSSmallPoolMemory.Load(EMemoryOrder::Relaxed);
+	const int64  LocalAllocatedLargePoolMemory           = AllocatedLargePoolMemory.Load(EMemoryOrder::Relaxed);
+	const int64  LocalAllocatedLargePoolMemoryWAlignment = AllocatedLargePoolMemoryWAlignment.Load(EMemoryOrder::Relaxed);
 
 	OutStats.Add(TEXT("AllocatedSmallPoolMemory"), TotalAllocatedSmallPoolMemory);
 	OutStats.Add(TEXT("AllocatedOSSmallPoolMemory"), LocalAllocatedOSSmallPoolMemory);
 	OutStats.Add(TEXT("AllocatedLargePoolMemory"), LocalAllocatedLargePoolMemory);
 	OutStats.Add(TEXT("AllocatedLargePoolMemoryWAlignment"), LocalAllocatedLargePoolMemoryWAlignment);
 
-	uint64 TotalAllocated = TotalAllocatedSmallPoolMemory + LocalAllocatedLargePoolMemory;
-	uint64 TotalOSAllocated = LocalAllocatedOSSmallPoolMemory + LocalAllocatedLargePoolMemoryWAlignment + OSPageAllocatorCachedFreeSize;
+	const uint64 TotalAllocated = TotalAllocatedSmallPoolMemory + LocalAllocatedLargePoolMemory;
+	const uint64 TotalOSAllocated = LocalAllocatedOSSmallPoolMemory + LocalAllocatedLargePoolMemoryWAlignment + OSPageAllocatorCachedFreeSize;
 
 	OutStats.Add(TEXT("TotalAllocated"), TotalAllocated);
 	OutStats.Add(TEXT("TotalOSAllocated"), TotalOSAllocated);
@@ -1335,17 +1328,17 @@ void FMallocBinned2::GetAllocatorStats( FGenericMemoryStats& OutStats )
 
 void FMallocBinned2::DumpAllocatorStats(class FOutputDevice& Ar)
 {
-#if BINNED2_ALLOCATOR_STATS
+#if UE_MB2_ALLOCATOR_STATS
 
-	int64  TotalAllocatedSmallPoolMemory           = GetTotalAllocatedSmallPoolMemory();
-	int64  LocalAllocatedLargePoolMemory           = AllocatedLargePoolMemory.Load(EMemoryOrder::Relaxed);
-	int64  LocalAllocatedLargePoolMemoryWAlignment = AllocatedLargePoolMemoryWAlignment.Load(EMemoryOrder::Relaxed);
-	uint64 OSPageAllocatorCachedFreeSize           = CachedOSPageAllocator.GetCachedFreeTotal();
+	const int64  TotalAllocatedSmallPoolMemory           = GetTotalAllocatedSmallPoolMemory();
+	const int64  LocalAllocatedLargePoolMemory           = AllocatedLargePoolMemory.Load(EMemoryOrder::Relaxed);
+	const int64  LocalAllocatedLargePoolMemoryWAlignment = AllocatedLargePoolMemoryWAlignment.Load(EMemoryOrder::Relaxed);
+	const uint64 OSPageAllocatorCachedFreeSize           = CachedOSPageAllocator.GetCachedFreeTotal();
 
 	Ar.Logf(TEXT("FMallocBinned2 Mem report"));
 	Ar.Logf(TEXT("Constants.BinnedPageSize = %d"), int32(PageSize));
 	Ar.Logf(TEXT("Constants.BinnedAllocationGranularity = %d"), int32(OsAllocationGranularity));
-	Ar.Logf(TEXT("Small Pool Allocations: %fmb  (including block size padding)"), ((double)TotalAllocatedSmallPoolMemory) / (1024.0f * 1024.0f));
+	Ar.Logf(TEXT("Small Pool Allocations: %fmb  (including bin size padding)"), ((double)TotalAllocatedSmallPoolMemory) / (1024.0f * 1024.0f));
 	Ar.Logf(TEXT("Small Pool OS Allocated: %fmb"), ((double)AllocatedOSSmallPoolMemory) / (1024.0f * 1024.0f));
 	Ar.Logf(TEXT("Large Pool Requested Allocations: %fmb"), ((double)AllocatedLargePoolMemory) / (1024.0f * 1024.0f));
 	Ar.Logf(TEXT("Large Pool OS Allocated: %fmb"), ((double)AllocatedLargePoolMemoryWAlignment) / (1024.0f * 1024.0f));
@@ -1367,10 +1360,10 @@ void FMallocBinned2::DumpAllocatorStats(class FOutputDevice& Ar)
 #endif
 
 #else
-	Ar.Logf(TEXT("Allocator Stats for binned2 are not in this build set BINNED2_ALLOCATOR_STATS 1 in MallocBinned2.cpp"));
+	Ar.Logf(TEXT("Allocator Stats for binned2 are not in this build set UE_MB2_ALLOCATOR_STATS 1 in MallocBinned2.cpp"));
 #endif
 }
-#if !BINNED2_INLINE
+#if !UE_MB2_INLINE
 	#if PLATFORM_USES_FIXED_GMalloc_CLASS && !FORCE_ANSI_ALLOCATOR && USE_MALLOC_BINNED2
 		//#define FMEMORY_INLINE_FUNCTION_DECORATOR  FORCEINLINE
 		#define FMEMORY_INLINE_GMalloc (FMallocBinned2::MallocBinned2)
diff --git a/Engine/Source/Runtime/Core/Private/HAL/MallocBinned3.cpp b/Engine/Source/Runtime/Core/Private/HAL/MallocBinned3.cpp
--- a/Engine/Source/Runtime/Core/Private/HAL/MallocBinned3.cpp
+++ b/Engine/Source/Runtime/Core/Private/HAL/MallocBinned3.cpp
@@ -15,12 +15,12 @@ PRAGMA_DISABLE_UNSAFE_TYPECAST_WARNINGS
 #include "HAL/MemoryMisc.h"
 #include "HAL/PlatformMisc.h"
 
-#if USE_CACHED_PAGE_ALLOCATOR_FOR_LARGE_ALLOCS
+#if UE_MB3_USE_CACHED_PAGE_ALLOCATOR_FOR_LARGE_ALLOCS
 #include "HAL/Allocators/CachedOSPageAllocator.h"
-#define BINNED3_MAX_CACHED_OS_FREES (64)
-#define BINNED3_MAX_CACHED_OS_FREES_BYTE_LIMIT (64*1024*1024)
+#define UE_MB3_MAX_CACHED_OS_FREES (64)
+#define UE_MB3_MAX_CACHED_OS_FREES_BYTE_LIMIT (64*1024*1024)
 
-typedef TCachedOSPageAllocator<BINNED3_MAX_CACHED_OS_FREES, BINNED3_MAX_CACHED_OS_FREES_BYTE_LIMIT> TBinned3CachedOSPageAllocator;
+typedef TCachedOSPageAllocator<UE_MB3_MAX_CACHED_OS_FREES, UE_MB3_MAX_CACHED_OS_FREES_BYTE_LIMIT> TBinned3CachedOSPageAllocator;
 
 TBinned3CachedOSPageAllocator& GetCachedOSPageAllocator()
 {
@@ -30,63 +30,56 @@ TBinned3CachedOSPageAllocator& GetCachedOSPageAllocator()
 
 #endif
 
-#if BINNED3_ALLOW_RUNTIME_TWEAKING
-
-int32 GMallocBinned3PerThreadCaches = DEFAULT_GMallocBinned3PerThreadCaches;
-static FAutoConsoleVariableRef GMallocBinned3PerThreadCachesCVar(
-	TEXT("MallocBinned3.PerThreadCaches"),
-	GMallocBinned3PerThreadCaches,
-	TEXT("Enables per-thread caches of small (<= 32768 byte) allocations from FMallocBinned3")
-	);
-
-int32 GMallocBinned3MaxBundlesBeforeRecycle = BINNED3_MAX_GMallocBinned3MaxBundlesBeforeRecycle;
-static FAutoConsoleVariableRef GMallocBinned3MaxBundlesBeforeRecycleCVar(
-	TEXT("MallocBinned3.BundleRecycleCount"),
-	GMallocBinned3MaxBundlesBeforeRecycle,
-	TEXT("Number of freed bundles in the global recycler before it returns them to the system, per-block size. Limited by BINNED3_MAX_GMallocBinned3MaxBundlesBeforeRecycle (currently 4)")
-	);
-
-int32 GMallocBinned3AllocExtra = DEFAULT_GMallocBinned3AllocExtra;
-static FAutoConsoleVariableRef GMallocBinned3AllocExtraCVar(
-	TEXT("MallocBinned3.AllocExtra"),
-	GMallocBinned3AllocExtra,
-	TEXT("When we do acquire the lock, how many blocks cached in TLS caches. In no case will we grab more than a page.")
-	);
-
+#if UE_MB3_ALLOW_RUNTIME_TWEAKING
+	int32 GBinned3PerThreadCaches = UE_DEFAULT_GBinned3PerThreadCaches;
+	static FAutoConsoleVariableRef GMallocBinned3PerThreadCachesCVar(
+		TEXT("MallocBinned3.PerThreadCaches"),
+		GBinned3PerThreadCaches,
+		TEXT("Enables per-thread caches of small (<= 32768 byte) allocations from FMallocBinned3")
+		);
+
+	int32 GBinned3MaxBundlesBeforeRecycle = UE_DEFAULT_GBinned3MaxBundlesBeforeRecycle;
+	static FAutoConsoleVariableRef GMallocBinned3MaxBundlesBeforeRecycleCVar(
+		TEXT("MallocBinned3.BundleRecycleCount"),
+		GBinned3MaxBundlesBeforeRecycle,
+		TEXT("Number of freed bundles in the global recycler before it returns them to the system, per-block size. Limited by UE_DEFAULT_GBinned3MaxBundlesBeforeRecycle (currently 4)")
+		);
+
+	int32 GBinned3AllocExtra = UE_DEFAULT_GBinned3AllocExtra;
+	static FAutoConsoleVariableRef GMallocBinned3AllocExtraCVar(
+		TEXT("MallocBinned3.AllocExtra"),
+		GBinned3AllocExtra,
+		TEXT("When we do acquire the lock, how many bins cached in TLS caches. In no case will we grab more than a page.")
+		);
 #endif
 
-#if BINNED3_ALLOCATOR_STATS
-int64 Binned3AllocatedSmallPoolMemory = 0; // memory that's requested to be allocated by the game
-int64 Binned3AllocatedOSSmallPoolMemory = 0;
-
-int64 Binned3AllocatedLargePoolMemory = 0; // memory requests to the OS which don't fit in the small pool
-int64 Binned3AllocatedLargePoolMemoryWAlignment = 0; // when we allocate at OS level we need to align to a size
-
-TAtomic<int64> Binned3Commits;
-TAtomic<int64> Binned3Decommits;
-int64 Binned3PoolInfoMemory = 0;
-int64 Binned3HashMemory = 0;
-int64 Binned3FreeBitsMemory = 0;
-TAtomic<int64> Binned3TotalPoolSearches;
-TAtomic<int64> Binned3TotalPointerTests;
+#if UE_MB3_ALLOCATOR_STATS
+	int64 Binned3AllocatedSmallPoolMemory = 0; // memory that's requested to be allocated by the game
+	int64 Binned3AllocatedOSSmallPoolMemory = 0;
 
+	int64 Binned3AllocatedLargePoolMemory = 0; // memory requests to the OS which don't fit in the small pool
+	int64 Binned3AllocatedLargePoolMemoryWAlignment = 0; // when we allocate at OS level we need to align to a size
 
+	TAtomic<int64> Binned3Commits;
+	TAtomic<int64> Binned3Decommits;
+	int64 Binned3PoolInfoMemory = 0;
+	int64 Binned3HashMemory = 0;
+	int64 Binned3FreeBitsMemory = 0;
+	TAtomic<int64> Binned3TotalPoolSearches;
+	TAtomic<int64> Binned3TotalPointerTests;
 #endif
 
-#define BINNED3_TIME_LARGE_BLOCKS (0)
+#define UE_MB3_TIME_LARGE_BLOCKS (0)
 
-#if BINNED3_TIME_LARGE_BLOCKS
-TAtomic<double> MemoryRangeReserveTotalTime(0.0);
-TAtomic<int32> MemoryRangeReserveTotalCount(0);
+#if UE_MB3_TIME_LARGE_BLOCKS
+	TAtomic<double> MemoryRangeReserveTotalTime(0.0);
+	TAtomic<int32> MemoryRangeReserveTotalCount(0);
 
-TAtomic<double> MemoryRangeFreeTotalTime(0.0);
-TAtomic<int32> MemoryRangeFreeTotalCount(0);
+	TAtomic<double> MemoryRangeFreeTotalTime(0.0);
+	TAtomic<int32> MemoryRangeFreeTotalCount(0);
 #endif
 
-#define BINNED3_LARGE_POOL_CANARIES 1 // need to repad the data structure so that the pagesize divides by this to disable
-
-MS_ALIGN(PLATFORM_CACHE_LINE_SIZE) static uint8 Binned3UnusedAlignPadding[PLATFORM_CACHE_LINE_SIZE] GCC_ALIGN(PLATFORM_CACHE_LINE_SIZE) = { 0 };
-uint16 FMallocBinned3::SmallBlockSizesReversedShifted[BINNED3_SMALL_POOL_COUNT + 1] = { 0 };
+uint16 FMallocBinned3::SmallBinSizesReversedShifted[UE_MB3_SMALL_POOL_COUNT + 1] = { 0 };
 uint32 FMallocBinned3::OsAllocationGranularity = 0;
 
 #if !BINNED3_USE_SEPARATE_VM_PER_POOL
@@ -94,14 +87,14 @@ uint32 FMallocBinned3::OsAllocationGranularity = 0;
 #else
 	uint64 FMallocBinned3::PoolSearchDiv = 0;
 	uint8* FMallocBinned3::HighestPoolBaseVMPtr = nullptr;
-	uint8* FMallocBinned3::PoolBaseVMPtr[BINNED3_SMALL_POOL_COUNT] = { nullptr };
+	uint8* FMallocBinned3::PoolBaseVMPtr[UE_MB3_SMALL_POOL_COUNT] = { nullptr };
 #endif
 
 FMallocBinned3* FMallocBinned3::MallocBinned3 = nullptr;
 // Mapping of sizes to small table indices
-uint8 FMallocBinned3::MemSizeToIndex[1 + (BINNED3_MAX_SMALL_POOL_SIZE >> BINNED3_MINIMUM_ALIGNMENT_SHIFT)] = { 0 };
+uint8 FMallocBinned3::MemSizeToPoolIndex[1 + (UE_MB3_MAX_SMALL_POOL_SIZE >> UE_MB3_MINIMUM_ALIGNMENT_SHIFT)] = { 0 };
 
-struct FMallocBinned3::FPoolInfoSmall
+struct FMallocBinned3::FPoolInfoSmall		//This is more like BlockInfoSmall as it stores info per block
 {
 	enum ECanary
 	{
@@ -111,7 +104,7 @@ struct FMallocBinned3::FPoolInfoSmall
 
 	uint32 Canary : 2;
 	uint32 Taken : 15;
-	uint32 NoFirstFreeIndex : 1;
+	uint32 NoFirstFreeIndex : 1;	// if a small block info has a FirstFreeIndex as it's only 14 bits and can't store UINT_MAX to say there's no first free index
 	uint32 FirstFreeIndex : 14;
 
 	FPoolInfoSmall()
@@ -122,6 +115,7 @@ struct FMallocBinned3::FPoolInfoSmall
 	{
 		static_assert(sizeof(FPoolInfoSmall) == 4, "Padding fail");
 	}
+
 	void CheckCanary(ECanary ShouldBe) const
 	{
 		if (Canary != ShouldBe)
@@ -129,6 +123,7 @@ struct FMallocBinned3::FPoolInfoSmall
 			UE_LOG(LogMemory, Fatal, TEXT("MallocBinned3 Corruption Canary was 0x%x, should be 0x%x"), int32(Canary), int32(ShouldBe));
 		}
 	}
+
 	void SetCanary(ECanary ShouldBe, bool bPreexisting, bool bGuarnteedToBeNew)
 	{
 		if (bPreexisting)
@@ -165,31 +160,32 @@ struct FMallocBinned3::FPoolInfoSmall
 		}
 		Canary = ShouldBe;
 	}
-	bool HasFreeRegularBlock() const
+
+	bool HasFreeBin() const
 	{
 		CheckCanary(ECanary::SmallAssigned);
 		return !NoFirstFreeIndex;
 	}
 
-	void* AllocateRegularBlock(uint8* BlockOfBlocksPtr, uint32 BlockSize)
+	void* AllocateBin(uint8* BlockPtr, uint32 BinSize)
 	{
-		check(HasFreeRegularBlock());
+		check(HasFreeBin());
 		++Taken;
 		check(Taken != 0);
-		FFreeBlock* Free = (FFreeBlock*)(BlockOfBlocksPtr + BlockSize * FirstFreeIndex);
-		void* Result = Free->AllocateRegularBlock();
-		if (Free->GetNumFreeRegularBlocks() == 0)
+		FFreeBlock* Free = (FFreeBlock*)(BlockPtr + BinSize * FirstFreeIndex);
+		void* Result = Free->AllocateBin();
+		if (Free->GetNumFreeBins() == 0)
 		{
-			if (Free->NextFreeIndex == MAX_uint32)
+			if (Free->NextFreeBlockIndex == MAX_uint32)
 			{
 				FirstFreeIndex = 0;
 				NoFirstFreeIndex = 1;
 			}
 			else
 			{
-				FirstFreeIndex = Free->NextFreeIndex;
-				check(uint32(FirstFreeIndex) == Free->NextFreeIndex);
-				check(((FFreeBlock*)(BlockOfBlocksPtr + BlockSize * FirstFreeIndex))->GetNumFreeRegularBlocks());
+				FirstFreeIndex = Free->NextFreeBlockIndex;
+				check(uint32(FirstFreeIndex) == Free->NextFreeBlockIndex);
+				check(((FFreeBlock*)(BlockPtr + BinSize * FirstFreeIndex))->GetNumFreeBins());
 			}
 		}
 
@@ -206,11 +202,11 @@ struct FMallocBinned3::FPoolInfoLarge
 	};
 
 public:
-	ECanary	Canary;	// See ECanary
+	ECanary	Canary;
 private:
-	uint32 AllocSize;      // Number of bytes allocated
-	uint32 VMSizeDivVirtualSizeAlignment;    // Number of VM bytes allocated aligned for OS
-	uint32 CommitSize;    // Number of bytes committed by the OS
+	uint32 AllocSize;						// Number of bytes allocated
+	uint32 VMSizeDivVirtualSizeAlignment;	// Number of VM bytes allocated aligned for OS
+	uint32 CommitSize;						// Number of bytes committed by the OS
 
 public:
 	FPoolInfoLarge() :
@@ -220,6 +216,7 @@ public:
 		CommitSize(0)
 	{
 	}
+
 	void CheckCanary(ECanary ShouldBe) const
 	{
 		if (Canary != ShouldBe)
@@ -227,6 +224,7 @@ public:
 			UE_LOG(LogMemory, Fatal, TEXT("MallocBinned3 Corruption Canary was 0x%x, should be 0x%x"), int32(Canary), int32(ShouldBe));
 		}
 	}
+
 	void SetCanary(ECanary ShouldBe, bool bPreexisting, bool bGuarnteedToBeNew)
 	{
 		if (bPreexisting)
@@ -263,14 +261,15 @@ public:
 		}
 		Canary = ShouldBe;
 	}
+
 	uint32 GetOSRequestedBytes() const
 	{
 		return AllocSize;
 	}
 
-	UPTRINT GetOsCommittedBytes() const
+	uint32 GetOsCommittedBytes() const
 	{
-		return (UPTRINT)CommitSize;
+		return CommitSize;
 	}
 
 	uint32 GetOsVMPages() const
@@ -279,7 +278,7 @@ public:
 		return VMSizeDivVirtualSizeAlignment;
 	}
 
-	void SetOSAllocationSizes(uint32 InRequestedBytes)
+	void SetOSAllocationSize(uint32 InRequestedBytes)
 	{
 		CheckCanary(ECanary::LargeAssigned);
 		AllocSize = InRequestedBytes;
@@ -308,16 +307,16 @@ struct FMallocBinned3::Private
 	/**
 	* Gets the FPoolInfoSmall for a small block memory address. If no valid info exists one is created.
 	*/
-	static FPoolInfoSmall* GetOrCreatePoolInfoSmall(FMallocBinned3& Allocator, uint32 InPoolIndex, uint32 BlockOfBlocksIndex)
+	static FPoolInfoSmall* GetOrCreatePoolInfoSmall(FMallocBinned3& Allocator, uint32 InPoolIndex, uint32 BlockIndex)
 	{
 		const uint32 InfosPerPage = Allocator.SmallPoolInfosPerPlatformPage;
-		const uint32 InfoOuterIndex = BlockOfBlocksIndex / InfosPerPage;
-		const uint32 InfoInnerIndex = BlockOfBlocksIndex % InfosPerPage;
+		const uint32 InfoOuterIndex = BlockIndex / InfosPerPage;
+		const uint32 InfoInnerIndex = BlockIndex % InfosPerPage;
 		FPoolInfoSmall*& InfoBlock = Allocator.SmallPoolTables[InPoolIndex].PoolInfos[InfoOuterIndex];
 		if (!InfoBlock)
 		{
 			InfoBlock = (FPoolInfoSmall*)Allocator.AllocateMetaDataMemory(Allocator.OsAllocationGranularity);
-#if BINNED3_ALLOCATOR_STATS
+#if UE_MB3_ALLOCATOR_STATS
 			Binned3PoolInfoMemory += Allocator.OsAllocationGranularity;
 #endif
 			DefaultConstructItems<FPoolInfoSmall>((void*)InfoBlock, InfosPerPage);
@@ -326,14 +325,15 @@ struct FMallocBinned3::Private
 		FPoolInfoSmall* Result = &InfoBlock[InfoInnerIndex];
 
 		bool bGuaranteedToBeNew = false;
-		if (BlockOfBlocksIndex >= Allocator.SmallPoolTables[InPoolIndex].NumEverUsedBlockOfBlocks)
+		if (BlockIndex >= Allocator.SmallPoolTables[InPoolIndex].NumEverUsedBlocks)
 		{
 			bGuaranteedToBeNew = true;
-			Allocator.SmallPoolTables[InPoolIndex].NumEverUsedBlockOfBlocks = BlockOfBlocksIndex + 1;
+			Allocator.SmallPoolTables[InPoolIndex].NumEverUsedBlocks = BlockIndex + 1;
 		}
 		Result->SetCanary(FPoolInfoSmall::ECanary::SmallAssigned, false, bGuaranteedToBeNew);
 		return Result;
 	}
+
 	/**
 	 * Gets the FPoolInfoLarge for a large block memory address. If no valid info exists one is created.
 	 */
@@ -344,10 +344,10 @@ struct FMallocBinned3::Private
 		 */
 		auto CreatePoolArray = [&Allocator](uint64 NumPools)
 		{
-			uint64 PoolArraySize = NumPools * sizeof(FPoolInfoLarge);
+			const uint64 PoolArraySize = NumPools * sizeof(FPoolInfoLarge);
 
 			void* Result = Allocator.AllocateMetaDataMemory(PoolArraySize);
-#if BINNED3_ALLOCATOR_STATS
+#if UE_MB3_ALLOCATOR_STATS
 			Binned3PoolInfoMemory += PoolArraySize;
 #endif
 			if (!Result)
@@ -359,7 +359,7 @@ struct FMallocBinned3::Private
 			return (FPoolInfoLarge*)Result;
 		};
 
-		uint32 BucketIndex;
+		uint32  BucketIndex;
 		UPTRINT BucketIndexCollision;
 		uint32  PoolIndex;
 		Allocator.PtrToPoolMapping.GetHashBucketAndPoolIndices(InPtr, BucketIndex, BucketIndexCollision, PoolIndex);
@@ -391,7 +391,7 @@ struct FMallocBinned3::Private
 		{
 			{
 				Allocator.HashBucketFreeList = (PoolHashBucket*)Allocator.AllocateMetaDataMemory(FMallocBinned3::OsAllocationGranularity);
-#if BINNED3_ALLOCATOR_STATS
+#if UE_MB3_ALLOCATOR_STATS
 				Binned3HashMemory += FMallocBinned3::OsAllocationGranularity;
 #endif
 			}
@@ -432,12 +432,11 @@ struct FMallocBinned3::Private
 
 	static FPoolInfoLarge* FindPoolInfo(FMallocBinned3& Allocator, void* InPtr)
 	{
-		uint32 BucketIndex;
+		uint32  BucketIndex;
 		UPTRINT BucketIndexCollision;
 		uint32  PoolIndex;
 		Allocator.PtrToPoolMapping.GetHashBucketAndPoolIndices(InPtr, BucketIndex, BucketIndexCollision, PoolIndex);
 
-
 		PoolHashBucket* FirstBucket = &Allocator.HashBuckets[BucketIndex];
 		PoolHashBucket* Collision   = FirstBucket;
 		do
@@ -458,7 +457,7 @@ struct FMallocBinned3::Private
 	{
 		bool PushBundle(uint32 InPoolIndex, FBundleNode* InBundle)
 		{
-			uint32 NumCachedBundles = FMath::Min<uint32>(GMallocBinned3MaxBundlesBeforeRecycle, BINNED3_MAX_GMallocBinned3MaxBundlesBeforeRecycle);
+			const uint32 NumCachedBundles = FMath::Min<uint32>(GBinned3MaxBundlesBeforeRecycle, UE_DEFAULT_GBinned3MaxBundlesBeforeRecycle);
 			for (uint32 Slot = 0; Slot < NumCachedBundles; Slot++)
 			{
 				if (!Bundles[InPoolIndex].FreeBundles[Slot])
@@ -474,7 +473,7 @@ struct FMallocBinned3::Private
 
 		FBundleNode* PopBundle(uint32 InPoolIndex)
 		{
-			uint32 NumCachedBundles = FMath::Min<uint32>(GMallocBinned3MaxBundlesBeforeRecycle, BINNED3_MAX_GMallocBinned3MaxBundlesBeforeRecycle);
+			const uint32 NumCachedBundles = FMath::Min<uint32>(GBinned3MaxBundlesBeforeRecycle, UE_DEFAULT_GBinned3MaxBundlesBeforeRecycle);
 			for (uint32 Slot = 0; Slot < NumCachedBundles; Slot++)
 			{
 				FBundleNode* Result = Bundles[InPoolIndex].FreeBundles[Slot];
@@ -492,23 +491,23 @@ struct FMallocBinned3::Private
 	private:
 		struct FPaddedBundlePointer
 		{
-			FBundleNode* FreeBundles[BINNED3_MAX_GMallocBinned3MaxBundlesBeforeRecycle];
-#define BINNED3_BUNDLE_PADDING (PLATFORM_CACHE_LINE_SIZE - sizeof(FBundleNode*) * BINNED3_MAX_GMallocBinned3MaxBundlesBeforeRecycle)
-#if (4 + (4 * PLATFORM_64BITS)) * BINNED3_MAX_GMallocBinned3MaxBundlesBeforeRecycle < PLATFORM_CACHE_LINE_SIZE
-			uint8 Padding[BINNED3_BUNDLE_PADDING];
+			FBundleNode* FreeBundles[UE_DEFAULT_GBinned3MaxBundlesBeforeRecycle];
+#if (4 + (4 * PLATFORM_64BITS)) * UE_DEFAULT_GBinned3MaxBundlesBeforeRecycle < PLATFORM_CACHE_LINE_SIZE
+#	define UE_MB3_BUNDLE_PADDING (PLATFORM_CACHE_LINE_SIZE - sizeof(FBundleNode*) * UE_DEFAULT_GBinned3MaxBundlesBeforeRecycle)
+			uint8 Padding[UE_MB3_BUNDLE_PADDING];
 #endif
 			FPaddedBundlePointer()
 			{
-				DefaultConstructItems<FBundleNode*>(FreeBundles, BINNED3_MAX_GMallocBinned3MaxBundlesBeforeRecycle);
+				DefaultConstructItems<FBundleNode*>(FreeBundles, UE_DEFAULT_GBinned3MaxBundlesBeforeRecycle);
 			}
 		};
 		static_assert(sizeof(FPaddedBundlePointer) == PLATFORM_CACHE_LINE_SIZE, "FPaddedBundlePointer should be the same size as a cache line");
-		MS_ALIGN(PLATFORM_CACHE_LINE_SIZE) FPaddedBundlePointer Bundles[BINNED3_SMALL_POOL_COUNT] GCC_ALIGN(PLATFORM_CACHE_LINE_SIZE);
+		MS_ALIGN(PLATFORM_CACHE_LINE_SIZE) FPaddedBundlePointer Bundles[UE_MB3_SMALL_POOL_COUNT] GCC_ALIGN(PLATFORM_CACHE_LINE_SIZE);
 	};
 
 	static FGlobalRecycler GGlobalRecycler;
 
-	static void FreeBundles(FMallocBinned3& Allocator, FBundleNode* BundlesToRecycle, uint32 InBlockSize, uint32 InPoolIndex)
+	static void FreeBundles(FMallocBinned3& Allocator, FBundleNode* BundlesToRecycle, uint32 InBinSize, uint32 InPoolIndex)
 	{
 		FPoolTable& Table = Allocator.SmallPoolTables[InPoolIndex];
 
@@ -522,54 +521,54 @@ struct FMallocBinned3::Private
 			{
 				FBundleNode* NextNode = Node->NextNodeInCurrentBundle;
 
-				uint32 OutBlockOfBlocksIndex;
-				void* BasePtrOfNode = Allocator.BlockOfBlocksPointerFromContainedPtr(Node, Allocator.SmallPoolTables[InPoolIndex].PagesPlatformForBlockOfBlocks, OutBlockOfBlocksIndex);
-				uint32 BlockWithinIndex = (((uint8*)Node) - ((uint8*)BasePtrOfNode)) / Allocator.SmallPoolTables[InPoolIndex].BlockSize;
+				uint32 OutBlockIndex;
+				void* BaseBlockPtr = Allocator.BlockPointerFromContainedPtr(Node, Allocator.SmallPoolTables[InPoolIndex].NumMemoryPagesPerBlock, OutBlockIndex);
+				const uint32 BinIndexWithinBlock = (((uint8*)Node) - ((uint8*)BaseBlockPtr)) / Allocator.SmallPoolTables[InPoolIndex].BinSize;
 
-				FPoolInfoSmall* NodePoolBlock = Allocator.SmallPoolTables[InPoolIndex].PoolInfos[OutBlockOfBlocksIndex / Allocator.SmallPoolInfosPerPlatformPage];
+				FPoolInfoSmall* NodePoolBlock = Allocator.SmallPoolTables[InPoolIndex].PoolInfos[OutBlockIndex / Allocator.SmallPoolInfosPerPlatformPage];
 				if (!NodePoolBlock)
 				{
 					UE_LOG(LogMemory, Fatal, TEXT("FMallocBinned3 Attempt to free an unrecognized small block %p"), Node);
 				}
-				FPoolInfoSmall* NodePool = &NodePoolBlock[OutBlockOfBlocksIndex % Allocator.SmallPoolInfosPerPlatformPage];
+				FPoolInfoSmall* NodePool = &NodePoolBlock[OutBlockIndex % Allocator.SmallPoolInfosPerPlatformPage];
 
 				NodePool->CheckCanary(FPoolInfoSmall::ECanary::SmallAssigned);
 
-				bool bWasExhaused = NodePool->NoFirstFreeIndex;
+				const bool bWasExhaused = NodePool->NoFirstFreeIndex;
 
 				// Free a pooled allocation.
 				FFreeBlock* Free = (FFreeBlock*)Node;
-				Free->NumFreeBlocks = 1;
-				Free->NextFreeIndex = NodePool->NoFirstFreeIndex ? MAX_uint32 : NodePool->FirstFreeIndex;
-				Free->BlockSizeShifted = (InBlockSize >> BINNED3_MINIMUM_ALIGNMENT_SHIFT);
+				Free->NumFreeBins = 1;
+				Free->NextFreeBlockIndex = NodePool->NoFirstFreeIndex ? MAX_uint32 : NodePool->FirstFreeIndex;
+				Free->BinSizeShifted = (InBinSize >> UE_MB3_MINIMUM_ALIGNMENT_SHIFT);
 				Free->Canary = FFreeBlock::CANARY_VALUE;
 				Free->PoolIndex = InPoolIndex;
-				NodePool->FirstFreeIndex = BlockWithinIndex;
+				NodePool->FirstFreeIndex = BinIndexWithinBlock;
 				NodePool->NoFirstFreeIndex = 0;
-				check(uint32(NodePool->FirstFreeIndex) == BlockWithinIndex);
+				check(uint32(NodePool->FirstFreeIndex) == BinIndexWithinBlock);
 
 				// Free this pool.
 				check(NodePool->Taken >= 1);
 				if (--NodePool->Taken == 0)
 				{
 					NodePool->SetCanary(FPoolInfoSmall::ECanary::SmallUnassigned, true, false);
-					Table.BlockOfBlockAllocationBits.FreeBit(OutBlockOfBlocksIndex);
+					Table.BlocksAllocatedBits.FreeBit(OutBlockIndex);
 
-					uint64 AllocSize = Allocator.SmallPoolTables[InPoolIndex].PagesPlatformForBlockOfBlocks * Allocator.OsAllocationGranularity;
+					const uint64 AllocSize = Allocator.SmallPoolTables[InPoolIndex].NumMemoryPagesPerBlock * Allocator.OsAllocationGranularity;
 
 					if (!bWasExhaused)
 					{
-						Table.BlockOfBlockIsExhausted.AllocBit(OutBlockOfBlocksIndex);
+						Table.BlocksExhaustedBits.AllocBit(OutBlockIndex);
 					}
 
-					Allocator.Decommit(InPoolIndex, BasePtrOfNode, AllocSize);
-#if BINNED3_ALLOCATOR_STATS
+					Allocator.Decommit(InPoolIndex, BaseBlockPtr, AllocSize);
+#if UE_MB3_ALLOCATOR_STATS
 					Binned3AllocatedOSSmallPoolMemory -= AllocSize;
 #endif
 				}
 				else if (bWasExhaused)
 				{
-					Table.BlockOfBlockIsExhausted.FreeBit(OutBlockOfBlocksIndex);
+					Table.BlocksExhaustedBits.FreeBit(OutBlockIndex);
 				}
 
 				Node = NextNode;
@@ -579,27 +578,29 @@ struct FMallocBinned3::Private
 		}
 	}
 
-
 	static FCriticalSection& GetFreeBlockListsRegistrationMutex()
 	{
 		static FCriticalSection FreeBlockListsRegistrationMutex;
 		return FreeBlockListsRegistrationMutex;
 	}
+
 	static TArray<FPerThreadFreeBlockLists*>& GetRegisteredFreeBlockLists()
 	{
 		static TArray<FPerThreadFreeBlockLists*> RegisteredFreeBlockLists;
 		return RegisteredFreeBlockLists;
 	}
-	static void RegisterThreadFreeBlockLists( FPerThreadFreeBlockLists* FreeBlockLists )
+
+	static void RegisterThreadFreeBlockLists(FPerThreadFreeBlockLists* FreeBlockLists)
 	{
 		FScopeLock Lock(&GetFreeBlockListsRegistrationMutex());
 		GetRegisteredFreeBlockLists().Add(FreeBlockLists);
 	}
-	static void UnregisterThreadFreeBlockLists( FPerThreadFreeBlockLists* FreeBlockLists )
+
+	static void UnregisterThreadFreeBlockLists(FPerThreadFreeBlockLists* FreeBlockLists)
 	{
 		FScopeLock Lock(&GetFreeBlockListsRegistrationMutex());
 		GetRegisteredFreeBlockLists().Remove(FreeBlockLists);
-#if BINNED3_ALLOCATOR_STATS
+#if UE_MB3_ALLOCATOR_STATS
 		ConsolidatedMemory.fetch_add(FreeBlockLists->AllocatedMemory, std::memory_order_relaxed);
 #endif
 	}
@@ -609,7 +610,7 @@ FMallocBinned3::Private::FGlobalRecycler FMallocBinned3::Private::GGlobalRecycle
 
 void FMallocBinned3::FreeBundles(FBundleNode* Bundles, uint32 PoolIndex)
 {
-	Private::FreeBundles(*this, Bundles, PoolIndexToBlockSize(PoolIndex), PoolIndex);
+	Private::FreeBundles(*this, Bundles, PoolIndexToBinSize(PoolIndex), PoolIndex);
 }
 
 void FMallocBinned3::RegisterThreadFreeBlockLists(FPerThreadFreeBlockLists* FreeBlockLists)
@@ -622,52 +623,51 @@ void FMallocBinned3::UnregisterThreadFreeBlockLists(FPerThreadFreeBlockLists* Fr
 	Private::UnregisterThreadFreeBlockLists(FreeBlockLists);
 }
 
-FMallocBinned3::FPoolInfoSmall* FMallocBinned3::PushNewPoolToFront(FMallocBinned3::FPoolTable& Table, uint32 InBlockSize, uint32 InPoolIndex, uint32& OutBlockOfBlocksIndex)
+FMallocBinned3::FPoolInfoSmall* FMallocBinned3::PushNewPoolToFront(FMallocBinned3::FPoolTable& Table, uint32 InBinSize, uint32 InPoolIndex, uint32& OutBlockIndex)
 {
-	const uint32 BlockOfBlocksSize = OsAllocationGranularity * Table.PagesPlatformForBlockOfBlocks;
+	const uint32 BlockSize = OsAllocationGranularity * Table.NumMemoryPagesPerBlock;
 
 	// Allocate memory.
-
-	uint32 BlockOfBlocksIndex = Table.BlockOfBlockAllocationBits.AllocBit();
-	if (BlockOfBlocksIndex == MAX_uint32)
+	const uint32 BlockIndex = Table.BlocksAllocatedBits.AllocBit();
+	if (BlockIndex == MAX_uint32)
 	{
 		return nullptr;
 	}
-	uint8* FreePtr = BlockPointerFromIndecies(InPoolIndex, BlockOfBlocksIndex, BlockOfBlocksSize);
+	uint8* FreePtr = BlockPointerFromIndecies(InPoolIndex, BlockIndex, BlockSize);
 
 	LLM_PLATFORM_SCOPE(ELLMTag::FMalloc);
-	Commit(InPoolIndex, FreePtr, BlockOfBlocksSize);
-	uint64 EndOffset = UPTRINT(FreePtr + BlockOfBlocksSize) - UPTRINT(PoolBasePtr(InPoolIndex));
+	Commit(InPoolIndex, FreePtr, BlockSize);
+	const uint64 EndOffset = UPTRINT(FreePtr + BlockSize) - UPTRINT(PoolBasePtr(InPoolIndex));
 	if (EndOffset > Table.UnusedAreaOffsetLow)
 	{
 		Table.UnusedAreaOffsetLow = EndOffset;
 	}
-	FFreeBlock* Free = new ((void*)FreePtr) FFreeBlock(BlockOfBlocksSize, InBlockSize, InPoolIndex);
-#if BINNED3_ALLOCATOR_STATS
-	Binned3AllocatedOSSmallPoolMemory += (int64)BlockOfBlocksSize;
+	FFreeBlock* Free = new ((void*)FreePtr) FFreeBlock(BlockSize, InBinSize, InPoolIndex);
+#if UE_MB3_ALLOCATOR_STATS
+	Binned3AllocatedOSSmallPoolMemory += (int64)BlockSize;
 #endif
 	check(IsAligned(Free, OsAllocationGranularity));
 	// Create pool
-	FPoolInfoSmall* Result = Private::GetOrCreatePoolInfoSmall(*this, InPoolIndex, BlockOfBlocksIndex);
+	FPoolInfoSmall* Result = Private::GetOrCreatePoolInfoSmall(*this, InPoolIndex, BlockIndex);
 	Result->CheckCanary(FPoolInfoSmall::ECanary::SmallAssigned);
 	Result->Taken = 0;
 	Result->FirstFreeIndex = 0;
 	Result->NoFirstFreeIndex = 0;
-	Table.BlockOfBlockIsExhausted.FreeBit(BlockOfBlocksIndex);
+	Table.BlocksExhaustedBits.FreeBit(BlockIndex);
 
-	OutBlockOfBlocksIndex = BlockOfBlocksIndex;
+	OutBlockIndex = BlockIndex;
 
 	return Result;
 }
 
-FMallocBinned3::FPoolInfoSmall* FMallocBinned3::GetFrontPool(FPoolTable& Table, uint32 InPoolIndex, uint32& OutBlockOfBlocksIndex)
+FMallocBinned3::FPoolInfoSmall* FMallocBinned3::GetFrontPool(FPoolTable& Table, uint32 InPoolIndex, uint32& OutBlockIndex)
 {
-	OutBlockOfBlocksIndex = Table.BlockOfBlockIsExhausted.NextAllocBit();
-	if (OutBlockOfBlocksIndex == MAX_uint32)
+	OutBlockIndex = Table.BlocksExhaustedBits.NextAllocBit();
+	if (OutBlockIndex == MAX_uint32)
 	{
 		return nullptr;
 	}
-	return Private::GetOrCreatePoolInfoSmall(*this, InPoolIndex, OutBlockOfBlocksIndex);
+	return Private::GetOrCreatePoolInfoSmall(*this, InPoolIndex, OutBlockIndex);
 }
 
 
@@ -678,23 +678,21 @@ FMallocBinned3::FMallocBinned3()
 	check(!bOnce); // this is now a singleton-like thing and you cannot make multiple copies
 	bOnce = true;
 
-	check(!PLATFORM_32BITS);
-
 	OsAllocationGranularity = FPlatformMemory::FPlatformVirtualMemoryBlock::GetCommitAlignment();
 	checkf(FMath::IsPowerOfTwo(OsAllocationGranularity), TEXT("OS page size must be a power of two"));
 
 	// First thing we try to allocate address space for bins as it might help us to move forward Constants.AddressStart and reduce the amount of available address space for the Large OS Allocs
 	// Available address space is used to reserve hash map that can address all of that range, so less addressable space means less memory is allocated for book keeping
 #if !BINNED3_USE_SEPARATE_VM_PER_POOL
-	Binned3BaseVMBlock = FPlatformMemory::FPlatformVirtualMemoryBlock::AllocateVirtual(BINNED3_SMALL_POOL_COUNT * MAX_MEMORY_PER_BLOCK_SIZE, OsAllocationGranularity);
+	Binned3BaseVMBlock = FPlatformMemory::FPlatformVirtualMemoryBlock::AllocateVirtual(UE_MB3_SMALL_POOL_COUNT * UE_MB3_MAX_MEMORY_PER_POOL_SIZE, OsAllocationGranularity);
 	Binned3BaseVMPtr = (uint8*)Binned3BaseVMBlock.GetVirtualPointer();
 	check(IsAligned(Binned3BaseVMPtr, OsAllocationGranularity));
 	verify(Binned3BaseVMPtr);
 #else
 
-	for (uint32 Index = 0; Index < BINNED3_SMALL_POOL_COUNT; ++Index)
+	for (uint32 Index = 0; Index < UE_MB3_SMALL_POOL_COUNT; ++Index)
 	{
-		FPlatformMemory::FPlatformVirtualMemoryBlock NewBLock = FPlatformMemory::FPlatformVirtualMemoryBlock::AllocateVirtual(MAX_MEMORY_PER_BLOCK_SIZE, OsAllocationGranularity);
+		FPlatformMemory::FPlatformVirtualMemoryBlock NewBLock = FPlatformMemory::FPlatformVirtualMemoryBlock::AllocateVirtual(UE_MB3_MAX_MEMORY_PER_POOL_SIZE, OsAllocationGranularity);
 
 		uint8* NewVM = (uint8*)NewBLock.GetVirtualPointer();
 		check(IsAligned(NewVM, OsAllocationGranularity));
@@ -724,25 +722,25 @@ FMallocBinned3::FMallocBinned3()
 			PoolBaseVMBlock[Index] = NewBLock;
 		}
 	}
-	HighestPoolBaseVMPtr = PoolBaseVMPtr[BINNED3_SMALL_POOL_COUNT - 1];
+	HighestPoolBaseVMPtr = PoolBaseVMPtr[UE_MB3_SMALL_POOL_COUNT - 1];
 	uint64 TotalGaps = 0;
-	for (uint32 Index = 0; Index < BINNED3_SMALL_POOL_COUNT - 1; ++Index)
+	for (uint32 Index = 0; Index < UE_MB3_SMALL_POOL_COUNT - 1; ++Index)
 	{
 		check(PoolBaseVMPtr[Index + 1] > PoolBaseVMPtr[Index]); // we sorted it
-		check(PoolBaseVMPtr[Index + 1] >= PoolBaseVMPtr[Index] + MAX_MEMORY_PER_BLOCK_SIZE); // and blocks are non-overlapping
-		TotalGaps += PoolBaseVMPtr[Index + 1] - (PoolBaseVMPtr[Index] + MAX_MEMORY_PER_BLOCK_SIZE);
+		check(PoolBaseVMPtr[Index + 1] >= PoolBaseVMPtr[Index] + UE_MB3_MAX_MEMORY_PER_POOL_SIZE); // and pools are non-overlapping
+		TotalGaps += PoolBaseVMPtr[Index + 1] - (PoolBaseVMPtr[Index] + UE_MB3_MAX_MEMORY_PER_POOL_SIZE);
 	}
 	if (TotalGaps == 0)
 	{
 		PoolSearchDiv = 0;
 	}
-	else if (TotalGaps < MAX_MEMORY_PER_BLOCK_SIZE)
+	else if (TotalGaps < UE_MB3_MAX_MEMORY_PER_POOL_SIZE)
 	{
-		PoolSearchDiv = MAX_MEMORY_PER_BLOCK_SIZE; // the gaps are not significant, ignoring them should give accurate searches
+		PoolSearchDiv = UE_MB3_MAX_MEMORY_PER_POOL_SIZE; // the gaps are not significant, ignoring them should give accurate searches
 	}
 	else
 	{
-		PoolSearchDiv = MAX_MEMORY_PER_BLOCK_SIZE + ((TotalGaps + BINNED3_SMALL_POOL_COUNT - 2) / (BINNED3_SMALL_POOL_COUNT - 1));
+		PoolSearchDiv = UE_MB3_MAX_MEMORY_PER_POOL_SIZE + ((TotalGaps + UE_MB3_SMALL_POOL_COUNT - 2) / (UE_MB3_SMALL_POOL_COUNT - 1));
 	}
 #endif
 
@@ -750,12 +748,12 @@ FMallocBinned3::FMallocBinned3()
 #if !BINNED3_USE_SEPARATE_VM_PER_POOL
 	if (Constants.AddressStart == (uint64)Binned3BaseVMPtr)
 	{
-		Constants.AddressStart += Align(BINNED3_SMALL_POOL_COUNT * MAX_MEMORY_PER_BLOCK_SIZE, OsAllocationGranularity);
+		Constants.AddressStart += Align(UE_MB3_SMALL_POOL_COUNT * UE_MB3_MAX_MEMORY_PER_POOL_SIZE, OsAllocationGranularity);
 	}
 #else
 	if (!TotalGaps && Constants.AddressStart == (uint64)PoolBaseVMPtr[0])
 	{
-		Constants.AddressStart += Align(BINNED3_SMALL_POOL_COUNT * MAX_MEMORY_PER_BLOCK_SIZE, OsAllocationGranularity);
+		Constants.AddressStart += Align(UE_MB3_SMALL_POOL_COUNT * UE_MB3_MAX_MEMORY_PER_POOL_SIZE, OsAllocationGranularity);
 	}
 #endif
 
@@ -765,47 +763,45 @@ FMallocBinned3::FMallocBinned3()
 	PtrToPoolMapping.Init(OsAllocationGranularity, NumLargePoolsPerPage, Constants.AddressStart, Constants.AddressLimit);
 
 	checkf(Constants.AddressLimit > OsAllocationGranularity, TEXT("OS address limit must be greater than the page size")); // Check to catch 32 bit overflow in AddressLimit
-	static_assert(BINNED3_SMALL_POOL_COUNT <= 256, "Small block size array size must fit in a byte");
-	static_assert(sizeof(FFreeBlock) <= BINNED3_MINIMUM_ALIGNMENT, "Free block struct must be small enough to fit into a block.");
+	static_assert(UE_MB3_SMALL_POOL_COUNT <= 256, "Small bins size array size must fit in a byte");
+	static_assert(sizeof(FFreeBlock) <= UE_MB3_MINIMUM_ALIGNMENT, "Free block struct must be small enough to fit into the smallest bin");
 
 	// Init pool tables.
+	FSizeTableEntry SizeTable[UE_MB3_SMALL_POOL_COUNT];
 
-	FSizeTableEntry SizeTable[BINNED3_SMALL_POOL_COUNT];
-
-	verify(FSizeTableEntry::FillSizeTable(OsAllocationGranularity, SizeTable, BINNED3_BASE_PAGE_SIZE, BINNED3_MINIMUM_ALIGNMENT, BINNED3_MAX_SMALL_POOL_SIZE, BINNED3_BASE_PAGE_SIZE) == BINNED3_SMALL_POOL_COUNT);
-	checkf(SizeTable[BINNED3_SMALL_POOL_COUNT - 1].BlockSize == BINNED3_MAX_SMALL_POOL_SIZE, TEXT("BINNED3_MAX_SMALL_POOL_SIZE must equal the largest block size"));
-	checkf(sizeof(FMallocBinned3::FFreeBlock) <= SizeTable[0].BlockSize, TEXT("Pool header must be able to fit into the smallest block"));
+	verify(FSizeTableEntry::FillSizeTable(OsAllocationGranularity, SizeTable, UE_MB3_BASE_PAGE_SIZE, UE_MB3_MINIMUM_ALIGNMENT, UE_MB3_MAX_SMALL_POOL_SIZE, UE_MB3_BASE_PAGE_SIZE) == UE_MB3_SMALL_POOL_COUNT);
+	checkf(SizeTable[UE_MB3_SMALL_POOL_COUNT - 1].BinSize == UE_MB3_MAX_SMALL_POOL_SIZE, TEXT("UE_MB3_MAX_SMALL_POOL_SIZE must be equal to the largest bin size"));
+	checkf(sizeof(FMallocBinned3::FFreeBlock) <= SizeTable[0].BinSize, TEXT("Pool header must be able to fit into the smallest bin"));
 
 	SmallPoolInfosPerPlatformPage = OsAllocationGranularity / sizeof(FPoolInfoSmall);
 
 	uint32 RequiredMetaMem = 0;
-	for (uint32 Index = 0; Index < BINNED3_SMALL_POOL_COUNT; ++Index)
+	for (uint32 Index = 0; Index < UE_MB3_SMALL_POOL_COUNT; ++Index)
 	{
-		checkf(Index == 0 || SizeTable[Index - 1].BlockSize < SizeTable[Index].BlockSize, TEXT("Small block sizes must be strictly increasing"));
-		checkf(SizeTable[Index].BlockSize % BINNED3_MINIMUM_ALIGNMENT == 0, TEXT("Small block size must be a multiple of BINNED3_MINIMUM_ALIGNMENT"));
+		checkf(Index == 0 || SizeTable[Index - 1].BinSize < SizeTable[Index].BinSize, TEXT("Small bin sizes must be strictly increasing"));
+		checkf(SizeTable[Index].BinSize % UE_MB3_MINIMUM_ALIGNMENT == 0, TEXT("Small bin size must be a multiple of UE_MB3_MINIMUM_ALIGNMENT"));
 
-		SmallPoolTables[Index].BlockSize = SizeTable[Index].BlockSize;
-		SmallPoolTables[Index].BlocksPerBlockOfBlocks = SizeTable[Index].BlocksPerBlockOfBlocks;
-		SmallPoolTables[Index].PagesPlatformForBlockOfBlocks = SizeTable[Index].PagesPlatformForBlockOfBlocks;
+		SmallPoolTables[Index].BinSize = SizeTable[Index].BinSize;
+		SmallPoolTables[Index].NumMemoryPagesPerBlock = SizeTable[Index].NumMemoryPagesPerBlock;
 
 		SmallPoolTables[Index].UnusedAreaOffsetLow = 0;
-		SmallPoolTables[Index].NumEverUsedBlockOfBlocks = 0;
-#if BINNED3_ALLOCATOR_PER_BIN_STATS
+		SmallPoolTables[Index].NumEverUsedBlocks = 0;
+#if UE_M3_ALLOCATOR_PER_BIN_STATS
 		SmallPoolTables[Index].TotalRequestedAllocSize.Store(0);
 		SmallPoolTables[Index].TotalAllocCount.Store(0);
 		SmallPoolTables[Index].TotalFreeCount.Store(0);
 #endif
 
-		const int64 TotalNumberOfBlocksOfBlocks = MAX_MEMORY_PER_BLOCK_SIZE / (SizeTable[Index].PagesPlatformForBlockOfBlocks * OsAllocationGranularity);
-		const uint32 Size = Align(sizeof(FPoolInfoSmall**) * (TotalNumberOfBlocksOfBlocks + SmallPoolInfosPerPlatformPage - 1) / SmallPoolInfosPerPlatformPage, PLATFORM_CACHE_LINE_SIZE);
+		const int64 TotalNumberOfBlocks = UE_MB3_MAX_MEMORY_PER_POOL_SIZE / (SizeTable[Index].NumMemoryPagesPerBlock * OsAllocationGranularity);
+		const uint32 Size = Align(sizeof(FPoolInfoSmall**) * (TotalNumberOfBlocks + SmallPoolInfosPerPlatformPage - 1) / SmallPoolInfosPerPlatformPage, PLATFORM_CACHE_LINE_SIZE);
 		RequiredMetaMem += Size;
-#if BINNED3_ALLOCATOR_STATS
+#if UE_MB3_ALLOCATOR_STATS
 		Binned3PoolInfoMemory += Size;
 #endif
 
-		const int64 AllocationSize = Align(FBitTree::GetMemoryRequirements(TotalNumberOfBlocksOfBlocks), PLATFORM_CACHE_LINE_SIZE);
+		const int64 AllocationSize = Align(FBitTree::GetMemoryRequirements(TotalNumberOfBlocks), PLATFORM_CACHE_LINE_SIZE);
 		RequiredMetaMem += AllocationSize * 2;
-#if BINNED3_ALLOCATOR_STATS
+#if UE_MB3_ALLOCATOR_STATS
 		Binned3FreeBitsMemory += AllocationSize * 2;
 #endif
 	}
@@ -815,50 +811,50 @@ FMallocBinned3::FMallocBinned3()
 	const uint8* MetaMemEnd = MetaMem + RequiredMetaMem;
 	FMemory::Memzero(MetaMem, RequiredMetaMem);
 
-	for (uint32 Index = 0; Index < BINNED3_SMALL_POOL_COUNT; ++Index)
+	for (uint32 Index = 0; Index < UE_MB3_SMALL_POOL_COUNT; ++Index)
 	{
-		const int64 TotalNumberOfBlocksOfBlocks = MAX_MEMORY_PER_BLOCK_SIZE / (SizeTable[Index].PagesPlatformForBlockOfBlocks * OsAllocationGranularity);
-		const uint32 Size = Align(sizeof(FPoolInfoSmall**) * (TotalNumberOfBlocksOfBlocks + SmallPoolInfosPerPlatformPage - 1) / SmallPoolInfosPerPlatformPage, PLATFORM_CACHE_LINE_SIZE);
+		const int64 TotalNumberOfBlocks = UE_MB3_MAX_MEMORY_PER_POOL_SIZE / (SizeTable[Index].NumMemoryPagesPerBlock * OsAllocationGranularity);
+		const uint32 Size = Align(sizeof(FPoolInfoSmall**) * (TotalNumberOfBlocks + SmallPoolInfosPerPlatformPage - 1) / SmallPoolInfosPerPlatformPage, PLATFORM_CACHE_LINE_SIZE);
 
 		SmallPoolTables[Index].PoolInfos = (FPoolInfoSmall**)MetaMem;
 		MetaMem += Size;
 
-		const int64 AllocationSize = Align(FBitTree::GetMemoryRequirements(TotalNumberOfBlocksOfBlocks), PLATFORM_CACHE_LINE_SIZE);
-		SmallPoolTables[Index].BlockOfBlockAllocationBits.FBitTreeInit(TotalNumberOfBlocksOfBlocks, MetaMem, AllocationSize, false);
+		const int64 AllocationSize = Align(FBitTree::GetMemoryRequirements(TotalNumberOfBlocks), PLATFORM_CACHE_LINE_SIZE);
+		SmallPoolTables[Index].BlocksAllocatedBits.FBitTreeInit(TotalNumberOfBlocks, MetaMem, AllocationSize, false);
 		MetaMem += AllocationSize;
 
-		SmallPoolTables[Index].BlockOfBlockIsExhausted.FBitTreeInit(TotalNumberOfBlocksOfBlocks, MetaMem, AllocationSize, true);
+		SmallPoolTables[Index].BlocksExhaustedBits.FBitTreeInit(TotalNumberOfBlocks, MetaMem, AllocationSize, true);
 		MetaMem += AllocationSize;
 	}
 	check(MetaMem <= MetaMemEnd);
 
 	// Set up pool mappings
-	uint8* IndexEntry = MemSizeToIndex;
-	uint32  PoolIndex  = 0;
-	for (uint32 Index = 0; Index != 1 + (BINNED3_MAX_SMALL_POOL_SIZE >> BINNED3_MINIMUM_ALIGNMENT_SHIFT); ++Index)
+	uint8* IndexEntry = MemSizeToPoolIndex;
+	uint32 PoolIndex  = 0;
+	for (uint32 Index = 0; Index != 1 + (UE_MB3_MAX_SMALL_POOL_SIZE >> UE_MB3_MINIMUM_ALIGNMENT_SHIFT); ++Index)
 	{
-		uint32 BlockSize = Index << BINNED3_MINIMUM_ALIGNMENT_SHIFT; // inverse of int32 Index = int32((Size >> BINNED3_MINIMUM_ALIGNMENT_SHIFT));
-		while (SizeTable[PoolIndex].BlockSize < BlockSize)
+		const uint32 BinSize = Index << UE_MB3_MINIMUM_ALIGNMENT_SHIFT; // inverse of int32 Index = int32((Size >> UE_MB3_MINIMUM_ALIGNMENT_SHIFT));
+		while (SizeTable[PoolIndex].BinSize < BinSize)
 		{
 			++PoolIndex;
-			check(PoolIndex != BINNED3_SMALL_POOL_COUNT);
+			check(PoolIndex != UE_MB3_SMALL_POOL_COUNT);
 		}
 		check(PoolIndex < 256);
 		*IndexEntry++ = uint8(PoolIndex);
 	}
 
 	// now reverse the pool sizes for cache coherency
-	for (uint32 Index = 0; Index != BINNED3_SMALL_POOL_COUNT; ++Index)
+	for (uint32 Index = 0; Index != UE_MB3_SMALL_POOL_COUNT; ++Index)
 	{
-		uint32 Partner = BINNED3_SMALL_POOL_COUNT - Index - 1;
-		SmallBlockSizesReversedShifted[Index] = (SizeTable[Partner].BlockSize >> BINNED3_MINIMUM_ALIGNMENT_SHIFT);
+		uint32 Partner = UE_MB3_SMALL_POOL_COUNT - Index - 1;
+		SmallBinSizesReversedShifted[Index] = (SizeTable[Partner].BinSize >> UE_MB3_MINIMUM_ALIGNMENT_SHIFT);
 	}
 
 	uint64 MaxHashBuckets = PtrToPoolMapping.GetMaxHashBuckets();
 	{
 		int64 HashAllocSize = Align(MaxHashBuckets * sizeof(PoolHashBucket), OsAllocationGranularity);
 		HashBuckets = (PoolHashBucket*)AllocateMetaDataMemory(HashAllocSize);
-#if BINNED3_ALLOCATOR_STATS
+#if UE_MB3_ALLOCATOR_STATS
 		Binned3HashMemory += HashAllocSize;
 #endif
 		verify(HashBuckets);
@@ -875,24 +871,27 @@ FMallocBinned3::~FMallocBinned3()
 
 void FMallocBinned3::Commit(uint32 InPoolIndex, void *Ptr, SIZE_T Size)
 {
-
-#if BINNED3_ALLOCATOR_STATS
+#if UE_MB3_ALLOCATOR_STATS
 	Binned3Commits++;
 #endif
+
 #if !BINNED3_USE_SEPARATE_VM_PER_POOL
 	Binned3BaseVMBlock.CommitByPtr(Ptr, Size);
 #else
 	PoolBaseVMBlock[InPoolIndex].CommitByPtr(Ptr, Size);
 #endif
+
 	LLM_IF_ENABLED(FLowLevelMemTracker::Get().OnLowLevelAlloc(ELLMTracker::Platform, Ptr, Size));
 }
+
 void FMallocBinned3::Decommit(uint32 InPoolIndex, void *Ptr, SIZE_T Size)
 {
 	LLM_IF_ENABLED(FLowLevelMemTracker::Get().OnLowLevelFree(ELLMTracker::Platform, Ptr));
 
-#if BINNED3_ALLOCATOR_STATS
+#if UE_MB3_ALLOCATOR_STATS
 	Binned3Decommits++;
 #endif
+
 #if !BINNED3_USE_SEPARATE_VM_PER_POOL
 	Binned3BaseVMBlock.DecommitByPtr(Ptr, Size);
 #else
@@ -903,9 +902,9 @@ void FMallocBinned3::Decommit(uint32 InPoolIndex, void *Ptr, SIZE_T Size)
 void* FMallocBinned3::AllocateMetaDataMemory(SIZE_T Size)
 {
 	LLM_PLATFORM_SCOPE(ELLMTag::FMalloc);
-	size_t VirtualAlignedSize = Align(Size, FPlatformMemory::FPlatformVirtualMemoryBlock::GetVirtualSizeAlignment());
+	const size_t VirtualAlignedSize = Align(Size, FPlatformMemory::FPlatformVirtualMemoryBlock::GetVirtualSizeAlignment());
 	FPlatformMemory::FPlatformVirtualMemoryBlock Block = FPlatformMemory::FPlatformVirtualMemoryBlock::AllocateVirtual(VirtualAlignedSize);
-	size_t CommitAlignedSize = Align(Size, FPlatformMemory::FPlatformVirtualMemoryBlock::GetCommitAlignment());
+	const size_t CommitAlignedSize = Align(Size, FPlatformMemory::FPlatformVirtualMemoryBlock::GetCommitAlignment());
 	Block.Commit(0, CommitAlignedSize);
 	LLM_IF_ENABLED(FLowLevelMemTracker::Get().OnLowLevelAlloc(ELLMTracker::Platform, Block.GetVirtualPointer(), CommitAlignedSize));
 	return Block.GetVirtualPointer();
@@ -930,36 +929,36 @@ bool FMallocBinned3::IsInternallyThreadSafe() const
 
 void* FMallocBinned3::MallocExternal(SIZE_T Size, uint32 Alignment)
 {
-	static_assert(DEFAULT_ALIGNMENT <= BINNED3_MINIMUM_ALIGNMENT, "DEFAULT_ALIGNMENT is assumed to be zero"); // used below
+	static_assert(DEFAULT_ALIGNMENT <= UE_MB3_MINIMUM_ALIGNMENT, "DEFAULT_ALIGNMENT is assumed to be zero"); // used below
 
 	// Fast path: Allocate from the small pools if the size is small enough and the alignment <= binned3 min alignment.
 	//            Larger alignments can waste a lot of memory allocating an entire page, so some smaller alignments are 
 	//			  handled in the fallback path if less than a predefined max small pool alignment. 
 
-	bool UsePools = (Size <= BINNED3_MAX_SMALL_POOL_SIZE) && (Alignment <= BINNED3_MINIMUM_ALIGNMENT);
+	bool UsePools = (Size <= UE_MB3_MAX_SMALL_POOL_SIZE) && (Alignment <= UE_MB3_MINIMUM_ALIGNMENT);
 	
 	if (!UsePools)
 	{
-		// fallback: test for non-default/mininum alignments and handle a subset of them to reduce memory waste of 64KB page requirements
-		// 			 e.g. code that wants to use aligned avx loads will use 32 byte alignments 
-
+		// check if allocations that require alignment larger than UE_MB3_MINIMUM_ALIGNMENT can be promoted to a bin with a natural alignment that matches
+		// i.e. 16 bytes allocation with 128 bytes alignment can be promoted to 128 bytes bin
+		// this will save us a lot of memory as otherwise allocations will be promoted to OS allocs that are at least 64 KB large, depending on a UE_MB3_MAX_SMALL_POOL_SIZE
 		UsePools = PromoteToLargerBin(Size, Alignment, *this);
 	}
 
 	if (UsePools) 
 	{
-		uint32 PoolIndex = BoundSizeToPoolIndex(Size, MemSizeToIndex);
-		FPerThreadFreeBlockLists* Lists = GMallocBinned3PerThreadCaches ? FPerThreadFreeBlockLists::Get() : nullptr;
+		const uint32 PoolIndex = BoundSizeToPoolIndex(Size, MemSizeToPoolIndex);
+		FPerThreadFreeBlockLists* Lists = GBinned3PerThreadCaches ? FPerThreadFreeBlockLists::Get() : nullptr;
 		if (Lists)
 		{
 			if (Lists->ObtainRecycledPartial(PoolIndex, Private::GGlobalRecycler))
 			{
 				if (void* Result = Lists->Malloc(PoolIndex))
 				{
-#if BINNED3_ALLOCATOR_STATS
+#if UE_MB3_ALLOCATOR_STATS
 					SmallPoolTables[PoolIndex].HeadEndAlloc(Size);
-					uint32 BlockSize = PoolIndexToBlockSize(PoolIndex);
-					Lists->AllocatedMemory += BlockSize;
+					const uint32 BinSize = PoolIndexToBinSize(PoolIndex);
+					Lists->AllocatedMemory += BinSize;
 #endif
 					return Result;
 				}
@@ -971,71 +970,71 @@ void* FMallocBinned3::MallocExternal(SIZE_T Size, uint32 Alignment)
 		// Allocate from small object pool.
 		FPoolTable& Table = SmallPoolTables[PoolIndex];
 
-		uint32 BlockOfBlocksIndex = MAX_uint32;
-		FPoolInfoSmall* Pool = GetFrontPool(Table, PoolIndex, BlockOfBlocksIndex);
+		uint32 BlockIndex = MAX_uint32;
+		FPoolInfoSmall* Pool = GetFrontPool(Table, PoolIndex, BlockIndex);
 		if (!Pool)
 		{
-			Pool = PushNewPoolToFront(Table, Table.BlockSize, PoolIndex, BlockOfBlocksIndex);
+			Pool = PushNewPoolToFront(Table, Table.BinSize, PoolIndex, BlockIndex);
 			
-			//Indicates that we run out of Pool memory (512 MB) for this block type
+			//Indicates that we run out of Pool memory (512 MB) for this bin type
 			if (!Pool)
 			{
-				if ((PoolIndex + 1) < BINNED3_SMALL_POOL_COUNT)
+				if ((PoolIndex + 1) < UE_MB3_SMALL_POOL_COUNT)
 				{
-					return MallocExternal(SmallPoolTables[PoolIndex + 1].BlockSize, Alignment);
+					return MallocExternal(SmallPoolTables[PoolIndex + 1].BinSize, Alignment);
 				}
 				else
 				{
-					return MallocExternal(BINNED3_MAX_SMALL_POOL_SIZE + 1, Alignment);
+					return MallocExternal(UE_MB3_MAX_SMALL_POOL_SIZE + 1, Alignment);
 				}
 			}
 		}
 
-		const uint32 BlockOfBlocksSize = OsAllocationGranularity * Table.PagesPlatformForBlockOfBlocks;
-		uint8* BlockOfBlocksPtr = BlockPointerFromIndecies(PoolIndex, BlockOfBlocksIndex, BlockOfBlocksSize);
+		const uint32 BlockSize = OsAllocationGranularity * Table.NumMemoryPagesPerBlock;
+		uint8* BlockPtr = BlockPointerFromIndecies(PoolIndex, BlockIndex, BlockSize);
 
-		void* Result = Pool->AllocateRegularBlock(BlockOfBlocksPtr, Table.BlockSize);
-#if BINNED3_ALLOCATOR_STATS
+		void* Result = Pool->AllocateBin(BlockPtr, Table.BinSize);
+#if UE_MB3_ALLOCATOR_STATS
 		Table.HeadEndAlloc(Size);
-		Binned3AllocatedSmallPoolMemory += PoolIndexToBlockSize(PoolIndex);
-#endif // BINNED3_ALLOCATOR_STATS
-		if (GMallocBinned3AllocExtra)
+		Binned3AllocatedSmallPoolMemory += PoolIndexToBinSize(PoolIndex);
+#endif
+		if (GBinned3AllocExtra)
 		{
 			if (Lists)
 			{
 				// prefill the free list with some allocations so we are less likely to hit this slow path with the mutex 
-				for (int32 Index = 0; Index < GMallocBinned3AllocExtra && Pool->HasFreeRegularBlock(); Index++)
+				for (int32 Index = 0; Index < GBinned3AllocExtra && Pool->HasFreeBin(); Index++)
 				{
-					if (!Lists->Free(Result, PoolIndex, Table.BlockSize))
+					if (!Lists->Free(Result, PoolIndex, Table.BinSize))
 					{
 						break;
 					}
-					Result = Pool->AllocateRegularBlock(BlockOfBlocksPtr, Table.BlockSize);
+					Result = Pool->AllocateBin(BlockPtr, Table.BinSize);
 				}
 			}
 		}
-		if (!Pool->HasFreeRegularBlock())
+		if (!Pool->HasFreeBin())
 		{
-			Table.BlockOfBlockIsExhausted.AllocBit(BlockOfBlocksIndex);
+			Table.BlocksExhaustedBits.AllocBit(BlockIndex);
 		}
 
 		return Result;
 	}
-	Alignment = FMath::Max<uint32>(Alignment, BINNED3_MINIMUM_ALIGNMENT);
+	Alignment = FMath::Max<uint32>(Alignment, UE_MB3_MINIMUM_ALIGNMENT);
 	Size = Align(FMath::Max((SIZE_T)1, Size), Alignment);
 
 	check(FMath::IsPowerOfTwo(Alignment));
 
 	// Use OS for non-pooled allocations.
-	UPTRINT AlignedSize = Align(Size, FPlatformMemory::FPlatformVirtualMemoryBlock::GetCommitAlignment());
+	const uint64 AlignedSize = Align(Size, FPlatformMemory::FPlatformVirtualMemoryBlock::GetCommitAlignment());
 
-#if BINNED3_TIME_LARGE_BLOCKS
-	double StartTime = FPlatformTime::Seconds();
+#if UE_MB3_TIME_LARGE_BLOCKS
+	const double StartTime = FPlatformTime::Seconds();
 #endif
 
 	LLM_PLATFORM_SCOPE(ELLMTag::FMalloc);
 
-#if USE_CACHED_PAGE_ALLOCATOR_FOR_LARGE_ALLOCS
+#if UE_MB3_USE_CACHED_PAGE_ALLOCATOR_FOR_LARGE_ALLOCS
 	FScopeLock Lock(&Mutex);
 	void* Result = GetCachedOSPageAllocator().Allocate(AlignedSize);
 	check(IsAligned(Result, Alignment));
@@ -1047,8 +1046,8 @@ void* FMallocBinned3::MallocExternal(SIZE_T Size, uint32 Alignment)
 	LLM_IF_ENABLED(FLowLevelMemTracker::Get().OnLowLevelAlloc(ELLMTracker::Platform, Result, AlignedSize));
 #endif
 
-#if BINNED3_TIME_LARGE_BLOCKS
-	double Add = FPlatformTime::Seconds() - StartTime;
+#if UE_MB3_TIME_LARGE_BLOCKS
+	const double Add = FPlatformTime::Seconds() - StartTime;
 	double Old;
 	do
 	{
@@ -1057,27 +1056,27 @@ void* FMallocBinned3::MallocExternal(SIZE_T Size, uint32 Alignment)
 	MemoryRangeReserveTotalCount++;
 #endif
 
-	UE_CLOG(!IsAligned(Result, Alignment) ,LogMemory, Fatal, TEXT("FMallocBinned3 alignment was too large for OS. Alignment=%d   Ptr=%p"), Alignment, Result);
+	UE_CLOG(!IsAligned(Result, Alignment) ,LogMemory, Fatal, TEXT("FMallocBinned3 alignment was too large for OS. Alignment=%d Ptr=%p"), Alignment, Result);
 
 	if (!Result)
 	{
 		Private::OutOfMemory(AlignedSize);
 	}
 	check(IsOSAllocation(Result));
-#if! USE_CACHED_PAGE_ALLOCATOR_FOR_LARGE_ALLOCS
+
+#if! UE_MB3_USE_CACHED_PAGE_ALLOCATOR_FOR_LARGE_ALLOCS
 	FScopeLock Lock(&Mutex);
 #endif
 
-
-#if BINNED3_ALLOCATOR_STATS
-	Binned3AllocatedLargePoolMemory += (int64)Size;
+#if UE_MB3_ALLOCATOR_STATS
+	Binned3AllocatedLargePoolMemory += Size;
 	Binned3AllocatedLargePoolMemoryWAlignment += AlignedSize;
 #endif
 
 	// Create pool.
 	FPoolInfoLarge* Pool = Private::GetOrCreatePoolInfoLarge(*this, Result);
 	check(Size > 0 && Size <= AlignedSize && AlignedSize >= FPlatformMemory::FPlatformVirtualMemoryBlock::GetCommitAlignment());
-#if USE_CACHED_PAGE_ALLOCATOR_FOR_LARGE_ALLOCS
+#if UE_MB3_USE_CACHED_PAGE_ALLOCATOR_FOR_LARGE_ALLOCS
 	Pool->SetOSAllocationSizes(Size, AlignedSize, AlignedSize / FPlatformMemory::FPlatformVirtualMemoryBlock::GetCommitAlignment());
 #else
 	Pool->SetOSAllocationSizes(Size, AlignedSize, Block.GetActualSizeInPages());
@@ -1086,7 +1085,6 @@ void* FMallocBinned3::MallocExternal(SIZE_T Size, uint32 Alignment)
 	return Result;
 }
 
-
 void* FMallocBinned3::ReallocExternal(void* Ptr, SIZE_T NewSize, uint32 Alignment)
 {
 	if (NewSize == 0)
@@ -1094,20 +1092,20 @@ void* FMallocBinned3::ReallocExternal(void* Ptr, SIZE_T NewSize, uint32 Alignmen
 		FMallocBinned3::FreeExternal(Ptr);
 		return nullptr;
 	}
-	static_assert(DEFAULT_ALIGNMENT <= BINNED3_MINIMUM_ALIGNMENT, "DEFAULT_ALIGNMENT is assumed to be zero"); // used below
+	static_assert(DEFAULT_ALIGNMENT <= UE_MB3_MINIMUM_ALIGNMENT, "DEFAULT_ALIGNMENT is assumed to be zero"); // used below
 	check(FMath::IsPowerOfTwo(Alignment));
 	check(Alignment <= OsAllocationGranularity);
 
-	uint64 PoolIndex = PoolIndexFromPtr(Ptr);
-	if (PoolIndex < BINNED3_SMALL_POOL_COUNT)
+	const uint64 PoolIndex = PoolIndexFromPtr(Ptr);
+	if (PoolIndex < UE_MB3_SMALL_POOL_COUNT)
 	{
 		check(Ptr); // null is an OS allocation because it will not fall in our VM block
-		uint32 BlockSize = PoolIndexToBlockSize(PoolIndex);
+		const uint32 BinSize = PoolIndexToBinSize(PoolIndex);
 		if (
-			((NewSize <= BlockSize) & (IsAligned(BlockSize, Alignment))) && // one branch, not two
-			(PoolIndex == 0 || NewSize > PoolIndexToBlockSize(PoolIndex - 1)))
+			((NewSize <= BinSize) & (IsAligned(BinSize, Alignment))) && // one branch, not two
+			(PoolIndex == 0 || NewSize > PoolIndexToBinSize(PoolIndex - 1)))
 		{
-#if BINNED3_ALLOCATOR_STATS
+#if UE_MB3_ALLOCATOR_STATS
 			SmallPoolTables[PoolIndex].HeadEndAlloc(NewSize);
 			SmallPoolTables[PoolIndex].HeadEndFree();
 #endif
@@ -1116,7 +1114,7 @@ void* FMallocBinned3::ReallocExternal(void* Ptr, SIZE_T NewSize, uint32 Alignmen
 
 		// Reallocate and copy the data across
 		void* Result = FMallocBinned3::MallocExternal(NewSize, Alignment);
-		FMemory::Memcpy(Result, Ptr, FMath::Min<SIZE_T>(NewSize, BlockSize));
+		FMemory::Memcpy(Result, Ptr, FMath::Min<SIZE_T>(NewSize, BinSize));
 		FMallocBinned3::FreeExternal(Ptr);
 		return Result;
 	}
@@ -1132,13 +1130,13 @@ void* FMallocBinned3::ReallocExternal(void* Ptr, SIZE_T NewSize, uint32 Alignmen
 	FPoolInfoLarge* Pool = Private::FindPoolInfo(*this, Ptr);
 	if (!Pool)
 	{
-		UE_LOG(LogMemory, Fatal, TEXT("FMallocBinned3 Attempt to realloc an unrecognized block %p"), Ptr);
+		UE_LOG(LogMemory, Fatal, TEXT("FMallocBinned3 Attempt to realloc an unrecognized pointer %p"), Ptr);
 	}
-	UPTRINT PoolOsBytes = Pool->GetOsCommittedBytes();
-	uint32 PoolOSRequestedBytes = Pool->GetOSRequestedBytes();
-	checkf(PoolOSRequestedBytes <= PoolOsBytes, TEXT("FMallocBinned3::ReallocExternal %d %d"), int32(PoolOSRequestedBytes), int32(PoolOsBytes));
+	const uint32 PoolOsBytes = Pool->GetOsCommittedBytes();
+	const uint32 PoolOSRequestedBytes = Pool->GetOSRequestedBytes();
+	checkf(PoolOSRequestedBytes <= PoolOsBytes, TEXT("FMallocBinned3::ReallocExternal %u %u"), PoolOSRequestedBytes, PoolOsBytes);
 	if (NewSize > PoolOsBytes || // can't fit in the old block
-		(NewSize <= BINNED3_MAX_SMALL_POOL_SIZE && Alignment <= BINNED3_MINIMUM_ALIGNMENT) || // can switch to the small block allocator
+		(NewSize <= UE_MB3_MAX_SMALL_POOL_SIZE && Alignment <= UE_MB3_MINIMUM_ALIGNMENT) || // can switch to the small bin allocator
 		Align(NewSize, OsAllocationGranularity) < PoolOsBytes) // we can get some pages back
 	{
 		Mutex.Unlock();
@@ -1150,34 +1148,34 @@ void* FMallocBinned3::ReallocExternal(void* Ptr, SIZE_T NewSize, uint32 Alignmen
 		return Result;
 	}
 
-#if BINNED3_ALLOCATOR_STATS
+#if UE_MB3_ALLOCATOR_STATS
 	Binned3AllocatedLargePoolMemory += ((int64)NewSize) - ((int64)PoolOSRequestedBytes);
 	// don't need to change the Binned3AllocatedLargePoolMemoryWAlignment because we didn't reallocate so it's the same size
 #endif
 	
-	Pool->SetOSAllocationSizes(NewSize);
+	Pool->SetOSAllocationSize(NewSize);
 	Mutex.Unlock();
 	return Ptr;
 }
 
 void FMallocBinned3::FreeExternal(void* Ptr)
 {
-	uint64 PoolIndex = PoolIndexFromPtr(Ptr);
-	if (PoolIndex < BINNED3_SMALL_POOL_COUNT)
+	const uint64 PoolIndex = PoolIndexFromPtr(Ptr);
+	if (PoolIndex < UE_MB3_SMALL_POOL_COUNT)
 	{
 		check(Ptr); // null is an OS allocation because it will not fall in our VM block
-		uint32 BlockSize = PoolIndexToBlockSize(PoolIndex);
+		const uint32 BinSize = PoolIndexToBinSize(PoolIndex);
 
 		FBundleNode* BundlesToRecycle = nullptr;
-		FPerThreadFreeBlockLists* Lists = GMallocBinned3PerThreadCaches ? FPerThreadFreeBlockLists::Get() : nullptr;
+		FPerThreadFreeBlockLists* Lists = GBinned3PerThreadCaches ? FPerThreadFreeBlockLists::Get() : nullptr;
 		if (Lists)
 		{
 			BundlesToRecycle = Lists->RecycleFullBundle(PoolIndex, Private::GGlobalRecycler);
-			bool bPushed = Lists->Free(Ptr, PoolIndex, BlockSize);
+			const bool bPushed = Lists->Free(Ptr, PoolIndex, BinSize);
 			check(bPushed);
-#if BINNED3_ALLOCATOR_STATS
+#if UE_MB3_ALLOCATOR_STATS
 			SmallPoolTables[PoolIndex].HeadEndFree();
-			Lists->AllocatedMemory -= BlockSize;
+			Lists->AllocatedMemory -= BinSize;
 #endif
 		}
 		else
@@ -1189,60 +1187,60 @@ void FMallocBinned3::FreeExternal(void* Ptr)
 		{
 			BundlesToRecycle->NextBundle = nullptr;
 			FScopeLock Lock(&Mutex);
-			Private::FreeBundles(*this, BundlesToRecycle, BlockSize, PoolIndex);
-#if BINNED3_ALLOCATOR_STATS
+			Private::FreeBundles(*this, BundlesToRecycle, BinSize, PoolIndex);
+#if UE_MB3_ALLOCATOR_STATS
 			if (!Lists)
 			{
 				SmallPoolTables[PoolIndex].HeadEndFree();
 				// lists track their own stat track them instead in the global stat if we don't have lists
-				Binned3AllocatedSmallPoolMemory -= ((int64)(BlockSize));
+				Binned3AllocatedSmallPoolMemory -= ((int64)(BinSize));
 			}
 #endif
 		}
 	}
 	else if (Ptr)
 	{
-#if USE_CACHED_PAGE_ALLOCATOR_FOR_LARGE_ALLOCS
+#if UE_MB3_USE_CACHED_PAGE_ALLOCATOR_FOR_LARGE_ALLOCS
 		FScopeLock Lock(&Mutex);
 #endif
 		uint32 VMPages;
 		{
-#if !USE_CACHED_PAGE_ALLOCATOR_FOR_LARGE_ALLOCS
+#if !UE_MB3_USE_CACHED_PAGE_ALLOCATOR_FOR_LARGE_ALLOCS
 			FScopeLock Lock(&Mutex);
 #endif
 			FPoolInfoLarge* Pool = Private::FindPoolInfo(*this, Ptr);
 			if (!Pool)
 			{
-				UE_LOG(LogMemory, Fatal, TEXT("FMallocBinned3 Attempt to free an unrecognized block %p"), Ptr);
+				UE_LOG(LogMemory, Fatal, TEXT("FMallocBinned3 Attempt to free an unrecognized pointer %p"), Ptr);
 			}
-			UPTRINT PoolOsBytes = Pool->GetOsCommittedBytes();
-			uint32 PoolOSRequestedBytes = Pool->GetOSRequestedBytes();
+			const uint32 PoolOsBytes = Pool->GetOsCommittedBytes();
+			const uint32 PoolOSRequestedBytes = Pool->GetOSRequestedBytes();
 			VMPages = Pool->GetOsVMPages();
 
-#if BINNED3_ALLOCATOR_STATS
-			Binned3AllocatedLargePoolMemory -= ((int64)PoolOSRequestedBytes);
-			Binned3AllocatedLargePoolMemoryWAlignment -= ((int64)PoolOsBytes);
+#if UE_MB3_ALLOCATOR_STATS
+			Binned3AllocatedLargePoolMemory -= PoolOSRequestedBytes;
+			Binned3AllocatedLargePoolMemoryWAlignment -= PoolOsBytes;
 #endif
 
-			checkf(PoolOSRequestedBytes <= PoolOsBytes, TEXT("FMallocBinned3::FreeExternal %d %d"), int32(PoolOSRequestedBytes), int32(PoolOsBytes));
+			checkf(PoolOSRequestedBytes <= PoolOsBytes, TEXT("FMallocBinned3::FreeExternal %u %u"), int32(PoolOSRequestedBytes), int32(PoolOsBytes));
 			Pool->SetCanary(FPoolInfoLarge::ECanary::LargeUnassigned, true, false);
 		}
 
 		// Free an OS allocation.
-#if BINNED3_TIME_LARGE_BLOCKS
-		double StartTime = FPlatformTime::Seconds();
+#if UE_MB3_TIME_LARGE_BLOCKS
+		const double StartTime = FPlatformTime::Seconds();
 #endif
 		{
 			LLM_IF_ENABLED(FLowLevelMemTracker::Get().OnLowLevelFree(ELLMTracker::Platform, Ptr));
-#if USE_CACHED_PAGE_ALLOCATOR_FOR_LARGE_ALLOCS
+#if UE_MB3_USE_CACHED_PAGE_ALLOCATOR_FOR_LARGE_ALLOCS
 			GetCachedOSPageAllocator().Free(Ptr, VMPages * FPlatformMemory::FPlatformVirtualMemoryBlock::GetCommitAlignment());
 #else
 			FPlatformMemory::FPlatformVirtualMemoryBlock Block(Ptr, VMPages);
 			Block.FreeVirtual();
 #endif
 		}
-#if BINNED3_TIME_LARGE_BLOCKS
-		double Add = FPlatformTime::Seconds() - StartTime;
+#if UE_MB3_TIME_LARGE_BLOCKS
+		const double Add = FPlatformTime::Seconds() - StartTime;
 		double Old;
 		do
 		{
@@ -1255,11 +1253,11 @@ void FMallocBinned3::FreeExternal(void* Ptr)
 
 bool FMallocBinned3::GetAllocationSizeExternal(void* Ptr, SIZE_T& SizeOut)
 {
-	uint64 PoolIndex = PoolIndexFromPtr(Ptr);
-	if (PoolIndex < BINNED3_SMALL_POOL_COUNT)
+	const uint64 PoolIndex = PoolIndexFromPtr(Ptr);
+	if (PoolIndex < UE_MB3_SMALL_POOL_COUNT)
 	{
 		check(Ptr); // null is an OS allocation because it will not fall in our VM block
-		SizeOut = PoolIndexToBlockSize(PoolIndex);
+		SizeOut = PoolIndexToBinSize(PoolIndex);
 		return true;
 	}
 	if (!Ptr)
@@ -1270,11 +1268,11 @@ bool FMallocBinned3::GetAllocationSizeExternal(void* Ptr, SIZE_T& SizeOut)
 	FPoolInfoLarge* Pool = Private::FindPoolInfo(*this, Ptr);
 	if (!Pool)
 	{
-		UE_LOG(LogMemory, Fatal, TEXT("FMallocBinned3 Attempt to GetAllocationSizeExternal an unrecognized block %p"), Ptr);
+		UE_LOG(LogMemory, Fatal, TEXT("FMallocBinned3 Attempt to GetAllocationSizeExternal an unrecognized pointer %p"), Ptr);
 	}
-	UPTRINT PoolOsBytes = Pool->GetOsCommittedBytes();
-	uint32 PoolOSRequestedBytes = Pool->GetOSRequestedBytes();
-	checkf(PoolOSRequestedBytes <= PoolOsBytes, TEXT("FMallocBinned3::GetAllocationSizeExternal %d %d"), int32(PoolOSRequestedBytes), int32(PoolOsBytes));
+	const uint32 PoolOsBytes = Pool->GetOsCommittedBytes();
+	const uint32 PoolOSRequestedBytes = Pool->GetOSRequestedBytes();
+	checkf(PoolOSRequestedBytes <= PoolOsBytes, TEXT("FMallocBinned3::GetAllocationSizeExternal %u %u"), PoolOSRequestedBytes, PoolOsBytes);
 	SizeOut = PoolOsBytes;
 	return true;
 }
@@ -1282,7 +1280,7 @@ bool FMallocBinned3::GetAllocationSizeExternal(void* Ptr, SIZE_T& SizeOut)
 bool FMallocBinned3::ValidateHeap()
 {
 	// Not implemented
-	// NumEverUsedBlockOfBlocks gives us all of the information we need to examine each pool, so it is doable.
+	// NumEverUsedBlocks gives us all of the information we need to examine each pool, so it is doable.
 	return true;
 }
 
@@ -1293,7 +1291,7 @@ const TCHAR* FMallocBinned3::GetDescriptiveName()
 
 void FMallocBinned3::Trim(bool bTrimThreadCaches)
 {
-	if (GMallocBinned3PerThreadCaches && bTrimThreadCaches)
+	if (GBinned3PerThreadCaches && bTrimThreadCaches)
 	{
 		// Trim memory and increase the Epoch.
 		FMallocBinnedCommonUtils::Trim(*this);
@@ -1312,7 +1310,7 @@ TArray<FMallocBinned3::FPerThreadFreeBlockLists*>& FMallocBinned3::GetRegistered
 
 void FMallocBinned3::SetupTLSCachesOnCurrentThread()
 {
-	if (!BINNED3_ALLOW_RUNTIME_TWEAKING && !GMallocBinned3PerThreadCaches)
+	if (!UE_MB3_ALLOW_RUNTIME_TWEAKING && !GBinned3PerThreadCaches)
 	{
 		return;
 	}
@@ -1326,7 +1324,7 @@ void FMallocBinned3::SetupTLSCachesOnCurrentThread()
 
 void FMallocBinned3::ClearAndDisableTLSCachesOnCurrentThread()
 {
-	if (!BINNED3_ALLOW_RUNTIME_TWEAKING && !GMallocBinned3PerThreadCaches)
+	if (!UE_MB3_ALLOW_RUNTIME_TWEAKING && !GBinned3PerThreadCaches)
 	{
 		return;
 	}
@@ -1337,7 +1335,7 @@ void FMallocBinned3::ClearAndDisableTLSCachesOnCurrentThread()
 
 void FMallocBinned3::MarkTLSCachesAsUsedOnCurrentThread()
 {
-	if (!BINNED3_ALLOW_RUNTIME_TWEAKING && !GMallocBinned3PerThreadCaches)
+	if (!UE_MB3_ALLOW_RUNTIME_TWEAKING && !GBinned3PerThreadCaches)
 	{
 		return;
 	}
@@ -1347,7 +1345,7 @@ void FMallocBinned3::MarkTLSCachesAsUsedOnCurrentThread()
 
 void FMallocBinned3::MarkTLSCachesAsUnusedOnCurrentThread()
 {
-	if (!BINNED3_ALLOW_RUNTIME_TWEAKING && !GMallocBinned3PerThreadCaches)
+	if (!UE_MB3_ALLOW_RUNTIME_TWEAKING && !GBinned3PerThreadCaches)
 	{
 		return;
 	}
@@ -1360,10 +1358,10 @@ void FMallocBinned3::MarkTLSCachesAsUnusedOnCurrentThread()
 
 void FMallocBinned3::FFreeBlock::CanaryFail() const
 {
-	UE_LOG(LogMemory, Fatal, TEXT("FMallocBinned3 Attempt to realloc an unrecognized block %p   canary == 0x%x != 0x%x"), (void*)this, (int32)Canary, (int32)FMallocBinned3::FFreeBlock::CANARY_VALUE);
+	UE_LOG(LogMemory, Fatal, TEXT("FMallocBinned3 Attempt to realloc an unrecognized pointer %p   canary == 0x%x != 0x%x"), (void*)this, (int32)Canary, (int32)FMallocBinned3::FFreeBlock::CANARY_VALUE);
 }
 
-#if BINNED3_ALLOCATOR_STATS
+#if UE_MB3_ALLOCATOR_STATS
 int64 FMallocBinned3::GetTotalAllocatedSmallPoolMemory() const
 {
 	int64 FreeBlockAllocatedMemory = 0;
@@ -1380,19 +1378,19 @@ int64 FMallocBinned3::GetTotalAllocatedSmallPoolMemory() const
 }
 #endif
 
-void FMallocBinned3::GetAllocatorStats( FGenericMemoryStats& OutStats )
+void FMallocBinned3::GetAllocatorStats(FGenericMemoryStats& OutStats)
 {
-#if BINNED3_ALLOCATOR_STATS
+#if UE_MB3_ALLOCATOR_STATS
 
-	int64 TotalAllocatedSmallPoolMemory = GetTotalAllocatedSmallPoolMemory();
+	const int64 TotalAllocatedSmallPoolMemory = GetTotalAllocatedSmallPoolMemory();
 
 	OutStats.Add(TEXT("Binned3AllocatedSmallPoolMemory"), TotalAllocatedSmallPoolMemory);
 	OutStats.Add(TEXT("Binned3AllocatedOSSmallPoolMemory"), Binned3AllocatedOSSmallPoolMemory);
 	OutStats.Add(TEXT("Binned3AllocatedLargePoolMemory"), Binned3AllocatedLargePoolMemory);
 	OutStats.Add(TEXT("Binned3AllocatedLargePoolMemoryWAlignment"), Binned3AllocatedLargePoolMemoryWAlignment);
 
-	uint64 TotalAllocated = TotalAllocatedSmallPoolMemory + Binned3AllocatedLargePoolMemory;
-	uint64 TotalOSAllocated = Binned3AllocatedOSSmallPoolMemory + Binned3AllocatedLargePoolMemoryWAlignment;
+	const uint64 TotalAllocated = TotalAllocatedSmallPoolMemory + Binned3AllocatedLargePoolMemory;
+	const uint64 TotalOSAllocated = Binned3AllocatedOSSmallPoolMemory + Binned3AllocatedLargePoolMemoryWAlignment;
 
 	OutStats.Add(TEXT("TotalAllocated"), TotalAllocated);
 	OutStats.Add(TEXT("TotalOSAllocated"), TotalOSAllocated);
@@ -1400,26 +1398,24 @@ void FMallocBinned3::GetAllocatorStats( FGenericMemoryStats& OutStats )
 	FMalloc::GetAllocatorStats(OutStats);
 }
 
-#if BINNED3_ALLOCATOR_STATS && BINNED3_USE_SEPARATE_VM_PER_POOL
-void FMallocBinned3::RecordPoolSearch(uint32 Tests)
+#if UE_MB3_ALLOCATOR_STATS && BINNED3_USE_SEPARATE_VM_PER_POOL
+void FMallocBinned3::RecordPoolSearch(uint32 Tests) const
 {
 	Binned3TotalPoolSearches++;
 	Binned3TotalPointerTests += Tests;
 }
 #endif
 
-
 void FMallocBinned3::DumpAllocatorStats(class FOutputDevice& Ar)
 {
-#if BINNED3_ALLOCATOR_STATS
-
-	int64 TotalAllocatedSmallPoolMemory = GetTotalAllocatedSmallPoolMemory();
+#if UE_MB3_ALLOCATOR_STATS
+	const int64 TotalAllocatedSmallPoolMemory = GetTotalAllocatedSmallPoolMemory();
 
 	Ar.Logf(TEXT("FMallocBinned3 Mem report"));
 	Ar.Logf(TEXT("Constants.BinnedAllocationGranularity = %d"), int32(OsAllocationGranularity));
-	Ar.Logf(TEXT("BINNED3_MAX_SMALL_POOL_SIZE = %d"), int32(BINNED3_MAX_SMALL_POOL_SIZE));
-	Ar.Logf(TEXT("MAX_MEMORY_PER_BLOCK_SIZE = %llu"), uint64(MAX_MEMORY_PER_BLOCK_SIZE));
-	Ar.Logf(TEXT("Small Pool Allocations: %fmb  (including block size padding)"), ((double)TotalAllocatedSmallPoolMemory) / (1024.0f * 1024.0f));
+	Ar.Logf(TEXT("UE_MB3_MAX_SMALL_POOL_SIZE = %d"), int32(UE_MB3_MAX_SMALL_POOL_SIZE));
+	Ar.Logf(TEXT("UE_MB3_MAX_MEMORY_PER_POOL_SIZE = %llu"), uint64(UE_MB3_MAX_MEMORY_PER_POOL_SIZE));
+	Ar.Logf(TEXT("Small Pool Allocations: %fmb  (including bin size padding)"), ((double)TotalAllocatedSmallPoolMemory) / (1024.0f * 1024.0f));
 	Ar.Logf(TEXT("Small Pool OS Allocated: %fmb"), ((double)Binned3AllocatedOSSmallPoolMemory) / (1024.0f * 1024.0f));
 	Ar.Logf(TEXT("Large Pool Requested Allocations: %fmb"), ((double)Binned3AllocatedLargePoolMemory) / (1024.0f * 1024.0f));
 	Ar.Logf(TEXT("Large Pool OS Allocated: %fmb"), ((double)Binned3AllocatedLargePoolMemoryWAlignment) / (1024.0f * 1024.0f));
@@ -1434,8 +1430,8 @@ void FMallocBinned3::DumpAllocatorStats(class FOutputDevice& Ar)
 	if (PoolSearchDiv)
 	{
 		Ar.Logf(TEXT("%llu Pointer Searches   %llu Pointer Compares    %llu Compares/Search"), Binned3TotalPoolSearches.Load(), Binned3TotalPointerTests.Load(), Binned3TotalPointerTests.Load() / Binned3TotalPoolSearches.Load());
-		uint64 TotalMem = PoolBaseVMPtr[BINNED3_SMALL_POOL_COUNT - 1] + MAX_MEMORY_PER_BLOCK_SIZE - PoolBaseVMPtr[0];
-		uint64 MinimumMem = uint64(BINNED3_SMALL_POOL_COUNT) * MAX_MEMORY_PER_BLOCK_SIZE;
+		const uint64 TotalMem = PoolBaseVMPtr[UE_MB3_SMALL_POOL_COUNT - 1] + UE_MB3_MAX_MEMORY_PER_POOL_SIZE - PoolBaseVMPtr[0];
+		const uint64 MinimumMem = uint64(UE_MB3_SMALL_POOL_COUNT) * UE_MB3_MAX_MEMORY_PER_POOL_SIZE;
 		Ar.Logf(TEXT("Percent of gaps in the address range %6.4f  (hopefully < 1, or the searches above will suffer)"), 100.0f * (1.0f - float(MinimumMem) / float(TotalMem)));
 	}
 #else
@@ -1447,47 +1443,45 @@ void FMallocBinned3::DumpAllocatorStats(class FOutputDevice& Ar)
 			) / (1024.0f * 1024.0f));
 
 
-#if BINNED3_TIME_LARGE_BLOCKS
+#if UE_MB3_TIME_LARGE_BLOCKS
 	Ar.Logf(TEXT("MemoryRangeReserve %d calls %6.3fs    %6.3fus / call"), MemoryRangeReserveTotalCount.Load(), float(MemoryRangeReserveTotalTime.Load()), float(MemoryRangeReserveTotalTime.Load()) * 1000000.0f / float(MemoryRangeReserveTotalCount.Load()));
 	Ar.Logf(TEXT("MemoryRangeFree    %d calls %6.3fs    %6.3fus / call"), MemoryRangeFreeTotalCount.Load(), float(MemoryRangeFreeTotalTime.Load()), float(MemoryRangeFreeTotalTime.Load()) * 1000000.0f / float(MemoryRangeFreeTotalCount.Load()));
 #endif
 
-#if BINNED3_ALLOCATOR_PER_BIN_STATS
-	for (int32 PoolIndex = 0; PoolIndex < BINNED3_SMALL_POOL_COUNT; PoolIndex++)
+#if UE_M3_ALLOCATOR_PER_BIN_STATS
+	for (int32 PoolIndex = 0; PoolIndex < UE_MB3_SMALL_POOL_COUNT; PoolIndex++)
 	{
-		
-		int64 VM = SmallPoolTables[PoolIndex].UnusedAreaOffsetLow;
-		uint32 CommittedBlocks = SmallPoolTables[PoolIndex].BlockOfBlockAllocationBits.CountOnes(SmallPoolTables[PoolIndex].NumEverUsedBlockOfBlocks);
-		uint32 PartialBlocks = SmallPoolTables[PoolIndex].NumEverUsedBlockOfBlocks - SmallPoolTables[PoolIndex].BlockOfBlockIsExhausted.CountOnes(SmallPoolTables[PoolIndex].NumEverUsedBlockOfBlocks);
-		uint32 FullBlocks = CommittedBlocks - PartialBlocks;
-		int64 ComittedVM = VM - (SmallPoolTables[PoolIndex].NumEverUsedBlockOfBlocks - CommittedBlocks) * SmallPoolTables[PoolIndex].PagesPlatformForBlockOfBlocks * OsAllocationGranularity;
+		const int64 VM = SmallPoolTables[PoolIndex].UnusedAreaOffsetLow;
+		const uint32 CommittedBlocks = SmallPoolTables[PoolIndex].BlocksAllocatedBits.CountOnes(SmallPoolTables[PoolIndex].NumEverUsedBlocks);
+		const uint32 PartialBlocks = SmallPoolTables[PoolIndex].NumEverUsedBlocks - SmallPoolTables[PoolIndex].BlocksExhaustedBits.CountOnes(SmallPoolTables[PoolIndex].NumEverUsedBlocks);
+		const uint32 FullBlocks = CommittedBlocks - PartialBlocks;
+		const int64 ComittedVM = VM - (SmallPoolTables[PoolIndex].NumEverUsedBlocks - CommittedBlocks) * SmallPoolTables[PoolIndex].NumMemoryPagesPerBlock * OsAllocationGranularity;
 
-		int64 AveSize = SmallPoolTables[PoolIndex].TotalAllocCount.Load() ? SmallPoolTables[PoolIndex].TotalRequestedAllocSize.Load() / SmallPoolTables[PoolIndex].TotalAllocCount.Load() : 0;
-		int64 EstPadWaste = ((SmallPoolTables[PoolIndex].TotalAllocCount.Load() - SmallPoolTables[PoolIndex].TotalFreeCount.Load()) * (PoolIndexToBlockSize(PoolIndex) - AveSize));
+		const int64 AveSize = SmallPoolTables[PoolIndex].TotalAllocCount.Load() ? SmallPoolTables[PoolIndex].TotalRequestedAllocSize.Load() / SmallPoolTables[PoolIndex].TotalAllocCount.Load() : 0;
+		const int64 EstPadWaste = ((SmallPoolTables[PoolIndex].TotalAllocCount.Load() - SmallPoolTables[PoolIndex].TotalFreeCount.Load()) * (PoolIndexToBinSize(PoolIndex) - AveSize));
 
 		Ar.Logf(TEXT("Pool %2d   Size %6d   Allocs %8lld  Frees %8lld  AveAllocSize %6d  EstPadWaste %4dKB  UsedVM %3dMB  CommittedVM %3dMB  HighSlabs %6d  CommittedSlabs %6d  FullSlabs %6d  PartialSlabs  %6d"), 
 			PoolIndex,
-			PoolIndexToBlockSize(PoolIndex),
+			PoolIndexToBinSize(PoolIndex),
 			SmallPoolTables[PoolIndex].TotalAllocCount.Load(),
 			SmallPoolTables[PoolIndex].TotalFreeCount.Load(),
 			AveSize,
 			EstPadWaste / 1024,
 			VM / (1024 * 1024),
 			ComittedVM / (1024 * 1024),
-			SmallPoolTables[PoolIndex].NumEverUsedBlockOfBlocks,
+			SmallPoolTables[PoolIndex].NumEverUsedBlocks,
 			CommittedBlocks,
 			FullBlocks,
 			PartialBlocks
 			);
 	}
-#else
 #endif
 
 #else
-	Ar.Logf(TEXT("Allocator Stats for Binned3 are not in this build set BINNED3_ALLOCATOR_STATS 1 in MallocBinned3.cpp"));
+	Ar.Logf(TEXT("Allocator Stats for Binned3 are not in this build set UE_MB3_ALLOCATOR_STATS 1 in MallocBinned3.cpp"));
 #endif
 }
-#if !BINNED3_INLINE
+#if !UE_MB3_INLINE
 	#if PLATFORM_USES_FIXED_GMalloc_CLASS && !FORCE_ANSI_ALLOCATOR && USE_MALLOC_BINNED3
 		//#define FMEMORY_INLINE_FUNCTION_DECORATOR  FORCEINLINE
 		#define FMEMORY_INLINE_GMalloc (FMallocBinned3::MallocBinned3)
diff --git a/Engine/Source/Runtime/Core/Private/HAL/MallocBinnedCommon.cpp b/Engine/Source/Runtime/Core/Private/HAL/MallocBinnedCommon.cpp
--- a/Engine/Source/Runtime/Core/Private/HAL/MallocBinnedCommon.cpp
+++ b/Engine/Source/Runtime/Core/Private/HAL/MallocBinnedCommon.cpp
@@ -13,10 +13,10 @@ PRAGMA_DISABLE_UNSAFE_TYPECAST_WARNINGS
 
 #if PLATFORM_HAS_FPlatformVirtualMemoryBlock
 
-// Block sizes are based around getting the maximum amount of allocations per pool, with as little alignment waste as possible.
-// Block sizes should be close to even divisors of the system page size, and well distributed.
+// Bin sizes are based around getting the maximum amount of allocations per block, with as little alignment waste as possible.
+// Bin sizes should be close to even divisors of the system page size, and well distributed.
 // They must be 16-byte aligned as well.
-static uint32 BinnedCommonSmallBlockSizes4k[] = 
+static constexpr uint32 BinnedCommonSmallBinSizes4k[] =
 {
 	16, 32, 48, 64, 80, 96, 112, 128, 144, 160, 176, 192, // +16
 	224, 256, 288, 320, // +32
@@ -33,7 +33,7 @@ static uint32 BinnedCommonSmallBlockSizes4k[] =
 	4096 // /1
 };
 
-static uint32 BinnedCommonSmallBlockSizes8k[] =
+static constexpr uint32 BinnedCommonSmallBinSizes8k[] =
 {
 	736,  // /11 ish
 	1168, // /7 ish
@@ -42,7 +42,7 @@ static uint32 BinnedCommonSmallBlockSizes8k[] =
 	8192  // /1
 };
 
-static uint32 BinnedCommonSmallBlockSizes12k[] =
+static constexpr uint32 BinnedCommonSmallBinSizes12k[] =
 {
 	//1104, // /11 ish
 	//1216, // /10 ish
@@ -54,7 +54,7 @@ static uint32 BinnedCommonSmallBlockSizes12k[] =
 	12288 // /1
 };
 
-static uint32 BinnedCommonSmallBlockSizes16k[] =
+static constexpr uint32 BinnedCommonSmallBinSizes16k[] =
 {
 	//1488, // /11 ish
 	//1808, // /9 ish
@@ -64,7 +64,7 @@ static uint32 BinnedCommonSmallBlockSizes16k[] =
 	16384 // /1
 };
 
-static uint32 BinnedCommonSmallBlockSizes20k[] =
+static constexpr uint32 BinnedCommonSmallBinSizes20k[] =
 {
 	// 2912, // /7 ish
 	//3408, // /6 ish
@@ -74,12 +74,12 @@ static uint32 BinnedCommonSmallBlockSizes20k[] =
 	20480  // /1
 };
 
-static uint32 BinnedCommonSmallBlockSizes24k[] = // 1 total
+static constexpr uint32 BinnedCommonSmallBinSizes24k[] = // 1 total
 {
 	24576  // /1
 };
 
-static uint32 BinnedCommonSmallBlockSizes28k[] = // 6 total
+static constexpr uint32 BinnedCommonSmallBinSizes28k[] = // 6 total
 {
 	4768,  // /6 ish
 	5728,  // /5 ish
@@ -89,40 +89,39 @@ static uint32 BinnedCommonSmallBlockSizes28k[] = // 6 total
 	28672  // /1
 };
 
-FSizeTableEntry::FSizeTableEntry(uint32 InBlockSize, uint64 PlatformPageSize, uint8 Pages4k, uint32 BasePageSize, uint32 MinimumAlignment)
-	: BlockSize(InBlockSize)
+FSizeTableEntry::FSizeTableEntry(uint32 InBinSize, uint64 PlatformPageSize, uint8 Num4kbPages, uint32 BasePageSize, uint32 MinimumAlignment)
+	: BinSize(InBinSize)
 {
-	check((PlatformPageSize & (BasePageSize - 1)) == 0 && PlatformPageSize >= BasePageSize && InBlockSize % MinimumAlignment == 0);
+	check((PlatformPageSize & (BasePageSize - 1)) == 0 && PlatformPageSize >= BasePageSize && InBinSize % MinimumAlignment == 0);
 
-	uint64 Page4kPerPlatformPage = PlatformPageSize / BasePageSize;
+	const uint64 NumBasePagesPerPlatformPage = PlatformPageSize / BasePageSize;
 
-	PagesPlatformForBlockOfBlocks = 0;
+	NumMemoryPagesPerBlock = 0;
 	while (true)
 	{
-		check(PagesPlatformForBlockOfBlocks < MAX_uint8);
-		PagesPlatformForBlockOfBlocks++;
-		if (PagesPlatformForBlockOfBlocks * Page4kPerPlatformPage < Pages4k)
+		check(NumMemoryPagesPerBlock < MAX_uint8);
+		NumMemoryPagesPerBlock++;
+		if (NumMemoryPagesPerBlock * NumBasePagesPerPlatformPage < Num4kbPages)
 		{
 			continue;
 		}
-		if (PagesPlatformForBlockOfBlocks * Page4kPerPlatformPage % Pages4k != 0)
+		if (NumMemoryPagesPerBlock * NumBasePagesPerPlatformPage % Num4kbPages != 0)
 		{
 			continue;
 		}
 		break;
 	}
-	check((PlatformPageSize * PagesPlatformForBlockOfBlocks) / BlockSize <= MAX_uint16);
-	BlocksPerBlockOfBlocks = (PlatformPageSize * PagesPlatformForBlockOfBlocks) / BlockSize;
+	check((PlatformPageSize * NumMemoryPagesPerBlock) / BinSize <= MAX_uint32);
 }
 
-static inline void FillTable(FSizeTableEntry* SizeTable, int32& Index, const uint32* BlockList, uint32 BlockListSize, uint32 NumPages, uint64 PlatformPageSize, uint32 BasePageSize, uint32 MinimumAlignment)
+static inline void FillTable(FSizeTableEntry* SizeTable, int32& Index, const uint32* BinsList, uint32 BinListSize, uint32 Num4kbPages, uint64 PlatformPageSize, uint32 BasePageSize, uint32 MinimumAlignment)
 {
-	for (uint32 Sub = 0; Sub < BlockListSize; Sub++)
+	for (uint32 Sub = 0; Sub < BinListSize; Sub++)
 	{
-		// if we override BINNEDCOMMON_MAX_LISTED_SMALL_POOL_SIZE externally, we need to filter out predefined bins of a larger size
-		if (BlockList[Sub] <= BINNEDCOMMON_MAX_LISTED_SMALL_POOL_SIZE)
+		// if we override UE_MBC_MAX_LISTED_SMALL_POOL_SIZE externally, we need to filter out predefined bins of a larger size
+		if (BinsList[Sub] <= UE_MBC_MAX_LISTED_SMALL_POOL_SIZE)
 		{
-			SizeTable[Index++] = FSizeTableEntry(BlockList[Sub], PlatformPageSize, NumPages, BasePageSize, MinimumAlignment);
+			SizeTable[Index++] = FSizeTableEntry(BinsList[Sub], PlatformPageSize, Num4kbPages, BasePageSize, MinimumAlignment);
 		}
 	}
 }
@@ -130,20 +129,20 @@ static inline void FillTable(FSizeTableEntry* SizeTable, int32& Index, const uin
 uint8 FSizeTableEntry::FillSizeTable(uint64 PlatformPageSize, FSizeTableEntry* SizeTable, uint32 BasePageSize, uint32 MinimumAlignment, uint32 MaxSize, uint32 SizeIncrement)
 {
 	int32 Index = 0;
-	FillTable(SizeTable, Index, BinnedCommonSmallBlockSizes4k,  UE_ARRAY_COUNT(BinnedCommonSmallBlockSizes4k),  1, PlatformPageSize, BasePageSize, MinimumAlignment);
-	FillTable(SizeTable, Index, BinnedCommonSmallBlockSizes8k,  UE_ARRAY_COUNT(BinnedCommonSmallBlockSizes8k),  2, PlatformPageSize, BasePageSize, MinimumAlignment);
-	FillTable(SizeTable, Index, BinnedCommonSmallBlockSizes12k, UE_ARRAY_COUNT(BinnedCommonSmallBlockSizes12k), 3, PlatformPageSize, BasePageSize, MinimumAlignment);
-	FillTable(SizeTable, Index, BinnedCommonSmallBlockSizes16k, UE_ARRAY_COUNT(BinnedCommonSmallBlockSizes16k), 4, PlatformPageSize, BasePageSize, MinimumAlignment);
-	FillTable(SizeTable, Index, BinnedCommonSmallBlockSizes20k, UE_ARRAY_COUNT(BinnedCommonSmallBlockSizes20k), 5, PlatformPageSize, BasePageSize, MinimumAlignment);
-	FillTable(SizeTable, Index, BinnedCommonSmallBlockSizes24k, UE_ARRAY_COUNT(BinnedCommonSmallBlockSizes24k), 6, PlatformPageSize, BasePageSize, MinimumAlignment);
-	FillTable(SizeTable, Index, BinnedCommonSmallBlockSizes28k, UE_ARRAY_COUNT(BinnedCommonSmallBlockSizes28k), 7, PlatformPageSize, BasePageSize, MinimumAlignment);
+	FillTable(SizeTable, Index, BinnedCommonSmallBinSizes4k,  UE_ARRAY_COUNT(BinnedCommonSmallBinSizes4k),  1, PlatformPageSize, BasePageSize, MinimumAlignment);
+	FillTable(SizeTable, Index, BinnedCommonSmallBinSizes8k,  UE_ARRAY_COUNT(BinnedCommonSmallBinSizes8k),  2, PlatformPageSize, BasePageSize, MinimumAlignment);
+	FillTable(SizeTable, Index, BinnedCommonSmallBinSizes12k, UE_ARRAY_COUNT(BinnedCommonSmallBinSizes12k), 3, PlatformPageSize, BasePageSize, MinimumAlignment);
+	FillTable(SizeTable, Index, BinnedCommonSmallBinSizes16k, UE_ARRAY_COUNT(BinnedCommonSmallBinSizes16k), 4, PlatformPageSize, BasePageSize, MinimumAlignment);
+	FillTable(SizeTable, Index, BinnedCommonSmallBinSizes20k, UE_ARRAY_COUNT(BinnedCommonSmallBinSizes20k), 5, PlatformPageSize, BasePageSize, MinimumAlignment);
+	FillTable(SizeTable, Index, BinnedCommonSmallBinSizes24k, UE_ARRAY_COUNT(BinnedCommonSmallBinSizes24k), 6, PlatformPageSize, BasePageSize, MinimumAlignment);
+	FillTable(SizeTable, Index, BinnedCommonSmallBinSizes28k, UE_ARRAY_COUNT(BinnedCommonSmallBinSizes28k), 7, PlatformPageSize, BasePageSize, MinimumAlignment);
 
-	check(Index == BINNEDCOMMON_NUM_LISTED_SMALL_POOLS);
+	check(Index == UE_MBC_NUM_LISTED_SMALL_POOLS);
 
 	Algo::Sort(MakeArrayView(SizeTable, Index));
-	check(SizeTable[Index - 1].BlockSize == BINNEDCOMMON_MAX_LISTED_SMALL_POOL_SIZE);
-	check(IsAligned(BINNEDCOMMON_MAX_LISTED_SMALL_POOL_SIZE, BasePageSize));
-	for (uint32 Size = BINNEDCOMMON_MAX_LISTED_SMALL_POOL_SIZE + BasePageSize; Size <= MaxSize; Size += SizeIncrement)
+	check(SizeTable[Index - 1].BinSize == UE_MBC_MAX_LISTED_SMALL_POOL_SIZE);
+	check(IsAligned(UE_MBC_MAX_LISTED_SMALL_POOL_SIZE, BasePageSize));
+	for (uint32 Size = UE_MBC_MAX_LISTED_SMALL_POOL_SIZE + BasePageSize; Size <= MaxSize; Size += SizeIncrement)
 	{
 		SizeTable[Index++] = FSizeTableEntry(Size, PlatformPageSize, Size / BasePageSize, BasePageSize, MinimumAlignment);
 	}
@@ -224,7 +223,7 @@ uint32 FBitTree::AllocBit()
 		{
 			uint64* At = Bits + Offset;
 			check(At >= Bits && At < Bits + AllocationSize / 8);
-			uint32 LowestZeroBit = FMath::CountTrailingZeros64(~*At);
+			const uint32 LowestZeroBit = FMath::CountTrailingZeros64(~*At);
 			check(LowestZeroBit < 64);
 			Result = Result * 64 + LowestZeroBit;
 			if (Row == Rows - 1)
@@ -237,7 +236,7 @@ uint32 FBitTree::AllocBit()
 					{
 						do
 						{
-							uint32 Rem = (Offset - 1) % 64;
+							const uint32 Rem = (Offset - 1) % 64;
 							Offset = (Offset - 1) / 64;
 							At = Bits + Offset;
 							check(At >= Bits && At < Bits + AllocationSize / 8);
@@ -264,10 +263,10 @@ uint32 FBitTree::AllocBit()
 bool FBitTree::IsAllocated(uint32 Index) const
 {
 	check(Index < DesiredCapacity);
-	uint32 Row = Rows - 1;
-	uint32 Rem = Index % 64;
-	uint32 Offset = OffsetOfLastRow + Index / 64;
-	uint64* At = Bits + Offset;
+	const uint32 Row = Rows - 1;
+	const uint32 Rem = Index % 64;
+	const uint32 Offset = OffsetOfLastRow + Index / 64;
+	const uint64* At = Bits + Offset;
 	check(At >= Bits && At < Bits + AllocationSize / 8);
 	return !!((*At) & (1ull << Rem));
 }
@@ -301,6 +300,7 @@ void FBitTree::AllocBit(uint32 Index)
 	}
 
 }
+
 uint32 FBitTree::NextAllocBit() const
 {
 	uint32 Result = MAX_uint32;
@@ -311,9 +311,9 @@ uint32 FBitTree::NextAllocBit() const
 		uint32 Row = 0;
 		while (true)
 		{
-			uint64* At = Bits + Offset;
+			const uint64* At = Bits + Offset;
 			check(At >= Bits && At < Bits + AllocationSize / 8);
-			uint32 LowestZeroBit = FMath::CountTrailingZeros64(~*At);
+			const uint32 LowestZeroBit = FMath::CountTrailingZeros64(~*At);
 			check(LowestZeroBit < 64);
 			Result = Result * 64 + LowestZeroBit;
 			if (Row == Rows - 1)
@@ -333,12 +333,12 @@ uint32 FBitTree::NextAllocBit(uint32 StartIndex) const
 {
 	if (*Bits != MAX_uint64) // else we are full
 	{
-		uint32 Index = StartIndex;
+		const uint32 Index = StartIndex;
 		check(Index < DesiredCapacity);
 		uint32 Row = Rows - 1;
 		uint32 Rem = Index % 64;
 		uint32 Offset = OffsetOfLastRow + Index / 64;
-		uint64* At = Bits + Offset;
+		const uint64* At = Bits + Offset;
 		check(At >= Bits && At < Bits + AllocationSize / 8);
 		uint64 LocalAt = *At;
 		if (!(LocalAt & (1ull << Rem)))
@@ -350,7 +350,7 @@ uint32 FBitTree::NextAllocBit(uint32 StartIndex) const
 		if (LocalAt != MAX_uint64)
 		{
 			// this qword has an item we can use
-			uint32 LowestZeroBit = FMath::CountTrailingZeros64(~LocalAt);
+			const uint32 LowestZeroBit = FMath::CountTrailingZeros64(~LocalAt);
 			check(LowestZeroBit < 64);
 			return Index - Rem + LowestZeroBit;
 
@@ -373,12 +373,12 @@ uint32 FBitTree::NextAllocBit(uint32 StartIndex) const
 					// now search down the tree
 					while (true)
 					{
-						uint32 LowestZeroBit = FMath::CountTrailingZeros64(~LocalAt);
+						const uint32 LowestZeroBit = FMath::CountTrailingZeros64(~LocalAt);
 						check(LowestZeroBit < 64);
 						if (Row == Rows - 1)
 						{
 							check(!(LocalAt & (1ull << LowestZeroBit))); // this was already allocated?
-							uint32 Result = (Offset - OffsetOfLastRow) * 64 + LowestZeroBit;
+							const uint32 Result = (Offset - OffsetOfLastRow) * 64 + LowestZeroBit;
 							check(Result < DesiredCapacity);
 							return Result;
 
@@ -397,7 +397,6 @@ uint32 FBitTree::NextAllocBit(uint32 StartIndex) const
 	return MAX_uint32;
 }
 
-
 void FBitTree::FreeBit(uint32 Index)
 {
 	check(Index < DesiredCapacity);
@@ -431,7 +430,7 @@ void FBitTree::FreeBit(uint32 Index)
 uint32 FBitTree::CountOnes(uint32 UpTo) const
 {
 	uint32 Result = 0;
-	uint64* At = Bits + OffsetOfLastRow;
+	const uint64* At = Bits + OffsetOfLastRow;
 	while (UpTo >= 64)
 	{
 		Result += FMath::CountBits(*At);
@@ -536,16 +535,16 @@ static FAutoConsoleVariableRef GMallocBinnedFlushRegisteredThreadCachesOnOneThre
 	GMallocBinnedFlushRegisteredThreadCachesOnOneThread,
 	TEXT("Wether or not to attempt to flush registered thread caches on one thread (enabled by default)."));
 
-#if UE_BINNEDCOMMON_ALLOW_RUNTIME_TWEAKING
+#if UE_MBC_ALLOW_RUNTIME_TWEAKING
 
-int32 GMallocBinnedBundleSize = DEFAULT_GMallocBinnedBundleSize;
+int32 GMallocBinnedBundleSize = UE_DEFAULT_GMallocBinnedBundleSize;
 static FAutoConsoleVariableRef GMallocBinned3BundleSizeCVar(
 	TEXT("MallocBinned.BundleSize"),
 	GMallocBinnedBundleSize,
 	TEXT("Max size in bytes of per-block bundles used in the recycling process")
 );
 
-int32 GMallocBinnedBundleCount = DEFAULT_GMallocBinnedBundleCount;
+int32 GMallocBinnedBundleCount = UE_DEFAULT_GMallocBinnedBundleCount;
 static FAutoConsoleVariableRef GMallocBinned3BundleCountCVar(
 	TEXT("MallocBinned.BundleCount"),
 	GMallocBinnedBundleCount,
@@ -555,9 +554,9 @@ static FAutoConsoleVariableRef GMallocBinned3BundleCountCVar(
 #endif
 
 uint32 FMallocBinnedCommonBase::BinnedTlsSlot = FPlatformTLS::InvalidTlsSlot;
-#if UE_BINNEDCOMMON_ALLOCATOR_STATS
-std::atomic<int64> FMallocBinnedCommonBase::TLSMemory(0);
-std::atomic<int64> FMallocBinnedCommonBase::ConsolidatedMemory(0);
+#if UE_MBC_ALLOCATOR_STATS
+	std::atomic<int64> FMallocBinnedCommonBase::TLSMemory(0);
+	std::atomic<int64> FMallocBinnedCommonBase::ConsolidatedMemory(0);
 #endif
 
 PRAGMA_RESTORE_UNSAFE_TYPECAST_WARNINGS
diff --git a/Engine/Source/Runtime/Core/Private/HAL/MallocBinnedGPU.cpp b/Engine/Source/Runtime/Core/Private/HAL/MallocBinnedGPU.cpp
--- a/Engine/Source/Runtime/Core/Private/HAL/MallocBinnedGPU.cpp
+++ b/Engine/Source/Runtime/Core/Private/HAL/MallocBinnedGPU.cpp
@@ -573,7 +573,7 @@ void FMallocBinnedGPU::InitMallocBinned()
 
 	if (ArenaParams.bUseStandardSmallPoolSizes)
 	{
-		SizeTable.AddZeroed(BINNEDCOMMON_NUM_LISTED_SMALL_POOLS + ArenaParams.MaxStandardPoolSize / ArenaParams.BasePageSize); // overestimate
+		SizeTable.AddZeroed(UE_MBC_NUM_LISTED_SMALL_POOLS + ArenaParams.MaxStandardPoolSize / ArenaParams.BasePageSize); // overestimate
 		ArenaParams.PoolCount = FSizeTableEntry::FillSizeTable(ArenaParams.AllocationGranularity, &SizeTable[0], ArenaParams.BasePageSize, ArenaParams.MinimumAlignment, ArenaParams.MaxStandardPoolSize, ArenaParams.BasePageSize);
 		SizeTable.RemoveAt(ArenaParams.PoolCount, SizeTable.Num() - ArenaParams.PoolCount);
 	}
@@ -602,7 +602,7 @@ void FMallocBinnedGPU::InitMallocBinned()
 	}
 	check(ArenaParams.PoolCount == SizeTable.Num());
 	check(SizeTable.Num() < 256);
-	ArenaParams.MaxPoolSize = SizeTable[ArenaParams.PoolCount - 1].BlockSize;
+	ArenaParams.MaxPoolSize = SizeTable[ArenaParams.PoolCount - 1].BinSize;
 
 	check(ArenaParams.BasePageSize % sizeof(FPoolInfoSmall) == 0);
 	SmallPoolInfosPerPlatformPage = ArenaParams.BasePageSize / sizeof(FPoolInfoSmall);
@@ -619,23 +619,22 @@ void FMallocBinnedGPU::InitMallocBinned()
 
 	for (uint32 Index = 0; Index < ArenaParams.PoolCount; ++Index)
 	{
-		checkf(Index == 0 || SizeTable[Index - 1].BlockSize < SizeTable[Index].BlockSize, TEXT("Small block sizes must be strictly increasing"));
-		checkf(SizeTable[Index].BlockSize % ArenaParams.MinimumAlignment == 0, TEXT("Small block size must be a multiple of ArenaParams.MinimumAlignment"));
+		checkf(Index == 0 || SizeTable[Index - 1].BinSize < SizeTable[Index].BinSize, TEXT("Small block sizes must be strictly increasing"));
+		checkf(SizeTable[Index].BinSize % ArenaParams.MinimumAlignment == 0, TEXT("Small block size must be a multiple of ArenaParams.MinimumAlignment"));
 
 		// determine the largest alignment that we can cover with a small block
-		while (ArenaParams.MaximumAlignmentForSmallBlock < ArenaParams.AllocationGranularity && IsAligned(SizeTable[Index].BlockSize, ArenaParams.MaximumAlignmentForSmallBlock * 2))
+		while (ArenaParams.MaximumAlignmentForSmallBlock < ArenaParams.AllocationGranularity && IsAligned(SizeTable[Index].BinSize, ArenaParams.MaximumAlignmentForSmallBlock * 2))
 		{
 			ArenaParams.MaximumAlignmentForSmallBlock *= 2;
 		}
 
-		SmallPoolTables[Index].BlockSize = SizeTable[Index].BlockSize;
-		SmallPoolTables[Index].BlocksPerBlockOfBlocks = SizeTable[Index].BlocksPerBlockOfBlocks;
-		SmallPoolTables[Index].PagesPlatformForBlockOfBlocks = SizeTable[Index].PagesPlatformForBlockOfBlocks;
+		SmallPoolTables[Index].BlockSize = SizeTable[Index].BinSize;
+		SmallPoolTables[Index].PagesPlatformForBlockOfBlocks = SizeTable[Index].NumMemoryPagesPerBlock;
 
 		SmallPoolTables[Index].UnusedAreaOffsetLow = 0;
 		SmallPoolTables[Index].NumEverUsedBlockOfBlocks = 0;
 
-		int64 TotalNumberOfBlocksOfBlocks = ArenaParams.MaxMemoryPerBlockSize / (SizeTable[Index].PagesPlatformForBlockOfBlocks * ArenaParams.AllocationGranularity);
+		int64 TotalNumberOfBlocksOfBlocks = ArenaParams.MaxMemoryPerBlockSize / (SizeTable[Index].NumMemoryPagesPerBlock * ArenaParams.AllocationGranularity);
 
 		int64 MaxPoolInfoMemory = sizeof(FPoolInfoSmall**) * (TotalNumberOfBlocksOfBlocks + SmallPoolInfosPerPlatformPage - 1) / SmallPoolInfosPerPlatformPage;
 		SmallPoolTables[Index].PoolInfos = (FPoolInfoSmall**)FMemory::Malloc(MaxPoolInfoMemory);
@@ -672,7 +671,7 @@ void FMallocBinnedGPU::InitMallocBinned()
 	{
 
 		uint32 BlockSize = Index << ArenaParams.MinimumAlignmentShift; // inverse of int32 Index = int32((Size >> ArenaParams.MinimumAlignmentShift));
-		while (SizeTable[PoolIndex].BlockSize < BlockSize)
+		while (SizeTable[PoolIndex].BinSize < BlockSize)
 		{
 			++PoolIndex;
 			check(PoolIndex != ArenaParams.PoolCount);
@@ -685,7 +684,7 @@ void FMallocBinnedGPU::InitMallocBinned()
 	for (uint32 Index = 0; Index != ArenaParams.PoolCount; ++Index)
 	{
 		uint32 Partner = ArenaParams.PoolCount - Index - 1;
-		SmallBlockSizesReversedShifted[Index] = (SizeTable[Partner].BlockSize >> ArenaParams.MinimumAlignmentShift);
+		SmallBlockSizesReversedShifted[Index] = (SizeTable[Partner].BinSize >> ArenaParams.MinimumAlignmentShift);
 	}
 	uint64 MaxHashBuckets = PtrToPoolMapping.GetMaxHashBuckets();
 
diff --git a/Engine/Source/Runtime/Core/Public/HAL/Allocators/CachedOSVeryLargePageAllocator.h b/Engine/Source/Runtime/Core/Public/HAL/Allocators/CachedOSVeryLargePageAllocator.h
--- a/Engine/Source/Runtime/Core/Public/HAL/Allocators/CachedOSVeryLargePageAllocator.h
+++ b/Engine/Source/Runtime/Core/Public/HAL/Allocators/CachedOSVeryLargePageAllocator.h
@@ -66,7 +66,7 @@ public:
 		return CachedFree + CachedOSPageAllocator.GetCachedFreeTotal();
 	}
 
-	FORCEINLINE bool IsSmallBlockAllocation(const void* Ptr)
+	FORCEINLINE bool IsSmallBlockAllocation(const void* Ptr) const
 	{
 		if (((uintptr_t)Ptr - AddressSpaceReserved) < AddressSpaceToReserveForSmallPool)
 		{
diff --git a/Engine/Source/Runtime/Core/Public/HAL/MallocBinned2.h b/Engine/Source/Runtime/Core/Public/HAL/MallocBinned2.h
--- a/Engine/Source/Runtime/Core/Public/HAL/MallocBinned2.h
+++ b/Engine/Source/Runtime/Core/Public/HAL/MallocBinned2.h
@@ -17,72 +17,71 @@
 
 struct FGenericMemoryStats;
 
-#define BINNED2_MAX_CACHED_OS_FREES (64)
+#define UE_MB2_MAX_CACHED_OS_FREES (64)
 #if PLATFORM_64BITS
-	#define BINNED2_MAX_CACHED_OS_FREES_BYTE_LIMIT (64*1024*1024)
+#	define UE_MB2_MAX_CACHED_OS_FREES_BYTE_LIMIT (64*1024*1024)
 #else
-	#define BINNED2_MAX_CACHED_OS_FREES_BYTE_LIMIT (16*1024*1024)
+#	define UE_MB2_MAX_CACHED_OS_FREES_BYTE_LIMIT (16*1024*1024)
 #endif
 
-#define BINNED2_LARGE_ALLOC					65536		// Alignment of OS-allocated pointer - pool-allocated pointers will have a non-aligned pointer
-#define BINNED2_MINIMUM_ALIGNMENT_SHIFT		4			// Alignment of blocks, expressed as a shift
-#define BINNED2_MINIMUM_ALIGNMENT			16			// Alignment of blocks
-#define BINNED2_MAXIMUM_ALIGNMENT			128
-#define BINNED2_MAX_SMALL_POOL_SIZE			(32768-16)	// Maximum block size in GMallocBinned2SmallBlockSizes
-#define BINNED2_SMALL_POOL_COUNT			45
+#define UE_MB2_LARGE_ALLOC					65536		// Alignment of OS-allocated pointer - pool-allocated pointers will have a non-aligned pointer
+#define UE_MB2_MINIMUM_ALIGNMENT_SHIFT		4			// Alignment of bins, expressed as a shift
+#define UE_MB2_MINIMUM_ALIGNMENT			16			// Alignment of bins
+#define UE_MB2_MAXIMUM_ALIGNMENT			128
+#define UE_MB2_MAX_SMALL_POOL_SIZE			(32768-16)	// Maximum bin size in SmallBinSizes in cpp file
+#define UE_MB2_SMALL_POOL_COUNT				45
 
 
-#define DEFAULT_GMallocBinned2PerThreadCaches 1
-#define DEFAULT_GMallocBinned2LockFreeCaches 0
-#define DEFAULT_GMallocBinned2AllocExtra 32
-#define BINNED2_MAX_GMallocBinned2MaxBundlesBeforeRecycle 8
-#define DEFAULT_GMallocBinned2MoveOSFreesOffTimeCriticalThreads 1
+#define UE_DEFAULT_GBinned2PerThreadCaches						1
+#define UE_DEFAULT_GBinned2AllocExtra							32
+#define UE_DEFAULT_GBinned2MaxBundlesBeforeRecycle				8
+#define UE_DEFAULT_GBinned2MoveOSFreesOffTimeCriticalThreads	1
 
 // When book keeping is at the end of FFreeBlock, MallocBinned2 cannot tell if the allocation comes from a large allocation (higher than 64KB, also named as "OSAllocation") 
 // or from VeryLargePageAllocator that fell back to FCachedOSPageAllocator. In both cases the allocation (large or small) might be aligned to 64KB.
 // bookKeeping at the end needs to be disabled if we want VeryLargePageAllocator to properly fallback to regular TCachedOSPageAllocator if needed
-#ifndef BINNED2_BOOKKEEPING_AT_THE_END_OF_LARGEBLOCK
-#define BINNED2_BOOKKEEPING_AT_THE_END_OF_LARGEBLOCK 0
+#ifndef UE_MB2_BOOKKEEPING_AT_THE_END_OF_LARGEBLOCK
+#	define UE_MB2_BOOKKEEPING_AT_THE_END_OF_LARGEBLOCK 0
 #endif
 
 // If we are emulating forking on a windows server or are a linux server, enable support for avoiding dirtying pages owned by the parent. 
 #ifndef BINNED2_FORK_SUPPORT
-	#define BINNED2_FORK_SUPPORT (UE_SERVER && (PLATFORM_UNIX || DEFAULT_SERVER_FAKE_FORKS))
+#	define BINNED2_FORK_SUPPORT (UE_SERVER && (PLATFORM_UNIX || DEFAULT_SERVER_FAKE_FORKS))
 #endif
 
 
-#define BINNED2_ALLOW_RUNTIME_TWEAKING UE_BINNEDCOMMON_ALLOW_RUNTIME_TWEAKING
-#if BINNED2_ALLOW_RUNTIME_TWEAKING
-	extern CORE_API int32 GMallocBinned2PerThreadCaches;
-	extern CORE_API int32 GMallocBinned2MaxBundlesBeforeRecycle;
-	extern CORE_API int32 GMallocBinned2AllocExtra;
-	extern CORE_API int32 GMallocBinned2MoveOSFreesOffTimeCriticalThreads;
-
+#define UE_MB2_ALLOW_RUNTIME_TWEAKING UE_MBC_ALLOW_RUNTIME_TWEAKING
+#if UE_MB2_ALLOW_RUNTIME_TWEAKING
+	extern CORE_API int32 GBinned2PerThreadCaches;
+	extern CORE_API int32 GBinned2MaxBundlesBeforeRecycle;
+	extern CORE_API int32 GBinned2AllocExtra;
+	extern CORE_API int32 GBinned2MoveOSFreesOffTimeCriticalThreads;
 #else
-	#define GMallocBinned2PerThreadCaches DEFAULT_GMallocBinned2PerThreadCaches
-	#define GMallocBinned2MaxBundlesBeforeRecycle BINNED2_MAX_GMallocBinned2MaxBundlesBeforeRecycle
-	#define GMallocBinned2AllocExtra DEFAULT_GMallocBinned2AllocExtra
-	#define GMallocBinned2MoveOSFreesOffTimeCriticalThreads DEFAULT_GMallocBinned2MoveOSFreesOffTimeCriticalThreads
+#	define GBinned2PerThreadCaches						UE_DEFAULT_GBinned2PerThreadCaches
+#	define GBinned2MaxBundlesBeforeRecycle				UE_DEFAULT_GBinned2MaxBundlesBeforeRecycle
+#	define GBinned2AllocExtra							UE_DEFAULT_GBinned2AllocExtra
+#	define GBinned2MoveOSFreesOffTimeCriticalThreads	UE_DEFAULT_GBinned2MoveOSFreesOffTimeCriticalThreads
 #endif
 
-#define BINNED2_ALLOCATOR_STATS UE_BINNEDCOMMON_ALLOCATOR_STATS
-#define BINNED2_ALLOCATOR_STATS_VALIDATION (BINNED2_ALLOCATOR_STATS && 0)
+#define UE_MB2_ALLOCATOR_STATS				UE_MBC_ALLOCATOR_STATS
+#define UE_MB2_ALLOCATOR_STATS_VALIDATION	(UE_MB2_ALLOCATOR_STATS && 0)
 
 
-#if BINNED2_ALLOCATOR_STATS
-//////////////////////////////////////////////////////////////////////////
-// the following don't need a critical section because they are covered by the critical section called Mutex
-extern TAtomic<int64> AllocatedSmallPoolMemory; // memory that's requested to be allocated by the game
-extern TAtomic<int64> AllocatedOSSmallPoolMemory;
-extern TAtomic<int64> AllocatedLargePoolMemory; // memory requests to the OS which don't fit in the small pool
-extern TAtomic<int64> AllocatedLargePoolMemoryWAlignment; // when we allocate at OS level we need to align to a size
+#if UE_MB2_ALLOCATOR_STATS
+	//////////////////////////////////////////////////////////////////////////
+	// the following don't need a critical section because they are covered by the critical section called Mutex
+	extern TAtomic<int64> AllocatedSmallPoolMemory; // memory that's requested to be allocated by the game
+	extern TAtomic<int64> AllocatedOSSmallPoolMemory;
+	extern TAtomic<int64> AllocatedLargePoolMemory; // memory requests to the OS which don't fit in the small pool
+	extern TAtomic<int64> AllocatedLargePoolMemoryWAlignment; // when we allocate at OS level we need to align to a size
 #endif
-#if BINNED2_ALLOCATOR_STATS_VALIDATION
-#include "Misc/ScopeLock.h"
 
-extern int64 AllocatedSmallPoolMemoryValidation;
-extern FCriticalSection ValidationCriticalSection;
-extern int32 RecursionCounter;
+#if UE_MB2_ALLOCATOR_STATS_VALIDATION
+#	include "Misc/ScopeLock.h"
+
+	extern int64 AllocatedSmallPoolMemoryValidation;
+	extern FCriticalSection ValidationCriticalSection;
+	extern int32 RecursionCounter;
 #endif
 
 // Canary value used in FFreeBlock
@@ -103,7 +102,7 @@ enum class EBlockCanary : uint8
 //
 // Optimized virtual memory allocator.
 //
-class FMallocBinned2 : public TMallocBinnedCommon<FMallocBinned2, BINNED2_MINIMUM_ALIGNMENT, BINNED2_MAXIMUM_ALIGNMENT, BINNED2_MINIMUM_ALIGNMENT_SHIFT, BINNED2_SMALL_POOL_COUNT, BINNED2_MAX_SMALL_POOL_SIZE>
+class FMallocBinned2 : public TMallocBinnedCommon<FMallocBinned2, UE_MB2_MINIMUM_ALIGNMENT, UE_MB2_MAXIMUM_ALIGNMENT, UE_MB2_MINIMUM_ALIGNMENT_SHIFT, UE_MB2_SMALL_POOL_COUNT, UE_MB2_MAX_SMALL_POOL_SIZE>
 {
 	// Forward declares.
 	struct FPoolInfo;
@@ -113,62 +112,61 @@ class FMallocBinned2 : public TMallocBinnedCommon<FMallocBinned2, BINNED2_MINIMU
 	/** Information about a piece of free memory. */
 	struct FFreeBlock
 	{
-		FORCEINLINE FFreeBlock(uint32 InPageSize, uint16 InBlockSize, uint8 InPoolIndex, EBlockCanary InCanary)
-			: BlockSize(InBlockSize)
+		FORCEINLINE FFreeBlock(uint32 InPageSize, uint16 InBinSize, uint8 InPoolIndex, EBlockCanary InCanary)
+			: BinSize(InBinSize)
 			, PoolIndex(InPoolIndex)
 			, CanaryAndForkState(InCanary)
 			, NextFreeBlock(nullptr)
 		{
-			check(InPoolIndex < MAX_uint8 && InBlockSize <= MAX_uint16);
-			NumFreeBlocks = InPageSize / InBlockSize;
-			if (NumFreeBlocks * InBlockSize + sizeof(FFreeBlock) > InPageSize)
+			check(InPoolIndex < MAX_uint8 && InBinSize <= MAX_uint16);
+			NumFreeBins = InPageSize / InBinSize;
+			if (NumFreeBins * InBinSize + sizeof(FFreeBlock) > InPageSize)
 			{
-				NumFreeBlocks--;
+				NumFreeBins--;
 			}
-			check(NumFreeBlocks * InBlockSize + sizeof(FFreeBlock) <= InPageSize);
+			check(NumFreeBins * InBinSize + sizeof(FFreeBlock) <= InPageSize);
 		}
 
-		FORCEINLINE uint32 GetNumFreeRegularBlocks() const
+		FORCEINLINE uint32 GetNumFreeBins() const
 		{
-			return NumFreeBlocks;
+			return NumFreeBins;
 		}
 
-
-		FORCEINLINE void* AllocateRegularBlock()
+		FORCEINLINE void* AllocateBin()
 		{
-			--NumFreeBlocks;
-#if !UE_USE_VERYLARGEPAGEALLOCATOR || !BINNED2_BOOKKEEPING_AT_THE_END_OF_LARGEBLOCK
-			if (IsAligned(this, BINNED2_LARGE_ALLOC))
+			--NumFreeBins;
+#if !UE_USE_VERYLARGEPAGEALLOCATOR || !UE_MB2_BOOKKEEPING_AT_THE_END_OF_LARGEBLOCK
+			if (IsAligned(this, UE_MB2_LARGE_ALLOC))
 			{
-				return (uint8*)this + BINNED2_LARGE_ALLOC - (NumFreeBlocks + 1) * BlockSize;
+				return (uint8*)this + UE_MB2_LARGE_ALLOC - (NumFreeBins + 1) * BinSize;
 			}
 #else
-			if (IsAligned(((uintptr_t)this)+sizeof(FFreeBlock), BINNED2_LARGE_ALLOC))
+			if (IsAligned((uintptr_t)this + sizeof(FFreeBlock), UE_MB2_LARGE_ALLOC))
 			{
 				// The book keeping FreeBlock is at the end of the "page" so we align down to get to the beginning of the page
-				uintptr_t ptr = AlignDown((uintptr_t)this, BINNED2_LARGE_ALLOC);
+				const uintptr_t Ptr = AlignDown((uintptr_t)this, UE_MB2_LARGE_ALLOC);
 				// And we offset the returned pointer based on how many free blocks are left.
-				return (uint8*) ptr + (NumFreeBlocks * BlockSize);
+				return (uint8*)Ptr + (NumFreeBins * BinSize);
 			}
 #endif
-			return (uint8*)this + (NumFreeBlocks)* BlockSize;
+			return (uint8*)this + (NumFreeBins * BinSize);
 		}
 
-		uint16 BlockSize;				// Size of the blocks that this list points to
-		uint8 PoolIndex;				// Index of this pool
+		uint16 BinSize;				// Size of the bins that this list points to
+		uint8 PoolIndex;			// Index of this pool
 
 		// Normally this value just functions as a canary to detect invalid memory state.
 		// When process forking is supported, it's still a canary but it has two valid values.
 		// One value is used pre-fork and one post-fork and the value is used to avoid freeing memory in pages shared with the parent process.
 		EBlockCanary CanaryAndForkState; 
 
-		uint32 NumFreeBlocks;          // Number of consecutive free blocks here, at least 1.
-		void*  NextFreeBlock;          // Next free block in another pool
+		uint32 NumFreeBins;          // Number of consecutive free bins here, at least 1.
+		FFreeBlock*  NextFreeBlock;  // Next free block in another pool
 	};
 
 	struct FPoolList
 	{
-		FPoolList();
+		FPoolList() = default;
 
 		void Clear();
 		bool IsEmpty() const;
@@ -180,11 +178,11 @@ class FMallocBinned2 : public TMallocBinnedCommon<FMallocBinned2, BINNED2_MINIMU
 
 		FPoolInfo& PushNewPoolToFront(FMallocBinned2& Allocator, uint32 InBytes, uint32 InPoolIndex);
 
-		void ValidateActivePools();
-		void ValidateExhaustedPools();
+		void ValidateActivePools() const;
+		void ValidateExhaustedPools() const;
 
 	private:
-		FPoolInfo* Front;
+		FPoolInfo* Front = nullptr;
 	};
 
 	/** Pool table. */
@@ -192,13 +190,13 @@ class FMallocBinned2 : public TMallocBinnedCommon<FMallocBinned2, BINNED2_MINIMU
 	{
 		FPoolList ActivePools;
 		FPoolList ExhaustedPools;
-		uint32    BlockSize;
+		uint32    BinSize = 0;
 
-		FPoolTable();
+		FPoolTable() = default;
 	};
 
 	// Pool tables for different pool sizes
-	FPoolTable SmallPoolTables[BINNED2_SMALL_POOL_COUNT];
+	FPoolTable SmallPoolTables[UE_MB2_SMALL_POOL_COUNT];
 
 	PoolHashBucket* HashBuckets;
 	PoolHashBucket* HashBucketFreeList;
@@ -211,56 +209,54 @@ class FMallocBinned2 : public TMallocBinnedCommon<FMallocBinned2, BINNED2_MINIMU
 #endif
 
 #if !PLATFORM_UNIX && !PLATFORM_ANDROID
-#if UE_USE_VERYLARGEPAGEALLOCATOR
-	FCachedOSVeryLargePageAllocator CachedOSPageAllocator;
-#else
-	TCachedOSPageAllocator<BINNED2_MAX_CACHED_OS_FREES, BINNED2_MAX_CACHED_OS_FREES_BYTE_LIMIT> CachedOSPageAllocator;
-#endif
+#	if UE_USE_VERYLARGEPAGEALLOCATOR
+		FCachedOSVeryLargePageAllocator CachedOSPageAllocator;
+#	else
+		TCachedOSPageAllocator<UE_MB2_MAX_CACHED_OS_FREES, UE_MB2_MAX_CACHED_OS_FREES_BYTE_LIMIT> CachedOSPageAllocator;
+#	endif
 #else
 	FPooledVirtualMemoryAllocator CachedOSPageAllocator;
 #endif
 
 	FCriticalSection Mutex;
 
-	FORCEINLINE bool IsOSAllocation(const void* Ptr)
+	FORCEINLINE bool IsOSAllocation(const void* Ptr) const
 	{
 #if UE_USE_VERYLARGEPAGEALLOCATOR && !PLATFORM_UNIX && !PLATFORM_ANDROID
-		return !CachedOSPageAllocator.IsSmallBlockAllocation(Ptr) && IsAligned(Ptr, BINNED2_LARGE_ALLOC);
+		return !CachedOSPageAllocator.IsSmallBlockAllocation(Ptr) && IsAligned(Ptr, UE_MB2_LARGE_ALLOC);
 #else
-		return IsAligned(Ptr, BINNED2_LARGE_ALLOC);
+		return IsAligned(Ptr, UE_MB2_LARGE_ALLOC);
 #endif
 	}
 
 	static FORCEINLINE FFreeBlock* GetPoolHeaderFromPointer(void* Ptr)
 	{
-#if !UE_USE_VERYLARGEPAGEALLOCATOR || !BINNED2_BOOKKEEPING_AT_THE_END_OF_LARGEBLOCK
-		return (FFreeBlock*)AlignDown(Ptr, BINNED2_LARGE_ALLOC);
+#if !UE_USE_VERYLARGEPAGEALLOCATOR || !UE_MB2_BOOKKEEPING_AT_THE_END_OF_LARGEBLOCK
+		return (FFreeBlock*)AlignDown(Ptr, UE_MB2_LARGE_ALLOC);
 #else
-		return (FFreeBlock*) (AlignDown((uintptr_t) Ptr, BINNED2_LARGE_ALLOC) + BINNED2_LARGE_ALLOC - sizeof(FFreeBlock));
+		return (FFreeBlock*)(AlignDown((uintptr_t) Ptr, UE_MB2_LARGE_ALLOC) + UE_MB2_LARGE_ALLOC - sizeof(FFreeBlock));
 #endif
 	}
 
 public:
-
-
 	CORE_API FMallocBinned2();
-
 	CORE_API virtual ~FMallocBinned2();
 
 	// FMalloc interface.
 	CORE_API virtual bool IsInternallyThreadSafe() const override;
+
 	FORCEINLINE virtual void* Malloc(SIZE_T Size, uint32 Alignment) override
 	{
-#if BINNED2_ALLOCATOR_STATS_VALIDATION
+#if UE_MB2_ALLOCATOR_STATS_VALIDATION
 		FScopeLock Lock(&ValidationCriticalSection);
 		++RecursionCounter;
 		void *Result = MallocInline(Size, Alignment);
-		if ( !IsOSAllocation(Result))
+		if (!IsOSAllocation(Result))
 		{
 			SIZE_T OutSize;
-			ensure( GetAllocationSize(Result, OutSize) );
+			ensure(GetAllocationSize(Result, OutSize));
 			AllocatedSmallPoolMemoryValidation += OutSize;
-			if (RecursionCounter==1)
+			if (RecursionCounter == 1)
 			{
 				check(GetTotalAllocatedSmallPoolMemory() == AllocatedSmallPoolMemoryValidation);
 				if (GetTotalAllocatedSmallPoolMemory() != AllocatedSmallPoolMemoryValidation)
@@ -275,6 +271,7 @@ public:
 		return MallocInline(Size, Alignment);
 #endif
 	}
+
 	FORCEINLINE void* MallocInline(SIZE_T Size, uint32 Alignment)
 	{
 		// Only allocate from the small pools if the size is small enough and the alignment isn't crazy large.
@@ -282,15 +279,15 @@ public:
 		const bool bUseSmallPool = UseSmallAlloc(Size, Alignment);
 		if (bUseSmallPool)
 		{
-			FPerThreadFreeBlockLists* Lists = GMallocBinned2PerThreadCaches ? FPerThreadFreeBlockLists::Get() : nullptr;
+			FPerThreadFreeBlockLists* Lists = GBinned2PerThreadCaches ? FPerThreadFreeBlockLists::Get() : nullptr;
 			if (Lists)
 			{
-				const uint32 PoolIndex = BoundSizeToPoolIndex(Size, MemSizeToIndex);
+				const uint32 PoolIndex = BoundSizeToPoolIndex(Size, MemSizeToPoolIndex);
 				if (void* Result = Lists->Malloc(PoolIndex))
 				{
-#if BINNED2_ALLOCATOR_STATS
-					const uint32 BlockSize = PoolIndexToBlockSize(PoolIndex);
-					Lists->AllocatedMemory += BlockSize;
+#if UE_MB2_ALLOCATOR_STATS
+					const uint32 BinSize = PoolIndexToBinSize(PoolIndex);
+					Lists->AllocatedMemory += BinSize;
 #endif
 					return Result;
 				}
@@ -299,37 +296,41 @@ public:
 
 		return MallocSelect(Size, Alignment, bUseSmallPool);
 	}
+
 	FORCEINLINE static bool UseSmallAlloc(SIZE_T Size, uint32 Alignment)
 	{
-#if UE_USE_VERYLARGEPAGEALLOCATOR && BINNED2_BOOKKEEPING_AT_THE_END_OF_LARGEBLOCK
-		if (Alignment > BINNED2_MINIMUM_ALIGNMENT)
+#if UE_USE_VERYLARGEPAGEALLOCATOR && UE_MB2_BOOKKEEPING_AT_THE_END_OF_LARGEBLOCK
+		if (Alignment > UE_MB2_MINIMUM_ALIGNMENT)
 		{
 			Size = Align(Size, Alignment);
 		}
-		bool bResult = (Size <= BINNED2_MAX_SMALL_POOL_SIZE);
+		const bool bResult = (Size <= UE_MB2_MAX_SMALL_POOL_SIZE);
 #else
-		bool bResult = ((Size <= BINNED2_MAX_SMALL_POOL_SIZE) & (Alignment <= BINNED2_MINIMUM_ALIGNMENT)); // one branch, not two
+		const bool bResult = ((Size <= UE_MB2_MAX_SMALL_POOL_SIZE) & (Alignment <= UE_MB2_MINIMUM_ALIGNMENT)); // one branch, not two
 #endif
 		return bResult;
 	}
+
 	CORE_API void* MallocSelect(SIZE_T Size, uint32 Alignment, bool bUseSmallPool);
 	FORCEINLINE void* MallocSelect(SIZE_T Size, uint32 Alignment)
 	{
 		return MallocSelect(Size, Alignment, UseSmallAlloc(Size, Alignment));
 	}
+
 	FORCEINLINE virtual void* Realloc(void* Ptr, SIZE_T NewSize, uint32 Alignment) override
 	{
-#if BINNED2_ALLOCATOR_STATS_VALIDATION
-		bool bOldIsOsAllocation = IsOSAllocation(Ptr);
+#if UE_MB2_ALLOCATOR_STATS_VALIDATION
+		const bool bOldIsOsAllocation = IsOSAllocation(Ptr);
 		SIZE_T OldSize;
 		if (!bOldIsOsAllocation)
 		{
 			ensure(GetAllocationSize(Ptr, OldSize));
 		}
+
 		FScopeLock Lock(&ValidationCriticalSection);
 		++RecursionCounter;
 		void *Result = ReallocInline(Ptr, NewSize, Alignment);
-		if ( !bOldIsOsAllocation )
+		if (!bOldIsOsAllocation)
 		{
 			AllocatedSmallPoolMemoryValidation -= OldSize;
 		}
@@ -353,22 +354,23 @@ public:
 		return ReallocInline(Ptr, NewSize, Alignment);
 #endif
 	}
+
 	FORCEINLINE void* ReallocInline(void* Ptr, SIZE_T NewSize, uint32 Alignment) 
 	{
-#if UE_USE_VERYLARGEPAGEALLOCATOR && BINNED2_BOOKKEEPING_AT_THE_END_OF_LARGEBLOCK
-		if (Alignment > BINNED2_MINIMUM_ALIGNMENT && (NewSize <= BINNED2_MAX_SMALL_POOL_SIZE))
+#if UE_USE_VERYLARGEPAGEALLOCATOR && UE_MB2_BOOKKEEPING_AT_THE_END_OF_LARGEBLOCK
+		if (Alignment > UE_MB2_MINIMUM_ALIGNMENT && (NewSize <= UE_MB2_MAX_SMALL_POOL_SIZE))
 		{
 			NewSize = Align(NewSize, Alignment);
 		}
-		if (NewSize <= BINNED2_MAX_SMALL_POOL_SIZE)
+		if (NewSize <= UE_MB2_MAX_SMALL_POOL_SIZE)
 #else
-		if (NewSize <= BINNED2_MAX_SMALL_POOL_SIZE && Alignment <= BINNED2_MINIMUM_ALIGNMENT) // one branch, not two
+		if (NewSize <= UE_MB2_MAX_SMALL_POOL_SIZE && Alignment <= UE_MB2_MINIMUM_ALIGNMENT) // one branch, not two
 #endif
 		{
-			FPerThreadFreeBlockLists* Lists = GMallocBinned2PerThreadCaches ? FPerThreadFreeBlockLists::Get() : nullptr;
+			FPerThreadFreeBlockLists* Lists = GBinned2PerThreadCaches ? FPerThreadFreeBlockLists::Get() : nullptr;
 			if (Lists && (!Ptr || !IsOSAllocation(Ptr)))
 			{
-				uint32 BlockSize = 0;
+				uint32 BinSize = 0;
 				uint32 PoolIndex = 0;
 
 				bool bCanFree = true; // the nullptr is always "freeable"
@@ -376,39 +378,39 @@ public:
 				{
 					// Reallocate to a smaller/bigger pool if necessary
 					FFreeBlock* Free = GetPoolHeaderFromPointer(Ptr);
-					BlockSize = Free->BlockSize;
+					BinSize = Free->BinSize;
 					PoolIndex = Free->PoolIndex;
 					// If canary is invalid we will assert in ReallocExternal. Otherwise it's the pre-fork canary and we will allocate new memory without touching this allocation.
 					bCanFree = Free->CanaryAndForkState == CurrentCanary;
-					if (NewSize && bCanFree && NewSize <= BlockSize && (PoolIndex == 0 || NewSize > PoolIndexToBlockSize(PoolIndex - 1)))
+					if (NewSize && bCanFree && NewSize <= BinSize && (PoolIndex == 0 || NewSize > PoolIndexToBinSize(PoolIndex - 1)))
 					{
 						return Ptr;
 					}
-					bCanFree = bCanFree && Lists->CanFree(PoolIndex, BlockSize);
+					bCanFree = bCanFree && Lists->CanFree(PoolIndex, BinSize);
 				}
 				if (bCanFree)
 				{
-					uint32 NewPoolIndex = BoundSizeToPoolIndex(NewSize, MemSizeToIndex);
-					uint32 NewBlockSize = PoolIndexToBlockSize(NewPoolIndex);
+					const uint32 NewPoolIndex = BoundSizeToPoolIndex(NewSize, MemSizeToPoolIndex);
+					const uint32 NewBinSize = PoolIndexToBinSize(NewPoolIndex);
 					void* Result = NewSize ? Lists->Malloc(NewPoolIndex) : nullptr;
-#if BINNED2_ALLOCATOR_STATS
+#if UE_MB2_ALLOCATOR_STATS
 					if (Result)
 					{
-						Lists->AllocatedMemory += NewBlockSize;
+						Lists->AllocatedMemory += NewBinSize;
 					}
 #endif
 					if (Result || !NewSize)
 					{
 						if (Result && Ptr)
 						{
-							FMemory::Memcpy(Result, Ptr, FPlatformMath::Min<SIZE_T>(NewSize, BlockSize));
+							FMemory::Memcpy(Result, Ptr, FPlatformMath::Min<SIZE_T>(NewSize, BinSize));
 						}
 						if (Ptr)
 						{
-							bool bDidPush = Lists->Free(Ptr, PoolIndex, BlockSize);
+							const bool bDidPush = Lists->Free(Ptr, PoolIndex, BinSize);
 							checkSlow(bDidPush);
-#if BINNED2_ALLOCATOR_STATS
-							Lists->AllocatedMemory -= BlockSize;
+#if UE_MB2_ALLOCATOR_STATS
+							Lists->AllocatedMemory -= BinSize;
 #endif
 						}
 
@@ -423,7 +425,7 @@ public:
 
 	FORCEINLINE virtual void Free(void* Ptr) override
 	{
-#if BINNED2_ALLOCATOR_STATS_VALIDATION
+#if UE_MB2_ALLOCATOR_STATS_VALIDATION
 		FScopeLock Lock(&ValidationCriticalSection);
 		++RecursionCounter;
 		if (!IsOSAllocation(Ptr))
@@ -451,16 +453,16 @@ public:
 	{
 		if (!IsOSAllocation(Ptr))
 		{
-			FPerThreadFreeBlockLists* Lists = GMallocBinned2PerThreadCaches ? FPerThreadFreeBlockLists::Get() : nullptr;
+			FPerThreadFreeBlockLists* Lists = GBinned2PerThreadCaches ? FPerThreadFreeBlockLists::Get() : nullptr;
 			if (Lists)
 			{
 				FFreeBlock* BasePtr = GetPoolHeaderFromPointer(Ptr);
-				int32 BlockSize = BasePtr->BlockSize;
+				const int32 BinSize = BasePtr->BinSize;
 				// If canary is invalid we will assert in FreeExternal. Otherwise it's the pre-fork canary and we will turn this free into a no-op.
-				if (BasePtr->CanaryAndForkState == CurrentCanary && Lists->Free(Ptr, BasePtr->PoolIndex, BasePtr->BlockSize))
+				if (BasePtr->CanaryAndForkState == CurrentCanary && Lists->Free(Ptr, BasePtr->PoolIndex, BinSize))
 				{
-#if BINNED2_ALLOCATOR_STATS
-					Lists->AllocatedMemory -= BasePtr->BlockSize;
+#if UE_MB2_ALLOCATOR_STATS
+					Lists->AllocatedMemory -= BinSize;
 #endif
 					return;
 				}
@@ -468,6 +470,7 @@ public:
 		}
 		FreeExternal(Ptr);
 	}
+
 	FORCEINLINE virtual bool GetAllocationSize(void *Ptr, SIZE_T &SizeOut) override
 	{
 		if (!IsOSAllocation(Ptr))
@@ -479,7 +482,7 @@ public:
 			if (Free->CanaryAndForkState == CurrentCanary)
 #endif
 			{
-				SizeOut = Free->BlockSize;
+				SizeOut = Free->BinSize;
 				return true;
 			}
 		}
@@ -513,19 +516,19 @@ public:
 	CORE_API void CanaryTest(const FFreeBlock* Block) const;
 	CORE_API void CanaryFail(const FFreeBlock* Block) const;
 
-#if BINNED2_ALLOCATOR_STATS
+#if UE_MB2_ALLOCATOR_STATS
 	CORE_API int64 GetTotalAllocatedSmallPoolMemory() const;
 #endif
 	CORE_API virtual void GetAllocatorStats( FGenericMemoryStats& out_Stats ) override;
 	/** Dumps current allocator stats to the log. */
 	CORE_API virtual void DumpAllocatorStats(class FOutputDevice& Ar) override;
 	
-	static CORE_API uint16 SmallBlockSizesReversed[BINNED2_SMALL_POOL_COUNT]; // this is reversed to get the smallest elements on our main cache line
+	static CORE_API uint16 SmallBinSizesReversed[UE_MB2_SMALL_POOL_COUNT]; // this is reversed to get the smallest elements on our main cache line
 	static CORE_API FMallocBinned2* MallocBinned2;
 	static CORE_API uint32 PageSize;
 	static CORE_API uint32 OsAllocationGranularity;
 	// Mapping of sizes to small table indices
-	static CORE_API uint8 MemSizeToIndex[1 + (BINNED2_MAX_SMALL_POOL_SIZE >> BINNED2_MINIMUM_ALIGNMENT_SHIFT)];
+	static CORE_API uint8 MemSizeToPoolIndex[1 + (UE_MB2_MAX_SMALL_POOL_SIZE >> UE_MB2_MINIMUM_ALIGNMENT_SHIFT)];
 
 	static void RegisterThreadFreeBlockLists(FPerThreadFreeBlockLists* FreeBlockLists);
 	static void UnregisterThreadFreeBlockLists(FPerThreadFreeBlockLists* FreeBlockLists);
@@ -535,17 +538,17 @@ public:
 	static void* AllocateMetaDataMemory(SIZE_T Size);
 	static void FreeMetaDataMemory(void* Ptr, SIZE_T Size);
 
-	FORCEINLINE uint32 PoolIndexToBlockSize(uint32 PoolIndex) const
+	FORCEINLINE uint32 PoolIndexToBinSize(uint32 PoolIndex) const
 	{
-		return SmallBlockSizesReversed[BINNED2_SMALL_POOL_COUNT - PoolIndex - 1];
+		return SmallBinSizesReversed[UE_MB2_SMALL_POOL_COUNT - PoolIndex - 1];
 	}
 
 	void FreeBundles(FBundleNode* Bundles, uint32 PoolIndex);
 	FCriticalSection& GetMutex() { return Mutex; }
 };
 
-#define BINNED2_INLINE (1)
-#if BINNED2_INLINE // during development, it helps with iteration time to not include these here, but rather in the .cpp
+#define UE_MB2_INLINE (1)
+#if UE_MB2_INLINE // during development, it helps with iteration time to not include these here, but rather in the .cpp
 	#if PLATFORM_USES_FIXED_GMalloc_CLASS && !FORCE_ANSI_ALLOCATOR && USE_MALLOC_BINNED2
 		#define FMEMORY_INLINE_FUNCTION_DECORATOR  FORCEINLINE
 		#define FMEMORY_INLINE_GMalloc (FMallocBinned2::MallocBinned2)
diff --git a/Engine/Source/Runtime/Core/Public/HAL/MallocBinned3.h b/Engine/Source/Runtime/Core/Public/HAL/MallocBinned3.h
--- a/Engine/Source/Runtime/Core/Public/HAL/MallocBinned3.h
+++ b/Engine/Source/Runtime/Core/Public/HAL/MallocBinned3.h
@@ -19,68 +19,68 @@
 
 struct FGenericMemoryStats;
 
-#define USE_CACHED_PAGE_ALLOCATOR_FOR_LARGE_ALLOCS (0)
+#define UE_MB3_USE_CACHED_PAGE_ALLOCATOR_FOR_LARGE_ALLOCS (0)
 
 #ifndef USE_512MB_MAX_MEMORY_PER_BLOCK_SIZE 
-#define USE_512MB_MAX_MEMORY_PER_BLOCK_SIZE 0
+#	define USE_512MB_MAX_MEMORY_PER_BLOCK_SIZE		0
 #endif
 
-#define BINNED3_BASE_PAGE_SIZE				4096			// Minimum "page size" for binned3
-#define BINNED3_MINIMUM_ALIGNMENT_SHIFT		4				// Alignment of blocks, expressed as a shift
-#define BINNED3_MINIMUM_ALIGNMENT			16				// Alignment of blocks
+#define UE_MB3_BASE_PAGE_SIZE						4096			// Minimum "page size" for binned3
+#define UE_MB3_MINIMUM_ALIGNMENT_SHIFT				4				// Alignment of bins, expressed as a shift
+#define UE_MB3_MINIMUM_ALIGNMENT					16				// Alignment of bins
 
 #ifndef BINNED3_MAX_SMALL_POOL_ALIGNMENT
-#define BINNED3_MAX_SMALL_POOL_ALIGNMENT	128
+#	define BINNED3_MAX_SMALL_POOL_ALIGNMENT			128
 #endif
 
 
-#ifndef BINNED3_MAX_SMALL_POOL_SIZE
-#if USE_CACHED_PAGE_ALLOCATOR_FOR_LARGE_ALLOCS
-#define BINNED3_MAX_SMALL_POOL_SIZE			(BINNEDCOMMON_MAX_LISTED_SMALL_POOL_SIZE)	// Maximum medium block size
-#else
-#define BINNED3_MAX_SMALL_POOL_SIZE			(128 * 1024)	// Maximum medium block size
-#endif
+#ifndef UE_MB3_MAX_SMALL_POOL_SIZE
+#	if UE_MB3_USE_CACHED_PAGE_ALLOCATOR_FOR_LARGE_ALLOCS
+#		define UE_MB3_MAX_SMALL_POOL_SIZE			(UE_MBC_MAX_LISTED_SMALL_POOL_SIZE)	// Maximum small bin size
+#	else
+#		define UE_MB3_MAX_SMALL_POOL_SIZE			(128 * 1024)	// Maximum small bin size
+#	endif
 #endif
-#define BINNED3_SMALL_POOL_COUNT			(BINNEDCOMMON_NUM_LISTED_SMALL_POOLS + (BINNED3_MAX_SMALL_POOL_SIZE - BINNEDCOMMON_MAX_LISTED_SMALL_POOL_SIZE) / BINNED3_BASE_PAGE_SIZE)
+#define UE_MB3_SMALL_POOL_COUNT						(UE_MBC_NUM_LISTED_SMALL_POOLS + (UE_MB3_MAX_SMALL_POOL_SIZE - UE_MBC_MAX_LISTED_SMALL_POOL_SIZE) / UE_MB3_BASE_PAGE_SIZE)
 
 #if USE_512MB_MAX_MEMORY_PER_BLOCK_SIZE
-#define MAX_MEMORY_PER_BLOCK_SIZE_SHIFT (29) // maximum of 512MB per block size
+#	define UE_MB3_MAX_MEMORY_PER_POOL_SIZE_SHIFT	(29) // maximum of 512MB per pool
 #else
-#define MAX_MEMORY_PER_BLOCK_SIZE_SHIFT (30) // maximum of 1GB per block size
+#	define UE_MB3_MAX_MEMORY_PER_POOL_SIZE_SHIFT	(30) // maximum of 1GB per pool
 #endif
 
-#define MAX_MEMORY_PER_BLOCK_SIZE (1ull << MAX_MEMORY_PER_BLOCK_SIZE_SHIFT) 
+#define UE_MB3_MAX_MEMORY_PER_POOL_SIZE				(1ull << UE_MB3_MAX_MEMORY_PER_POOL_SIZE_SHIFT) 
 
 // This choice depends on how efficient the OS is with sparse commits in large VM blocks
 #if !defined(BINNED3_USE_SEPARATE_VM_PER_POOL)
-	#if PLATFORM_WINDOWS
-		#define BINNED3_USE_SEPARATE_VM_PER_POOL (1)
-	#else
-		#define BINNED3_USE_SEPARATE_VM_PER_POOL (0)
-	#endif
+#	if PLATFORM_WINDOWS
+#		define BINNED3_USE_SEPARATE_VM_PER_POOL		(1)
+#	else
+#		define BINNED3_USE_SEPARATE_VM_PER_POOL		(0)
+#	endif
 #endif
 
-#define DEFAULT_GMallocBinned3PerThreadCaches 1
-#define DEFAULT_GMallocBinned3AllocExtra 32
-#define BINNED3_MAX_GMallocBinned3MaxBundlesBeforeRecycle 8
+#define UE_DEFAULT_GBinned3PerThreadCaches			1
+#define UE_DEFAULT_GBinned3AllocExtra				32
+#define UE_DEFAULT_GBinned3MaxBundlesBeforeRecycle	8
 
-#define BINNED3_ALLOW_RUNTIME_TWEAKING UE_BINNEDCOMMON_ALLOW_RUNTIME_TWEAKING
-#if BINNED3_ALLOW_RUNTIME_TWEAKING
-	extern CORE_API int32 GMallocBinned3PerThreadCaches;
-	extern CORE_API int32 GMallocBinned3MaxBundlesBeforeRecycle;
-	extern CORE_API int32 GMallocBinned3AllocExtra;
+#define UE_MB3_ALLOW_RUNTIME_TWEAKING				UE_MBC_ALLOW_RUNTIME_TWEAKING
+#if UE_MB3_ALLOW_RUNTIME_TWEAKING
+	extern CORE_API int32 GBinned3PerThreadCaches;
+	extern CORE_API int32 GBinned3MaxBundlesBeforeRecycle;
+	extern CORE_API int32 GBinned3AllocExtra;
 #else
-	#define GMallocBinned3PerThreadCaches DEFAULT_GMallocBinned3PerThreadCaches
-	#define GMallocBinned3MaxBundlesBeforeRecycle BINNED3_MAX_GMallocBinned3MaxBundlesBeforeRecycle
-	#define GMallocBinned3AllocExtra DEFAULT_GMallocBinned3AllocExtra
+	#define GBinned3PerThreadCaches					UE_DEFAULT_GBinned3PerThreadCaches
+	#define GBinned3MaxBundlesBeforeRecycle			UE_DEFAULT_GBinned3MaxBundlesBeforeRecycle
+	#define GBinned3AllocExtra						UE_DEFAULT_GBinned3AllocExtra
 #endif
 
-#define BINNED3_ALLOCATOR_STATS UE_BINNEDCOMMON_ALLOCATOR_STATS
+#define UE_MB3_ALLOCATOR_STATS						UE_MBC_ALLOCATOR_STATS
 
-#if BINNED3_ALLOCATOR_STATS
-	#define BINNED3_ALLOCATOR_PER_BIN_STATS !(UE_BUILD_SHIPPING || UE_BUILD_TEST)
+#if UE_MB3_ALLOCATOR_STATS
+#	define UE_M3_ALLOCATOR_PER_BIN_STATS			!(UE_BUILD_SHIPPING || UE_BUILD_TEST)
 #else
-	#define BINNED3_ALLOCATOR_PER_BIN_STATS 0
+#	define UE_M3_ALLOCATOR_PER_BIN_STATS			0
 #endif
 
 
@@ -88,10 +88,27 @@ PRAGMA_DISABLE_UNSAFE_TYPECAST_WARNINGS
 
 //
 // Optimized virtual memory allocator.
+// 
+// MallocBinned3 supports two types of allocations - large and small pool allocation:
+// 1. For small pool allocation MallocBinned3 reserves contiguous range of virtual memory (a Pool) for each allocation size (a Bin).
+//    These Pools can be adjacent to each other in memory if BINNED3_USE_SEPARATE_VM_PER_POOL is set to 0.
+//    By default each Pool reserves 1 GB of address space, unless USE_512MB_MAX_MEMORY_PER_BLOCK_SIZE is defined to 1.
+//    Each Pool commits and decommits reserved memory in Blocks. Each Block is equal to at least one memory page in it's size.
+//    A Block contains N Bins in a way to minimize the tail memory waste. For example 16 bytes bins fit inside 4KB memory page without any waste;
+//    i.e. one Block for 16 byte Bins contains 256 bins. However 736 bytes Bin uses two 4KB pages in it's Block to minimize memory waste.
+//    Each Pool manages it's Blocks allocation via a Bit Tree (BlocksAllocatedBits\BlocksExhaustedBits members).
+//    And every Block manages it's Bins via a number of additional data structures - FPoolInfoSmall and FFreeBlock.
+//    FFreeBlock is an in-place header for the Block that's stored at the end of each block and contains info on number of free Bins and index of the next free Block if any
+//    Memory is allocated top down in a Block
+// 
+// 2. For large allocations we go directly to OS, unless UE_MB3_USE_CACHED_PAGE_ALLOCATOR_FOR_LARGE_ALLOCS is defined to 1
+//    Each allocation is handled via FPlatformVirtualMemoryBlock and information about it is stored in FPoolInfoLarge
+//    FPoolInfoLarge stores the original allocation size and how much memory was committed by the OS in case we need to realloc that allocation and there's enough tail waste to do that in place
+//    All FPoolInfoLarge are stored in the PoolHashBucket and HashBucketFreeList
 //
-class CORE_API FMallocBinned3 : public TMallocBinnedCommon<FMallocBinned3, BINNED3_MINIMUM_ALIGNMENT, BINNED3_MAX_SMALL_POOL_ALIGNMENT, BINNED3_MINIMUM_ALIGNMENT_SHIFT, BINNED3_SMALL_POOL_COUNT, BINNED3_MAX_SMALL_POOL_SIZE>
+class CORE_API FMallocBinned3 : public TMallocBinnedCommon<FMallocBinned3, UE_MB3_MINIMUM_ALIGNMENT, BINNED3_MAX_SMALL_POOL_ALIGNMENT, UE_MB3_MINIMUM_ALIGNMENT_SHIFT, UE_MB3_SMALL_POOL_COUNT, UE_MB3_MAX_SMALL_POOL_SIZE>
 {	
-	// Forward declares.
+	// Forward declarations
 	struct FPoolInfoLarge;
 	struct FPoolInfoSmall;
 	struct FPoolTable;
@@ -107,20 +124,21 @@ class CORE_API FMallocBinned3 : public TMallocBinnedCommon<FMallocBinned3, BINNE
 			CANARY_VALUE = 0xe7
 		};
 
-		FORCEINLINE FFreeBlock(uint32 InPageSize, uint32 InBlockSize, uint8 InPoolIndex)
-			: BlockSizeShifted(InBlockSize >> BINNED3_MINIMUM_ALIGNMENT_SHIFT)
+		FORCEINLINE FFreeBlock(uint32 InBlockSize, uint32 InBinSize, uint8 InPoolIndex)
+			: BinSizeShifted(InBinSize >> UE_MB3_MINIMUM_ALIGNMENT_SHIFT)
 			, PoolIndex(InPoolIndex)
 			, Canary(CANARY_VALUE)
-			, NextFreeIndex(MAX_uint32)
+			, NextFreeBlockIndex(MAX_uint32)
 		{
-			check(InPoolIndex < MAX_uint8 && (InBlockSize >> BINNED3_MINIMUM_ALIGNMENT_SHIFT) <= MAX_uint16);
-			NumFreeBlocks = InPageSize / InBlockSize;
+			check(InPoolIndex < MAX_uint8 && (InBinSize >> UE_MB3_MINIMUM_ALIGNMENT_SHIFT) <= MAX_uint16);
+			NumFreeBins = InBlockSize / InBinSize;
 		}
 
-		FORCEINLINE uint32 GetNumFreeRegularBlocks() const
+		FORCEINLINE uint32 GetNumFreeBins() const
 		{
-			return NumFreeBlocks;
+			return NumFreeBins;
 		}
+
 		FORCEINLINE bool IsCanaryOk() const
 		{
 			return Canary == FFreeBlock::CANARY_VALUE;
@@ -132,39 +150,39 @@ class CORE_API FMallocBinned3 : public TMallocBinnedCommon<FMallocBinned3, BINNE
 			{
 				CanaryFail();
 			}
-			//checkSlow(PoolIndex == BoundSizeToPoolIndex(BlockSize));
+			//checkSlow(PoolIndex == BoundSizeToPoolIndex(uint32(BinSizeShifted) << UE_MB3_MINIMUM_ALIGNMENT_SHIFT));
 		}
 		void CanaryFail() const;
 
-		FORCEINLINE void* AllocateRegularBlock()
+		FORCEINLINE void* AllocateBin()
 		{
-			--NumFreeBlocks;
-			return (uint8*)this + NumFreeBlocks* (uint32(BlockSizeShifted) << BINNED3_MINIMUM_ALIGNMENT_SHIFT);
+			--NumFreeBins;
+			return (uint8*)this + NumFreeBins * (uint32(BinSizeShifted) << UE_MB3_MINIMUM_ALIGNMENT_SHIFT);
 		}
 
-		uint16 BlockSizeShifted;		// Size of the blocks that this list points to >> BINNED3_MINIMUM_ALIGNMENT_SHIFT
-		uint8 PoolIndex;				// Index of this pool
-		uint8 Canary;					// Constant value of 0xe3
-		uint32 NumFreeBlocks;          // Number of consecutive free blocks here, at least 1.
-		uint32 NextFreeIndex;          // Next free block or MAX_uint32
+		uint16 BinSizeShifted;		// Size of the bins that this list points to >> UE_MB3_MINIMUM_ALIGNMENT_SHIFT
+		uint8  PoolIndex;			// Index of this pool
+		uint8  Canary;				// Constant value of 0xe3
+		uint32 NumFreeBins;			// Number of consecutive free bins here, at least 1.
+		uint32 NextFreeBlockIndex;	// Next free block or MAX_uint32
 	};
 
 	/** Pool table. */
 	struct FPoolTable
 	{
-		uint32 BlockSize;
-		uint16 BlocksPerBlockOfBlocks;
-		uint8 PagesPlatformForBlockOfBlocks;
+		uint32 BinSize;						// Bin size, i.e. 16, 32, 64 etc bytes. Pretty redundant and can be computed
+		//uint32 BlockSize;					// Size of a block. A block contains N bins in a way to minimize tail memory waste
+		uint32 NumMemoryPagesPerBlock;		// Num memory pages needed to allocate one block. 
 
-		FBitTree BlockOfBlockAllocationBits; // one bits in here mean the virtual memory is committed
-		FBitTree BlockOfBlockIsExhausted;    // one bit in here means the pool is completely full
+		FBitTree BlocksAllocatedBits;		// One bits in here mean the virtual memory is committed
+		FBitTree BlocksExhaustedBits;		// One bit in here means the pool is completely full
 
-		uint32 NumEverUsedBlockOfBlocks;
-		FPoolInfoSmall** PoolInfos;
+		uint32 NumEverUsedBlocks;
+		FPoolInfoSmall** PoolInfos;			// Book keeping info about every allocated area for this bin's pool
 
-		uint64 UnusedAreaOffsetLow;
+		uint64 UnusedAreaOffsetLow;			// High watermark for allocated vm memory for this pool
 
-#if BINNED3_ALLOCATOR_PER_BIN_STATS
+#if UE_M3_ALLOCATOR_PER_BIN_STATS
 		// these are "head end" stats, above the TLS cache
 		TAtomic<int64> TotalRequestedAllocSize;
 		TAtomic<int64> TotalAllocCount;
@@ -172,7 +190,7 @@ class CORE_API FMallocBinned3 : public TMallocBinnedCommon<FMallocBinned3, BINNE
 
 		FORCEINLINE void HeadEndAlloc(SIZE_T Size)
 		{
-			check(Size >= 0 && Size <= BlockSize);
+			check(Size >= 0 && Size <= BinSize);
 			TotalRequestedAllocSize += Size;
 			TotalAllocCount++;
 		}
@@ -191,47 +209,48 @@ class CORE_API FMallocBinned3 : public TMallocBinnedCommon<FMallocBinned3, BINNE
 	};
 
 	// Pool tables for different pool sizes
-	FPoolTable SmallPoolTables[BINNED3_SMALL_POOL_COUNT];
+	FPoolTable SmallPoolTables[UE_MB3_SMALL_POOL_COUNT];
 
 	uint32 SmallPoolInfosPerPlatformPage;
 
-	PoolHashBucket* HashBuckets;
-	PoolHashBucket* HashBucketFreeList;
+	PoolHashBucket* HashBuckets;			// Hash buckets for external allocations, reserved in constructor based on the platform constants like page size and virtual address high\low hints
+	PoolHashBucket* HashBucketFreeList;		// Hash buckets for allocations that were allocated outside of the platform constants virtual address high\low hints
 	uint64 NumLargePoolsPerPage;
 
 	FCriticalSection Mutex;
 
 #if !BINNED3_USE_SEPARATE_VM_PER_POOL
-	FORCEINLINE uint64 PoolIndexFromPtr(const void* Ptr) // returns a uint64 for it can also be used to check if it is an OS allocation
+	FORCEINLINE uint64 PoolIndexFromPtr(const void* Ptr) const		// returns a uint64 for it can also be used to check if it is an OS allocation
 	{
-		return (UPTRINT(Ptr) - UPTRINT(Binned3BaseVMPtr)) >> MAX_MEMORY_PER_BLOCK_SIZE_SHIFT;
+		return (UPTRINT(Ptr) - UPTRINT(Binned3BaseVMPtr)) >> UE_MB3_MAX_MEMORY_PER_POOL_SIZE_SHIFT;
 	}
-	FORCEINLINE uint8* PoolBasePtr(uint32 InPoolIndex)
+	FORCEINLINE uint8* PoolBasePtr(uint32 InPoolIndex) const
 	{
-		return Binned3BaseVMPtr + InPoolIndex * MAX_MEMORY_PER_BLOCK_SIZE;
+		return Binned3BaseVMPtr + InPoolIndex * UE_MB3_MAX_MEMORY_PER_POOL_SIZE;
 	}
 #else
-#if BINNED3_ALLOCATOR_STATS
-	void RecordPoolSearch(uint32 Tests);
-#else
-	FORCEINLINE void RecordPoolSearch(uint32 Tests)
-	{
+#	if UE_MB3_ALLOCATOR_STATS
+		void RecordPoolSearch(uint32 Tests) const;
+#	else
+		FORCEINLINE void RecordPoolSearch(uint32 Tests) const
+		{
 
-	}
-#endif
-	FORCEINLINE uint64 PoolIndexFromPtr(const void* Ptr) // returns a uint64 for it can also be used to check if it is an OS allocation
+		}
+#	endif
+
+	FORCEINLINE uint64 PoolIndexFromPtr(const void* Ptr) const		// returns a uint64 for it can also be used to check if it is an OS allocation
 	{
 		if (PoolSearchDiv == 0)
 		{
-			return (UPTRINT(Ptr) - UPTRINT(PoolBaseVMPtr[0])) >> MAX_MEMORY_PER_BLOCK_SIZE_SHIFT;
+			return (UPTRINT(Ptr) - UPTRINT(PoolBaseVMPtr[0])) >> UE_MB3_MAX_MEMORY_PER_POOL_SIZE_SHIFT;
 		}
-		uint64 PoolIndex = BINNED3_SMALL_POOL_COUNT;
-		if (((uint8*)Ptr >= PoolBaseVMPtr[0]) & ((uint8*)Ptr < HighestPoolBaseVMPtr + MAX_MEMORY_PER_BLOCK_SIZE))
+		uint64 PoolIndex = UE_MB3_SMALL_POOL_COUNT;
+		if (((uint8*)Ptr >= PoolBaseVMPtr[0]) & ((uint8*)Ptr < HighestPoolBaseVMPtr + UE_MB3_MAX_MEMORY_PER_POOL_SIZE))
 		{
 			PoolIndex = uint64((uint8*)Ptr - PoolBaseVMPtr[0]) / PoolSearchDiv;
-			if (PoolIndex >= BINNED3_SMALL_POOL_COUNT)
+			if (PoolIndex >= UE_MB3_SMALL_POOL_COUNT)
 			{
-				PoolIndex = BINNED3_SMALL_POOL_COUNT - 1;
+				PoolIndex = UE_MB3_SMALL_POOL_COUNT - 1;
 			}
 			uint32 Tests = 1; // we are counting potential cache misses here, not actual comparisons
 			if ((uint8*)Ptr < PoolBaseVMPtr[PoolIndex])
@@ -240,24 +259,24 @@ class CORE_API FMallocBinned3 : public TMallocBinnedCommon<FMallocBinned3, BINNE
 				{
 					Tests++;
 					PoolIndex--;
-					check(PoolIndex < BINNED3_SMALL_POOL_COUNT);
+					check(PoolIndex < UE_MB3_SMALL_POOL_COUNT);
 				} while ((uint8*)Ptr < PoolBaseVMPtr[PoolIndex]);
-				if ((uint8*)Ptr >= PoolBaseVMPtr[PoolIndex] + MAX_MEMORY_PER_BLOCK_SIZE)
+				if ((uint8*)Ptr >= PoolBaseVMPtr[PoolIndex] + UE_MB3_MAX_MEMORY_PER_POOL_SIZE)
 				{
-					PoolIndex = BINNED3_SMALL_POOL_COUNT; // was in the gap
+					PoolIndex = UE_MB3_SMALL_POOL_COUNT; // was in the gap
 				}
 			}
-			else if ((uint8*)Ptr >= PoolBaseVMPtr[PoolIndex] + MAX_MEMORY_PER_BLOCK_SIZE)
+			else if ((uint8*)Ptr >= PoolBaseVMPtr[PoolIndex] + UE_MB3_MAX_MEMORY_PER_POOL_SIZE)
 			{
 				do
 				{
 					Tests++;
 					PoolIndex++;
-					check(PoolIndex < BINNED3_SMALL_POOL_COUNT);
-				} while ((uint8*)Ptr >= PoolBaseVMPtr[PoolIndex] + MAX_MEMORY_PER_BLOCK_SIZE);
+					check(PoolIndex < UE_MB3_SMALL_POOL_COUNT);
+				} while ((uint8*)Ptr >= PoolBaseVMPtr[PoolIndex] + UE_MB3_MAX_MEMORY_PER_POOL_SIZE);
 				if ((uint8*)Ptr < PoolBaseVMPtr[PoolIndex])
 				{
-					PoolIndex = BINNED3_SMALL_POOL_COUNT; // was in the gap
+					PoolIndex = UE_MB3_SMALL_POOL_COUNT; // was in the gap
 				}
 			}
 			RecordPoolSearch(Tests);
@@ -265,45 +284,47 @@ class CORE_API FMallocBinned3 : public TMallocBinnedCommon<FMallocBinned3, BINNE
 		return PoolIndex;
 	}
 
-	FORCEINLINE uint8* PoolBasePtr(uint32 InPoolIndex)
+	FORCEINLINE uint8* PoolBasePtr(uint32 InPoolIndex) const
 	{
 		return PoolBaseVMPtr[InPoolIndex];
 	}
-#endif
-	FORCEINLINE uint32 PoolIndexFromPtrChecked(const void* Ptr)
+#endif //~!BINNED3_USE_SEPARATE_VM_PER_POOL
+
+	FORCEINLINE uint32 PoolIndexFromPtrChecked(const void* Ptr) const
 	{
-		uint64 Result = PoolIndexFromPtr(Ptr);
-		check(Result < BINNED3_SMALL_POOL_COUNT);
+		const uint64 Result = PoolIndexFromPtr(Ptr);
+		check(Result < UE_MB3_SMALL_POOL_COUNT);
 		return (uint32)Result;
 	}
 
-	FORCEINLINE bool IsOSAllocation(const void* Ptr)
+	FORCEINLINE bool IsOSAllocation(const void* Ptr) const
 	{
-		return PoolIndexFromPtr(Ptr) >= BINNED3_SMALL_POOL_COUNT;
+		return PoolIndexFromPtr(Ptr) >= UE_MB3_SMALL_POOL_COUNT;
 	}
 
-
-	FORCEINLINE void* BlockOfBlocksPointerFromContainedPtr(const void* Ptr, uint8 PagesPlatformForBlockOfBlocks, uint32& OutBlockOfBlocksIndex)
+	FORCEINLINE void* BlockPointerFromContainedPtr(const void* Ptr, uint8 NumMemoryPagesPerBlock, uint32& OutBlockIndex) const
 	{
-		uint32 PoolIndex = PoolIndexFromPtrChecked(Ptr);
+		const uint32 PoolIndex = PoolIndexFromPtrChecked(Ptr);
 		uint8* PoolStart = PoolBasePtr(PoolIndex);
-		uint64 BlockOfBlocksIndex = (UPTRINT(Ptr) - UPTRINT(PoolStart)) / (UPTRINT(PagesPlatformForBlockOfBlocks) * UPTRINT(OsAllocationGranularity));
-		OutBlockOfBlocksIndex = BlockOfBlocksIndex;
+		const uint64 BlockIndex = (UPTRINT(Ptr) - UPTRINT(PoolStart)) / (UPTRINT(NumMemoryPagesPerBlock) * UPTRINT(OsAllocationGranularity));
+		OutBlockIndex = BlockIndex;
 
-		uint8* Result = PoolStart + BlockOfBlocksIndex * UPTRINT(PagesPlatformForBlockOfBlocks) * UPTRINT(OsAllocationGranularity);
+		uint8* Result = PoolStart + BlockIndex * UPTRINT(NumMemoryPagesPerBlock) * UPTRINT(OsAllocationGranularity);
 
-		check(Result < PoolStart + MAX_MEMORY_PER_BLOCK_SIZE);
+		check(Result < PoolStart + UE_MB3_MAX_MEMORY_PER_POOL_SIZE);
 		return Result;
 	}
-	FORCEINLINE uint8* BlockPointerFromIndecies(uint32 InPoolIndex, uint32 BlockOfBlocksIndex, uint32 BlockOfBlocksSize)
+
+	FORCEINLINE uint8* BlockPointerFromIndecies(uint32 InPoolIndex, uint32 BlockIndex, uint32 BlockSize) const
 	{
 		uint8* PoolStart = PoolBasePtr(InPoolIndex);
-		uint8* Ptr = PoolStart + BlockOfBlocksIndex * uint64(BlockOfBlocksSize);
-		check(Ptr + BlockOfBlocksSize <= PoolStart + MAX_MEMORY_PER_BLOCK_SIZE);
+		uint8* Ptr = PoolStart + BlockIndex * uint64(BlockSize);
+		check(Ptr + BlockSize <= PoolStart + UE_MB3_MAX_MEMORY_PER_POOL_SIZE);
 		return Ptr;
 	}
-	FPoolInfoSmall* PushNewPoolToFront(FPoolTable& Table, uint32 InBlockSize, uint32 InPoolIndex, uint32& OutBlockOfBlocksIndex);
-	FPoolInfoSmall* GetFrontPool(FPoolTable& Table, uint32 InPoolIndex, uint32& OutBlockOfBlocksIndex);
+
+	FPoolInfoSmall* PushNewPoolToFront(FPoolTable& Table, uint32 InBinSize, uint32 InPoolIndex, uint32& OutBlockIndex);
+	FPoolInfoSmall* GetFrontPool(FPoolTable& Table, uint32 InPoolIndex, uint32& OutBlockIndex);
 
 public:
 
@@ -319,19 +340,19 @@ public:
 	
 		// Only allocate from the small pools if the size is small enough and the alignment isn't crazy large.
 		// With large alignments, we'll waste a lot of memory allocating an entire page, but such alignments are highly unlikely in practice.
-		if ((Size <= BINNED3_MAX_SMALL_POOL_SIZE) & (Alignment <= BINNED3_MINIMUM_ALIGNMENT)) // one branch, not two
+		if ((Size <= UE_MB3_MAX_SMALL_POOL_SIZE) & (Alignment <= UE_MB3_MINIMUM_ALIGNMENT)) // one branch, not two
 		{
-			FPerThreadFreeBlockLists* Lists = GMallocBinned3PerThreadCaches ? FPerThreadFreeBlockLists::Get() : nullptr;
+			FPerThreadFreeBlockLists* Lists = GBinned3PerThreadCaches ? FPerThreadFreeBlockLists::Get() : nullptr;
 			if (Lists)
 			{
-				uint32 PoolIndex = BoundSizeToPoolIndex(Size, MemSizeToIndex);
-				uint32 BlockSize = PoolIndexToBlockSize(PoolIndex);
+				const uint32 PoolIndex = BoundSizeToPoolIndex(Size, MemSizeToPoolIndex);
+				const uint32 BinSize = PoolIndexToBinSize(PoolIndex);
 				Result = Lists->Malloc(PoolIndex);
-#if BINNED3_ALLOCATOR_STATS
+#if UE_MB3_ALLOCATOR_STATS
 				if (Result)
 				{
 					SmallPoolTables[PoolIndex].HeadEndAlloc(Size);
-					Lists->AllocatedMemory += BlockSize;
+					Lists->AllocatedMemory += BinSize;
 				}
 #endif
 			}
@@ -343,57 +364,58 @@ public:
 
 		return Result;
 	}
+
 	FORCEINLINE virtual void* Realloc(void* Ptr, SIZE_T NewSize, uint32 Alignment) override
 	{
-		if (NewSize <= BINNED3_MAX_SMALL_POOL_SIZE && Alignment <= BINNED3_MINIMUM_ALIGNMENT) // one branch, not two
+		if (NewSize <= UE_MB3_MAX_SMALL_POOL_SIZE && Alignment <= UE_MB3_MINIMUM_ALIGNMENT) // one branch, not two
 		{
-			FPerThreadFreeBlockLists* Lists = GMallocBinned3PerThreadCaches ? FPerThreadFreeBlockLists::Get() : nullptr;
+			FPerThreadFreeBlockLists* Lists = GBinned3PerThreadCaches ? FPerThreadFreeBlockLists::Get() : nullptr;
 
-			uint64 PoolIndex = PoolIndexFromPtr(Ptr);
-			if ((!!Lists) & ((!Ptr) | (PoolIndex < BINNED3_SMALL_POOL_COUNT)))
+			const uint64 PoolIndex = PoolIndexFromPtr(Ptr);
+			if ((!!Lists) & ((!Ptr) | (PoolIndex < UE_MB3_SMALL_POOL_COUNT)))
 			{
-				uint32 BlockSize = 0;
+				uint32 BinSize = 0;
 
 				bool bCanFree = true; // the nullptr is always "freeable"
 				if (Ptr)
 				{
 					// Reallocate to a smaller/bigger pool if necessary
-					BlockSize = PoolIndexToBlockSize(PoolIndex);
-					if ((!!NewSize) & (NewSize <= BlockSize) & ((!PoolIndex) | (NewSize > PoolIndexToBlockSize(static_cast<uint32>(PoolIndex - 1)))))
+					BinSize = PoolIndexToBinSize(PoolIndex);
+					if ((!!NewSize) & (NewSize <= BinSize) & ((!PoolIndex) | (NewSize > PoolIndexToBinSize(static_cast<uint32>(PoolIndex - 1)))))
 					{
-#if BINNED3_ALLOCATOR_STATS
+#if UE_MB3_ALLOCATOR_STATS
 						SmallPoolTables[PoolIndex].HeadEndAlloc(NewSize);
 						SmallPoolTables[PoolIndex].HeadEndFree();
 #endif
 						return Ptr;
 					}
-					bCanFree = Lists->CanFree(PoolIndex, BlockSize);
+					bCanFree = Lists->CanFree(PoolIndex, BinSize);
 				}
 				if (bCanFree)
 				{
-					uint32 NewPoolIndex = BoundSizeToPoolIndex(NewSize, MemSizeToIndex);
-					uint32 NewBlockSize = PoolIndexToBlockSize(NewPoolIndex);
+					const uint32 NewPoolIndex = BoundSizeToPoolIndex(NewSize, MemSizeToPoolIndex);
+					const uint32 NewBinSize = PoolIndexToBinSize(NewPoolIndex);
 					void* Result = NewSize ? Lists->Malloc(NewPoolIndex) : nullptr;
-#if BINNED3_ALLOCATOR_STATS
+#if UE_MB3_ALLOCATOR_STATS
 					if (Result)
 					{
 						SmallPoolTables[NewPoolIndex].HeadEndAlloc(NewSize);
-						Lists->AllocatedMemory += NewBlockSize;
+						Lists->AllocatedMemory += NewBinSize;
 					}
 #endif
 					if (Result || !NewSize)
 					{
 						if (Result && Ptr)
 						{
-							FMemory::Memcpy(Result, Ptr, FPlatformMath::Min<SIZE_T>(NewSize, BlockSize));
+							FMemory::Memcpy(Result, Ptr, FPlatformMath::Min<SIZE_T>(NewSize, BinSize));
 						}
 						if (Ptr)
 						{
-							bool bDidPush = Lists->Free(Ptr, PoolIndex, BlockSize);
+							const bool bDidPush = Lists->Free(Ptr, PoolIndex, BinSize);
 							checkSlow(bDidPush);
-#if BINNED3_ALLOCATOR_STATS
+#if UE_MB3_ALLOCATOR_STATS
 							SmallPoolTables[PoolIndex].HeadEndFree();
-							Lists->AllocatedMemory -= BlockSize;
+							Lists->AllocatedMemory -= BinSize;
 #endif
 						}
 
@@ -408,18 +430,18 @@ public:
 
 	FORCEINLINE virtual void Free(void* Ptr) override
 	{
-		uint64 PoolIndex = PoolIndexFromPtr(Ptr);
-		if (PoolIndex < BINNED3_SMALL_POOL_COUNT)
+		const uint64 PoolIndex = PoolIndexFromPtr(Ptr);
+		if (PoolIndex < UE_MB3_SMALL_POOL_COUNT)
 		{
-			FPerThreadFreeBlockLists* Lists = GMallocBinned3PerThreadCaches ? FPerThreadFreeBlockLists::Get() : nullptr;
+			FPerThreadFreeBlockLists* Lists = GBinned3PerThreadCaches ? FPerThreadFreeBlockLists::Get() : nullptr;
 			if (Lists)
 			{
-				int32 BlockSize = PoolIndexToBlockSize(PoolIndex);
-				if (Lists->Free(Ptr, PoolIndex, BlockSize))
+				const int32 BinSize = PoolIndexToBinSize(PoolIndex);
+				if (Lists->Free(Ptr, PoolIndex, BinSize))
 				{
-#if BINNED3_ALLOCATOR_STATS
+#if UE_MB3_ALLOCATOR_STATS
 					SmallPoolTables[PoolIndex].HeadEndFree();
-					Lists->AllocatedMemory -= BlockSize;
+					Lists->AllocatedMemory -= BinSize;
 #endif
 					return;
 				}
@@ -427,12 +449,13 @@ public:
 		}
 		FreeExternal(Ptr);
 	}
+
 	FORCEINLINE virtual bool GetAllocationSize(void *Ptr, SIZE_T &SizeOut) override
 	{
-		uint64 PoolIndex = PoolIndexFromPtr(Ptr);
-		if (PoolIndex < BINNED3_SMALL_POOL_COUNT)
+		const uint64 PoolIndex = PoolIndexFromPtr(Ptr);
+		if (PoolIndex < UE_MB3_SMALL_POOL_COUNT)
 		{
-			SizeOut = PoolIndexToBlockSize(PoolIndex);
+			SizeOut = PoolIndexToBinSize(PoolIndex);
 			return true;
 		}
 		return GetAllocationSizeExternal(Ptr, SizeOut);
@@ -457,16 +480,16 @@ public:
 	void FreeExternal(void *Ptr);
 	bool GetAllocationSizeExternal(void* Ptr, SIZE_T& SizeOut);
 
-#if BINNED3_ALLOCATOR_STATS
+#if UE_MB3_ALLOCATOR_STATS
 	int64 GetTotalAllocatedSmallPoolMemory() const;
 #endif
 	virtual void GetAllocatorStats( FGenericMemoryStats& out_Stats ) override;
 	/** Dumps current allocator stats to the log. */
 	virtual void DumpAllocatorStats(class FOutputDevice& Ar) override;
 	
-	// +1 enables PoolIndexToBlockSize(~0u / -1) dummy access that helps avoid PoolIndex == 0 branching in Realloc(),
-	// see ((!PoolIndex) | (NewSize > PoolIndexToBlockSize(static_cast<uint32>(PoolIndex - 1)))
-	static uint16 SmallBlockSizesReversedShifted[BINNED3_SMALL_POOL_COUNT + 1]; // this is reversed to get the smallest elements on our main cache line
+	// +1 enables PoolIndexToBinSize(~0u / -1) dummy access that helps avoid PoolIndex == 0 branching in Realloc(),
+	// see ((!PoolIndex) | (NewSize > PoolIndexToBinSize(static_cast<uint32>(PoolIndex - 1)))
+	static uint16 SmallBinSizesReversedShifted[UE_MB3_SMALL_POOL_COUNT + 1]; // this is reversed to get the smallest elements on our main cache line
 	static FMallocBinned3* MallocBinned3;
 	static uint32 OsAllocationGranularity;
 
@@ -480,16 +503,16 @@ public:
 	FPlatformMemory::FPlatformVirtualMemoryBlock Binned3BaseVMBlock;
 #else
 	static uint64 PoolSearchDiv; // if this is zero, the VM turned out to be contiguous anyway so we use a simple subtract and shift
-	static uint8* HighestPoolBaseVMPtr; // this is a duplicate of PoolBaseVMPtr[BINNED3_SMALL_POOL_COUNT - 1]
-	static uint8* PoolBaseVMPtr[BINNED3_SMALL_POOL_COUNT];
-	FPlatformMemory::FPlatformVirtualMemoryBlock PoolBaseVMBlock[BINNED3_SMALL_POOL_COUNT];
+	static uint8* HighestPoolBaseVMPtr; // this is a duplicate of PoolBaseVMPtr[UE_MB3_SMALL_POOL_COUNT - 1]
+	static uint8* PoolBaseVMPtr[UE_MB3_SMALL_POOL_COUNT];
+	FPlatformMemory::FPlatformVirtualMemoryBlock PoolBaseVMBlock[UE_MB3_SMALL_POOL_COUNT];
 #endif
 	// Mapping of sizes to small table indices
-	static uint8 MemSizeToIndex[1 + (BINNED3_MAX_SMALL_POOL_SIZE >> BINNED3_MINIMUM_ALIGNMENT_SHIFT)];
+	static uint8 MemSizeToPoolIndex[1 + (UE_MB3_MAX_SMALL_POOL_SIZE >> UE_MB3_MINIMUM_ALIGNMENT_SHIFT)];
 
-	FORCEINLINE uint32 PoolIndexToBlockSize(uint32 PoolIndex) const
+	FORCEINLINE uint32 PoolIndexToBinSize(uint32 PoolIndex) const
 	{
-		return uint32(SmallBlockSizesReversedShifted[BINNED3_SMALL_POOL_COUNT - PoolIndex - 1]) << BINNED3_MINIMUM_ALIGNMENT_SHIFT;
+		return uint32(SmallBinSizesReversedShifted[UE_MB3_SMALL_POOL_COUNT - PoolIndex - 1]) << UE_MB3_MINIMUM_ALIGNMENT_SHIFT;
 	}
 
 	void FreeBundles(FBundleNode* Bundles, uint32 PoolIndex);
@@ -504,8 +527,8 @@ public:
 
 PRAGMA_RESTORE_UNSAFE_TYPECAST_WARNINGS
 
-#define BINNED3_INLINE (1)
-#if BINNED3_INLINE // during development, it helps with iteration time to not include these here, but rather in the .cpp
+#define UE_MB3_INLINE (1)
+#if UE_MB3_INLINE // during development, it helps with iteration time to not include these here, but rather in the .cpp
 	#if PLATFORM_USES_FIXED_GMalloc_CLASS && !FORCE_ANSI_ALLOCATOR && 0/*USE_MALLOC_BINNED3*/
 		#define FMEMORY_INLINE_FUNCTION_DECORATOR  FORCEINLINE
 		#define FMEMORY_INLINE_GMalloc (FMallocBinned3::MallocBinned3)
diff --git a/Engine/Source/Runtime/Core/Public/HAL/MallocBinnedCommon.h b/Engine/Source/Runtime/Core/Public/HAL/MallocBinnedCommon.h
--- a/Engine/Source/Runtime/Core/Public/HAL/MallocBinnedCommon.h
+++ b/Engine/Source/Runtime/Core/Public/HAL/MallocBinnedCommon.h
@@ -17,31 +17,31 @@
 #include "HAL/PlatformMemory.h"
 #include "Templates/Function.h"
 
-// A project can define it's own BINNEDCOMMON_MAX_LISTED_SMALL_POOL_SIZE and BINNEDCOMMON_NUM_LISTED_SMALL_POOLS to reduce runtime memory usage
+// A project can define it's own UE_MBC_MAX_LISTED_SMALL_POOL_SIZE and UE_MBC_NUM_LISTED_SMALL_POOLS to reduce runtime memory usage
 // MallocBinnedCommon.cpp has a list of predefined bins that go up to 28672
 // By default allocators (i.e. MB3) that use these bins will rely on this number as a baseline for a small bins count
 // These allocators can increase the amount of small bins they want to manage by going over the default MBC bins list
 // In MB3 case that means defining BINNED3_MAX_SMALL_POOL_SIZE to something like 65536
-// Every bin over the BINNEDCOMMON_MAX_LISTED_SMALL_POOL_SIZE would come with a 4kb increment
+// Every bin over the UE_MBC_MAX_LISTED_SMALL_POOL_SIZE would come with a 4kb increment
 // These small bins would be kept in user mode, increasing application's memory footprint and reducing the time it takes to allocate memory from the said bins
 // If application needs to aggressively reduce it's memory footprint, potentially trading some perf due to an increased amount of kernel calls to allocate memory
-// it can redefine BINNEDCOMMON_MAX_LISTED_SMALL_POOL_SIZE and BINNED3_MAX_SMALL_POOL_SIZE to smaller numbers, the good value is 16384 for both
-// This, however, would require the app to redefine BINNEDCOMMON_NUM_LISTED_SMALL_POOLS too to match the number of bins that fall under the new define's threshold
-// In case of 16384, we'll skip 3 larger bins and so BINNEDCOMMON_NUM_LISTED_SMALL_POOLS should be set to 48 at the time of writing
-#if !defined(BINNEDCOMMON_MAX_LISTED_SMALL_POOL_SIZE)
-	#define BINNEDCOMMON_MAX_LISTED_SMALL_POOL_SIZE	28672
+// it can redefine UE_MBC_MAX_LISTED_SMALL_POOL_SIZE and BINNED3_MAX_SMALL_POOL_SIZE to smaller numbers, the good value is 16384 for both
+// This, however, would require the app to redefine UE_MBC_NUM_LISTED_SMALL_POOLS too to match the number of bins that fall under the new define's threshold
+// In case of 16384, we'll skip 3 larger bins and so UE_MBC_NUM_LISTED_SMALL_POOLS should be set to 48 at the time of writing
+#if !defined(UE_MBC_MAX_LISTED_SMALL_POOL_SIZE)
+#	define UE_MBC_MAX_LISTED_SMALL_POOL_SIZE	28672
 #endif
 
-#if !defined(BINNEDCOMMON_NUM_LISTED_SMALL_POOLS)
-	#define BINNEDCOMMON_NUM_LISTED_SMALL_POOLS	51
+#if !defined(UE_MBC_NUM_LISTED_SMALL_POOLS)
+#	define UE_MBC_NUM_LISTED_SMALL_POOLS	51
 #endif
 
 #if !defined(BINNEDCOMMON_USE_SEPARATE_VM_PER_POOL)
-	#if PLATFORM_WINDOWS
-		#define BINNEDCOMMON_USE_SEPARATE_VM_PER_POOL (1)
-	#else
-		#define BINNEDCOMMON_USE_SEPARATE_VM_PER_POOL (0)
-	#endif
+#	if PLATFORM_WINDOWS
+#		define BINNEDCOMMON_USE_SEPARATE_VM_PER_POOL (1)
+#	else
+#		define BINNEDCOMMON_USE_SEPARATE_VM_PER_POOL (0)
+#	endif
 #endif
 
 
@@ -95,19 +95,15 @@ public:
 
 struct FSizeTableEntry
 {
-	uint32 BlockSize;
-	uint16 BlocksPerBlockOfBlocks;
-	uint8 PagesPlatformForBlockOfBlocks;
+	uint32 BinSize;
+	uint32 NumMemoryPagesPerBlock;
 
-	FSizeTableEntry()
-	{
-	}
-
-	FSizeTableEntry(uint32 InBlockSize, uint64 PlatformPageSize, uint8 Pages4k, uint32 BasePageSize, uint32 MinimumAlignment);
+	FSizeTableEntry() = default;
+	FSizeTableEntry(uint32 InBinSize, uint64 PlatformPageSize, uint8 Num4kbPages, uint32 BasePageSize, uint32 MinimumAlignment);
 
 	bool operator<(const FSizeTableEntry& Other) const
 	{
-		return BlockSize < Other.BlockSize;
+		return BinSize < Other.BinSize;
 	}
 	static uint8 FillSizeTable(uint64 PlatformPageSize, FSizeTableEntry* SizeTable, uint32 BasePageSize, uint32 MinimumAlignment, uint32 MaxSize, uint32 SizeIncrement);
 };
@@ -157,27 +153,27 @@ struct FArenaParams
 #endif
 
 #if AGGRESSIVE_MEMORY_SAVING
-#	define DEFAULT_GMallocBinnedBundleSize 8192
+#	define UE_DEFAULT_GMallocBinnedBundleSize 8192
 #else
-#	define DEFAULT_GMallocBinnedBundleSize 65536
+#	define UE_DEFAULT_GMallocBinnedBundleSize 65536
 #endif
 
-#define DEFAULT_GMallocBinnedBundleCount 64
+#define UE_DEFAULT_GMallocBinnedBundleCount 64
 
-#ifndef UE_BINNEDCOMMON_ALLOW_RUNTIME_TWEAKING
-	#define UE_BINNEDCOMMON_ALLOW_RUNTIME_TWEAKING 0
+#ifndef UE_MBC_ALLOW_RUNTIME_TWEAKING
+#	define UE_MBC_ALLOW_RUNTIME_TWEAKING 0
 #endif
 
-#if UE_BINNEDCOMMON_ALLOW_RUNTIME_TWEAKING
-extern CORE_API int32 GMallocBinnedBundleSize;
-extern CORE_API int32 GMallocBinnedBundleCount;
+#if UE_MBC_ALLOW_RUNTIME_TWEAKING
+	extern CORE_API int32 GMallocBinnedBundleSize;
+	extern CORE_API int32 GMallocBinnedBundleCount;
 #else
-#	define GMallocBinnedBundleSize	DEFAULT_GMallocBinnedBundleSize
-#	define GMallocBinnedBundleCount	DEFAULT_GMallocBinnedBundleCount
+#	define GMallocBinnedBundleSize	UE_DEFAULT_GMallocBinnedBundleSize
+#	define GMallocBinnedBundleCount	UE_DEFAULT_GMallocBinnedBundleSize
 #endif
 
-#ifndef UE_BINNEDCOMMON_ALLOCATOR_STATS
-#	define UE_BINNEDCOMMON_ALLOCATOR_STATS (!UE_BUILD_SHIPPING || WITH_EDITOR)
+#ifndef UE_MBC_ALLOCATOR_STATS
+#	define UE_MBC_ALLOCATOR_STATS (!UE_BUILD_SHIPPING || WITH_EDITOR)
 #endif
 
 extern CORE_API float GMallocBinnedFlushThreadCacheMaxWaitTime;
@@ -196,6 +192,7 @@ protected:
 			, AddressSpaceBase(0)
 		{
 		}
+
 		explicit FPtrToPoolMapping(uint32 InPageSize, uint64 InNumPoolsPerPage, uint64 AddressBase, uint64 AddressLimit)
 		{
 			Init(InPageSize, InNumPoolsPerPage, AddressBase, AddressLimit);
@@ -203,7 +200,7 @@ protected:
 
 		void Init(uint32 InPageSize, uint64 InNumPoolsPerPage, uint64 AddressBase, uint64 AddressLimit)
 		{
-			uint64 PoolPageToPoolBitShift = FPlatformMath::CeilLogTwo64(InNumPoolsPerPage);
+			const uint64 PoolPageToPoolBitShift = FPlatformMath::CeilLogTwo64(InNumPoolsPerPage);
 
 			PtrToPoolPageBitShift = FPlatformMath::CeilLogTwo(InPageSize);
 			HashKeyShift = PtrToPoolPageBitShift + PoolPageToPoolBitShift;
@@ -295,8 +292,8 @@ protected:
 	template <class T>
 	struct TPoolHashBucket
 	{
-		UPTRINT         BucketIndex;
-		T* FirstPool;
+		UPTRINT			 BucketIndex;
+		T*				 FirstPool;
 		TPoolHashBucket* Prev;
 		TPoolHashBucket* Next;
 
@@ -328,7 +325,7 @@ protected:
 	FPtrToPoolMapping PtrToPoolMapping;
 	static CORE_API uint32 BinnedTlsSlot;
 
-#if UE_BINNEDCOMMON_ALLOCATOR_STATS
+#if UE_MBC_ALLOCATOR_STATS
 	static std::atomic<int64> TLSMemory;
 	static std::atomic<int64> ConsolidatedMemory;
 #endif
@@ -351,11 +348,11 @@ protected:
 	struct FFreeBlockList
 	{
 		// return true if we actually pushed it
-		FORCEINLINE bool PushToFront(void* InPtr, uint32 InPoolIndex, uint32 InBlockSize)
+		FORCEINLINE bool PushToFront(void* InPtr, uint32 InPoolIndex, uint32 InBinSize)
 		{
 			checkSlow(InPtr);
 
-			if ((PartialBundle.Count >= (uint32)GMallocBinnedBundleCount) | (PartialBundle.Count * InBlockSize >= (uint32)GMallocBinnedBundleSize))
+			if ((PartialBundle.Count >= (uint32)GMallocBinnedBundleCount) | (PartialBundle.Count * InBinSize >= (uint32)GMallocBinnedBundleSize))
 			{
 				if (FullBundle.Head)
 				{
@@ -368,9 +365,9 @@ protected:
 			return true;
 		}
 
-		FORCEINLINE bool CanPushToFront(uint32 InPoolIndex, uint32 InBlockSize) const
+		FORCEINLINE bool CanPushToFront(uint32 InPoolIndex, uint32 InBinSize) const
 		{
-			return !((!!FullBundle.Head) & ((PartialBundle.Count >= (uint32)GMallocBinnedBundleCount) | (PartialBundle.Count * InBlockSize >= (uint32)GMallocBinnedBundleSize)));
+			return !((!!FullBundle.Head) & ((PartialBundle.Count >= (uint32)GMallocBinnedBundleCount) | (PartialBundle.Count * InBinSize >= (uint32)GMallocBinnedBundleSize)));
 		}
 
 		FORCEINLINE void* PopFromFront(uint32 InPoolIndex)
@@ -479,7 +476,7 @@ protected:
 			{
 				const int64 TLSSize = Align(sizeof(FPerThreadFreeBlockLists), AllocType::OsAllocationGranularity);
 				ThreadSingleton = new (AllocType::AllocateMetaDataMemory(TLSSize)) FPerThreadFreeBlockLists();
-#if UE_BINNEDCOMMON_ALLOCATOR_STATS
+#if UE_MBC_ALLOCATOR_STATS
 				TLSMemory.fetch_add(TLSSize, std::memory_order_relaxed);
 #endif
 				verify(ThreadSingleton);
@@ -517,7 +514,7 @@ protected:
 			if (ThreadSingleton)
 			{
 				const int64 TLSSize = Align(sizeof(FPerThreadFreeBlockLists), AllocType::OsAllocationGranularity);
-#if UE_BINNEDCOMMON_ALLOCATOR_STATS
+#if UE_MBC_ALLOCATOR_STATS
 				TLSMemory.fetch_sub(TLSSize, std::memory_order_relaxed);
 #endif
 				AllocType::UnregisterThreadFreeBlockLists(ThreadSingleton);
@@ -536,15 +533,15 @@ protected:
 		}
 
 		// return true if the pointer was pushed
-		FORCEINLINE bool Free(void* InPtr, uint32 InPoolIndex, uint32 InBlockSize)
+		FORCEINLINE bool Free(void* InPtr, uint32 InPoolIndex, uint32 InBinSize)
 		{
-			return FreeLists[InPoolIndex].PushToFront(InPtr, InPoolIndex, InBlockSize);
+			return FreeLists[InPoolIndex].PushToFront(InPtr, InPoolIndex, InBinSize);
 		}
 
 		// return true if a pointer can be pushed
-		FORCEINLINE bool CanFree(uint32 InPoolIndex, uint32 InBlockSize) const
+		FORCEINLINE bool CanFree(uint32 InPoolIndex, uint32 InBinSize) const
 		{
-			return FreeLists[InPoolIndex].CanPushToFront(InPoolIndex, InBlockSize);
+			return FreeLists[InPoolIndex].CanPushToFront(InPoolIndex, InBinSize);
 		}
 
 		// returns a bundle that needs to be freed if it can't be recycled
@@ -593,7 +590,7 @@ protected:
 			return true;
 		}
 
-#if UE_BINNEDCOMMON_ALLOCATOR_STATS
+#if UE_MBC_ALLOCATOR_STATS
 	public:
 		int64 AllocatedMemory = 0;
 #endif
@@ -607,11 +604,11 @@ protected:
 	FORCEINLINE SIZE_T QuantizeSizeCommon(SIZE_T Count, uint32 Alignment, const AllocType& Alloc) const
 	{
 		static_assert(DEFAULT_ALIGNMENT <= MinAlign, "DEFAULT_ALIGNMENT is assumed to be zero"); // used below
-		checkSlow((Alignment & (Alignment - 1)) == 0); // Check the alignment is a power of two
+		checkSlow(FMath::IsPowerOfTwo(Alignment));
 		SIZE_T SizeOut;
 		if ((Count <= MaxSmallPoolSize) & (Alignment <= MinAlign)) // one branch, not two
 		{
-			SizeOut = Alloc.PoolIndexToBlockSize(BoundSizeToPoolIndex(Count, Alloc.MemSizeToIndex));
+			SizeOut = Alloc.PoolIndexToBinSize(BoundSizeToPoolIndex(Count, Alloc.MemSizeToPoolIndex));
 			check(SizeOut >= Count);
 			return SizeOut;
 		}
@@ -619,13 +616,13 @@ protected:
 		Count = Align(Count, Alignment);
 		if ((Count <= MaxSmallPoolSize) & (Alignment <= MaxAlign))
 		{
-			uint32 PoolIndex = BoundSizeToPoolIndex(Count, Alloc.MemSizeToIndex);
+			uint32 PoolIndex = BoundSizeToPoolIndex(Count, Alloc.MemSizeToPoolIndex);
 			do
 			{
-				uint32 BlockSize = Alloc.PoolIndexToBlockSize(PoolIndex);
-				if (IsAligned(BlockSize, Alignment))
+				const uint32 BinSize = Alloc.PoolIndexToBinSize(PoolIndex);
+				if (IsAligned(BinSize, Alignment))
 				{
-					SizeOut = SIZE_T(BlockSize);
+					SizeOut = SIZE_T(BinSize);
 					check(SizeOut >= Count);
 					return SizeOut;
 				}
@@ -640,11 +637,11 @@ protected:
 		return SizeOut;
 	}
 
-	FORCEINLINE uint32 BoundSizeToPoolIndex(SIZE_T Size, const uint8(&MemSizeToIndex)[1 + (MaxSmallPoolSize >> MinAlignShift)]) const
+	FORCEINLINE uint32 BoundSizeToPoolIndex(SIZE_T Size, const uint8(&MemSizeToPoolIndex)[1 + (MaxSmallPoolSize >> MinAlignShift)]) const
 	{
-		auto Index = ((Size + MinAlign - 1) >> MinAlignShift);
+		const auto Index = ((Size + MinAlign - 1) >> MinAlignShift);
 		checkSlow(Index >= 0 && Index <= (MaxSmallPoolSize >> MinAlignShift)); // and it should be in the table
-		uint32 PoolIndex = uint32(MemSizeToIndex[Index]);
+		const uint32 PoolIndex = uint32(MemSizeToPoolIndex[Index]);
 		checkSlow(PoolIndex >= 0 && PoolIndex < NumSmallPools);
 		return PoolIndex;
 	}
@@ -658,10 +655,10 @@ protected:
 		const SIZE_T AlignedSize = Align(Size, Alignment);
 		if (UNLIKELY((AlignedSize <= MaxSmallPoolSize) && (Alignment <= MaxAlign)))
 		{
-			uint32 PoolIndex = BoundSizeToPoolIndex(AlignedSize, Alloc.MemSizeToIndex);
+			uint32 PoolIndex = BoundSizeToPoolIndex(AlignedSize, Alloc.MemSizeToPoolIndex);
 			do
 			{
-				uint32 BlockSize = Alloc.PoolIndexToBlockSize(PoolIndex);
+				const uint32 BlockSize = Alloc.PoolIndexToBinSize(PoolIndex);
 				if (IsAligned(BlockSize, Alignment))
 				{
 					// we found a matching pool for our alignment and size requirements, so modify the size request to match
diff --git a/Engine/Source/Runtime/Core/Public/HAL/MallocBinnedCommonUtils.h b/Engine/Source/Runtime/Core/Public/HAL/MallocBinnedCommonUtils.h
--- a/Engine/Source/Runtime/Core/Public/HAL/MallocBinnedCommonUtils.h
+++ b/Engine/Source/Runtime/Core/Public/HAL/MallocBinnedCommonUtils.h
@@ -39,7 +39,7 @@ public:
 				TRACE_CPUPROFILER_EVENT_SCOPE(FMallocBinnedCommonUtils::FlushCurrentThreadCache);
 				QUICK_SCOPE_CYCLE_COUNTER(STAT_FMallocBinnedCommonUtils_FlushCurrentThreadCache);
 
-				double StartTimeInner = FPlatformTime::Seconds();
+				const double StartTimeInner = FPlatformTime::Seconds();
 
 				double WaitForMutexTime = 0.0f;
 				double WaitForMutexAndTrimTime = 0.0f;
