Fixing non-unity build on WindowsPlatform.h

diff --git a/Engine/Source/Runtime/Core/Public/Windows/WindowsPlatform.h b/Engine/Source/Runtime/Core/Public/Windows/WindowsPlatform.h
--- a/Engine/Source/Runtime/Core/Public/Windows/WindowsPlatform.h
+++ b/Engine/Source/Runtime/Core/Public/Windows/WindowsPlatform.h
@@ -18,6 +18,16 @@
 	#error "Windows Vista and earlier are no longer supported"
 #endif
 
+#ifndef PLATFORM_COMPILER_CLANG
+#if defined(__clang__)
+#define PLATFORM_COMPILER_CLANG			1
+#else
+#define PLATFORM_COMPILER_CLANG			0
+#endif // defined(__clang__)
+#endif
+
+#include "GenericPlatform/GenericPlatform.h"	//for FGenericPlatformTypes
+
 /**
 * Windows specific types
 **/
@@ -56,11 +66,12 @@ typedef FWindowsPlatformTypes FPlatformTypes;
 #define PLATFORM_SUPPORTS_UNALIGNED_LOADS					1
 
 #define PLATFORM_SUPPORTS_PRAGMA_PACK						1
-#if defined(_M_ARM) || defined(_M_ARM64) || defined(_M_ARM64EC)
+#if (defined(__arm__) || defined(_M_ARM) || defined(__aarch64__) || defined(_M_ARM64) || defined(_M_ARM64EC))
 	#define PLATFORM_CPU_ARM_FAMILY							1
 	#define PLATFORM_ENABLE_VECTORINTRINSICS_NEON			1
 	#define PLATFORM_ENABLE_VECTORINTRINSICS				1
 #elif (defined(_M_IX86) || defined(_M_X64))
+	#define PLATFORM_CPU_ARM_FAMILY							0
 	#define PLATFORM_CPU_X86_FAMILY							1
 	#define PLATFORM_ENABLE_VECTORINTRINSICS				1
 
