[AutoRTFM] Work around premature inlining of AutoRTFM functions

Unity builds can prematurely inline the core AutoRTFM functions, before the AutoRTFM compiler pass can switch the call sites for the closed variant. Remove the autortfm_is_closed() implementation and mark the other functions with clang::optnone to prevent this inlining.

diff --git a/Engine/Source/Runtime/Core/Private/AutoRTFM/API.cpp b/Engine/Source/Runtime/Core/Private/AutoRTFM/API.cpp
--- a/Engine/Source/Runtime/Core/Private/AutoRTFM/API.cpp
+++ b/Engine/Source/Runtime/Core/Private/AutoRTFM/API.cpp
@@ -220,11 +220,6 @@ extern "C" UE_AUTORTFM_AUTORTFM("RTFM_autortfm_is_transactional") bool autortfm_
 	return false;
 }
 
-extern "C" UE_AUTORTFM_AUTORTFM("RTFM_autortfm_is_closed") bool autortfm_is_closed()
-{
-    return false;
-}
-
 extern "C" bool autortfm_is_committing_or_aborting()
 {
 	if (ForTheRuntime::IsAutoRTFMRuntimeEnabled())
@@ -484,11 +479,6 @@ extern "C" UE_AUTORTFM_NOAUTORTFM bool RTFM_autortfm_is_transactional()
     return true;
 }
 
-extern "C" UE_AUTORTFM_NOAUTORTFM bool RTFM_autortfm_is_closed()
-{
-    return true;
-}
-
 extern "C" UE_AUTORTFM_NOAUTORTFM autortfm_result RTFM_autortfm_transact(void (*UninstrumentedWork)(void*), void (*InstrumentedWork)(void*), void* Arg)
 {
 	FContext* Context = FContext::Get();
diff --git a/Engine/Source/Runtime/Core/Public/AutoRTFM/AutoRTFM.h b/Engine/Source/Runtime/Core/Public/AutoRTFM/AutoRTFM.h
--- a/Engine/Source/Runtime/Core/Public/AutoRTFM/AutoRTFM.h
+++ b/Engine/Source/Runtime/Core/Public/AutoRTFM/AutoRTFM.h
@@ -29,7 +29,8 @@
 #endif
 
 #if UE_AUTORTFM
-#define UE_AUTORTFM_AUTORTFM(F) [[clang::autortfm(F), clang::noinline]]
+// #jira SOL-6799: clang::optnone is used to prevent inlining, despite clang::noinline
+#define UE_AUTORTFM_AUTORTFM(F) [[clang::autortfm(F), clang::noinline, clang::optnone]]
 #define UE_AUTORTFM_NOAUTORTFM [[clang::noautortfm, clang::noinline]]
 // #jira SOL-6589: remove clang::noinline once this JIRA is fixed.
 #define UE_AUTORTFM_ALWAYS_OPEN [[clang::autortfm_always_open, clang::noinline]]
