Increase the string builder size for ToString() to account for possible bogus TC values.  Also add a checkSlow to try to catch cases where we are building a non-standard Timecode value.

diff --git a/Engine/Source/Runtime/Core/Public/Misc/Timecode.h b/Engine/Source/Runtime/Core/Public/Misc/Timecode.h
--- a/Engine/Source/Runtime/Core/Public/Misc/Timecode.h
+++ b/Engine/Source/Runtime/Core/Public/Misc/Timecode.h
@@ -41,7 +41,9 @@ struct FTimecode
 		, Seconds(InSeconds)
 		, Frames(InFrames)
 		, bDropFrameFormat(InbDropFrame)
-	{}
+	{
+		checkSlow(IsValid());
+	}
 
 	/**
 	 * User construction from a number of hours minutes seconds frames, and subframes.
@@ -56,7 +58,9 @@ struct FTimecode
 		, Frames(InFrames)
 		, Subframe(InSubframe)
 		, bDropFrameFormat(InbDropFrame)
-	{}
+	{
+		checkSlow(IsValid());
+	}
 
 	/**
 	 * User construction from a time in seconds
@@ -352,8 +356,9 @@ public:
 			SignText = PositiveSign;
 		}
 
-		TStringBuilder<32> Builder;
-
+		// Use a buffer that will hold 64 chars to account for maximum int sizes for hours, min, seconds, frames + 3 charcs for subframe +/- 3 chars for sign.
+		//
+		TStringBuilder<64> Builder;
 		if (bDropFrameFormat)
 		{
 			Builder.Appendf(TEXT("%s%02d:%02d:%02d;%02d"), SignText, FMath::Abs(Hours), FMath::Abs(Minutes), FMath::Abs(Seconds), FMath::Abs(Frames));
@@ -371,6 +376,24 @@ public:
 		return Builder.ToString();
 	}
 
+	/*
+	 * Will return true if the timecode represents a valid timecode value where
+	 *
+	 * Hours is +/- [0,23]
+	 * Minutes is +/- [0, 59]
+	 * Seconds is +/- [0, 59]
+	 * Frames [INT_MIN, INT_MAX]
+	 * Subframes > 0
+	 *
+	 */
+	bool IsValid() const
+	{
+		auto InAbsRange = [](int32 Val, int32 MinVal, int32 MaxVal)
+		{
+		    return FMath::Abs(Val) >= MinVal && FMath::Abs(Val) <= MaxVal;
+		};
+		return InAbsRange(Hours,0,23) && InAbsRange(Minutes, 0, 59) && InAbsRange(Seconds, 0, 59) && Subframe >= 0;
+	}
 public:
 
 	/** How many hours does this timecode represent */
