[AutoRTFM] Reintroduce a definition for autortfm_is_closed()

Intead use a volatile local to prevent the premature inlining of this function.
For some reason the linux builds are failing to link. This is weird as the compiler should be replacing all callsites with a literal true or false.

diff --git a/Engine/Source/Runtime/Core/Private/AutoRTFM/API.cpp b/Engine/Source/Runtime/Core/Private/AutoRTFM/API.cpp
--- a/Engine/Source/Runtime/Core/Private/AutoRTFM/API.cpp
+++ b/Engine/Source/Runtime/Core/Private/AutoRTFM/API.cpp
@@ -220,6 +220,12 @@ extern "C" UE_AUTORTFM_AUTORTFM("RTFM_autortfm_is_transactional") bool autortfm_
 	return false;
 }
 
+extern "C" UE_AUTORTFM_AUTORTFM("RTFM_autortfm_is_closed") bool autortfm_is_closed()
+{
+	volatile bool ReallyDoNotInlineMe = true; // #jira SOL-6799
+	return !ReallyDoNotInlineMe;
+}
+
 extern "C" bool autortfm_is_committing_or_aborting()
 {
 	if (ForTheRuntime::IsAutoRTFMRuntimeEnabled())
@@ -479,6 +485,11 @@ extern "C" UE_AUTORTFM_NOAUTORTFM bool RTFM_autortfm_is_transactional()
     return true;
 }
 
+extern "C" UE_AUTORTFM_NOAUTORTFM bool RTFM_autortfm_is_closed()
+{
+	return true;
+}
+
 extern "C" UE_AUTORTFM_NOAUTORTFM autortfm_result RTFM_autortfm_transact(void (*UninstrumentedWork)(void*), void (*InstrumentedWork)(void*), void* Arg)
 {
 	FContext* Context = FContext::Get();
