Added overload to HasSameIndexAndSerialNumber in TWeakObjectPtr.

The new overload allows comparing TWeakObjectPtrs that are not the same type. An example of this would be a TWeakObjectPtr<UStruct> with a TWeakObjectPtr<const UClass>.

[RN] HasSameIndexAndSerialNumber in TWeakObjectPtr can now be used for related types.

diff --git a/Engine/Source/Runtime/Core/Public/UObject/WeakObjectPtrTemplates.h b/Engine/Source/Runtime/Core/Public/UObject/WeakObjectPtrTemplates.h
--- a/Engine/Source/Runtime/Core/Public/UObject/WeakObjectPtrTemplates.h
+++ b/Engine/Source/Runtime/Core/Public/UObject/WeakObjectPtrTemplates.h
@@ -231,6 +231,19 @@ public:
 		return static_cast<const TWeakObjectPtrBase&>(*this).HasSameIndexAndSerialNumber(static_cast<const TWeakObjectPtrBase&>(Other));
 	}
 
+	/**
+	 * Returns true if two weak pointers were originally set to the same object, even if they are now stale
+	 * @param Other weak pointer to compare to
+	 */
+	template <
+		typename OtherT
+		UE_REQUIRES(UE_REQUIRES_EXPR((T*)nullptr == (OtherT*)nullptr))
+	>
+	FORCEINLINE bool HasSameIndexAndSerialNumber(const TWeakObjectPtr<OtherT, TWeakObjectPtrBase>& Other) const
+	{
+		return ((const TWeakObjectPtrBase&)*this).HasSameIndexAndSerialNumber((const TWeakObjectPtrBase&)Other);
+	}
+
 	/**
 	 * Weak object pointer serialization, this forwards to FArchive::operator<<(struct FWeakObjectPtr&) or an override
 	 */
