Use JOBOBJECT_LIMIT_VIOLATION_INFORMATION to more accurately establish container memory usage

Laura Hermanns recently found a better way to get the memory used by a Windows job than what I had used before. Instead of having to estimate the memory usage by summing up private bytes, we can query JOBOBJECT_LIMIT_VIOLATION_INFORMATION and read the memory usage directly. This is orders of magnitude faster and yields more accurate results. I had missed this when first implementing this and was under the impression that JOBOBJECT_LIMIT_VIOLATION_INFORMATION was delivered via the IoCompletionPort, but it can actually just be used with QueryInformationJobObject.

#rnx

diff --git a/Engine/Source/Runtime/Core/Private/Windows/WindowsPlatformMemory.cpp b/Engine/Source/Runtime/Core/Private/Windows/WindowsPlatformMemory.cpp
--- a/Engine/Source/Runtime/Core/Private/Windows/WindowsPlatformMemory.cpp
+++ b/Engine/Source/Runtime/Core/Private/Windows/WindowsPlatformMemory.cpp
@@ -288,66 +288,19 @@ FMalloc* FWindowsPlatformMemory::BaseAllocator()
 	return Instance;
 }
 
-static uint64 EstimateContainerCommittedMemory()
-{
-	// These is unfortunately no way to just get the resources used by a job group. This is an unfortunate hole
-	// in the API provided by Windows. We instead sum up the commit usage of all processes in the job group.
-	// This is slow and overestimates the actual commit usage, but that is close enough and we'd rather be slow
-	// than terminate the process because we are running out of memory.
-	TRACE_CPUPROFILER_EVENT_SCOPE(EstimateContainerCommittedMemory);
-	constexpr int32 MaxProcessesAccountedFor = 256;
-	struct ProcessList {
-		DWORD NumberOfAssignedProcesses;
-		DWORD NumberOfProcessIdsInList;
-		ULONG_PTR ProcessIdList[MaxProcessesAccountedFor];
-	};
-	ProcessList Processes{};
+static uint64 QueryContainerCommittedMemory()
+{	
+	JOBOBJECT_LIMIT_VIOLATION_INFORMATION LimitInformation{};
 	DWORD LengthWritten;
-	if (!QueryInformationJobObject(NULL, JobObjectBasicProcessIdList, &Processes, sizeof(Processes), &LengthWritten))
+	if (!QueryInformationJobObject(NULL, JobObjectLimitViolationInformation, &LimitInformation, sizeof(LimitInformation), &LengthWritten))
 	{
 		DWORD ErrNo = GetLastError();
-		UE_LOG(LogMemory, Error, TEXT("QueryInformationJobObject(NULL, JobObjectBasicProcessIdList, &Processes, sizeof(Processes), &LengthWritten) failed with GetLastError() = %d"),
+		UE_LOG(LogMemory, Error, TEXT("QueryInformationJobObject(NULL, JobObjectLimitViolationInformation, &LimitInformation, sizeof(LimitInformation), &LengthWritten) failed with GetLastError() = %d"),
 			ErrNo
 		);
 		return 0;
 	}
-	if (Processes.NumberOfProcessIdsInList < Processes.NumberOfAssignedProcesses)
-	{
-		UE_LOG(LogMemory, Warning, TEXT("The number of processes in this container %d is larger than %d, memory accounting may be more inaccurate than usual."),
-			Processes.NumberOfAssignedProcesses,
-			Processes.NumberOfProcessIdsInList
-		);
-	}
-	ensureMsgf(Processes.NumberOfProcessIdsInList <= MaxProcessesAccountedFor,
-		TEXT("Number of processes returned by QueryInformationJobObject is too large: %d (must be below %d)"),
-		Processes.NumberOfProcessIdsInList,
-		MaxProcessesAccountedFor
-	);
-	ensureMsgf(Processes.NumberOfProcessIdsInList > 0, TEXT("Number of processes returned by QueryInformationJobObject is 0"));
-	PROCESS_MEMORY_COUNTERS ProcessMemoryCounters;
-	FPlatformMemory::Memzero(&ProcessMemoryCounters, sizeof(ProcessMemoryCounters));
-	uint64 TotalCommittedMemory = 0;
-	for (int32 ProcessIndex = 0; ProcessIndex < (int32)Processes.NumberOfProcessIdsInList; ProcessIndex++)
-	{
-		HANDLE ProcHandle = OpenProcess(PROCESS_QUERY_LIMITED_INFORMATION, FALSE, (DWORD)Processes.ProcessIdList[ProcessIndex]);
-		if (ProcHandle != INVALID_HANDLE_VALUE)
-		{
-			if (GetProcessMemoryInfo(ProcHandle, &ProcessMemoryCounters, sizeof(ProcessMemoryCounters)))
-			{
-				TotalCommittedMemory += ProcessMemoryCounters.PagefileUsage;
-			}
-		}
-		else
-		{
-			DWORD ErrNo = GetLastError();
-			UE_LOG(LogMemory, Warning, TEXT("OpenProcess(PROCESS_QUERY_LIMITED_INFORMATION, FALSE, %d) failed with GetLastError() = %d"),
-				(DWORD)Processes.ProcessIdList[ProcessIndex],
-				ErrNo
-			);
-		}
-		CloseHandle(ProcHandle);
-	}
-	return TotalCommittedMemory;
+	return LimitInformation.JobMemory;
 }
 
 FPlatformMemoryStats FWindowsPlatformMemory::GetStats()
@@ -425,7 +378,7 @@ FPlatformMemoryStats FWindowsPlatformMemory::GetStats()
 
 	if (GWindowsUseContainerMemory && IsRunningAsJob())
 	{
-		const uint64 TotalUsage = EstimateContainerCommittedMemory();
+		const uint64 TotalUsage = QueryContainerCommittedMemory();
 		ensureMsgf(TotalUsage != 0, TEXT("Estimated container memory usage is 0. This should never happen"));
 		if (TotalUsage != 0)
 		{
