Updated GetPakchunkIndexFromPakFile to parse PAK chunknames ending with the chunk index

#rnx

diff --git a/Engine/Source/Runtime/Core/Private/GenericPlatform/GenericPlatformMisc.cpp b/Engine/Source/Runtime/Core/Private/GenericPlatform/GenericPlatformMisc.cpp
--- a/Engine/Source/Runtime/Core/Private/GenericPlatform/GenericPlatformMisc.cpp
+++ b/Engine/Source/Runtime/Core/Private/GenericPlatform/GenericPlatformMisc.cpp
@@ -1990,7 +1990,7 @@ int32 FGenericPlatformMisc::GetPakchunkIndexFromPakFile(FStringView InFilename)
 			DigitCount++;
 		}
 
-		if (DigitCount > 0 && (StartOfNumber + DigitCount) < BaseFilename.Len())
+		if (DigitCount > 0 && (StartOfNumber + DigitCount) <= BaseFilename.Len())
 		{
 			// FromString can't take a view
 			TStringBuilder<16> ChunkNumberString = WriteToString<16>(BaseFilename.Mid(StartOfNumber, DigitCount));
diff --git a/Engine/Source/Runtime/Core/Tests/GenericPlatform/GenericPlatformMiscTest.cpp b/Engine/Source/Runtime/Core/Tests/GenericPlatform/GenericPlatformMiscTest.cpp
new file mode 100644
--- /dev/null
+++ b/Engine/Source/Runtime/Core/Tests/GenericPlatform/GenericPlatformMiscTest.cpp
@@ -0,0 +1,19 @@
+// Copyright Epic Games, Inc. All Rights Reserved.
+
+#if WITH_TESTS 
+
+#include "GenericPlatform/GenericPlatformMisc.h"
+#include "Tests/TestHarnessAdapter.h"
+
+TEST_CASE_NAMED(FGenericPlatformMiscTest, "System::Core::GenericPlatform", "[Core][GenericPlatform][SmokeFilter]")
+{
+	SECTION("Parse pakchunk index")
+	{
+		CHECK(1 == FGenericPlatformMisc::GetPakchunkIndexFromPakFile(FString(TEXT("pakchunk1-Windows"))));
+		CHECK(12 == FGenericPlatformMisc::GetPakchunkIndexFromPakFile(FString(TEXT("../../../Content/Paks/pakchunk12-Windows"))));
+		CHECK(42 == FGenericPlatformMisc::GetPakchunkIndexFromPakFile(FString(TEXT("pakchunk42"))));
+		CHECK(INDEX_NONE == FGenericPlatformMisc::GetPakchunkIndexFromPakFile(FString(TEXT("pakchunk"))));
+	}
+}
+
+#endif // WITH_TESTS
