Make AutoRTFM 9.5% faster by removing the TLS context that serves no purpose now we're running exclusively single threaded.

diff --git a/Engine/Source/Runtime/Core/Private/AutoRTFM/Context.cpp b/Engine/Source/Runtime/Core/Private/AutoRTFM/Context.cpp
--- a/Engine/Source/Runtime/Core/Private/AutoRTFM/Context.cpp
+++ b/Engine/Source/Runtime/Core/Private/AutoRTFM/Context.cpp
@@ -13,6 +13,11 @@
 
 #include "Templates/UniquePtr.h"
 #include "Containers/StringConv.h"
+#include "HAL/PlatformTLS.h"
+
+#if PLATFORM_WINDOWS
+extern "C" __declspec(dllimport) void __stdcall GetCurrentThreadStackLimits(void**, void**);
+#endif
 
 namespace
 {
@@ -33,85 +38,20 @@ FAutoRTFMMetrics GetAutoRTFMMetrics()
 	return GAutoRTFMMetrics;
 }
 
-// A memory blob that can placement new store an FContext class.
-struct alignas(alignof(FContext)) FContextData final
-{
-	uint8 Payload[sizeof(FContext)] = {};
-};
-
-thread_local FContextData ContextDataTls;
-
-// This holder's entire job is to call the destructor on the memory we placement newed
-// into `ContextDataTls` above, and then zero out the data. This is to prevent the
-// problem where global destruction order is undefined, and we could call a destructor
-// that wants to call into the runtime after we have destroyed it. By zeroing out the
-// memory the runtime will *always* return safe defaults if it is queried after these
-// values are destroyed.
-struct FContextHolder final
-{
-	~FContextHolder()
-	{
-		if (Context)
-		{
-			Context->~FContext();
-			memset(Context, 0, sizeof(FContext));
-			Context = nullptr;
-		}
-	}
-
-	FContext* Context = nullptr;
-};
-
-thread_local FContextHolder ContextTls;
-
 void FContext::InitializeGlobalData()
 {
 }
 
-FContext* FContext::TryGet()
-{
-    return ContextTls.Context;
-}
-
-void FContext::Set()
-{
-    ContextTls.Context = this;
-}
-
-FContext* FContext::Get()
-{
-    FContext* Result = TryGet();
-
-    if (!Result)
-    {
-        Result = new (&ContextDataTls) FContext();
-        Result->Set();
-    }
-
-    return Result;
-}
+FContext FContext::ContextSingleton;
 
 bool FContext::IsTransactional()
 {
-    FContext* Context = TryGet();
-    if (!Context)
-    {
-        return false;
-    }
-
-    return Context->GetStatus() == EContextStatus::OnTrack;
+    return Get()->GetStatus() == EContextStatus::OnTrack;
 }
 
 bool FContext::IsCommittingOrAborting()
 {
-	FContext* Context = TryGet();
-
-	if (!Context)
-	{
-		return false;
-	}
-
-	switch (Context->GetStatus())
+	switch (Get()->GetStatus())
 	{
 	default:
 		return true;
@@ -356,8 +296,6 @@ ETransactionResult FContext::Transact(void (*InstrumentedFunction)(void*), void*
 	NewTransaction->SetIsScopedTransaction();
 
 	void* TransactStackAddress = &NewTransaction;
-	ASSERT(TransactStackAddress > StackBegin);
-	ASSERT(TransactStackAddress < StackEnd);
 	TScopedGuard<void*> CurrentNestStackAddressGuard(CurrentTransactStackAddress, TransactStackAddress);
 
 	ETransactionResult Result = ETransactionResult::Committed; // Initialize to something to make the compiler happy.
@@ -366,6 +304,30 @@ ETransactionResult FContext::Transact(void (*InstrumentedFunction)(void*), void*
     {
         ASSERT(Status == EContextStatus::Idle);
 
+		ASSERT(FPlatformTLS::InvalidTlsSlot == CurrentThreadId);
+		CurrentThreadId = FPlatformTLS::GetCurrentThreadId();
+
+		ASSERT(nullptr == StackBegin);
+		ASSERT(nullptr == StackEnd);
+
+#if PLATFORM_WINDOWS
+		GetCurrentThreadStackLimits(&StackBegin, &StackEnd);
+#elif defined(__APPLE__)         
+		StackEnd = pthread_get_stackaddr_np(pthread_self());
+		size_t StackSize = pthread_get_stacksize_np(pthread_self());
+		StackBegin = static_cast<char*>(StackEnd) - StackSize;
+#else
+		pthread_attr_t Attr;
+		pthread_getattr_np(pthread_self(), &Attr);
+		size_t StackSize;
+		pthread_attr_getstack(&Attr, &StackBegin, &StackSize);
+		StackEnd = static_cast<char*>(StackBegin) + StackSize;
+#endif
+		ASSERT(StackEnd > StackBegin);
+
+		ASSERT(TransactStackAddress > StackBegin);
+		ASSERT(TransactStackAddress < StackEnd);
+
 		PushTransaction(NewTransaction);
 		PushCallNest(NewNest);
         OuterTransactStackAddress = TransactStackAddress;
@@ -447,6 +409,13 @@ ETransactionResult FContext::Transact(void (*InstrumentedFunction)(void*), void*
 		// This transaction is within another transaction
 		ASSERT(Status == EContextStatus::OnTrack);
 
+		ASSERT(CurrentThreadId == FPlatformTLS::GetCurrentThreadId());
+
+		ASSERT(nullptr != StackBegin);
+		ASSERT(nullptr != StackEnd);
+		ASSERT(TransactStackAddress > StackBegin);
+		ASSERT(TransactStackAddress < StackEnd);
+
 		PushTransaction(NewTransaction);
 		PushCallNest(NewNest);
 
@@ -523,32 +492,13 @@ void FContext::AbortByLanguageAndThrow()
     CurrentTransaction->AbortAndThrow();
 }
 
-#if PLATFORM_WINDOWS
-extern "C" __declspec(dllimport) void __stdcall GetCurrentThreadStackLimits(void**, void**);
-#endif
-
-FContext::FContext()
-{
-#if PLATFORM_WINDOWS
-    GetCurrentThreadStackLimits(&StackBegin, &StackEnd);
-#elif defined(__APPLE__)         
-   StackEnd = pthread_get_stackaddr_np(pthread_self());   
-   size_t StackSize = pthread_get_stacksize_np(pthread_self());    
-   StackBegin = static_cast<char*>(StackEnd) - StackSize;
-#else
-    pthread_attr_t Attr;
-    pthread_getattr_np(pthread_self(), &Attr);
-    size_t StackSize;
-    pthread_attr_getstack(&Attr, &StackBegin, &StackSize);
-    StackEnd = static_cast<char*>(StackBegin) + StackSize;
-#endif
-    ASSERT(StackEnd > StackBegin);
-}
-
 void FContext::Reset()
 {
+	CurrentThreadId = FPlatformTLS::InvalidTlsSlot;
     OuterTransactStackAddress = nullptr;
 	CurrentTransactStackAddress = nullptr;
+	StackBegin = nullptr;
+	StackEnd = nullptr;
     CurrentTransaction = nullptr;
 	CurrentNest = nullptr;
     Status = EContextStatus::Idle;
diff --git a/Engine/Source/Runtime/Core/Private/AutoRTFM/Context.h b/Engine/Source/Runtime/Core/Private/AutoRTFM/Context.h
--- a/Engine/Source/Runtime/Core/Private/AutoRTFM/Context.h
+++ b/Engine/Source/Runtime/Core/Private/AutoRTFM/Context.h
@@ -15,7 +15,6 @@ class FCallNest;
 class FContext
 {
 public:
-    static FContext* TryGet();
     static FContext* Get();
     static bool IsTransactional();
 	static bool IsCommittingOrAborting();
@@ -60,7 +59,9 @@ public:
     static void InitializeGlobalData();
 
 private:
-    FContext();
+	static FContext ContextSingleton;
+
+	FContext() { Reset(); }
     FContext(const FContext&) = delete;
 
 	void PushCallNest(FCallNest* NewCallNest);
@@ -71,8 +72,6 @@ private:
 
 	ETransactionResult ResolveNestedTransaction(FTransaction* NewTransaction);
 	bool AttemptToCommitTransaction(FTransaction* const Transaction);
-
-    void Set();
     
     // All of this other stuff ought to be private?
     void Reset();
@@ -85,6 +84,7 @@ private:
     void* OuterTransactStackAddress{nullptr};
     void* CurrentTransactStackAddress{nullptr};
     EContextStatus Status{EContextStatus::Idle};
+	uint32 CurrentThreadId{ ~0u };
 };
 
 } // namespace AutoRTFM
diff --git a/Engine/Source/Runtime/Core/Private/AutoRTFM/ContextInlines.h b/Engine/Source/Runtime/Core/Private/AutoRTFM/ContextInlines.h
--- a/Engine/Source/Runtime/Core/Private/AutoRTFM/ContextInlines.h
+++ b/Engine/Source/Runtime/Core/Private/AutoRTFM/ContextInlines.h
@@ -8,6 +8,11 @@
 namespace AutoRTFM
 {
 
+UE_AUTORTFM_FORCEINLINE FContext* FContext::Get()
+{
+	return &FContext::ContextSingleton;
+}
+
 UE_AUTORTFM_FORCEINLINE void FContext::CheckOpenRecordWrite(void* LogicalAddress)
 {
     // We don't record any writes to unscoped transactions when the write-address
