[octree] added LexToString for FOctreeElementId

diff --git a/Engine/Source/Runtime/Core/Public/Math/GenericOctreePublic.h b/Engine/Source/Runtime/Core/Public/Math/GenericOctreePublic.h
--- a/Engine/Source/Runtime/Core/Public/Math/GenericOctreePublic.h
+++ b/Engine/Source/Runtime/Core/Public/Math/GenericOctreePublic.h
@@ -51,6 +51,11 @@ private:
 	{
 		return ElementIndex;
 	}
+
+	friend FString LexToString(const FOctreeElementId2 ElementId)
+	{
+		return FString::Printf(TEXT("[%s:%s]"), *LexToString(ElementId.NodeIndex), *LexToString(ElementId.ElementIndex));
+	}
 };
 
 /**
