Logging: Added DispatchDynamicLogRecord for advanced use cases

diff --git a/Engine/Source/Runtime/Core/Private/Logging/StructuredLog.cpp b/Engine/Source/Runtime/Core/Private/Logging/StructuredLog.cpp
--- a/Engine/Source/Runtime/Core/Private/Logging/StructuredLog.cpp
+++ b/Engine/Source/Runtime/Core/Private/Logging/StructuredLog.cpp
@@ -829,18 +829,10 @@ inline static FLogRecord CreateLogRecord(const FLogCategoryBase& Category, const
 	return Record;
 }
 
-template <typename StaticLogRecordType>
-inline static void DispatchLogRecord(const StaticLogRecordType& Log, const FLogRecord& Record)
+inline static void DispatchLogRecord(const FLogRecord& Record)
 {
-#if LOGTRACE_ENABLED
-	if (UE_TRACE_CHANNELEXPR_IS_ENABLED(LogChannel))
-	{
-		LogToTrace(&Log, Record);
-	}
-#endif
-
 	FOutputDevice* OutputDevice = nullptr;
-	switch (Log.Verbosity)
+	switch (Record.GetVerbosity())
 	{
 	case ELogVerbosity::Error:
 	case ELogVerbosity::Warning:
@@ -854,10 +846,23 @@ inline static void DispatchLogRecord(const StaticLogRecordType& Log, const FLogR
 	(OutputDevice ? OutputDevice : GLog)->SerializeRecord(Record);
 }
 
+template <typename StaticLogRecordType>
+inline static void DispatchStaticLogRecord(const StaticLogRecordType& Log, const FLogRecord& Record)
+{
+#if LOGTRACE_ENABLED
+	if (UE_TRACE_CHANNELEXPR_IS_ENABLED(LogChannel))
+	{
+		LogToTrace(&Log, Record);
+	}
+#endif
+
+	DispatchLogRecord(Record);
+}
+
 void LogWithFieldArray(const FLogCategoryBase& Category, const FStaticLogRecord& Log, const FLogField* Fields, const int32 FieldCount)
 {
 #if !NO_LOGGING
-	DispatchLogRecord(Log, CreateLogRecord(Category, Log, Fields, FieldCount));
+	DispatchStaticLogRecord(Log, CreateLogRecord(Category, Log, Fields, FieldCount));
 #endif
 }
 
@@ -876,7 +881,7 @@ void LogWithFieldArray(const FLogCategoryBase& Category, const FStaticLocalizedL
 	FLogRecord Record = CreateLogRecord(Category, Log, Fields, FieldCount);
 	Record.SetTextNamespace(Log.TextNamespace);
 	Record.SetTextKey(Log.TextKey);
-	DispatchLogRecord(Log, Record);
+	DispatchStaticLogRecord(Log, Record);
 #endif
 }
 
@@ -1102,4 +1107,16 @@ void BasicFatalLog(const FLogCategoryBase& Category, const FStaticBasicLogRecord
 
 } // UE::Logging::Private
 
+///////////////////////////////////////////////////////////////////////////////////////////////////////////////////////
+
+namespace UE
+{
+
+void DispatchDynamicLogRecord(const FLogRecord& Record)
+{
+	Logging::Private::DispatchLogRecord(Record);
+}
+
+} // UE
+
 #endif // !NO_LOGGING
diff --git a/Engine/Source/Runtime/Core/Public/Logging/StructuredLog.h b/Engine/Source/Runtime/Core/Public/Logging/StructuredLog.h
--- a/Engine/Source/Runtime/Core/Public/Logging/StructuredLog.h
+++ b/Engine/Source/Runtime/Core/Public/Logging/StructuredLog.h
@@ -210,6 +210,16 @@ private:
 	const TCHAR* TextKey = nullptr;
 };
 
+/**
+ * Advanced function to dispatch a log record to active output devices.
+ *
+ * Always use UE_LOGFMT or its variants when possible.
+ * Dynamic dispatch bypasses many optimizations provided by the macros.
+ * Anything pointed to by the record MUST remain valid until threaded logs have been flushed.
+ * Filtering by category or verbosity is the responsibility of the caller.
+ */
+UE_API void DispatchDynamicLogRecord(const FLogRecord& Record);
+
 /**
  * Serializes the value to be used in a log message.
  *
