Fixed 3ds Max exporter compilation error

#rnx

diff --git a/Engine/Source/Runtime/Core/Public/Misc/TransactionallySafeCriticalSection.h b/Engine/Source/Runtime/Core/Public/Misc/TransactionallySafeCriticalSection.h
--- a/Engine/Source/Runtime/Core/Public/Misc/TransactionallySafeCriticalSection.h
+++ b/Engine/Source/Runtime/Core/Public/Misc/TransactionallySafeCriticalSection.h
@@ -35,9 +35,10 @@ struct FTransactionallySafeCriticalSectionDefinition final
 				{
 					// We explicitly copy the state here for the case that `this` was stack
 					// allocated and has already died before the on-abort is hit.
-					AutoRTFM::OnAbort([State = this->State]
+					TSharedPtr<FState> LocalState = this->State;
+					AutoRTFM::OnAbort([LocalState]
 						{
-							ensure(0 == State->TransactionalLockCount); 
+							ensure(0 == LocalState->TransactionalLockCount);
 						});
 				});
 
@@ -53,9 +54,10 @@ struct FTransactionallySafeCriticalSectionDefinition final
 				{
 					// We explicitly copy the state here for the case that `this` was stack
 					// allocated and has already died before the on-commit is hit.
-					AutoRTFM::OnCommit([State = this->State]
+					TSharedPtr<FState> LocalState = this->State;
+					AutoRTFM::OnCommit([LocalState]
 						{
-							ensure(0 == State->TransactionalLockCount);
+							ensure(0 == LocalState->TransactionalLockCount);
 						});
 				});
 
