Someone fixed a similar issue in CL 3778182 in ICUBreakIterator. So applying the same fix to ICULineBreakIterator is a good idea.

PR #12067: Ensure the FICUBreakIteratorManager is valid on destruction

#rnx
#tests start and end application, but breakpoints were not hit.

diff --git a/Engine/Source/Runtime/Core/Private/Internationalization/ICULineBreakIterator.cpp b/Engine/Source/Runtime/Core/Private/Internationalization/ICULineBreakIterator.cpp
--- a/Engine/Source/Runtime/Core/Private/Internationalization/ICULineBreakIterator.cpp
+++ b/Engine/Source/Runtime/Core/Private/Internationalization/ICULineBreakIterator.cpp
@@ -81,8 +81,11 @@ FICULineBreakIterator::FICULineBreakIterator()
 
 FICULineBreakIterator::~FICULineBreakIterator()
 {
-	// This assumes that FICULineBreakIterator owns the iterators, and that nothing ever copies an FICULineBreakIterator instance
-	FICUBreakIteratorManager::Get().DestroyIterator(ICULineBreakIteratorHandle);
+	if (FICUBreakIteratorManager::IsInitialized())
+	{
+		// This assumes that FICULineBreakIterator owns the iterators, and that nothing ever copies an FICULineBreakIterator instance
+		FICUBreakIteratorManager::Get().DestroyIterator(ICULineBreakIteratorHandle);
+	}
 }
 
 void FICULineBreakIterator::SetString(FString&& InString)
