Added a TString<CharType> alias to map a char type to its corresponding string container type.

diff --git a/Engine/Source/Runtime/Core/Public/Containers/ContainersFwd.h b/Engine/Source/Runtime/Core/Public/Containers/ContainersFwd.h
--- a/Engine/Source/Runtime/Core/Public/Containers/ContainersFwd.h
+++ b/Engine/Source/Runtime/Core/Public/Containers/ContainersFwd.h
@@ -20,6 +20,33 @@ template<> struct TIsContiguousContainer<FString>     { static constexpr bool Va
 template<> struct TIsContiguousContainer<FAnsiString> { static constexpr bool Value = true; };
 template<> struct TIsContiguousContainer<FUtf8String> { static constexpr bool Value = true; };
 
+namespace UE::Core::Private
+{
+	template <typename CharType>
+	struct TCharTypeToStringType;
+
+	template <>
+	struct TCharTypeToStringType<TCHAR>
+	{
+		using Type = FString;
+	};
+
+	template <>
+	struct TCharTypeToStringType<ANSICHAR>
+	{
+		using Type = FAnsiString;
+	};
+
+	template <>
+	struct TCharTypeToStringType<UTF8CHAR>
+	{
+		using Type = FUtf8String;
+	};
+}
+
+template <typename CharType>
+using TString = typename UE::Core::Private::TCharTypeToStringType<CharType>::Type;
+
 template<typename T, typename Allocator = FDefaultAllocator> class TArray;
 template<typename T> using TArray64 = TArray<T, FDefaultAllocator64>;
 template<typename T, typename SizeType = int32> class TArrayView;
diff --git a/Engine/Source/Runtime/Core/Public/Containers/UnrealString.h.inl b/Engine/Source/Runtime/Core/Public/Containers/UnrealString.h.inl
--- a/Engine/Source/Runtime/Core/Public/Containers/UnrealString.h.inl
+++ b/Engine/Source/Runtime/Core/Public/Containers/UnrealString.h.inl
@@ -1569,7 +1569,7 @@ public:
 	UE_STRING_CLASS& Appendf(const FmtType& Fmt, Types... Args)
 	{
 		static_assert(TIsArrayOrRefOfTypeByPredicate<FmtType, TIsCharEncodingCompatibleWithPrintfFmt>::Value, "Formatting string must be a literal " PREPROCESSOR_TO_STRING(UE_STRING_PRINTF_FMT_CHARTYPE) " array.");
-		static_assert((TIsValidVariadicFunctionArg<Types>::Value && ...), "Invalid argument(s) passed to TString::Appendf");
+		static_assert((TIsValidVariadicFunctionArg<Types>::Value && ...), "Invalid argument(s) passed to Appendf");
 
 		AppendfImpl(*this, (const ElementType*)Fmt, Args...);
 		return *this;
