[Windows] Make IsWine() publicly accessible to fix build issue.

#rnx

diff --git a/Engine/Source/Runtime/Core/Public/Windows/WindowsPlatformMisc.h b/Engine/Source/Runtime/Core/Public/Windows/WindowsPlatformMisc.h
--- a/Engine/Source/Runtime/Core/Public/Windows/WindowsPlatformMisc.h
+++ b/Engine/Source/Runtime/Core/Public/Windows/WindowsPlatformMisc.h
@@ -216,7 +216,7 @@ struct FWindowsPlatformMisc
 	 *
 	 * @return	Returns true if the current runtime environment is Wine
 	 */
-	static bool IsWine();
+	static CORE_API bool IsWine();
 
 #if !UE_BUILD_SHIPPING
 	static CORE_API void PromptForRemoteDebugging(bool bIsEnsure);
