For iOS, add a call to get the current state of background app refresh and add a call to open Settings on iOS

diff --git a/Engine/Source/Runtime/Core/Private/IOS/IOSPlatformMisc.cpp b/Engine/Source/Runtime/Core/Private/IOS/IOSPlatformMisc.cpp
--- a/Engine/Source/Runtime/Core/Private/IOS/IOSPlatformMisc.cpp
+++ b/Engine/Source/Runtime/Core/Private/IOS/IOSPlatformMisc.cpp
@@ -1615,6 +1615,29 @@ FString FIOSPlatformMisc::GetBuildNumber()
 	return BuildString;
 }
 
+bool FIOSPlatformMisc::IsBackgroundAppRefreshAvailable()
+{
+	return (UIBackgroundRefreshStatusAvailable == [[UIApplication sharedApplication] backgroundRefreshStatus]);
+}
+
+void FIOSPlatformMisc::OpenAppNotificationSettings()
+{
+	dispatch_async(dispatch_get_main_queue(), ^{
+		NSURL * SettingsUrl = [[NSURL alloc]initWithString:UIApplicationOpenNotificationSettingsURLString];
+		[[UIApplication sharedApplication]openURL:SettingsUrl options:@{} completionHandler:nil];
+		[SettingsUrl release];
+	});
+}
+
+void FIOSPlatformMisc::OpenAppCustomSettings()
+{
+	dispatch_async(dispatch_get_main_queue(), ^{
+		NSURL * SettingsUrl = [[NSURL alloc]initWithString:UIApplicationOpenSettingsURLString];
+		[[UIApplication sharedApplication]openURL:SettingsUrl options:@{} completionHandler:nil];
+		[SettingsUrl release];
+	});
+}
+
 bool FIOSPlatformMisc::RequestDeviceCheckToken(TFunction<void(const TArray<uint8>&)> QuerySucceededFunc, TFunction<void(const FString&, const FString&)> QueryFailedFunc)
 {
 	DCDevice* DeviceCheckDevice = [DCDevice currentDevice];
diff --git a/Engine/Source/Runtime/Core/Public/IOS/IOSPlatformMisc.h b/Engine/Source/Runtime/Core/Public/IOS/IOSPlatformMisc.h
--- a/Engine/Source/Runtime/Core/Public/IOS/IOSPlatformMisc.h
+++ b/Engine/Source/Runtime/Core/Public/IOS/IOSPlatformMisc.h
@@ -190,6 +190,21 @@ struct CORE_API FIOSPlatformMisc : public FApplePlatformMisc
 	static FString GetProjectVersion();
 	static FString GetBuildNumber();
 
+    /**
+     * @return true if the app is able to refresh in the background
+     */
+    static bool IsBackgroundAppRefreshAvailable();
+
+    /**
+     * Open the Settings app at the Notification page for this app.
+     */
+    static void OpenAppNotificationSettings();
+	
+    /**
+     * Open the app's custom settings in the Settings App.
+     */
+    static void OpenAppCustomSettings();
+
 	static void SetGracefulTerminationHandler();
 	static void SetCrashHandler(void(*CrashHandler)(const FGenericCrashContext& Context));
 
