[Pipe/Tasks]
 - Make sure tasks inside a Pipe are all marked as IsCompleted before returning from FPipe::WaitUntilEmpty

#rnx

diff --git a/Engine/Source/Runtime/Core/Public/Tasks/TaskPrivate.h b/Engine/Source/Runtime/Core/Public/Tasks/TaskPrivate.h
--- a/Engine/Source/Runtime/Core/Public/Tasks/TaskPrivate.h
+++ b/Engine/Source/Runtime/Core/Public/Tasks/TaskPrivate.h
@@ -530,12 +530,7 @@ public:
 			{
 				TASKGRAPH_VERBOSE_EVENT_SCOPE(FTaskBase::Close);
 				checkSlow(!IsCompleted());
-
-				if (GetPipe() != nullptr)
-				{
-					ClearPipe();
-				}
-
+				
 				// Push the first subsequent to the local queue so we pick it up directly as our next task.
 				// This saves us the cost of going to the global queue and performing a wake-up.
 				// But if we're a task event, always wake up new workers because the current task could continue executing for a long time after the trigger.
@@ -548,6 +543,13 @@ public:
 					Subsequent->TryUnlock(bWakeUpWorker);
 				}
 
+				// Clear the pipe after the task is completed (subsequents closed) so that any tasks part of the
+				// pipe are not seen still being executed after FPipe::WaitUntilEmpty has returned.
+				if (GetPipe() != nullptr)
+				{
+					ClearPipe();
+				}
+
 				// release nested tasks
 				ReleasePrerequisites();
 
