[AutoRTFM] Prevent premature inlining of RTFM core functions and prevent mixed -fautortfm PCHs

The compiler can inline function calls before the AutoRTFMCompilerPass splits functions into open and closed variants, which can lead to incorrect RTFM behaviour.

When AutoRTFM is enabled (clang-verse with -fautortfm) remove the implementation of autortfm_is_closed(). This prevents the compiler from inlining the function (as there is nothing to inline).

Replace the use of UE_AUTORTFM_AUTORTFM() with function bodies that explicitly branches on autortfm_is_closed(), calling either a UE_AUTORTFM_ALWAYS_OPEN annotated closed function, or instead performing the open logic. Because the function body calls autortfm_is_closed() (which cannot be inlined), these functions should avoid the premature inlining problem.

When AutoRTFM is not enabled (clang-verse with -fautortfm, or a different compiler) the AutoRTFM core functions are stubbed with `static inline` implementations. The static qualifier should be sufficient to prevent this implementation leaking into other modules with AutoRTFM enabled.

Change BENIGN_LANGOPT(AutoRTFM, ...) to LANGOPT(AutoRTFM, ...). This makes Clang consider the flag as incompatible when mixing -fautortfm / no -fautortfm between precompiled headers and users. This is important as the new __AUTORTFM_ENABLED macro will be baked into the PCH.

Update the UnrealBuildTool to separate PCHs with -fautortfm enabled and without.

diff --git a/Engine/Source/Programs/UnrealBuildTool/Configuration/UEBuildModuleCPP.cs b/Engine/Source/Programs/UnrealBuildTool/Configuration/UEBuildModuleCPP.cs
--- a/Engine/Source/Programs/UnrealBuildTool/Configuration/UEBuildModuleCPP.cs
+++ b/Engine/Source/Programs/UnrealBuildTool/Configuration/UEBuildModuleCPP.cs
@@ -1419,6 +1419,12 @@ namespace UnrealBuildTool
 			{
 				return false;
 			}
+
+			if (ModuleCompileEnvironment.bEnableAutoRTFMInstrumentation != CompileEnvironment.bEnableAutoRTFMInstrumentation)
+			{
+				return false;
+			}
+
 			return true;
 		}
 
@@ -1464,6 +1470,17 @@ namespace UnrealBuildTool
 					Variant += ".NonRTTI";
 				}
 			}
+			if (CompileEnvironment.bEnableAutoRTFMInstrumentation != BaseCompileEnvironment.bEnableAutoRTFMInstrumentation)
+			{
+				if (CompileEnvironment.bEnableAutoRTFMInstrumentation)
+				{
+					Variant += ".AutoRTFM";
+				}
+				else
+				{
+					Variant += ".NonAutoRTFM";
+				}
+			}
 			if (CompileEnvironment.bEnableExceptions != BaseCompileEnvironment.bEnableExceptions)
 			{
 				if (CompileEnvironment.bEnableExceptions)
diff --git a/Engine/Source/Programs/UnrealBuildTool/Platform/Windows/VCToolChain.cs b/Engine/Source/Programs/UnrealBuildTool/Platform/Windows/VCToolChain.cs
--- a/Engine/Source/Programs/UnrealBuildTool/Platform/Windows/VCToolChain.cs
+++ b/Engine/Source/Programs/UnrealBuildTool/Platform/Windows/VCToolChain.cs
@@ -1154,7 +1154,7 @@ namespace UnrealBuildTool
 				// Additional disabled warnings for msvc. Everything below should be checked if it is necessary
 				ClangWarnings.GetVCDisabledWarnings(Arguments);
 
-				if (!CompileEnvironment.bDisableAutoRTFMInstrumentation && CompileEnvironment.bUseAutoRTFMCompiler)
+				if (CompileEnvironment.bEnableAutoRTFMInstrumentation)
 				{
 					Arguments.Add("-fautortfm");
 				}
@@ -1195,7 +1195,7 @@ namespace UnrealBuildTool
 
 			if (Target.WindowsPlatform.Compiler.IsClang())
 			{
-				if (!CompileEnvironment.bDisableAutoRTFMInstrumentation && CompileEnvironment.bUseAutoRTFMCompiler)
+				if (CompileEnvironment.bEnableAutoRTFMInstrumentation)
 				{
 					Arguments.Add("-fautortfm");
 				}
diff --git a/Engine/Source/Programs/UnrealBuildTool/System/CppCompileEnvironment.cs b/Engine/Source/Programs/UnrealBuildTool/System/CppCompileEnvironment.cs
--- a/Engine/Source/Programs/UnrealBuildTool/System/CppCompileEnvironment.cs
+++ b/Engine/Source/Programs/UnrealBuildTool/System/CppCompileEnvironment.cs
@@ -287,6 +287,11 @@ namespace UnrealBuildTool
 		/// </summary>
 		public bool bDisableAutoRTFMInstrumentation = false;
 
+		/// <summary>
+		/// Whether AutoRTFM instrumentation is enabled or not.
+		/// </summary>
+		public bool bEnableAutoRTFMInstrumentation => bUseAutoRTFMCompiler && !bDisableAutoRTFMInstrumentation;
+
 		/// <summary>
 		/// If unity builds are enabled this can be used to override if this specific module will build using Unity.
 		/// This is set using the per module configurations in BuildConfiguration.
diff --git a/Engine/Source/Programs/UnrealBuildTool/ToolChain/ClangToolChain.cs b/Engine/Source/Programs/UnrealBuildTool/ToolChain/ClangToolChain.cs
--- a/Engine/Source/Programs/UnrealBuildTool/ToolChain/ClangToolChain.cs
+++ b/Engine/Source/Programs/UnrealBuildTool/ToolChain/ClangToolChain.cs
@@ -390,8 +390,8 @@ namespace UnrealBuildTool
 				// Validate PCH inputs by content if mtime check fails
 				Arguments.Add("-fpch-validate-input-files-content");
 			}
-
-			if (!CompileEnvironment.bDisableAutoRTFMInstrumentation && CompileEnvironment.bUseAutoRTFMCompiler)
+			
+			if (CompileEnvironment.bEnableAutoRTFMInstrumentation)
 			{
 				Arguments.Add("-fautortfm");
 			}
diff --git a/Engine/Source/Runtime/Core/Private/AutoRTFM/API.cpp b/Engine/Source/Runtime/Core/Private/AutoRTFM/API.cpp
--- a/Engine/Source/Runtime/Core/Private/AutoRTFM/API.cpp
+++ b/Engine/Source/Runtime/Core/Private/AutoRTFM/API.cpp
@@ -5,6 +5,10 @@
 #include "HAL/IConsoleManager.h"
 #include "GenericPlatform/GenericPlatformCrashContext.h"
 
+#if UE_AUTORTFM
+static_assert(UE_AUTORTFM_ENABLED, "AutoRTFM/API.cpp requires the compiler flag '-fautortfm'");
+#endif
+
 namespace
 {
 	// Move this to a local only and use functions to access this
@@ -182,7 +186,7 @@ namespace AutoRTFM
 	}
 }
 
-#if (defined(__AUTORTFM) && __AUTORTFM)
+#if (defined(__AUTORTFM_ENABLED) && __AUTORTFM_ENABLED)
 #include "AutoRTFM/AutoRTFMConstants.h"
 #include "CallNest.h"
 #include "Context.h"
@@ -201,6 +205,21 @@ namespace AutoRTFM
 namespace AutoRTFM
 {
 
+namespace
+{
+
+// Internal closed-variant implementations.
+UE_AUTORTFM_ALWAYS_OPEN bool RTFM_autortfm_is_transactional()
+{
+	return true;
+}
+
+UE_AUTORTFM_ALWAYS_OPEN autortfm_result RTFM_autortfm_transact(void (*UninstrumentedWork)(void*), void (*InstrumentedWork)(void*), void* Arg)
+{
+	FContext* Context = FContext::Get();
+	return static_cast<autortfm_result>(Context->Transact(InstrumentedWork, Arg));
+}
+
 UE_AUTORTFM_FORCEINLINE autortfm_result TransactThenOpenImpl(void (*UninstrumentedWork)(void*), void* Arg)
 {
 	return static_cast<autortfm_result>(
@@ -210,20 +229,169 @@ UE_AUTORTFM_FORCEINLINE autortfm_result TransactThenOpenImpl(void (*Uninstrument
 		}));
 }
 
-extern "C" UE_AUTORTFM_AUTORTFM("RTFM_autortfm_is_transactional") bool autortfm_is_transactional()
+UE_AUTORTFM_ALWAYS_OPEN autortfm_result RTFM_autortfm_transact_then_open(void (*UninstrumentedWork)(void*), void (*InstrumentedWork)(void*), void* Arg)
 {
-	if (ForTheRuntime::IsAutoRTFMRuntimeEnabled())
+	return TransactThenOpenImpl(UninstrumentedWork, Arg);
+}
+
+UE_AUTORTFM_ALWAYS_OPEN void RTFM_autortfm_commit(void (*UninstrumentedWork)(void*), void (*InstrumentedWork)(void*), void* Arg)
+{
+	autortfm_result Result = autortfm_transact(UninstrumentedWork, InstrumentedWork, Arg);
+	UE_CLOG(Result != autortfm_committed, LogAutoRTFM, Fatal, TEXT("Unexpected transaction result: %u."), Result);
+}
+
+UE_AUTORTFM_ALWAYS_OPEN void RTFM_autortfm_abort()
+{
+	FContext* Context = FContext::Get();
+	Context->AbortByRequestAndThrow();
+}
+
+UE_AUTORTFM_ALWAYS_OPEN void RTFM_autortfm_start_transaction()
+{
+	UE_LOG(LogAutoRTFM, Fatal, TEXT("The function `autortfm_start_transaction` was called from closed code."));
+}
+
+UE_AUTORTFM_ALWAYS_OPEN void RTFM_autortfm_commit_transaction()
+{
+	UE_LOG(LogAutoRTFM, Fatal, TEXT("The function `RTFM_autortfm_commit_transaction` was called from closed code."));
+}
+
+UE_AUTORTFM_ALWAYS_OPEN autortfm_result RTFM_autortfm_abort_transaction()
+{
+	FContext* const Context = FContext::Get();
+	return static_cast<autortfm_result>(Context->AbortTransaction(/* bIsClosed */ true, /* bIsCascading */ false));
+}
+
+UE_AUTORTFM_ALWAYS_OPEN autortfm_result RTFM_autortfm_cascading_abort_transaction()
+{
+	FContext* const Context = FContext::Get();
+	return static_cast<autortfm_result>(Context->AbortTransaction(/* bIsClosed */ true, /* bIsCascading */ true));
+}
+
+UE_AUTORTFM_ALWAYS_OPEN void RTFM_autortfm_clear_transaction_status()
+{
+	UE_LOG(LogAutoRTFM, Fatal, TEXT("The function `autortfm_clear_transaction_status` was called from closed code."));
+	AutoRTFM::Unreachable();
+}
+
+UE_AUTORTFM_ALWAYS_OPEN void RTFM_autortfm_open(void (*Work)(void*), void* Arg)
+{
+	Work(Arg);
+
+	FContext* Context = FContext::Get();
+	if (Context->IsAborting())
 	{
-		return FContext::IsTransactional();
+		Context->Throw();
 	}
+}
 
-	return false;
+UE_AUTORTFM_ALWAYS_OPEN autortfm_status RTFM_autortfm_close(void (*UninstrumentedWork)(void*), void (*InstrumentedWork)(void*), void* Arg)
+{
+	FContext* const Context = FContext::Get();
+
+	if (InstrumentedWork)
+	{
+		InstrumentedWork(Arg);
+	}
+	else
+	{
+		ensureMsgf(!ForTheRuntime::IsEnsureOnAbortByLanguageEnabled(), TEXT("Could not find function %p '%s' where '%s'."), UninstrumentedWork, *GetFunctionDescription(UninstrumentedWork), ANSI_TO_TCHAR("autortfm_close"));
+		Context->AbortByLanguageAndThrow();
+	}
+
+	return static_cast<autortfm_status>(Context->GetStatus());
 }
 
-extern "C" UE_AUTORTFM_AUTORTFM("RTFM_autortfm_is_closed") bool autortfm_is_closed()
+UE_AUTORTFM_ALWAYS_OPEN void RTFM_autortfm_record_open_write(void*, size_t)
 {
-	volatile bool ReallyDoNotInlineMe = true; // #jira SOL-6799
-	return !ReallyDoNotInlineMe;
+	UE_LOG(LogAutoRTFM, Fatal, TEXT("The function `autortfm_record_open_write` was called from closed code."));
+}
+
+UE_AUTORTFM_ALWAYS_OPEN void RTFM_OnCommitInternal(TFunction<void()> && Work)
+{
+	FContext* Context = FContext::Get();
+	ASSERT(Context->GetStatus() == EContextStatus::OnTrack);
+	Context->GetCurrentTransaction()->DeferUntilCommit(MoveTemp(Work));
+}
+
+UE_AUTORTFM_ALWAYS_OPEN void RTFM_OnAbortInternal(TFunction<void()>&& Work)
+{
+	FContext* Context = FContext::Get();
+	ASSERT(Context->GetStatus() == EContextStatus::OnTrack);
+	Context->GetCurrentTransaction()->DeferUntilAbort(MoveTemp(Work));
+}
+
+UE_AUTORTFM_ALWAYS_OPEN void RTFM_PushOnAbortHandlerInternal(const void* Key, TFunction<void()>&& Work)
+{
+	FContext* Context = FContext::Get();
+	ASSERT(Context->GetStatus() == EContextStatus::OnTrack);
+	Context->GetCurrentTransaction()->PushDeferUntilAbortHandler(Key, MoveTemp(Work));
+}
+
+UE_AUTORTFM_ALWAYS_OPEN void RTFM_PopOnAbortHandlerInternal(const void* Key)
+{
+	FContext* Context = FContext::Get();
+	ASSERT(Context->GetStatus() == EContextStatus::OnTrack);
+	Context->GetCurrentTransaction()->PopDeferUntilAbortHandler(Key);
+}
+
+UE_AUTORTFM_ALWAYS_OPEN void RTFM_autortfm_on_commit(void (*Work)(void*), void* Arg)
+{
+	RTFM_OnCommitInternal([Work, Arg] { Work(Arg); });
+}
+
+UE_AUTORTFM_ALWAYS_OPEN void RTFM_autortfm_on_abort(void (*Work)(void*), void* Arg)
+{
+	RTFM_OnAbortInternal([Work, Arg] { Work(Arg); });
+}
+
+UE_AUTORTFM_ALWAYS_OPEN void RTFM_autortfm_push_on_abort_handler(const void* Key, void (*Work)(void*), void* Arg)
+{
+	RTFM_PushOnAbortHandlerInternal(Key, [Work, Arg] { Work(Arg); });
+}
+
+UE_AUTORTFM_ALWAYS_OPEN void RTFM_autortfm_pop_on_abort_handler(const void* Key)
+{
+	RTFM_PopOnAbortHandlerInternal(Key);
+}
+
+UE_AUTORTFM_ALWAYS_OPEN void* RTFM_autortfm_did_allocate(void* Ptr, size_t Size)
+{
+	FContext* Context = FContext::Get();
+	Context->DidAllocate(Ptr, Size);
+	return Ptr;
+}
+
+UE_AUTORTFM_ALWAYS_OPEN void RTFM_autortfm_did_free(void* Ptr)
+{
+	// We should never-ever-ever actually free memory from within closed code of
+	// a transaction.
+	AutoRTFM::Unreachable();
+}
+
+UE_AUTORTFM_ALWAYS_OPEN void RTFM_autortfm_check_consistency_assuming_no_races()
+{
+}
+UE_AUTORTFM_REGISTER_OPEN_FUNCTION(autortfm_check_consistency_assuming_no_races);
+
+}  // anonymous namespace
+
+// The AutoRTFM public API.
+// Each function will be forked by the compiler into an open and closed variant.
+// autortfm_is_closed() is used to branch to the closed variants declared above.
+extern "C" bool autortfm_is_transactional()
+{
+	if (autortfm_is_closed())
+	{
+		return RTFM_autortfm_is_transactional();
+	}
+
+	if (ForTheRuntime::IsAutoRTFMRuntimeEnabled())
+	{
+		return FContext::IsTransactional();
+	}
+
+	return false;
 }
 
 extern "C" bool autortfm_is_committing_or_aborting()
@@ -236,9 +404,13 @@ extern "C" bool autortfm_is_committing_or_aborting()
 	return false;
 }
 
-// First Part - the API exposed outside transactions.
-extern "C" UE_AUTORTFM_AUTORTFM("RTFM_autortfm_transact") autortfm_result autortfm_transact(void (*UninstrumentedWork)(void*), void (*InstrumentedWork)(void*), void* Arg)
+extern "C" autortfm_result autortfm_transact(void (*UninstrumentedWork)(void*), void (*InstrumentedWork)(void*), void* Arg)
 {
+	if (autortfm_is_closed())
+	{
+		return RTFM_autortfm_transact(UninstrumentedWork, InstrumentedWork, Arg);
+	}
+
 	if (ForTheRuntime::IsAutoRTFMRuntimeEnabled())
 	{
 	    return static_cast<autortfm_result>(FContext::Get()->Transact(InstrumentedWork, Arg));
@@ -248,51 +420,93 @@ extern "C" UE_AUTORTFM_AUTORTFM("RTFM_autortfm_transact") autortfm_result autort
 	return autortfm_committed;
 }
 
-extern "C" UE_AUTORTFM_AUTORTFM("RTFM_autortfm_transact_then_open") autortfm_result autortfm_transact_then_open(void (*UninstrumentedWork)(void*), void (*InstrumentedWork)(void*), void* Arg)
+extern "C" autortfm_result autortfm_transact_then_open(void (*UninstrumentedWork)(void*), void (*InstrumentedWork)(void*), void* Arg)
 {
+	if (autortfm_is_closed())
+	{
+		return RTFM_autortfm_transact_then_open(UninstrumentedWork, InstrumentedWork, Arg);
+	}
+
 	return TransactThenOpenImpl(UninstrumentedWork, Arg);
 }
 
-extern "C" UE_AUTORTFM_AUTORTFM("RTFM_autortfm_commit") void autortfm_commit(void (*UninstrumentedWork)(void*), void (*InstrumentedWork)(void*), void* Arg)
+extern "C" void autortfm_commit(void (*UninstrumentedWork)(void*), void (*InstrumentedWork)(void*), void* Arg)
 {
+	if (autortfm_is_closed())
+	{
+		return RTFM_autortfm_commit(UninstrumentedWork, InstrumentedWork, Arg);
+	}
+
     autortfm_result Result = autortfm_transact(UninstrumentedWork, InstrumentedWork, Arg);
 	UE_CLOG(Result != autortfm_committed, LogAutoRTFM, Fatal, TEXT("Unexpected transaction result: %u."), Result);
 }
 
-extern "C" UE_AUTORTFM_AUTORTFM("RTFM_autortfm_abort") void autortfm_abort()
+extern "C" void autortfm_abort()
 {
+	if (autortfm_is_closed())
+	{
+		return RTFM_autortfm_abort();
+	}
+
 	UE_CLOG(!FContext::IsTransactional(), LogAutoRTFM, Fatal, TEXT("The function `autortfm_abort` was called from outside a transaction."));
 	FContext::Get()->AbortByRequestAndThrow();
 }
 
-extern "C" UE_AUTORTFM_AUTORTFM("RTFM_autortfm_start_transaction") bool autortfm_start_transaction()
+extern "C" bool autortfm_start_transaction()
 {
+	if (autortfm_is_closed())
+	{
+		RTFM_autortfm_start_transaction();
+		return false;
+	}
+
 	UE_CLOG(!FContext::IsTransactional(), LogAutoRTFM, Fatal, TEXT("The function `autortfm_start_transaction` was called from outside a transact."));
 	return FContext::Get()->StartTransaction();
 }
 
-extern "C" UE_AUTORTFM_AUTORTFM("RTFM_autortfm_commit_transaction") autortfm_result autortfm_commit_transaction()
+extern "C" autortfm_result autortfm_commit_transaction()
 {
+	if (autortfm_is_closed())
+	{
+		RTFM_autortfm_commit_transaction();
+		return autortfm_aborted_by_language;
+	}
+
 	UE_CLOG(!FContext::IsTransactional(), LogAutoRTFM, Fatal, TEXT("The function `autortfm_commit_transaction` was called from outside a transact."));
 	return static_cast<autortfm_result>(FContext::Get()->CommitTransaction());
 }
 
-extern "C" UE_AUTORTFM_AUTORTFM("RTFM_autortfm_abort_transaction") autortfm_result autortfm_abort_transaction()
+extern "C" autortfm_result autortfm_abort_transaction()
 {
+	if (autortfm_is_closed())
+	{
+		return RTFM_autortfm_abort_transaction();
+	}
+
 	UE_CLOG(!FContext::IsTransactional(), LogAutoRTFM, Fatal, TEXT("The function `autortfm_abort_transaction` was called from outside a transact."));
 	FContext* const Context = FContext::Get();
 	return static_cast<autortfm_result>(Context->AbortTransaction(false, false));
 }
 
-extern "C" UE_AUTORTFM_AUTORTFM("RTFM_autortfm_cascading_abort_transaction") autortfm_result autortfm_cascading_abort_transaction()
+extern "C" autortfm_result autortfm_cascading_abort_transaction()
 {
+	if (autortfm_is_closed())
+	{
+		return RTFM_autortfm_cascading_abort_transaction();
+	}
+
 	UE_CLOG(!FContext::IsTransactional(), LogAutoRTFM, Fatal, TEXT("The function `autortfm_cascading_abort_transaction` was called from outside a transact."));
 	FContext* const Context = FContext::Get();
 	return static_cast<autortfm_result>(Context->AbortTransaction(false, true));
 }
 
-extern "C" UE_AUTORTFM_AUTORTFM("RTFM_autortfm_clear_transaction_status") void autortfm_clear_transaction_status()
+extern "C" void autortfm_clear_transaction_status()
 {
+	if (autortfm_is_closed())
+	{
+		return RTFM_autortfm_clear_transaction_status();
+	}
+
 	ASSERT(FContext::Get()->IsAborting());
 	FContext::Get()->ClearTransactionStatus();
 }
@@ -313,13 +527,23 @@ extern "C" UE_AUTORTFM_NOAUTORTFM bool autortfm_current_nest_throw()
 	return true;
 }
 
-extern "C" UE_AUTORTFM_AUTORTFM("RTFM_autortfm_open") void autortfm_open(void (*Work)(void*), void* Arg)
+extern "C" void autortfm_open(void (*Work)(void*), void* Arg)
 {
+	if (autortfm_is_closed())
+	{
+		return RTFM_autortfm_open(Work, Arg);
+	}
+
 	Work(Arg);
 }
 
-extern "C" UE_AUTORTFM_AUTORTFM("RTFM_autortfm_close") autortfm_status autortfm_close(void (*UninstrumentedWork)(void*), void (*InstrumentedWork)(void*), void* Arg)
+extern "C" autortfm_status autortfm_close(void (*UninstrumentedWork)(void*), void (*InstrumentedWork)(void*), void* Arg)
 {
+	if (autortfm_is_closed())
+	{
+		return RTFM_autortfm_close(UninstrumentedWork, InstrumentedWork, Arg);
+	}
+
 	autortfm_status Result = autortfm_status_ontrack;
 
 	if (ForTheRuntime::IsAutoRTFMRuntimeEnabled())
@@ -346,8 +570,13 @@ extern "C" UE_AUTORTFM_AUTORTFM("RTFM_autortfm_close") autortfm_status autortfm_
 	return Result;
 }
 
-extern "C" UE_AUTORTFM_AUTORTFM("RTFM_autortfm_record_open_write") void autortfm_record_open_write(void* Ptr, size_t Size)
+extern "C" void autortfm_record_open_write(void* Ptr, size_t Size)
 {
+	if (autortfm_is_closed())
+	{
+		return RTFM_autortfm_record_open_write(Ptr, Size);
+	}
+
 	FContext* const Context = FContext::Get();
 
     Context->CheckOpenRecordWrite(Ptr);
@@ -360,21 +589,38 @@ extern "C" UE_AUTORTFM_NOAUTORTFM void autortfm_register_open_function(void* Ori
     FunctionMapAdd(OriginalFunction, NewFunction);
 }
 
-UE_AUTORTFM_AUTORTFM("RTFM_OnCommitInternal") void ForTheRuntime::OnCommitInternal(TFunction<void()> && Work)
+void ForTheRuntime::OnCommitInternal(TFunction<void()> && Work)
 {
-    Work();
+	if (autortfm_is_closed())
+	{
+		return RTFM_OnCommitInternal(MoveTemp(Work));
+	}
+
+	Work();
 }
 
-UE_AUTORTFM_AUTORTFM("RTFM_OnAbortInternal") void ForTheRuntime::OnAbortInternal(TFunction<void()> && Work)
+void ForTheRuntime::OnAbortInternal(TFunction<void()> && Work)
 {
+	if (autortfm_is_closed())
+	{
+		return RTFM_OnAbortInternal(MoveTemp(Work));
+	}
 }
 
-UE_AUTORTFM_AUTORTFM("RTFM_PushOnAbortHandlerInternal") void ForTheRuntime::PushOnAbortHandlerInternal(const void* Key, TFunction<void()> && Work)
+void ForTheRuntime::PushOnAbortHandlerInternal(const void* Key, TFunction<void()> && Work)
 {
+	if (autortfm_is_closed())
+	{
+		return RTFM_PushOnAbortHandlerInternal(Key, MoveTemp(Work));
+	}
 }
 
-UE_AUTORTFM_AUTORTFM("RTFM_PopOnAbortHandlerInternal") void ForTheRuntime::PopOnAbortHandlerInternal(const void* Key)
+void ForTheRuntime::PopOnAbortHandlerInternal(const void* Key)
 {
+	if (autortfm_is_closed())
+	{
+		return RTFM_PopOnAbortHandlerInternal(Key);
+	}
 }
 
 void OpenCommit(TFunction<void()>&& Work)
@@ -387,30 +633,60 @@ void OpenAbort(TFunction<void()>&& Work)
 	ForTheRuntime::OnAbortInternal(MoveTemp(Work));
 }
 
-extern "C" UE_AUTORTFM_AUTORTFM("RTFM_autortfm_on_commit") void autortfm_on_commit(void (*Work)(void*), void* Arg)
+extern "C" void autortfm_on_commit(void (*Work)(void*), void* Arg)
 {
+	if (autortfm_is_closed())
+	{
+		return RTFM_autortfm_on_commit(Work, Arg);
+	}
+
     Work(Arg);
 }
 
-extern "C" UE_AUTORTFM_AUTORTFM("RTFM_autortfm_on_abort") void autortfm_on_abort(void (*Work)(void*), void* Arg)
+extern "C" void autortfm_on_abort(void (*Work)(void*), void* Arg)
 {
+	if (autortfm_is_closed())
+	{
+		return RTFM_autortfm_on_abort(Work, Arg);
+	}
+
 }
 
-extern "C" UE_AUTORTFM_AUTORTFM("RTFM_autortfm_push_on_abort_handler") void autortfm_push_on_abort_handler(const void* Key, void (*Work)(void*), void* Arg)
+extern "C" void autortfm_push_on_abort_handler(const void* Key, void (*Work)(void*), void* Arg)
 {
+	if (autortfm_is_closed())
+	{
+		return RTFM_autortfm_push_on_abort_handler(Key, Work, Arg);
+	}
+
 }
 
-extern "C" UE_AUTORTFM_AUTORTFM("RTFM_autortfm_pop_on_abort_handler") void autortfm_pop_on_abort_handler(const void* Key)
+extern "C" void autortfm_pop_on_abort_handler(const void* Key)
 {
+	if (autortfm_is_closed())
+	{
+		return RTFM_autortfm_pop_on_abort_handler(Key);
+	}
+
 }
 
-extern "C" UE_AUTORTFM_AUTORTFM("RTFM_autortfm_did_allocate") void* autortfm_did_allocate(void* Ptr, size_t Size)
+extern "C" void* autortfm_did_allocate(void* Ptr, size_t Size)
 {
+	if (autortfm_is_closed())
+	{
+		return RTFM_autortfm_did_allocate(Ptr, Size);
+	}
+
     return Ptr;
 }
 
-extern "C" UE_AUTORTFM_AUTORTFM("RTFM_autortfm_did_free") void autortfm_did_free(void* Ptr)
+extern "C" void autortfm_did_free(void* Ptr)
 {
+	if (autortfm_is_closed())
+	{
+		return RTFM_autortfm_did_free(Ptr);
+	}
+
 	// We only need to process did free if we need to track allocation locations.
 	if constexpr (bTrackAllocationLocations)
 	{
@@ -433,8 +709,13 @@ extern "C" UE_AUTORTFM_AUTORTFM("RTFM_autortfm_did_free") void autortfm_did_free
 	}
 }
 
-extern "C" UE_AUTORTFM_AUTORTFM("RTFM_autortfm_check_consistency_assuming_no_races") void autortfm_check_consistency_assuming_no_races()
+extern "C" void autortfm_check_consistency_assuming_no_races()
 {
+	if (autortfm_is_closed())
+	{
+		return RTFM_autortfm_check_consistency_assuming_no_races();
+	}
+
     if (FContext::IsTransactional())
     {
         AutoRTFM::Unreachable();
@@ -477,170 +758,6 @@ extern "C" UE_AUTORTFM_NOAUTORTFM void autortfm_check_abi(void* const Ptr, const
 	UE_CLOG(RuntimeConstants != *CompilerConstants, LogAutoRTFM, Fatal, TEXT("ABI error between AutoRTFM compiler and runtime."));
 }
 
-// Second Part - the same API exposed inside transactions. Note that we don't expose all of the API
-// to transactions! That's intentional. However, things like autortfm_defer_until_commit can be called
-// from an open nest in a transaction.
-extern "C" UE_AUTORTFM_NOAUTORTFM bool RTFM_autortfm_is_transactional()
-{
-    return true;
-}
-
-extern "C" UE_AUTORTFM_NOAUTORTFM bool RTFM_autortfm_is_closed()
-{
-	return true;
-}
-
-extern "C" UE_AUTORTFM_NOAUTORTFM autortfm_result RTFM_autortfm_transact(void (*UninstrumentedWork)(void*), void (*InstrumentedWork)(void*), void* Arg)
-{
-	FContext* Context = FContext::Get();
-    return static_cast<autortfm_result>(Context->Transact(InstrumentedWork, Arg));
-}
-
-extern "C" UE_AUTORTFM_NOAUTORTFM autortfm_result RTFM_autortfm_transact_then_open(void (*UninstrumentedWork)(void*), void (*InstrumentedWork)(void*), void* Arg)
-{
-    return TransactThenOpenImpl(UninstrumentedWork, Arg);
-}
-
-extern "C" UE_AUTORTFM_NOAUTORTFM void RTFM_autortfm_commit(void (*UninstrumentedWork)(void*), void (*InstrumentedWork)(void*), void* Arg)
-{
-    autortfm_result Result = autortfm_transact(UninstrumentedWork, InstrumentedWork, Arg);
-	UE_CLOG(Result != autortfm_committed, LogAutoRTFM, Fatal, TEXT("Unexpected transaction result: %u."), Result);
-}
-
-extern "C" UE_AUTORTFM_NOAUTORTFM void RTFM_autortfm_abort()
-{
-	FContext* Context = FContext::Get();
-    Context->AbortByRequestAndThrow();
-}
-
-extern "C" UE_AUTORTFM_NOAUTORTFM void RTFM_autortfm_start_transaction()
-{
-	UE_LOG(LogAutoRTFM, Fatal, TEXT("The function `autortfm_start_transaction` was called from closed code."));
-}
-
-extern "C" UE_AUTORTFM_NOAUTORTFM void RTFM_autortfm_commit_transaction()
-{
-	UE_LOG(LogAutoRTFM, Fatal, TEXT("The function `RTFM_autortfm_commit_transaction` was called from closed code."));
-}
-
-extern "C" UE_AUTORTFM_NOAUTORTFM autortfm_result RTFM_autortfm_abort_transaction()
-{
-	FContext* const Context = FContext::Get();
-	return static_cast<autortfm_result>(Context->AbortTransaction(true, false));
-}
-
-extern "C" UE_AUTORTFM_NOAUTORTFM autortfm_result RTFM_autortfm_cascading_abort_transaction()
-{
-	FContext* const Context = FContext::Get();
-	return static_cast<autortfm_result>(Context->AbortTransaction(true, true));
-}
-
-extern "C" UE_AUTORTFM_NOAUTORTFM void RTFM_autortfm_clear_transaction_status()
-{
-	UE_LOG(LogAutoRTFM, Fatal, TEXT("The function `autortfm_clear_transaction_status` was called from closed code."));
-	AutoRTFM::Unreachable();
-}
-
-extern "C" UE_AUTORTFM_NOAUTORTFM void RTFM_autortfm_open(void (*Work)(void*), void* Arg)
-{
-	Work(Arg);
-
-	FContext* Context = FContext::Get();
-	if (Context->IsAborting())
-	{
-		Context->Throw();
-	}
-}
-
-extern "C" UE_AUTORTFM_NOAUTORTFM autortfm_status RTFM_autortfm_close(void (*UninstrumentedWork)(void*), void (*InstrumentedWork)(void*), void* Arg)
-{
-	FContext* const Context = FContext::Get();
-
-	if (InstrumentedWork)
-    {
-        InstrumentedWork(Arg);
-    }
-	else
-	{
-		ensureMsgf(!ForTheRuntime::IsEnsureOnAbortByLanguageEnabled(), TEXT("Could not find function %p '%s' where '%s'."), UninstrumentedWork, *GetFunctionDescription(UninstrumentedWork), ANSI_TO_TCHAR("autortfm_close"));
-	    Context->AbortByLanguageAndThrow();
-	}
-
-	return static_cast<autortfm_status>(Context->GetStatus());
-}
-
-extern "C" UE_AUTORTFM_NOAUTORTFM void RTFM_autortfm_record_open_write(void*, size_t)
-{
-	UE_LOG(LogAutoRTFM, Fatal, TEXT("The function `autortfm_record_open_write` was called from closed code."));
-}
-
-extern "C" UE_AUTORTFM_NOAUTORTFM void RTFM_OnCommitInternal(TFunction<void()> && Work)
-{
-	FContext* Context = FContext::Get();
-    ASSERT(Context->GetStatus() == EContextStatus::OnTrack);
-    Context->GetCurrentTransaction()->DeferUntilCommit(MoveTemp(Work));
-}
-
-extern "C" UE_AUTORTFM_NOAUTORTFM void RTFM_OnAbortInternal(TFunction<void()>&& Work)
-{
-	FContext* Context = FContext::Get();
-    ASSERT(Context->GetStatus() == EContextStatus::OnTrack);
-    Context->GetCurrentTransaction()->DeferUntilAbort(MoveTemp(Work));
-}
-
-extern "C" UE_AUTORTFM_NOAUTORTFM void RTFM_PushOnAbortHandlerInternal(const void* Key, TFunction<void()>&& Work)
-{
-	FContext* Context = FContext::Get();
-    ASSERT(Context->GetStatus() == EContextStatus::OnTrack);
-    Context->GetCurrentTransaction()->PushDeferUntilAbortHandler(Key, MoveTemp(Work));
-}
-
-extern "C" UE_AUTORTFM_NOAUTORTFM void RTFM_PopOnAbortHandlerInternal(const void* Key)
-{
-	FContext* Context = FContext::Get();
-    ASSERT(Context->GetStatus() == EContextStatus::OnTrack);
-    Context->GetCurrentTransaction()->PopDeferUntilAbortHandler(Key);
-}
-
-extern "C" UE_AUTORTFM_NOAUTORTFM void RTFM_autortfm_on_commit(void (*Work)(void*), void* Arg)
-{
-    RTFM_OnCommitInternal([Work, Arg] { Work(Arg); });
-}
-
-extern "C" UE_AUTORTFM_NOAUTORTFM void RTFM_autortfm_on_abort(void (*Work)(void*), void* Arg)
-{
-    RTFM_OnAbortInternal([Work, Arg] { Work(Arg); });
-}
-
-extern "C" UE_AUTORTFM_NOAUTORTFM void RTFM_autortfm_push_on_abort_handler(const void* Key, void (*Work)(void*), void* Arg)
-{
-    RTFM_PushOnAbortHandlerInternal(Key, [Work, Arg] { Work(Arg); });
-}
-
-extern "C" UE_AUTORTFM_NOAUTORTFM void RTFM_autortfm_pop_on_abort_handler(const void* Key)
-{
-    RTFM_PopOnAbortHandlerInternal(Key);
-}
-
-extern "C" UE_AUTORTFM_NOAUTORTFM void* RTFM_autortfm_did_allocate(void* Ptr, size_t Size)
-{
-	FContext* Context = FContext::Get();
-    Context->DidAllocate(Ptr, Size);
-    return Ptr;
-}
-
-extern "C" UE_AUTORTFM_NOAUTORTFM void RTFM_autortfm_did_free(void* Ptr)
-{
-	// We should never-ever-ever actually free memory from within closed code of
-	// a transaction.
-	AutoRTFM::Unreachable();
-}
-
-extern "C" UE_AUTORTFM_NOAUTORTFM void RTFM_autortfm_check_consistency_assuming_no_races()
-{
-}
-UE_AUTORTFM_REGISTER_OPEN_FUNCTION(autortfm_check_consistency_assuming_no_races);
-
 } // namespace AutoRTFM
 
 #endif // defined(__AUTORTFM) && __AUTORTFM
diff --git a/Engine/Source/Runtime/Core/Public/AutoRTFM/AutoRTFM.h b/Engine/Source/Runtime/Core/Public/AutoRTFM/AutoRTFM.h
--- a/Engine/Source/Runtime/Core/Public/AutoRTFM/AutoRTFM.h
+++ b/Engine/Source/Runtime/Core/Public/AutoRTFM/AutoRTFM.h
@@ -3,11 +3,17 @@
 #pragma once
 
 #if (defined(__AUTORTFM) && __AUTORTFM)
-#define UE_AUTORTFM 1
+#define UE_AUTORTFM 1  // Compiler is 'verse-clang'
 #else
 #define UE_AUTORTFM 0
 #endif
 
+#if (defined(__AUTORTFM_ENABLED) && __AUTORTFM_ENABLED)
+#define UE_AUTORTFM_ENABLED 1  // Compiled with '-fautortfm'
+#else
+#define UE_AUTORTFM_ENABLED 0
+#endif
+
 #if !defined(UE_AUTORTFM_ENABLED_RUNTIME_BY_DEFAULT)
 #define UE_AUTORTFM_ENABLED_RUNTIME_BY_DEFAULT 1
 #endif
@@ -29,13 +35,10 @@
 #endif
 
 #if UE_AUTORTFM
-// #jira SOL-6799: clang::optnone is used to prevent inlining, despite clang::noinline
-#define UE_AUTORTFM_AUTORTFM(F) [[clang::autortfm(F), clang::noinline, clang::optnone]]
 #define UE_AUTORTFM_NOAUTORTFM [[clang::noautortfm, clang::noinline]]
 // #jira SOL-6589: remove clang::noinline once this JIRA is fixed.
 #define UE_AUTORTFM_ALWAYS_OPEN [[clang::autortfm_always_open, clang::noinline]]
 #else
-#define UE_AUTORTFM_AUTORTFM(F)
 #define UE_AUTORTFM_NOAUTORTFM
 #define UE_AUTORTFM_ALWAYS_OPEN
 #endif
@@ -57,6 +60,8 @@
 #define UE_AUTORTFM_FORCEINLINE inline
 #define UE_AUTORTFM_MEMCPY ::memcpy
 #define UE_AUTORTFM_MOVE std::move
+#define PRAGMA_DISABLE_UNREACHABLE_CODE_WARNINGS
+#define PRAGMA_RESTORE_UNREACHABLE_CODE_WARNINGS
 #else
 #include <HAL/Platform.h>
 #include <HAL/PlatformMemory.h>
@@ -118,37 +123,39 @@ typedef enum
 	autortfm_status_aborted_by_cascade
 } autortfm_status;
 
-#if UE_AUTORTFM
-UE_AUTORTFM_API bool autortfm_is_transactional(void);
+#if UE_AUTORTFM_ENABLED
+// Note: There is no implementation of this function.
+// The AutoRTFM compiler will replace all calls to this function with a constant boolean value.
+UE_AUTORTFM_API bool autortfm_is_closed(void);
 #else
-UE_AUTORTFM_FORCEINLINE bool autortfm_is_transactional(void)
+static UE_AUTORTFM_FORCEINLINE bool autortfm_is_closed(void)
 {
     return false;
 }
 #endif
 
-#if UE_AUTORTFM
-UE_AUTORTFM_API bool autortfm_is_closed(void);
+#if UE_AUTORTFM_ENABLED
+UE_AUTORTFM_API bool autortfm_is_transactional(void);
 #else
-UE_AUTORTFM_FORCEINLINE bool autortfm_is_closed(void)
+static UE_AUTORTFM_FORCEINLINE bool autortfm_is_transactional(void)
 {
     return false;
 }
 #endif
 
-#if UE_AUTORTFM
+#if UE_AUTORTFM_ENABLED
 UE_AUTORTFM_API bool autortfm_is_committing_or_aborting(void);
 #else
-UE_AUTORTFM_FORCEINLINE bool autortfm_is_committing_or_aborting(void)
+static UE_AUTORTFM_FORCEINLINE bool autortfm_is_committing_or_aborting(void)
 {
 	return false;
 }
 #endif
 
-#if UE_AUTORTFM
+#if UE_AUTORTFM_ENABLED
 UE_AUTORTFM_API autortfm_result autortfm_transact(void (*UninstrumentedWork)(void*), void (*InstrumentedWork)(void*), void* Arg);
 #else
-UE_AUTORTFM_FORCEINLINE autortfm_result autortfm_transact(void (*UninstrumentedWork)(void*), void (*InstrumentedWork)(void*), void* Arg)
+static UE_AUTORTFM_FORCEINLINE autortfm_result autortfm_transact(void (*UninstrumentedWork)(void*), void (*InstrumentedWork)(void*), void* Arg)
 {
 	UE_AUTORTFM_UNUSED(InstrumentedWork);
 	UninstrumentedWork(Arg);
@@ -156,10 +163,10 @@ UE_AUTORTFM_FORCEINLINE autortfm_result autortfm_transact(void (*UninstrumentedW
 }
 #endif
 
-#if UE_AUTORTFM
+#if UE_AUTORTFM_ENABLED
 UE_AUTORTFM_API autortfm_result autortfm_transact_then_open(void (*UninstrumentedWork)(void*), void (*InstrumentedWork)(void*), void* Arg);
 #else
-UE_AUTORTFM_FORCEINLINE autortfm_result autortfm_transact_then_open(void (*UninstrumentedWork)(void*), void (*InstrumentedWork)(void*), void* Arg)
+static UE_AUTORTFM_FORCEINLINE autortfm_result autortfm_transact_then_open(void (*UninstrumentedWork)(void*), void (*InstrumentedWork)(void*), void* Arg)
 {
 	UE_AUTORTFM_UNUSED(InstrumentedWork);
 	UninstrumentedWork(Arg);
@@ -167,48 +174,48 @@ UE_AUTORTFM_FORCEINLINE autortfm_result autortfm_transact_then_open(void (*Unins
 }
 #endif
 
-#if UE_AUTORTFM
+#if UE_AUTORTFM_ENABLED
 UE_AUTORTFM_API void autortfm_commit(void (*UninstrumentedWork)(void*), void (*InstrumentedWork)(void*), void* Arg);
 #else
-UE_AUTORTFM_FORCEINLINE void autortfm_commit(void (*UninstrumentedWork)(void*), void (*InstrumentedWork)(void*), void* Arg)
+static UE_AUTORTFM_FORCEINLINE void autortfm_commit(void (*UninstrumentedWork)(void*), void (*InstrumentedWork)(void*), void* Arg)
 {
 	UE_AUTORTFM_UNUSED(InstrumentedWork);
 	UninstrumentedWork(Arg);
 }
 #endif
 
-#if UE_AUTORTFM
+#if UE_AUTORTFM_ENABLED
 UE_AUTORTFM_API autortfm_result autortfm_abort_transaction();
 #else
-UE_AUTORTFM_FORCEINLINE autortfm_result autortfm_abort_transaction() { return autortfm_aborted_by_request; }
+static UE_AUTORTFM_FORCEINLINE autortfm_result autortfm_abort_transaction() { return autortfm_aborted_by_request; }
 #endif
 
-#if UE_AUTORTFM
+#if UE_AUTORTFM_ENABLED
 UE_AUTORTFM_API autortfm_result autortfm_cascading_abort_transaction();
 #else
-UE_AUTORTFM_FORCEINLINE autortfm_result autortfm_cascading_abort_transaction() { return autortfm_aborted_by_cascade; }
+static UE_AUTORTFM_FORCEINLINE autortfm_result autortfm_cascading_abort_transaction() { return autortfm_aborted_by_cascade; }
 #endif
 
-#if UE_AUTORTFM
+#if UE_AUTORTFM_ENABLED
 UE_AUTORTFM_API bool autortfm_start_transaction();
 #else
-UE_AUTORTFM_FORCEINLINE bool autortfm_start_transaction() { return false; }
+static UE_AUTORTFM_FORCEINLINE bool autortfm_start_transaction() { return false; }
 #endif
 
-#if UE_AUTORTFM
+#if UE_AUTORTFM_ENABLED
 UE_AUTORTFM_API autortfm_result autortfm_commit_transaction();
 #else
-UE_AUTORTFM_FORCEINLINE autortfm_result autortfm_commit_transaction() { return autortfm_aborted_by_language; }
+static UE_AUTORTFM_FORCEINLINE autortfm_result autortfm_commit_transaction() { return autortfm_aborted_by_language; }
 #endif
 
-#if UE_AUTORTFM
+#if UE_AUTORTFM_ENABLED
 UE_AUTORTFM_API void autortfm_clear_transaction_status();
 #else
-UE_AUTORTFM_FORCEINLINE void autortfm_clear_transaction_status() {}
+static UE_AUTORTFM_FORCEINLINE void autortfm_clear_transaction_status() {}
 #endif
 
-#if UE_AUTORTFM
-UE_AUTORTFM_FORCEINLINE void autortfm_abort_if_transactional(void)
+#if UE_AUTORTFM_ENABLED
+static UE_AUTORTFM_FORCEINLINE void autortfm_abort_if_transactional(void)
 {
 	if (autortfm_is_transactional())
 	{
@@ -216,11 +223,11 @@ UE_AUTORTFM_FORCEINLINE void autortfm_abort_if_transactional(void)
 	}
 }
 #else
-UE_AUTORTFM_FORCEINLINE void autortfm_abort_if_transactional(void) { }
+static UE_AUTORTFM_FORCEINLINE void autortfm_abort_if_transactional(void) { }
 #endif
 
-#if UE_AUTORTFM
-UE_AUTORTFM_FORCEINLINE void autortfm_abort_if_closed(void)
+#if UE_AUTORTFM_ENABLED
+static UE_AUTORTFM_FORCEINLINE void autortfm_abort_if_closed(void)
 {
 	if (autortfm_is_closed())
 	{
@@ -228,20 +235,20 @@ UE_AUTORTFM_FORCEINLINE void autortfm_abort_if_closed(void)
 	}
 }
 #else
-UE_AUTORTFM_FORCEINLINE void autortfm_abort_if_closed(void) { }
+static UE_AUTORTFM_FORCEINLINE void autortfm_abort_if_closed(void) { }
 #endif
 
-#if UE_AUTORTFM
+#if UE_AUTORTFM_ENABLED
 UE_AUTORTFM_API void autortfm_open(void (*work)(void* arg), void* arg);
 #else
-UE_AUTORTFM_FORCEINLINE void autortfm_open(void (*work)(void* arg), void* arg) { work(arg); }
+static UE_AUTORTFM_FORCEINLINE void autortfm_open(void (*work)(void* arg), void* arg) { work(arg); }
 #endif
 
-#if UE_AUTORTFM
+#if UE_AUTORTFM_ENABLED
 [[nodiscard]] UE_AUTORTFM_API autortfm_status autortfm_close(void (*UninstrumentedWork)(void*), void (*InstrumentedWork)(void*), void* Arg);
 #else
 PRAGMA_DISABLE_UNREACHABLE_CODE_WARNINGS
-[[nodiscard]] UE_AUTORTFM_FORCEINLINE autortfm_status autortfm_close(void (*UninstrumentedWork)(void*), void (*InstrumentedWork)(void*), void* Arg)
+[[nodiscard]] static UE_AUTORTFM_FORCEINLINE autortfm_status autortfm_close(void (*UninstrumentedWork)(void*), void (*InstrumentedWork)(void*), void* Arg)
 {
 	UE_AUTORTFM_UNUSED(UninstrumentedWork);
 	UE_AUTORTFM_UNUSED(InstrumentedWork);
@@ -252,97 +259,97 @@ PRAGMA_DISABLE_UNREACHABLE_CODE_WARNINGS
 PRAGMA_RESTORE_UNREACHABLE_CODE_WARNINGS
 #endif
 
-#if UE_AUTORTFM
+#if UE_AUTORTFM_ENABLED
 UE_AUTORTFM_API void autortfm_record_open_write(void* Ptr, size_t Size);
 #else
-UE_AUTORTFM_FORCEINLINE void autortfm_record_open_write(void* Ptr, size_t Size)
+static UE_AUTORTFM_FORCEINLINE void autortfm_record_open_write(void* Ptr, size_t Size)
 {
 	UE_AUTORTFM_UNUSED(Ptr);
 	UE_AUTORTFM_UNUSED(Size);
 }
 #endif
 
-#if UE_AUTORTFM
+#if UE_AUTORTFM_ENABLED
 UE_AUTORTFM_API void autortfm_register_open_function(void* original_function, void* new_function);
 #else
-UE_AUTORTFM_FORCEINLINE void autortfm_register_open_function(void* original_function, void* new_function) 
+static UE_AUTORTFM_FORCEINLINE void autortfm_register_open_function(void* original_function, void* new_function) 
 { 
 	UE_AUTORTFM_UNUSED(original_function);
 	UE_AUTORTFM_UNUSED(new_function);
 }
 #endif
 
-#if UE_AUTORTFM
+#if UE_AUTORTFM_ENABLED
 UE_AUTORTFM_API void autortfm_on_commit(void (*work)(void* arg), void* arg);
 #else
-UE_AUTORTFM_FORCEINLINE void autortfm_on_commit(void (*work)(void* arg), void* arg)
+static UE_AUTORTFM_FORCEINLINE void autortfm_on_commit(void (*work)(void* arg), void* arg)
 {
     work(arg);
 }
 #endif
 
 [[deprecated("Use autortfm_on_commit instead.")]]
-UE_AUTORTFM_FORCEINLINE void autortfm_open_commit(void (*work)(void* arg), void* arg)
+static UE_AUTORTFM_FORCEINLINE void autortfm_open_commit(void (*work)(void* arg), void* arg)
 {
 	autortfm_on_commit(work, arg);
 }
 
-#if UE_AUTORTFM
+#if UE_AUTORTFM_ENABLED
 UE_AUTORTFM_API void autortfm_on_abort(void (*work)(void* arg), void* arg);
 #else
-UE_AUTORTFM_FORCEINLINE void autortfm_on_abort(void (*work)(void* arg), void* arg)
+static UE_AUTORTFM_FORCEINLINE void autortfm_on_abort(void (*work)(void* arg), void* arg)
 {
 	UE_AUTORTFM_UNUSED(work);
 	UE_AUTORTFM_UNUSED(arg);
 }
 #endif
 
-#if UE_AUTORTFM
+#if UE_AUTORTFM_ENABLED
 UE_AUTORTFM_API void autortfm_push_on_abort_handler(const void* key, void (*work)(void* arg), void* arg);
 UE_AUTORTFM_API void autortfm_pop_on_abort_handler(const void* key);
 #else
-UE_AUTORTFM_FORCEINLINE void autortfm_push_on_abort_handler(const void* key, void (*work)(void* arg), void* arg)
+static UE_AUTORTFM_FORCEINLINE void autortfm_push_on_abort_handler(const void* key, void (*work)(void* arg), void* arg)
 {
 	UE_AUTORTFM_UNUSED(key);
 	UE_AUTORTFM_UNUSED(work);
 	UE_AUTORTFM_UNUSED(arg);
 }
 
-UE_AUTORTFM_FORCEINLINE void autortfm_pop_on_abort_handler(const void* key)
+static UE_AUTORTFM_FORCEINLINE void autortfm_pop_on_abort_handler(const void* key)
 {
 	UE_AUTORTFM_UNUSED(key);
 }
 #endif
 
 [[deprecated("Use autortfm_on_abort instead.")]]
-UE_AUTORTFM_FORCEINLINE void autortfm_open_abort(void (*work)(void* arg), void* arg)
+static UE_AUTORTFM_FORCEINLINE void autortfm_open_abort(void (*work)(void* arg), void* arg)
 {
 	autortfm_on_abort(work, arg);
 }
 
-#if UE_AUTORTFM
+#if UE_AUTORTFM_ENABLED
 UE_AUTORTFM_API void* autortfm_did_allocate(void* ptr, size_t size);
 #else
-UE_AUTORTFM_FORCEINLINE void* autortfm_did_allocate(void* ptr, size_t size)
+static UE_AUTORTFM_FORCEINLINE void* autortfm_did_allocate(void* ptr, size_t size)
 {
 	UE_AUTORTFM_UNUSED(size);
     return ptr;
 }
 #endif
 
-#if UE_AUTORTFM
+#if UE_AUTORTFM_ENABLED
 UE_AUTORTFM_API void autortfm_did_free(void* ptr);
 #else
-UE_AUTORTFM_FORCEINLINE void autortfm_did_free(void* ptr)
+static UE_AUTORTFM_FORCEINLINE void autortfm_did_free(void* ptr)
 {
 	UE_AUTORTFM_UNUSED(ptr);
 }
 #endif
 
-#if UE_AUTORTFM
+#if UE_AUTORTFM_ENABLED
 UE_AUTORTFM_API void autortfm_check_consistency_assuming_no_races(void);
 #else
-UE_AUTORTFM_FORCEINLINE void autortfm_check_consistency_assuming_no_races(void) { }
+static UE_AUTORTFM_FORCEINLINE void autortfm_check_consistency_assuming_no_races(void) { }
 #endif
 
 // If running with AutoRTFM enabled, then perform an ABI check between the
@@ -351,10 +358,10 @@ UE_AUTORTFM_FORCEINLINE void autortfm_check_consistency_assuming_no_races(void)
 // compiler pass. Should not be called manually by the user, a call to this will
 // be injected by the compiler into a global constructor in the AutoRTFM compiled
 // code.
-#if UE_AUTORTFM
+#if UE_AUTORTFM_ENABLED
 UE_AUTORTFM_API void autortfm_check_abi(void* ptr, size_t size);
 #else
-UE_AUTORTFM_FORCEINLINE void autortfm_check_abi(void* ptr, size_t size)
+static UE_AUTORTFM_FORCEINLINE void autortfm_check_abi(void* ptr, size_t size)
 {
 	UE_AUTORTFM_UNUSED(ptr);
 	UE_AUTORTFM_UNUSED(size);
@@ -462,17 +469,17 @@ auto AutoRTFMLookupInstrumentedFunctorInvoker(const TFunctor& Functor) -> void(*
 // Tells if we are currently running in a transaction. This will return true in an open nest
 // (see `Open`). This function is handled specially in the compiler and will be constant folded
 // as true in closed code, or preserved as a function call in open code.
-UE_AUTORTFM_FORCEINLINE bool IsTransactional() { return autortfm_is_transactional(); }
+static UE_AUTORTFM_FORCEINLINE bool IsTransactional() { return autortfm_is_transactional(); }
 
 // Tells if we are currently running in the closed nest of a transaction. By default,
 // transactional code is in a closed nest; the only way to be in an open nest is to request it
 // via `Open`. This function is handled specially in the compiler and will be constant folded
 // as true in closed code, and false in open code.
-UE_AUTORTFM_FORCEINLINE bool IsClosed() { return autortfm_is_closed(); }
+static UE_AUTORTFM_FORCEINLINE bool IsClosed() { return autortfm_is_closed(); }
 
 // Tells us if we are currently committing or aborting a transaction. This will return true
 // in an on-abort or on-commit.
-UE_AUTORTFM_FORCEINLINE bool IsCommittingOrAborting() { return autortfm_is_committing_or_aborting(); }
+static UE_AUTORTFM_FORCEINLINE bool IsCommittingOrAborting() { return autortfm_is_committing_or_aborting(); }
 
 // Run the functor in a transaction. Memory writes and other side effects get instrumented
 // and will be reversed if the transaction aborts.
@@ -483,7 +490,7 @@ UE_AUTORTFM_FORCEINLINE bool IsCommittingOrAborting() { return autortfm_is_commi
 //
 // If AutoRTFM is disabled, the code will be ran non-transactionally.
 template<typename TFunctor>
-UE_AUTORTFM_FORCEINLINE ETransactionResult Transact(const TFunctor& Functor)
+static UE_AUTORTFM_FORCEINLINE ETransactionResult Transact(const TFunctor& Functor)
 {
 	ETransactionResult Result =
 		static_cast<ETransactionResult>(
@@ -500,7 +507,7 @@ UE_AUTORTFM_FORCEINLINE ETransactionResult Transact(const TFunctor& Functor)
 // be compiled with the AutoRTFM instrumentation of functions if the only
 // thing that's being invoked is a function in the open.
 template<typename TFunctor>
-UE_AUTORTFM_FORCEINLINE ETransactionResult TransactThenOpen(const TFunctor& Functor)
+static UE_AUTORTFM_FORCEINLINE ETransactionResult TransactThenOpen(const TFunctor& Functor)
 {
 	ETransactionResult Result =
 		static_cast<ETransactionResult>(
@@ -516,7 +523,7 @@ UE_AUTORTFM_FORCEINLINE ETransactionResult TransactThenOpen(const TFunctor& Func
 // execution if the result is anything other than autortfm_committed.
 // Useful for testing.
 template<typename TFunctor>
-UE_AUTORTFM_FORCEINLINE void Commit(const TFunctor& Functor)
+static UE_AUTORTFM_FORCEINLINE void Commit(const TFunctor& Functor)
 {
     autortfm_commit(
 		&AutoRTFMFunctorInvoker<TFunctor>,
@@ -525,7 +532,7 @@ UE_AUTORTFM_FORCEINLINE void Commit(const TFunctor& Functor)
 }
 
 // End a transaction and discard all effects.
-UE_AUTORTFM_FORCEINLINE ETransactionResult AbortTransaction()
+static UE_AUTORTFM_FORCEINLINE ETransactionResult AbortTransaction()
 {
 	return static_cast<ETransactionResult>(autortfm_abort_transaction());
 }
@@ -533,26 +540,26 @@ UE_AUTORTFM_FORCEINLINE ETransactionResult AbortTransaction()
 // End a transaction nest and discard all effects. This cascades, meaning
 // an abort of a nested transaction will cause all transactions in the
 // nest to abort.
-UE_AUTORTFM_FORCEINLINE ETransactionResult CascadingAbortTransaction()
+static UE_AUTORTFM_FORCEINLINE ETransactionResult CascadingAbortTransaction()
 {
 	return static_cast<ETransactionResult>(autortfm_cascading_abort_transaction());
 }
 
 // Abort if running in a transaction.
-UE_AUTORTFM_FORCEINLINE void AbortIfTransactional()
+static UE_AUTORTFM_FORCEINLINE void AbortIfTransactional()
 {
     autortfm_abort_if_transactional();
 }
 
 // Abort if running in closed code.
-UE_AUTORTFM_FORCEINLINE void AbortIfClosed()
+static UE_AUTORTFM_FORCEINLINE void AbortIfClosed()
 {
     autortfm_abort_if_closed();
 }
 
 // Executes the given code non-transactionally regardless of whether we are in
 // a transaction or not.
-template<typename TFunctor> UE_AUTORTFM_FORCEINLINE void Open(const TFunctor& Functor)
+template<typename TFunctor> static UE_AUTORTFM_FORCEINLINE void Open(const TFunctor& Functor)
 {
     autortfm_open(
         [] (void* Arg) { (*static_cast<const TFunctor*>(Arg))(); },
@@ -563,7 +570,7 @@ template<typename TFunctor> UE_AUTORTFM_FORCEINLINE void Open(const TFunctor& Fu
 // (whether we are in open or closed code).
 //
 // Will crash if called outside of a transaction nest.
-template<typename TFunctor> [[nodiscard]] UE_AUTORTFM_FORCEINLINE EContextStatus Close(const TFunctor& Functor)
+template<typename TFunctor> [[nodiscard]] static UE_AUTORTFM_FORCEINLINE EContextStatus Close(const TFunctor& Functor)
 {
     return static_cast<EContextStatus>(
 		autortfm_close(
@@ -577,7 +584,7 @@ template<typename TFunctor> [[nodiscard]] UE_AUTORTFM_FORCEINLINE EContextStatus
 // this just adds the work to the work deferred until the outer nest's commit.
 // If this is called outside a transaction or from an open nest then the work
 // happens immediately.
-template<typename TFunctor> UE_AUTORTFM_FORCEINLINE void OnCommit(const TFunctor& Work)
+template<typename TFunctor> static UE_AUTORTFM_FORCEINLINE void OnCommit(const TFunctor& Work)
 {
 	if (autortfm_is_closed())
 	{
@@ -593,13 +600,13 @@ template<typename TFunctor> UE_AUTORTFM_FORCEINLINE void OnCommit(const TFunctor
 // this just adds the work to the work deferred until the outer nest's commit.
 // If this is called outside a transaction or from an open nest then the work
 // happens immediately.
-template<typename TFunctor> UE_AUTORTFM_FORCEINLINE void OnCommit(const TFunctor& Work) { Work(); }
+template<typename TFunctor> static UE_AUTORTFM_FORCEINLINE void OnCommit(const TFunctor& Work) { Work(); }
 #endif
 
 #if UE_AUTORTFM
 // Have some work happen when this transaction aborts. If this is called
 // outside a transaction or from an open nest then the work is ignored.
-template<typename TFunctor> UE_AUTORTFM_FORCEINLINE void OnAbort(const TFunctor& Work)
+template<typename TFunctor> static UE_AUTORTFM_FORCEINLINE void OnAbort(const TFunctor& Work)
 {
 	if (autortfm_is_closed())
 	{
@@ -611,7 +618,7 @@ template<typename TFunctor> UE_AUTORTFM_FORCEINLINE void OnAbort(const TFunctor&
 // the handler can be unregistered (see PopOnAbortHandler). This is useful
 // for scoped mutations that need an abort handler present unless execution
 // reaches the end of the relevant scope.
-template<typename TFunctor> UE_AUTORTFM_FORCEINLINE void PushOnAbortHandler(const void* Key, const TFunctor& Work)
+template<typename TFunctor> static UE_AUTORTFM_FORCEINLINE void PushOnAbortHandler(const void* Key, const TFunctor& Work)
 {
 	if (autortfm_is_closed())
 	{
@@ -621,7 +628,7 @@ template<typename TFunctor> UE_AUTORTFM_FORCEINLINE void PushOnAbortHandler(cons
 
 // Unregister all handlers for transaction abort that were previously pushed
 // via PushOnAbortHandler with the given key
-UE_AUTORTFM_FORCEINLINE void PopOnAbortHandler(const void* Key)
+static UE_AUTORTFM_FORCEINLINE void PopOnAbortHandler(const void* Key)
 {
 	if (autortfm_is_closed())
 	{
@@ -631,17 +638,17 @@ UE_AUTORTFM_FORCEINLINE void PopOnAbortHandler(const void* Key)
 #else
 // Have some work happen when this transaction aborts. If this is called
 // outside a transaction or from an open nest then the work is ignored.
-template<typename TFunctor> UE_AUTORTFM_FORCEINLINE void OnAbort(const TFunctor&) {}
+template<typename TFunctor> static UE_AUTORTFM_FORCEINLINE void OnAbort(const TFunctor&) {}
 
 // Register a handler for transaction abort. Takes a key parameter so that
 // the handler can be unregistered (see PopOnAbortHandler). This is useful
 // for scoped mutations that need an abort handler present unless execution
 // reaches the end of the relevant scope.
-template<typename TFunctor> UE_AUTORTFM_FORCEINLINE void PushOnAbortHandler(const void* Key, const TFunctor&) {}
+template<typename TFunctor> static UE_AUTORTFM_FORCEINLINE void PushOnAbortHandler(const void* Key, const TFunctor&) {}
 
 // Unregister all handlers for transaction abort that were previously pushed
 // via PushOnAbortHandler with the given key
-UE_AUTORTFM_FORCEINLINE void PopOnAbortHandler(const void* Key) {}
+static UE_AUTORTFM_FORCEINLINE void PopOnAbortHandler(const void* Key) {}
 #endif
 
 #if UE_AUTORTFM
@@ -650,7 +657,7 @@ UE_AUTORTFM_API void OpenCommit(TFunction<void()>&& Work);
 #else
 template<typename TFunctor>
 [[deprecated("Use OnCommit instead.")]]
-UE_AUTORTFM_FORCEINLINE void OpenCommit(const TFunctor& Work) { Work(); }
+static UE_AUTORTFM_FORCEINLINE void OpenCommit(const TFunctor& Work) { Work(); }
 #endif
 
 #if UE_AUTORTFM
@@ -659,7 +666,7 @@ UE_AUTORTFM_API void OpenAbort(TFunction<void()>&& Work);
 #else
 template<typename TFunctor>
 [[deprecated("Use OnAbort instead.")]]
-UE_AUTORTFM_FORCEINLINE void OpenAbort(const TFunctor& Work) {}
+static UE_AUTORTFM_FORCEINLINE void OpenAbort(const TFunctor& Work) {}
 #endif
 
 // Inform the runtime that we have performed a new object allocation. It's only
@@ -668,13 +675,13 @@ UE_AUTORTFM_FORCEINLINE void OpenAbort(const TFunctor& Work) {}
 // returned by this function. It's guaranteed to be equal to the pointer you
 // passed, but it's blessed specially from the compiler's perspective, leading
 // to some nice optimizations. This does nothing when called from open code.
-UE_AUTORTFM_FORCEINLINE void* DidAllocate(void* Ptr, size_t Size)
+static UE_AUTORTFM_FORCEINLINE void* DidAllocate(void* Ptr, size_t Size)
 {
     return autortfm_did_allocate(Ptr, Size);
 }
 
 // Inform the runtime that we have free'd a given memory location.
-UE_AUTORTFM_FORCEINLINE void DidFree(void* Ptr)
+static UE_AUTORTFM_FORCEINLINE void DidFree(void* Ptr)
 {
     autortfm_did_free(Ptr);
 }
@@ -748,7 +755,7 @@ namespace ForTheRuntime
 	// Manually create a new transaction from open code and push it as a transaction nest.
 	// Can only be called within an already active parent transaction (EG. this cannot start
 	// a transaction nest itself).
-	UE_AUTORTFM_FORCEINLINE bool StartTransaction()
+	static UE_AUTORTFM_FORCEINLINE bool StartTransaction()
 	{
 		return autortfm_start_transaction();
 	}
@@ -756,13 +763,13 @@ namespace ForTheRuntime
 	// Manually commit the top transaction nest, popping it from the execution scope.
 	// Can only be called within an already active parent transaction (EG. this cannot end
 	// a transaction nest itself).
-	UE_AUTORTFM_FORCEINLINE ETransactionResult CommitTransaction()
+	static UE_AUTORTFM_FORCEINLINE ETransactionResult CommitTransaction()
 	{
 		return static_cast<ETransactionResult>(autortfm_commit_transaction());
 	}
 
 	// Manually clear the status of a user abort from the top transaction in a nest.
-	UE_AUTORTFM_FORCEINLINE void ClearTransactionStatus()
+	static UE_AUTORTFM_FORCEINLINE void ClearTransactionStatus()
 	{
 		autortfm_clear_transaction_status();
 	}
@@ -779,7 +786,7 @@ namespace ForTheRuntime
 	// This results in calls to ClosedVariant to happen in open mode. We will call
 	// ClosedVariant's nontransactional version within the transaction. This happens
 	// with the additional caveat that the function signatures must match.
-	UE_AUTORTFM_FORCEINLINE void RegisterOpenFunction(void* const OpenFunction, void* const ClosedVariant)
+	static UE_AUTORTFM_FORCEINLINE void RegisterOpenFunction(void* const OpenFunction, void* const ClosedVariant)
 	{
 		autortfm_register_open_function(OpenFunction, ClosedVariant);
 	}
@@ -793,32 +800,32 @@ namespace ForTheRuntime
 	};
 
 	// Manually records that the memory span was written in the current transaction.
-	UE_AUTORTFM_FORCEINLINE void RecordOpenWrite(void* Ptr, size_t Size)
+	static UE_AUTORTFM_FORCEINLINE void RecordOpenWrite(void* Ptr, size_t Size)
 	{
 		autortfm_record_open_write(Ptr, Size);
 	}
 
 	// Manually records that the memory span was written in the current transaction.
-	template<typename TTYPE> UE_AUTORTFM_FORCEINLINE void RecordOpenWrite(TTYPE* Ptr)
+	template<typename TTYPE> static UE_AUTORTFM_FORCEINLINE void RecordOpenWrite(TTYPE* Ptr)
 	{
 		autortfm_record_open_write(Ptr, sizeof(TTYPE));
 	}
 
 	// Reserved for future.
-	UE_AUTORTFM_FORCEINLINE void RecordOpenRead(void const*, size_t) {}
+	static UE_AUTORTFM_FORCEINLINE void RecordOpenRead(void const*, size_t) {}
 
 	// Reserved for future.
-	template<typename TTYPE> UE_AUTORTFM_FORCEINLINE void RecordOpenRead(TTYPE*) {}
+	template<typename TTYPE> static UE_AUTORTFM_FORCEINLINE void RecordOpenRead(TTYPE*) {}
 
 	// WriteMemory first records the memory span as written (see RecordOpenWrite) and then copies the specified value into it.
-	UE_AUTORTFM_FORCEINLINE void WriteMemory(void* DestPtr, void const* SrcPtr, size_t Size)
+	static UE_AUTORTFM_FORCEINLINE void WriteMemory(void* DestPtr, void const* SrcPtr, size_t Size)
 	{
 		RecordOpenWrite(DestPtr, Size);
 		UE_AUTORTFM_MEMCPY(DestPtr, SrcPtr, Size);
 	}
 
 	// WriteMemory first records the memory span as written (see RecordOpenWrite) and then copies the specified value into it.
-	template<typename TTYPE> UE_AUTORTFM_FORCEINLINE void WriteMemory(TTYPE* DestPtr, TTYPE const* SrcPtr)
+	template<typename TTYPE> static UE_AUTORTFM_FORCEINLINE void WriteMemory(TTYPE* DestPtr, TTYPE const* SrcPtr)
 	{
 		if constexpr (std::is_trivially_copyable<TTYPE>::value)
 		{
@@ -832,7 +839,7 @@ namespace ForTheRuntime
 	}
 
 	// WriteMemory first records the memory span as written (see RecordOpenWrite) and then copies the specified value into it.
-	template<typename TTYPE> UE_AUTORTFM_FORCEINLINE void WriteMemory(TTYPE* DestPtr, TTYPE const SrcValue)
+	template<typename TTYPE> static UE_AUTORTFM_FORCEINLINE void WriteMemory(TTYPE* DestPtr, TTYPE const SrcValue)
 	{
 		if constexpr (std::is_trivially_copyable<TTYPE>::value)
 		{
@@ -850,7 +857,7 @@ namespace ForTheRuntime
 	// expected values with the actual values in global memory. Does nothing when
 	// called outside of a transaction. May do nothing if debugging features aren't
 	// enabled in the autortfm runtime.
-	UE_AUTORTFM_FORCEINLINE void CheckConsistencyAssumingNoRaces()
+	static UE_AUTORTFM_FORCEINLINE void CheckConsistencyAssumingNoRaces()
 	{
 		autortfm_check_consistency_assuming_no_races();
 	}
@@ -866,55 +873,55 @@ enum EAutoRTFMEnabledState
 
 // Deprecated use AutoRTFM::ForTheRuntime::SetAutoRTFMRuntime instead.
 [[deprecated("Use AutoRTFM::ForTheRuntime::SetAutoRTFMRuntime instead.")]]
-UE_AUTORTFM_FORCEINLINE bool SetAutoRTFMRuntime(EAutoRTFMEnabledState bEnabled) { return ForTheRuntime::SetAutoRTFMRuntime(static_cast<ForTheRuntime::EAutoRTFMEnabledState>(bEnabled)); }
+static UE_AUTORTFM_FORCEINLINE bool SetAutoRTFMRuntime(EAutoRTFMEnabledState bEnabled) { return ForTheRuntime::SetAutoRTFMRuntime(static_cast<ForTheRuntime::EAutoRTFMEnabledState>(bEnabled)); }
 
 // Deprecated use AutoRTFM::ForTheRuntime::IsAutoRTFMRuntimeEnabled instead.
 [[deprecated("Use AutoRTFM::ForTheRuntime::IsAutoRTFMRuntimeEnabled instead.")]]
-UE_AUTORTFM_FORCEINLINE bool IsAutoRTFMRuntimeEnabled() { return ForTheRuntime::IsAutoRTFMRuntimeEnabled(); }
+static UE_AUTORTFM_FORCEINLINE bool IsAutoRTFMRuntimeEnabled() { return ForTheRuntime::IsAutoRTFMRuntimeEnabled(); }
 
 // Deprecated use AutoRTFM::ForTheRuntime::StartTransaction instead.
 [[deprecated("Use AutoRTFM::ForTheRuntime::StartTransaction instead.")]]
-UE_AUTORTFM_FORCEINLINE bool StartTransaction() { return ForTheRuntime::StartTransaction(); }
+static UE_AUTORTFM_FORCEINLINE bool StartTransaction() { return ForTheRuntime::StartTransaction(); }
 
 // Deprecated use AutoRTFM::ForTheRuntime::CommitTransaction instead.
 [[deprecated("Use AutoRTFM::ForTheRuntime::CommitTransaction instead.")]]
-UE_AUTORTFM_FORCEINLINE ETransactionResult CommitTransaction() { return ForTheRuntime::CommitTransaction(); }
+static UE_AUTORTFM_FORCEINLINE ETransactionResult CommitTransaction() { return ForTheRuntime::CommitTransaction(); }
 
 // Deprecated use AutoRTFM::ForTheRuntime::ClearTransactionStatus instead.
 [[deprecated("Use AutoRTFM::ForTheRuntime::ClearTransactionStatus instead.")]]
-UE_AUTORTFM_FORCEINLINE void ClearTransactionStatus() { return ForTheRuntime::ClearTransactionStatus(); }
+static UE_AUTORTFM_FORCEINLINE void ClearTransactionStatus() { return ForTheRuntime::ClearTransactionStatus(); }
 
 // Deprecated use AutoRTFM::ForTheRuntime::RegisterOpenFunction instead.
 [[deprecated("Use AutoRTFM::ForTheRuntime::RegisterOpenFunction instead.")]]
-UE_AUTORTFM_FORCEINLINE void RegisterOpenFunction(void* OriginalFunction, void* NewFunction) { return ForTheRuntime::RegisterOpenFunction(OriginalFunction, NewFunction); }
+static UE_AUTORTFM_FORCEINLINE void RegisterOpenFunction(void* OriginalFunction, void* NewFunction) { return ForTheRuntime::RegisterOpenFunction(OriginalFunction, NewFunction); }
 
 // Deprecated use AutoRTFM::ForTheRuntime::RecordOpenWrite instead.
 [[deprecated("Use AutoRTFM::ForTheRuntime::RecordOpenWrite instead.")]]
-UE_AUTORTFM_FORCEINLINE void RecordOpenWrite(void* Ptr, size_t Size) { ForTheRuntime::RecordOpenWrite(Ptr, Size); }
+static UE_AUTORTFM_FORCEINLINE void RecordOpenWrite(void* Ptr, size_t Size) { ForTheRuntime::RecordOpenWrite(Ptr, Size); }
 
 // Deprecated use AutoRTFM::ForTheRuntime::RecordOpenWrite instead.
-template<typename TTYPE> [[deprecated("Use AutoRTFM::ForTheRuntime::RecordOpenWrite instead.")]] UE_AUTORTFM_FORCEINLINE void RecordOpenWrite(TTYPE* Ptr) { return ForTheRuntime::RecordOpenWrite(Ptr); }
+template<typename TTYPE> [[deprecated("Use AutoRTFM::ForTheRuntime::RecordOpenWrite instead.")]] static UE_AUTORTFM_FORCEINLINE void RecordOpenWrite(TTYPE* Ptr) { return ForTheRuntime::RecordOpenWrite(Ptr); }
 
 // Deprecated use AutoRTFM::ForTheRuntime::RecordOpenRead instead.
 [[deprecated("Use AutoRTFM::ForTheRuntime::RecordOpenRead instead.")]]
-UE_AUTORTFM_FORCEINLINE void RecordOpenRead(void const* Ptr, size_t Size) { return ForTheRuntime::RecordOpenRead(Ptr, Size); }
+static UE_AUTORTFM_FORCEINLINE void RecordOpenRead(void const* Ptr, size_t Size) { return ForTheRuntime::RecordOpenRead(Ptr, Size); }
 
 // Deprecated use AutoRTFM::ForTheRuntime::RecordOpenRead instead.
-template<typename TTYPE> [[deprecated("Use AutoRTFM::ForTheRuntime::RecordOpenRead instead.")]] UE_AUTORTFM_FORCEINLINE void RecordOpenRead(TTYPE* Ptr) { return ForTheRuntime::RecordOpenRead(Ptr); }
+template<typename TTYPE> [[deprecated("Use AutoRTFM::ForTheRuntime::RecordOpenRead instead.")]] static UE_AUTORTFM_FORCEINLINE void RecordOpenRead(TTYPE* Ptr) { return ForTheRuntime::RecordOpenRead(Ptr); }
 
 // Deprecated use AutoRTFM::ForTheRuntime::WriteMemory instead.
 [[deprecated("Use AutoRTFM::ForTheRuntime::WriteMemory instead.")]]
-UE_AUTORTFM_FORCEINLINE void WriteMemory(void* D, void const* S, size_t Size) { return ForTheRuntime::WriteMemory(D, S, Size); }
+static UE_AUTORTFM_FORCEINLINE void WriteMemory(void* D, void const* S, size_t Size) { return ForTheRuntime::WriteMemory(D, S, Size); }
 
 // Deprecated use AutoRTFM::ForTheRuntime::WriteMemory instead.
-template<typename TTYPE> [[deprecated("Use AutoRTFM::ForTheRuntime::WriteMemory instead.")]] UE_AUTORTFM_FORCEINLINE void WriteMemory(TTYPE* D, TTYPE const* S) { return ForTheRuntime::WriteMemory(D, S); }
+template<typename TTYPE> [[deprecated("Use AutoRTFM::ForTheRuntime::WriteMemory instead.")]] static UE_AUTORTFM_FORCEINLINE void WriteMemory(TTYPE* D, TTYPE const* S) { return ForTheRuntime::WriteMemory(D, S); }
 
 // Deprecated use AutoRTFM::ForTheRuntime::WriteMemory instead.
-template<typename TTYPE> [[deprecated("Use AutoRTFM::ForTheRuntime::WriteMemory instead.")]] UE_AUTORTFM_FORCEINLINE void WriteMemory(TTYPE* D, TTYPE const S) { return ForTheRuntime::WriteMemory(D, S); }
+template<typename TTYPE> [[deprecated("Use AutoRTFM::ForTheRuntime::WriteMemory instead.")]] static UE_AUTORTFM_FORCEINLINE void WriteMemory(TTYPE* D, TTYPE const S) { return ForTheRuntime::WriteMemory(D, S); }
 
 // Deprecated use AutoRTFM::ForTheRuntime::CheckConsistencyAssumingNoRaces instead.
 [[deprecated("Use AutoRTFM::ForTheRuntime::CheckConsistencyAssumingNoRaces instead.")]]
-UE_AUTORTFM_FORCEINLINE void CheckConsistencyAssumingNoRaces() { return ForTheRuntime::CheckConsistencyAssumingNoRaces(); }
+static UE_AUTORTFM_FORCEINLINE void CheckConsistencyAssumingNoRaces() { return ForTheRuntime::CheckConsistencyAssumingNoRaces(); }
 
 } // namespace AutoRTFM
 
