- Always remove empty configs. It's possible for titles to have many configs and since empty config instances are not free, keeping around empty configs can result in non-insignificant memory waste on memory constrained platforms

#rnx

diff --git a/Engine/Source/Runtime/Core/Private/Misc/ConfigContext.cpp b/Engine/Source/Runtime/Core/Private/Misc/ConfigContext.cpp
--- a/Engine/Source/Runtime/Core/Private/Misc/ConfigContext.cpp
+++ b/Engine/Source/Runtime/Core/Private/Misc/ConfigContext.cpp
@@ -282,12 +282,8 @@ bool FConfigContext::Load(const TCHAR* InBaseIniName, FString& OutFinalFilename)
 		// delete the branch if it is useless
 		if (!bSuccess && ConfigSystem != nullptr && ExistingFile == nullptr && TemporaryBranch == nullptr)
 		{
-			static bool bRemoveEmptyPlugins = FParse::Param(FCommandLine::Get(), TEXT("RemoveEmptyConfigs"));
-			if (bRemoveEmptyPlugins)
-			{
-				ConfigSystem->Remove(DestIniFilename);
-				Branch = nullptr;
-			}
+			ConfigSystem->Remove(DestIniFilename);
+			Branch = nullptr;
 		}
 	}
 	return bSuccess;
