Added FindLastFrom to TBitArray
#tests FoundationTests

diff --git a/Engine/Source/Runtime/Core/Public/Containers/BitArray.h b/Engine/Source/Runtime/Core/Public/Containers/BitArray.h
--- a/Engine/Source/Runtime/Core/Public/Containers/BitArray.h
+++ b/Engine/Source/Runtime/Core/Public/Containers/BitArray.h
@@ -501,9 +501,9 @@ private:
 		return FBitSet::CalculateNumWords(MaxBits);
 	}
 
-	FORCEINLINE uint32 GetLastWordMask() const
+	FORCEINLINE uint32 GetLastWordMask(int32 EndIndexExclusive) const
 	{
-		const uint32 UnusedBits = (FBitSet::BitsPerWord - static_cast<uint32>(NumBits) % FBitSet::BitsPerWord) % FBitSet::BitsPerWord;
+		const uint32 UnusedBits = (FBitSet::BitsPerWord - static_cast<uint32>(EndIndexExclusive) % FBitSet::BitsPerWord) % FBitSet::BitsPerWord;
 		return ~0u >> UnusedBits;
 	}
 
@@ -1083,8 +1083,6 @@ private:
 
 	int32 FindFromImpl(bool bValue, int32 StartIndex) const
 	{
-		checkSlow((StartIndex >= 0) & (StartIndex <= NumBits));
-
 		// Produce a mask for the first iteration
 		uint32 Mask = ~0u << (StartIndex % FBitSet::BitsPerWord);
 
@@ -1116,6 +1114,45 @@ private:
 		return INDEX_NONE;
 	}
 
+	int32 FindLastFromImpl(bool bValue, int32 EndIndexExclusive) const
+	{
+		// We need to early out on empty bit arrays in order to avoid iterating on the first (inexistent) DWord :
+		if (NumBits == 0)
+		{
+			return INDEX_NONE;
+		}
+
+		// Produce a mask for the first iteration
+		uint32 Mask = GetLastWordMask(EndIndexExclusive);
+
+		// Iterate over the array until we see a word with a matching bit
+		const uint32* RESTRICT DwordArray = GetData();
+		uint32 DwordIndex = FBitSet::CalculateNumWords(EndIndexExclusive);
+		const uint32 Test = bValue ? 0u : ~0u;
+		for (;;)
+		{
+			if (DwordIndex == 0)
+			{
+				return INDEX_NONE;
+			}
+			--DwordIndex;
+			if ((DwordArray[DwordIndex] & Mask) != (Test & Mask))
+			{
+				break;
+			}
+			Mask = ~0u;
+		}
+
+		// If we're looking for a false, then we flip the bits - then we only need to find the last one bit
+		const uint32 Bits = (bValue ? DwordArray[DwordIndex] : ~DwordArray[DwordIndex]) & Mask;
+		UE_ASSUME(Bits != 0);
+
+		uint32 BitIndex = (NumBitsPerDWORD - 1) - FMath::CountLeadingZeros(Bits);
+
+		int32 Result = BitIndex + (DwordIndex << NumBitsPerDWORDLogTwo);
+		return Result;
+	}
+
 public:
 
 	/**
@@ -1133,6 +1170,7 @@ public:
 	FORCEINLINE int32 FindFrom(bool bValue, IndexType StartIndex) const
 	{
 		static_assert(!std::is_same_v<IndexType, bool>, "TBitArray::FindFrom: unexpected bool passed as the StartIndex argument");
+		checkSlow((StartIndex >= 0) && (StartIndex <= NumBits));
 		return FindFromImpl(bValue, StartIndex);
 	}
 
@@ -1146,37 +1184,27 @@ public:
 	 */
 	int32 FindLast(bool bValue) const 
 	{
-		const int32 LocalNumBits = NumBits;
-
-		// Get the correct mask for the last word
-		uint32 Mask = GetLastWordMask();
-
-		// Iterate over the array until we see a word with a zero bit.
-		uint32 DwordIndex = FBitSet::CalculateNumWords(LocalNumBits);
-		const uint32* RESTRICT DwordArray = GetData();
-		const uint32 Test = bValue ? 0u : ~0u;
-		for (;;)
-		{
-			if (DwordIndex == 0)
-			{
-				return INDEX_NONE;
-			}
-			--DwordIndex;
-			if ((DwordArray[DwordIndex] & Mask) != (Test & Mask))
-			{
-				break;
-			}
-			Mask = ~0u;
-		}
-
-		// Flip the bits, then we only need to find the first one bit -- easy.
-		const uint32 Bits = (bValue ? DwordArray[DwordIndex] : ~DwordArray[DwordIndex]) & Mask;
-		UE_ASSUME(Bits != 0);
-
-		uint32 BitIndex = (NumBitsPerDWORD - 1) - FMath::CountLeadingZeros(Bits);
+		return FindLastFromImpl(bValue, /*EndIndexExclusive = */NumBits);
+	}
 
-		int32 Result = BitIndex + (DwordIndex << NumBitsPerDWORDLogTwo);
-		return Result;
+	/**
+	 * Finds the last occurrence of the specified value (true/false) in the array starting the (reverse) search from the given bit index, and returns the bit index.
+	 * If the specified value is not found, INDEX_NONE is returned.
+	 *
+	 * @param  bValue  The value (true/false) to search for.
+	 * @param  EndIndexInclusive  The index to start the (reverse) search from.
+	 *
+	 * @return The index of the last occurrence of the specified value (true/false) from EndIndexInclusive, or INDEX_NONE if not found.
+	 */
+	template <typename IndexType>
+	int32 FindLastFrom(bool bValue, IndexType EndIndexInclusive) const
+	{
+		static_assert(!std::is_same_v<IndexType, bool>, "TBitArray::FindLastFrom: unexpected bool passed as the EndIndexInclusive argument");
+		// This range check here ([-1, NumBits)) is because the end index is inclusive and because this is the reverse logic of the FindFrom range ([0, NumBits]) 
+		//  It doesn't make much sense to allow EndIndexInclusive to be -1 in FindLastFrom (since that will always fail), much like it doesn't make sense to allow
+		//  StartIndex to be NumBits in FindFrom (for the same reason) but this keeps the implementations consistent
+		checkSlow((EndIndexInclusive >= -1) && (EndIndexInclusive < NumBits)); 
+		return FindLastFromImpl(bValue, /*EndIndexExclusive = */EndIndexInclusive + 1);
 	}
 
 	/**
