[AutoRTFM] Fix int32 integer overflow in HitSet

The assert in IncreaseCapacity() was not catching the integer overflow because it was testing for 64-bit overflow, when the size types were all 32 bit. Make these 64 bit to match the SixtyFourMinusLogCapacity logic.

diff --git a/Engine/Source/Runtime/Core/Private/AutoRTFM/HitSet.h b/Engine/Source/Runtime/Core/Private/AutoRTFM/HitSet.h
--- a/Engine/Source/Runtime/Core/Private/AutoRTFM/HitSet.h
+++ b/Engine/Source/Runtime/Core/Private/AutoRTFM/HitSet.h
@@ -10,24 +10,24 @@ namespace AutoRTFM
 	class FHitSet final
 	{
 		// TODO: Revisit a good probe depth for the hashset.
-		static constexpr uint32_t LinearProbeDepth = 16;
+		static constexpr uint64_t LinearProbeDepth = 16;
 
 		// TODO: Revisit a good initial capacity for the hitset.
-		static constexpr uint32_t LogInitialCapacity = 4;
+		static constexpr uint64_t LogInitialCapacity = 4;
 
 	public:
 		using Key = TTaggedPtr<void>;
 
 		explicit FHitSet()
 		{
-			constexpr uint32_t InitialCapacity = 1u << LogInitialCapacity;
+			constexpr uint64_t InitialCapacity = 1u << LogInitialCapacity;
 
 			// Do not want to deal with null payloads.
 			static_assert(0 != InitialCapacity);
 
 			// The capacity is always a power of two so that the range reduction is optimal.
 			static_assert(0 == (InitialCapacity & (InitialCapacity - 1)));
-			
+
 			Payload = static_cast<uintptr_t*>(FMemory::MallocZeroed(InitialCapacity * sizeof(uintptr_t)));
 			ASSERT(nullptr != Payload);
 
@@ -73,12 +73,12 @@ namespace AutoRTFM
 			Size = 0;
 		}
 
-		uint32_t GetCapacity() const
+		uint64_t GetCapacity() const
 		{
 			return Capacity();
 		}
 
-		uint32_t GetSize() const
+		uint64_t GetSize() const
 		{
 			return Size;
 		}
@@ -86,13 +86,13 @@ namespace AutoRTFM
 	private:
 		uintptr_t* Payload;
 
-		uint32_t SixtyFourMinusLogCapacity;
+		uint64_t SixtyFourMinusLogCapacity;
 
-		uint32_t Size;
+		uint64_t Size;
 
-		UE_AUTORTFM_FORCEINLINE uint32_t Capacity() const
+		UE_AUTORTFM_FORCEINLINE uint64_t Capacity() const
 		{
-			return 1u << (64 - SixtyFourMinusLogCapacity);
+			return static_cast<uint64_t>(1) << (64 - SixtyFourMinusLogCapacity);
 		}
 
 		UE_AUTORTFM_FORCEINLINE void IncreaseCapacity()
@@ -108,8 +108,8 @@ namespace AutoRTFM
 		void Resize()
 		{
 			uintptr_t* const OldPayload = Payload;
-			const uint32_t OldCapacity = Capacity();
-			const uint32_t OldSize = Size;
+			const uint64_t OldCapacity = Capacity();
+			const uint64_t OldSize = Size;
 
 			while (true)
 			{
@@ -232,7 +232,7 @@ namespace AutoRTFM
 #ifdef __clang__
 #pragma clang loop unroll(disable)
 #endif
-			for (uint32_t D = 0; D < LinearProbeDepth; D++)
+			for (uint64_t D = 0; D < LinearProbeDepth; D++)
 			{
 				const uintptr_t I = FirstHashInRange(Hash + D);
 
