RHI: Default the NullRHIShaderFormat to SM6 unless specified

#rnx

diff --git a/Engine/Source/Runtime/Core/Private/GenericPlatform/GenericPlatformMisc.cpp b/Engine/Source/Runtime/Core/Private/GenericPlatform/GenericPlatformMisc.cpp
--- a/Engine/Source/Runtime/Core/Private/GenericPlatform/GenericPlatformMisc.cpp
+++ b/Engine/Source/Runtime/Core/Private/GenericPlatform/GenericPlatformMisc.cpp
@@ -1231,7 +1231,14 @@ const TCHAR* FGenericPlatformMisc::LaunchDir()
 
 const TCHAR* FGenericPlatformMisc::GetNullRHIShaderFormat()
 {
-	return TEXT("PCD3D_SM5");
+	if (FParse::Param(FCommandLine::Get(), TEXT("sm5")))
+	{
+		return TEXT("PCD3D_SM5");
+	}
+	else
+	{
+		return TEXT("PCD3D_SM6");
+	}
 }
 
 IPlatformChunkInstall* FGenericPlatformMisc::GetPlatformChunkInstall()
diff --git a/Engine/Source/Runtime/Core/Private/Unix/UnixPlatformMisc.cpp b/Engine/Source/Runtime/Core/Private/Unix/UnixPlatformMisc.cpp
--- a/Engine/Source/Runtime/Core/Private/Unix/UnixPlatformMisc.cpp
+++ b/Engine/Source/Runtime/Core/Private/Unix/UnixPlatformMisc.cpp
@@ -763,7 +763,14 @@ int32 FUnixPlatformMisc::NumberOfCoresIncludingHyperthreads()
 
 const TCHAR* FUnixPlatformMisc::GetNullRHIShaderFormat()
 {
-	return TEXT("SF_VULKAN_SM5");
+	if (FParse::Param(FCommandLine::Get(), TEXT("sm5")))
+	{
+		return TEXT("SF_VULKAN_SM5");
+	}
+	else
+	{
+		return TEXT("SF_VULKAN_SM6");
+	}
 }
 
 #define CPUINFO_TOKENS                            \
diff --git a/Engine/Source/Runtime/Core/Public/Mac/MacPlatformMisc.h b/Engine/Source/Runtime/Core/Public/Mac/MacPlatformMisc.h
--- a/Engine/Source/Runtime/Core/Public/Mac/MacPlatformMisc.h
+++ b/Engine/Source/Runtime/Core/Public/Mac/MacPlatformMisc.h
@@ -54,7 +54,14 @@ struct CORE_API FMacPlatformMisc : public FApplePlatformMisc
 	 */
 	FORCEINLINE static const TCHAR* GetNullRHIShaderFormat() 
 	{ 
-		return TEXT("SF_METAL_SM5"); 
+		if (FParse::Param(FCommandLine::Get(), TEXT("sm5")))
+		{
+			return TEXT("SF_METAL_SM5");
+		}
+		else
+		{
+			return TEXT("SF_METAL_SM6");
+		}
 	}
 
 	/**
