PR #10381: Fix potential FIntRect operator<< ambiguous compilation error.

#rnx

- This might have been required before the FArchiveBase refactor several years ago but I cannot think of any reason we need to keep this code so deleting it is probably better than updating it to template<typename> TIntRect

diff --git a/Engine/Source/Runtime/Core/Public/Serialization/Archive.h b/Engine/Source/Runtime/Core/Public/Serialization/Archive.h
--- a/Engine/Source/Runtime/Core/Public/Serialization/Archive.h
+++ b/Engine/Source/Runtime/Core/Public/Serialization/Archive.h
@@ -1563,14 +1563,6 @@ public:
 		return Ar;
 	}
 
-	/**
-	 * Serializes an FIntRect value from or into an archive.
-	 *
-	 * @param Ar The archive to serialize from or to.
-	 * @param Value The value to serialize.
-	 */
-	friend FArchive& operator<<(FArchive& Ar, FIntRect& Value);
-
 	/**
 	 * Serializes an FString value from or into an archive.
 	 *
