Mac: Fix build with missing header.

#rnx

diff --git a/Engine/Source/Runtime/Core/Private/Mac/MacPlatformMisc.cpp b/Engine/Source/Runtime/Core/Private/Mac/MacPlatformMisc.cpp
--- a/Engine/Source/Runtime/Core/Private/Mac/MacPlatformMisc.cpp
+++ b/Engine/Source/Runtime/Core/Private/Mac/MacPlatformMisc.cpp
@@ -779,6 +779,18 @@ bool FMacPlatformMisc::CommandLineCommands()
 	return HandleFirstInstall();
 }
 
+const TCHAR* FMacPlatformMisc::GetNullRHIShaderFormat()
+{
+	if (FParse::Param(FCommandLine::Get(), TEXT("sm5")))
+	{
+		return TEXT("SF_METAL_SM5");
+	}
+	else
+	{
+		return TEXT("SF_METAL_SM6");
+	}
+}
+
 int32 FMacPlatformMisc::NumberOfCores()
 {	
 	static int32 NumberOfCores = -1;
diff --git a/Engine/Source/Runtime/Core/Public/Mac/MacPlatformMisc.h b/Engine/Source/Runtime/Core/Public/Mac/MacPlatformMisc.h
--- a/Engine/Source/Runtime/Core/Public/Mac/MacPlatformMisc.h
+++ b/Engine/Source/Runtime/Core/Public/Mac/MacPlatformMisc.h
@@ -52,17 +52,7 @@ struct CORE_API FMacPlatformMisc : public FApplePlatformMisc
 	 *
 	 * @return	Returns the shader format to be used by that platform
 	 */
-	FORCEINLINE static const TCHAR* GetNullRHIShaderFormat() 
-	{ 
-		if (FParse::Param(FCommandLine::Get(), TEXT("sm5")))
-		{
-			return TEXT("SF_METAL_SM5");
-		}
-		else
-		{
-			return TEXT("SF_METAL_SM6");
-		}
-	}
+	static CORE_API const TCHAR* GetNullRHIShaderFormat();
 
 	/**
 	 * Uses cpuid instruction to get the vendor string
