- Added GConffg->RemoveSectionFromBranch() that will destructively remove a section from the entirety of a branch, which can be useful for large sections that are read in and them never needed again. The will not be reloaded without fully reloading the entire branch (unlike SafeUnloadBranch)
- Added "config removesection <branch> <section>" console command for testing purposes

diff --git a/Engine/Source/Runtime/Core/Private/Misc/ConfigCacheIni.cpp b/Engine/Source/Runtime/Core/Private/Misc/ConfigCacheIni.cpp
--- a/Engine/Source/Runtime/Core/Private/Misc/ConfigCacheIni.cpp
+++ b/Engine/Source/Runtime/Core/Private/Misc/ConfigCacheIni.cpp
@@ -3286,6 +3286,28 @@ bool FConfigBranch::SafeUnload()
 	return true;
 }
 
+bool FConfigBranch::RemoveSection(const TCHAR* Section)
+{
+	int NumRemoved = 0;
+	
+	FString SectionName(Section);
+	for (TPair<FString, FConfigCommandStream>& Pair : StaticLayers)
+	{
+		NumRemoved += Pair.Value.Remove(SectionName);
+	}
+	for (DynamicLayerList::TIterator Node(DynamicLayers.GetHead()); Node; ++Node)
+	{
+		NumRemoved += Node->Remove(SectionName);
+	}
+
+	NumRemoved += InMemoryFile.Remove(SectionName);
+	NumRemoved += CombinedStaticLayers.Remove(SectionName);
+	NumRemoved += SavedLayer.Remove(SectionName);
+	NumRemoved += CommandLineOverrides.Remove(SectionName);
+	NumRemoved += FinalCombinedLayers.Remove(SectionName);
+
+	return NumRemoved > 0;
+}
 
 void FConfigBranch::Flush()
 {
@@ -4031,7 +4053,16 @@ bool FConfigCacheIni::SafeUnloadBranch(const TCHAR* BranchName)
 	return false;
 }
 
+bool FConfigCacheIni::RemoveSectionFromBranch(const TCHAR* Section, const TCHAR* Filename)
+{
+	FConfigBranch* Branch = FindBranchWithNoReload(Filename, Filename);
+	if (Branch)
+	{
+		return Branch->RemoveSection(Section);
+	}
 
+	return false;	
+}
 
 bool FConfigCacheIni::EmptySection( const TCHAR* Section, const FString& Filename )
 {
@@ -6245,6 +6276,30 @@ class FIniExec : public FSelfRegisteringExec
 			});
 		}
 		
+		if (FParse::Command(&Cmd, TEXT("RemoveSection")))
+		{
+			TCHAR BranchName[256];
+			TCHAR Section[256];
+
+			if (FParse::Token(Cmd, BranchName, UE_ARRAY_COUNT(BranchName), true) &&
+				FParse::Token(Cmd, Section, UE_ARRAY_COUNT(Section), true))
+			{
+				bool bRemovedSomething = GConfig->RemoveSectionFromBranch(Section, BranchName);
+				if (bRemovedSomething)
+				{
+					Ar.Logf(TEXT("Successfully removed '%s' from layer(s) in %s"), Section, BranchName);						
+				}
+				else
+				{
+					Ar.Logf(TEXT("Nothing was removed from %s (either branch wasn't found or the section '%s' wasn't)"), BranchName, Section);
+				}
+			}
+			else
+			{
+				Ar.Logf(TEXT("Usage: config RemoveSection <BranchName> <Section>"));	
+			}
+		}
+		
 		if (FParse::Command(&Cmd, TEXT("Timing")))
 		{
 			Ar.Logf(TEXT("INITIME : PrepareForLoad: %fms, PreformLoad: %fms"), GPrepareForLoadTime * 1000.0, GPerformLoadTime * 1000.0);
diff --git a/Engine/Source/Runtime/Core/Public/Misc/ConfigCacheIni.h b/Engine/Source/Runtime/Core/Public/Misc/ConfigCacheIni.h
--- a/Engine/Source/Runtime/Core/Public/Misc/ConfigCacheIni.h
+++ b/Engine/Source/Runtime/Core/Public/Misc/ConfigCacheIni.h
@@ -1020,6 +1020,14 @@ public:
 	 * if something does try to access it, it will reload in-place (can hitch your game, so be aware)
 	 */
 	CORE_API bool SafeUnload();
+	
+	/**
+	  * Removes the section completely from all layers of this branch 
+	  * This is destructive! It will not reload on demand (like SafeUnload does). Use this only when you know you will _never_ need
+	  * the values in the section again!
+	  * @return true if anything was removed
+	 */
+	CORE_API bool RemoveSection(const TCHAR* Section);
 
 	CORE_API void Flush();
 	
@@ -1549,6 +1557,13 @@ public:
 	 */
 	CORE_API bool SafeUnloadBranch(const TCHAR* Filename);
 
+	/**
+	  * Removes the section completely from all layers of a branch. 
+	  * This is destructive! It will not reload on demand (like SafeUnloadBranch does). Use this only when you know you will _never_ need
+	  * the values in the section again!
+	  * @return true if anything was removed
+	 */
+	CORE_API bool RemoveSectionFromBranch(const TCHAR* Section, const TCHAR* Filename);
 
 	// Static helper functions
 
