[Modules] Provide more details in assertion of LoadModuleChecked().

ShaderCompileWorker loads many modules and some crash reports only show what module could not be loaded but without details why it failed.
This change adds more details to the assertion in LoadModuleChecked(), to show whether the module was loaded from the Game thread and the failure status.

#rnx

diff --git a/Engine/Source/Runtime/Core/Private/Modules/ModuleManager.cpp b/Engine/Source/Runtime/Core/Private/Modules/ModuleManager.cpp
--- a/Engine/Source/Runtime/Core/Private/Modules/ModuleManager.cpp
+++ b/Engine/Source/Runtime/Core/Private/Modules/ModuleManager.cpp
@@ -465,17 +465,27 @@ void FModuleManager::RefreshModuleFilenameFromManifest(const FName InModuleName)
 #endif	// !IS_MONOLITHIC
 
 IModuleInterface* FModuleManager::LoadModule(const FName InModuleName, ELoadModuleFlags InLoadModuleFlags)
+{
+	EModuleLoadResult FailureReason = EModuleLoadResult::Success;
+	return GetOrLoadModule(InModuleName, FailureReason, InLoadModuleFlags);
+}
+
+IModuleInterface* FModuleManager::GetOrLoadModule(const FName InModuleName, EModuleLoadResult& OutFailureReason, ELoadModuleFlags InLoadModuleFlags)
 {
 	LLM_SCOPE_BYNAME(TEXT("Modules"));
 	// We allow an already loaded module to be returned in other threads to simplify
 	// parallel processing scenarios but they must have been loaded from the main thread beforehand.
 	if(!IsInGameThread())
 	{
-		return GetModule(InModuleName);
+		IModuleInterface* Module = GetModule(InModuleName);
+		if (!Module)
+		{
+			OutFailureReason = EModuleLoadResult::NotLoadedByGameThread;
+		}
+		return Module;
 	}
 
-	EModuleLoadResult FailureReason;
-	IModuleInterface* Result = LoadModuleWithFailureReason(InModuleName, FailureReason, InLoadModuleFlags);
+	IModuleInterface* Result = LoadModuleWithFailureReason(InModuleName, OutFailureReason, InLoadModuleFlags);
 
 	// This should return a valid pointer only if and only if the module is loaded
 	checkSlow((Result != nullptr) == IsModuleLoaded(InModuleName));
@@ -483,11 +493,29 @@ IModuleInterface* FModuleManager::LoadModule(const FName InModuleName, ELoadModu
 	return Result;
 }
 
+static const TCHAR* LexToString(EModuleLoadResult LoadResult)
+{
+	switch (LoadResult)
+	{
+	case EModuleLoadResult::Success:				return TEXT("Success");
+	case EModuleLoadResult::FileNotFound:			return TEXT("FileNotFound");
+	case EModuleLoadResult::FileIncompatible:		return TEXT("FileIncompatible");
+	case EModuleLoadResult::CouldNotBeLoadedByOS:	return TEXT("CouldNotBeLoadedByOS");
+	case EModuleLoadResult::FailedToInitialize:		return TEXT("FailedToInitialize");
+	case EModuleLoadResult::NotLoadedByGameThread:	return TEXT("NotLoadedByGameThread");
+	default:										return TEXT("<Unknown>");
+	}
+}
 
 IModuleInterface& FModuleManager::LoadModuleChecked( const FName InModuleName )
 {
-	IModuleInterface* Module = LoadModule(InModuleName, ELoadModuleFlags::LogFailures);
-	checkf(Module, TEXT("%s"), *InModuleName.ToString());
+	EModuleLoadResult FailureReason = EModuleLoadResult::Success;
+	IModuleInterface* Module = GetOrLoadModule(InModuleName, FailureReason, ELoadModuleFlags::LogFailures);
+
+	checkf(Module, TEXT("ModuleName=%s, Failure=%s, IsInGameThread=%s"),
+		*InModuleName.ToString(),
+		LexToString(FailureReason),
+		IsInGameThread() ? TEXT("Yes") : TEXT("No"));
 
 	return *Module;
 }
diff --git a/Engine/Source/Runtime/Core/Public/Modules/ModuleManager.h b/Engine/Source/Runtime/Core/Public/Modules/ModuleManager.h
--- a/Engine/Source/Runtime/Core/Public/Modules/ModuleManager.h
+++ b/Engine/Source/Runtime/Core/Public/Modules/ModuleManager.h
@@ -60,7 +60,10 @@ enum class EModuleLoadResult
 	CouldNotBeLoadedByOS,
 
 	/** Module initialization failed. */
-	FailedToInitialize
+	FailedToInitialize,
+
+	/** A thread attempted to load the module before the Game thread did. */
+	NotLoadedByGameThread
 };
 
 /**
@@ -337,6 +340,8 @@ public:
 private:
 	static CORE_API IModuleInterface* GetModulePtr_Internal(FName ModuleName);
 
+	IModuleInterface* GetOrLoadModule(const FName InModuleName, EModuleLoadResult& OutFailureReason, ELoadModuleFlags InLoadModuleFlags = ELoadModuleFlags::None);
+
 public:
 
 	/**
