Optimizing MB2 bins to reduce waste

diff --git a/Engine/Source/Runtime/Core/Private/HAL/MallocBinned2.cpp b/Engine/Source/Runtime/Core/Private/HAL/MallocBinned2.cpp
--- a/Engine/Source/Runtime/Core/Private/HAL/MallocBinned2.cpp
+++ b/Engine/Source/Runtime/Core/Private/HAL/MallocBinned2.cpp
@@ -106,12 +106,41 @@ PRAGMA_DISABLE_UNSAFE_TYPECAST_WARNINGS
 // They must be 16-byte aligned as well.
 static constexpr uint16 SmallBinSizes[] =
 {
-	16, 32, 48, 64, 80, 96, 128,
-	160, 192, 224, 256, 288, 320, 384, 448,
-	512, 576, 640, 704, 768, 896, 1024 - 16, 1168,
-	1488, 1632, 2048 - 16, 2336, 2720, 3264, 4096 - 16, 4368,
-	5040, 5456, 5952, 6544 - 16, 7280, 8192 - 16, 9360,
-	10912, 13104, 16384 - 16, 19104, 21840, 27024, 32768 - 16
+	16, 32, 48, 64, 80, 96, 112, 128, 144, 160, 176, 192, 208,// +16
+	256, 288, 320, 384, 448,
+	512,	// /128
+	560,	// /117
+	624,	// /105
+	720,	// /91
+	816,	// /80 31
+	912,	// /71 86
+	1024 - 16,	// /64
+	1168,	// /56 11
+	1392,	// /47 8
+	1520,	// /43 11
+	1680,	// /39
+	1872,	// /35
+	2048 - 16,	// /32
+	2256,	// /29 5
+	2608,	// /25 13
+	2976,	// /22
+	3264,	// /20 8
+	3632,	// /18 4
+	4096 - 16,	// /16
+	4368,	// /15
+	4672,	// /14
+	5040,	// /13
+	5456,	// /12
+	5952,	// /11
+	6544,	// /10
+	7280,	// /9
+	8192 - 16,	// /8
+	9360,	// /7
+	10912,	// /6
+	13104,	// /5
+	16384 - 16,	// /4
+	21840,	// /3
+	32768 - 16	// /2
 };
 
 MS_ALIGN(PLATFORM_CACHE_LINE_SIZE) static uint8 UnusedAlignPadding[PLATFORM_CACHE_LINE_SIZE] GCC_ALIGN(PLATFORM_CACHE_LINE_SIZE) = { 0 };
diff --git a/Engine/Source/Runtime/Core/Public/HAL/MallocBinned2.h b/Engine/Source/Runtime/Core/Public/HAL/MallocBinned2.h
--- a/Engine/Source/Runtime/Core/Public/HAL/MallocBinned2.h
+++ b/Engine/Source/Runtime/Core/Public/HAL/MallocBinned2.h
@@ -29,7 +29,7 @@ struct FGenericMemoryStats;
 #define UE_MB2_MINIMUM_ALIGNMENT			16			// Alignment of bins
 #define UE_MB2_MAXIMUM_ALIGNMENT			128
 #define UE_MB2_MAX_SMALL_POOL_SIZE			(32768-16)	// Maximum bin size in SmallBinSizes in cpp file
-#define UE_MB2_SMALL_POOL_COUNT				45
+#define UE_MB2_SMALL_POOL_COUNT				51
 
 
 #define UE_DEFAULT_GBinned2PerThreadCaches						1
