Add callsite always_inline to Open and OnCommit in the AutoRTFM runtime.

diff --git a/Engine/Source/Runtime/Core/Public/AutoRTFM/AutoRTFM.h b/Engine/Source/Runtime/Core/Public/AutoRTFM/AutoRTFM.h
--- a/Engine/Source/Runtime/Core/Public/AutoRTFM/AutoRTFM.h
+++ b/Engine/Source/Runtime/Core/Public/AutoRTFM/AutoRTFM.h
@@ -34,10 +34,12 @@
 #define UE_AUTORTFM_NOAUTORTFM [[clang::noautortfm, clang::noinline]]
 // #jira SOL-6589: remove clang::noinline once this JIRA is fixed.
 #define UE_AUTORTFM_ALWAYS_OPEN [[clang::autortfm_always_open, clang::noinline]]
+#define UE_AUTORTFM_CALLSITE_FORCEINLINE [[clang::always_inline]]
 #else
 #define UE_AUTORTFM_AUTORTFM(F)
 #define UE_AUTORTFM_NOAUTORTFM
 #define UE_AUTORTFM_ALWAYS_OPEN
+#define UE_AUTORTFM_CALLSITE_FORCEINLINE
 #endif
 
 #if UE_AUTORTFM && UE_AUTORTFM_STATIC_VERIFIER
@@ -555,7 +557,7 @@ UE_AUTORTFM_FORCEINLINE void AbortIfClosed()
 template<typename TFunctor> UE_AUTORTFM_FORCEINLINE void Open(const TFunctor& Functor)
 {
     autortfm_open(
-        [] (void* Arg) { (*static_cast<const TFunctor*>(Arg))(); },
+        [] (void* Arg) { UE_AUTORTFM_CALLSITE_FORCEINLINE (*static_cast<const TFunctor*>(Arg))(); },
         const_cast<void*>(static_cast<const void*>(&Functor)));
 }
 
@@ -585,7 +587,7 @@ template<typename TFunctor> UE_AUTORTFM_FORCEINLINE void OnCommit(const TFunctor
 	}
 	else
 	{
-		Work();
+		UE_AUTORTFM_CALLSITE_FORCEINLINE Work();
 	}
 }
 #else
