FPackageId - Add support for appending the value (and optionally name) to a TStringBuilderBase instance.

diff --git a/Engine/Source/Runtime/Core/Public/IO/PackageId.h b/Engine/Source/Runtime/Core/Public/IO/PackageId.h
--- a/Engine/Source/Runtime/Core/Public/IO/PackageId.h
+++ b/Engine/Source/Runtime/Core/Public/IO/PackageId.h
@@ -81,3 +81,12 @@ public:
 
 CORE_API FString LexToString(const FPackageId& PackageId);
 
+template <typename CharType>
+TStringBuilderBase<CharType>& operator<<(TStringBuilderBase<CharType>& Builder, const FPackageId& PackageId)
+{
+	Builder.Appendf(CHARTEXT(CharType, "0x%llX"), PackageId.Value());
+#if WITH_PACKAGEID_NAME_MAP
+	Builder << " (" << PackageId.GetName() << ")";
+#endif // WITH_PACKAGEID_NAME_MAP
+	return Builder;
+}
