Cast to void* when calling MemoryOps functions to help TArray work with const elements.

diff --git a/Engine/Source/Runtime/Core/Public/Containers/Array.h b/Engine/Source/Runtime/Core/Public/Containers/Array.h
--- a/Engine/Source/Runtime/Core/Public/Containers/Array.h
+++ b/Engine/Source/Runtime/Core/Public/Containers/Array.h
@@ -1410,7 +1410,7 @@ private:
 			SlackTrackerNumChanged();
 		}
 		ElementType* Data = GetData() + Index;
-		RelocateConstructItems<ElementType>(Data + 1, Data, OldNum - Index);
+		RelocateConstructItems<ElementType>((void*)(Data + 1), Data, OldNum - Index);
 	}
 	template <typename OtherSizeType>
 	void InsertUninitializedImpl(SizeType Index, OtherSizeType Count)
@@ -1440,7 +1440,7 @@ private:
 			SlackTrackerNumChanged();
 		}
 		ElementType* Data = GetData() + Index;
-		RelocateConstructItems<ElementType>(Data + Count, Data, OldNum - Index);
+		RelocateConstructItems<ElementType>((void*)(Data + Count), Data, OldNum - Index);
 	}
 
 public:
@@ -1518,12 +1518,12 @@ public:
 	void InsertDefaulted(SizeType Index)
 	{
 		InsertUninitializedImpl(Index);
-		DefaultConstructItems<ElementType>(GetData() + Index, 1);
+		DefaultConstructItems<ElementType>((void*)(GetData() + Index), 1);
 	}
 	void InsertDefaulted(SizeType Index, SizeType Count)
 	{
 		InsertUninitializedImpl(Index, Count);
-		DefaultConstructItems<ElementType>(GetData() + Index, Count);
+		DefaultConstructItems<ElementType>((void*)(GetData() + Index), Count);
 	}
 
 	/**
@@ -1538,7 +1538,7 @@ public:
 	{
 		InsertUninitializedImpl(Index, 1);
 		ElementType* Ptr = GetData() + Index;
-		DefaultConstructItems<ElementType>(Ptr, 1);
+		DefaultConstructItems<ElementType>((void*)Ptr, 1);
 		return *Ptr;
 	}
 
@@ -1554,7 +1554,7 @@ public:
 		SizeType NumNewElements = (SizeType)InitList.size();
 
 		InsertUninitializedImpl(InIndex, NumNewElements);
-		ConstructItems<ElementType>(GetData() + InIndex, InitList.begin(), NumNewElements);
+		ConstructItems<ElementType>((void*)(GetData() + InIndex), InitList.begin(), NumNewElements);
 
 		return InIndex;
 	}
@@ -1574,7 +1574,7 @@ public:
 		auto NumNewElements = Items.Num();
 
 		InsertUninitializedImpl(InIndex, NumNewElements);
-		ConstructItems<ElementType>(GetData() + InIndex, Items.GetData(), NumNewElements);
+		ConstructItems<ElementType>((void*)(GetData() + InIndex), Items.GetData(), NumNewElements);
 
 		return InIndex;
 	}
@@ -1594,7 +1594,7 @@ public:
 		auto NumNewElements = Items.Num();
 
 		InsertUninitializedImpl(InIndex, NumNewElements);
-		RelocateConstructItems<ElementType>(GetData() + InIndex, Items.GetData(), NumNewElements);
+		RelocateConstructItems<ElementType>((void*)(GetData() + InIndex), Items.GetData(), NumNewElements);
 		Items.ArrayNum = 0;
 
 		Items.SlackTrackerNumChanged();
@@ -1616,7 +1616,7 @@ public:
 		check(Ptr != nullptr);
 
 		InsertUninitializedImpl(Index, Count);
-		ConstructItems<ElementType>(GetData() + Index, Ptr, Count);
+		ConstructItems<ElementType>((void*)(GetData() + Index), Ptr, Count);
 
 		return Index;
 	}
@@ -1725,7 +1725,7 @@ private:
 		SizeType NumToMove = (ArrayNum - Index) - 1;
 		if (NumToMove)
 		{
-			RelocateConstructItems<ElementType>(Dest, Dest + 1, NumToMove);
+			RelocateConstructItems<ElementType>((void*)Dest, Dest + 1, NumToMove);
 		}
 		--ArrayNum;
 
@@ -1742,7 +1742,7 @@ private:
 		SizeType NumToMove = (ArrayNum - Index) - Count;
 		if (NumToMove)
 		{
-			RelocateConstructItems<ElementType>(Dest, Dest + Count, NumToMove);
+			RelocateConstructItems<ElementType>((void*)Dest, Dest + Count, NumToMove);
 		}
 		ArrayNum -= Count;
 
@@ -1812,7 +1812,7 @@ private:
 		const SizeType NumElementsToMoveIntoHole = FPlatformMath::Min(1, NumElementsAfterHole);
 		if (NumElementsToMoveIntoHole)
 		{
-			RelocateConstructItems<ElementType>(Dest, Data + (ArrayNum - NumElementsToMoveIntoHole), NumElementsToMoveIntoHole);
+			RelocateConstructItems<ElementType>((void*)Dest, Data + (ArrayNum - NumElementsToMoveIntoHole), NumElementsToMoveIntoHole);
 		}
 		--ArrayNum;
 
@@ -1831,7 +1831,7 @@ private:
 		const SizeType NumElementsToMoveIntoHole = FPlatformMath::Min(Count, NumElementsAfterHole);
 		if (NumElementsToMoveIntoHole)
 		{
-			RelocateConstructItems<ElementType>(Dest, Data + (ArrayNum - NumElementsToMoveIntoHole), NumElementsToMoveIntoHole);
+			RelocateConstructItems<ElementType>((void*)Dest, Data + (ArrayNum - NumElementsToMoveIntoHole), NumElementsToMoveIntoHole);
 		}
 		ArrayNum -= Count;
 
@@ -1962,7 +1962,7 @@ public:
 		{
 			const SizeType Diff = NewNum - ArrayNum;
 			const SizeType Index = AddUninitialized(Diff);
-			DefaultConstructItems<ElementType>((uint8*)AllocatorInstance.GetAllocation() + Index * sizeof(ElementType), Diff);
+			DefaultConstructItems<ElementType>((void*)((uint8*)AllocatorInstance.GetAllocation() + Index * sizeof(ElementType)), Diff);
 		}
 		else if (NewNum < 0)
 		{
@@ -2072,7 +2072,7 @@ public:
 
 		// Allocate memory for the new elements.
 		SizeType Pos = AddUninitialized(SourceCount);
-		ConstructItems<ElementType>(GetData() + Pos, Source.GetData(), SourceCount);
+		ConstructItems<ElementType>((void*)(GetData() + Pos), Source.GetData(), SourceCount);
 	}
 
 	/**
@@ -2096,7 +2096,7 @@ public:
 
 		// Allocate memory for the new elements.
 		SizeType Pos = AddUninitialized(SourceCount);
-		RelocateConstructItems<ElementType>(GetData() + Pos, Source.GetData(), SourceCount);
+		RelocateConstructItems<ElementType>((void*)(GetData() + Pos), Source.GetData(), SourceCount);
 		Source.ArrayNum = 0;
 
 		Source.SlackTrackerNumChanged();
@@ -2131,7 +2131,7 @@ public:
 
 		// Allocate memory for the new elements.
 		SizeType Pos = AddUninitialized(SourceCount);
-		ConstructItems<ElementType>(GetData() + Pos, UE::Core::Private::GetDataHelper(Source), SourceCount);
+		ConstructItems<ElementType>((void*)(GetData() + Pos), UE::Core::Private::GetDataHelper(Source), SourceCount);
 	}
 
 	/**
@@ -2146,7 +2146,7 @@ public:
 		check(Ptr != nullptr || Count == 0);
 
 		SizeType Pos = AddUninitialized(Count);
-		ConstructItems<ElementType>(GetData() + Pos, Ptr, Count);
+		ConstructItems<ElementType>((void*)(GetData() + Pos), Ptr, Count);
 	}
 
 	/**
@@ -2160,7 +2160,7 @@ public:
 		SizeType Count = (SizeType)InitList.size();
 
 		SizeType Pos = AddUninitialized(Count);
-		ConstructItems<ElementType>(GetData() + Pos, InitList.begin(), Count);
+		ConstructItems<ElementType>((void*)(GetData() + Pos), InitList.begin(), Count);
 	}
 
 	/**
@@ -2392,13 +2392,13 @@ public:
 	SizeType AddDefaulted()
 	{
 		const SizeType Index = AddUninitialized();
-		DefaultConstructItems<ElementType>((uint8*)AllocatorInstance.GetAllocation() + Index * sizeof(ElementType), 1);
+		DefaultConstructItems<ElementType>((void*)((uint8*)AllocatorInstance.GetAllocation() + Index * sizeof(ElementType)), 1);
 		return Index;
 	}
 	SizeType AddDefaulted(SizeType Count)
 	{
 		const SizeType Index = AddUninitialized(Count);
-		DefaultConstructItems<ElementType>((uint8*)AllocatorInstance.GetAllocation() + Index * sizeof(ElementType), Count);
+		DefaultConstructItems<ElementType>((void*)((uint8*)AllocatorInstance.GetAllocation() + Index * sizeof(ElementType)), Count);
 		return Index;
 	}
 
@@ -2413,7 +2413,7 @@ public:
 	{
 		const SizeType Index = AddUninitialized();
 		ElementType* Ptr = GetData() + Index;
-		DefaultConstructItems<ElementType>(Ptr, 1);
+		DefaultConstructItems<ElementType>((void*)Ptr, 1);
 		return *Ptr;
 	}
 
@@ -2501,7 +2501,7 @@ public:
 		auto NumNewElements = Items.Num();
 
 		InsertUninitializedImpl(InIndex, NumNewElements);
-		ConstructItems<ElementType>(GetData() + InIndex, Items.GetData(), NumNewElements);
+		ConstructItems<ElementType>((void*)(GetData() + InIndex), Items.GetData(), NumNewElements);
 
 		return InIndex;
 	}
@@ -2527,7 +2527,7 @@ public:
 		auto NumNewElements = Items.Num();
 
 		InsertUninitializedImpl(InIndex, NumNewElements);
-		RelocateConstructItems<ElementType>(GetData() + InIndex, Items.GetData(), NumNewElements);
+		RelocateConstructItems<ElementType>((void*)(GetData() + InIndex), Items.GetData(), NumNewElements);
 		Items.ArrayNum = 0;
 
 		Items.SlackTrackerNumChanged();
@@ -2553,7 +2553,7 @@ public:
 		TContainerElementTypeCompatibility<ElementType>::CopyingFromOtherType();
 
 		SizeType Pos = AddUninitialized(Count);
-		ConstructItems<ElementType>(GetData() + Pos, Ptr, Count);
+		ConstructItems<ElementType>((void*)(GetData() + Pos), Ptr, Count);
 	}
 
 private:
@@ -2659,7 +2659,7 @@ public:
 
 		// Destruct items that match the specified Item.
 		DestructItems(RemovePtr, 1);
-		RelocateConstructItems<ElementType>(RemovePtr, RemovePtr + 1, ArrayNum - (Index + 1));
+		RelocateConstructItems<ElementType>((void*)RemovePtr, RemovePtr + 1, ArrayNum - (Index + 1));
 
 		// Update the array count
 		--ArrayNum;
@@ -2722,7 +2722,7 @@ public:
 				// this was a non-matching run, we need to move it
 				if (WriteIndex != RunStartIndex)
 				{
-					RelocateConstructItems<ElementType>(Data + WriteIndex, Data + RunStartIndex, RunLength);
+					RelocateConstructItems<ElementType>((void*)(Data + WriteIndex), Data + RunStartIndex, RunLength);
 				}
 				WriteIndex += RunLength;
 			}
@@ -3180,7 +3180,7 @@ private:
 		if (OtherNum || PrevMax)
 		{
 			ResizeForCopy(NewNum, PrevMax);
-			ConstructItems<ElementType>(GetData(), OtherData, OtherNum);
+			ConstructItems<ElementType>((void*)GetData(), OtherData, OtherNum);
 		}
 		else
 		{
@@ -3217,7 +3217,7 @@ private:
 			}
 
 			ResizeForCopy(NewNum + ExtraSlack, PrevMax);
-			ConstructItems<ElementType>(GetData(), OtherData, OtherNum);
+			ConstructItems<ElementType>((void*)GetData(), OtherData, OtherNum);
 		}
 		else
 		{
