Cleaned up MB2 bins and added commentary with explanation for the bins size

diff --git a/Engine/Source/Runtime/Core/Private/HAL/MallocBinned2.cpp b/Engine/Source/Runtime/Core/Private/HAL/MallocBinned2.cpp
--- a/Engine/Source/Runtime/Core/Private/HAL/MallocBinned2.cpp
+++ b/Engine/Source/Runtime/Core/Private/HAL/MallocBinned2.cpp
@@ -106,41 +106,46 @@ PRAGMA_DISABLE_UNSAFE_TYPECAST_WARNINGS
 // They must be 16-byte aligned as well.
 static constexpr uint16 SmallBinSizes[] =
 {
-	16, 32, 48, 64, 80, 96, 112, 128, 144, 160, 176, 192, 208,// +16
-	256, 288, 320, 384, 448,
-	512,	// /128
-	560,	// /117
-	624,	// /105
-	720,	// /91
-	816,	// /80 31
-	912,	// /71 86
-	1024 - 16,	// /64
-	1168,	// /56 11
-	1392,	// /47 8
-	1520,	// /43 11
-	1680,	// /39
-	1872,	// /35
-	2048 - 16,	// /32
-	2256,	// /29 5
-	2608,	// /25 13
-	2976,	// /22
-	3264,	// /20 8
-	3632,	// /18 4
-	4096 - 16,	// /16
-	4368,	// /15
-	4672,	// /14
-	5040,	// /13
-	5456,	// /12
-	5952,	// /11
-	6544,	// /10
-	7280,	// /9
-	8192 - 16,	// /8
-	9360,	// /7
-	10912,	// /6
-	13104,	// /5
-	16384 - 16,	// /4
-	21840,	// /3
-	32768 - 16	// /2
+	16, 32, 48, 64, 80, 96, 112, 128, 144, 160, 176, 192, 208, // +16
+//	Bin		|	Divider | Slack leftover per page (64KB)
+	256,		// /256
+	288,		// /227 160b
+	320,		// /204 256b
+	384,		// /170 256b
+	448,		// /146 128b
+	512,		// /128	
+	560,		// /117	16b
+	624,		// /105	16b
+	720,		// /91	16b
+	816,		// /80	256b
+	912,		// /71	784b
+	1024-16,	// /64
+	1168,		// /56	128b
+	1392,		// /47	112b
+	1520,		// /43	176b
+	1680,		// /39	16b
+	1872,		// /35	16b
+	2048-16,	// /32
+	2256,		// /29	112b
+	2608,		// /25	336b
+	2976,		// /22	64b
+	3264,		// /20	256b
+	3632,		// /18	160b
+	4096-16,	// /16
+	4368,		// /15	16b
+	4672,		// /14	128b
+	5040,		// /13	16b
+	5456,		// /12	64b
+	5952,		// /11	64b
+	6544,		// /10	96b
+	7280,		// /9	16b
+	8192-16,	// /8
+	9360,		// /7	16b
+	10912,		// /6	64b
+	13104,		// /5	16b
+	16384-16,	// /4
+	21840,		// /3	16b
+	32768-16	// /2
 };
 
 MS_ALIGN(PLATFORM_CACHE_LINE_SIZE) static uint8 UnusedAlignPadding[PLATFORM_CACHE_LINE_SIZE] GCC_ALIGN(PLATFORM_CACHE_LINE_SIZE) = { 0 };
@@ -893,7 +898,7 @@ void* FMallocBinned2::MallocExternalSmall(SIZE_T Size, uint32 Alignment)
 
 	void* Result = Pool->AllocateBin();
 #if UE_MB2_ALLOCATOR_STATS
-	AllocatedSmallPoolMemory += PoolIndexToBinSize(PoolIndex);
+	AllocatedSmallPoolMemory += Table.BinSize;
 #endif
 
 	if (GBinned2AllocExtra)
@@ -1092,7 +1097,7 @@ void FMallocBinned2::FreeExternal(void* Ptr)
 		FPerThreadFreeBlockLists* Lists = GBinned2PerThreadCaches ? FPerThreadFreeBlockLists::Get() : nullptr;
 		if (Lists)
 		{
-			BundlesToRecycle = Lists->RecycleFullBundle(BasePtr->PoolIndex, Private::GGlobalRecycler);
+			BundlesToRecycle = Lists->RecycleFullBundle(PoolIndex, Private::GGlobalRecycler);
 			const bool bPushed = Lists->Free(Ptr, PoolIndex, BinSize);
 			check(bPushed);
 #if UE_MB2_ALLOCATOR_STATS
