Move away from our legacy AutoRTFM assert.

diff --git a/Engine/Source/Runtime/Core/Private/AutoRTFM/Utils.cpp b/Engine/Source/Runtime/Core/Private/AutoRTFM/Utils.cpp
--- a/Engine/Source/Runtime/Core/Private/AutoRTFM/Utils.cpp
+++ b/Engine/Source/Runtime/Core/Private/AutoRTFM/Utils.cpp
@@ -2,9 +2,6 @@
 
 #if (defined(__AUTORTFM) && __AUTORTFM)
 #include "Utils.h"
-#include <errno.h>
-#include <stdio.h>
-#include <string.h>
 
 #if PLATFORM_WINDOWS
 #include "Windows/WindowsHWrapper.h"
@@ -21,17 +18,6 @@ DEFINE_LOG_CATEGORY(LogAutoRTFM)
 namespace AutoRTFM
 {
 
-void PrettyAbort(const char* const File, const unsigned Line, const char* const Function, const char* const Expression)
-{
-    UE_LOG(LogAutoRTFM, Fatal, TEXT("%s:%d:%s: assertion %s failed."), ANSI_TO_TCHAR(File), Line, ANSI_TO_TCHAR(Function), ANSI_TO_TCHAR(Expression));
-
-#if defined(_MSC_VER) && !defined(__clang__)
-	__assume(false);
-#else
-	__builtin_unreachable();
-#endif // PLATFORM_WINDOWS
-}
-
 FString GetFunctionDescription(void* FunctionPtr)
 {
 #if PLATFORM_WINDOWS
diff --git a/Engine/Source/Runtime/Core/Private/AutoRTFM/Utils.h b/Engine/Source/Runtime/Core/Private/AutoRTFM/Utils.h
--- a/Engine/Source/Runtime/Core/Private/AutoRTFM/Utils.h
+++ b/Engine/Source/Runtime/Core/Private/AutoRTFM/Utils.h
@@ -3,24 +3,19 @@
 #pragma once
 
 #include "AutoRTFM/AutoRTFM.h"
-#define __STDC_FORMAT_MACROS 1
-#include <inttypes.h>
-#include <stdio.h>
-#include <stdlib.h>
 
 #include "Containers/UnrealString.h"
 #include "Logging/LogMacros.h"
+#include "Misc/AssertionMacros.h"
 
 DECLARE_LOG_CATEGORY_EXTERN(LogAutoRTFM, Display, All)
 
 namespace AutoRTFM
 {
 
-[[noreturn]] void PrettyAbort(const char* const File, const unsigned Line, const char* const Function, const char* const Expression);
-
 [[noreturn]] inline void Unreachable()
 {
-	UE_LOG(LogAutoRTFM, Fatal, TEXT("Unreachable encountered!"));
+	LowLevelFatalError(TEXT("Unreachable encountered!"));
 
 #if defined(_MSC_VER) && !defined(__clang__)
     __assume(false);
@@ -37,19 +32,9 @@ FString GetFunctionDescription(TReturnType (*FunctionPtr)(TParameterTypes...))
     return GetFunctionDescription(reinterpret_cast<void*>(FunctionPtr));
 }
 
-template<size_t A, size_t B> struct PrettyStaticAssert final
-{
-  static_assert(A == B, "Not equal");
-  static constexpr bool _cResult = (A == B);
-};
-
 } // namespace AutoRTFM
 
-#if defined(_MSC_VER) && !defined(__clang__)
-	#define ASSERT(exp) do { if (UNLIKELY(!(exp))) { UE_DEBUG_BREAK(); PrettyAbort(__FILE__, __LINE__, __FUNCSIG__, #exp); } } while (false)
-#else
-	#define ASSERT(exp) do { if (UNLIKELY(!(exp))) { UE_DEBUG_BREAK(); PrettyAbort(__FILE__, __LINE__, __PRETTY_FUNCTION__, #exp); } } while (false)
-#endif
+#define ASSERT(exp) UE_CLOG(UNLIKELY(!(exp)), LogAutoRTFM, Fatal, TEXT("AutoRTFM assert!"))
 
 #if defined(__has_feature)
 	#if __has_feature(address_sanitizer)
