Added more UE_LIFETIMEBOUND attributes to TArray functions which return element references.
Added MSVC support to UE_LIFETIMEBOUND, though it needs a special static analysis pass and so is currently useless.

diff --git a/Engine/Source/Runtime/Core/Public/Containers/Array.h b/Engine/Source/Runtime/Core/Public/Containers/Array.h
--- a/Engine/Source/Runtime/Core/Public/Containers/Array.h
+++ b/Engine/Source/Runtime/Core/Public/Containers/Array.h
@@ -900,7 +900,7 @@ public:
 	 * @returns Reference to the top element.
 	 * @see Pop, Push
 	 */
-	FORCEINLINE ElementType& Top()
+	FORCEINLINE ElementType& Top() UE_LIFETIMEBOUND
 	{
 		return Last();
 	}
@@ -913,7 +913,7 @@ public:
 	 * @returns Reference to the top element.
 	 * @see Pop, Push
 	 */
-	FORCEINLINE const ElementType& Top() const
+	FORCEINLINE const ElementType& Top() const UE_LIFETIMEBOUND
 	{
 		return Last();
 	}
@@ -924,7 +924,7 @@ public:
 	 * @param IndexFromTheEnd (Optional) Index from the end of array (default = 0).
 	 * @returns Reference to n-th last element from the array.
 	 */
-	FORCEINLINE ElementType& Last(SizeType IndexFromTheEnd = 0)
+	FORCEINLINE ElementType& Last(SizeType IndexFromTheEnd = 0) UE_LIFETIMEBOUND
 	{
 		RangeCheck(ArrayNum - IndexFromTheEnd - 1);
 		return GetData()[ArrayNum - IndexFromTheEnd - 1];
@@ -938,7 +938,7 @@ public:
 	 * @param IndexFromTheEnd (Optional) Index from the end of array (default = 0).
 	 * @returns Reference to n-th last element from the array.
 	 */
-	FORCEINLINE const ElementType& Last(SizeType IndexFromTheEnd = 0) const
+	FORCEINLINE const ElementType& Last(SizeType IndexFromTheEnd = 0) const UE_LIFETIMEBOUND
 	{
 		RangeCheck(ArrayNum - IndexFromTheEnd - 1);
 		return GetData()[ArrayNum - IndexFromTheEnd - 1];
@@ -1499,7 +1499,7 @@ public:
 	 * @return A reference to the newly-inserted element.
 	 * @see Insert_GetRef, InsertDefaulted_GetRef
 	 */
-	ElementType& InsertZeroed_GetRef(SizeType Index)
+	ElementType& InsertZeroed_GetRef(SizeType Index) UE_LIFETIMEBOUND
 	{
 		InsertUninitializedImpl(Index, 1);
 		ElementType* Ptr = GetData() + Index;
@@ -1534,7 +1534,7 @@ public:
 	 * @return A reference to the newly-inserted element.
 	 * @see Insert_GetRef, InsertZeroed_GetRef
 	 */
-	ElementType& InsertDefaulted_GetRef(SizeType Index)
+	ElementType& InsertDefaulted_GetRef(SizeType Index) UE_LIFETIMEBOUND
 	{
 		InsertUninitializedImpl(Index, 1);
 		ElementType* Ptr = GetData() + Index;
@@ -1682,7 +1682,7 @@ public:
 	 * @return A reference to the newly-inserted element.
 	 * @see Add, Remove
 	 */
-	ElementType& Insert_GetRef(ElementType&& Item, SizeType Index)
+	ElementType& Insert_GetRef(ElementType&& Item, SizeType Index) UE_LIFETIMEBOUND
 	{
 		CheckAddress(&Item);
 
@@ -1702,7 +1702,7 @@ public:
 	 * @return A reference to the newly-inserted element.
 	 * @see Add, Remove
 	 */
-	ElementType& Insert_GetRef(const ElementType& Item, SizeType Index)
+	ElementType& Insert_GetRef(const ElementType& Item, SizeType Index) UE_LIFETIMEBOUND
 	{
 		CheckAddress(&Item);
 
@@ -2244,7 +2244,7 @@ public:
 	 * @return A reference to the newly-inserted element.
 	 */
 	template <typename... ArgsType>
-	FORCEINLINE ElementType& Emplace_GetRef(ArgsType&&... Args)
+	FORCEINLINE ElementType& Emplace_GetRef(ArgsType&&... Args) UE_LIFETIMEBOUND
 	{
 		const SizeType Index = AddUninitialized();
 		ElementType* Ptr = GetData() + Index;
@@ -2273,7 +2273,7 @@ public:
 	 * @return A reference to the newly-inserted element.
 	 */
 	template <typename... ArgsType>
-	FORCEINLINE ElementType& EmplaceAt_GetRef(SizeType Index, ArgsType&&... Args)
+	FORCEINLINE ElementType& EmplaceAt_GetRef(SizeType Index, ArgsType&&... Args) UE_LIFETIMEBOUND
 	{
 		InsertUninitializedImpl(Index, 1);
 		ElementType* Ptr = GetData() + Index;
@@ -2318,7 +2318,7 @@ public:
 	 * @return A reference to the newly-inserted element.
 	 * @see AddDefaulted_GetRef, AddUnique_GetRef, AddZeroed_GetRef, Insert_GetRef
 	 */
-	FORCEINLINE ElementType& Add_GetRef(ElementType&& Item)
+	FORCEINLINE ElementType& Add_GetRef(ElementType&& Item) UE_LIFETIMEBOUND
 	{
 		CheckAddress(&Item);
 		return Emplace_GetRef(MoveTempIfPossible(Item));
@@ -2331,7 +2331,7 @@ public:
 	 * @return A reference to the newly-inserted element.
 	 * @see AddDefaulted_GetRef, AddUnique_GetRef, AddZeroed_GetRef, Insert_GetRef
 	 */
-	FORCEINLINE ElementType& Add_GetRef(const ElementType& Item)
+	FORCEINLINE ElementType& Add_GetRef(const ElementType& Item) UE_LIFETIMEBOUND
 	{
 		CheckAddress(&Item);
 		return Emplace_GetRef(Item);
@@ -2373,7 +2373,7 @@ public:
 	 * @return A reference to the newly-inserted element.
 	 * @see Add_GetRef, AddDefaulted_GetRef, AddUnique_GetRef, Insert_GetRef
 	 */
-	ElementType& AddZeroed_GetRef()
+	ElementType& AddZeroed_GetRef() UE_LIFETIMEBOUND
 	{
 		const SizeType Index = AddUninitialized();
 		ElementType* Ptr = GetData() + Index;
@@ -2409,7 +2409,7 @@ public:
 	 * @return A reference to the newly-inserted element.
 	 * @see Add_GetRef, AddZeroed_GetRef, AddUnique_GetRef, Insert_GetRef
 	 */
-	ElementType& AddDefaulted_GetRef()
+	ElementType& AddDefaulted_GetRef() UE_LIFETIMEBOUND
 	{
 		const SizeType Index = AddUninitialized();
 		ElementType* Ptr = GetData() + Index;
@@ -3503,7 +3503,7 @@ public:
 	 *
 	 * @returns The reference to the top element from the heap.
 	 */
-	const ElementType& HeapTop() const
+	const ElementType& HeapTop() const UE_LIFETIMEBOUND
 	{
 		return (*this)[0];
 	}
@@ -3513,7 +3513,7 @@ public:
 	 *
 	 * @returns The reference to the top element from the heap.
 	 */
-	ElementType& HeapTop()
+	ElementType& HeapTop() UE_LIFETIMEBOUND
 	{
 		return (*this)[0];
 	}
diff --git a/Engine/Source/Runtime/Core/Public/MSVC/MSVCPlatform.h b/Engine/Source/Runtime/Core/Public/MSVC/MSVCPlatform.h
--- a/Engine/Source/Runtime/Core/Public/MSVC/MSVCPlatform.h
+++ b/Engine/Source/Runtime/Core/Public/MSVC/MSVCPlatform.h
@@ -27,3 +27,8 @@
 // Ensure we can use this builtin - seems to be present on Clang 9, GCC 11 and MSVC 19.26
 #define PLATFORM_COMPILER_SUPPORTS_BUILTIN_BITCAST (_MSC_VER >= 1926)
 
+#ifdef __has_cpp_attribute
+	#if __has_cpp_attribute(msvc::lifetimebound)
+		#define UE_LIFETIMEBOUND [[msvc::lifetimebound]]
+	#endif
+#endif
