Fixed moving arrays of const TArray elements into array of non-const TArray elements.

diff --git a/Engine/Source/Runtime/Core/Private/Containers/ContainersTest.cpp b/Engine/Source/Runtime/Core/Private/Containers/ContainersTest.cpp
--- a/Engine/Source/Runtime/Core/Private/Containers/ContainersTest.cpp
+++ b/Engine/Source/Runtime/Core/Private/Containers/ContainersTest.cpp
@@ -934,6 +934,53 @@ bool FContainersTSetTest::RunTest(const FString& Parameters)
 	return !HasAnyErrors();
 }
 
+IMPLEMENT_SIMPLE_AUTOMATION_TEST(FContainersTArrayTest, "System.Core.Containers.TArray", EAutomationTestFlags_ApplicationContextMask | EAutomationTestFlags::SmokeFilter)
+bool FContainersTArrayTest::RunTest(const FString& Parameters)
+{
+	// Move semantics
+	{
+		// Move array to another array of the same type
+		{
+			TArray<int32> From = { 1, 2, 3, 4, 5 };
+			TArray<int32> To = MoveTemp(From);
+			TestTrue(TEXT("Move constructing an array relocates the elements"), From.IsEmpty() && To == TArray<int32>{ 1, 2, 3, 4, 5 });
+		}
+
+		// Move array to another array of bitwise-compatible type
+		{
+			// We can transfer memory between arrays of signed and unsigned
+			TArray<int32> From = { 1, 2, 3, 4, 5 };
+			TArray<uint32> To(MoveTemp(From));
+			TestTrue(TEXT("Move constructing an array relocates the elements with bitwise compatible elements"), From.IsEmpty() && To == TArray<uint32>{ 1, 2, 3, 4, 5 });
+		}
+
+		// Move array to another array of bitwise-incompatible type
+		{
+			// We can't transfer memory, but we can copy
+			TArray<int32> From = { 1, 2, 3, 4, 5 };
+			TArray<int64> To(MoveTemp(From));
+			TestTrue(TEXT("Move constructing an array does not relocate the elements with bitwise incompatible elements"), From == TArray<int32>{ 1, 2, 3, 4, 5 } && To == TArray<int64>{ 1, 2, 3, 4, 5 });
+		}
+
+		// Move array of unqualified type to an array of qualified type
+		{
+			// We can transfer memory from an array of non-const to an array of const
+			TArray<int32> From = { 1, 2, 3, 4, 5 };
+			TArray<const int32> To(MoveTemp(From));
+			TestTrue(TEXT("Move constructing an array when adding const relocates the elements"), From.IsEmpty() && To == TArray<const int32>{ 1, 2, 3, 4, 5 });
+		}
+
+		// Move array of qualified type to an array of unqualified type
+		{
+			// We can't transfer memory from an array of const to an array of non-const, because that would be const-incorrect
+			TArray<const int32> From = { 1, 2, 3, 4, 5 };
+			TArray<int32> To(MoveTemp(From));
+			TestTrue(TEXT("Move constructing an array when removing conts copies the elements"), From == TArray<const int32>{ 1, 2, 3, 4, 5 } && To == TArray<int32>{ 1, 2, 3, 4, 5 });
+		}
+	}
+
+	return !HasAnyErrors();
+}
 
 namespace ArrayViewTests
 {
diff --git a/Engine/Source/Runtime/Core/Public/Containers/Array.h b/Engine/Source/Runtime/Core/Public/Containers/Array.h
--- a/Engine/Source/Runtime/Core/Public/Containers/Array.h
+++ b/Engine/Source/Runtime/Core/Public/Containers/Array.h
@@ -23,6 +23,7 @@
 #include "Templates/IdentityFunctor.h"
 #include "Templates/Invoke.h"
 #include "Templates/Less.h"
+#include "Templates/LosesQualifiersFromTo.h"
 #include "Templates/Requires.h"
 #include "Templates/Sorting.h"
 #include "Templates/AlignmentTemplates.h"
@@ -304,17 +305,22 @@ namespace UE::Core::Private
 	template <typename FromArrayType, typename ToArrayType>
 	constexpr bool CanMoveTArrayPointersBetweenArrayTypes()
 	{
-		using FromAllocatorType = typename FromArrayType::AllocatorType;
-		using ToAllocatorType   = typename ToArrayType::AllocatorType;
-		using FromElementType   = typename FromArrayType::ElementType;
-		using ToElementType     = typename ToArrayType::ElementType;
+		using FromAllocatorType          = typename FromArrayType::AllocatorType;
+		using ToAllocatorType            = typename ToArrayType::AllocatorType;
+		using FromElementType            = typename FromArrayType::ElementType;
+		using ToElementType              = typename ToArrayType::ElementType;
+		using UnqualifiedFromElementType = std::remove_cv_t<FromElementType>;
+		using UnqualifiedToElementType   = std::remove_cv_t<ToElementType>;
 
 		// Allocators must be equal or move-compatible...
 		if constexpr (std::is_same_v<FromAllocatorType, ToAllocatorType> || TCanMoveBetweenAllocators<FromAllocatorType, ToAllocatorType>::Value)
 		{
 			return
-				std::is_same_v         <ToElementType, FromElementType> ||      // The element type of the container must be the same, or...
-				TIsBitwiseConstructible<ToElementType, FromElementType>::Value; // ... the element type of the source container must be bitwise constructible from the element type in the destination container
+				!TLosesQualifiersFromTo<FromElementType, ToElementType>::Value &&
+				(
+					std::is_same_v         <const ToElementType, const FromElementType> ||               // The element type of the container must be the same, or...
+					TIsBitwiseConstructible<UnqualifiedToElementType, UnqualifiedFromElementType>::Value // ... the element type of the source container must be bitwise constructible from the element type in the destination container
+				);
 		}
 		else
 		{
