Fix crash when a thread accesses config while another thread is performing dynamic config on the same branch

#rnx

diff --git a/Engine/Source/Runtime/Core/Private/Misc/ConfigCacheIni.cpp b/Engine/Source/Runtime/Core/Private/Misc/ConfigCacheIni.cpp
--- a/Engine/Source/Runtime/Core/Private/Misc/ConfigCacheIni.cpp
+++ b/Engine/Source/Runtime/Core/Private/Misc/ConfigCacheIni.cpp
@@ -928,6 +928,7 @@ FConfigFile::FConfigFile(FConfigFile&& Other)
 
 FConfigFile& FConfigFile::operator=(const FConfigFile& Other)
 {
+	FWriteScopeLock ScopeLock(ConfigFileMapLock);
 	this->FConfigFileMap::operator=(Other);
 	Dirty = Other.Dirty;
 	NoSave = Other.NoSave;
@@ -974,6 +975,7 @@ FConfigFile& FConfigFile::operator=(const FConfigFile& Other)
 
 FConfigFile& FConfigFile::operator=(FConfigFile&& Other)
 {
+	FWriteScopeLock ScopeLock(ConfigFileMapLock);
 	this->FConfigFileMap::operator=(MoveTemp(Other));
 	Dirty = Other.Dirty;
 	NoSave = Other.NoSave;
@@ -1049,6 +1051,8 @@ UE::ConfigAccessTracking::FFile* FConfigFile::GetFileAccess() const
 
 bool FConfigFile::operator==( const FConfigFile& Other ) const
 {
+	FReadScopeLock ScopeLock(ConfigFileMapLock);
+
 	if ( Pairs.Num() != Other.Pairs.Num() )
 		return 0;
 
@@ -1104,6 +1108,7 @@ bool FConfigFile::Combine(const FString& Filename)
 
 void FConfigFile::Shrink()
 {
+	FWriteScopeLock ScopeLock(ConfigFileMapLock);
 	FConfigFileMap::Shrink();
 	for (FConfigFileMap::TIterator It(*this); It; ++It)
 	{
@@ -5915,6 +5920,7 @@ FCriticalSection FConfigCacheIni::ConfigForPlatformLock;
 TMap<FName, FConfigCacheIni::FPluginInfo*> FConfigCacheIni::RegisteredPlugins;
 FCriticalSection FConfigCacheIni::RegisteredPluginsLock;
 
+FRWLock FConfigFile::ConfigFileMapLock;
 
 void FConfigCacheIni::AddPluginToAllBranches(FName PluginName, FConfigModificationTracker* ModificationTracker)
 {
diff --git a/Engine/Source/Runtime/Core/Public/Misc/ConfigCacheIni.h b/Engine/Source/Runtime/Core/Public/Misc/ConfigCacheIni.h
--- a/Engine/Source/Runtime/Core/Public/Misc/ConfigCacheIni.h
+++ b/Engine/Source/Runtime/Core/Public/Misc/ConfigCacheIni.h
@@ -31,6 +31,7 @@
 #include "Misc/ConfigAccessTracking.h"
 #include "Misc/ConfigTypes.h"
 #include "Misc/Paths.h"
+#include "Misc/ScopeRWLock.h"
 #include "Serialization/Archive.h"
 #include "Serialization/StructuredArchive.h"
 #include "Serialization/StructuredArchiveAdapters.h"
@@ -608,45 +609,133 @@ public:
 	UE_DEPRECATED(5.4, "Use FindSection, and/or use the new AddToSection, etc APIs to modify sections without retrieving the section. See top of ConfigCacheIni.h for more info.")
 	FORCEINLINE const FConfigSection* Find(const FString& SectionName) const
 	{
+		FReadScopeLock ScopeLock(ConfigFileMapLock);
 		return FConfigFileMap::Find(SectionName);
 	}
+
 	UE_DEPRECATED(5.4, "Use FindSection, and/or use the new AddToSection, etc APIs to modify sections without retrieving the section. See top of ConfigCacheIni.h for more info.")
 	FORCEINLINE FConfigSection* Find(const FString& SectionName)
 	{
+		FReadScopeLock ScopeLock(ConfigFileMapLock);
 		return FConfigFileMap::Find(SectionName);
 	}
 
 	FORCEINLINE const FConfigSection* FindSection(const FString& SectionName) const
 	{
+		FReadScopeLock ScopeLock(ConfigFileMapLock);
 		return FConfigFileMap::Find(SectionName);
 	}
 
-	FORCEINLINE int32 Num() const 								{ return FConfigFileMap::Num(); }
-	FORCEINLINE bool IsEmpty() const							{ return FConfigFileMap::IsEmpty(); }
-	FORCEINLINE void Empty(int32 ExpectedNumElements = 0) 		{ FConfigFileMap::Empty(ExpectedNumElements); }
-	FORCEINLINE bool Contains(const FString& SectionName) const	{ return FConfigFileMap::Contains(SectionName); }
-	FORCEINLINE int32 GetKeys(TArray<FString>& Keys) const 		{ return FConfigFileMap::GetKeys(Keys); }
-	FORCEINLINE int32 GetKeys(TSet<FString>& Keys) const		{ return FConfigFileMap::GetKeys(Keys); }
-	FORCEINLINE int32 Remove(KeyConstPointerType InKey) 		{ return FConfigFileMap::Remove(InKey); }
+	FORCEINLINE int32 Num() const
+	{
+		FReadScopeLock ScopeLock(ConfigFileMapLock);
+		return FConfigFileMap::Num();
+	}
+
+	FORCEINLINE bool IsEmpty() const
+	{
+		FReadScopeLock ScopeLock(ConfigFileMapLock);
+		return FConfigFileMap::IsEmpty();
+	}
+
+	FORCEINLINE void Empty(int32 ExpectedNumElements = 0)
+	{
+		FWriteScopeLock ScopeLock(ConfigFileMapLock);
+		FConfigFileMap::Empty(ExpectedNumElements);
+	}
+
+	FORCEINLINE bool Contains(const FString& SectionName) const
+	{
+		FReadScopeLock ScopeLock(ConfigFileMapLock);
+		return FConfigFileMap::Contains(SectionName);
+	}
+
+	FORCEINLINE int32 GetKeys(TArray<FString>& Keys) const
+	{
+		FReadScopeLock ScopeLock(ConfigFileMapLock);
+		return FConfigFileMap::GetKeys(Keys);
+	}
+
+	FORCEINLINE int32 GetKeys(TSet<FString>& Keys) const
+	{
+		FReadScopeLock ScopeLock(ConfigFileMapLock);
+		return FConfigFileMap::GetKeys(Keys);
+	}
+
+	FORCEINLINE int32 Remove(KeyConstPointerType InKey)
+	{
+		FWriteScopeLock ScopeLock(ConfigFileMapLock);
+		return FConfigFileMap::Remove(InKey);
+	}
+
+	FORCEINLINE ValueType& Add(const KeyType& InKey, const ValueType& InValue)
+	{
+		FWriteScopeLock ScopeLock(ConfigFileMapLock);
+		return FConfigFileMap::Add(InKey, InValue);
+	}
+
+	FORCEINLINE ValueType& Add(const KeyType&  InKey, ValueType&& InValue)
+	{
+		FWriteScopeLock ScopeLock(ConfigFileMapLock);
+		return FConfigFileMap::Add(InKey, MoveTempIfPossible(InValue));
+	}
+
+	FORCEINLINE ValueType& Add(KeyType&& InKey, const ValueType&  InValue)
+	{
+		FWriteScopeLock ScopeLock(ConfigFileMapLock);
+		return FConfigFileMap::Add(MoveTempIfPossible(InKey), InValue);
+	}
 
-	FORCEINLINE ValueType& Add(const KeyType&  InKey, const ValueType&  InValue) { return FConfigFileMap::Add(InKey, InValue); }
-	FORCEINLINE ValueType& Add(const KeyType&  InKey,		ValueType&& InValue) { return FConfigFileMap::Add(InKey, MoveTempIfPossible(InValue)); }
-	FORCEINLINE ValueType& Add(		 KeyType&& InKey, const ValueType&  InValue) { return FConfigFileMap::Add(MoveTempIfPossible(InKey), InValue); }
-	FORCEINLINE ValueType& Add(		 KeyType&& InKey,		ValueType&& InValue) { return FConfigFileMap::Add(MoveTempIfPossible(InKey), MoveTempIfPossible(InValue)); }
+	FORCEINLINE ValueType& Add(KeyType&& InKey,	ValueType&& InValue)
+	{
+		FWriteScopeLock ScopeLock(ConfigFileMapLock);
+		return FConfigFileMap::Add(MoveTempIfPossible(InKey), MoveTempIfPossible(InValue));
+	}
 	
-	FORCEINLINE void Append(TMap<FString, FConfigSection> Other) { FConfigFileMap::Append(MoveTemp(Other)); }
-	FORCEINLINE void Reset() { FConfigFileMap::Reset(); }
+	FORCEINLINE void Append(TMap<FString, FConfigSection> Other)
+	{
+		FWriteScopeLock ScopeLock(ConfigFileMapLock);
+		FConfigFileMap::Append(MoveTemp(Other));
+	}
 
+	FORCEINLINE void Reset()
+	{
+		FWriteScopeLock ScopeLock(ConfigFileMapLock);
+		FConfigFileMap::Reset();
+	}
 
 	UE_DEPRECATED(5.4, "Use FindOrAddConfigSection, and/or use the new AddToSection, etc APIs to modify sections without retrieving the section. See top of ConfigCacheIni.h for more info.")
-	FORCEINLINE ValueType& FindOrAdd(const FString& Key) 		{ return FConfigFileMap::FindOrAdd(Key); }
+	FORCEINLINE ValueType& FindOrAdd(const FString& Key)
+	{
+		FWriteScopeLock ScopeLock(ConfigFileMapLock);
+		return FConfigFileMap::FindOrAdd(Key);
+	}
 
 	UE_DEPRECATED(5.4, "Use const ranged for iterators, (wrap your FConfigFile variable in AsConst to force the const iterator). See top of ConfigCacheIni.h for more info.")
-	FORCEINLINE TRangedForIterator      begin() { return TRangedForIterator(Pairs.begin()); }
-	FORCEINLINE TRangedForConstIterator begin() const { return TRangedForConstIterator(Pairs.begin()); }
+	FORCEINLINE TRangedForIterator begin()
+	{
+		FReadScopeLock ScopeLock(ConfigFileMapLock);
+		return TRangedForIterator(Pairs.begin());
+	}
+
+	FORCEINLINE TRangedForConstIterator begin() const
+	{
+		FReadScopeLock ScopeLock(ConfigFileMapLock); 
+		return TRangedForConstIterator(Pairs.begin());
+	}
+
 	UE_DEPRECATED(5.4, "Use const ranged for iterators, (wrap your FConfigFile variable in AsConst to force the const iterator). See top of ConfigCacheIni.h for more info.")
-	FORCEINLINE TRangedForIterator      end() { return TRangedForIterator(Pairs.end()); }
-	FORCEINLINE TRangedForConstIterator end() const { return TRangedForConstIterator(Pairs.end()); }
+	FORCEINLINE TRangedForIterator end()
+	{
+		FReadScopeLock ScopeLock(ConfigFileMapLock); 
+		return TRangedForIterator(Pairs.end());
+	}
+
+	FORCEINLINE TRangedForConstIterator end() const
+	{
+		FReadScopeLock ScopeLock(ConfigFileMapLock);
+		return TRangedForConstIterator(Pairs.end());
+	}
 	
 	///////////////////////////////////
 
@@ -718,7 +807,12 @@ private:
 	void ProcessCommand(FConfigSection* Section, FStringView SectionName, FConfigValue::EValueType Command, FName Key, FString&& Value);
 
 	FConfigSection* FindOrAddSectionInternal(const FString& SectionName);
-	FORCEINLINE FConfigSection* FindInternal(const FString& SectionName) { return FConfigFileMap::Find(SectionName); };
+
+	FORCEINLINE FConfigSection* FindInternal(const FString& SectionName)
+	{
+		FReadScopeLock ScopeLock(ConfigFileMapLock);
+		return FConfigFileMap::Find(SectionName);
+	};
 
 	// allow the templated helper to access FindOrAddSectionInternal
 	template<typename FileType>
@@ -914,6 +1008,8 @@ private:
 	// for AddStaticLayersToHierarchy
 	friend class FConfigCacheIni;
 	friend FConfigContext;
+
+	static CORE_API FRWLock ConfigFileMapLock;
 };
 
 /**
