Add a comment to warn developers that they should update NoExportTypes.h when changing the struct data.

diff --git a/Engine/Source/Runtime/Core/Public/Misc/FrameRate.h b/Engine/Source/Runtime/Core/Public/Misc/FrameRate.h
--- a/Engine/Source/Runtime/Core/Public/Misc/FrameRate.h
+++ b/Engine/Source/Runtime/Core/Public/Misc/FrameRate.h
@@ -31,6 +31,8 @@ struct FFrameRate
 		: Numerator(InNumerator), Denominator(InDenominator)
 	{}
 
+	/** IMPORTANT: If you change the struct data, ensure that you also update the version in NoExportTypes.h  */
+
 	/**
 	 * The numerator of the framerate represented as a number of frames per second (e.g. 60 for 60 fps)
 	 */
diff --git a/Engine/Source/Runtime/Core/Public/Misc/FrameTime.h b/Engine/Source/Runtime/Core/Public/Misc/FrameTime.h
--- a/Engine/Source/Runtime/Core/Public/Misc/FrameTime.h
+++ b/Engine/Source/Runtime/Core/Public/Misc/FrameTime.h
@@ -93,6 +93,8 @@ public:
 	 */
 	static FFrameTime FromDecimal(double InDecimalFrame);
 
+	/** IMPORTANT: If you change the struct data, ensure that you also update the version in NoExportTypes.h  */
+
 	FFrameNumber FrameNumber;
 
 private:
diff --git a/Engine/Source/Runtime/Core/Public/Misc/QualifiedFrameTime.h b/Engine/Source/Runtime/Core/Public/Misc/QualifiedFrameTime.h
--- a/Engine/Source/Runtime/Core/Public/Misc/QualifiedFrameTime.h
+++ b/Engine/Source/Runtime/Core/Public/Misc/QualifiedFrameTime.h
@@ -79,6 +79,8 @@ public:
 
 public:
 
+	/** IMPORTANT: If you change the struct data, ensure that you also update the version in NoExportTypes.h  */
+
 	/** The frame time */
 	FFrameTime Time;
 
diff --git a/Engine/Source/Runtime/Core/Public/Misc/Timecode.h b/Engine/Source/Runtime/Core/Public/Misc/Timecode.h
--- a/Engine/Source/Runtime/Core/Public/Misc/Timecode.h
+++ b/Engine/Source/Runtime/Core/Public/Misc/Timecode.h
@@ -376,7 +376,7 @@ public:
 		return Builder.ToString();
 	}
 
-	/*
+	/**
 	 * Will return true if the timecode represents a valid timecode value where
 	 *
 	 * Hours is +/- [0,23]
@@ -396,6 +396,8 @@ public:
 	}
 public:
 
+	/** IMPORTANT: If you change the struct data, ensure that you also update the version in NoExportTypes.h  */
+
 	/** How many hours does this timecode represent */
 	int32 Hours;
 
