Simplified the live table management in FTextLocalizationManager

This change aims to significantly simplify live table management by removing yet more vestiges of logic from before FTextHistory_Base has an embedded text ID (and so needed to look up its ID via the display string pointer).

Under this new change, the only thing that ever adds new entries to the live table is UpdateFromNative and UpdateFromLocalizations, and always does so using "clean" text IDs (the same IDs that are gathered and localized). GetDisplayString no longer attempts to patch the live table for source mismatches or partial text ID matches, as we no longer need every text ID to have a unique and up-to-date display string pointer (and haven't since 5.0).

Additionally the debug cultures (ie, LEET and keys) no longer clobber the data in the live table, as FindDisplayString and GetDisplayString now generate those debug display strings on-demand. This avoids UpdateFromNative and UpdateFromLocalizations needing to spin through the entire live table to restore any clobbered display strings with their original values. "-LEETifyUnlocalized" also no longer clobbers the data in the live table, as GetDisplayString now generates that debug string on-demand if there is no valid live table entry.

#rnx

diff --git a/Engine/Source/Developer/TranslationEditor/Private/TranslationPickerEditWindow.cpp b/Engine/Source/Developer/TranslationEditor/Private/TranslationPickerEditWindow.cpp
--- a/Engine/Source/Developer/TranslationEditor/Private/TranslationPickerEditWindow.cpp
+++ b/Engine/Source/Developer/TranslationEditor/Private/TranslationPickerEditWindow.cpp
@@ -576,7 +576,7 @@ FReply STranslationPickerEditWidget::SaveAndPreview()
 
 #else
 
-	FTextLocalizationManager::Get().ReplaceStringInLiveTable(TranslationUnit->Namespace, TranslationUnit->Key, TranslationUnit->Translation);
+	FTextLocalizationManager::Get().AddOrUpdateDisplayStringInLiveTable(TranslationUnit->Namespace, TranslationUnit->Key, TranslationUnit->Translation, &TranslationUnit->Source);
 
 	if (IConsoleObject* CObj = IConsoleManager::Get().FindConsoleObject(TEXT("Slate.TriggerInvalidate")))
 	{
diff --git a/Engine/Source/Runtime/Core/Private/Internationalization/TextLocalizationManager.cpp b/Engine/Source/Runtime/Core/Private/Internationalization/TextLocalizationManager.cpp
--- a/Engine/Source/Runtime/Core/Private/Internationalization/TextLocalizationManager.cpp
+++ b/Engine/Source/Runtime/Core/Private/Internationalization/TextLocalizationManager.cpp
@@ -123,7 +123,7 @@ static FAutoConsoleCommand CmdDumpLiveTable(
 
 static FAutoConsoleCommand SetDisplayString(
 	TEXT("Localization.SetDisplayString"),
-	TEXT("Replaces DisplayString in the live table given required arguments: 'Namespace', 'Key', and 'DisplayString'"),
+	TEXT("Add/Update DisplayString in the live table given required arguments: 'Namespace', 'Key', and 'DisplayString'. 'SourceString' is optional, but needed if adding a new display string."),
 	FConsoleCommandWithArgsDelegate::CreateStatic([](const TArray<FString>& Args)
 	{
 		// Rebuild the full string of arguments, since values within quotes may have been split on spaces
@@ -132,6 +132,7 @@ static FAutoConsoleCommand SetDisplayString(
 		FString Namespace;
 		TOptional<FString> Key;
 		TOptional<FString> DisplayString;
+		TOptional<FString> SourceString;
 
 		FString Tmp;
 		if (FParse::Value(*Arguments, TEXT("Namespace="), Tmp))
@@ -149,6 +150,11 @@ static FAutoConsoleCommand SetDisplayString(
 			DisplayString = MoveTemp(Tmp);
 		}
 
+		if (FParse::Value(*Arguments, TEXT("SourceString="), Tmp))
+		{
+			SourceString = MoveTemp(Tmp);
+		}
+
 		// Namespace is optional, assumed to be empty if not provided
 		if (!Key.IsSet() || !DisplayString.IsSet())
 		{
@@ -162,7 +168,7 @@ static FAutoConsoleCommand SetDisplayString(
 			return;
 		}
 
-		FTextLocalizationManager::Get().ReplaceStringInLiveTable(Namespace, Key.GetValue(), DisplayString.GetValue());
+		FTextLocalizationManager::Get().AddOrUpdateDisplayStringInLiveTable(Namespace, Key.GetValue(), DisplayString.GetValue(), SourceString.GetPtrOrNull());
 
 		// For live game, force widget invalidation to update the text onscreen. This isn't needed in the editor.
 		if (IConsoleObject* CObj = IConsoleManager::Get().FindConsoleObject(TEXT("Slate.TriggerInvalidate")))
@@ -801,22 +807,34 @@ void FTextLocalizationManager::DumpLiveTable(const FString& OutputFilename, cons
 	FFileHelper::SaveStringToFile(DumpString, *OutputFilename, FFileHelper::EEncodingOptions::ForceUTF8WithoutBOM);
 }
 
-void FTextLocalizationManager::ReplaceStringInLiveTable(const FString& Namespace, const FString& Key, const FString& DisplayString)
+void FTextLocalizationManager::AddOrUpdateDisplayStringInLiveTable(const FString& Namespace, const FString& Key, const FString& DisplayString, const FString* const SourceStringPtr)
 {
-	// Lock while updating the table
-	FScopeLock ScopeLock(&DisplayStringTableCS);
-
 	const FTextId TextId(Namespace, Key);
 
-	FDisplayStringEntry* LiveEntry = DisplayStringLookupTable.Find(TextId);
+	// Lock while updating the table
+	FScopeLock ScopeLock(&DisplayStringTableCS);
 
-	if (LiveEntry)
+	if (FDisplayStringEntry* LiveEntry = DisplayStringLookupTable.Find(TextId))
 	{
 		LiveEntry->DisplayString = MakeTextDisplayString(CopyTemp(DisplayString));
 		DirtyLocalRevisionForTextId(TextId);
 
 		UE_LOG(LogConsoleResponse, Display, TEXT("Updated string for Namespace='%s', Key='%s' to DisplayString='%s'"), *Namespace, *Key, *DisplayString);
 	}
+	else if (SourceStringPtr)
+	{
+		// Add new entry
+		FDisplayStringEntry NewLiveEntry(
+			FTextKey(),													/*LocResID*/
+			INDEX_NONE,													/*LocalizationTargetPathId*/
+			FTextLocalizationResource::HashString(*SourceStringPtr),	/*SourceStringHash*/
+			MakeTextDisplayString(CopyTemp(DisplayString))				/*String*/
+		);
+
+		DisplayStringLookupTable.Emplace(TextId, NewLiveEntry);
+
+		UE_LOG(LogConsoleResponse, Display, TEXT("Added string for Namespace='%s', Key='%s' to DisplayString='%s'"), *Namespace, *Key, *DisplayString);
+	}
 	else
 	{
 		UE_LOG(LogConsoleResponse, Display, TEXT("String not found for Namespace='%s', Key='%s'"), *Namespace, *Key);
@@ -977,31 +995,30 @@ void FTextLocalizationManager::RegisterPolyglotTextData(TArrayView<const FPolygl
 	}
 }
 
-FTextConstDisplayStringPtr FTextLocalizationManager::FindDisplayString(const FTextKey& Namespace, const FTextKey& Key, const FString* const SourceStringPtr) const
+FTextConstDisplayStringPtr FTextLocalizationManager::FindDisplayString_Internal(const FTextId& TextId, const FString& SourceString) const
 {
-	if (Key.IsEmpty() || !FTextLocalizationManager::IsDisplayStringSupportEnabled())
+	if (!IsInitialized())
 	{
 		return nullptr;
 	}
 
-	FScopeLock ScopeLock(&DisplayStringTableCS);
+#if ENABLE_LOC_TESTING
+	if (FInternationalization::Get().GetCurrentLanguage()->GetName() == FLeetCulture::StaticGetName())
+	{
+		FTextDisplayStringRef TmpDisplayString = MakeTextDisplayString(CopyTemp(SourceString));
+		FInternationalization::Leetify(*TmpDisplayString);
+		return TmpDisplayString;
+	}
+	if (FInternationalization::Get().GetCurrentLanguage()->GetName() == FKeysCulture::StaticGetName())
+	{
+		return MakeTextDisplayString(TextLocalizationManager::KeyifyTextId(TextId));
+	}
+#endif	// ENABLE_LOC_TESTING
 
-	const FTextId TextId(Namespace, Key);
+	FScopeLock ScopeLock(&DisplayStringTableCS);
 
 	if (const FDisplayStringEntry* LiveEntry = DisplayStringLookupTable.Find(TextId))
 	{
-		auto GetSourceStringRef = [SourceStringPtr]() -> const FString&
-		{
-			if (SourceStringPtr)
-			{
-				return *SourceStringPtr;
-			}
-
-			static const FString EmptyString;
-			return EmptyString;
-		};
-
-		const FString& SourceString = GetSourceStringRef();
 		if (SourceString.IsEmpty() || LiveEntry->SourceStringHash == FTextLocalizationResource::HashString(SourceString))
 		{
 			return LiveEntry->DisplayString;
@@ -1011,10 +1028,9 @@ FTextConstDisplayStringPtr FTextLocalizationManager::FindDisplayString(const FTe
 	return nullptr;
 }
 
-FTextConstDisplayStringPtr FTextLocalizationManager::GetDisplayString(const FTextKey& Namespace, const FTextKey& Key, const FString* const SourceStringPtr)
+FTextConstDisplayStringPtr FTextLocalizationManager::FindDisplayString(const FTextKey& Namespace, const FTextKey& Key, const FString* const SourceStringPtr) const
 {
-	TRACE_CPUPROFILER_EVENT_SCOPE(FTextLocalizationManager::GetDisplayString);
-	LLM_SCOPE_BYNAME(TEXT("Localization/DisplayStrings"));
+	TRACE_CPUPROFILER_EVENT_SCOPE(FTextLocalizationManager::FindDisplayString);
 
 	if (Key.IsEmpty() || !FTextLocalizationManager::IsDisplayStringSupportEnabled())
 	{
@@ -1032,144 +1048,61 @@ FTextConstDisplayStringPtr FTextLocalizationManager::GetDisplayString(const FTex
 		return EmptyString;
 	};
 
-	const FString& SourceString = GetSourceStringRef();
-
-	FScopeLock ScopeLock(&DisplayStringTableCS);
+	return FindDisplayString_Internal(FTextId(Namespace, Key), GetSourceStringRef());
+}
 
-#if ENABLE_LOC_TESTING
-	const bool bShouldLEETIFYAll = IsInitialized() && FInternationalization::Get().GetCurrentLanguage()->GetName() == FLeetCulture::StaticGetName();
-	const bool bShouldKeyifyAll = IsInitialized() && FInternationalization::Get().GetCurrentLanguage()->GetName() == FKeysCulture::StaticGetName();
+FTextConstDisplayStringPtr FTextLocalizationManager::GetDisplayString(const FTextKey& Namespace, const FTextKey& Key, const FString* const SourceStringPtr) const
+{
+	TRACE_CPUPROFILER_EVENT_SCOPE(FTextLocalizationManager::GetDisplayString);
 
-	// Attempt to set bShouldLEETIFYUnlocalizedString appropriately, only once, after the commandline is initialized and parsed.
-	static bool bShouldLEETIFYUnlocalizedString = false;
+	if (Key.IsEmpty() || !FTextLocalizationManager::IsDisplayStringSupportEnabled())
 	{
-		static bool bHasParsedCommandLine = false;
-		if (!bHasParsedCommandLine && FCommandLine::IsInitialized())
-		{
-			bShouldLEETIFYUnlocalizedString = FParse::Param(FCommandLine::Get(), TEXT("LEETIFYUnlocalized"));
-			bHasParsedCommandLine = true;
-		}
+		return nullptr;
 	}
-#endif
-
-	const FTextId TextId(Namespace, Key);
-
-	const uint32 SourceStringHash = !SourceString.IsEmpty() ? FTextLocalizationResource::HashString(SourceString) : 0;
 
-	FDisplayStringEntry* LiveEntry = DisplayStringLookupTable.Find(TextId);
+	FTextId TextId(Namespace, Key);
 
-	// In builds with stable keys enabled, we want to use the display string from the "clean" version of the text (if the sources match) as this is the only version that is translated
-	FTextConstDisplayStringPtr SourceDisplayString;
-	FDisplayStringEntry* SourceLiveEntry = nullptr;
+	// In builds with stable keys enabled, we want to find the display string for the "clean" version of the text (if the sources match) as this is the only version that is translated
 #if USE_STABLE_LOCALIZATION_KEYS
 	{
-		const FTextKey DisplayNamespace = TextNamespaceUtil::StripPackageNamespace(TextId.GetNamespace().GetChars());
-		if (DisplayNamespace != TextId.GetNamespace())
+		const FString DisplayNamespace = TextNamespaceUtil::StripPackageNamespace(TextId.GetNamespace().GetChars());
+		if (FCString::Strcmp(*DisplayNamespace, TextId.GetNamespace().GetChars()) != 0)
 		{
-			SourceLiveEntry = DisplayStringLookupTable.Find(FTextId(DisplayNamespace, TextId.GetKey()));
-
-			if (SourceLiveEntry)
-			{
-				if (SourceString.IsEmpty() || SourceLiveEntry->SourceStringHash == SourceStringHash)
-				{
-					SourceDisplayString = SourceLiveEntry->DisplayString;
-				}
-				else
-				{
-					SourceLiveEntry = nullptr;
-				}
-			}
+			TextId = FTextId(DisplayNamespace, TextId.GetKey());
 		}
 	}
-#endif // USE_STABLE_LOCALIZATION_KEYS
+#endif	// USE_STABLE_LOCALIZATION_KEYS
 
-	// Entry is present.
-	if (LiveEntry)
+	auto GetSourceStringRef = [SourceStringPtr]() -> const FString&
 	{
-		// If the source string (hash) is different, the local source has changed and should override
-		if (!SourceString.IsEmpty() && SourceStringHash != LiveEntry->SourceStringHash)
+		if (SourceStringPtr)
 		{
-			LiveEntry->SourceStringHash = SourceStringHash;
-			LiveEntry->DisplayString = SourceDisplayString ? SourceDisplayString.ToSharedRef() : MakeTextDisplayString(CopyTemp(SourceString));
-			DirtyLocalRevisionForTextId(TextId);
-
-#if ENABLE_LOC_TESTING
-			if (bShouldKeyifyAll)
-			{
-				DisplayStringBackupTable.Add(TextId, LiveEntry->DisplayString);
-				LiveEntry->DisplayString = MakeTextDisplayString(TextLocalizationManager::KeyifyTextId(TextId));
-			}
-			else if ((bShouldLEETIFYAll || bShouldLEETIFYUnlocalizedString) && !LiveEntry->DisplayString->IsEmpty())
-			{
-				if (!bShouldLEETIFYUnlocalizedString)
-				{
-					DisplayStringBackupTable.Add(TextId, LiveEntry->DisplayString);
-				}
-
-				FTextDisplayStringRef TmpDisplayString = MakeTextDisplayString(CopyTemp(*LiveEntry->DisplayString));
-				FInternationalization::Leetify(*TmpDisplayString);
-				LiveEntry->DisplayString = TmpDisplayString;
-			}
-#endif
-
-			UE_LOG(LogTextLocalizationManager, Verbose, TEXT("An attempt was made to get a localized string (Namespace:%s, Key:%s), but the source string hash does not match - the source string (%s) will be used."), TextId.GetNamespace().GetChars(), TextId.GetKey().GetChars(), **LiveEntry->DisplayString);
+			return *SourceStringPtr;
 		}
 
-		return LiveEntry->DisplayString;
-	}
-	// Entry is absent, but has a related entry to clone.
-	else if (SourceLiveEntry)
-	{
-		TRACE_CPUPROFILER_EVENT_SCOPE(FTextLocalizationManager::GetDisplayString_AddRelatedEntry);
-
-		check(SourceString.IsEmpty() || SourceLiveEntry->SourceStringHash == SourceStringHash);
-		check(SourceDisplayString && SourceLiveEntry->DisplayString == SourceDisplayString);
+		static const FString EmptyString;
+		return EmptyString;
+	};
 
-		// Clone the entry for the active ID
-		FDisplayStringEntry NewEntry(*SourceLiveEntry);
-		DisplayStringLookupTable.Emplace(TextId, NewEntry);
+	const FString& SourceString = GetSourceStringRef();
 
-		return NewEntry.DisplayString;
+	if (FTextConstDisplayStringPtr DisplayString = FindDisplayString_Internal(TextId, SourceString))
+	{
+		return DisplayString;
 	}
+
 #if ENABLE_LOC_TESTING
-	// Entry is absent.
-	else if (bShouldKeyifyAll || ((bShouldLEETIFYAll || bShouldLEETIFYUnlocalizedString) && !SourceString.IsEmpty()))
+	if (FCommandLine::IsInitialized())
 	{
-		TRACE_CPUPROFILER_EVENT_SCOPE(FTextLocalizationManager::GetDisplayString_AddNewEntry);
-
-		FTextConstDisplayStringRef UnlocalizedString = MakeTextDisplayString(CopyTemp(SourceString));
-
-		if (bShouldKeyifyAll)
+		static const bool bShouldLEETIFYUnlocalizedString = FParse::Param(FCommandLine::Get(), TEXT("LEETIFYUnlocalized"));
+		if (bShouldLEETIFYUnlocalizedString)
 		{
-			DisplayStringBackupTable.Add(TextId, UnlocalizedString);
-			UnlocalizedString = MakeTextDisplayString(TextLocalizationManager::KeyifyTextId(TextId));
-		}
-		else if (bShouldLEETIFYAll || bShouldLEETIFYUnlocalizedString)
-		{
-			check(!SourceString.IsEmpty());
-
-			if (!bShouldLEETIFYUnlocalizedString)
-			{
-				DisplayStringBackupTable.Add(TextId, UnlocalizedString);
-			}
-
-			FTextDisplayStringRef TmpDisplayString = MakeTextDisplayString(CopyTemp(*UnlocalizedString));
+			FTextDisplayStringRef TmpDisplayString = MakeTextDisplayString(CopyTemp(SourceString));
 			FInternationalization::Leetify(*TmpDisplayString);
-			UnlocalizedString = TmpDisplayString;
+			return TmpDisplayString;
 		}
-
-		FDisplayStringEntry NewEntry(
-			FTextKey(),					/*LocResID*/
-			INDEX_NONE,					/*LocalizationTargetPathId*/
-			SourceStringHash,			/*SourceStringHash*/
-			UnlocalizedString			/*String*/
-		);
-
-		DisplayStringLookupTable.Emplace(TextId, NewEntry);
-
-		return UnlocalizedString;
 	}
-#endif // ENABLE_LOC_TESTING
+#endif
 
 	return nullptr;
 }
@@ -1381,6 +1314,15 @@ void FTextLocalizationManager::OnCultureChanged()
 		return;
 	}
 
+#if ENABLE_LOC_TESTING
+	if (FInternationalization::Get().GetCurrentLanguage()->GetName() == FLeetCulture::StaticGetName() || FInternationalization::Get().GetCurrentLanguage()->GetName() == FKeysCulture::StaticGetName())
+	{
+		// When switching to a debug culture, just bump the text revision (so that the dynamically generated debug display strings are rebuilt) and bail
+		DirtyTextRevision();
+		return;
+	}
+#endif	// ENABLE_LOC_TESTING
+
 	RefreshResources();
 
 	if (!TextLocalizationManager::AsyncLoadLocalizationDataOnLanguageChange)
@@ -1428,15 +1370,7 @@ void FTextLocalizationManager::LoadLocalizationResourcesForPrioritizedCultures_S
 		return;
 	}
 
-	// Leet-ify always needs the native text to operate on, so force native data if we're loading for LEET
-	// The keys culture also uses native so that we have a good way to restore the original text from the keys state 
-	ELocalizationLoadFlags FinalLocLoadFlags = LocLoadFlags | (ShouldForceLoadGameLocalization() ? ELocalizationLoadFlags::ForceLocalizedGame : ELocalizationLoadFlags::None);
-#if ENABLE_LOC_TESTING
-	if ((PrioritizedCultureNames[0] == FLeetCulture::StaticGetName()) || (PrioritizedCultureNames[0] == FKeysCulture::StaticGetName()))
-	{
-		FinalLocLoadFlags |= ELocalizationLoadFlags::Native;
-	}
-#endif
+	const ELocalizationLoadFlags FinalLocLoadFlags = LocLoadFlags | (ShouldForceLoadGameLocalization() ? ELocalizationLoadFlags::ForceLocalizedGame : ELocalizationLoadFlags::None);
 
 	// Load the resources from each text source
 	FTextLocalizationResource NativeResource;
@@ -1447,32 +1381,14 @@ void FTextLocalizationManager::LoadLocalizationResourcesForPrioritizedCultures_S
 		LocalizedTextSource->LoadLocalizedResources(FinalLocLoadFlags, PrioritizedCultureNames, NativeResource, LocalizedResource);
 	}
 
-	// When loc testing is enabled, UpdateFromNative also takes care of restoring non-localized text which is why the condition below is gated
-#if !ENABLE_LOC_TESTING
+	// Apply the new display string data
 	if (!NativeResource.IsEmpty())
-#endif
 	{
 		UpdateFromNative(MoveTemp(NativeResource), /*bDirtyTextRevision*/false);
 	}
-
-#if ENABLE_LOC_TESTING
-	// The leet culture is fake. Just leet-ify existing strings.
-	if (PrioritizedCultureNames[0] == FLeetCulture::StaticGetName())
-	{
-		LeetifyAllDisplayStrings();
-	}
-	else if (PrioritizedCultureNames[0] == FKeysCulture::StaticGetName())
-	{
-		KeyifyAllDisplayStrings();
-	}
-	else
-#endif
+	if (!LocalizedResource.IsEmpty())
 	{
-		// Replace localizations with those of the loaded localization resources.
-		if (!LocalizedResource.IsEmpty())
-		{
-			UpdateFromLocalizations(MoveTemp(LocalizedResource), /*bDirtyTextRevision*/false);
-		}
+		UpdateFromLocalizations(MoveTemp(LocalizedResource), /*bDirtyTextRevision*/false);
 	}
 
 	DirtyTextRevision();
@@ -1730,9 +1646,6 @@ void FTextLocalizationManager::UpdateFromNative(FTextLocalizationResource&& Text
 #endif	// WITH_EDITORONLY_DATA
 					DisplayStringsByLocalizationTargetId.TrackTextId(LiveEntry->LocalizationTargetPathId, NewEntry.LocalizationTargetPathId, TextId);
 					LiveEntry->LocalizationTargetPathId = NewEntry.LocalizationTargetPathId;
-#if ENABLE_LOC_TESTING
-					DisplayStringBackupTable.Remove(TextId);
-#endif	// ENABLE_LOC_TESTING
 				}
 			}
 			else
@@ -1752,52 +1665,6 @@ void FTextLocalizationManager::UpdateFromNative(FTextLocalizationResource&& Text
 
 		// Note: Do not use TextLocalizationResource after this point as we may have stolen some of its strings
 		TextLocalizationResource.Entries.Reset();
-
-		// Perform any additional processing over existing entries
-#if ENABLE_LOC_TESTING || USE_STABLE_LOCALIZATION_KEYS
-		for (auto& DisplayStringPair : DisplayStringLookupTable)
-		{
-			const FTextId& TextId = DisplayStringPair.Key;
-			FDisplayStringEntry& LiveEntry = DisplayStringPair.Value;
-
-#if USE_STABLE_LOCALIZATION_KEYS
-			// In builds with stable keys enabled, we have to update the display strings from the "clean" version of the text (if the sources match) as this is the only version that is translated
-			{
-				const FTextKey LiveNamespace = TextId.GetNamespace();
-				const FTextKey DisplayNamespace = TextNamespaceUtil::StripPackageNamespace(LiveNamespace.GetChars());
-				if (LiveNamespace != DisplayNamespace)
-				{
-					const FDisplayStringEntry* DisplayStringEntry = DisplayStringLookupTable.Find(FTextId(DisplayNamespace, TextId.GetKey()));
-					if (DisplayStringEntry && ((LiveEntry.SourceStringHash == DisplayStringEntry->SourceStringHash) || LiveEntry.IsEmpty()))
-					{
-						// this is to account for the case where the LiveString is empty and we are overwriting the value 
-						// we could do an if check, but copying an int is cheaper 
-						LiveEntry.SourceStringHash = DisplayStringEntry->SourceStringHash;
-						// There is no good way to copy over the NativeBackupString. We disregard it as it is only for testing 
-						LiveEntry.DisplayString = DisplayStringEntry->DisplayString;
-#if WITH_EDITORONLY_DATA
-						LiveEntry.LocResID = DisplayStringEntry->LocResID;
-#endif	// WITH_EDITORONLY_DATA
-						DisplayStringsByLocalizationTargetId.TrackTextId(LiveEntry.LocalizationTargetPathId, DisplayStringEntry->LocalizationTargetPathId, TextId);
-						LiveEntry.LocalizationTargetPathId = DisplayStringEntry->LocalizationTargetPathId;
-#if ENABLE_LOC_TESTING
-						DisplayStringBackupTable.Remove(TextId);
-#endif	// ENABLE_LOC_TESTING
-					}
-				}
-			}
-#endif	// USE_STABLE_LOCALIZATION_KEYS
-
-#if ENABLE_LOC_TESTING
-			// Restore the pre-leet state (if any)
-			if (FTextConstDisplayStringPtr DisplayStringBackup;
-				DisplayStringBackupTable.RemoveAndCopyValue(TextId, DisplayStringBackup))
-			{
-				LiveEntry.DisplayString = DisplayStringBackup.ToSharedRef();
-			}
-#endif	// ENABLE_LOC_TESTING
-		}
-#endif	// ENABLE_LOC_TESTING || USE_STABLE_LOCALIZATION_KEYS
 	}
 
 	if (bDirtyTextRevision)
@@ -1817,8 +1684,6 @@ void FTextLocalizationManager::UpdateFromLocalizations(FTextLocalizationResource
 		return;
 	}
 
-	static const bool bShouldLEETIFYUnlocalizedString = FParse::Param(FCommandLine::Get(), TEXT("LEETIFYUnlocalized"));
-
 	// Lock while updating the tables
 	{
 		FScopeLock ScopeLock(&DisplayStringTableCS);
@@ -1843,25 +1708,11 @@ void FTextLocalizationManager::UpdateFromLocalizations(FTextLocalizationResource
 					// this is to account for the case where the LiveString is empty and we are overwriting the value 
 					// we could do an if check, but copying an int is cheaper 
 					LiveEntry->SourceStringHash = NewEntry.SourceStringHash;
-					// @TODO: Currently no way to copy over the NativeBackupString
 					LiveEntry->DisplayString = NewEntry.LocalizedString.ToSharedRef();
 #if WITH_EDITORONLY_DATA
 					LiveEntry->LocResID = NewEntry.LocResID;
 #endif	// WITH_EDITORONLY_DATA
 				}
-#if ENABLE_LOC_TESTING
-				else if (bShouldLEETIFYUnlocalizedString && !LiveEntry->DisplayString->IsEmpty())
-				{
-					FTextDisplayStringRef TmpDisplayString = MakeTextDisplayString(CopyTemp(*LiveEntry->DisplayString));
-					FInternationalization::Leetify(*TmpDisplayString);
-					LiveEntry->DisplayString = TmpDisplayString;
-#if WITH_EDITORONLY_DATA
-					LiveEntry->LocResID = FTextKey();
-#endif	// WITH_EDITORONLY_DATA
-					DisplayStringsByLocalizationTargetId.TrackTextId(LiveEntry->LocalizationTargetPathId, INDEX_NONE, TextId);
-					LiveEntry->LocalizationTargetPathId = INDEX_NONE;
-				}
-#endif	// ENABLE_LOC_TESTING
 			}
 			else
 			{
@@ -1880,54 +1731,6 @@ void FTextLocalizationManager::UpdateFromLocalizations(FTextLocalizationResource
 
 		// Note: Do not use TextLocalizationResource after this point as we may have stolen some of its strings
 		TextLocalizationResource.Entries.Reset();
-
-		// Perform any additional processing over existing entries
-#if USE_STABLE_LOCALIZATION_KEYS
-		{
-			for (auto& DisplayStringPair : DisplayStringLookupTable)
-			{
-				const FTextId& TextId = DisplayStringPair.Key;
-				FDisplayStringEntry& LiveEntry = DisplayStringPair.Value;
-
-				// In builds with stable keys enabled, we have to update the display strings from the "clean" version of the text (if the sources match) as this is the only version that is translated
-				const FTextKey LiveNamespace = TextId.GetNamespace();
-				const FTextKey DisplayNamespace = TextNamespaceUtil::StripPackageNamespace(LiveNamespace.GetChars());
-				if (LiveNamespace != DisplayNamespace)
-				{
-					const FDisplayStringEntry* DisplayStringEntry = DisplayStringLookupTable.Find(FTextId(DisplayNamespace, TextId.GetKey()));
-
-					// If the source string hashes are are the same, we can replace the display string.
-					// Otherwise, it would suggest the source string has changed and the new localization may be based off of an old source string.
-					if (DisplayStringEntry && ((LiveEntry.SourceStringHash == DisplayStringEntry->SourceStringHash) || LiveEntry.IsEmpty()))
-					{
-						// this is to account for the case where the LiveString is empty and we are overwriting the value 
-						// we could do an if check, but copying an int is cheaper 
-						LiveEntry.SourceStringHash = DisplayStringEntry->SourceStringHash;
-						// There is no good way to copy over the NativeBackupString. We disregard it as it is only for testing 
-						LiveEntry.DisplayString = DisplayStringEntry->DisplayString;
-#if WITH_EDITORONLY_DATA
-						LiveEntry.LocResID = DisplayStringEntry->LocResID;
-#endif	// WITH_EDITORONLY_DATA
-						DisplayStringsByLocalizationTargetId.TrackTextId(LiveEntry.LocalizationTargetPathId, DisplayStringEntry->LocalizationTargetPathId, TextId);
-						LiveEntry.LocalizationTargetPathId = DisplayStringEntry->LocalizationTargetPathId;
-					}
-#if ENABLE_LOC_TESTING
-					else if (bShouldLEETIFYUnlocalizedString && !LiveEntry.DisplayString->IsEmpty())
-					{
-						FTextDisplayStringRef TmpDisplayString = MakeTextDisplayString(CopyTemp(*LiveEntry.DisplayString));
-						FInternationalization::Leetify(*TmpDisplayString);
-						LiveEntry.DisplayString = TmpDisplayString;
-#if WITH_EDITORONLY_DATA
-						LiveEntry.LocResID = FTextKey();
-#endif	// WITH_EDITORONLY_DATA
-						DisplayStringsByLocalizationTargetId.TrackTextId(LiveEntry.LocalizationTargetPathId, INDEX_NONE, TextId);
-						LiveEntry.LocalizationTargetPathId = INDEX_NONE;
-					}
-#endif	// ENABLE_LOC_TESTING
-				}
-			}
-		}
-#endif	// USE_STABLE_LOCALIZATION_KEYS
 	}
 
 	if (bDirtyTextRevision)
@@ -2080,40 +1883,3 @@ bool FTextLocalizationManager::ShouldForceLoadGameLocalization() const
 	return false;
 #endif	// WITH_EDITOR
 }
-
-#if ENABLE_LOC_TESTING
-void FTextLocalizationManager::LeetifyAllDisplayStrings()
-{
-	// Lock while updating the tables
-	FScopeLock ScopeLock(&DisplayStringTableCS);
-
-	DisplayStringBackupTable.Reset();
-	for (auto& DisplayStringPair : DisplayStringLookupTable)
-	{
-		FDisplayStringEntry& LiveEntry = DisplayStringPair.Value;
-		DisplayStringBackupTable.Add(DisplayStringPair.Key, LiveEntry.DisplayString);
-
-		if (!LiveEntry.DisplayString->IsEmpty())
-		{
-			FTextDisplayStringRef TmpDisplayString = MakeTextDisplayString(CopyTemp(*LiveEntry.DisplayString));
-			FInternationalization::Leetify(*TmpDisplayString);
-			LiveEntry.DisplayString = TmpDisplayString;
-		}
-	}
-}
-
-void FTextLocalizationManager::KeyifyAllDisplayStrings()
-{
-	// Lock while updating the tables
-	FScopeLock ScopeLock(&DisplayStringTableCS);
-
-	DisplayStringBackupTable.Reset();
-	for (auto& DisplayStringPair : DisplayStringLookupTable)
-	{
-		FDisplayStringEntry& LiveEntry = DisplayStringPair.Value;
-		DisplayStringBackupTable.Add(DisplayStringPair.Key, LiveEntry.DisplayString);
-
-		LiveEntry.DisplayString = MakeTextDisplayString(TextLocalizationManager::KeyifyTextId(DisplayStringPair.Key));
-	}
-}
-#endif 
diff --git a/Engine/Source/Runtime/Core/Public/Internationalization/TextLocalizationManager.h b/Engine/Source/Runtime/Core/Public/Internationalization/TextLocalizationManager.h
--- a/Engine/Source/Runtime/Core/Public/Internationalization/TextLocalizationManager.h
+++ b/Engine/Source/Runtime/Core/Public/Internationalization/TextLocalizationManager.h
@@ -122,9 +122,6 @@ private:
 
 	mutable FCriticalSection DisplayStringTableCS;
 	FDisplayStringLookupTable DisplayStringLookupTable;
-#if ENABLE_LOC_TESTING
-	TMap<FTextId, FTextConstDisplayStringPtr> DisplayStringBackupTable;
-#endif
 	FDisplayStringsByLocalizationTargetId DisplayStringsByLocalizationTargetId;
 
 	mutable FRWLock TextRevisionRW;
@@ -164,7 +161,7 @@ private:
 public:
 	CORE_API void DumpLiveTable(const FString* NamespaceFilter = nullptr, const FString* KeyFilter = nullptr, const FString* DisplayStringFilter = nullptr, const FLogCategoryBase* CategoryOverride = nullptr) const;
 	CORE_API void DumpLiveTable(const FString& OutputFilename, const FString* NamespaceFilter = nullptr, const FString* KeyFilter = nullptr, const FString* DisplayStringFilter = nullptr) const;
-	CORE_API void ReplaceStringInLiveTable(const FString& Namespace, const FString& Key, const FString& DisplayString);
+	CORE_API void AddOrUpdateDisplayStringInLiveTable(const FString& Namespace, const FString& Key, const FString& DisplayString, const FString* const SourceStringPtr = nullptr);
 #endif
 
 	/**
@@ -205,17 +202,21 @@ public:
 	CORE_API void RegisterPolyglotTextData(const FPolyglotTextData& InPolyglotTextData, const bool InAddDisplayString = true);
 	CORE_API void RegisterPolyglotTextData(TArrayView<const FPolyglotTextData> InPolyglotTextDataArray, const bool InAddDisplayStrings = true);
 
-	/**	Finds and returns the display string with the given namespace and key, if it exists.
-	 *	Additionally, if a non-null and non-empty source string is specified and the found localized display string was not localized from that source string, null will be returned. */
+	/**
+	 * Finds and returns the display string with the given namespace and key, if it exists.
+	 * @note If a non-null and non-empty source string is specified and the found localized display string was not localized from that source string, null will be returned.
+	 */
 	CORE_API FTextConstDisplayStringPtr FindDisplayString(const FTextKey& Namespace, const FTextKey& Key, const FString* const SourceStringPtr = nullptr) const;
 
-	/**	Returns a display string with the given namespace and key.
-	 *	If no display string exists, it will be created using the source string or an empty string if no source string is provided.
-	 *	If a display string exists ...
-	 *		... but it was not localized from the specified source string, the display string will be set to the specified source and returned.
-	 *		... and it was localized from the specified source string (or the source string was null or empty), the display string will be returned.
-	*/
-	CORE_API FTextConstDisplayStringPtr GetDisplayString(const FTextKey& Namespace, const FTextKey& Key, const FString* const SourceStringPtr);
+	/**
+	 * Get the current display string for the given namespace and key, if any.
+	 * @note If a non-null and non-empty source string is specified and the found localized display string was not localized from that source string, it will be considered unlocalized.
+	 * 
+	 * Unlike FindDisplayString:
+	 *   * This function may adjust the given text ID (eg, when USE_STABLE_LOCALIZATION_KEYS is enabled).
+	 *   * This function may return a value for unlocalized strings (eg, when using -LEETifyUnlocalized).
+	 */
+	CORE_API FTextConstDisplayStringPtr GetDisplayString(const FTextKey& Namespace, const FTextKey& Key, const FString* const SourceStringPtr) const;
 
 #if WITH_EDITORONLY_DATA
 	/** If an entry exists for the specified namespace and key, returns true and provides the localization resource identifier from which it was loaded. Otherwise, returns false. */
@@ -380,14 +381,11 @@ private:
 	/** Dirties the local revision counter for the given text ID by incrementing it (or adding it) */
 	CORE_API void DirtyLocalRevisionForTextId(const FTextId& InTextId);
 
+	/** Internal version of FindDisplayString, shared between FindDisplayString and GetDisplayString */
+	FTextConstDisplayStringPtr FindDisplayString_Internal(const FTextId& TextId, const FString& SourceString) const;
+
 	/** Dirties the text revision counter by incrementing it, causing a revision mismatch for any information cached before this happens.  */
 	CORE_API void DirtyTextRevision();
-#if ENABLE_LOC_TESTING
-	/** A helper function that leetifies all of the display strings when the LEET culture is active. */
-	CORE_API void LeetifyAllDisplayStrings();
-	/** A helper function that converts all of the display strings to show the localization key and namespace associated with the string when the keys culture is active. */
-	CORE_API void KeyifyAllDisplayStrings();
-#endif
 
 	/** Array of registered localized text sources, sorted by priority (@see RegisterTextSource) */
 	TArray<TSharedPtr<ILocalizedTextSource>> LocalizedTextSources;
