- Added some calls to Shrink which saves a decent amount of config memory (added it to the "config timing" output)
- Added detailed stats to "config memusage" - per layer, and per section tracing for large layers/sections (with -detailed)
- Added -csv[=path] option to config memusage (mimicking dumpcvars csv handling)
- Added FConfigBranch::RunOnEachFile and RunOnEachCommandStream so reduce how many places are hardcoding the list of layers

diff --git a/Engine/Source/Runtime/Core/Private/Misc/ConfigCacheIni.cpp b/Engine/Source/Runtime/Core/Private/Misc/ConfigCacheIni.cpp
--- a/Engine/Source/Runtime/Core/Private/Misc/ConfigCacheIni.cpp
+++ b/Engine/Source/Runtime/Core/Private/Misc/ConfigCacheIni.cpp
@@ -1108,8 +1108,14 @@ bool FConfigFile::Combine(const FString& Filename)
 
 void FConfigFile::Shrink()
 {
+#if !UE_BUILD_SHIPPING
+	extern double GConfigShrinkTime;
+	if (IsInGameThread()) GConfigShrinkTime -= FPlatformTime::Seconds();
+#endif
+
 	FWriteScopeLock ScopeLock(ConfigFileMapLock);
 	FConfigFileMap::Shrink();
+
 	for (FConfigFileMap::TIterator It(*this); It; ++It)
 	{
 		It.Value().Shrink();
@@ -1120,6 +1126,10 @@ void FConfigFile::Shrink()
 	{
 		Pair.Value.Shrink();
 	}
+
+#if !UE_BUILD_SHIPPING
+	if (IsInGameThread()) GConfigShrinkTime += FPlatformTime::Seconds();
+#endif
 }
 
 // Assumes GetTypeHash(AltKeyType) matches GetTypeHash(KeyType)
@@ -1482,6 +1492,10 @@ void FillFileFromBuffer(FileType* File, FStringView Buffer, bool bHandleSymbolCo
 				
 				CurrentSectionName = *FoundRemap;
 			}
+			if (CurrentSection)
+			{
+				CurrentSection->Shrink();
+			}
 			CurrentSection = File->FindOrAddSectionInternal(CurrentSectionName);
 
 			// look to see if there is a set of key remaps for this section
@@ -2922,6 +2936,11 @@ FConfigCommandStreamSection* FConfigCommandStream::FindOrAddSectionInternal(cons
 
 void FConfigCommandStream::Shrink()
 {
+#if !UE_BUILD_SHIPPING
+	extern double GConfigShrinkTime;
+	if (IsInGameThread()) GConfigShrinkTime -= FPlatformTime::Seconds();
+#endif
+
 	TMap<FString, FConfigCommandStreamSection>::Shrink();
 	for (auto& Pair : *this)
 	{
@@ -2933,6 +2952,10 @@ void FConfigCommandStream::Shrink()
 	{
 		Pair.Value.Shrink();
 	}
+
+#if !UE_BUILD_SHIPPING
+	if (IsInGameThread()) GConfigShrinkTime += FPlatformTime::Seconds();
+#endif
 }
 
 
@@ -2991,6 +3014,32 @@ void FConfigBranch::InitFiles()
 	}
 }
 
+void FConfigBranch::RunOnEachFile(TFunction<void(FConfigFile& File, const FString& Name)> Func)
+{
+	// cache the static layers so when remaking dynamic layers after removing a dynamic layer it's faster
+	Func(CombinedStaticLayers, TEXT("CombinedStaticLayers"));
+	Func(FinalCombinedLayers, TEXT("FinalCombinedLayers"));
+	Func(InMemoryFile, TEXT("InMemoryFile"));
+}
+
+void FConfigBranch::RunOnEachCommandStream(TFunction<void(FConfigCommandStream& File, const FString& Name)> Func)
+{
+	for (TPair<FString, FConfigCommandStream>& Pair : StaticLayers)
+	{
+		Func(Pair.Value, Pair.Key);
+	}
+
+	for (DynamicLayerList::TIterator Node(DynamicLayers.GetHead()); Node; ++Node)
+	{
+		Func(*Node.GetNode()->GetValue(), Node->Filename);
+	}
+
+	Func(SavedLayer, TEXT("SavedLayer"));
+	Func(CommandLineOverrides, TEXT("CommandLineOverrides"));
+//	Func(RuntimeChanges, TEXT("RuntimeChanges"));
+}
+
+
 bool FConfigBranch::AddDynamicLayerToHierarchy(const FString& Filename, FConfigModificationTracker* ModificationTracker)
 {
 	if (!DoesConfigFileExistWrapper(*Filename))
@@ -3105,6 +3154,9 @@ bool FConfigBranch::AddDynamicLayersToHierarchy(const TArray<FString>& Filenames
 
 		// re-apply the in-memory changes
 		InMemoryFile.ApplyFile(&Patch);
+
+		FinalCombinedLayers.Shrink();
+		InMemoryFile.Shrink();
 	}
 
 
@@ -3196,6 +3248,9 @@ bool FConfigBranch::AddDynamicLayerStringToHierarchy(const FString& Filename, co
 	// re-apply the in-memory changes
 	InMemoryFile.ApplyFile(&Patch);
 
+	FinalCombinedLayers.Shrink();
+	InMemoryFile.Shrink();
+
 	return true;
 }
 
@@ -3248,6 +3303,9 @@ bool FConfigBranch::RemoveDynamicLayersFromHierarchy(const TArray<FString>& File
 	InMemoryFile = FinalCombinedLayers;
 	InMemoryFile.bCanSaveAllSections = bOldSaveAllSections;
 
+	FinalCombinedLayers.Shrink();
+	InMemoryFile.Shrink();
+
 	// re-apply the in-memory changes
 	InMemoryFile.ApplyFile(&Patch);
 	
@@ -3280,13 +3338,15 @@ bool FConfigBranch::SafeUnload()
 {
 	bIsSafeUnloaded = true;
 
-	InMemoryFile.Cleanup();
-	CombinedStaticLayers.Cleanup();
-	SavedLayer.Empty();
-	CommandLineOverrides.Empty();
-	StaticLayers.Empty();
-	DynamicLayers.Empty();
-	FinalCombinedLayers.Empty();
+	RunOnEachFile([](FConfigFile& File, const FString& Name)
+	{
+		File.Cleanup();
+	});
+
+	RunOnEachCommandStream([](FConfigCommandStream& Stream, const FString& Name)
+	{
+		Stream.Empty();
+	});
 
 	return true;
 }
@@ -3314,6 +3374,19 @@ bool FConfigBranch::RemoveSection(const TCHAR* Section)
 	return NumRemoved > 0;
 }
 
+void FConfigBranch::Shrink()
+{
+	RunOnEachFile([](FConfigFile& File, const FString& Name)
+	{
+		File.Shrink();
+	});
+
+	RunOnEachCommandStream([](FConfigCommandStream& Stream, const FString& Name)
+	{
+		Stream.Shrink();
+	});
+}
+
 void FConfigBranch::Flush()
 {
 	SaveBranch(*this);
@@ -4728,6 +4801,69 @@ protected:
 	SIZE_T Num, Max;
 };
 
+struct FDetailedConfigMemUsage : public FArchiveCountConfigMem
+{
+	TMap<FString, FArchiveCountConfigMem> PerLayerInfo;
+	TMap<FString, FArchiveCountConfigMem> PerSectionInfo;
+	TMap<FString, FArchiveCountConfigMem> PerSectionValueInfo;
+
+	FDetailedConfigMemUsage(FConfigBranch* Branch, bool bTrackDetails)
+	{
+		(*this) << *Branch;
+
+		if (bTrackDetails)
+		{
+			Branch->RunOnEachFile([this](FConfigFile& File, const FString& Name)
+			{
+				TrackFile(Name, File);
+			});
+
+			Branch->RunOnEachCommandStream([this](FConfigCommandStream& Stream, const FString& Name)
+			{
+				TrackCommandStream(Name, Stream);
+			});
+		}
+	}
+
+private:
+	void TrackFile(const FString& Name, FConfigFile& File)
+	{
+		FArchiveCountConfigMem& Ar = PerLayerInfo.FindOrAdd(Name);
+		Ar << File;
+
+		for (const TPair<FString, FConfigSection>& Pair: AsConst(File))
+		{
+			FArchiveCountConfigMem& SectionAr = PerSectionInfo.FindOrAdd(Pair.Key);
+			SectionAr << const_cast<FConfigSection&>(Pair.Value);
+
+			FArchiveCountConfigMem& ValueAr = PerSectionValueInfo.FindOrAdd(Pair.Key);
+			for (const TPair<FName, FConfigValue>& Pair2 : Pair.Value)
+			{
+				ValueAr << const_cast<FConfigValue&>(Pair2.Value);
+			}
+		}
+	}
+
+	void TrackCommandStream(const FString& Name, FConfigCommandStream& Stream)
+	{
+		FArchiveCountConfigMem& Ar = PerLayerInfo.FindOrAdd(Name);
+		Ar << Stream;
+
+		for (auto& Pair : Stream)
+		{
+			FArchiveCountConfigMem& SectionAr = const_cast<FArchiveCountConfigMem&>(PerSectionInfo.FindOrAdd(Pair.Key));
+			SectionAr << Pair.Value;
+
+			FArchiveCountConfigMem& ValueAr = PerSectionValueInfo.FindOrAdd(Pair.Key);
+			for (const TPair<FName, FConfigValue>& Pair2 : Pair.Value)
+			{
+				ValueAr << const_cast<FConfigValue&>(Pair2.Value);
+			}
+		}
+	}
+};
+
+
 
 /**
  * Tracks the amount of memory used by a single config or loc file
@@ -6146,7 +6282,7 @@ void FConfigCacheIni::RegisterPlugin(FName PluginName, const FString& PluginDir,
 
 double GPrepareForLoadTime = 0;
 double GPerformLoadTime = 0;
-
+double GConfigShrinkTime = 0;
 
 class FIniExec : public FSelfRegisteringExec
 {
@@ -6308,11 +6444,51 @@ class FIniExec : public FSelfRegisteringExec
 		
 		if (FParse::Command(&Cmd, TEXT("Timing")))
 		{
-			Ar.Logf(TEXT("INITIME : PrepareForLoad: %fms, PreformLoad: %fms"), GPrepareForLoadTime * 1000.0, GPerformLoadTime * 1000.0);
+			Ar.Logf(TEXT("INITIME : PrepareForLoad: %fms, PreformLoad: %fms, Shrink: %fms"), GPrepareForLoadTime * 1000.0, GPerformLoadTime * 1000.0, GConfigShrinkTime * 1000.0);
+		}
+
+		if (FParse::Command(&Cmd, TEXT("Shrink")))
+		{
+			TCHAR BranchName[256];
+			if (FParse::Token(Cmd, BranchName, UE_ARRAY_COUNT(BranchName), true))
+			{
+				FConfigBranch* Branch = GConfig->FindBranch(BranchName, BranchName);
+				if (Branch)
+				{
+					Branch->Shrink();
+				}
+			}
 		}
 		
 		if (FParse::Command(&Cmd, TEXT("MemUsage")))
 		{
+			// parse options (default is simple, print to log, 10kb cutoff)
+			bool bUseDetailed = FParse::Param(Cmd, TEXT("detailed"));
+			FString CSVFilename;
+			bool bWriteToCSV = FParse::Value(Cmd, TEXT("-csv="), CSVFilename);
+			bWriteToCSV = bWriteToCSV || FParse::Param(Cmd, TEXT("csv"));
+			int CutoffKB = 10;
+			FParse::Value(Cmd, TEXT("Cutoff="), CutoffKB);
+
+			// handle CSV output
+			FArchive* CSV = nullptr;
+			if (bWriteToCSV)
+			{
+				if (CSVFilename.IsEmpty())
+				{
+					CSVFilename = FPaths::Combine(FPaths::ProjectLogDir(), TEXT("ConfigMemUsage.csv"));
+				}
+				CSV = IFileManager::Get().CreateFileWriter(*CSVFilename, FILEWRITE_AllowRead);
+				if (CSV == nullptr)
+				{
+					Ar.Logf(TEXT("Unable to create CSV file for writing: '%s'"), *CSVFilename);
+					return true;
+				}
+
+				Ar.Logf(TEXT("Dumping to CSV file: '%s'"), *CSVFilename);
+			}
+
+			// init counters
 			uint64 Total = 0;
 			int NumSkipped = 0;
 			uint64 SkippedTotal = 0;
@@ -6320,15 +6496,17 @@ class FIniExec : public FSelfRegisteringExec
 			uint64 SingleSectionTotal = 0;
 			int NoSection = 0;
 			uint64 NoSectionTotal = 0;
+
+			uint64 SlackTotal = 0;
 			for (const FString& Filename : GConfig->GetFilenames())
 			{
-				FArchiveCountConfigMem MemAr;
-
 				FConfigBranch* Branch = GConfig->FindBranchWithNoReload(*Filename, Filename);
-				MemAr << *Branch;
+
+				FDetailedConfigMemUsage MemAr(Branch, bUseDetailed);
 
 				uint64 Mem = MemAr.GetMax();
 				Total += Mem;
+				SlackTotal += MemAr.GetMax() - MemAr.GetNum();
 
 				if (Branch->InMemoryFile.Num() == 1)
 				{
@@ -6342,21 +6520,146 @@ class FIniExec : public FSelfRegisteringExec
 				}
 
 				// don't bother printing the neglibly sized ones as they are just noise, so cut off anything < 10kb
-				if (Mem < 10 * 1024)
+				if (Mem < CutoffKB * 1024)
 				{
 					NumSkipped++;
 					SkippedTotal += Mem;
 				}
 				else
 				{
-					Ar.Logf(TEXT("[%0.2fmb] - %s"), (double)MemAr.GetMax() / 1024.0 / 1024.0, *Filename);
+					FArchiveCountConfigMem SectionMem;
+					FArchiveCountConfigMem ValueMem;
+
+
+
+
+					if (bWriteToCSV)
+					{
+						CSV->Logf(TEXT("%0.2fmb,%0.2fmb,%s"), (double)MemAr.GetNum() / 1024.0 / 1024.0, (double)MemAr.GetMax() / 1024.0 / 1024.0, *Filename);
+					}
+					else
+					{
+						Ar.Logf(TEXT("[%0.2fmb / %0.2fmb] - %s"), (double)MemAr.GetNum() / 1024.0 / 1024.0, (double)MemAr.GetMax() / 1024.0 / 1024.0, *Filename);
+					}
+					bool bPrintedHeader = false;
+					for (auto& Pair : MemAr.PerLayerInfo)
+					{
+						if (Pair.Value.GetMax() >= CutoffKB * 1024)
+						{
+							if (bWriteToCSV)
+							{
+								if (!bPrintedHeader)
+								{
+									CSV->Logf(TEXT(",Large layers:"));
+								}
+								CSV->Logf(TEXT(",,%0.2fmb,%0.2fmb,%s"), (double)Pair.Value.GetNum() / 1024.0 / 1024.0, (double)Pair.Value.GetMax() / 1024.0 / 1024.0, *Pair.Key);
+							}
+							else
+							{
+								if (!bPrintedHeader)
+								{
+									Ar.Logf(TEXT("  Large layers:"));
+								}
+								Ar.Logf(TEXT("    [%0.2fmb / %0.2fmb] - %s"), (double)Pair.Value.GetNum() / 1024.0 / 1024.0, (double)Pair.Value.GetMax() / 1024.0 / 1024.0, *Pair.Key);
+							}
+							bPrintedHeader = true;
+						}
+					}
+					bPrintedHeader = false;
+					for (auto& Pair : MemAr.PerSectionInfo)
+					{
+						if (Pair.Value.GetMax() >= CutoffKB * 1024)
+						{
+							if (bWriteToCSV)
+							{
+								if (!bPrintedHeader)
+								{
+									CSV->Logf(TEXT(",Large sections (across all layers):"));
+								}
+								CSV->Logf(TEXT(",,%0.2fmb,%0.2fmb,%s"), (double)Pair.Value.GetNum() / 1024.0 / 1024.0, (double)Pair.Value.GetMax() / 1024.0 / 1024.0, *Pair.Key);
+							}
+							else
+							{
+								if (!bPrintedHeader)
+								{
+									Ar.Logf(TEXT("  Large sections (across all layers):"));
+								}
+								Ar.Logf(TEXT("    [%0.2fmb / %0.2fmb] - %s"), (double)Pair.Value.GetNum() / 1024.0 / 1024.0, (double)Pair.Value.GetMax() / 1024.0 / 1024.0, *Pair.Key);
+							}
+							bPrintedHeader = true;
+						}
+					}
+					bPrintedHeader = false;
+					for (auto& Pair : MemAr.PerSectionValueInfo	)
+					{
+						if (Pair.Value.GetMax() >= CutoffKB * 1024)
+						{
+							if (bWriteToCSV)
+							{
+								if (!bPrintedHeader)
+								{
+									CSV->Logf(TEXT(",Large sections (by values):"));
+								}
+								CSV->Logf(TEXT(",,%0.2fmb,%0.2fmb,%s"), (double)Pair.Value.GetNum() / 1024.0 / 1024.0, (double)Pair.Value.GetMax() / 1024.0 / 1024.0, *Pair.Key);
+							}
+							else
+							{
+								if (!bPrintedHeader)
+								{
+									Ar.Logf(TEXT("  Large sections (by values):"));
+								}
+								Ar.Logf(TEXT("    [%0.2fmb / %0.2fmb] - %s"), (double)Pair.Value.GetNum() / 1024.0 / 1024.0, (double)Pair.Value.GetMax() / 1024.0 / 1024.0, *Pair.Key);
+							}
+							bPrintedHeader = true;
+						}
+					}
 				}
 			}
 
+			if (bWriteToCSV)
+			{
+				CSV->Logf(TEXT(""));
+				CSV->Logf(TEXT("%0.2fmb,%d All Configs"), (double)Total / 1024.0 / 1024.0, GConfig->GetFilenames().Num());
+				CSV->Logf(TEXT("%0.2fmb,%d Tiny Configs (not displayed above)"), (double)SkippedTotal / 1024.0 / 1024.0, NumSkipped);
+				CSV->Logf(TEXT("%0.2fmb,%d Single Section Configs"), (double)SingleSectionTotal / 1024.0 / 1024.0, SingleSection);
+				CSV->Logf(TEXT("%0.2fmb,%d ZeroSection Configs"), (double)NoSectionTotal / 1024.0 / 1024.0, NoSection);
+				CSV->Logf(TEXT("%0.2fmb,Total Slack (wasted memory)"), (double)SlackTotal / 1024.0 / 1024.0);
+				CSV->Logf(TEXT(""));
+				if (!bUseDetailed)
+				{
+					CSV->Logf(TEXT("To get more detailed information, use \"config memusage -detailed\""));
+				}
+				if (CutoffKB == 10)
+				{
+					CSV->Logf(TEXT("To change the cutoff, in KB, for small files/layers/sections, use \"config memusage -cutoff=<value>\""));
+				}
+#if WITH_EDITOR
+				CSV->Logf(TEXT("(Note: Editor builds store more layer state, so the memory usage will be higher than in a client build)"));
+#endif
+			}
+			else
+			{
+				Ar.Logf(TEXT(""));
 			Ar.Logf(TEXT("[%0.2fmb] - %d All Configs"), (double)Total / 1024.0 / 1024.0, GConfig->GetFilenames().Num());
-			Ar.Logf(TEXT("[%0.2fmb] - %d Tiny Configs (not displayed)"), (double)SkippedTotal / 1024.0 / 1024.0, NumSkipped);
+				Ar.Logf(TEXT("[%0.2fmb] - %d Tiny Configs (not displayed above)"), (double)SkippedTotal / 1024.0 / 1024.0, NumSkipped);
 			Ar.Logf(TEXT("[%0.2fmb] - %d Single Section Configs"), (double)SingleSectionTotal / 1024.0 / 1024.0, SingleSection);
 			Ar.Logf(TEXT("[%0.2fmb] - %d ZeroSection Configs"), (double)NoSectionTotal / 1024.0 / 1024.0, NoSection);
+				Ar.Logf(TEXT("[%0.2fmb] - Total Slack (wasted memory)"), (double)SlackTotal / 1024.0 / 1024.0);
+				Ar.Logf(TEXT(""));
+				if (!bUseDetailed)
+				{
+					Ar.Logf(TEXT("To get more detailed information, use \"config memusage -detailed\""));
+				}
+				if (CutoffKB == 10)
+				{
+					Ar.Logf(TEXT("To change the cutoff, in KB, for small files/layers/sections, use \"config memusage -cutoff=<value>\""));
+				}
+				Ar.Logf(TEXT("To save to .csv, use \"config memusage -csv or -csv=<filepath>\""));
+#if WITH_EDITOR
+				Ar.Logf(TEXT("(Note: Editor builds store more layer state, so the memory usage will be higher than in a client build)"));
+#endif
+			}
+			delete CSV;
 		}
 
 		return true;
diff --git a/Engine/Source/Runtime/Core/Private/Misc/ConfigContext.cpp b/Engine/Source/Runtime/Core/Private/Misc/ConfigContext.cpp
--- a/Engine/Source/Runtime/Core/Private/Misc/ConfigContext.cpp
+++ b/Engine/Source/Runtime/Core/Private/Misc/ConfigContext.cpp
@@ -394,7 +394,9 @@ bool FConfigContext::PerformSingleFileLoad()
 
 bool FConfigContext::PrepareForLoad(bool& bPerformLoad)
 {
+#if !UE_BUILD_SHIPPING
 	if (IsInGameThread()) GPrepareForLoadTime -= FPlatformTime::Seconds();
+#endif
 
 	checkf(ConfigSystem != nullptr || ExistingFile != nullptr, TEXT("Loading config expects to either have a ConfigFile already passed in, or have a ConfigSystem passed in"));
 
@@ -488,7 +490,9 @@ bool FConfigContext::PerformLoad()
 {
 	LLM_SCOPE(ELLMTag::ConfigSystem);
 
+#if !UE_BUILD_SHIPPING
 	if (IsInGameThread()) GPerformLoadTime -= FPlatformTime::Seconds();
+#endif
 
 #if DISABLE_GENERATED_INI_WHEN_COOKED
 	if (BaseIniName == TEXT("GameUserSettings"))
@@ -584,7 +588,9 @@ bool FConfigContext::PerformLoad()
 		}
 	}
 
+#if !UE_BUILD_SHIPPING
 	if (IsInGameThread()) GPerformLoadTime += FPlatformTime::Seconds();
+#endif
 
 	return bGeneratedFile;
 }
@@ -989,11 +995,18 @@ bool FConfigContext::LoadIniFileHierarchy()
 	// if we had been reading into the Static cache, not InMemory, then start the InMemory from this point
 	if (Branch->ReplayMethod != EBranchReplayMethod::NoReplay)
 	{
+		Branch->CombinedStaticLayers.Shrink();
+
 		Branch->InMemoryFile = Branch->CombinedStaticLayers;
 		
 		// need to reset this since it just got blown away
 		Branch->InMemoryFile.ChangeTracker = &Branch->SavedLayer;
 	}
+	else
+	{
+		Branch->InMemoryFile.Shrink();
+	}
+
 	Branch->FinalCombinedLayers = Branch->InMemoryFile;
 	return bReadAnyFile;
 }
diff --git a/Engine/Source/Runtime/Core/Public/Misc/ConfigCacheIni.h b/Engine/Source/Runtime/Core/Public/Misc/ConfigCacheIni.h
--- a/Engine/Source/Runtime/Core/Public/Misc/ConfigCacheIni.h
+++ b/Engine/Source/Runtime/Core/Public/Misc/ConfigCacheIni.h
@@ -316,19 +316,19 @@ public:
 	bool operator==(const FConfigValue& Other) const { return SavedValueHash == Other.SavedValueHash; }
 	bool operator!=(const FConfigValue& Other) const { return !(FConfigValue::operator==(Other)); }
 
-	friend FArchive& operator<<(FArchive& Ar, FConfigValue& ConfigSection)
+	friend FArchive& operator<<(FArchive& Ar, FConfigValue& ConfigValue)
 	{
-		FStructuredArchiveFromArchive(Ar).GetSlot() << ConfigSection;
+		FStructuredArchiveFromArchive(Ar).GetSlot() << ConfigValue;
 		return Ar;
 	}
 
-	friend void operator<<(FStructuredArchive::FSlot Slot, FConfigValue& ConfigSection)
+	friend void operator<<(FStructuredArchive::FSlot Slot, FConfigValue& ConfigValue)
 	{
-		Slot << ConfigSection.SavedValue;
+		Slot << ConfigValue.SavedValue;
 
 		if (Slot.GetUnderlyingArchive().IsLoading())
 		{
-			ConfigSection.ExpandValueInternal();
+			ConfigValue.ExpandValueInternal();
 		}
 	}
 
@@ -387,6 +387,7 @@ public:
 	friend class FConfigCacheIni;
 	friend class FConfigFile;
 	friend class FConfigBranch;
+	friend struct FDetailedConfigMemUsage;
 	const FString& GetSavedValueForWriting() const
 	{
 		return SavedValue;
@@ -1126,9 +1127,18 @@ public:
 	CORE_API bool RemoveSection(const TCHAR* Section);
 
 	CORE_API void Flush();
-	
+
+	CORE_API void Shrink();
+
 	CORE_API void Dump(FOutputDevice& Ar);
-	
+
+
+	/**
+	 * Run a function on every file in the branch
+	 */
+	CORE_API void RunOnEachFile(TFunction<void(FConfigFile& File, const FString& Name)> Func);
+	CORE_API void RunOnEachCommandStream(TFunction<void(FConfigCommandStream& File, const FString& Name)> Func);
+
 private:
 	void InitFiles();
 	
