Updated FTextLocalizationManager to use a RWLock for its live table data

This is possible now that GetDisplayString no longer ad-hoc mutates the live table

#rnx

diff --git a/Engine/Source/Runtime/Core/Private/Internationalization/TextLocalizationManager.cpp b/Engine/Source/Runtime/Core/Private/Internationalization/TextLocalizationManager.cpp
--- a/Engine/Source/Runtime/Core/Private/Internationalization/TextLocalizationManager.cpp
+++ b/Engine/Source/Runtime/Core/Private/Internationalization/TextLocalizationManager.cpp
@@ -711,7 +711,7 @@ FTextLocalizationManager::~FTextLocalizationManager()
 void FTextLocalizationManager::DumpMemoryInfo() const
 {
 	{
-		FScopeLock ScopeLock(&DisplayStringTableCS);
+		FReadScopeLock ScopeLock(DisplayStringTableRW);
 		UE_LOG(LogTextLocalizationManager, Log, TEXT("DisplayStringLookupTable.GetAllocatedSize()=%d elems=%d"), DisplayStringLookupTable.GetAllocatedSize(), DisplayStringLookupTable.Num());
 	}
 	{
@@ -727,7 +727,7 @@ void FTextLocalizationManager::CompactDataStructures()
 
 	double StartTime = FPlatformTime::Seconds();
 	{
-		FScopeLock ScopeLock(&DisplayStringTableCS);
+		FWriteScopeLock ScopeLock(DisplayStringTableRW);
 		DisplayStringLookupTable.Shrink();
 	}
 	{
@@ -750,7 +750,7 @@ void FTextLocalizationManager::DumpLiveTableImpl(const FString* NamespaceFilter,
 			return !Filter || Str.MatchesWildcard(*Filter, ESearchCase::IgnoreCase); // Note: This is case insensitive since its used from a debug command
 		};
 
-		FScopeLock ScopeLock(&DisplayStringTableCS);
+		FReadScopeLock ScopeLock(DisplayStringTableRW);
 		DisplayStringLookupTableToDump.Reserve(DisplayStringLookupTable.Num());
 		for (const auto& DisplayStringPair : DisplayStringLookupTable)
 		{
@@ -761,17 +761,18 @@ void FTextLocalizationManager::DumpLiveTableImpl(const FString* NamespaceFilter,
 				DisplayStringLookupTableToDump.Add(DisplayStringPair.Key, DisplayStringPair.Value);
 			}
 		}
-		DisplayStringLookupTableToDump.KeySort([](const FTextId& A, const FTextId& B)
-		{
-			const int32 NamespaceResult = FCString::Strcmp(A.GetNamespace().GetChars(), B.GetNamespace().GetChars());
-			if (NamespaceResult != 0)
-			{
-				return NamespaceResult < 0;
-			}
-			return FCString::Strcmp(A.GetKey().GetChars(), B.GetKey().GetChars()) < 0;
-		});
 	}
 
+	DisplayStringLookupTableToDump.KeySort([](const FTextId& A, const FTextId& B)
+	{
+		const int32 NamespaceResult = FCString::Strcmp(A.GetNamespace().GetChars(), B.GetNamespace().GetChars());
+		if (NamespaceResult != 0)
+		{
+			return NamespaceResult < 0;
+		}
+		return FCString::Strcmp(A.GetKey().GetChars(), B.GetKey().GetChars()) < 0;
+	});
+
 	for (const auto& DisplayStringPair : DisplayStringLookupTableToDump)
 	{
 		Callback(DisplayStringPair.Key, DisplayStringPair.Value.DisplayString);
@@ -812,7 +813,7 @@ void FTextLocalizationManager::AddOrUpdateDisplayStringInLiveTable(const FString
 	const FTextId TextId(Namespace, Key);
 
 	// Lock while updating the table
-	FScopeLock ScopeLock(&DisplayStringTableCS);
+	FWriteScopeLock ScopeLock(DisplayStringTableRW);
 
 	if (FDisplayStringEntry* LiveEntry = DisplayStringLookupTable.Find(TextId))
 	{
@@ -891,7 +892,7 @@ TArray<FString> FTextLocalizationManager::GetLocalizedCultureNames(const ELocali
 
 int32 FTextLocalizationManager::GetLocalizationTargetPathId(FStringView InLocalizationTargetPath)
 {
-	FScopeLock ScopeLock(&DisplayStringTableCS);
+	FWriteScopeLock ScopeLock(DisplayStringTableRW);
 
 	int32 LocalizationTargetPathId = INDEX_NONE;
 	DisplayStringsByLocalizationTargetId.FindOrAdd(InLocalizationTargetPath, &LocalizationTargetPathId);
@@ -1015,7 +1016,7 @@ FTextConstDisplayStringPtr FTextLocalizationManager::FindDisplayString_Internal(
 	}
 #endif	// ENABLE_LOC_TESTING
 
-	FScopeLock ScopeLock(&DisplayStringTableCS);
+	FReadScopeLock ScopeLock(DisplayStringTableRW);
 
 	if (const FDisplayStringEntry* LiveEntry = DisplayStringLookupTable.Find(TextId))
 	{
@@ -1110,7 +1111,7 @@ FTextConstDisplayStringPtr FTextLocalizationManager::GetDisplayString(const FTex
 #if WITH_EDITORONLY_DATA
 bool FTextLocalizationManager::GetLocResID(const FTextKey& Namespace, const FTextKey& Key, FString& OutLocResId) const
 {
-	FScopeLock ScopeLock(&DisplayStringTableCS);
+	FReadScopeLock ScopeLock(DisplayStringTableRW);
 
 	const FTextId TextId(Namespace, Key);
 
@@ -1217,7 +1218,7 @@ void FTextLocalizationManager::HandleLocalizationTargetsMounted(TArrayView<const
 
 	// Mark the targets as mounted before loading any of their data
 	{
-		FScopeLock ScopeLock(&DisplayStringTableCS);
+		FWriteScopeLock ScopeLock(DisplayStringTableRW);
 		for (const FString& LocalizationTargetPath : LocalizationTargetPaths)
 		{
 			FDisplayStringsForLocalizationTarget& DisplayStringsForLocalizationTarget = DisplayStringsByLocalizationTargetId.FindOrAdd(LocalizationTargetPath);
@@ -1259,7 +1260,7 @@ void FTextLocalizationManager::HandleLocalizationTargetsUnmounted(TArrayView<con
 		FTextCache& TextCache = FTextCache::Get();
 
 		// Lock while updating the tables
-		FScopeLock ScopeLock(&TLM.DisplayStringTableCS);
+		FWriteScopeLock ScopeLock(TLM.DisplayStringTableRW);
 
 		// Discard the data for each localization target that was unmounted, and mark the target as no longer mounted so that we no longer track its text IDs
 		for (const FString& LocalizationTargetPath : LocalizationTargetPaths)
@@ -1618,7 +1619,7 @@ void FTextLocalizationManager::UpdateFromNative(FTextLocalizationResource&& Text
 
 	// Lock while updating the tables
 	{
-		FScopeLock ScopeLock(&DisplayStringTableCS);
+		FWriteScopeLock ScopeLock(DisplayStringTableRW);
 
 		DisplayStringLookupTable.Reserve(TextLocalizationResource.Entries.Num());
 
@@ -1686,7 +1687,7 @@ void FTextLocalizationManager::UpdateFromLocalizations(FTextLocalizationResource
 
 	// Lock while updating the tables
 	{
-		FScopeLock ScopeLock(&DisplayStringTableCS);
+		FWriteScopeLock ScopeLock(DisplayStringTableRW);
 
 		DisplayStringLookupTable.Reserve(TextLocalizationResource.Entries.Num());
 
diff --git a/Engine/Source/Runtime/Core/Public/Internationalization/TextLocalizationManager.h b/Engine/Source/Runtime/Core/Public/Internationalization/TextLocalizationManager.h
--- a/Engine/Source/Runtime/Core/Public/Internationalization/TextLocalizationManager.h
+++ b/Engine/Source/Runtime/Core/Public/Internationalization/TextLocalizationManager.h
@@ -120,7 +120,7 @@ private:
 		return InitializedFlags != ETextLocalizationManagerInitializedFlags::None;
 	}
 
-	mutable FCriticalSection DisplayStringTableCS;
+	mutable FRWLock DisplayStringTableRW;
 	FDisplayStringLookupTable DisplayStringLookupTable;
 	FDisplayStringsByLocalizationTargetId DisplayStringsByLocalizationTargetId;
 
