Fix remaining new compiler analyze warnings

#rnx

diff --git a/Engine/Plugins/Tests/RHITests/Source/RHITests/Private/RHIClearTextureTests.cpp b/Engine/Plugins/Tests/RHITests/Source/RHITests/Private/RHIClearTextureTests.cpp
--- a/Engine/Plugins/Tests/RHITests/Source/RHITests/Private/RHIClearTextureTests.cpp
+++ b/Engine/Plugins/Tests/RHITests/Source/RHITests/Private/RHIClearTextureTests.cpp
@@ -1541,7 +1541,8 @@ bool FRHIClearTextureTests::Test_ClearTexture(FRHICommandListImmediate& RHICmdLi
 		FTestContext TestContext(VertexDeclarationRHI, RHITextureCreateDescs, SentinelColor);
 		EPixelFormat CurrentPixelFormat = PF_Unknown;
 
-		bool bGPUCapture = GGPUCaptureTest && IRenderCaptureProvider::IsAvailable();
+		bool bGGPUCaptureTest = GGPUCaptureTest;
+		bool bGPUCapture = bGGPUCaptureTest && IRenderCaptureProvider::IsAvailable();
 		if (bGPUCapture)
 		{
 			Test_ClearTextureBeginCapture(RHICmdList);
diff --git a/Engine/Source/Runtime/Core/Public/HAL/IPlatformFileCachedWrapper.h b/Engine/Source/Runtime/Core/Public/HAL/IPlatformFileCachedWrapper.h
--- a/Engine/Source/Runtime/Core/Public/HAL/IPlatformFileCachedWrapper.h
+++ b/Engine/Source/Runtime/Core/Public/HAL/IPlatformFileCachedWrapper.h
@@ -285,7 +285,9 @@ public:
 #define PLATFORM_PROVIDES_FILE_CACHE 0
 #endif
 		// Default to false on platforms that already do platform file level caching
-		bool bResult = !PLATFORM_PROVIDES_FILE_CACHE && !PLATFORM_WINDOWS && FPlatformProperties::RequiresCookedData();
+		bool bPlatformProvidesFileCache = PLATFORM_PROVIDES_FILE_CACHE;
+		bool bPlatformWindows = PLATFORM_WINDOWS;
+		bool bResult = !bPlatformProvidesFileCache && !bPlatformWindows && FPlatformProperties::RequiresCookedData();
 
 		// Allow a choice between shorter load times or less memory on desktop platforms.
 		// Note: this cannot be in config since they aren't read at that point.
diff --git a/Engine/Source/Runtime/Engine/Private/Animation/AnimStreamable.cpp b/Engine/Source/Runtime/Engine/Private/Animation/AnimStreamable.cpp
--- a/Engine/Source/Runtime/Engine/Private/Animation/AnimStreamable.cpp
+++ b/Engine/Source/Runtime/Engine/Private/Animation/AnimStreamable.cpp
@@ -682,7 +682,7 @@ void UAnimStreamable::RequestCompressedDataForChunk(const FString& ChunkDDCKey,
 		const FString FinalDDCKey = FDerivedDataCacheInterface::BuildCacheKey(AnimCompressor->GetPluginName(), AnimCompressor->GetVersionString(), *AnimCompressor->GetPluginSpecificCacheKeySuffix());
 
 		// For debugging DDC/Compression issues		
-		const bool bSkipDDC = false;
+		bool bSkipDDC = false;
 
 		const int32 ChunkNumFrames = FrameEnd - FrameStart;
 		Chunk.StartTime = SamplingFrameRate.AsSeconds(static_cast<int32>(FrameStart));
diff --git a/Engine/Source/Runtime/Engine/Private/PacketHandlers/StatelessConnectHandlerComponent.cpp b/Engine/Source/Runtime/Engine/Private/PacketHandlers/StatelessConnectHandlerComponent.cpp
--- a/Engine/Source/Runtime/Engine/Private/PacketHandlers/StatelessConnectHandlerComponent.cpp
+++ b/Engine/Source/Runtime/Engine/Private/PacketHandlers/StatelessConnectHandlerComponent.cpp
@@ -1685,8 +1685,8 @@ bool StatelessConnectHandlerComponent::CheckVersion(const FParsedHandshakeData&
 	}
 
 	const EEngineNetworkRuntimeFeatures LocalNetworkFeatures = Driver ? Driver->GetNetworkRuntimeFeatures() : EEngineNetworkRuntimeFeatures::None;
-
-	const bool bCheckNetVersion = !UE_BUILD_SHIPPING || GHandshakeEnforceNetworkCLVersion;
+	bool bUEBuildShipping = UE_BUILD_SHIPPING;
+	const bool bCheckNetVersion = !bUEBuildShipping || GHandshakeEnforceNetworkCLVersion;
 	bool bValidNetCLVersion = true;
 	bool bIsNetFeatureCompatible = true;
 
diff --git a/Engine/Source/Runtime/SlateCore/Private/Test/SlateInvalidationWidgetListTest.cpp b/Engine/Source/Runtime/SlateCore/Private/Test/SlateInvalidationWidgetListTest.cpp
--- a/Engine/Source/Runtime/SlateCore/Private/Test/SlateInvalidationWidgetListTest.cpp
+++ b/Engine/Source/Runtime/SlateCore/Private/Test/SlateInvalidationWidgetListTest.cpp
@@ -367,7 +367,8 @@ bool FSlateInvalidationWidgetListTest::RunTest(const FString& Parameters)
 
 			{
 				bool bChildRemoved = TestRemoveWidget(WidgetC, TEXT("C"));
-				AddErrorIfFalse(!bChildRemoved || !UE::Slate::Private::bUpdateOnlyWhatIsNeeded, TEXT("C no widget was removed"));
+				bool bUpdateOnlyWhatIsNeeded = UE::Slate::Private::bUpdateOnlyWhatIsNeeded;
+				AddErrorIfFalse(!bChildRemoved || !bUpdateOnlyWhatIsNeeded, TEXT("C no widget was removed"));
 			}
 		}
 	}
diff --git a/Engine/Source/Runtime/VulkanRHI/Private/VulkanLayers.cpp b/Engine/Source/Runtime/VulkanRHI/Private/VulkanLayers.cpp
--- a/Engine/Source/Runtime/VulkanRHI/Private/VulkanLayers.cpp
+++ b/Engine/Source/Runtime/VulkanRHI/Private/VulkanLayers.cpp
@@ -399,7 +399,8 @@ void FVulkanIntanceSetupHelper::AddDebugLayers(const TArray<FLayerWithExtensions
 		}
 	}
 
-	const bool bForceDebugUtils = VULKAN_ENABLE_DRAW_MARKERS || FParse::Param(FCommandLine::Get(), TEXT("vulkandebugutils"));
+	bool bVulkanEnableDrawMarkers = VULKAN_ENABLE_DRAW_MARKERS;
+	const bool bForceDebugUtils = bVulkanEnableDrawMarkers || FParse::Param(FCommandLine::Get(), TEXT("vulkandebugutils"));
 	if ((bUseVulkanValidation || bForceDebugUtils) && (ActiveDebugLayerExtension == FVulkanDynamicRHI::EActiveDebugLayerExtension::None))
 	{
 		auto FindLayerContainingExtension = [](const ANSICHAR* ExtensionName, const TArray<FLayerWithExtensions>& LayerProperties)
