- Fixed issues with adding multiple dynamic layers to a branch at one time (potentially only the last one added would be applied, and cvar tracking could also only use the final file's sections)
- Allow dynamic layer platform inis without the non-platform base ini existing

diff --git a/Engine/Source/Runtime/Core/Private/Misc/ConfigCacheIni.cpp b/Engine/Source/Runtime/Core/Private/Misc/ConfigCacheIni.cpp
--- a/Engine/Source/Runtime/Core/Private/Misc/ConfigCacheIni.cpp
+++ b/Engine/Source/Runtime/Core/Private/Misc/ConfigCacheIni.cpp
@@ -3052,74 +3052,96 @@ bool FConfigBranch::AddDynamicLayerToHierarchy(const FString& Filename, FConfigM
 
 bool FConfigBranch::AddDynamicLayersToHierarchy(const TArray<FString>& Filenames, FName Tag, DynamicLayerPriority Priority, FConfigModificationTracker* ModificationTracker)
 {
+	static bool bDumpIniLoadInfo = FParse::Param(FCommandLine::Get(), TEXT("dumpiniloads"));
+
 	bool bFoundAFile = false;
-	bool bInsertedAtEnd = false;
+	bool bInsertedBeforeEnd = false;
 
 	// calculate a patch so we don't lose in-memory changes
 	FConfigCommandStream Patch;
-	FConfigCommandStream* DynamicLayer=nullptr;
-	FConfigCommandStream LocalLayer;
+
+	TArray<FConfigCommandStream*> AddedLayers;
 	for (const FString& Filename : Filenames)
 	{
+		UE_CLOG(bDumpIniLoadInfo, LogConfig, Display, TEXT("Looking for file: %s"), *Filename);
+
+		UE_LOG(LogConfig, Verbose, TEXT("Adding Dynamic layer %s to Branch %s"), *Filename, *IniName.ToString());
+
 		if (!DoesConfigFileExistWrapper(*Filename))
 		{
+			UE_LOG(LogConfig, Verbose, TEXT("  .. doesn't exist!"));
 			continue;
 		}
 
-		if (!bFoundAFile)
+		UE_CLOG(bDumpIniLoadInfo, LogConfig, Display, TEXT("   Found %s!"), *Filename);
+
+		if (AddedLayers.Num() == 0)
 		{
 			Patch = CalculateDiff(FinalCombinedLayers, InMemoryFile);
+			UE_LOG(LogConfig, Verbose, TEXT("  .. calculating diff on first file"));
 		}
 		
+		// make and read in the layer
+		FConfigCommandStream* DynamicLayer = new FConfigCommandStream;
+		FillFileFromDisk(DynamicLayer, Filename, true);
+		DynamicLayer->Priority = (uint16)Priority;
+		DynamicLayer->Filename = Filename;
+		DynamicLayer->Tag = Tag;
+
+		// remember in local array, then figure out how to remember it permanently
+		AddedLayers.Add(DynamicLayer);
+
 		// if we aren't caching dynamic layers, then we need a temp layer
 		if (ReplayMethod == EBranchReplayMethod::NoReplay)
 		{
-			DynamicLayer = &LocalLayer;
-			bInsertedAtEnd = true;
+			UE_LOG(LogConfig, Verbose, TEXT("  .. no replay, so just adding at end"));
 		}
 		else
 		{
 			// find the first node with higher priority
 			// @todo move this to a function
 			bool bInserted = false;
-			DynamicLayer = new FConfigCommandStream;
-			DynamicLayer->Priority = (uint16)Priority;
-			DynamicLayer->Filename = Filename;
 			for (DynamicLayerList::TIterator Node(DynamicLayers.GetHead()); Node; ++Node)
 			{
 				if (Node->Priority > DynamicLayer->Priority)
 				{
+					UE_LOG(LogConfig, Verbose, TEXT("  .. inserted in middle of dynamic layers"));
 					DynamicLayers.InsertNode(DynamicLayer, Node.GetNode());
+					bInsertedBeforeEnd = true;
 					bInserted = true;
 					break;
 				}
 			}
 			if (!bInserted)
 			{
+				UE_LOG(LogConfig, Verbose, TEXT("  .. inserting at end of layers"));
 				DynamicLayers.AddTail(DynamicLayer);
-				bInsertedAtEnd = true;
 			}
 		}
 
-		// load in the file
-		FillFileFromDisk(DynamicLayer, Filename, true);
-		DynamicLayer->Tag = Tag;
-		
 		// track modified section names if desired
 		if (ModificationTracker != nullptr)
 		{
 			if (ModificationTracker->bTrackModifiedSections)
 			{
+				UE_LOG(LogConfig, Verbose, TEXT("  .. tracking sections:"));
 				for (const TPair<FString, FConfigCommandStreamSection>& Pair : *DynamicLayer)
 				{
 					TSet<FString>& ModifiedSections = ModificationTracker->ModifiedSectionsPerBranch.FindOrAdd(IniName);
 					ModifiedSections.Add(Pair.Key);
+					UE_LOG(LogConfig, Verbose, TEXT("  .. .. %s"), *Pair.Key);
 					if (FConfigModificationTracker::FCVarTracker* CVarTracker = ModificationTracker->CVars.Find(Pair.Key))
 					{
-						FConfigSection NewSection;
-						// copy just the SectionMap parts
-						(FConfigSectionMap&)NewSection = (const FConfigSectionMap&)Pair.Value;
-						CVarTracker->CVarEntriesPerBranch.Add(IniName, NewSection);
+						UE_LOG(LogConfig, Verbose, TEXT("  .. .. .. tracking cvars"), *Pair.Key);
+						
+						const FConfigSectionMap& ModifiedCVars = (const FConfigSectionMap&)Pair.Value;  
+						FConfigSection& TrackedCVarSection = CVarTracker->CVarEntriesPerBranch.FindOrAdd(IniName);
+						for (const TPair<FName, FConfigValue>& CVarPair : ModifiedCVars)
+						{
+							UE_LOG(LogConfig, Verbose, TEXT("  .. .. .. .. %s = %s"), *CVarPair.Key.ToString(), *CVarPair.Value.GetSavedValue());
+							TrackedCVarSection.Remove(CVarPair.Key);
+							TrackedCVarSection.Add(CVarPair.Key, CVarPair.Value);
+						}
 					}
 				}
 			}
@@ -3128,16 +3150,25 @@ bool FConfigBranch::AddDynamicLayersToHierarchy(const TArray<FString>& Filenames
 				ModificationTracker->LoadedFiles.Add(Filename);
 			}
 		}
-
-		bFoundAFile = true;
 	}
 	
-	if (bFoundAFile)
+	if (AddedLayers.Num() > 0)
 	{
-		if (!bInsertedAtEnd)
+		// if all were added at the end (or there's no replay), we can just apply them without rewinding
+		if (!bInsertedBeforeEnd)
+		{
+			for (FConfigCommandStream* NewLayer : AddedLayers)
+			{
+				UE_LOG(LogConfig, Verbose, TEXT("  .. reapplying layer with %d sections"), NewLayer->Num());
+				FinalCombinedLayers.ApplyFile(NewLayer);
+				InMemoryFile.ApplyFile(NewLayer);
+			}
+		}
+		else
 		{
 			// rebuild
 			FinalCombinedLayers = CombinedStaticLayers;
+			UE_LOG(LogConfig, Verbose, TEXT("  .. reapplying all dynamic layers"));
 			for (DynamicLayerList::TIterator Node(DynamicLayers.GetHead()); Node; ++Node)
 			{
 				FinalCombinedLayers.ApplyFile(*Node);
@@ -3146,11 +3177,6 @@ bool FConfigBranch::AddDynamicLayersToHierarchy(const TArray<FString>& Filenames
 			InMemoryFile = FinalCombinedLayers;
 			InMemoryFile.bCanSaveAllSections = bOldSaveAllSections;
 		}
-		else
-		{
-			FinalCombinedLayers.ApplyFile(DynamicLayer);
-			InMemoryFile.ApplyFile(DynamicLayer);
-		}
 
 		// re-apply the in-memory changes
 		InMemoryFile.ApplyFile(&Patch);
@@ -3159,7 +3185,6 @@ bool FConfigBranch::AddDynamicLayersToHierarchy(const TArray<FString>& Filenames
 		InMemoryFile.Shrink();
 	}
 
-
 	return bFoundAFile;
 }
 
@@ -6103,23 +6128,44 @@ void FConfigCacheIni::AddPluginToBranches(FName PluginName, FConfigModificationT
 		}
 	}
 	
+	FString PlatformNameStr(PlatformName.ToString());
 	TArray<FString> PluginConfigs;
 	FString PluginConfigDir = FPaths::Combine(PluginInfo->PluginDir, TEXT("Config"));
+	FString PlatformConfigDir = FPaths::Combine(PluginConfigDir, PlatformNameStr);
 	IFileManager::Get().FindFiles(PluginConfigs, *PluginConfigDir, TEXT("ini"));
+	IFileManager::Get().FindFiles(PluginConfigs, *PlatformConfigDir, TEXT("ini"));
+	
+	// if this plugin has any platform extensions, then we need to look in them for files, in so that we can load them
+	// even if there is no platform-less config file in the plugin itself
+	for (FString& ChildPluginDir : PluginInfo->ChildPluginDirs)
+	{
+		if (ChildPluginDir.Contains(*FString::Printf(TEXT("/%s/"), *PlatformNameStr)))
+		{
+			FString PlatformExtConfigDir = FPaths::Combine(ChildPluginDir, TEXT("Config"));
+			IFileManager::Get().FindFiles(PluginConfigs, *PlatformExtConfigDir, TEXT("ini"));
+		}
+	}
 	
 	// make a single context that can be used for all the branches modified by this plugin
-	FConfigContext Context = FConfigContext::ReadIntoConfigSystem(this, PlatformName.ToString());
+	FConfigContext Context = FConfigContext::ReadIntoConfigSystem(this, PlatformNameStr);
 	Context.bIsForPluginModification = true;
 	Context.PluginModificationPriority = PluginInfo->Priority;
 	Context.bIncludeTagNameInBranchName = PluginInfo->bIncludePluginNameInBranchName;
 	Context.ChangeTracker = ModificationTracker;
 	
-	// find branches that are found in the plugin dir, platform extension dirs can be used, but there must be a .ini file in Config for it to search
+	// find branches that are found in the plugin dir or it's platform dirs
 	FString StrippedPart = PluginInfo->bIncludePluginNameInBranchName ? PluginName.ToString() : FString();
 	FName CurrentPlatform(FPlatformProperties::IniPlatformName());
+	
+	TSet<FName> LoadedBranches;
 	for (const FString& ConfigFilename : PluginConfigs)
 	{
-		FName BranchName = *FPaths::GetBaseFilename(ConfigFilename).Replace(*StrippedPart, TEXT(""));
+		FName BranchName = *FPaths::GetBaseFilename(ConfigFilename).Replace(*StrippedPart, TEXT("")).Replace(*PlatformNameStr, TEXT(""));
+		if (LoadedBranches.Contains(BranchName))
+		{
+			continue;
+		}
+		LoadedBranches.Add(BranchName);
 		
 		// if we have been tracking loaded files, and we've already loaded this file, we can skip it (it would be the DefaultMyPlugin.ini type of file)
 		if (ModificationTracker && ModificationTracker->bTrackLoadedFiles)
