[TaskGraph]
 - Add CVar to control if we should always pump named threads queues while waiting on named threads or just when strictly necessary

#rnx

diff --git a/Engine/Source/Runtime/Core/Private/Tasks/TaskPrivate.cpp b/Engine/Source/Runtime/Core/Private/Tasks/TaskPrivate.cpp
--- a/Engine/Source/Runtime/Core/Private/Tasks/TaskPrivate.cpp
+++ b/Engine/Source/Runtime/Core/Private/Tasks/TaskPrivate.cpp
@@ -5,6 +5,14 @@
 
 #include "Async/TaskGraphInterfaces.h"
 
+CORE_API bool GTaskGraphAlwaysWaitWithNamedThreadSupport = 1;
+static FAutoConsoleVariableRef CVarTaskGraphAlwaysWaitWithNamedThreadSupport(
+	TEXT("TaskGraph.AlwaysWaitWithNamedThreadSupport"),
+	GTaskGraphAlwaysWaitWithNamedThreadSupport,
+	TEXT("Default to pumping the named thread tasks when waiting on named threads to avoid potential deadlocks."),
+	ECVF_ReadOnly
+);
+
 namespace UE::Tasks
 {
 	namespace Private
@@ -197,9 +205,17 @@ namespace UE::Tasks
 			return WaitImpl(Timeout);
 		}
 
+		bool ShouldForceWaitWithNamedThreadsSupport(EExtendedTaskPriority Priority);
 		void FTaskBase::Wait()
 		{
-			WaitWithNamedThreadsSupport();
+			if (GTaskGraphAlwaysWaitWithNamedThreadSupport || ShouldForceWaitWithNamedThreadsSupport(ExtendedPriority))
+			{
+				WaitWithNamedThreadsSupport();
+			}
+			else
+			{
+				WaitImpl(FTimeout::Never());
+			}
 		}
 
 		void FTaskBase::WaitWithNamedThreadsSupport()
@@ -556,5 +572,25 @@ namespace UE::Tasks
 		}
 
 #endif // !TASKGRAPH_NEW_FRONTEND
+
+		bool ShouldForceWaitWithNamedThreadsSupport(EExtendedTaskPriority ExtendedPriority)
+		{
+			// We force wait named thread support when we're waiting on a task that must run on the same thread we're currently on.
+			// If we don't do this, it's a guaranteed deadlock.
+#if TASKGRAPH_NEW_FRONTEND
+			const bool bIsNamedThreadTask = ExtendedPriority >= EExtendedTaskPriority::GameThreadNormalPri;
+			if (bIsNamedThreadTask)
+			{
+				FTaskGraphInterface& TaskGraph = FTaskGraphInterface::Get();
+				ENamedThreads::Type CurrentThreadIndex = ENamedThreads::GetThreadIndex(TaskGraph.GetCurrentThreadIfKnown());
+				if (CurrentThreadIndex <= ENamedThreads::ActualRenderingThread)
+				{
+					ENamedThreads::Type TaskThreadIndex = ENamedThreads::GetThreadIndex(TranslatePriority(ExtendedPriority));
+					return TaskThreadIndex == CurrentThreadIndex;
+				}
+			}
+#endif
+			return false;
+		}
 	}
 }
