[TaskGraph]
 - Add missing public APIs on FTask

diff --git a/Engine/Source/Runtime/Core/Public/Tasks/Task.h b/Engine/Source/Runtime/Core/Public/Tasks/Task.h
--- a/Engine/Source/Runtime/Core/Public/Tasks/Task.h
+++ b/Engine/Source/Runtime/Core/Public/Tasks/Task.h
@@ -123,6 +123,14 @@ namespace UE::Tasks
 				PRAGMA_ENABLE_DEPRECATION_WARNINGS
 			}
 
+			// try to finish the task on the current thread if it's not already started.
+			// @return true if the task is completed
+			bool TryRetractAndExecute()
+			{
+				Pimpl->TryRetractAndExecute(FTimeout::Never());
+				return IsCompleted();
+			}
+
 			// launches a task for asynchronous execution
 			// @param DebugName - a unique name for task identification in debugger and profiler, is compiled out in test/shipping builds
 			// @param TaskBody - a functor that will be executed asynchronously
@@ -189,6 +197,17 @@ namespace UE::Tasks
 			{
 				return Pimpl != Other.Pimpl;
 			}
+			
+			ETaskPriority GetPriority() const
+			{
+				return Pimpl->GetPriority();
+			}
+
+			EExtendedTaskPriority GetExtendedPriority() const
+			{
+				return Pimpl->GetExtendedPriority();
+			}
+		
 
 		protected:
 			TRefCountPtr<FTaskBase> Pimpl;
