[TaskGraph]
 - Fix potential deadlock when all workers are waiting and a different thread triggers an event
 - This can also reduce latency spikes since events will now be triggered ASAP instead of being scheduled for later

diff --git a/Engine/Source/Runtime/Core/Private/Async/TaskGraph.cpp b/Engine/Source/Runtime/Core/Private/Async/TaskGraph.cpp
--- a/Engine/Source/Runtime/Core/Private/Async/TaskGraph.cpp
+++ b/Engine/Source/Runtime/Core/Private/Async/TaskGraph.cpp
@@ -1574,7 +1574,24 @@ private:
 			InEvent->Trigger();
 			return;
 		}
+#if TASKGRAPH_NEW_FRONTEND
+		// When the new taskgraph frontend is active, FGraphEventArray can be used as prerequisites.
+		// We prefer the new task system because it avoids a potential deadlock situation if all threads are used
+		// while the condition is met because this one will execute inline from the thread that completed the prereqs
+		// instead of scheduling yet another task just to trigger the event that may never run if the taskgraph is already full.
+		UE::Tasks::Launch(
+			TEXT("TriggerEventWhenTaskComplete"),
+			[InEvent]()
+			{
+				InEvent->Trigger();
+			},
+			Tasks,
+			LowLevelTasks::ETaskPriority::Normal,
+			UE::Tasks::EExtendedTaskPriority::Inline
+		);
+#else
 		TGraphTask<FTriggerEventGraphTask>::CreateTask(&Tasks, CurrentThreadIfKnown).ConstructAndDispatchWhenReady(InEvent, TriggerThread);
+#endif
 	}
 
 	void AddShutdownCallback(TFunction<void()>& Callback) override
