[Scheduler]
 - Add delegate fired when the scheduler reaches its maximum oversubscription capacity
 - Add function to determine when the scheduler has exceeded its oversubscription capacity
 - Add CSV accumulative stat each time the scheduler reaches its maximum oversubscription capacity

diff --git a/Engine/Source/Runtime/Core/Private/Async/Fundamental/Scheduler.cpp b/Engine/Source/Runtime/Core/Private/Async/Fundamental/Scheduler.cpp
--- a/Engine/Source/Runtime/Core/Private/Async/Fundamental/Scheduler.cpp
+++ b/Engine/Source/Runtime/Core/Private/Async/Fundamental/Scheduler.cpp
@@ -209,6 +209,29 @@ namespace LowLevelTasks
 			}
 		}
 	}
+	
+	bool FScheduler::IsOversubscriptionLimitReached(ETaskPriority TaskPriority) const
+	{
+		const bool bIsBackgroundTask = TaskPriority >= ETaskPriority::ForegroundCount;
+		if (bIsBackgroundTask)
+		{
+			return 
+				WaitingQueue[1].IsOversubscriptionLimitReached();
+		}
+		else
+		{
+			// Since we are allowing background thread to run foreground task we need both waiting queue
+			// to reach their limit to consider that priority's limit reached.
+			return 
+				WaitingQueue[0].IsOversubscriptionLimitReached() &&
+				WaitingQueue[1].IsOversubscriptionLimitReached();
+		}
+	}
+	
+	FOversubscriptionLimitReached& FScheduler::GetOversubscriptionLimitReachedEvent()
+	{
+		return OversubscriptionLimitReachedEvent;
+	}
 
 	inline FTask* FScheduler::ExecuteTask(FTask* InTask)
 	{
diff --git a/Engine/Source/Runtime/Core/Private/Async/Fundamental/WaitingQueue.cpp b/Engine/Source/Runtime/Core/Private/Async/Fundamental/WaitingQueue.cpp
--- a/Engine/Source/Runtime/Core/Private/Async/Fundamental/WaitingQueue.cpp
+++ b/Engine/Source/Runtime/Core/Private/Async/Fundamental/WaitingQueue.cpp
@@ -161,6 +161,11 @@ void FWaitingQueue::PrepareWait(FWaitEvent* Node)
 	State.fetch_add(WaiterInc, std::memory_order_relaxed);
 }
 
+bool FWaitingQueue::IsOversubscriptionLimitReached() const
+{
+	return Oversubscription.load(std::memory_order_relaxed) >= MaxThreadCount;
+}
+
 void FWaitingQueue::CheckState(uint64 InState, bool bInIsWaiter)
 {
 	using namespace WaitingQueueImpl;
@@ -404,7 +409,12 @@ void FWaitingQueue::IncrementOversubscription()
 {
 	using namespace WaitingQueueImpl;
 
-	++Oversubscription;
+	if (++Oversubscription >= MaxThreadCount)
+	{
+		TRACE_CPUPROFILER_EVENT_SCOPE(FWaitingQueue::OversubscriptionLimitReached);
+		CSV_CUSTOM_STAT(Scheduler, OversubscriptionLimitReached, 1, ECsvCustomStatOp::Accumulate);
+		OversubscriptionLimitReachedEvent.Broadcast();
+	}
 
 	// This is important that StandbyState is invalidated after Oversubscription is increased so we
 	// can detect stale decisions and reevaluate oversubscription.
diff --git a/Engine/Source/Runtime/Core/Public/Async/Fundamental/Scheduler.h b/Engine/Source/Runtime/Core/Public/Async/Fundamental/Scheduler.h
--- a/Engine/Source/Runtime/Core/Public/Async/Fundamental/Scheduler.h
+++ b/Engine/Source/Runtime/Core/Public/Async/Fundamental/Scheduler.h
@@ -208,6 +208,15 @@ namespace LowLevelTasks
 
 		//get the background priority set when workers were started
 		inline EThreadPriority GetBackgroundPriority() const { return BackgroundPriority; }
+
+		//determine if we're currently out of workers for a given task priority
+		CORE_API bool IsOversubscriptionLimitReached(ETaskPriority TaskPriority) const;
+
+		//event that will fire when the scheduler has reached its oversubscription limit (all threads are waiting).
+		//note: This event can be broadcasted from any thread so the receiver needs to be thread-safe
+		//      For optimal performance, avoid binding UObjects to this event and use AddRaw/AddLambda instead.
+		//      Also, what's happening inside that callback should be as brief and simple as possible (i.e. raising an event)
+		CORE_API FOversubscriptionLimitReached& GetOversubscriptionLimitReachedEvent();
 	public:
 		FScheduler() = default;
 		~FScheduler();
@@ -228,7 +237,7 @@ namespace LowLevelTasks
 
 		friend class FOversubscriptionScope;
 	private:
-		Private::FWaitingQueue                         WaitingQueue[2] = { WorkerEvents, WorkerEvents };
+		Private::FWaitingQueue                         WaitingQueue[2] = { { WorkerEvents, OversubscriptionLimitReachedEvent }, { WorkerEvents, OversubscriptionLimitReachedEvent } };
 		FSchedulerTls::FQueueRegistry                  QueueRegistry;
 		FCriticalSection                               WorkerThreadsCS;
 		TArray<TUniquePtr<FThread>>                    WorkerThreads;
@@ -243,6 +252,7 @@ namespace LowLevelTasks
 		EThreadPriority                                WorkerPriority = EThreadPriority::TPri_Normal;
 		EThreadPriority                                BackgroundPriority = EThreadPriority::TPri_BelowNormal;
 		std::atomic_bool                               TemporaryShutdown{ false };
+		FOversubscriptionLimitReached                  OversubscriptionLimitReachedEvent;
 	};
 
 	namespace Private
diff --git a/Engine/Source/Runtime/Core/Public/Async/Fundamental/TaskShared.h b/Engine/Source/Runtime/Core/Public/Async/Fundamental/TaskShared.h
--- a/Engine/Source/Runtime/Core/Public/Async/Fundamental/TaskShared.h
+++ b/Engine/Source/Runtime/Core/Public/Async/Fundamental/TaskShared.h
@@ -4,10 +4,13 @@
 
 #include "Containers/Array.h"
 #include "Containers/ContainerAllocationPolicies.h"
+#include "Delegates/Delegate.h"
 #include "ProfilingDebugging/CpuProfilerTrace.h"
 
 namespace LowLevelTasks {
 
+DECLARE_TS_MULTICAST_DELEGATE(FOversubscriptionLimitReached);
+
 template<typename NodeType>
 using TAlignedArray = TArray<NodeType, TAlignedHeapAllocator<alignof(NodeType)>>;
 
diff --git a/Engine/Source/Runtime/Core/Public/Async/Fundamental/WaitingQueue.h b/Engine/Source/Runtime/Core/Public/Async/Fundamental/WaitingQueue.h
--- a/Engine/Source/Runtime/Core/Public/Async/Fundamental/WaitingQueue.h
+++ b/Engine/Source/Runtime/Core/Public/Async/Fundamental/WaitingQueue.h
@@ -41,18 +41,19 @@ namespace LowLevelTasks::Private
 
 	class FWaitingQueue
 	{
-		uint32                     ThreadCount{ 0 };    // Normal amount of threads when there is no oversubscription.
-		uint32                     MaxThreadCount{ 0 }; // Max limit that can be reached during oversubscription period.
-		TFunction<void()>          CreateThread;
-		std::atomic<uint32>        Oversubscription{ 0 };
-		std::atomic<uint64>        State;
-		std::atomic<uint64>        StandbyState;
-		TAlignedArray<FWaitEvent>& NodesArray;
-		std::atomic<bool>          bIsShuttingDown{ false };
-
+		uint32                         ThreadCount{ 0 };    // Normal amount of threads when there is no oversubscription.
+		uint32                         MaxThreadCount{ 0 }; // Max limit that can be reached during oversubscription period.
+		TFunction<void()>              CreateThread;
+		std::atomic<uint32>            Oversubscription{ 0 };
+		std::atomic<uint64>            State;
+		std::atomic<uint64>            StandbyState;
+		TAlignedArray<FWaitEvent>&     NodesArray;
+		std::atomic<bool>              bIsShuttingDown{ false };
+		FOversubscriptionLimitReached& OversubscriptionLimitReachedEvent;
 	public:
-		FWaitingQueue(TAlignedArray<FWaitEvent>& InNodesArray)
+		FWaitingQueue(TAlignedArray<FWaitEvent>& InNodesArray, FOversubscriptionLimitReached& InOversubscriptionLimitReachedEvent)
 			: NodesArray(InNodesArray)
+			, OversubscriptionLimitReachedEvent(InOversubscriptionLimitReachedEvent)
 		{
 		}
 
@@ -85,6 +86,9 @@ namespace LowLevelTasks::Private
 		// go to sleep if conditional standby determines we're now over the active thread count.
 		CORE_API void DecrementOversubscription();
 
+		// Is the current waiting queue out of workers
+		CORE_API bool IsOversubscriptionLimitReached() const;
+
 		// Try to wake up the amount of workers passed in the parameters.
 		// Return the number that were woken up.
 		int32 Notify(int32 Count = 1)
