Use gnu::returns_nonnull on Clang for functions returning known non-null values.

diff --git a/Engine/Source/Runtime/Core/Public/Clang/ClangPlatform.h b/Engine/Source/Runtime/Core/Public/Clang/ClangPlatform.h
--- a/Engine/Source/Runtime/Core/Public/Clang/ClangPlatform.h
+++ b/Engine/Source/Runtime/Core/Public/Clang/ClangPlatform.h
@@ -42,3 +42,5 @@
 
 // Ensure we can use this builtin - seems to be present on Clang 9, GCC 11 and MSVC 19.26,
 #define PLATFORM_COMPILER_SUPPORTS_BUILTIN_BITCAST (__clang_major__ >= 9)
+
+#define FUNCTION_NON_NULL_RETURN_END [[gnu::returns_nonnull]]
