Fix the returns_nonnull build breakage I introduced before coffee had flown through my veins.

diff --git a/Engine/Source/Runtime/Core/Public/Clang/ClangPlatform.h b/Engine/Source/Runtime/Core/Public/Clang/ClangPlatform.h
--- a/Engine/Source/Runtime/Core/Public/Clang/ClangPlatform.h
+++ b/Engine/Source/Runtime/Core/Public/Clang/ClangPlatform.h
@@ -43,4 +43,4 @@
 // Ensure we can use this builtin - seems to be present on Clang 9, GCC 11 and MSVC 19.26,
 #define PLATFORM_COMPILER_SUPPORTS_BUILTIN_BITCAST (__clang_major__ >= 9)
 
-#define FUNCTION_NON_NULL_RETURN_END [[gnu::returns_nonnull]]
+#define FUNCTION_NON_NULL_RETURN_START [[gnu::returns_nonnull]]
diff --git a/Engine/Source/Runtime/Core/Public/Windows/WindowsPlatform.h b/Engine/Source/Runtime/Core/Public/Windows/WindowsPlatform.h
--- a/Engine/Source/Runtime/Core/Public/Windows/WindowsPlatform.h
+++ b/Engine/Source/Runtime/Core/Public/Windows/WindowsPlatform.h
@@ -148,7 +148,11 @@ typedef FWindowsPlatformTypes FPlatformTypes;
 #define STDCALL		__stdcall										/* Standard calling convention */
 #define FORCEINLINE __forceinline									/* Force code to be inline */
 #define FORCENOINLINE __declspec(noinline)							/* Force code to NOT be inline */
+
+// When using clang-cl we will pick up the ClangPlatform.h definition of this macro before this, so check first!
+#ifndef FUNCTION_NON_NULL_RETURN_START
 #define FUNCTION_NON_NULL_RETURN_START _Ret_notnull_				/* Indicate that the function never returns nullptr. */
+#endif
 
 #define DECLARE_UINT64(x)	x
 
diff --git a/Engine/Source/Runtime/CoreUObject/Private/Serialization/AsyncLoading.cpp b/Engine/Source/Runtime/CoreUObject/Private/Serialization/AsyncLoading.cpp
--- a/Engine/Source/Runtime/CoreUObject/Private/Serialization/AsyncLoading.cpp
+++ b/Engine/Source/Runtime/CoreUObject/Private/Serialization/AsyncLoading.cpp
@@ -2069,7 +2069,7 @@ FObjectImport* FAsyncPackage::FindExistingImport(int32 LocalImportIndex)
 		if (Import->OuterIndex.IsNull())
 		{
 			Import->XObject = StaticFindObjectFast(UPackage::StaticClass(), nullptr, Linker->GetInstancingContext().RemapPackage(Import->ObjectName), true);
-			check(!Import->XObject || CastChecked<UPackage>(Import->XObject));
+			check(!Import->XObject || CastChecked<UPackage>(Import->XObject, ECastCheckedType::NullChecked));
 		}
 		else if (Import->OuterIndex.IsImport())
 		{
