Silence the non-unity header complaint (it's compiling the Clang-only ClangPlatform.h with the MSVC compiler!).

diff --git a/Engine/Source/Runtime/Core/Public/Clang/ClangPlatform.h b/Engine/Source/Runtime/Core/Public/Clang/ClangPlatform.h
--- a/Engine/Source/Runtime/Core/Public/Clang/ClangPlatform.h
+++ b/Engine/Source/Runtime/Core/Public/Clang/ClangPlatform.h
@@ -8,6 +8,8 @@
 
 // HEADER_UNIT_UNSUPPORTED - Clang not supporting header units
 
+#if defined(__clang__)
+
 #if !defined(__cpp_if_constexpr)
 	#error "Compiler is expected to support if constexpr"
 #endif
@@ -44,3 +46,5 @@
 #define PLATFORM_COMPILER_SUPPORTS_BUILTIN_BITCAST (__clang_major__ >= 9)
 
 #define FUNCTION_NON_NULL_RETURN_START [[gnu::returns_nonnull]]
+
+#endif//defined(__clang__)
