Memory insights for iOS - capture callstacks and report modules

diff --git a/Engine/Source/Runtime/Core/Core.Build.cs b/Engine/Source/Runtime/Core/Core.Build.cs
--- a/Engine/Source/Runtime/Core/Core.Build.cs
+++ b/Engine/Source/Runtime/Core/Core.Build.cs
@@ -149,6 +149,13 @@ public class Core : ModuleRules
 			{
 				PublicDefinitions.Add("USE_MALLOC_BINNED3=1");
 				PublicDefinitions.Add("USE_512MB_MAX_MEMORY_PER_BLOCK_SIZE=1");
+				
+				if (Target.Configuration != UnrealTargetConfiguration.Shipping && Target.Type != TargetType.Program)
+				{
+					PublicDefinitions.Add("UE_MEMORY_TRACE_AVAILABLE=1");
+					PublicDefinitions.Add("UE_MEMORY_TAGS_TRACE_ENABLED=1");
+					PublicDefinitions.Add("UE_CALLSTACK_TRACE_ENABLED=1");
+				}
 			}
 		}
 		else if (Target.IsInPlatformGroup(UnrealPlatformGroup.Android))
diff --git a/Engine/Source/Runtime/Core/Private/IOS/IOSCallstackTrace.cpp b/Engine/Source/Runtime/Core/Private/IOS/IOSCallstackTrace.cpp
new file mode 100644
--- /dev/null
+++ b/Engine/Source/Runtime/Core/Private/IOS/IOSCallstackTrace.cpp
@@ -0,0 +1,163 @@
+﻿// Copyright Epic Games, Inc. All Rights Reserved.
+
+#include "IOSCallstackTrace.h"
+#include "ProfilingDebugging/CallstackTrace.h"
+
+#if UE_CALLSTACK_TRACE_ENABLED
+
+#include "CoreTypes.h"
+#include "HAL/Platform.h"
+#include "HAL/PlatformStackWalk.h"
+#include "ProfilingDebugging/CallstackTracePrivate.h"
+
+////////////////////////////////////////////////////////////////////////////////
+class FBacktracer
+{
+public:
+	FBacktracer(FMalloc* InMalloc);
+	~FBacktracer();
+	static FBacktracer*	Get();
+	FORCEINLINE uint32 GetBacktraceId(uint64* StackFrames, uint32 NumStackFrames);
+	FORCEINLINE uint32 GetBacktraceId(void* ReturnAddress);
+
+private:
+	static FBacktracer* Instance;
+	FCallstackTracer CallstackTracer;
+};
+
+////////////////////////////////////////////////////////////////////////////////
+FBacktracer* FBacktracer::Instance = nullptr;
+
+////////////////////////////////////////////////////////////////////////////////
+FBacktracer::FBacktracer(FMalloc* InMalloc)
+	: CallstackTracer(InMalloc)
+{
+	Instance = this;
+}
+
+////////////////////////////////////////////////////////////////////////////////
+FBacktracer::~FBacktracer()
+{
+}
+
+////////////////////////////////////////////////////////////////////////////////
+FBacktracer* FBacktracer::Get()
+{
+	return Instance;
+}
+
+////////////////////////////////////////////////////////////////////////////////
+FORCEINLINE uint32 FBacktracer::GetBacktraceId(uint64* StackFrames, uint32 NumStackFrames)
+{
+#if !UE_BUILD_SHIPPING
+	if (NumStackFrames == 0)
+	{
+		return 0;
+	}
+
+	FCallstackTracer::FBacktraceEntry BacktraceEntry;
+	uint64 BacktraceId = 0;
+	for (int32 Index = 0; Index < NumStackFrames; Index++)
+	{
+		// This is a simple order-dependent LCG. Should be sufficient enough
+		BacktraceId += StackFrames[Index];
+		BacktraceId *= 0x30be8efa499c249dull;
+	}
+
+	// Save the collected id
+	BacktraceEntry.Hash = BacktraceId;
+	BacktraceEntry.FrameCount = NumStackFrames;
+	BacktraceEntry.Frames = StackFrames;
+
+	// Add to queue to be processed. This might block until there is room in the
+	// queue (i.e. the processing thread has caught up processing).
+	return CallstackTracer.AddCallstack(BacktraceEntry);
+#endif
+
+	return 0;
+}
+
+////////////////////////////////////////////////////////////////////////////////
+FORCEINLINE uint32 FBacktracer::GetBacktraceId(void* ReturnAddress)
+{
+#if !UE_BUILD_SHIPPING
+	uint64 StackFrames[64];
+
+	uint32 NumStackFrames = FPlatformStackWalk::CaptureStackBackTrace(StackFrames, UE_ARRAY_COUNT(StackFrames));
+
+	if (NumStackFrames > 0)
+	{
+		// Skip until we find ReturnAddress in the stack, otherwise start from the beginning.
+		uint32 StartFromIndex = 0;
+		for (int32 Index = 0; Index < NumStackFrames; Index++)
+		{
+			if (StackFrames[Index] == (uint64)ReturnAddress)
+			{
+				StartFromIndex = Index;
+				break;
+			}
+		}
+
+		return GetBacktraceId(StackFrames + StartFromIndex, NumStackFrames - StartFromIndex);
+	}
+#endif
+
+	return 0;
+}
+
+////////////////////////////////////////////////////////////////////////////////
+void Modules_Initialize();
+
+////////////////////////////////////////////////////////////////////////////////
+void CallstackTrace_CreateInternal(FMalloc* Malloc)
+{
+	if (FBacktracer::Get() != nullptr)
+	{
+		return;
+	}
+
+	// Allocate, construct and intentionally leak backtracer
+	void* Alloc = Malloc->Malloc(sizeof(FBacktracer), alignof(FBacktracer));
+	new (Alloc) FBacktracer(Malloc);
+}
+
+////////////////////////////////////////////////////////////////////////////////
+void CallstackTrace_InitializeInternal()
+{
+	if (FBacktracer* Instance = FBacktracer::Get())
+	{
+#if !UE_BUILD_SHIPPING
+		Modules_Initialize();
+#endif
+	}
+}
+
+////////////////////////////////////////////////////////////////////////////////
+uint32 CallstackTrace_GetCurrentId()
+{
+	if (!UE_TRACE_CHANNELEXPR_IS_ENABLED(CallstackChannel))
+	{
+		return 0;
+	}
+
+	void* ReturnAddress = PLATFORM_RETURN_ADDRESS_FOR_CALLSTACKTRACING();
+	if (FBacktracer* Instance = FBacktracer::Get())
+	{
+		return Instance->GetBacktraceId(ReturnAddress);
+	}
+
+	return 0;
+}
+
+////////////////////////////////////////////////////////////////////////////////
+uint32 CallstackTrace_GetExternalCallstackId(uint64* Frames, uint32 Count)
+{
+	if (FBacktracer* Instance = FBacktracer::Get())
+	{
+		return Instance->GetBacktraceId(Frames, Count);
+	}
+
+	return 0;
+}
+
+#endif
diff --git a/Engine/Source/Runtime/Core/Private/IOS/IOSCallstackTrace.h b/Engine/Source/Runtime/Core/Private/IOS/IOSCallstackTrace.h
new file mode 100644
--- /dev/null
+++ b/Engine/Source/Runtime/Core/Private/IOS/IOSCallstackTrace.h
@@ -0,0 +1,13 @@
+﻿// Copyright Epic Games, Inc. All Rights Reserved.
+
+#pragma once
+
+#include "ProfilingDebugging/CallstackTrace.h"
+
+#if UE_CALLSTACK_TRACE_ENABLED
+
+#include "CoreTypes.h"
+
+uint32 CallstackTrace_GetExternalCallstackId(uint64* Frames, uint32 Count);
+
+#endif
diff --git a/Engine/Source/Runtime/Core/Private/IOS/IOSMemoryTrace.cpp b/Engine/Source/Runtime/Core/Private/IOS/IOSMemoryTrace.cpp
new file mode 100644
--- /dev/null
+++ b/Engine/Source/Runtime/Core/Private/IOS/IOSMemoryTrace.cpp
@@ -0,0 +1,59 @@
+// Copyright Epic Games, Inc. All Rights Reserved.
+
+#include "ProfilingDebugging/MemoryTrace.h"
+
+#if UE_MEMORY_TRACE_ENABLED
+
+#include "ProfilingDebugging/CallstackTrace.h"
+#include "ProfilingDebugging/TraceMalloc.h"
+#include "Containers/StringView.h"
+#include "CoreTypes.h"
+#include "HAL/MemoryBase.h"
+#include "HAL/Platform.h"
+#include "HAL/PlatformTime.h"
+#include "Misc/CoreDelegates.h"
+#include "Misc/CString.h"
+#include "Trace/Trace.inl"
+
+
+////////////////////////////////////////////////////////////////////////////////
+FMalloc* MemoryTrace_CreateInternal(FMalloc*, int, const ANSICHAR* const*);
+
+////////////////////////////////////////////////////////////////////////////////
+FMalloc* MemoryTrace_Create(FMalloc* InMalloc)
+{
+    FMalloc* Result = InMalloc;
+    
+    NSArray* DocumentsPaths = NSSearchPathForDirectoriesInDomains(NSDocumentDirectory, NSUserDomainMask, YES);
+    NSString* DocumentsDirectory = [DocumentsPaths objectAtIndex:0];
+    NSString* DocsCmdLinePath = [[DocumentsDirectory stringByAppendingPathComponent:@"uecommandline.txt"] retain];
+    NSString* BundleCmdLinePath = [[[[NSBundle mainBundle] bundlePath] stringByAppendingPathComponent:@"uecommandline.txt"] retain];
+    NSFileManager* FileManager = [NSFileManager defaultManager];
+    
+    // try Documents folder and if there is no command line there, try the one in the bundle
+    NSString* CmdLinePaths[] = { DocsCmdLinePath, BundleCmdLinePath };
+    
+    for (int i=0;i<2;i++)
+    {
+        if ([FileManager fileExistsAtPath:CmdLinePaths[i]])
+        {
+            NSError* Error = nil;
+            NSString* CmdLine = [[NSString stringWithContentsOfFile:CmdLinePaths[i] encoding:NSUTF8StringEncoding error:&Error] retain];
+            if (Error == nil)
+            {
+                int32 ArgC = 2;
+                const char* ArgV[] = { "UE5", (char*)[CmdLine UTF8String] };
+                Result = MemoryTrace_CreateInternal(InMalloc, ArgC, ArgV);
+                [CmdLine release];
+                break;
+            }
+        }
+    }
+    
+    [DocsCmdLinePath release];
+    [BundleCmdLinePath release];
+
+    return Result;
+}
+
+#endif // UE_MEMORY_TRACE_ENABLED
diff --git a/Engine/Source/Runtime/Core/Private/IOS/IOSModuleDiagnostics.cpp b/Engine/Source/Runtime/Core/Private/IOS/IOSModuleDiagnostics.cpp
new file mode 100644
--- /dev/null
+++ b/Engine/Source/Runtime/Core/Private/IOS/IOSModuleDiagnostics.cpp
@@ -0,0 +1,96 @@
+﻿// Copyright Epic Games, Inc. All Rights Reserved.
+
+#include "CoreTypes.h"
+#include "HAL/LowLevelMemTracker.h"
+#include "HAL/PlatformProcess.h"
+#include "Math/UnrealMathUtility.h"
+#include "Misc/Paths.h"
+#include "ProfilingDebugging/ModuleDiagnostics.h"
+#include "ProfilingDebugging/MemoryTrace.h"
+#include "ProfilingDebugging/MetadataTrace.h"
+
+#if !UE_BUILD_SHIPPING
+#include "Apple/PreAppleSystemHeaders.h"
+#include <mach/mach.h>
+#include <mach-o/dyld_images.h>
+#include <mach-o/loader.h>
+#include "Apple/PostAppleSystemHeaders.h"
+#endif
+
+////////////////////////////////////////////////////////////////////////////////
+void Modules_Initialize()
+{
+#if !UE_BUILD_SHIPPING
+	using namespace UE::Trace;
+
+	constexpr uint32 SizeOfSymbolFormatString = 4;
+	UE_TRACE_LOG(Diagnostics, ModuleInit, ModuleChannel, sizeof(ANSICHAR) * SizeOfSymbolFormatString)
+		<< ModuleInit.SymbolFormat("psym", SizeOfSymbolFormatString)
+		<< ModuleInit.ModuleBaseShift(uint8(0));
+    
+    struct task_dyld_info dyld_info;
+    mach_msg_type_number_t count = TASK_DYLD_INFO_COUNT;
+    if (task_info(mach_task_self(), TASK_DYLD_INFO, (task_info_t)&dyld_info, &count) == KERN_SUCCESS)
+    {
+        HeapId ProgramHeapId = MemoryTrace_HeapSpec(EMemoryTraceRootHeap::SystemMemory, TEXT("Program"), EMemoryTraceHeapFlags::NeverFrees);
+        FString ExecutablePath = FPaths::GetPath(FString([[NSBundle mainBundle] executablePath]));
+        
+        struct dyld_all_image_infos* infos = (struct dyld_all_image_infos*)dyld_info.all_image_info_addr;
+        // FPlatformMisc::LowLevelOutputDebugStringf(TEXT("Found %d modules"), infos->infoArrayCount);
+
+        // Send Mach-O's uuid as the BuildId. Psym generation seems to have an extra 0 but we will ignore it on the other end.
+        constexpr uint32 BuildIdSize = 16;
+        uint8 BuildId[BuildIdSize] = {0};
+
+        for (int i=0; i<infos->infoArrayCount; i++)
+        {
+            const struct dyld_image_info* Image = &infos->infoArray[i];
+            // FPlatformMisc::LowLevelOutputDebugStringf(TEXT("Found module %s"), ANSI_TO_TCHAR(Image->imageFilePath));
+
+            // calc image size by adding header and size of segments
+            struct mach_header_64* Header = (struct mach_header_64*)Image->imageLoadAddress;
+            size_t ImageSize = sizeof(*Header) + Header->sizeofcmds;
+            char* CmdPtr = (char*)(Header + 1);
+            for (int cmd=0; cmd < Header->ncmds; cmd++)
+            {
+                struct load_command* LoadCommand = (struct load_command*)CmdPtr;
+                if (LoadCommand->cmd == LC_SEGMENT_64)
+                {
+                    ImageSize += ((struct segment_command_64*)LoadCommand)->vmsize;
+                }
+                else if (LoadCommand->cmd == LC_UUID)
+                {
+                    FMemory::Memcpy(BuildId, ((struct uuid_command*)LoadCommand)->uuid, BuildIdSize);
+                }
+                CmdPtr = CmdPtr + LoadCommand->cmdsize;
+            }
+                        
+            FString ImageName(Image->imageFilePath);
+            bool bInsideExecutablePath = ImageName.StartsWith(ExecutablePath);
+            
+            // trim path to leave just image name
+            ImageName = FPaths::GetCleanFilename(ImageName);
+            uint64 ImageLoadAddress = (uint64)Image->imageLoadAddress;
+            
+            UE_TRACE_LOG(Diagnostics, ModuleLoad, ModuleChannel, sizeof(TCHAR) * ImageName.Len() + BuildIdSize)
+                << ModuleLoad.Name(*ImageName, ImageName.Len())
+                << ModuleLoad.Base(ImageLoadAddress)
+                << ModuleLoad.Size(ImageSize)
+                << ModuleLoad.ImageId(BuildId, BuildIdSize);
+                       
+#if UE_MEMORY_TRACE_ENABLED
+            // Only count the main executable and any other libraries inside our bundle as "Program Size"
+            if (bInsideExecutablePath)
+            {
+                UE_TRACE_METADATA_CLEAR_SCOPE();
+                LLM(UE_MEMSCOPE(ELLMTag::ProgramSize));
+                MemoryTrace_Alloc(ImageLoadAddress, ImageSize, 1);
+                MemoryTrace_MarkAllocAsHeap(ImageLoadAddress, ProgramHeapId);
+                MemoryTrace_Alloc(ImageLoadAddress, ImageSize, 1);
+            }
+#endif // UE_MEMORY_TRACE_ENABLED
+
+        }
+    }
+#endif
+}
diff --git a/Engine/Source/Runtime/Core/Public/IOS/IOSPlatform.h b/Engine/Source/Runtime/Core/Public/IOS/IOSPlatform.h
--- a/Engine/Source/Runtime/Core/Public/IOS/IOSPlatform.h
+++ b/Engine/Source/Runtime/Core/Public/IOS/IOSPlatform.h
@@ -83,6 +83,8 @@ typedef FIOSPlatformTypes FPlatformTypes;
 #undef PLATFORM_COMPILER_SUPPORTS_BUILTIN_BITCAST
 #define PLATFORM_COMPILER_SUPPORTS_BUILTIN_BITCAST (__clang_major__ >= 13)
 
+#define PLATFORM_RETURN_ADDRESS_FOR_CALLSTACKTRACING    PLATFORM_RETURN_ADDRESS
+
 #define PLATFORM_GLOBAL_LOG_CATEGORY					LogIOS
 
 #define PLATFORM_BREAK()                                __builtin_debugtrap()
