Move MB3 IOS setting to project specific location

diff --git a/Engine/Source/Runtime/Core/Core.Build.cs b/Engine/Source/Runtime/Core/Core.Build.cs
--- a/Engine/Source/Runtime/Core/Core.Build.cs
+++ b/Engine/Source/Runtime/Core/Core.Build.cs
@@ -147,9 +147,6 @@ public class Core : ModuleRules
 			
 			if (Target.Platform == UnrealTargetPlatform.IOS)
 			{
-				PublicDefinitions.Add("USE_MALLOC_BINNED3=1");
-				PublicDefinitions.Add("USE_512MB_MAX_MEMORY_PER_BLOCK_SIZE=1");
-				
 				if (Target.Configuration != UnrealTargetConfiguration.Shipping && Target.Type != TargetType.Program)
 				{
 					PublicDefinitions.Add("UE_MEMORY_TRACE_AVAILABLE=1");
