SafeUnload empty branches when we fail to load the config rather than removing the branch outright. This allows plugins which may later try to append to 'empty' branches to do so as there isn't a means to AddOrCreate branches when plugins load.

#rnx

diff --git a/Engine/Source/Runtime/Core/Private/Misc/ConfigContext.cpp b/Engine/Source/Runtime/Core/Private/Misc/ConfigContext.cpp
--- a/Engine/Source/Runtime/Core/Private/Misc/ConfigContext.cpp
+++ b/Engine/Source/Runtime/Core/Private/Misc/ConfigContext.cpp
@@ -281,15 +281,10 @@ bool FConfigContext::Load(const TCHAR* InBaseIniName, FString& OutFinalFilename)
 			// we need to copy the temporary branch's final result back into the output
 			*ExistingFile = TemporaryBranch->InMemoryFile;
 		}
-		// delete the branch if it is useless
+		// Unload the branch if it is empty. SafeUnload so that we may re-use the branch should it need to be added to later (e.g. by a plugin)
 		if (!bSuccess && ConfigSystem != nullptr && ExistingFile == nullptr && TemporaryBranch == nullptr)
 		{
-			static bool bRemoveEmptyPlugins = FParse::Param(FCommandLine::Get(), TEXT("RemoveEmptyConfigs"));
-			if (bRemoveEmptyPlugins)
-			{
-				ConfigSystem->Remove(DestIniFilename);
-				Branch = nullptr;
-			}
+			ConfigSystem->SafeUnloadBranch(InBaseIniName);
 		}
 	}
 	return bSuccess;
