Disable MB3 by default on Apple platforms

diff --git a/Engine/Source/Runtime/Core/Public/Apple/ApplePlatformMemory.h b/Engine/Source/Runtime/Core/Public/Apple/ApplePlatformMemory.h
--- a/Engine/Source/Runtime/Core/Public/Apple/ApplePlatformMemory.h
+++ b/Engine/Source/Runtime/Core/Public/Apple/ApplePlatformMemory.h
@@ -14,12 +14,9 @@
 
 NS_ASSUME_NONNULL_BEGIN
 
-// if USE_MALLOC_BINNED3 isn't configured enable it by default
+// if USE_MALLOC_BINNED3 isn't configured disable it by default
 #ifndef USE_MALLOC_BINNED3
-	#define USE_MALLOC_BINNED3 1
-	#if PLATFORM_IOS
-		#define USE_512MB_MAX_MEMORY_PER_BLOCK_SIZE 1
-	#endif
+	#define USE_MALLOC_BINNED3 0
 #endif // USE_MALLOC_BINNED3
 
 
