Get the current state of user iOS notification approval

diff --git a/Engine/Source/Runtime/Core/Private/IOS/IOSPlatformMisc.cpp b/Engine/Source/Runtime/Core/Private/IOS/IOSPlatformMisc.cpp
--- a/Engine/Source/Runtime/Core/Private/IOS/IOSPlatformMisc.cpp
+++ b/Engine/Source/Runtime/Core/Private/IOS/IOSPlatformMisc.cpp
@@ -1283,6 +1283,43 @@ void FIOSPlatformMisc::UnregisterForRemoteNotifications()
 
 }
 
+FIOSPlatformMisc::EIOSAuthNotificationStatus FIOSPlatformMisc::GetNotificationAuthorizationStatus()
+{
+	dispatch_semaphore_t Semaphore = dispatch_semaphore_create(0);
+	static EIOSAuthNotificationStatus CurrentAuthStatus = Unknown;
+	dispatch_async(dispatch_get_global_queue(QOS_CLASS_BACKGROUND, 0), ^{
+		[[UNUserNotificationCenter currentNotificationCenter] getNotificationSettingsWithCompletionHandler:^(UNNotificationSettings* Settings)
+		{
+			switch (Settings.authorizationStatus)
+			{
+				case UNAuthorizationStatusNotDetermined:
+					CurrentAuthStatus = NotDetermined;
+					break;
+				case UNAuthorizationStatusDenied:
+					CurrentAuthStatus = Denied;
+					break;
+				case UNAuthorizationStatusAuthorized:
+					CurrentAuthStatus = Authorized;
+					break;
+				case UNAuthorizationStatusProvisional:
+					CurrentAuthStatus = Provisional;
+					break;
+#if !PLATFORM_TVOS
+				case UNAuthorizationStatusEphemeral:
+					CurrentAuthStatus = Ephemeral;
+					break;
+#endif
+				default:
+					CurrentAuthStatus = Unknown;
+			}
+			dispatch_semaphore_signal(Semaphore);
+		}];
+	});
+
+	// wait for a result, but timeout after 1s
+	dispatch_semaphore_wait(Semaphore, dispatch_time(DISPATCH_TIME_NOW, NSEC_PER_SEC));
+	return CurrentAuthStatus;
+}
 
 // See for more information about the Blobs
 // https://opensource.apple.com/source/xnu/xnu-4570.61.1/osfmk/kern/cs_blobs.h.auto.html
diff --git a/Engine/Source/Runtime/Core/Public/IOS/IOSPlatformMisc.h b/Engine/Source/Runtime/Core/Public/IOS/IOSPlatformMisc.h
--- a/Engine/Source/Runtime/Core/Public/IOS/IOSPlatformMisc.h
+++ b/Engine/Source/Runtime/Core/Public/IOS/IOSPlatformMisc.h
@@ -73,6 +73,17 @@ struct CORE_API FIOSPlatformMisc : public FApplePlatformMisc
 	// Check if notifications are allowed if min iOS version is < 10
 	UE_DEPRECATED(4.21, "IsAllowedRemoteNotifications is deprecated. Use FIOSLocalNotificationService::CheckAllowedNotifications instead.")
 	static bool IsAllowedRemoteNotifications();
+
+	enum EIOSAuthNotificationStatus
+	{
+		NotDetermined,
+		Denied,
+		Authorized,
+		Provisional,
+		Ephemeral,
+		Unknown
+	};
+	static EIOSAuthNotificationStatus GetNotificationAuthorizationStatus();
     
     static bool IsEntitlementEnabled(const char *EntitlementToCheck);
 	
