Common assertions header for LLTs and Functional Tests
#rnx

diff --git a/Engine/Source/Developer/LowLevelTestsRunner/Public/TestMacros/Assertions.h b/Engine/Source/Developer/LowLevelTestsRunner/Public/TestMacros/Assertions.h
--- a/Engine/Source/Developer/LowLevelTestsRunner/Public/TestMacros/Assertions.h
+++ b/Engine/Source/Developer/LowLevelTestsRunner/Public/TestMacros/Assertions.h
@@ -2,8 +2,8 @@
 #pragma once
 
 #include "Misc/AssertionMacros.h"
-#include "LowLevelTestsRunner/EnsureScope.h"
-#include "LowLevelTestsRunner/CheckScope.h"
+#include "Tests/EnsureScope.h"
+#include "Tests/CheckScope.h"
 
 //requires that an UE `ensure` fails in this call
 #define REQUIRE_ENSURE(...) INTERNAL_UE_ENSURE( "REQUIRE_ENSURE", DO_ENSURE, Catch::ResultDisposition::Normal, #__VA_ARGS__, __VA_ARGS__ )
@@ -33,7 +33,7 @@
 	do { \
 		Catch::AssertionHandler catchAssertionHandler( macroName##_catch_sr, CATCH_INTERNAL_LINEINFO, ensureExpr, resultDisposition ); \
 		INTERNAL_CATCH_TRY { \
-			::UE::LowLevelTests::FEnsureScope scope; \
+			FEnsureScope scope; \
 			static_cast<void>(__VA_ARGS__); \
 			bool bEncounteredEnsure = scope.GetCount() > 0; \
 			if (doEnsure && !bEncounteredEnsure) \
@@ -46,7 +46,7 @@
 	do { \
 		Catch::AssertionHandler catchAssertionHandler( macroName##_catch_sr, CATCH_INTERNAL_LINEINFO, ensureExpr, resultDisposition ); \
 		INTERNAL_CATCH_TRY { \
-			::UE::LowLevelTests::FEnsureScope scope(msg); \
+			FEnsureScope scope(msg); \
 			static_cast<void>(__VA_ARGS__); \
 			bool bEncounteredEnsure = scope.GetCount() > 0; \
 			if (doEnsure && !bEncounteredEnsure) \
@@ -59,7 +59,7 @@
 	do { \
 		Catch::AssertionHandler catchAssertionHandler( macroName##_catch_sr, CATCH_INTERNAL_LINEINFO, checkExpr, resultDisposition ); \
 		INTERNAL_CATCH_TRY { \
-			::UE::LowLevelTests::FCheckScope scope; \
+			FCheckScope scope; \
 			__VA_ARGS__; \
 			bool bEncounteredEnsure = scope.GetCount() > 0; \
 			if (doCheck && !bEncounteredEnsure) \
@@ -72,7 +72,7 @@
 	do { \
 		Catch::AssertionHandler catchAssertionHandler( macroName##_catch_sr, CATCH_INTERNAL_LINEINFO, checkExpr, resultDisposition ); \
 		INTERNAL_CATCH_TRY { \
-			::UE::LowLevelTests::FCheckScope scope(msg); \
+			FCheckScope scope(msg); \
 			__VA_ARGS__; \
 			bool bEncounteredEnsure = scope.GetCount() > 0; \
 			if (doCheck && !bEncounteredEnsure) \
diff --git a/Engine/Source/Programs/SwitchboardListener/SwitchboardListenerTests/Private/SwitchboardAuthTests.cpp b/Engine/Source/Programs/SwitchboardListener/SwitchboardListenerTests/Private/SwitchboardAuthTests.cpp
--- a/Engine/Source/Programs/SwitchboardListener/SwitchboardListenerTests/Private/SwitchboardAuthTests.cpp
+++ b/Engine/Source/Programs/SwitchboardListener/SwitchboardListenerTests/Private/SwitchboardAuthTests.cpp
@@ -8,7 +8,7 @@
 #include "JsonWebToken.h"
 #include "HAL/PlatformProcess.h"
 #include "HAL/PlatformTime.h"
-#include "LowLevelTestsRunner/WarnFilterScope.h"
+#include "Tests/WarnFilterScope.h"
 
 
 namespace UE::SwitchboardListener::Private::Tests
@@ -118,7 +118,7 @@ namespace UE::SwitchboardListener::Private::Tests
 
 			// Previously issued JWT was expired by a password change.
 			{
-				UE::Testing::FWarnFilterScope _([](const TCHAR* Message, ELogVerbosity::Type Verbosity, const FName& Category)
+				FWarnFilterScope _([](const TCHAR* Message, ELogVerbosity::Type Verbosity, const FName& Category)
 					{
 						constexpr FStringView ExpectedError = TEXTVIEW("JWT issued at time predates last password change");
 						const bool bFiltered =
diff --git a/Engine/Source/Runtime/Core/Private/Tests/CheckScope.cpp b/Engine/Source/Runtime/Core/Private/Tests/CheckScope.cpp
new file mode 100644
--- /dev/null
+++ b/Engine/Source/Runtime/Core/Private/Tests/CheckScope.cpp
@@ -0,0 +1,96 @@
+// Copyright Epic Games, Inc. All Rights Reserved.
+
+#include "Tests/CheckScope.h"
+#include "Containers/StringConv.h"
+#include "CoreGlobals.h"
+#include "Misc/OutputDeviceError.h"
+
+class FCheckScopeOutputDeviceError : public FOutputDeviceError
+{
+public:
+	FCheckScopeOutputDeviceError(FOutputDeviceError* Error, const ANSICHAR* Msg)
+		: DeviceError(Error)
+		, ExpectedMsg(Msg)
+		, Count(0)
+	{
+	}
+
+	virtual void HandleError() override
+	{
+		if (DeviceError)
+		{
+			DeviceError->HandleError();
+		}
+	}
+
+	void Serialize(const TCHAR* V, ELogVerbosity::Type Verbosity, const FName& Category) override
+	{
+		if (ExpectedMsg)
+		{
+			auto Str = StringCast<TCHAR>(ExpectedMsg);
+			auto Found = FPlatformString::Strstr(V, Str.Get());
+			if (Found)
+			{
+				++Count;
+			}
+		}
+		else
+		{
+			++Count;
+		}
+	}
+
+	void SetDeviceError(FOutputDeviceError* Error)
+	{
+		DeviceError = Error;
+	}
+
+	FOutputDeviceError* GetDeviceError() const
+	{
+		return DeviceError;
+	}
+
+	int GetCount() { return Count; }
+
+private:
+	FOutputDeviceError* DeviceError;
+	const ANSICHAR* ExpectedMsg;
+	int Count;
+};
+
+FCheckScope::FCheckScope(const ANSICHAR* Msg)
+	: DeviceError(new FCheckScopeOutputDeviceError(GError, Msg))
+#if !UE_BUILD_SHIPPING
+	, bIgnoreDebugger(GIgnoreDebugger)
+#else
+	, bIgnoreDebugger(false)
+#endif
+	, bCriticalError(GIsCriticalError)
+{
+#if !UE_BUILD_SHIPPING
+	GIgnoreDebugger = true;
+#endif
+	GError = DeviceError;
+	GIsCriticalError = true; //set to true to disable printf of error message which causes Horde to flag the lines as an error
+}
+
+FCheckScope::FCheckScope()
+	: FCheckScope(nullptr)	
+{
+}
+
+FCheckScope::~FCheckScope()
+{
+#if !UE_BUILD_SHIPPING
+	GIgnoreDebugger = bIgnoreDebugger;
+#endif
+	GError = DeviceError->GetDeviceError();
+	GIsCriticalError = bCriticalError;
+	delete DeviceError;
+}
+
+int
+FCheckScope::GetCount()
+{
+	return DeviceError->GetCount();
+}
\ No newline at end of file
diff --git a/Engine/Source/Runtime/Core/Private/Tests/EnsureScope.cpp b/Engine/Source/Runtime/Core/Private/Tests/EnsureScope.cpp
new file mode 100644
--- /dev/null
+++ b/Engine/Source/Runtime/Core/Private/Tests/EnsureScope.cpp
@@ -0,0 +1,41 @@
+// Copyright Epic Games, Inc. All Rights Reserved.
+
+#include "Tests/EnsureScope.h"
+#include "Containers/StringConv.h"
+
+FEnsureScope::FEnsureScope()
+	: FEnsureScope([](const FEnsureHandlerArgs&) { return true; })
+{
+}
+
+
+FEnsureScope::FEnsureScope(const ANSICHAR* ExpectedMsg)
+	: FEnsureScope([ExpectedMsg](const FEnsureHandlerArgs& Args)
+		{
+			return FPlatformString::Stricmp(ExpectedMsg, Args.Message) == 0;
+		})
+{
+}
+
+FEnsureScope::FEnsureScope(TFunction<bool(const FEnsureHandlerArgs& Args)> EnsureFunc)
+	: Count(0)
+{
+#if DO_ENSURE
+	OldHandler = SetEnsureHandler([this, EnsureFunc](const FEnsureHandlerArgs& Args) -> bool
+		{
+			bool Handled = EnsureFunc(Args);
+			if (Handled)
+			{
+				++Count;
+			}
+			return Handled;
+		});
+#endif
+}
+
+FEnsureScope::~FEnsureScope()
+{
+#if DO_ENSURE
+	SetEnsureHandler(OldHandler);
+#endif
+}
\ No newline at end of file
diff --git a/Engine/Source/Runtime/Core/Private/Tests/WarnFilterScope.cpp b/Engine/Source/Runtime/Core/Private/Tests/WarnFilterScope.cpp
new file mode 100644
--- /dev/null
+++ b/Engine/Source/Runtime/Core/Private/Tests/WarnFilterScope.cpp
@@ -0,0 +1,46 @@
+// Copyright Epic Games, Inc. All Rights Reserved.
+
+#include "Tests/WarnFilterScope.h"
+#include "Misc/FeedbackContext.h"
+
+struct FFilterFeedback : public FFeedbackContext
+{
+
+	FFilterFeedback(FFeedbackContext* OldFeedback, TFunction<bool(const TCHAR* Message, ELogVerbosity::Type Verbosity, const FName& Category)> Handler)
+		: OldFeedback(OldFeedback)
+		, Handler(Handler)
+	{
+
+	}
+
+	virtual void Serialize(const TCHAR* V, ELogVerbosity::Type Verbosity, const FName& Category) override
+	{
+		if (!Handler(V, Verbosity, Category))
+		{
+			OldFeedback->Serialize(V, Verbosity, Category);
+		}
+	}
+	virtual void Serialize(const TCHAR* V, ELogVerbosity::Type Verbosity, const FName& Category, double Time) override
+	{
+		if (!Handler(V, Verbosity, Category))
+		{
+			OldFeedback->Serialize(V, Verbosity, Category, Time);
+		}
+	}
+
+	FFeedbackContext* OldFeedback;
+	TFunction<bool(const TCHAR* Message, ELogVerbosity::Type Verbosity, const FName& Category)> Handler;
+};
+
+FWarnFilterScope::FWarnFilterScope(TFunction<bool(const TCHAR* Message, ELogVerbosity::Type Verbosity, const FName& Category)> LogHandler)
+	: OldWarn(GWarn)
+	, Feedback(new FFilterFeedback(OldWarn, LogHandler))
+{
+	GWarn = Feedback;
+}
+
+FWarnFilterScope::~FWarnFilterScope()
+{
+	GWarn = OldWarn;
+	delete Feedback;
+}
\ No newline at end of file
diff --git a/Engine/Source/Runtime/Core/Public/Tests/Assertions.h b/Engine/Source/Runtime/Core/Public/Tests/Assertions.h
new file mode 100644
--- /dev/null
+++ b/Engine/Source/Runtime/Core/Public/Tests/Assertions.h
@@ -0,0 +1,73 @@
+// Copyright Epic Games, Inc. All Rights Reserved.
+
+#pragma once
+
+#if WITH_LOW_LEVEL_TESTS
+#include "TestMacros/Assertions.h" // HEADER_UNIT_IGNORE
+#elif defined(WITH_AUTOMATION_TESTS) || (WITH_DEV_AUTOMATION_TESTS || WITH_PERF_AUTOMATION_TESTS)
+#include "Misc/AssertionMacros.h"
+#include "Misc/AutomationTest.h"
+#include "Tests/EnsureScope.h"
+#include "Tests/CheckScope.h"
+
+//requires that an UE `ensure` fails in this call
+#define REQUIRE_ENSURE(...) INTERNAL_UE_ENSURE( "REQUIRE_ENSURE", DO_ENSURE, #__VA_ARGS__, __VA_ARGS__ )
+
+//requires that an UE `ensure` fails with a message that matches the supplied message
+#define REQUIRE_ENSURE_MSG(msg, ...) INTERNAL_UE_ENSURE_MSG(msg, "REQUIRE_ENSURE", DO_ENSURE, #__VA_ARGS__, __VA_ARGS__ )
+
+//checks that an UE `ensure` fails in this call
+#define CHECK_ENSURE(...) INTERNAL_UE_ENSURE( "CHECK_ENSURE", DO_ENSURE, Catch::ResultDisposition::ContinueOnFailure, #__VA_ARGS__, __VA_ARGS__ )
+
+//checks that an UE `ensure` fails with a message that matches the supplied message
+#define CHECK_ENSURE_MSG(msg, ...) INTERNAL_UE_ENSURE_MSG(msg, "CHECK_ENSURE", DO_ENSURE, Catch::ResultDisposition::ContinueOnFailure, #__VA_ARGS__, __VA_ARGS__ )
+
+//requires that a UE `check` fails in this call
+#define REQUIRE_CHECK(...) INTERNAL_UE_CHECK( "REQUIRE_CHECK", DO_CHECK, #__VA_ARGS__, __VA_ARGS__ )
+
+//requires that a UE `check` fails in this call contains the supplies message
+#define REQUIRE_CHECK_MSG(msg, ...) INTERNAL_UE_CHECK_MSG(msg, "REQUIRE_CHECK", DO_CHECK, #__VA_ARGS__, __VA_ARGS__ )
+
+//requires that a UE `checkSlow` fails in this call
+#define REQUIRE_CHECK_SLOW(...) INTERNAL_UE_CHECK( "REQUIRE_CHECK_SLOW", DO_GUARD_SLOW, #__VA_ARGS__, __VA_ARGS__ )
+
+//requires that a UE `checkSlow` fails in this call contains the supplies message
+#define REQUIRE_CHECK_SLOW_MSG(msg, ...) INTERNAL_UE_CHECK_MSG(msg, "REQUIRE_CHECK_SLOW", DO_GUARD_SLOW, #__VA_ARGS__, __VA_ARGS__ )
+
+#define INTERNAL_UE_ENSURE( macroName, doEnsure, ensureExpr, ... ) \
+	do { \
+		FEnsureScope scope; \
+		static_cast<void>(__VA_ARGS__); \
+		bool bEncounteredEnsure = scope.GetCount() > 0; \
+		if (doEnsure && !bEncounteredEnsure) \
+			FAutomationTestFramework::Get().GetCurrentTest()->AddError(TEXT("Expected failure of `ensure` not received.")); \
+	} while(false) \
+
+#define INTERNAL_UE_ENSURE_MSG(msg, macroName, doEnsure, ensureExpr, ... ) \
+	do { \
+		FEnsureScope scope(msg); \
+		static_cast<void>(__VA_ARGS__); \
+		bool bEncounteredEnsure = scope.GetCount() > 0; \
+		if (doEnsure && !bEncounteredEnsure) \
+			FAutomationTestFramework::Get().GetCurrentTest()->AddError(TEXT("Expected failure of `ensure` with message %s not received", msg)); \
+	} while(false) \
+
+#define INTERNAL_UE_CHECK(macroName, doCheck, checkExpr, ... ) \
+	do { \
+		FCheckScope scope; \
+		__VA_ARGS__; \
+		bool bEncounteredEnsure = scope.GetCount() > 0; \
+		if (doCheck && !bEncounteredEnsure) \
+			FAutomationTestFramework::Get().GetCurrentTest()->AddError(TEXT("Expected failure of `check` not received")); \
+	} while(false) \
+
+#define INTERNAL_UE_CHECK_MSG(msg, macroName, doCheck, checkExpr, ... ) \
+	do { \
+		FCheckScope scope(msg); \
+		__VA_ARGS__; \
+		bool bEncounteredEnsure = scope.GetCount() > 0; \
+		if (doCheck && !bEncounteredEnsure) \
+			FAutomationTestFramework::Get().GetCurrentTest()->AddError(TEXT("Expected failure of `check` containing message %s not received", msg)); \
+	} while(false) \
+
+#endif
\ No newline at end of file
diff --git a/Engine/Source/Runtime/Core/Public/Tests/CheckScope.h b/Engine/Source/Runtime/Core/Public/Tests/CheckScope.h
new file mode 100644
--- /dev/null
+++ b/Engine/Source/Runtime/Core/Public/Tests/CheckScope.h
@@ -0,0 +1,23 @@
+// Copyright Epic Games, Inc. All Rights Reserved.
+
+#pragma once
+#include "HAL/Platform.h"
+
+class FCheckScopeOutputDeviceError;
+
+///@brief Scope that captures failed `check` calls
+struct FCheckScope
+{
+	FCheckScope();
+
+	///@brief Only catches failed `check` calls that contain Msg
+	///@param Msg string to look for in the `check` message
+	FCheckScope(const ANSICHAR* Msg);
+	~FCheckScope();
+
+	int GetCount();
+private:
+	FCheckScopeOutputDeviceError* DeviceError;
+	bool bIgnoreDebugger;
+	bool bCriticalError;
+};
\ No newline at end of file
diff --git a/Engine/Source/Runtime/Core/Public/Tests/EnsureScope.h b/Engine/Source/Runtime/Core/Public/Tests/EnsureScope.h
new file mode 100644
--- /dev/null
+++ b/Engine/Source/Runtime/Core/Public/Tests/EnsureScope.h
@@ -0,0 +1,27 @@
+// Copyright Epic Games, Inc. All Rights Reserved.
+
+#pragma once
+#include "Templates/Function.h"
+
+struct FEnsureHandlerArgs;
+
+///@brief Scope that captures failed `ensure` calls
+struct FEnsureScope
+{
+	FEnsureScope();
+
+	///@brief Only captures failed `ensure` calls with the exact message 
+	///@param Msg string to compare against failed `ensure` message
+	FEnsureScope(const ANSICHAR* Msg);
+
+	///@brief Captures failed `ensure` calls that return true from callback
+	///@param EnsureFunc lambda called for each failed `ensure` return true to handle the error
+	FEnsureScope(TFunction<bool(const FEnsureHandlerArgs&)> EnsureFunc);
+	~FEnsureScope();
+
+	int GetCount() { return Count; }
+
+private:
+	TFunction<bool(const FEnsureHandlerArgs&)> OldHandler;
+	int Count;
+};
\ No newline at end of file
diff --git a/Engine/Source/Runtime/Core/Public/Tests/WarnFilterScope.h b/Engine/Source/Runtime/Core/Public/Tests/WarnFilterScope.h
new file mode 100644
--- /dev/null
+++ b/Engine/Source/Runtime/Core/Public/Tests/WarnFilterScope.h
@@ -0,0 +1,24 @@
+// Copyright Epic Games, Inc. All Rights Reserved.
+
+#pragma once
+#include "HAL/Platform.h"
+#include "Templates/Function.h"
+
+class FFeedbackContext;
+
+struct FFilterFeedback;
+
+///@brief Scope that captures warning log messages
+///necessary to filter out warning message from appearing the log file as Horde will flag the warnings
+struct FWarnFilterScope
+{
+	/// @param LogHandler function to filter log messages. returning true filters message from log file
+	FWarnFilterScope(TFunction<bool(const TCHAR* Message, ELogVerbosity::Type Verbosity, const FName& Category)> LogHandler);
+	FWarnFilterScope(const FWarnFilterScope&) = delete;
+	FWarnFilterScope& operator=(const FWarnFilterScope&) = delete;
+
+	~FWarnFilterScope();
+private:
+	FFeedbackContext* OldWarn;
+	FFilterFeedback* Feedback;
+};
\ No newline at end of file
diff --git a/Engine/Source/Runtime/Core/Tests/Containers/BitArrayTests.cpp b/Engine/Source/Runtime/Core/Tests/Containers/BitArrayTests.cpp
--- a/Engine/Source/Runtime/Core/Tests/Containers/BitArrayTests.cpp
+++ b/Engine/Source/Runtime/Core/Tests/Containers/BitArrayTests.cpp
@@ -9,6 +9,7 @@
 #include "Templates/Function.h"
 
 #include "Tests/TestHarnessAdapter.h"
+#include "Tests/Assertions.h"
 
 namespace UE
 {
diff --git a/Engine/Source/Runtime/Core/Tests/Misc/WarnFilterScopeTest.cpp b/Engine/Source/Runtime/Core/Tests/Misc/WarnFilterScopeTest.cpp
--- a/Engine/Source/Runtime/Core/Tests/Misc/WarnFilterScopeTest.cpp
+++ b/Engine/Source/Runtime/Core/Tests/Misc/WarnFilterScopeTest.cpp
@@ -3,7 +3,7 @@
 #if WITH_LOW_LEVEL_TESTS
 
 #include "TestHarness.h"
-#include "LowLevelTestsRunner/WarnFilterScope.h"
+#include "Tests/WarnFilterScope.h"
 #include "Logging/LogMacros.h"
 
 DECLARE_LOG_CATEGORY_EXTERN(LogTesting, Log, All);
@@ -12,7 +12,7 @@ DEFINE_LOG_CATEGORY(LogTesting);
 TEST_CASE("UE::Testing::WarnFilterScope")
 {
 	bool bFiltered = false;
-	UE::Testing::FWarnFilterScope _([&bFiltered](const TCHAR* Message, ELogVerbosity::Type Verbosity, const FName& Category)
+	FWarnFilterScope _([&bFiltered](const TCHAR* Message, ELogVerbosity::Type Verbosity, const FName& Category)
 		{
 			bFiltered = FCString::Strcmp(Message, TEXT("Test Message")) == 0 && Verbosity == ELogVerbosity::Type::Warning && Category == TEXT("LogTesting");
 			return bFiltered;
diff --git a/Engine/Source/Runtime/CoreUObject/Tests/ObjectPropertyTest.cpp b/Engine/Source/Runtime/CoreUObject/Tests/ObjectPropertyTest.cpp
--- a/Engine/Source/Runtime/CoreUObject/Tests/ObjectPropertyTest.cpp
+++ b/Engine/Source/Runtime/CoreUObject/Tests/ObjectPropertyTest.cpp
@@ -11,7 +11,7 @@
 #include "UObject/LinkerPlaceholderExportObject.h"
 #include "UObject/LinkerPlaceholderClass.h"
 #include "UObject/ObjectHandleTracking.h"
-#include "LowLevelTestsRunner/WarnFilterScope.h"
+#include "Tests/WarnFilterScope.h"
 #include "Misc/AssetRegistryInterface.h"
 #include "AssetRegistry/AssetData.h"
 #include "ObjectRefTrackingTestBase.h"
@@ -62,7 +62,7 @@ TEST_CASE("UE::CoreUObject::FObjectProperty::CheckValidAddress")
 	Obj->ObjectPtr = reinterpret_cast<UObjectPtrTestClass*>(Obj);
 	CHECK(Obj->ObjectPtr != nullptr);
 
-	UE::Testing::FWarnFilterScope _([](const TCHAR* Message, ELogVerbosity::Type Verbosity, const FName& Category)
+	FWarnFilterScope _([](const TCHAR* Message, ELogVerbosity::Type Verbosity, const FName& Category)
 		{
 			if (Category == TEXT("LogProperty") && FCString::Strstr(Message, TEXT("Reference will be nullptred")) && Verbosity == ELogVerbosity::Type::Warning)
 			{
@@ -126,7 +126,7 @@ TEST_CASE("UE::CoreUObject::FObjectProperty::CheckValidAddressNonNullable")
 	Obj->ObjectPtrNonNullable = reinterpret_cast<UObjectPtrTestClass*>(OtherTestPackage);
 	CHECK(Obj->ObjectPtrNonNullable != nullptr);
 
-	UE::Testing::FWarnFilterScope _([](const TCHAR* Message, ELogVerbosity::Type Verbosity, const FName& Category)
+	FWarnFilterScope _([](const TCHAR* Message, ELogVerbosity::Type Verbosity, const FName& Category)
 		{
 			if (Category == TEXT("LogProperty") && FCString::Strstr(Message, TEXT("Reference will be defaulted to")) && Verbosity == ELogVerbosity::Type::Warning)
 			{
diff --git a/Engine/Source/Runtime/CoreUObject/Tests/PluginManagerTest.cpp b/Engine/Source/Runtime/CoreUObject/Tests/PluginManagerTest.cpp
--- a/Engine/Source/Runtime/CoreUObject/Tests/PluginManagerTest.cpp
+++ b/Engine/Source/Runtime/CoreUObject/Tests/PluginManagerTest.cpp
@@ -14,7 +14,7 @@
 #include "UObject/Package.h"
 #include "HAL/IConsoleManager.h"
 #include "UObject/Class.h"
-#include "LowLevelTestsRunner/WarnFilterScope.h"
+#include "Tests/WarnFilterScope.h"
 #include "UObject/FastReferenceCollector.h"
 
 
@@ -174,7 +174,7 @@ namespace UE::CoreUObject::Private::Tests
 		};
 		RefObj->ObjectPtr = Obj; //add a reference
 
-		UE::Testing::FWarnFilterScope _([](const TCHAR* Message, ELogVerbosity::Type Verbosity, const FName& Category)
+		FWarnFilterScope _([](const TCHAR* Message, ELogVerbosity::Type Verbosity, const FName& Category)
 			{
 				bool bFiltered = FCString::Strncmp(Message, TEXT("Marking leaking package"), FCString::Strlen(TEXT("Marking leaking package"))) == 0 && Verbosity == ELogVerbosity::Type::Warning && Category == TEXT("PluginHandlerLog");
 				return bFiltered;
diff --git a/Engine/Source/Runtime/CoreUObject/Tests/PropertyClassTest.cpp b/Engine/Source/Runtime/CoreUObject/Tests/PropertyClassTest.cpp
--- a/Engine/Source/Runtime/CoreUObject/Tests/PropertyClassTest.cpp
+++ b/Engine/Source/Runtime/CoreUObject/Tests/PropertyClassTest.cpp
@@ -11,7 +11,7 @@
 #include "UObject/LinkerPlaceholderExportObject.h"
 #include "UObject/LinkerPlaceholderClass.h"
 #include "UObject/ObjectHandleTracking.h"
-#include "LowLevelTestsRunner/WarnFilterScope.h"
+#include "Tests/WarnFilterScope.h"
 #include "ObjectRefTrackingTestBase.h"
 
 TEST_CASE("UE::CoreUObject::FClassProperty::Identical")
