Upgrade AutoRTFM_Enabled to apply to *all* Verse code.

diff --git a/Engine/Source/Programs/AutoRTFMTests/Private/AutoRTFMTests.cpp b/Engine/Source/Programs/AutoRTFMTests/Private/AutoRTFMTests.cpp
--- a/Engine/Source/Programs/AutoRTFMTests/Private/AutoRTFMTests.cpp
+++ b/Engine/Source/Programs/AutoRTFMTests/Private/AutoRTFMTests.cpp
@@ -44,8 +44,8 @@ int main(int ArgC, const char* ArgV[]) {
 	GEngineLoop.PreInit(0, nullptr);
 	FModuleManager::Get().StartProcessingNewlyLoadedObjects();
 
-	// Enable all Verse code to run under AutoRTFM (shouldn't affect our tests here, but better safe than sorry).
-	AutoRTFM::ForTheRuntime::SetAutoRTFMRuntime(AutoRTFM::ForTheRuntime::EAutoRTFMEnabledState::AutoRTFM_EnabledForAllVerse);
+	// Enable AutoRTFM.
+	AutoRTFM::ForTheRuntime::SetAutoRTFMRuntime(AutoRTFM::ForTheRuntime::EAutoRTFMEnabledState::AutoRTFM_Enabled);
 
 	// We don't want to trigger ensure's on abort because we are going to test that.
 	AutoRTFM::ForTheRuntime::SetEnsureOnAbortByLanguage(false);
diff --git a/Engine/Source/Runtime/Core/Private/AutoRTFM/API.cpp b/Engine/Source/Runtime/Core/Private/AutoRTFM/API.cpp
--- a/Engine/Source/Runtime/Core/Private/AutoRTFM/API.cpp
+++ b/Engine/Source/Runtime/Core/Private/AutoRTFM/API.cpp
@@ -101,7 +101,10 @@ namespace AutoRTFM
 				return false;
 			case EAutoRTFMEnabledState::AutoRTFM_Enabled:
 			case EAutoRTFMEnabledState::AutoRTFM_ForcedEnabled:
+#pragma clang diagnostic push
+#pragma clang diagnostic ignored "-Wdeprecated-declarations"
 			case EAutoRTFMEnabledState::AutoRTFM_EnabledForAllVerse:
+#pragma clang diagnostic pop
 				return true;
 			}
 #else
@@ -111,18 +114,7 @@ namespace AutoRTFM
 
 		bool IsAutoRTFMRuntimeEnabledForAllVerse()
 		{
-			// #noop if AutoRTFM is not compiled in
-#if UE_AUTORTFM
-			switch (GAutoRTFMRuntimeEnabled)
-			{
-			default:
-				return false;
-			case EAutoRTFMEnabledState::AutoRTFM_EnabledForAllVerse:
-				return true;
-			}
-#else
-			return false;
-#endif
+			return IsAutoRTFMRuntimeEnabled();
 		}
 
 		void SetEnsureOnAbortByLanguage(bool bEnabled)
diff --git a/Engine/Source/Runtime/Core/Public/AutoRTFM/AutoRTFM.h b/Engine/Source/Runtime/Core/Public/AutoRTFM/AutoRTFM.h
--- a/Engine/Source/Runtime/Core/Public/AutoRTFM/AutoRTFM.h
+++ b/Engine/Source/Runtime/Core/Public/AutoRTFM/AutoRTFM.h
@@ -695,7 +695,7 @@ namespace ForTheRuntime
 		// Disable AutoRTFM.
 		AutoRTFM_Disabled = 0,
 
-		// Enable AutoRTFM.
+		// Enable AutoRTFM for *all* Verse code (not just failure contexts).
 		AutoRTFM_Enabled,
 
 		// Force disable AutoRTFM - once set the AutoRTFM runtime cannot be re-enabled.
@@ -704,8 +704,7 @@ namespace ForTheRuntime
 		// Force enable AutoRTFM - once set the AutoRTFM runtime cannot be re-enabled.
 		AutoRTFM_ForcedEnabled,
 
-		// Enable AutoRTFM for *all* Verse code (not just failure contexts).
-		AutoRTFM_EnabledForAllVerse,
+		AutoRTFM_EnabledForAllVerse [[deprecated("Use AutoRTFM_Enabled instead!")]],
 	};
 
 	// An enum to represent whether we should abort and retry transactions (for testing purposes).
@@ -732,6 +731,7 @@ namespace ForTheRuntime
 	UE_AUTORTFM_API bool IsAutoRTFMRuntimeEnabled();
 
 	// Query whether the AutoRTFM runtime is enabled for *all* Verse code (not just failure contexts).
+	[[deprecated("Use IsAutoRTFMRuntimeEnabled instead!")]]
 	UE_AUTORTFM_API bool IsAutoRTFMRuntimeEnabledForAllVerse();
 
 	// Set whether we should trigger an ensure on an abort-by-language.
