Updating SymLoadModuleExW failure warning to include more context about what symbols the program was attempting to load when the error occured along with additional error context.

diff --git a/Engine/Source/Runtime/Core/Private/Windows/WindowsPlatformStackWalk.cpp b/Engine/Source/Runtime/Core/Private/Windows/WindowsPlatformStackWalk.cpp
--- a/Engine/Source/Runtime/Core/Private/Windows/WindowsPlatformStackWalk.cpp
+++ b/Engine/Source/Runtime/Core/Private/Windows/WindowsPlatformStackWalk.cpp
@@ -782,7 +782,9 @@ void LoadSymbolsForModule(HMODULE ModuleHandle, const FString& RemoteStorage)
 		// If the module is already loaded, the return value is zero and GetLastError returns ERROR_SUCCESS.
 		if (ErrorCode != ERROR_SUCCESS)
 		{
-			UE_LOG(LogWindows, Warning, TEXT("SymLoadModuleExW. Error: %d"), ErrorCode);
+			TCHAR ErrorMsgBuffer[2048];
+			const TCHAR* ErrorMsg = FWindowsPlatformMisc::GetSystemErrorMessage(ErrorMsgBuffer, UE_ARRAY_COUNT(ErrorMsgBuffer), ErrorCode);
+			UE_LOG(LogWindows, Warning, TEXT("Failed to load module '%s' with SymLoadModuleExW. Error %d: %s"), ModuleName, ErrorCode, ErrorMsg);
 		}
 	}
 }
