New tests for TBitArray
#tests FoundationTests

diff --git a/Engine/Source/Runtime/Core/Tests/Containers/BitArrayTests.cpp b/Engine/Source/Runtime/Core/Tests/Containers/BitArrayTests.cpp
--- a/Engine/Source/Runtime/Core/Tests/Containers/BitArrayTests.cpp
+++ b/Engine/Source/Runtime/Core/Tests/Containers/BitArrayTests.cpp
@@ -7,7 +7,6 @@
 #include "Serialization/MemoryReader.h"
 #include "Serialization/MemoryWriter.h"
 #include "Templates/Function.h"
-
 #include "Tests/TestHarnessAdapter.h"
 #include "Tests/Assertions.h"
 
@@ -2075,6 +2074,20 @@ TEST_CASE_NAMED(FBitArrayFindTest, "System::Core::Containers::BitArray::Find", "
 
 TEST_CASE_NAMED(FBitArrayFindFromTest, "System::Core::Containers::BitArray::FindFrom", "[ApplicationContextMask][SmokeFilter]")
 {
+	{
+		TBitArray<> ArrEmpty;
+		REQUIRE_CHECK_SLOW(ArrEmpty.FindFrom(true, -1)); // prior to the [0, NumBits] range, we should assert
+		CHECK_MESSAGE(TEXT("FindFrom true in an empty array"), ArrEmpty.FindFrom(true, 0) == INDEX_NONE);
+		CHECK_MESSAGE(TEXT("FindFrom false in an empty array"), ArrEmpty.FindFrom(false, 0) == INDEX_NONE);
+		REQUIRE_CHECK_SLOW(ArrEmpty.FindFrom(true, 1)); // past the [0, NumBits] range, we should assert
+	}
+
+	{
+		TBitArray<> Arr(true, 1);
+		REQUIRE_CHECK_SLOW(Arr.FindFrom(true, -1)); // prior to the [-1, NumBits) range, we should assert
+		REQUIRE_CHECK_SLOW(Arr.FindFrom(true, Arr.Num() + 1)); // prior to the [-1, NumBits) range, we should assert
+	}
+
 	{
 		TBitArray<> ArrTrue(true, NumBitsPerDWORD);
 		CHECK_MESSAGE(TEXT("Find true in an array full of true, starting at index 0"), ArrTrue.FindFrom(true, 0) == 0);
@@ -2145,6 +2158,98 @@ TEST_CASE_NAMED(FBitArrayTestFindLastTest, "System::Core::Containers::BitArray::
 	}
 }
 
+TEST_CASE_NAMED(FBitArrayTestFindLastTestFrom, "System::Core::Containers::BitArray::FindLastFrom", "[ApplicationContextMask][SmokeFilter]")
+{
+	{
+		TBitArray<> ArrEmpty;
+		REQUIRE_CHECK_SLOW(ArrEmpty.FindLastFrom(true, -2)); // prior to the [-1, NumBits) range, we should assert
+		CHECK_MESSAGE(TEXT("FindLastFrom true in an empty array"), ArrEmpty.FindLastFrom(true, -1) == INDEX_NONE);
+		CHECK_MESSAGE(TEXT("FindLastFrom false in an empty array"), ArrEmpty.FindLastFrom(false, -1) == INDEX_NONE);
+		REQUIRE_CHECK_SLOW(ArrEmpty.FindLastFrom(true, 0)); // past the [-1, NumBits) range, we should assert
+	}
+
+	{
+		TBitArray<> Arr(true, 1);
+		REQUIRE_CHECK_SLOW(Arr.FindLastFrom(true, -2)); // prior to the [-1, NumBits) range, we should assert
+		REQUIRE_CHECK_SLOW(Arr.FindLastFrom(true, Arr.Num() + 1)); // prior to the [-1, NumBits) range, we should assert
+	}
+
+	{
+		TBitArray<> ArrTrue(true, NumBitsPerDWORD);
+		CHECK_MESSAGE(TEXT("Find last true in an array containing 32 true, starting from 0"), ArrTrue.FindLastFrom(true, 0) == 0);
+		CHECK_MESSAGE(TEXT("Find last true in an array containing 32 true, starting from the middle"), ArrTrue.FindLastFrom(true, NumBitsPerDWORD / 2) == NumBitsPerDWORD / 2);
+		CHECK_MESSAGE(TEXT("Find last true in an array containing 32 true, starting from the end"), ArrTrue.FindLastFrom(true, NumBitsPerDWORD - 1) == NumBitsPerDWORD - 1);
+		CHECK_MESSAGE(TEXT("Find last false in an array containing 32 true, starting from 0"), ArrTrue.FindLastFrom(false, 0) == INDEX_NONE);
+		CHECK_MESSAGE(TEXT("Find last false in an array containing 32 true, starting from the middle"), ArrTrue.FindLastFrom(false, NumBitsPerDWORD / 2) == INDEX_NONE);
+		CHECK_MESSAGE(TEXT("Find last false in an array containing 32 true, starting from the end"), ArrTrue.FindLastFrom(false, NumBitsPerDWORD - 1) == INDEX_NONE);
+	}
+
+	{
+		TBitArray<> ArrFalse(false, NumBitsPerDWORD);
+		CHECK_MESSAGE(TEXT("Find last false in an array containing 32 false, starting from 0"), ArrFalse.FindLastFrom(false, 0) == 0);
+		CHECK_MESSAGE(TEXT("Find last false in an array containing 32 false, starting from the middle"), ArrFalse.FindLastFrom(false, NumBitsPerDWORD / 2) == NumBitsPerDWORD / 2);
+		CHECK_MESSAGE(TEXT("Find last false in an array containing 32 false, starting from the end"), ArrFalse.FindLastFrom(false, NumBitsPerDWORD - 1) == NumBitsPerDWORD - 1);
+		CHECK_MESSAGE(TEXT("Find last true in an array containing 32 false, starting from 0"), ArrFalse.FindLastFrom(true, 0) == INDEX_NONE);
+		CHECK_MESSAGE(TEXT("Find last true in an array containing 32 false, starting from the middle"), ArrFalse.FindLastFrom(true, NumBitsPerDWORD / 2) == INDEX_NONE);
+		CHECK_MESSAGE(TEXT("Find last true in an array containing 32 false, starting from the end"), ArrFalse.FindLastFrom(true, NumBitsPerDWORD - 1) == INDEX_NONE);
+	}
+
+	{
+		TBitArray<> ArrEndTrue;
+		ArrEndTrue.Add(false);
+		ArrEndTrue.Add(false);
+		ArrEndTrue.Add(true);
+		ArrEndTrue.Add(true);
+		CHECK_MESSAGE(TEXT("Find last true in an array containing 32 false, starting from the end"), ArrEndTrue.FindLastFrom(true, 0) == INDEX_NONE);
+		CHECK_MESSAGE(TEXT("Find last true in an array containing 32 false, starting from the end"), ArrEndTrue.FindLastFrom(true, 1) == INDEX_NONE);
+		CHECK_MESSAGE(TEXT("Find last true in an array containing 32 false, starting from the end"), ArrEndTrue.FindLastFrom(true, 2) == 2);
+		CHECK_MESSAGE(TEXT("Find last true in an array containing 32 false, starting from the end"), ArrEndTrue.FindLastFrom(true, 3) == 3);
+	}
+
+	{
+		uint32 InitialDWords[3] =
+			{ 0b00010000000000000000000011110100
+			, 0b10000001000000000000000000000001 
+			, 0b00000000000000000000000000000010 };
+		TBitArray<> Arr(false, std::size(InitialDWords) * NumBitsPerDWORD);
+		Arr.SetRangeFromRange(0, Arr.Num(), InitialDWords);
+
+		auto TestArr = [&Arr](bool bInTestValue, const TCHAR* InTestValueStr)
+		{
+			auto TestFindLastFrom = [&Arr](bool bInTestValue, const TCHAR* InTestValueStr, int32 InStartIndex, int32 InExpectedValue)
+			{
+				CHECK_MESSAGE(FString::Printf(TEXT("Find last %s in an array starting from %i"), InTestValueStr, InStartIndex), Arr.FindLastFrom(bInTestValue, InStartIndex) == InExpectedValue);
+			};
+
+			// First DWord
+			TestFindLastFrom(bInTestValue, InTestValueStr, 31, 28);
+			TestFindLastFrom(bInTestValue, InTestValueStr, 28, 28);
+			TestFindLastFrom(bInTestValue, InTestValueStr, 27, 7);
+			TestFindLastFrom(bInTestValue, InTestValueStr, 7, 7);
+			TestFindLastFrom(bInTestValue, InTestValueStr, 4, 4);
+			TestFindLastFrom(bInTestValue, InTestValueStr, 3, 2);
+			TestFindLastFrom(bInTestValue, InTestValueStr, 1, INDEX_NONE);
+			TestFindLastFrom(bInTestValue, InTestValueStr, 0, INDEX_NONE);
+			// Second DWord
+			TestFindLastFrom(bInTestValue, InTestValueStr, 32, 32);
+			TestFindLastFrom(bInTestValue, InTestValueStr, 33, 32);
+			TestFindLastFrom(bInTestValue, InTestValueStr, 55, 32);
+			TestFindLastFrom(bInTestValue, InTestValueStr, 56, 56);
+			TestFindLastFrom(bInTestValue, InTestValueStr, 57, 56);
+			TestFindLastFrom(bInTestValue, InTestValueStr, 62, 56);
+			TestFindLastFrom(bInTestValue, InTestValueStr, 63, 63);
+			// Third DWord
+			TestFindLastFrom(bInTestValue, InTestValueStr, 64, 63);
+			TestFindLastFrom(bInTestValue, InTestValueStr, 65, 65);
+			TestFindLastFrom(bInTestValue, InTestValueStr, 95, 65);
+		};
+
+		TestArr(true, TEXT("true"));
+		Arr.BitwiseNOT();
+		TestArr(false, TEXT("false"));
+	}
+}
+
 TEST_CASE_NAMED(FBitArrayTestFindAndSetFirstZeroBitTest, "System::Core::Containers::BitArray::FindAndSetFirstZeroBit", "[ApplicationContextMask][SmokeFilter]")
 {
 	{
