UnrealMathSSE simplifications for SSE4.2 min spec

We've required SSE4.2 since last year, so remove some fallback paths and also use SSE4.1 instructions where appropriate.
Also add tests for modified functions that didn't have any tests yet.

diff --git a/Engine/Source/Runtime/Core/Public/Math/UnrealMathSSE.h b/Engine/Source/Runtime/Core/Public/Math/UnrealMathSSE.h
--- a/Engine/Source/Runtime/Core/Public/Math/UnrealMathSSE.h
+++ b/Engine/Source/Runtime/Core/Public/Math/UnrealMathSSE.h
@@ -842,12 +842,9 @@ FORCEINLINE void VectorStoreAlignedStreamed(const VectorRegister4Double& Vec, do
  */
 FORCEINLINE void VectorStoreFloat3( const VectorRegister4Float& Vec, float* Ptr )
 {
-	union { VectorRegister4Float v; float f[4]; } Tmp;
-	Tmp.v = Vec;
-	float* FloatPtr = (float*)(Ptr);
-	FloatPtr[0] = Tmp.f[0];
-	FloatPtr[1] = Tmp.f[1];
-	FloatPtr[2] = Tmp.f[2];
+	VectorRegister4Float Tmp = Vec;
+	_mm_storel_pi((__m64*)(Ptr), Tmp);
+	_mm_store_ss(&Ptr[2], _mm_movehl_ps(Tmp, Tmp));
 }
 
 FORCEINLINE void VectorStoreFloat3(const VectorRegister4Double& Vec, double* Dst)
@@ -1803,12 +1800,14 @@ FORCEINLINE VectorRegister4Double VectorCompareLE(const VectorRegister4Double& V
 
 FORCEINLINE VectorRegister4Float VectorSelect(const VectorRegister4Float& Mask, const VectorRegister4Float& Vec1, const VectorRegister4Float& Vec2 )
 {
-	return _mm_xor_ps(Vec2, _mm_and_ps(Mask, _mm_xor_ps(Vec1, Vec2)));
+	// Can't (in general) use BLENDVPS despite our SSE4.1 minimum requirement since
+	// this is defined to be bitwise, not element-wise with MSB as toggle
+	return _mm_or_ps(_mm_and_ps(Mask, Vec1), _mm_andnot_ps(Mask, Vec2));
 }
 
 FORCEINLINE VectorRegister2Double VectorSelect(const VectorRegister2Double& Mask, const VectorRegister2Double& Vec1, const VectorRegister2Double& Vec2)
 {
-	return _mm_xor_pd(Vec2, _mm_and_pd(Mask, _mm_xor_pd(Vec1, Vec2)));
+	return _mm_or_pd(_mm_and_pd(Mask, Vec1), _mm_andnot_pd(Mask, Vec2));
 }
 
 FORCEINLINE VectorRegister4Double VectorSelect(const VectorRegister4Double& Mask, const VectorRegister4Double& Vec1, const VectorRegister4Double& Vec2)
@@ -1818,7 +1817,7 @@ FORCEINLINE VectorRegister4Double VectorSelect(const VectorRegister4Double& Mask
 	Result.XY = VectorSelect(Mask.XY, Vec1.XY, Vec2.XY);
 	Result.ZW = VectorSelect(Mask.ZW, Vec1.ZW, Vec2.ZW);
 #else
-	Result = _mm256_xor_pd(Vec2, _mm256_and_pd(Mask, _mm256_xor_pd(Vec1, Vec2)));
+	Result = _mm256_or_pd(_mm256_and_pd(Mask, Vec1), _mm256_andnot_pd(Mask, Vec2));
 #endif
 	return Result;
 }
@@ -2182,7 +2181,7 @@ FORCEINLINE VectorRegister4Double VectorReciprocalEstimate(const VectorRegister4
 */
 FORCEINLINE VectorRegister4Float VectorSet_W0(const VectorRegister4Float& Vec)
 {
-	return _mm_and_ps(Vec, GlobalVectorConstants::XYZMask());
+	return _mm_insert_ps(Vec, Vec, 0x08); // Copies lane 0 to lane 0 and zero-masks lane 3
 }
 
 FORCEINLINE VectorRegister4Double VectorSet_W0(const VectorRegister4Double& Vec)
@@ -2205,11 +2204,7 @@ FORCEINLINE VectorRegister4Double VectorSet_W0(const VectorRegister4Double& Vec)
 */
 FORCEINLINE VectorRegister4Float VectorSet_W1( const VectorRegister4Float& Vec)
 {
-	// Temp = (Vector[2]. Vector[3], 1.0f, 1.0f)
-	VectorRegister4Float Temp = _mm_movehl_ps( VectorOneFloat(), Vec);
-
-	// Return (Vector[0], Vector[1], Vector[2], 1.0f)
-	return VectorShuffle(Vec, Temp, 0, 1, 0, 3 );
+	return _mm_blend_ps(Vec, VectorOneFloat(), 0b1000);
 }
 
 FORCEINLINE VectorRegister4Double VectorSet_W1(const VectorRegister4Double& Vec)
@@ -2489,12 +2484,11 @@ FORCEINLINE int VectorMaskBits(const VectorRegister4Double& VecMask)
  */
 FORCEINLINE VectorRegister4Float VectorMergeVecXYZ_VecW(const VectorRegister4Float& VecXYZ, const VectorRegister4Float& VecW)
 {
-	return VectorSelect(GlobalVectorConstants::XYZMask(), VecXYZ, VecW);
+	return _mm_blend_ps(VecXYZ, VecW, 0b1000);
 }
 
 FORCEINLINE VectorRegister4Double VectorMergeVecXYZ_VecW(const VectorRegister4Double& VecXYZ, const VectorRegister4Double& VecW)
 {
-	//return VectorSelect(GlobalVectorConstants::DoubleXYZMask(), VecXYZ, VecW);
 	return VectorRegister4Double(VecXYZ.XY, _mm_move_sd(VecW.ZW, VecXYZ.ZW));
 }
 
@@ -2505,8 +2499,8 @@ FORCEINLINE VectorRegister4Double VectorMergeVecXYZ_VecW(const VectorRegister4Do
  * @return				VectorRegister4Float( float(Ptr[0]), float(Ptr[1]), float(Ptr[2]), float(Ptr[3]) )
  */
 // Looks complex but is really quite straightforward:
-// Load as 32-bit value, unpack 4x unsigned bytes to 4x 16-bit ints, then unpack again into 4x 32-bit ints, then convert to 4x floats
-#define VectorLoadByte4( Ptr )			_mm_cvtepi32_ps(_mm_unpacklo_epi16(_mm_unpacklo_epi8(_mm_cvtsi32_si128(*(int32*)Ptr), _mm_setzero_si128()), _mm_setzero_si128()))
+// Load as 32-bit value, convert into 4x 32-bit ints, then convert to 4x floats
+#define VectorLoadByte4( Ptr )			_mm_cvtepi32_ps(_mm_cvtepu8_epi32(_mm_cvtsi32_si128(*(int32*)Ptr)))
 
 /**
 * Loads 4 signed BYTEs from unaligned memory and converts them into 4 FLOATs.
@@ -2515,13 +2509,11 @@ FORCEINLINE VectorRegister4Double VectorMergeVecXYZ_VecW(const VectorRegister4Do
 * @return				VectorRegister4Float( float(Ptr[0]), float(Ptr[1]), float(Ptr[2]), float(Ptr[3]) )
 */
 // Looks complex but is really quite straightforward:
-// Load as 32-bit value, unpack 4x unsigned bytes to 4x 16-bit ints, then unpack again into 4x 32-bit ints, then convert to 4x floats
+// Load as 32-bit value, unpack 4x signed 8-bit ints to 4x 32-bit ints, then convert to 4x floats
 FORCEINLINE VectorRegister4Float VectorLoadSignedByte4(const void* Ptr)
 {
-	auto Temp = _mm_unpacklo_epi16(_mm_unpacklo_epi8(_mm_cvtsi32_si128(*(int32*)Ptr), _mm_setzero_si128()), _mm_setzero_si128());
-	auto Mask = _mm_cmpgt_epi32(Temp, _mm_set1_epi32(127));
-	auto Comp = _mm_and_si128(Mask, _mm_set1_epi32(~127));
-	return _mm_cvtepi32_ps(_mm_or_si128(Comp, Temp));
+	__m128i Temp = _mm_cvtepi8_epi32(_mm_cvtsi32_si128(*(int32*)Ptr));
+	return _mm_cvtepi32_ps(Temp);
 }
 
 /**
@@ -2546,11 +2538,14 @@ FORCEINLINE void VectorStoreByte4( const VectorRegister4Float& Vec, void* Ptr )
 {
 	// Looks complex but is really quite straightforward:
 	// Convert 4x floats to 4x 32-bit ints, then pack into 4x 16-bit ints, then into 4x 8-bit unsigned ints, then store as a 32-bit value
-	*(int32*)Ptr = _mm_cvtsi128_si32(_mm_packus_epi16(_mm_packs_epi32(_mm_cvttps_epi32(Vec), _mm_setzero_si128()), _mm_setzero_si128()));
+	__m128i VecInt32 = _mm_cvttps_epi32(Vec);
+	__m128i VecInt16 = _mm_packs_epi32(VecInt32, VecInt32);
+	__m128i VecUInt8 = _mm_packus_epi16(VecInt16, VecInt16);
+	*(int32*)Ptr = _mm_cvtsi128_si32(VecUInt8);
 }
 
 /**
-* Converts the 4 FLOATs in the vector to 4 BYTEs, clamped to [-127,127], and stores to unaligned memory.
+* Converts the 4 FLOATs in the vector to 4 BYTEs, clamped to [-128,127], and stores to unaligned memory.
 *
 * @param Vec			Vector containing 4 FLOATs
 * @param Ptr			Unaligned memory pointer to store the 4 BYTEs.
@@ -2559,7 +2554,10 @@ FORCEINLINE void VectorStoreSignedByte4(const VectorRegister4Float& Vec, void* P
 {
 	// Looks complex but is really quite straightforward:
 	// Convert 4x floats to 4x 32-bit ints, then pack into 4x 16-bit ints, then into 4x 8-bit unsigned ints, then store as a 32-bit value
-	*(int32*)Ptr = _mm_cvtsi128_si32(_mm_packs_epi16(_mm_packs_epi32(_mm_cvttps_epi32(Vec), _mm_setzero_si128()), _mm_setzero_si128()));
+	__m128i VecInt32 = _mm_cvttps_epi32(Vec);
+	__m128i VecInt16 = _mm_packs_epi32(VecInt32, VecInt32);
+	__m128i VecInt8 = _mm_packs_epi16(VecInt16, VecInt16);
+	*(int32*)Ptr = _mm_cvtsi128_si32(VecInt8);
 }
 
 
@@ -2616,7 +2614,7 @@ FORCEINLINE void VectorStoreURGB10A2N(const VectorRegister4Float& Vec, void* Ptr
  * @param Ptr			Unaligned memory pointer to the RGBA16(8 bytes).
  * @return				VectorRegister4Float with 4 FLOATs loaded from Ptr.
  */
-#define VectorLoadURGBA16N( Ptr ) _mm_cvtepi32_ps(_mm_unpacklo_epi16(_mm_loadl_epi64((const __m128i*)Ptr), _mm_setzero_si128()))
+#define VectorLoadURGBA16N( Ptr ) _mm_cvtepi32_ps(_mm_cvtepu16_epi32(_mm_loadl_epi64((const __m128i*)Ptr)))
 
 /**
  * Loads packed signed RGBA16(8 bytes) from unaligned memory and converts them into 4 FLOATs.
@@ -2626,33 +2624,25 @@ FORCEINLINE void VectorStoreURGB10A2N(const VectorRegister4Float& Vec, void* Ptr
  */
 FORCEINLINE VectorRegister4Float VectorLoadSRGBA16N(const void* Ptr)
 {
-	auto Temp = _mm_unpacklo_epi16(_mm_loadl_epi64((const __m128i*)Ptr), _mm_setzero_si128());
-	auto Mask = _mm_cmpgt_epi32(Temp, _mm_set1_epi32(32767));
-	auto Comp = _mm_and_si128(Mask, _mm_set1_epi32(~32767));
-	return _mm_cvtepi32_ps(_mm_or_si128(Comp, Temp));
+	__m128i Temp = _mm_cvtepi16_epi32(_mm_loadl_epi64((const __m128i*)Ptr));
+	return _mm_cvtepi32_ps(Temp);
 }
 
 /**
-* Converts the 4 FLOATs in the vector RGBA16, clamped to [0, 65535], and stores to unaligned memory.
+* Converts the 4 FLOATs in the vector to RGBA16, clamped to [0, 65535], and stores to unaligned memory.
 *
 * @param Vec			Vector containing 4 FLOATs
-* @param Ptr			Unaligned memory pointer to store the packed RGB10A2(4 bytes).
+* @param Ptr			Unaligned memory pointer to store the packed RGBA16(8 bytes).
 */
 FORCEINLINE void VectorStoreURGBA16N(const VectorRegister4Float& Vec, void* Ptr)
 {
-
 	VectorRegister4Float Tmp;
 	Tmp = _mm_max_ps(Vec, MakeVectorRegisterFloat(0.0f, 0.0f, 0.0f, 0.0f));
 	Tmp = _mm_min_ps(Tmp, MakeVectorRegisterFloat(1.0f, 1.0f, 1.0f, 1.0f));
 	Tmp = _mm_mul_ps(Tmp, MakeVectorRegisterFloat(65535.0f, 65535.0f, 65535.0f, 65535.0f));
 
 	VectorRegister4Int TmpI = _mm_cvtps_epi32(Tmp);
-
-	uint16* Out = (uint16*)Ptr;
-	Out[0] = static_cast<int16>(_mm_extract_epi16(TmpI, 0));
-	Out[1] = static_cast<int16>(_mm_extract_epi16(TmpI, 2));
-	Out[2] = static_cast<int16>(_mm_extract_epi16(TmpI, 4));
-	Out[3] = static_cast<int16>(_mm_extract_epi16(TmpI, 6));
+	_mm_storel_epi64((__m128i*)Ptr, _mm_packus_epi32(TmpI, TmpI));
 }
 
 /**
@@ -2802,26 +2792,12 @@ FORCEINLINE bool VectorContainsNaNOrInfinite(const VectorRegister4Double& Vec)
 
 FORCEINLINE VectorRegister4Float VectorTruncate(const VectorRegister4Float& Vec)
 {
-#if UE_PLATFORM_MATH_USE_SSE4_1
 	return _mm_round_ps(Vec, _MM_FROUND_TRUNC);
-#else
-	return _mm_cvtepi32_ps(_mm_cvttps_epi32(Vec));
-#endif
 }
 
 FORCEINLINE VectorRegister2Double TruncateVectorRegister2d(const VectorRegister2Double& V)
 {
-#if UE_PLATFORM_MATH_USE_SSE4_1
 	return _mm_round_pd(V, _MM_FROUND_TRUNC);
-#else
-	// TODO: LWC: Optimize
-	// Note: SSE2 just has _mm_cvttsd_si64(), which extracts only the truncated lower element, so there is some extra shuffling to get both values out.
-	int64 X = _mm_cvttsd_si64(V);
-	int64 Y = _mm_cvttsd_si64(_mm_shuffle_pd(V, V, SHUFFLEMASK2(1, 0)));
-	VectorRegister2Double A = _mm_cvtsi64_sd(V, X); // Converts to lowest element, copies upper.
-	VectorRegister2Double B = _mm_cvtsi64_sd(V, Y); // Converts to lowest element, copies upper.
-	return _mm_shuffle_pd(A, B, SHUFFLEMASK2(0, 0));
-#endif // UE_PLATFORM_MATH_USE_SSE4_1	
 }
 
 FORCEINLINE VectorRegister4Double VectorTruncate(const VectorRegister4Double& V)
@@ -2838,12 +2814,7 @@ FORCEINLINE VectorRegister4Double VectorTruncate(const VectorRegister4Double& V)
 
 FORCEINLINE VectorRegister4Float VectorRound(const VectorRegister4Float &Vec)
 {
-#if UE_PLATFORM_MATH_USE_SSE4_1
 	return _mm_round_ps(Vec, _MM_FROUND_TO_NEAREST_INT |_MM_FROUND_NO_EXC);
-#else
-	VectorRegister4Float Trunc = VectorTruncate(Vec);
-	return VectorAdd(Trunc, VectorTruncate(VectorMultiply(VectorSubtract(Vec, Trunc), GlobalVectorConstants::FloatAlmostTwo())));
-#endif
 }
 
 FORCEINLINE VectorRegister4Int VectorRoundToIntHalfToEven(const VectorRegister4Float& Vec)
@@ -2853,15 +2824,7 @@ FORCEINLINE VectorRegister4Int VectorRoundToIntHalfToEven(const VectorRegister4F
 
 FORCEINLINE VectorRegister4Float VectorCeil(const VectorRegister4Float& V)
 {
-#if UE_PLATFORM_MATH_USE_SSE4_1
 	return _mm_ceil_ps(V);
-#else
-	const VectorRegister4Float Trunc = VectorTruncate(V);
-	const VectorRegister4Float Frac = VectorSubtract(V, Trunc);
-	const VectorRegister4Float FracMask = VectorCompareGT(Frac, (GlobalVectorConstants::FloatZero));
-	const VectorRegister4Float Add = VectorSelect(FracMask, (GlobalVectorConstants::FloatOne), (GlobalVectorConstants::FloatZero));
-	return VectorAdd(Trunc, Add);
-#endif
 }
 
 FORCEINLINE VectorRegister4Double VectorCeil(const VectorRegister4Double& V)
@@ -2870,31 +2833,17 @@ FORCEINLINE VectorRegister4Double VectorCeil(const VectorRegister4Double& V)
 	VectorRegister4Double Result;
 	Result = _mm256_round_pd(V, _MM_FROUND_TO_POS_INF | _MM_FROUND_NO_EXC);
 	return Result;
-#elif UE_PLATFORM_MATH_USE_SSE4_1
+#else
 	VectorRegister4Double Result;
 	Result.XY = _mm_ceil_pd(V.XY);
 	Result.ZW = _mm_ceil_pd(V.ZW);
 	return Result;
-#else
-	const VectorRegister4Double Trunc = VectorTruncate(V);
-	const VectorRegister4Double Frac = VectorSubtract(V, Trunc);
-	const VectorRegister4Double FracMask = VectorCompareGT(Frac, (GlobalVectorConstants::DoubleZero));
-	const VectorRegister4Double Add = VectorSelect(FracMask, (GlobalVectorConstants::DoubleOne), (GlobalVectorConstants::DoubleZero));
-	return VectorAdd(Trunc, Add);
 #endif
 }
 
 FORCEINLINE VectorRegister4Float VectorFloor(const VectorRegister4Float& V)
 {
-#if UE_PLATFORM_MATH_USE_SSE4_1
 	return _mm_floor_ps(V);
-#else
-	const VectorRegister4Float Trunc = VectorTruncate(V);
-	const VectorRegister4Float Frac = VectorSubtract(V, Trunc);
-	const VectorRegister4Float FracMask = VectorCompareLT(Frac, (GlobalVectorConstants::FloatZero));
-	const VectorRegister4Float Add = VectorSelect(FracMask, (GlobalVectorConstants::FloatMinusOne), (GlobalVectorConstants::FloatZero));
-	return VectorAdd(Trunc, Add);
-#endif
 }
 
 FORCEINLINE VectorRegister4Double VectorFloor(const VectorRegister4Double& V)
@@ -2903,17 +2852,11 @@ FORCEINLINE VectorRegister4Double VectorFloor(const VectorRegister4Double& V)
 	VectorRegister4Double Result;
 	Result = _mm256_round_pd(V, _MM_FROUND_TO_NEG_INF | _MM_FROUND_NO_EXC);
 	return Result;
-#elif UE_PLATFORM_MATH_USE_SSE4_1
+#else
 	VectorRegister4Double Result;
 	Result.XY = _mm_floor_pd(V.XY);
 	Result.ZW = _mm_floor_pd(V.ZW);
 	return Result;
-#else
-	const VectorRegister4Double Trunc = VectorTruncate(V);
-	const VectorRegister4Double Frac = VectorSubtract(V, Trunc);
-	const VectorRegister4Double FracMask = VectorCompareLT(Frac, (GlobalVectorConstants::DoubleZero));
-	const VectorRegister4Double Add = VectorSelect(FracMask, (GlobalVectorConstants::DoubleMinusOne), (GlobalVectorConstants::DoubleZero));
-	return VectorAdd(Trunc, Add);
 #endif
 }
 
@@ -3420,7 +3363,8 @@ FORCEINLINE VectorRegister4Double VectorATan2(const VectorRegister4Double& Y, co
 
 FORCEINLINE VectorRegister4Int VectorIntSelect(const VectorRegister4Int& Mask, const VectorRegister4Int& Vec1, const VectorRegister4Int& Vec2)
 {
-	return _mm_xor_si128(Vec2, _mm_and_si128(Mask, _mm_xor_si128(Vec1, Vec2)));
+	// Can't use PBLENDVB in general because this is a bitwise select, not byte-lane-wise
+	return _mm_or_si128(_mm_and_si128(Mask, Vec1), _mm_andnot_si128(Mask, Vec2));
 }
 
 //Arithmetic
@@ -3429,34 +3373,24 @@ FORCEINLINE VectorRegister4Int VectorIntSelect(const VectorRegister4Int& Mask, c
 
 FORCEINLINE VectorRegister4Int VectorIntMultiply(const VectorRegister4Int& A, const VectorRegister4Int& B)
 {
-#if UE_PLATFORM_MATH_USE_SSE4_1
 	return _mm_mullo_epi32(A, B);
-#else
-	//SSE2 doesn't have a multiply op for 4 32bit ints. Ugh.
-	__m128i Temp0 = _mm_mul_epu32(A, B);
-	__m128i Temp1 = _mm_mul_epu32(_mm_srli_si128(A, 4), _mm_srli_si128(B, 4));
-	return _mm_unpacklo_epi32(_mm_shuffle_epi32(Temp0, _MM_SHUFFLE(0, 0, 2, 0)), _mm_shuffle_epi32(Temp1, _MM_SHUFFLE(0, 0, 2, 0)));
-#endif
 }
 
 #define VectorIntNegate(A) VectorIntSubtract( GlobalVectorConstants::IntZero, A)
 
 FORCEINLINE VectorRegister4Int VectorIntMin(const VectorRegister4Int& A, const VectorRegister4Int& B)
 {
-	VectorRegister4Int Mask = VectorIntCompareLT(A, B);
-	return VectorIntSelect(Mask, A, B);
+	return _mm_min_epi32(A, B);
 }
 
 FORCEINLINE VectorRegister4Int VectorIntMax(const VectorRegister4Int& A, const VectorRegister4Int& B)
 {
-	VectorRegister4Int Mask = VectorIntCompareGT(A, B);
-	return VectorIntSelect(Mask, A, B);
+	return _mm_max_epi32(A, B);
 }
 
 FORCEINLINE VectorRegister4Int VectorIntAbs(const VectorRegister4Int& A)
 {
-	VectorRegister4Int Mask = VectorIntCompareGE(A, GlobalVectorConstants::IntZero);
-	return VectorIntSelect(Mask, A, VectorIntNegate(A));
+	return _mm_abs_epi32(A);
 }
 
 FORCEINLINE VectorRegister4Int VectorIntClamp(const VectorRegister4Int& Vec1, const VectorRegister4Int& Vec2, const VectorRegister4Int& Vec3) 
diff --git a/Engine/Source/Runtime/Core/Tests/Math/UnrealMathTest.cpp b/Engine/Source/Runtime/Core/Tests/Math/UnrealMathTest.cpp
--- a/Engine/Source/Runtime/Core/Tests/Math/UnrealMathTest.cpp
+++ b/Engine/Source/Runtime/Core/Tests/Math/UnrealMathTest.cpp
@@ -2315,7 +2315,7 @@ TEST_CASE_NAMED(FVectorRegisterAbstractionTest, "System::Core::Math::Vector Regi
 	float F1 = 1.f;
 	uint32 U1 = *(uint32*)&F1;
 	VectorRegister4Float V0, V1, V2, V3;
-	VectorRegister4Int VI0, VI1;
+	VectorRegister4Int VI0, VI1, VI2, VI3;
 	float Float0, Float1, Float2, Float3;
 
 	// Using a union as we need to do a bitwise cast of 0xFFFFFFFF into a float for NaN.
@@ -2606,8 +2606,7 @@ TEST_CASE_NAMED(FVectorRegisterAbstractionTest, "System::Core::Math::Vector Regi
 	V0 = MakeVectorRegister(1.0f, 3.0f, 2.0f, 4.0f);
 	LogTest<float>(TEXT("VectorCombineLow"), TestVectorsEqual(V0, V1));
 
-	// LoadByte4
-
+	// LoadByte4 and StoreByte4
 	uint8 Bytes[4] = { 25, 75, 125, 200 };
 	V0 = VectorLoadByte4(Bytes);
 	V1 = MakeVectorRegister(25.f, 75.f, 125.f, 200.f);
@@ -2623,6 +2622,38 @@ TEST_CASE_NAMED(FVectorRegisterAbstractionTest, "System::Core::Math::Vector Regi
 	V1 = VectorLoadByte4(Bytes);
 	LogTest<float>(TEXT("VectorStoreByte4"), TestVectorsEqual(V0, V1));
 
+	// LoadSignedByte4 and StoreSignedByte4
+	int8 SignedBytes[4] = { 25, -75, 125, -50 };
+	V0 = VectorLoadSignedByte4(SignedBytes);
+	V1 = MakeVectorRegister(25.f, -75.f, 125.f, -50.f);
+	LogTest<float>(TEXT("VectorLoadSignedByte4"), TestVectorsEqual(V0, V1));
+
+	V0 = MakeVectorRegister(-4.0f, 3.0f, -2.0f, 1.0f);
+	VectorStoreSignedByte4(V0, SignedBytes);
+	V1 = VectorLoadSignedByte4(SignedBytes);
+	LogTest<float>(TEXT("VectorStoreSignedByte4"), TestVectorsEqual(V0, V1));
+
+	// LoadURGBA16N and StoreURGBA16N
+	// NOTE: StoreURGBA16N is _not_ symmetric with Load. Load just gives you the
+	// values as floats in [0,65535], but Load expects the values as floats in
+	// [0,1], clamps to that range first, then scales to write ints.
+	uint16 UInt16s[4] = { 1234, 2345, 3456, 64000 };
+	V0 = VectorLoadURGBA16N(UInt16s);
+	V1 = MakeVectorRegister(1234.f, 2345.f, 3456.f, 64000.f);
+	LogTest<float>(TEXT("VectorLoadURGBA16N"), TestVectorsEqual(V0, V1));
+
+	V0 = MakeVectorRegister(-200000.f, 123.f / 65535.f, 234.f / 65535.f, 1.5f);
+	VectorStoreURGBA16N(V0, UInt16s);
+	V0 = MakeVectorRegister(0.f, 123.f, 234.f, 65535.f);
+	V1 = VectorLoadURGBA16N(UInt16s);
+	LogTest<float>(TEXT("VectorStoreURGBA16N"), TestVectorsEqual(V0, V1));
+
+	// LoadSRGBA16N
+	int16 Int16s[4] = { 1234, -2345, 3456, -32500 };
+	V0 = VectorLoadSRGBA16N(Int16s);
+	V1 = MakeVectorRegister(1234.f, -2345.f, 3456.f, -32500.f);
+	LogTest<float>(TEXT("VectorLoadSRGBA16N"), TestVectorsEqual(V0, V1));
+
 	// Vector Any/All comparisons
 	V0 = MakeVectorRegister(2.0f, 4.0f, 6.0f, 8.0f);
 	V1 = MakeVectorRegister(4.0f, 3.0f, 2.0f, 1.0f);
@@ -2941,6 +2972,38 @@ TEST_CASE_NAMED(FVectorRegisterAbstractionTest, "System::Core::Math::Vector Regi
 	MaskBits = VectorMaskBits(V0);
 	LogTest<float>(TEXT("VectorMaskBits"), MaskBits == 5);
 
+	// VectorIntSelect
+	VI0 = MakeVectorRegisterInt(1, 3, 2, 8);
+	VI1 = MakeVectorRegisterInt(2, 4, 2, 1);
+	VI2 = MakeVectorRegisterInt(-1, 0, 0, -1);
+	VI2 = VectorIntSelect(VI2, VI0, VI1);
+	VI3 = MakeVectorRegisterInt(1, 4, 2, 8);
+	LogTest<int32>(TEXT("VectorIntSelect"), TestVectorsEqualBitwise(VectorCastIntToFloat(VI2), VectorCastIntToFloat(VI3)));
+
+	VI0 = MakeVectorRegisterInt(1, 3, 5, 7);
+	VI1 = MakeVectorRegisterInt(2, 4, 6, 8);
+	VI2 = MakeVectorRegisterInt(0, -1, -1, 0);
+	VI2 = VectorIntSelect(VI2, VI0, VI1);
+	VI3 = MakeVectorRegisterInt(2, 3, 5, 8);
+	LogTest<int32>(TEXT("VectorIntSelect"), TestVectorsEqualBitwise(VectorCastIntToFloat(VI2), VectorCastIntToFloat(VI3)));
+
+	// VectorIntMin/VectorIntMax
+	VI0 = MakeVectorRegisterInt(10, 12, 30, 22);
+	VI1 = MakeVectorRegisterInt(5, 14, 32, -100);
+	VI2 = VectorIntMin(VI0, VI1);
+	VI3 = MakeVectorRegisterInt(5, 12, 30, -100);
+	LogTest<int32>(TEXT("VectorIntMin"), TestVectorsEqualBitwise(VectorCastIntToFloat(VI2), VectorCastIntToFloat(VI3)));
+
+	VI2 = VectorIntMax(VI0, VI1);
+	VI3 = MakeVectorRegisterInt(10, 14, 32, 22);
+	LogTest<int32>(TEXT("VectorIntMax"), TestVectorsEqualBitwise(VectorCastIntToFloat(VI2), VectorCastIntToFloat(VI3)));
+
+	// VectorIntAbs
+	VI0 = MakeVectorRegisterInt(1, -2, 3, -4);
+	VI2 = VectorIntAbs(VI0);
+	VI3 = MakeVectorRegisterInt(1, 2, 3, 4);
+	LogTest<int32>(TEXT("VectorIntAbs"), TestVectorsEqualBitwise(VectorCastIntToFloat(VI2), VectorCastIntToFloat(VI3)));
+
 	// Matrix multiplications and transformations
 	{
 		FMatrix44f	M0, M1, M2, M3;
