Improved Commit\Decommit on Apple platforms to reuduce RSS

#tests local testing on replays

diff --git a/Engine/Source/Runtime/Core/Private/Apple/ApplePlatformMemory.cpp b/Engine/Source/Runtime/Core/Private/Apple/ApplePlatformMemory.cpp
--- a/Engine/Source/Runtime/Core/Private/Apple/ApplePlatformMemory.cpp
+++ b/Engine/Source/Runtime/Core/Private/Apple/ApplePlatformMemory.cpp
@@ -807,13 +807,14 @@ void FApplePlatformMemory::FPlatformVirtualMemoryBlock::Commit(size_t InOffset,
 {
 	check(IsAligned(InOffset, GetCommitAlignment()) && IsAligned(InSize, GetCommitAlignment()));
 	check(InOffset >= 0 && InSize >= 0 && InOffset + InSize <= GetActualSize() && Ptr);
+	madvise(((uint8*)Ptr) + InOffset, InSize, MADV_FREE_REUSE);
 }
 
 void FApplePlatformMemory::FPlatformVirtualMemoryBlock::Decommit(size_t InOffset, size_t InSize)
 {
 	check(IsAligned(InOffset, GetCommitAlignment()) && IsAligned(InSize, GetCommitAlignment()));
 	check(InOffset >= 0 && InSize >= 0 && InOffset + InSize <= GetActualSize() && Ptr);
-	if (madvise(((uint8*)Ptr) + InOffset, InSize, MADV_DONTNEED) != 0)
+	if (madvise(((uint8*)Ptr) + InOffset, InSize, MADV_FREE_REUSABLE) != 0)
 	{
 		// we can ran out of VMAs here too!
 		FPlatformMemory::OnOutOfMemory(InSize, 0);
