Don't allow dynamically loaded additional localization to replace any existing data in the live table

Additional localization usually loads with a lower priority than the base game localization, however dynamically loaded localization was being allowed to replace existing data. This would lead to an inconsistency between a full reload and a dynamic load of the localization data, which this change corrects.

#rnx

diff --git a/Engine/Source/Runtime/Core/Private/Internationalization/TextLocalizationManager.cpp b/Engine/Source/Runtime/Core/Private/Internationalization/TextLocalizationManager.cpp
--- a/Engine/Source/Runtime/Core/Private/Internationalization/TextLocalizationManager.cpp
+++ b/Engine/Source/Runtime/Core/Private/Internationalization/TextLocalizationManager.cpp
@@ -1230,6 +1230,8 @@ void FTextLocalizationManager::HandleLocalizationTargetsMounted(TArrayView<const
 	LocLoadFlags |= (WITH_EDITOR ? ELocalizationLoadFlags::Editor : ELocalizationLoadFlags::None);
 	LocLoadFlags |= (FApp::IsGame() ? ELocalizationLoadFlags::Game : ELocalizationLoadFlags::None);
 	LocLoadFlags |= ELocalizationLoadFlags::Engine;
+	// We don't allow dynamically loaded additional data to replace any existing data, as additional localization is lower priority when performing a full update, so this flag emulates that behavior
+	LocLoadFlags |= ELocalizationLoadFlags::SkipExisting;
 
 	const TArray<FString> PrioritizedCultureNames = FInternationalization::Get().GetPrioritizedCultureNames(FInternationalization::Get().GetCurrentLanguage()->GetName());
 
@@ -1382,14 +1384,18 @@ void FTextLocalizationManager::LoadLocalizationResourcesForPrioritizedCultures_S
 		LocalizedTextSource->LoadLocalizedResources(FinalLocLoadFlags, PrioritizedCultureNames, NativeResource, LocalizedResource);
 	}
 
+	FUpdateLiveTableOptions UpdateOptions;
+	UpdateOptions.bDirtyTextRevision = false;
+	UpdateOptions.bReplaceExisting = !EnumHasAnyFlags(FinalLocLoadFlags, ELocalizationLoadFlags::SkipExisting);
+
 	// Apply the new display string data
 	if (!NativeResource.IsEmpty())
 	{
-		UpdateFromNative(MoveTemp(NativeResource), /*bDirtyTextRevision*/false);
+		UpdateFromNative(MoveTemp(NativeResource), UpdateOptions);
 	}
 	if (!LocalizedResource.IsEmpty())
 	{
-		UpdateFromLocalizations(MoveTemp(LocalizedResource), /*bDirtyTextRevision*/false);
+		UpdateFromLocalizations(MoveTemp(LocalizedResource), UpdateOptions);
 	}
 
 	DirtyTextRevision();
@@ -1466,8 +1472,11 @@ void FTextLocalizationManager::LoadLocalizationTargetsForPrioritizedCultures_Syn
 	}
 	else
 	{
+		FUpdateLiveTableOptions UpdateOptions;
+		UpdateOptions.bReplaceExisting = !EnumHasAnyFlags(LocLoadFlags, ELocalizationLoadFlags::SkipExisting);
+
 		UE_LOG(LogTextLocalizationManager, Verbose, TEXT("Patching LocRes data for %d entries"), LocalizedResource.Entries.Num());
-		UpdateFromLocalizations(MoveTemp(LocalizedResource), /*bDirtyTextRevision*/true);
+		UpdateFromLocalizations(MoveTemp(LocalizedResource), UpdateOptions);
 	}
 }
 
@@ -1606,7 +1615,7 @@ void FTextLocalizationManager::QueueAsyncTask(TUniqueFunction<void()>&& Task)
 	}
 }
 
-void FTextLocalizationManager::UpdateFromNative(FTextLocalizationResource&& TextLocalizationResource, const bool bDirtyTextRevision)
+void FTextLocalizationManager::UpdateFromNative(FTextLocalizationResource&& TextLocalizationResource, const FUpdateLiveTableOptions& UpdateOptions)
 {
 	TRACE_CPUPROFILER_EVENT_SCOPE(FTextLocalizationManager::UpdateFromNative);
 	LLM_SCOPE_BYNAME(TEXT("Localization/DisplayStrings"));
@@ -1631,7 +1640,20 @@ void FTextLocalizationManager::UpdateFromNative(FTextLocalizationResource&& Text
 			FTextLocalizationResource::FEntry& NewEntry = EntryPair.Value;
 
 			FDisplayStringEntry* LiveEntry = DisplayStringLookupTable.Find(TextId);
-			if (LiveEntry)
+			if (!LiveEntry)
+			{
+				// Add new entry
+				FDisplayStringEntry NewLiveEntry(
+					NewEntry.LocResID,						/*LocResID*/
+					NewEntry.LocalizationTargetPathId,		/*LocalizationTargetPathId*/
+					NewEntry.SourceStringHash,				/*SourceStringHash*/
+					NewEntry.LocalizedString.ToSharedRef()	/*String*/
+				);
+
+				DisplayStringLookupTable.Emplace(TextId, NewLiveEntry);
+				DisplayStringsByLocalizationTargetId.TrackTextId(INDEX_NONE, NewEntry.LocalizationTargetPathId, TextId);
+			}
+			else if (UpdateOptions.bReplaceExisting)
 			{
 				// Update existing entry
 				// If the existing entry is empty, we just overwrite it 
@@ -1649,32 +1671,19 @@ void FTextLocalizationManager::UpdateFromNative(FTextLocalizationResource&& Text
 					LiveEntry->LocalizationTargetPathId = NewEntry.LocalizationTargetPathId;
 				}
 			}
-			else
-			{
-				// Add new entry
-				FDisplayStringEntry NewLiveEntry(
-					NewEntry.LocResID,						/*LocResID*/
-					NewEntry.LocalizationTargetPathId,		/*LocalizationTargetPathId*/
-					NewEntry.SourceStringHash,				/*SourceStringHash*/
-					NewEntry.LocalizedString.ToSharedRef()	/*String*/
-				);
-
-				DisplayStringLookupTable.Emplace(TextId, NewLiveEntry);
-				DisplayStringsByLocalizationTargetId.TrackTextId(INDEX_NONE, NewEntry.LocalizationTargetPathId, TextId);
-			}
 		}
 
 		// Note: Do not use TextLocalizationResource after this point as we may have stolen some of its strings
 		TextLocalizationResource.Entries.Reset();
 	}
 
-	if (bDirtyTextRevision)
+	if (UpdateOptions.bDirtyTextRevision)
 	{
 		DirtyTextRevision();
 	}
 }
 
-void FTextLocalizationManager::UpdateFromLocalizations(FTextLocalizationResource&& TextLocalizationResource, const bool bDirtyTextRevision)
+void FTextLocalizationManager::UpdateFromLocalizations(FTextLocalizationResource&& TextLocalizationResource, const FUpdateLiveTableOptions& UpdateOptions)
 {
 	TRACE_CPUPROFILER_EVENT_SCOPE(FTextLocalizationManager::UpdateFromLocalizations);
 	LLM_SCOPE_BYNAME(TEXT("Localization/DisplayStrings"));
@@ -1698,7 +1707,20 @@ void FTextLocalizationManager::UpdateFromLocalizations(FTextLocalizationResource
 			FTextLocalizationResource::FEntry& NewEntry = EntryPair.Value;
 
 			FDisplayStringEntry* LiveEntry = DisplayStringLookupTable.Find(TextId);
-			if (LiveEntry)
+			if (!LiveEntry)
+			{
+				// Add new entry
+				FDisplayStringEntry NewLiveEntry(
+					NewEntry.LocResID,						/*LocResID*/
+					NewEntry.LocalizationTargetPathId,		/*LocalizationTargetPathId*/
+					NewEntry.SourceStringHash,				/*SourceStringHash*/
+					NewEntry.LocalizedString.ToSharedRef()	/*String*/
+				);
+
+				DisplayStringLookupTable.Emplace(TextId, NewLiveEntry);
+				DisplayStringsByLocalizationTargetId.TrackTextId(INDEX_NONE, NewEntry.LocalizationTargetPathId, TextId);
+			}
+			else if (UpdateOptions.bReplaceExisting)
 			{
 				// Update existing entry
 				// If the source string hashes are are the same, we can replace the display string.
@@ -1715,26 +1737,13 @@ void FTextLocalizationManager::UpdateFromLocalizations(FTextLocalizationResource
 #endif	// WITH_EDITORONLY_DATA
 				}
 			}
-			else
-			{
-				// Add new entry
-				FDisplayStringEntry NewLiveEntry(
-					NewEntry.LocResID,						/*LocResID*/
-					NewEntry.LocalizationTargetPathId,		/*LocalizationTargetPathId*/
-					NewEntry.SourceStringHash,				/*SourceStringHash*/
-					NewEntry.LocalizedString.ToSharedRef()	/*String*/
-				);
-
-				DisplayStringLookupTable.Emplace(TextId, NewLiveEntry);
-				DisplayStringsByLocalizationTargetId.TrackTextId(INDEX_NONE, NewEntry.LocalizationTargetPathId, TextId);
-			}
 		}
 
 		// Note: Do not use TextLocalizationResource after this point as we may have stolen some of its strings
 		TextLocalizationResource.Entries.Reset();
 	}
 
-	if (bDirtyTextRevision)
+	if (UpdateOptions.bDirtyTextRevision)
 	{
 		DirtyTextRevision();
 	}
diff --git a/Engine/Source/Runtime/Core/Public/Internationalization/LocalizedTextSourceTypes.h b/Engine/Source/Runtime/Core/Public/Internationalization/LocalizedTextSourceTypes.h
--- a/Engine/Source/Runtime/Core/Public/Internationalization/LocalizedTextSourceTypes.h
+++ b/Engine/Source/Runtime/Core/Public/Internationalization/LocalizedTextSourceTypes.h
@@ -59,6 +59,12 @@ enum class ELocalizationLoadFlags : uint8
 
 	/** Force localized game data to be loaded, even when running in the editor */
 	ForceLocalizedGame = 1<<5,
+
+	/**
+	 * Skip updating any entries that already exist in the live table
+	 * @note Not useful when performing a full update, but has utility when patching in new untrusted localization data (eg, loading UGC localization data over the base localization data)
+	 */
+	SkipExisting = 1<<6,
 };
 ENUM_CLASS_FLAGS(ELocalizationLoadFlags);
 
diff --git a/Engine/Source/Runtime/Core/Public/Internationalization/TextLocalizationManager.h b/Engine/Source/Runtime/Core/Public/Internationalization/TextLocalizationManager.h
--- a/Engine/Source/Runtime/Core/Public/Internationalization/TextLocalizationManager.h
+++ b/Engine/Source/Runtime/Core/Public/Internationalization/TextLocalizationManager.h
@@ -372,11 +372,21 @@ private:
 	/** Queue the given task to run async (where possible), chained to any existing AsyncLocalizationTask */
 	CORE_API void QueueAsyncTask(TUniqueFunction<void()>&& Task);
 
+	struct FUpdateLiveTableOptions
+	{
+		FUpdateLiveTableOptions()
+		{
+		}
+
+		bool bDirtyTextRevision = true;
+		bool bReplaceExisting = true;
+	};
+
 	/** Updates display string entries and adds new display string entries based on provided native text. */
-	CORE_API void UpdateFromNative(FTextLocalizationResource&& TextLocalizationResource, const bool bDirtyTextRevision = true);
+	CORE_API void UpdateFromNative(FTextLocalizationResource&& TextLocalizationResource, const FUpdateLiveTableOptions& UpdateOptions = FUpdateLiveTableOptions());
 
 	/** Updates display string entries and adds new display string entries based on provided localizations. */
-	CORE_API void UpdateFromLocalizations(FTextLocalizationResource&& TextLocalizationResource, const bool bDirtyTextRevision = true);
+	CORE_API void UpdateFromLocalizations(FTextLocalizationResource&& TextLocalizationResource, const FUpdateLiveTableOptions& UpdateOptions = FUpdateLiveTableOptions());
 
 	/** Dirties the local revision counter for the given text ID by incrementing it (or adding it) */
 	CORE_API void DirtyLocalRevisionForTextId(const FTextId& InTextId);
