Proactive fix to potential write failure on iOS

diff --git a/Engine/Source/Runtime/Core/Private/IOS/IOSPlatformFile.cpp b/Engine/Source/Runtime/Core/Private/IOS/IOSPlatformFile.cpp
--- a/Engine/Source/Runtime/Core/Private/IOS/IOSPlatformFile.cpp
+++ b/Engine/Source/Runtime/Core/Private/IOS/IOSPlatformFile.cpp
@@ -279,19 +279,26 @@ public:
         }
 	}
 
-	virtual bool Write( const uint8* Source, int64 BytesToWrite ) override
+	virtual bool Write(const uint8* Source, int64 BytesToWrite) override
 	{
-		while (BytesToWrite)
+		while (BytesToWrite > 0)
 		{
-			check(BytesToWrite >= 0);
-			int64 ThisSize = FMath::Min<int64>(READWRITE_SIZE, BytesToWrite);
-			check(Source);
-			if (write(FileHandle, Source, ThisSize) != ThisSize)
+			const int64 ThisSize = FMath::Min<int64>(READWRITE_SIZE, BytesToWrite);
+			const int64 Written = write(FileHandle, Source, ThisSize);
+			if (Written < 0)
 			{
-				return false;
+				if (errno == EINTR)
+				{
+					continue;
+				}
+				else
+				{
+					return false;
+				}
 			}
-			Source += ThisSize;
-			BytesToWrite -= ThisSize;
+			check(Written <= ThisSize);
+			Source += Written;
+			BytesToWrite -= Written;
 		}
 		return true;
 	}
