Fixed potential recursive OOM or memory corruption crash due to getting memory stats. Changed FGenericMemoryStats structure to use inline allocators to avoid allocations in most/all cases.

diff --git a/Engine/Source/Editor/UnrealEd/Private/CookOnTheFlyServer.cpp b/Engine/Source/Editor/UnrealEd/Private/CookOnTheFlyServer.cpp
--- a/Engine/Source/Editor/UnrealEd/Private/CookOnTheFlyServer.cpp
+++ b/Engine/Source/Editor/UnrealEd/Private/CookOnTheFlyServer.cpp
@@ -5714,7 +5714,7 @@ void UCookOnTheFlyServer::EvaluateGarbageCollectionResults(bool bWasDueToOOM, bo
 		auto AllocatorStatsToString = [](const FGenericMemoryStats& AllocatorStats)
 		{
 			TStringBuilder<256> Writer;
-			for (const TPair<FString, SIZE_T>& Item : AllocatorStats.Data)
+			for (const TPair<FStringView, SIZE_T>& Item : AllocatorStats)
 			{
 				Writer << TEXT("\n\t\tItem ") << Item.Key << TEXT(" ") << (uint64)Item.Value;
 			}
diff --git a/Engine/Source/Runtime/Core/Public/HAL/MemoryMisc.h b/Engine/Source/Runtime/Core/Public/HAL/MemoryMisc.h
--- a/Engine/Source/Runtime/Core/Public/HAL/MemoryMisc.h
+++ b/Engine/Source/Runtime/Core/Public/HAL/MemoryMisc.h
@@ -5,16 +5,81 @@
 #include "CoreTypes.h"
 #include "Containers/UnrealString.h"
 #include "Containers/Map.h"
+#include "Containers/StringView.h"
+#include "Misc/StringBuilder.h"
+
+#ifndef UE_MEMORY_STAT_DESCRIPTION_LENGTH_DEFAULT
+#define UE_MEMORY_STAT_DESCRIPTION_LENGTH 64
+#endif
+
+#ifndef UE_MEMORY_STAT_PREALLOCATION_COUNT
+#define UE_MEMORY_STAT_PREALLOCATION_COUNT 32
+#endif
 
 /** Holds generic memory stats, internally implemented as a map. */
 struct FGenericMemoryStats
 {
-	void Add( const TCHAR* StatDescription, const SIZE_T StatValue )
+private:
+	PRAGMA_DISABLE_DEPRECATION_WARNINGS
+	static constexpr int32 EntryCount = UE_MEMORY_STAT_PREALLOCATION_COUNT;
+	static constexpr int32 StringLength = UE_MEMORY_STAT_DESCRIPTION_LENGTH;
+	using FInternalMap = TMap<FStringView, SIZE_T, TInlineSetAllocator<EntryCount>>;
+
+public:
+	void Add(const FStringView& InDescription, SIZE_T InValue)
+	{
+		const FStringView DescriptionView = Descriptions.Emplace_GetRef(InPlace, InDescription).ToView();
+		Data.FInternalMap::Add(DescriptionView, InValue);
+	}
+
+	SIZE_T* Find(const FStringView& InDescription)
+	{
+		return Data.FInternalMap::Find(InDescription);
+	}
+
+	const SIZE_T* Find(const FStringView& InDescription) const
+	{
+		return const_cast<FGenericMemoryStats*>(this)->Find(InDescription);
+	}
+
+	SIZE_T FindRef(const FStringView& InDescription) const
 	{
-		Data.Add( FString(StatDescription), StatValue );
+		return Data.FInternalMap::FindRef(InDescription);
 	}
 
-	TMap<FString, SIZE_T> Data;
+	// Expose iterators only for const ranged-for iteration
+	FInternalMap::TRangedForConstIterator begin() const { return Data.begin(); }
+	FInternalMap::TRangedForConstIterator end() const { return Data.end(); }
+
+private:
+	TArray<TStringBuilder<StringLength>, TInlineAllocator<EntryCount>> Descriptions;
+
+	/** Wrapper on a TMap<TStringView> to allow passing both ANSI and TCHAR strings, for deprecation phase */
+	class FGenericMemoryStatsMap : public FInternalMap
+	{
+	public:
+		SIZE_T* Find(const FString& InDescription)
+		{
+			const FStringView DescriptionView(InDescription);
+			return FInternalMap::Find(DescriptionView);
+		}
+
+		const SIZE_T* Find(const FString& InDescription) const
+		{
+			return const_cast<FGenericMemoryStatsMap*>(this)->Find(InDescription);
+		}
+
+		SIZE_T FindRef(const FString& InDescription) const
+		{
+			const FStringView DescriptionView(InDescription);
+			return FInternalMap::FindRef(DescriptionView);
+		}
+	};
+	PRAGMA_ENABLE_DEPRECATION_WARNINGS
+
+public:
+	UE_DEPRECATED(5.5, "Using Data directly is deprecated and will be made private. Use methods on FGenericMemoryStats directly.")
+	FGenericMemoryStatsMap Data;
 };
 
 #ifndef ENABLE_MEMORY_SCOPE_STATS
diff --git a/Engine/Source/Runtime/Engine/Private/ProfilingDebugging/HealthSnapshot.cpp b/Engine/Source/Runtime/Engine/Private/ProfilingDebugging/HealthSnapshot.cpp
--- a/Engine/Source/Runtime/Engine/Private/ProfilingDebugging/HealthSnapshot.cpp
+++ b/Engine/Source/Runtime/Engine/Private/ProfilingDebugging/HealthSnapshot.cpp
@@ -75,8 +75,8 @@ void FHealthSnapshot::CaptureMemoryStats()
 		SIZE_T Allocated = 0;
 
 		GMalloc->GetAllocatorStats(MallocStats);
-		SIZE_T *MallocTotalAllocated = nullptr;
-		if ((MallocTotalAllocated = MallocStats.Data.Find("TotalAllocated")) != nullptr)
+		const SIZE_T *MallocTotalAllocated = nullptr;
+		if ((MallocTotalAllocated = MallocStats.Find(TEXT("TotalAllocated"))) != nullptr)
 		{
 			Allocated = *MallocTotalAllocated;
 		}
