Reduce number of max in-flight tasks allowed in a worker thread queue when AGGRESSIVE_MEMORY_SAVING is enabled. Saves ~320KB per worker thread

#rnx

diff --git a/Engine/Source/Runtime/Core/Public/Async/Fundamental/LocalQueue.h b/Engine/Source/Runtime/Core/Public/Async/Fundamental/LocalQueue.h
--- a/Engine/Source/Runtime/Core/Public/Async/Fundamental/LocalQueue.h
+++ b/Engine/Source/Runtime/Core/Public/Async/Fundamental/LocalQueue.h
@@ -9,6 +9,14 @@
 
 #include <atomic>
 
+#if AGGRESSIVE_MEMORY_SAVING
+	#define LOCALQUEUEREGISTRYDEFAULTS_MAX_LOCALQUEUES 1024
+	#define LOCALQUEUEREGISTRYDEFAULTS_MAX_ITEMCOUNT 512
+#else
+	#define LOCALQUEUEREGISTRYDEFAULTS_MAX_LOCALQUEUES 1024
+	#define LOCALQUEUEREGISTRYDEFAULTS_MAX_ITEMCOUNT 1024
+#endif
+
 namespace LowLevelTasks
 {
 namespace LocalQueue_Impl
@@ -144,7 +152,7 @@ enum class ELocalQueueType
  * A Dequeue Operation can only be done starting from a LocalQueue, than the GlobalQueue will be checked.                                   *
  * Finally Items might get Stolen from other LocalQueues that are registered with the LocalQueueRegistry.                                   *
  ********************************************************************************************************************************************/
-template<uint32 NumLocalItems = 1024, uint32 MaxLocalQueues = 1024>
+template<uint32 NumLocalItems = LOCALQUEUEREGISTRYDEFAULTS_MAX_ITEMCOUNT, uint32 MaxLocalQueues = LOCALQUEUEREGISTRYDEFAULTS_MAX_LOCALQUEUES>
 class TLocalQueueRegistry
 {
 	static uint32 Rand()
@@ -266,7 +274,7 @@ private:
 	void AddLocalQueue(TLocalQueue* QueueToAdd)
 	{
 		uint32 Index = NumLocalQueues.fetch_add(1, std::memory_order_relaxed);
-		check(Index < MaxLocalQueues);
+		UE_CLOG(Index >= MaxLocalQueues, LowLevelTasks, Fatal, TEXT("Attempting to add more than the maximum allowed number of queues (%d)"), MaxLocalQueues);
 
 		// std::memory_order_release to make sure values are all written to the TLocalQueue before publishing.
 		LocalQueues[Index].store(QueueToAdd, std::memory_order_release);
