[Core]
* Fixed header units compile warnings

diff --git a/Engine/Source/Runtime/Core/Public/Misc/AutomationTest.h b/Engine/Source/Runtime/Core/Public/Misc/AutomationTest.h
--- a/Engine/Source/Runtime/Core/Public/Misc/AutomationTest.h
+++ b/Engine/Source/Runtime/Core/Public/Misc/AutomationTest.h
@@ -132,12 +132,12 @@ enum class EAutomationTestFlags
 
 ENUM_CLASS_FLAGS(EAutomationTestFlags)
 
-constexpr EAutomationTestFlags EAutomationTestFlags_ApplicationContextMask = EAutomationTestFlags::EditorContext | EAutomationTestFlags::ClientContext | EAutomationTestFlags::ServerContext | EAutomationTestFlags::CommandletContext | EAutomationTestFlags::ProgramContext;
-constexpr EAutomationTestFlags EAutomationTestFlags_FeatureMask            = EAutomationTestFlags::NonNullRHI | EAutomationTestFlags::RequiresUser;
-constexpr EAutomationTestFlags EAutomationTestFlags_HighPriorityAndAbove   = EAutomationTestFlags::CriticalPriority | EAutomationTestFlags::HighPriority;
-constexpr EAutomationTestFlags EAutomationTestFlags_MediumPriorityAndAbove = EAutomationTestFlags::CriticalPriority | EAutomationTestFlags::HighPriority | EAutomationTestFlags::MediumPriority;
-constexpr EAutomationTestFlags EAutomationTestFlags_PriorityMask           = EAutomationTestFlags::CriticalPriority | EAutomationTestFlags::HighPriority | EAutomationTestFlags::MediumPriority | EAutomationTestFlags::LowPriority;
-constexpr EAutomationTestFlags EAutomationTestFlags_FilterMask             = EAutomationTestFlags::SmokeFilter | EAutomationTestFlags::EngineFilter | EAutomationTestFlags::ProductFilter | EAutomationTestFlags::PerfFilter | EAutomationTestFlags::StressFilter | EAutomationTestFlags::NegativeFilter;
+inline constexpr EAutomationTestFlags EAutomationTestFlags_ApplicationContextMask = EAutomationTestFlags::EditorContext | EAutomationTestFlags::ClientContext | EAutomationTestFlags::ServerContext | EAutomationTestFlags::CommandletContext | EAutomationTestFlags::ProgramContext;
+inline constexpr EAutomationTestFlags EAutomationTestFlags_FeatureMask            = EAutomationTestFlags::NonNullRHI | EAutomationTestFlags::RequiresUser;
+inline constexpr EAutomationTestFlags EAutomationTestFlags_HighPriorityAndAbove   = EAutomationTestFlags::CriticalPriority | EAutomationTestFlags::HighPriority;
+inline constexpr EAutomationTestFlags EAutomationTestFlags_MediumPriorityAndAbove = EAutomationTestFlags::CriticalPriority | EAutomationTestFlags::HighPriority | EAutomationTestFlags::MediumPriority;
+inline constexpr EAutomationTestFlags EAutomationTestFlags_PriorityMask           = EAutomationTestFlags::CriticalPriority | EAutomationTestFlags::HighPriority | EAutomationTestFlags::MediumPriority | EAutomationTestFlags::LowPriority;
+inline constexpr EAutomationTestFlags EAutomationTestFlags_FilterMask             = EAutomationTestFlags::SmokeFilter | EAutomationTestFlags::EngineFilter | EAutomationTestFlags::ProductFilter | EAutomationTestFlags::PerfFilter | EAutomationTestFlags::StressFilter | EAutomationTestFlags::NegativeFilter;
 
 CORE_API const TMap<FString, EAutomationTestFlags>& EAutomationTestFlags_GetTestFlagsMap();
 
