Fix mac file write edge cases

diff --git a/Engine/Source/Runtime/Core/Private/Apple/ApplePlatformFile.cpp b/Engine/Source/Runtime/Core/Private/Apple/ApplePlatformFile.cpp
--- a/Engine/Source/Runtime/Core/Private/Apple/ApplePlatformFile.cpp
+++ b/Engine/Source/Runtime/Core/Private/Apple/ApplePlatformFile.cpp
@@ -198,20 +198,26 @@ public:
 		check(IsValid());
 		TRACE_PLATFORMFILE_BEGIN_WRITE(this, FileHandle, 0, BytesToWrite);
 		int64 TotalBytesWritten = 0;
-		while (BytesToWrite)
+		while (BytesToWrite > 0)
 		{
-			check(BytesToWrite >= 0);
 			int64 ThisSize = FMath::Min<int64>(READWRITE_SIZE, BytesToWrite);
 			check(Source);
 			int64 BytesWritten = write(FileHandle, Source, ThisSize);
-			TotalBytesWritten += BytesWritten;
-			if (BytesWritten != ThisSize)
-			{
-				TRACE_PLATFORMFILE_END_WRITE(this, TotalBytesWritten);
-				return false;
-			}
-			Source += ThisSize;
-			BytesToWrite -= ThisSize;
+            if (BytesWritten < 0)
+            {
+                if (errno == EINTR)
+                {
+                    continue;
+                }
+                else
+                {
+                    TRACE_PLATFORMFILE_END_WRITE(this, TotalBytesWritten);
+                    return false;
+                }
+            }
+            TotalBytesWritten += BytesWritten;
+			Source += BytesWritten;
+			BytesToWrite -= BytesWritten;
 		}
 		TRACE_PLATFORMFILE_END_WRITE(this, TotalBytesWritten);
 		return true;
