Use a TSet for the internal FTextKey string index when using the slab allocator, as FKeyData::Str was always the same value as the TMap value in that config

#rnx

diff --git a/Engine/Source/Runtime/Core/Private/Internationalization/TextKey.cpp b/Engine/Source/Runtime/Core/Private/Internationalization/TextKey.cpp
--- a/Engine/Source/Runtime/Core/Private/Internationalization/TextKey.cpp
+++ b/Engine/Source/Runtime/Core/Private/Internationalization/TextKey.cpp
@@ -139,12 +139,16 @@ private:
 	{
 		FScopeLock ScopeLock(&SynchronizationObject);
 #if UE_TEXTKEY_USE_SLAB_ALLOCATOR
-		const TCHAR* StrPtr = KeysTable.FindRef(KeyData);
+		const TCHAR* StrPtr = nullptr;
+		if (const FKeyData* FoundKeyData = KeysTable.Find(KeyData))
+		{
+			StrPtr = FoundKeyData->Str;
+		}
 		if (!StrPtr)
 		{
 			LLM_SCOPE_BYNAME(TEXT("Localization/TextKeys"));
 			StrPtr = StringAllocations.Add(KeyData.ToView());
-			KeysTable.Add(FKeyData(StrPtr, KeyData), StrPtr);
+			KeysTable.Add(FKeyData(StrPtr, KeyData));
 		}
 		return StrPtr;
 #else
@@ -256,7 +260,7 @@ private:
 	FCriticalSection SynchronizationObject;
 #if UE_TEXTKEY_USE_SLAB_ALLOCATOR
 	FStringSlabAllocator StringAllocations;
-	TMap<FKeyData, const TCHAR*> KeysTable;
+	TSet<FKeyData> KeysTable;
 #else
 	TMap<FKeyData, FString> KeysTable;
 #endif
