Handle write returning 0 on iOS as an error

diff --git a/Engine/Source/Runtime/Core/Private/IOS/IOSPlatformFile.cpp b/Engine/Source/Runtime/Core/Private/IOS/IOSPlatformFile.cpp
--- a/Engine/Source/Runtime/Core/Private/IOS/IOSPlatformFile.cpp
+++ b/Engine/Source/Runtime/Core/Private/IOS/IOSPlatformFile.cpp
@@ -285,7 +285,7 @@ public:
 		{
 			const int64 ThisSize = FMath::Min<int64>(READWRITE_SIZE, BytesToWrite);
 			const int64 Written = write(FileHandle, Source, ThisSize);
-			if (Written < 0)
+			if (Written <= 0)
 			{
 				if (errno == EINTR)
 				{
