Handle write returning 0 on mac as an error

diff --git a/Engine/Source/Runtime/Core/Private/Apple/ApplePlatformFile.cpp b/Engine/Source/Runtime/Core/Private/Apple/ApplePlatformFile.cpp
--- a/Engine/Source/Runtime/Core/Private/Apple/ApplePlatformFile.cpp
+++ b/Engine/Source/Runtime/Core/Private/Apple/ApplePlatformFile.cpp
@@ -203,7 +203,7 @@ public:
 			int64 ThisSize = FMath::Min<int64>(READWRITE_SIZE, BytesToWrite);
 			check(Source);
 			int64 BytesWritten = write(FileHandle, Source, ThisSize);
-            if (BytesWritten < 0)
+            if (BytesWritten <= 0)
             {
                 if (errno == EINTR)
                 {
