UnrealMathNEON various tweaks.

Also contains one small additional cleanup for UnrealMathSSE.

diff --git a/Engine/Source/Runtime/Core/Public/Math/UnrealMathNeon.h b/Engine/Source/Runtime/Core/Public/Math/UnrealMathNeon.h
--- a/Engine/Source/Runtime/Core/Public/Math/UnrealMathNeon.h
+++ b/Engine/Source/Runtime/Core/Public/Math/UnrealMathNeon.h
@@ -2,10 +2,14 @@
 
 #pragma once
 
+#include "HAL/Platform.h"
+
 // HEADER_UNIT_SKIP - Not included directly
 
 PRAGMA_DISABLE_SHADOW_VARIABLE_WARNINGS
 
+#if PLATFORM_ENABLE_VECTORINTRINSICS_NEON
+
 // Include the intrinsic functions header
 #if ((PLATFORM_WINDOWS || PLATFORM_HOLOLENS) && PLATFORM_64BITS)
 #include <arm64_neon.h>
@@ -503,16 +507,10 @@ FORCEINLINE VectorRegister4Float VectorLoadFloat2(const float* Ptr)
  */
 FORCEINLINE VectorRegister4Double VectorLoadFloat3(const double* Ptr)
 {
-	union U
-	{
-		VectorRegister4Double V; double D[4];
-		inline U() : V() {}
-	} Tmp;
-
-	Tmp.V.XY = vld1q_f64(Ptr);
-	Tmp.D[2] = Ptr[2];
-	Tmp.D[3] = 0.0;
-	return Tmp.V;
+	VectorRegister4Double Result;
+	Result.XY = vld1q_f64(Ptr);
+	Result.ZW = vcombine_f64(vld1_f64(&Ptr[2]), vdup_n_f64(0.0));
+	return Result;
 }
 
 /**
@@ -523,7 +521,10 @@ FORCEINLINE VectorRegister4Double VectorLoadFloat3(const double* Ptr)
  */
 FORCEINLINE VectorRegister4Double VectorLoadFloat3_W1(const double* Ptr)
 {
-	return MakeVectorRegisterDouble(Ptr[0], Ptr[1], Ptr[2], 1.0f);
+	VectorRegister4Double Result;
+	Result.XY = vld1q_f64(Ptr);
+	Result.ZW = vcombine_f64(vld1_f64(&Ptr[2]), vdup_n_f64(1.0));
+	return Result;
 }
 
 /**
@@ -1477,20 +1478,18 @@ FORCEINLINE VectorRegister4Double VectorShuffleImpl(VectorRegister4Double Vec1,
  */
 FORCEINLINE uint32 VectorMaskBits(VectorRegister4Float VecMask)
 {
-	uint32x4_t mmA = vtstq_u32(vreinterpretq_u32_f32(VecMask), GlobalVectorConstants::SignBit()); // mask with 1s every bit for vector element if it's sign is negative
-	uint32x4_t mmB = vandq_u32(mmA, MakeVectorRegisterInt(0x1, 0x2, 0x4, 0x8)); // pick only one bit on it's corresponding position
-	uint32x2_t mmC = vorr_u32(vget_low_u32(mmB), vget_high_u32(mmB));           // now combine the result
-	return vget_lane_u32(mmC, 0) | vget_lane_u32(mmC, 1);                       // reduce the result from 2 elements to one
+	int32x4_t Signs = vshrq_n_s32(vreinterpretq_s32_f32(VecMask), 31); // sign bit of each lane replicated 32x
+	int32x4_t Masked = vandq_s32(Signs, MakeVectorRegisterInt(0x1, 0x2, 0x4, 0x8)); // pick bit for lane position
+	return uint32(vaddvq_s32(Masked)); // reduce via add
 }
 
 FORCEINLINE uint32 VectorMaskBits(VectorRegister4Double VecMask)
 {
-	uint64x2_t mmA = vtstq_u64(vreinterpretq_u64_f64(VecMask.XY), GlobalVectorConstants::DoubleSignBit().XY); // mask with 1s every bit for vector element if it's sign is negative
-	uint64x2_t mmA1 = vtstq_u64(vreinterpretq_u64_f64(VecMask.ZW), GlobalVectorConstants::DoubleSignBit().XY);
-	uint64x2_t mmB = vandq_u64(mmA, MakeVectorRegisterInt64(0x1, 0x2)); // pick only one bit on it's corresponding position
-	uint64x2_t mmB1 = vandq_u64(mmA1, MakeVectorRegisterInt64(0x4, 0x8));
-	uint64x2_t mmC = vorrq_u64(mmB, mmB1);								// now combine the result
-	return (uint32)(vgetq_lane_u64(mmC, 0) | vgetq_lane_u64(mmC, 1));     // reduce the result from 2 elements to one
+	int64x2_t Signs0 = vshrq_n_s64(vreinterpretq_s64_f32(VecMask.XY), 63); // sign bit of each lane replicated 64x
+	int64x2_t Signs1 = vshrq_n_s64(vreinterpretq_s64_f32(VecMask.ZW), 63); // sign bit of each lane replicated 64x
+	int32x4_t Signs = vuzp1q_s32(Signs0, Signs1); // 32-bit masks
+	int32x4_t Masked = vandq_s32(Signs, MakeVectorRegisterInt(0x1, 0x2, 0x4, 0x8)); // pick bit for lane position
+	return uint32(vaddvq_s32(Masked)); // reduce via add
 }
 
 /**
@@ -1707,9 +1706,6 @@ FORCEINLINE VectorRegister4Double VectorReciprocalSqrtEstimate(const VectorRegis
  */
 FORCEINLINE VectorRegister4Float VectorReciprocalSqrt(const VectorRegister4Float& Vec)
 {
-	// Perform a single pass of Newton-Raphson iteration on the hardware estimate
-	// This is a builtin instruction (VRSQRTS)
-
 	// Initial estimate
 	VectorRegister4Float RecipSqrt = VectorReciprocalSqrtEstimate(Vec);
 
@@ -1720,9 +1716,6 @@ FORCEINLINE VectorRegister4Float VectorReciprocalSqrt(const VectorRegister4Float
 
 FORCEINLINE VectorRegister4Double VectorReciprocalSqrt(const VectorRegister4Double& Vec)
 {
-	// Perform a single pass of Newton-Raphson iteration on the hardware estimate
-	// This is a builtin instruction (VRSQRTS)
-
 	// Initial estimate
 	VectorRegister4Double RecipSqrt = VectorReciprocalSqrtEstimate(Vec);
 
@@ -2089,78 +2082,72 @@ FORCEINLINE VectorRegister4Double VectorMergeVecXYZ_VecW(const VectorRegister4Do
 
 /**
  * Loads 4 uint8s from unaligned memory and converts them into 4 floats.
- * IMPORTANT: You need to call VectorResetFloatRegisters() before using scalar floats after you've used this intrinsic!
  *
  * @param Ptr			Unaligned memory pointer to the 4 uint8s.
  * @return				VectorRegister4Float( float(Ptr[0]), float(Ptr[1]), float(Ptr[2]), float(Ptr[3]) )
  */
 FORCEINLINE VectorRegister4Float VectorLoadByte4( const void* Ptr )
 {
-	// OPTIMIZE ME!
-	const uint8 *P = (const uint8 *)Ptr;
-	return MakeVectorRegister( (float)P[0], (float)P[1], (float)P[2], (float)P[3] );
+	uint8x8_t AsUInt8 = vreinterpret_u8_u32(vld1_dup_u32((const uint32*)Ptr));
+	uint16x8_t AsUInt16 = vmovl_u8(AsUInt8);
+	uint32x4_t AsUInt32 = vmovl_u16(vget_low_u16(AsUInt16));
+	return vcvtq_f32_u32(AsUInt32);
 }
 
 /**
 * Loads 4 int8s from unaligned memory and converts them into 4 floats.
-* IMPORTANT: You need to call VectorResetFloatRegisters() before using scalar floats after you've used this intrinsic!
 *
 * @param Ptr			Unaligned memory pointer to the 4 uint8s.
 * @return				VectorRegister4Float( float(Ptr[0]), float(Ptr[1]), float(Ptr[2]), float(Ptr[3]) )
 */
 FORCEINLINE VectorRegister4Float VectorLoadSignedByte4(const void* Ptr)
 {
-	// OPTIMIZE ME!
-	const int8 *P = (const int8 *)Ptr;
-	return MakeVectorRegister((float)P[0], (float)P[1], (float)P[2], (float)P[3]);
+	int8x8_t AsInt8 = vreinterpret_s8_u32(vld1_dup_u32((const uint32*)Ptr));
+	int16x8_t AsInt16 = vmovl_s8(AsInt8);
+	int32x4_t AsInt32 = vmovl_s16(vget_low_u16(AsInt16));
+	return vcvtq_f32_s32(AsInt32);
 }
 
 /**
  * Loads 4 uint8s from unaligned memory and converts them into 4 floats in reversed order.
- * IMPORTANT: You need to call VectorResetFloatRegisters() before using scalar floats after you've used this intrinsic!
  *
  * @param Ptr			Unaligned memory pointer to the 4 uint8s.
  * @return				VectorRegister4Float( float(Ptr[3]), float(Ptr[2]), float(Ptr[1]), float(Ptr[0]) )
  */
 FORCEINLINE VectorRegister4Float VectorLoadByte4Reverse( const uint8* Ptr )
 {
-	// OPTIMIZE ME!
-	const uint8 *P = (const uint8 *)Ptr;
-	return MakeVectorRegister( (float)P[3], (float)P[2], (float)P[1], (float)P[0] );
+	uint8x8_t AsUInt8 = vrev32_u8(vreinterpret_u8_u32(vld1_dup_u32((const uint32*)Ptr)));
+	uint16x8_t AsUInt16 = vmovl_u8(AsUInt8);
+	uint32x4_t AsUInt32 = vmovl_u16(vget_low_u16(AsUInt16));
+	return vcvtq_f32_u32(AsUInt32);
 }
 
 /**
  * Converts the 4 floats in the vector to 4 uint8s, clamped to [0,255], and stores to unaligned memory.
- * IMPORTANT: You need to call VectorResetFloatRegisters() before using scalar floats after you've used this intrinsic!
  *
  * @param Vec			Vector containing 4 floats
  * @param Ptr			Unaligned memory pointer to store the 4 uint8s.
  */
 FORCEINLINE void VectorStoreByte4( VectorRegister4Float Vec, void* Ptr )
 {
-	uint16x8_t u16x8 = (uint16x8_t)vcvtq_u32_f32(VectorMin(Vec, GlobalVectorConstants::Float255));
-	uint8x8_t u8x8 = (uint8x8_t)vget_low_u16( vuzpq_u16( u16x8, u16x8 ).val[0] );
-	u8x8 = vuzp_u8( u8x8, u8x8 ).val[0];
-	uint32_t buf[2];
-	vst1_u8( (uint8_t *)buf, u8x8 );
-	*(uint32_t *)Ptr = buf[0]; 
+	uint32x4_t AsUInt32 = vcvtq_u32_f32(Vec); // Saturates (clamps) to [0,2^32 - 1]
+	uint16x4_t AsUInt16 = vqmovn_u32(AsUInt32); // Saturates further to [0,2^16 - 1]
+	uint8x8_t AsUInt8 = vqmovn_u16(vcombine_u16(AsUInt16, vdup_n_u16(0))); // Saturates to [0,255]
+	vst1_lane_u32((uint32_t*)Ptr, AsUInt8, 0);
 }
 
 /**
-* Converts the 4 floats in the vector to 4 int8s, clamped to [-127, 127], and stores to unaligned memory.
-* IMPORTANT: You need to call VectorResetFloatRegisters() before using scalar floats after you've used this intrinsic!
+* Converts the 4 floats in the vector to 4 int8s, clamped to [-128, 127], and stores to unaligned memory.
 *
 * @param Vec			Vector containing 4 floats
 * @param Ptr			Unaligned memory pointer to store the 4 uint8s.
 */
 FORCEINLINE void VectorStoreSignedByte4(VectorRegister4Float Vec, void* Ptr)
 {
-	int16x8_t s16x8 = (int16x8_t)vcvtq_s32_f32(VectorMax(VectorMin(Vec, GlobalVectorConstants::Float127), GlobalVectorConstants::FloatNeg127));
-	int8x8_t s8x8 = (int8x8_t)vget_low_s16(vuzpq_s16(s16x8, s16x8).val[0]);
-	s8x8 = vuzp_s8(s8x8, s8x8).val[0];
-	int32_t buf[2];
-	vst1_s8((int8_t *)buf, s8x8);
-	*(int32_t *)Ptr = buf[0];
+	int32x4_t AsInt32 = vcvtq_s32_f32(Vec); // Saturates (clamps) to [-2^31,2^31 - 1]
+	int16x4_t AsInt16 = vqmovn_s32(AsInt32); // Saturates further to [-32768,32767]
+	int8x8_t AsInt8 = vqmovn_s16(vcombine_s16(AsInt16, vdup_n_s16(0))); // Saturates to [-128,127]
+	vst1_lane_u32((uint32_t*)Ptr, AsInt8, 0);
 }
 
 /**
@@ -2173,25 +2160,11 @@ template <bool bAligned>
 FORCEINLINE void VectorStoreHalf4(VectorRegister4Float Vec, void* RESTRICT Ptr)
 {
 	float16x4_t f16x4 = vcvt_f16_f32(Vec);
-
-	if (bAligned)
-	{
-		vst1_u8( (uint8_t *)Ptr, f16x4 );
-	}
-	else
-	{
-		alignas(16) uint16_t Buf[4];
-		vst1_u8( (uint8_t *)Buf, f16x4 );
-		for (int i = 0; i < 4; ++i)
-		{
-			((uint16_t*)Ptr)[i] = Buf[i];
-		}
-	}
+	vst1_u8((uint8_t*)Ptr, f16x4);
 }
 
 /**
 * Loads packed RGB10A2(4 bytes) from unaligned memory and converts them into 4 FLOATs.
-* IMPORTANT: You need to call VectorResetFloatRegisters() before using scalar FLOATs after you've used this intrinsic!
 *
 * @param Ptr			Unaligned memory pointer to the RGB10A2(4 bytes).
 * @return				VectorRegister4Float with 4 FLOATs loaded from Ptr.
@@ -2211,7 +2184,6 @@ FORCEINLINE VectorRegister4Float VectorLoadURGB10A2N(void* Ptr)
 
 /**
 * Converts the 4 FLOATs in the vector RGB10A2, clamped to [0, 1023] and [0, 3], and stores to unaligned memory.
-* IMPORTANT: You need to call VectorResetFloatRegisters() before using scalar FLOATs after you've used this intrinsic!
 *
 * @param Vec			Vector containing 4 FLOATs
 * @param Ptr			Unaligned memory pointer to store the packed RGB10A2(4 bytes).
@@ -2250,12 +2222,11 @@ FORCEINLINE int32 VectorAnyGreaterThan(VectorRegister4Double Vec1, VectorRegiste
 {
 	uint32x4_t MaskXY = (uint32x4_t)vcgtq_f64(Vec1.XY, Vec2.XY);
 	uint32x4_t MaskZW = (uint32x4_t)vcgtq_f64(Vec1.ZW, Vec2.ZW);
-	return vmaxvq_u32(MaskXY) || vmaxvq_u32(MaskZW);
+	return vmaxvq_u32(vorrq_u32(MaskXY, MaskZW));
 }
 
 /**
  * Resets the floating point registers so that they can be used again.
- * Some intrinsics use these for MMX purposes (e.g. VectorLoadByte4 and VectorStoreByte4).
  */
 #define VectorResetFloatRegisters()
 
@@ -2747,46 +2718,33 @@ FORCEINLINE void VectorSinCos(VectorRegister4Double* RESTRICT VSinAngles, Vector
 }
 
 /**
-* Loads packed RGBA16(4 bytes) from unaligned memory and converts them into 4 FLOATs.
-* IMPORTANT: You need to call VectorResetFloatRegisters() before using scalar FLOATs after you've used this intrinsic!
+* Loads packed RGBA16(8 bytes) from unaligned memory and converts them into 4 FLOATs.
 *
 * @param Ptr			Unaligned memory pointer to the RGBA16(8 bytes).
 * @return				VectorRegister4Float with 4 FLOATs loaded from Ptr.
 */
 FORCEINLINE VectorRegister4Float VectorLoadURGBA16N(const uint16* E)
 {
-	alignas(16) float V[4];
-	V[0] = float(E[0]);
-	V[1] = float(E[1]);
-	V[2] = float(E[2]);
-	V[3] = float(E[3]);
-
-	return VectorLoad(V);
+	uint16x4_t UInt16s = vld1_u16(E);
+	uint32x4_t UInt32s = vmovl_u16(UInt16s);
+	return vcvtq_f32_u32(UInt32s);
 }
 
 /**
-* Loads packed signed RGBA16(4 bytes) from unaligned memory and converts them into 4 FLOATs.
-* IMPORTANT: You need to call VectorResetFloatRegisters() before using scalar FLOATs after you've used this intrinsic!
+* Loads packed signed RGBA16(8 bytes) from unaligned memory and converts them into 4 FLOATs.
 *
 * @param Ptr			Unaligned memory pointer to the RGBA16(8 bytes).
 * @return				VectorRegister4Float with 4 FLOATs loaded from Ptr.
 */
 FORCEINLINE VectorRegister4Float VectorLoadSRGBA16N(const void* Ptr)
 {
-	alignas(16) float V[4];
-	int16* E = (int16*)Ptr;
-
-	V[0] = float(E[0]);
-	V[1] = float(E[1]);
-	V[2] = float(E[2]);
-	V[3] = float(E[3]);
-
-	return VectorLoad(V);
+	int16x4_t Int16s = vld1_s16((const int16 *)Ptr);
+	int32x4_t Int32s = vmovl_s16(Int16s);
+	return vcvtq_f32_s32(Int32s);
 }
 
 /**
 * Converts the 4 FLOATs in the vector RGBA16, clamped to [0, 65535], and stores to unaligned memory.
-* IMPORTANT: You need to call VectorResetFloatRegisters() before using scalar FLOATs after you've used this intrinsic!
 *
 * @param Vec			Vector containing 4 FLOATs
 * @param Ptr			Unaligned memory pointer to store the packed RGBA16(8 bytes).
@@ -2796,16 +2754,10 @@ FORCEINLINE void VectorStoreURGBA16N(const VectorRegister4Float& Vec, uint16* Ou
 	VectorRegister4Float Tmp;
 	Tmp = VectorMax(Vec, VectorZeroFloat());
 	Tmp = VectorMin(Tmp, VectorOneFloat());
-	Tmp = VectorMultiplyAdd(Tmp, vdupq_n_f32(65535.0f), vdupq_n_f32(0.5f));
-	Tmp = VectorTruncate(Tmp);
-
-	alignas(16) float F[4];
-	VectorStoreAligned(Tmp, F);
+	Tmp = VectorMultiply(Tmp, vdupq_n_f32(65535.0f));
 
-	Out[0] = (uint16)F[0];
-	Out[1] = (uint16)F[1];
-	Out[2] = (uint16)F[2];
-	Out[3] = (uint16)F[3];
+	uint32x4_t TmpUInt = vcvtnq_u32_f32(Tmp);
+	vst1_u16(Out, vmovn_u32(TmpUInt));
 }
 
 //////////////////////////////////////////////////////////////////////////
@@ -2819,7 +2771,7 @@ FORCEINLINE void VectorStoreURGBA16N(const VectorRegister4Float& Vec, uint16* Ou
 /** = a ^ b */
 #define VectorIntXor(A, B)		veorq_s32(A, B)
 /** = (~a) & b to match _mm_andnot_si128 */
-#define VectorIntAndNot(A, B)	vandq_s32(vmvnq_s32(A), B)
+#define VectorIntAndNot(A, B)	vbicq_s32(B, A)
 /** = ~a */
 #define VectorIntNot(A)	vmvnq_s32(A)
 
@@ -2834,7 +2786,7 @@ FORCEINLINE void VectorStoreURGBA16N(const VectorRegister4Float& Vec, uint16* Ou
 
 FORCEINLINE VectorRegister4Int VectorIntSelect(const VectorRegister4Int& Mask, const VectorRegister4Int& Vec1, const VectorRegister4Int& Vec2)
 {
-	return VectorIntXor(Vec2, VectorIntAnd(Mask, VectorIntXor(Vec1, Vec2)));
+	return vbslq_s32(Mask, Vec1, Vec2);
 }
 
 //Arithmetic
@@ -2845,7 +2797,7 @@ FORCEINLINE VectorRegister4Int VectorIntSelect(const VectorRegister4Int& Mask, c
 #define VectorIntMin(A, B) vminq_s32(A,B)
 #define VectorIntMax(A, B) vmaxq_s32(A,B)
 #define VectorIntClamp(A, B, C) VectorIntMin(VectorIntMax(A, B), C)
-#define VectorIntAbs(A) vabdq_s32(A, GlobalVectorConstants::IntZero)
+#define VectorIntAbs(A) vabsq_s32(A)
 
 #define VectorIntSign(A) VectorIntSelect( VectorIntCompareGE(A, GlobalVectorConstants::IntZero), GlobalVectorConstants::IntOne, GlobalVectorConstants::IntMinusOne )
 
@@ -2937,12 +2889,13 @@ FORCEINLINE VectorRegister4Int VectorRoundToIntHalfToEven(const VectorRegister4F
 }
 
 FORCEINLINE VectorRegister4i VectorIntExpandLow16To32(VectorRegister4i V) {
-	int16x4x2_t res = vzip_s16(vget_low_u16(V), vdup_n_u16(0));
-	return vcombine_s16(res.val[0], res.val[1]);
+	return vmovl_u16(vget_low_u16(V));
 }
 
 // To be continued...
 
+#endif // #if PLATFORM_ENABLE_VECTORINTRINSICS_NEON
+
 PRAGMA_ENABLE_SHADOW_VARIABLE_WARNINGS
 
 #if UE_ENABLE_INCLUDE_ORDER_DEPRECATED_IN_5_4
diff --git a/Engine/Source/Runtime/Core/Public/Math/UnrealMathSSE.h b/Engine/Source/Runtime/Core/Public/Math/UnrealMathSSE.h
--- a/Engine/Source/Runtime/Core/Public/Math/UnrealMathSSE.h
+++ b/Engine/Source/Runtime/Core/Public/Math/UnrealMathSSE.h
@@ -2173,6 +2173,23 @@ FORCEINLINE VectorRegister4Double VectorReciprocalEstimate(const VectorRegister4
 	return VectorReciprocal(Vec);
 }
 
+/**
+ * Merges the XYZ components of one vector with the W component of another vector and returns the result.
+ *
+ * @param VecXYZ	Source vector for XYZ_
+ * @param VecW		Source register for ___W (note: the fourth component is used, not the first)
+ * @return			VectorRegister4Float(VecXYZ.x, VecXYZ.y, VecXYZ.z, VecW.w)
+ */
+FORCEINLINE VectorRegister4Float VectorMergeVecXYZ_VecW(const VectorRegister4Float& VecXYZ, const VectorRegister4Float& VecW)
+{
+	return _mm_blend_ps(VecXYZ, VecW, 0b1000);
+}
+
+FORCEINLINE VectorRegister4Double VectorMergeVecXYZ_VecW(const VectorRegister4Double& VecXYZ, const VectorRegister4Double& VecW)
+{
+	return VectorRegister4Double(VecXYZ.XY, _mm_move_sd(VecW.ZW, VecXYZ.ZW));
+}
+
 /**
 * Loads XYZ and sets W=0
 *
@@ -2204,7 +2221,7 @@ FORCEINLINE VectorRegister4Double VectorSet_W0(const VectorRegister4Double& Vec)
 */
 FORCEINLINE VectorRegister4Float VectorSet_W1( const VectorRegister4Float& Vec)
 {
-	return _mm_blend_ps(Vec, VectorOneFloat(), 0b1000);
+	return VectorMergeVecXYZ_VecW(Vec, VectorOneFloat());
 }
 
 FORCEINLINE VectorRegister4Double VectorSet_W1(const VectorRegister4Double& Vec)
@@ -2475,23 +2492,6 @@ FORCEINLINE int VectorMaskBits(const VectorRegister4Double& VecMask)
 #endif
 }
 
-/**
- * Merges the XYZ components of one vector with the W component of another vector and returns the result.
- *
- * @param VecXYZ	Source vector for XYZ_
- * @param VecW		Source register for ___W (note: the fourth component is used, not the first)
- * @return			VectorRegister4Float(VecXYZ.x, VecXYZ.y, VecXYZ.z, VecW.w)
- */
-FORCEINLINE VectorRegister4Float VectorMergeVecXYZ_VecW(const VectorRegister4Float& VecXYZ, const VectorRegister4Float& VecW)
-{
-	return _mm_blend_ps(VecXYZ, VecW, 0b1000);
-}
-
-FORCEINLINE VectorRegister4Double VectorMergeVecXYZ_VecW(const VectorRegister4Double& VecXYZ, const VectorRegister4Double& VecW)
-{
-	return VectorRegister4Double(VecXYZ.XY, _mm_move_sd(VecW.ZW, VecXYZ.ZW));
-}
-
 /**
  * Loads 4 BYTEs from unaligned memory and converts them into 4 FLOATs.
  *
