Fix crash after call to FAutomationSpecBase::Redefine().

The fix is to reinitialize the RootDefinitionScope and DefinitionScopeStack, matching the logic from the constructor.

diff --git a/Engine/Source/Runtime/Core/Public/Misc/AutomationTest.h b/Engine/Source/Runtime/Core/Public/Misc/AutomationTest.h
--- a/Engine/Source/Runtime/Core/Public/Misc/AutomationTest.h
+++ b/Engine/Source/Runtime/Core/Public/Misc/AutomationTest.h
@@ -3495,6 +3495,9 @@ protected:
 		RootDefinitionScope.Reset();
 		DefinitionScopeStack.Empty();
 		bHasBeenDefined = false;
+
+		RootDefinitionScope = MakeShareable(new FSpecDefinitionScope());
+		DefinitionScopeStack.Push(RootDefinitionScope.ToSharedRef());
 	}
 
 private:
