Additional clean-up of FTextLocalizationManager

* Merged UpdateFromNative and UpdateFromLocalizations into UpdateLiveTable, as the two are now identical.
* Load both the native and localized display strings into the same FTextLocalizationResource so that we only need to call UpdateLiveTable once.
* Skip the hash check when replacing an existing entry, as that was old logic to avoid stomping a non-localized display string.
* Fixed some dynamic load paths (chunked and plugins) failing to load the native localization data to serve as a consistent base.
* Fixed some duplicate querying of FInternationalization::Get() when finding a display string in a non-shipping build.
* Removed API export from private functions.

#rnx

diff --git a/Engine/Source/Runtime/Core/Private/Internationalization/LocalizationResourceTextSource.cpp b/Engine/Source/Runtime/Core/Private/Internationalization/LocalizationResourceTextSource.cpp
--- a/Engine/Source/Runtime/Core/Private/Internationalization/LocalizationResourceTextSource.cpp
+++ b/Engine/Source/Runtime/Core/Private/Internationalization/LocalizationResourceTextSource.cpp
@@ -181,7 +181,7 @@ void FLocalizationResourceTextSource::LoadLocalizedResourcesFromPaths(TArrayView
 	};
 
 	// Load the native texts first to ensure we always apply translations to a consistent base
-	if (InPrioritizedNativePaths.Num() > 0)
+	if (ShouldLoadNative(InLoadFlags) && InPrioritizedNativePaths.Num() > 0)
 	{
 		for (const FString& LocalizationPath : InPrioritizedNativePaths)
 		{
diff --git a/Engine/Source/Runtime/Core/Private/Internationalization/TextLocalizationManager.cpp b/Engine/Source/Runtime/Core/Private/Internationalization/TextLocalizationManager.cpp
--- a/Engine/Source/Runtime/Core/Private/Internationalization/TextLocalizationManager.cpp
+++ b/Engine/Source/Runtime/Core/Private/Internationalization/TextLocalizationManager.cpp
@@ -991,7 +991,7 @@ void FTextLocalizationManager::RegisterPolyglotTextData(TArrayView<const FPolygl
 
 		if (!TextLocalizationResource.IsEmpty())
 		{
-			UpdateFromLocalizations(MoveTemp(TextLocalizationResource));
+			UpdateLiveTable(MoveTemp(TextLocalizationResource));
 		}
 	}
 }
@@ -1004,15 +1004,18 @@ FTextConstDisplayStringPtr FTextLocalizationManager::FindDisplayString_Internal(
 	}
 
 #if ENABLE_LOC_TESTING
-	if (FInternationalization::Get().GetCurrentLanguage()->GetName() == FLeetCulture::StaticGetName())
 	{
-		FTextDisplayStringRef TmpDisplayString = MakeTextDisplayString(CopyTemp(SourceString));
-		FInternationalization::Leetify(*TmpDisplayString);
-		return TmpDisplayString;
-	}
-	if (FInternationalization::Get().GetCurrentLanguage()->GetName() == FKeysCulture::StaticGetName())
-	{
-		return MakeTextDisplayString(TextLocalizationManager::KeyifyTextId(TextId));
+		FInternationalization& I18N = FInternationalization::Get();
+		if (I18N.GetCurrentLanguage()->GetName() == FLeetCulture::StaticGetName())
+		{
+			FTextDisplayStringRef TmpDisplayString = MakeTextDisplayString(CopyTemp(SourceString));
+			FInternationalization::Leetify(*TmpDisplayString);
+			return TmpDisplayString;
+		}
+		if (I18N.GetCurrentLanguage()->GetName() == FKeysCulture::StaticGetName())
+		{
+			return MakeTextDisplayString(TextLocalizationManager::KeyifyTextId(TextId));
+		}
 	}
 #endif	// ENABLE_LOC_TESTING
 
@@ -1167,12 +1170,12 @@ void FTextLocalizationManager::UpdateFromLocalizationResource(const FString& Loc
 {
 	FTextLocalizationResource TextLocalizationResource;
 	TextLocalizationResource.LoadFromFile(LocalizationResourceFilePath, 0);
-	UpdateFromLocalizations(MoveTemp(TextLocalizationResource));
+	UpdateLiveTable(MoveTemp(TextLocalizationResource));
 }
 
 void FTextLocalizationManager::UpdateFromLocalizationResource(const FTextLocalizationResource& TextLocalizationResource)
 {
-	UpdateFromLocalizations(CopyTemp(TextLocalizationResource));
+	UpdateLiveTable(CopyTemp(TextLocalizationResource));
 }
 
 void FTextLocalizationManager::WaitForAsyncTasks()
@@ -1230,6 +1233,7 @@ void FTextLocalizationManager::HandleLocalizationTargetsMounted(TArrayView<const
 	LocLoadFlags |= (WITH_EDITOR ? ELocalizationLoadFlags::Editor : ELocalizationLoadFlags::None);
 	LocLoadFlags |= (FApp::IsGame() ? ELocalizationLoadFlags::Game : ELocalizationLoadFlags::None);
 	LocLoadFlags |= ELocalizationLoadFlags::Engine;
+	LocLoadFlags |= ELocalizationLoadFlags::Native;
 	// We don't allow dynamically loaded additional data to replace any existing data, as additional localization is lower priority when performing a full update, so this flag emulates that behavior
 	LocLoadFlags |= ELocalizationLoadFlags::SkipExisting;
 
@@ -1376,29 +1380,20 @@ void FTextLocalizationManager::LoadLocalizationResourcesForPrioritizedCultures_S
 	const ELocalizationLoadFlags FinalLocLoadFlags = LocLoadFlags | (ShouldForceLoadGameLocalization() ? ELocalizationLoadFlags::ForceLocalizedGame : ELocalizationLoadFlags::None);
 
 	// Load the resources from each text source
-	FTextLocalizationResource NativeResource;
-	FTextLocalizationResource LocalizedResource;
+	FTextLocalizationResource TextLocalizationResource;
 	for (const TSharedPtr<ILocalizedTextSource>& LocalizedTextSource : AvailableTextSources)
 	{
 		LLM_SCOPE_BYNAME(TEXT("Localization/DisplayStrings"));
-		LocalizedTextSource->LoadLocalizedResources(FinalLocLoadFlags, PrioritizedCultureNames, NativeResource, LocalizedResource);
+		LocalizedTextSource->LoadLocalizedResources(FinalLocLoadFlags, PrioritizedCultureNames, TextLocalizationResource, TextLocalizationResource);
 	}
 
-	FUpdateLiveTableOptions UpdateOptions;
-	UpdateOptions.bDirtyTextRevision = false;
-	UpdateOptions.bReplaceExisting = !EnumHasAnyFlags(FinalLocLoadFlags, ELocalizationLoadFlags::SkipExisting);
-
 	// Apply the new display string data
-	if (!NativeResource.IsEmpty())
 	{
-		UpdateFromNative(MoveTemp(NativeResource), UpdateOptions);
-	}
-	if (!LocalizedResource.IsEmpty())
-	{
-		UpdateFromLocalizations(MoveTemp(LocalizedResource), UpdateOptions);
-	}
+		FUpdateLiveTableOptions UpdateOptions;
+		UpdateOptions.bReplaceExisting = !EnumHasAnyFlags(FinalLocLoadFlags, ELocalizationLoadFlags::SkipExisting);
 
-	DirtyTextRevision();
+		UpdateLiveTable(MoveTemp(TextLocalizationResource), UpdateOptions);
+	}
 }
 
 void FTextLocalizationManager::LoadLocalizationResourcesForPrioritizedCultures_Async(TArrayView<const FString> PrioritizedCultureNames, const ELocalizationLoadFlags LocLoadFlags)
@@ -1421,15 +1416,14 @@ void FTextLocalizationManager::LoadLocalizationTargetsForPrioritizedCultures_Syn
 	}
 
 	// Load the resources from each localization target
-	FTextLocalizationResource UnusedNativeResource;
-	FTextLocalizationResource LocalizedResource;
+	FTextLocalizationResource TextLocalizationResource;
 	for (const FString& LocalizationTargetPath : LocalizationTargetPaths)
 	{
 		UE_LOG(LogTextLocalizationManager, Verbose, TEXT("Loading LocRes data from '%s'"), *LocalizationTargetPath);
 	}
 	{
 		LLM_SCOPE_BYNAME(TEXT("Localization/DisplayStrings"));
-		LocResTextSource->LoadLocalizedResourcesFromPaths(TArrayView<FString>(), LocalizationTargetPaths, TArrayView<FString>(), LocLoadFlags, PrioritizedCultureNames, UnusedNativeResource, LocalizedResource);
+		LocResTextSource->LoadLocalizedResourcesFromPaths(LocalizationTargetPaths, LocalizationTargetPaths, TArrayView<FString>(), LocLoadFlags, PrioritizedCultureNames, TextLocalizationResource, TextLocalizationResource);
 	}
 
 	// Allow any higher priority text sources to override the additional text loaded (eg, to allow polyglot hot-fixes to take priority)
@@ -1438,7 +1432,7 @@ void FTextLocalizationManager::LoadLocalizationTargetsForPrioritizedCultures_Syn
 	{
 		// Copy the IDs array as QueryLocalizedResource can update the map
 		TArray<FTextId> NewTextIds;
-		LocalizedResource.Entries.GenerateKeyArray(NewTextIds);
+		TextLocalizationResource.Entries.GenerateKeyArray(NewTextIds);
 
 		for (const TSharedPtr<ILocalizedTextSource>& LocalizedTextSource : AvailableTextSources)
 		{
@@ -1450,7 +1444,7 @@ void FTextLocalizationManager::LoadLocalizationTargetsForPrioritizedCultures_Syn
 			LLM_SCOPE_BYNAME(TEXT("Localization/DisplayStrings"));
 			for (const FTextId& NewTextId : NewTextIds)
 			{
-				if (LocalizedTextSource->QueryLocalizedResource(LocLoadFlags, PrioritizedCultureNames, NewTextId, UnusedNativeResource, LocalizedResource) == EQueryLocalizedResourceResult::NotImplemented)
+				if (LocalizedTextSource->QueryLocalizedResource(LocLoadFlags, PrioritizedCultureNames, NewTextId, TextLocalizationResource, TextLocalizationResource) == EQueryLocalizedResourceResult::NotImplemented)
 				{
 					bNeedsFullRefresh = true;
 					break;
@@ -1475,8 +1469,8 @@ void FTextLocalizationManager::LoadLocalizationTargetsForPrioritizedCultures_Syn
 		FUpdateLiveTableOptions UpdateOptions;
 		UpdateOptions.bReplaceExisting = !EnumHasAnyFlags(LocLoadFlags, ELocalizationLoadFlags::SkipExisting);
 
-		UE_LOG(LogTextLocalizationManager, Verbose, TEXT("Patching LocRes data for %d entries"), LocalizedResource.Entries.Num());
-		UpdateFromLocalizations(MoveTemp(LocalizedResource), UpdateOptions);
+		UE_LOG(LogTextLocalizationManager, Verbose, TEXT("Patching LocRes data for %d entries"), TextLocalizationResource.Entries.Num());
+		UpdateLiveTable(MoveTemp(TextLocalizationResource), UpdateOptions);
 	}
 }
 
@@ -1585,7 +1579,7 @@ void FTextLocalizationManager::LoadChunkedLocalizationResources_Sync(TArrayView<
 	}
 
 	const TArray<FString> PrioritizedCultureNames = FInternationalization::Get().GetPrioritizedCultureNames(FInternationalization::Get().GetCurrentLanguage()->GetName());
-	LoadLocalizationTargetsForPrioritizedCultures_Sync(AvailableTextSources, PrioritizedLocalizationPaths, PrioritizedCultureNames, ELocalizationLoadFlags::Game);
+	LoadLocalizationTargetsForPrioritizedCultures_Sync(AvailableTextSources, PrioritizedLocalizationPaths, PrioritizedCultureNames, ELocalizationLoadFlags::Native | ELocalizationLoadFlags::Game);
 }
 
 void FTextLocalizationManager::LoadChunkedLocalizationResources_Async(const int32 ChunkId, const FString& PakFilename)
@@ -1615,81 +1609,13 @@ void FTextLocalizationManager::QueueAsyncTask(TUniqueFunction<void()>&& Task)
 	}
 }
 
-void FTextLocalizationManager::UpdateFromNative(FTextLocalizationResource&& TextLocalizationResource, const FUpdateLiveTableOptions& UpdateOptions)
-{
-	TRACE_CPUPROFILER_EVENT_SCOPE(FTextLocalizationManager::UpdateFromNative);
-	LLM_SCOPE_BYNAME(TEXT("Localization/DisplayStrings"));
-
-	// Nothing to do?
-	if (!FTextLocalizationManager::IsDisplayStringSupportEnabled())
-	{
-		return;
-	}
-
-	// Lock while updating the tables
-	{
-		FWriteScopeLock ScopeLock(DisplayStringTableRW);
-
-		DisplayStringLookupTable.Reserve(TextLocalizationResource.Entries.Num());
-
-		// Add/update entries
-		// Note: This code doesn't handle "leet-ification" itself as it is resetting everything to a known "good" state ("leet-ification" happens later on the "good" native text)
-		for (auto& EntryPair : TextLocalizationResource.Entries)
-		{
-			const FTextId& TextId = EntryPair.Key;
-			FTextLocalizationResource::FEntry& NewEntry = EntryPair.Value;
-
-			FDisplayStringEntry* LiveEntry = DisplayStringLookupTable.Find(TextId);
-			if (!LiveEntry)
-			{
-				// Add new entry
-				FDisplayStringEntry NewLiveEntry(
-					NewEntry.LocResID,						/*LocResID*/
-					NewEntry.LocalizationTargetPathId,		/*LocalizationTargetPathId*/
-					NewEntry.SourceStringHash,				/*SourceStringHash*/
-					NewEntry.LocalizedString.ToSharedRef()	/*String*/
-				);
-
-				DisplayStringLookupTable.Emplace(TextId, NewLiveEntry);
-				DisplayStringsByLocalizationTargetId.TrackTextId(INDEX_NONE, NewEntry.LocalizationTargetPathId, TextId);
-			}
-			else if (UpdateOptions.bReplaceExisting)
-			{
-				// Update existing entry
-				// If the existing entry is empty, we just overwrite it 
-				if ((LiveEntry->SourceStringHash == NewEntry.SourceStringHash) || LiveEntry->IsEmpty())
-				{
-					// this is to account for the case where the LiveString is empty and we are overwriting the value 
-					// we could do an if check, but copying an int is cheaper 
-					LiveEntry->SourceStringHash = NewEntry.SourceStringHash;
-					// There is no good way to copy over the NativeBackupString. We disregard it as it is only for testing 
-					LiveEntry->DisplayString = NewEntry.LocalizedString.ToSharedRef();
-#if WITH_EDITORONLY_DATA
-					LiveEntry->LocResID = NewEntry.LocResID;
-#endif	// WITH_EDITORONLY_DATA
-					DisplayStringsByLocalizationTargetId.TrackTextId(LiveEntry->LocalizationTargetPathId, NewEntry.LocalizationTargetPathId, TextId);
-					LiveEntry->LocalizationTargetPathId = NewEntry.LocalizationTargetPathId;
-				}
-			}
-		}
-
-		// Note: Do not use TextLocalizationResource after this point as we may have stolen some of its strings
-		TextLocalizationResource.Entries.Reset();
-	}
-
-	if (UpdateOptions.bDirtyTextRevision)
-	{
-		DirtyTextRevision();
-	}
-}
-
-void FTextLocalizationManager::UpdateFromLocalizations(FTextLocalizationResource&& TextLocalizationResource, const FUpdateLiveTableOptions& UpdateOptions)
+void FTextLocalizationManager::UpdateLiveTable(FTextLocalizationResource&& TextLocalizationResource, const FUpdateLiveTableOptions& UpdateOptions)
 {
-	TRACE_CPUPROFILER_EVENT_SCOPE(FTextLocalizationManager::UpdateFromLocalizations);
+	TRACE_CPUPROFILER_EVENT_SCOPE(FTextLocalizationManager::UpdateLiveTable);
 	LLM_SCOPE_BYNAME(TEXT("Localization/DisplayStrings"));
 
 	// Nothing to do?
-	if (!FTextLocalizationManager::IsDisplayStringSupportEnabled())
+	if (!FTextLocalizationManager::IsDisplayStringSupportEnabled() || TextLocalizationResource.IsEmpty())
 	{
 		return;
 	}
@@ -1723,19 +1649,13 @@ void FTextLocalizationManager::UpdateFromLocalizations(FTextLocalizationResource
 			else if (UpdateOptions.bReplaceExisting)
 			{
 				// Update existing entry
-				// If the source string hashes are are the same, we can replace the display string.
-				// Otherwise, it would suggest the source string has changed and the new localization may be based off of an old source string.
-				// Alternatively, if the display string is empty, we can just overwrite the data.
-				if ((LiveEntry->SourceStringHash == NewEntry.SourceStringHash) || LiveEntry->IsEmpty())
-				{
-					// this is to account for the case where the LiveString is empty and we are overwriting the value 
-					// we could do an if check, but copying an int is cheaper 
-					LiveEntry->SourceStringHash = NewEntry.SourceStringHash;
-					LiveEntry->DisplayString = NewEntry.LocalizedString.ToSharedRef();
+				LiveEntry->SourceStringHash = NewEntry.SourceStringHash;
+				LiveEntry->DisplayString = NewEntry.LocalizedString.ToSharedRef();
 #if WITH_EDITORONLY_DATA
-					LiveEntry->LocResID = NewEntry.LocResID;
+				LiveEntry->LocResID = NewEntry.LocResID;
 #endif	// WITH_EDITORONLY_DATA
-				}
+				DisplayStringsByLocalizationTargetId.TrackTextId(LiveEntry->LocalizationTargetPathId, NewEntry.LocalizationTargetPathId, TextId);
+				LiveEntry->LocalizationTargetPathId = NewEntry.LocalizationTargetPathId;
 			}
 		}
 
diff --git a/Engine/Source/Runtime/Core/Public/Internationalization/TextLocalizationManager.h b/Engine/Source/Runtime/Core/Public/Internationalization/TextLocalizationManager.h
--- a/Engine/Source/Runtime/Core/Public/Internationalization/TextLocalizationManager.h
+++ b/Engine/Source/Runtime/Core/Public/Internationalization/TextLocalizationManager.h
@@ -66,14 +66,6 @@ private:
 			, SourceStringHash(InSourceStringHash)
 		{
 		}
-
-		/** 
-		* Returns true if the display string entry contains invalid display string data. 
-		*/
-		bool IsEmpty() const
-		{
-			return SourceStringHash == 0 && DisplayString->IsEmpty();
-		}
 	};
 
 	/** Manages the currently loaded or registered text localizations. */
@@ -157,7 +149,7 @@ public:
 	 * @note Calling this function with no filters specified will dump the entire live table.
 	 */
 private:
-	CORE_API void DumpLiveTableImpl(const FString* NamespaceFilter, const FString* KeyFilter, const FString* DisplayStringFilter, TFunctionRef<void(const FTextId& Id, const FTextConstDisplayStringRef& DisplayString)> Callback) const;
+	void DumpLiveTableImpl(const FString* NamespaceFilter, const FString* KeyFilter, const FString* DisplayStringFilter, TFunctionRef<void(const FTextId& Id, const FTextConstDisplayStringRef& DisplayString)> Callback) const;
 public:
 	CORE_API void DumpLiveTable(const FString* NamespaceFilter = nullptr, const FString* KeyFilter = nullptr, const FString* DisplayStringFilter = nullptr, const FLogCategoryBase* CategoryOverride = nullptr) const;
 	CORE_API void DumpLiveTable(const FString& OutputFilename, const FString* NamespaceFilter = nullptr, const FString* KeyFilter = nullptr, const FString* DisplayStringFilter = nullptr) const;
@@ -348,29 +340,29 @@ public:
 
 private:
 	/** Callback for when a PAK file is loaded. Async loads any chunk specific localization resources. */
-	CORE_API void OnPakFileMounted(const IPakFile& PakFile);
+	void OnPakFileMounted(const IPakFile& PakFile);
 
 	/** Callback for changes in culture. Async loads the new culture's localization resources. */
-	CORE_API void OnCultureChanged();
+	void OnCultureChanged();
 
 	/** Loads localization resources for the specified culture, optionally loading localization resources that are editor-specific or game-specific. */
-	CORE_API void LoadLocalizationResourcesForCulture_Sync(TArrayView<const TSharedPtr<ILocalizedTextSource>> AvailableTextSources, const FString& CultureName, const ELocalizationLoadFlags LocLoadFlags);
-	CORE_API void LoadLocalizationResourcesForCulture_Async(const FString& CultureName, const ELocalizationLoadFlags LocLoadFlags);
+	void LoadLocalizationResourcesForCulture_Sync(TArrayView<const TSharedPtr<ILocalizedTextSource>> AvailableTextSources, const FString& CultureName, const ELocalizationLoadFlags LocLoadFlags);
+	void LoadLocalizationResourcesForCulture_Async(const FString& CultureName, const ELocalizationLoadFlags LocLoadFlags);
 
 	/** Loads localization resources for the specified prioritized cultures, optionally loading localization resources that are editor-specific or game-specific. */
-	CORE_API void LoadLocalizationResourcesForPrioritizedCultures_Sync(TArrayView<const TSharedPtr<ILocalizedTextSource>> AvailableTextSources, TArrayView<const FString> PrioritizedCultureNames, const ELocalizationLoadFlags LocLoadFlags);
-	CORE_API void LoadLocalizationResourcesForPrioritizedCultures_Async(TArrayView<const FString> PrioritizedCultureNames, const ELocalizationLoadFlags LocLoadFlags);
+	void LoadLocalizationResourcesForPrioritizedCultures_Sync(TArrayView<const TSharedPtr<ILocalizedTextSource>> AvailableTextSources, TArrayView<const FString> PrioritizedCultureNames, const ELocalizationLoadFlags LocLoadFlags);
+	void LoadLocalizationResourcesForPrioritizedCultures_Async(TArrayView<const FString> PrioritizedCultureNames, const ELocalizationLoadFlags LocLoadFlags);
 
 	/** Loads the specified localization targets for the specified prioritized cultures, optionally loading localization resources that are editor-specific or game-specific. */
-	CORE_API void LoadLocalizationTargetsForPrioritizedCultures_Sync(TArrayView<const TSharedPtr<ILocalizedTextSource>> AvailableTextSources, TArrayView<const FString> LocalizationTargetPaths, TArrayView<const FString> PrioritizedCultureNames, const ELocalizationLoadFlags LocLoadFlags);
-	CORE_API void LoadLocalizationTargetsForPrioritizedCultures_Async(TArrayView<const FString> LocalizationTargetPaths, TArrayView<const FString> PrioritizedCultureNames, const ELocalizationLoadFlags LocLoadFlags);
+	void LoadLocalizationTargetsForPrioritizedCultures_Sync(TArrayView<const TSharedPtr<ILocalizedTextSource>> AvailableTextSources, TArrayView<const FString> LocalizationTargetPaths, TArrayView<const FString> PrioritizedCultureNames, const ELocalizationLoadFlags LocLoadFlags);
+	void LoadLocalizationTargetsForPrioritizedCultures_Async(TArrayView<const FString> LocalizationTargetPaths, TArrayView<const FString> PrioritizedCultureNames, const ELocalizationLoadFlags LocLoadFlags);
 
 	/** Loads the chunked localization resource data for the specified chunk */
-	CORE_API void LoadChunkedLocalizationResources_Sync(TArrayView<const TSharedPtr<ILocalizedTextSource>> AvailableTextSources, const int32 ChunkId, const FString& PakFilename);
-	CORE_API void LoadChunkedLocalizationResources_Async(const int32 ChunkId, const FString& PakFilename);
+	void LoadChunkedLocalizationResources_Sync(TArrayView<const TSharedPtr<ILocalizedTextSource>> AvailableTextSources, const int32 ChunkId, const FString& PakFilename);
+	void LoadChunkedLocalizationResources_Async(const int32 ChunkId, const FString& PakFilename);
 
 	/** Queue the given task to run async (where possible), chained to any existing AsyncLocalizationTask */
-	CORE_API void QueueAsyncTask(TUniqueFunction<void()>&& Task);
+	void QueueAsyncTask(TUniqueFunction<void()>&& Task);
 
 	struct FUpdateLiveTableOptions
 	{
@@ -382,20 +374,17 @@ private:
 		bool bReplaceExisting = true;
 	};
 
-	/** Updates display string entries and adds new display string entries based on provided native text. */
-	CORE_API void UpdateFromNative(FTextLocalizationResource&& TextLocalizationResource, const FUpdateLiveTableOptions& UpdateOptions = FUpdateLiveTableOptions());
-
-	/** Updates display string entries and adds new display string entries based on provided localizations. */
-	CORE_API void UpdateFromLocalizations(FTextLocalizationResource&& TextLocalizationResource, const FUpdateLiveTableOptions& UpdateOptions = FUpdateLiveTableOptions());
+	/** Updates display string entries and adds new display string entries. */
+	void UpdateLiveTable(FTextLocalizationResource&& TextLocalizationResource, const FUpdateLiveTableOptions& UpdateOptions = FUpdateLiveTableOptions());
 
 	/** Dirties the local revision counter for the given text ID by incrementing it (or adding it) */
-	CORE_API void DirtyLocalRevisionForTextId(const FTextId& InTextId);
+	void DirtyLocalRevisionForTextId(const FTextId& InTextId);
 
 	/** Internal version of FindDisplayString, shared between FindDisplayString and GetDisplayString */
 	FTextConstDisplayStringPtr FindDisplayString_Internal(const FTextId& TextId, const FString& SourceString) const;
 
 	/** Dirties the text revision counter by incrementing it, causing a revision mismatch for any information cached before this happens.  */
-	CORE_API void DirtyTextRevision();
+	void DirtyTextRevision();
 
 	/** Array of registered localized text sources, sorted by priority (@see RegisterTextSource) */
 	TArray<TSharedPtr<ILocalizedTextSource>> LocalizedTextSources;
