- SafeUnload config branches after 5 seconds of not being used
- FIxed re-loading of SafeUnloaded branches to re-load the dynamic layers
- Added a local cache to some functions that was querying GConfig every frame, keeping Engine config in memory always

diff --git a/Engine/Source/Runtime/ApplicationCore/Private/IOS/IOSAppDelegate.cpp b/Engine/Source/Runtime/ApplicationCore/Private/IOS/IOSAppDelegate.cpp
--- a/Engine/Source/Runtime/ApplicationCore/Private/IOS/IOSAppDelegate.cpp
+++ b/Engine/Source/Runtime/ApplicationCore/Private/IOS/IOSAppDelegate.cpp
@@ -1134,53 +1134,60 @@ static FAutoConsoleVariableRef CVarGEnableThermalsReport(
 #if !PLATFORM_TVOS
 - (UIInterfaceOrientationMask)application:(UIApplication *)application supportedInterfaceOrientationsForWindow:(UIWindow*)window
 {
-	bool bSupportsPortrait;
-	bool bSupportsPortraitUpsideDown;
-	bool bSupportsLandscapeLeft;
-	bool bSupportsLandscapeRight;
+	static bool bInitialized = false;
+	static UIInterfaceOrientationMask Mask = 0;
 	
-	// This is called during app startup and IOSRuntimeSettings may not have been loaded yet
-	bool hasValue = GConfig->GetBool(TEXT("/Script/IOSRuntimeSettings.IOSRuntimeSettings"), TEXT("bSupportsPortraitOrientation"), bSupportsPortrait, GEngineIni);
-	
-	NSArray<NSString*> *SupportedOrientations = [[NSBundle mainBundle] objectForInfoDictionaryKey:@"UISupportedInterfaceOrientations"];
-	if (!hasValue && SupportedOrientations != NULL)
+	if (!bInitialized)
 	{
-		// Loop through the Info.plist UISupportedInterfaceOrientations array values looking for "Portrait", "Left" and "Right"
-		NSPredicate *predicate = [NSPredicate predicateWithFormat:@"SELF == %@", @"UIInterfaceOrientationPortrait"];
-		bSupportsPortrait = ([SupportedOrientations filteredArrayUsingPredicate:predicate].count > 0);
+		bInitialized = true;
 		
-		NSPredicate *predicateDown = [NSPredicate predicateWithFormat:@"SELF == %@", @"UIInterfaceOrientationPortraitUpsideDown"];
-		bSupportsPortraitUpsideDown = ([SupportedOrientations filteredArrayUsingPredicate:predicateDown].count > 0);
+		bool bSupportsPortrait;
+		bool bSupportsPortraitUpsideDown;
+		bool bSupportsLandscapeLeft;
+		bool bSupportsLandscapeRight;
 		
-		NSPredicate *PredicateLeft = [NSPredicate predicateWithFormat:@"SELF == %@", @"UIInterfaceOrientationLandscapeLeft"];
-		bSupportsLandscapeLeft = ([SupportedOrientations filteredArrayUsingPredicate:PredicateLeft].count > 0);
+		// This is called during app startup and IOSRuntimeSettings may not have been loaded yet
+		bool hasValue = GConfig->GetBool(TEXT("/Script/IOSRuntimeSettings.IOSRuntimeSettings"), TEXT("bSupportsPortraitOrientation"), bSupportsPortrait, GEngineIni);
 		
-		NSPredicate *PredicateRight = [NSPredicate predicateWithFormat:@"SELF == %@", @"UIInterfaceOrientationLandscapeRight"];
-		bSupportsLandscapeRight = ([SupportedOrientations filteredArrayUsingPredicate:PredicateRight].count > 0);
-	}
-	else
-	{
-		GConfig->GetBool(TEXT("/Script/IOSRuntimeSettings.IOSRuntimeSettings"), TEXT("bSupportsUpsideDownOrientation"), bSupportsPortraitUpsideDown, GEngineIni);
-		GConfig->GetBool(TEXT("/Script/IOSRuntimeSettings.IOSRuntimeSettings"), TEXT("bSupportsLandscapeLeftOrientation"), bSupportsLandscapeLeft, GEngineIni);
-		GConfig->GetBool(TEXT("/Script/IOSRuntimeSettings.IOSRuntimeSettings"), TEXT("bSupportsLandscapeRightOrientation"), bSupportsLandscapeRight, GEngineIni);
-	}
-	
-	UIInterfaceOrientationMask Mask = 0;
-	if (bSupportsPortrait)
-	{
-		Mask |= UIInterfaceOrientationMaskPortrait;
-	}
-	if (bSupportsPortraitUpsideDown)
-	{
-		Mask |= UIInterfaceOrientationMaskPortraitUpsideDown;
-	}
-	if (bSupportsLandscapeLeft)
-	{
-		Mask |= UIInterfaceOrientationMaskLandscapeLeft;
-	}
-	if (bSupportsLandscapeRight)
-	{
-		Mask |= UIInterfaceOrientationMaskLandscapeRight;
+		NSArray<NSString*> *SupportedOrientations = [[NSBundle mainBundle] objectForInfoDictionaryKey:@"UISupportedInterfaceOrientations"];
+		if (!hasValue && SupportedOrientations != NULL)
+		{
+			// Loop through the Info.plist UISupportedInterfaceOrientations array values looking for "Portrait", "Left" and "Right"
+			NSPredicate *predicate = [NSPredicate predicateWithFormat:@"SELF == %@", @"UIInterfaceOrientationPortrait"];
+			bSupportsPortrait = ([SupportedOrientations filteredArrayUsingPredicate:predicate].count > 0);
+			
+			NSPredicate *predicateDown = [NSPredicate predicateWithFormat:@"SELF == %@", @"UIInterfaceOrientationPortraitUpsideDown"];
+			bSupportsPortraitUpsideDown = ([SupportedOrientations filteredArrayUsingPredicate:predicateDown].count > 0);
+			
+			NSPredicate *PredicateLeft = [NSPredicate predicateWithFormat:@"SELF == %@", @"UIInterfaceOrientationLandscapeLeft"];
+			bSupportsLandscapeLeft = ([SupportedOrientations filteredArrayUsingPredicate:PredicateLeft].count > 0);
+			
+			NSPredicate *PredicateRight = [NSPredicate predicateWithFormat:@"SELF == %@", @"UIInterfaceOrientationLandscapeRight"];
+			bSupportsLandscapeRight = ([SupportedOrientations filteredArrayUsingPredicate:PredicateRight].count > 0);
+		}
+		else
+		{
+			GConfig->GetBool(TEXT("/Script/IOSRuntimeSettings.IOSRuntimeSettings"), TEXT("bSupportsUpsideDownOrientation"), bSupportsPortraitUpsideDown, GEngineIni);
+			GConfig->GetBool(TEXT("/Script/IOSRuntimeSettings.IOSRuntimeSettings"), TEXT("bSupportsLandscapeLeftOrientation"), bSupportsLandscapeLeft, GEngineIni);
+			GConfig->GetBool(TEXT("/Script/IOSRuntimeSettings.IOSRuntimeSettings"), TEXT("bSupportsLandscapeRightOrientation"), bSupportsLandscapeRight, GEngineIni);
+		}
+		
+		if (bSupportsPortrait)
+		{
+			Mask |= UIInterfaceOrientationMaskPortrait;
+		}
+		if (bSupportsPortraitUpsideDown)
+		{
+			Mask |= UIInterfaceOrientationMaskPortraitUpsideDown;
+		}
+		if (bSupportsLandscapeLeft)
+		{
+			Mask |= UIInterfaceOrientationMaskLandscapeLeft;
+		}
+		if (bSupportsLandscapeRight)
+		{
+			Mask |= UIInterfaceOrientationMaskLandscapeRight;
+		}
 	}
 	
 	// If no orientation constraints are set, default to MaskAll
diff --git a/Engine/Source/Runtime/Core/Private/GenericPlatform/GenericPlatformMisc.cpp b/Engine/Source/Runtime/Core/Private/GenericPlatform/GenericPlatformMisc.cpp
--- a/Engine/Source/Runtime/Core/Private/GenericPlatform/GenericPlatformMisc.cpp
+++ b/Engine/Source/Runtime/Core/Private/GenericPlatform/GenericPlatformMisc.cpp
@@ -2028,12 +2028,8 @@ bool FGenericPlatformMisc::IsPGIActive()
 
 int FGenericPlatformMisc::GetMobilePropagateAlphaSetting()
 {
-	static int PropagateAlpha = -1;
-	if (PropagateAlpha < 0)
-	{
-		GConfig->GetInt(TEXT("/Script/Engine.RendererSettings"), TEXT("r.Mobile.PropagateAlpha"), PropagateAlpha, GEngineIni);
-	}
-	return PropagateAlpha;
+	static auto* PropagateAlphaCVar = IConsoleManager::Get().FindConsoleVariable(TEXT("r.Mobile.PropagateAlpha"));
+	return PropagateAlphaCVar ? PropagateAlphaCVar->GetInt() : 0;
 }
 
 void FGenericPlatformMisc::ShowConsoleWindow()
diff --git a/Engine/Source/Runtime/Core/Private/HAL/ThreadHeartBeat.cpp b/Engine/Source/Runtime/Core/Private/HAL/ThreadHeartBeat.cpp
--- a/Engine/Source/Runtime/Core/Private/HAL/ThreadHeartBeat.cpp
+++ b/Engine/Source/Runtime/Core/Private/HAL/ThreadHeartBeat.cpp
@@ -977,15 +977,17 @@ void FGameThreadHitchHeartBeatThreaded::InitSettings()
 	}
 	else
 	{
-		float Config_Duration = -1.0f;
-		bool Config_StackWalk = false;
+		static float Config_Duration = -1.0f;
+		static bool Config_StackWalk = false;
 
 		// Read from config files
 		bool bReadFromConfig = false;
 		if (GConfig)
 		{
-			bReadFromConfig |= GConfig->GetFloat(TEXT("Core.System"), TEXT("GameThreadHeartBeatHitchDuration"), Config_Duration, GEngineIni);
-			bReadFromConfig |= GConfig->GetBool(TEXT("Core.System"), TEXT("GameThreadHeartBeatStackWalk"), Config_StackWalk, GEngineIni);
+			static bool bLocalReadFromConfig = 
+				GConfig->GetFloat(TEXT("Core.System"), TEXT("GameThreadHeartBeatHitchDuration"), Config_Duration, GEngineIni) ||
+				GConfig->GetBool(TEXT("Core.System"), TEXT("GameThreadHeartBeatStackWalk"), Config_StackWalk, GEngineIni);
+			bReadFromConfig = bLocalReadFromConfig;
 		}
 
 		if (bReadFromConfig)
diff --git a/Engine/Source/Runtime/Core/Private/Misc/ConfigCacheIni.cpp b/Engine/Source/Runtime/Core/Private/Misc/ConfigCacheIni.cpp
--- a/Engine/Source/Runtime/Core/Private/Misc/ConfigCacheIni.cpp
+++ b/Engine/Source/Runtime/Core/Private/Misc/ConfigCacheIni.cpp
@@ -68,6 +68,13 @@ static FAutoConsoleVariableRef CVarUseNewSaveTracking(
 	GUseNewSaveTracking,
 	TEXT("If true, use the new method for tracking modifications to GConfig when saving"));
 
+static int GTimeToUnloadConfig = 0;
+static FAutoConsoleVariableRef CVarTimeToUnloadConfig(
+	TEXT("ini.TimeToUnloadConfig"),
+	GTimeToUnloadConfig,
+	TEXT("If > 0, when a config branch hasn't been accessed in this many seconds, SafeUnload the branch"));
+
+
 #if WITH_EDITOR
 // editor wants full replay - and we can't put this in a cvar since we need way too early for cvar processing!
 static int GDefaultReplayMethod = 2;
@@ -2988,6 +2995,8 @@ FConfigBranch::FConfigBranch()
 	}
 
 	InitFiles();
+
+	InactiveTimer = -1;
 }
 
 FConfigBranch::FConfigBranch(const FConfigFile& ExistingFile)
@@ -3359,23 +3368,53 @@ void FConfigBranch::RemoveTagFromHierarchy(FName Tag, FConfigModificationTracker
 	}
 }
 
-bool FConfigBranch::SafeUnload()
+void FConfigBranch::SafeUnload()
 {
 	bIsSafeUnloaded = true;
 
-	RunOnEachFile([](FConfigFile& File, const FString& Name)
+	InMemoryFile.Cleanup();
+	CombinedStaticLayers.Cleanup();
+	FinalCombinedLayers.Cleanup();
+
+	// empty the command streams for the static and dynamic layers, but leave any other streams alone
+	// note that we keep the dynamic list around, but without the section data, because we use the 
+	// dynamic layer filename, tag, and priority to load again
+	StaticLayers.Empty();
+	for (DynamicLayerList::TIterator Node(DynamicLayers.GetHead()); Node; ++Node)
 	{
-		File.Cleanup();
-	});
+		Node->Empty();
+	}
+}
 
-	RunOnEachCommandStream([](FConfigCommandStream& Stream, const FString& Name)
+void FConfigBranch::SafeReload()
+{
+	double StartTime = FPlatformTime::Seconds();
+	
+	// read static layers back in from disk
+	// @todo make sure we only Unload from GConfig
+	FConfigContext Context = FConfigContext::ReadIntoConfigSystem(GConfig, Platform.ToString());
+	Context.Branch = this;
+	Context.DestIniFilename = IniPath;
+	Context.Load(*IniName.ToString());
+
+	// read dynamic layers back in from disk
+	FConfigBranch::DynamicLayerList EmptiedDynamicLayers;
+	while (!DynamicLayers.IsEmpty())
 	{
-		Stream.Empty();
-	});
+		FConfigBranch::DynamicLayerList::TDoubleLinkedListNode* HeadNode = DynamicLayers.GetHead();
+		DynamicLayers.RemoveNode(HeadNode, false);
+		EmptiedDynamicLayers.AddTail(HeadNode);
+	}
+	for (FConfigBranch::DynamicLayerList::TIterator Node(EmptiedDynamicLayers.GetHead()); Node; ++Node)
+	{
+		FConfigCommandStream& S = *Node.GetNode()->GetValue();
+		AddDynamicLayersToHierarchy({ S.Filename }, S.Tag, (DynamicLayerPriority)S.Priority);
+	}
 
-	return true;
+	UE_LOG(LogConfig, Log, TEXT("Branch '%s' had been unloaded. Reloading on-demand took %.2fms"), *IniName.ToString(), (FPlatformTime::Seconds() - StartTime) * 1000.0f);
 }
 
+
 bool FConfigBranch::RemoveSection(const TCHAR* Section)
 {
 	int NumRemoved = 0;
@@ -3482,6 +3521,62 @@ FConfigCacheIni::~FConfigCacheIni()
 	Flush( 1 );
 }
 
+void FConfigCacheIni::Tick(float DeltaSeconds)
+{
+	if (GTimeToUnloadConfig == 0)
+	{
+		return;
+	}
+
+	FConfigBranch* BranchesToCheck[2];
+
+	// find next known file to check
+	static int KnownFileToCheckForUnload = 0;
+	if (KnownFileToCheckForUnload >= (uint8)EKnownIniFile::NumKnownFiles)
+	{
+		KnownFileToCheckForUnload = 0;
+	}
+	BranchesToCheck[0] = &KnownFiles.Branches[KnownFileToCheckForUnload++];
+	
+	// find next unknown file to check
+	static int OtherFileToCheckForUnload = 0;
+	if (OtherFileToCheckForUnload >= OtherFileNames.Num())
+	{
+		OtherFileToCheckForUnload = 0;
+	}
+	BranchesToCheck[1] = OtherFiles.FindRef(OtherFileNames[OtherFileToCheckForUnload++]);
+
+	checkf(OtherFileNames.Num() == OtherFiles.Num(), TEXT("OtherFIles and OtherFileNames are out of sync! %d other files, %d other file names!"), OtherFileNames.Num(), OtherFiles.Num());
+	
+	// now check for unused files
+	double Now = FPlatformTime::Seconds();
+	for (FConfigBranch* Branch : BranchesToCheck)
+	{
+		if (Branch == nullptr || Branch->bIsSafeUnloaded)
+		{
+			continue;
+		}
+		
+		// we start out negative so that we ignre the long startup time without ticking, so on first tick we allow it to be tracked
+		if (Branch->InactiveTimer < 0)
+		{
+			Branch->InactiveTimer = Now;
+		}
+		else if (Branch->InactiveTimer > 0)
+		{
+			if (Now - Branch->InactiveTimer > GTimeToUnloadConfig)
+			{
+				UE_LOG(LogConfig, Log, TEXT("Unloading %s due to inactivity"), *Branch->IniPath);
+				
+				Branch->SafeUnload();
+				Branch->InactiveTimer = 0;
+			}
+		}
+	}
+}
+
+
+
 FConfigBranch* FConfigCacheIni::FindBranchWithNoReload(FName BaseIniName, const FString& Filename)
 {
 	// look for a known file, if there's no ini extension
@@ -3522,13 +3617,14 @@ FConfigBranch* FConfigCacheIni::FindBranch(FName BaseIniName, const FString& Fil
 
 	if (Branch && Branch->bIsSafeUnloaded)
 	{
-		double StartTime = FPlatformTime::Seconds();
-		FConfigContext Context = FConfigContext::ReadIntoConfigSystem(this, PlatformName.ToString());
-		Context.Branch = Branch;
-		Context.DestIniFilename = Branch->IniPath;
-		Context.Load(*BaseIniName.ToString());
+		Branch->SafeReload();
+	}
 
-		UE_LOG(LogConfig, Log, TEXT("Branch '%s' had been unloaded. Reloading on-demand took %.2fms"), *Branch->IniName.ToString(), (FPlatformTime::Seconds() - StartTime) * 1000.0f);
+	// track that this branch is being used, so re-set the time
+	if (Branch && Branch->InactiveTimer >= 0 && GTimeToUnloadConfig > 0)
+	{
+		Branch->InactiveTimer = FPlatformTime::Seconds();
+		UE_LOG(LogConfig, Log, TEXT("REsetting InactiveTimer for %s"), *Branch->IniName.ToString());
 	}
 
 	return Branch;
@@ -3547,24 +3643,52 @@ FConfigBranch& FConfigCacheIni::AddNewBranch(const FString& Filename)
 		FileAccess->OverrideFilenameToLoad = FName(FStringView(Filename));
 	}
 #endif
+	if (OtherFiles.Find(Filename) == nullptr)
+	{
+		OtherFileNames.Add(Filename);
+	}
 	FConfigBranch*& Existing = OtherFiles.FindOrAdd(Filename);
 	delete Existing;
 	Existing = Branch;
 	return *Branch;
 }
 
+int32 FConfigCacheIni::Remove(const FString& Filename)
+{
+	OtherFileNames.Remove(Filename);
+	delete OtherFiles.FindRef(Filename);
+	return OtherFiles.Remove(Filename);
+}
+
 
 FConfigFile* FConfigCacheIni::FindConfigFile( const FString& Filename )
 {
-	// look for a known file, if there's no ini extension
-	FConfigFile* Result = Filename.EndsWith(TEXT(".ini")) ? nullptr : KnownFiles.GetMutableFile(FName(*Filename));
+	FConfigBranch* Result;
+	if (!Filename.EndsWith(TEXT(".ini")))
+	{
+		Result = KnownFiles.GetBranch(*Filename);
+	}
+	else
+	{
+		Result = OtherFiles.FindRef(Filename);
+	}
 
-	if (Result == nullptr)
+	if (Result)
 	{
-		FConfigBranch* Branch = OtherFiles.FindRef(Filename);
-		Result = Branch ? &Branch->InMemoryFile : nullptr;
+		if (Result->bIsSafeUnloaded)
+		{
+			Result->SafeReload();
+		};
+		// track that this branch is being used, so re-set the time
+		if (Result && Result->InactiveTimer >= 0 && GTimeToUnloadConfig > 0)
+	{
+			Result->InactiveTimer = FPlatformTime::Seconds();
+			UE_LOG(LogConfig, VeryVerbose, TEXT("REsetting InactiveTimer for %s"), *Result->IniName.ToString());
+		}
+		return &Result->InMemoryFile;
 	}
-	return Result;
+
+	return nullptr;
 }
 
 FConfigFile* FConfigCacheIni::Find(const FString& Filename)
@@ -3613,17 +3737,32 @@ FConfigFile* FConfigCacheIni::Find(const FString& Filename)
 
 FConfigFile* FConfigCacheIni::FindConfigFileWithBaseName(FName BaseName)
 {
-	if (FConfigFile* Result = KnownFiles.GetMutableFile(BaseName))
+	FConfigBranch* Result = KnownFiles.GetBranch(BaseName);
+	if (Result == nullptr)
 	{
-		return Result;
+		for (TPair<FString,FConfigBranch*>& CurrentFilePair : OtherFiles)
+		{
+			if (CurrentFilePair.Value->InMemoryFile.Name == BaseName)
+			{
+				Result = CurrentFilePair.Value;
+				break;
+			}
+		}
 	}
 
-	for (TPair<FString,FConfigBranch*>& CurrentFilePair : OtherFiles)
+	if (Result)
 	{
-		if (CurrentFilePair.Value->InMemoryFile.Name == BaseName)
+		if (Result->bIsSafeUnloaded)
+	{
+			Result->SafeReload();
+		}
+		// track that this branch is being used, so re-set the time
+		if (Result && Result->InactiveTimer >= 0 && GTimeToUnloadConfig > 0)
 		{
-			return &CurrentFilePair.Value->InMemoryFile;
+			Result->InactiveTimer = FPlatformTime::Seconds();
+			UE_LOG(LogConfig, Log, TEXT("REsetting InactiveTimer for %s"), *Result->IniName.ToString());
 		}
+		return &Result->InMemoryFile;
 	}
 	return nullptr;
 }
@@ -3641,6 +3780,10 @@ FConfigFile& FConfigCacheIni::Add(const FString& Filename, const FConfigFile& Fi
 		FileAccess->OverrideFilenameToLoad = FName(FStringView(Filename));
 	}
 #endif
+	if (OtherFiles.Find(Filename) == nullptr)
+	{
+		OtherFileNames.Add(Filename);
+	}
 	FConfigBranch*& Existing = OtherFiles.FindOrAdd(Filename);
 	delete Existing;
 	Existing = Branch;
@@ -3676,8 +3819,7 @@ bool FConfigCacheIni::ContainsConfigFile(const FConfigFile* ConfigFile) const
 
 TArray<FString> FConfigCacheIni::GetFilenames()
 {
-	TArray<FString> Result;
-	OtherFiles.GetKeys(Result);
+	TArray<FString> Result = OtherFileNames;
 
 	for (const FConfigBranch& Branch : KnownFiles.Branches)
 	{
@@ -3742,6 +3884,7 @@ void FConfigCacheIni::Flush(bool bRemoveFromCache, const FString& Filename )
 				delete It.Value;
 			}
 			OtherFiles.Empty();
+			OtherFileNames.Empty();
 		}
 	}
 }
@@ -4150,7 +4293,8 @@ bool FConfigCacheIni::SafeUnloadBranch(const TCHAR* BranchName)
 	FConfigBranch* Branch = FindBranchWithNoReload(BranchName, BranchName);
 	if (Branch)
 	{
-		return Branch->SafeUnload();
+		Branch->SafeUnload();
+		return true;
 	}
 
 	return false;
@@ -4981,6 +5125,7 @@ void FConfigCacheIni::ShowMemoryUsage( FOutputDevice& Ar )
 	// record the memory used by the FConfigCacheIni's TMap
 	FArchiveCountConfigMem MemAr;
 	OtherFiles.CountBytes(MemAr);
+	OtherFileNames.CountBytes(MemAr);
 
 	SIZE_T TotalMemoryUsage=MemAr.GetNum();
 	SIZE_T MaxMemoryUsage=MemAr.GetMax();
@@ -5015,6 +5160,7 @@ SIZE_T FConfigCacheIni::GetMaxMemoryUsage()
 	// record the memory used by the FConfigCacheIni's TMap
 	FArchiveCountConfigMem MemAr;
 	OtherFiles.CountBytes(MemAr);
+	OtherFileNames.CountBytes(MemAr);
 
 	SIZE_T TotalMemoryUsage=MemAr.GetNum();
 	SIZE_T MaxMemoryUsage=MemAr.GetMax();
@@ -5128,6 +5274,7 @@ void FConfigCacheIni::Serialize(FArchive& Ar)
 			Ar << Filename;
 			Ar << *Branch;
 			OtherFiles.Add(Filename, Branch);
+			OtherFileNames.Add(Filename);
 		}
 	}
 	else
@@ -6437,6 +6584,14 @@ class FIniExec : public FSelfRegisteringExec
 			}
 		}
 
+		if (FParse::Command(&Cmd, TEXT("UnloadAll")))
+		{
+			for (const FString& Filename : GConfig->GetFilenames())
+			{
+				GConfig->SafeUnloadBranch(*Filename);
+			}
+		}
+
 		if (FParse::Command(&Cmd, TEXT("AddHotFix")))
 		{
 			TCHAR FileName[256];
@@ -6538,6 +6693,8 @@ class FIniExec : public FSelfRegisteringExec
 			uint64 Total = 0;
 			int NumSkipped = 0;
 			uint64 SkippedTotal = 0;
+			int Unloaded = 0;
+			uint64 UnloadedTotal = 0;
 			int SingleSection = 0;
 			uint64 SingleSectionTotal = 0;
 			int NoSection = 0;
@@ -6554,12 +6711,17 @@ class FIniExec : public FSelfRegisteringExec
 				Total += Mem;
 				SlackTotal += MemAr.GetMax() - MemAr.GetNum();
 
-				if (Branch->InMemoryFile.Num() == 1)
+				if (Branch->bIsSafeUnloaded)
+				{
+					Unloaded++;
+					UnloadedTotal += Mem;
+				}
+				else if (Branch->InMemoryFile.Num() == 1)
 				{
 					SingleSection++;
 					SingleSectionTotal += Mem;
 				}
-				if (Branch->InMemoryFile.Num() == 0)
+				else if (Branch->InMemoryFile.Num() == 0)
 				{
 					NoSection++;
 					NoSectionTotal += Mem;
@@ -6687,7 +6849,8 @@ class FIniExec : public FSelfRegisteringExec
 			{
 				Ar.Logf(TEXT(""));
 			Ar.Logf(TEXT("[%0.2fmb] - %d All Configs"), (double)Total / 1024.0 / 1024.0, GConfig->GetFilenames().Num());
-				Ar.Logf(TEXT("[%0.2fmb] - %d Tiny Configs (not displayed above)"), (double)SkippedTotal / 1024.0 / 1024.0, NumSkipped);
+			Ar.Logf(TEXT("[%0.2fmb] - %d SafeUnloaded Configs"), (double)UnloadedTotal / 1024.0 / 1024.0, Unloaded);
+			Ar.Logf(TEXT("[%0.2fmb] - %d Tiny Configs (not displayed above)"), (double)SkippedTotal / 1024.0 / 1024.0, NumSkipped);
 			Ar.Logf(TEXT("[%0.2fmb] - %d Single Section Configs"), (double)SingleSectionTotal / 1024.0 / 1024.0, SingleSection);
 			Ar.Logf(TEXT("[%0.2fmb] - %d ZeroSection Configs"), (double)NoSectionTotal / 1024.0 / 1024.0, NoSection);
 				Ar.Logf(TEXT("[%0.2fmb] - Total Slack (wasted memory)"), (double)SlackTotal / 1024.0 / 1024.0);
diff --git a/Engine/Source/Runtime/Core/Public/Misc/ConfigCacheIni.h b/Engine/Source/Runtime/Core/Public/Misc/ConfigCacheIni.h
--- a/Engine/Source/Runtime/Core/Public/Misc/ConfigCacheIni.h
+++ b/Engine/Source/Runtime/Core/Public/Misc/ConfigCacheIni.h
@@ -1116,8 +1116,9 @@ public:
 	 * Frees up the static layer memory, which can be useful if a branch is loaded from, cached, and never used again. However,
 	 * if something does try to access it, it will reload in-place (can hitch your game, so be aware)
 	 */
-	CORE_API bool SafeUnload();
-	
+	CORE_API void SafeUnload();
+	CORE_API void SafeReload();
+
 	/**
 	  * Removes the section completely from all layers of this branch 
 	  * This is destructive! It will not reload on demand (like SafeUnload does). Use this only when you know you will _never_ need
@@ -1145,6 +1146,9 @@ private:
 	void RemoveTagFromHierarchy(FName Tag, FConfigModificationTracker* ModificationTracker);
 	
 	friend class FConfigCacheIni;
+	
+	// when we last Found the branch to pull data from it, this is used to unload after it's been unused for some time
+	double InactiveTimer;
 };
 
 
@@ -1189,6 +1193,11 @@ public:
 		return bGloballyRegistered;
 	}
 
+	/**
+	 * Allow for periodic cleanup or other tasks
+	 */
+	CORE_API void Tick(float DeltaSeconds);
+
 	/**
 	* Prases apart an ini section that contains a list of 1-to-N mappings of strings in the following format
 	*	 [PerMapPackages]
@@ -1269,11 +1278,8 @@ public:
     /** Create a new branch for FIlename, and return it */
 	CORE_API FConfigBranch& AddNewBranch(const FString& Filename);
 	
-	int32 Remove(const FString& Filename)
-	{
-		delete OtherFiles.FindRef(Filename);
-		return OtherFiles.Remove(Filename);
-	}
+	CORE_API int32 Remove(const FString& Filename);
+	
 	CORE_API TArray<FString> GetFilenames();
 
 
@@ -1911,6 +1917,9 @@ private:
 	FKnownConfigFiles KnownFiles;
 
 	TMap<FString, FConfigBranch*> OtherFiles;
+	
+	/** Parallel array to OtherFiles */
+	TArray<FString> OtherFileNames;
 
 	struct FPluginInfo
 	{
diff --git a/Engine/Source/Runtime/Engine/Private/GameEngine.cpp b/Engine/Source/Runtime/Engine/Private/GameEngine.cpp
--- a/Engine/Source/Runtime/Engine/Private/GameEngine.cpp
+++ b/Engine/Source/Runtime/Engine/Private/GameEngine.cpp
@@ -1739,6 +1739,8 @@ void UGameEngine::Tick( float DeltaSeconds, bool bIdleMode )
 		FEngineAnalytics::Tick(DeltaSeconds);
 	}
 
+	GConfig->Tick(DeltaSeconds);
+
 	// -----------------------------------------------------
 	// Begin ticking worlds
 	// -----------------------------------------------------
diff --git a/Engine/Source/Runtime/RHI/Private/RHI.cpp b/Engine/Source/Runtime/RHI/Private/RHI.cpp
--- a/Engine/Source/Runtime/RHI/Private/RHI.cpp
+++ b/Engine/Source/Runtime/RHI/Private/RHI.cpp
@@ -1178,9 +1178,25 @@ static ERHIBindlessConfiguration ParseConfigurationFromString(const FString& InS
 static bool GetBindlessConfigurationSetting(FString& OutSetting, EShaderPlatform Platform, const TCHAR* SettingName)
 {
 	const FString ShaderFormat = FDataDrivenShaderPlatformInfo::GetShaderFormat(Platform).ToString();
+	static TMap<FString, FString> ShaderFormatToBindlessConfig;
 	if (!ShaderFormat.IsEmpty())
 	{
-		return GConfig->GetString(*ShaderFormat, SettingName, OutSetting, GEngineIni);
+		FString* Config = ShaderFormatToBindlessConfig.Find(ShaderFormat);
+		if (Config == nullptr)
+		{
+			if (!GConfig->GetString(*ShaderFormat, SettingName, OutSetting, GEngineIni))
+			{
+				OutSetting = TEXT("");
+			}
+			ShaderFormatToBindlessConfig.Add(ShaderFormat, OutSetting);
+		}
+		else
+		{
+			OutSetting = *Config;
+		}
+
+		// return true if we ever found a setting
+		return OutSetting.Len() > 0;
 	}
 
 	return false;
