Apply CodeAnalysis modifications for all Microsoft platforms not just Windows

#rnx

diff --git a/Engine/Source/Runtime/Core/Public/HAL/Platform.h b/Engine/Source/Runtime/Core/Public/HAL/Platform.h
--- a/Engine/Source/Runtime/Core/Public/HAL/Platform.h
+++ b/Engine/Source/Runtime/Core/Public/HAL/Platform.h
@@ -110,8 +110,8 @@
 
 #if PLATFORM_COMPILER_CLANG
 	#include "Clang/ClangPlatformCodeAnalysis.h"
-#elif PLATFORM_WINDOWS
-	#include "Windows/WindowsPlatformCodeAnalysis.h"
+#elif PLATFORM_MICROSOFT
+	#include "Microsoft/MicrosoftPlatformCodeAnalysis.h"
 #endif
 
 #ifndef USING_ADDRESS_SANITISER
diff --git a/Engine/Source/Runtime/Core/Public/Microsoft/MicrosoftPlatformCodeAnalysis.h b/Engine/Source/Runtime/Core/Public/Microsoft/MicrosoftPlatformCodeAnalysis.h
new file mode 100644
--- /dev/null
+++ b/Engine/Source/Runtime/Core/Public/Microsoft/MicrosoftPlatformCodeAnalysis.h
@@ -0,0 +1,316 @@
+// Copyright Epic Games, Inc. All Rights Reserved.
+
+#pragma once
+
+
+// Code analysis features
+#if defined( _PREFAST_ ) || defined( PVS_STUDIO ) || defined(__clang_analyzer__)
+	#define USING_CODE_ANALYSIS 1
+#else
+	#define USING_CODE_ANALYSIS 0
+#endif
+
+//
+// NOTE: To suppress a single occurrence of a code analysis warning:
+//
+// 		CA_SUPPRESS( <WarningNumber> )
+//		...code that triggers warning...
+//
+
+//
+// NOTE: To disable all code analysis warnings for a section of code (such as include statements
+//       for a third party library), you can use the following:
+//
+// 		#if USING_CODE_ANALYSIS
+// 			MSVC_PRAGMA( warning( push ) )
+// 			MSVC_PRAGMA( warning( disable : ALL_CODE_ANALYSIS_WARNINGS ) )
+// 		#endif	// USING_CODE_ANALYSIS
+//
+//		<code with warnings>
+//
+// 		#if USING_CODE_ANALYSIS
+// 			MSVC_PRAGMA( warning( pop ) )
+// 		#endif	// USING_CODE_ANALYSIS
+//
+
+#if USING_CODE_ANALYSIS
+
+#if !defined(__clang_analyzer__)
+
+	// Input argument
+	// Example:  void SetValue( CA_IN bool bReadable );
+	#define CA_IN __in
+
+	// Output argument
+	// Example:  void FillValue( CA_OUT bool& bWriteable );
+	#define CA_OUT __out
+
+	// Specifies that a function parameter may only be read from, never written.
+	// NOTE: CA_READ_ONLY is inferred automatically if your parameter is has a const qualifier.
+	// Example:  void SetValue( CA_READ_ONLY bool bReadable );
+	#define CA_READ_ONLY [Pre(Access=Read)]
+
+	// Specifies that a function parameter may only be written to, never read.
+	// Example:  void FillValue( CA_WRITE_ONLY bool& bWriteable );
+	#define CA_WRITE_ONLY [Pre(Access=Write)]
+
+	// Incoming pointer parameter must not be NULL and must point to a valid location in memory.
+	// Place before a function parameter's type name.
+	// Example:  void SetPointer( CA_VALID_POINTER void* Pointer );
+	#define CA_VALID_POINTER [Pre(Null=No,Valid=Yes)]
+
+	// Caller must check the return value.  Place before the return value in a function declaration.
+	// Example:  CA_CHECK_RETVAL int32 GetNumber();
+	#define CA_CHECK_RETVAL [returnvalue:Post(MustCheck=Yes)]
+
+	// Function is expected to never return
+	#define CA_NO_RETURN __declspec(noreturn)
+
+	// Suppresses a warning for a single occurrence.  Should be used only for code analysis warnings on Windows platform!
+	#define CA_SUPPRESS( WarningNumber ) __pragma( warning( suppress: WarningNumber ) )
+
+	// Tells the code analysis engine to assume the statement to be true.  Useful for suppressing false positive warnings.
+	// NOTE: We use a double operator not here to avoid issues with passing certain class objects directly into __analysis_assume (which may cause a bogus compiler warning)
+	#define CA_ASSUME( Expr ) __analysis_assume( !!( Expr ) )
+
+	// Does a simple 'if (Condition)', but disables warnings about using constants in the condition.  Helps with some macro expansions.
+	#define CA_CONSTANT_IF(Condition) __pragma(warning(push)) __pragma(warning(disable:6326)) if (Condition) __pragma(warning(pop))
+
+
+	//
+	// Disable some code analysis warnings that we are NEVER interested in
+	//
+
+	// NOTE: Please be conservative about adding new suppressions here!  If you add a suppression, please
+	//       add a comment that explains the rationale.
+
+	// We don't use exceptions or care to gracefully handle _alloca() failure.  Also, we wrap _alloca in an
+	// appAlloca macro (not inline methods) and don't want to suppress at all call sites.
+	#pragma warning(disable : 6255) // warning C6255: _alloca indicates failure by raising a stack overflow exception.  Consider using _malloca instead.
+
+	// This a very common false positive warning (but some cases are not false positives!). Disabling for now so that we
+	// can more quickly see the benefits of static analysis on new code.
+	#pragma warning(disable : 6102) // warning C6102: Using 'variable' from failed function call at line 'line'.
+
+	// We use this exception handler in Windows code, it may be worth a closer look, but disabling for now
+	// so we can get the benefits of analysis on cross-platform code sooner.
+	#pragma warning(disable : 6320) // warning C6320: Exception-filter expression is the constant EXCEPTION_EXECUTE_HANDLER. This might mask exceptions that were not intended to be handled.
+
+	// Branching on constants allows us to ensure code paths compile even if they are inactive (eg. if(PLATFORM_FOO){ ... }), and is good practice that should not be discouraged.
+	#pragma warning(disable : 6326) // warning C6326 : Potential comparison of a constant with another constant.
+
+	// Likewise, expressions involving constants can also be valuable to check code paths compile.
+	#pragma warning(disable : 6240) // warning C6240 : (<expression> && <non-zero constant>) always evaluates to the result of <expression>. Did you intend to use the bitwise-and operator?
+
+	// This is a warning that a static array is exactly 365 members long.  We get false positives because our code gen can make arrays exactly this size.  
+	// EX: 365 properties in a UObject will auto gen a table exactly this long.
+	#pragma warning(disable : 6393) // warning C6393: A lookup table of size 365 is not sufficient to handle leap years.
+
+	//PVS-Studio settings:
+	//-V::505,542,581,601,623,668,677,690,704,719,720,730,735,751,758,780,781,1002,1008,1055,1062,1100
+	//-V:TRYCOMPRESSION:519,547
+	//-V:check(:501,547,560,605
+	//-V:checkf(:510
+	//-V:dtAssert(:568
+	//-V:rcAssert(:568
+	//-V:GET_FUNCTION_NAME_CHECKED:521
+	//-V:ENABLE_TEXT_ERROR_CHECKING_RESULTS:560
+	//-V:ENABLE_LOC_TESTING:560,617
+	//-V:WITH_EDITOR:560
+	//-V:UE_LOG_ACTIVE:560
+	//-V:verify:501
+	//-V:%n:609
+	//-V:UE_BUILD_SHIPPING:501
+	//-V:WITH_EDITOR:501
+	//-V:TestTrueExpr:501
+	//-V:PLATFORM_:517,547
+	//-V:ensureMsgf:562
+	//-V:WindowsMinorVersion:547
+	//-V:Import.XObject:547,560
+	//-V:MotionControllerComponent:547,560
+	//-V:AddUninitialized(sizeof(void*)/:514
+	//-V:TestTrue:678
+	//-V:SetViewTarget:678
+	//-V:Slot:607
+	//-V:RESIDENCY_CHECK_RESULT:607
+	//-V:bHitTesting:581
+	//-V:OptionalType:580 
+	//-V:GetNextNode:681
+	//-V:ConvertToAbsolutePathForExternalAppFor:524
+	//-V:CopySingleValue:524
+	//-V:bTimeLimitReached:560
+	//-V:bRedirectionAllowed:560
+	//-V:NumFailures:560
+	//-V:bAllowInstantToolTips:560
+	//-V:bIsRealTime:560
+	//-V:Position:519
+	//-V:DynamicParameterValue[ParameterIndex]:557
+	//-V:ViewIndex:557
+	//-V:DeviceIndex:557
+	//-V:Interpolation:560
+	//-V:storePortals:560
+	//-V:bDefaultShouldBeMaximized:560
+	//-V:bAllowPerfHUD:560
+	//-V:bUseClientStorage:560
+	//-V:bCalculateThisMapping:560
+	//-V:bDebugSelectedTaskOnly:560
+	//-V:bDebugSelectedTaskOnly:560
+	//-V:bIsPreview:560
+	//-V:bSupportsFastClear:560
+	//-V:bUseAPILibaries:560
+	//-V:bUseCachedBlobs:560
+	//-V:bWireframe:560
+	//-V:Num():560
+	//-V:PLATFORM_MAC:560
+	//-V:Particle->Size.Z:570
+	//-V:ComponentMaskParameter:601
+	//-V:Format(:601
+	//-V:SelectedEmitter:519
+	//-V:MAX_VERTS_PER_POLY:512
+	//-V:127:547
+	//-V:0x7F:547
+	//-V:WARN_COLOR:547
+	//-V:<<:614
+	//-V:FT_LOAD_TARGET_NORMAL:616
+	//-V:OPENGL_PERFORMANCE_DATA_INVALID:564
+	//-V:HLSLCC_VersionMajor:616
+	//-V:bIgnoreFieldReferences:519
+	//-V:CachedQueryInstance:519
+	//-V:MeshContext:519
+	//-V:bAffectedByMarquee:519
+	//-V:CopyCompleteValueFromScriptVM:524
+	//-V:OnStopWatchingPin:524
+	//-V:GetMinChildNodes:524
+	//-V:FromWorldMatrix:524
+	//-V:RemoveSelectedActorsFromSelectedLayer_CanExecute:524
+	//-V:NotifyLevelRemovedFromWorld:524
+	//-V:SPAWN_INIT:595
+	//-V:BEGIN_UPDATE_LOOP:595
+	//-V:OPENGL_PERFORMANCE_DATA_INVALID:560
+	//-V:bSkipTranslationTrack:560
+	//-V:NumSelected>0:581
+	//-V:bTryPerTrackBitwiseCompression:581
+	//-V:DataStripped:581
+	//-V:FromInt:601
+	//-V:UE_CLOG(:501,560
+	//-V:UE_LOG(:501,510,560
+	//-V:UGL_REQUIRED_VOID:501
+	//-V:AnimScriptInstance:595
+	//-V:Driver:595
+	//-V:PSceneAsync->lockWrite:595
+	//-V:Context.World():595
+	//-V:UNIT_LOG:595
+	//-V:ensure(:595
+	//-V:ALLOCATE_VERTEX_DATA_TEMPLATE:501
+	//-V:UGL_REQUIRED:501
+	//-V:DEBUG_LOG_HTTP:523
+	//-V:GIsEditor:560
+	//-V:bHasEditorToken:560
+	//-V:GEventDrivenLoaderEnabled:501
+	//-V:WALK_TO_CHARACTER:519
+	//-V:IMPLEMENT_AI_INSTANT_TEST:773
+	//-V:ENABLE_VERIFY_GL:564
+	//-V:INC_MEMORY_STAT_BY:568
+	//-V:DEC_MEMORY_STAT_BY:568
+	//-V:Key():568
+	//-V:Modify:762
+	//-V:GetTransitionList:762
+	//-V:Execute:768
+	//-V:LAUNCHERSERVICES_SHAREABLEPROJECTPATHS:768
+	//-V:SELECT_STATIC_MESH_VERTEX_TYPE:622
+	//-V:GET_FUNCTION_NAME_CHECKED:685
+	//-V:This(:678
+	//-V:state->error:649
+	//-V:ProjModifiers:616
+	//-V:PERF_DETAILED_PER_CLASS_GC_STATS:686
+	//-V:FMath:656
+	//-V:->*:607
+	//-V:GENERATED_UCLASS_BODY:764
+	//-V:CalcSegmentCostOnPoly:764
+	//-V:DrawLine:764
+	//-V:vrapi_SubmitFrame:641
+	//-V:VertexData:773
+	//-V:Linker:678
+	//-V:self:678
+	//-V:AccumulateParentID:678
+	//-V:FindChar:679
+
+	// The following classes retain a reference to data supplied in the constructor by the derived class which can not yet be initialized.
+	//-V:FMemoryWriter(:1050
+	//-V:FObjectWriter(:1050
+	//-V:FDurationTimer(:1050
+	//-V:FScopedDurationTimer(:1050
+	//-V:FQueryFastData(:1050
+
+	// Exclude all generated protobuf files
+	//V_EXCLUDE_PATH *.pb.cc
+
+	// warning V530: The return value of function 'HashCombine' is required to be utilized.
+	//-V::530
+
+	// warning V630: Instantiation of TRingBuffer < FString >: The 'Malloc' function is used to allocate memory for an array of objects which are classes containing constructors and destructors.
+	//-V::630
+
+	// warning V1043: A global object variable 'GSmallNumber' is declared in the header. Multiple copies of it will be created in all translation units that include this header file.
+	//-V::1043
+
+	// warning V1051: Consider checking for misprints. It's possible that the 'LayerInfo' should be checked here.
+	//-V::1051
+
+	// Disabling because incorrectly flagging all TStaticArrays
+	// V557: Array overrun is possible
+	//-V::557
+
+	// Disabling because too many virtuals currently in use in constructors/destructors, need to revist
+	// V1053: Calling the 'foo' virtual function in the constructor/destructor may lead to unexpected result at runtime.
+	//-V::1053
+
+#else // defined(__clang_analyzer__)
+
+// A fake function marked with noreturn that acts as a marker for CA_ASSUME to ensure the
+// static analyzer doesn't take an analysis path that is assumed not to be navigable.
+__declspec(dllimport, noreturn) void CA_AssumeNoReturn();
+
+// Input argument
+// Example:  void SetValue( CA_IN bool bReadable );
+#define CA_IN
+
+// Output argument
+// Example:  void FillValue( CA_OUT bool& bWriteable );
+#define CA_OUT
+
+// Specifies that a function parameter may only be read from, never written.
+// NOTE: CA_READ_ONLY is inferred automatically if your parameter is has a const qualifier.
+// Example:  void SetValue( CA_READ_ONLY bool bReadable );
+#define CA_READ_ONLY
+
+// Specifies that a function parameter may only be written to, never read.
+// Example:  void FillValue( CA_WRITE_ONLY bool& bWriteable );
+#define CA_WRITE_ONLY
+
+// Incoming pointer parameter must not be NULL and must point to a valid location in memory.
+// Place before a function parameter's type name.
+// Example:  void SetPointer( CA_VALID_POINTER void* Pointer );
+#define CA_VALID_POINTER
+
+// Caller must check the return value.  Place before the return value in a function declaration.
+// Example:  CA_CHECK_RETVAL int32 GetNumber();
+#define CA_CHECK_RETVAL
+
+// Function is expected to never return
+#define CA_NO_RETURN __declspec(noreturn)
+
+// Suppresses a warning for a single occurrence.  Should be used only for code analysis warnings on Windows platform!
+#define CA_SUPPRESS( WarningNumber )
+
+// Tells the code analysis engine to assume the statement to be true.  Useful for suppressing false positive warnings.
+#define CA_ASSUME( Expr )  (__builtin_expect(!bool(Expr), 0) ? CA_AssumeNoReturn() : (void)0)
+
+// Does a simple 'if (Condition)', but disables warnings about using constants in the condition.  Helps with some macro expansions.
+#define CA_CONSTANT_IF(Condition) if (Condition)
+
+
+#endif
+
+
+#endif
diff --git a/Engine/Source/Runtime/Core/Public/Windows/WindowsPlatformCodeAnalysis.h b/Engine/Source/Runtime/Core/Public/Windows/WindowsPlatformCodeAnalysis.h
--- a/Engine/Source/Runtime/Core/Public/Windows/WindowsPlatformCodeAnalysis.h
+++ b/Engine/Source/Runtime/Core/Public/Windows/WindowsPlatformCodeAnalysis.h
@@ -2,315 +2,7 @@
 
 #pragma once
 
-
-// Code analysis features
-#if defined( _PREFAST_ ) || defined( PVS_STUDIO ) || defined(__clang_analyzer__)
-	#define USING_CODE_ANALYSIS 1
-#else
-	#define USING_CODE_ANALYSIS 0
-#endif
-
-//
-// NOTE: To suppress a single occurrence of a code analysis warning:
-//
-// 		CA_SUPPRESS( <WarningNumber> )
-//		...code that triggers warning...
-//
-
-//
-// NOTE: To disable all code analysis warnings for a section of code (such as include statements
-//       for a third party library), you can use the following:
-//
-// 		#if USING_CODE_ANALYSIS
-// 			MSVC_PRAGMA( warning( push ) )
-// 			MSVC_PRAGMA( warning( disable : ALL_CODE_ANALYSIS_WARNINGS ) )
-// 		#endif	// USING_CODE_ANALYSIS
-//
-//		<code with warnings>
-//
-// 		#if USING_CODE_ANALYSIS
-// 			MSVC_PRAGMA( warning( pop ) )
-// 		#endif	// USING_CODE_ANALYSIS
-//
-
-#if USING_CODE_ANALYSIS
-
-#if !defined(__clang_analyzer__)
-
-	// Input argument
-	// Example:  void SetValue( CA_IN bool bReadable );
-	#define CA_IN __in
-
-	// Output argument
-	// Example:  void FillValue( CA_OUT bool& bWriteable );
-	#define CA_OUT __out
-
-	// Specifies that a function parameter may only be read from, never written.
-	// NOTE: CA_READ_ONLY is inferred automatically if your parameter is has a const qualifier.
-	// Example:  void SetValue( CA_READ_ONLY bool bReadable );
-	#define CA_READ_ONLY [Pre(Access=Read)]
-
-	// Specifies that a function parameter may only be written to, never read.
-	// Example:  void FillValue( CA_WRITE_ONLY bool& bWriteable );
-	#define CA_WRITE_ONLY [Pre(Access=Write)]
-
-	// Incoming pointer parameter must not be NULL and must point to a valid location in memory.
-	// Place before a function parameter's type name.
-	// Example:  void SetPointer( CA_VALID_POINTER void* Pointer );
-	#define CA_VALID_POINTER [Pre(Null=No,Valid=Yes)]
-
-	// Caller must check the return value.  Place before the return value in a function declaration.
-	// Example:  CA_CHECK_RETVAL int32 GetNumber();
-	#define CA_CHECK_RETVAL [returnvalue:Post(MustCheck=Yes)]
-
-	// Function is expected to never return
-	#define CA_NO_RETURN __declspec(noreturn)
-
-	// Suppresses a warning for a single occurrence.  Should be used only for code analysis warnings on Windows platform!
-	#define CA_SUPPRESS( WarningNumber ) __pragma( warning( suppress: WarningNumber ) )
-
-	// Tells the code analysis engine to assume the statement to be true.  Useful for suppressing false positive warnings.
-	// NOTE: We use a double operator not here to avoid issues with passing certain class objects directly into __analysis_assume (which may cause a bogus compiler warning)
-	#define CA_ASSUME( Expr ) __analysis_assume( !!( Expr ) )
-
-	// Does a simple 'if (Condition)', but disables warnings about using constants in the condition.  Helps with some macro expansions.
-	#define CA_CONSTANT_IF(Condition) __pragma(warning(push)) __pragma(warning(disable:6326)) if (Condition) __pragma(warning(pop))
-
-
-	//
-	// Disable some code analysis warnings that we are NEVER interested in
-	//
-
-	// NOTE: Please be conservative about adding new suppressions here!  If you add a suppression, please
-	//       add a comment that explains the rationale.
-
-	// We don't use exceptions or care to gracefully handle _alloca() failure.  Also, we wrap _alloca in an
-	// appAlloca macro (not inline methods) and don't want to suppress at all call sites.
-	#pragma warning(disable : 6255) // warning C6255: _alloca indicates failure by raising a stack overflow exception.  Consider using _malloca instead.
-
-	// This a very common false positive warning (but some cases are not false positives!). Disabling for now so that we
-	// can more quickly see the benefits of static analysis on new code.
-	#pragma warning(disable : 6102) // warning C6102: Using 'variable' from failed function call at line 'line'.
-
-	// We use this exception handler in Windows code, it may be worth a closer look, but disabling for now
-	// so we can get the benefits of analysis on cross-platform code sooner.
-	#pragma warning(disable : 6320) // warning C6320: Exception-filter expression is the constant EXCEPTION_EXECUTE_HANDLER. This might mask exceptions that were not intended to be handled.
-
-	// Branching on constants allows us to ensure code paths compile even if they are inactive (eg. if(PLATFORM_FOO){ ... }), and is good practice that should not be discouraged.
-	#pragma warning(disable : 6326) // warning C6326 : Potential comparison of a constant with another constant.
-
-	// Likewise, expressions involving constants can also be valuable to check code paths compile.
-	#pragma warning(disable : 6240) // warning C6240 : (<expression> && <non-zero constant>) always evaluates to the result of <expression>. Did you intend to use the bitwise-and operator?
-
-	// This is a warning that a static array is exactly 365 members long.  We get false positives because our code gen can make arrays exactly this size.  
-	// EX: 365 properties in a UObject will auto gen a table exactly this long.
-	#pragma warning(disable : 6393) // warning C6393: A lookup table of size 365 is not sufficient to handle leap years.
-
-	//PVS-Studio settings:
-	//-V::505,542,581,601,623,668,677,690,704,719,720,730,735,751,758,780,781,1002,1008,1055,1062,1100
-	//-V:TRYCOMPRESSION:519,547
-	//-V:check(:501,547,560,605
-	//-V:checkf(:510
-	//-V:dtAssert(:568
-	//-V:rcAssert(:568
-	//-V:GET_FUNCTION_NAME_CHECKED:521
-	//-V:ENABLE_TEXT_ERROR_CHECKING_RESULTS:560
-	//-V:ENABLE_LOC_TESTING:560,617
-	//-V:WITH_EDITOR:560
-	//-V:UE_LOG_ACTIVE:560
-	//-V:verify:501
-	//-V:%n:609
-	//-V:UE_BUILD_SHIPPING:501
-	//-V:WITH_EDITOR:501
-	//-V:TestTrueExpr:501
-	//-V:PLATFORM_:517,547
-	//-V:ensureMsgf:562
-	//-V:WindowsMinorVersion:547
-	//-V:Import.XObject:547,560
-	//-V:MotionControllerComponent:547,560
-	//-V:AddUninitialized(sizeof(void*)/:514
-	//-V:TestTrue:678
-	//-V:SetViewTarget:678
-	//-V:Slot:607
-	//-V:RESIDENCY_CHECK_RESULT:607
-	//-V:bHitTesting:581
-	//-V:OptionalType:580 
-	//-V:GetNextNode:681
-	//-V:ConvertToAbsolutePathForExternalAppFor:524
-	//-V:CopySingleValue:524
-	//-V:bTimeLimitReached:560
-	//-V:bRedirectionAllowed:560
-	//-V:NumFailures:560
-	//-V:bAllowInstantToolTips:560
-	//-V:bIsRealTime:560
-	//-V:Position:519
-	//-V:DynamicParameterValue[ParameterIndex]:557
-	//-V:ViewIndex:557
-	//-V:DeviceIndex:557
-	//-V:Interpolation:560
-	//-V:storePortals:560
-	//-V:bDefaultShouldBeMaximized:560
-	//-V:bAllowPerfHUD:560
-	//-V:bUseClientStorage:560
-	//-V:bCalculateThisMapping:560
-	//-V:bDebugSelectedTaskOnly:560
-	//-V:bDebugSelectedTaskOnly:560
-	//-V:bIsPreview:560
-	//-V:bSupportsFastClear:560
-	//-V:bUseAPILibaries:560
-	//-V:bUseCachedBlobs:560
-	//-V:bWireframe:560
-	//-V:Num():560
-	//-V:PLATFORM_MAC:560
-	//-V:Particle->Size.Z:570
-	//-V:ComponentMaskParameter:601
-	//-V:Format(:601
-	//-V:SelectedEmitter:519
-	//-V:MAX_VERTS_PER_POLY:512
-	//-V:127:547
-	//-V:0x7F:547
-	//-V:WARN_COLOR:547
-	//-V:<<:614
-	//-V:FT_LOAD_TARGET_NORMAL:616
-	//-V:OPENGL_PERFORMANCE_DATA_INVALID:564
-	//-V:HLSLCC_VersionMajor:616
-	//-V:bIgnoreFieldReferences:519
-	//-V:CachedQueryInstance:519
-	//-V:MeshContext:519
-	//-V:bAffectedByMarquee:519
-	//-V:CopyCompleteValueFromScriptVM:524
-	//-V:OnStopWatchingPin:524
-	//-V:GetMinChildNodes:524
-	//-V:FromWorldMatrix:524
-	//-V:RemoveSelectedActorsFromSelectedLayer_CanExecute:524
-	//-V:NotifyLevelRemovedFromWorld:524
-	//-V:SPAWN_INIT:595
-	//-V:BEGIN_UPDATE_LOOP:595
-	//-V:OPENGL_PERFORMANCE_DATA_INVALID:560
-	//-V:bSkipTranslationTrack:560
-	//-V:NumSelected>0:581
-	//-V:bTryPerTrackBitwiseCompression:581
-	//-V:DataStripped:581
-	//-V:FromInt:601
-	//-V:UE_CLOG(:501,560
-	//-V:UE_LOG(:501,510,560
-	//-V:UGL_REQUIRED_VOID:501
-	//-V:AnimScriptInstance:595
-	//-V:Driver:595
-	//-V:PSceneAsync->lockWrite:595
-	//-V:Context.World():595
-	//-V:UNIT_LOG:595
-	//-V:ensure(:595
-	//-V:ALLOCATE_VERTEX_DATA_TEMPLATE:501
-	//-V:UGL_REQUIRED:501
-	//-V:DEBUG_LOG_HTTP:523
-	//-V:GIsEditor:560
-	//-V:bHasEditorToken:560
-	//-V:GEventDrivenLoaderEnabled:501
-	//-V:WALK_TO_CHARACTER:519
-	//-V:IMPLEMENT_AI_INSTANT_TEST:773
-	//-V:ENABLE_VERIFY_GL:564
-	//-V:INC_MEMORY_STAT_BY:568
-	//-V:DEC_MEMORY_STAT_BY:568
-	//-V:Key():568
-	//-V:Modify:762
-	//-V:GetTransitionList:762
-	//-V:Execute:768
-	//-V:LAUNCHERSERVICES_SHAREABLEPROJECTPATHS:768
-	//-V:SELECT_STATIC_MESH_VERTEX_TYPE:622
-	//-V:GET_FUNCTION_NAME_CHECKED:685
-	//-V:This(:678
-	//-V:state->error:649
-	//-V:ProjModifiers:616
-	//-V:PERF_DETAILED_PER_CLASS_GC_STATS:686
-	//-V:FMath:656
-	//-V:->*:607
-	//-V:GENERATED_UCLASS_BODY:764
-	//-V:CalcSegmentCostOnPoly:764
-	//-V:DrawLine:764
-	//-V:vrapi_SubmitFrame:641
-	//-V:VertexData:773
-	//-V:Linker:678
-	//-V:self:678
-	//-V:AccumulateParentID:678
-	//-V:FindChar:679
-
-	// The following classes retain a reference to data supplied in the constructor by the derived class which can not yet be initialized.
-	//-V:FMemoryWriter(:1050
-	//-V:FObjectWriter(:1050
-	//-V:FDurationTimer(:1050
-	//-V:FScopedDurationTimer(:1050
-	//-V:FQueryFastData(:1050
-
-	// Exclude all generated protobuf files
-	//V_EXCLUDE_PATH *.pb.cc
-
-	// warning V530: The return value of function 'HashCombine' is required to be utilized.
-	//-V::530
-
-	// warning V630: Instantiation of TRingBuffer < FString >: The 'Malloc' function is used to allocate memory for an array of objects which are classes containing constructors and destructors.
-	//-V::630
-
-	// warning V1043: A global object variable 'GSmallNumber' is declared in the header. Multiple copies of it will be created in all translation units that include this header file.
-	//-V::1043
-
-	// warning V1051: Consider checking for misprints. It's possible that the 'LayerInfo' should be checked here.
-	//-V::1051
-
-	// Disabling because incorrectly flagging all TStaticArrays
-	// V557: Array overrun is possible
-	//-V::557
-
-	// Disabling because too many virtuals currently in use in constructors/destructors, need to revist
-	// V1053: Calling the 'foo' virtual function in the constructor/destructor may lead to unexpected result at runtime.
-	//-V::1053
-
-#else // defined(__clang_analyzer__)
-
-// A fake function marked with noreturn that acts as a marker for CA_ASSUME to ensure the
-// static analyzer doesn't take an analysis path that is assumed not to be navigable.
-__declspec(dllimport, noreturn) void CA_AssumeNoReturn();
-
-// Input argument
-// Example:  void SetValue( CA_IN bool bReadable );
-#define CA_IN
-
-// Output argument
-// Example:  void FillValue( CA_OUT bool& bWriteable );
-#define CA_OUT
-
-// Specifies that a function parameter may only be read from, never written.
-// NOTE: CA_READ_ONLY is inferred automatically if your parameter is has a const qualifier.
-// Example:  void SetValue( CA_READ_ONLY bool bReadable );
-#define CA_READ_ONLY
-
-// Specifies that a function parameter may only be written to, never read.
-// Example:  void FillValue( CA_WRITE_ONLY bool& bWriteable );
-#define CA_WRITE_ONLY
-
-// Incoming pointer parameter must not be NULL and must point to a valid location in memory.
-// Place before a function parameter's type name.
-// Example:  void SetPointer( CA_VALID_POINTER void* Pointer );
-#define CA_VALID_POINTER
-
-// Caller must check the return value.  Place before the return value in a function declaration.
-// Example:  CA_CHECK_RETVAL int32 GetNumber();
-#define CA_CHECK_RETVAL
-
-// Function is expected to never return
-#define CA_NO_RETURN __declspec(noreturn)
-
-// Suppresses a warning for a single occurrence.  Should be used only for code analysis warnings on Windows platform!
-#define CA_SUPPRESS( WarningNumber )
-
-// Tells the code analysis engine to assume the statement to be true.  Useful for suppressing false positive warnings.
-#define CA_ASSUME( Expr )  (__builtin_expect(!bool(Expr), 0) ? CA_AssumeNoReturn() : (void)0)
-
-// Does a simple 'if (Condition)', but disables warnings about using constants in the condition.  Helps with some macro expansions.
-#define CA_CONSTANT_IF(Condition) if (Condition)
-
-
-#endif
-
-
-#endif
+// Not marking as deprecated as eventually we'll need to properly do 
+// a per platform CodeAnalysis header, but for now just reroute to generic 
+// Microsoft header which includes windows specific elements
+#include "Microsoft/MicrosoftPlatformCodeAnalysis.h"
\ No newline at end of file
