Fixing -inifile not being considered in DoesConfigFileExistWrapper by checking OverrideFileFromCommandline, which prevented an ini file specified on the command line that wasn't in the cook from still being used.

OverrideFileFromCommandline now caches the results of parsing -inifile from the command line, instead of doing it every time the function was called.

Also added some more verbose logging that was provided by Josh.

#rnx

diff --git a/Engine/Source/Runtime/Core/Private/Misc/ConfigCacheIni.cpp b/Engine/Source/Runtime/Core/Private/Misc/ConfigCacheIni.cpp
--- a/Engine/Source/Runtime/Core/Private/Misc/ConfigCacheIni.cpp
+++ b/Engine/Source/Runtime/Core/Private/Misc/ConfigCacheIni.cpp
@@ -82,6 +82,8 @@ static int GDefaultReplayMethod = 2;
 static int GDefaultReplayMethod = 1;
 #endif
 
+static bool OverrideFileFromCommandline(FString & InOutFilename);
+
 /*-----------------------------------------------------------------------------
 FConfigValue
 -----------------------------------------------------------------------------*/
@@ -508,6 +510,12 @@ void FixupArrayOfStructKeysForSection(SectionType* Section, const FString& Secti
 		return true;
 	}
 
+	FString OverriddenIniFile = IniFile;
+	if (OverrideFileFromCommandline(OverriddenIniFile))
+	{
+		return true;
+	}
+
 	// So far, testing on cooked consoles, cooked desktop, and the Editor
 	// works fine.
 	// However, there was an issue where INIs wouldn't be found during cooking,
@@ -1273,16 +1281,27 @@ static bool OverrideFileFromCommandline(FString& InOutFilename)
 	//                under the same folder structure.
 	//          Ex1: D:\<some_folder>\Engine\Config\BaseEngine.ini
 	//			Ex2: D:\<some_folder>\QAGame\Config\Windows\WindowsEngine.ini
-	FString StagedFilePaths;
-	if(FParse::Value(FCommandLine::Get(), CommandlineOverrideSpecifiers::IniFileOverrideIdentifier, StagedFilePaths, false))
+	static bool bHasCachedData = false;
+
+	static TArray<FString> Files;
+	if (UNLIKELY(!bHasCachedData))
+	{
+		FString StagedFilePaths;
+		if (FParse::Value(FCommandLine::Get(), CommandlineOverrideSpecifiers::IniFileOverrideIdentifier, StagedFilePaths, false))
+		{
+			StagedFilePaths.ParseIntoArray(Files, TEXT(","), true);
+		}
+
+		bHasCachedData = true;
+	}
+
+	if (Files.Num() > 0)
 	{
 		FString RelativePath = InOutFilename;
 		if (FPaths::IsUnderDirectory(RelativePath, FPaths::RootDir()))
 		{
 			FPaths::MakePathRelativeTo(RelativePath, *FPaths::RootDir());
 
-			TArray<FString> Files;
-			StagedFilePaths.ParseIntoArray(Files, TEXT(","), true);
 			for (int32 Index = 0; Index < Files.Num(); Index++)
 			{
 				FString NormalizedOverride = Files[Index];
@@ -6282,6 +6301,13 @@ void FConfigCacheIni::AddPluginToBranches(FName PluginName, FConfigModificationT
 	IFileManager::Get().FindFiles(PluginConfigs, *PluginConfigDir, TEXT("ini"));
 	IFileManager::Get().FindFiles(PluginConfigs, *PlatformConfigDir, TEXT("ini"));
 	
+#if !UE_BUILD_SHIPPING
+	for (const FString& F : PluginConfigs)
+	{
+		UE_LOG(LogConfig, Verbose, TEXT("Found config file %s in plugin dir %s"), *F, *PluginInfo->PluginDir);
+	}
+#endif // !UE_BUILD_SHIPPING
+	
 	// if this plugin has any platform extensions, then we need to look in them for files, in so that we can load them
 	// even if there is no platform-less config file in the plugin itself
 	for (FString& ChildPluginDir : PluginInfo->ChildPluginDirs)
@@ -6320,6 +6346,7 @@ void FConfigCacheIni::AddPluginToBranches(FName PluginName, FConfigModificationT
 			FString FullPath = FPaths::Combine(PluginConfigDir, ConfigFilename);
 			if (ModificationTracker->LoadedFiles.Contains(FullPath))
 			{
+				UE_LOG(LogConfig, Verbose, TEXT("Skipping already loaded file %s"), *FullPath);
 				continue;
 			}
 		}
