autortfm fixes:

- inside ResolveName2, move the LoadPackage call into the Open. Aborting should leave the package orphaned and cleaned up by the GC
- PlatformStackWalk: resolve symbols in the open because they rely on disallowed platform syscalls or locks
- PhysicsTrace: move the call to PhysicsTraceChannel/PhysicsTraceProfile/PhysicsTraceObject to transaction commit time. These calls register some async work to be performed and consumed later, so deferring these should be OK

diff --git a/Engine/Source/Runtime/Core/Private/Unix/UnixPlatformStackWalk.cpp b/Engine/Source/Runtime/Core/Private/Unix/UnixPlatformStackWalk.cpp
--- a/Engine/Source/Runtime/Core/Private/Unix/UnixPlatformStackWalk.cpp
+++ b/Engine/Source/Runtime/Core/Private/Unix/UnixPlatformStackWalk.cpp
@@ -15,6 +15,7 @@
 #include "HAL/ExceptionHandling.h"
 #include "HAL/PlatformProcess.h"
 #include "HAL/PlatformTime.h"
+#include "AutoRTFM/AutoRTFM.h"
 
 #include <link.h>
 #include <signal.h>
@@ -491,7 +492,10 @@ namespace
 
 void FUnixPlatformStackWalk::ProgramCounterToSymbolInfo( uint64 ProgramCounter, FProgramCounterSymbolInfo& out_SymbolInfo )
 {
-	PopulateProgramCounterSymbolInfoFromSymbolFile(ProgramCounter, out_SymbolInfo);
+	UE_AUTORTFM_OPEN2
+	{
+		PopulateProgramCounterSymbolInfoFromSymbolFile(ProgramCounter, out_SymbolInfo);
+	};
 }
 
 bool FUnixPlatformStackWalk::ProgramCounterToHumanReadableString( int32 CurrentCallDepth, uint64 ProgramCounter, ANSICHAR* HumanReadableString, SIZE_T HumanReadableStringSize, FGenericCrashContext* Context )
diff --git a/Engine/Source/Runtime/Core/Private/Windows/WindowsPlatformStackWalk.cpp b/Engine/Source/Runtime/Core/Private/Windows/WindowsPlatformStackWalk.cpp
--- a/Engine/Source/Runtime/Core/Private/Windows/WindowsPlatformStackWalk.cpp
+++ b/Engine/Source/Runtime/Core/Private/Windows/WindowsPlatformStackWalk.cpp
@@ -458,74 +458,77 @@ uint32 FWindowsPlatformStackWalk::CaptureStackBackTrace( uint64* BackTrace, uint
 
 void FWindowsPlatformStackWalk::ProgramCounterToSymbolInfo( uint64 ProgramCounter, FProgramCounterSymbolInfo& out_SymbolInfo )
 {
-	UE::TUniqueLock GlobalLock(GStackWalkingLock);
+	UE_AUTORTFM_OPEN2
+	{
+		UE::TUniqueLock GlobalLock(GStackWalkingLock);
 
-	// Initialize stack walking as it loads up symbol information which we require.
-	InitStackWalking();
+		// Initialize stack walking as it loads up symbol information which we require.
+		InitStackWalking();
 
 #if ON_DEMAND_SYMBOL_LOADING
-	// Load symbols for the module
-	bool bShouldReloadModuleMissingDebugSymbols = !FPlatformProperties::IsMonolithicBuild() && FPlatformStackWalk::WantsDetailedCallstacksInNonMonolithicBuilds();
-	LoadSymbolsForModuleByAddress(ProgramCounter, GetSymbolSearchPath(), bShouldReloadModuleMissingDebugSymbols);
+		// Load symbols for the module
+		bool bShouldReloadModuleMissingDebugSymbols = !FPlatformProperties::IsMonolithicBuild() && FPlatformStackWalk::WantsDetailedCallstacksInNonMonolithicBuilds();
+		LoadSymbolsForModuleByAddress(ProgramCounter, GetSymbolSearchPath(), bShouldReloadModuleMissingDebugSymbols);
 #endif
 
-	// Set the program counter.
-	out_SymbolInfo.ProgramCounter = ProgramCounter;
+		// Set the program counter.
+		out_SymbolInfo.ProgramCounter = ProgramCounter;
 
-	uint32 LastError = 0;
-	HANDLE ProcessHandle = GProcessHandle;
+		uint32 LastError = 0;
+		HANDLE ProcessHandle = GProcessHandle;
 
-	// Initialize symbol.
-	ANSICHAR SymbolBuffer[sizeof( SYMBOL_INFO ) + FProgramCounterSymbolInfo::MAX_NAME_LENGTH] = {0};
-	SYMBOL_INFO* Symbol = (SYMBOL_INFO*)SymbolBuffer;
-	Symbol->SizeOfStruct = sizeof(SYMBOL_INFO);
-	Symbol->MaxNameLen = FProgramCounterSymbolInfo::MAX_NAME_LENGTH;
+		// Initialize symbol.
+		ANSICHAR SymbolBuffer[sizeof( SYMBOL_INFO ) + FProgramCounterSymbolInfo::MAX_NAME_LENGTH] = {0};
+		SYMBOL_INFO* Symbol = (SYMBOL_INFO*)SymbolBuffer;
+		Symbol->SizeOfStruct = sizeof(SYMBOL_INFO);
+		Symbol->MaxNameLen = FProgramCounterSymbolInfo::MAX_NAME_LENGTH;
 
-	// Get function name.
-	if( SymFromAddr( ProcessHandle, ProgramCounter, nullptr, Symbol ) )
-	{
-		// Skip any funky chars in the beginning of a function name.
-		int32 Offset = 0;
-		while( Symbol->Name[Offset] < 32 || Symbol->Name[Offset] > 127 )
+		// Get function name.
+		if( SymFromAddr( ProcessHandle, ProgramCounter, nullptr, Symbol ) )
 		{
-			Offset++;
-		}
+			// Skip any funky chars in the beginning of a function name.
+			int32 Offset = 0;
+			while( Symbol->Name[Offset] < 32 || Symbol->Name[Offset] > 127 )
+			{
+				Offset++;
+			}
 
-		// Write out function name.
-		FCStringAnsi::Strncpy( out_SymbolInfo.FunctionName, Symbol->Name + Offset, FProgramCounterSymbolInfo::MAX_NAME_LENGTH ); 
-		FCStringAnsi::Strncat( out_SymbolInfo.FunctionName, "()", FProgramCounterSymbolInfo::MAX_NAME_LENGTH );
-	}
-	else
-	{
-		// No symbol found for this address.
-		LastError = GetLastError();
-	}
+			// Write out function name.
+			FCStringAnsi::Strncpy( out_SymbolInfo.FunctionName, Symbol->Name + Offset, FProgramCounterSymbolInfo::MAX_NAME_LENGTH ); 
+			FCStringAnsi::Strncat( out_SymbolInfo.FunctionName, "()", FProgramCounterSymbolInfo::MAX_NAME_LENGTH );
+		}
+		else
+		{
+			// No symbol found for this address.
+			LastError = GetLastError();
+		}
 
-	// Get filename and line number.
-	IMAGEHLP_LINE64	ImageHelpLine = {0};
-	ImageHelpLine.SizeOfStruct = sizeof( ImageHelpLine );
-	if( SymGetLineFromAddr64( ProcessHandle, ProgramCounter, (::DWORD *)&out_SymbolInfo.SymbolDisplacement, &ImageHelpLine ) )
-	{
-		FCStringAnsi::Strncpy( out_SymbolInfo.Filename, ImageHelpLine.FileName, FProgramCounterSymbolInfo::MAX_NAME_LENGTH );
-		out_SymbolInfo.LineNumber = ImageHelpLine.LineNumber;
-	}
-	else
-	{
-		LastError = GetLastError();
-	}
+		// Get filename and line number.
+		IMAGEHLP_LINE64	ImageHelpLine = {0};
+		ImageHelpLine.SizeOfStruct = sizeof( ImageHelpLine );
+		if( SymGetLineFromAddr64( ProcessHandle, ProgramCounter, (::DWORD *)&out_SymbolInfo.SymbolDisplacement, &ImageHelpLine ) )
+		{
+			FCStringAnsi::Strncpy( out_SymbolInfo.Filename, ImageHelpLine.FileName, FProgramCounterSymbolInfo::MAX_NAME_LENGTH );
+			out_SymbolInfo.LineNumber = ImageHelpLine.LineNumber;
+		}
+		else
+		{
+			LastError = GetLastError();
+		}
 
-	// Get module name.
-	IMAGEHLP_MODULE64 ImageHelpModule = {0};
-	ImageHelpModule.SizeOfStruct = sizeof( ImageHelpModule );
-	if( SymGetModuleInfo64( ProcessHandle, ProgramCounter, &ImageHelpModule) )
-	{
-		// Write out module information.
-		FCStringAnsi::Strncpy( out_SymbolInfo.ModuleName, ImageHelpModule.ImageName, FProgramCounterSymbolInfo::MAX_NAME_LENGTH );
-	}
-	else
-	{
-		LastError = GetLastError();
-	}
+		// Get module name.
+		IMAGEHLP_MODULE64 ImageHelpModule = {0};
+		ImageHelpModule.SizeOfStruct = sizeof( ImageHelpModule );
+		if( SymGetModuleInfo64( ProcessHandle, ProgramCounter, &ImageHelpModule) )
+		{
+			// Write out module information.
+			FCStringAnsi::Strncpy( out_SymbolInfo.ModuleName, ImageHelpModule.ImageName, FProgramCounterSymbolInfo::MAX_NAME_LENGTH );
+		}
+		else
+		{
+			LastError = GetLastError();
+		}
+	};
 }
 
 void FWindowsPlatformStackWalk::ProgramCounterToSymbolInfoEx(uint64 ProgramCounter, FProgramCounterSymbolInfoEx& out_SymbolInfo)
diff --git a/Engine/Source/Runtime/CoreUObject/Private/UObject/UObjectGlobals.cpp b/Engine/Source/Runtime/CoreUObject/Private/UObject/UObjectGlobals.cpp
--- a/Engine/Source/Runtime/CoreUObject/Private/UObject/UObjectGlobals.cpp
+++ b/Engine/Source/Runtime/CoreUObject/Private/UObject/UObjectGlobals.cpp
@@ -1265,17 +1265,25 @@ bool ResolveName2(UObject*& InPackage, FStringBuilderBase& InOutName, bool Creat
 		{
 			// Try to find the package in memory first, should be faster than attempting to load or create
 			InPackage = InPackage ? nullptr : StaticFindObjectFast(UPackage::StaticClass(), InPackage, *PartialName);
+			
 			if (!bIsScriptPackage && !InPackage)
 			{
-				InPackage = LoadPackage(Cast<UPackage>(InPackage), *PartialName, LoadFlags, nullptr, InstancingContext);
+				UE_AUTORTFM_OPEN2
+				{
+					InPackage = LoadPackage(Cast<UPackage>(InPackage), *PartialName, LoadFlags, nullptr, InstancingContext);
+				};
 			}
+
 			if (!InPackage)
 			{
-				InPackage = CreatePackage(*PartialName);
-				if (bIsScriptPackage)
+				UE_AUTORTFM_OPEN2
 				{
-					Cast<UPackage>(InPackage)->SetPackageFlags(PKG_CompiledIn);
-				}
+					InPackage = CreatePackage(*PartialName);
+					if (bIsScriptPackage)
+					{
+						Cast<UPackage>(InPackage)->SetPackageFlags(PKG_CompiledIn);
+					}
+				};
 			}
 
 			check(InPackage);
