It is now possible to poll the virtualization interface for the name of the currently enabled system implementation (if any)

#rnx

diff --git a/Engine/Source/Runtime/Core/Private/Virtualization/VirtualizationSystem.cpp b/Engine/Source/Runtime/Core/Private/Virtualization/VirtualizationSystem.cpp
--- a/Engine/Source/Runtime/Core/Private/Virtualization/VirtualizationSystem.cpp
+++ b/Engine/Source/Runtime/Core/Private/Virtualization/VirtualizationSystem.cpp
@@ -166,6 +166,7 @@ public:
 };
 
 TUniquePtr<IVirtualizationSystem> GVirtualizationSystem = nullptr;
+FName GVirtualizationSystemName;
 
 /**
  * Utility to check if either cmd is present in the command line. Useful when transitioning from one
@@ -329,6 +330,7 @@ void Initialize(const FInitParams& InitParams, EInitializationFlags Flags)
 			{
 				check(!IVirtualizationSystem::IsInitialized());
 				GVirtualizationSystem = MoveTemp(NewSystem);
+				GVirtualizationSystemName = SystemName;
 			}
 			else
 			{
@@ -377,6 +379,11 @@ bool IVirtualizationSystem::IsInitialized()
 	return GVirtualizationSystem != nullptr;
 }
 
+FName IVirtualizationSystem::GetSystemName()
+{
+	return GVirtualizationSystemName;
+}
+
 IVirtualizationSystem& IVirtualizationSystem::Get()
 {
 	// For now allow Initialize to be called directly if it was not called explicitly.
diff --git a/Engine/Source/Runtime/Core/Public/Virtualization/VirtualizationSystem.h b/Engine/Source/Runtime/Core/Public/Virtualization/VirtualizationSystem.h
--- a/Engine/Source/Runtime/Core/Public/Virtualization/VirtualizationSystem.h
+++ b/Engine/Source/Runtime/Core/Public/Virtualization/VirtualizationSystem.h
@@ -605,6 +605,9 @@ public:
 	/** Returns true if a virtualization system has been initialized and false if not */
 	CORE_API static bool IsInitialized();
 
+	/** Returns the name of the currently initialized system, will be NAME_None if initialization failed or the system is not enabled */
+	CORE_API static FName GetSystemName();
+
 	/** 
 	 * Gain access to the current virtualization system active for the project. If the system has not yet been 
 	 * initialized then calling this method will initialize it.
