Cache Network Connection Type on iOS

diff --git a/Engine/Source/Runtime/Core/Private/IOS/IOSPlatformMisc.cpp b/Engine/Source/Runtime/Core/Private/IOS/IOSPlatformMisc.cpp
--- a/Engine/Source/Runtime/Core/Private/IOS/IOSPlatformMisc.cpp
+++ b/Engine/Source/Runtime/Core/Private/IOS/IOSPlatformMisc.cpp
@@ -1545,7 +1545,7 @@ void FIOSPlatformMisc::GetValidTargetPlatforms(TArray<FString>& TargetPlatformNa
 #endif
 }
 
-ENetworkConnectionType FIOSPlatformMisc::GetNetworkConnectionType()
+static inline ENetworkConnectionType CheckNetworkConnectionType()
 {
 	struct sockaddr_in ZeroAddress;
 	FMemory::Memzero(&ZeroAddress, sizeof(ZeroAddress));
@@ -1587,6 +1587,24 @@ ENetworkConnectionType FIOSPlatformMisc::GetNetworkConnectionType()
     return ENetworkConnectionType::None;
 }
 
+ENetworkConnectionType FIOSPlatformMisc::GetNetworkConnectionType()
+{
+	static TOptional<ENetworkConnectionType> ConnectionType = {};
+	static double LastCheckTime = 0;
+
+	const double CurrentTime = FPlatformTime::Seconds();
+	const double CheckInterval = 0.2;
+
+	if (!ConnectionType.IsSet() || CurrentTime >= LastCheckTime + CheckInterval)
+	{
+		ConnectionType = CheckNetworkConnectionType();
+		LastCheckTime = CurrentTime;
+	}
+
+	ensure(ConnectionType.IsSet());
+	return ConnectionType.GetValue();
+}
+
 bool FIOSPlatformMisc::HasActiveWiFiConnection()
 {
     return GetNetworkConnectionType() == ENetworkConnectionType::WiFi;
