MonitoredProcess: Only pass write end of stdout pipe to createprocess in FMonitoredProcess launch

This fixes issues with subprocesses that use stdin or that pass inherited pipes to further subprocesses

diff --git a/Engine/Source/Runtime/Core/Private/Misc/MonitoredProcess.cpp b/Engine/Source/Runtime/Core/Private/Misc/MonitoredProcess.cpp
--- a/Engine/Source/Runtime/Core/Private/Misc/MonitoredProcess.cpp
+++ b/Engine/Source/Runtime/Core/Private/Misc/MonitoredProcess.cpp
@@ -72,7 +72,7 @@ bool FMonitoredProcess::Launch()
 		return false;
 	}
 
-	ProcessHandle = FPlatformProcess::CreateProc(*URL, *Params, false, Hidden, Hidden, nullptr, 0, *WorkingDir, WritePipe, ReadPipe);
+	ProcessHandle = FPlatformProcess::CreateProc(*URL, *Params, false, Hidden, Hidden, nullptr, 0, *WorkingDir, WritePipe, nullptr);
 
 	if (!ProcessHandle.IsValid())
 	{
diff --git a/Engine/Source/Runtime/Core/Public/Misc/MonitoredProcess.h b/Engine/Source/Runtime/Core/Public/Misc/MonitoredProcess.h
--- a/Engine/Source/Runtime/Core/Public/Misc/MonitoredProcess.h
+++ b/Engine/Source/Runtime/Core/Public/Misc/MonitoredProcess.h
@@ -229,7 +229,7 @@ protected:
 	// Holds the handle to the process. */
 	FProcHandle ProcessHandle;
 
-	// Holds the read pipe. */
+	// Holds the non-inheritable read end of stdout pipe. */
 	void* ReadPipe = nullptr;
 
 	// Holds the return code. */
@@ -250,7 +250,7 @@ protected:
 	// Holds the URL of the working dir for the process. */
 	FString WorkingDir;
 
-	// Holds the write pipe. */
+	// Holds the inheritable write end of stdout pipe. */
 	void* WritePipe = nullptr;
 
 	// Holds if we should create pipes
