Allow direct Utf8 base64 encoding/decoding

diff --git a/Engine/Source/Runtime/Core/Private/Misc/Base64.cpp b/Engine/Source/Runtime/Core/Private/Misc/Base64.cpp
--- a/Engine/Source/Runtime/Core/Private/Misc/Base64.cpp
+++ b/Engine/Source/Runtime/Core/Private/Misc/Base64.cpp
@@ -119,13 +119,13 @@ template<typename CharType> uint32 FBase64::Encode(const uint8* Source, uint32 L
 		}
 		else
 		{
-			EncodedBytes[2] = EncodingAlphabet[ByteTriplet & 0x3F];
+			EncodedBytes[2] = (CharType)EncodingAlphabet[ByteTriplet & 0x3F];
 		}
 		// Now encode the remaining bits the same way
 		ByteTriplet >>= 6;
-		EncodedBytes[1] = EncodingAlphabet[ByteTriplet & 0x3F];
+		EncodedBytes[1] = (CharType)EncodingAlphabet[ByteTriplet & 0x3F];
 		ByteTriplet >>= 6;
-		EncodedBytes[0] = EncodingAlphabet[ByteTriplet & 0x3F];
+		EncodedBytes[0] = (CharType)EncodingAlphabet[ByteTriplet & 0x3F];
 	}
 	
 	Length -= PaddingLength;
@@ -144,13 +144,13 @@ template<typename CharType> uint32 FBase64::Encode(const uint8* Source, uint32 L
 		uint32 ByteTriplet = A << 16 | B << 8 | C;
 
 		// Use the 6bit block to find the representation ascii character for it
-		EncodedBytes[3] = EncodingAlphabet[ByteTriplet & 0x3F];
+		EncodedBytes[3] = (CharType)EncodingAlphabet[ByteTriplet & 0x3F];
 		ByteTriplet >>= 6;
-		EncodedBytes[2] = EncodingAlphabet[ByteTriplet & 0x3F];
+		EncodedBytes[2] = (CharType)EncodingAlphabet[ByteTriplet & 0x3F];
 		ByteTriplet >>= 6;
-		EncodedBytes[1] = EncodingAlphabet[ByteTriplet & 0x3F];
+		EncodedBytes[1] = (CharType)EncodingAlphabet[ByteTriplet & 0x3F];
 		ByteTriplet >>= 6;
-		EncodedBytes[0] = EncodingAlphabet[ByteTriplet & 0x3F];
+		EncodedBytes[0] = (CharType)EncodingAlphabet[ByteTriplet & 0x3F];
 	}
 
 	verify(EncodedBytes == Dest);
@@ -160,6 +160,7 @@ template<typename CharType> uint32 FBase64::Encode(const uint8* Source, uint32 L
 
 template CORE_API uint32 FBase64::Encode<ANSICHAR>(const uint8* Source, uint32 Length, ANSICHAR* Dest, EBase64Mode Mode);
 template CORE_API uint32 FBase64::Encode<WIDECHAR>(const uint8* Source, uint32 Length, WIDECHAR* Dest, EBase64Mode Mode);
+template CORE_API uint32 FBase64::Encode<UTF8CHAR>(const uint8* Source, uint32 Length, UTF8CHAR* Dest, EBase64Mode Mode);
 
 bool FBase64::Decode(const FString& Source, FString& OutDest, EBase64Mode Mode)
 {
@@ -181,7 +182,8 @@ bool FBase64::Decode(const FString& Source, TArray<uint8>& OutDest, EBase64Mode
 	return Decode(*Source, Source.Len(), OutDest.GetData(), Mode);
 }
 
-template<typename CharType> bool FBase64::Decode(const CharType* Source, uint32 Length, uint8* Dest, EBase64Mode Mode)
+template<typename CharType> 
+bool FBase64::Decode(const CharType* Source, uint32 Length, uint8* Dest, EBase64Mode Mode)
 {
 	check(Mode == EBase64Mode::Standard || Mode == EBase64Mode::UrlSafe);
 
@@ -268,6 +270,7 @@ template<typename CharType> bool FBase64::Decode(const CharType* Source, uint32
 
 template CORE_API bool FBase64::Decode<ANSICHAR>(const ANSICHAR* Source, uint32 Length, uint8* Dest, EBase64Mode Mode);
 template CORE_API bool FBase64::Decode<WIDECHAR>(const WIDECHAR* Source, uint32 Length, uint8* Dest, EBase64Mode Mode);
+template CORE_API bool FBase64::Decode<UTF8CHAR>(const UTF8CHAR* Source, uint32 Length, uint8* Dest, EBase64Mode Mode);
 
 uint32 FBase64::GetDecodedDataSize(const FString& Source)
 {
@@ -301,3 +304,4 @@ template<typename CharType> uint32 FBase64::GetDecodedDataSize(const CharType* S
 
 template CORE_API uint32 FBase64::GetDecodedDataSize(const ANSICHAR* Source, uint32 Length);
 template CORE_API uint32 FBase64::GetDecodedDataSize(const WIDECHAR* Source, uint32 Length);
+template CORE_API uint32 FBase64::GetDecodedDataSize(const UTF8CHAR* Source, uint32 Length);
