Add missing math functions to stdlib

diff --git a/Engine/Source/Runtime/Core/Private/AutoRTFM/StdLibSupport.cpp b/Engine/Source/Runtime/Core/Private/AutoRTFM/StdLibSupport.cpp
--- a/Engine/Source/Runtime/Core/Private/AutoRTFM/StdLibSupport.cpp
+++ b/Engine/Source/Runtime/Core/Private/AutoRTFM/StdLibSupport.cpp
@@ -235,6 +235,7 @@ UE_AUTORTFM_REGISTER_SELF_FUNCTION(log);
 UE_AUTORTFM_REGISTER_SELF_FUNCTION(pow);
 UE_AUTORTFM_REGISTER_SELF_FUNCTION(llrint);
 UE_AUTORTFM_REGISTER_SELF_FUNCTION(fmod);
+UE_AUTORTFM_REGISTER_SELF_FUNCTION(modf);
 // Linux (likely Mac) have ambiguous overrides to these math functions
 #else
 UE_AUTORTFM_REGISTER_SELF_FUNCTION(static_cast<float(*)(float)>(&sqrt));
@@ -275,6 +276,8 @@ UE_AUTORTFM_REGISTER_SELF_FUNCTION(static_cast<long long(*)(float)>(&llrint));
 UE_AUTORTFM_REGISTER_SELF_FUNCTION(static_cast<long long(*)(long double)>(&llrint));
 UE_AUTORTFM_REGISTER_SELF_FUNCTION(static_cast<float(*)(float, float)>(&fmod));
 UE_AUTORTFM_REGISTER_SELF_FUNCTION(static_cast<long double(*)(long double, long double)>(&fmod));
+UE_AUTORTFM_REGISTER_SELF_FUNCTION(static_cast<float(*)(float, float*)>(&modf));
+UE_AUTORTFM_REGISTER_SELF_FUNCTION(static_cast<long double(*)(long double, long double*)>(&modf));
 #endif
 
 // Self register Math functions
@@ -299,6 +302,7 @@ UE_AUTORTFM_REGISTER_SELF_FUNCTION(llrintf);
 UE_AUTORTFM_REGISTER_SELF_FUNCTION(fmodf);
 UE_AUTORTFM_REGISTER_SELF_FUNCTION(fmodl);
 UE_AUTORTFM_REGISTER_SELF_FUNCTION(rand);
+UE_AUTORTFM_REGISTER_SELF_FUNCTION(modff);
 
 // FIXME: Does not currently support %n format specifiers.
 int RTFM_vsnprintf(char* Str, size_t Size, const char* Format, va_list ArgList)
