Fixed LLM for FCachedOSVeryLargePageAllocator

diff --git a/Engine/Source/Runtime/Core/Private/HAL/Allocators/CachedOSPageAllocator.cpp b/Engine/Source/Runtime/Core/Private/HAL/Allocators/CachedOSPageAllocator.cpp
--- a/Engine/Source/Runtime/Core/Private/HAL/Allocators/CachedOSPageAllocator.cpp
+++ b/Engine/Source/Runtime/Core/Private/HAL/Allocators/CachedOSPageAllocator.cpp
@@ -55,20 +55,18 @@ void* FCachedOSPageAllocator::AllocateImpl(SIZE_T Size, uint32 CachedByteLimit,
 				return Result;
 			}
 
-		{
-			LLM_PLATFORM_SCOPE(ELLMTag::FMalloc);
-			void* Ptr = nullptr;
 			{
 #if UE_ALLOW_OSMEMORYLOCKFREE
 				FScopeUnlock ScopeUnlock(Mutex);
 #endif
-				Ptr = FPlatformMemory::BinnedAllocFromOS(Size);
-			}
-			if(Ptr)
-			{
-				return Ptr;
+				LLM_PLATFORM_SCOPE(ELLMTag::FMalloc);
+				void* Ptr = FPlatformMemory::BinnedAllocFromOS(Size);
+				if (Ptr)
+				{
+					return Ptr;
+				}
 			}
-		}
+
 			// Are we holding on to much mem? Release it all.
 			for (FFreePageBlock* Block = First; Block != Last; ++Block)
 			{
@@ -81,13 +79,12 @@ void* FCachedOSPageAllocator::AllocateImpl(SIZE_T Size, uint32 CachedByteLimit,
 		}
 	}
 
-	LLM_PLATFORM_SCOPE(ELLMTag::FMalloc);
-
 	void* Ret = nullptr;
 	{
 #if UE_ALLOW_OSMEMORYLOCKFREE
 		FScopeUnlock ScopeUnlock(Mutex);
 #endif
+		LLM_PLATFORM_SCOPE(ELLMTag::FMalloc);
 		Ret = FPlatformMemory::BinnedAllocFromOS(Size);
 	}
 	return Ret;
diff --git a/Engine/Source/Runtime/Core/Private/HAL/Allocators/CachedOSVeryLargePageAllocator.cpp b/Engine/Source/Runtime/Core/Private/HAL/Allocators/CachedOSVeryLargePageAllocator.cpp
--- a/Engine/Source/Runtime/Core/Private/HAL/Allocators/CachedOSVeryLargePageAllocator.cpp
+++ b/Engine/Source/Runtime/Core/Private/HAL/Allocators/CachedOSVeryLargePageAllocator.cpp
@@ -74,7 +74,7 @@ static volatile int32 GLargePageAllocatorDecommitCount = 0;
 
 void FCachedOSVeryLargePageAllocator::Init()
 {
-	Block = FPlatformMemory::FPlatformVirtualMemoryBlock::AllocateVirtual(AddressSpaceToReserve);
+	Block = FPlatformMemory::FPlatformVirtualMemoryBlock::AllocateVirtual(AddressSpaceToReserve, SizeOfLargePage);
 	AddressSpaceReserved = (uintptr_t)Block.GetVirtualPointer();
 	AddressSpaceReservedEnd = AddressSpaceReserved + AddressSpaceToReserve;
 	AddressSpaceReservedEndSmallPool = AddressSpaceReserved + (AddressSpaceToReserve / 2);
@@ -276,6 +276,7 @@ FCachedOSVeryLargePageAllocator::FLargePage* FCachedOSVeryLargePageAllocator::Al
 				bOutCommitFailure = true;
 				return nullptr;
 			}
+			LLM_IF_ENABLED(FLowLevelMemTracker::Get().OnLowLevelAlloc(ELLMTracker::Platform, (void*)LargePage->BaseAddress, SizeOfLargePage));
 			// A new large page has been created. Add it to CachedFree counter
 			CachedFree += SizeOfLargePage;
 			CommittedLargePagesCount[AllocationHint] += 1;
@@ -317,6 +318,7 @@ void FCachedOSVeryLargePageAllocator::Free(void* Ptr, SIZE_T Size, FCriticalSect
 					FScopeUnlock ScopeUnlock(Mutex);
 #endif
 					Block.Decommit(LargePage->BaseAddress - AddressSpaceReserved, SizeOfLargePage);
+					LLM_IF_ENABLED(FLowLevelMemTracker::Get().OnLowLevelFree(ELLMTracker::Platform, (void*)LargePage->BaseAddress));
 				}
 
 				CommittedLargePagesCount[LargePage->AllocationHint] -= 1;
@@ -369,7 +371,7 @@ void FCachedOSVeryLargePageAllocator::ShrinkEmptyBackStore(int32 NewEmptyBackSto
 		}
 		LargePage->Unlink();
 		Block.Decommit(LargePage->BaseAddress - AddressSpaceReserved, SizeOfLargePage);
-
+		LLM_IF_ENABLED(FLowLevelMemTracker::Get().OnLowLevelFree(ELLMTracker::Platform, (void*)LargePage->BaseAddress));
 		LargePage->LinkHead(FreeLargePagesHead[LargePage->AllocationHint]);
 		CachedFree -= SizeOfLargePage;
 		EmptyBackStoreCount[AllocationHint] -= 1;
