Adding whether or not the user is running Wine to the Windows crash report context.

diff --git a/Engine/Source/Runtime/Core/Private/Windows/WindowsPlatformCrashContext.cpp b/Engine/Source/Runtime/Core/Private/Windows/WindowsPlatformCrashContext.cpp
--- a/Engine/Source/Runtime/Core/Private/Windows/WindowsPlatformCrashContext.cpp
+++ b/Engine/Source/Runtime/Core/Private/Windows/WindowsPlatformCrashContext.cpp
@@ -211,6 +211,7 @@ void FGenericCrashContext::CleanupPlatformSpecificFiles()
 void FWindowsPlatformCrashContext::AddPlatformSpecificProperties() const
 {
 	AddCrashProperty(TEXT("PlatformIsRunningWindows"), 1);
+	AddCrashProperty(TEXT("PlatformIsRunningWine"), FWindowsPlatformMisc::IsWine());
 	AddCrashProperty(TEXT("IsRunningOnBattery"), FPlatformMisc::IsRunningOnBattery());
 	WIDECHAR DriveName = 0;
 	const TCHAR* BaseDir = FWindowsPlatformProcess::BaseDir();
