De-inlined complex MemoryOps functions.
Coerced TFunction and TUniqueFunction bindings into function pointers if possible to reduce the number of template instantiations.
Fixed MemoryOps friendship.

diff --git a/Engine/Source/Editor/Sequencer/Private/Tools/SequencerSnapField.cpp b/Engine/Source/Editor/Sequencer/Private/Tools/SequencerSnapField.cpp
--- a/Engine/Source/Editor/Sequencer/Private/Tools/SequencerSnapField.cpp
+++ b/Engine/Source/Editor/Sequencer/Private/Tools/SequencerSnapField.cpp
@@ -41,11 +41,11 @@ struct FSnapGridVisitor : ISequencerEntityVisitor, UE::Sequencer::ISnapField
 
 				const int32 StartNum = Snaps.Num();
 				const int32 NumKeys = TimesScratch.Num();
-				Snaps.SetNum(Snaps.Num() + NumKeys);
+				Snaps.Reserve(StartNum + NumKeys);
 
 				for (int32 Index = 0; Index < NumKeys; ++Index)
 				{
-					Snaps[StartNum + Index] = FSnapPoint(FSnapPoint::Key, TimesScratch[Index]);
+					Snaps.Emplace(FSnapPoint::Key, TimesScratch[Index]);
 				}
 			}
 		}
diff --git a/Engine/Source/Editor/Sequencer/Public/MVVM/Extensions/ISnappableExtension.h b/Engine/Source/Editor/Sequencer/Public/MVVM/Extensions/ISnappableExtension.h
--- a/Engine/Source/Editor/Sequencer/Public/MVVM/Extensions/ISnappableExtension.h
+++ b/Engine/Source/Editor/Sequencer/Public/MVVM/Extensions/ISnappableExtension.h
@@ -50,8 +50,6 @@ struct SEQUENCER_API FSnapPoint
 
 private:
 
-	friend void ::DefaultConstructItems<FSnapPoint, int32>(void*, int32);
-
 	// uninitialized default constructor - only for use by containers
 	FSnapPoint()
 	{}
diff --git a/Engine/Source/Runtime/Core/Public/Templates/Function.h b/Engine/Source/Runtime/Core/Public/Templates/Function.h
--- a/Engine/Source/Runtime/Core/Public/Templates/Function.h
+++ b/Engine/Source/Runtime/Core/Public/Templates/Function.h
@@ -552,6 +552,12 @@ namespace UE::Core::Private::Function
 		// A pointer to the callable object
 		void* Ptr = nullptr;
 	};
+
+	template <typename FunctorType> auto ResolveFuncPtrTypeIfPossible(FunctorType&&, int) -> decltype(+std::declval<FunctorType>());
+	template <typename FunctorType> auto ResolveFuncPtrTypeIfPossible(FunctorType&&, ...) -> FunctorType&&;
+
+	template <typename FunctorType>
+	using TFuncPtrTypeIfPossible_T = decltype(ResolveFuncPtrTypeIfPossible(std::declval<FunctorType&&>(), 0));
 }
 
 /**
@@ -625,6 +631,9 @@ public:
 		: Super(Forward<FunctorType>(InFunc))
 	{
 		// This constructor is disabled for TFunctionRef types so it isn't incorrectly selected as copy/move constructors.
+
+		// Unlike TFunction and TUniqueFunction, we do not coerce the function type to a pointer here because the functionref would
+		// end up pointing to the temporary pointer we created on the stack.
 	}
 
 	/////////////////////////////////////////////////////
@@ -714,7 +723,7 @@ public:
 		)
 	>
 	TFunction(FunctorType&& InFunc)
-		: Super(Forward<FunctorType>(InFunc))
+		: Super(static_cast<UE::Core::Private::Function::TFuncPtrTypeIfPossible_T<FunctorType>>(InFunc))
 	{
 		// This constructor is disabled for TFunction types so it isn't incorrectly selected as copy/move constructors.
 
@@ -826,7 +835,7 @@ public:
 		)
 	>
 	TUniqueFunction(FunctorType&& InFunc)
-		: Super(Forward<FunctorType>(InFunc))
+		: Super(static_cast<UE::Core::Private::Function::TFuncPtrTypeIfPossible_T<FunctorType>>(InFunc))
 	{
 		// This constructor is disabled for TUniqueFunction types so it isn't incorrectly selected as copy/move constructors.
 
diff --git a/Engine/Source/Runtime/Core/Public/Templates/MemoryOps.h b/Engine/Source/Runtime/Core/Public/Templates/MemoryOps.h
--- a/Engine/Source/Runtime/Core/Public/Templates/MemoryOps.h
+++ b/Engine/Source/Runtime/Core/Public/Templates/MemoryOps.h
@@ -6,6 +6,7 @@
 #include "HAL/UnrealMemory.h"
 #include "Templates/IsTriviallyCopyAssignable.h"
 #include "Templates/IsTriviallyCopyConstructible.h"
+#include "Templates/Requires.h"
 #include "Templates/UnrealTypeTraits.h"
 #include "Traits/UseBitwiseSwap.h"
 #include <new> // IWYU pragma: export
@@ -28,26 +29,28 @@ namespace UE::Core::Private::MemoryOps
  * @param	Elements	The address of the first memory location to construct at.
  * @param	Count		The number of elements to destruct.
  */
-template <typename ElementType, typename SizeType>
+template <
+	typename ElementType,
+	typename SizeType
+	UE_REQUIRES(sizeof(ElementType) > 0 && TIsZeroConstructType<ElementType>::Value) // the sizeof here should improve the error messages we get when we try to call this function with incomplete types
+>
 FORCEINLINE void DefaultConstructItems(void* Address, SizeType Count)
 {
-	if constexpr (sizeof(ElementType) == 0)
-	{
-		// Should never get here, but this construct should improve the error messages we get when we try to call this function with incomplete types
-	}
-	else if constexpr (TIsZeroConstructType<ElementType>::Value)
-	{
-		FMemory::Memset(Address, 0, sizeof(ElementType) * Count);
-	}
-	else
+	FMemory::Memset(Address, 0, sizeof(ElementType) * Count);
+}
+template <
+	typename ElementType,
+	typename SizeType
+	UE_REQUIRES(sizeof(ElementType) > 0 && !TIsZeroConstructType<ElementType>::Value) // the sizeof here should improve the error messages we get when we try to call this function with incomplete types
+>
+FORCENOINLINE void DefaultConstructItems(void* Address, SizeType Count)
+{
+	ElementType* Element = (ElementType*)Address;
+	while (Count)
 	{
-		ElementType* Element = (ElementType*)Address;
-		while (Count)
-		{
-			::new ((void*)Element) ElementType;
-			++Element;
-			--Count;
-		}
+		::new ((void*)Element) ElementType;
+		++Element;
+		--Count;
 	}
 }
 
@@ -82,24 +85,29 @@ FORCEINLINE void DestructItem(ElementType* Element)
  *
  * @note: This function is optimized for values of T, and so will not dynamically dispatch destructor calls if T's destructor is virtual.
  */
-template <typename ElementType, typename SizeType>
+template <
+	typename ElementType,
+	typename SizeType
+	UE_REQUIRES(sizeof(ElementType) > 0 && std::is_trivially_destructible_v<ElementType>) // the sizeof here should improve the error messages we get when we try to call this function with incomplete types
+>
 FORCEINLINE void DestructItems(ElementType* Element, SizeType Count)
 {
-	if constexpr (sizeof(ElementType) == 0)
-	{
-		// Should never get here, but this construct should improve the error messages we get when we try to call this function with incomplete types
-	}
-	else if constexpr (!std::is_trivially_destructible_v<ElementType>)
+}
+template <
+	typename ElementType,
+	typename SizeType
+	UE_REQUIRES(sizeof(ElementType) > 0 && !std::is_trivially_destructible_v<ElementType>) // the sizeof here should improve the error messages we get when we try to call this function with incomplete types
+>
+FORCENOINLINE void DestructItems(ElementType* Element, SizeType Count)
+{
+	while (Count)
 	{
-		while (Count)
-		{
-			// We need a typedef here because VC won't compile the destructor call below if ElementType itself has a member called ElementType
-			typedef ElementType DestructItemsElementTypeTypedef;
+		// We need a typedef here because VC won't compile the destructor call below if ElementType itself has a member called ElementType
+		typedef ElementType DestructItemsElementTypeTypedef;
 
-			Element->DestructItemsElementTypeTypedef::~DestructItemsElementTypeTypedef();
-			++Element;
-			--Count;
-		}
+		Element->DestructItemsElementTypeTypedef::~DestructItemsElementTypeTypedef();
+		++Element;
+		--Count;
 	}
 }
 
@@ -110,29 +118,33 @@ FORCEINLINE void DestructItems(ElementType* Element, SizeType Count)
  * @param	Source		A pointer to the first argument to pass to the constructor.
  * @param	Count		The number of elements to copy.
  */
-template <typename DestinationElementType, typename SourceElementType, typename SizeType>
+template <
+	typename DestinationElementType,
+	typename SourceElementType,
+	typename SizeType
+	UE_REQUIRES(sizeof(DestinationElementType) > 0 && sizeof(SourceElementType) > 0 && TIsBitwiseConstructible<DestinationElementType, SourceElementType>::Value) // the sizeof here should improve the error messages we get when we try to call this function with incomplete types
+>
 FORCEINLINE void ConstructItems(void* Dest, const SourceElementType* Source, SizeType Count)
 {
-	if constexpr (sizeof(DestinationElementType) == 0 || sizeof(SourceElementType) == 0)
+	if (Count)
 	{
-		// Should never get here, but this construct should improve the error messages we get when we try to call this function with incomplete types
+		FMemory::Memcpy(Dest, Source, sizeof(SourceElementType) * Count);
 	}
-	else if constexpr (TIsBitwiseConstructible<DestinationElementType, SourceElementType>::Value)
-	{
-		if (Count)
-		{
-			FMemory::Memcpy(Dest, Source, sizeof(SourceElementType) * Count);
-		}
-	}
-	else
+}
+template <
+	typename DestinationElementType,
+	typename SourceElementType,
+	typename SizeType
+	UE_REQUIRES(sizeof(DestinationElementType) > 0 && sizeof(SourceElementType) > 0 && !TIsBitwiseConstructible<DestinationElementType, SourceElementType>::Value) // the sizeof here should improve the error messages we get when we try to call this function with incomplete types
+>
+FORCENOINLINE void ConstructItems(void* Dest, const SourceElementType* Source, SizeType Count)
+{
+	while (Count)
 	{
-		while (Count)
-		{
-			::new ((void*)Dest) DestinationElementType(*Source);
-			++(DestinationElementType*&)Dest;
-			++Source;
-			--Count;
-		}
+		::new ((void*)Dest) DestinationElementType(*Source);
+		++(DestinationElementType*&)Dest;
+		++Source;
+		--Count;
 	}
 }
 
@@ -143,26 +155,28 @@ FORCEINLINE void ConstructItems(void* Dest, const SourceElementType* Source, Siz
  * @param	Source		A pointer to the first item to assign.
  * @param	Count		The number of elements to assign.
  */
-template <typename ElementType, typename SizeType>
+template <
+	typename ElementType,
+	typename SizeType
+	UE_REQUIRES(sizeof(ElementType) > 0 && TIsTriviallyCopyAssignable<ElementType>::Value) // the sizeof here should improve the error messages we get when we try to call this function with incomplete types
+>
 FORCEINLINE void CopyAssignItems(ElementType* Dest, const ElementType* Source, SizeType Count)
 {
-	if constexpr (sizeof(ElementType) == 0)
-	{
-		// Should never get here, but this construct should improve the error messages we get when we try to call this function with incomplete types
-	}
-	else if constexpr (TIsTriviallyCopyAssignable<ElementType>::Value)
-	{
-		FMemory::Memcpy(Dest, Source, sizeof(ElementType) * Count);
-	}
-	else
+	FMemory::Memcpy(Dest, Source, sizeof(ElementType) * Count);
+}
+template <
+	typename ElementType,
+	typename SizeType
+	UE_REQUIRES(sizeof(ElementType) > 0 && !TIsTriviallyCopyAssignable<ElementType>::Value) // the sizeof here should improve the error messages we get when we try to call this function with incomplete types
+>
+FORCENOINLINE void CopyAssignItems(ElementType* Dest, const ElementType* Source, SizeType Count)
+{
+	while (Count)
 	{
-		while (Count)
-		{
-			*Dest = *Source;
-			++Dest;
-			++Source;
-			--Count;
-		}
+		*Dest = *Source;
+		++Dest;
+		++Source;
+		--Count;
 	}
 }
 
@@ -211,41 +225,45 @@ FORCEINLINE void RelocateConstructItem(void* Dest, const SourceElementType* Sour
  * @param	Source		A pointer to the first item to relocate.
  * @param	Count		The number of elements to relocate.
  */
-template <typename DestinationElementType, typename SourceElementType, typename SizeType>
+template <
+	typename DestinationElementType,
+	typename SourceElementType,
+	typename SizeType
+	UE_REQUIRES(sizeof(DestinationElementType) > 0 && sizeof(SourceElementType) > 0 && UE::Core::Private::MemoryOps::TCanBitwiseRelocate_V<DestinationElementType, SourceElementType>) // the sizeof here should improve the error messages we get when we try to call this function with incomplete types
+>
 FORCEINLINE void RelocateConstructItems(void* Dest, SourceElementType* Source, SizeType Count)
 {
-	if constexpr (sizeof(DestinationElementType) == 0 || sizeof(SourceElementType) == 0)
-	{
-		// Should never get here, but this construct should improve the error messages we get when we try to call this function with incomplete types
-	}
-	else if constexpr (std::is_const_v<SourceElementType>)
-	{
-		static_assert(sizeof(SourceElementType) == 0, "RelocateConstructItems: Source cannot be const");
-	}
-	else if constexpr (UE::Core::Private::MemoryOps::TCanBitwiseRelocate_V<DestinationElementType, SourceElementType>)
-	{
-		/* All existing UE containers seem to assume trivial relocatability (i.e. memcpy'able) of their members,
-		 * so we're going to assume that this is safe here.  However, it's not generally possible to assume this
-		 * in general as objects which contain pointers/references to themselves are not safe to be trivially
-		 * relocated.
-		 *
-		 * However, it is not yet possible to automatically infer this at compile time, so we can't enable
-		 * different (i.e. safer) implementations anyway. */
+	static_assert(!std::is_const_v<SourceElementType>, "RelocateConstructItems: Source cannot be const");
 
-		FMemory::Memmove(Dest, Source, sizeof(SourceElementType) * Count);
-	}
-	else
+	/* All existing UE containers seem to assume trivial relocatability (i.e. memcpy'able) of their members,
+	 * so we're going to assume that this is safe here.  However, it's not generally possible to assume this
+	 * in general as objects which contain pointers/references to themselves are not safe to be trivially
+	 * relocated.
+	 *
+	 * However, it is not yet possible to automatically infer this at compile time, so we can't enable
+	 * different (i.e. safer) implementations anyway. */
+
+	FMemory::Memmove(Dest, Source, sizeof(SourceElementType) * Count);
+}
+template <
+	typename DestinationElementType,
+	typename SourceElementType,
+	typename SizeType
+	UE_REQUIRES(sizeof(DestinationElementType) > 0 && sizeof(SourceElementType) > 0 && !UE::Core::Private::MemoryOps::TCanBitwiseRelocate_V<DestinationElementType, SourceElementType>) // the sizeof here should improve the error messages we get when we try to call this function with incomplete types
+>
+FORCENOINLINE void RelocateConstructItems(void* Dest, SourceElementType* Source, SizeType Count)
+{
+	static_assert(!std::is_const_v<SourceElementType>, "RelocateConstructItems: Source cannot be const");
+
+	while (Count)
 	{
-		while (Count)
-		{
-			// We need a typedef here because VC won't compile the destructor call below if SourceElementType itself has a member called SourceElementType
-			typedef SourceElementType RelocateConstructItemsElementTypeTypedef;
+		// We need a typedef here because VC won't compile the destructor call below if SourceElementType itself has a member called SourceElementType
+		typedef SourceElementType RelocateConstructItemsElementTypeTypedef;
 
-			::new ((void*)Dest) DestinationElementType((SourceElementType&&)*Source);
-			++(DestinationElementType*&)Dest;
-			(Source++)->RelocateConstructItemsElementTypeTypedef::~RelocateConstructItemsElementTypeTypedef();
-			--Count;
-		}
+		::new ((void*)Dest) DestinationElementType((SourceElementType&&)*Source);
+		++(DestinationElementType*&)Dest;
+		(Source++)->RelocateConstructItemsElementTypeTypedef::~RelocateConstructItemsElementTypeTypedef();
+		--Count;
 	}
 }
 
@@ -256,26 +274,28 @@ FORCEINLINE void RelocateConstructItems(void* Dest, SourceElementType* Source, S
  * @param	Source		A pointer to the first item to move from.
  * @param	Count		The number of elements to move.
  */
-template <typename ElementType, typename SizeType>
+template <
+	typename ElementType,
+	typename SizeType
+	UE_REQUIRES(sizeof(ElementType) > 0 && TIsTriviallyCopyConstructible<ElementType>::Value) // the sizeof here should improve the error messages we get when we try to call this function with incomplete types
+>
 FORCEINLINE void MoveConstructItems(void* Dest, const ElementType* Source, SizeType Count)
 {
-	if constexpr (sizeof(ElementType) == 0)
-	{
-		// Should never get here, but this construct should improve the error messages we get when we try to call this function with incomplete types
-	}
-	else if constexpr (TIsTriviallyCopyConstructible<ElementType>::Value)
-	{
-		FMemory::Memmove(Dest, Source, sizeof(ElementType) * Count);
-	}
-	else
+	FMemory::Memmove(Dest, Source, sizeof(ElementType) * Count);
+}
+template <
+	typename ElementType,
+	typename SizeType
+	UE_REQUIRES(sizeof(ElementType) > 0 && !TIsTriviallyCopyConstructible<ElementType>::Value) // the sizeof here should improve the error messages we get when we try to call this function with incomplete types
+>
+FORCENOINLINE void MoveConstructItems(void* Dest, const ElementType* Source, SizeType Count)
+{
+	while (Count)
 	{
-		while (Count)
-		{
-			::new ((void*)Dest) ElementType((ElementType&&)*Source);
-			++(ElementType*&)Dest;
-			++Source;
-			--Count;
-		}
+		::new ((void*)Dest) ElementType((ElementType&&)*Source);
+		++(ElementType*&)Dest;
+		++Source;
+		--Count;
 	}
 }
 
@@ -286,56 +306,60 @@ FORCEINLINE void MoveConstructItems(void* Dest, const ElementType* Source, SizeT
  * @param	Source		A pointer to the first item to move assign.
  * @param	Count		The number of elements to move assign.
  */
-template <typename ElementType, typename SizeType>
+template <
+	typename ElementType,
+	typename SizeType
+	UE_REQUIRES(sizeof(ElementType) > 0 && TIsTriviallyCopyAssignable<ElementType>::Value) // the sizeof here should improve the error messages we get when we try to call this function with incomplete types
+>
 FORCEINLINE void MoveAssignItems(ElementType* Dest, const ElementType* Source, SizeType Count)
 {
-	if constexpr (sizeof(ElementType) == 0)
-	{
-		// Should never get here, but this construct should improve the error messages we get when we try to call this function with incomplete types
-	}
-	else if constexpr (TIsTriviallyCopyAssignable<ElementType>::Value)
-	{
-		FMemory::Memmove(Dest, Source, sizeof(ElementType) * Count);
-	}
-	else
+	FMemory::Memmove(Dest, Source, sizeof(ElementType) * Count);
+}
+template <
+	typename ElementType,
+	typename SizeType
+	UE_REQUIRES(sizeof(ElementType) > 0 && !TIsTriviallyCopyAssignable<ElementType>::Value) // the sizeof here should improve the error messages we get when we try to call this function with incomplete types
+>
+FORCENOINLINE void MoveAssignItems(ElementType* Dest, const ElementType* Source, SizeType Count)
+{
+	while (Count)
 	{
-		while (Count)
-		{
-			*Dest = (ElementType&&)*Source;
-			++Dest;
-			++Source;
-			--Count;
-		}
+		*Dest = (ElementType&&)*Source;
+		++Dest;
+		++Source;
+		--Count;
 	}
 }
 
-template <typename ElementType, typename SizeType>
+template <
+	typename ElementType,
+	typename SizeType
+	UE_REQUIRES(sizeof(ElementType) > 0 && TTypeTraits<ElementType>::IsBytewiseComparable) // the sizeof here should improve the error messages we get when we try to call this function with incomplete types
+>
 FORCEINLINE bool CompareItems(const ElementType* A, const ElementType* B, SizeType Count)
 {
-	if constexpr (sizeof(ElementType) == 0)
-	{
-		// Should never get here, but this construct should improve the error messages we get when we try to call this function with incomplete types
-	}
-	else if constexpr (TTypeTraits<ElementType>::IsBytewiseComparable)
-	{
-		return !Count || !FMemory::Memcmp(A, B, sizeof(ElementType) * Count);
-	}
-	else
+	return !Count || !FMemory::Memcmp(A, B, sizeof(ElementType) * Count);
+}
+template <
+	typename ElementType,
+	typename SizeType
+	UE_REQUIRES(sizeof(ElementType) > 0 && !TTypeTraits<ElementType>::IsBytewiseComparable) // the sizeof here should improve the error messages we get when we try to call this function with incomplete types
+>
+FORCENOINLINE bool CompareItems(const ElementType* A, const ElementType* B, SizeType Count)
+{
+	while (Count)
 	{
-		while (Count)
+		if (!(*A == *B))
 		{
-			if (!(*A == *B))
-			{
-				return false;
-			}
-
-			++A;
-			++B;
-			--Count;
+			return false;
 		}
 
-		return true;
+		++A;
+		++B;
+		--Count;
 	}
+
+	return true;
 }
 
 #if UE_ENABLE_INCLUDE_ORDER_DEPRECATED_IN_5_5
