Use a new encoding to reduce the binary size by 7.83% when using the AutoRTFM compiler.

diff --git a/Engine/Source/Runtime/Core/Private/AutoRTFM/FunctionMapInlines.h b/Engine/Source/Runtime/Core/Private/AutoRTFM/FunctionMapInlines.h
--- a/Engine/Source/Runtime/Core/Private/AutoRTFM/FunctionMapInlines.h
+++ b/Engine/Source/Runtime/Core/Private/AutoRTFM/FunctionMapInlines.h
@@ -12,22 +12,48 @@ namespace AutoRTFM
 
 inline void* FunctionMapLookup(void* OldFunction, const char* Where)
 {
-    void* Result = FunctionMapTryLookup(OldFunction);
-    if (!Result)
-    {
-		if (Where)
-		{
-			ensureMsgf(!ForTheRuntime::IsEnsureOnAbortByLanguageEnabled(), TEXT("Could not find function %p '%s' where '%s'."), OldFunction, *GetFunctionDescription(OldFunction), ANSI_TO_TCHAR(Where));
-		}
-		else
-		{
-			ensureMsgf(!ForTheRuntime::IsEnsureOnAbortByLanguageEnabled(), TEXT("Could not find function %p '%s'."), OldFunction, *GetFunctionDescription(OldFunction));
-		}
-		
-		FContext* Context = FContext::Get();
-        Context->AbortByLanguageAndThrow();
-    }
-    return Result;
+	// We use prefix data in our custom LLVM pass to stuff some data just
+	// before the address of all open function pointers (that we have
+	// definitions for!). We have two 32-bit values, the Offset from our
+	// open function -> closed function, and the Magic Mike constant which
+	// is just the bottom 32-bits of OldFunction. We use this Magic Mike
+	// constant to ensure that we are *actually* reading a legit address
+	// from the 8-bytes before a functions address, and not some random
+	// memory location (imagine you were trying to lookup a function
+	// address from a third-party library or in the c stdlib, it won't
+	// have been compiled with our compiler so *won't* have this prefix
+	// data). Only if Magic Mike constant we compute from OldFunction
+	// matches what the compiler injected, can we calculate the closed
+	// function address using the Offset.
+	const int32 Offset = *(reinterpret_cast<int32*>(OldFunction) - 1);
+	const uint32 MagicMike = *(reinterpret_cast<uint32*>(OldFunction) - 2);
+
+	if (LIKELY(MagicMike == (reinterpret_cast<uint64>(OldFunction) & 0xffffffff)))
+	{
+		const int64 Relocated = reinterpret_cast<int64>(OldFunction) + Offset;
+		return reinterpret_cast<void*>(Relocated);
+	}
+
+	// Instead fall back to the slower function map lookup.
+    void* const Result = FunctionMapTryLookup(OldFunction);
+
+	if (Result)
+	{
+		return Result;
+	}
+
+	if (Where)
+	{
+		ensureMsgf(!ForTheRuntime::IsEnsureOnAbortByLanguageEnabled(), TEXT("Could not find function %p '%s' where '%s'."), OldFunction, *GetFunctionDescription(OldFunction), ANSI_TO_TCHAR(Where));
+	}
+	else
+	{
+		ensureMsgf(!ForTheRuntime::IsEnsureOnAbortByLanguageEnabled(), TEXT("Could not find function %p '%s'."), OldFunction, *GetFunctionDescription(OldFunction));
+	}
+
+	FContext* Context = FContext::Get();
+	Context->AbortByLanguageAndThrow();
+	return nullptr;
 }
 
 template<typename TReturnType, typename... TParameterTypes>
