CBWriter: Add platform support for SIZE_T, which does not implicitly convert to uint64 on Linux and Mac.

#rnx

diff --git a/Engine/Source/Runtime/Core/Public/Serialization/CompactBinaryWriter.h b/Engine/Source/Runtime/Core/Public/Serialization/CompactBinaryWriter.h
--- a/Engine/Source/Runtime/Core/Public/Serialization/CompactBinaryWriter.h
+++ b/Engine/Source/Runtime/Core/Public/Serialization/CompactBinaryWriter.h
@@ -390,6 +390,17 @@ public:
 		return *this;
 	}
 
+#if PLATFORM_LINUX || PLATFORM_APPLE
+	// On Unix-like platforms, size_t is an unsigned long, whereas UE defines uint64 as unsigned long long,
+	// making the two types immiscible for template resolution. However, SSIZE_T is currently defined as
+	// int64 (see FGenericPlatformTypes) and so no special overload is required for that type.
+	inline FCbWriter& operator<<(SIZE_T Value)
+	{
+		AddInteger(static_cast<uint64>(Value));
+		return *this;
+	}
+#endif
+
 	inline FCbWriter& operator<<(float Value)
 	{
 		AddFloat(Value);
diff --git a/Engine/Source/Runtime/Core/Tests/Serialization/CompactBinaryWriterTest.cpp b/Engine/Source/Runtime/Core/Tests/Serialization/CompactBinaryWriterTest.cpp
--- a/Engine/Source/Runtime/Core/Tests/Serialization/CompactBinaryWriterTest.cpp
+++ b/Engine/Source/Runtime/Core/Tests/Serialization/CompactBinaryWriterTest.cpp
@@ -1010,6 +1010,7 @@ TEST_CASE_NAMED(FCbWriterStreamTest, "System::Core::Serialization::CbWriter::Str
 		Writer << ANSITEXTVIEW("Integers");
 		Writer.BeginArray();
 		Writer << int32(-1) << int64(-1) << uint32(1) << uint64(1);
+		Writer << SSIZE_T(-1) << SIZE_T(-1) << PTRINT(-1) << UPTRINT(-1);
 		Writer.EndArray();
 
 		Writer << ANSITEXTVIEW("Float32") << 1.0f;
