Better handling of pread/pwrite for Android filesystem

#android

diff --git a/Engine/Source/Runtime/Core/Private/Android/AndroidPlatformFile.cpp b/Engine/Source/Runtime/Core/Private/Android/AndroidPlatformFile.cpp
--- a/Engine/Source/Runtime/Core/Private/Android/AndroidPlatformFile.cpp
+++ b/Engine/Source/Runtime/Core/Private/Android/AndroidPlatformFile.cpp
@@ -345,10 +345,16 @@ public:
 #endif
 			if (ThisSize < 0)
 			{
+				if (errno == EINTR)
+				{
+					// interrupted by signal, no error
+					continue;
+				}
 				return false;
 			}
 			else if (ThisSize == 0)
 			{
+				// 0 is EOF
 				break;
 			}
 			CurrentOffset += ThisSize;
@@ -369,25 +375,38 @@ public:
 		}
 
 		bool bSuccess = true;
-		while (BytesToWrite)
+		while (BytesToWrite > 0)
 		{
-			check(BytesToWrite >= 0);
 			int64 ThisSize = FMath::Min<int64>(READWRITE_SIZE, BytesToWrite);
 			check(Source);
-			if (__pwrite(File->Handle, Source, ThisSize, CurrentOffset) != ThisSize)
+			errno = EINTR;
+			int64 Result = __pwrite(File->Handle, Source, ThisSize, CurrentOffset);
+			if (Result <= 0)
 			{
+				if (errno == EINTR)
+				{
+					// interrupted by signal, no error
+					continue;
+				}
+#if LOG_ANDROID_FILE
+				int32 SaveErrno = errno;
+				FPlatformMisc::LowLevelOutputDebugStringf(
+					TEXT("(%d/%d) FFileHandleAndroid:Write => Path = %s, this size = %d, CurrentOffset = %d, Source = %p, Result = %d, errno = %d"),
+					FAndroidTLS::GetCurrentThreadId(), File->Handle,
+					*(File->Path), int32(ThisSize), CurrentOffset, Source, int32(Result), SaveErrno);
+#endif
 				bSuccess = false;
 				break;
 			}
 #if LOG_ANDROID_FILE
 			FPlatformMisc::LowLevelOutputDebugStringf(
-				TEXT("(%d/%d) FFileHandleAndroid:Write => Path = %s, this size = %d, CurrentOffset = %d, Source = %p"),
+				TEXT("(%d/%d) FFileHandleAndroid:Write => Path = %s, this size = %d, CurrentOffset = %d, Source = %p, Result = %d"),
 				FAndroidTLS::GetCurrentThreadId(), File->Handle,
-				*(File->Path), int32(ThisSize), CurrentOffset, Source);
+				*(File->Path), int32(ThisSize), CurrentOffset, Source, int32(Result));
 #endif
-			CurrentOffset += ThisSize;
-			Source += ThisSize;
-			BytesToWrite -= ThisSize;
+			CurrentOffset += Result;
+			Source += Result;
+			BytesToWrite -= Result;
 		}
 		
 		// Update the cached file length
