Fixed excessive alignment requirement in FCachedOSVeryLargePageAllocator

diff --git a/Engine/Source/Runtime/Core/Private/HAL/Allocators/CachedOSVeryLargePageAllocator.cpp b/Engine/Source/Runtime/Core/Private/HAL/Allocators/CachedOSVeryLargePageAllocator.cpp
--- a/Engine/Source/Runtime/Core/Private/HAL/Allocators/CachedOSVeryLargePageAllocator.cpp
+++ b/Engine/Source/Runtime/Core/Private/HAL/Allocators/CachedOSVeryLargePageAllocator.cpp
@@ -74,7 +74,7 @@ static volatile int32 GLargePageAllocatorDecommitCount = 0;
 
 void FCachedOSVeryLargePageAllocator::Init()
 {
-	Block = FPlatformMemory::FPlatformVirtualMemoryBlock::AllocateVirtual(AddressSpaceToReserve, SizeOfLargePage);
+	Block = FPlatformMemory::FPlatformVirtualMemoryBlock::AllocateVirtual(AddressSpaceToReserve);
 	AddressSpaceReserved = (uintptr_t)Block.GetVirtualPointer();
 	AddressSpaceReservedEnd = AddressSpaceReserved + AddressSpaceToReserve;
 	AddressSpaceReservedEndSmallPool = AddressSpaceReserved + (AddressSpaceToReserve / 2);
