Fixed text about the use of non-public constructors in the comments of MakeShared, which conflicts with advice given elsewhere in the engine when using emplacement functions.

diff --git a/Engine/Source/Runtime/Core/Public/Templates/SharedPointer.h b/Engine/Source/Runtime/Core/Public/Templates/SharedPointer.h
--- a/Engine/Source/Runtime/Core/Public/Templates/SharedPointer.h
+++ b/Engine/Source/Runtime/Core/Public/Templates/SharedPointer.h
@@ -2229,9 +2229,22 @@ template< class ObjectType, class DeleterType >
  * MakeShared utility function.  Allocates a new ObjectType and reference controller in a single memory block.
  * Equivalent to std::make_shared.
  *
- * NOTE: If the constructor is private/protected you will need to friend the intrusive reference controller in your class. e.g.
- * 	  template <typename ObjectType>
- *	  friend class SharedPointerInternals::TIntrusiveReferenceController;
+ * NOTE: If the constructor is private/protected you will need to utilize a private token.  Do not make SharedPointerInternals::TIntrusiveReferenceController a friend.
+ *
+ *     class FMyType
+ *     {
+ *     private:
+ *         struct FPrivateToken { explicit FPrivateToken() = default; };
+ *
+ *     public:
+ *         // This has an equivalent access level to a private constructor,
+ *         // as only friends of FMyType will have access to FPrivateToken,
+ *         // but MakeShared can legally call it since it's public.
+ *         explicit FMyType(FPrivateToken, int32 Int, float Real, const TCHAR* String);
+ *     };
+ *
+ *     // Won't compile if the caller doesn't have access to FMyType::FPrivateToken
+ *     TSharedPtr<FMyType> Val = MakeShared<FMyType>(FMyType::FPrivateToken{}, 5, 3.14f, TEXT("Banana"));
  */
 template <typename InObjectType, ESPMode InMode = ESPMode::ThreadSafe, typename... InArgTypes>
 [[nodiscard]] FORCEINLINE TSharedRef<InObjectType, InMode> MakeShared(InArgTypes&&... Args)
