Fix non-unity error

diff --git a/Engine/Source/Runtime/Core/Private/ProfilingDebugging/TraceAuxiliary.cpp b/Engine/Source/Runtime/Core/Private/ProfilingDebugging/TraceAuxiliary.cpp
--- a/Engine/Source/Runtime/Core/Private/ProfilingDebugging/TraceAuxiliary.cpp
+++ b/Engine/Source/Runtime/Core/Private/ProfilingDebugging/TraceAuxiliary.cpp
@@ -1871,6 +1871,7 @@ void FTraceAuxiliary::EnumerateFixedChannelPresets(PresetCallback Callback)
 
 void FTraceAuxiliary::EnumerateChannelPresetsFromSettings(PresetCallback Callback)
 {
+#if UE_TRACE_ENABLED
 	TArray<FString> PresetStrings;
 	GConfig->GetSection(TEXT("Trace.ChannelPresets"), PresetStrings, GEngineIni);
 
@@ -1885,6 +1886,7 @@ void FTraceAuxiliary::EnumerateChannelPresetsFromSettings(PresetCallback Callbac
 			return;
 		}
 	}
+#endif
 }
 
 FTraceAuxiliary::ETraceSystemStatus FTraceAuxiliary::GetTraceSystemStatus()
