Add missing Oodle define.

#rnx

diff --git a/Engine/Source/Runtime/Core/Private/Compression/OodleDataCompression.cpp b/Engine/Source/Runtime/Core/Private/Compression/OodleDataCompression.cpp
--- a/Engine/Source/Runtime/Core/Private/Compression/OodleDataCompression.cpp
+++ b/Engine/Source/Runtime/Core/Private/Compression/OodleDataCompression.cpp
@@ -81,7 +81,8 @@ static struct { ECompressionLevel Level; const TCHAR* Name; } CompressionLevelNa
 	{ECompressionLevel::Optimal1, TEXT("Optimal1")},
 	{ECompressionLevel::Optimal2, TEXT("Optimal2")},
 	{ECompressionLevel::Optimal3, TEXT("Optimal3")},
-	{ECompressionLevel::Optimal4, TEXT("Optimal4")}
+	{ECompressionLevel::Optimal4, TEXT("Optimal4")},
+	{ECompressionLevel::Optimal5, TEXT("Optimal5")}
 };
 
 CORE_API bool ECompressionLevelFromString(const TCHAR* InName, ECompressionLevel& OutLevel)
diff --git a/Engine/Source/Runtime/Core/Public/Compression/OodleDataCompression.h b/Engine/Source/Runtime/Core/Public/Compression/OodleDataCompression.h
--- a/Engine/Source/Runtime/Core/Public/Compression/OodleDataCompression.h
+++ b/Engine/Source/Runtime/Core/Public/Compression/OodleDataCompression.h
@@ -104,7 +104,8 @@ enum class ECompressionLevel : int8
 	Optimal1 = 5,
 	Optimal2 = 6,
 	Optimal3 = 7,
-	Optimal4 = 8, // if another added update CompressionLevelNameMap
+	Optimal4 = 8,
+	Optimal5 = 9, // if another added update CompressionLevelNameMap
 };
 
 // If the input is invalid, returns false, and doesn't touch output.
