Make function map lookup 7.49% faster by moving the unlikely error case to another function that we `[[clang::musttail]]` call, avoiding a lot of boilerplate stack crap in a function that needs to be fast.

diff --git a/Engine/Source/Runtime/Core/Private/AutoRTFM/Context.h b/Engine/Source/Runtime/Core/Private/AutoRTFM/Context.h
--- a/Engine/Source/Runtime/Core/Private/AutoRTFM/Context.h
+++ b/Engine/Source/Runtime/Core/Private/AutoRTFM/Context.h
@@ -45,14 +45,14 @@ public:
     void DidFree(void* LogicalAddress);
 
     // The rest of this is internalish.
-    void AbortByLanguageAndThrow();
+    [[noreturn]] void AbortByLanguageAndThrow();
 
 	inline FTransaction* GetCurrentTransaction() const { return CurrentTransaction; }
 	inline FCallNest* GetCurrentNest() const { return CurrentNest; }
     inline bool IsTransactionStack(const void* LogicalAddress) const { return LogicalAddress >= StackBegin && LogicalAddress < OuterTransactStackAddress; }
 	inline bool IsInnerTransactionStack(const void* LogicalAddress) const { return LogicalAddress >= StackBegin && LogicalAddress < CurrentTransactStackAddress; }
 	inline EContextStatus GetStatus() const { return Status; }
-	void Throw();
+	[[noreturn]] void Throw();
 	
     void DumpState() const;
 
diff --git a/Engine/Source/Runtime/Core/Private/AutoRTFM/FunctionMap.cpp b/Engine/Source/Runtime/Core/Private/AutoRTFM/FunctionMap.cpp
--- a/Engine/Source/Runtime/Core/Private/AutoRTFM/FunctionMap.cpp
+++ b/Engine/Source/Runtime/Core/Private/AutoRTFM/FunctionMap.cpp
@@ -2,6 +2,7 @@
 
 #if (defined(__AUTORTFM) && __AUTORTFM)
 #include "FunctionMap.h"
+#include "ContextInlines.h"
 #include "Atomic.h"
 #include "FastLock.h"
 #include "LockGuard.h"
@@ -103,6 +104,21 @@ void FunctionMapAdd(void* OldFunction, void* NewFunction)
     }
 }
 
+void* FunctionMapReportError(void* OldFunction, const char* Where)
+{
+	if (Where)
+	{
+		ensureMsgf(!ForTheRuntime::IsEnsureOnAbortByLanguageEnabled(), TEXT("Could not find function %p '%s' where '%s'."), OldFunction, *GetFunctionDescription(OldFunction), ANSI_TO_TCHAR(Where));
+	}
+	else
+	{
+		ensureMsgf(!ForTheRuntime::IsEnsureOnAbortByLanguageEnabled(), TEXT("Could not find function %p '%s'."), OldFunction, *GetFunctionDescription(OldFunction));
+	}
+
+	FContext* Context = FContext::Get();
+	Context->AbortByLanguageAndThrow();
+}
+
 } // namespace AutoRTFM
 
 #endif // defined(__AUTORTFM) && __AUTORTFM
diff --git a/Engine/Source/Runtime/Core/Private/AutoRTFM/FunctionMap.h b/Engine/Source/Runtime/Core/Private/AutoRTFM/FunctionMap.h
--- a/Engine/Source/Runtime/Core/Private/AutoRTFM/FunctionMap.h
+++ b/Engine/Source/Runtime/Core/Private/AutoRTFM/FunctionMap.h
@@ -36,6 +36,8 @@ struct FFunctionMap {
 
 void FunctionMapAdd(void* OldFunction, void* NewFuncton);
 
+void* FunctionMapReportError(void* OldFunction, const char* Where);
+
 inline void* FunctionMapTryLookup(void* OldFunction)
 {
     FFunctionMap* Map = GlobalData->FunctionMap;
diff --git a/Engine/Source/Runtime/Core/Private/AutoRTFM/FunctionMapInlines.h b/Engine/Source/Runtime/Core/Private/AutoRTFM/FunctionMapInlines.h
--- a/Engine/Source/Runtime/Core/Private/AutoRTFM/FunctionMapInlines.h
+++ b/Engine/Source/Runtime/Core/Private/AutoRTFM/FunctionMapInlines.h
@@ -37,23 +37,15 @@ inline void* FunctionMapLookup(void* OldFunction, const char* Where)
 	// Instead fall back to the slower function map lookup.
     void* const Result = FunctionMapTryLookup(OldFunction);
 
-	if (Result)
+	if (UNLIKELY(!Result))
 	{
-		return Result;
+#ifdef __clang__
+		[[clang::musttail]]
+#endif
+		return FunctionMapReportError(OldFunction, Where);
 	}
 
-	if (Where)
-	{
-		ensureMsgf(!ForTheRuntime::IsEnsureOnAbortByLanguageEnabled(), TEXT("Could not find function %p '%s' where '%s'."), OldFunction, *GetFunctionDescription(OldFunction), ANSI_TO_TCHAR(Where));
-	}
-	else
-	{
-		ensureMsgf(!ForTheRuntime::IsEnsureOnAbortByLanguageEnabled(), TEXT("Could not find function %p '%s'."), OldFunction, *GetFunctionDescription(OldFunction));
-	}
-
-	FContext* Context = FContext::Get();
-	Context->AbortByLanguageAndThrow();
-	return nullptr;
+	return Result;
 }
 
 template<typename TReturnType, typename... TParameterTypes>
diff --git a/Engine/Source/Runtime/Core/Private/AutoRTFM/LongJump.h b/Engine/Source/Runtime/Core/Private/AutoRTFM/LongJump.h
--- a/Engine/Source/Runtime/Core/Private/AutoRTFM/LongJump.h
+++ b/Engine/Source/Runtime/Core/Private/AutoRTFM/LongJump.h
@@ -24,7 +24,7 @@ public:
     template<typename TTryFunctor, typename TCatchFunctor>
     void TryCatch(const TTryFunctor& TryFunctor, const TCatchFunctor& CatchFunctor);
     
-    void Throw();
+    [[noreturn]] void Throw();
     
 private:
     jmp_buf JmpBuf;
diff --git a/Engine/Source/Runtime/Core/Private/AutoRTFM/Transaction.h b/Engine/Source/Runtime/Core/Private/AutoRTFM/Transaction.h
--- a/Engine/Source/Runtime/Core/Private/AutoRTFM/Transaction.h
+++ b/Engine/Source/Runtime/Core/Private/AutoRTFM/Transaction.h
@@ -66,7 +66,7 @@ public:
     void PushDeferUntilAbortHandler(const void* Key, TFunction<void()>&&);
     bool PopDeferUntilAbortHandler(const void* Key);
 
-    void AbortAndThrow();
+    [[noreturn]] void AbortAndThrow();
     void AbortWithoutThrowing();
     bool AttemptToCommit();
 
