Fix the Linux Arm LTCG problem when using the AutoRTFM compiler.

diff --git a/Engine/Source/Runtime/Core/Private/AutoRTFM/API.cpp b/Engine/Source/Runtime/Core/Private/AutoRTFM/API.cpp
--- a/Engine/Source/Runtime/Core/Private/AutoRTFM/API.cpp
+++ b/Engine/Source/Runtime/Core/Private/AutoRTFM/API.cpp
@@ -364,7 +364,6 @@ UE_AUTORTFM_ALWAYS_OPEN void RTFM_autortfm_did_free(void* Ptr)
 UE_AUTORTFM_ALWAYS_OPEN void RTFM_autortfm_check_consistency_assuming_no_races()
 {
 }
-UE_AUTORTFM_REGISTER_OPEN_FUNCTION(autortfm_check_consistency_assuming_no_races);
 
 }  // anonymous namespace
 
diff --git a/Engine/Source/Runtime/Core/Private/AutoRTFM/FunctionMap.cpp b/Engine/Source/Runtime/Core/Private/AutoRTFM/FunctionMap.cpp
--- a/Engine/Source/Runtime/Core/Private/AutoRTFM/FunctionMap.cpp
+++ b/Engine/Source/Runtime/Core/Private/AutoRTFM/FunctionMap.cpp
@@ -12,6 +12,8 @@
 
 #include "Containers/Array.h"
 
+#include "AutoRTFM/AutoRTFMConstants.h"
+
 namespace AutoRTFM
 {
 
@@ -91,6 +93,13 @@ void FunctionMapAdd(void* OldFunction, void* NewFunction)
         return;
     }
 
+	// We first assert that we don't accidentally have the Magic Mike
+	// constant in the special location before the old function. This
+	// is to ensure that any functions added to the map are not jointly
+	// using the prefix data mechanism and the hashmap lookup.
+	const uint64 PrefixData = *(reinterpret_cast<uint64*>(OldFunction) - 1);
+	ASSERT(Constants::MagicMike != (PrefixData & 0xffff000000000000));
+
     InitializeGlobalDataIfNecessary();
     
     TLockGuard<FFastLock> LockGuard(GlobalData->FunctionMapLock);
diff --git a/Engine/Source/Runtime/Core/Private/AutoRTFM/FunctionMapInlines.h b/Engine/Source/Runtime/Core/Private/AutoRTFM/FunctionMapInlines.h
--- a/Engine/Source/Runtime/Core/Private/AutoRTFM/FunctionMapInlines.h
+++ b/Engine/Source/Runtime/Core/Private/AutoRTFM/FunctionMapInlines.h
@@ -2,8 +2,10 @@
 
 #pragma once
 
+#include "AutoRTFM/AutoRTFMConstants.h"
 #include "Context.h"
 #include "FunctionMap.h"
+#include "Utils.h"
 
 #include "Containers/StringConv.h"
 
@@ -14,24 +16,15 @@ inline void* FunctionMapLookup(void* OldFunction, const char* Where)
 {
 	// We use prefix data in our custom LLVM pass to stuff some data just
 	// before the address of all open function pointers (that we have
-	// definitions for!). We have two 32-bit values, the Offset from our
-	// open function -> closed function, and the Magic Mike constant which
-	// is just the bottom 32-bits of OldFunction. We use this Magic Mike
-	// constant to ensure that we are *actually* reading a legit address
-	// from the 8-bytes before a functions address, and not some random
-	// memory location (imagine you were trying to lookup a function
-	// address from a third-party library or in the c stdlib, it won't
-	// have been compiled with our compiler so *won't* have this prefix
-	// data). Only if Magic Mike constant we compute from OldFunction
-	// matches what the compiler injected, can we calculate the closed
-	// function address using the Offset.
-	const int32 Offset = *(reinterpret_cast<int32*>(OldFunction) - 1);
-	const uint32 MagicMike = *(reinterpret_cast<uint32*>(OldFunction) - 2);
-
-	if (LIKELY(MagicMike == (reinterpret_cast<uint64>(OldFunction) & 0xffffffff)))
+	// definitions for!). We use the special Magic Mike constant in the top
+	// 16-bits of the function pointer address as a magic constant check
+	// to give us a much higher confidence that there is actually a closed
+	// variant pointer residing 8-bytes before our function address.
+	const uint64 PrefixData = *(reinterpret_cast<uint64*>(OldFunction) - 1);
+
+	if (LIKELY(Constants::MagicMike == (PrefixData & 0xffff000000000000)))
 	{
-		const int64 Relocated = reinterpret_cast<int64>(OldFunction) + Offset;
-		return reinterpret_cast<void*>(Relocated);
+		return reinterpret_cast<void*>(PrefixData & 0x0000ffffffffffff);
 	}
 
 	// Instead fall back to the slower function map lookup.
diff --git a/Engine/Source/Runtime/Core/Public/AutoRTFM/AutoRTFMConstants.h b/Engine/Source/Runtime/Core/Public/AutoRTFM/AutoRTFMConstants.h
--- a/Engine/Source/Runtime/Core/Public/AutoRTFM/AutoRTFMConstants.h
+++ b/Engine/Source/Runtime/Core/Public/AutoRTFM/AutoRTFMConstants.h
@@ -16,5 +16,10 @@ namespace Constants
 	inline constexpr uint32_t Minor = 2;
 	inline constexpr uint32_t Patch = 0;
 
+	// The Magic Mike constant - shared with the compiler implementation.
+	// Used to ensure that prefix data pointers are what they claim to be
+	// and were definitely injected by us.
+	inline constexpr uint64_t MagicMike = 0xa273000000000000;
+
 } // namespace Constants
 } // namespace AutoRTFM
