Cleaned up MallocBinnedCommon.h from FArenaParams that's used by the MallocBinnedGPU only as it creates circular includes

diff --git a/Engine/Source/Runtime/Core/Public/HAL/MallocBinnedCommon.h b/Engine/Source/Runtime/Core/Public/HAL/MallocBinnedCommon.h
--- a/Engine/Source/Runtime/Core/Public/HAL/MallocBinnedCommon.h
+++ b/Engine/Source/Runtime/Core/Public/HAL/MallocBinnedCommon.h
@@ -13,9 +13,7 @@
 
 #if PLATFORM_HAS_FPlatformVirtualMemoryBlock
 
-#include "Containers/Array.h"
 #include "HAL/PlatformMemory.h"
-#include "Templates/Function.h"
 
 // A project can define it's own UE_MBC_MAX_LISTED_SMALL_POOL_SIZE and UE_MBC_NUM_LISTED_SMALL_POOLS to reduce runtime memory usage
 // MallocBinnedCommon.cpp has a list of predefined bins that go up to 28672
@@ -108,43 +106,6 @@ struct FSizeTableEntry
 	static uint8 FillSizeTable(uint64 PlatformPageSize, FSizeTableEntry* SizeTable, uint32 BasePageSize, uint32 MinimumAlignment, uint32 MaxSize, uint32 SizeIncrement);
 };
 
-struct FArenaParams
-{
-	// these are parameters you set
-	uint64 AddressLimit = 1024 * 1024 * 1024; // this controls the size of the root hash table
-	uint32 BasePageSize = 4096; // this is used to make sensible calls to malloc and figures into the standard pool sizes if bUseStandardSmallPoolSizes is true
-	uint32 AllocationGranularity = 4096; // this is the granularity of the commit and decommit calls used on the VM slabs
-	uint32 MaxSizePerBundle = 8192;
-	uint32 MaxStandardPoolSize = 128 * 1024; // these are added to the standard pool sizes, mainly to use the TLS caches, they are typically one block per slab
-	uint16 MaxBlocksPerBundle = 64;
-	uint8 MaxMemoryPerBlockSizeShift = 29;
-	uint8 EmptyCacheAllocExtra = 32;
-	uint8 MaxGlobalBundles = 32;
-	uint8 MinimumAlignmentShift = 4;
-	uint8 PoolCount;
-	bool bUseSeparateVMPerPool = !!(BINNEDCOMMON_USE_SEPARATE_VM_PER_POOL);
-	bool bPerThreadCaches = true;
-	bool bUseStandardSmallPoolSizes = true;
-	bool bAttemptToAlignSmallBocks = true;
-	TArray<uint32> AdditionalBlockSizes;
-
-	// This lambdas is similar to the platform virtual memory HAL and by default just call that. 
-	TFunction<FPlatformMemory::FPlatformVirtualMemoryBlock(SIZE_T)> ReserveVM;
-
-	// These allow you to override the large block allocator. The value add here is that MBA tracks the metadata for you and call tell the difference between a large block pointer and a small block pointer.
-	// By defaults these just use the platform VM interface to allocate some committed memory
-	TFunction<void*(SIZE_T, SIZE_T, SIZE_T&, uint32&)> LargeBlockAlloc;
-	TFunction<void(void*, uint32)> LargeBlockFree;
-
-
-	// these are parameters are derived from other parameters
-	uint64 MaxMemoryPerBlockSize;
-	uint32 MaxPoolSize;
-	uint32 MinimumAlignment;
-	uint32 MaximumAlignmentForSmallBlock;
-
-};
-
 #endif	//~PLATFORM_HAS_FPlatformVirtualMemoryBlock
 
 
diff --git a/Engine/Source/Runtime/Core/Public/HAL/MallocBinnedGPU.h b/Engine/Source/Runtime/Core/Public/HAL/MallocBinnedGPU.h
--- a/Engine/Source/Runtime/Core/Public/HAL/MallocBinnedGPU.h
+++ b/Engine/Source/Runtime/Core/Public/HAL/MallocBinnedGPU.h
@@ -25,6 +25,7 @@
 #include "Misc/ScopeLock.h"
 #include "Misc/ScopeLock.h"
 #include "Templates/AlignmentTemplates.h"
+#include "Templates/Function.h"
 
 
 #define BINNEDGPU_MAX_GMallocBinnedGPUMaxBundlesBeforeRecycle (8)
@@ -39,6 +40,41 @@
 
 PRAGMA_DISABLE_UNSAFE_TYPECAST_WARNINGS
 
+struct FArenaParams
+{
+	// these are parameters you set
+	uint64 AddressLimit = 1024 * 1024 * 1024; // this controls the size of the root hash table
+	uint32 BasePageSize = 4096; // this is used to make sensible calls to malloc and figures into the standard pool sizes if bUseStandardSmallPoolSizes is true
+	uint32 AllocationGranularity = 4096; // this is the granularity of the commit and decommit calls used on the VM slabs
+	uint32 MaxSizePerBundle = 8192;
+	uint32 MaxStandardPoolSize = 128 * 1024; // these are added to the standard pool sizes, mainly to use the TLS caches, they are typically one block per slab
+	uint16 MaxBlocksPerBundle = 64;
+	uint8 MaxMemoryPerBlockSizeShift = 29;
+	uint8 EmptyCacheAllocExtra = 32;
+	uint8 MaxGlobalBundles = 32;
+	uint8 MinimumAlignmentShift = 4;
+	uint8 PoolCount;
+	bool bUseSeparateVMPerPool = !!(BINNEDCOMMON_USE_SEPARATE_VM_PER_POOL);
+	bool bPerThreadCaches = true;
+	bool bUseStandardSmallPoolSizes = true;
+	bool bAttemptToAlignSmallBocks = true;
+	TArray<uint32> AdditionalBlockSizes;
+
+	// This lambdas is similar to the platform virtual memory HAL and by default just call that. 
+	TFunction<FPlatformMemory::FPlatformVirtualMemoryBlock(SIZE_T)> ReserveVM;
+
+	// These allow you to override the large block allocator. The value add here is that MBA tracks the metadata for you and call tell the difference between a large block pointer and a small block pointer.
+	// By defaults these just use the platform VM interface to allocate some committed memory
+	TFunction<void* (SIZE_T, SIZE_T, SIZE_T&, uint32&)> LargeBlockAlloc;
+	TFunction<void(void*, uint32)> LargeBlockFree;
+
+	// these are parameters are derived from other parameters
+	uint64 MaxMemoryPerBlockSize;
+	uint32 MaxPoolSize;
+	uint32 MinimumAlignment;
+	uint32 MaximumAlignmentForSmallBlock;
+};
+
 class CORE_API FMallocBinnedGPU final : public FMalloc
 {
 	struct FGlobalRecycler;
