Always open around creating GUID

diff --git a/Engine/Source/Programs/AutoRTFMTests/Private/UEPlatform.cpp b/Engine/Source/Programs/AutoRTFMTests/Private/UEPlatform.cpp
new file mode 100644
--- /dev/null
+++ b/Engine/Source/Programs/AutoRTFMTests/Private/UEPlatform.cpp
@@ -0,0 +1,19 @@
+// Copyright Epic Games, Inc. All Rights Reserved.
+
+#include "Catch2Includes.h"
+#include "AutoRTFM/AutoRTFM.h"
+#include "HAL/PlatformMisc.h"
+
+TEST_CASE("FPlatformMisc.CreateGUID")
+{
+	AutoRTFM::ForTheRuntime::SetEnsureOnAbortByLanguage(true);
+	FGuid Guid;
+	AutoRTFM::ETransactionResult Result = AutoRTFM::Transact([&]
+	{
+		FPlatformMisc::CreateGuid(Guid);
+	});
+	AutoRTFM::ForTheRuntime::SetEnsureOnAbortByLanguage(false);
+
+	REQUIRE(AutoRTFM::ETransactionResult::Committed == Result);
+	REQUIRE(Guid != FGuid());
+}
diff --git a/Engine/Source/Runtime/Core/Private/Windows/WindowsPlatformMisc.cpp b/Engine/Source/Runtime/Core/Private/Windows/WindowsPlatformMisc.cpp
--- a/Engine/Source/Runtime/Core/Private/Windows/WindowsPlatformMisc.cpp
+++ b/Engine/Source/Runtime/Core/Private/Windows/WindowsPlatformMisc.cpp
@@ -35,6 +35,7 @@
 #include "GenericPlatform/GenericPlatformCrashContext.h"
 #include "Windows/WindowsPlatformCrashContext.h"
 #include "HAL/PlatformOutputDevices.h"
+#include "AutoRTFM/AutoRTFM.h"
 
 #include "GenericPlatform/GenericPlatformChunkInstall.h"
 #include "GenericPlatform/GenericPlatformDriver.h"
@@ -1401,7 +1402,7 @@ const TCHAR* FWindowsPlatformMisc::GetSystemErrorMessage(TCHAR* OutBuffer, int32
 	return OutBuffer;
 }
 
-void FWindowsPlatformMisc::CreateGuid(FGuid& Result)
+UE_AUTORTFM_ALWAYS_OPEN void FWindowsPlatformMisc::CreateGuid(FGuid& Result)
 {
 	verify( CoCreateGuid( (GUID*)&Result )==S_OK );
 }
