[Insights][CpuProfilerTrace]
  - Added FCpuProfilerTrace::OutputEventType(Name, File, Line), with Name as an FName.
  - Added FCpuProfilerTrace::OutputDynamicEventType(Name, File, Line) with Name as ANSICHAR*, TCHAR* or FName.
  - Added FCpuProfilerTrace::OutputDynamicEventTypeWithId(Id, Name, File, Line) with Name as ANSICHAR* or TCHAR*.
  - Added FCpuProfilerTrace::OutputBeginDynamicEventWithId(Id, const ANSICHAR* Name, File, Line), in addition to existing function that receives Name as TCHAR*.
  - Added TRACE_CPUPROFILER_EVENT_MANUAL_END() macro to match TRACE_CPUPROFILER_EVENT_MANUAL_START macro.

diff --git a/Engine/Source/Runtime/Core/Private/ProfilingDebugging/CpuProfilerTrace.cpp b/Engine/Source/Runtime/Core/Private/ProfilingDebugging/CpuProfilerTrace.cpp
--- a/Engine/Source/Runtime/Core/Private/ProfilingDebugging/CpuProfilerTrace.cpp
+++ b/Engine/Source/Runtime/Core/Private/ProfilingDebugging/CpuProfilerTrace.cpp
@@ -172,69 +172,35 @@
 void FCpuProfilerTrace::OutputBeginDynamicEvent(const ANSICHAR* Name, const ANSICHAR* File, uint32 Line)
 {
 	CPUPROFILERTRACE_OUTPUTBEGINEVENT_PROLOGUE();
-	uint32 SpecId = ThreadBuffer->DynamicAnsiScopeNamesMap.FindRef(Name);
-	if (!SpecId)
-	{
-		LLM_SCOPE_BYNAME(TEXT("Trace/CpuProfiler"));
-		int32 NameSize = strlen(Name) + 1;
-		ANSICHAR* NameCopy = reinterpret_cast<ANSICHAR*>(ThreadBuffer->DynamicScopeNamesMemory.Alloc(NameSize, alignof(ANSICHAR)));
-		FMemory::Memmove(NameCopy, Name, NameSize);
-		SpecId = OutputEventType(NameCopy, File, Line);
-		ThreadBuffer->DynamicAnsiScopeNamesMap.Add(NameCopy, SpecId);
-	}
+	uint32 SpecId = OutputDynamicEventType(Name, File, Line);
 	CPUPROFILERTRACE_OUTPUTBEGINEVENT_EPILOGUE();
 }
 
 void FCpuProfilerTrace::OutputBeginDynamicEvent(const TCHAR* Name, const ANSICHAR* File, uint32 Line)
 {
 	CPUPROFILERTRACE_OUTPUTBEGINEVENT_PROLOGUE();
-	uint32 SpecId = ThreadBuffer->DynamicTCharScopeNamesMap.FindRef(Name);
-	if (!SpecId)
-	{
-		LLM_SCOPE_BYNAME(TEXT("Trace/CpuProfiler"));
-		int32 NameSize = (FCString::Strlen(Name) + 1) * sizeof(TCHAR);
-		TCHAR* NameCopy = reinterpret_cast<TCHAR*>(ThreadBuffer->DynamicScopeNamesMemory.Alloc(NameSize, alignof(TCHAR)));
-		FMemory::Memmove(NameCopy, Name, NameSize);
-		SpecId = OutputEventType(NameCopy, File, Line);
-		ThreadBuffer->DynamicTCharScopeNamesMap.Add(NameCopy, SpecId);
-	}
+	uint32 SpecId = OutputDynamicEventType(Name, File, Line);
 	CPUPROFILERTRACE_OUTPUTBEGINEVENT_EPILOGUE();
 }
 
 void FCpuProfilerTrace::OutputBeginDynamicEvent(const FName Name, const ANSICHAR* File, uint32 Line)
 {
-	OutputBeginDynamicEventWithId(Name, nullptr, File, Line);
+	CPUPROFILERTRACE_OUTPUTBEGINEVENT_PROLOGUE();
+	uint32 SpecId = OutputDynamicEventType(Name, File, Line);
+	CPUPROFILERTRACE_OUTPUTBEGINEVENT_EPILOGUE();
 }
 
+void FCpuProfilerTrace::OutputBeginDynamicEventWithId(const FName Id, const ANSICHAR* Name, const ANSICHAR* File, uint32 Line)
+{
+	CPUPROFILERTRACE_OUTPUTBEGINEVENT_PROLOGUE();
+	uint32 SpecId = OutputDynamicEventTypeWithId(Id, Name, File, Line);
+	CPUPROFILERTRACE_OUTPUTBEGINEVENT_EPILOGUE();
+}
+
 void FCpuProfilerTrace::OutputBeginDynamicEventWithId(const FName Id, const TCHAR* Name, const ANSICHAR* File, uint32 Line)
 {
 	CPUPROFILERTRACE_OUTPUTBEGINEVENT_PROLOGUE();
-	uint32 SpecId = ThreadBuffer->DynamicFNameScopeNamesMap.FindRef(Id.GetComparisonIndex());
-	if (!SpecId)
-	{
-		LLM_SCOPE_BYNAME(TEXT("Trace/CpuProfiler"));
-		if (Name != nullptr)
-		{
-			SpecId = OutputEventType(Name, File, Line);
-		}
-		else
-		{
-			const FNameEntry* NameEntry = Id.GetDisplayNameEntry();
-			if (NameEntry->IsWide())
-			{
-				WIDECHAR WideName[NAME_SIZE];
-				NameEntry->GetWideName(WideName);
-				SpecId = OutputEventType(WideName, File, Line);
-			}
-			else
-			{
-				ANSICHAR AnsiName[NAME_SIZE];
-				NameEntry->GetAnsiName(AnsiName);
-				SpecId = OutputEventType(AnsiName, File, Line);
-			}
-		}
-		ThreadBuffer->DynamicFNameScopeNamesMap.Add(Id.GetComparisonIndex(), SpecId);
-	}
+	uint32 SpecId = OutputDynamicEventTypeWithId(Id, Name, File, Line);
 	CPUPROFILERTRACE_OUTPUTBEGINEVENT_EPILOGUE();
 }
 
@@ -357,6 +323,134 @@
 	return SpecId;
 }
 
+UE_AUTORTFM_ALWAYS_OPEN
+uint32 FCpuProfilerTrace::OutputEventType(const FName Name, const ANSICHAR* File, uint32 Line)
+{
+	const FNameEntry* NameEntry = Name.GetDisplayNameEntry();
+	if (NameEntry->IsWide())
+	{
+		WIDECHAR WideName[NAME_SIZE];
+		NameEntry->GetWideName(WideName);
+		return OutputEventType(WideName, File, Line);
+	}
+	else
+	{
+		ANSICHAR AnsiName[NAME_SIZE];
+		NameEntry->GetAnsiName(AnsiName);
+		return OutputEventType(AnsiName, File, Line);
+	}
+}
+
+UE_AUTORTFM_ALWAYS_OPEN
+uint32 FCpuProfilerTrace::OutputDynamicEventType(const ANSICHAR* Name, const ANSICHAR* File, uint32 Line)
+{
+	FCpuProfilerTraceInternal::FThreadBuffer* ThreadBuffer = FCpuProfilerTraceInternal::ThreadBuffer;
+	if (!ThreadBuffer)
+	{
+		ThreadBuffer = FCpuProfilerTraceInternal::CreateThreadBuffer();
+	}
+	uint32 SpecId = ThreadBuffer->DynamicAnsiScopeNamesMap.FindRef(Name);
+	if (!SpecId)
+	{
+		LLM_SCOPE_BYNAME(TEXT("Trace/CpuProfiler"));
+		int32 NameSize = strlen(Name) + 1;
+		ANSICHAR* NameCopy = reinterpret_cast<ANSICHAR*>(ThreadBuffer->DynamicScopeNamesMemory.Alloc(NameSize, alignof(ANSICHAR)));
+		FMemory::Memmove(NameCopy, Name, NameSize);
+		SpecId = OutputEventType(NameCopy, File, Line);
+		ThreadBuffer->DynamicAnsiScopeNamesMap.Add(NameCopy, SpecId);
+	}
+	return SpecId;
+}
+
+UE_AUTORTFM_ALWAYS_OPEN
+uint32 FCpuProfilerTrace::OutputDynamicEventType(const TCHAR* Name, const ANSICHAR* File, uint32 Line)
+{
+	FCpuProfilerTraceInternal::FThreadBuffer* ThreadBuffer = FCpuProfilerTraceInternal::ThreadBuffer;
+	if (!ThreadBuffer)
+	{
+		ThreadBuffer = FCpuProfilerTraceInternal::CreateThreadBuffer();
+	}
+	uint32 SpecId = ThreadBuffer->DynamicTCharScopeNamesMap.FindRef(Name);
+	if (!SpecId)
+	{
+		LLM_SCOPE_BYNAME(TEXT("Trace/CpuProfiler"));
+		int32 NameSize = (FCString::Strlen(Name) + 1) * sizeof(TCHAR);
+		TCHAR* NameCopy = reinterpret_cast<TCHAR*>(ThreadBuffer->DynamicScopeNamesMemory.Alloc(NameSize, alignof(TCHAR)));
+		FMemory::Memmove(NameCopy, Name, NameSize);
+		SpecId = OutputEventType(NameCopy, File, Line);
+		ThreadBuffer->DynamicTCharScopeNamesMap.Add(NameCopy, SpecId);
+	}
+	return SpecId;
+}
+
+UE_AUTORTFM_ALWAYS_OPEN
+uint32 FCpuProfilerTrace::OutputDynamicEventType(const FName Name, const ANSICHAR* File, uint32 Line)
+{
+	FCpuProfilerTraceInternal::FThreadBuffer* ThreadBuffer = FCpuProfilerTraceInternal::ThreadBuffer;
+	if (!ThreadBuffer)
+	{
+		ThreadBuffer = FCpuProfilerTraceInternal::CreateThreadBuffer();
+	}
+	uint32 SpecId = ThreadBuffer->DynamicFNameScopeNamesMap.FindRef(Name.GetComparisonIndex());
+	if (!SpecId)
+	{
+		LLM_SCOPE_BYNAME(TEXT("Trace/CpuProfiler"));
+		SpecId = OutputEventType(Name, File, Line);
+		ThreadBuffer->DynamicFNameScopeNamesMap.Add(Name.GetComparisonIndex(), SpecId);
+	}
+	return SpecId;
+}
+
+UE_AUTORTFM_ALWAYS_OPEN
+uint32 FCpuProfilerTrace::OutputDynamicEventTypeWithId(const FName Id, const ANSICHAR* Name, const ANSICHAR* File, uint32 Line)
+{
+	FCpuProfilerTraceInternal::FThreadBuffer* ThreadBuffer = FCpuProfilerTraceInternal::ThreadBuffer;
+	if (!ThreadBuffer)
+	{
+		ThreadBuffer = FCpuProfilerTraceInternal::CreateThreadBuffer();
+	}
+	uint32 SpecId = ThreadBuffer->DynamicFNameScopeNamesMap.FindRef(Id.GetComparisonIndex());
+	if (!SpecId)
+	{
+		LLM_SCOPE_BYNAME(TEXT("Trace/CpuProfiler"));
+		if (Name != nullptr)
+		{
+			SpecId = OutputEventType(Name, File, Line);
+		}
+		else
+		{
+			SpecId = OutputEventType(Id, File, Line);
+		}
+		ThreadBuffer->DynamicFNameScopeNamesMap.Add(Id.GetComparisonIndex(), SpecId);
+	}
+	return SpecId;
+}
+
+UE_AUTORTFM_ALWAYS_OPEN
+uint32 FCpuProfilerTrace::OutputDynamicEventTypeWithId(const FName Id, const TCHAR* Name, const ANSICHAR* File, uint32 Line)
+{
+	FCpuProfilerTraceInternal::FThreadBuffer* ThreadBuffer = FCpuProfilerTraceInternal::ThreadBuffer;
+	if (!ThreadBuffer)
+	{
+		ThreadBuffer = FCpuProfilerTraceInternal::CreateThreadBuffer();
+	}
+	uint32 SpecId = ThreadBuffer->DynamicFNameScopeNamesMap.FindRef(Id.GetComparisonIndex());
+	if (!SpecId)
+	{
+		LLM_SCOPE_BYNAME(TEXT("Trace/CpuProfiler"));
+		if (Name != nullptr)
+		{
+			SpecId = OutputEventType(Name, File, Line);
+		}
+		else
+		{
+			SpecId = OutputEventType(Id, File, Line);
+		}
+		ThreadBuffer->DynamicFNameScopeNamesMap.Add(Id.GetComparisonIndex(), SpecId);
+	}
+	return SpecId;
+}
+
 void FCpuProfilerTrace::FlushThreadBuffer()
 {
 	FCpuProfilerTraceInternal::FThreadBuffer* ThreadBuffer = FCpuProfilerTraceInternal::ThreadBuffer;
diff --git a/Engine/Source/Runtime/Core/Public/ProfilingDebugging/CpuProfilerTrace.h b/Engine/Source/Runtime/Core/Public/ProfilingDebugging/CpuProfilerTrace.h
--- a/Engine/Source/Runtime/Core/Public/ProfilingDebugging/CpuProfilerTrace.h
+++ b/Engine/Source/Runtime/Core/Public/ProfilingDebugging/CpuProfilerTrace.h
@@ -19,7 +19,6 @@
 #endif
 #endif
 
-
 #if CPUPROFILERTRACE_ENABLED
 
 UE_TRACE_CHANNEL_EXTERN(CpuChannel, CORE_API);
@@ -27,7 +26,7 @@ UE_TRACE_CHANNEL_EXTERN(CpuChannel, CORE_API);
 class FName;
 
 /*
- * Facilities for tracing timed cpu events. Two types of events are supported, static events where the identifier is
+ * Facilities for tracing timed CPU events. Two types of events are supported, static events where the identifier is
  * known at compile time, and dynamic event were identifiers can be constructed in runtime. Static events have lower overhead
  * so always prefer to use them if possible.
  *
@@ -39,71 +38,92 @@ class FName;
 struct FCpuProfilerTrace
 {
 	/*
-	 * Output cpu event definition (spec).
-	 * @param Name Event name
-	 * @param File Source filename
-	 * @param Line Line number in source file
-	 * @return Event definition id
+	 * Output CPU event definition (spec).
+	 * The trace event emitted by this function is an "important event" (so all events emitted will add to the Trace System's cache).
+	 * It is the responsibility of the caller code to ensure this function is not abused.
+	 * @param Name - The event name
+	 * @param File - The source filename
+	 * @param Line - The line number in source file
+	 * @returns The event definition id
 	 */
 	FORCENOINLINE CORE_API static uint32 OutputEventType(const ANSICHAR* Name, const ANSICHAR* File = nullptr, uint32 Line = 0);
+	FORCENOINLINE CORE_API static uint32 OutputEventType(const TCHAR* Name, const ANSICHAR* File = nullptr, uint32 Line = 0);
+	FORCENOINLINE CORE_API static uint32 OutputEventType(const FName Name, const ANSICHAR* File = nullptr, uint32 Line = 0);
+
 	/*
-	 * Output cpu event definition (spec).
-	 * @param Name Event name
-	 * @param File Source filename
-	 * @param Line Line number in source file
-	 * @return Event definition id
+	 * Output CPU event definition (spec) for a dynamic event.
+	 * The Name will be cached and the trace event will only be emitted once (for each unique Name; even if File or Line changes).
+	 * @param Name - The event name
+	 * @param File - The source filename
+	 * @param Line - The line number in source file
+	 * @returns The event definition id
 	 */
-	FORCENOINLINE CORE_API static uint32 OutputEventType(const TCHAR* Name, const ANSICHAR* File = nullptr, uint32 Line = 0);
+	FORCENOINLINE CORE_API static uint32 OutputDynamicEventType(const ANSICHAR* Name, const ANSICHAR* File = nullptr, uint32 Line = 0);
+	FORCENOINLINE CORE_API static uint32 OutputDynamicEventType(const TCHAR* Name, const ANSICHAR* File = nullptr, uint32 Line = 0);
+	FORCENOINLINE CORE_API static uint32 OutputDynamicEventType(const FName Name, const ANSICHAR* File = nullptr, uint32 Line = 0);
+
+	/*
+	 * Output CPU event definition (spec) for a dynamic event identified by an FName.
+	 * The Id will be cached and the trace event will only be emitted once (for each unique Id; even if Name, File or Line changes).
+	 * This is faster and less memory expensive than \ref OutputDynamicEventType that receives ANSICHAR* or TCHAR* name.
+	 * @param Id - The id of event
+	 * @param Name - The name of event
+	 * @param File - The source filename
+	 * @param Line - The line number in source file
+	 */
+	FORCENOINLINE CORE_API static uint32 OutputDynamicEventTypeWithId(const FName Id, const ANSICHAR* Name, const ANSICHAR* File = nullptr, uint32 Line = 0);
+	FORCENOINLINE CORE_API static uint32 OutputDynamicEventTypeWithId(const FName Id, const TCHAR* Name, const ANSICHAR* File = nullptr, uint32 Line = 0);
+
 	/*
 	 * Output begin event marker for a given spec. Must always be matched with an end event.
-	 * @param SpecId Event definition id.
+	 * @param SpecId - The event definition id.
 	 */
 	CORE_API static void OutputBeginEvent(uint32 SpecId);
+
 	/*
 	 * Output begin event marker for a dynamic event name. This is more expensive than statically known event
 	 * names using \ref OutputBeginEvent. Must always be matched with an end event.
-	 * @param Name Name of event
-	 * @param File Source filename
-	 * @param Line Line number in source file
+	 * @param Name - The name of event
+	 * @param File - The source filename
+	 * @param Line - The line number in source file
 	 */
 	CORE_API static void OutputBeginDynamicEvent(const ANSICHAR* Name, const ANSICHAR* File = nullptr, uint32 Line = 0);
-	/*
-	 * Output begin event marker for a dynamic event name. This is more expensive than statically known event
-	 * names using \ref OutputBeginEvent. Must always be matched with an end event.
-	 * @param Name Name of event
-	 * @param File Source filename
-	 * @param Line Line number in source file
-	 */
 	CORE_API static void OutputBeginDynamicEvent(const TCHAR* Name, const ANSICHAR* File = nullptr, uint32 Line = 0);
+
 	/*
 	 * Output begin event marker for a dynamic event identified by an FName. This is more expensive than
 	 * statically known event names using \ref OutputBeginEvent, but it is faster than \ref OutputBeginDynamicEvent
 	 * that receives ANSICHAR* / TCHAR* name. Must always be matched with an end event.
-	 * @param Name Name of event
-	 * @param File Source filename
-	 * @param Line Line number in source file
+	 * @param Name - The name of event
+	 * @param File - The source filename
+	 * @param Line - The line number in source file
 	 */
-	FORCENOINLINE CORE_API static void OutputBeginDynamicEvent(const FName Name, const ANSICHAR* File = nullptr, uint32 Line = 0);
+	CORE_API static void OutputBeginDynamicEvent(const FName Name, const ANSICHAR* File = nullptr, uint32 Line = 0);
+
 	/*
 	 * Output begin event marker for a dynamic event identified by an FName. This is more expensive than
 	 * statically known event names using \ref OutputBeginEvent, but it is faster than \ref OutputBeginDynamicEvent
 	 * that receives ANSICHAR* / TCHAR* name. Must always be matched with an end event.
-	 * @param Id Id of event
-	 * @param Name Name of event
-	 * @param File Source filename
-	 * @param Line Line number in source file
+	 * @param Id - The id of event
+	 * @param Name - The name of event
+	 * @param File - The source filename
+	 * @param Line - The line number in source file
 	 */
+	CORE_API static void OutputBeginDynamicEventWithId(const FName Id, const ANSICHAR* Name, const ANSICHAR* File = nullptr, uint32 Line = 0);
 	CORE_API static void OutputBeginDynamicEventWithId(const FName Id, const TCHAR* Name, const ANSICHAR* File = nullptr, uint32 Line = 0);
+
 	/*
 	 * Output end event marker for static or dynamic event for the currently open scope.
 	 */
 	CORE_API static void OutputEndEvent();
+
 	/*
 	* Output resume marker for a given spec. Must always be matched with an suspend event.
-	* @param SpecId unique Resume Event definition id.
-	* @param TimerScopeDepth updates the depth of the current OutputBeginEvent depth.
+	* @param SpecId - The unique Resume Event definition id
+	* @param TimerScopeDepth - updates the depth of the current OutputBeginEvent depth
 	*/
 	CORE_API static void OutputResumeEvent(uint64 SpecId, uint32& TimerScopeDepth);
+
 	/*
 	* Output suspend event marker for the currently open resume event.
 	*/
@@ -253,22 +273,22 @@ struct FCpuProfilerTrace
 	};
 };
 
-// Advanced macro for integrating e.g. stats/named events with cpu trace
-// Declares a cpu timing event for future use, conditionally and with a particular variable name for storing the id
+// Advanced macro for integrating e.g. stats/named events with CPU trace.
+// Declares a CPU timing event for future use, conditionally and with a particular variable name for storing the id.
 #define TRACE_CPUPROFILER_EVENT_DECLARE(DeclName) \
 	static uint32 DeclName;
 
-// Advanced macro for integrating e.g. stats/named events with cpu trace
-// Traces a scoped event previously declared with TRACE_CPUPROFILER_EVENT_DECLARE, conditionally
+// Advanced macro for integrating e.g. stats/named events with CPU trace.
+// Traces a scoped event previously declared with TRACE_CPUPROFILER_EVENT_DECLARE, conditionally.
 #define TRACE_CPUPROFILER_EVENT_SCOPE_USE(DeclName, NameStr, ScopeName, Condition) \
 	FCpuProfilerTrace::FEventScope ScopeName(DeclName, NameStr, Condition, __FILE__, __LINE__);
 
-// Advanced macro for integrating e.g. stats/named events with cpu trace
+// Advanced macro for integrating e.g. stats/named events with CPU trace
 // Traces a scoped event previously declared with TRACE_CPUPROFILER_EVENT_DECLARE, conditionally
 #define TRACE_CPUPROFILER_EVENT_SCOPE_USE_ON_CHANNEL(DeclName, NameStr, ScopeName, Channel, Condition) \
 	FCpuProfilerTrace::FEventScope ScopeName(DeclName, NameStr, Channel, Condition, __FILE__, __LINE__);
 
-// Advanced macro that will check if CpuChannel is enabled and, if so, declare a new stat event and start it
+// Advanced macro that will check if CpuChannel is enabled and, if so, declares a new CPU event and starts it.
 #define TRACE_CPUPROFILER_EVENT_MANUAL_START(EventNameStr) \
 	if (CpuChannel) \
 	{ \
@@ -276,56 +296,63 @@ struct FCpuProfilerTrace
 		FCpuProfilerTrace::OutputBeginEvent(PREPROCESSOR_JOIN(__CpuProfilerEventSpecId, __LINE__)); \
 	}
 
-// Advanced macro that can be used with TRACE_CPUPROFILER_EVENT_MANUAL_START to wrap code that should only be executed if 
-// the event was actually started
+// Advanced macro that will check if CpuChannel is enabled and, if so, ends the previously started CPU event.
+#define TRACE_CPUPROFILER_EVENT_MANUAL_END() \
+	if (CpuChannel) \
+	{ \
+		FCpuProfilerTrace::OutputEndEvent(); \
+	}
+
+// Advanced macro that can be used with TRACE_CPUPROFILER_EVENT_MANUAL_START to wrap code that should only be executed if
+// the event was actually started.
 #define TRACE_CPUPROFILER_EVENT_MANUAL_IS_ENABLED() \
 	bool(CpuChannel)
 
-// Conditionally trace a scoped cpu timing event providing a static string (const ANSICHAR* or const TCHAR*)
-// as the scope name and a condition under which to create the trace. It will use the Cpu trace channel.
+// Conditionally trace a scoped CPU timing event providing a static string (const ANSICHAR* or const TCHAR*)
+// as the scope name and a condition under which to create the trace. It will use the CPU trace channel.
 // Example: TRACE_CPUPROFILER_EVENT_SCOPE_STR_CONDITIONAL("My Scoped Timer A", Condition)
 #define TRACE_CPUPROFILER_EVENT_SCOPE_STR_CONDITIONAL(NameStr, Condition) \
 	TRACE_CPUPROFILER_EVENT_DECLARE(PREPROCESSOR_JOIN(__CpuProfilerEventSpecId, __LINE__)); \
-	TRACE_CPUPROFILER_EVENT_SCOPE_USE(PREPROCESSOR_JOIN(__CpuProfilerEventSpecId, __LINE__), NameStr, PREPROCESSOR_JOIN(__CpuProfilerEventScope, __LINE__), (Condition)); 
+	TRACE_CPUPROFILER_EVENT_SCOPE_USE(PREPROCESSOR_JOIN(__CpuProfilerEventSpecId, __LINE__), NameStr, PREPROCESSOR_JOIN(__CpuProfilerEventScope, __LINE__), (Condition));
 
-// Trace a scoped cpu timing event providing a static string (const ANSICHAR* or const TCHAR*)
-// as the scope name. It will use the Cpu trace channel.
+// Trace a scoped CPU timing event providing a static string (const ANSICHAR* or const TCHAR*)
+// as the scope name. It will use the CPU trace channel.
 // Example: TRACE_CPUPROFILER_EVENT_SCOPE_STR("My Scoped Timer A")
 #define TRACE_CPUPROFILER_EVENT_SCOPE_STR(NameStr) \
 	TRACE_CPUPROFILER_EVENT_SCOPE_STR_CONDITIONAL(NameStr, true)
 
-// Conditionally trace a scoped cpu timing event providing a static string (const ANSICHAR* or const TCHAR*)
+// Conditionally trace a scoped CPU timing event providing a static string (const ANSICHAR* or const TCHAR*)
 // as the scope name and a trace channel and a condition.
 // Example: TRACE_CPUPROFILER_EVENT_SCOPE_ON_CHANNEL_STR_CONDITIONAL("My Scoped Timer A", CpuChannel, Condition)
 // Note: The event will be emitted only if both the given channel and CpuChannel is enabled.
 #define TRACE_CPUPROFILER_EVENT_SCOPE_ON_CHANNEL_STR_CONDITIONAL(NameStr, Channel, Condition) \
 	TRACE_CPUPROFILER_EVENT_DECLARE(PREPROCESSOR_JOIN(__CpuProfilerEventSpecId, __LINE__)); \
-	TRACE_CPUPROFILER_EVENT_SCOPE_USE_ON_CHANNEL(PREPROCESSOR_JOIN(__CpuProfilerEventSpecId, __LINE__), NameStr, PREPROCESSOR_JOIN(__CpuProfilerEventScope, __LINE__), Channel, (Condition)); 
+	TRACE_CPUPROFILER_EVENT_SCOPE_USE_ON_CHANNEL(PREPROCESSOR_JOIN(__CpuProfilerEventSpecId, __LINE__), NameStr, PREPROCESSOR_JOIN(__CpuProfilerEventScope, __LINE__), Channel, (Condition));
 
-// Trace a scoped cpu timing event providing a static string (const ANSICHAR* or const TCHAR*)
+// Trace a scoped CPU timing event providing a static string (const ANSICHAR* or const TCHAR*)
 // as the scope name and a trace channel.
 // Example: TRACE_CPUPROFILER_EVENT_SCOPE_ON_CHANNEL_STR("My Scoped Timer A", CpuChannel)
 // Note: The event will be emitted only if both the given channel and CpuChannel is enabled.
 #define TRACE_CPUPROFILER_EVENT_SCOPE_ON_CHANNEL_STR(NameStr, Channel) \
 	TRACE_CPUPROFILER_EVENT_SCOPE_ON_CHANNEL_STR_CONDITIONAL(NameStr, Channel, true)
 
-// Conditionally trace a scoped cpu timing event providing a scope name (plain text) and a condition.
-// It will use the Cpu trace channel.
+// Conditionally trace a scoped CPU timing event providing a scope name (plain text) and a condition.
+// It will use the CPU trace channel.
 // Example: TRACE_CPUPROFILER_EVENT_SCOPE_CONDITIONAL(MyScopedTimer::A, Condition)
 // Note: Do not use this macro with a static string because, in that case, additional quotes will
 //       be added around the event scope name.
 #define TRACE_CPUPROFILER_EVENT_SCOPE_CONDITIONAL(Name, Condition) \
 	TRACE_CPUPROFILER_EVENT_SCOPE_STR_CONDITIONAL(#Name, (Condition))
 
-// Trace a scoped cpu timing event providing a scope name (plain text).
-// It will use the Cpu trace channel.
+// Trace a scoped CPU timing event providing a scope name (plain text).
+// It will use the CPU trace channel.
 // Example: TRACE_CPUPROFILER_EVENT_SCOPE(MyScopedTimer::A)
 // Note: Do not use this macro with a static string because, in that case, additional quotes will
 //       be added around the event scope name.
 #define TRACE_CPUPROFILER_EVENT_SCOPE(Name) \
 	TRACE_CPUPROFILER_EVENT_SCOPE_CONDITIONAL(Name, true)
 
-// Conditionally trace a scoped cpu timing event providing a scope name (plain text), a trace channel, and a condition.
+// Conditionally trace a scoped CPU timing event providing a scope name (plain text), a trace channel, and a condition.
 // Example: TRACE_CPUPROFILER_EVENT_SCOPE_ON_CHANNEL_CONDITIONAL(MyScopedTimer::A, CpuChannel, Condition)
 // Note: Do not use this macro with a static string because, in that case, additional quotes will
 //       be added around the event scope name.
@@ -333,7 +360,7 @@ struct FCpuProfilerTrace
 #define TRACE_CPUPROFILER_EVENT_SCOPE_ON_CHANNEL_CONDITIONAL(Name, Channel, Condition) \
 	TRACE_CPUPROFILER_EVENT_SCOPE_ON_CHANNEL_STR_CONDITIONAL(#Name, Channel, (Condition))
 
-// Trace a scoped cpu timing event providing a scope name (plain text) and a trace channel.
+// Trace a scoped CPU timing event providing a scope name (plain text) and a trace channel.
 // Example: TRACE_CPUPROFILER_EVENT_SCOPE_ON_CHANNEL(MyScopedTimer::A, CpuChannel)
 // Note: Do not use this macro with a static string because, in that case, additional quotes will
 //       be added around the event scope name.
@@ -341,7 +368,7 @@ struct FCpuProfilerTrace
 #define TRACE_CPUPROFILER_EVENT_SCOPE_ON_CHANNEL(Name, Channel) \
 	TRACE_CPUPROFILER_EVENT_SCOPE_ON_CHANNEL_CONDITIONAL(Name, Channel, true)
 
-// Conditionally trace a scoped cpu timing event providing a dynamic string (const ANSICHAR* or const TCHAR*)
+// Conditionally trace a scoped CPU timing event providing a dynamic string (const ANSICHAR* or const TCHAR*)
 // as the scope name and a trace channel.
 // Example: TRACE_CPUPROFILER_EVENT_SCOPE_TEXT_ON_CHANNEL_CONDITIONAL(*MyScopedTimerNameString, CpuChannel, Condition)
 // Note: This macro has a larger overhead compared to macro that accepts a plain text name
@@ -349,7 +376,7 @@ struct FCpuProfilerTrace
 #define TRACE_CPUPROFILER_EVENT_SCOPE_TEXT_ON_CHANNEL_CONDITIONAL(Name, Channel, Condition) \
 	FCpuProfilerTrace::FDynamicEventScope PREPROCESSOR_JOIN(__CpuProfilerEventScope, __LINE__)(Name, Channel, (Condition), __FILE__, __LINE__);
 
-// Trace a scoped cpu timing event providing a dynamic string (const ANSICHAR* or const TCHAR*)
+// Trace a scoped CPU timing event providing a dynamic string (const ANSICHAR* or const TCHAR*)
 // as the scope name and a trace channel.
 // Example: TRACE_CPUPROFILER_EVENT_SCOPE_TEXT_ON_CHANNEL(*MyScopedTimerNameString, CpuChannel)
 // Note: This macro has a larger overhead compared to macro that accepts a plain text name
@@ -357,16 +384,16 @@ struct FCpuProfilerTrace
 #define TRACE_CPUPROFILER_EVENT_SCOPE_TEXT_ON_CHANNEL(Name, Channel) \
 	TRACE_CPUPROFILER_EVENT_SCOPE_TEXT_ON_CHANNEL_CONDITIONAL(Name, Channel, true)
 
-// Conditionally trace a scoped cpu timing event providing a dynamic string (const ANSICHAR* or const TCHAR*)
-// as the scope name. It will use the Cpu trace channel.
+// Conditionally trace a scoped CPU timing event providing a dynamic string (const ANSICHAR* or const TCHAR*)
+// as the scope name. It will use the CPU trace channel.
 // Example: TRACE_CPUPROFILER_EVENT_SCOPE_TEXT_CONDITIONAL(*MyScopedTimerNameString, Condition)
 // Note: This macro has a larger overhead compared to macro that accepts a plain text name
 //       or a static string. Use it only if scope name really needs to be a dynamic string.
 #define TRACE_CPUPROFILER_EVENT_SCOPE_TEXT_CONDITIONAL(Name, Condition) \
 	FCpuProfilerTrace::FDynamicEventScope PREPROCESSOR_JOIN(__CpuProfilerEventScope, __LINE__)(Name, (Condition), __FILE__, __LINE__);
 
-// Trace a scoped cpu timing event providing a dynamic string (const ANSICHAR* or const TCHAR*)
-// as the scope name. It will use the Cpu trace channel.
+// Trace a scoped CPU timing event providing a dynamic string (const ANSICHAR* or const TCHAR*)
+// as the scope name. It will use the CPU trace channel.
 // Example: TRACE_CPUPROFILER_EVENT_SCOPE_TEXT(*MyScopedTimerNameString)
 // Note: This macro has a larger overhead compared to macro that accepts a plain text name
 //       or a static string. Use it only if scope name really needs to be a dynamic string.
@@ -376,14 +403,15 @@ struct FCpuProfilerTrace
 // Make sure all thread data has reached the destination.
 // Note: Can be useful to call this before entering a wait condition that might take a while.
 #define TRACE_CPUPROFILER_EVENT_FLUSH() \
-	FCpuProfilerTrace::FlushThreadBuffer(); 
+	FCpuProfilerTrace::FlushThreadBuffer();
 
-#else
+#else // CPUPROFILERTRACE_ENABLED
 
 #define TRACE_CPUPROFILER_EVENT_DECLARE(DeclName)
 #define TRACE_CPUPROFILER_EVENT_SCOPE_USE(DeclName, NameStr, ScopeName, Condition)
 #define TRACE_CPUPROFILER_EVENT_SCOPE_USE_ON_CHANNEL(DeclName, NameStr, ScopeName, Channel, Condition)
 #define TRACE_CPUPROFILER_EVENT_MANUAL_START(EventNameStr)
+#define TRACE_CPUPROFILER_EVENT_MANUAL_END()
 #define TRACE_CPUPROFILER_EVENT_MANUAL_IS_ENABLED() false
 #define TRACE_CPUPROFILER_EVENT_SCOPE_STR(NameStr)
 #define TRACE_CPUPROFILER_EVENT_SCOPE_STR_CONDITIONAL(NameStr, Condition)
@@ -399,4 +427,4 @@ struct FCpuProfilerTrace
 #define TRACE_CPUPROFILER_EVENT_SCOPE_TEXT_CONDITIONAL(Name, Condition)
 #define TRACE_CPUPROFILER_EVENT_FLUSH()
 
-#endif
+#endif // CPUPROFILERTRACE_ENABLED
