Added clamp function to Interval class

diff --git a/Engine/Source/Runtime/Core/Public/Math/Interval.h b/Engine/Source/Runtime/Core/Public/Math/Interval.h
--- a/Engine/Source/Runtime/Core/Public/Math/Interval.h
+++ b/Engine/Source/Runtime/Core/Public/Math/Interval.h
@@ -201,6 +201,22 @@ public:
 		return ElementType();
 	}
 
+	/**
+	 * Clamps X to be between the interval inclusively.
+	 *
+	 * @param X the element to clamp
+	 * @return zero if invalid(Min > Max), clamped result otherwise
+	 */
+	ElementType Clamp( ElementType X ) const
+	{
+		if (!IsValid())
+		{
+			return ElementType();
+		}
+
+		return FMath::Clamp(X, Min, Max);
+	}
+
 public:
 
 	/**
