Avoid a strlen() when adding to KeysTable, as we already know the length of the string

#rnx

diff --git a/Engine/Source/Runtime/Core/Private/Internationalization/TextKey.cpp b/Engine/Source/Runtime/Core/Private/Internationalization/TextKey.cpp
--- a/Engine/Source/Runtime/Core/Private/Internationalization/TextKey.cpp
+++ b/Engine/Source/Runtime/Core/Private/Internationalization/TextKey.cpp
@@ -148,7 +148,7 @@ private:
 		{
 			LLM_SCOPE_BYNAME(TEXT("Localization/TextKeys"));
 			StrPtr = StringAllocations.Add(KeyData.ToView());
-			KeysTable.Add(FKeyData(StrPtr, KeyData));
+			KeysTable.Add(FKeyData(FStringView(StrPtr, KeyData.StrLen), KeyData));
 		}
 		return StrPtr;
 #else
