Reduce virtual address space to 512 MB per bin when using MB3 on iOS

diff --git a/Engine/Source/Runtime/Core/Public/HAL/MallocBinned3.h b/Engine/Source/Runtime/Core/Public/HAL/MallocBinned3.h
--- a/Engine/Source/Runtime/Core/Public/HAL/MallocBinned3.h
+++ b/Engine/Source/Runtime/Core/Public/HAL/MallocBinned3.h
@@ -43,6 +43,10 @@ struct FGenericMemoryStats;
 #endif
 #define UE_MB3_SMALL_POOL_COUNT						(UE_MBC_NUM_LISTED_SMALL_POOLS + (UE_MB3_MAX_SMALL_POOL_SIZE - UE_MBC_MAX_LISTED_SMALL_POOL_SIZE) / UE_MB3_BASE_PAGE_SIZE)
 
+#if PLATFORM_IOS
+#	define USE_512MB_MAX_MEMORY_PER_BLOCK_SIZE		1
+#endif
+
 #if USE_512MB_MAX_MEMORY_PER_BLOCK_SIZE
 #	define UE_MB3_MAX_MEMORY_PER_POOL_SIZE_SHIFT	(29) // maximum of 512MB per pool
 #else
