Validate that the requested OS locale is supported by ICU, and if not, use the requested language instead

We have some platforms that are returning unsupported locales (eg, a locale of en-IT, with a language of it-IT), so this change hardens the default locale set-up to validate that ICU has resource data for the requested locale, and if not, we use the requested language as the locale instead

diff --git a/Engine/Source/Runtime/Core/Private/Internationalization/ICUInternationalization.cpp b/Engine/Source/Runtime/Core/Private/Internationalization/ICUInternationalization.cpp
--- a/Engine/Source/Runtime/Core/Private/Internationalization/ICUInternationalization.cpp
+++ b/Engine/Source/Runtime/Core/Private/Internationalization/ICUInternationalization.cpp
@@ -178,7 +178,16 @@ bool FICUInternationalization::Initialize()
 		I18N->InvariantCulture = FindOrMakeCanonizedCulture(FString(), EAllowDefaultCultureFallback::Yes);
 	}
 	I18N->DefaultLanguage = FindOrMakeCulture(FPlatformMisc::GetDefaultLanguage(), EAllowDefaultCultureFallback::Yes);
-	I18N->DefaultLocale = FindOrMakeCulture(FPlatformMisc::GetDefaultLocale(), EAllowDefaultCultureFallback::Yes);
+	if (const FString& DefaultLocaleName = FPlatformMisc::GetDefaultLocale();
+		AllAvailableCulturesMap.Contains(DefaultLocaleName))
+	{
+		I18N->DefaultLocale = FindOrMakeCulture(DefaultLocaleName, EAllowDefaultCultureFallback::Yes);
+	}
+	else
+	{
+		UE_LOG(LogICUInternationalization, Log, TEXT("OS requested locale '%s' is not supported. Using the OS requested language of '%s' as the locale."), *DefaultLocaleName, *I18N->DefaultLanguage->GetName());
+		I18N->DefaultLocale = I18N->DefaultLanguage;
+	}
 	I18N->CurrentLanguage = I18N->DefaultLanguage;
 	I18N->CurrentLocale = I18N->DefaultLocale;
 
