Add member initializers for FMeshCardsBuildData to avoid Bounds being serialized with NaN values if MaxLumenMeshCards is 0. In this case FMeshUtilities::GenerateCardRepresentationData doesn't populate FCardRepresentationData.
Bumped the static mesh DDC key to flush out any cached data containing NaNs.

#tests Rebuilt mesh cards while MaxLumenMeshCards was set to 0. Observed in Visual Studio that the bounds now consistently initialize zeroed.

diff --git a/Engine/Source/Runtime/Core/Private/UObject/FortniteMainBranchObjectVersion.cpp b/Engine/Source/Runtime/Core/Private/UObject/FortniteMainBranchObjectVersion.cpp
--- a/Engine/Source/Runtime/Core/Private/UObject/FortniteMainBranchObjectVersion.cpp
+++ b/Engine/Source/Runtime/Core/Private/UObject/FortniteMainBranchObjectVersion.cpp
@@ -13,7 +13,7 @@ TMap<FGuid, FGuid> FFortniteMainBranchObjectVersion::GetSystemGuids()
 	SystemGuids.Add(DevGuids.NIAGARASHADERMAP_DERIVEDDATA_VER, FGuid("C417A50271B3427D98785388266E3B66"));
 	SystemGuids.Add(DevGuids.Niagara_LatestScriptCompileVersion, FGuid("E486231DA967C64B95D4E73131056FA0"));
 	SystemGuids.Add(DevGuids.SkeletalMeshDerivedDataVersion, FGuid("B9AEAA3EE7AC4FCFB3860E8F3E06DD75"));
-	SystemGuids.Add(DevGuids.STATICMESH_DERIVEDDATA_VER, FGuid("7F8E1537-E2A9-4A9C-8CF2-A6272A5A8176"));
+	SystemGuids.Add(DevGuids.STATICMESH_DERIVEDDATA_VER, FGuid("67EF7DA7-C4F9-4751-85DB-5A6E7E74860C"));
 	SystemGuids.Add(DevGuids.MaterialTranslationDDCVersion, FGuid("5F69E03D3B204DA38D8CF564CB24F23F"));
 	return SystemGuids;
 }
diff --git a/Engine/Source/Runtime/Engine/Public/MeshCardBuild.h b/Engine/Source/Runtime/Engine/Public/MeshCardBuild.h
--- a/Engine/Source/Runtime/Engine/Public/MeshCardBuild.h
+++ b/Engine/Source/Runtime/Engine/Public/MeshCardBuild.h
@@ -87,8 +87,8 @@ public:
 class FMeshCardsBuildData
 {
 public:
-	FBox Bounds;
-	bool bMostlyTwoSided;
+	FBox Bounds = FBox(ForceInit);
+	bool bMostlyTwoSided = false;
 	TArray<FLumenCardBuildData> CardBuildData;
 
 	// Temporary debug visualization data, don't serialize
