Reduced delegate implementation impact by reimplementing TDelegateBase::CreateDelegateInstance as an operator new overload.

diff --git a/Engine/Source/Runtime/Core/Private/Delegates/DelegateBase.cpp b/Engine/Source/Runtime/Core/Private/Delegates/DelegateBase.cpp
new file mode 100644
--- /dev/null
+++ b/Engine/Source/Runtime/Core/Private/Delegates/DelegateBase.cpp
@@ -0,0 +1,15 @@
+// Copyright Epic Games, Inc. All Rights Reserved.
+
+#include "Delegates/DelegateBase.h"
+
+void* UE::Core::Private::DelegateAllocate(size_t Size, FDelegateAllocation& Allocation)
+{
+	int32 NewDelegateSize = FMath::DivideAndRoundUp((int32)Size, (int32)sizeof(FAlignedInlineDelegateType));
+	if (Allocation.DelegateSize != NewDelegateSize)
+	{
+		Allocation.DelegateAllocator.ResizeAllocation(0, NewDelegateSize, sizeof(FAlignedInlineDelegateType));
+		Allocation.DelegateSize = NewDelegateSize;
+	}
+
+	return Allocation.DelegateAllocator.GetAllocation();
+}
diff --git a/Engine/Source/Runtime/Core/Public/Delegates/DelegateAccessHandler.h b/Engine/Source/Runtime/Core/Public/Delegates/DelegateAccessHandler.h
--- a/Engine/Source/Runtime/Core/Public/Delegates/DelegateAccessHandler.h
+++ b/Engine/Source/Runtime/Core/Public/Delegates/DelegateAccessHandler.h
@@ -30,6 +30,9 @@ struct FNotThreadSafeDelegateMode;
 template<typename ThreadSafetyMode>
 class TDelegateAccessHandlerBase;
 
+template <typename ThreadSafetyMode>
+struct TWriteLockedDelegateAllocation;
+
 /**
  * non thread-safe version that does not do any race detection. supposed to be used in a controlled environment that provides own
  * detection or synchronisation.
@@ -37,6 +40,8 @@ class TDelegateAccessHandlerBase;
 template<>
 class TDelegateAccessHandlerBase<FNotThreadSafeNotCheckedDelegateMode>
 {
+	friend struct TWriteLockedDelegateAllocation<FNotThreadSafeNotCheckedDelegateMode>;
+
 protected:
 	struct FReadAccessScope {};
 	struct FWriteAccessScope {};
@@ -66,6 +71,8 @@ struct TIsZeroConstructType<TDelegateAccessHandlerBase<FNotThreadSafeNotCheckedD
 template<>
 class TDelegateAccessHandlerBase<FThreadSafeDelegateMode>
 {
+	friend struct TWriteLockedDelegateAllocation<FThreadSafeDelegateMode>;
+
 protected:
 	struct FReadAccessScope { FTransactionallySafeScopeLock Lock; };
 	struct FWriteAccessScope { FTransactionallySafeScopeLock Lock; };
@@ -98,6 +105,8 @@ struct TIsZeroConstructType<TDelegateAccessHandlerBase<FThreadSafeDelegateMode>>
 template<>
 class TDelegateAccessHandlerBase<FNotThreadSafeDelegateMode>
 {
+	friend struct TWriteLockedDelegateAllocation<FNotThreadSafeDelegateMode>;
+
 protected:
 #if !UE_AUTORTFM
 	class FReadAccessScope
diff --git a/Engine/Source/Runtime/Core/Public/Delegates/DelegateBase.h b/Engine/Source/Runtime/Core/Public/Delegates/DelegateBase.h
--- a/Engine/Source/Runtime/Core/Public/Delegates/DelegateBase.h
+++ b/Engine/Source/Runtime/Core/Public/Delegates/DelegateBase.h
@@ -91,12 +91,44 @@ struct FNotThreadSafeNotCheckedDelegateUserPolicy
 	using FMulticastDelegateExtras = TMulticastDelegateBase<FNotThreadSafeNotCheckedDelegateUserPolicy>;
 };
 
+struct FDelegateAllocation
+{
+	FDelegateAllocatorType::ForElementType<FAlignedInlineDelegateType> DelegateAllocator;
+	int32 DelegateSize = 0;
+};
+
+template <typename ThreadSafetyMode>
+struct TWriteLockedDelegateAllocation;
+
+template <typename ThreadSafetyMode>
+void* operator new(size_t Size, const TWriteLockedDelegateAllocation<ThreadSafetyMode>& LockedAllocation);
+
+template <typename ThreadSafetyMode>
+struct TWriteLockedDelegateAllocation
+{
+	UE_NONCOPYABLE(TWriteLockedDelegateAllocation)
+
+	friend void* operator new<ThreadSafetyMode>(size_t Size, const TWriteLockedDelegateAllocation<ThreadSafetyMode>& LockedAllocation);
+
+	explicit TWriteLockedDelegateAllocation(TDelegateBase<ThreadSafetyMode>& Delegate)
+		: WriteScope(Delegate.GetWriteAccessScope())
+		, Allocation(Delegate)
+	{
+	}
+
+private:
+	typename TDelegateAccessHandlerBase<ThreadSafetyMode>::FWriteAccessScope WriteScope;
+	FDelegateAllocation& Allocation;
+};
+
 /**
  * Base class for unicast delegates.
  */
 template<typename ThreadSafetyMode>
-class TDelegateBase : public TDelegateAccessHandlerBase<ThreadSafetyMode>
+class TDelegateBase : public TDelegateAccessHandlerBase<ThreadSafetyMode>, private FDelegateAllocation
 {
+	friend struct TWriteLockedDelegateAllocation<ThreadSafetyMode>;
+
 private:
 	using Super = TDelegateAccessHandlerBase<ThreadSafetyMode>;
 
@@ -130,6 +162,8 @@ private:
 	template <typename, typename, typename, typename, typename...>
 	friend class TWeakBaseFunctorDelegateInstance;
 
+	friend class FWriteLockedDelegateAllocation;
+
 protected:
 	using typename Super::FReadAccessScope;
 	using Super::GetReadAccessScope;
@@ -303,23 +337,6 @@ public:
 	}
 
 protected:
-	/**
-	 * "emplacement" of delegate instance of the given type
-	 */
-	template<typename DelegateInstanceType, typename... DelegateInstanceParams>
-	void CreateDelegateInstance(DelegateInstanceParams&&... Params)
-	{
-		FWriteAccessScope WriteScope = GetWriteAccessScope();
-
-		IDelegateInstance* DelegateInstance = GetDelegateInstanceProtected();
-		if (DelegateInstance)
-		{
-			DelegateInstance->~IDelegateInstance();
-		}
-
-		new(Allocate(sizeof(DelegateInstanceType))) DelegateInstanceType(Forward<DelegateInstanceParams>(Params)...);
-	}
-
 	/**
 	 * Gets the delegate instance.  Not intended for use by user code.
 	 *
@@ -337,18 +354,6 @@ protected:
 	}
 
 private:
-	void* Allocate(int32 Size)
-	{
-		int32 NewDelegateSize = FMath::DivideAndRoundUp(Size, (int32)sizeof(FAlignedInlineDelegateType));
-		if (DelegateSize != NewDelegateSize)
-		{
-			DelegateAllocator.ResizeAllocation(0, NewDelegateSize, sizeof(FAlignedInlineDelegateType));
-			DelegateSize = NewDelegateSize;
-		}
-
-		return DelegateAllocator.GetAllocation();
-	}
-
 	template<typename OtherThreadSafetyMode>
 	void MoveConstruct(TDelegateBase<OtherThreadSafetyMode>&& Other)
 	{
@@ -391,6 +396,18 @@ private:
 		}
 	}
 
-	FDelegateAllocatorType::ForElementType<FAlignedInlineDelegateType> DelegateAllocator;
-	int32 DelegateSize = 0;
+	using FDelegateAllocation::DelegateAllocator;
+	using FDelegateAllocation::DelegateSize;
 };
+
+namespace UE::Core::Private
+{
+	// Should only be called by the TWriteLockedDelegateAllocation<ThreadSafetyMode> overload, because it obtains the write lock
+	CORE_API void* DelegateAllocate(size_t Size, FDelegateAllocation& Allocation);
+}
+
+template <typename ThreadSafetyMode>
+FORCEINLINE void* operator new(size_t Size, const TWriteLockedDelegateAllocation<ThreadSafetyMode>& LockedAllocation)
+{
+	return UE::Core::Private::DelegateAllocate(Size, LockedAllocation.Allocation);
+}
diff --git a/Engine/Source/Runtime/Core/Public/Delegates/DelegateInstancesImpl.h b/Engine/Source/Runtime/Core/Public/Delegates/DelegateInstancesImpl.h
--- a/Engine/Source/Runtime/Core/Public/Delegates/DelegateInstancesImpl.h
+++ b/Engine/Source/Runtime/Core/Public/Delegates/DelegateInstancesImpl.h
@@ -129,17 +129,17 @@ public:
 
 	void CreateCopy(TDelegateBase<FThreadSafeDelegateMode>& Base) const final
 	{
-		Base.template CreateDelegateInstance<TBaseUFunctionDelegateInstance>(*this);
+		new (TWriteLockedDelegateAllocation{Base}) TBaseUFunctionDelegateInstance(*this);
 	}
 
 	void CreateCopy(TDelegateBase<FNotThreadSafeDelegateMode>& Base) const final
 	{
-		Base.template CreateDelegateInstance<TBaseUFunctionDelegateInstance>(*this);
+		new (TWriteLockedDelegateAllocation{Base}) TBaseUFunctionDelegateInstance(*this);
 	}
 
 	void CreateCopy(TDelegateBase<FNotThreadSafeNotCheckedDelegateMode>& Base) const final
 	{
-		Base.template CreateDelegateInstance<TBaseUFunctionDelegateInstance>(*this);
+		new (TWriteLockedDelegateAllocation{Base}) TBaseUFunctionDelegateInstance(*this);
 	}
 
 	RetValType Execute(ParamTypes... Params) const final
@@ -251,17 +251,17 @@ public:
 
 	void CreateCopy(TDelegateBase<FThreadSafeDelegateMode>& Base) const final
 	{
-		Base.template CreateDelegateInstance<TBaseSPMethodDelegateInstance>(*this);
+		new (TWriteLockedDelegateAllocation{Base}) TBaseSPMethodDelegateInstance(*this);
 	}
 
 	void CreateCopy(TDelegateBase<FNotThreadSafeDelegateMode>& Base) const final
 	{
-		Base.template CreateDelegateInstance<TBaseSPMethodDelegateInstance>(*this);
+		new (TWriteLockedDelegateAllocation{Base}) TBaseSPMethodDelegateInstance(*this);
 	}
 
 	void CreateCopy(TDelegateBase<FNotThreadSafeNotCheckedDelegateMode>& Base) const final
 	{
-		Base.template CreateDelegateInstance<TBaseSPMethodDelegateInstance>(*this);
+		new (TWriteLockedDelegateAllocation{Base}) TBaseSPMethodDelegateInstance(*this);
 	}
 
 	RetValType Execute(ParamTypes... Params) const final
@@ -375,17 +375,17 @@ public:
 	// IBaseDelegateInstance interface
 	void CreateCopy(TDelegateBase<FThreadSafeDelegateMode>& Base) const final
 	{
-		Base.template CreateDelegateInstance<TBaseSPLambdaDelegateInstance>(*this);
+		new (TWriteLockedDelegateAllocation{Base}) TBaseSPLambdaDelegateInstance(*this);
 	}
 
 	void CreateCopy(TDelegateBase<FNotThreadSafeDelegateMode>& Base) const final
 	{
-		Base.template CreateDelegateInstance<TBaseSPLambdaDelegateInstance>(*this);
+		new (TWriteLockedDelegateAllocation{Base}) TBaseSPLambdaDelegateInstance(*this);
 	}
 
 	void CreateCopy(TDelegateBase<FNotThreadSafeNotCheckedDelegateMode>& Base) const final
 	{
-		Base.template CreateDelegateInstance<TBaseSPLambdaDelegateInstance>(*this);
+		new (TWriteLockedDelegateAllocation{Base}) TBaseSPLambdaDelegateInstance(*this);
 	}
 
 	RetValType Execute(ParamTypes... Params) const final
@@ -491,17 +491,17 @@ public:
 
 	void CreateCopy(TDelegateBase<FThreadSafeDelegateMode>& Base) const final
 	{
-		Base.template CreateDelegateInstance<TBaseRawMethodDelegateInstance>(*this);
+		new (TWriteLockedDelegateAllocation{Base}) TBaseRawMethodDelegateInstance(*this);
 	}
 
 	void CreateCopy(TDelegateBase<FNotThreadSafeDelegateMode>& Base) const final
 	{
-		Base.template CreateDelegateInstance<TBaseRawMethodDelegateInstance>(*this);
+		new (TWriteLockedDelegateAllocation{Base}) TBaseRawMethodDelegateInstance(*this);
 	}
 
 	void CreateCopy(TDelegateBase<FNotThreadSafeNotCheckedDelegateMode>& Base) const final
 	{
-		Base.template CreateDelegateInstance<TBaseRawMethodDelegateInstance>(*this);
+		new (TWriteLockedDelegateAllocation{Base}) TBaseRawMethodDelegateInstance(*this);
 	}
 
 	RetValType Execute(ParamTypes... Params) const final
@@ -620,17 +620,17 @@ public:
 
 	void CreateCopy(TDelegateBase<FThreadSafeDelegateMode>& Base) const final
 	{
-		Base.template CreateDelegateInstance<TBaseUObjectMethodDelegateInstance>(*this);
+		new (TWriteLockedDelegateAllocation{Base}) TBaseUObjectMethodDelegateInstance(*this);
 	}
 
 	void CreateCopy(TDelegateBase<FNotThreadSafeDelegateMode>& Base) const final
 	{
-		Base.template CreateDelegateInstance<TBaseUObjectMethodDelegateInstance>(*this);
+		new (TWriteLockedDelegateAllocation{Base}) TBaseUObjectMethodDelegateInstance(*this);
 	}
 
 	void CreateCopy(TDelegateBase<FNotThreadSafeNotCheckedDelegateMode>& Base) const final
 	{
-		Base.template CreateDelegateInstance<TBaseUObjectMethodDelegateInstance>(*this);
+		new (TWriteLockedDelegateAllocation{Base}) TBaseUObjectMethodDelegateInstance(*this);
 	}
 
 	RetValType Execute(ParamTypes... Params) const final
@@ -748,17 +748,17 @@ public:
 
 	void CreateCopy(TDelegateBase<FThreadSafeDelegateMode>& Base) const final
 	{
-		Base.template CreateDelegateInstance<TBaseStaticDelegateInstance>(*this);
+		new (TWriteLockedDelegateAllocation{Base}) TBaseStaticDelegateInstance(*this);
 	}
 
 	void CreateCopy(TDelegateBase<FNotThreadSafeDelegateMode>& Base) const final
 	{
-		Base.template CreateDelegateInstance<TBaseStaticDelegateInstance>(*this);
+		new (TWriteLockedDelegateAllocation{Base}) TBaseStaticDelegateInstance(*this);
 	}
 
 	void CreateCopy(TDelegateBase<FNotThreadSafeNotCheckedDelegateMode>& Base) const final
 	{
-		Base.template CreateDelegateInstance<TBaseStaticDelegateInstance>(*this);
+		new (TWriteLockedDelegateAllocation{Base}) TBaseStaticDelegateInstance(*this);
 	}
 
 	RetValType Execute(ParamTypes... Params) const final
@@ -846,17 +846,17 @@ public:
 	// IBaseDelegateInstance interface
 	void CreateCopy(TDelegateBase<FThreadSafeDelegateMode>& Base) const final
 	{
-		Base.template CreateDelegateInstance<TBaseFunctorDelegateInstance>(*this);
+		new (TWriteLockedDelegateAllocation{Base}) TBaseFunctorDelegateInstance(*this);
 	}
 
 	void CreateCopy(TDelegateBase<FNotThreadSafeDelegateMode>& Base) const final
 	{
-		Base.template CreateDelegateInstance<TBaseFunctorDelegateInstance>(*this);
+		new (TWriteLockedDelegateAllocation{Base}) TBaseFunctorDelegateInstance(*this);
 	}
 
 	void CreateCopy(TDelegateBase<FNotThreadSafeNotCheckedDelegateMode>& Base) const final
 	{
-		Base.template CreateDelegateInstance<TBaseFunctorDelegateInstance>(*this);
+		new (TWriteLockedDelegateAllocation{Base}) TBaseFunctorDelegateInstance(*this);
 	}
 
 	RetValType Execute(ParamTypes... Params) const final
@@ -945,17 +945,17 @@ public:
 	// IBaseDelegateInstance interface
 	void CreateCopy(TDelegateBase<FThreadSafeDelegateMode>& Base) const final
 	{
-		Base.template CreateDelegateInstance<TWeakBaseFunctorDelegateInstance>(*this);
+		new (TWriteLockedDelegateAllocation{Base}) TWeakBaseFunctorDelegateInstance(*this);
 	}
 
 	void CreateCopy(TDelegateBase<FNotThreadSafeDelegateMode>& Base) const final
 	{
-		Base.template CreateDelegateInstance<TWeakBaseFunctorDelegateInstance>(*this);
+		new (TWriteLockedDelegateAllocation{Base}) TWeakBaseFunctorDelegateInstance(*this);
 	}
 
 	void CreateCopy(TDelegateBase<FNotThreadSafeNotCheckedDelegateMode>& Base) const final
 	{
-		Base.template CreateDelegateInstance<TWeakBaseFunctorDelegateInstance>(*this);
+		new (TWriteLockedDelegateAllocation{Base}) TWeakBaseFunctorDelegateInstance(*this);
 	}
 
 	RetValType Execute(ParamTypes... Params) const final
diff --git a/Engine/Source/Runtime/Core/Public/Delegates/DelegateSignatureImpl.inl b/Engine/Source/Runtime/Core/Public/Delegates/DelegateSignatureImpl.inl
--- a/Engine/Source/Runtime/Core/Public/Delegates/DelegateSignatureImpl.inl
+++ b/Engine/Source/Runtime/Core/Public/Delegates/DelegateSignatureImpl.inl
@@ -116,7 +116,7 @@ public:
 	template <typename... VarTypes>
 	inline void BindStatic(typename TBaseStaticDelegateInstance<FuncType, UserPolicy, std::decay_t<VarTypes>...>::FFuncPtr InFunc, VarTypes&&... Vars)
 	{
-		Super::template CreateDelegateInstance<TBaseStaticDelegateInstance<FuncType, UserPolicy, std::decay_t<VarTypes>...>>(InFunc, Forward<VarTypes>(Vars)...);
+		new (TWriteLockedDelegateAllocation{*this}) TBaseStaticDelegateInstance<FuncType, UserPolicy, std::decay_t<VarTypes>...>(InFunc, Forward<VarTypes>(Vars)...);
 	}
 	
 	/**
@@ -126,7 +126,7 @@ public:
 	template<typename FunctorType, typename... VarTypes>
 	inline void BindLambda(FunctorType&& InFunctor, VarTypes&&... Vars)
 	{
-		Super::template CreateDelegateInstance<TBaseFunctorDelegateInstance<FuncType, UserPolicy, typename TRemoveReference<FunctorType>::Type, std::decay_t<VarTypes>...>>(Forward<FunctorType>(InFunctor), Forward<VarTypes>(Vars)...);
+		new (TWriteLockedDelegateAllocation{*this}) TBaseFunctorDelegateInstance<FuncType, UserPolicy, typename TRemoveReference<FunctorType>::Type, std::decay_t<VarTypes>...>(Forward<FunctorType>(InFunctor), Forward<VarTypes>(Vars)...);
 	}
 
 	/**
@@ -136,12 +136,12 @@ public:
 	template<typename UserClass, ESPMode Mode, typename FunctorType, typename... VarTypes>
 	inline void BindSPLambda(const TSharedRef<UserClass, Mode>& InUserObjectRef, FunctorType&& InFunctor, VarTypes&&... Vars)
 	{
-		Super::template CreateDelegateInstance<TBaseSPLambdaDelegateInstance<const UserClass, Mode, FuncType, UserPolicy, typename TRemoveReference<FunctorType>::Type, std::decay_t<VarTypes>...>>(InUserObjectRef, Forward<FunctorType>(InFunctor), Forward<VarTypes>(Vars)...);
+		new (TWriteLockedDelegateAllocation{*this}) TBaseSPLambdaDelegateInstance<const UserClass, Mode, FuncType, UserPolicy, typename TRemoveReference<FunctorType>::Type, std::decay_t<VarTypes>...>(InUserObjectRef, Forward<FunctorType>(InFunctor), Forward<VarTypes>(Vars)...);
 	}
 	template <typename UserClass, typename FunctorType, typename... VarTypes>
 	inline void BindSPLambda(const UserClass* InUserObject, FunctorType&& InFunctor, VarTypes&&... Vars)
 	{
-		Super::template CreateDelegateInstance<TBaseSPLambdaDelegateInstance<const UserClass, decltype(InUserObject->AsShared())::Mode, FuncType, UserPolicy, typename TRemoveReference<FunctorType>::Type, std::decay_t<VarTypes>...>>(StaticCastSharedRef<const UserClass>(InUserObject->AsShared()), Forward<FunctorType>(InFunctor), Forward<VarTypes>(Vars)...);
+		new (TWriteLockedDelegateAllocation{*this}) TBaseSPLambdaDelegateInstance<const UserClass, decltype(InUserObject->AsShared())::Mode, FuncType, UserPolicy, typename TRemoveReference<FunctorType>::Type, std::decay_t<VarTypes>...>(StaticCastSharedRef<const UserClass>(InUserObject->AsShared()), Forward<FunctorType>(InFunctor), Forward<VarTypes>(Vars)...);
 	}
 
 	/**
@@ -151,7 +151,7 @@ public:
 	template<typename UserClass, typename FunctorType, typename... VarTypes>
 	inline void BindWeakLambda(UserClass* InUserObject, FunctorType&& InFunctor, VarTypes&&... Vars)
 	{
-		Super::template CreateDelegateInstance<TWeakBaseFunctorDelegateInstance<UserClass, FuncType, UserPolicy, typename TRemoveReference<FunctorType>::Type, std::decay_t<VarTypes>...>>(InUserObject, Forward<FunctorType>(InFunctor), Forward<VarTypes>(Vars)...);
+		new (TWriteLockedDelegateAllocation{*this}) TWeakBaseFunctorDelegateInstance<UserClass, FuncType, UserPolicy, typename TRemoveReference<FunctorType>::Type, std::decay_t<VarTypes>...>(InUserObject, Forward<FunctorType>(InFunctor), Forward<VarTypes>(Vars)...);
 	}
 
 	/**
@@ -165,12 +165,12 @@ public:
 	{
 		static_assert(!TIsConst<UserClass>::Value, "Attempting to bind a delegate with a const object pointer and non-const member function.");
 
-		Super::template CreateDelegateInstance<TBaseRawMethodDelegateInstance<false, UserClass, FuncType, UserPolicy, std::decay_t<VarTypes>...>>(InUserObject, InFunc, Forward<VarTypes>(Vars)...);
+		new (TWriteLockedDelegateAllocation{*this}) TBaseRawMethodDelegateInstance<false, UserClass, FuncType, UserPolicy, std::decay_t<VarTypes>...>(InUserObject, InFunc, Forward<VarTypes>(Vars)...);
 	}
 	template <typename UserClass, typename... VarTypes>
 	inline void BindRaw(const UserClass* InUserObject, typename TMemFunPtrType<true, UserClass, RetValType (ParamTypes..., std::decay_t<VarTypes>...)>::Type InFunc, VarTypes&&... Vars)
 	{
-		Super::template CreateDelegateInstance<TBaseRawMethodDelegateInstance<true, const UserClass, FuncType, UserPolicy, std::decay_t<VarTypes>...>>(InUserObject, InFunc, Forward<VarTypes>(Vars)...);
+		new (TWriteLockedDelegateAllocation{*this}) TBaseRawMethodDelegateInstance<true, const UserClass, FuncType, UserPolicy, std::decay_t<VarTypes>...>(InUserObject, InFunc, Forward<VarTypes>(Vars)...);
 	}
 
 	/**
@@ -181,12 +181,12 @@ public:
 	{
 		static_assert(!TIsConst<UserClass>::Value, "Attempting to bind a delegate with a const object pointer and non-const member function.");
 
-		Super::template CreateDelegateInstance<TBaseSPMethodDelegateInstance<false, UserClass, Mode, FuncType, UserPolicy, std::decay_t<VarTypes>...>>(InUserObjectRef, InFunc, Forward<VarTypes>(Vars)...);
+		new (TWriteLockedDelegateAllocation{*this}) TBaseSPMethodDelegateInstance<false, UserClass, Mode, FuncType, UserPolicy, std::decay_t<VarTypes>...>(InUserObjectRef, InFunc, Forward<VarTypes>(Vars)...);
 	}
 	template <typename UserClass, ESPMode Mode, typename... VarTypes>
 	inline void BindSP(const TSharedRef<UserClass, Mode>& InUserObjectRef, typename TMemFunPtrType<true, UserClass, RetValType (ParamTypes..., std::decay_t<VarTypes>...)>::Type InFunc, VarTypes&&... Vars)
 	{
-		Super::template CreateDelegateInstance<TBaseSPMethodDelegateInstance<true, const UserClass, Mode, FuncType, UserPolicy, std::decay_t<VarTypes>...>>(InUserObjectRef, InFunc, Forward<VarTypes>(Vars)...);
+		new (TWriteLockedDelegateAllocation{*this}) TBaseSPMethodDelegateInstance<true, const UserClass, Mode, FuncType, UserPolicy, std::decay_t<VarTypes>...>(InUserObjectRef, InFunc, Forward<VarTypes>(Vars)...);
 	}
 
 	/**
@@ -200,12 +200,12 @@ public:
 	{
 		static_assert(!TIsConst<UserClass>::Value, "Attempting to bind a delegate with a const object pointer and non-const member function.");
 
-		Super::template CreateDelegateInstance<TBaseSPMethodDelegateInstance<false, UserClass, decltype(InUserObject->AsShared())::Mode, FuncType, UserPolicy, std::decay_t<VarTypes>...>>(StaticCastSharedRef<UserClass>(InUserObject->AsShared()), InFunc, Forward<VarTypes>(Vars)...);
+		new (TWriteLockedDelegateAllocation{*this}) TBaseSPMethodDelegateInstance<false, UserClass, decltype(InUserObject->AsShared())::Mode, FuncType, UserPolicy, std::decay_t<VarTypes>...>(StaticCastSharedRef<UserClass>(InUserObject->AsShared()), InFunc, Forward<VarTypes>(Vars)...);
 	}
 	template <typename UserClass, typename... VarTypes>
 	inline void BindSP(const UserClass* InUserObject, typename TMemFunPtrType<true, UserClass, RetValType (ParamTypes..., std::decay_t<VarTypes>...)>::Type InFunc, VarTypes&&... Vars)
 	{
-		Super::template CreateDelegateInstance<TBaseSPMethodDelegateInstance<true, const UserClass, decltype(InUserObject->AsShared())::Mode, FuncType, UserPolicy, std::decay_t<VarTypes>...>>(StaticCastSharedRef<const UserClass>(InUserObject->AsShared()), InFunc, Forward<VarTypes>(Vars)...);
+		new (TWriteLockedDelegateAllocation{*this}) TBaseSPMethodDelegateInstance<true, const UserClass, decltype(InUserObject->AsShared())::Mode, FuncType, UserPolicy, std::decay_t<VarTypes>...>(StaticCastSharedRef<const UserClass>(InUserObject->AsShared()), InFunc, Forward<VarTypes>(Vars)...);
 	}
 
 	/**
@@ -221,12 +221,12 @@ public:
 	{
 		static_assert(!TIsConst<UserClass>::Value, "Attempting to bind a delegate with a const object pointer and non-const member function.");
 
-		Super::template CreateDelegateInstance<TBaseSPMethodDelegateInstance<false, UserClass, ESPMode::ThreadSafe, FuncType, UserPolicy, std::decay_t<VarTypes>...>>(InUserObjectRef, InFunc, Forward<VarTypes>(Vars)...);
+		new (TWriteLockedDelegateAllocation{*this}) TBaseSPMethodDelegateInstance<false, UserClass, ESPMode::ThreadSafe, FuncType, UserPolicy, std::decay_t<VarTypes>...>(InUserObjectRef, InFunc, Forward<VarTypes>(Vars)...);
 	}
 	template <typename UserClass, typename... VarTypes>
 	inline void BindThreadSafeSP(const TSharedRef<UserClass, ESPMode::ThreadSafe>& InUserObjectRef, typename TMemFunPtrType<true, UserClass, RetValType (ParamTypes..., std::decay_t<VarTypes>...)>::Type InFunc, VarTypes&&... Vars)
 	{
-		Super::template CreateDelegateInstance<TBaseSPMethodDelegateInstance<true, const UserClass, ESPMode::ThreadSafe, FuncType, UserPolicy, std::decay_t<VarTypes>...>>(InUserObjectRef, InFunc, Forward<VarTypes>(Vars)...);
+		new (TWriteLockedDelegateAllocation{*this}) TBaseSPMethodDelegateInstance<true, const UserClass, ESPMode::ThreadSafe, FuncType, UserPolicy, std::decay_t<VarTypes>...>(InUserObjectRef, InFunc, Forward<VarTypes>(Vars)...);
 	}
 
 	/**
@@ -242,12 +242,12 @@ public:
 	{
 		static_assert(!TIsConst<UserClass>::Value, "Attempting to bind a delegate with a const object pointer and non-const member function.");
 
-		Super::template CreateDelegateInstance<TBaseSPMethodDelegateInstance<false, UserClass, ESPMode::ThreadSafe, FuncType, UserPolicy, std::decay_t<VarTypes>...>>(StaticCastSharedRef<UserClass>(InUserObject->AsShared()), InFunc, Forward<VarTypes>(Vars)...);
+		new (TWriteLockedDelegateAllocation{*this}) TBaseSPMethodDelegateInstance<false, UserClass, ESPMode::ThreadSafe, FuncType, UserPolicy, std::decay_t<VarTypes>...>(StaticCastSharedRef<UserClass>(InUserObject->AsShared()), InFunc, Forward<VarTypes>(Vars)...);
 	}
 	template <typename UserClass, typename... VarTypes>
 	inline void BindThreadSafeSP(const UserClass* InUserObject, typename TMemFunPtrType<true, UserClass, RetValType (ParamTypes..., std::decay_t<VarTypes>...)>::Type InFunc, VarTypes&&... Vars)
 	{
-		Super::template CreateDelegateInstance<TBaseSPMethodDelegateInstance<true, const UserClass, ESPMode::ThreadSafe, FuncType, UserPolicy, std::decay_t<VarTypes>...>>(StaticCastSharedRef<const UserClass>(InUserObject->AsShared()), InFunc, Forward<VarTypes>(Vars)...);
+		new (TWriteLockedDelegateAllocation{*this}) TBaseSPMethodDelegateInstance<true, const UserClass, ESPMode::ThreadSafe, FuncType, UserPolicy, std::decay_t<VarTypes>...>(StaticCastSharedRef<const UserClass>(InUserObject->AsShared()), InFunc, Forward<VarTypes>(Vars)...);
 	}
 
 	/**
@@ -259,12 +259,12 @@ public:
 	template <typename UObjectTemplate, typename... VarTypes>
 	inline void BindUFunction(UObjectTemplate* InUserObject, const FName& InFunctionName, VarTypes&&... Vars)
 	{
-		Super::template CreateDelegateInstance<TBaseUFunctionDelegateInstance<UObjectTemplate, FuncType, UserPolicy, std::decay_t<VarTypes>...>>(InUserObject, InFunctionName, Forward<VarTypes>(Vars)...);
+		new (TWriteLockedDelegateAllocation{*this}) TBaseUFunctionDelegateInstance<UObjectTemplate, FuncType, UserPolicy, std::decay_t<VarTypes>...>(InUserObject, InFunctionName, Forward<VarTypes>(Vars)...);
 	}
 	template <typename UObjectTemplate, typename... VarTypes>
 	inline void BindUFunction(TObjectPtr<UObjectTemplate> InUserObject, const FName& InFunctionName, VarTypes&&... Vars)
 	{
-		Super::template CreateDelegateInstance<TBaseUFunctionDelegateInstance<UObjectTemplate, FuncType, UserPolicy, std::decay_t<VarTypes>...>>(ToRawPtr(InUserObject), InFunctionName, Forward<VarTypes>(Vars)...);
+		new (TWriteLockedDelegateAllocation{*this}) TBaseUFunctionDelegateInstance<UObjectTemplate, FuncType, UserPolicy, std::decay_t<VarTypes>...>(ToRawPtr(InUserObject), InFunctionName, Forward<VarTypes>(Vars)...);
 	}
 
 	/**
@@ -278,24 +278,24 @@ public:
 	{
 		static_assert(!TIsConst<UserClass>::Value, "Attempting to bind a delegate with a const object pointer and non-const member function.");
 
-		Super::template CreateDelegateInstance<TBaseUObjectMethodDelegateInstance<false, UserClass, FuncType, UserPolicy, std::decay_t<VarTypes>...>>(InUserObject, InFunc, Forward<VarTypes>(Vars)...);
+		new (TWriteLockedDelegateAllocation{*this}) TBaseUObjectMethodDelegateInstance<false, UserClass, FuncType, UserPolicy, std::decay_t<VarTypes>...>(InUserObject, InFunc, Forward<VarTypes>(Vars)...);
 	}
 	template <typename UserClass, typename... VarTypes>
 	inline void BindUObject(const UserClass* InUserObject, typename TMemFunPtrType<true, UserClass, RetValType (ParamTypes..., std::decay_t<VarTypes>...)>::Type InFunc, VarTypes&&... Vars)
 	{
-		Super::template CreateDelegateInstance<TBaseUObjectMethodDelegateInstance<true, const UserClass, FuncType, UserPolicy, std::decay_t<VarTypes>...>>(InUserObject, InFunc, Forward<VarTypes>(Vars)...);
+		new (TWriteLockedDelegateAllocation{*this}) TBaseUObjectMethodDelegateInstance<true, const UserClass, FuncType, UserPolicy, std::decay_t<VarTypes>...>(InUserObject, InFunc, Forward<VarTypes>(Vars)...);
 	}
 	template <typename UserClass, typename... VarTypes>
 	inline void BindUObject(TObjectPtr<UserClass> InUserObject, typename TMemFunPtrType<false, UserClass, RetValType (ParamTypes..., std::decay_t<VarTypes>...)>::Type InFunc, VarTypes&&... Vars)
 	{
 		static_assert(!TIsConst<UserClass>::Value, "Attempting to bind a delegate with a const object pointer and non-const member function.");
 
-		Super::template CreateDelegateInstance<TBaseUObjectMethodDelegateInstance<false, UserClass, FuncType, UserPolicy, std::decay_t<VarTypes>...>>(ToRawPtr(InUserObject), InFunc, Forward<VarTypes>(Vars)...);
+		new (TWriteLockedDelegateAllocation{*this}) TBaseUObjectMethodDelegateInstance<false, UserClass, FuncType, UserPolicy, std::decay_t<VarTypes>...>(ToRawPtr(InUserObject), InFunc, Forward<VarTypes>(Vars)...);
 	}
 	template <typename UserClass, typename... VarTypes>
 	inline void BindUObject(TObjectPtr<UserClass> InUserObject, typename TMemFunPtrType<true, UserClass, RetValType (ParamTypes..., std::decay_t<VarTypes>...)>::Type InFunc, VarTypes&&... Vars)
 	{
-		Super::template CreateDelegateInstance<TBaseUObjectMethodDelegateInstance<true, const UserClass, FuncType, UserPolicy, std::decay_t<VarTypes>...>>(ToRawPtr(InUserObject), InFunc, Forward<VarTypes>(Vars)...);
+		new (TWriteLockedDelegateAllocation{*this}) TBaseUObjectMethodDelegateInstance<true, const UserClass, FuncType, UserPolicy, std::decay_t<VarTypes>...>(ToRawPtr(InUserObject), InFunc, Forward<VarTypes>(Vars)...);
 	}
 
 	// Executing via a delegate registration reference is not allowed
@@ -366,7 +366,7 @@ public:
 	[[nodiscard]] inline static TDelegate<RetValType(ParamTypes...), UserPolicy> CreateStatic(typename TIdentity<RetValType (*)(ParamTypes..., std::decay_t<VarTypes>...)>::Type InFunc, VarTypes&&... Vars)
 	{
 		TDelegate<RetValType(ParamTypes...), UserPolicy> Result;
-		Result.template CreateDelegateInstance<TBaseStaticDelegateInstance<FuncType, UserPolicy, std::decay_t<VarTypes>...>>(InFunc, Forward<VarTypes>(Vars)...);
+		new (TWriteLockedDelegateAllocation{Result}) TBaseStaticDelegateInstance<FuncType, UserPolicy, std::decay_t<VarTypes>...>(InFunc, Forward<VarTypes>(Vars)...);
 		return Result;
 	}
 
@@ -378,7 +378,7 @@ public:
 	[[nodiscard]] inline static TDelegate<RetValType(ParamTypes...), UserPolicy> CreateLambda(FunctorType&& InFunctor, VarTypes&&... Vars)
 	{
 		TDelegate<RetValType(ParamTypes...), UserPolicy> Result;
-		Result.template CreateDelegateInstance<TBaseFunctorDelegateInstance<FuncType, UserPolicy, typename TRemoveReference<FunctorType>::Type, std::decay_t<VarTypes>...>>(Forward<FunctorType>(InFunctor), Forward<VarTypes>(Vars)...);
+		new (TWriteLockedDelegateAllocation{Result}) TBaseFunctorDelegateInstance<FuncType, UserPolicy, typename TRemoveReference<FunctorType>::Type, std::decay_t<VarTypes>...>(Forward<FunctorType>(InFunctor), Forward<VarTypes>(Vars)...);
 		return Result;
 	}
 
@@ -390,14 +390,14 @@ public:
 	[[nodiscard]] inline static TDelegate<RetValType(ParamTypes...), UserPolicy> CreateSPLambda(const TSharedRef<UserClass, Mode>& InUserObjectRef, FunctorType&& InFunctor, VarTypes&&... Vars)
 	{
 		TDelegate<RetValType(ParamTypes...), UserPolicy> Result;
-		Result.template CreateDelegateInstance<TBaseSPLambdaDelegateInstance<UserClass, Mode, FuncType, UserPolicy, typename TRemoveReference<FunctorType>::Type, std::decay_t<VarTypes>...>>(InUserObjectRef, Forward<FunctorType>(InFunctor), Forward<VarTypes>(Vars)...);
+		new (TWriteLockedDelegateAllocation{Result}) TBaseSPLambdaDelegateInstance<UserClass, Mode, FuncType, UserPolicy, typename TRemoveReference<FunctorType>::Type, std::decay_t<VarTypes>...>(InUserObjectRef, Forward<FunctorType>(InFunctor), Forward<VarTypes>(Vars)...);
 		return Result;
 	}
 	template <typename UserClass, typename FunctorType, typename... VarTypes>
 	[[nodiscard]] inline static TDelegate<RetValType(ParamTypes...), UserPolicy> CreateSPLambda(UserClass* InUserObject, FunctorType&& InFunctor, VarTypes&&... Vars)
 	{
 		TDelegate<RetValType(ParamTypes...), UserPolicy> Result;
-		Result.template CreateDelegateInstance<TBaseSPLambdaDelegateInstance<UserClass, decltype(InUserObject->AsShared())::Mode, FuncType, UserPolicy, typename TRemoveReference<FunctorType>::Type, std::decay_t<VarTypes>...>>(StaticCastSharedRef<UserClass>(InUserObject->AsShared()), Forward<FunctorType>(InFunctor), Forward<VarTypes>(Vars)...);
+		new (TWriteLockedDelegateAllocation{Result}) TBaseSPLambdaDelegateInstance<UserClass, decltype(InUserObject->AsShared())::Mode, FuncType, UserPolicy, typename TRemoveReference<FunctorType>::Type, std::decay_t<VarTypes>...>(StaticCastSharedRef<UserClass>(InUserObject->AsShared()), Forward<FunctorType>(InFunctor), Forward<VarTypes>(Vars)...);
 		return Result;
 	}
 
@@ -409,7 +409,7 @@ public:
 	[[nodiscard]] inline static TDelegate<RetValType(ParamTypes...), UserPolicy> CreateWeakLambda(UserClass* InUserObject, FunctorType&& InFunctor, VarTypes&&... Vars)
 	{
 		TDelegate<RetValType(ParamTypes...), UserPolicy> Result;
-		Result.template CreateDelegateInstance<TWeakBaseFunctorDelegateInstance<UserClass, FuncType, UserPolicy, typename TRemoveReference<FunctorType>::Type, std::decay_t<VarTypes>...>>(InUserObject, Forward<FunctorType>(InFunctor), Forward<VarTypes>(Vars)...);
+		new (TWriteLockedDelegateAllocation{Result}) TWeakBaseFunctorDelegateInstance<UserClass, FuncType, UserPolicy, typename TRemoveReference<FunctorType>::Type, std::decay_t<VarTypes>...>(InUserObject, Forward<FunctorType>(InFunctor), Forward<VarTypes>(Vars)...);
 		return Result;
 	}
 
@@ -425,14 +425,14 @@ public:
 		static_assert(!TIsConst<UserClass>::Value, "Attempting to bind a delegate with a const object pointer and non-const member function.");
 
 		TDelegate<RetValType(ParamTypes...), UserPolicy> Result;
-		Result.template CreateDelegateInstance<TBaseRawMethodDelegateInstance<false, UserClass, FuncType, UserPolicy, std::decay_t<VarTypes>...>>(InUserObject, InFunc, Forward<VarTypes>(Vars)...);
+		new (TWriteLockedDelegateAllocation{Result}) TBaseRawMethodDelegateInstance<false, UserClass, FuncType, UserPolicy, std::decay_t<VarTypes>...>(InUserObject, InFunc, Forward<VarTypes>(Vars)...);
 		return Result;
 	}
 	template <typename UserClass, typename... VarTypes>
 	[[nodiscard]] inline static TDelegate<RetValType(ParamTypes...), UserPolicy> CreateRaw(const UserClass* InUserObject, typename TMemFunPtrType<true, UserClass, RetValType (ParamTypes..., std::decay_t<VarTypes>...)>::Type InFunc, VarTypes&&... Vars)
 	{
 		TDelegate<RetValType(ParamTypes...), UserPolicy> Result;
-		Result.template CreateDelegateInstance<TBaseRawMethodDelegateInstance<true, const UserClass, FuncType, UserPolicy, std::decay_t<VarTypes>...>>(InUserObject, InFunc, Forward<VarTypes>(Vars)...);
+		new (TWriteLockedDelegateAllocation{Result}) TBaseRawMethodDelegateInstance<true, const UserClass, FuncType, UserPolicy, std::decay_t<VarTypes>...>(InUserObject, InFunc, Forward<VarTypes>(Vars)...);
 		return Result;
 	}
 
@@ -448,14 +448,14 @@ public:
 		static_assert(!TIsConst<UserClass>::Value, "Attempting to bind a delegate with a const object pointer and non-const member function.");
 
 		TDelegate<RetValType(ParamTypes...), UserPolicy> Result;
-		Result.template CreateDelegateInstance<TBaseSPMethodDelegateInstance<false, UserClass, Mode, FuncType, UserPolicy, std::decay_t<VarTypes>...>>(InUserObjectRef, InFunc, Forward<VarTypes>(Vars)...);
+		new (TWriteLockedDelegateAllocation{Result}) TBaseSPMethodDelegateInstance<false, UserClass, Mode, FuncType, UserPolicy, std::decay_t<VarTypes>...>(InUserObjectRef, InFunc, Forward<VarTypes>(Vars)...);
 		return Result;
 	}
 	template <typename UserClass, ESPMode Mode, typename... VarTypes>
 	[[nodiscard]] inline static TDelegate<RetValType(ParamTypes...), UserPolicy> CreateSP(const TSharedRef<UserClass, Mode>& InUserObjectRef, typename TMemFunPtrType<true, UserClass, RetValType (ParamTypes..., std::decay_t<VarTypes>...)>::Type InFunc, VarTypes&&... Vars)
 	{
 		TDelegate<RetValType(ParamTypes...), UserPolicy> Result;
-		Result.template CreateDelegateInstance<TBaseSPMethodDelegateInstance<true, const UserClass, Mode, FuncType, UserPolicy, std::decay_t<VarTypes>...>>(InUserObjectRef, InFunc, Forward<VarTypes>(Vars)...);
+		new (TWriteLockedDelegateAllocation{Result}) TBaseSPMethodDelegateInstance<true, const UserClass, Mode, FuncType, UserPolicy, std::decay_t<VarTypes>...>(InUserObjectRef, InFunc, Forward<VarTypes>(Vars)...);
 		return Result;
 	}
 
@@ -471,14 +471,14 @@ public:
 		static_assert(!TIsConst<UserClass>::Value, "Attempting to bind a delegate with a const object pointer and non-const member function.");
 
 		TDelegate<RetValType(ParamTypes...), UserPolicy> Result;
-		Result.template CreateDelegateInstance<TBaseSPMethodDelegateInstance<false, UserClass, decltype(InUserObject->AsShared())::Mode, FuncType, UserPolicy, std::decay_t<VarTypes>...>>(StaticCastSharedRef<UserClass>(InUserObject->AsShared()), InFunc, Forward<VarTypes>(Vars)...);
+		new (TWriteLockedDelegateAllocation{Result}) TBaseSPMethodDelegateInstance<false, UserClass, decltype(InUserObject->AsShared())::Mode, FuncType, UserPolicy, std::decay_t<VarTypes>...>(StaticCastSharedRef<UserClass>(InUserObject->AsShared()), InFunc, Forward<VarTypes>(Vars)...);
 		return Result;
 	}
 	template <typename UserClass, typename... VarTypes>
 	[[nodiscard]] inline static TDelegate<RetValType(ParamTypes...), UserPolicy> CreateSP(const UserClass* InUserObject, typename TMemFunPtrType<true, UserClass, RetValType (ParamTypes..., std::decay_t<VarTypes>...)>::Type InFunc, VarTypes&&... Vars)
 	{
 		TDelegate<RetValType(ParamTypes...), UserPolicy> Result;
-		Result.template CreateDelegateInstance<TBaseSPMethodDelegateInstance<true, const UserClass, decltype(InUserObject->AsShared())::Mode, FuncType, UserPolicy, std::decay_t<VarTypes>...>>(StaticCastSharedRef<const UserClass>(InUserObject->AsShared()), InFunc, Forward<VarTypes>(Vars)...);
+		new (TWriteLockedDelegateAllocation{Result}) TBaseSPMethodDelegateInstance<true, const UserClass, decltype(InUserObject->AsShared())::Mode, FuncType, UserPolicy, std::decay_t<VarTypes>...>(StaticCastSharedRef<const UserClass>(InUserObject->AsShared()), InFunc, Forward<VarTypes>(Vars)...);
 		return Result;
 	}
 
@@ -496,14 +496,14 @@ public:
 		static_assert(!TIsConst<UserClass>::Value, "Attempting to bind a delegate with a const object pointer and non-const member function.");
 
 		TDelegate<RetValType(ParamTypes...), UserPolicy> Result;
-		Result.template CreateDelegateInstance<TBaseSPMethodDelegateInstance<false, UserClass, ESPMode::ThreadSafe, FuncType, UserPolicy, std::decay_t<VarTypes>...>>(InUserObjectRef, InFunc, Forward<VarTypes>(Vars)...);
+		new (TWriteLockedDelegateAllocation{Result}) TBaseSPMethodDelegateInstance<false, UserClass, ESPMode::ThreadSafe, FuncType, UserPolicy, std::decay_t<VarTypes>...>(InUserObjectRef, InFunc, Forward<VarTypes>(Vars)...);
 		return Result;
 	}
 	template <typename UserClass, typename... VarTypes>
 	[[nodiscard]] inline static TDelegate<RetValType(ParamTypes...), UserPolicy> CreateThreadSafeSP(const TSharedRef<UserClass, ESPMode::ThreadSafe>& InUserObjectRef, typename TMemFunPtrType<true, UserClass, RetValType (ParamTypes..., std::decay_t<VarTypes>...)>::Type InFunc, VarTypes&&... Vars)
 	{
 		TDelegate<RetValType(ParamTypes...), UserPolicy> Result;
-		Result.template CreateDelegateInstance<TBaseSPMethodDelegateInstance<true, const UserClass, ESPMode::ThreadSafe, FuncType, UserPolicy, std::decay_t<VarTypes>...>>(InUserObjectRef, InFunc, Forward<VarTypes>(Vars)...);
+		new (TWriteLockedDelegateAllocation{Result}) TBaseSPMethodDelegateInstance<true, const UserClass, ESPMode::ThreadSafe, FuncType, UserPolicy, std::decay_t<VarTypes>...>(InUserObjectRef, InFunc, Forward<VarTypes>(Vars)...);
 		return Result;
 	}
 
@@ -521,14 +521,14 @@ public:
 		static_assert(!TIsConst<UserClass>::Value, "Attempting to bind a delegate with a const object pointer and non-const member function.");
 
 		TDelegate<RetValType(ParamTypes...), UserPolicy> Result;
-		Result.template CreateDelegateInstance<TBaseSPMethodDelegateInstance<false, UserClass, ESPMode::ThreadSafe, FuncType, UserPolicy, std::decay_t<VarTypes>...>>(StaticCastSharedRef<UserClass>(InUserObject->AsShared()), InFunc, Forward<VarTypes>(Vars)...);
+		new (TWriteLockedDelegateAllocation{Result}) TBaseSPMethodDelegateInstance<false, UserClass, ESPMode::ThreadSafe, FuncType, UserPolicy, std::decay_t<VarTypes>...>(StaticCastSharedRef<UserClass>(InUserObject->AsShared()), InFunc, Forward<VarTypes>(Vars)...);
 		return Result;
 	}
 	template <typename UserClass, typename... VarTypes>
 	[[nodiscard]] inline static TDelegate<RetValType(ParamTypes...), UserPolicy> CreateThreadSafeSP(const UserClass* InUserObject, typename TMemFunPtrType<true, UserClass, RetValType (ParamTypes..., std::decay_t<VarTypes>...)>::Type InFunc, VarTypes&&... Vars)
 	{
 		TDelegate<RetValType(ParamTypes...), UserPolicy> Result;
-		Result.template CreateDelegateInstance<TBaseSPMethodDelegateInstance<true, const UserClass, ESPMode::ThreadSafe, FuncType, UserPolicy, std::decay_t<VarTypes>...>>(StaticCastSharedRef<const UserClass>(InUserObject->AsShared()), InFunc, Forward<VarTypes>(Vars)...);
+		new (TWriteLockedDelegateAllocation{Result}) TBaseSPMethodDelegateInstance<true, const UserClass, ESPMode::ThreadSafe, FuncType, UserPolicy, std::decay_t<VarTypes>...>(StaticCastSharedRef<const UserClass>(InUserObject->AsShared()), InFunc, Forward<VarTypes>(Vars)...);
 		return Result;
 	}
 
@@ -542,14 +542,14 @@ public:
 	[[nodiscard]] inline static TDelegate<RetValType(ParamTypes...), UserPolicy> CreateUFunction(UObjectTemplate* InUserObject, const FName& InFunctionName, VarTypes&&... Vars)
 	{
 		TDelegate<RetValType(ParamTypes...), UserPolicy> Result;
-		Result.template CreateDelegateInstance<TBaseUFunctionDelegateInstance<UObjectTemplate, FuncType, UserPolicy, std::decay_t<VarTypes>...>>(InUserObject, InFunctionName, Forward<VarTypes>(Vars)...);
+		new (TWriteLockedDelegateAllocation{Result}) TBaseUFunctionDelegateInstance<UObjectTemplate, FuncType, UserPolicy, std::decay_t<VarTypes>...>(InUserObject, InFunctionName, Forward<VarTypes>(Vars)...);
 		return Result;
 	}
 	template <typename UObjectTemplate, typename... VarTypes>
 	[[nodiscard]] inline static TDelegate<RetValType(ParamTypes...), UserPolicy> CreateUFunction(TObjectPtr<UObjectTemplate> InUserObject, const FName& InFunctionName, VarTypes&&... Vars)
 	{
 		TDelegate<RetValType(ParamTypes...), UserPolicy> Result;
-		Result.template CreateDelegateInstance<TBaseUFunctionDelegateInstance<UObjectTemplate, FuncType, UserPolicy, std::decay_t<VarTypes>...>>(ToRawPtr(InUserObject), InFunctionName, Forward<VarTypes>(Vars)...);
+		new (TWriteLockedDelegateAllocation{Result}) TBaseUFunctionDelegateInstance<UObjectTemplate, FuncType, UserPolicy, std::decay_t<VarTypes>...>(ToRawPtr(InUserObject), InFunctionName, Forward<VarTypes>(Vars)...);
 		return Result;
 	}
 
@@ -565,14 +565,14 @@ public:
 		static_assert(!TIsConst<UserClass>::Value, "Attempting to bind a delegate with a const object pointer and non-const member function.");
 
 		TDelegate<RetValType(ParamTypes...), UserPolicy> Result;
-		Result.template CreateDelegateInstance<TBaseUObjectMethodDelegateInstance<false, UserClass, FuncType, UserPolicy, std::decay_t<VarTypes>...>>(InUserObject, InFunc, Forward<VarTypes>(Vars)...);
+		new (TWriteLockedDelegateAllocation{Result}) TBaseUObjectMethodDelegateInstance<false, UserClass, FuncType, UserPolicy, std::decay_t<VarTypes>...>(InUserObject, InFunc, Forward<VarTypes>(Vars)...);
 		return Result;
 	}
 	template <typename UserClass, typename... VarTypes>
 	[[nodiscard]] inline static TDelegate<RetValType(ParamTypes...), UserPolicy> CreateUObject(const UserClass* InUserObject, typename TMemFunPtrType<true, UserClass, RetValType (ParamTypes..., std::decay_t<VarTypes>...)>::Type InFunc, VarTypes&&... Vars)
 	{
 		TDelegate<RetValType(ParamTypes...), UserPolicy> Result;
-		Result.template CreateDelegateInstance<TBaseUObjectMethodDelegateInstance<true, const UserClass, FuncType, UserPolicy, std::decay_t<VarTypes>...>>(InUserObject, InFunc, Forward<VarTypes>(Vars)...);
+		new (TWriteLockedDelegateAllocation{Result}) TBaseUObjectMethodDelegateInstance<true, const UserClass, FuncType, UserPolicy, std::decay_t<VarTypes>...>(InUserObject, InFunc, Forward<VarTypes>(Vars)...);
 		return Result;
 	}
 	template <typename UserClass, typename... VarTypes>
@@ -581,14 +581,14 @@ public:
 		static_assert(!TIsConst<UserClass>::Value, "Attempting to bind a delegate with a const object pointer and non-const member function.");
 
 		TDelegate<RetValType(ParamTypes...), UserPolicy> Result;
-		Result.template CreateDelegateInstance<TBaseUObjectMethodDelegateInstance<false, UserClass, FuncType, UserPolicy, std::decay_t<VarTypes>...>>(ToRawPtr(InUserObject), InFunc, Forward<VarTypes>(Vars)...);
+		new (TWriteLockedDelegateAllocation{Result}) TBaseUObjectMethodDelegateInstance<false, UserClass, FuncType, UserPolicy, std::decay_t<VarTypes>...>(ToRawPtr(InUserObject), InFunc, Forward<VarTypes>(Vars)...);
 		return Result;
 	}
 	template <typename UserClass, typename... VarTypes>
 	[[nodiscard]] inline static TDelegate<RetValType(ParamTypes...), UserPolicy> CreateUObject(TObjectPtr<UserClass> InUserObject, typename TMemFunPtrType<true, UserClass, RetValType (ParamTypes..., std::decay_t<VarTypes>...)>::Type InFunc, VarTypes&&... Vars)
 	{
 		TDelegate<RetValType(ParamTypes...), UserPolicy> Result;
-		Result.template CreateDelegateInstance<TBaseUObjectMethodDelegateInstance<true, const UserClass, FuncType, UserPolicy, std::decay_t<VarTypes>...>>(ToRawPtr(InUserObject), InFunc, Forward<VarTypes>(Vars)...);
+		new (TWriteLockedDelegateAllocation{Result}) TBaseUObjectMethodDelegateInstance<true, const UserClass, FuncType, UserPolicy, std::decay_t<VarTypes>...>(ToRawPtr(InUserObject), InFunc, Forward<VarTypes>(Vars)...);
 		return Result;
 	}
 
