Delegate compile error workaround

TDelegateAccessHandlerBase is being used by a delegate with a different thread safety mode. This works around the compile error for now.

#rnx

diff --git a/Engine/Source/Runtime/Core/Public/Delegates/DelegateAccessHandler.h b/Engine/Source/Runtime/Core/Public/Delegates/DelegateAccessHandler.h
--- a/Engine/Source/Runtime/Core/Public/Delegates/DelegateAccessHandler.h
+++ b/Engine/Source/Runtime/Core/Public/Delegates/DelegateAccessHandler.h
@@ -40,7 +40,8 @@ struct TWriteLockedDelegateAllocation;
 template<>
 class TDelegateAccessHandlerBase<FNotThreadSafeNotCheckedDelegateMode>
 {
-	friend struct TWriteLockedDelegateAllocation<FNotThreadSafeNotCheckedDelegateMode>;
+	template <typename T>
+	friend struct TWriteLockedDelegateAllocation;
 
 protected:
 	struct FReadAccessScope {};
