Fixed a compile error in TDeque when using an allocator with NeedsElementType == false.

diff --git a/Engine/Source/Runtime/Core/Public/Containers/Deque.h b/Engine/Source/Runtime/Core/Public/Containers/Deque.h
--- a/Engine/Source/Runtime/Core/Public/Containers/Deque.h
+++ b/Engine/Source/Runtime/Core/Public/Containers/Deque.h
@@ -556,7 +556,7 @@ private:
 			RelocateConstructItems<ElementType>(TempStorage.GetAllocation(), GetData(), Tail);
 			const SizeType HeadToEndOffset = Capacity - Head;
 			ShiftElementsLeft(HeadToEndOffset);
-			RelocateConstructItems<ElementType>(GetData() + HeadToEndOffset, TempStorage.GetAllocation(), Tail);
+			RelocateConstructItems<ElementType>(GetData() + HeadToEndOffset, (ElementType*)TempStorage.GetAllocation(), Tail);
 		}
 	}
 
