Changed some typedefs to aliases.

#rnx

diff --git a/Engine/Source/Runtime/Core/Public/Delegates/DelegateBase.h b/Engine/Source/Runtime/Core/Public/Delegates/DelegateBase.h
--- a/Engine/Source/Runtime/Core/Public/Delegates/DelegateBase.h
+++ b/Engine/Source/Runtime/Core/Public/Delegates/DelegateBase.h
@@ -12,19 +12,19 @@
 #include "Delegates/IDelegateInstance.h"
 
 #if !defined(_WIN32) || defined(_WIN64) || (defined(ALLOW_DELEGATE_INLINE_ALLOCATORS_ON_WIN32) && ALLOW_DELEGATE_INLINE_ALLOCATORS_ON_WIN32)
-	typedef TAlignedBytes<16, 16> FAlignedInlineDelegateType;
+	using FAlignedInlineDelegateType = TAlignedBytes<16, 16>;
 	#if !defined(NUM_DELEGATE_INLINE_BYTES) || NUM_DELEGATE_INLINE_BYTES == 0
-		typedef FHeapAllocator FDelegateAllocatorType;
+		using FDelegateAllocatorType = FHeapAllocator;
 	#elif NUM_DELEGATE_INLINE_BYTES < 0 || (NUM_DELEGATE_INLINE_BYTES % 16) != 0
 		#error NUM_DELEGATE_INLINE_BYTES must be a multiple of 16
 	#else
-		typedef TInlineAllocator<(NUM_DELEGATE_INLINE_BYTES / 16)> FDelegateAllocatorType;
+		using FDelegateAllocatorType = TInlineAllocator<(NUM_DELEGATE_INLINE_BYTES / 16)>;
 	#endif
 #else
 	// ... except on Win32, because we can't pass 16-byte aligned types by value, as some delegates are
 	// so we'll just keep it heap-allocated, which are always sufficiently aligned.
-	typedef TAlignedBytes<16, 8> FAlignedInlineDelegateType;
-	typedef FHeapAllocator FDelegateAllocatorType;
+	using FAlignedInlineDelegateType = TAlignedBytes<16, 8>;
+	using FDelegateAllocatorType     = FHeapAllocator;
 #endif
 
 template <typename UserPolicy>
