Use the more codegen optimal ClangPlatformAtomics when using clang-cl on Windows.

diff --git a/Engine/Source/Runtime/Core/Public/Clang/ClangPlatformAtomics.h b/Engine/Source/Runtime/Core/Public/Clang/ClangPlatformAtomics.h
--- a/Engine/Source/Runtime/Core/Public/Clang/ClangPlatformAtomics.h
+++ b/Engine/Source/Runtime/Core/Public/Clang/ClangPlatformAtomics.h
@@ -1,7 +1,7 @@
 // Copyright Epic Games, Inc. All Rights Reserved.
 
 /*=============================================================================================
-	ClangPlatformAtomics.h: Apple platform Atomics functions
+	ClangPlatformAtomics.h: Clang Atomics functions
 ==============================================================================================*/
 
 #pragma once
diff --git a/Engine/Source/Runtime/Core/Public/Windows/WindowsPlatformAtomics.h b/Engine/Source/Runtime/Core/Public/Windows/WindowsPlatformAtomics.h
--- a/Engine/Source/Runtime/Core/Public/Windows/WindowsPlatformAtomics.h
+++ b/Engine/Source/Runtime/Core/Public/Windows/WindowsPlatformAtomics.h
@@ -441,6 +441,10 @@ protected:
 	static CORE_API void HandleAtomicsFailure( const TCHAR* InFormat, ... );
 };
 
+#ifdef __clang__
+#include "Clang/ClangPlatformAtomics.h"
 
-
+typedef FClangPlatformAtomics FPlatformAtomics;
+#else
 typedef FWindowsPlatformAtomics FPlatformAtomics;
+#endif // __clang__
