Tests: Core: Fixed failing 'System.Core.HAL.Platform Verification' test because of asserts emitted by FName::AutoTest().

This test expects the system to distinguish between FNames with different casing. However, this features is only available when building with WITH_CASE_PRESERVING_NAME, which is enabled only on desktop platforms (Win, Unix, Mac).

#rnx

diff --git a/Engine/Source/Runtime/Core/Private/UObject/UnrealNames.cpp b/Engine/Source/Runtime/Core/Private/UObject/UnrealNames.cpp
--- a/Engine/Source/Runtime/Core/Private/UObject/UnrealNames.cpp
+++ b/Engine/Source/Runtime/Core/Private/UObject/UnrealNames.cpp
@@ -3753,7 +3753,15 @@ void CheckLazyName(const CharType(&Literal)[N])
 	CharType Literal2[N];
 	FMemory::Memcpy(Literal2, Literal);
 	check(FLazyName(Literal) == FLazyName(Literal2));
-	check(WriteToString<64>(FLazyName(Literal).Resolve()).ToView().Equals(WriteToString<64>(Literal).ToView(), ESearchCase::CaseSensitive));
+
+	constexpr ESearchCase::Type ComparisonMode =
+#if WITH_CASE_PRESERVING_NAME
+		ESearchCase::CaseSensitive;
+#else
+		ESearchCase::IgnoreCase;
+#endif
+
+	check(WriteToString<64>(FLazyName(Literal).Resolve()).ToView().Equals(WriteToString<64>(Literal).ToView(), ComparisonMode));
 }
 
 static void TestNameBatch();
