[AutoRTFM] Support returning values from AutoRTFM::Open() lambdas

diff --git a/Engine/Source/Programs/AutoRTFMTests/Private/Open.cpp b/Engine/Source/Programs/AutoRTFMTests/Private/Open.cpp
--- a/Engine/Source/Programs/AutoRTFMTests/Private/Open.cpp
+++ b/Engine/Source/Programs/AutoRTFMTests/Private/Open.cpp
@@ -102,3 +102,60 @@ TEST_CASE("Open.Atomics")
         }));
     REQUIRE(bDidRun);
 }
+
+TEST_CASE("Open.ReturnValue")
+{
+    static_assert(AutoRTFM::TIsSafeToReturnFromOpen<int>::value);
+    static_assert(AutoRTFM::TIsSafeToReturnFromOpen<float>::value);
+    static_assert(AutoRTFM::TIsSafeToReturnFromOpen<int*>::value);
+    static_assert(AutoRTFM::TIsSafeToReturnFromOpen<void>::value);
+    static_assert(AutoRTFM::TIsSafeToReturnFromOpen<std::tuple<int, float>>::value);
+    static_assert(!AutoRTFM::TIsSafeToReturnFromOpen<std::string>::value);
+    static_assert(!AutoRTFM::TIsSafeToReturnFromOpen<std::tuple<int, std::string>>::value);
+
+    SECTION("int")
+    {
+        int Value = 10;
+        AutoRTFM::ETransactionResult Result = AutoRTFM::Transact([&] ()
+            {
+                int Got = AutoRTFM::Open([] { return 42; });
+                AutoRTFM::Open([&] { Value = Got; });
+            });
+        REQUIRE(AutoRTFM::ETransactionResult::Committed == Result);
+        REQUIRE(42 == Value);
+    }
+
+    SECTION("char*")
+    {
+        std::string Value = "<unassigned>";
+        AutoRTFM::ETransactionResult Result = AutoRTFM::Transact([&] ()
+            {
+                // Note: AutoRTFM::Open() is returning a const char*
+                std::string Got = AutoRTFM::Open([] { return "meow"; });
+                AutoRTFM::Open([&] { Value = Got; });
+            });
+        REQUIRE(AutoRTFM::ETransactionResult::Committed == Result);
+        REQUIRE("meow" == Value);
+    }
+    
+    SECTION("tuple")
+    {
+        int Int = 0;
+        std::string String = "<unassigned>";
+        AutoRTFM::ETransactionResult Result = AutoRTFM::Transact([&] ()
+            {
+                auto [I, S] = AutoRTFM::Open([] 
+                    { 
+                        return std::make_tuple(42, "woof"); 
+                    });
+                AutoRTFM::Open([&]
+                    { 
+                        Int = I; 
+                        String = S; 
+                    });
+            });
+        REQUIRE(AutoRTFM::ETransactionResult::Committed == Result);
+        REQUIRE(42 == Int);
+        REQUIRE("woof" == String);
+    }
+}
diff --git a/Engine/Source/Runtime/Core/Public/AutoRTFM/AutoRTFM.h b/Engine/Source/Runtime/Core/Public/AutoRTFM/AutoRTFM.h
--- a/Engine/Source/Runtime/Core/Public/AutoRTFM/AutoRTFM.h
+++ b/Engine/Source/Runtime/Core/Public/AutoRTFM/AutoRTFM.h
@@ -373,6 +373,8 @@ static UE_AUTORTFM_FORCEINLINE void autortfm_check_abi(void* ptr, size_t size)
 #ifdef __cplusplus
 }
 
+#include <tuple>
+
 namespace AutoRTFM
 {
 
@@ -559,13 +561,55 @@ static UE_AUTORTFM_FORCEINLINE void AbortIfClosed()
     autortfm_abort_if_closed();
 }
 
+// Traits helper that is used to determine whether type T can be used as the 
+// return type of a function called by Open().
+template<typename T, typename = void>
+class TIsSafeToReturnFromOpen : public std::false_type {};
+
+// Specializations of TIsSafeToReturnFromOpen for fundamental types.
+template <typename T>
+class TIsSafeToReturnFromOpen<T, std::enable_if_t<std::is_fundamental_v<T>>> : public std::true_type {};
+
+// Specializations of TIsSafeToReturnFromOpen for raw pointer types.
+template <typename T>
+class TIsSafeToReturnFromOpen<T*, void> : public std::true_type {};
+
+// Specializations of TIsSafeToReturnFromOpen for std::tuple.
+template <typename ... TYPES>
+class TIsSafeToReturnFromOpen<std::tuple<TYPES...>, void> : public std::conditional_t<
+	(TIsSafeToReturnFromOpen<TYPES>::value && ...), std::true_type, std::false_type> {};
+
 // Executes the given code non-transactionally regardless of whether we are in
-// a transaction or not.
-template<typename TFunctor> static UE_AUTORTFM_FORCEINLINE void Open(const TFunctor& Functor)
+// a transaction or not. Returns the value returned by Functor.
+template<typename TFunctor, typename TReturn = decltype(std::declval<TFunctor>()())> 
+static UE_AUTORTFM_FORCEINLINE TReturn Open(const TFunctor& Functor)
 {
-    autortfm_open(
-        [] (void* Arg) { UE_AUTORTFM_CALLSITE_FORCEINLINE (*static_cast<const TFunctor*>(Arg))(); },
-        const_cast<void*>(static_cast<const void*>(&Functor)));
+	static_assert(TIsSafeToReturnFromOpen<TReturn>::value,
+		"function return type is not safe to return from Open()");
+
+	if constexpr (std::is_same_v<void, TReturn>)
+	{
+		autortfm_open(
+			[] (void* Arg) { UE_AUTORTFM_CALLSITE_FORCEINLINE (*static_cast<const TFunctor*>(Arg))(); },
+			const_cast<void*>(static_cast<const void*>(&Functor)));
+	}
+	else
+	{
+		TReturn ReturnValue;
+		struct FData
+		{
+			const TFunctor& Functor;
+			TReturn& ReturnValue;
+		};
+		FData Data{Functor, ReturnValue};
+		autortfm_open([](void* Arg)
+			{
+				FData& Data = *reinterpret_cast<FData*>(Arg);
+				UE_AUTORTFM_CALLSITE_FORCEINLINE Data.ReturnValue = Data.Functor();
+			},
+			reinterpret_cast<void*>(&Data));
+		return Data.ReturnValue;
+	}
 }
 
 // Always executes the given code transactionally when called from a transaction nest
