Core: Added a container for global debugging pointers that is visible to natvis

This new concept aims to abstract the complexities of dealing global debugging state that is tricky with the symbol resolution scoping rules in Visual Studio's natvis specification, and fragile with live-coding.
The specification causes natvis authors to often need to add explicit scoping to global variables (using {,,UnrealEditor-ModuleName.dll} or UnrealEditor-ModuleName.dll! scoping), which is fragile for different build types (Editor vs Server vs Client) and configurations (Debug, Shipping etc).

This change adds a GDebuggingState symbol to every dll with which any engine module or third-party library can register global pointers using a unique GUID. The GUID string can then be used to locate that pointer from any other module through a natvis intrinsic function using strstr, without having to inject a specific symbol name into every DLL.

This allows authors to write generic natvis expressions that will work robustly when being viewed from the context of any other module, or in code that has been patched by LiveCoding.

diff --git a/Engine/Extras/NatvisHelpers/EntryPoint.cpp b/Engine/Extras/NatvisHelpers/EntryPoint.cpp
--- a/Engine/Extras/NatvisHelpers/EntryPoint.cpp
+++ b/Engine/Extras/NatvisHelpers/EntryPoint.cpp
@@ -2,8 +2,9 @@
 
 #include "Globals.h"
 
-extern "C" __declspec(dllexport) void InitNatvisHelpers(uint8** NameTable, FChunkedFixedUObjectArray* ObjectArray)
+extern "C" __declspec(dllexport) void InitNatvisHelpers(uint8** NameTable, FChunkedFixedUObjectArray* ObjectArray, UE::Core::FVisualizerDebuggingState* DebuggingState)
 {
 	GNameBlocksDebug = NameTable;
 	GObjectArrayForDebugVisualizers = ObjectArray;
+	GDebuggingState = DebuggingState;
 }
diff --git a/Engine/Extras/NatvisHelpers/Globals.cpp b/Engine/Extras/NatvisHelpers/Globals.cpp
--- a/Engine/Extras/NatvisHelpers/Globals.cpp
+++ b/Engine/Extras/NatvisHelpers/Globals.cpp
@@ -5,3 +5,6 @@ class FChunkedFixedUObjectArray;
 
 uint8** GNameBlocksDebug = nullptr;
 FChunkedFixedUObjectArray* GObjectArrayForDebugVisualizers = nullptr;
+
+namespace UE { namespace Core { struct FVisualizerDebuggingState; } }
+UE::Core::FVisualizerDebuggingState* GDebuggingState = nullptr;
\ No newline at end of file
diff --git a/Engine/Extras/NatvisHelpers/Globals.h b/Engine/Extras/NatvisHelpers/Globals.h
--- a/Engine/Extras/NatvisHelpers/Globals.h
+++ b/Engine/Extras/NatvisHelpers/Globals.h
@@ -7,3 +7,6 @@ class FChunkedFixedUObjectArray;
 
 extern uint8** GNameBlocksDebug;
 extern FChunkedFixedUObjectArray* GObjectArrayForDebugVisualizers;
+
+namespace UE { namespace Core { struct FVisualizerDebuggingState; } }
+extern UE::Core::FVisualizerDebuggingState* GDebuggingState;
diff --git a/Engine/Source/Developer/Windows/LiveCoding/Private/External/LC_ClientCommandActions.cpp b/Engine/Source/Developer/Windows/LiveCoding/Private/External/LC_ClientCommandActions.cpp
--- a/Engine/Source/Developer/Windows/LiveCoding/Private/External/LC_ClientCommandActions.cpp
+++ b/Engine/Source/Developer/Windows/LiveCoding/Private/External/LC_ClientCommandActions.cpp
@@ -104,6 +104,9 @@ extern uint8** GNameBlocksDebug;
 
 class FChunkedFixedUObjectArray;
 extern FChunkedFixedUObjectArray*& GObjectArrayForDebugVisualizers;
+
+namespace UE::Core{ struct FVisualizerDebuggingState; }
+extern UE::Core::FVisualizerDebuggingState*& GDebuggingState;
 // END EPIC MOD
 
 // BEGIN EPIC MOD - Support for object reinstancing. We need to call a global post-patch handler, rather than just getting individual callbacks for modified modules.
@@ -128,12 +131,12 @@ bool actions::LoadPatch::Execute(const CommandType* command, const DuplexPipe* p
 	// BEGIN EPIC MOD - Support for UE debug visualizers
 	if (module != nullptr)
 	{
-		typedef void InitNatvisHelpersFunc(uint8** NameTable, FChunkedFixedUObjectArray* ObjectArray);
+		typedef void InitNatvisHelpersFunc(uint8** NameTable, FChunkedFixedUObjectArray* ObjectArray, UE::Core::FVisualizerDebuggingState* DebuggingState);
 
 		InitNatvisHelpersFunc* InitNatvisHelpers = (InitNatvisHelpersFunc*)(void*)GetProcAddress(module, "InitNatvisHelpers");
 		if (InitNatvisHelpers != nullptr)
 		{
-			(*InitNatvisHelpers)(GNameBlocksDebug, GObjectArrayForDebugVisualizers);
+			(*InitNatvisHelpers)(GNameBlocksDebug, GObjectArrayForDebugVisualizers, GDebuggingState);
 		}
 	}
 	// END EPIC MOD
diff --git a/Engine/Source/Runtime/Core/Private/Modules/Tests/VisualizerDebuggingStateTests.cpp b/Engine/Source/Runtime/Core/Private/Modules/Tests/VisualizerDebuggingStateTests.cpp
new file mode 100644
--- /dev/null
+++ b/Engine/Source/Runtime/Core/Private/Modules/Tests/VisualizerDebuggingStateTests.cpp
@@ -0,0 +1,118 @@
+// Copyright Epic Games, Inc. All Rights Reserved.
+
+#include "Misc/AutomationTest.h"
+#include "Modules/VisualizerDebuggingState.h"
+#include "HAL/UnrealMemory.h"
+
+#if WITH_DEV_AUTOMATION_TESTS && UE_VISUALIZER_DEBUGGING_STATE
+
+IMPLEMENT_SIMPLE_AUTOMATION_TEST(FCoreDebuggingStateTest, 
+		"System.Core.DebuggingState", 
+		EAutomationTestFlags::EditorContext | EAutomationTestFlags::EngineFilter)
+bool FCoreDebuggingStateTest::RunTest(const FString& Parameters)
+{
+	using namespace UE::Core;
+
+	struct FTestData
+	{
+		FGuid ID;
+		const char* String;
+		uint32* Pointer;
+	};
+
+	struct FTestState : FVisualizerDebuggingState
+	{
+		EVisualizerDebuggingStateResult Assign(const FGuid& UniqueId, void* DebugPtr)
+		{
+			return AssignImpl(UniqueId, DebugPtr);
+		}
+		void* Find(const FGuid& UniqueId) const
+		{
+			return FVisualizerDebuggingState::Find(UniqueId);
+		}
+	};
+
+	FTestState TestState;
+
+	uint32 ResultValues[] = {
+		1u,
+		10u,
+		100u,
+		1000u,
+		2u,
+		20u,
+		200u,
+		2000u
+	};
+
+	FTestData Tests[] = {
+		{ FGuid(), "93a891b0f3404d9c9b1f51981966e1e0", &ResultValues[0] },
+		{ FGuid(), "06a5fe3be35d4d2987abfeaea8c54035", &ResultValues[1] },
+		{ FGuid(), "e87d5c5d7f9948d4a051de51ecfb9b25", &ResultValues[2] },
+		{ FGuid(), "1233cca4d6ee400cad1ca3f8802ac523", &ResultValues[3] },
+	};
+
+	// Try assigning the pointers multiple times with different values each time
+	for (int32 Pass = 0; Pass < 2; ++Pass)
+	{
+		for (int32 TestIndex = 0; TestIndex < UE_ARRAY_COUNT(Tests); ++TestIndex)
+		{
+			FTestData& Test = Tests[TestIndex];
+
+			Test.Pointer = &ResultValues[TestIndex + Pass*UE_ARRAY_COUNT(Tests)];
+
+			if (!FGuid::ParseExact(Test.String, EGuidFormats::DigitsLower, Test.ID))
+			{
+				this->AddError(FString::Printf(TEXT("Error parsing GUID string %hs!"), Test.String));
+				continue;
+			}
+
+			EVisualizerDebuggingStateResult Result = TestState.Assign(Test.ID, Test.Pointer);
+			if (Result != EVisualizerDebuggingStateResult::Success)
+			{
+				this->AddError(FString::Printf(TEXT("There was an error registering %hs was not found!"), Test.String));
+			}
+		}
+	}
+
+	void* ExpectedResults[] = {
+		&ResultValues[4],
+		&ResultValues[5],
+		&ResultValues[6],
+		&ResultValues[7],
+	};
+
+	for (int32 Index = 0; Index < UE_ARRAY_COUNT(Tests); ++Index)
+	{
+		const FTestData& Test = Tests[Index];
+		void* Result = TestState.Find(Test.ID);
+		if (Result == nullptr)
+		{
+			this->AddError(FString::Printf(TEXT("Debugging pointer for moniker %hs was not found!"), Test.String));
+		}
+		else if (Result != Test.Pointer || Result != ExpectedResults[Index])
+		{
+			this->AddError(FString::Printf(TEXT("Debugging pointer for moniker %hs was incorrect!"), Test.String));
+		}
+	}
+
+	// Test that assigning a fake GUID comprising the others results in a string collision
+	static const char* const CollisionGuid = "9b1f51981966e1e006a5fe3be35d4d29";
+	FGuid CollisionId;
+	if (!FGuid::ParseExact(CollisionGuid, EGuidFormats::DigitsLower, CollisionId))
+	{
+		this->AddError(FString::Printf(TEXT("Error parsing GUID string %hs!"), CollisionGuid));
+		return false;
+	}
+
+	if (TestState.Assign(CollisionId, nullptr) != EVisualizerDebuggingStateResult::StringCollision)
+	{
+		this->AddError(FString::Printf(TEXT("Expected FVisualizerDebuggingState::Assign(\"%s\") to result in a string collision but it did not"), CollisionGuid));
+		return false;
+	}
+
+	return true;
+}
+
+#endif
+
diff --git a/Engine/Source/Runtime/Core/Private/Modules/VisualizerDebuggingState.cpp b/Engine/Source/Runtime/Core/Private/Modules/VisualizerDebuggingState.cpp
new file mode 100644
--- /dev/null
+++ b/Engine/Source/Runtime/Core/Private/Modules/VisualizerDebuggingState.cpp
@@ -0,0 +1,122 @@
+// Copyright Epic Games, Inc. All Rights Reserved.
+
+#include "Modules/VisualizerDebuggingState.h"
+#include "CoreGlobals.h"
+#include "HAL/UnrealMemory.h"
+#include "Containers/StringView.h"
+#include "Containers/AnsiString.h"
+#include "Algo/IndexOf.h"
+#include "Misc/Guid.h"
+
+UE::Core::FVisualizerDebuggingState* GCoreDebuggingState = nullptr;
+
+namespace UE::Core
+{
+
+#if UE_VISUALIZER_DEBUGGING_STATE
+struct FVisualizerDebuggingStateImpl
+{
+	/**
+	 * Array of GUIDs associated with indices within DebugPtrs
+	 */
+	TArray<FGuid> UniqueIds;
+
+	/**
+	 * Array of user-provided pointers that are associated with the guids in UniqueIds.
+	 */
+	TArray<void*> DebugPtrs;
+
+	/**
+	 * Single string containing all Guids as a string at position Index*32 to support searching with a single strstr instrinsic call, and dividing the result by 32
+	 * Example layout:
+	 *   OrderedGuidString: "d9ad42709d2c4bc8a2f8f925e1617b288456cfc3222f4833a8afa45a6ed73b5a9873a9fedd8441d49ef6b3258e8a6c60#"
+	 *   Ptrs:              [Ptr1                           ,Ptr2                           ,Ptr3]
+	 */
+	FAnsiString OrderedGuidString;
+};
+#endif // UE_VISUALIZER_DEBUGGING_STATE
+
+FVisualizerDebuggingState::FVisualizerDebuggingState()
+{
+#if UE_VISUALIZER_DEBUGGING_STATE
+	PimplData = new FVisualizerDebuggingStateImpl;
+#endif
+}
+
+FVisualizerDebuggingState::~FVisualizerDebuggingState()
+{
+#if UE_VISUALIZER_DEBUGGING_STATE
+	delete PimplData;
+#endif
+}
+
+#if UE_VISUALIZER_DEBUGGING_STATE
+
+EVisualizerDebuggingStateResult FVisualizerDebuggingState::Assign(const FGuid& UniqueId, void* Ptr)
+{
+	if (!GCoreDebuggingState)
+	{
+		static FVisualizerDebuggingState State;
+		GCoreDebuggingState = &State;
+	}
+	return GCoreDebuggingState->AssignImpl(UniqueId, Ptr);
+}
+
+void* FVisualizerDebuggingState::Find(const FGuid& UniqueId) const
+{
+	if (GuidString != nullptr)
+	{
+		const int32 Index = Algo::IndexOf(PimplData->UniqueIds, UniqueId);
+		if (Index != INDEX_NONE)
+		{
+			return Ptrs[Index];
+		}
+	}
+
+	return nullptr;
+}
+
+EVisualizerDebuggingStateResult FVisualizerDebuggingState::AssignImpl(const FGuid& UniqueId, void* Ptr)
+{
+	check(UniqueId.IsValid());
+
+	// Locate an existing entry from the string
+	const int32 ExistingIndex = Algo::IndexOf(PimplData->UniqueIds, UniqueId);
+	if (ExistingIndex != INDEX_NONE)
+	{
+		Ptrs[ExistingIndex] = Ptr;
+		return EVisualizerDebuggingStateResult::Success;
+	}
+
+	TAnsiStringBuilder<32> ThisGuidString;
+	UniqueId.AppendString(ThisGuidString, EGuidFormats::DigitsLower);
+
+	// Search the existing string for a potential string collision (this is extremely unlikely)
+	if (GuidString != nullptr && FCStringAnsi::Strstr(GuidString, ThisGuidString.ToString()) != nullptr)
+	{
+		return EVisualizerDebuggingStateResult::StringCollision;
+	}
+
+	// Add a new entry:
+	const int32 PreviousNum = PimplData->UniqueIds.Num();
+	const int32 NewNum      = PreviousNum + 1;
+
+	// Append the new GUID string to the hasystack string
+	PimplData->OrderedGuidString.Append(ThisGuidString);
+	PimplData->UniqueIds.Emplace(UniqueId);
+	PimplData->DebugPtrs.Emplace(Ptr);
+
+	// Check invariants
+	check(PimplData->DebugPtrs.Num()    == PimplData->UniqueIds.Num());
+	check(PimplData->DebugPtrs.Num()*32 == PimplData->OrderedGuidString.Len());
+
+	// Update the cached ptrs for direct access inside natvis
+	Ptrs = PimplData->DebugPtrs.GetData();
+	GuidString = *PimplData->OrderedGuidString;
+
+	return EVisualizerDebuggingStateResult::Success;
+}
+
+#endif // UE_VISUALIZER_DEBUGGING_STATE
+
+} // namespace UE::Core
\ No newline at end of file
diff --git a/Engine/Source/Runtime/Core/Public/CoreGlobals.h b/Engine/Source/Runtime/Core/Public/CoreGlobals.h
--- a/Engine/Source/Runtime/Core/Public/CoreGlobals.h
+++ b/Engine/Source/Runtime/Core/Public/CoreGlobals.h
@@ -822,8 +822,10 @@ namespace UE::CoreUObject::Private
 	/** Array to help visualize object handles in the debugger */
 	struct FObjectHandlePackageDebugData;
 }
+namespace UE::Core { struct FVisualizerDebuggingState; }
 extern CORE_API UE::CoreUObject::Private::FStoredObjectPathDebug* GCoreComplexObjectPathDebug;
 extern CORE_API UE::CoreUObject::Private::FObjectHandlePackageDebugData* GCoreObjectHandlePackageDebug;
+extern CORE_API UE::Core::FVisualizerDebuggingState* GCoreDebuggingState;
 
 /** @return True if running cook-on-the-fly. */
 bool CORE_API IsRunningCookOnTheFly();
diff --git a/Engine/Source/Runtime/Core/Public/Modules/Boilerplate/ModuleBoilerplate.h b/Engine/Source/Runtime/Core/Public/Modules/Boilerplate/ModuleBoilerplate.h
--- a/Engine/Source/Runtime/Core/Public/Modules/Boilerplate/ModuleBoilerplate.h
+++ b/Engine/Source/Runtime/Core/Public/Modules/Boilerplate/ModuleBoilerplate.h
@@ -5,6 +5,10 @@
 #include "CoreTypes.h"
 #include "HAL/UnrealMemory.h"
 #include "UObject/NameTypes.h"
+// Include the full definition of FVisualizerDebuggingState so visualizers can see the full type information
+//   Without this include (and just a fwd decl), visualizers in some modules will not be able to resolve
+//   FVisualizerDebuggingState::Ptrs or FVisualizerDebuggingState::GuidString
+#include "Modules/VisualizerDebuggingState.h"
 
 class FChunkedFixedUObjectArray;
 
@@ -93,7 +97,8 @@ class FChunkedFixedUObjectArray;
 		uint8** GNameBlocksDebug = FNameDebugVisualizer::GetBlocks(); \
 		FChunkedFixedUObjectArray*& GObjectArrayForDebugVisualizers = GCoreObjectArrayForDebugVisualizers; \
 		UE::CoreUObject::Private::FStoredObjectPathDebug*& GComplexObjectPathDebug = GCoreComplexObjectPathDebug; \
-		UE::CoreUObject::Private::FObjectHandlePackageDebugData*& GObjectHandlePackageDebug = GCoreObjectHandlePackageDebug;
+		UE::CoreUObject::Private::FObjectHandlePackageDebugData*& GObjectHandlePackageDebug = GCoreObjectHandlePackageDebug; \
+		UE::Core::FVisualizerDebuggingState*& GDebuggingState = GCoreDebuggingState;
 #endif
 
 // in DLL builds, these are done per-module, otherwise we just need one in the application
diff --git a/Engine/Source/Runtime/Core/Public/Modules/VisualizerDebuggingState.h b/Engine/Source/Runtime/Core/Public/Modules/VisualizerDebuggingState.h
new file mode 100644
--- /dev/null
+++ b/Engine/Source/Runtime/Core/Public/Modules/VisualizerDebuggingState.h
@@ -0,0 +1,148 @@
+// Copyright Epic Games, Inc. All Rights Reserved.
+
+#pragma once
+
+#include "CoreTypes.h"
+
+#ifndef UE_VISUALIZER_DEBUGGING_STATE
+	#define UE_VISUALIZER_DEBUGGING_STATE !UE_BUILD_SHIPPING
+#endif
+
+struct FGuid;
+
+namespace UE::Core
+{
+
+struct FVisualizerDebuggingStateImpl;
+
+/** Enumeration specifying how a call to FVisualizerDebuggingState::Assign succeeded or failed */
+enum class EVisualizerDebuggingStateResult : uint8
+{
+	/** The operation completed successfully */
+	Success,
+	/** Failure: The supplied guid happened to collide with the concatenation of 2 other pre-assigned GUIDs Please use a different GUID. */
+	StringCollision,
+	/** Failure: The feature is disabled entirely. */
+	FeatureDisabled,
+};
+
+/**
+ * Global visualizer state manager that handles the complexities of natvis lookup rules for DLLs and LiveCoding
+ *
+ * The singleton instance of this class is forward declared in CoreGlobals.h (GCoreDebuggingState), with
+ *   additional symbols added to each DLL through UE_VISUALIZERS_HELPERS (GDebuggingState).
+ *
+ * When a plugin or engine system needs to provide a globally accessible pointer to a natvis visualizer,
+ *   that pointer can be added by calling FVisualizerDebuggingState::Assign using a globally unique ID
+ *   that is used to find the pointer in the natvis. GUIDs may be generated using any readily-available GUID generator tool.
+
+ * Pointers are located using an intrinsic call to strstr that finds the pointer's offset using a string match
+ *   of the original GUID's lower-case, 32-character string representation of the form "c198d347ec1243ec833c423bd5fb6084".
+ *   The offset of the address of the  match divided by 32 gives the entry index from which the actual debug ptr can be retrieved.
+ *
+ * An example visualizer syntax is included below. It is recommended that most natvis files utilizing this method would include an
+ *   global intrinsic function at the top for quick access:
+ * 
+ *   <Intrinsic Name="GetDebugState" Expression="GDebuggingState->Ptrs[(GDebuggingState->GuidString - strstr(GDebuggingState->GuidString, GuidString))/32]">
+ *     <Parameter Name="GuidString" Type="char*"/>
+ *   </Intrinsic>
+ *
+ *   // C++:
+ *   struct FMyGlobalState
+ *   {
+ *     static FMyGlobalState* GetSingleton();
+ *     TArray<FString> TypeNames;
+ *   };
+ *   struct FMyStruct
+ *   {
+ *     int32 TypeNameIndex;
+ *   };
+ * 
+ *   // Register the global state
+ *   FGuid DebugVisualizerID(0x1bdc1747, 0x6b924697, 0xbd8fccc8, 0xb26b4f79)
+ *   EVisualizerDebuggingStateResult Result = FVisualizerDebuggingState::Assign(DebugVisualizerID, FMyGlobalState::GetSingleton());
+ *
+ *   // .natvis:
+ *   <!-- Example natvis syntax for finding a FMyGlobalState* registered with the guid 1bdc1747-6b92-4697-bd8f-ccc8b26b4f79 -->
+ *   <Type Name="FMyStruct">
+ *     <!-- less verbose, more fragile version -->
+ *     <Intrinsic Name="GetGlobalState" Expression="((FMyGlobalState*)GetDebugState(&quot;1bdc17476b924697bd8fccc8b26b4f79&quot;))""></Intrinsic>
+ *
+ *     <DisplayString>Type={ GetGlobalState()->TypeNames[TypeNameIndex] }</DisplayString>
+ *   </Type>
+ */
+struct FVisualizerDebuggingState
+{
+	FVisualizerDebuggingState();
+
+	FVisualizerDebuggingState(const FVisualizerDebuggingState&) = delete;
+	void operator=(const FVisualizerDebuggingState&) = delete;
+
+	FVisualizerDebuggingState(FVisualizerDebuggingState&&) = delete;
+	void operator=(FVisualizerDebuggingState&&) = delete;
+
+	~FVisualizerDebuggingState();
+
+#if !UE_VISUALIZER_DEBUGGING_STATE
+
+	/**
+	 * Assign a globally accessible debugging state ptr by name, potentially overwriting a previously assigned ptr
+	 * 
+	 * @param UniqueId   A guid that uniquely identifies this debug ptr. This should be the same as is used to access the ptr inside a natvis file.
+	 * @param DebugPtr   The pointer that needs to be accessible from natvis
+	 * @return EVisualizerDebuggingStateResult::Success if the operation completed successfully, otherwise an error code
+	 */
+	[[nodiscard]] static EVisualizerDebuggingStateResult Assign(const FGuid& UniqueId, void* DebugPtr)
+	{
+		return EVisualizerDebuggingStateResult::FeatureDisabled;
+	}
+
+#else
+
+	/**
+	 * Assign a globally accessible debugging state ptr by name, potentially overwriting a previously assigned ptr
+	 * 
+	 * @param UniqueId   A guid that uniquely identifies this debug ptr. This should be the same as is used to access the ptr inside a natvis file.
+	 * @param DebugPtr   The pointer that needs to be accessible from natvis
+	 * @return EVisualizerDebuggingStateResult::Success if the operation completed successfully, otherwise an error code
+	 */
+	[[nodiscard]] CORE_API static EVisualizerDebuggingStateResult Assign(const FGuid& UniqueId, void* DebugPtr);
+
+
+protected:
+	// protected in order to support automation test introspection
+
+	/**
+	 * Implementation function for Assign
+	 */
+	EVisualizerDebuggingStateResult AssignImpl(const FGuid& UniqueId, void* DebugPtr);
+
+	/**
+	 * Attempt to locate the index of the specified unique ID
+	 */
+	void* Find(const FGuid& UniqueId) const;
+
+	/**
+	 * Single string containing all GUID strings ordered by their insertion order.
+	 * The single string is designed to support finding a string offset using a single strstr instrinsic call, and dividing the result by 32.
+	 * This string exists for easy referencing inside natvis files; it points to PimplData->OrderedGuidString.
+	 * Example layout:
+	 *   
+	 *   OrderedGuidString: "d9ad42709d2c4bc8a2f8f925e1617b288456cfc3222f4833a8afa45a6ed73b5a9873a9fedd8441d49ef6b3258e8a6c60#"
+	 *   Ptrs:              [Ptr1                           ,Ptr2                           ,Ptr3]
+	 */
+	const char* GuidString = nullptr;
+
+	/**
+	 * Array of void* ptrs to the user-provided debugging state ptrs for each entry.
+	 * Points directly to PimplData->DebugPtrs.GetData() for quick referencing inside natvis without needing any other symbols or casting.
+	 */
+	void** Ptrs = nullptr;
+
+	/** Pimpl to avoid pulling in too many headers since this header is included everywhere */
+	FVisualizerDebuggingStateImpl* PimplData = nullptr;
+
+#endif
+};
+
+} // namespace UE::Core
