TSpscQueue support for holding UniquePtr.

diff --git a/Engine/Source/Runtime/Core/Public/Containers/SpscQueue.h b/Engine/Source/Runtime/Core/Public/Containers/SpscQueue.h
--- a/Engine/Source/Runtime/Core/Public/Containers/SpscQueue.h
+++ b/Engine/Source/Runtime/Core/Public/Containers/SpscQueue.h
@@ -86,7 +86,7 @@ public:
 		TOptional<ElementType> LocalElement = Dequeue();
 		if (LocalElement.IsSet())
 		{
-			OutElem = LocalElement.GetValue();
+			OutElem = MoveTempIfPossible(LocalElement.GetValue());
 			return true;
 		}
 		
