- Fixed several #ifdef 's that should be #if (caused Mac code to execute on IOS)

diff --git a/Engine/Source/Runtime/Core/Private/Apple/ApplePlatformStackWalk.cpp b/Engine/Source/Runtime/Core/Private/Apple/ApplePlatformStackWalk.cpp
--- a/Engine/Source/Runtime/Core/Private/Apple/ApplePlatformStackWalk.cpp
+++ b/Engine/Source/Runtime/Core/Private/Apple/ApplePlatformStackWalk.cpp
@@ -22,7 +22,7 @@
 #include "PLCrashReporter.h"
 #endif
 
-#ifdef PLATFORM_MAC
+#if PLATFORM_MAC
 static TAutoConsoleVariable<int32> CVarApplePlatformThreadCallStackEnabled(
 	TEXT("ApplePlatformThreadStackWalk.Enable"),
 	1,
@@ -304,7 +304,7 @@ int32 FApplePlatformStackWalk::GetProcessModuleSignatures(FStackWalkModuleInfo *
 }
 
 // Consider this for iOS
-#ifdef PLATFORM_MAC
+#if PLATFORM_MAC
 FCriticalSection GThreadCallStackMutex;
 
 // These are used in the SIGUSR2 callback to hand over the user provided buffer to copy to
@@ -350,7 +350,7 @@ static void RaiseSIGUSR2ForThreadAndWait(uint32 ThreadId)
 void FApplePlatformStackWalk::ThreadStackWalkAndDump(ANSICHAR* HumanReadableString, SIZE_T HumanReadableStringSize, int32 IgnoreCount, uint32 ThreadId)
 {
 // Consider this for iOS
-#ifdef PLATFORM_MAC
+#if PLATFORM_MAC
 	if (CVarApplePlatformThreadCallStackEnabled.AsVariable()->GetInt())
 	{
 		FScopeLock Lock(&GThreadCallStackMutex);
@@ -370,7 +370,7 @@ void FApplePlatformStackWalk::ThreadStackWalkAndDump(ANSICHAR* HumanReadableStri
 uint32 FApplePlatformStackWalk::CaptureThreadStackBackTrace(uint64 ThreadId, uint64* BackTrace, uint32 MaxDepth, void* Context)
 {
 // Consider this for iOS
-#ifdef PLATFORM_MAC
+#if PLATFORM_MAC
 	GThreadBackTraceCount = 0;
 
 	if (CVarApplePlatformThreadCallStackEnabled.AsVariable()->GetInt())
@@ -388,6 +388,8 @@ uint32 FApplePlatformStackWalk::CaptureThreadStackBackTrace(uint64 ThreadId, uin
 	}
 
 	return GThreadBackTraceCount;
+#else
+	return 0;
 #endif
 }
 
