[AutoRTFM] Make FPooledLargeMemoryData transaction-safe

diff --git a/Engine/Source/Runtime/Core/Private/Serialization/LargeMemoryData.cpp b/Engine/Source/Runtime/Core/Private/Serialization/LargeMemoryData.cpp
--- a/Engine/Source/Runtime/Core/Private/Serialization/LargeMemoryData.cpp
+++ b/Engine/Source/Runtime/Core/Private/Serialization/LargeMemoryData.cpp
@@ -2,6 +2,7 @@
 
 #include "Serialization/LargeMemoryData.h"
 
+#include "AutoRTFM/AutoRTFM.h"
 #include "Containers/LockFreeList.h"
 #include "HAL/IConsoleManager.h"
 #include "HAL/UnrealMemory.h"
@@ -127,18 +128,30 @@ std::atomic<int32> FPooledLargeMemoryData::FreeListLength;
 
 FPooledLargeMemoryData::FPooledLargeMemoryData()
 {
-	Data = FreeList.Pop();
-	if (Data != nullptr)
+	AutoRTFM::Open([&]
 	{
-		--FreeListLength;
-	}
-	else
-	{
-		Data = new FLargeMemoryData();
-	}
+		Data = FreeList.Pop();
+		if (Data != nullptr)
+		{
+			--FreeListLength;
+		}
+		else
+		{
+			Data = new FLargeMemoryData();
+		}
+	});
+	AutoRTFM::OnAbort([Data = this->Data] { Free(Data); });
 }
 
 FPooledLargeMemoryData::~FPooledLargeMemoryData()
+{
+	// If outside of a closed transaction, immediately free the FLargeMemoryData.
+	// If inside a closed transaction, then defer the free until the transaction
+	// completes as writes to the memory may be undone if the transaction aborts.
+	AutoRTFM::OnCommit([Data = this->Data] { Free(Data); });
+}
+
+void FPooledLargeMemoryData::Free(FLargeMemoryData* Data)
 {
 	if (Data->GetSize() > GPooledLargeMemoryData_MaxElementSize)
 	{
diff --git a/Engine/Source/Runtime/Core/Public/Serialization/LargeMemoryData.h b/Engine/Source/Runtime/Core/Public/Serialization/LargeMemoryData.h
--- a/Engine/Source/Runtime/Core/Public/Serialization/LargeMemoryData.h
+++ b/Engine/Source/Runtime/Core/Public/Serialization/LargeMemoryData.h
@@ -90,6 +90,7 @@ public:
 	CORE_API ~FPooledLargeMemoryData();
 	FLargeMemoryData& Get() { return *Data; }
 private:
+	static void Free(FLargeMemoryData* Data);
 	FLargeMemoryData* Data;
 	static CORE_API TLockFreePointerListUnordered<FLargeMemoryData, 0> FreeList;
 	static CORE_API std::atomic<int32> FreeListLength;
diff --git a/Engine/Source/Runtime/Engine/Private/Tests/AutoRTFM/AutoRTFMLargeMemoryDataTest.cpp b/Engine/Source/Runtime/Engine/Private/Tests/AutoRTFM/AutoRTFMLargeMemoryDataTest.cpp
new file mode 100644
--- /dev/null
+++ b/Engine/Source/Runtime/Engine/Private/Tests/AutoRTFM/AutoRTFMLargeMemoryDataTest.cpp
@@ -0,0 +1,41 @@
+// Copyright Epic Games, Inc. All Rights Reserved.
+
+#include "Misc/AutomationTest.h"
+#include "AutoRTFM/AutoRTFM.h"
+#include "Serialization/LargeMemoryData.h"
+
+#if WITH_DEV_AUTOMATION_TESTS
+
+IMPLEMENT_SIMPLE_AUTOMATION_TEST(FAutoRTFMLargeMemoryDataTest, "AutoRTFM + FPooledLargeMemoryData", EAutomationTestFlags::EngineFilter | EAutomationTestFlags::ClientContext | EAutomationTestFlags::ServerContext | EAutomationTestFlags::CommandletContext)
+
+bool FAutoRTFMLargeMemoryDataTest::RunTest(const FString & Parameters)
+{
+	if (!AutoRTFM::ForTheRuntime::IsAutoRTFMRuntimeEnabled())
+	{
+		ExecutionInfo.AddEvent(FAutomationEvent(EAutomationEventType::Info, TEXT("SKIPPED 'FAutoRTFMLargeMemoryDataTest' test. AutoRTFM disabled.")));
+		return true;
+	}
+
+	AutoRTFM::Transact([&]
+	{
+		FPooledLargeMemoryData Data;
+	});
+
+	AutoRTFM::Transact([&]
+	{
+		FPooledLargeMemoryData Data;
+		AutoRTFM::AbortTransaction();
+	});
+
+	AutoRTFM::Transact([&]
+	{
+		{
+			FPooledLargeMemoryData Data;
+		}
+		AutoRTFM::AbortTransaction();
+	});
+
+	return true;
+}
+
+#endif //WITH_DEV_AUTOMATION_TESTS
