Reduce thread contention on MB2/3

diff --git a/Engine/Source/Runtime/Core/Private/HAL/MallocBinned2.cpp b/Engine/Source/Runtime/Core/Private/HAL/MallocBinned2.cpp
--- a/Engine/Source/Runtime/Core/Private/HAL/MallocBinned2.cpp
+++ b/Engine/Source/Runtime/Core/Private/HAL/MallocBinned2.cpp
@@ -85,11 +85,11 @@ PRAGMA_DISABLE_UNSAFE_TYPECAST_WARNINGS
 
 
 #if UE_MB2_ALLOCATOR_STATS
-	TAtomic<int64> AllocatedSmallPoolMemory(0); // memory that's requested to be allocated by the game
-	TAtomic<int64> AllocatedOSSmallPoolMemory(0);
+	std::atomic<int64> AllocatedSmallPoolMemory(0); // memory that's requested to be allocated by the game
+	std::atomic<int64> AllocatedOSSmallPoolMemory(0);
 
-	TAtomic<int64> AllocatedLargePoolMemory(0); // memory requests to the OS which don't fit in the small pool
-	TAtomic<int64> AllocatedLargePoolMemoryWAlignment(0); // when we allocate at OS level we need to align to a size
+	std::atomic<int64> AllocatedLargePoolMemory(0); // memory requests to the OS which don't fit in the small pool
+	std::atomic<int64> AllocatedLargePoolMemoryWAlignment(0); // when we allocate at OS level we need to align to a size
 
 	int64 Binned2PoolInfoMemory = 0;
 	int64 Binned2HashMemory = 0;
@@ -499,6 +499,8 @@ struct FMallocBinned2::Private
 	{
 		FPoolTable& Table = Allocator.SmallPoolTables[InPoolIndex];
 
+		FScopeLock Lock(&Table.Mutex);
+
 		FBundleNode* Bundle = BundlesToRecycle;
 		while (Bundle)
 		{
@@ -544,9 +546,12 @@ struct FMallocBinned2::Private
 
 					// Free the OS memory.
 					NodePool->Unlink();
-					Allocator.CachedOSPageAllocator.Free(BasePtrOfNode, Allocator.PageSize);
+					{
+						FScopeLock InnerLock(&Allocator.Mutex);
+						Allocator.CachedOSPageAllocator.Free(BasePtrOfNode, Allocator.PageSize);
+					}
 #if UE_MB2_ALLOCATOR_STATS
-					AllocatedOSSmallPoolMemory -= ((int64)Allocator.PageSize);
+					AllocatedOSSmallPoolMemory.fetch_sub((int64)Allocator.PageSize, std::memory_order_relaxed);
 #endif
 				}
 
@@ -659,10 +664,14 @@ FMallocBinned2::FPoolInfo& FMallocBinned2::FPoolList::PushNewPoolToFront(FMalloc
 
 	// Allocate memory.
 	LLM_PLATFORM_SCOPE(ELLMTag::FMalloc);
-	void* FreePtr = Allocator.CachedOSPageAllocator.Allocate(LocalPageSize, FMemory::AllocationHints::SmallPool);
+
+	void* FreePtr;
+	{
+		FScopeLock Lock(&Allocator.Mutex);
+		FreePtr = Allocator.CachedOSPageAllocator.Allocate(LocalPageSize, FMemory::AllocationHints::SmallPool);
+	}
 	if (!FreePtr)
 	{
-		FScopeUnlock TempUnlock(&Allocator.Mutex);
 		Private::OutOfMemory(LocalPageSize);
 	}
 
@@ -675,8 +684,10 @@ FMallocBinned2::FPoolInfo& FMallocBinned2::FPoolList::PushNewPoolToFront(FMalloc
 #endif
 
 #if UE_MB2_ALLOCATOR_STATS
-	AllocatedOSSmallPoolMemory += (int64)LocalPageSize;
+	AllocatedOSSmallPoolMemory.fetch_add((int64)LocalPageSize, std::memory_order_relaxed);
 #endif
+
+	FScopeLock Lock(&Allocator.Mutex);
 	// Create pool
 	FPoolInfo* Result = Private::GetOrCreatePoolInfo(Allocator, Free, FPoolInfo::ECanary::FirstFreeBlockIsPtr, false);
 	Result->Link(Front);
@@ -824,8 +835,6 @@ void FMallocBinned2::OnPostFork()
 		FMallocBinned2::Private::CheckThreadFreeBlockListsForFork();
 	}
 
-	FScopeLock Lock(&Mutex);
-
 	// This will be compared against the pool header of existing allocations to turn Free into a no-op for pages shared with the parent process
 	UE_CLOG(CurrentCanary != EBlockCanary::PreFork, LogMemory, Fatal, TEXT("FMallocBinned2 only supports forking once!"));
 
@@ -834,6 +843,7 @@ void FMallocBinned2::OnPostFork()
 
 	for (FPoolTable& Table : SmallPoolTables)
 	{
+		FScopeLock Lock(&Table.Mutex);
 		// Clear our list of partially used pages so we don't dirty them and cause them to become unshared with the parent process
 		Table.ActivePools.Clear();
 		Table.ExhaustedPools.Clear();
@@ -881,11 +891,11 @@ void* FMallocBinned2::MallocExternalSmall(SIZE_T Size, uint32 Alignment)
 
 	NOALLOC_SCOPE_CYCLE_COUNTER(STAT_FMallocBinned2_MallocExternalSmall);
 
-	FScopeLock Lock(&Mutex);
-
 	// Allocate from small object pool.
 	FPoolTable& Table = SmallPoolTables[PoolIndex];
 
+	FScopeLock Lock(&Table.Mutex);
+
 	FPoolInfo* Pool;
 	if (!Table.ActivePools.IsEmpty())
 	{
@@ -898,7 +908,7 @@ void* FMallocBinned2::MallocExternalSmall(SIZE_T Size, uint32 Alignment)
 
 	void* Result = Pool->AllocateBin();
 #if UE_MB2_ALLOCATOR_STATS
-	AllocatedSmallPoolMemory += Table.BinSize;
+	AllocatedSmallPoolMemory.fetch_add(Table.BinSize, std::memory_order_relaxed);
 #endif
 
 	if (GBinned2AllocExtra)
@@ -961,8 +971,8 @@ void* FMallocBinned2::MallocExternalLarge(SIZE_T Size, uint32 Alignment)
 	check(IsAligned(Result, PageSize) && IsOSAllocation(Result));
 
 #if UE_MB2_ALLOCATOR_STATS
-	AllocatedLargePoolMemory += Size;
-	AllocatedLargePoolMemoryWAlignment += AlignedSize;
+	AllocatedLargePoolMemory.fetch_add(Size, std::memory_order_relaxed);
+	AllocatedLargePoolMemoryWAlignment.fetch_add(AlignedSize, std::memory_order_relaxed);
 #endif
 
 	// Create pool.
@@ -1064,8 +1074,7 @@ void* FMallocBinned2::ReallocExternal(void* Ptr, SIZE_T NewSize, uint32 Alignmen
 	checkf(FMallocBinned2::FPoolInfo::IsSupportedSize(NewSize), TEXT("Invalid Realloc size: '%" SIZE_T_FMT "'"), NewSize);
 
 #if UE_MB2_ALLOCATOR_STATS
-	AllocatedLargePoolMemory += ((int64)NewSize) - ((int64)PoolOSRequestedBytes);
-	// don't need to change the AllocatedLargePoolMemoryWAlignment because we didn't reallocate so it's the same size
+	AllocatedLargePoolMemory.fetch_add((int64)NewSize - (int64)PoolOSRequestedBytes, std::memory_order_relaxed);
 #endif
 
 	Pool->SetOSAllocationSizes(NewSize, PoolOsBytes);
@@ -1113,13 +1122,12 @@ void FMallocBinned2::FreeExternal(void* Ptr)
 		if (BundlesToRecycle)
 		{
 			BundlesToRecycle->NextBundle = nullptr;
-			FScopeLock Lock(&Mutex);
 			Private::FreeBundles(*this, BundlesToRecycle, BinSize, PoolIndex);
 #if UE_MB2_ALLOCATOR_STATS
 			if (!Lists)
 			{
 				// lists track their own stat track them instead in the global stat if we don't have lists
-				AllocatedSmallPoolMemory -= ((int64)(BinSize));
+				AllocatedSmallPoolMemory.fetch_sub((int64)BinSize, std::memory_order_relaxed);
 			}
 #endif
 		}
@@ -1136,8 +1144,8 @@ void FMallocBinned2::FreeExternal(void* Ptr)
 		const SIZE_T PoolOSRequestedBytes = Pool->GetOSRequestedBytes();
 
 #if UE_MB2_ALLOCATOR_STATS
-		AllocatedLargePoolMemory -= ((int64)PoolOSRequestedBytes);
-		AllocatedLargePoolMemoryWAlignment -= ((int64)PoolOsBytes);
+		AllocatedLargePoolMemory.fetch_sub((int64)PoolOSRequestedBytes, std::memory_order_relaxed);
+		AllocatedLargePoolMemoryWAlignment.fetch_sub((int64)PoolOsBytes, std::memory_order_relaxed);
 #endif
 
 		checkf(PoolOSRequestedBytes <= PoolOsBytes, TEXT("FMallocBinned2::FreeExternal %d %d"), int32(PoolOSRequestedBytes), int32(PoolOsBytes));
@@ -1207,10 +1215,10 @@ void FMallocBinned2::FPoolList::ValidateExhaustedPools() const
 bool FMallocBinned2::ValidateHeap()
 {
 	NOALLOC_SCOPE_CYCLE_COUNTER(STAT_FMallocBinned2_ValidateHeap);
-	FScopeLock Lock(&Mutex);
-
+	
 	for (FPoolTable& Table : SmallPoolTables)
 	{
+		FScopeLock Lock(&Table.Mutex);
 		Table.ActivePools.ValidateActivePools();
 		Table.ExhaustedPools.ValidateExhaustedPools();
 	}
@@ -1329,7 +1337,7 @@ int64 FMallocBinned2::GetTotalAllocatedSmallPoolMemory() const
 		FreeBlockAllocatedMemory += ConsolidatedMemory.load(std::memory_order_relaxed);
 	}
 
-	return AllocatedSmallPoolMemory.Load(EMemoryOrder::Relaxed) + FreeBlockAllocatedMemory;
+	return AllocatedSmallPoolMemory.load(std::memory_order_relaxed) + FreeBlockAllocatedMemory;
 }
 #endif
 
@@ -1341,9 +1349,9 @@ void FMallocBinned2::GetAllocatorStats( FGenericMemoryStats& OutStats )
 
 #if UE_MB2_ALLOCATOR_STATS
 	const int64  TotalAllocatedSmallPoolMemory           = GetTotalAllocatedSmallPoolMemory();
-	const int64  LocalAllocatedOSSmallPoolMemory         = AllocatedOSSmallPoolMemory.Load(EMemoryOrder::Relaxed);
-	const int64  LocalAllocatedLargePoolMemory           = AllocatedLargePoolMemory.Load(EMemoryOrder::Relaxed);
-	const int64  LocalAllocatedLargePoolMemoryWAlignment = AllocatedLargePoolMemoryWAlignment.Load(EMemoryOrder::Relaxed);
+	const int64  LocalAllocatedOSSmallPoolMemory         = AllocatedOSSmallPoolMemory.load(std::memory_order_relaxed);
+	const int64  LocalAllocatedLargePoolMemory           = AllocatedLargePoolMemory.load(std::memory_order_relaxed);
+	const int64  LocalAllocatedLargePoolMemoryWAlignment = AllocatedLargePoolMemoryWAlignment.load(std::memory_order_relaxed);
 
 	OutStats.Add(TEXT("AllocatedSmallPoolMemory"), TotalAllocatedSmallPoolMemory);
 	OutStats.Add(TEXT("AllocatedOSSmallPoolMemory"), LocalAllocatedOSSmallPoolMemory);
@@ -1365,15 +1373,15 @@ void FMallocBinned2::DumpAllocatorStats(class FOutputDevice& Ar)
 #if UE_MB2_ALLOCATOR_STATS
 
 	const int64  TotalAllocatedSmallPoolMemory           = GetTotalAllocatedSmallPoolMemory();
-	const int64  LocalAllocatedLargePoolMemory           = AllocatedLargePoolMemory.Load(EMemoryOrder::Relaxed);
-	const int64  LocalAllocatedLargePoolMemoryWAlignment = AllocatedLargePoolMemoryWAlignment.Load(EMemoryOrder::Relaxed);
+	const int64  LocalAllocatedLargePoolMemory           = AllocatedLargePoolMemory.load(std::memory_order_relaxed);
+	const int64  LocalAllocatedLargePoolMemoryWAlignment = AllocatedLargePoolMemoryWAlignment.load(std::memory_order_relaxed);
 	const uint64 OSPageAllocatorCachedFreeSize           = CachedOSPageAllocator.GetCachedFreeTotal();
 
 	Ar.Logf(TEXT("FMallocBinned2 Mem report"));
 	Ar.Logf(TEXT("Constants.BinnedPageSize = %d"), int32(PageSize));
 	Ar.Logf(TEXT("Constants.BinnedAllocationGranularity = %d"), int32(OsAllocationGranularity));
 	Ar.Logf(TEXT("Small Pool Allocations: %fmb  (including bin size padding)"), ((double)TotalAllocatedSmallPoolMemory) / (1024.0f * 1024.0f));
-	Ar.Logf(TEXT("Small Pool OS Allocated: %fmb"), ((double)AllocatedOSSmallPoolMemory) / (1024.0f * 1024.0f));
+	Ar.Logf(TEXT("Small Pool OS Allocated: %fmb"), ((double)AllocatedOSSmallPoolMemory.load(std::memory_order_relaxed)) / (1024.0f * 1024.0f));
 	Ar.Logf(TEXT("Large Pool Requested Allocations: %fmb"), ((double)AllocatedLargePoolMemory) / (1024.0f * 1024.0f));
 	Ar.Logf(TEXT("Large Pool OS Allocated: %fmb"), ((double)AllocatedLargePoolMemoryWAlignment) / (1024.0f * 1024.0f));
 	Ar.Logf(TEXT("Requested Allocations: %fmb"), ((double)LocalAllocatedLargePoolMemory) / (1024.0f * 1024.0f));
diff --git a/Engine/Source/Runtime/Core/Private/HAL/MallocBinned3.cpp b/Engine/Source/Runtime/Core/Private/HAL/MallocBinned3.cpp
--- a/Engine/Source/Runtime/Core/Private/HAL/MallocBinned3.cpp
+++ b/Engine/Source/Runtime/Core/Private/HAL/MallocBinned3.cpp
@@ -54,7 +54,7 @@ TBinned3CachedOSPageAllocator& GetCachedOSPageAllocator()
 #endif
 
 #if UE_MB3_ALLOCATOR_STATS
-	int64 Binned3AllocatedSmallPoolMemory = 0; // memory that's requested to be allocated by the game
+	std::atomic<int64> Binned3AllocatedSmallPoolMemory(0); // memory that's requested to be allocated by the game
 	int64 Binned3AllocatedOSSmallPoolMemory = 0;
 
 	int64 Binned3AllocatedLargePoolMemory = 0; // memory requests to the OS which don't fit in the small pool
@@ -510,6 +510,7 @@ struct FMallocBinned3::Private
 	static void FreeBundles(FMallocBinned3& Allocator, FBundleNode* BundlesToRecycle, uint32 InBinSize, uint32 InPoolIndex)
 	{
 		FPoolTable& Table = Allocator.SmallPoolTables[InPoolIndex];
+		FScopeLock Lock(&Table.Mutex);
 
 		FBundleNode* Bundle = BundlesToRecycle;
 		while (Bundle)
@@ -787,9 +788,9 @@ FMallocBinned3::FMallocBinned3()
 		SmallPoolTables[Index].UnusedAreaOffsetLow = 0;
 		SmallPoolTables[Index].NumEverUsedBlocks = 0;
 #if UE_M3_ALLOCATOR_PER_BIN_STATS
-		SmallPoolTables[Index].TotalRequestedAllocSize.Store(0);
-		SmallPoolTables[Index].TotalAllocCount.Store(0);
-		SmallPoolTables[Index].TotalFreeCount.Store(0);
+		SmallPoolTables[Index].TotalRequestedAllocSize.store(0, std::memory_order_relaxed);
+		SmallPoolTables[Index].TotalAllocCount.store(0, std::memory_order_relaxed);
+		SmallPoolTables[Index].TotalFreeCount.store(0, std::memory_order_relaxed);
 #endif
 
 		const int64 TotalNumberOfBlocks = UE_MB3_MAX_MEMORY_PER_POOL_SIZE / (SizeTable[Index].NumMemoryPagesPerBlock * OsAllocationGranularity);
@@ -965,11 +966,11 @@ void* FMallocBinned3::MallocExternal(SIZE_T Size, uint32 Alignment)
 			}
 		}
 
-		FScopeLock Lock(&Mutex);
-
 		// Allocate from small object pool.
 		FPoolTable& Table = SmallPoolTables[PoolIndex];
 
+		FScopeLock Lock(&Table.Mutex);
+
 		uint32 BlockIndex = MAX_uint32;
 		FPoolInfoSmall* Pool = GetFrontPool(Table, PoolIndex, BlockIndex);
 		if (!Pool)
@@ -996,7 +997,7 @@ void* FMallocBinned3::MallocExternal(SIZE_T Size, uint32 Alignment)
 		void* Result = Pool->AllocateBin(BlockPtr, Table.BinSize);
 #if UE_MB3_ALLOCATOR_STATS
 		Table.HeadEndAlloc(Size);
-		Binned3AllocatedSmallPoolMemory += PoolIndexToBinSize(PoolIndex);
+		Binned3AllocatedSmallPoolMemory.fetch_add(PoolIndexToBinSize(PoolIndex), std::memory_order_relaxed);
 #endif
 		if (GBinned3AllocExtra)
 		{
@@ -1186,14 +1187,13 @@ void FMallocBinned3::FreeExternal(void* Ptr)
 		if (BundlesToRecycle)
 		{
 			BundlesToRecycle->NextBundle = nullptr;
-			FScopeLock Lock(&Mutex);
 			Private::FreeBundles(*this, BundlesToRecycle, BinSize, PoolIndex);
 #if UE_MB3_ALLOCATOR_STATS
 			if (!Lists)
 			{
 				SmallPoolTables[PoolIndex].HeadEndFree();
 				// lists track their own stat track them instead in the global stat if we don't have lists
-				Binned3AllocatedSmallPoolMemory -= ((int64)(BinSize));
+				Binned3AllocatedSmallPoolMemory.fetch_sub((int64)BinSize, std::memory_order_relaxed);
 			}
 #endif
 		}
@@ -1264,8 +1264,13 @@ bool FMallocBinned3::GetAllocationSizeExternal(void* Ptr, SIZE_T& SizeOut)
 	{
 		return false;
 	}
-	FScopeLock Lock(&Mutex);
-	FPoolInfoLarge* Pool = Private::FindPoolInfo(*this, Ptr);
+
+	FPoolInfoLarge* Pool;
+	{
+		FScopeLock Lock(&Mutex);
+		Pool = Private::FindPoolInfo(*this, Ptr);
+	}
+
 	if (!Pool)
 	{
 		UE_LOG(LogMemory, Fatal, TEXT("FMallocBinned3 Attempt to GetAllocationSizeExternal an unrecognized pointer %p"), Ptr);
@@ -1374,7 +1379,7 @@ int64 FMallocBinned3::GetTotalAllocatedSmallPoolMemory() const
 		FreeBlockAllocatedMemory += ConsolidatedMemory.load(std::memory_order_relaxed);
 	}
 
-	return Binned3AllocatedSmallPoolMemory + FreeBlockAllocatedMemory;
+	return Binned3AllocatedSmallPoolMemory.load(std::memory_order_relaxed) + FreeBlockAllocatedMemory;
 }
 #endif
 
@@ -1457,14 +1462,14 @@ void FMallocBinned3::DumpAllocatorStats(class FOutputDevice& Ar)
 		const uint32 FullBlocks = CommittedBlocks - PartialBlocks;
 		const int64 ComittedVM = VM - (SmallPoolTables[PoolIndex].NumEverUsedBlocks - CommittedBlocks) * SmallPoolTables[PoolIndex].NumMemoryPagesPerBlock * OsAllocationGranularity;
 
-		const int64 AveSize = SmallPoolTables[PoolIndex].TotalAllocCount.Load() ? SmallPoolTables[PoolIndex].TotalRequestedAllocSize.Load() / SmallPoolTables[PoolIndex].TotalAllocCount.Load() : 0;
-		const int64 EstPadWaste = ((SmallPoolTables[PoolIndex].TotalAllocCount.Load() - SmallPoolTables[PoolIndex].TotalFreeCount.Load()) * (PoolIndexToBinSize(PoolIndex) - AveSize));
+		const int64 AveSize = SmallPoolTables[PoolIndex].TotalAllocCount.load(std::memory_order_relaxed) ? SmallPoolTables[PoolIndex].TotalRequestedAllocSize.load(std::memory_order_relaxed) / SmallPoolTables[PoolIndex].TotalAllocCount.load(std::memory_order_relaxed) : 0;
+		const int64 EstPadWaste = ((SmallPoolTables[PoolIndex].TotalAllocCount.load(std::memory_order_relaxed) - SmallPoolTables[PoolIndex].TotalFreeCount.load(std::memory_order_relaxed)) * (PoolIndexToBinSize(PoolIndex) - AveSize));
 
 		Ar.Logf(TEXT("Pool %2d   Size %6d   Allocs %8lld  Frees %8lld  AveAllocSize %6d  EstPadWaste %4dKB  UsedVM %3dMB  CommittedVM %3dMB  HighSlabs %6d  CommittedSlabs %6d  FullSlabs %6d  PartialSlabs  %6d"), 
 			PoolIndex,
 			PoolIndexToBinSize(PoolIndex),
-			SmallPoolTables[PoolIndex].TotalAllocCount.Load(),
-			SmallPoolTables[PoolIndex].TotalFreeCount.Load(),
+			SmallPoolTables[PoolIndex].TotalAllocCount.load(std::memory_order_relaxed),
+			SmallPoolTables[PoolIndex].TotalFreeCount.load(std::memory_order_relaxed),
 			AveSize,
 			EstPadWaste / 1024,
 			VM / (1024 * 1024),
diff --git a/Engine/Source/Runtime/Core/Public/HAL/MallocBinned2.h b/Engine/Source/Runtime/Core/Public/HAL/MallocBinned2.h
--- a/Engine/Source/Runtime/Core/Public/HAL/MallocBinned2.h
+++ b/Engine/Source/Runtime/Core/Public/HAL/MallocBinned2.h
@@ -66,16 +66,6 @@ struct FGenericMemoryStats;
 #define UE_MB2_ALLOCATOR_STATS				UE_MBC_ALLOCATOR_STATS
 #define UE_MB2_ALLOCATOR_STATS_VALIDATION	(UE_MB2_ALLOCATOR_STATS && 0)
 
-
-#if UE_MB2_ALLOCATOR_STATS
-	//////////////////////////////////////////////////////////////////////////
-	// the following don't need a critical section because they are covered by the critical section called Mutex
-	extern TAtomic<int64> AllocatedSmallPoolMemory; // memory that's requested to be allocated by the game
-	extern TAtomic<int64> AllocatedOSSmallPoolMemory;
-	extern TAtomic<int64> AllocatedLargePoolMemory; // memory requests to the OS which don't fit in the small pool
-	extern TAtomic<int64> AllocatedLargePoolMemoryWAlignment; // when we allocate at OS level we need to align to a size
-#endif
-
 #if UE_MB2_ALLOCATOR_STATS_VALIDATION
 #	include "Misc/ScopeLock.h"
 
@@ -192,6 +182,8 @@ class FMallocBinned2 : public TMallocBinnedCommon<FMallocBinned2, UE_MB2_MINIMUM
 		FPoolList ExhaustedPools;
 		uint32    BinSize = 0;
 
+		FCriticalSection Mutex;
+
 		FPoolTable() = default;
 	};
 
@@ -544,7 +536,6 @@ public:
 	}
 
 	void FreeBundles(FBundleNode* Bundles, uint32 PoolIndex);
-	FCriticalSection& GetMutex() { return Mutex; }
 };
 
 #define UE_MB2_INLINE (1)
diff --git a/Engine/Source/Runtime/Core/Public/HAL/MallocBinned3.h b/Engine/Source/Runtime/Core/Public/HAL/MallocBinned3.h
--- a/Engine/Source/Runtime/Core/Public/HAL/MallocBinned3.h
+++ b/Engine/Source/Runtime/Core/Public/HAL/MallocBinned3.h
@@ -184,21 +184,23 @@ class CORE_API FMallocBinned3 : public TMallocBinnedCommon<FMallocBinned3, UE_MB
 
 		uint64 UnusedAreaOffsetLow;			// High watermark for allocated vm memory for this pool
 
+		FCriticalSection Mutex;
+
 #if UE_M3_ALLOCATOR_PER_BIN_STATS
 		// these are "head end" stats, above the TLS cache
-		TAtomic<int64> TotalRequestedAllocSize;
-		TAtomic<int64> TotalAllocCount;
-		TAtomic<int64> TotalFreeCount;
+		std::atomic<int64> TotalRequestedAllocSize;
+		std::atomic<int64> TotalAllocCount;
+		std::atomic<int64> TotalFreeCount;
 
 		FORCEINLINE void HeadEndAlloc(SIZE_T Size)
 		{
 			check(Size >= 0 && Size <= BinSize);
-			TotalRequestedAllocSize += Size;
-			TotalAllocCount++;
+			TotalRequestedAllocSize.fetch_add(Size, std::memory_order_relaxed);
+			TotalAllocCount.fetch_add(1, std::memory_order_relaxed);
 		}
 		FORCEINLINE void HeadEndFree()
 		{
-			TotalFreeCount++;
+			TotalFreeCount.fetch_add(1, std::memory_order_relaxed);
 		}
 #else
 		FORCEINLINE void HeadEndAlloc(SIZE_T Size)
@@ -518,7 +520,6 @@ public:
 	}
 
 	void FreeBundles(FBundleNode* Bundles, uint32 PoolIndex);
-	FCriticalSection& GetMutex() { return Mutex; }
 
 	void Commit(uint32 InPoolIndex, void *Ptr, SIZE_T Size);
 	void Decommit(uint32 InPoolIndex, void *Ptr, SIZE_T Size);
diff --git a/Engine/Source/Runtime/Core/Public/HAL/MallocBinnedCommonUtils.h b/Engine/Source/Runtime/Core/Public/HAL/MallocBinnedCommonUtils.h
--- a/Engine/Source/Runtime/Core/Public/HAL/MallocBinnedCommonUtils.h
+++ b/Engine/Source/Runtime/Core/Public/HAL/MallocBinnedCommonUtils.h
@@ -40,22 +40,9 @@ public:
 				QUICK_SCOPE_CYCLE_COUNTER(STAT_FMallocBinnedCommonUtils_FlushCurrentThreadCache);
 
 				const double StartTimeInner = FPlatformTime::Seconds();
+				TrimThreadFreeBlockLists(Allocator, Lists);
+				const double WaitForMutexAndTrimTime = FPlatformTime::Seconds() - StartTimeInner;
 
-				double WaitForMutexTime = 0.0f;
-				double WaitForMutexAndTrimTime = 0.0f;
-
-				{
-					FScopeLock Lock(&Allocator.GetMutex());
-					WaitForMutexTime = FPlatformTime::Seconds() - StartTimeInner;
-					TrimThreadFreeBlockLists(Allocator, Lists);
-					WaitForMutexAndTrimTime = FPlatformTime::Seconds() - StartTimeInner;
-				}
-
-				// These logs must happen outside the above mutex to avoid deadlocks
-				if (WaitForMutexTime > GMallocBinnedFlushThreadCacheMaxWaitTime)
-				{
-					UE_LOG(LogMemory, Warning, TEXT("FMalloc%s took %6.2fms to wait for mutex for trim."), Allocator.GetDescriptiveName(), WaitForMutexTime * 1000.0f);
-				}
 				if (WaitForMutexAndTrimTime > GMallocBinnedFlushThreadCacheMaxWaitTime)
 				{
 					UE_LOG(LogMemory, Warning, TEXT("FMalloc%s took %6.2fms to wait for mutex AND trim."), Allocator.GetDescriptiveName(), WaitForMutexAndTrimTime * 1000.0f);
@@ -76,7 +63,6 @@ public:
 		// Skip on desktop as we may have too many threads and this could cause some hitches.
 		if (!PLATFORM_DESKTOP && GMallocBinnedFlushRegisteredThreadCachesOnOneThread != 0)
 		{
-			FScopeLock Lock(&Allocator.GetMutex());
 			FScopeLock FreeBlockLock(&AllocType::GetFreeBlockListsRegistrationMutex());
 			for (typename AllocType::FPerThreadFreeBlockLists* BlockList : AllocType::GetRegisteredFreeBlockLists())
 			{
