[AutoRTFM] Return values from AutoRTFM::Open()

Cleans up a few trivial places where we can avoid the lambda-ref assignment.

diff --git a/Engine/Source/Runtime/Core/Public/HAL/FMemory.inl b/Engine/Source/Runtime/Core/Public/HAL/FMemory.inl
--- a/Engine/Source/Runtime/Core/Public/HAL/FMemory.inl
+++ b/Engine/Source/Runtime/Core/Public/HAL/FMemory.inl
@@ -16,29 +16,29 @@ struct FScopedMallocTimer;
 
 FMEMORY_INLINE_FUNCTION_DECORATOR void* FMemory::Malloc(SIZE_T Count, uint32 Alignment)
 {
-	void* Ptr = nullptr; // Silence bogus static analysis warnings.
-
 	// AutoRTFM: For non-transactional code, all of these calls optimize away and the
 	// behavior is the same as it always has been.
 	// For transactional code, we call the allocator in the 'open' as an optimization, so that
 	// we don't end up keeping track of the writes to the allocator's internal data structures.
 	// This is because allocators are already transactional - malloc can be rolled back by
 	// calling free.
-	UE_AUTORTFM_OPEN2
+	void* Ptr = AutoRTFM::Open([Count, Alignment]
 	{
+		void* Alloc = nullptr;
 		if (!FMEMORY_INLINE_GMalloc)
 		{
-			Ptr = MallocExternal(Count, Alignment);
+			Alloc = MallocExternal(Count, Alignment);
 		}
 		else
 		{
 			DoGamethreadHook(0);
 			FScopedMallocTimer Timer(0);
-			Ptr = FMEMORY_INLINE_GMalloc->Malloc(Count, Alignment);
+			Alloc = FMEMORY_INLINE_GMalloc->Malloc(Count, Alignment);
 		}
 		// optional tracking of every allocation
-		LLM_IF_ENABLED(FLowLevelMemTracker::Get().OnLowLevelAlloc(ELLMTracker::Default, Ptr, Count, ELLMTag::Untagged, ELLMAllocType::FMalloc));
-	};
+		LLM_IF_ENABLED(FLowLevelMemTracker::Get().OnLowLevelAlloc(ELLMTracker::Default, Alloc, Count, ELLMTag::Untagged, ELLMAllocType::FMalloc));
+		return Alloc;
+	});
 
 	// AutoRTFM: This is a no-op for non-transactional code.
 	// For transactional code, this defers a call to Free if the transaction aborts,
diff --git a/Engine/Source/Runtime/CoreUObject/Public/UObject/Stack.h b/Engine/Source/Runtime/CoreUObject/Public/UObject/Stack.h
--- a/Engine/Source/Runtime/CoreUObject/Public/UObject/Stack.h
+++ b/Engine/Source/Runtime/CoreUObject/Public/UObject/Stack.h
@@ -35,9 +35,9 @@ FORCEINLINE void* UeVstackAllocHelper(FVirtualStackAllocator* Allocator, size_t
 	{
 		// in transactional code we redirect the virtual stack allocator to FMemory::Malloc
 		// until we properly implement a transaction-friendly stack allocator
-		AutoRTFM::Open([&Result, Size, Align]()
+		Result = AutoRTFM::Open([Size, Align]
 		{
-			Result = FMemory::Malloc(Size, Align);
+			return FMemory::Malloc(Size, Align);
 		});
 
 		// for these 'stack' allocations we call Free on both commit and abort, since
diff --git a/Engine/Source/Runtime/Engine/Private/PhysicsEngine/PhysDrawing.cpp b/Engine/Source/Runtime/Engine/Private/PhysicsEngine/PhysDrawing.cpp
--- a/Engine/Source/Runtime/Engine/Private/PhysicsEngine/PhysDrawing.cpp
+++ b/Engine/Source/Runtime/Engine/Private/PhysicsEngine/PhysDrawing.cpp
@@ -1238,11 +1238,9 @@ void FKAggregateGeom::FreeRenderInfo()
 {
 	// See if we have rendering resources to free. We do this in the open (non-transactionally) because of
 	// the threaded nature of this operation.
-	FKConvexGeomRenderInfo* RenderInfo = nullptr;
-
-	AutoRTFM::Open([&]
+	FKConvexGeomRenderInfo* RenderInfo = AutoRTFM::Open([&]
 		{
-			RenderInfo = RenderInfoPtr.exchange(nullptr, std::memory_order_acq_rel);
+			return RenderInfoPtr.exchange(nullptr, std::memory_order_acq_rel);
 		});
 
 	if (RenderInfo)
diff --git a/Engine/Source/Runtime/Engine/Private/TraceFilter.cpp b/Engine/Source/Runtime/Engine/Private/TraceFilter.cpp
--- a/Engine/Source/Runtime/Engine/Private/TraceFilter.cpp
+++ b/Engine/Source/Runtime/Engine/Private/TraceFilter.cpp
@@ -277,9 +277,7 @@ FAutoConsoleCommand FlushFilterStateCommand(TEXT("TraceFilter.FlushState"), TEXT
 template<bool bForceThreadSafe>
 bool FTraceFilter::IsObjectTraceable(const UObject* InObject)
 {
-	bool bResult = false;
-
-	AutoRTFM::Open([&]
+	return AutoRTFM::Open([&]
 		{
 			if constexpr (!bForceThreadSafe)
 			{
@@ -287,10 +285,8 @@ bool FTraceFilter::IsObjectTraceable(const UObject* InObject)
 			}
 
 			// Object not found in the AnnotationMap means that it is at the default value, which is bIsTraceable == true
-			bResult = GObjectFilterAnnotations.GetAnnotationMap().Find(InObject) == nullptr;
+			return GObjectFilterAnnotations.GetAnnotationMap().Find(InObject) == nullptr;
 		});
-
-	return bResult;
 }
 
 template bool FTraceFilter::IsObjectTraceable</*bForceThreadSafe = */ false>(const UObject* InObject);
