VR Template | Add visionOS specific .ini
-Added support for reading VisionOS***.ini.  IOS is a parent platform, so it's INIs are read as well.
-Added VisionOSEngine.ini to vr template which disables MobileMultiview and InstancedStereoRendering, which are not currently supported.
-Fixed build break where notification manager is used, but was not included on visionos.  Perhaps it wasn't suppoted in an earlier version, but the types in quesiton exist now so I just removed the ifdef removal on VisionOS.

diff --git a/Engine/Source/Developer/IOS/IOSTargetPlatformSettings/Private/IOSTargetPlatformSettings.cpp b/Engine/Source/Developer/IOS/IOSTargetPlatformSettings/Private/IOSTargetPlatformSettings.cpp
--- a/Engine/Source/Developer/IOS/IOSTargetPlatformSettings/Private/IOSTargetPlatformSettings.cpp
+++ b/Engine/Source/Developer/IOS/IOSTargetPlatformSettings/Private/IOSTargetPlatformSettings.cpp
@@ -20,7 +20,7 @@
 
 FIOSTargetPlatformSettings::FIOSTargetPlatformSettings(bool bInIsTVOS, bool bInIsVisionOS)
 	// override the ini name up in the base classes, which will go into the FTargetPlatformInfo
-	: TTargetPlatformSettingsBase()
+	: TTargetPlatformSettingsBase(nullptr, bInIsVisionOS ? TEXT("VisionOS") : nullptr)
 	, bIsTVOS(bInIsTVOS)
 	, bIsVisionOS(bInIsVisionOS)
 	, MobileShadingPath(0)
diff --git a/Engine/Source/Runtime/Core/Private/IOS/IOSPlatformMisc.cpp b/Engine/Source/Runtime/Core/Private/IOS/IOSPlatformMisc.cpp
--- a/Engine/Source/Runtime/Core/Private/IOS/IOSPlatformMisc.cpp
+++ b/Engine/Source/Runtime/Core/Private/IOS/IOSPlatformMisc.cpp
@@ -48,9 +48,8 @@
 #if UE_WITH_STORE_KIT
 #import <StoreKit/StoreKit.h>
 #endif
-#if !PLATFORM_VISIONOS
+
 #import <UserNotifications/UserNotifications.h>
-#endif
 
 #include <sys/sysctl.h> // sysctlbyname
 
@@ -1235,7 +1234,7 @@ void FIOSPlatformMisc::RegisterForRemoteNotifications()
 	}
 
     dispatch_async(dispatch_get_main_queue(), ^{
-#if !PLATFORM_TVOS && !PLATFORM_VISIONOS && NOTIFICATIONS_ENABLED
+#if !PLATFORM_TVOS && NOTIFICATIONS_ENABLED
 		UNUserNotificationCenter *Center = [UNUserNotificationCenter currentNotificationCenter];
 		[Center requestAuthorizationWithOptions:(UNAuthorizationOptionBadge | UNAuthorizationOptionSound | UNAuthorizationOptionAlert)
 							  completionHandler:^(BOOL granted, NSError * _Nullable error) {
@@ -1270,7 +1269,7 @@ bool FIOSPlatformMisc::IsRegisteredForRemoteNotifications()
 
 bool FIOSPlatformMisc::IsAllowedRemoteNotifications()
 {
-#if !PLATFORM_TVOS && !PLATFORM_VISIONOS && NOTIFICATIONS_ENABLED
+#if !PLATFORM_TVOS && NOTIFICATIONS_ENABLED
 	checkf(false, TEXT("For min iOS version >= 10 use FIOSLocalNotificationService::CheckAllowedNotifications."));
 	return true;
 #else
diff --git a/Engine/Source/Runtime/Core/Public/IOS/IOSPlatformProperties.h b/Engine/Source/Runtime/Core/Public/IOS/IOSPlatformProperties.h
--- a/Engine/Source/Runtime/Core/Public/IOS/IOSPlatformProperties.h
+++ b/Engine/Source/Runtime/Core/Public/IOS/IOSPlatformProperties.h
@@ -125,9 +125,25 @@ struct FTVOSPlatformProperties : public FIOSPlatformProperties
 	}
 };
 
-#ifdef PROPERTY_HEADER_SHOULD_DEFINE_TYPE
+struct FVisionOSPlatformProperties : public FIOSPlatformProperties
+{
+	static FORCEINLINE const char* PlatformName()
+	{
+		return "IOS";
+	}
 
+	static FORCEINLINE const char* IniPlatformName()
+	{
+		return "VisionOS";
+	}
+};
 
+#ifdef PROPERTY_HEADER_SHOULD_DEFINE_TYPE
+
+#if PLATFORM_VISIONOS
+typedef FVisionOSPlatformProperties FPlatformProperties;
+#else
 typedef FIOSPlatformProperties FPlatformProperties;
+#endif
 
 #endif
diff --git a/Templates/TP_VirtualRealityBP/Config/VisionOS/VisionOSEngine.ini b/Templates/TP_VirtualRealityBP/Config/VisionOS/VisionOSEngine.ini
new file mode 100644
--- /dev/null
+++ b/Templates/TP_VirtualRealityBP/Config/VisionOS/VisionOSEngine.ini
@@ -0,0 +1,4 @@
+
+[/Script/Engine.RendererSettings]
+vr.InstancedStereo=False
+vr.MobileMultiView=False
\ No newline at end of file
