Added GetRangePct function to TInterval

diff --git a/Engine/Source/Runtime/Core/Public/Math/Interval.h b/Engine/Source/Runtime/Core/Public/Math/Interval.h
--- a/Engine/Source/Runtime/Core/Public/Math/Interval.h
+++ b/Engine/Source/Runtime/Core/Public/Math/Interval.h
@@ -216,7 +216,21 @@ public:
 
 		return FMath::Clamp(X, Min, Max);
 	}
+	
+	/**
+	 * Calculate the Percentage of X in the Interval
+	 * @param X the element to calculate the percentage
+	 * @return zero if invalid(Min > Max), percentage otherwise
+	 */
+	ElementType GetRangePct(ElementType X) const
+	{
+		if (!IsValid())
+		{
+			return ElementType();
+		}
 
+		return FMath::GetRangePct(Min, Max, X);
+	}
 public:
 
 	/**
