Fixed a delegate compile error in configurations that do not detect delegate race conditions

Removed the previous workaround of an overly-broad friend declaration.

#rnx

diff --git a/Engine/Source/Runtime/Core/Public/Delegates/DelegateAccessHandler.h b/Engine/Source/Runtime/Core/Public/Delegates/DelegateAccessHandler.h
--- a/Engine/Source/Runtime/Core/Public/Delegates/DelegateAccessHandler.h
+++ b/Engine/Source/Runtime/Core/Public/Delegates/DelegateAccessHandler.h
@@ -40,8 +40,7 @@ struct TWriteLockedDelegateAllocation;
 template<>
 class TDelegateAccessHandlerBase<FNotThreadSafeNotCheckedDelegateMode>
 {
-	template <typename T>
-	friend struct TWriteLockedDelegateAllocation;
+	friend struct TWriteLockedDelegateAllocation<FNotThreadSafeNotCheckedDelegateMode>;
 
 protected:
 	struct FReadAccessScope {};
@@ -193,6 +192,7 @@ private:
 template<>
 class TDelegateAccessHandlerBase<FNotThreadSafeDelegateMode> : public TDelegateAccessHandlerBase<FNotThreadSafeNotCheckedDelegateMode>
 {
+	friend struct TWriteLockedDelegateAllocation<FNotThreadSafeDelegateMode>;
 };
 
 #endif // UE_DETECT_DELEGATES_RACE_CONDITIONS
