Remove a bunch of the long deprecated APIs in AutoRTFM to clean out the header.

diff --git a/Engine/Source/Programs/AutoRTFMTests/Private/Abort.cpp b/Engine/Source/Programs/AutoRTFMTests/Private/Abort.cpp
--- a/Engine/Source/Programs/AutoRTFMTests/Private/Abort.cpp
+++ b/Engine/Source/Programs/AutoRTFMTests/Private/Abort.cpp
@@ -83,7 +83,7 @@ TEST_CASE("Abort.NestedAbortOrder")
 	REQUIRE(2 == Orderer);
 }
 
-TEST_CASE("Abort.TransactionInOpenCommit")
+TEST_CASE("Abort.TransactionInOnCommit")
 {
 	AutoRTFM::ETransactionResult InnerResult;
 
@@ -105,7 +105,7 @@ TEST_CASE("Abort.TransactionInOpenCommit")
 	REQUIRE(AutoRTFM::ETransactionResult::AbortedByTransactInOnCommit == InnerResult);
 }
 
-TEST_CASE("Abort.TransactionInOpenAbort")
+TEST_CASE("Abort.TransactionInOnAbort")
 {
 	AutoRTFM::ETransactionResult Result;
 	AutoRTFM::ETransactionResult InnerResult;
diff --git a/Engine/Source/Programs/AutoRTFMTests/Private/OpenAPI.cpp b/Engine/Source/Programs/AutoRTFMTests/Private/OpenAPI.cpp
--- a/Engine/Source/Programs/AutoRTFMTests/Private/OpenAPI.cpp
+++ b/Engine/Source/Programs/AutoRTFMTests/Private/OpenAPI.cpp
@@ -740,7 +740,7 @@ TEST_CASE("OpenAPI.Footgun2")
 }
 
 #if 0
-TEST_CASE("OpenAPI.StartCloseOpenCommit")
+TEST_CASE("OpenAPI.StartCloseOnCommit")
 {
 	REQUIRE(!AutoRTFM::IsTransactional());
 
diff --git a/Engine/Source/Runtime/Core/Private/AutoRTFM/API.cpp b/Engine/Source/Runtime/Core/Private/AutoRTFM/API.cpp
--- a/Engine/Source/Runtime/Core/Private/AutoRTFM/API.cpp
+++ b/Engine/Source/Runtime/Core/Private/AutoRTFM/API.cpp
@@ -13,9 +13,9 @@ namespace
 {
 	// Move this to a local only and use functions to access this
 #if UE_AUTORTFM_ENABLED_RUNTIME_BY_DEFAULT
-	int GAutoRTFMRuntimeEnabled = AutoRTFM::EAutoRTFMEnabledState::AutoRTFM_Enabled;
+	int GAutoRTFMRuntimeEnabled = AutoRTFM::ForTheRuntime::EAutoRTFMEnabledState::AutoRTFM_Enabled;
 #else
-	int GAutoRTFMRuntimeEnabled = AutoRTFM::EAutoRTFMEnabledState::AutoRTFM_Disabled;
+	int GAutoRTFMRuntimeEnabled = AutoRTFM::ForTheRuntime::EAutoRTFMEnabledState::AutoRTFM_Disabled;
 #endif // UE_AUTORTFM_ENABLED_RUNTIME_BY_DEFAULT
 
 	void UpdateAutoRTFMRuntimeCrashData()
@@ -623,16 +623,6 @@ void ForTheRuntime::PopOnAbortHandlerInternal(const void* Key)
 	}
 }
 
-void OpenCommit(TFunction<void()>&& Work)
-{
-	ForTheRuntime::OnCommitInternal(MoveTemp(Work));
-}
-
-void OpenAbort(TFunction<void()>&& Work)
-{
-	ForTheRuntime::OnAbortInternal(MoveTemp(Work));
-}
-
 extern "C" void autortfm_on_commit(void (*Work)(void*), void* Arg)
 {
 	if (autortfm_is_closed())
diff --git a/Engine/Source/Runtime/Core/Private/HAL/ThreadingBase.cpp b/Engine/Source/Runtime/Core/Private/HAL/ThreadingBase.cpp
--- a/Engine/Source/Runtime/Core/Private/HAL/ThreadingBase.cpp
+++ b/Engine/Source/Runtime/Core/Private/HAL/ThreadingBase.cpp
@@ -1456,7 +1456,6 @@ FTlsAutoCleanup* FThreadSingletonInitializer::Get( TFunctionRef<FTlsAutoCleanup*
 		ThreadSingleton = (FTlsAutoCleanup*)FPlatformTLS::GetTlsValue( TlsSlot );
 		if( !ThreadSingleton )
 		{
-			// these are generally left open and only get cleaned up on thread exit so avoiding dealing with an OPENABORT here to clean this up
 			ThreadSingleton = CreateInstance();
 			ThreadSingleton->Register();
 			FPlatformTLS::SetTlsValue( TlsSlot, ThreadSingleton );
diff --git a/Engine/Source/Runtime/Core/Public/AutoRTFM/AutoRTFM.h b/Engine/Source/Runtime/Core/Public/AutoRTFM/AutoRTFM.h
--- a/Engine/Source/Runtime/Core/Public/AutoRTFM/AutoRTFM.h
+++ b/Engine/Source/Runtime/Core/Public/AutoRTFM/AutoRTFM.h
@@ -109,8 +109,6 @@ typedef enum
 	autortfm_aborted_by_transact_in_on_commit,
 	autortfm_aborted_by_transact_in_on_abort,
 	autortfm_aborted_by_cascade,
-	autortfm_aborted_by_transact_in_open_commit [[deprecated("Use autortfm_aborted_by_transact_in_on_commit instead.")]] = autortfm_aborted_by_transact_in_on_commit,
-	autortfm_aborted_by_transact_in_open_abort [[deprecated("Use autortfm_aborted_by_transact_in_on_abort instead.")]] = autortfm_aborted_by_transact_in_on_abort,
 } autortfm_result;
 
 // This must match AutoRTFM::EContextStatus.
@@ -290,12 +288,6 @@ static UE_AUTORTFM_FORCEINLINE void autortfm_on_commit(void (*work)(void* arg),
 }
 #endif
 
-[[deprecated("Use autortfm_on_commit instead.")]]
-static UE_AUTORTFM_FORCEINLINE void autortfm_open_commit(void (*work)(void* arg), void* arg)
-{
-	autortfm_on_commit(work, arg);
-}
-
 #if UE_AUTORTFM_ENABLED
 UE_AUTORTFM_API void autortfm_on_abort(void (*work)(void* arg), void* arg);
 #else
@@ -323,12 +315,6 @@ static UE_AUTORTFM_FORCEINLINE void autortfm_pop_on_abort_handler(const void* ke
 }
 #endif
 
-[[deprecated("Use autortfm_on_abort instead.")]]
-static UE_AUTORTFM_FORCEINLINE void autortfm_open_abort(void (*work)(void* arg), void* arg)
-{
-	autortfm_on_abort(work, arg);
-}
-
 #if UE_AUTORTFM_ENABLED
 UE_AUTORTFM_API void* autortfm_did_allocate(void* ptr, size_t size);
 #else
@@ -398,18 +384,12 @@ enum class ETransactionResult
 	// The transaction aborted because in a call to OnAbort, a new transaction nest was attempted which is not allowed.
 	AbortedByTransactInOnAbort = autortfm_aborted_by_transact_in_on_abort,
 
-	// Deprecated use AbortedByTransactInOnCommit instead.
-	AbortedByTransactInOpenCommit [[deprecated("Use AbortedByTransactInOnCommit instead.")]] = autortfm_aborted_by_transact_in_on_commit,
-
-	// Deprecated use AbortedByTransactInOnAbort instead.
-	AbortedByTransactInOpenAbort [[deprecated("Use AbortedByTransactInOnAbort instead.")]] = autortfm_aborted_by_transact_in_on_abort,
-
 	// The transaction aborted because of an explicit call to CascadingAbortTransaction.
 	AbortedByCascade = autortfm_aborted_by_cascade
 };
 
 // The context status shows what state the AutoRTFM context is currently in. 
-enum class EContextStatus
+enum class EContextStatus : uint8_t
 {
 	// An Idle status means we are not in transactional code.
 	Idle = autortfm_status_idle,
@@ -697,24 +677,6 @@ template<typename TFunctor> static UE_AUTORTFM_FORCEINLINE void PushOnAbortHandl
 static UE_AUTORTFM_FORCEINLINE void PopOnAbortHandler(const void* Key) {}
 #endif
 
-#if UE_AUTORTFM
-[[deprecated("Use OnCommit instead.")]]
-UE_AUTORTFM_API void OpenCommit(TFunction<void()>&& Work);
-#else
-template<typename TFunctor>
-[[deprecated("Use OnCommit instead.")]]
-static UE_AUTORTFM_FORCEINLINE void OpenCommit(const TFunctor& Work) { Work(); }
-#endif
-
-#if UE_AUTORTFM
-[[deprecated("Use OnAbort instead.")]]
-UE_AUTORTFM_API void OpenAbort(TFunction<void()>&& Work);
-#else
-template<typename TFunctor>
-[[deprecated("Use OnAbort instead.")]]
-static UE_AUTORTFM_FORCEINLINE void OpenAbort(const TFunctor& Work) {}
-#endif
-
 // Inform the runtime that we have performed a new object allocation. It's only
 // necessary to call this inside of custom malloc implementations. As an
 // optimization, you can choose to then only have your malloc return the pointer
@@ -915,65 +877,6 @@ namespace ForTheRuntime
 
 } // namespace ForTheRuntime
 
-enum EAutoRTFMEnabledState
-{
-	AutoRTFM_Disabled = ForTheRuntime::AutoRTFM_Disabled,
-	AutoRTFM_Enabled = ForTheRuntime::AutoRTFM_Enabled,
-	AutoRTFM_ForcedDisabled = ForTheRuntime::AutoRTFM_ForcedDisabled,
-};
-
-// Deprecated use AutoRTFM::ForTheRuntime::SetAutoRTFMRuntime instead.
-[[deprecated("Use AutoRTFM::ForTheRuntime::SetAutoRTFMRuntime instead.")]]
-static UE_AUTORTFM_FORCEINLINE bool SetAutoRTFMRuntime(EAutoRTFMEnabledState bEnabled) { return ForTheRuntime::SetAutoRTFMRuntime(static_cast<ForTheRuntime::EAutoRTFMEnabledState>(bEnabled)); }
-
-// Deprecated use AutoRTFM::ForTheRuntime::IsAutoRTFMRuntimeEnabled instead.
-[[deprecated("Use AutoRTFM::ForTheRuntime::IsAutoRTFMRuntimeEnabled instead.")]]
-static UE_AUTORTFM_FORCEINLINE bool IsAutoRTFMRuntimeEnabled() { return ForTheRuntime::IsAutoRTFMRuntimeEnabled(); }
-
-// Deprecated use AutoRTFM::ForTheRuntime::StartTransaction instead.
-[[deprecated("Use AutoRTFM::ForTheRuntime::StartTransaction instead.")]]
-static UE_AUTORTFM_FORCEINLINE bool StartTransaction() { return ForTheRuntime::StartTransaction(); }
-
-// Deprecated use AutoRTFM::ForTheRuntime::CommitTransaction instead.
-[[deprecated("Use AutoRTFM::ForTheRuntime::CommitTransaction instead.")]]
-static UE_AUTORTFM_FORCEINLINE ETransactionResult CommitTransaction() { return ForTheRuntime::CommitTransaction(); }
-
-// Deprecated use AutoRTFM::ForTheRuntime::ClearTransactionStatus instead.
-[[deprecated("Use AutoRTFM::ForTheRuntime::ClearTransactionStatus instead.")]]
-static UE_AUTORTFM_FORCEINLINE void ClearTransactionStatus() { return ForTheRuntime::ClearTransactionStatus(); }
-
-// Deprecated use AutoRTFM::ForTheRuntime::RegisterOpenFunction instead.
-[[deprecated("Use AutoRTFM::ForTheRuntime::RegisterOpenFunction instead.")]]
-static UE_AUTORTFM_FORCEINLINE void RegisterOpenFunction(void* OriginalFunction, void* NewFunction) { return ForTheRuntime::RegisterOpenFunction(OriginalFunction, NewFunction); }
-
-// Deprecated use AutoRTFM::ForTheRuntime::RecordOpenWrite instead.
-[[deprecated("Use AutoRTFM::ForTheRuntime::RecordOpenWrite instead.")]]
-static UE_AUTORTFM_FORCEINLINE void RecordOpenWrite(void* Ptr, size_t Size) { ForTheRuntime::RecordOpenWrite(Ptr, Size); }
-
-// Deprecated use AutoRTFM::ForTheRuntime::RecordOpenWrite instead.
-template<typename TTYPE> [[deprecated("Use AutoRTFM::ForTheRuntime::RecordOpenWrite instead.")]] static UE_AUTORTFM_FORCEINLINE void RecordOpenWrite(TTYPE* Ptr) { return ForTheRuntime::RecordOpenWrite(Ptr); }
-
-// Deprecated use AutoRTFM::ForTheRuntime::RecordOpenRead instead.
-[[deprecated("Use AutoRTFM::ForTheRuntime::RecordOpenRead instead.")]]
-static UE_AUTORTFM_FORCEINLINE void RecordOpenRead(void const* Ptr, size_t Size) { return ForTheRuntime::RecordOpenRead(Ptr, Size); }
-
-// Deprecated use AutoRTFM::ForTheRuntime::RecordOpenRead instead.
-template<typename TTYPE> [[deprecated("Use AutoRTFM::ForTheRuntime::RecordOpenRead instead.")]] static UE_AUTORTFM_FORCEINLINE void RecordOpenRead(TTYPE* Ptr) { return ForTheRuntime::RecordOpenRead(Ptr); }
-
-// Deprecated use AutoRTFM::ForTheRuntime::WriteMemory instead.
-[[deprecated("Use AutoRTFM::ForTheRuntime::WriteMemory instead.")]]
-static UE_AUTORTFM_FORCEINLINE void WriteMemory(void* D, void const* S, size_t Size) { return ForTheRuntime::WriteMemory(D, S, Size); }
-
-// Deprecated use AutoRTFM::ForTheRuntime::WriteMemory instead.
-template<typename TTYPE> [[deprecated("Use AutoRTFM::ForTheRuntime::WriteMemory instead.")]] static UE_AUTORTFM_FORCEINLINE void WriteMemory(TTYPE* D, TTYPE const* S) { return ForTheRuntime::WriteMemory(D, S); }
-
-// Deprecated use AutoRTFM::ForTheRuntime::WriteMemory instead.
-template<typename TTYPE> [[deprecated("Use AutoRTFM::ForTheRuntime::WriteMemory instead.")]] static UE_AUTORTFM_FORCEINLINE void WriteMemory(TTYPE* D, TTYPE const S) { return ForTheRuntime::WriteMemory(D, S); }
-
-// Deprecated use AutoRTFM::ForTheRuntime::CheckConsistencyAssumingNoRaces instead.
-[[deprecated("Use AutoRTFM::ForTheRuntime::CheckConsistencyAssumingNoRaces instead.")]]
-static UE_AUTORTFM_FORCEINLINE void CheckConsistencyAssumingNoRaces() { return ForTheRuntime::CheckConsistencyAssumingNoRaces(); }
-
 } // namespace AutoRTFM
 
 // Macro-based variants so we completely compile away when not in use, even in debug builds
@@ -1074,12 +977,6 @@ namespace AutoRTFM::Private
 // New version is used like a block: UE_AUTORTFM_TRANSACT2 { ... code ... };
 #define UE_AUTORTFM_TRANSACT2 UE_AUTORTFM_TRANSACT_IMPL2
 
-// Deprecated. Use UE_AUTORTFM_ONABORT instead.
-#define UE_AUTORTFM_OPENABORT(...) UE_AUTORTFM_ONABORT(AutoRTFM::ForTheRuntime::DeprecatedUseOnAbortMacro(); __VA_ARGS__)
-
-// Deprecated. Use UE_AUTORTFM_ONCOMMIT instead.
-#define UE_AUTORTFM_OPENCOMMIT(...) UE_AUTORTFM_ONCOMMIT(AutoRTFM::ForTheRuntime::DeprecatedUseOnCommitMacro(); __VA_ARGS__)
-
 #define UE_AUTORTFM_CONCAT_IMPL(A, B) A ## B
 #define UE_AUTORTFM_CONCAT(A, B) UE_AUTORTFM_CONCAT_IMPL(A, B)
 
