Add ASAN poisoning to FVirtualStackAllocator to allow ASAN to detect accesses to regions of the stack that aren't currently allocated.

diff --git a/Engine/Source/Runtime/Core/Private/Memory/VirtualStackAllocator.cpp b/Engine/Source/Runtime/Core/Private/Memory/VirtualStackAllocator.cpp
--- a/Engine/Source/Runtime/Core/Private/Memory/VirtualStackAllocator.cpp
+++ b/Engine/Source/Runtime/Core/Private/Memory/VirtualStackAllocator.cpp
@@ -40,6 +40,10 @@ FVirtualStackAllocator::FVirtualStackAllocator(size_t RequestedStackSize, EVirtu
 		NextUncommittedPage = VirtualMemory.GetVirtualPointer();
 		NextAllocationStart = NextUncommittedPage;
 		RecentHighWaterMark = NextUncommittedPage;
+
+#if PLATFORM_HAS_ASAN_INCLUDE
+		ASAN_POISON_MEMORY_REGION(NextUncommittedPage, TotalReservationSize);
+#endif
 	}
 }
 
@@ -85,9 +89,12 @@ void* FVirtualStackAllocator::Allocate(size_t Size, size_t Alignment)
 		}
 
 		NextAllocationStart = AllocationEnd;
+
+#if PLATFORM_HAS_ASAN_INCLUDE
+		ASAN_UNPOISON_MEMORY_REGION(AllocationStart, Size);
+#endif
 	}
  
-
     return AllocationStart;
 }
 
diff --git a/Engine/Source/Runtime/Core/Public/Memory/VirtualStackAllocator.h b/Engine/Source/Runtime/Core/Public/Memory/VirtualStackAllocator.h
--- a/Engine/Source/Runtime/Core/Public/Memory/VirtualStackAllocator.h
+++ b/Engine/Source/Runtime/Core/Public/Memory/VirtualStackAllocator.h
@@ -7,6 +7,10 @@
 #include "HAL/PlatformMemory.h"
 #include "Misc/AssertionMacros.h"
 
+#if PLATFORM_HAS_ASAN_INCLUDE
+#include <sanitizer/asan_interface.h>
+#endif
+
 class FVirtualStackAllocator;
 
 /**
@@ -105,6 +109,10 @@ private:
 	{
 		check(RestorePointer <= NextAllocationStart);
 
+#if PLATFORM_HAS_ASAN_INCLUDE
+		ASAN_POISON_MEMORY_REGION(RestorePointer, static_cast<uint8*>(NextAllocationStart) - static_cast<uint8*>(RestorePointer));
+#endif
+
 		NextAllocationStart = RestorePointer;
 		if (UNLIKELY(NextAllocationStart == VirtualMemory.GetVirtualPointer() && DecommitMode != EVirtualStackAllocatorDecommitMode::AllOnDestruction))
 		{
