Move FConfigIni FScopeLock to using FTransactionallySafeScopeLock

diff --git a/Engine/Source/Runtime/Core/Private/Misc/ConfigCacheIni.cpp b/Engine/Source/Runtime/Core/Private/Misc/ConfigCacheIni.cpp
--- a/Engine/Source/Runtime/Core/Private/Misc/ConfigCacheIni.cpp
+++ b/Engine/Source/Runtime/Core/Private/Misc/ConfigCacheIni.cpp
@@ -22,6 +22,7 @@
 #include "Misc/DataDrivenPlatformInfoRegistry.h"
 #include "Misc/StringBuilder.h"
 #include "Misc/PathViews.h"
+#include "Misc/TransactionallySafeScopeLock.h"
 #include "ProfilingDebugging/CpuProfilerTrace.h"
 #include "ProfilingDebugging/AssetMetadataTrace.h"
 #include "Serialization/MemoryReader.h"
@@ -1034,7 +1035,7 @@ FConfigFile& FConfigFile::operator=(FConfigFile&& Other)
 	return *this;
 }
 
-void FConfigFile::Cleanup()
+UE_AUTORTFM_ALWAYS_OPEN void FConfigFile::Cleanup()
 {
 	Empty();
 }
@@ -6245,7 +6246,7 @@ FCriticalSection FConfigCacheIni::ConfigForPlatformLock;
 #endif
 
 TMap<FName, FConfigCacheIni::FPluginInfo*> FConfigCacheIni::RegisteredPlugins;
-FCriticalSection FConfigCacheIni::RegisteredPluginsLock;
+FTransactionallySafeCriticalSection FConfigCacheIni::RegisteredPluginsLock;
 
 FRWLock FConfigFile::ConfigFileMapLock;
 
@@ -6284,7 +6285,7 @@ void FConfigCacheIni::AddPluginToBranches(FName PluginName, FConfigModificationT
 	
 	FPluginInfo* PluginInfo = nullptr;
 	{
-		FScopeLock Lock(&RegisteredPluginsLock);
+		FTransactionallySafeScopeLock Lock(&RegisteredPluginsLock);
 		
 		PluginInfo = RegisteredPlugins.FindRef(PluginName);
 		if (PluginInfo == nullptr)
@@ -6494,7 +6495,7 @@ void FConfigCacheIni::RegisterPlugin(FName PluginName, const FString& PluginDir,
 	Info->Priority = Priority;
 	Info->bIncludePluginNameInBranchName = bIncludePluginNameInBranchName;
 
-	FScopeLock Lock(&RegisteredPluginsLock);
+	FTransactionallySafeScopeLock Lock(&RegisteredPluginsLock);
 	RegisteredPlugins.Add(PluginName, Info);
 }
 
diff --git a/Engine/Source/Runtime/Core/Private/Misc/ConfigContext.cpp b/Engine/Source/Runtime/Core/Private/Misc/ConfigContext.cpp
--- a/Engine/Source/Runtime/Core/Private/Misc/ConfigContext.cpp
+++ b/Engine/Source/Runtime/Core/Private/Misc/ConfigContext.cpp
@@ -83,7 +83,7 @@ void FConfigContext::CachePaths()
 		// are we loading a plugin?
 		if (ConfigSystem != nullptr)
 		{
-			FScopeLock Lock(&FConfigCacheIni::RegisteredPluginsLock);
+			FTransactionallySafeScopeLock Lock(&FConfigCacheIni::RegisteredPluginsLock);
 			
 			FConfigCacheIni::FPluginInfo* PluginInfo = ConfigSystem->RegisteredPlugins.FindRef(ConfigFileTag == NAME_None ? FName(*BaseIniName) : ConfigFileTag);
 			if (PluginInfo != nullptr)
diff --git a/Engine/Source/Runtime/Core/Public/Misc/ConfigCacheIni.h b/Engine/Source/Runtime/Core/Public/Misc/ConfigCacheIni.h
--- a/Engine/Source/Runtime/Core/Public/Misc/ConfigCacheIni.h
+++ b/Engine/Source/Runtime/Core/Public/Misc/ConfigCacheIni.h
@@ -1931,7 +1931,7 @@ private:
 	};
 	static TMap<FName, FPluginInfo*> RegisteredPlugins;
 	TArray<FName> PendingModificationPlugins;
-	static FCriticalSection RegisteredPluginsLock;
+	static FTransactionallySafeCriticalSection RegisteredPluginsLock;
 	
 #if ALLOW_OTHER_PLATFORM_CONFIG
 	static TMap<FName, FConfigCacheIni*> ConfigForPlatform;
