Corrected erroneous ptr offset computation in natvis instrisic example for GDebuggingState

diff --git a/Engine/Source/Runtime/Core/Public/Modules/VisualizerDebuggingState.h b/Engine/Source/Runtime/Core/Public/Modules/VisualizerDebuggingState.h
--- a/Engine/Source/Runtime/Core/Public/Modules/VisualizerDebuggingState.h
+++ b/Engine/Source/Runtime/Core/Public/Modules/VisualizerDebuggingState.h
@@ -43,7 +43,7 @@ enum class EVisualizerDebuggingStateResult : uint8
  * An example visualizer syntax is included below. It is recommended that most natvis files utilizing this method would include an
  *   global intrinsic function at the top for quick access:
  * 
- *   <Intrinsic Name="GetDebugState" Expression="GDebuggingState->Ptrs[(GDebuggingState->GuidString - strstr(GDebuggingState->GuidString, GuidString))/32]">
+ *   <Intrinsic Name="GetDebugState" Expression="GDebuggingState->Ptrs[(strstr(GDebuggingState->GuidString, GuidString) - GDebuggingState->GuidString)/32]">
  *     <Parameter Name="GuidString" Type="char*"/>
  *   </Intrinsic>
  *
