Adding TSparseArray::IndexOfByPredicate()

diff --git a/Engine/Source/Runtime/Core/Public/Containers/SparseArray.h b/Engine/Source/Runtime/Core/Public/Containers/SparseArray.h
--- a/Engine/Source/Runtime/Core/Public/Containers/SparseArray.h
+++ b/Engine/Source/Runtime/Core/Public/Containers/SparseArray.h
@@ -678,6 +678,26 @@ public:
 		// Set up the new head
 		FirstFreeIndex = CurrentHeadIndex;
 	}
+	
+	/**
+	 * Finds an item by predicate.
+	 *
+	 * @param Pred The predicate to match.
+	 * @returns Index to the first matching element, or INDEX_NONE if none is found.
+	 */
+	template <typename Predicate>
+	int32 IndexOfByPredicate(Predicate Pred) const
+	{	
+		for (TConstIterator It = CreateConstIterator(); It; ++It)
+		{
+			if (::Invoke(Pred, *It))
+			{
+				return It.GetIndex();
+			}
+		}
+
+		return INDEX_NONE;
+	}
 
 	/** 
 	 * Helper function to return the amount of memory allocated by this container 
