[AutoRTFM] Make CreateComponentFromTemplate() transaction safe

* Migrate FStaticCustomVersionRegistry::Lock to use FTransactionallySafeRWLock.
* Disable caching of FUnversionedStructSchema when in a closed transaction.
* Ensure that the loader logic of StaticDuplicateObjectEx() is done in the open.
* Allocate PrimitiveSceneIds in the open.

diff --git a/Engine/Source/Runtime/Core/Private/Serialization/CustomVersion.cpp b/Engine/Source/Runtime/Core/Private/Serialization/CustomVersion.cpp
--- a/Engine/Source/Runtime/Core/Private/Serialization/CustomVersion.cpp
+++ b/Engine/Source/Runtime/Core/Private/Serialization/CustomVersion.cpp
@@ -12,6 +12,8 @@
 #include "HAL/CriticalSection.h"
 #include "Misc/AssertionMacros.h"
 #include "Misc/ScopeRWLock.h"
+#include "Misc/TransactionallySafeRWLock.h"
+#include "Misc/TransactionallySafeRWScopeLock.h"
 #include "Serialization/Archive.h"
 #include "Serialization/StructuredArchive.h"
 #include "Serialization/StructuredArchiveAdapters.h"
@@ -94,7 +96,7 @@ struct FStaticCustomVersionRegistry
 	};
 	typedef TMap<FGuid, FPendingRegistration, TInlineSetAllocator<64>> RegistrationQueue;
 
-	FRWLock Lock;
+	FTransactionallySafeRWLock Lock;
 	FCustomVersionContainer Registered;
 	RegistrationQueue Queue;
 
@@ -179,7 +181,7 @@ FCustomVersionContainer FCurrentCustomVersions::GetAll()
 	FStaticCustomVersionRegistry& Registry = FStaticCustomVersionRegistry::Get();
 
 	{
-		FReadScopeLock Scope(Registry.Lock);
+		FTransactionallySafeReadScopeLock Scope(Registry.Lock);
 
 		if (Registry.Queue.Num() == 0)
 		{
@@ -187,7 +189,7 @@ FCustomVersionContainer FCurrentCustomVersions::GetAll()
 		}
 	}
 
-	FWriteScopeLock Scope(Registry.Lock);
+	FTransactionallySafeWriteScopeLock Scope(Registry.Lock);
 	Registry.RegisterQueue();
 	return Registry.Registered;
 }
@@ -196,7 +198,7 @@ TOptional<FCustomVersion> FCurrentCustomVersions::Get(const FGuid& Guid)
 {
 	FStaticCustomVersionRegistry& Registry = FStaticCustomVersionRegistry::Get();
 
-	FReadScopeLock Scope(Registry.Lock);
+	FTransactionallySafeReadScopeLock Scope(Registry.Lock);
 	return Registry.Find(Guid);
 }
 
@@ -208,7 +210,7 @@ TArray<FCustomVersionDifference> FCurrentCustomVersions::Compare(const FCustomVe
 	{
 		FStaticCustomVersionRegistry& Registry = FStaticCustomVersionRegistry::Get();
 
-		FReadScopeLock Scope(Registry.Lock);
+		FTransactionallySafeReadScopeLock Scope(Registry.Lock);
 
 		for (const FCustomVersion& CompareVersion : CompareVersions)
 		{
@@ -238,7 +240,7 @@ void FCurrentCustomVersions::Register(const FGuid& Key, int32 Version, const TCH
 {
 	FStaticCustomVersionRegistry& Registry = FStaticCustomVersionRegistry::Get();
 
-	FWriteScopeLock Scope(Registry.Lock);
+	FTransactionallySafeWriteScopeLock Scope(Registry.Lock);
 	check(Registry.Queue.Find(Key) == nullptr);
 	Registry.Queue.Add(Key, { Version, Name, ValidatorFunc });
 }
@@ -247,7 +249,7 @@ void FCurrentCustomVersions::Unregister(const FGuid& Key)
 {
 	FStaticCustomVersionRegistry& Registry = FStaticCustomVersionRegistry::Get();
 
-	FWriteScopeLock Scope(Registry.Lock);
+	FTransactionallySafeWriteScopeLock Scope(Registry.Lock);
 	Registry.Unregister(Key);
 }
 
@@ -268,7 +270,7 @@ const FCustomVersionContainer& FCustomVersionContainer::GetRegistered()
 	FStaticCustomVersionRegistry& Registry = FStaticCustomVersionRegistry::Get();
 
 	// Even though returning a reference isn't thread-safe, we can still synchronize access to the queue
-	FWriteScopeLock Scope(Registry.Lock);
+	FTransactionallySafeWriteScopeLock Scope(Registry.Lock);
 	Registry.RegisterQueue();
 	return Registry.Registered;
 }
diff --git a/Engine/Source/Runtime/CoreUObject/Private/Serialization/UnversionedPropertySerialization.cpp b/Engine/Source/Runtime/CoreUObject/Private/Serialization/UnversionedPropertySerialization.cpp
--- a/Engine/Source/Runtime/CoreUObject/Private/Serialization/UnversionedPropertySerialization.cpp
+++ b/Engine/Source/Runtime/CoreUObject/Private/Serialization/UnversionedPropertySerialization.cpp
@@ -1,6 +1,7 @@
 // Copyright Epic Games, Inc. All Rights Reserved.
 
 #include "Serialization/UnversionedPropertySerialization.h"
+#include "AutoRTFM/AutoRTFM.h"
 #include "Serialization/UnversionedPropertySerializationTest.h"
 #include "Hash/Blake3.h"
 #include "Interfaces/ITargetPlatform.h"
@@ -452,11 +453,15 @@ const FUnversionedStructSchema& GetOrCreateUnversionedSchema(const UStruct* Stru
 
 	FUnversionedStructSchema* CreatedSchema = FUnversionedStructSchema::Create(Struct, bSkipEditorOnly);
 
-	void** CachedSchemaPtr = reinterpret_cast<void**>(const_cast<FUnversionedStructSchema**>(&GetUnversionedSchema(Struct, bSkipEditorOnly)));
-	if (const FUnversionedStructSchema* ExistingSchema = reinterpret_cast<const FUnversionedStructSchema*>(FPlatformAtomics::InterlockedCompareExchangePointer(CachedSchemaPtr, CreatedSchema, nullptr)))
+	// AutoRTFM cannot undo atomic swaps, so skip the caching if called from a transaction.
+	if (!AutoRTFM::IsClosed())
 	{
-		FUnversionedStructSchema::Delete(CreatedSchema);
-		return *ExistingSchema;
+		void** CachedSchemaPtr = reinterpret_cast<void**>(const_cast<FUnversionedStructSchema**>(&GetUnversionedSchema(Struct, bSkipEditorOnly)));
+		if (const FUnversionedStructSchema* ExistingSchema = reinterpret_cast<const FUnversionedStructSchema*>(FPlatformAtomics::InterlockedCompareExchangePointer(CachedSchemaPtr, CreatedSchema, nullptr)))
+		{
+			FUnversionedStructSchema::Delete(CreatedSchema);
+			return *ExistingSchema;
+		}
 	}
 
 	return *CreatedSchema;
diff --git a/Engine/Source/Runtime/CoreUObject/Private/UObject/UObjectGlobals.cpp b/Engine/Source/Runtime/CoreUObject/Private/UObject/UObjectGlobals.cpp
--- a/Engine/Source/Runtime/CoreUObject/Private/UObject/UObjectGlobals.cpp
+++ b/Engine/Source/Runtime/CoreUObject/Private/UObject/UObjectGlobals.cpp
@@ -3061,11 +3061,16 @@ UObject* StaticDuplicateObjectEx( FObjectDuplicationParameters& Parameters )
 	// Make sure we're not duplicating the AsyncLoading, Async or LoaderImport internal flags, they will prevent the object from being gcd.
 	Parameters.InternalFlagMask &= ~(EInternalObjectFlags::Async | EInternalObjectFlags::LoaderImport | EInternalObjectFlags_AsyncLoading);
 
-	if (!IsAsyncLoading() && Parameters.SourceObject->HasAnyFlags(RF_ClassDefaultObject))
+	// We can't modify the loader from a transaction, so check for async loading and reset the 
+	// loaders in the open.
+	AutoRTFM::Open([&]
 	{
-		// Detach linker for the outer if it already exists, to avoid problems with PostLoad checking the Linker version
-		ResetLoaders(Parameters.DestOuter);
-	}
+		if (!IsAsyncLoading() && Parameters.SourceObject->HasAnyFlags(RF_ClassDefaultObject))
+		{
+			// Detach linker for the outer if it already exists, to avoid problems with PostLoad checking the Linker version
+			ResetLoaders(Parameters.DestOuter);
+		}
+	});
 
 	FObjectInstancingGraph InstanceGraph;
 
diff --git a/Engine/Source/Runtime/Engine/Private/Tests/AutoRTFM/AutoRTFMActorComponentTests.cpp b/Engine/Source/Runtime/Engine/Private/Tests/AutoRTFM/AutoRTFMActorComponentTests.cpp
--- a/Engine/Source/Runtime/Engine/Private/Tests/AutoRTFM/AutoRTFMActorComponentTests.cpp
+++ b/Engine/Source/Runtime/Engine/Private/Tests/AutoRTFM/AutoRTFMActorComponentTests.cpp
@@ -425,6 +425,20 @@ AUTORTFM_ACTOR_COMPONENT_TEST(USkeletalMeshComponent)
 	TEST_CHECK_TRUE(AutoRTFM::ETransactionResult::AbortedByRequest == Result);
 }
 
+
+// Test aborting a call to AAutoRTFMTestActor::CreateComponentFromTemplate().
+// See: SOL-7002
+AUTORTFM_ACTOR_COMPONENT_TEST(CreateComponentFromTemplate)
+{
+	AutoRTFM::ETransactionResult Result = AutoRTFM::Transact([&]
+		{
+			Actor->CreateComponentFromTemplate(Component);
+			AutoRTFM::AbortTransaction();
+		});
+
+	TEST_CHECK_TRUE(AutoRTFM::ETransactionResult::AbortedByRequest == Result);
+}
+
 }  // anonymous namespace
 
 #undef TEST_CHECK_FALSE
diff --git a/Engine/Source/Runtime/Engine/Public/PrimitiveSceneInfoData.h b/Engine/Source/Runtime/Engine/Public/PrimitiveSceneInfoData.h
--- a/Engine/Source/Runtime/Engine/Public/PrimitiveSceneInfoData.h
+++ b/Engine/Source/Runtime/Engine/Public/PrimitiveSceneInfoData.h
@@ -2,6 +2,7 @@
 
 #pragma once 
 
+#include "AutoRTFM/AutoRTFM.h"
 #include "HAL/ThreadSafeCounter.h"
 #include "PrimitiveComponentId.h"
 
@@ -91,10 +92,10 @@ public:
 		: RegistrationSerialNumber(-1)
 		, bAlwaysVisible(false)
 	{
-		PrimitiveSceneId.PrimIDValue = NextPrimitiveId.Increment();
+		PrimitiveSceneId.PrimIDValue = AutoRTFM::Open([] { return NextPrimitiveId.Increment(); });
 	}
 };
 
 #if UE_ENABLE_INCLUDE_ORDER_DEPRECATED_IN_5_5
 #include "RendererInterface.h"
-#endif
\ No newline at end of file
+#endif
