Use FCriticalSection in FAsyncLatentCommand and FAsyncUntilDoneLatentCommand to avoid race condition when Future is assigned

#rnx

diff --git a/Engine/Source/Runtime/Core/Public/Misc/AutomationTest.h b/Engine/Source/Runtime/Core/Public/Misc/AutomationTest.h
--- a/Engine/Source/Runtime/Core/Public/Misc/AutomationTest.h
+++ b/Engine/Source/Runtime/Core/Public/Misc/AutomationTest.h
@@ -2832,12 +2832,12 @@ private:
 					return true;
 				}
 
+				FScopeLock CriticalSection(&ActionCS);
 				bDone = false;
+				StartedRunning = FDateTime::UtcNow();
 				Future = Async(Execution, [this]() {
 					Predicate(FDoneDelegate::CreateRaw(this, &FAsyncUntilDoneLatentCommand::Done));
 				});
-
-				StartedRunning = FDateTime::UtcNow();
 			}
 
 			if (bDone)
@@ -2859,6 +2859,7 @@ private:
 
 		void Done()
 		{
+			FScopeLock CriticalSection(&ActionCS);
 			if (Future.IsValid())
 			{
 				bDone = true;
@@ -2867,6 +2868,7 @@ private:
 
 		void Reset()
 		{
+			FScopeLock CriticalSection(&ActionCS);
 			// Reset the done for the next potential run of this command
 			bDone = false;
 			Future.Reset();
@@ -2883,6 +2885,7 @@ private:
 		FThreadSafeBool bDone;
 		FDateTime StartedRunning;
 		TFuture<void> Future;
+		FCriticalSection ActionCS;
 	};
 
 	class FAsyncLatentCommand : public IAutomationLatentCommand
@@ -2910,13 +2913,13 @@ private:
 					return true;
 				}
 
+				FScopeLock CriticalSection(&ActionCS);
 				bDone = false;
+				StartedRunning = FDateTime::UtcNow();
 				Future = Async(Execution, [this]() {
 					Predicate();
 					Done();
 				});
-
-				StartedRunning = FDateTime::UtcNow();
 			}
 
 			if (bDone)
@@ -2938,6 +2941,7 @@ private:
 
 		void Done()
 		{
+			FScopeLock CriticalSection(&ActionCS);
 			if (Future.IsValid())
 			{
 				bDone = true;
@@ -2946,6 +2950,7 @@ private:
 
 		void Reset()
 		{
+			FScopeLock CriticalSection(&ActionCS);
 			// Reset the done for the next potential run of this command
 			bDone = false;
 			Future.Reset();
@@ -2962,6 +2967,7 @@ private:
 		FThreadSafeBool bDone;
 		FDateTime StartedRunning;
 		TFuture<void> Future;
+		FCriticalSection ActionCS;
 	};
 
 	struct FSpecIt
