Logging warning message when an already registered test name is found

#rnx

diff --git a/Engine/Source/Runtime/Core/Private/Misc/AutomationTest.cpp b/Engine/Source/Runtime/Core/Private/Misc/AutomationTest.cpp
--- a/Engine/Source/Runtime/Core/Private/Misc/AutomationTest.cpp
+++ b/Engine/Source/Runtime/Core/Private/Misc/AutomationTest.cpp
@@ -153,7 +153,11 @@ FAutomationTestBase::FAutomationTestBase( const FString& InName, const bool bInC
 	LLM_SCOPE_BYNAME(TEXT("AutomationTest/Framework"));
 	TestName = InName;
 	// Register the newly created automation test into the automation testing framework
-	FAutomationTestFramework::Get().RegisterAutomationTest( InName, this );
+	const bool bRegistered = FAutomationTestFramework::Get().RegisterAutomationTest( InName, this );
+	if (!bRegistered)
+	{
+		UE_LOG(LogAutomationTest, Warning, TEXT("Failed to register test with the name '%s'. Test with the same name is already registered and will not be overridden."), *InName);
+	}
 }
 
 /** Destructor */
