FPersistentStorageManager - make sure replacement directories have a trailing /

diff --git a/Engine/Source/Runtime/Core/Private/HAL/IPlatformFileManagedStorageWrapper.cpp b/Engine/Source/Runtime/Core/Private/HAL/IPlatformFileManagedStorageWrapper.cpp
--- a/Engine/Source/Runtime/Core/Private/HAL/IPlatformFileManagedStorageWrapper.cpp
+++ b/Engine/Source/Runtime/Core/Private/HAL/IPlatformFileManagedStorageWrapper.cpp
@@ -222,6 +222,14 @@ FPersistentStorageManager::FCategoryInfo FPersistentStorageManager::InitCategori
 	CustomDirectoryReplace.Add(TEXT("[persistent]"), FPaths::ConvertRelativePathToFull(FPaths::ProjectPersistentDownloadDir()));
 	CustomDirectoryReplace.Add(TEXT("[saved]"), FPaths::ConvertRelativePathToFull(FPaths::ProjectSavedDir()));
 
+	for (TPair<FString, FString>& Pair : CustomDirectoryReplace)
+	{
+		if (!Pair.Value.EndsWith(TEXT("/")))
+		{
+			Pair.Value += TEXT("/");
+		}
+	}
+
 	// Takes on the pattern
 	// (Name="CategoryName",QuotaMB=100,Directories=("Dir1","Dir2","Dir3"))
 	TArray<FString> CategoryConfigs;
