Tests: String: Fixed an incorrect test case in System.Core.String.ConstructFromPtrSizeWithSlack, which was imposing too strict requirements.

#rnx

diff --git a/Engine/Source/Runtime/Core/Private/Containers/StringTest.cpp b/Engine/Source/Runtime/Core/Private/Containers/StringTest.cpp
--- a/Engine/Source/Runtime/Core/Private/Containers/StringTest.cpp
+++ b/Engine/Source/Runtime/Core/Private/Containers/StringTest.cpp
@@ -908,7 +908,7 @@ bool FStringConstructFromPtrSizeWithSlackTest::RunTest(const FString& Parameters
 
 		const TArray<TCHAR>& StrArr = Str.GetCharArray();
 		int32 ActualSlack = StrArr.Max() - StrArr.Num();
-		bool  bValidSlack = (ExtraSlack == 0 && *Ptr == TEXT('\0')) ? (ActualSlack == 0) : (ActualSlack >= ExtraSlack);
+		bool  bValidSlack = (ActualSlack >= ExtraSlack);
 		if (StrArr != Expected || !bValidSlack)
 		{
 			AddError(
