CompareMap: Move the implementation from CoreRedirects into Compare and make it fully templated (no reliance on Map.h).
#rnx

diff --git a/Engine/Source/Runtime/Core/Public/Algo/Compare.h b/Engine/Source/Runtime/Core/Public/Algo/Compare.h
--- a/Engine/Source/Runtime/Core/Public/Algo/Compare.h
+++ b/Engine/Source/Runtime/Core/Public/Algo/Compare.h
@@ -6,6 +6,7 @@
 #include "Templates/EqualTo.h"
 #include "Templates/IdentityFunctor.h"
 #include "Templates/Invoke.h"
+#include "Templates/Less.h"
 #include "Templates/UnrealTemplate.h"
 
 namespace Algo::Private
@@ -104,4 +105,183 @@ namespace Algo
 	{
 		return Private::Compare(Forward<InAT>(InputA), Forward<InBT>(InputB), MoveTemp(Projection), MoveTemp(Predicate));
 	}
+
+	/**
+	 * Compares two unique-key maps (e.g. TMap) as if they were sorted arrays of key,value pairs sorted by Key and then
+	 * by Value (almost, see note on sort order). Does not support multiple values per key (e.g. TMultiMap).
+	 * 
+	 * Note on sort order: Maps with a smaller number of elements are considered less than maps with a larger number
+	 * of elements, no matter what keys are present in each map. This drastically improves performance when comparing
+	 * maps of different sizes. This is different than would be exepcted from a lexical compare of strings, but it does
+	 * match the comparison of two numbers represented as a string of digits.
+	 * 
+	 * MapType interface:
+	 *  typename KeyType, typename ValueType,
+	 *  IntType Num(), ValueType* Find(const KeyType&), IteratorType<PairType> begin(), end()
+	 * PairType interface:
+	 *  const KeyType& Key, const ValueType& Value
+	 * 
+	 * @param A Left-Hand-Side map container.
+	 * @param B Right-Hand-Side map container.
+	 * @param KeyLessThan bool(const KeyType& A, const KeyType& B) that returns A < B.
+	 * @param ValueLessThan bool(const ValueType& A, const ValueType& B) that returns A < B.
+	 *
+	 * @return -1 if A < B, 0 if A == B, 1 if A > B
+	 */
+	template <typename MapType, typename KeyLessThanType, typename ValueLessThanType>
+	int CompareMap(const MapType& A, const MapType& B, KeyLessThanType KeyLessThan, ValueLessThanType ValueLessThan)
+	{
+		using KeyType = typename MapType::KeyType;
+		using ValueType = typename MapType::ValueType;
+		if (A.Num() != B.Num())
+		{
+			return A.Num() < B.Num() ? -1 : 1;
+		}
+		if (A.Num() == 0)
+		{
+			return 0;
+		}
+
+		bool bAllKeysOfAAreInB = true;
+		const KeyType* MinKeyWithDifference = nullptr;
+		bool bMinKeyIsLessInA = false;
+		for (const auto& Pair : A) // TPair<KeyType, ValueType>
+		{
+			const ValueType* BValue = B.Find(Pair.Key);
+			int Compare = 0;
+			if (!BValue)
+			{
+				bAllKeysOfAAreInB = false;
+				Compare = -1;
+			}
+			else if (ValueLessThan(Pair.Value, *BValue))
+			{
+				Compare = -1;
+			}
+			else if (ValueLessThan(*BValue, Pair.Value))
+			{
+				Compare = 1;
+			}
+			if (Compare != 0)
+			{
+				if (!MinKeyWithDifference || KeyLessThan(Pair.Key, *MinKeyWithDifference))
+				{
+					MinKeyWithDifference = &Pair.Key;
+					bMinKeyIsLessInA = Compare < 0;
+				}
+			}
+		}
+
+		// The number of keys in A and B is the same (checked above), so if all keys of A are in B, then there are
+		// no additional keys in B that are not in A and we don't need to iterate over B.
+		if (!bAllKeysOfAAreInB)
+		{
+			// If B has additional keys not in A then we need to check each of those not-in-A keys to see if they
+			// are smaller than the MinKeyWithDifference, and if so they become the MinKeyWithDifference.
+			for (const auto& Pair : B) // TPair<KeyType, ValueType>
+			{
+				if (!A.Contains(Pair.Key))
+				{
+					if (!MinKeyWithDifference || KeyLessThan(Pair.Key, *MinKeyWithDifference))
+					{
+						MinKeyWithDifference = &Pair.Key;
+						bMinKeyIsLessInA = false;
+					}
+				}
+			}
+		}
+		if (MinKeyWithDifference)
+		{
+			return bMinKeyIsLessInA ? -1 : 1;
+		}
+		return 0;
+	}
+
+	template <typename MapType>
+	int CompareMap(const MapType& A, const MapType& B)
+	{
+		return CompareMap(A, B, TLess<>(), TLess<>());
+	}
+
+	template <typename MapType, typename KeyLessThanType>
+	int CompareMap(const MapType& A, const MapType& B, KeyLessThanType KeyLessThan)
+	{
+		return CompareMap(A, B, MoveTempIfPossible(KeyLessThan), TLess<>());
+	}
+
+	/**
+	 * Compares two sets (e.g. TSet) as if they were sorted arrays of keys (almost, see note on sort order).
+	 *
+	 * Note on sort order: Maps with a smaller number of elements are considered less than maps with a larger number
+	 * of elements, no matter what keys are present in each map. This drastically improves performance when comparing
+	 * maps of different sizes. This is different than would be exepcted from a lexical compare of strings, but it does
+	 * match the comparison of two numbers represented as a string of digits.
+	 *
+	 * SetType interface:
+	 *  typename ElementType
+	 *  IntType Num(), bool Contains(const KeyType&), IteratorType<KeyType> begin(), end()
+	 *
+	 * @param A Left-Hand-Side set container.
+	 * @param B Right-Hand-Side set container.
+	 * @param KeyLessThan bool(const KeyType& A, const KeyType& B) that returns A < B.
+	 *
+	 * @return -1 if A < B, 0 if A == B, 1 if A > B
+	 */
+	template <typename SetType, typename KeyLessThanType>
+	int CompareSet(const SetType& A, const SetType& B, KeyLessThanType KeyLessThan)
+	{
+		using KeyType = typename SetType::ElementType;
+		if (A.Num() != B.Num())
+		{
+			return A.Num() < B.Num() ? -1 : 1;
+		}
+		if (A.Num() == 0)
+		{
+			return 0;
+		}
+
+		const KeyType* MinKeyWithDifference = nullptr;
+		bool bMinKeyIsLessInA = false;
+		for (const KeyType& AKey : A)
+		{
+			if (!B.Contains(AKey))
+			{
+				if (!MinKeyWithDifference || KeyLessThan(AKey, *MinKeyWithDifference))
+				{
+					MinKeyWithDifference = &AKey;
+					bMinKeyIsLessInA = true;
+				}
+			}
+		}
+
+		// The number of keys in A and B is the same (checked above), so if all keys of A are in B, then there are
+		// no additional keys in B that are not in A and we don't need to iterate over B.
+		if (MinKeyWithDifference)
+		{
+			// If B has additional keys not in A then we need to check each of those not-in-A keys to see if they
+			// are smaller than the MinKeyWithDifference, and if so they become the MinKeyWithDifference.
+			for (const KeyType& BKey : B)
+			{
+				if (!A.Contains(BKey))
+				{
+					if (!MinKeyWithDifference || KeyLessThan(BKey, *MinKeyWithDifference))
+					{
+						MinKeyWithDifference = &BKey;
+						bMinKeyIsLessInA = false;
+					}
+				}
+			}
+		}
+		if (MinKeyWithDifference)
+		{
+			return bMinKeyIsLessInA ? -1 : 1;
+		}
+		return 0;
+	}
+
+	template <typename SetType>
+	int CompareSet(const SetType& A, const SetType& B)
+	{
+		return CompareSet(A, B, TLess<>());
+	}
 }
diff --git a/Engine/Source/Runtime/CoreUObject/Private/UObject/CoreRedirects.cpp b/Engine/Source/Runtime/CoreUObject/Private/UObject/CoreRedirects.cpp
--- a/Engine/Source/Runtime/CoreUObject/Private/UObject/CoreRedirects.cpp
+++ b/Engine/Source/Runtime/CoreUObject/Private/UObject/CoreRedirects.cpp
@@ -8,6 +8,7 @@
 #include "UObject/TopLevelAssetPath.h"
 #include "UObject/UnrealType.h"
 
+#include "Algo/Compare.h"
 #include "GenericPlatform/GenericPlatformFile.h"
 #include "Hash/Blake3.h"
 #include "Logging/StructuredLog.h"
@@ -1189,63 +1190,6 @@ void FCoreRedirect::AppendHash(FBlake3& Hasher) const
 	}
 }
 
-template <typename KeyType, typename ValueType>
-int DeterministicCompare(const TMap<KeyType, ValueType>& A, const TMap<KeyType, ValueType>& B)
-{
-	if (A.Num() != B.Num())
-	{
-		return A.Num() < B.Num() ? -1 : 1;
-	}
-	if (A.Num() == 0)
-	{
-		return 0;
-	}
-
-	const KeyType* MinKeyWithDifference = nullptr;
-	bool bMinKeyIsLessInA = false;
-	for (const TPair<KeyType, ValueType>& Pair : A)
-	{
-		const ValueType* BValue = B.Find(Pair.Key);
-		int Compare = 0;
-		if (!BValue)
-		{
-			Compare = -1;
-		}
-		else if (Pair.Value < *BValue)
-		{
-			Compare = -1;
-		}
-		else if (*BValue < Pair.Value)
-		{
-			Compare = 1;
-		}
-		if (Compare != 0)
-		{
-			if (!MinKeyWithDifference || *Pair.Key < *MinKeyWithDifference)
-			{
-				MinKeyWithDifference = &Pair.Key;
-				bMinKeyIsLessInA = Compare < 0;
-			}
-		}
-	}
-	for (const TPair<KeyType, ValueType>& Pair : B)
-	{
-		if (!A.Contains(Pair.Key))
-		{
-			if (!MinKeyWithDifference || *Pair.Key < *MinKeyWithDifference)
-			{
-				MinKeyWithDifference = &Pair.Key;
-				bMinKeyIsLessInA = false;
-			}
-		}
-	}
-	if (MinKeyWithDifference)
-	{
-		return bMinKeyIsLessInA ? -1 : 1;
-	}
-	return 0;
-}
-
 int FCoreRedirect::Compare(const FCoreRedirect& Other) const
 {
 	if (RedirectFlags != Other.RedirectFlags)
@@ -1267,7 +1211,7 @@ int FCoreRedirect::Compare(const FCoreRedirect& Other) const
 	{
 		return Compare;
 	}
-	Compare = DeterministicCompare(ValueChanges, Other.ValueChanges);
+	Compare = Algo::CompareMap(ValueChanges, Other.ValueChanges);
 	if (Compare != 0)
 	{
 		return Compare;
