Fixed an issue where the PreloadModuleSymbolFile could fail if the symbol file was larger than SSIZE_MAX. Also fixed an issue with the pointer not being adjusted by PageSize if it did fail and needed to be cleaned up.

diff --git a/Engine/Source/Runtime/Core/Private/Unix/UnixPlatformStackWalk.cpp b/Engine/Source/Runtime/Core/Private/Unix/UnixPlatformStackWalk.cpp
--- a/Engine/Source/Runtime/Core/Private/Unix/UnixPlatformStackWalk.cpp
+++ b/Engine/Source/Runtime/Core/Private/Unix/UnixPlatformStackWalk.cpp
@@ -67,13 +67,27 @@ void CORE_API UnixPlatformStackWalk_PreloadModuleSymbolFile()
 			GModuleSymbolFileMemory = (uint8_t*)FMemory::Malloc(GModuleSymbolFileMemorySize +  2 * (FPlatformMemory::GetConstants().PageSize));
 			GModuleSymbolFileMemory += FPlatformMemory::GetConstants().PageSize;
 			
-			ssize_t BytesRead = read(SymbolFileFD, GModuleSymbolFileMemory, GModuleSymbolFileMemorySize);
+			ssize_t BytesRead = 0;
+			{
+				ssize_t RemainingBytes = GModuleSymbolFileMemorySize;
+				uint8_t* CurrentModulePos = GModuleSymbolFileMemory;
+				while(RemainingBytes > SSIZE_MAX)
+				{
+					BytesRead += read(SymbolFileFD, CurrentModulePos, SSIZE_MAX);
+					RemainingBytes -= SSIZE_MAX;
+					CurrentModulePos += SSIZE_MAX;
+				}
+				BytesRead += read(SymbolFileFD, CurrentModulePos, RemainingBytes);
+			}
+
 
 			close(SymbolFileFD);
 
 			// Did not read expected amount of bytes
 			if (BytesRead != GModuleSymbolFileMemorySize)
 			{
+				UE_LOG(LogHAL, Warning, TEXT("UnixPlatformStackWalk_UnloadPreloadedModuleSymbol: BytesRead %d Expected %d"), BytesRead, GModuleSymbolFileMemorySize);
+				GModuleSymbolFileMemory -= FPlatformMemory::GetConstants().PageSize;
 				FMemory::Free(GModuleSymbolFileMemory);
 
 				if (BytesRead == -1)
