Slight change incase the full SSIZE_MAX isn't read from the file we don't advance the pointed too far.

diff --git a/Engine/Source/Runtime/Core/Private/Unix/UnixPlatformStackWalk.cpp b/Engine/Source/Runtime/Core/Private/Unix/UnixPlatformStackWalk.cpp
--- a/Engine/Source/Runtime/Core/Private/Unix/UnixPlatformStackWalk.cpp
+++ b/Engine/Source/Runtime/Core/Private/Unix/UnixPlatformStackWalk.cpp
@@ -71,11 +71,17 @@ void CORE_API UnixPlatformStackWalk_PreloadModuleSymbolFile()
 			{
 				ssize_t RemainingBytes = GModuleSymbolFileMemorySize;
 				uint8_t* CurrentModulePos = GModuleSymbolFileMemory;
+				ssize_t CurrentBytesRead = 0;
 				while(RemainingBytes > SSIZE_MAX)
 				{
-					BytesRead += read(SymbolFileFD, CurrentModulePos, SSIZE_MAX);
-					RemainingBytes -= SSIZE_MAX;
-					CurrentModulePos += SSIZE_MAX;
+					CurrentBytesRead = read(SymbolFileFD, CurrentModulePos, SSIZE_MAX);
+					if(CurrentBytesRead < 0)
+					{
+						break;
+					}
+					RemainingBytes -= CurrentBytesRead;
+					CurrentModulePos += CurrentBytesRead;
+					BytesRead += CurrentBytesRead;
 				}
 				BytesRead += read(SymbolFileFD, CurrentModulePos, RemainingBytes);
 			}
