Fix Android CIS

diff --git a/Engine/Source/Runtime/Core/Private/Android/AndroidSwappyLoader.cpp b/Engine/Source/Runtime/Core/Private/Android/AndroidSwappyLoader.cpp
--- a/Engine/Source/Runtime/Core/Private/Android/AndroidSwappyLoader.cpp
+++ b/Engine/Source/Runtime/Core/Private/Android/AndroidSwappyLoader.cpp
@@ -2,8 +2,6 @@
 
 #include "AndroidSwappyLoader.h"
 
-#if USE_ANDROID_OPENGL_SWAPPY || USE_ANDROID_VULKAN_SWAPPY
-
 #undef VK_NO_PROTOTYPES
 #include "swappy/swappyVk.h"
 #include "swappy/swappyGL.h"
@@ -357,6 +355,4 @@ void SwappyGL_getStats(SwappyStats* swappyStats)
 void SwappyGL_uninjectTracer(const SwappyTracer* t)
 {
 	fpSwappyGL_uninjectTracer(t);
-}
-
-#endif // USE_ANDROID_OPENGL_SWAPPY || USE_ANDROID_VULKAN_SWAPPY
\ No newline at end of file
+}
\ No newline at end of file
diff --git a/Engine/Source/Runtime/Core/Private/Android/AndroidSwappyLoader.h b/Engine/Source/Runtime/Core/Private/Android/AndroidSwappyLoader.h
--- a/Engine/Source/Runtime/Core/Private/Android/AndroidSwappyLoader.h
+++ b/Engine/Source/Runtime/Core/Private/Android/AndroidSwappyLoader.h
@@ -4,14 +4,5 @@
 #include "CoreMinimal.h"
 #include "CoreTypes.h"
 
-#ifndef USE_ANDROID_OPENGL_SWAPPY
-#define USE_ANDROID_OPENGL_SWAPPY 0
-#endif
-#ifndef USE_ANDROID_VULKAN_SWAPPY
-#define USE_ANDROID_VULKAN_SWAPPY 0
-#endif
-
-#if USE_ANDROID_OPENGL_SWAPPY || USE_ANDROID_VULKAN_SWAPPY
 void LoadSwappy();
 void UnloadSwappy();
-#endif
\ No newline at end of file
