Added dll exports to fix console modular build linker errors

diff --git a/Engine/Source/Runtime/Core/Public/GenericPlatform/GenericPlatformCrashContext.h b/Engine/Source/Runtime/Core/Public/GenericPlatform/GenericPlatformCrashContext.h
--- a/Engine/Source/Runtime/Core/Public/GenericPlatform/GenericPlatformCrashContext.h
+++ b/Engine/Source/Runtime/Core/Public/GenericPlatform/GenericPlatformCrashContext.h
@@ -861,11 +861,11 @@ private:
 	int32 CrashContextIndex;
 
 	/** Engine and game data set / reset delegates */
-	static FEngineDataResetDelegate OnEngineDataReset;
-	static FEngineDataSetDelegate OnEngineDataSet;
+	CORE_API static FEngineDataResetDelegate OnEngineDataReset;
+	CORE_API static FEngineDataSetDelegate OnEngineDataSet;
 
-	static FGameDataResetDelegate OnGameDataReset;
-	static FGameDataSetDelegate OnGameDataSet;
+	CORE_API static FGameDataResetDelegate OnGameDataReset;
+	CORE_API static FGameDataSetDelegate OnGameDataSet;
 
 	// FNoncopyable
 	FGenericCrashContext( const FGenericCrashContext& ) = delete;
