Do not use the address of the return address by default with clang windows platforms for backtrace a captured stack.
Clang __builtin_frame_address(0) is not behaving like msvc _AddressOfReturnAddress() on all windows platforms causing issues in the stacktrace tracking.
More investigation is needed to see if we can use the fast unwind path on clang as well.

diff --git a/Engine/Source/Runtime/Core/Private/ProfilingDebugging/Microsoft/WindowsCallstackTrace.cpp b/Engine/Source/Runtime/Core/Private/ProfilingDebugging/Microsoft/WindowsCallstackTrace.cpp
--- a/Engine/Source/Runtime/Core/Private/ProfilingDebugging/Microsoft/WindowsCallstackTrace.cpp
+++ b/Engine/Source/Runtime/Core/Private/ProfilingDebugging/Microsoft/WindowsCallstackTrace.cpp
@@ -50,6 +50,16 @@ static bool GModulesAreInitialized = false;
 	#endif
 #endif
 
+// stacktrace tracking using clang intrinsic __builtin_frame_address(0) doesn't work correctly on all windows platforms
+#if !defined(PLATFORM_USE_CALLSTACK_ADDRESS_POINTER)
+	#if defined(__clang__)
+		#define PLATFORM_USE_CALLSTACK_ADDRESS_POINTER 0
+	#else
+		#define PLATFORM_USE_CALLSTACK_ADDRESS_POINTER 1
+	#endif
+#endif
+
+
 #if UE_CALLSTACK_TRACE_USE_UNWIND_TABLES
 
 /*
@@ -741,7 +751,8 @@ public:
 	FBacktracer(FMalloc* InMalloc);
 	~FBacktracer();
 	static FBacktracer*	Get();
-	uint32 GetBacktraceId(void* AddressOfReturnAddress);
+	inline uint32 GetBacktraceId(void* AddressOfReturnAddress);
+	uint32 GetBacktraceId(uint64 ReturnAddress);
 	void AddModule(UPTRINT Base, const TCHAR* Name) {}
 	void RemoveModule(UPTRINT Base) {}
 
@@ -774,10 +785,16 @@ FBacktracer* FBacktracer::Get()
 }
 
 ////////////////////////////////////////////////////////////////////////////////
-uint32 FBacktracer::GetBacktraceId(void* AddressOfReturnAddress) 
+uint32 FBacktracer::GetBacktraceId(void* AddressOfReturnAddress)
 {
-#if !UE_BUILD_SHIPPING
 	const uint64 ReturnAddress = *(uint64*)AddressOfReturnAddress;
+	return GetBacktraceId(ReturnAddress);
+}
+
+////////////////////////////////////////////////////////////////////////////////
+uint32 FBacktracer::GetBacktraceId(uint64 ReturnAddress)
+{
+#if !UE_BUILD_SHIPPING
 	uint64 StackFrames[256];
 	int32 NumStackFrames = FPlatformStackWalk::CaptureStackBackTrace(StackFrames, UE_ARRAY_COUNT(StackFrames));
 	if (NumStackFrames > 0)
@@ -871,10 +888,14 @@ uint32 CallstackTrace_GetCurrentId()
 		return 0;
 	}
 
-	void* AddressOfReturnAddress = PLATFORM_RETURN_ADDRESS_FOR_CALLSTACKTRACING();
+	void* StackAddress = PLATFORM_RETURN_ADDRESS_FOR_CALLSTACKTRACING();
 	if (FBacktracer* Instance = FBacktracer::Get())
 	{
-		return Instance->GetBacktraceId(AddressOfReturnAddress);
+#if PLATFORM_USE_CALLSTACK_ADDRESS_POINTER
+		return Instance->GetBacktraceId(StackAddress);
+#else
+		return Instance->GetBacktraceId((uint64)StackAddress);
+#endif
 	}
 
 	return 0;
diff --git a/Engine/Source/Runtime/Core/Public/Windows/WindowsPlatform.h b/Engine/Source/Runtime/Core/Public/Windows/WindowsPlatform.h
--- a/Engine/Source/Runtime/Core/Public/Windows/WindowsPlatform.h
+++ b/Engine/Source/Runtime/Core/Public/Windows/WindowsPlatform.h
@@ -120,6 +120,7 @@ typedef FWindowsPlatformTypes FPlatformTypes;
 #define PLATFORM_SUPPORTS_BORDERLESS_WINDOW					1
 
 #define PLATFORM_RETURN_ADDRESS_FOR_CALLSTACKTRACING		PLATFORM_RETURN_ADDRESS_POINTER
+#define PLATFORM_USE_CALLSTACK_ADDRESS_POINTER				1
 
 #define WINDOWS_USE_FEATURE_APPLICATIONMISC_CLASS			1
 #define WINDOWS_USE_FEATURE_PLATFORMMISC_CLASS				1
