Merging more code between MB2 and 3 and moving it to MBC

diff --git a/Engine/Source/Runtime/Core/Private/HAL/MallocBinned2.cpp b/Engine/Source/Runtime/Core/Private/HAL/MallocBinned2.cpp
--- a/Engine/Source/Runtime/Core/Private/HAL/MallocBinned2.cpp
+++ b/Engine/Source/Runtime/Core/Private/HAL/MallocBinned2.cpp
@@ -1,9 +1,6 @@
 // Copyright Epic Games, Inc. All Rights Reserved.
 
 #include "HAL/MallocBinned2.h"
-#include "HAL/MallocBinnedCommonUtils.h"
-#include "Logging/LogMacros.h"
-#include "Misc/ScopeLock.h"
 #include "Templates/Function.h"
 #include "GenericPlatform/GenericPlatformProcess.h"
 #include "Stats/Stats.h"
@@ -14,84 +11,33 @@
 #include "Misc/CoreDelegates.h"
 #include "HAL/MallocTimer.h"
 #include "ProfilingDebugging/CsvProfiler.h"
-#include "FramePro/FrameProProfiler.h"
+#include "HAL/MallocBinnedCommonUtils.h"
 
 CSV_DEFINE_CATEGORY_MODULE(CORE_API, FMemory, true);
 
-#if FRAMEPRO_ENABLED
-/** Pushes a profiler scope if it's safe to do so without any new allocations. */
-class FNoAllocScopeCycleCounter
-{
-public:
-	FORCEINLINE FNoAllocScopeCycleCounter(const ANSICHAR* InStatString)
-		: bPop(false)
-		, StatString(InStatString)
-	{
-		if (FFrameProProfiler::IsThreadContextReady() && GCycleStatsShouldEmitNamedEvents)
-		{
-			bPop = true;
-			FFrameProProfiler::PushEvent(StatString);
-		}
-	}
-
-	FORCEINLINE ~FNoAllocScopeCycleCounter()
-	{
-		if (bPop)
-		{
-			FFrameProProfiler::PopEvent(StatString);
-		}
-	}
-private:
-	bool bPop;
-	const ANSICHAR* StatString;
-};
-#	define NOALLOC_SCOPE_CYCLE_COUNTER(Stat) FNoAllocScopeCycleCounter NoAllocCycleCounter_##Stat(#Stat)
-#else
-#	define NOALLOC_SCOPE_CYCLE_COUNTER(Stat)
-#endif // FRAMEPRO_ENABLED
+PRAGMA_DISABLE_UNSAFE_TYPECAST_WARNINGS
 
 
-PRAGMA_DISABLE_UNSAFE_TYPECAST_WARNINGS
+#define UE_DEFAULT_GBinned2MoveOSFreesOffTimeCriticalThreads 1
 
-#if UE_MB2_ALLOW_RUNTIME_TWEAKING
-	int32 GBinned2PerThreadCaches = UE_DEFAULT_GBinned2PerThreadCaches;
-	static FAutoConsoleVariableRef GMallocBinned2PerThreadCachesCVar(
-		TEXT("MallocBinned2.PerThreadCaches"),
-		GBinned2PerThreadCaches,
-		TEXT("Enables per-thread caches of small (<= 32768 byte) allocations from FMallocBinned2")
-		);
-
-	int32 GBinned2MaxBundlesBeforeRecycle = UE_DEFAULT_GBinned2MaxBundlesBeforeRecycle;
-	static FAutoConsoleVariableRef GMallocBinned2MaxBundlesBeforeRecycleCVar(
-		TEXT("MallocBinned2.BundleRecycleCount"),
-		GBinned2MaxBundlesBeforeRecycle,
-		TEXT("Number of freed bundles in the global recycler before it returns them to the system, per-block size. Limited by UE_DEFAULT_GBinned2MaxBundlesBeforeRecycle (currently 4)")
-		);
-
-	int32 GBinned2AllocExtra = UE_DEFAULT_GBinned2AllocExtra;
-	static FAutoConsoleVariableRef GMallocBinned2AllocExtraCVar(
-		TEXT("MallocBinned2.AllocExtra"),
-		GBinned2AllocExtra,
-		TEXT("When we do acquire the lock, how many bins cached in TLS caches. In no case will we grab more than a page.")
-		);
-
-	int32 GBinned2MoveOSFreesOffTimeCriticalThreads = 1;
+#if UE_MBC_ALLOW_RUNTIME_TWEAKING
+	int32 GBinned2MoveOSFreesOffTimeCriticalThreads = UE_DEFAULT_GBinned2MoveOSFreesOffTimeCriticalThreads;
 	static FAutoConsoleVariableRef GGMallocBinned2MoveOSFreesOffTimeCriticalThreadsCVar(
 		TEXT("MallocBinned2.MoveOSFreesOffTimeCriticalThreads"),
 		GBinned2MoveOSFreesOffTimeCriticalThreads,
 		TEXT("When the OS needs to free memory hint to the underlying cache that we are on a time critical thread, it may decide to delay the free for a non time critical thread")
 	);
+#else
+#	define GBinned2MoveOSFreesOffTimeCriticalThreads	UE_DEFAULT_GBinned2MoveOSFreesOffTimeCriticalThreads
 #endif
 
 
 #if UE_MB2_ALLOCATOR_STATS
-	std::atomic<int64> AllocatedSmallPoolMemory(0); // memory that's requested to be allocated by the game
 	std::atomic<int64> AllocatedOSSmallPoolMemory(0);
 
 	std::atomic<int64> AllocatedLargePoolMemory(0); // memory requests to the OS which don't fit in the small pool
 	std::atomic<int64> AllocatedLargePoolMemoryWAlignment(0); // when we allocate at OS level we need to align to a size
 
-	int64 Binned2PoolInfoMemory = 0;
 	int64 Binned2HashMemory = 0;
 #endif
 
@@ -150,350 +96,143 @@ static constexpr uint16 SmallBinSizes[] =
 
 MS_ALIGN(PLATFORM_CACHE_LINE_SIZE) static uint8 UnusedAlignPadding[PLATFORM_CACHE_LINE_SIZE] GCC_ALIGN(PLATFORM_CACHE_LINE_SIZE) = { 0 };
 uint16 FMallocBinned2::SmallBinSizesReversed[UE_MB2_SMALL_POOL_COUNT] = { 0 };
-uint32 FMallocBinned2::OsAllocationGranularity = 0;
 uint32 FMallocBinned2::PageSize = 0;
 FMallocBinned2* FMallocBinned2::MallocBinned2 = nullptr;
 // Mapping of sizes to small table indices
-uint8 FMallocBinned2::MemSizeToPoolIndex[1 + (UE_MB2_MAX_SMALL_POOL_SIZE >> UE_MB2_MINIMUM_ALIGNMENT_SHIFT)] = { 0 };
+uint8 FMallocBinned2::MemSizeToPoolIndex[1 + (UE_MB2_MAX_SMALL_POOL_SIZE >> UE_MBC_MIN_SMALL_POOL_ALIGNMENT_SHIFT)] = { 0 };
 
 
-struct FMallocBinned2::FPoolInfo
+static FORCEINLINE bool IsSupportedSize(SIZE_T Size)
 {
-	static FORCEINLINE bool IsSupportedSize(SIZE_T Size)
-	{
-		bool bResult = IsAligned(Size, UE_MB2_MINIMUM_ALIGNMENT);
-		bResult = bResult && (Size >> UE_MB2_MINIMUM_ALIGNMENT_SHIFT) <= SIZE_T(MAX_uint32);
-		return bResult;
-	}
+	bool bResult = IsAligned(Size, UE_MBC_MIN_SMALL_POOL_ALIGNMENT);
+	bResult = bResult && (Size >> UE_MBC_MIN_SMALL_POOL_ALIGNMENT_SHIFT) <= SIZE_T(MAX_uint32);
+	return bResult;
+}
 
-	enum class ECanary : uint16
-	{
-		Unassigned = 0x3941,
-		FirstFreeBlockIsOSAllocSize = 0x17ea,
-		FirstFreeBlockIsPtr = 0xf317
-	};
-
-	uint16      Taken;          // Number of allocated elements in this pool, when counts down to zero can free the entire pool	
-	ECanary     Canary;         // See ECanary
-	uint32      AllocSize;      // Number of bytes allocated
-	FFreeBlock* FirstFreeBlock; // Pointer to first free memory in this pool or the OS Allocation Size in bytes if this allocation is not binned
-	FPoolInfo*  Next;           // Pointer to next pool
-	FPoolInfo** PtrToPrevNext;  // Pointer to whichever pointer points to this pool
-
-public:
-	FPoolInfo()
-		: Taken(0)
-		, Canary(ECanary::Unassigned)
-		, AllocSize(0)
-		, FirstFreeBlock(nullptr)
-		, Next(nullptr)
-		, PtrToPrevNext(nullptr)
-	{
-	}
+FMallocBinned2::FPoolInfo::FPoolInfo()
+	: Taken(0)
+	, Canary(ECanary::Unassigned)
+	, AllocSize(0)
+	, FirstFreeBlock(nullptr)
+	, Next(nullptr)
+	, PtrToPrevNext(nullptr)
+{
+}
 
-	void CheckCanary(ECanary ShouldBe) const
+void FMallocBinned2::FPoolInfo::CheckCanary(ECanary ShouldBe) const
+{
+	if (Canary != ShouldBe)
 	{
-		if (Canary != ShouldBe)
-		{
-			UE_LOG(LogMemory, Fatal, TEXT("MallocBinned2 Corruption Canary was 0x%x, should be 0x%x"), int32(Canary), int32(ShouldBe));
-		}
+		UE_LOG(LogMemory, Fatal, TEXT("MallocBinned2 Corruption Canary was 0x%x, should be 0x%x"), int32(Canary), int32(ShouldBe));
 	}
+}
 
-	void SetCanary(ECanary ShouldBe, bool bPreexisting, bool bGuaranteedToBeNew)
+void FMallocBinned2::FPoolInfo::SetCanary(ECanary ShouldBe, bool bPreexisting, bool bGuaranteedToBeNew)
+{
+	if (bPreexisting)
 	{
-		if (bPreexisting)
+		if (bGuaranteedToBeNew)
 		{
-			if (bGuaranteedToBeNew)
-			{
-				UE_LOG(LogMemory, Fatal, TEXT("MallocBinned2 Corruption Canary was 0x%x, should be 0x%x. This block is both preexisting and guaranteed to be new; which makes no sense."), int32(Canary), int32(ShouldBe));
-			}
-			if (ShouldBe == ECanary::Unassigned)
-			{
-				if (Canary != ECanary::FirstFreeBlockIsOSAllocSize && Canary != ECanary::FirstFreeBlockIsPtr)
-				{
-					UE_LOG(LogMemory, Fatal, TEXT("MallocBinned2 Corruption Canary was 0x%x, will be 0x%x because this block should be preexisting and in use."), int32(Canary), int32(ShouldBe));
-				}
-			}
-			else if (Canary != ShouldBe)
-			{
-				UE_LOG(LogMemory, Fatal, TEXT("MallocBinned2 Corruption Canary was 0x%x, should be 0x%x because this block should be preexisting."), int32(Canary), int32(ShouldBe));
-			}
+			UE_LOG(LogMemory, Fatal, TEXT("MallocBinned2 Corruption Canary was 0x%x, should be 0x%x. This block is both preexisting and guaranteed to be new; which makes no sense."), int32(Canary), int32(ShouldBe));
 		}
-		else
+		if (ShouldBe == ECanary::Unassigned)
 		{
-			if (bGuaranteedToBeNew)
+			if (Canary != ECanary::FirstFreeBlockIsOSAllocSize && Canary != ECanary::FirstFreeBlockIsPtr)
 			{
-				if (Canary != ECanary::Unassigned)
-				{
-					UE_LOG(LogMemory, Fatal, TEXT("MallocBinned2 Corruption Canary was 0x%x, will be 0x%x. This block is guaranteed to be new yet is it already assigned."), int32(Canary), int32(ShouldBe));
-				}
-			}
-			else if (Canary != ShouldBe && Canary != ECanary::Unassigned)
-			{
-				UE_LOG(LogMemory, Fatal, TEXT("MallocBinned2 Corruption Canary was 0x%x, will be 0x%x does not have an expected value."), int32(Canary), int32(ShouldBe));
+				UE_LOG(LogMemory, Fatal, TEXT("MallocBinned2 Corruption Canary was 0x%x, will be 0x%x because this block should be preexisting and in use."), int32(Canary), int32(ShouldBe));
 			}
 		}
-		Canary = ShouldBe;
-	}
-
-	bool HasFreeBin() const
-	{
-		CheckCanary(ECanary::FirstFreeBlockIsPtr);
-		return FirstFreeBlock && FirstFreeBlock->GetNumFreeBins() != 0;
-	}
-
-	void* AllocateBin()
-	{
-		check(HasFreeBin());
-		++Taken;
-		void* Result = FirstFreeBlock->AllocateBin();
-		ExhaustPoolIfNecessary();
-		return Result;
-	}
-
-	SIZE_T GetOSRequestedBytes() const
-	{
-		return SIZE_T(AllocSize) << UE_MB2_MINIMUM_ALIGNMENT_SHIFT;
-	}
-
-	UPTRINT GetOsAllocatedBytes() const
-	{
-		CheckCanary(ECanary::FirstFreeBlockIsOSAllocSize);
-		return (UPTRINT)FirstFreeBlock;
-	}
-
-	void SetOSAllocationSizes(SIZE_T InRequestedBytes, UPTRINT InAllocatedBytes)
-	{
-		CheckCanary(ECanary::FirstFreeBlockIsOSAllocSize);
-		checkSlow(InRequestedBytes != 0);                // Shouldn't be pooling zero byte allocations
-		checkSlow(InAllocatedBytes >= InRequestedBytes); // We must be allocating at least as much as we requested
-		checkSlow(IsSupportedSize(InRequestedBytes));    // We must be allocating a size we can store
-
-		AllocSize      = uint32(InRequestedBytes >> UE_MB2_MINIMUM_ALIGNMENT_SHIFT);
-		FirstFreeBlock = (FFreeBlock*)InAllocatedBytes;
-	}
-
-	void Link(FPoolInfo*& PrevNext)
-	{
-		if (PrevNext)
+		else if (Canary != ShouldBe)
 		{
-			PrevNext->PtrToPrevNext = &Next;
+			UE_LOG(LogMemory, Fatal, TEXT("MallocBinned2 Corruption Canary was 0x%x, should be 0x%x because this block should be preexisting."), int32(Canary), int32(ShouldBe));
 		}
-		Next          = PrevNext;
-		PtrToPrevNext = &PrevNext;
-		PrevNext      = this;
 	}
-
-	void Unlink()
-	{
-		if (Next)
-		{
-			Next->PtrToPrevNext = PtrToPrevNext;
-		}
-		*PtrToPrevNext = Next;
-	}
-
-private:
-	void ExhaustPoolIfNecessary()
-	{
-		if (FirstFreeBlock->GetNumFreeBins() == 0)
-		{
-			FirstFreeBlock = FirstFreeBlock->NextFreeBlock;
-		}
-		check(!FirstFreeBlock || FirstFreeBlock->GetNumFreeBins() != 0);
-	}
-};
-
-
-struct FMallocBinned2::Private
-{
-	// Implementation. 
-	static CA_NO_RETURN void OutOfMemory(uint64 Size, uint32 Alignment=0)
-	{
-		// this is expected not to return
-		FPlatformMemory::OnOutOfMemory(Size, Alignment);
-	}
-
-	/**
-	 * Gets the FPoolInfo for a memory address. If no valid info exists one is created.
-	 */
-	static FPoolInfo* GetOrCreatePoolInfo(FMallocBinned2& Allocator, void* InPtr, FMallocBinned2::FPoolInfo::ECanary Kind, bool bPreexisting)
+	else
 	{
-		/** 
-		 * Creates an array of FPoolInfo structures for tracking allocations.
-		 */
-		auto CreatePoolArray = [&Allocator](uint64 NumPools)
+		if (bGuaranteedToBeNew)
 		{
-			const uint64 PoolArraySize = NumPools * sizeof(FPoolInfo);
-
-			void* Result;
+			if (Canary != ECanary::Unassigned)
 			{
-				LLM_PLATFORM_SCOPE(ELLMTag::FMalloc);
-				Result = FPlatformMemory::BinnedAllocFromOS(PoolArraySize);
-#if UE_MB2_ALLOCATOR_STATS
-				Binned2PoolInfoMemory += PoolArraySize;
-#endif
-			}
-
-			if (!Result)
-			{
-				FScopeUnlock TempUnlock(&Allocator.Mutex);
-				OutOfMemory(PoolArraySize);
-			}
-
-			DefaultConstructItems<FPoolInfo>(Result, NumPools);
-			return (FPoolInfo*)Result;
-		};
-
-		uint32 BucketIndex;
-		UPTRINT BucketIndexCollision;
-		uint32  PoolIndex;
-		Allocator.PtrToPoolMapping.GetHashBucketAndPoolIndices(InPtr, BucketIndex, BucketIndexCollision, PoolIndex);
-
-		PoolHashBucket* FirstBucket = &Allocator.HashBuckets[BucketIndex];
-		PoolHashBucket* Collision   = FirstBucket;
-		do
-		{
-			if (!Collision->FirstPool)
-			{
-				Collision->BucketIndex = BucketIndexCollision;
-				Collision->FirstPool   = CreatePoolArray(Allocator.NumPoolsPerPage);
-				Collision->FirstPool[PoolIndex].SetCanary(Kind, bPreexisting, true);
-				return &Collision->FirstPool[PoolIndex];
-			}
-
-			if (Collision->BucketIndex == BucketIndexCollision)
-			{
-				Collision->FirstPool[PoolIndex].SetCanary(Kind, bPreexisting, false);
-				return &Collision->FirstPool[PoolIndex];
+				UE_LOG(LogMemory, Fatal, TEXT("MallocBinned2 Corruption Canary was 0x%x, will be 0x%x. This block is guaranteed to be new yet is it already assigned."), int32(Canary), int32(ShouldBe));
 			}
-
-			Collision = Collision->Next;
 		}
-		while (Collision != FirstBucket);
-
-		// Create a new hash bucket entry
-		if (!Allocator.HashBucketFreeList)
+		else if (Canary != ShouldBe && Canary != ECanary::Unassigned)
 		{
-			{
-				LLM_PLATFORM_SCOPE(ELLMTag::FMalloc);
-				Allocator.HashBucketFreeList = (PoolHashBucket*)FPlatformMemory::BinnedAllocFromOS(FMallocBinned2::PageSize);
-#if UE_MB2_ALLOCATOR_STATS
-				Binned2HashMemory += FMallocBinned2::PageSize;
-#endif
-			}
-
-			for (UPTRINT i = 0, n = PageSize / sizeof(PoolHashBucket); i < n; ++i)
-			{
-				Allocator.HashBucketFreeList->Link(new (Allocator.HashBucketFreeList + i) PoolHashBucket());
-			}
+			UE_LOG(LogMemory, Fatal, TEXT("MallocBinned2 Corruption Canary was 0x%x, will be 0x%x does not have an expected value."), int32(Canary), int32(ShouldBe));
 		}
+	}
+	Canary = ShouldBe;
+}
 
-		PoolHashBucket* NextFree  = Allocator.HashBucketFreeList->Next;
-		PoolHashBucket* NewBucket = Allocator.HashBucketFreeList;
+bool FMallocBinned2::FPoolInfo::HasFreeBin() const
+{
+	CheckCanary(ECanary::FirstFreeBlockIsPtr);
+	return FirstFreeBlock && FirstFreeBlock->GetNumFreeBins() != 0;
+}
 
-		NewBucket->Unlink();
+void* FMallocBinned2::FPoolInfo::AllocateBin()
+{
+	check(HasFreeBin());
+	++Taken;
+	void* Result = FirstFreeBlock->AllocateBin();
+	ExhaustPoolIfNecessary();
+	return Result;
+}
 
-		if (NextFree == NewBucket)
-		{
-			NextFree = nullptr;
-		}
-		Allocator.HashBucketFreeList = NextFree;
+SIZE_T FMallocBinned2::FPoolInfo::GetOSRequestedBytes() const
+{
+	return SIZE_T(AllocSize) << UE_MBC_MIN_SMALL_POOL_ALIGNMENT_SHIFT;
+}
 
-		if (!NewBucket->FirstPool)
-		{
-			NewBucket->FirstPool = CreatePoolArray(Allocator.NumPoolsPerPage);
-			NewBucket->FirstPool[PoolIndex].SetCanary(Kind, bPreexisting, true);
-		}
-		else
-		{
-			NewBucket->FirstPool[PoolIndex].SetCanary(Kind, bPreexisting, false);
-		}
+SIZE_T FMallocBinned2::FPoolInfo::GetOsAllocatedBytes() const
+{
+	CheckCanary(ECanary::FirstFreeBlockIsOSAllocSize);
+	return (SIZE_T)FirstFreeBlock;
+}
 
-		NewBucket->BucketIndex = BucketIndexCollision;
+void FMallocBinned2::FPoolInfo::SetOSAllocationSizes(SIZE_T InRequestedBytes, UPTRINT InAllocatedBytes)
+{
+	CheckCanary(ECanary::FirstFreeBlockIsOSAllocSize);
+	checkSlow(InRequestedBytes != 0);                // Shouldn't be pooling zero byte allocations
+	checkSlow(InAllocatedBytes >= InRequestedBytes); // We must be allocating at least as much as we requested
+	checkSlow(IsSupportedSize(InRequestedBytes));    // We must be allocating a size we can store
 
-		FirstBucket->Link(NewBucket);
+	AllocSize      = uint32(InRequestedBytes >> UE_MBC_MIN_SMALL_POOL_ALIGNMENT_SHIFT);
+	FirstFreeBlock = (FFreeBlock*)InAllocatedBytes;
+}
 
-		return &NewBucket->FirstPool[PoolIndex];
+void FMallocBinned2::FPoolInfo::Link(FPoolInfo*& PrevNext)
+{
+	if (PrevNext)
+	{
+		PrevNext->PtrToPrevNext = &Next;
 	}
+	Next          = PrevNext;
+	PtrToPrevNext = &PrevNext;
+	PrevNext      = this;
+}
 
-	static FPoolInfo* FindPoolInfo(FMallocBinned2& Allocator, void* InPtr)
+void FMallocBinned2::FPoolInfo::Unlink()
+{
+	if (Next)
 	{
-		uint32 BucketIndex;
-		UPTRINT BucketIndexCollision;
-		uint32  PoolIndex;
-		Allocator.PtrToPoolMapping.GetHashBucketAndPoolIndices(InPtr, BucketIndex, BucketIndexCollision, PoolIndex);
-
-
-		PoolHashBucket* FirstBucket = &Allocator.HashBuckets[BucketIndex];
-		PoolHashBucket* Collision   = FirstBucket;
-		do
-		{
-			if (Collision->BucketIndex == BucketIndexCollision)
-			{
-				return &Collision->FirstPool[PoolIndex];
-			}
-
-			Collision = Collision->Next;
-		}
-		while (Collision != FirstBucket);
-
-		return nullptr;
+		Next->PtrToPrevNext = PtrToPrevNext;
 	}
+	*PtrToPrevNext = Next;
+}
 
-	struct FGlobalRecycler
+void FMallocBinned2::FPoolInfo::ExhaustPoolIfNecessary()
+{
+	if (FirstFreeBlock->GetNumFreeBins() == 0)
 	{
-		bool PushBundle(uint32 InPoolIndex, FBundleNode* InBundle)
-		{
-			const uint32 NumCachedBundles = FMath::Min<uint32>(GBinned2MaxBundlesBeforeRecycle, UE_DEFAULT_GBinned2MaxBundlesBeforeRecycle);
-			for (uint32 Slot = 0; Slot < NumCachedBundles; Slot++)
-			{
-				if (!Bundles[InPoolIndex].FreeBundles[Slot])
-				{
-					if (!FPlatformAtomics::InterlockedCompareExchangePointer((void**)&Bundles[InPoolIndex].FreeBundles[Slot], InBundle, nullptr))
-					{
-						return true;
-					}
-				}
-			}
-			return false;
-		}
-
-		FBundleNode* PopBundle(uint32 InPoolIndex)
-		{
-			const uint32 NumCachedBundles = FMath::Min<uint32>(GBinned2MaxBundlesBeforeRecycle, UE_DEFAULT_GBinned2MaxBundlesBeforeRecycle);
-			for (uint32 Slot = 0; Slot < NumCachedBundles; Slot++)
-			{
-				FBundleNode* Result = Bundles[InPoolIndex].FreeBundles[Slot];
-				if (Result)
-				{
-					if (FPlatformAtomics::InterlockedCompareExchangePointer((void**)&Bundles[InPoolIndex].FreeBundles[Slot], nullptr, Result) == Result)
-					{
-						return Result;
-					}
-				}
-			}
-			return nullptr;
-		}
-
-	private:
-		struct FPaddedBundlePointer
-		{
-			FBundleNode* FreeBundles[UE_DEFAULT_GBinned2MaxBundlesBeforeRecycle];
-#if (4 + (4 * PLATFORM_64BITS)) * UE_DEFAULT_GBinned2MaxBundlesBeforeRecycle < PLATFORM_CACHE_LINE_SIZE
-#	define UE_MB2_BUNDLE_PADDING (PLATFORM_CACHE_LINE_SIZE - sizeof(FBundleNode*) * UE_DEFAULT_GBinned2MaxBundlesBeforeRecycle)
-			uint8 Padding[UE_MB2_BUNDLE_PADDING];
-#endif
-			FPaddedBundlePointer()
-			{
-				DefaultConstructItems<FBundleNode*>(FreeBundles, UE_DEFAULT_GBinned2MaxBundlesBeforeRecycle);
-			}
-		};
-		static_assert(sizeof(FPaddedBundlePointer) == PLATFORM_CACHE_LINE_SIZE, "FPaddedBundlePointer should be the same size as a cache line");
-		MS_ALIGN(PLATFORM_CACHE_LINE_SIZE) FPaddedBundlePointer Bundles[UE_MB2_SMALL_POOL_COUNT] GCC_ALIGN(PLATFORM_CACHE_LINE_SIZE);
-	};
+		FirstFreeBlock = FirstFreeBlock->NextFreeBlock;
+	}
+	check(!FirstFreeBlock || FirstFreeBlock->GetNumFreeBins() != 0);
+}
 
-	static FGlobalRecycler GGlobalRecycler;
+struct FMallocBinned2::Private
+{
+	static MallocBinnedPrivate::TGlobalRecycler<UE_MB2_SMALL_POOL_COUNT> GGlobalRecycler;
 
 	static void FreeBundles(FMallocBinned2& Allocator, FBundleNode* BundlesToRecycle, uint32 InBinSize, uint32 InPoolIndex)
 	{
@@ -510,7 +249,7 @@ struct FMallocBinned2::Private
 			do
 			{
 				FBundleNode* NextNode = Node->NextNodeInCurrentBundle;
-				FPoolInfo*   NodePool = FindPoolInfo(Allocator, Node);
+				FPoolInfo*   NodePool = Internal::FindPoolInfo(Allocator, Node);
 				if (!NodePool)
 				{
 					UE_LOG(LogMemory, Fatal, TEXT("FMallocBinned2 Attempt to free an unrecognized small block %p"), Node);
@@ -547,7 +286,7 @@ struct FMallocBinned2::Private
 					// Free the OS memory.
 					NodePool->Unlink();
 					{
-						FScopeLock InnerLock(&Allocator.Mutex);
+						FScopeLock InnerLock(&Allocator.ExternalAllocMutex);
 						Allocator.CachedOSPageAllocator.Free(BasePtrOfNode, Allocator.PageSize);
 					}
 #if UE_MB2_ALLOCATOR_STATS
@@ -562,52 +301,10 @@ struct FMallocBinned2::Private
 		}
 	}
 
-
-	static FCriticalSection& GetFreeBlockListsRegistrationMutex()
-	{
-		static FCriticalSection FreeBlockListsRegistrationMutex;
-		return FreeBlockListsRegistrationMutex;
-	}
-
-	static TArray<FPerThreadFreeBlockLists*>& GetRegisteredFreeBlockLists()
-	{
-		static TArray<FPerThreadFreeBlockLists*> RegisteredFreeBlockLists;
-		return RegisteredFreeBlockLists;
-	}
-
-	static void RegisterThreadFreeBlockLists(FPerThreadFreeBlockLists* FreeBlockLists)
-	{
-		NOALLOC_SCOPE_CYCLE_COUNTER(STAT_FMallocBinned2_RegisterThreadFreeBlockLists);
-		FScopeLock Lock(&GetFreeBlockListsRegistrationMutex());
-#if UE_MB2_ALLOCATOR_STATS_VALIDATION
-		++RecursionCounter;
-#endif
-		GetRegisteredFreeBlockLists().Add(FreeBlockLists);
-#if UE_MB2_ALLOCATOR_STATS_VALIDATION
-		--RecursionCounter;
-#endif
-	}
-
-	static void UnregisterThreadFreeBlockLists(FPerThreadFreeBlockLists* FreeBlockLists)
-	{
-		NOALLOC_SCOPE_CYCLE_COUNTER(STAT_FMallocBinned2_UnregisterThreadFreeBlockLists);
-		FScopeLock Lock(&GetFreeBlockListsRegistrationMutex());
-#if UE_MB2_ALLOCATOR_STATS_VALIDATION
-		++RecursionCounter;
-#endif
-		GetRegisteredFreeBlockLists().Remove(FreeBlockLists);
-#if UE_MB2_ALLOCATOR_STATS_VALIDATION
-		--RecursionCounter;
-#endif
-#if UE_MB2_ALLOCATOR_STATS
-		ConsolidatedMemory.fetch_add(FreeBlockLists->AllocatedMemory, std::memory_order_relaxed);
-#endif
-	}
-
 	static void CheckThreadFreeBlockListsForFork()
 	{
 #if BINNED2_FORK_SUPPORT
-		if(GBinned2PerThreadCaches)
+		if(GMallocBinnedPerThreadCaches)
 		{
 			FScopeLock Lock(&GetFreeBlockListsRegistrationMutex());
 			TArray<FPerThreadFreeBlockLists*>& List = GetRegisteredFreeBlockLists();
@@ -618,17 +315,7 @@ struct FMallocBinned2::Private
 	}
 };
 
-FMallocBinned2::Private::FGlobalRecycler FMallocBinned2::Private::GGlobalRecycler;
-
-void FMallocBinned2::RegisterThreadFreeBlockLists(FPerThreadFreeBlockLists* FreeBlockLists)
-{
-	Private::RegisterThreadFreeBlockLists(FreeBlockLists);
-}
-
-void FMallocBinned2::UnregisterThreadFreeBlockLists(FPerThreadFreeBlockLists* FreeBlockLists)
-{
-	Private::UnregisterThreadFreeBlockLists(FreeBlockLists);
-}
+MallocBinnedPrivate::TGlobalRecycler<UE_MB2_SMALL_POOL_COUNT> FMallocBinned2::Private::GGlobalRecycler;
 
 FORCEINLINE void FMallocBinned2::FPoolList::Clear()
 {
@@ -667,12 +354,12 @@ FMallocBinned2::FPoolInfo& FMallocBinned2::FPoolList::PushNewPoolToFront(FMalloc
 
 	void* FreePtr;
 	{
-		FScopeLock Lock(&Allocator.Mutex);
+		FScopeLock Lock(&Allocator.ExternalAllocMutex);
 		FreePtr = Allocator.CachedOSPageAllocator.Allocate(LocalPageSize, FMemory::AllocationHints::SmallPool);
 	}
 	if (!FreePtr)
 	{
-		Private::OutOfMemory(LocalPageSize);
+		OutOfMemory(LocalPageSize);
 	}
 
 #if !UE_USE_VERYLARGEPAGEALLOCATOR || !UE_MB2_BOOKKEEPING_AT_THE_END_OF_LARGEBLOCK
@@ -687,9 +374,9 @@ FMallocBinned2::FPoolInfo& FMallocBinned2::FPoolList::PushNewPoolToFront(FMalloc
 	AllocatedOSSmallPoolMemory.fetch_add((int64)LocalPageSize, std::memory_order_relaxed);
 #endif
 
-	FScopeLock Lock(&Allocator.Mutex);
+	FScopeLock Lock(&Allocator.ExternalAllocMutex);
 	// Create pool
-	FPoolInfo* Result = Private::GetOrCreatePoolInfo(Allocator, Free, FPoolInfo::ECanary::FirstFreeBlockIsPtr, false);
+	FPoolInfo* Result = Internal::GetOrCreatePoolInfo(Allocator, Free, FPoolInfo::ECanary::FirstFreeBlockIsPtr, false);
 	Result->Link(Front);
 	Result->Taken          = 0;
 	Result->FirstFreeBlock = Free;
@@ -698,7 +385,6 @@ FMallocBinned2::FPoolInfo& FMallocBinned2::FPoolList::PushNewPoolToFront(FMalloc
 }
 
 FMallocBinned2::FMallocBinned2()
-	: HashBucketFreeList(nullptr)
 {
 	static bool bOnce = false;
 	check(!bOnce); // this is now a singleton-like thing and you cannot make multiple copies
@@ -726,14 +412,14 @@ FMallocBinned2::FMallocBinned2()
 	static_assert(sizeof(FMallocBinned2::FFreeBlock) <= SmallBinSizes[0], "Pool header must be able to fit into the smallest bin");
 	static_assert(UE_ARRAY_COUNT(SmallBinSizes) == UE_MB2_SMALL_POOL_COUNT, "Number of bins in SmallBinSizes must match UE_MB2_SMALL_POOL_COUNT");
 	static_assert(UE_ARRAY_COUNT(SmallBinSizes) <= 256, "Number of bins in SmallBinSizes fit in one byte");
-	static_assert(sizeof(FFreeBlock) <= UE_MB2_MINIMUM_ALIGNMENT, "Free block struct must be small enough to fit into the smallest bin.");
+	static_assert(sizeof(FFreeBlock) <= UE_MBC_MIN_SMALL_POOL_ALIGNMENT, "Free block struct must be small enough to fit into the smallest bin.");
 
 	// Init pool tables.
 	for (uint32 Index = 0; Index != UE_MB2_SMALL_POOL_COUNT; ++Index)
 	{
 		checkf(Index == 0 || SmallBinSizes[Index - 1] < SmallBinSizes[Index], TEXT("Small bin sizes must be strictly increasing"));
 		checkf(SmallBinSizes[Index] <= PageSize, TEXT("Small bin size must be small enough to fit into a page"));
-		checkf(SmallBinSizes[Index] % UE_MB2_MINIMUM_ALIGNMENT == 0, TEXT("Small bin size must be a multiple of UE_MB2_MINIMUM_ALIGNMENT"));
+		checkf(SmallBinSizes[Index] % UE_MBC_MIN_SMALL_POOL_ALIGNMENT == 0, TEXT("Small bin size must be a multiple of UE_MBC_MINIMUM_ALIGNMENT"));
 
 		SmallPoolTables[Index].BinSize = SmallBinSizes[Index];
 	}
@@ -741,9 +427,9 @@ FMallocBinned2::FMallocBinned2()
 	// Set up pool mappings
 	uint8* IndexEntry = MemSizeToPoolIndex;
 	uint32 PoolIndex  = 0;
-	for (uint32 Index = 0; Index != 1 + (UE_MB2_MAX_SMALL_POOL_SIZE >> UE_MB2_MINIMUM_ALIGNMENT_SHIFT); ++Index)
+	for (uint32 Index = 0; Index != 1 + (UE_MB2_MAX_SMALL_POOL_SIZE >> UE_MBC_MIN_SMALL_POOL_ALIGNMENT_SHIFT); ++Index)
 	{
-		const uint32 BinSize = Index << UE_MB2_MINIMUM_ALIGNMENT_SHIFT; // inverse of int32 Index = int32((Size >> UE_MB2_MINIMUM_ALIGNMENT_SHIFT));
+		const uint32 BinSize = Index << UE_MBC_MIN_SMALL_POOL_ALIGNMENT_SHIFT; // inverse of int32 Index = int32((Size >> UE_MBC_MIN_SMALL_POOL_ALIGNMENT_SHIFT));
 		while (SmallBinSizes[PoolIndex] < BinSize)
 		{
 			++PoolIndex;
@@ -760,17 +446,8 @@ FMallocBinned2::FMallocBinned2()
 		SmallBinSizesReversed[Index] = SmallBinSizes[Partner];
 	}
 
-	const uint64 MaxHashBuckets = PtrToPoolMapping.GetMaxHashBuckets();
-	{
-		LLM_PLATFORM_SCOPE(ELLMTag::FMalloc);
-		size_t AllocationSize = Align(MaxHashBuckets * sizeof(PoolHashBucket), OsAllocationGranularity);
-		HashBuckets = (PoolHashBucket*)FPlatformMemory::BinnedAllocFromOS(AllocationSize);
-#if UE_MB2_ALLOCATOR_STATS
-		Binned2HashMemory += AllocationSize;
-#endif
-	}
+	AllocateHashBuckets();
 
-	DefaultConstructItems<PoolHashBucket>(HashBuckets, MaxHashBuckets);
 	MallocBinned2 = this;
 	GFixedMallocLocationPtr = (FMalloc**)(&MallocBinned2);
 }
@@ -784,14 +461,14 @@ void FMallocBinned2::OnMallocInitialized()
 #if UE_USE_VERYLARGEPAGEALLOCATOR
 	FCoreDelegates::GetLowLevelAllocatorMemoryTrimDelegate().AddLambda([this]()
 		{
-			FScopeLock Lock(&Mutex);
-			CachedOSPageAllocator.FreeAll(&Mutex);
+			FScopeLock Lock(&ExternalAllocMutex);
+			CachedOSPageAllocator.FreeAll(&ExternalAllocMutex);
 		}
 	);
 
 	FCoreDelegates::GetRefreshLowLevelAllocatorDelegate().AddLambda([this]()
 		{
-			FScopeLock Lock(&Mutex);
+			FScopeLock Lock(&ExternalAllocMutex);
 			CachedOSPageAllocator.Refresh();
 		}
 	);
@@ -803,7 +480,7 @@ void FMallocBinned2::OnPreFork()
 {
 #if BINNED2_FORK_SUPPORT
 	// Trim caches so we don't use them in the child process and cause pages to be copied
-	if (GBinned2PerThreadCaches)
+	if (GMallocBinnedPerThreadCaches)
 	{
 		FMallocBinnedCommonUtils::FlushCurrentThreadCache(*this);
 		FMallocBinned2::Private::CheckThreadFreeBlockListsForFork();
@@ -829,7 +506,7 @@ void FMallocBinned2::OnPreFork()
 void FMallocBinned2::OnPostFork()
 {
 #if BINNED2_FORK_SUPPORT
-	if (GBinned2PerThreadCaches)
+	if (GMallocBinnedPerThreadCaches)
 	{
 		FMallocBinnedCommonUtils::FlushCurrentThreadCache(*this);
 		FMallocBinned2::Private::CheckThreadFreeBlockListsForFork();
@@ -873,7 +550,7 @@ void* FMallocBinned2::MallocSelect(SIZE_T Size, uint32 Alignment, bool bUseSmall
 void* FMallocBinned2::MallocExternalSmall(SIZE_T Size, uint32 Alignment)
 {
 	const uint32 PoolIndex = BoundSizeToPoolIndex(Size, MemSizeToPoolIndex);
-	FPerThreadFreeBlockLists* Lists = GBinned2PerThreadCaches ? FPerThreadFreeBlockLists::Get() : nullptr;
+	FPerThreadFreeBlockLists* Lists = GMallocBinnedPerThreadCaches ? FPerThreadFreeBlockLists::Get() : nullptr;
 	if (Lists)
 	{
 		if (Lists->ObtainRecycledPartial(PoolIndex, Private::GGlobalRecycler))
@@ -911,12 +588,12 @@ void* FMallocBinned2::MallocExternalSmall(SIZE_T Size, uint32 Alignment)
 	AllocatedSmallPoolMemory.fetch_add(Table.BinSize, std::memory_order_relaxed);
 #endif
 
-	if (GBinned2AllocExtra)
+	if (GMallocBinnedAllocExtra)
 	{
 		if (Lists)
 		{
 			// prefill the free list with some allocations so we are less likely to hit this slow path with the mutex 
-			for (int32 Index = 0; Index < GBinned2AllocExtra && Pool->HasFreeBin(); Index++)
+			for (int32 Index = 0; Index < GMallocBinnedAllocExtra && Pool->HasFreeBin(); Index++)
 			{
 				if (!Lists->Free(Result, PoolIndex, Table.BinSize))
 				{
@@ -936,14 +613,14 @@ void* FMallocBinned2::MallocExternalSmall(SIZE_T Size, uint32 Alignment)
 
 void* FMallocBinned2::MallocExternalLarge(SIZE_T Size, uint32 Alignment)
 {
-	static_assert(DEFAULT_ALIGNMENT <= UE_MB2_MINIMUM_ALIGNMENT, "DEFAULT_ALIGNMENT is assumed to be zero"); // used below
+	static_assert(DEFAULT_ALIGNMENT <= UE_MBC_MIN_SMALL_POOL_ALIGNMENT, "DEFAULT_ALIGNMENT is assumed to be zero"); // used below
 
-	Alignment = FMath::Max<uint32>(Alignment, UE_MB2_MINIMUM_ALIGNMENT);
+	Alignment = FMath::Max<uint32>(Alignment, UE_MBC_MIN_SMALL_POOL_ALIGNMENT);
 	Size = Align(FMath::Max((SIZE_T)1, Size), Alignment);
 
 	checkf(FMath::IsPowerOfTwo(Alignment), TEXT("Invalid Malloc alignment: '%u' is not a power of two"), Alignment);
 	checkf(Alignment <= PageSize, TEXT("Invalid Malloc alignment: '%u' is greater than the page size '%u'"), Alignment, PageSize);
-	checkf(FMallocBinned2::FPoolInfo::IsSupportedSize(Size), TEXT("Invalid Malloc size: '%" SIZE_T_FMT "'"), Size);
+	checkf(IsSupportedSize(Size), TEXT("Invalid Malloc size: '%" SIZE_T_FMT "'"), Size);
 
 	const SIZE_T AlignedSize = Align(Size, OsAllocationGranularity);
 
@@ -954,17 +631,17 @@ void* FMallocBinned2::MallocExternalLarge(SIZE_T Size, uint32 Alignment)
 	{
 		NOALLOC_SCOPE_CYCLE_COUNTER(STAT_FMallocBinned2_MallocExternalLarge);
 
-		FScopeLock Lock(&Mutex);
+		FScopeLock Lock(&ExternalAllocMutex);
 
 		// Use OS for non-pooled allocations.
-		Result = CachedOSPageAllocator.Allocate(AlignedSize, 0, &Mutex);
+		Result = CachedOSPageAllocator.Allocate(AlignedSize, 0, &ExternalAllocMutex);
 		if (!Result)
 		{
-			FScopeUnlock TempUnlock(&Mutex);
-			Private::OutOfMemory(AlignedSize);
+			FScopeUnlock TempUnlock(&ExternalAllocMutex);
+			OutOfMemory(AlignedSize);
 		}
 
-		Pool = Private::GetOrCreatePoolInfo(*this, Result, FPoolInfo::ECanary::FirstFreeBlockIsOSAllocSize, false);
+		Pool = Internal::GetOrCreatePoolInfo(*this, Result, FPoolInfo::ECanary::FirstFreeBlockIsOSAllocSize, false);
 	}
 
 	UE_CLOG(!IsAligned(Result, Alignment) ,LogMemory, Fatal, TEXT("FMallocBinned2 alignment was too large for OS. Alignment=%d   Ptr=%p"), Alignment, Result);
@@ -990,7 +667,7 @@ void* FMallocBinned2::ReallocExternal(void* Ptr, SIZE_T NewSize, uint32 Alignmen
 		return nullptr;
 	}
 
-	static_assert(DEFAULT_ALIGNMENT <= UE_MB2_MINIMUM_ALIGNMENT, "DEFAULT_ALIGNMENT is assumed to be zero"); // used below
+	static_assert(DEFAULT_ALIGNMENT <= UE_MBC_MIN_SMALL_POOL_ALIGNMENT, "DEFAULT_ALIGNMENT is assumed to be zero"); // used below
 
 	if (!IsOSAllocation(Ptr))
 	{
@@ -1027,8 +704,8 @@ void* FMallocBinned2::ReallocExternal(void* Ptr, SIZE_T NewSize, uint32 Alignmen
 	NOALLOC_SCOPE_CYCLE_COUNTER(STAT_FMallocBinned2_ReallocExternal);
 
 	// Allocated from OS.
-	Mutex.Lock();
-	FPoolInfo* Pool = Private::FindPoolInfo(*this, Ptr);
+	ExternalAllocMutex.Lock();
+	FPoolInfo* Pool = Internal::FindPoolInfo(*this, Ptr);
 	UE_CLOG(!Pool, LogMemory, Fatal, TEXT("FMallocBinned2 Attempt to realloc an unrecognized pointer %p"), Ptr);
 
 	const UPTRINT PoolOsBytes = Pool->GetOsAllocatedBytes();
@@ -1049,14 +726,14 @@ void* FMallocBinned2::ReallocExternal(void* Ptr, SIZE_T NewSize, uint32 Alignmen
 		if (bUseSmallMalloc)
 		{
 			// Unlock before a small alloc, which rarely takes a lock
-			Mutex.Unlock();
+			ExternalAllocMutex.Unlock();
 			Result = MallocExternalSmall(NewSize, Alignment);
 		}
 		else
 		{
 			// Unlock after a large alloc, which does take a lock, to save unlocking and re-locking unnecessarily
 			Result = MallocExternalLarge(NewSize, Alignment);
-			Mutex.Unlock();
+			ExternalAllocMutex.Unlock();
 		}
 
 		FMemory::Memcpy(Result, Ptr, FMath::Min(NewSize, PoolOSRequestedBytes));
@@ -1064,14 +741,14 @@ void* FMallocBinned2::ReallocExternal(void* Ptr, SIZE_T NewSize, uint32 Alignmen
 		return Result;
 	}
 
-	Mutex.Unlock();
+	ExternalAllocMutex.Unlock();
 
-	Alignment = FMath::Max<uint32>(Alignment, UE_MB2_MINIMUM_ALIGNMENT);
+	Alignment = FMath::Max<uint32>(Alignment, UE_MBC_MIN_SMALL_POOL_ALIGNMENT);
 	NewSize = Align(FMath::Max((SIZE_T)1, NewSize), Alignment);
 
 	checkf(FMath::IsPowerOfTwo(Alignment), TEXT("Invalid Realloc alignment: '%u' is not a power of two"), Alignment);
 	checkf(Alignment <= PageSize, TEXT("Invalid Realloc alignment: '%u' is greater than the page size '%u'"), Alignment, PageSize);
-	checkf(FMallocBinned2::FPoolInfo::IsSupportedSize(NewSize), TEXT("Invalid Realloc size: '%" SIZE_T_FMT "'"), NewSize);
+	checkf(IsSupportedSize(NewSize), TEXT("Invalid Realloc size: '%" SIZE_T_FMT "'"), NewSize);
 
 #if UE_MB2_ALLOCATOR_STATS
 	AllocatedLargePoolMemory.fetch_add((int64)NewSize - (int64)PoolOSRequestedBytes, std::memory_order_relaxed);
@@ -1103,7 +780,7 @@ void FMallocBinned2::FreeExternal(void* Ptr)
 		const uint32 PoolIndex = BasePtr->PoolIndex;
 
 		FBundleNode* BundlesToRecycle = nullptr;
-		FPerThreadFreeBlockLists* Lists = GBinned2PerThreadCaches ? FPerThreadFreeBlockLists::Get() : nullptr;
+		FPerThreadFreeBlockLists* Lists = GMallocBinnedPerThreadCaches ? FPerThreadFreeBlockLists::Get() : nullptr;
 		if (Lists)
 		{
 			BundlesToRecycle = Lists->RecycleFullBundle(PoolIndex, Private::GGlobalRecycler);
@@ -1134,8 +811,8 @@ void FMallocBinned2::FreeExternal(void* Ptr)
 	}
 	else if (Ptr)
 	{
-		FScopeLock Lock(&Mutex);
-		FPoolInfo* Pool = Private::FindPoolInfo(*this, Ptr);
+		FScopeLock Lock(&ExternalAllocMutex);
+		FPoolInfo* Pool = Internal::FindPoolInfo(*this, Ptr);
 		if (!Pool)
 		{
 			UE_LOG(LogMemory, Fatal, TEXT("FMallocBinned2 Attempt to free an unrecognized pointer %p"), Ptr);
@@ -1151,43 +828,10 @@ void FMallocBinned2::FreeExternal(void* Ptr)
 		checkf(PoolOSRequestedBytes <= PoolOsBytes, TEXT("FMallocBinned2::FreeExternal %d %d"), int32(PoolOSRequestedBytes), int32(PoolOsBytes));
 		Pool->SetCanary(FPoolInfo::ECanary::Unassigned, true, false);
 		// Free an OS allocation.
-		CachedOSPageAllocator.Free(Ptr, PoolOsBytes, &Mutex, FPerThreadFreeBlockLists::Get() != nullptr && GBinned2MoveOSFreesOffTimeCriticalThreads != 0);
+		CachedOSPageAllocator.Free(Ptr, PoolOsBytes, &ExternalAllocMutex, FPerThreadFreeBlockLists::Get() != nullptr && GBinned2MoveOSFreesOffTimeCriticalThreads != 0);
 	}
 }
 
-bool FMallocBinned2::GetAllocationSizeExternal(void* Ptr, SIZE_T& SizeOut)
-{
-	if (!IsOSAllocation(Ptr))
-	{
-		check(Ptr); // null is 64k aligned so we should not be here
-		const FFreeBlock* Free = GetPoolHeaderFromPointer(Ptr);
-		CanaryTest(Free);
-		SizeOut = Free->BinSize;
-		return true;
-	}
-	if (!Ptr)
-	{
-		return false;
-	}
-
-	FPoolInfo* Pool;
-	{
-		NOALLOC_SCOPE_CYCLE_COUNTER(STAT_FMallocBinned2_GetAllocationSizeExternal);
-		FScopeLock Lock(&Mutex);
-		Pool = Private::FindPoolInfo(*this, Ptr);
-	}
-
-	if (!Pool)
-	{
-		UE_LOG(LogMemory, Fatal, TEXT("FMallocBinned2 Attempt to GetAllocationSizeExternal an unrecognized pointer %p"), Ptr);
-	}
-	const UPTRINT PoolOsBytes = Pool->GetOsAllocatedBytes();
-	const SIZE_T PoolOSRequestedBytes = Pool->GetOSRequestedBytes();
-	checkf(PoolOSRequestedBytes <= PoolOsBytes, TEXT("FMallocBinned2::GetAllocationSizeExternal %d %d"), int32(PoolOSRequestedBytes), int32(PoolOsBytes));
-	SizeOut = PoolOsBytes;
-	return true;
-}
-
 void FMallocBinned2::FPoolList::ValidateActivePools() const
 {
 	for (const FPoolInfo* const* PoolPtr = &Front; *PoolPtr; PoolPtr = &(*PoolPtr)->Next)
@@ -1236,82 +880,24 @@ void FMallocBinned2::FreeBundles(FBundleNode* Bundles, uint32 PoolIndex)
 	Private::FreeBundles(*this, Bundles, PoolIndexToBinSize(PoolIndex), PoolIndex);
 }
 
-FCriticalSection& FMallocBinned2::GetFreeBlockListsRegistrationMutex()
-{
-	return Private::GetFreeBlockListsRegistrationMutex();
-}
-
-TArray<FMallocBinned2::FPerThreadFreeBlockLists*>& FMallocBinned2::GetRegisteredFreeBlockLists()
-{
-	return Private::GetRegisteredFreeBlockLists();
-}
-
 void FMallocBinned2::Trim(bool bTrimThreadCaches)
 {
-	if (GBinned2PerThreadCaches && bTrimThreadCaches)
+	if (GMallocBinnedPerThreadCaches && bTrimThreadCaches)
 	{
 		FMallocBinnedCommonUtils::Trim(*this);
 
 #if !UE_USE_VERYLARGEPAGEALLOCATOR
-		FScopeLock Lock(&Mutex);
+		FScopeLock Lock(&ExternalAllocMutex);
 		// this cache is recycled anyway, if you need to trim it based on being OOM, it's already too late.
-		CachedOSPageAllocator.FreeAll(&Mutex);
+		CachedOSPageAllocator.FreeAll(&ExternalAllocMutex);
 #endif
 		//UE_LOG(LogTemp, Display, TEXT("Trim CachedOSPageAllocator = %6.2fms"), 1000.0f * float(FPlatformTime::Seconds() - StartTime));
 	}
 }
 
-void FMallocBinned2::SetupTLSCachesOnCurrentThread()
+void FMallocBinned2::FlushCurrentThreadCacheInternal(bool bNewEpochOnly)
 {
-	NOALLOC_SCOPE_CYCLE_COUNTER(STAT_FMallocBinned2_SetupTLSCachesOnCurrentThread);
-
-	if (!UE_MB2_ALLOW_RUNTIME_TWEAKING && !GBinned2PerThreadCaches)
-	{
-		return;
-	}
-	if (!FPlatformTLS::IsValidTlsSlot(FMallocBinned2::BinnedTlsSlot))
-	{
-		FMallocBinned2::BinnedTlsSlot = FPlatformTLS::AllocTlsSlot();
-	}
-	check(FPlatformTLS::IsValidTlsSlot(FMallocBinned2::BinnedTlsSlot));
-	FPerThreadFreeBlockLists::SetTLS();
-}
-
-void FMallocBinned2::ClearAndDisableTLSCachesOnCurrentThread()
-{
-	NOALLOC_SCOPE_CYCLE_COUNTER(STAT_FMallocBinned2_ClearTLSCachesOnCurrentThread);
-	FMallocBinnedCommonUtils::FlushCurrentThreadCache(*this);
-	FPerThreadFreeBlockLists::ClearTLS();
-}
-
-void FMallocBinned2::MarkTLSCachesAsUsedOnCurrentThread()
-{
-	NOALLOC_SCOPE_CYCLE_COUNTER(STAT_FMallocBinned2_MarkTLSCachesAsUsedOnCurrentThread);
-	FPerThreadFreeBlockLists::LockTLS();
-}
-
-void FMallocBinned2::MarkTLSCachesAsUnusedOnCurrentThread()
-{
-	NOALLOC_SCOPE_CYCLE_COUNTER(STAT_FMallocBinned2_MarkTLSCachesAsUnusedOnCurrentThread);
-
-	// Will only flush if memory trimming epoch has been bumped while the thread was active.
-	const bool bNewEpochOnly = true;
 	FMallocBinnedCommonUtils::FlushCurrentThreadCache(*this, bNewEpochOnly);
-
-	FPerThreadFreeBlockLists::UnlockTLS();
-}
-
-void FMallocBinned2::CanaryTest(const FFreeBlock* Block) const
-{
-#if BINNED2_FORK_SUPPORT
-	// When we support forking there are two valid canary values.
-	if (Block->CanaryAndForkState != CurrentCanary && Block->CanaryAndForkState != OldCanary)
-#else
-	if (Block->CanaryAndForkState != CurrentCanary)
-#endif
-	{
-		CanaryFail(Block);
-	}
 }
 
 void FMallocBinned2::CanaryFail(const FFreeBlock* Block) const
@@ -1323,24 +909,6 @@ void FMallocBinned2::CanaryFail(const FFreeBlock* Block) const
 #endif
 }
 
-#if UE_MB2_ALLOCATOR_STATS
-int64 FMallocBinned2::GetTotalAllocatedSmallPoolMemory() const
-{
-	NOALLOC_SCOPE_CYCLE_COUNTER(STAT_FMallocBinned2_GetTotalAllocatedSmallPoolMemory);
-	int64 FreeBlockAllocatedMemory = 0;
-	{
-		FScopeLock Lock(&Private::GetFreeBlockListsRegistrationMutex());
-		for (const FPerThreadFreeBlockLists* FreeBlockLists : Private::GetRegisteredFreeBlockLists())
-		{
-			FreeBlockAllocatedMemory += FreeBlockLists->AllocatedMemory;
-		}
-		FreeBlockAllocatedMemory += ConsolidatedMemory.load(std::memory_order_relaxed);
-	}
-
-	return AllocatedSmallPoolMemory.load(std::memory_order_relaxed) + FreeBlockAllocatedMemory;
-}
-#endif
-
 void FMallocBinned2::GetAllocatorStats( FGenericMemoryStats& OutStats )
 {
 	// Even if we have MB2 stats off, cached slack needs to be included as it might be needed by some platforms
@@ -1386,12 +954,12 @@ void FMallocBinned2::DumpAllocatorStats(class FOutputDevice& Ar)
 	Ar.Logf(TEXT("Large Pool OS Allocated: %fmb"), ((double)AllocatedLargePoolMemoryWAlignment) / (1024.0f * 1024.0f));
 	Ar.Logf(TEXT("Requested Allocations: %fmb"), ((double)LocalAllocatedLargePoolMemory) / (1024.0f * 1024.0f));
 	Ar.Logf(TEXT("OS Allocated: %fmb"), ((double)LocalAllocatedLargePoolMemoryWAlignment) / (1024.0f * 1024.0f));
-	Ar.Logf(TEXT("PoolInfo: %fmb"), ((double)Binned2PoolInfoMemory) / (1024.0f * 1024.0f));
+	Ar.Logf(TEXT("PoolInfo: %fmb"), ((double)PoolInfoMemory) / (1024.0f * 1024.0f));
 	Ar.Logf(TEXT("Hash: %fmb"), ((double)Binned2HashMemory) / (1024.0f * 1024.0f));
 	Ar.Logf(TEXT("TLS: %fmb"), ((double)TLSMemory.load(std::memory_order_relaxed)) / (1024.0f * 1024.0f));
 	Ar.Logf(TEXT("Total allocated from OS: %fmb"),
 		((double)
-			AllocatedOSSmallPoolMemory + AllocatedLargePoolMemoryWAlignment + Binned2PoolInfoMemory + Binned2HashMemory + TLSMemory.load(std::memory_order_relaxed)
+			AllocatedOSSmallPoolMemory + AllocatedLargePoolMemoryWAlignment + PoolInfoMemory + Binned2HashMemory + TLSMemory.load(std::memory_order_relaxed)
 			) / (1024.0f * 1024.0f));
 	Ar.Logf(TEXT("Cached free OS pages: %fmb"), ((double)OSPageAllocatorCachedFreeSize) / (1024.0f * 1024.0f));
 
@@ -1424,7 +992,7 @@ void FMallocBinned2::UpdateStats()
 void* FMallocBinned2::AllocateMetaDataMemory(SIZE_T Size)
 {
 	LLM_PLATFORM_SCOPE(ELLMTag::FMalloc);
-	Size = Align(Size, FMallocBinned2::OsAllocationGranularity);
+	Size = Align(Size, OsAllocationGranularity);
 	return FPlatformMemory::BinnedAllocFromOS(Size);
 }
 
diff --git a/Engine/Source/Runtime/Core/Private/HAL/MallocBinned3.cpp b/Engine/Source/Runtime/Core/Private/HAL/MallocBinned3.cpp
--- a/Engine/Source/Runtime/Core/Private/HAL/MallocBinned3.cpp
+++ b/Engine/Source/Runtime/Core/Private/HAL/MallocBinned3.cpp
@@ -1,12 +1,10 @@
 // Copyright Epic Games, Inc. All Rights Reserved.
 
 #include "HAL/MallocBinned3.h"
-#include "HAL/MallocBinnedCommonUtils.h"
 
 PRAGMA_DISABLE_UNSAFE_TYPECAST_WARNINGS
 
 #if PLATFORM_64BITS && PLATFORM_HAS_FPlatformVirtualMemoryBlock
-#include "Logging/LogMacros.h"
 #include "Misc/ScopeLock.h"
 #include "Templates/Function.h"
 #include "GenericPlatform/GenericPlatformProcess.h"
@@ -14,6 +12,8 @@ PRAGMA_DISABLE_UNSAFE_TYPECAST_WARNINGS
 #include "HAL/IConsoleManager.h"
 #include "HAL/MemoryMisc.h"
 #include "HAL/PlatformMisc.h"
+#include "HAL/MallocBinnedCommonUtils.h"
+
 
 #if UE_MB3_USE_CACHED_PAGE_ALLOCATOR_FOR_LARGE_ALLOCS
 #include "HAL/Allocators/CachedOSPageAllocator.h"
@@ -30,31 +30,7 @@ TBinned3CachedOSPageAllocator& GetCachedOSPageAllocator()
 
 #endif
 
-#if UE_MB3_ALLOW_RUNTIME_TWEAKING
-	int32 GBinned3PerThreadCaches = UE_DEFAULT_GBinned3PerThreadCaches;
-	static FAutoConsoleVariableRef GMallocBinned3PerThreadCachesCVar(
-		TEXT("MallocBinned3.PerThreadCaches"),
-		GBinned3PerThreadCaches,
-		TEXT("Enables per-thread caches of small (<= 32768 byte) allocations from FMallocBinned3")
-		);
-
-	int32 GBinned3MaxBundlesBeforeRecycle = UE_DEFAULT_GBinned3MaxBundlesBeforeRecycle;
-	static FAutoConsoleVariableRef GMallocBinned3MaxBundlesBeforeRecycleCVar(
-		TEXT("MallocBinned3.BundleRecycleCount"),
-		GBinned3MaxBundlesBeforeRecycle,
-		TEXT("Number of freed bundles in the global recycler before it returns them to the system, per-block size. Limited by UE_DEFAULT_GBinned3MaxBundlesBeforeRecycle (currently 4)")
-		);
-
-	int32 GBinned3AllocExtra = UE_DEFAULT_GBinned3AllocExtra;
-	static FAutoConsoleVariableRef GMallocBinned3AllocExtraCVar(
-		TEXT("MallocBinned3.AllocExtra"),
-		GBinned3AllocExtra,
-		TEXT("When we do acquire the lock, how many bins cached in TLS caches. In no case will we grab more than a page.")
-		);
-#endif
-
 #if UE_MB3_ALLOCATOR_STATS
-	std::atomic<int64> Binned3AllocatedSmallPoolMemory(0); // memory that's requested to be allocated by the game
 	int64 Binned3AllocatedOSSmallPoolMemory = 0;
 
 	int64 Binned3AllocatedLargePoolMemory = 0; // memory requests to the OS which don't fit in the small pool
@@ -62,8 +38,7 @@ TBinned3CachedOSPageAllocator& GetCachedOSPageAllocator()
 
 	TAtomic<int64> Binned3Commits;
 	TAtomic<int64> Binned3Decommits;
-	int64 Binned3PoolInfoMemory = 0;
-	int64 Binned3HashMemory = 0;
+
 	int64 Binned3FreeBitsMemory = 0;
 	TAtomic<int64> Binned3TotalPoolSearches;
 	TAtomic<int64> Binned3TotalPointerTests;
@@ -80,7 +55,6 @@ TBinned3CachedOSPageAllocator& GetCachedOSPageAllocator()
 #endif
 
 uint16 FMallocBinned3::SmallBinSizesReversedShifted[UE_MB3_SMALL_POOL_COUNT + 1] = { 0 };
-uint32 FMallocBinned3::OsAllocationGranularity = 0;
 
 #if !BINNED3_USE_SEPARATE_VM_PER_POOL
 	uint8* FMallocBinned3::Binned3BaseVMPtr = nullptr;
@@ -92,7 +66,63 @@ uint32 FMallocBinned3::OsAllocationGranularity = 0;
 
 FMallocBinned3* FMallocBinned3::MallocBinned3 = nullptr;
 // Mapping of sizes to small table indices
-uint8 FMallocBinned3::MemSizeToPoolIndex[1 + (UE_MB3_MAX_SMALL_POOL_SIZE >> UE_MB3_MINIMUM_ALIGNMENT_SHIFT)] = { 0 };
+uint8 FMallocBinned3::MemSizeToPoolIndex[1 + (UE_MB3_MAX_SMALL_POOL_SIZE >> UE_MBC_MIN_SMALL_POOL_ALIGNMENT_SHIFT)] = { 0 };
+
+/** Information about a piece of free memory. */
+struct FFreeBlock
+{
+	enum
+	{
+		CANARY_VALUE = 0xe7
+	};
+
+	FORCEINLINE FFreeBlock(uint32 InBlockSize, uint32 InBinSize, uint8 InPoolIndex)
+		: BinSizeShifted(InBinSize >> UE_MBC_MIN_SMALL_POOL_ALIGNMENT_SHIFT)
+		, PoolIndex(InPoolIndex)
+		, Canary(CANARY_VALUE)
+		, NextFreeBlockIndex(MAX_uint32)
+	{
+		check(InPoolIndex < MAX_uint8 && (InBinSize >> UE_MBC_MIN_SMALL_POOL_ALIGNMENT_SHIFT) <= MAX_uint16);
+		NumFreeBins = InBlockSize / InBinSize;
+	}
+
+	FORCEINLINE uint32 GetNumFreeBins() const
+	{
+		return NumFreeBins;
+	}
+
+	FORCEINLINE bool IsCanaryOk() const
+	{
+		return Canary == FFreeBlock::CANARY_VALUE;
+	}
+
+	FORCEINLINE void CanaryTest() const
+	{
+		if (!IsCanaryOk())
+		{
+			CanaryFail();
+		}
+		//checkSlow(PoolIndex == BoundSizeToPoolIndex(uint32(BinSizeShifted) << UE_MBC_MIN_SMALL_POOL_ALIGNMENT_SHIFT));
+	}
+
+	FORCEINLINE void CanaryFail() const
+	{
+		UE_LOG(LogMemory, Fatal, TEXT("FMallocBinned3 Attempt to realloc an unrecognized pointer %p   canary == 0x%x != 0x%x"), (void*)this, (int32)Canary, (int32)CANARY_VALUE);
+	}
+
+	FORCEINLINE void* AllocateBin()
+	{
+		--NumFreeBins;
+		return (uint8*)this + NumFreeBins * (uint32(BinSizeShifted) << UE_MBC_MIN_SMALL_POOL_ALIGNMENT_SHIFT);
+	}
+
+	uint16 BinSizeShifted;		// Size of the bins that this list points to >> UE_MBC_MIN_SMALL_POOL_ALIGNMENT_SHIFT
+	uint8  PoolIndex;			// Index of this pool
+	uint8  Canary;				// Constant value of 0xe3
+	uint32 NumFreeBins;			// Number of consecutive free bins here, at least 1.
+	uint32 NextFreeBlockIndex;	// Next free block or MAX_uint32
+};
+
 
 struct FMallocBinned3::FPoolInfoSmall		//This is more like BlockInfoSmall as it stores info per block
 {
@@ -193,117 +223,93 @@ struct FMallocBinned3::FPoolInfoSmall		//This is more like BlockInfoSmall as it
 	}
 };
 
-struct FMallocBinned3::FPoolInfoLarge
+FMallocBinned3::FPoolInfo::FPoolInfo()
+	: Canary(ECanary::LargeUnassigned)
+	, AllocSize(0)
+	, VMSizeDivVirtualSizeAlignment(0)
+	, CommitSize(0)
 {
-	enum class ECanary : uint32
-	{
-		LargeUnassigned = 0x39431234,
-		LargeAssigned = 0x17ea5678,
-	};
+}
 
-public:
-	ECanary	Canary;
-private:
-	uint32 AllocSize;						// Number of bytes allocated
-	uint32 VMSizeDivVirtualSizeAlignment;	// Number of VM bytes allocated aligned for OS
-	uint32 CommitSize;						// Number of bytes committed by the OS
-
-public:
-	FPoolInfoLarge() :
-		Canary(ECanary::LargeUnassigned),
-		AllocSize(0),
-		VMSizeDivVirtualSizeAlignment(0),
-		CommitSize(0)
+void FMallocBinned3::FPoolInfo::CheckCanary(ECanary ShouldBe) const
+{
+	if (Canary != ShouldBe)
 	{
+		UE_LOG(LogMemory, Fatal, TEXT("MallocBinned3 Corruption Canary was 0x%x, should be 0x%x"), int32(Canary), int32(ShouldBe));
 	}
+}
 
-	void CheckCanary(ECanary ShouldBe) const
+void FMallocBinned3::FPoolInfo::SetCanary(ECanary ShouldBe, bool bPreexisting, bool bGuarnteedToBeNew)
+{
+	if (bPreexisting)
 	{
-		if (Canary != ShouldBe)
+		if (bGuarnteedToBeNew)
 		{
-			UE_LOG(LogMemory, Fatal, TEXT("MallocBinned3 Corruption Canary was 0x%x, should be 0x%x"), int32(Canary), int32(ShouldBe));
+			UE_LOG(LogMemory, Fatal, TEXT("MallocBinned3 Corruption Canary was 0x%x, should be 0x%x. This block is both preexisting and guaranteed to be new; which makes no sense."), int32(Canary), int32(ShouldBe));
 		}
-	}
-
-	void SetCanary(ECanary ShouldBe, bool bPreexisting, bool bGuarnteedToBeNew)
-	{
-		if (bPreexisting)
+		if (ShouldBe == ECanary::LargeUnassigned)
 		{
-			if (bGuarnteedToBeNew)
-			{
-				UE_LOG(LogMemory, Fatal, TEXT("MallocBinned3 Corruption Canary was 0x%x, should be 0x%x. This block is both preexisting and guaranteed to be new; which makes no sense."), int32(Canary), int32(ShouldBe));
-			}
-			if (ShouldBe == ECanary::LargeUnassigned)
-			{
-				if (Canary != ECanary::LargeAssigned)
-				{
-					UE_LOG(LogMemory, Fatal, TEXT("MallocBinned3 Corruption Canary was 0x%x, will be 0x%x because this block should be preexisting and in use."), int32(Canary), int32(ShouldBe));
-				}
-			}
-			else if (Canary != ShouldBe)
+			if (Canary != ECanary::LargeAssigned)
 			{
-				UE_LOG(LogMemory, Fatal, TEXT("MallocBinned3 Corruption Canary was 0x%x, should be 0x%x because this block should be preexisting."), int32(Canary), int32(ShouldBe));
+				UE_LOG(LogMemory, Fatal, TEXT("MallocBinned3 Corruption Canary was 0x%x, will be 0x%x because this block should be preexisting and in use."), int32(Canary), int32(ShouldBe));
 			}
 		}
-		else
+		else if (Canary != ShouldBe)
 		{
-			if (bGuarnteedToBeNew)
-			{
-				if (Canary != ECanary::LargeUnassigned)
-				{
-					UE_LOG(LogMemory, Fatal, TEXT("MallocBinned3 Corruption Canary was 0x%x, will be 0x%x. This block is guaranteed to be new yet is it already assigned."), int32(Canary), int32(ShouldBe));
-				}
-			}
-			else if (Canary != ShouldBe && Canary != ECanary::LargeUnassigned)
+			UE_LOG(LogMemory, Fatal, TEXT("MallocBinned3 Corruption Canary was 0x%x, should be 0x%x because this block should be preexisting."), int32(Canary), int32(ShouldBe));
+		}
+	}
+	else
+	{
+		if (bGuarnteedToBeNew)
+		{
+			if (Canary != ECanary::LargeUnassigned)
 			{
-				UE_LOG(LogMemory, Fatal, TEXT("MallocBinned3 Corruption Canary was 0x%x, will be 0x%x does not have an expected value."), int32(Canary), int32(ShouldBe));
+				UE_LOG(LogMemory, Fatal, TEXT("MallocBinned3 Corruption Canary was 0x%x, will be 0x%x. This block is guaranteed to be new yet is it already assigned."), int32(Canary), int32(ShouldBe));
 			}
 		}
-		Canary = ShouldBe;
+		else if (Canary != ShouldBe && Canary != ECanary::LargeUnassigned)
+		{
+			UE_LOG(LogMemory, Fatal, TEXT("MallocBinned3 Corruption Canary was 0x%x, will be 0x%x does not have an expected value."), int32(Canary), int32(ShouldBe));
+		}
 	}
+	Canary = ShouldBe;
+}
 
-	uint32 GetOSRequestedBytes() const
-	{
-		return AllocSize;
-	}
+uint32 FMallocBinned3::FPoolInfo::GetOSRequestedBytes() const
+{
+	return AllocSize;
+}
 
-	uint32 GetOsCommittedBytes() const
-	{
-		return CommitSize;
-	}
+uint32 FMallocBinned3::FPoolInfo::GetOsCommittedBytes() const
+{
+	return CommitSize;
+}
 
-	uint32 GetOsVMPages() const
-	{
-		CheckCanary(ECanary::LargeAssigned);
-		return VMSizeDivVirtualSizeAlignment;
-	}
+uint32 FMallocBinned3::FPoolInfo::GetOsVMPages() const
+{
+	CheckCanary(ECanary::LargeAssigned);
+	return VMSizeDivVirtualSizeAlignment;
+}
 
-	void SetOSAllocationSize(uint32 InRequestedBytes)
-	{
-		CheckCanary(ECanary::LargeAssigned);
-		AllocSize = InRequestedBytes;
-		check(AllocSize > 0 && CommitSize >= AllocSize && VMSizeDivVirtualSizeAlignment * FPlatformMemory::FPlatformVirtualMemoryBlock::GetVirtualSizeAlignment() >= CommitSize);
-	}
+void FMallocBinned3::FPoolInfo::SetOSAllocationSize(uint32 InRequestedBytes)
+{
+	CheckCanary(ECanary::LargeAssigned);
+	AllocSize = InRequestedBytes;
+	check(AllocSize > 0 && CommitSize >= AllocSize && VMSizeDivVirtualSizeAlignment * FPlatformMemory::FPlatformVirtualMemoryBlock::GetVirtualSizeAlignment() >= CommitSize);
+}
 
-	void SetOSAllocationSizes(uint32 InRequestedBytes, UPTRINT InCommittedBytes, uint32 InVMSizeDivVirtualSizeAlignment)
-	{
-		CheckCanary(ECanary::LargeAssigned);
-		AllocSize = InRequestedBytes;
-		CommitSize = InCommittedBytes;
-		VMSizeDivVirtualSizeAlignment = InVMSizeDivVirtualSizeAlignment;
-		check(AllocSize > 0 && CommitSize >= AllocSize && VMSizeDivVirtualSizeAlignment * FPlatformMemory::FPlatformVirtualMemoryBlock::GetVirtualSizeAlignment() >= CommitSize);
-	}
-};
+void FMallocBinned3::FPoolInfo::SetOSAllocationSizes(uint32 InRequestedBytes, UPTRINT InCommittedBytes, uint32 InVMSizeDivVirtualSizeAlignment)
+{
+	CheckCanary(ECanary::LargeAssigned);
+	AllocSize = InRequestedBytes;
+	CommitSize = InCommittedBytes;
+	VMSizeDivVirtualSizeAlignment = InVMSizeDivVirtualSizeAlignment;
+	check(AllocSize > 0 && CommitSize >= AllocSize && VMSizeDivVirtualSizeAlignment * FPlatformMemory::FPlatformVirtualMemoryBlock::GetVirtualSizeAlignment() >= CommitSize);
+}
 
 struct FMallocBinned3::Private
 {
-	// Implementation. 
-	static CA_NO_RETURN void OutOfMemory(uint64 Size, uint32 Alignment=0)
-	{
-		// this is expected not to return
-		FPlatformMemory::OnOutOfMemory(Size, Alignment);
-	}
-
 	/**
 	* Gets the FPoolInfoSmall for a small block memory address. If no valid info exists one is created.
 	*/
@@ -317,7 +323,7 @@ struct FMallocBinned3::Private
 		{
 			InfoBlock = (FPoolInfoSmall*)Allocator.AllocateMetaDataMemory(Allocator.OsAllocationGranularity);
 #if UE_MB3_ALLOCATOR_STATS
-			Binned3PoolInfoMemory += Allocator.OsAllocationGranularity;
+			PoolInfoMemory += Allocator.OsAllocationGranularity;
 #endif
 			DefaultConstructItems<FPoolInfoSmall>((void*)InfoBlock, InfosPerPage);
 		}
@@ -334,178 +340,7 @@ struct FMallocBinned3::Private
 		return Result;
 	}
 
-	/**
-	 * Gets the FPoolInfoLarge for a large block memory address. If no valid info exists one is created.
-	 */
-	static FPoolInfoLarge* GetOrCreatePoolInfoLarge(FMallocBinned3& Allocator, void* InPtr)
-	{
-		/** 
-		 * Creates an array of FPoolInfo structures for tracking allocations.
-		 */
-		auto CreatePoolArray = [&Allocator](uint64 NumPools)
-		{
-			const uint64 PoolArraySize = NumPools * sizeof(FPoolInfoLarge);
-
-			void* Result = Allocator.AllocateMetaDataMemory(PoolArraySize);
-#if UE_MB3_ALLOCATOR_STATS
-			Binned3PoolInfoMemory += PoolArraySize;
-#endif
-			if (!Result)
-			{
-				OutOfMemory(PoolArraySize);
-			}
-
-			DefaultConstructItems<FPoolInfoLarge>(Result, NumPools);
-			return (FPoolInfoLarge*)Result;
-		};
-
-		uint32  BucketIndex;
-		UPTRINT BucketIndexCollision;
-		uint32  PoolIndex;
-		Allocator.PtrToPoolMapping.GetHashBucketAndPoolIndices(InPtr, BucketIndex, BucketIndexCollision, PoolIndex);
-
-		PoolHashBucket* FirstBucket = &Allocator.HashBuckets[BucketIndex];
-		PoolHashBucket* Collision   = FirstBucket;
-		do
-		{
-			if (!Collision->FirstPool)
-			{
-				Collision->BucketIndex = BucketIndexCollision;
-				Collision->FirstPool = CreatePoolArray(Allocator.NumLargePoolsPerPage);
-				Collision->FirstPool[PoolIndex].SetCanary(FPoolInfoLarge::ECanary::LargeAssigned, false, true);
-				return &Collision->FirstPool[PoolIndex];
-			}
-
-			if (Collision->BucketIndex == BucketIndexCollision)
-			{
-				Collision->FirstPool[PoolIndex].SetCanary(FPoolInfoLarge::ECanary::LargeAssigned, false, false);
-				return &Collision->FirstPool[PoolIndex];
-			}
-
-			Collision = Collision->Next;
-		}
-		while (Collision != FirstBucket);
-
-		// Create a new hash bucket entry
-		if (!Allocator.HashBucketFreeList)
-		{
-			{
-				Allocator.HashBucketFreeList = (PoolHashBucket*)Allocator.AllocateMetaDataMemory(FMallocBinned3::OsAllocationGranularity);
-#if UE_MB3_ALLOCATOR_STATS
-				Binned3HashMemory += FMallocBinned3::OsAllocationGranularity;
-#endif
-			}
-
-			for (UPTRINT i = 0, n = FMallocBinned3::OsAllocationGranularity / sizeof(PoolHashBucket); i < n; ++i)
-			{
-				Allocator.HashBucketFreeList->Link(new (Allocator.HashBucketFreeList + i) PoolHashBucket());
-			}
-		}
-
-		PoolHashBucket* NextFree  = Allocator.HashBucketFreeList->Next;
-		PoolHashBucket* NewBucket = Allocator.HashBucketFreeList;
-
-		NewBucket->Unlink();
-
-		if (NextFree == NewBucket)
-		{
-			NextFree = nullptr;
-		}
-		Allocator.HashBucketFreeList = NextFree;
-
-		if (!NewBucket->FirstPool)
-		{
-			NewBucket->FirstPool = CreatePoolArray(Allocator.NumLargePoolsPerPage);
-			NewBucket->FirstPool[PoolIndex].SetCanary(FPoolInfoLarge::ECanary::LargeAssigned, false, true);
-		}
-		else
-		{
-			NewBucket->FirstPool[PoolIndex].SetCanary(FPoolInfoLarge::ECanary::LargeAssigned, false, false);
-		}
-
-		NewBucket->BucketIndex = BucketIndexCollision;
-
-		FirstBucket->Link(NewBucket);
-
-		return &NewBucket->FirstPool[PoolIndex];
-	}
-
-	static FPoolInfoLarge* FindPoolInfo(FMallocBinned3& Allocator, void* InPtr)
-	{
-		uint32  BucketIndex;
-		UPTRINT BucketIndexCollision;
-		uint32  PoolIndex;
-		Allocator.PtrToPoolMapping.GetHashBucketAndPoolIndices(InPtr, BucketIndex, BucketIndexCollision, PoolIndex);
-
-		PoolHashBucket* FirstBucket = &Allocator.HashBuckets[BucketIndex];
-		PoolHashBucket* Collision   = FirstBucket;
-		do
-		{
-			if (Collision->BucketIndex == BucketIndexCollision)
-			{
-				return &Collision->FirstPool[PoolIndex];
-			}
-
-			Collision = Collision->Next;
-		}
-		while (Collision != FirstBucket);
-
-		return nullptr;
-	}
-
-	struct FGlobalRecycler
-	{
-		bool PushBundle(uint32 InPoolIndex, FBundleNode* InBundle)
-		{
-			const uint32 NumCachedBundles = FMath::Min<uint32>(GBinned3MaxBundlesBeforeRecycle, UE_DEFAULT_GBinned3MaxBundlesBeforeRecycle);
-			for (uint32 Slot = 0; Slot < NumCachedBundles; Slot++)
-			{
-				if (!Bundles[InPoolIndex].FreeBundles[Slot])
-				{
-					if (!FPlatformAtomics::InterlockedCompareExchangePointer((void**)&Bundles[InPoolIndex].FreeBundles[Slot], InBundle, nullptr))
-					{
-						return true;
-					}
-				}
-			}
-			return false;
-		}
-
-		FBundleNode* PopBundle(uint32 InPoolIndex)
-		{
-			const uint32 NumCachedBundles = FMath::Min<uint32>(GBinned3MaxBundlesBeforeRecycle, UE_DEFAULT_GBinned3MaxBundlesBeforeRecycle);
-			for (uint32 Slot = 0; Slot < NumCachedBundles; Slot++)
-			{
-				FBundleNode* Result = Bundles[InPoolIndex].FreeBundles[Slot];
-				if (Result)
-				{
-					if (FPlatformAtomics::InterlockedCompareExchangePointer((void**)&Bundles[InPoolIndex].FreeBundles[Slot], nullptr, Result) == Result)
-					{
-						return Result;
-					}
-				}
-			}
-			return nullptr;
-		}
-
-	private:
-		struct FPaddedBundlePointer
-		{
-			FBundleNode* FreeBundles[UE_DEFAULT_GBinned3MaxBundlesBeforeRecycle];
-#if (4 + (4 * PLATFORM_64BITS)) * UE_DEFAULT_GBinned3MaxBundlesBeforeRecycle < PLATFORM_CACHE_LINE_SIZE
-#	define UE_MB3_BUNDLE_PADDING (PLATFORM_CACHE_LINE_SIZE - sizeof(FBundleNode*) * UE_DEFAULT_GBinned3MaxBundlesBeforeRecycle)
-			uint8 Padding[UE_MB3_BUNDLE_PADDING];
-#endif
-			FPaddedBundlePointer()
-			{
-				DefaultConstructItems<FBundleNode*>(FreeBundles, UE_DEFAULT_GBinned3MaxBundlesBeforeRecycle);
-			}
-		};
-		static_assert(sizeof(FPaddedBundlePointer) == PLATFORM_CACHE_LINE_SIZE, "FPaddedBundlePointer should be the same size as a cache line");
-		MS_ALIGN(PLATFORM_CACHE_LINE_SIZE) FPaddedBundlePointer Bundles[UE_MB3_SMALL_POOL_COUNT] GCC_ALIGN(PLATFORM_CACHE_LINE_SIZE);
-	};
-
-	static FGlobalRecycler GGlobalRecycler;
+	static MallocBinnedPrivate::TGlobalRecycler<UE_MB3_SMALL_POOL_COUNT> GGlobalRecycler;
 
 	static void FreeBundles(FMallocBinned3& Allocator, FBundleNode* BundlesToRecycle, uint32 InBinSize, uint32 InPoolIndex)
 	{
@@ -541,7 +376,7 @@ struct FMallocBinned3::Private
 				FFreeBlock* Free = (FFreeBlock*)Node;
 				Free->NumFreeBins = 1;
 				Free->NextFreeBlockIndex = NodePool->NoFirstFreeIndex ? MAX_uint32 : NodePool->FirstFreeIndex;
-				Free->BinSizeShifted = (InBinSize >> UE_MB3_MINIMUM_ALIGNMENT_SHIFT);
+				Free->BinSizeShifted = (InBinSize >> UE_MBC_MIN_SMALL_POOL_ALIGNMENT_SHIFT);
 				Free->Canary = FFreeBlock::CANARY_VALUE;
 				Free->PoolIndex = InPoolIndex;
 				NodePool->FirstFreeIndex = BinIndexWithinBlock;
@@ -578,52 +413,15 @@ struct FMallocBinned3::Private
 			Bundle = NextBundle;
 		}
 	}
-
-	static FCriticalSection& GetFreeBlockListsRegistrationMutex()
-	{
-		static FCriticalSection FreeBlockListsRegistrationMutex;
-		return FreeBlockListsRegistrationMutex;
-	}
-
-	static TArray<FPerThreadFreeBlockLists*>& GetRegisteredFreeBlockLists()
-	{
-		static TArray<FPerThreadFreeBlockLists*> RegisteredFreeBlockLists;
-		return RegisteredFreeBlockLists;
-	}
-
-	static void RegisterThreadFreeBlockLists(FPerThreadFreeBlockLists* FreeBlockLists)
-	{
-		FScopeLock Lock(&GetFreeBlockListsRegistrationMutex());
-		GetRegisteredFreeBlockLists().Add(FreeBlockLists);
-	}
-
-	static void UnregisterThreadFreeBlockLists(FPerThreadFreeBlockLists* FreeBlockLists)
-	{
-		FScopeLock Lock(&GetFreeBlockListsRegistrationMutex());
-		GetRegisteredFreeBlockLists().Remove(FreeBlockLists);
-#if UE_MB3_ALLOCATOR_STATS
-		ConsolidatedMemory.fetch_add(FreeBlockLists->AllocatedMemory, std::memory_order_relaxed);
-#endif
-	}
 };
 
-FMallocBinned3::Private::FGlobalRecycler FMallocBinned3::Private::GGlobalRecycler;
+MallocBinnedPrivate::TGlobalRecycler<UE_MB3_SMALL_POOL_COUNT> FMallocBinned3::Private::GGlobalRecycler;
 
 void FMallocBinned3::FreeBundles(FBundleNode* Bundles, uint32 PoolIndex)
 {
 	Private::FreeBundles(*this, Bundles, PoolIndexToBinSize(PoolIndex), PoolIndex);
 }
 
-void FMallocBinned3::RegisterThreadFreeBlockLists(FPerThreadFreeBlockLists* FreeBlockLists)
-{
-	Private::RegisterThreadFreeBlockLists(FreeBlockLists);
-}
-
-void FMallocBinned3::UnregisterThreadFreeBlockLists(FPerThreadFreeBlockLists* FreeBlockLists)
-{
-	Private::UnregisterThreadFreeBlockLists(FreeBlockLists);
-}
-
 FMallocBinned3::FPoolInfoSmall* FMallocBinned3::PushNewPoolToFront(FMallocBinned3::FPoolTable& Table, uint32 InBinSize, uint32 InPoolIndex, uint32& OutBlockIndex)
 {
 	const uint32 BlockSize = OsAllocationGranularity * Table.NumMemoryPagesPerBlock;
@@ -673,7 +471,6 @@ FMallocBinned3::FPoolInfoSmall* FMallocBinned3::GetFrontPool(FPoolTable& Table,
 
 
 FMallocBinned3::FMallocBinned3()
-	: HashBucketFreeList(nullptr)
 {
 	static bool bOnce = false;
 	check(!bOnce); // this is now a singleton-like thing and you cannot make multiple copies
@@ -759,20 +556,20 @@ FMallocBinned3::FMallocBinned3()
 #endif
 
 	// large slab sizes are possible OsAllocationGranularity = 65536;
-	NumLargePoolsPerPage = OsAllocationGranularity / sizeof(FPoolInfoLarge);
-	check(OsAllocationGranularity % sizeof(FPoolInfoLarge) == 0);  // these need to divide evenly!
-	PtrToPoolMapping.Init(OsAllocationGranularity, NumLargePoolsPerPage, Constants.AddressStart, Constants.AddressLimit);
+	NumPoolsPerPage = OsAllocationGranularity / sizeof(FPoolInfo);
+	check(OsAllocationGranularity % sizeof(FPoolInfo) == 0);  // these need to divide evenly!
+	PtrToPoolMapping.Init(OsAllocationGranularity, NumPoolsPerPage, Constants.AddressStart, Constants.AddressLimit);
 
 	checkf(Constants.AddressLimit > OsAllocationGranularity, TEXT("OS address limit must be greater than the page size")); // Check to catch 32 bit overflow in AddressLimit
 	static_assert(UE_MB3_SMALL_POOL_COUNT <= 256, "Small bins size array size must fit in a byte");
-	static_assert(sizeof(FFreeBlock) <= UE_MB3_MINIMUM_ALIGNMENT, "Free block struct must be small enough to fit into the smallest bin");
+	static_assert(sizeof(FFreeBlock) <= UE_MBC_MIN_SMALL_POOL_ALIGNMENT, "Free block struct must be small enough to fit into the smallest bin");
 
 	// Init pool tables.
 	FSizeTableEntry SizeTable[UE_MB3_SMALL_POOL_COUNT];
 
-	verify(FSizeTableEntry::FillSizeTable(OsAllocationGranularity, SizeTable, UE_MB3_BASE_PAGE_SIZE, UE_MB3_MINIMUM_ALIGNMENT, UE_MB3_MAX_SMALL_POOL_SIZE, UE_MB3_BASE_PAGE_SIZE) == UE_MB3_SMALL_POOL_COUNT);
+	verify(FSizeTableEntry::FillSizeTable(OsAllocationGranularity, SizeTable, UE_MB3_BASE_PAGE_SIZE, UE_MB3_MAX_SMALL_POOL_SIZE, UE_MB3_BASE_PAGE_SIZE) == UE_MB3_SMALL_POOL_COUNT);
 	checkf(SizeTable[UE_MB3_SMALL_POOL_COUNT - 1].BinSize == UE_MB3_MAX_SMALL_POOL_SIZE, TEXT("UE_MB3_MAX_SMALL_POOL_SIZE must be equal to the largest bin size"));
-	checkf(sizeof(FMallocBinned3::FFreeBlock) <= SizeTable[0].BinSize, TEXT("Pool header must be able to fit into the smallest bin"));
+	checkf(sizeof(FFreeBlock) <= SizeTable[0].BinSize, TEXT("Pool header must be able to fit into the smallest bin"));
 
 	SmallPoolInfosPerPlatformPage = OsAllocationGranularity / sizeof(FPoolInfoSmall);
 
@@ -780,7 +577,6 @@ FMallocBinned3::FMallocBinned3()
 	for (uint32 Index = 0; Index < UE_MB3_SMALL_POOL_COUNT; ++Index)
 	{
 		checkf(Index == 0 || SizeTable[Index - 1].BinSize < SizeTable[Index].BinSize, TEXT("Small bin sizes must be strictly increasing"));
-		checkf(SizeTable[Index].BinSize % UE_MB3_MINIMUM_ALIGNMENT == 0, TEXT("Small bin size must be a multiple of UE_MB3_MINIMUM_ALIGNMENT"));
 
 		SmallPoolTables[Index].BinSize = SizeTable[Index].BinSize;
 		SmallPoolTables[Index].NumMemoryPagesPerBlock = SizeTable[Index].NumMemoryPagesPerBlock;
@@ -797,7 +593,7 @@ FMallocBinned3::FMallocBinned3()
 		const uint32 Size = Align(sizeof(FPoolInfoSmall**) * (TotalNumberOfBlocks + SmallPoolInfosPerPlatformPage - 1) / SmallPoolInfosPerPlatformPage, PLATFORM_CACHE_LINE_SIZE);
 		RequiredMetaMem += Size;
 #if UE_MB3_ALLOCATOR_STATS
-		Binned3PoolInfoMemory += Size;
+		PoolInfoMemory += Size;
 #endif
 
 		const int64 AllocationSize = Align(FBitTree::GetMemoryRequirements(TotalNumberOfBlocks), PLATFORM_CACHE_LINE_SIZE);
@@ -832,9 +628,9 @@ FMallocBinned3::FMallocBinned3()
 	// Set up pool mappings
 	uint8* IndexEntry = MemSizeToPoolIndex;
 	uint32 PoolIndex  = 0;
-	for (uint32 Index = 0; Index != 1 + (UE_MB3_MAX_SMALL_POOL_SIZE >> UE_MB3_MINIMUM_ALIGNMENT_SHIFT); ++Index)
+	for (uint32 Index = 0; Index != 1 + (UE_MB3_MAX_SMALL_POOL_SIZE >> UE_MBC_MIN_SMALL_POOL_ALIGNMENT_SHIFT); ++Index)
 	{
-		const uint32 BinSize = Index << UE_MB3_MINIMUM_ALIGNMENT_SHIFT; // inverse of int32 Index = int32((Size >> UE_MB3_MINIMUM_ALIGNMENT_SHIFT));
+		const uint32 BinSize = Index << UE_MBC_MIN_SMALL_POOL_ALIGNMENT_SHIFT; // inverse of int32 Index = int32((Size >> UE_MBC_MIN_SMALL_POOL_ALIGNMENT_SHIFT));
 		while (SizeTable[PoolIndex].BinSize < BinSize)
 		{
 			++PoolIndex;
@@ -848,20 +644,11 @@ FMallocBinned3::FMallocBinned3()
 	for (uint32 Index = 0; Index != UE_MB3_SMALL_POOL_COUNT; ++Index)
 	{
 		uint32 Partner = UE_MB3_SMALL_POOL_COUNT - Index - 1;
-		SmallBinSizesReversedShifted[Index] = (SizeTable[Partner].BinSize >> UE_MB3_MINIMUM_ALIGNMENT_SHIFT);
+		SmallBinSizesReversedShifted[Index] = (SizeTable[Partner].BinSize >> UE_MBC_MIN_SMALL_POOL_ALIGNMENT_SHIFT);
 	}
 
-	uint64 MaxHashBuckets = PtrToPoolMapping.GetMaxHashBuckets();
-	{
-		int64 HashAllocSize = Align(MaxHashBuckets * sizeof(PoolHashBucket), OsAllocationGranularity);
-		HashBuckets = (PoolHashBucket*)AllocateMetaDataMemory(HashAllocSize);
-#if UE_MB3_ALLOCATOR_STATS
-		Binned3HashMemory += HashAllocSize;
-#endif
-		verify(HashBuckets);
-	}
+	AllocateHashBuckets();
 
-	DefaultConstructItems<PoolHashBucket>(HashBuckets, MaxHashBuckets);
 	MallocBinned3 = this;
 	GFixedMallocLocationPtr = (FMalloc**)(&MallocBinned3);
 }
@@ -930,17 +717,17 @@ bool FMallocBinned3::IsInternallyThreadSafe() const
 
 void* FMallocBinned3::MallocExternal(SIZE_T Size, uint32 Alignment)
 {
-	static_assert(DEFAULT_ALIGNMENT <= UE_MB3_MINIMUM_ALIGNMENT, "DEFAULT_ALIGNMENT is assumed to be zero"); // used below
+	static_assert(DEFAULT_ALIGNMENT <= UE_MBC_MIN_SMALL_POOL_ALIGNMENT, "DEFAULT_ALIGNMENT is assumed to be zero"); // used below
 
 	// Fast path: Allocate from the small pools if the size is small enough and the alignment <= binned3 min alignment.
 	//            Larger alignments can waste a lot of memory allocating an entire page, so some smaller alignments are 
 	//			  handled in the fallback path if less than a predefined max small pool alignment. 
 
-	bool UsePools = (Size <= UE_MB3_MAX_SMALL_POOL_SIZE) && (Alignment <= UE_MB3_MINIMUM_ALIGNMENT);
+	bool UsePools = (Size <= UE_MB3_MAX_SMALL_POOL_SIZE) && (Alignment <= UE_MBC_MIN_SMALL_POOL_ALIGNMENT);
 	
 	if (!UsePools)
 	{
-		// check if allocations that require alignment larger than UE_MB3_MINIMUM_ALIGNMENT can be promoted to a bin with a natural alignment that matches
+		// check if allocations that require alignment larger than UE_MBC_MIN_SMALL_POOL_ALIGNMENT can be promoted to a bin with a natural alignment that matches
 		// i.e. 16 bytes allocation with 128 bytes alignment can be promoted to 128 bytes bin
 		// this will save us a lot of memory as otherwise allocations will be promoted to OS allocs that are at least 64 KB large, depending on a UE_MB3_MAX_SMALL_POOL_SIZE
 		UsePools = PromoteToLargerBin(Size, Alignment, *this);
@@ -949,7 +736,7 @@ void* FMallocBinned3::MallocExternal(SIZE_T Size, uint32 Alignment)
 	if (UsePools) 
 	{
 		const uint32 PoolIndex = BoundSizeToPoolIndex(Size, MemSizeToPoolIndex);
-		FPerThreadFreeBlockLists* Lists = GBinned3PerThreadCaches ? FPerThreadFreeBlockLists::Get() : nullptr;
+		FPerThreadFreeBlockLists* Lists = GMallocBinnedPerThreadCaches ? FPerThreadFreeBlockLists::Get() : nullptr;
 		if (Lists)
 		{
 			if (Lists->ObtainRecycledPartial(PoolIndex, Private::GGlobalRecycler))
@@ -966,6 +753,8 @@ void* FMallocBinned3::MallocExternal(SIZE_T Size, uint32 Alignment)
 			}
 		}
 
+		NOALLOC_SCOPE_CYCLE_COUNTER(STAT_FMallocBinned3_MallocExternalSmall);
+
 		// Allocate from small object pool.
 		FPoolTable& Table = SmallPoolTables[PoolIndex];
 
@@ -997,14 +786,14 @@ void* FMallocBinned3::MallocExternal(SIZE_T Size, uint32 Alignment)
 		void* Result = Pool->AllocateBin(BlockPtr, Table.BinSize);
 #if UE_MB3_ALLOCATOR_STATS
 		Table.HeadEndAlloc(Size);
-		Binned3AllocatedSmallPoolMemory.fetch_add(PoolIndexToBinSize(PoolIndex), std::memory_order_relaxed);
+		AllocatedSmallPoolMemory.fetch_add(PoolIndexToBinSize(PoolIndex), std::memory_order_relaxed);
 #endif
-		if (GBinned3AllocExtra)
+		if (GMallocBinnedAllocExtra)
 		{
 			if (Lists)
 			{
 				// prefill the free list with some allocations so we are less likely to hit this slow path with the mutex 
-				for (int32 Index = 0; Index < GBinned3AllocExtra && Pool->HasFreeBin(); Index++)
+				for (int32 Index = 0; Index < GMallocBinnedAllocExtra && Pool->HasFreeBin(); Index++)
 				{
 					if (!Lists->Free(Result, PoolIndex, Table.BinSize))
 					{
@@ -1021,7 +810,7 @@ void* FMallocBinned3::MallocExternal(SIZE_T Size, uint32 Alignment)
 
 		return Result;
 	}
-	Alignment = FMath::Max<uint32>(Alignment, UE_MB3_MINIMUM_ALIGNMENT);
+	Alignment = FMath::Max<uint32>(Alignment, UE_MBC_MIN_SMALL_POOL_ALIGNMENT);
 	Size = Align(FMath::Max((SIZE_T)1, Size), Alignment);
 
 	check(FMath::IsPowerOfTwo(Alignment));
@@ -1035,6 +824,8 @@ void* FMallocBinned3::MallocExternal(SIZE_T Size, uint32 Alignment)
 
 	LLM_PLATFORM_SCOPE(ELLMTag::FMalloc);
 
+	NOALLOC_SCOPE_CYCLE_COUNTER(STAT_FMallocBinned3_MallocExternalLarge);
+
 #if UE_MB3_USE_CACHED_PAGE_ALLOCATOR_FOR_LARGE_ALLOCS
 	FScopeLock Lock(&Mutex);
 	void* Result = GetCachedOSPageAllocator().Allocate(AlignedSize);
@@ -1061,12 +852,12 @@ void* FMallocBinned3::MallocExternal(SIZE_T Size, uint32 Alignment)
 
 	if (!Result)
 	{
-		Private::OutOfMemory(AlignedSize);
+		OutOfMemory(AlignedSize);
 	}
 	check(IsOSAllocation(Result));
 
 #if! UE_MB3_USE_CACHED_PAGE_ALLOCATOR_FOR_LARGE_ALLOCS
-	FScopeLock Lock(&Mutex);
+	FScopeLock Lock(&ExternalAllocMutex);
 #endif
 
 #if UE_MB3_ALLOCATOR_STATS
@@ -1075,7 +866,7 @@ void* FMallocBinned3::MallocExternal(SIZE_T Size, uint32 Alignment)
 #endif
 
 	// Create pool.
-	FPoolInfoLarge* Pool = Private::GetOrCreatePoolInfoLarge(*this, Result);
+	FPoolInfo* Pool = Internal::GetOrCreatePoolInfo(*this, Result, FPoolInfo::ECanary::LargeAssigned, false);
 	check(Size > 0 && Size <= AlignedSize && AlignedSize >= FPlatformMemory::FPlatformVirtualMemoryBlock::GetCommitAlignment());
 #if UE_MB3_USE_CACHED_PAGE_ALLOCATOR_FOR_LARGE_ALLOCS
 	Pool->SetOSAllocationSizes(Size, AlignedSize, AlignedSize / FPlatformMemory::FPlatformVirtualMemoryBlock::GetCommitAlignment());
@@ -1093,7 +884,7 @@ void* FMallocBinned3::ReallocExternal(void* Ptr, SIZE_T NewSize, uint32 Alignmen
 		FMallocBinned3::FreeExternal(Ptr);
 		return nullptr;
 	}
-	static_assert(DEFAULT_ALIGNMENT <= UE_MB3_MINIMUM_ALIGNMENT, "DEFAULT_ALIGNMENT is assumed to be zero"); // used below
+	static_assert(DEFAULT_ALIGNMENT <= UE_MBC_MIN_SMALL_POOL_ALIGNMENT, "DEFAULT_ALIGNMENT is assumed to be zero"); // used below
 	check(FMath::IsPowerOfTwo(Alignment));
 	check(Alignment <= OsAllocationGranularity);
 
@@ -1125,10 +916,12 @@ void* FMallocBinned3::ReallocExternal(void* Ptr, SIZE_T NewSize, uint32 Alignmen
 		return Result;
 	}
 
-	Mutex.Lock();
+	NOALLOC_SCOPE_CYCLE_COUNTER(STAT_FMallocBinned3_ReallocExternal);
+
+	ExternalAllocMutex.Lock();
 
 	// Allocated from OS.
-	FPoolInfoLarge* Pool = Private::FindPoolInfo(*this, Ptr);
+	FPoolInfo* Pool = Internal::FindPoolInfo(*this, Ptr);
 	if (!Pool)
 	{
 		UE_LOG(LogMemory, Fatal, TEXT("FMallocBinned3 Attempt to realloc an unrecognized pointer %p"), Ptr);
@@ -1137,10 +930,10 @@ void* FMallocBinned3::ReallocExternal(void* Ptr, SIZE_T NewSize, uint32 Alignmen
 	const uint32 PoolOSRequestedBytes = Pool->GetOSRequestedBytes();
 	checkf(PoolOSRequestedBytes <= PoolOsBytes, TEXT("FMallocBinned3::ReallocExternal %u %u"), PoolOSRequestedBytes, PoolOsBytes);
 	if (NewSize > PoolOsBytes || // can't fit in the old block
-		(NewSize <= UE_MB3_MAX_SMALL_POOL_SIZE && Alignment <= UE_MB3_MINIMUM_ALIGNMENT) || // can switch to the small bin allocator
+		(NewSize <= UE_MB3_MAX_SMALL_POOL_SIZE && Alignment <= UE_MBC_MIN_SMALL_POOL_ALIGNMENT) || // can switch to the small bin allocator
 		Align(NewSize, OsAllocationGranularity) < PoolOsBytes) // we can get some pages back
 	{
-		Mutex.Unlock();
+		ExternalAllocMutex.Unlock();
 		// Grow or shrink.
 		void* Result = FMallocBinned3::MallocExternal(NewSize, Alignment);
 		SIZE_T CopySize = FMath::Min<SIZE_T>(NewSize, PoolOSRequestedBytes);
@@ -1155,12 +948,14 @@ void* FMallocBinned3::ReallocExternal(void* Ptr, SIZE_T NewSize, uint32 Alignmen
 #endif
 	
 	Pool->SetOSAllocationSize(NewSize);
-	Mutex.Unlock();
+	ExternalAllocMutex.Unlock();
 	return Ptr;
 }
 
 void FMallocBinned3::FreeExternal(void* Ptr)
 {
+	NOALLOC_SCOPE_CYCLE_COUNTER(STAT_FMallocBinned3_FreeExternal);
+
 	const uint64 PoolIndex = PoolIndexFromPtr(Ptr);
 	if (PoolIndex < UE_MB3_SMALL_POOL_COUNT)
 	{
@@ -1168,7 +963,7 @@ void FMallocBinned3::FreeExternal(void* Ptr)
 		const uint32 BinSize = PoolIndexToBinSize(PoolIndex);
 
 		FBundleNode* BundlesToRecycle = nullptr;
-		FPerThreadFreeBlockLists* Lists = GBinned3PerThreadCaches ? FPerThreadFreeBlockLists::Get() : nullptr;
+		FPerThreadFreeBlockLists* Lists = GMallocBinnedPerThreadCaches ? FPerThreadFreeBlockLists::Get() : nullptr;
 		if (Lists)
 		{
 			BundlesToRecycle = Lists->RecycleFullBundle(PoolIndex, Private::GGlobalRecycler);
@@ -1193,7 +988,7 @@ void FMallocBinned3::FreeExternal(void* Ptr)
 			{
 				SmallPoolTables[PoolIndex].HeadEndFree();
 				// lists track their own stat track them instead in the global stat if we don't have lists
-				Binned3AllocatedSmallPoolMemory.fetch_sub((int64)BinSize, std::memory_order_relaxed);
+				AllocatedSmallPoolMemory.fetch_sub((int64)BinSize, std::memory_order_relaxed);
 			}
 #endif
 		}
@@ -1206,9 +1001,9 @@ void FMallocBinned3::FreeExternal(void* Ptr)
 		uint32 VMPages;
 		{
 #if !UE_MB3_USE_CACHED_PAGE_ALLOCATOR_FOR_LARGE_ALLOCS
-			FScopeLock Lock(&Mutex);
+			FScopeLock Lock(&ExternalAllocMutex);
 #endif
-			FPoolInfoLarge* Pool = Private::FindPoolInfo(*this, Ptr);
+			FPoolInfo* Pool = Internal::FindPoolInfo(*this, Ptr);
 			if (!Pool)
 			{
 				UE_LOG(LogMemory, Fatal, TEXT("FMallocBinned3 Attempt to free an unrecognized pointer %p"), Ptr);
@@ -1223,7 +1018,7 @@ void FMallocBinned3::FreeExternal(void* Ptr)
 #endif
 
 			checkf(PoolOSRequestedBytes <= PoolOsBytes, TEXT("FMallocBinned3::FreeExternal %u %u"), int32(PoolOSRequestedBytes), int32(PoolOsBytes));
-			Pool->SetCanary(FPoolInfoLarge::ECanary::LargeUnassigned, true, false);
+			Pool->SetCanary(FPoolInfo::ECanary::LargeUnassigned, true, false);
 		}
 
 		// Free an OS allocation.
@@ -1251,37 +1046,6 @@ void FMallocBinned3::FreeExternal(void* Ptr)
 	}
 }
 
-bool FMallocBinned3::GetAllocationSizeExternal(void* Ptr, SIZE_T& SizeOut)
-{
-	const uint64 PoolIndex = PoolIndexFromPtr(Ptr);
-	if (PoolIndex < UE_MB3_SMALL_POOL_COUNT)
-	{
-		check(Ptr); // null is an OS allocation because it will not fall in our VM block
-		SizeOut = PoolIndexToBinSize(PoolIndex);
-		return true;
-	}
-	if (!Ptr)
-	{
-		return false;
-	}
-
-	FPoolInfoLarge* Pool;
-	{
-		FScopeLock Lock(&Mutex);
-		Pool = Private::FindPoolInfo(*this, Ptr);
-	}
-
-	if (!Pool)
-	{
-		UE_LOG(LogMemory, Fatal, TEXT("FMallocBinned3 Attempt to GetAllocationSizeExternal an unrecognized pointer %p"), Ptr);
-	}
-	const uint32 PoolOsBytes = Pool->GetOsCommittedBytes();
-	const uint32 PoolOSRequestedBytes = Pool->GetOSRequestedBytes();
-	checkf(PoolOSRequestedBytes <= PoolOsBytes, TEXT("FMallocBinned3::GetAllocationSizeExternal %u %u"), PoolOSRequestedBytes, PoolOsBytes);
-	SizeOut = PoolOsBytes;
-	return true;
-}
-
 bool FMallocBinned3::ValidateHeap()
 {
 	// Not implemented
@@ -1296,93 +1060,18 @@ const TCHAR* FMallocBinned3::GetDescriptiveName()
 
 void FMallocBinned3::Trim(bool bTrimThreadCaches)
 {
-	if (GBinned3PerThreadCaches && bTrimThreadCaches)
+	if (GMallocBinnedPerThreadCaches && bTrimThreadCaches)
 	{
 		// Trim memory and increase the Epoch.
 		FMallocBinnedCommonUtils::Trim(*this);
 	}
 }
 
-FCriticalSection& FMallocBinned3::GetFreeBlockListsRegistrationMutex()
-{
-	return Private::GetFreeBlockListsRegistrationMutex();
-}
-
-TArray<FMallocBinned3::FPerThreadFreeBlockLists*>& FMallocBinned3::GetRegisteredFreeBlockLists()
-{
-	return Private::GetRegisteredFreeBlockLists();
-}
-
-void FMallocBinned3::SetupTLSCachesOnCurrentThread()
-{
-	if (!UE_MB3_ALLOW_RUNTIME_TWEAKING && !GBinned3PerThreadCaches)
-	{
-		return;
-	}
-	if (!FPlatformTLS::IsValidTlsSlot(FMallocBinned3::BinnedTlsSlot))
-	{
-		FMallocBinned3::BinnedTlsSlot = FPlatformTLS::AllocTlsSlot();
-	}
-	check(FPlatformTLS::IsValidTlsSlot(FMallocBinned3::BinnedTlsSlot));
-	FPerThreadFreeBlockLists::SetTLS();
-}
-
-void FMallocBinned3::ClearAndDisableTLSCachesOnCurrentThread()
-{
-	if (!UE_MB3_ALLOW_RUNTIME_TWEAKING && !GBinned3PerThreadCaches)
-	{
-		return;
-	}
-
-	FMallocBinnedCommonUtils::FlushCurrentThreadCache(*this);
-	FPerThreadFreeBlockLists::ClearTLS();
-}
-
-void FMallocBinned3::MarkTLSCachesAsUsedOnCurrentThread()
-{
-	if (!UE_MB3_ALLOW_RUNTIME_TWEAKING && !GBinned3PerThreadCaches)
-	{
-		return;
-	}
-
-	FPerThreadFreeBlockLists::LockTLS();
-}
-
-void FMallocBinned3::MarkTLSCachesAsUnusedOnCurrentThread()
+void FMallocBinned3::FlushCurrentThreadCacheInternal(bool bNewEpochOnly)
 {
-	if (!UE_MB3_ALLOW_RUNTIME_TWEAKING && !GBinned3PerThreadCaches)
-	{
-		return;
-	}
-
-	// Will only flush if memory trimming has been called while the thread was active.
-	const bool bNewEpochOnly = true;
 	FMallocBinnedCommonUtils::FlushCurrentThreadCache(*this, bNewEpochOnly);
-	FPerThreadFreeBlockLists::UnlockTLS();
 }
 
-void FMallocBinned3::FFreeBlock::CanaryFail() const
-{
-	UE_LOG(LogMemory, Fatal, TEXT("FMallocBinned3 Attempt to realloc an unrecognized pointer %p   canary == 0x%x != 0x%x"), (void*)this, (int32)Canary, (int32)FMallocBinned3::FFreeBlock::CANARY_VALUE);
-}
-
-#if UE_MB3_ALLOCATOR_STATS
-int64 FMallocBinned3::GetTotalAllocatedSmallPoolMemory() const
-{
-	int64 FreeBlockAllocatedMemory = 0;
-	{
-		FScopeLock Lock(&Private::GetFreeBlockListsRegistrationMutex());
-		for (const FPerThreadFreeBlockLists* FreeBlockLists : Private::GetRegisteredFreeBlockLists())
-		{
-			FreeBlockAllocatedMemory += FreeBlockLists->AllocatedMemory;
-		}
-		FreeBlockAllocatedMemory += ConsolidatedMemory.load(std::memory_order_relaxed);
-	}
-
-	return Binned3AllocatedSmallPoolMemory.load(std::memory_order_relaxed) + FreeBlockAllocatedMemory;
-}
-#endif
-
 void FMallocBinned3::GetAllocatorStats(FGenericMemoryStats& OutStats)
 {
 #if UE_MB3_ALLOCATOR_STATS
@@ -1424,8 +1113,8 @@ void FMallocBinned3::DumpAllocatorStats(class FOutputDevice& Ar)
 	Ar.Logf(TEXT("Small Pool OS Allocated: %fmb"), ((double)Binned3AllocatedOSSmallPoolMemory) / (1024.0f * 1024.0f));
 	Ar.Logf(TEXT("Large Pool Requested Allocations: %fmb"), ((double)Binned3AllocatedLargePoolMemory) / (1024.0f * 1024.0f));
 	Ar.Logf(TEXT("Large Pool OS Allocated: %fmb"), ((double)Binned3AllocatedLargePoolMemoryWAlignment) / (1024.0f * 1024.0f));
-	Ar.Logf(TEXT("PoolInfo: %fmb"), ((double)Binned3PoolInfoMemory) / (1024.0f * 1024.0f));
-	Ar.Logf(TEXT("Hash: %fmb"), ((double)Binned3HashMemory) / (1024.0f * 1024.0f));
+	Ar.Logf(TEXT("PoolInfo: %fmb"), ((double)PoolInfoMemory) / (1024.0f * 1024.0f));
+	Ar.Logf(TEXT("Hash: %fmb"), ((double)HashMemory) / (1024.0f * 1024.0f));
 	Ar.Logf(TEXT("Free Bits: %fmb"), ((double)Binned3FreeBitsMemory) / (1024.0f * 1024.0f));
 	Ar.Logf(TEXT("TLS: %fmb"), ((double)TLSMemory.load(std::memory_order_relaxed)) / (1024.0f * 1024.0f));
 	Ar.Logf(TEXT("Slab Commits: %llu"), Binned3Commits.Load());
@@ -1444,7 +1133,7 @@ void FMallocBinned3::DumpAllocatorStats(class FOutputDevice& Ar)
 #endif
 	Ar.Logf(TEXT("Total allocated from OS: %fmb"), 
 		((double)
-			Binned3AllocatedOSSmallPoolMemory + Binned3AllocatedLargePoolMemoryWAlignment + Binned3PoolInfoMemory + Binned3HashMemory + Binned3FreeBitsMemory + TLSMemory.load(std::memory_order_relaxed)
+			Binned3AllocatedOSSmallPoolMemory + Binned3AllocatedLargePoolMemoryWAlignment + PoolInfoMemory + HashMemory + Binned3FreeBitsMemory + TLSMemory.load(std::memory_order_relaxed)
 			) / (1024.0f * 1024.0f));
 
 
diff --git a/Engine/Source/Runtime/Core/Private/HAL/MallocBinnedCommon.cpp b/Engine/Source/Runtime/Core/Private/HAL/MallocBinnedCommon.cpp
--- a/Engine/Source/Runtime/Core/Private/HAL/MallocBinnedCommon.cpp
+++ b/Engine/Source/Runtime/Core/Private/HAL/MallocBinnedCommon.cpp
@@ -11,8 +11,6 @@
 
 PRAGMA_DISABLE_UNSAFE_TYPECAST_WARNINGS
 
-#if PLATFORM_HAS_FPlatformVirtualMemoryBlock
-
 // Bin sizes are based around getting the maximum amount of allocations per block, with as little alignment waste as possible.
 // Bin sizes should be close to even divisors of the system page size, and well distributed.
 // They must be 16-byte aligned as well.
@@ -89,10 +87,11 @@ static constexpr uint32 BinnedCommonSmallBinSizes28k[] = // 6 total
 	28672  // /1
 };
 
-FSizeTableEntry::FSizeTableEntry(uint32 InBinSize, uint64 PlatformPageSize, uint8 Num4kbPages, uint32 BasePageSize, uint32 MinimumAlignment)
+FSizeTableEntry::FSizeTableEntry(uint32 InBinSize, uint64 PlatformPageSize, uint8 Num4kbPages, uint32 BasePageSize)
 	: BinSize(InBinSize)
 {
-	check((PlatformPageSize & (BasePageSize - 1)) == 0 && PlatformPageSize >= BasePageSize && InBinSize % MinimumAlignment == 0);
+	check((PlatformPageSize & (BasePageSize - 1)) == 0 && PlatformPageSize >= BasePageSize);
+	checkf(InBinSize % UE_MBC_MIN_SMALL_POOL_ALIGNMENT == 0, TEXT("Small bin size must be a multiple of UE_MBC_MINIMUM_ALIGNMENT"));
 
 	const uint64 NumBasePagesPerPlatformPage = PlatformPageSize / BasePageSize;
 
@@ -114,28 +113,28 @@ FSizeTableEntry::FSizeTableEntry(uint32 InBinSize, uint64 PlatformPageSize, uint
 	check((PlatformPageSize * NumMemoryPagesPerBlock) / BinSize <= MAX_uint32);
 }
 
-static inline void FillTable(FSizeTableEntry* SizeTable, int32& Index, const uint32* BinsList, uint32 BinListSize, uint32 Num4kbPages, uint64 PlatformPageSize, uint32 BasePageSize, uint32 MinimumAlignment)
+static inline void FillTable(FSizeTableEntry* SizeTable, int32& Index, const uint32* BinsList, uint32 BinListSize, uint32 Num4kbPages, uint64 PlatformPageSize, uint32 BasePageSize)
 {
 	for (uint32 Sub = 0; Sub < BinListSize; Sub++)
 	{
 		// if we override UE_MBC_MAX_LISTED_SMALL_POOL_SIZE externally, we need to filter out predefined bins of a larger size
 		if (BinsList[Sub] <= UE_MBC_MAX_LISTED_SMALL_POOL_SIZE)
 		{
-			SizeTable[Index++] = FSizeTableEntry(BinsList[Sub], PlatformPageSize, Num4kbPages, BasePageSize, MinimumAlignment);
+			SizeTable[Index++] = FSizeTableEntry(BinsList[Sub], PlatformPageSize, Num4kbPages, BasePageSize);
 		}
 	}
 }
 
-uint8 FSizeTableEntry::FillSizeTable(uint64 PlatformPageSize, FSizeTableEntry* SizeTable, uint32 BasePageSize, uint32 MinimumAlignment, uint32 MaxSize, uint32 SizeIncrement)
+uint8 FSizeTableEntry::FillSizeTable(uint64 PlatformPageSize, FSizeTableEntry* SizeTable, uint32 BasePageSize, uint32 MaxSize, uint32 SizeIncrement)
 {
 	int32 Index = 0;
-	FillTable(SizeTable, Index, BinnedCommonSmallBinSizes4k,  UE_ARRAY_COUNT(BinnedCommonSmallBinSizes4k),  1, PlatformPageSize, BasePageSize, MinimumAlignment);
-	FillTable(SizeTable, Index, BinnedCommonSmallBinSizes8k,  UE_ARRAY_COUNT(BinnedCommonSmallBinSizes8k),  2, PlatformPageSize, BasePageSize, MinimumAlignment);
-	FillTable(SizeTable, Index, BinnedCommonSmallBinSizes12k, UE_ARRAY_COUNT(BinnedCommonSmallBinSizes12k), 3, PlatformPageSize, BasePageSize, MinimumAlignment);
-	FillTable(SizeTable, Index, BinnedCommonSmallBinSizes16k, UE_ARRAY_COUNT(BinnedCommonSmallBinSizes16k), 4, PlatformPageSize, BasePageSize, MinimumAlignment);
-	FillTable(SizeTable, Index, BinnedCommonSmallBinSizes20k, UE_ARRAY_COUNT(BinnedCommonSmallBinSizes20k), 5, PlatformPageSize, BasePageSize, MinimumAlignment);
-	FillTable(SizeTable, Index, BinnedCommonSmallBinSizes24k, UE_ARRAY_COUNT(BinnedCommonSmallBinSizes24k), 6, PlatformPageSize, BasePageSize, MinimumAlignment);
-	FillTable(SizeTable, Index, BinnedCommonSmallBinSizes28k, UE_ARRAY_COUNT(BinnedCommonSmallBinSizes28k), 7, PlatformPageSize, BasePageSize, MinimumAlignment);
+	FillTable(SizeTable, Index, BinnedCommonSmallBinSizes4k,  UE_ARRAY_COUNT(BinnedCommonSmallBinSizes4k),  1, PlatformPageSize, BasePageSize);
+	FillTable(SizeTable, Index, BinnedCommonSmallBinSizes8k,  UE_ARRAY_COUNT(BinnedCommonSmallBinSizes8k),  2, PlatformPageSize, BasePageSize);
+	FillTable(SizeTable, Index, BinnedCommonSmallBinSizes12k, UE_ARRAY_COUNT(BinnedCommonSmallBinSizes12k), 3, PlatformPageSize, BasePageSize);
+	FillTable(SizeTable, Index, BinnedCommonSmallBinSizes16k, UE_ARRAY_COUNT(BinnedCommonSmallBinSizes16k), 4, PlatformPageSize, BasePageSize);
+	FillTable(SizeTable, Index, BinnedCommonSmallBinSizes20k, UE_ARRAY_COUNT(BinnedCommonSmallBinSizes20k), 5, PlatformPageSize, BasePageSize);
+	FillTable(SizeTable, Index, BinnedCommonSmallBinSizes24k, UE_ARRAY_COUNT(BinnedCommonSmallBinSizes24k), 6, PlatformPageSize, BasePageSize);
+	FillTable(SizeTable, Index, BinnedCommonSmallBinSizes28k, UE_ARRAY_COUNT(BinnedCommonSmallBinSizes28k), 7, PlatformPageSize, BasePageSize);
 
 	check(Index == UE_MBC_NUM_LISTED_SMALL_POOLS);
 
@@ -144,7 +143,7 @@ uint8 FSizeTableEntry::FillSizeTable(uint64 PlatformPageSize, FSizeTableEntry* S
 	check(IsAligned(UE_MBC_MAX_LISTED_SMALL_POOL_SIZE, BasePageSize));
 	for (uint32 Size = UE_MBC_MAX_LISTED_SMALL_POOL_SIZE + BasePageSize; Size <= MaxSize; Size += SizeIncrement)
 	{
-		SizeTable[Index++] = FSizeTableEntry(Size, PlatformPageSize, Size / BasePageSize, BasePageSize, MinimumAlignment);
+		SizeTable[Index++] = FSizeTableEntry(Size, PlatformPageSize, Size / BasePageSize, BasePageSize);
 	}
 	check(Index < 256);
 	return (uint8)Index;
@@ -519,7 +518,7 @@ uint32 FBitTree::Slow_NextAllocBits(uint32 NumBits, uint64 StartIndex)
 	return MAX_uint32;
 }
 
-#endif
+uint32 FMallocBinnedCommonBase::OsAllocationGranularity = 0;
 
 float GMallocBinnedFlushThreadCacheMaxWaitTime = 0.2f;
 static FAutoConsoleVariableRef GMallocBinnedFlushThreadCacheMaxWaitTimeCVar(
@@ -537,26 +536,45 @@ static FAutoConsoleVariableRef GMallocBinnedFlushRegisteredThreadCachesOnOneThre
 
 #if UE_MBC_ALLOW_RUNTIME_TWEAKING
 
-int32 GMallocBinnedBundleSize = UE_DEFAULT_GMallocBinnedBundleSize;
-static FAutoConsoleVariableRef GMallocBinned3BundleSizeCVar(
-	TEXT("MallocBinned.BundleSize"),
-	GMallocBinnedBundleSize,
-	TEXT("Max size in bytes of per-block bundles used in the recycling process")
-);
-
-int32 GMallocBinnedBundleCount = UE_DEFAULT_GMallocBinnedBundleCount;
-static FAutoConsoleVariableRef GMallocBinned3BundleCountCVar(
-	TEXT("MallocBinned.BundleCount"),
-	GMallocBinnedBundleCount,
-	TEXT("Max count in blocks per-block bundles used in the recycling process")
-);
-
-#endif
+	int32 GMallocBinnedPerThreadCaches = UE_DEFAULT_GMallocBinnedPerThreadCaches;
+	static FAutoConsoleVariableRef GMallocBinned3PerThreadCachesCVar(
+		TEXT("MallocBinned.PerThreadCaches"),
+		GMallocBinnedPerThreadCaches,
+		TEXT("Enables per-thread caches of small (<= 32768 byte) allocations from FMallocBinned2/3")
+	);
+
+	int32 GMallocBinnedBundleSize = UE_DEFAULT_GMallocBinnedBundleSize;
+	static FAutoConsoleVariableRef GMallocBinned3BundleSizeCVar(
+		TEXT("MallocBinned.BundleSize"),
+		GMallocBinnedBundleSize,
+		TEXT("Max size in bytes of per-block bundles used in the recycling process")
+	);
+
+	int32 GMallocBinnedBundleCount = UE_DEFAULT_GMallocBinnedBundleCount;
+	static FAutoConsoleVariableRef GMallocBinned3BundleCountCVar(
+		TEXT("MallocBinned.BundleCount"),
+		GMallocBinnedBundleCount,
+		TEXT("Max count in blocks per-block bundles used in the recycling process")
+	);
+
+	int32 GMallocBinnedMaxBundlesBeforeRecycle = UE_DEFAULT_GMallocBinnedMaxBundlesBeforeRecycle;
+	static FAutoConsoleVariableRef GMallocBinned3MaxBundlesBeforeRecycleCVar(
+		TEXT("MallocBinned.BundleRecycleCount"),
+		GBinned3MaxBundlesBeforeRecycle,
+		TEXT("Number of freed bundles in the global recycler before it returns them to the system, per-block size. Limited by UE_DEFAULT_GBinned3MaxBundlesBeforeRecycle (currently 4)")
+	);
+
+#endif	//~UE_MBC_ALLOW_RUNTIME_TWEAKING
 
 uint32 FMallocBinnedCommonBase::BinnedTlsSlot = FPlatformTLS::InvalidTlsSlot;
 #if UE_MBC_ALLOCATOR_STATS
 	std::atomic<int64> FMallocBinnedCommonBase::TLSMemory(0);
 	std::atomic<int64> FMallocBinnedCommonBase::ConsolidatedMemory(0);
+
+	std::atomic<int64> FMallocBinnedCommonBase::AllocatedSmallPoolMemory(0); // memory that's requested to be allocated by the game
+
+	int64 FMallocBinnedCommonBase::PoolInfoMemory = 0;
+	int64 FMallocBinnedCommonBase::HashMemory = 0;
 #endif
 
 PRAGMA_RESTORE_UNSAFE_TYPECAST_WARNINGS
diff --git a/Engine/Source/Runtime/Core/Private/HAL/MallocBinnedGPU.cpp b/Engine/Source/Runtime/Core/Private/HAL/MallocBinnedGPU.cpp
--- a/Engine/Source/Runtime/Core/Private/HAL/MallocBinnedGPU.cpp
+++ b/Engine/Source/Runtime/Core/Private/HAL/MallocBinnedGPU.cpp
@@ -574,7 +574,7 @@ void FMallocBinnedGPU::InitMallocBinned()
 	if (ArenaParams.bUseStandardSmallPoolSizes)
 	{
 		SizeTable.AddZeroed(UE_MBC_NUM_LISTED_SMALL_POOLS + ArenaParams.MaxStandardPoolSize / ArenaParams.BasePageSize); // overestimate
-		ArenaParams.PoolCount = FSizeTableEntry::FillSizeTable(ArenaParams.AllocationGranularity, &SizeTable[0], ArenaParams.BasePageSize, ArenaParams.MinimumAlignment, ArenaParams.MaxStandardPoolSize, ArenaParams.BasePageSize);
+		ArenaParams.PoolCount = FSizeTableEntry::FillSizeTable(ArenaParams.AllocationGranularity, &SizeTable[0], ArenaParams.BasePageSize, ArenaParams.MaxStandardPoolSize, ArenaParams.BasePageSize);
 		SizeTable.RemoveAt(ArenaParams.PoolCount, SizeTable.Num() - ArenaParams.PoolCount);
 	}
 	else
@@ -587,12 +587,12 @@ void FMallocBinnedGPU::InitMallocBinned()
 		{
 			check(Size % 4096 == 0); // calculations are done assume 4k is the smallest page size we will ever see
 			check(Size / 4096 <= std::numeric_limits<uint8>::max())		// Make sure we don't try to allocate more pages than fits in our counter.
-			SizeTable.Emplace(Size, ArenaParams.AllocationGranularity, (uint8)(Size / 4096), ArenaParams.BasePageSize, ArenaParams.MinimumAlignment);
+			SizeTable.Emplace(Size, ArenaParams.AllocationGranularity, (uint8)(Size / 4096), ArenaParams.BasePageSize);
 		}
 		else
 		{
 			// it is difficult to test what would actually make a good bucket size here, wouldn't want a prime number, 33 for example because that would take 33 pages a slab
-			SizeTable.Emplace(Size, ArenaParams.AllocationGranularity, (uint8)1, ArenaParams.BasePageSize, ArenaParams.MinimumAlignment);
+			SizeTable.Emplace(Size, ArenaParams.AllocationGranularity, (uint8)1, ArenaParams.BasePageSize);
 		}
 		ArenaParams.PoolCount++;
 	}
diff --git a/Engine/Source/Runtime/Core/Public/HAL/MallocBinned2.h b/Engine/Source/Runtime/Core/Public/HAL/MallocBinned2.h
--- a/Engine/Source/Runtime/Core/Public/HAL/MallocBinned2.h
+++ b/Engine/Source/Runtime/Core/Public/HAL/MallocBinned2.h
@@ -25,23 +25,10 @@ struct FGenericMemoryStats;
 #endif
 
 #define UE_MB2_LARGE_ALLOC					65536		// Alignment of OS-allocated pointer - pool-allocated pointers will have a non-aligned pointer
-#define UE_MB2_MINIMUM_ALIGNMENT_SHIFT		4			// Alignment of bins, expressed as a shift
-#define UE_MB2_MINIMUM_ALIGNMENT			16			// Alignment of bins
-#define UE_MB2_MAXIMUM_ALIGNMENT			128
 #define UE_MB2_MAX_SMALL_POOL_SIZE			(32768-16)	// Maximum bin size in SmallBinSizes in cpp file
 #define UE_MB2_SMALL_POOL_COUNT				51
 
 
-#if !defined(UE_DEFAULT_GMallocBinnedPerThreadCaches)
-#	define UE_DEFAULT_GBinned2PerThreadCaches					1
-#else
-#	define UE_DEFAULT_GBinned2PerThreadCaches					UE_DEFAULT_GMallocBinnedPerThreadCaches
-#endif
-
-#define UE_DEFAULT_GBinned2AllocExtra							32
-#define UE_DEFAULT_GBinned2MaxBundlesBeforeRecycle				8
-#define UE_DEFAULT_GBinned2MoveOSFreesOffTimeCriticalThreads	1
-
 // When book keeping is at the end of FFreeBlock, MallocBinned2 cannot tell if the allocation comes from a large allocation (higher than 64KB, also named as "OSAllocation") 
 // or from VeryLargePageAllocator that fell back to FCachedOSPageAllocator. In both cases the allocation (large or small) might be aligned to 64KB.
 // bookKeeping at the end needs to be disabled if we want VeryLargePageAllocator to properly fallback to regular TCachedOSPageAllocator if needed
@@ -54,20 +41,6 @@ struct FGenericMemoryStats;
 #	define BINNED2_FORK_SUPPORT (UE_SERVER && (PLATFORM_UNIX || DEFAULT_SERVER_FAKE_FORKS))
 #endif
 
-
-#define UE_MB2_ALLOW_RUNTIME_TWEAKING UE_MBC_ALLOW_RUNTIME_TWEAKING
-#if UE_MB2_ALLOW_RUNTIME_TWEAKING
-	extern CORE_API int32 GBinned2PerThreadCaches;
-	extern CORE_API int32 GBinned2MaxBundlesBeforeRecycle;
-	extern CORE_API int32 GBinned2AllocExtra;
-	extern CORE_API int32 GBinned2MoveOSFreesOffTimeCriticalThreads;
-#else
-#	define GBinned2PerThreadCaches						UE_DEFAULT_GBinned2PerThreadCaches
-#	define GBinned2MaxBundlesBeforeRecycle				UE_DEFAULT_GBinned2MaxBundlesBeforeRecycle
-#	define GBinned2AllocExtra							UE_DEFAULT_GBinned2AllocExtra
-#	define GBinned2MoveOSFreesOffTimeCriticalThreads	UE_DEFAULT_GBinned2MoveOSFreesOffTimeCriticalThreads
-#endif
-
 #define UE_MB2_ALLOCATOR_STATS				UE_MBC_ALLOCATOR_STATS
 #define UE_MB2_ALLOCATOR_STATS_VALIDATION	(UE_MB2_ALLOCATOR_STATS && 0)
 
@@ -97,11 +70,48 @@ enum class EBlockCanary : uint8
 //
 // Optimized virtual memory allocator.
 //
-class FMallocBinned2 : public TMallocBinnedCommon<FMallocBinned2, UE_MB2_MINIMUM_ALIGNMENT, UE_MB2_MAXIMUM_ALIGNMENT, UE_MB2_MINIMUM_ALIGNMENT_SHIFT, UE_MB2_SMALL_POOL_COUNT, UE_MB2_MAX_SMALL_POOL_SIZE>
+class FMallocBinned2 : public TMallocBinnedCommon<FMallocBinned2, UE_MB2_SMALL_POOL_COUNT, UE_MB2_MAX_SMALL_POOL_SIZE>
 {
+	struct FFreeBlock;
+
+public:
+	struct FPoolInfo
+	{
+		enum class ECanary : uint16
+		{
+			Unassigned = 0x3941,
+			FirstFreeBlockIsOSAllocSize = 0x17ea,
+			FirstFreeBlockIsPtr = 0xf317
+		};
+
+		uint16      Taken;          // Number of allocated elements in this pool, when counts down to zero can free the entire pool	
+		ECanary     Canary;         // See ECanary
+		uint32      AllocSize;      // Number of bytes allocated
+		FFreeBlock* FirstFreeBlock; // Pointer to first free memory in this pool or the OS Allocation Size in bytes if this allocation is not binned
+		FPoolInfo*  Next;           // Pointer to next pool
+		FPoolInfo** PtrToPrevNext;  // Pointer to whichever pointer points to this pool
+
+		FPoolInfo();
+
+		void CheckCanary(ECanary ShouldBe) const;
+		void SetCanary(ECanary ShouldBe, bool bPreexisting, bool bGuaranteedToBeNew);
+
+		bool HasFreeBin() const;
+		void* AllocateBin();
+
+		SIZE_T GetOSRequestedBytes() const;
+		SIZE_T GetOsAllocatedBytes() const;
+		void SetOSAllocationSizes(SIZE_T InRequestedBytes, UPTRINT InAllocatedBytes);
+
+		void Link(FPoolInfo*& PrevNext);
+		void Unlink();
+
+	private:
+		void ExhaustPoolIfNecessary();
+	};
+
+private:
 	// Forward declares.
-	struct FPoolInfo;
-	using PoolHashBucket = TPoolHashBucket<FPoolInfo>;
 	struct Private;
 
 	/** Information about a piece of free memory. */
@@ -195,9 +205,6 @@ class FMallocBinned2 : public TMallocBinnedCommon<FMallocBinned2, UE_MB2_MINIMUM
 	// Pool tables for different pool sizes
 	FPoolTable SmallPoolTables[UE_MB2_SMALL_POOL_COUNT];
 
-	PoolHashBucket* HashBuckets;
-	PoolHashBucket* HashBucketFreeList;
-	uint64 NumPoolsPerPage;
 #if BINNED2_FORK_SUPPORT
 	EBlockCanary CurrentCanary = EBlockCanary::PreFork; // The value of the canary for pages we have allocated this side of the fork 
 	EBlockCanary OldCanary = EBlockCanary::PreFork;		// If we have forked, the value canary of old pages we should avoid touching 
@@ -215,8 +222,6 @@ class FMallocBinned2 : public TMallocBinnedCommon<FMallocBinned2, UE_MB2_MINIMUM
 	FPooledVirtualMemoryAllocator CachedOSPageAllocator;
 #endif
 
-	FCriticalSection Mutex;
-
 	FORCEINLINE bool IsOSAllocation(const void* Ptr) const
 	{
 #if UE_USE_VERYLARGEPAGEALLOCATOR && !PLATFORM_UNIX && !PLATFORM_ANDROID
@@ -276,7 +281,7 @@ public:
 		const bool bUseSmallPool = UseSmallAlloc(Size, Alignment);
 		if (bUseSmallPool)
 		{
-			FPerThreadFreeBlockLists* Lists = GBinned2PerThreadCaches ? FPerThreadFreeBlockLists::Get() : nullptr;
+			FPerThreadFreeBlockLists* Lists = GMallocBinnedPerThreadCaches ? FPerThreadFreeBlockLists::Get() : nullptr;
 			if (Lists)
 			{
 				const uint32 PoolIndex = BoundSizeToPoolIndex(Size, MemSizeToPoolIndex);
@@ -297,13 +302,13 @@ public:
 	FORCEINLINE static bool UseSmallAlloc(SIZE_T Size, uint32 Alignment)
 	{
 #if UE_USE_VERYLARGEPAGEALLOCATOR && UE_MB2_BOOKKEEPING_AT_THE_END_OF_LARGEBLOCK
-		if (Alignment > UE_MB2_MINIMUM_ALIGNMENT)
+		if (Alignment > UE_MBC_MIN_SMALL_POOL_ALIGNMENT)
 		{
 			Size = Align(Size, Alignment);
 		}
 		const bool bResult = (Size <= UE_MB2_MAX_SMALL_POOL_SIZE);
 #else
-		const bool bResult = ((Size <= UE_MB2_MAX_SMALL_POOL_SIZE) & (Alignment <= UE_MB2_MINIMUM_ALIGNMENT)); // one branch, not two
+		const bool bResult = ((Size <= UE_MB2_MAX_SMALL_POOL_SIZE) & (Alignment <= UE_MBC_MIN_SMALL_POOL_ALIGNMENT)); // one branch, not two
 #endif
 		return bResult;
 	}
@@ -355,16 +360,16 @@ public:
 	FORCEINLINE void* ReallocInline(void* Ptr, SIZE_T NewSize, uint32 Alignment) 
 	{
 #if UE_USE_VERYLARGEPAGEALLOCATOR && UE_MB2_BOOKKEEPING_AT_THE_END_OF_LARGEBLOCK
-		if (Alignment > UE_MB2_MINIMUM_ALIGNMENT && (NewSize <= UE_MB2_MAX_SMALL_POOL_SIZE))
+		if (Alignment > UE_MBC_MIN_SMALL_POOL_ALIGNMENT && (NewSize <= UE_MB2_MAX_SMALL_POOL_SIZE))
 		{
 			NewSize = Align(NewSize, Alignment);
 		}
 		if (NewSize <= UE_MB2_MAX_SMALL_POOL_SIZE)
 #else
-		if (NewSize <= UE_MB2_MAX_SMALL_POOL_SIZE && Alignment <= UE_MB2_MINIMUM_ALIGNMENT) // one branch, not two
+		if (NewSize <= UE_MB2_MAX_SMALL_POOL_SIZE && Alignment <= UE_MBC_MIN_SMALL_POOL_ALIGNMENT) // one branch, not two
 #endif
 		{
-			FPerThreadFreeBlockLists* Lists = GBinned2PerThreadCaches ? FPerThreadFreeBlockLists::Get() : nullptr;
+			FPerThreadFreeBlockLists* Lists = GMallocBinnedPerThreadCaches ? FPerThreadFreeBlockLists::Get() : nullptr;
 			if (Lists && (!Ptr || !IsOSAllocation(Ptr)))
 			{
 				uint32 BinSize = 0;
@@ -450,7 +455,7 @@ public:
 	{
 		if (!IsOSAllocation(Ptr))
 		{
-			FPerThreadFreeBlockLists* Lists = GBinned2PerThreadCaches ? FPerThreadFreeBlockLists::Get() : nullptr;
+			FPerThreadFreeBlockLists* Lists = GMallocBinnedPerThreadCaches ? FPerThreadFreeBlockLists::Get() : nullptr;
 			if (Lists)
 			{
 				FFreeBlock* BasePtr = GetPoolHeaderFromPointer(Ptr);
@@ -468,20 +473,23 @@ public:
 		FreeExternal(Ptr);
 	}
 
-	FORCEINLINE virtual bool GetAllocationSize(void *Ptr, SIZE_T &SizeOut) override
+	FORCEINLINE bool GetSmallAllocationSize(void* Ptr, SIZE_T& SizeOut) const
 	{
 		if (!IsOSAllocation(Ptr))
 		{
 			const FFreeBlock* Free = GetPoolHeaderFromPointer(Ptr);
-#if BINNED2_FORK_SUPPORT
-			if (Free->CanaryAndForkState == CurrentCanary || Free->CanaryAndForkState == OldCanary)
-#else
-			if (Free->CanaryAndForkState == CurrentCanary)
-#endif
-			{
-				SizeOut = Free->BinSize;
-				return true;
-			}
+			CanaryTest(Free);
+			SizeOut = Free->BinSize;
+			return true;
+		}
+		return false;
+	}
+
+	FORCEINLINE virtual bool GetAllocationSize(void *Ptr, SIZE_T &SizeOut) override
+	{
+		if (GetSmallAllocationSize(Ptr, SizeOut))
+		{
+			return true;
 		}
 		return GetAllocationSizeExternal(Ptr, SizeOut);
 	}
@@ -493,10 +501,6 @@ public:
 
 	CORE_API virtual bool ValidateHeap() override;
 	CORE_API virtual void Trim(bool bTrimThreadCaches) override;
-	CORE_API virtual void SetupTLSCachesOnCurrentThread() override;
-	CORE_API virtual void MarkTLSCachesAsUsedOnCurrentThread() override;
-	CORE_API virtual void MarkTLSCachesAsUnusedOnCurrentThread() override;
-	CORE_API virtual void ClearAndDisableTLSCachesOnCurrentThread() override;
 	CORE_API virtual const TCHAR* GetDescriptiveName() override;
 	CORE_API virtual void UpdateStats() override;
 	CORE_API virtual void OnMallocInitialized() override;
@@ -508,14 +512,21 @@ public:
 	CORE_API void* MallocExternalLarge(SIZE_T Size, uint32 Alignment);
 	CORE_API void* ReallocExternal(void* Ptr, SIZE_T NewSize, uint32 Alignment);
 	CORE_API void FreeExternal(void *Ptr);
-	CORE_API bool GetAllocationSizeExternal(void* Ptr, SIZE_T& SizeOut);
 
-	CORE_API void CanaryTest(const FFreeBlock* Block) const;
 	CORE_API void CanaryFail(const FFreeBlock* Block) const;
-
-#if UE_MB2_ALLOCATOR_STATS
-	CORE_API int64 GetTotalAllocatedSmallPoolMemory() const;
+	FORCEINLINE void CanaryTest(const FFreeBlock* Block) const
+	{
+#if BINNED2_FORK_SUPPORT
+		// When we support forking there are two valid canary values.
+		if (Block->CanaryAndForkState != CurrentCanary && Block->CanaryAndForkState != OldCanary)
+#else
+		if (Block->CanaryAndForkState != CurrentCanary)
 #endif
+		{
+			CanaryFail(Block);
+		}
+	}
+
 	CORE_API virtual void GetAllocatorStats( FGenericMemoryStats& out_Stats ) override;
 	/** Dumps current allocator stats to the log. */
 	CORE_API virtual void DumpAllocatorStats(class FOutputDevice& Ar) override;
@@ -523,14 +534,8 @@ public:
 	static CORE_API uint16 SmallBinSizesReversed[UE_MB2_SMALL_POOL_COUNT]; // this is reversed to get the smallest elements on our main cache line
 	static CORE_API FMallocBinned2* MallocBinned2;
 	static CORE_API uint32 PageSize;
-	static CORE_API uint32 OsAllocationGranularity;
 	// Mapping of sizes to small table indices
-	static CORE_API uint8 MemSizeToPoolIndex[1 + (UE_MB2_MAX_SMALL_POOL_SIZE >> UE_MB2_MINIMUM_ALIGNMENT_SHIFT)];
-
-	static void RegisterThreadFreeBlockLists(FPerThreadFreeBlockLists* FreeBlockLists);
-	static void UnregisterThreadFreeBlockLists(FPerThreadFreeBlockLists* FreeBlockLists);
-	static FCriticalSection& GetFreeBlockListsRegistrationMutex();
-	static TArray<FPerThreadFreeBlockLists*>& GetRegisteredFreeBlockLists();
+	static CORE_API uint8 MemSizeToPoolIndex[1 + (UE_MB2_MAX_SMALL_POOL_SIZE >> UE_MBC_MIN_SMALL_POOL_ALIGNMENT_SHIFT)];
 
 	static void* AllocateMetaDataMemory(SIZE_T Size);
 	static void FreeMetaDataMemory(void* Ptr, SIZE_T Size);
@@ -541,6 +546,8 @@ public:
 	}
 
 	void FreeBundles(FBundleNode* Bundles, uint32 PoolIndex);
+
+	void FlushCurrentThreadCacheInternal(bool bNewEpochOnly = false);
 };
 
 #define UE_MB2_INLINE (1)
diff --git a/Engine/Source/Runtime/Core/Public/HAL/MallocBinned3.h b/Engine/Source/Runtime/Core/Public/HAL/MallocBinned3.h
--- a/Engine/Source/Runtime/Core/Public/HAL/MallocBinned3.h
+++ b/Engine/Source/Runtime/Core/Public/HAL/MallocBinned3.h
@@ -28,12 +28,6 @@ struct FGenericMemoryStats;
 #endif
 
 #define UE_MB3_BASE_PAGE_SIZE						4096			// Minimum "page size" for binned3
-#define UE_MB3_MINIMUM_ALIGNMENT_SHIFT				4				// Alignment of bins, expressed as a shift
-#define UE_MB3_MINIMUM_ALIGNMENT					16				// Alignment of bins
-
-#ifndef BINNED3_MAX_SMALL_POOL_ALIGNMENT
-#	define BINNED3_MAX_SMALL_POOL_ALIGNMENT			128
-#endif
 
 
 #ifndef UE_MB3_MAX_SMALL_POOL_SIZE
@@ -62,26 +56,6 @@ struct FGenericMemoryStats;
 #	endif
 #endif
 
-#if !defined(UE_DEFAULT_GMallocBinnedPerThreadCaches)
-#	define UE_DEFAULT_GBinned3PerThreadCaches		1
-#else
-#	define UE_DEFAULT_GBinned3PerThreadCaches		UE_DEFAULT_GMallocBinnedPerThreadCaches
-#endif
-
-#define UE_DEFAULT_GBinned3AllocExtra				32
-#define UE_DEFAULT_GBinned3MaxBundlesBeforeRecycle	8
-
-#define UE_MB3_ALLOW_RUNTIME_TWEAKING				UE_MBC_ALLOW_RUNTIME_TWEAKING
-#if UE_MB3_ALLOW_RUNTIME_TWEAKING
-	extern CORE_API int32 GBinned3PerThreadCaches;
-	extern CORE_API int32 GBinned3MaxBundlesBeforeRecycle;
-	extern CORE_API int32 GBinned3AllocExtra;
-#else
-	#define GBinned3PerThreadCaches					UE_DEFAULT_GBinned3PerThreadCaches
-	#define GBinned3MaxBundlesBeforeRecycle			UE_DEFAULT_GBinned3MaxBundlesBeforeRecycle
-	#define GBinned3AllocExtra						UE_DEFAULT_GBinned3AllocExtra
-#endif
-
 #define UE_MB3_ALLOCATOR_STATS						UE_MBC_ALLOCATOR_STATS
 
 #if UE_MB3_ALLOCATOR_STATS
@@ -109,71 +83,52 @@ PRAGMA_DISABLE_UNSAFE_TYPECAST_WARNINGS
 //    Memory is allocated top down in a Block
 // 
 // 2. For large allocations we go directly to OS, unless UE_MB3_USE_CACHED_PAGE_ALLOCATOR_FOR_LARGE_ALLOCS is defined to 1
-//    Each allocation is handled via FPlatformVirtualMemoryBlock and information about it is stored in FPoolInfoLarge
-//    FPoolInfoLarge stores the original allocation size and how much memory was committed by the OS in case we need to realloc that allocation and there's enough tail waste to do that in place
-//    All FPoolInfoLarge are stored in the PoolHashBucket and HashBucketFreeList
+//    Each allocation is handled via FPlatformVirtualMemoryBlock and information about it is stored in FPoolInfo
+//    FPoolInfo stores the original allocation size and how much memory was committed by the OS in case we need to realloc that allocation and there's enough tail waste to do that in place
+//    All FPoolInfo are stored in the PoolHashBucket and HashBucketFreeList
 //
-class CORE_API FMallocBinned3 : public TMallocBinnedCommon<FMallocBinned3, UE_MB3_MINIMUM_ALIGNMENT, BINNED3_MAX_SMALL_POOL_ALIGNMENT, UE_MB3_MINIMUM_ALIGNMENT_SHIFT, UE_MB3_SMALL_POOL_COUNT, UE_MB3_MAX_SMALL_POOL_SIZE>
-{	
-	// Forward declarations
-	struct FPoolInfoLarge;
-	struct FPoolInfoSmall;
-	struct FPoolTable;
-	using PoolHashBucket = TPoolHashBucket<FPoolInfoLarge>;
-	struct Private;
-
-
-	/** Information about a piece of free memory. */
-	struct FFreeBlock
+class FMallocBinned3 : public TMallocBinnedCommon<FMallocBinned3, UE_MB3_SMALL_POOL_COUNT, UE_MB3_MAX_SMALL_POOL_SIZE>
+{
+public:
+	struct FPoolInfo
 	{
-		enum
+		enum class ECanary : uint32
 		{
-			CANARY_VALUE = 0xe7
+			LargeUnassigned = 0x39431234,
+			LargeAssigned = 0x17ea5678,
 		};
 
-		FORCEINLINE FFreeBlock(uint32 InBlockSize, uint32 InBinSize, uint8 InPoolIndex)
-			: BinSizeShifted(InBinSize >> UE_MB3_MINIMUM_ALIGNMENT_SHIFT)
-			, PoolIndex(InPoolIndex)
-			, Canary(CANARY_VALUE)
-			, NextFreeBlockIndex(MAX_uint32)
-		{
-			check(InPoolIndex < MAX_uint8 && (InBinSize >> UE_MB3_MINIMUM_ALIGNMENT_SHIFT) <= MAX_uint16);
-			NumFreeBins = InBlockSize / InBinSize;
-		}
+		FPoolInfo();
 
-		FORCEINLINE uint32 GetNumFreeBins() const
-		{
-			return NumFreeBins;
-		}
+		void CheckCanary(ECanary ShouldBe) const;
+		void SetCanary(ECanary ShouldBe, bool bPreexisting, bool bGuarnteedToBeNew);
 
-		FORCEINLINE bool IsCanaryOk() const
-		{
-			return Canary == FFreeBlock::CANARY_VALUE;
-		}
+		uint32 GetOSRequestedBytes() const;
+		uint32 GetOsCommittedBytes() const;
 
-		FORCEINLINE void CanaryTest() const
+		// And alias for GetOsCommittedBytes() to align with MB2 naming convention
+		uint32 GetOsAllocatedBytes() const
 		{
-			if (!IsCanaryOk())
-			{
-				CanaryFail();
-			}
-			//checkSlow(PoolIndex == BoundSizeToPoolIndex(uint32(BinSizeShifted) << UE_MB3_MINIMUM_ALIGNMENT_SHIFT));
+			return GetOsCommittedBytes();
 		}
-		void CanaryFail() const;
 
-		FORCEINLINE void* AllocateBin()
-		{
-			--NumFreeBins;
-			return (uint8*)this + NumFreeBins * (uint32(BinSizeShifted) << UE_MB3_MINIMUM_ALIGNMENT_SHIFT);
-		}
+		uint32 GetOsVMPages() const;
+		void SetOSAllocationSize(uint32 InRequestedBytes);
+		void SetOSAllocationSizes(uint32 InRequestedBytes, UPTRINT InCommittedBytes, uint32 InVMSizeDivVirtualSizeAlignment);
 
-		uint16 BinSizeShifted;		// Size of the bins that this list points to >> UE_MB3_MINIMUM_ALIGNMENT_SHIFT
-		uint8  PoolIndex;			// Index of this pool
-		uint8  Canary;				// Constant value of 0xe3
-		uint32 NumFreeBins;			// Number of consecutive free bins here, at least 1.
-		uint32 NextFreeBlockIndex;	// Next free block or MAX_uint32
+	private:
+		ECanary	Canary;
+		uint32 AllocSize;						// Number of bytes allocated
+		uint32 VMSizeDivVirtualSizeAlignment;	// Number of VM bytes allocated aligned for OS
+		uint32 CommitSize;						// Number of bytes committed by the OS
 	};
 
+private:
+	// Forward declarations
+	struct FPoolInfoSmall;
+	struct FPoolTable;
+	struct Private;
+
 	/** Pool table. */
 	struct FPoolTable
 	{
@@ -222,12 +177,6 @@ class CORE_API FMallocBinned3 : public TMallocBinnedCommon<FMallocBinned3, UE_MB
 
 	uint32 SmallPoolInfosPerPlatformPage;
 
-	PoolHashBucket* HashBuckets;			// Hash buckets for external allocations, reserved in constructor based on the platform constants like page size and virtual address high\low hints
-	PoolHashBucket* HashBucketFreeList;		// Hash buckets for allocations that were allocated outside of the platform constants virtual address high\low hints
-	uint64 NumLargePoolsPerPage;
-
-	FCriticalSection Mutex;
-
 #if !BINNED3_USE_SEPARATE_VM_PER_POOL
 	FORCEINLINE uint64 PoolIndexFromPtr(const void* Ptr) const		// returns a uint64 for it can also be used to check if it is an OS allocation
 	{
@@ -337,21 +286,21 @@ class CORE_API FMallocBinned3 : public TMallocBinnedCommon<FMallocBinned3, UE_MB
 
 public:
 
-	FMallocBinned3();
+	CORE_API FMallocBinned3();
 
-	virtual ~FMallocBinned3();
+	CORE_API virtual ~FMallocBinned3();
 
 	// FMalloc interface.
-	virtual bool IsInternallyThreadSafe() const override;
+	CORE_API virtual bool IsInternallyThreadSafe() const override;
 	FORCEINLINE virtual void* Malloc(SIZE_T Size, uint32 Alignment) override
 	{
 		void* Result = nullptr;
 	
 		// Only allocate from the small pools if the size is small enough and the alignment isn't crazy large.
 		// With large alignments, we'll waste a lot of memory allocating an entire page, but such alignments are highly unlikely in practice.
-		if ((Size <= UE_MB3_MAX_SMALL_POOL_SIZE) & (Alignment <= UE_MB3_MINIMUM_ALIGNMENT)) // one branch, not two
+		if ((Size <= UE_MB3_MAX_SMALL_POOL_SIZE) & (Alignment <= UE_MBC_MIN_SMALL_POOL_ALIGNMENT)) // one branch, not two
 		{
-			FPerThreadFreeBlockLists* Lists = GBinned3PerThreadCaches ? FPerThreadFreeBlockLists::Get() : nullptr;
+			FPerThreadFreeBlockLists* Lists = GMallocBinnedPerThreadCaches ? FPerThreadFreeBlockLists::Get() : nullptr;
 			if (Lists)
 			{
 				const uint32 PoolIndex = BoundSizeToPoolIndex(Size, MemSizeToPoolIndex);
@@ -376,9 +325,9 @@ public:
 
 	FORCEINLINE virtual void* Realloc(void* Ptr, SIZE_T NewSize, uint32 Alignment) override
 	{
-		if (NewSize <= UE_MB3_MAX_SMALL_POOL_SIZE && Alignment <= UE_MB3_MINIMUM_ALIGNMENT) // one branch, not two
+		if (NewSize <= UE_MB3_MAX_SMALL_POOL_SIZE && Alignment <= UE_MBC_MIN_SMALL_POOL_ALIGNMENT) // one branch, not two
 		{
-			FPerThreadFreeBlockLists* Lists = GBinned3PerThreadCaches ? FPerThreadFreeBlockLists::Get() : nullptr;
+			FPerThreadFreeBlockLists* Lists = GMallocBinnedPerThreadCaches ? FPerThreadFreeBlockLists::Get() : nullptr;
 
 			const uint64 PoolIndex = PoolIndexFromPtr(Ptr);
 			if ((!!Lists) & ((!Ptr) | (PoolIndex < UE_MB3_SMALL_POOL_COUNT)))
@@ -442,7 +391,7 @@ public:
 		const uint64 PoolIndex = PoolIndexFromPtr(Ptr);
 		if (PoolIndex < UE_MB3_SMALL_POOL_COUNT)
 		{
-			FPerThreadFreeBlockLists* Lists = GBinned3PerThreadCaches ? FPerThreadFreeBlockLists::Get() : nullptr;
+			FPerThreadFreeBlockLists* Lists = GMallocBinnedPerThreadCaches ? FPerThreadFreeBlockLists::Get() : nullptr;
 			if (Lists)
 			{
 				const int32 BinSize = PoolIndexToBinSize(PoolIndex);
@@ -459,7 +408,7 @@ public:
 		FreeExternal(Ptr);
 	}
 
-	FORCEINLINE virtual bool GetAllocationSize(void *Ptr, SIZE_T &SizeOut) override
+	FORCEINLINE bool GetSmallAllocationSize(void* Ptr, SIZE_T& SizeOut) const
 	{
 		const uint64 PoolIndex = PoolIndexFromPtr(Ptr);
 		if (PoolIndex < UE_MB3_SMALL_POOL_COUNT)
@@ -467,6 +416,16 @@ public:
 			SizeOut = PoolIndexToBinSize(PoolIndex);
 			return true;
 		}
+
+		return false;
+	}
+
+	FORCEINLINE virtual bool GetAllocationSize(void *Ptr, SIZE_T &SizeOut) override
+	{
+		if (GetSmallAllocationSize(Ptr, SizeOut))
+		{
+			return true;
+		}
 		return GetAllocationSizeExternal(Ptr, SizeOut);
 	}
 
@@ -475,37 +434,23 @@ public:
 		return QuantizeSizeCommon(Count, Alignment, *this);
 	}
 
-	virtual bool ValidateHeap() override;
-	virtual void Trim(bool bTrimThreadCaches) override;
-	virtual void SetupTLSCachesOnCurrentThread() override;
-	virtual void MarkTLSCachesAsUsedOnCurrentThread() override;
-	virtual void MarkTLSCachesAsUnusedOnCurrentThread() override;
-	virtual void ClearAndDisableTLSCachesOnCurrentThread() override;
-	virtual const TCHAR* GetDescriptiveName() override;
+	CORE_API virtual bool ValidateHeap() override;
+	CORE_API virtual void Trim(bool bTrimThreadCaches) override;
+	CORE_API virtual const TCHAR* GetDescriptiveName() override;
 	// End FMalloc interface.
 
-	void* MallocExternal(SIZE_T Size, uint32 Alignment);
-	void* ReallocExternal(void* Ptr, SIZE_T NewSize, uint32 Alignment);
-	void FreeExternal(void *Ptr);
-	bool GetAllocationSizeExternal(void* Ptr, SIZE_T& SizeOut);
+	CORE_API void* MallocExternal(SIZE_T Size, uint32 Alignment);
+	CORE_API void* ReallocExternal(void* Ptr, SIZE_T NewSize, uint32 Alignment);
+	CORE_API void FreeExternal(void *Ptr);
 
-#if UE_MB3_ALLOCATOR_STATS
-	int64 GetTotalAllocatedSmallPoolMemory() const;
-#endif
-	virtual void GetAllocatorStats( FGenericMemoryStats& out_Stats ) override;
+	CORE_API virtual void GetAllocatorStats( FGenericMemoryStats& out_Stats ) override;
 	/** Dumps current allocator stats to the log. */
-	virtual void DumpAllocatorStats(class FOutputDevice& Ar) override;
+	CORE_API virtual void DumpAllocatorStats(class FOutputDevice& Ar) override;
 	
 	// +1 enables PoolIndexToBinSize(~0u / -1) dummy access that helps avoid PoolIndex == 0 branching in Realloc(),
 	// see ((!PoolIndex) | (NewSize > PoolIndexToBinSize(static_cast<uint32>(PoolIndex - 1)))
 	static uint16 SmallBinSizesReversedShifted[UE_MB3_SMALL_POOL_COUNT + 1]; // this is reversed to get the smallest elements on our main cache line
 	static FMallocBinned3* MallocBinned3;
-	static uint32 OsAllocationGranularity;
-
-	static void RegisterThreadFreeBlockLists(FPerThreadFreeBlockLists* FreeBlockLists);
-	static void UnregisterThreadFreeBlockLists(FPerThreadFreeBlockLists* FreeBlockLists);
-	static FCriticalSection& GetFreeBlockListsRegistrationMutex();
-	static TArray<FPerThreadFreeBlockLists*>& GetRegisteredFreeBlockLists();
 
 #if !BINNED3_USE_SEPARATE_VM_PER_POOL
 	static uint8* Binned3BaseVMPtr;
@@ -517,11 +462,11 @@ public:
 	FPlatformMemory::FPlatformVirtualMemoryBlock PoolBaseVMBlock[UE_MB3_SMALL_POOL_COUNT];
 #endif
 	// Mapping of sizes to small table indices
-	static uint8 MemSizeToPoolIndex[1 + (UE_MB3_MAX_SMALL_POOL_SIZE >> UE_MB3_MINIMUM_ALIGNMENT_SHIFT)];
+	static uint8 MemSizeToPoolIndex[1 + (UE_MB3_MAX_SMALL_POOL_SIZE >> UE_MBC_MIN_SMALL_POOL_ALIGNMENT_SHIFT)];
 
 	FORCEINLINE uint32 PoolIndexToBinSize(uint32 PoolIndex) const
 	{
-		return uint32(SmallBinSizesReversedShifted[UE_MB3_SMALL_POOL_COUNT - PoolIndex - 1]) << UE_MB3_MINIMUM_ALIGNMENT_SHIFT;
+		return uint32(SmallBinSizesReversedShifted[UE_MB3_SMALL_POOL_COUNT - PoolIndex - 1]) << UE_MBC_MIN_SMALL_POOL_ALIGNMENT_SHIFT;
 	}
 
 	void FreeBundles(FBundleNode* Bundles, uint32 PoolIndex);
@@ -529,6 +474,8 @@ public:
 	void Commit(uint32 InPoolIndex, void *Ptr, SIZE_T Size);
 	void Decommit(uint32 InPoolIndex, void *Ptr, SIZE_T Size);
 
+	void FlushCurrentThreadCacheInternal(bool bNewEpochOnly = false);
+
 	static void* AllocateMetaDataMemory(SIZE_T Size);
 	static void FreeMetaDataMemory(void *Ptr, SIZE_T Size);
 };
diff --git a/Engine/Source/Runtime/Core/Public/HAL/MallocBinnedCommon.h b/Engine/Source/Runtime/Core/Public/HAL/MallocBinnedCommon.h
--- a/Engine/Source/Runtime/Core/Public/HAL/MallocBinnedCommon.h
+++ b/Engine/Source/Runtime/Core/Public/HAL/MallocBinnedCommon.h
@@ -3,17 +3,16 @@
 #pragma once
 
 #include "CoreTypes.h"
+#include "CoreGlobals.h"
+#include <atomic>
 #include "HAL/MemoryBase.h"
-#include "Math/UnrealMathUtility.h"
 #include "HAL/PlatformTLS.h"
+#include "HAL/CriticalSection.h"
 #include "Async/Mutex.h"
+#include "Misc/ScopeLock.h"
 #include "Templates/AlignmentTemplates.h"
+#include "Logging/LogMacros.h"
 
-#include <atomic>
 
-#if PLATFORM_HAS_FPlatformVirtualMemoryBlock
-
-#include "HAL/PlatformMemory.h"
-
 // A project can define it's own UE_MBC_MAX_LISTED_SMALL_POOL_SIZE and UE_MBC_NUM_LISTED_SMALL_POOLS to reduce runtime memory usage
 // MallocBinnedCommon.cpp has a list of predefined bins that go up to 28672
@@ -27,21 +26,67 @@
 // This, however, would require the app to redefine UE_MBC_NUM_LISTED_SMALL_POOLS too to match the number of bins that fall under the new define's threshold
 // In case of 16384, we'll skip 3 larger bins and so UE_MBC_NUM_LISTED_SMALL_POOLS should be set to 48 at the time of writing
 #if !defined(UE_MBC_MAX_LISTED_SMALL_POOL_SIZE)
-#	define UE_MBC_MAX_LISTED_SMALL_POOL_SIZE	28672
+#	define UE_MBC_MAX_LISTED_SMALL_POOL_SIZE			28672
 #endif
 
 #if !defined(UE_MBC_NUM_LISTED_SMALL_POOLS)
-#	define UE_MBC_NUM_LISTED_SMALL_POOLS	51
+#	define UE_MBC_NUM_LISTED_SMALL_POOLS				51
 #endif
 
 #if !defined(BINNEDCOMMON_USE_SEPARATE_VM_PER_POOL)
 #	if PLATFORM_WINDOWS
-#		define BINNEDCOMMON_USE_SEPARATE_VM_PER_POOL (1)
+#		define BINNEDCOMMON_USE_SEPARATE_VM_PER_POOL	1
 #	else
-#		define BINNEDCOMMON_USE_SEPARATE_VM_PER_POOL (0)
+#		define BINNEDCOMMON_USE_SEPARATE_VM_PER_POOL	0
 #	endif
 #endif
 
+#define UE_MBC_MIN_SMALL_POOL_ALIGNMENT					16	// Minimum alignment of bins
+#define UE_MBC_MIN_SMALL_POOL_ALIGNMENT_SHIFT			4	// Alignment of bins, expressed as a shift
+#define UE_MBC_MAX_SMALL_POOL_ALIGNMENT					128
 
+#if !defined(AGGRESSIVE_MEMORY_SAVING)
+#	error "AGGRESSIVE_MEMORY_SAVING must be defined"
+#endif
+
+#if AGGRESSIVE_MEMORY_SAVING
+#	define UE_DEFAULT_GMallocBinnedBundleSize			8192
+#else
+#	define UE_DEFAULT_GMallocBinnedBundleSize			65536
+#endif
+
+#if !defined(UE_DEFAULT_GMallocBinnedPerThreadCaches)
+#	define UE_DEFAULT_GMallocBinnedPerThreadCaches		1
+#endif
+#define UE_DEFAULT_GMallocBinnedBundleCount				64
+#define UE_DEFAULT_GMallocBinnedAllocExtra				32
+#define UE_DEFAULT_GMallocBinnedMaxBundlesBeforeRecycle	8
+
+#ifndef UE_MBC_ALLOW_RUNTIME_TWEAKING
+#	define UE_MBC_ALLOW_RUNTIME_TWEAKING				0
+#endif
+
+#if UE_MBC_ALLOW_RUNTIME_TWEAKING
+	extern CORE_API int32 GMallocBinnedPerThreadCaches;
+	extern CORE_API int32 GMallocBinnedBundleSize;
+	extern CORE_API int32 GMallocBinnedBundleCount;
+	extern CORE_API int32 GMallocBinnedAllocExtra;
+	extern CORE_API int32 GMallocBinnedMaxBundlesBeforeRecycle;
+#else
+#	define GMallocBinnedPerThreadCaches					UE_DEFAULT_GMallocBinnedPerThreadCaches
+#	define GMallocBinnedBundleSize						UE_DEFAULT_GMallocBinnedBundleSize
+#	define GMallocBinnedBundleCount						UE_DEFAULT_GMallocBinnedBundleSize
+#	define GMallocBinnedAllocExtra						UE_DEFAULT_GMallocBinnedAllocExtra
+#	define GMallocBinnedMaxBundlesBeforeRecycle			UE_DEFAULT_GMallocBinnedMaxBundlesBeforeRecycle
+#endif	//~UE_MBC_ALLOW_RUNTIME_TWEAKING
+
+#ifndef UE_MBC_ALLOCATOR_STATS
+#	define UE_MBC_ALLOCATOR_STATS (!UE_BUILD_SHIPPING || WITH_EDITOR)
+#endif
+
+extern CORE_API float GMallocBinnedFlushThreadCacheMaxWaitTime;
+extern CORE_API int32 GMallocBinnedFlushRegisteredThreadCachesOnOneThread;
+
+
 class FBitTree
 {
@@ -97,51 +142,35 @@
 	uint32 NumMemoryPagesPerBlock;
 
 	FSizeTableEntry() = default;
-	FSizeTableEntry(uint32 InBinSize, uint64 PlatformPageSize, uint8 Num4kbPages, uint32 BasePageSize, uint32 MinimumAlignment);
+	FSizeTableEntry(uint32 InBinSize, uint64 PlatformPageSize, uint8 Num4kbPages, uint32 BasePageSize);
 
 	bool operator<(const FSizeTableEntry& Other) const
 	{
 		return BinSize < Other.BinSize;
 	}
-	static uint8 FillSizeTable(uint64 PlatformPageSize, FSizeTableEntry* SizeTable, uint32 BasePageSize, uint32 MinimumAlignment, uint32 MaxSize, uint32 SizeIncrement);
+	static uint8 FillSizeTable(uint64 PlatformPageSize, FSizeTableEntry* SizeTable, uint32 BasePageSize, uint32 MaxSize, uint32 SizeIncrement);
 };
 
-#endif	//~PLATFORM_HAS_FPlatformVirtualMemoryBlock
 
-
-#if !defined(AGGRESSIVE_MEMORY_SAVING)
-#	error "AGGRESSIVE_MEMORY_SAVING must be defined"
-#endif
-
-#if AGGRESSIVE_MEMORY_SAVING
-#	define UE_DEFAULT_GMallocBinnedBundleSize 8192
-#else
-#	define UE_DEFAULT_GMallocBinnedBundleSize 65536
-#endif
-
-#define UE_DEFAULT_GMallocBinnedBundleCount 64
-
-#ifndef UE_MBC_ALLOW_RUNTIME_TWEAKING
-#	define UE_MBC_ALLOW_RUNTIME_TWEAKING 0
-#endif
-
-#if UE_MBC_ALLOW_RUNTIME_TWEAKING
-	extern CORE_API int32 GMallocBinnedBundleSize;
-	extern CORE_API int32 GMallocBinnedBundleCount;
-#else
-#	define GMallocBinnedBundleSize	UE_DEFAULT_GMallocBinnedBundleSize
-#	define GMallocBinnedBundleCount	UE_DEFAULT_GMallocBinnedBundleSize
-#endif
-
-#ifndef UE_MBC_ALLOCATOR_STATS
-#	define UE_MBC_ALLOCATOR_STATS (!UE_BUILD_SHIPPING || WITH_EDITOR)
-#endif
-
-extern CORE_API float GMallocBinnedFlushThreadCacheMaxWaitTime;
-extern CORE_API int32 GMallocBinnedFlushRegisteredThreadCachesOnOneThread;
-
 class FMallocBinnedCommonBase : public FMalloc
 {
+public:
+	// This needs to be small enough to fit inside the smallest allocation handled by MallocBinned2\3, hence the union.
+	struct FBundleNode
+	{
+		FBundleNode* NextNodeInCurrentBundle;
+
+		// NextBundle ptr is valid when node is stored in FFreeBlockList in a thread-local list of reusable allocations.
+		// Count is valid when node is stored in global recycler and caches the number of nodes in the list formed by NextNodeInCurrentBundle.
+		union
+		{
+			FBundleNode* NextBundle;
+			int32 Count;
+		};
+	};
+
+	static_assert(sizeof(FBundleNode) <= UE_MBC_MIN_SMALL_POOL_ALIGNMENT, "Bundle nodes must fit into the smallest block size");
+
 protected:
 	struct FPtrToPoolMapping
 	{
@@ -201,20 +230,7 @@
 		uint64 AddressSpaceBase;
 	};
 
-	// This needs to be small enough to fit inside the smallest allocation handled by MallocBinned2\3, hence the union.
-	struct FBundleNode
-	{
-		FBundleNode* NextNodeInCurrentBundle;
-
-		// NextBundle ptr is valid when node is stored in FFreeBlockList in a thread-local list of reusable allocations.
-		// Count is valid when node is stored in global recycler and caches the number of nodes in the list formed by NextNodeInCurrentBundle.
-		union
-		{
-			FBundleNode* NextBundle;
-			int32 Count;
-		};
-	};
-
+private:
 	struct FBundle
 	{
 		FORCEINLINE FBundle()
@@ -249,62 +265,6 @@
 		uint32       Count;
 	};
 
-	/** Hash table struct for retrieving allocation book keeping information */
-	template <class T>
-	struct TPoolHashBucket
-	{
-		UPTRINT			 BucketIndex;
-		T*				 FirstPool;
-		TPoolHashBucket* Prev;
-		TPoolHashBucket* Next;
-
-		TPoolHashBucket()
-		{
-			BucketIndex = 0;
-			FirstPool = nullptr;
-			Prev = this;
-			Next = this;
-		}
-
-		void Link(TPoolHashBucket* After)
-		{
-			After->Prev = Prev;
-			After->Next = this;
-			Prev->Next = After;
-			this->Prev = After;
-		}
-
-		void Unlink()
-		{
-			Next->Prev = Prev;
-			Prev->Next = Next;
-			Prev = this;
-			Next = this;
-		}
-	};
-
-	FPtrToPoolMapping PtrToPoolMapping;
-	static CORE_API uint32 BinnedTlsSlot;
-
-#if UE_MBC_ALLOCATOR_STATS
-	static std::atomic<int64> TLSMemory;
-	static std::atomic<int64> ConsolidatedMemory;
-#endif
-	std::atomic<uint64> MemoryTrimEpoch{ 0 };
-};
-
-template <class AllocType, int MinAlign, int MaxAlign, int MinAlignShift, int NumSmallPools, int MaxSmallPoolSize>
-class TMallocBinnedCommon : public FMallocBinnedCommonBase
-{
-	static_assert(sizeof(FBundleNode) <= MinAlign, "Bundle nodes must fit into the smallest block size");
-	friend class FMallocBinnedCommonUtils;
-
-	static constexpr int MIN_ALIGN           = MinAlign;
-	static constexpr int MAX_ALIGN           = MaxAlign;
-	static constexpr int MIN_ALIGN_SHIFT     = MinAlignShift;
-	static constexpr int NUM_SMALL_POOLS     = NumSmallPools;
-	static constexpr int MAX_SMALL_POOL_SIZE = MaxSmallPoolSize;
-	
 protected:
 	struct FFreeBlockList
 	{
@@ -411,6 +371,73 @@
 		FBundle FullBundle;
 	};
 
+	FPtrToPoolMapping PtrToPoolMapping;
+	uint64 NumPoolsPerPage;		// Number of AllocType::FPoolInfo
+	static uint32 OsAllocationGranularity;
+	FCriticalSection ExternalAllocMutex;
+
+	static CORE_API uint32 BinnedTlsSlot;
+
+#if UE_MBC_ALLOCATOR_STATS
+	static std::atomic<int64> TLSMemory;
+	static std::atomic<int64> ConsolidatedMemory;
+	static std::atomic<int64> AllocatedSmallPoolMemory;
+
+	static int64 PoolInfoMemory;
+	static int64 HashMemory;
+#endif
+
+	std::atomic<uint64> MemoryTrimEpoch{ 0 };
+
+	static CA_NO_RETURN void OutOfMemory(uint64 Size, uint32 Alignment = 0)
+	{
+		// this is expected not to return
+		FPlatformMemory::OnOutOfMemory(Size, Alignment);
+	}
+};
+
+template <class AllocType, int NumSmallPools, int MaxSmallPoolSize>
+class TMallocBinnedCommon : public FMallocBinnedCommonBase
+{
+	friend class FMallocBinnedCommonUtils;
+
+	static constexpr int NUM_SMALL_POOLS     = NumSmallPools;
+	static constexpr int MAX_SMALL_POOL_SIZE = MaxSmallPoolSize;
+
+	/** Hash table struct for retrieving allocation book keeping information */
+	struct FPoolHashBucket
+	{
+		UPTRINT			 BucketIndex;
+		AllocType::FPoolInfo* FirstPool;
+		FPoolHashBucket* Prev;
+		FPoolHashBucket* Next;
+
+		FPoolHashBucket()
+		{
+			BucketIndex = 0;
+			FirstPool = nullptr;
+			Prev = this;
+			Next = this;
+		}
+
+		void Link(FPoolHashBucket* After)
+		{
+			After->Prev = Prev;
+			After->Next = this;
+			Prev->Next = After;
+			this->Prev = After;
+		}
+
+		void Unlink()
+		{
+			Next->Prev = Prev;
+			Prev->Next = Next;
+			Prev = this;
+			Next = this;
+		}
+	};
+
+protected:
 	struct FPerThreadFreeBlockLists
 	{
 		FORCEINLINE static FPerThreadFreeBlockLists* Get() TSAN_SAFE
@@ -444,7 +471,7 @@
 				ThreadSingleton->bLockedByOwnerThread = true;
 				ThreadSingleton->Lock();
 				FPlatformTLS::SetTlsValue(BinnedTlsSlot, ThreadSingleton);
-				AllocType::RegisterThreadFreeBlockLists(ThreadSingleton);
+				RegisterThreadFreeBlockLists(ThreadSingleton);
 			}
 		}
 
@@ -478,7 +505,7 @@
 #if UE_MBC_ALLOCATOR_STATS
 				TLSMemory.fetch_sub(TLSSize, std::memory_order_relaxed);
 #endif
-				AllocType::UnregisterThreadFreeBlockLists(ThreadSingleton);
+				UnregisterThreadFreeBlockLists(ThreadSingleton);
 				ThreadSingleton->bLockedByOwnerThread = false;
 				ThreadSingleton->Unlock();
 				ThreadSingleton->~FPerThreadFreeBlockLists();
@@ -558,24 +585,212 @@
 	private:
 		UE::FMutex Mutex;
 		uint64 MemoryTrimEpoch = 0;
-		FFreeBlockList FreeLists[NumSmallPools];
+		FFreeBlockList FreeLists[NUM_SMALL_POOLS];
 		bool bLockedByOwnerThread = false;
 	};
 
+	struct Internal
+	{
+		using PoolInfo = typename AllocType::FPoolInfo;
+		/**
+		* Gets the PoolInfo for a large block memory address. If no valid info exists one is created.
+		*/
+		static PoolInfo* GetOrCreatePoolInfo(AllocType& Allocator, void* InPtr, PoolInfo::ECanary Kind, bool bPreexisting)
+		{
+			/**
+			* Creates an array of PoolInfo structures for tracking allocations.
+			*/
+			auto CreatePoolArray = [&Allocator](uint64 NumPools)
+				{
+					const uint64 PoolArraySize = NumPools * sizeof(PoolInfo);
+
+					void* Result = Allocator.AllocateMetaDataMemory(PoolArraySize);
+#if UE_MBC_ALLOCATOR_STATS
+					PoolInfoMemory += PoolArraySize;
+#endif
+					if (!Result)
+					{
+						FScopeUnlock TempUnlock(&Allocator.ExternalAllocMutex);
+						OutOfMemory(PoolArraySize);
+					}
+
+					DefaultConstructItems<PoolInfo>(Result, NumPools);
+					return (PoolInfo*)Result;
+				};
+
+			uint32  BucketIndex;
+			UPTRINT BucketIndexCollision;
+			uint32  PoolIndex;
+			Allocator.PtrToPoolMapping.GetHashBucketAndPoolIndices(InPtr, BucketIndex, BucketIndexCollision, PoolIndex);
+
+			FPoolHashBucket* FirstBucket = &Allocator.HashBuckets[BucketIndex];
+			FPoolHashBucket* Collision = FirstBucket;
+			do
+			{
+				if (!Collision->FirstPool)
+				{
+					Collision->BucketIndex = BucketIndexCollision;
+					Collision->FirstPool = CreatePoolArray(Allocator.NumPoolsPerPage);
+					Collision->FirstPool[PoolIndex].SetCanary(Kind, bPreexisting, true);
+					return &Collision->FirstPool[PoolIndex];
+				}
+
+				if (Collision->BucketIndex == BucketIndexCollision)
+				{
+					Collision->FirstPool[PoolIndex].SetCanary(Kind, bPreexisting, false);
+					return &Collision->FirstPool[PoolIndex];
+				}
+
+				Collision = Collision->Next;
+			} while (Collision != FirstBucket);
+
+			// Create a new hash bucket entry
+			if (!Allocator.HashBucketFreeList)
+			{
+				Allocator.HashBucketFreeList = (FPoolHashBucket*)Allocator.AllocateMetaDataMemory(AllocType::OsAllocationGranularity);
+#if UE_MBC_ALLOCATOR_STATS
+				HashMemory += AllocType::OsAllocationGranularity;
+#endif
+
+				for (UPTRINT i = 0, n = AllocType::OsAllocationGranularity / sizeof(FPoolHashBucket); i < n; ++i)
+				{
+					Allocator.HashBucketFreeList->Link(new (Allocator.HashBucketFreeList + i) FPoolHashBucket());
+				}
+			}
+
+			FPoolHashBucket* NextFree = Allocator.HashBucketFreeList->Next;
+			FPoolHashBucket* NewBucket = Allocator.HashBucketFreeList;
+
+			NewBucket->Unlink();
+
+			if (NextFree == NewBucket)
+			{
+				NextFree = nullptr;
+			}
+			Allocator.HashBucketFreeList = NextFree;
+
+			if (!NewBucket->FirstPool)
+			{
+				NewBucket->FirstPool = CreatePoolArray(Allocator.NumPoolsPerPage);
+				NewBucket->FirstPool[PoolIndex].SetCanary(Kind, bPreexisting, true);
+			}
+			else
+			{
+				NewBucket->FirstPool[PoolIndex].SetCanary(Kind, bPreexisting, false);
+			}
+
+			NewBucket->BucketIndex = BucketIndexCollision;
+
+			FirstBucket->Link(NewBucket);
+
+			return &NewBucket->FirstPool[PoolIndex];
+		}
+
+		static PoolInfo* FindPoolInfo(AllocType& Allocator, void* InPtr)
+		{
+			uint32  BucketIndex;
+			UPTRINT BucketIndexCollision;
+			uint32  PoolIndex;
+			Allocator.PtrToPoolMapping.GetHashBucketAndPoolIndices(InPtr, BucketIndex, BucketIndexCollision, PoolIndex);
+
+			FPoolHashBucket* FirstBucket = &Allocator.HashBuckets[BucketIndex];
+			FPoolHashBucket* Collision = FirstBucket;
+			do
+			{
+				if (Collision->BucketIndex == BucketIndexCollision)
+				{
+					return &Collision->FirstPool[PoolIndex];
+				}
+
+				Collision = Collision->Next;
+			} while (Collision != FirstBucket);
+
+			return nullptr;
+		}
+	};
+
+	static FCriticalSection& GetFreeBlockListsRegistrationMutex()
+	{
+		static FCriticalSection FreeBlockListsRegistrationMutex;
+		return FreeBlockListsRegistrationMutex;
+	}
+
+	static TArray<FPerThreadFreeBlockLists*>& GetRegisteredFreeBlockLists()
+	{
+		static TArray<FPerThreadFreeBlockLists*> RegisteredFreeBlockLists;
+		return RegisteredFreeBlockLists;
+	}
+
+	CORE_API virtual void SetupTLSCachesOnCurrentThread() override
+	{
+		//NOALLOC_SCOPE_CYCLE_COUNTER(STAT_FMallocBinned_SetupTLSCachesOnCurrentThread);
+
+		if (!UE_MBC_ALLOW_RUNTIME_TWEAKING && !GMallocBinnedPerThreadCaches)
+		{
+			return;
+		}
+		if (!FPlatformTLS::IsValidTlsSlot(BinnedTlsSlot))
+		{
+			BinnedTlsSlot = FPlatformTLS::AllocTlsSlot();
+		}
+		check(FPlatformTLS::IsValidTlsSlot(BinnedTlsSlot));
+		FPerThreadFreeBlockLists::SetTLS();
+	}
+
+	CORE_API virtual void ClearAndDisableTLSCachesOnCurrentThread() override
+	{
+		//NOALLOC_SCOPE_CYCLE_COUNTER(STAT_FMallocBinned_ClearTLSCachesOnCurrentThread);
+
+		if (!UE_MBC_ALLOW_RUNTIME_TWEAKING && !GMallocBinnedPerThreadCaches)
+		{
+			return;
+		}
+
+		((AllocType*)this)->FlushCurrentThreadCacheInternal();
+		FPerThreadFreeBlockLists::ClearTLS();
+	}
+
+	CORE_API virtual void MarkTLSCachesAsUsedOnCurrentThread() override
+	{
+		//NOALLOC_SCOPE_CYCLE_COUNTER(STAT_FMallocBinned_MarkTLSCachesAsUsedOnCurrentThread);
+
+		if (!UE_MBC_ALLOW_RUNTIME_TWEAKING && !GMallocBinnedPerThreadCaches)
+		{
+			return;
+		}
+
+		FPerThreadFreeBlockLists::LockTLS();
+	}
+
+	CORE_API virtual void MarkTLSCachesAsUnusedOnCurrentThread() override
+	{
+		//NOALLOC_SCOPE_CYCLE_COUNTER(STAT_FMallocBinned_MarkTLSCachesAsUnusedOnCurrentThread);
+
+		if (!UE_MBC_ALLOW_RUNTIME_TWEAKING && !GMallocBinnedPerThreadCaches)
+		{
+			return;
+		}
+
+		// Will only flush if memory trimming epoch has been bumped while the thread was active.
+		const bool bNewEpochOnly = true;
+		((AllocType*)this)->FlushCurrentThreadCacheInternal(bNewEpochOnly);
+		FPerThreadFreeBlockLists::UnlockTLS();
+	}
+
 	FORCEINLINE SIZE_T QuantizeSizeCommon(SIZE_T Count, uint32 Alignment, const AllocType& Alloc) const
 	{
-		static_assert(DEFAULT_ALIGNMENT <= MinAlign, "DEFAULT_ALIGNMENT is assumed to be zero"); // used below
+		static_assert(DEFAULT_ALIGNMENT <= UE_MBC_MIN_SMALL_POOL_ALIGNMENT, "DEFAULT_ALIGNMENT is assumed to be zero"); // used below
 		checkSlow(FMath::IsPowerOfTwo(Alignment));
 		SIZE_T SizeOut;
-		if ((Count <= MaxSmallPoolSize) & (Alignment <= MinAlign)) // one branch, not two
+		if ((Count <= MAX_SMALL_POOL_SIZE) & (Alignment <= UE_MBC_MIN_SMALL_POOL_ALIGNMENT)) // one branch, not two
 		{
 			SizeOut = Alloc.PoolIndexToBinSize(BoundSizeToPoolIndex(Count, Alloc.MemSizeToPoolIndex));
 			check(SizeOut >= Count);
 			return SizeOut;
 		}
-		Alignment = FMath::Max<uint32>(Alignment, MinAlign);
+		Alignment = FMath::Max<uint32>(Alignment, UE_MBC_MIN_SMALL_POOL_ALIGNMENT);
 		Count = Align(Count, Alignment);
-		if ((Count <= MaxSmallPoolSize) & (Alignment <= MaxAlign))
+		if ((Count <= MAX_SMALL_POOL_SIZE) & (Alignment <= UE_MBC_MAX_SMALL_POOL_ALIGNMENT))
 		{
 			uint32 PoolIndex = BoundSizeToPoolIndex(Count, Alloc.MemSizeToPoolIndex);
 			do
@@ -589,7 +804,7 @@
 				}
 
 				PoolIndex++;
-			} while (PoolIndex < NumSmallPools);
+			} while (PoolIndex < NUM_SMALL_POOLS);
 		}
 
 		Alignment = FPlatformMath::Max<uint32>(Alignment, Alloc.OsAllocationGranularity);
@@ -598,23 +813,23 @@
 		return SizeOut;
 	}
 
-	FORCEINLINE uint32 BoundSizeToPoolIndex(SIZE_T Size, const uint8(&MemSizeToPoolIndex)[1 + (MaxSmallPoolSize >> MinAlignShift)]) const
+	FORCEINLINE uint32 BoundSizeToPoolIndex(SIZE_T Size, const uint8(&MemSizeToPoolIndex)[1 + (MAX_SMALL_POOL_SIZE >> UE_MBC_MIN_SMALL_POOL_ALIGNMENT_SHIFT)]) const
 	{
-		const auto Index = ((Size + MinAlign - 1) >> MinAlignShift);
-		checkSlow(Index >= 0 && Index <= (MaxSmallPoolSize >> MinAlignShift)); // and it should be in the table
+		const auto Index = ((Size + UE_MBC_MIN_SMALL_POOL_ALIGNMENT - 1) >> UE_MBC_MIN_SMALL_POOL_ALIGNMENT_SHIFT);
+		checkSlow(Index >= 0 && Index <= (MAX_SMALL_POOL_SIZE >> UE_MBC_MIN_SMALL_POOL_ALIGNMENT_SHIFT)); // and it should be in the table
 		const uint32 PoolIndex = uint32(MemSizeToPoolIndex[Index]);
-		checkSlow(PoolIndex >= 0 && PoolIndex < NumSmallPools);
+		checkSlow(PoolIndex >= 0 && PoolIndex < NUM_SMALL_POOLS);
 		return PoolIndex;
 	}
 
 	bool PromoteToLargerBin(SIZE_T& Size, uint32& Alignment, const AllocType& Alloc) const
 	{
 		// try to promote our allocation request to a larger bin with a matching natural alignment
-		// if requested alignment is larger than MinAlign but smaller than MaxAlign
+		// if requested alignment is larger than UE_MBC_MIN_SMALL_POOL_ALIGNMENT but smaller than UE_MBC_MAX_SMALL_POOL_ALIGNMENT
 		// so we don't do a page allocation with a lot of memory waste
-		Alignment = FMath::Max<uint32>(Alignment, MinAlign);
+		Alignment = FMath::Max<uint32>(Alignment, UE_MBC_MIN_SMALL_POOL_ALIGNMENT);
 		const SIZE_T AlignedSize = Align(Size, Alignment);
-		if (UNLIKELY((AlignedSize <= MaxSmallPoolSize) && (Alignment <= MaxAlign)))
+		if (UNLIKELY((AlignedSize <= MAX_SMALL_POOL_SIZE) && (Alignment <= UE_MBC_MAX_SMALL_POOL_ALIGNMENT)))
 		{
 			uint32 PoolIndex = BoundSizeToPoolIndex(AlignedSize, Alloc.MemSizeToPoolIndex);
 			do
@@ -624,14 +839,95 @@
 				{
 					// we found a matching pool for our alignment and size requirements, so modify the size request to match
 					Size = SIZE_T(BlockSize);
-					Alignment = MinAlign;
+					Alignment = UE_MBC_MIN_SMALL_POOL_ALIGNMENT;
 					return true;
 				}
 
 				PoolIndex++;
-			} while (PoolIndex < NumSmallPools);
+			} while (PoolIndex < NUM_SMALL_POOLS);
 		}
 
 		return false;
 	}
+
+	bool GetAllocationSizeExternal(void* Ptr, SIZE_T& SizeOut)
+	{
+		if (((AllocType*)this)->GetSmallAllocationSize(Ptr, SizeOut))
+		{
+			return true;
+		}
+		if (!Ptr)
+		{
+			return false;
+		}
+
+		typename AllocType::FPoolInfo* Pool;
+		{
+			//NOALLOC_SCOPE_CYCLE_COUNTER(STAT_FMallocBinned_GetAllocationSizeExternal);
+			FScopeLock Lock(&ExternalAllocMutex);
+			Pool = Internal::FindPoolInfo(*(AllocType*)this, Ptr);
+		}
+
+		if (!Pool)
+		{
+			UE_LOG(LogMemory, Fatal, TEXT("FMallocBinned Attempt to GetAllocationSizeExternal an unrecognized pointer %p"), Ptr);
+		}
+		const SIZE_T PoolOsBytes = Pool->GetOsAllocatedBytes();
+		const SIZE_T PoolOSRequestedBytes = Pool->GetOSRequestedBytes();
+		checkf(PoolOSRequestedBytes <= PoolOsBytes, TEXT("FMallocBinned::GetAllocationSizeExternal %zu %zu"), PoolOSRequestedBytes, PoolOsBytes);
+		SizeOut = PoolOsBytes;
+		return true;
+	}
+
+#if UE_MBC_ALLOCATOR_STATS
+	int64 GetTotalAllocatedSmallPoolMemory() const
+	{
+		//NOALLOC_SCOPE_CYCLE_COUNTER(STAT_FMallocBinned_GetTotalAllocatedSmallPoolMemory);
+		int64 FreeBlockAllocatedMemory = 0;
+		{
+			FScopeLock Lock(&GetFreeBlockListsRegistrationMutex());
+			for (const FPerThreadFreeBlockLists* FreeBlockLists : GetRegisteredFreeBlockLists())
+			{
+				FreeBlockAllocatedMemory += FreeBlockLists->AllocatedMemory;
+			}
+			FreeBlockAllocatedMemory += ConsolidatedMemory.load(std::memory_order_relaxed);
+		}
+
+		return AllocatedSmallPoolMemory.load(std::memory_order_relaxed) + FreeBlockAllocatedMemory;
+	}
+#endif
+
+	void AllocateHashBuckets()
+	{
+		const uint64 MaxHashBuckets = PtrToPoolMapping.GetMaxHashBuckets();
+		const uint64 HashAllocSize = Align(MaxHashBuckets * sizeof(FPoolHashBucket), OsAllocationGranularity);
+		HashBuckets = (FPoolHashBucket*)AllocType::AllocateMetaDataMemory(HashAllocSize);
+#if UE_MBC_ALLOCATOR_STATS
+		HashMemory += HashAllocSize;
+#endif
+		verify(HashBuckets);
+
+		DefaultConstructItems<FPoolHashBucket>(HashBuckets, MaxHashBuckets);
+	}
+
+private:
+	FPoolHashBucket* HashBuckets = nullptr;				// Hash buckets for external allocations, reserved in constructor based on the platform constants like page size and virtual address high\low hints
+	FPoolHashBucket* HashBucketFreeList = nullptr;		// Hash buckets for allocations that were allocated outside of the platform constants virtual address high\low hints
+
+	static void RegisterThreadFreeBlockLists(FPerThreadFreeBlockLists* FreeBlockLists)
+	{
+		//NOALLOC_SCOPE_CYCLE_COUNTER(STAT_FMallocBinned_RegisterThreadFreeBlockLists);
+		FScopeLock Lock(&GetFreeBlockListsRegistrationMutex());
+		GetRegisteredFreeBlockLists().Add(FreeBlockLists);
+	}
+
+	static void UnregisterThreadFreeBlockLists(FPerThreadFreeBlockLists* FreeBlockLists)
+	{
+		//NOALLOC_SCOPE_CYCLE_COUNTER(STAT_FMallocBinned_UnregisterThreadFreeBlockLists);
+		FScopeLock Lock(&GetFreeBlockListsRegistrationMutex());
+		GetRegisteredFreeBlockLists().Remove(FreeBlockLists);
+#if UE_MBC_ALLOCATOR_STATS
+		ConsolidatedMemory.fetch_add(FreeBlockLists->AllocatedMemory, std::memory_order_relaxed);
+#endif
+	}
 };
diff --git a/Engine/Source/Runtime/Core/Public/HAL/MallocBinnedCommonUtils.h b/Engine/Source/Runtime/Core/Public/HAL/MallocBinnedCommonUtils.h
--- a/Engine/Source/Runtime/Core/Public/HAL/MallocBinnedCommonUtils.h
+++ b/Engine/Source/Runtime/Core/Public/HAL/MallocBinnedCommonUtils.h
@@ -7,6 +7,95 @@
 #include "HAL/MallocBinnedCommon.h"
 #include "Misc/App.h"
 #include "Stats/Stats.h"
+#include "FramePro/FrameProProfiler.h"
+
+
+#if FRAMEPRO_ENABLED
+//	Pushes a profiler scope if it's safe to do so without any new allocations.
+	class FNoAllocScopeCycleCounter
+	{
+	public:
+		FORCEINLINE FNoAllocScopeCycleCounter(const ANSICHAR* InStatString)
+		{
+			if (FFrameProProfiler::IsThreadContextReady() && GCycleStatsShouldEmitNamedEvents)
+			{
+				StatString = InStatString;
+				FFrameProProfiler::PushEvent(StatString);
+			}
+		}
+
+		FORCEINLINE ~FNoAllocScopeCycleCounter()
+		{
+			if (StatString)
+			{
+				FFrameProProfiler::PopEvent(StatString);
+			}
+		}
+	private:
+		const ANSICHAR* StatString = nullptr;
+	};
+
+#	define NOALLOC_SCOPE_CYCLE_COUNTER(Stat) FNoAllocScopeCycleCounter NoAllocCycleCounter_##Stat(#Stat)
+#else
+#	define NOALLOC_SCOPE_CYCLE_COUNTER(Stat)
+#endif	// ~FRAMEPRO_ENABLED
+
+namespace MallocBinnedPrivate
+{
+	template<int NumSmallPools>
+	struct TGlobalRecycler
+	{
+		bool PushBundle(uint32 InPoolIndex, FMallocBinnedCommonBase::FBundleNode* InBundle)
+		{
+			const uint32 NumCachedBundles = FMath::Min<uint32>(GMallocBinnedMaxBundlesBeforeRecycle, UE_DEFAULT_GMallocBinnedMaxBundlesBeforeRecycle);
+			for (uint32 Slot = 0; Slot < NumCachedBundles; Slot++)
+			{
+				if (!Bundles[InPoolIndex].FreeBundles[Slot])
+				{
+					if (!FPlatformAtomics::InterlockedCompareExchangePointer((void**)&Bundles[InPoolIndex].FreeBundles[Slot], InBundle, nullptr))
+					{
+						return true;
+					}
+				}
+			}
+			return false;
+		}
+
+		FMallocBinnedCommonBase::FBundleNode* PopBundle(uint32 InPoolIndex)
+		{
+			const uint32 NumCachedBundles = FMath::Min<uint32>(GMallocBinnedMaxBundlesBeforeRecycle, UE_DEFAULT_GMallocBinnedMaxBundlesBeforeRecycle);
+			for (uint32 Slot = 0; Slot < NumCachedBundles; Slot++)
+			{
+				FMallocBinnedCommonBase::FBundleNode* Result = Bundles[InPoolIndex].FreeBundles[Slot];
+				if (Result)
+				{
+					if (FPlatformAtomics::InterlockedCompareExchangePointer((void**)&Bundles[InPoolIndex].FreeBundles[Slot], nullptr, Result) == Result)
+					{
+						return Result;
+					}
+				}
+			}
+			return nullptr;
+		}
+
+	private:
+		struct FPaddedBundlePointer
+		{
+			FMallocBinnedCommonBase::FBundleNode* FreeBundles[UE_DEFAULT_GMallocBinnedMaxBundlesBeforeRecycle];
+#if (4 + (4 * PLATFORM_64BITS)) * UE_DEFAULT_GMallocBinnedMaxBundlesBeforeRecycle < PLATFORM_CACHE_LINE_SIZE
+#	define UE_MBC_BUNDLE_PADDING (PLATFORM_CACHE_LINE_SIZE - sizeof(FMallocBinnedCommonBase::FBundleNode*) * UE_DEFAULT_GMallocBinnedMaxBundlesBeforeRecycle)
+			uint8 Padding[UE_MBC_BUNDLE_PADDING];
+#endif
+			FPaddedBundlePointer()
+			{
+				DefaultConstructItems<FMallocBinnedCommonBase::FBundleNode*>(FreeBundles, UE_DEFAULT_GMallocBinnedMaxBundlesBeforeRecycle);
+			}
+		};
+		static_assert(sizeof(FPaddedBundlePointer) == PLATFORM_CACHE_LINE_SIZE, "FPaddedBundlePointer should be the same size as a cache line");
+		MS_ALIGN(PLATFORM_CACHE_LINE_SIZE) FPaddedBundlePointer Bundles[NumSmallPools] GCC_ALIGN(PLATFORM_CACHE_LINE_SIZE);
+	};
+}
+
 
 class FMallocBinnedCommonUtils
 {
@@ -41,11 +130,11 @@ public:
 
 				const double StartTimeInner = FPlatformTime::Seconds();
 				TrimThreadFreeBlockLists(Allocator, Lists);
-				const double WaitForMutexAndTrimTime = FPlatformTime::Seconds() - StartTimeInner;
+				const double WaitForTrimTime = FPlatformTime::Seconds() - StartTimeInner;
 
-				if (WaitForMutexAndTrimTime > GMallocBinnedFlushThreadCacheMaxWaitTime)
+				if (WaitForTrimTime > GMallocBinnedFlushThreadCacheMaxWaitTime)
 				{
-					UE_LOG(LogMemory, Warning, TEXT("FMalloc%s took %6.2fms to wait for mutex AND trim."), Allocator.GetDescriptiveName(), WaitForMutexAndTrimTime * 1000.0f);
+					UE_LOG(LogMemory, Warning, TEXT("FMalloc%s took %6.2fms to wait for mutex AND trim."), Allocator.GetDescriptiveName(), WaitForTrimTime * 1000.0f);
 				}
 			}
 		}
@@ -99,4 +188,4 @@ public:
 			FTaskGraphInterface::BroadcastSlow_OnlyUseForSpecialPurposes(FPlatformProcess::SupportsMultithreading() && FApp::ShouldUseThreadingForPerformance(), false, Broadcast);
 		}
 	}
-};
\ No newline at end of file
+};
