Fixed MB2 compilation on Linux

diff --git a/Engine/Source/Runtime/Core/Private/HAL/MallocBinned2.cpp b/Engine/Source/Runtime/Core/Private/HAL/MallocBinned2.cpp
--- a/Engine/Source/Runtime/Core/Private/HAL/MallocBinned2.cpp
+++ b/Engine/Source/Runtime/Core/Private/HAL/MallocBinned2.cpp
@@ -496,8 +496,8 @@ void FMallocBinned2::OnPreFork()
 	}
 
 #if !UE_USE_VERYLARGEPAGEALLOCATOR
-	FScopeLock Lock(&Mutex);
-	CachedOSPageAllocator.FreeAll(&Mutex);
+	FScopeLock Lock(&ExternalAllocMutex);
+	CachedOSPageAllocator.FreeAll(&ExternalAllocMutex);
 #endif
 
 #endif // BINNED2_FORK_SUPPORT
