Fixed MBC compilation

diff --git a/Engine/Source/Runtime/Core/Private/HAL/MallocBinnedCommon.cpp b/Engine/Source/Runtime/Core/Private/HAL/MallocBinnedCommon.cpp
--- a/Engine/Source/Runtime/Core/Private/HAL/MallocBinnedCommon.cpp
+++ b/Engine/Source/Runtime/Core/Private/HAL/MallocBinnedCommon.cpp
@@ -560,7 +560,7 @@ static FAutoConsoleVariableRef GMallocBinnedFlushRegisteredThreadCachesOnOneThre
 	int32 GMallocBinnedMaxBundlesBeforeRecycle = UE_DEFAULT_GMallocBinnedMaxBundlesBeforeRecycle;
 	static FAutoConsoleVariableRef GMallocBinned3MaxBundlesBeforeRecycleCVar(
 		TEXT("MallocBinned.BundleRecycleCount"),
-		GBinned3MaxBundlesBeforeRecycle,
+		GMallocBinnedMaxBundlesBeforeRecycle,
 		TEXT("Number of freed bundles in the global recycler before it returns them to the system, per-block size. Limited by UE_DEFAULT_GBinned3MaxBundlesBeforeRecycle (currently 4)")
 	);
 
