Fixed compilation in MBC on Android due to a missing typename for templated types

diff --git a/Engine/Source/Runtime/Core/Public/HAL/MallocBinnedCommon.h b/Engine/Source/Runtime/Core/Public/HAL/MallocBinnedCommon.h
--- a/Engine/Source/Runtime/Core/Public/HAL/MallocBinnedCommon.h
+++ b/Engine/Source/Runtime/Core/Public/HAL/MallocBinnedCommon.h
@@ -408,7 +408,7 @@ class TMallocBinnedCommon : public FMallocBinnedCommonBase
 	struct FPoolHashBucket
 	{
 		UPTRINT			 BucketIndex;
-		AllocType::FPoolInfo* FirstPool;
+		typename AllocType::FPoolInfo* FirstPool;
 		FPoolHashBucket* Prev;
 		FPoolHashBucket* Next;
 
@@ -595,7 +595,7 @@ protected:
 		/**
 		* Gets the PoolInfo for a large block memory address. If no valid info exists one is created.
 		*/
-		static PoolInfo* GetOrCreatePoolInfo(AllocType& Allocator, void* InPtr, PoolInfo::ECanary Kind, bool bPreexisting)
+		static PoolInfo* GetOrCreatePoolInfo(AllocType& Allocator, void* InPtr, typename PoolInfo::ECanary Kind, bool bPreexisting)
 		{
 			/**
 			* Creates an array of PoolInfo structures for tracking allocations.
