Fixed linker error in MallocBinnedCommon due to a missing CVar

diff --git a/Engine/Source/Runtime/Core/Private/HAL/MallocBinnedCommon.cpp b/Engine/Source/Runtime/Core/Private/HAL/MallocBinnedCommon.cpp
--- a/Engine/Source/Runtime/Core/Private/HAL/MallocBinnedCommon.cpp
+++ b/Engine/Source/Runtime/Core/Private/HAL/MallocBinnedCommon.cpp
@@ -537,28 +537,35 @@ static FAutoConsoleVariableRef GMallocBinnedFlushRegisteredThreadCachesOnOneThre
 #if UE_MBC_ALLOW_RUNTIME_TWEAKING
 
 	int32 GMallocBinnedPerThreadCaches = UE_DEFAULT_GMallocBinnedPerThreadCaches;
-	static FAutoConsoleVariableRef GMallocBinned3PerThreadCachesCVar(
+	static FAutoConsoleVariableRef GMallocBinnedPerThreadCachesCVar(
 		TEXT("MallocBinned.PerThreadCaches"),
 		GMallocBinnedPerThreadCaches,
 		TEXT("Enables per-thread caches of small (<= 32768 byte) allocations from FMallocBinned2/3")
 	);
 
 	int32 GMallocBinnedBundleSize = UE_DEFAULT_GMallocBinnedBundleSize;
-	static FAutoConsoleVariableRef GMallocBinned3BundleSizeCVar(
+	static FAutoConsoleVariableRef GMallocBinnedBundleSizeCVar(
 		TEXT("MallocBinned.BundleSize"),
 		GMallocBinnedBundleSize,
 		TEXT("Max size in bytes of per-block bundles used in the recycling process")
 	);
 
 	int32 GMallocBinnedBundleCount = UE_DEFAULT_GMallocBinnedBundleCount;
-	static FAutoConsoleVariableRef GMallocBinned3BundleCountCVar(
+	static FAutoConsoleVariableRef GMallocBinnedBundleCountCVar(
 		TEXT("MallocBinned.BundleCount"),
 		GMallocBinnedBundleCount,
 		TEXT("Max count in blocks per-block bundles used in the recycling process")
 	);
 
+	int32 GMallocBinnedAllocExtra = UE_DEFAULT_GMallocBinnedAllocExtra;
+	static FAutoConsoleVariableRef GMallocBinnedAllocExtraCVar(
+		TEXT("MallocBinned.AllocExtra"),
+		GMallocBinnedAllocExtra,
+		TEXT("When we do acquire the lock, how many bins cached in TLS caches. In no case will we grab more than a page.")
+	);
+
 	int32 GMallocBinnedMaxBundlesBeforeRecycle = UE_DEFAULT_GMallocBinnedMaxBundlesBeforeRecycle;
-	static FAutoConsoleVariableRef GMallocBinned3MaxBundlesBeforeRecycleCVar(
+	static FAutoConsoleVariableRef GMallocBinnedMaxBundlesBeforeRecycleCVar(
 		TEXT("MallocBinned.BundleRecycleCount"),
 		GMallocBinnedMaxBundlesBeforeRecycle,
 		TEXT("Number of freed bundles in the global recycler before it returns them to the system, per-block size. Limited by UE_DEFAULT_GBinned3MaxBundlesBeforeRecycle (currently 4)")
