TRingBuffer: Minor comments and spelling fixes

diff --git a/Engine/Source/Runtime/Core/Public/Containers/RingBuffer.h b/Engine/Source/Runtime/Core/Public/Containers/RingBuffer.h
--- a/Engine/Source/Runtime/Core/Public/Containers/RingBuffer.h
+++ b/Engine/Source/Runtime/Core/Public/Containers/RingBuffer.h
@@ -109,7 +109,7 @@ public:
 	void RemoveCurrent()
 	{
 		Container.RemoveAt(Index);
-		Index--;
+		--Index;
 	}
 
 	FORCEINLINE bool operator==(const TRingBufferIterator& Rhs) const { return &Container == &Rhs.Container && Index == Rhs.Index; }
@@ -124,10 +124,10 @@ private:
 /**
  * RingBuffer - an array with a Front and Back pointer and with implicit wraparound to the beginning of the array when reaching the end of the array when iterating from Front to Back
  * Useful for providing O(1) push/pop at the end of the array (for Queue or Stack) while still having high cache coherency during iteration.
- * Not threadsafe; caller must ensure there is no simultaneous access from multiple threads.
+ * Not thread-safe; caller must ensure there is no simultaneous access from multiple threads.
  *
  * Implementation Details:
- * Relies on unsigned arithmetics and ever increasing Front and Back indices to avoid having to store an extra element or maintain explicit empty state.
+ * Relies on unsigned arithmetics and ever-increasing Front and Back indices to avoid having to store an extra element or maintain explicit empty state.
  * Capacity will always be rounded up to the next power of two, to provide rapid masking of the index.
  */
 template<typename T, typename AllocatorT = FDefaultAllocator>
@@ -148,7 +148,7 @@ public:
 private:
 	/**
 	 * Type used for variables that are indexes into the underlying storage.
-	 * StorageModuloType types are offsets from the beginning of storage, and any bits outside of IndexMask are discarded.
+	 * StorageModuloType types are offsets from the beginning of storage, and any bits outside IndexMask are discarded.
 	 * They may have added on multiples of the capacity due to wrapping around, but will be interpreted as pointing to the value at (value & IndexMask).
 	 * StorageModuloTypes are also allowed to underflow/overflow their integer storage type; the only constraint is that X - Front <= Capacity for all valid indexes and for AfterBack.
 	 */
@@ -624,7 +624,7 @@ public:
 		return Index;
 	}
 
-	/** Remove the value at the given index from the RingBuffer, and shift values ahead or behind it into its location to fill the hole.  It is valid to call with Index outside of the range of the array; does nothing in that case. */
+	/** Remove the value at the given index from the RingBuffer, and shift values ahead or behind it into its location to fill the hole.  It is valid to call with Index outside the range of the array; does nothing in that case. */
 	void RemoveAt(IndexType Index)
 	{
 		RangeCheck(Index);
@@ -658,9 +658,9 @@ public:
 	}
 
 	/**
-	 * Removes as many instances of Item as there are in the array, maintaining order but not indices.
+	 * Removes all items for which a given predicate applies, maintaining order but not indices.
 	 *
-	 * @param Item Item to remove from array.
+	 * @param Predicate Any item for which the predicate returns true is removed.
 	 * @returns Number of removed elements.
 	 */
 	template <typename PredicateType>
@@ -787,7 +787,7 @@ private:
 		StorageModuloType NewIndexMask;
 		if (NewCapacity > 0)
 		{
-			NewData = reinterpret_cast<ElementType*>(FMemory::Malloc(sizeof(ElementType) * NewCapacity, alignof(ElementType)));
+			NewData = static_cast<ElementType*>(FMemory::Malloc(sizeof(ElementType) * NewCapacity, alignof(ElementType)));
 			NewIndexMask = NewCapacity - 1;
 			if (SrcNum > 0)
 			{
@@ -895,7 +895,7 @@ private:
 	 * RangeLast may be greater than or less than RangeFirst; the shift of the other elements will occur in the appropriate direction from RangeFirst down to RangeLast.
 	 * RangeFirst and RangeLast are given in StorageModulo space.
 	 * RangeDirection is -1 or 1 and indicates which direction we should move in StorageModulo space to iterate from RangeLast to RangeFirst.
-	 * Required constraint: (Rangelast - RangeFirst) <= capacity for RangeDirection == -1, and (RangeFirst - RangeLast) <= capacity for RangeDirection == 1
+	 * Required constraint: (RangeLast - RangeFirst) <= capacity for RangeDirection == -1, and (RangeFirst - RangeLast) <= capacity for RangeDirection == 1
 	 */
 	void ShiftLastToFirst(StorageModuloType RangeFirst, StorageModuloType RangeLast, int RangeDirection)
 	{
