Optimize MB3 memory usage when AGGRESSIVE_MEMORY_SAVING is set to 1

diff --git a/Engine/Source/Runtime/Core/Private/HAL/MallocBinnedCommon.cpp b/Engine/Source/Runtime/Core/Private/HAL/MallocBinnedCommon.cpp
--- a/Engine/Source/Runtime/Core/Private/HAL/MallocBinnedCommon.cpp
+++ b/Engine/Source/Runtime/Core/Private/HAL/MallocBinnedCommon.cpp
@@ -117,7 +117,7 @@ static inline void FillTable(FSizeTableEntry* SizeTable, int32& Index, const uin
 {
 	for (uint32 Sub = 0; Sub < BinListSize; Sub++)
 	{
-		// if we override UE_MBC_MAX_LISTED_SMALL_POOL_SIZE externally, we need to filter out predefined bins of a larger size
+		// UE_MBC_MAX_LISTED_SMALL_POOL_SIZE will differ depending on the value AGGRESSIVE_MEMORY_SAVING
 		if (BinsList[Sub] <= UE_MBC_MAX_LISTED_SMALL_POOL_SIZE)
 		{
 			SizeTable[Index++] = FSizeTableEntry(BinsList[Sub], PlatformPageSize, Num4kbPages, BasePageSize);
diff --git a/Engine/Source/Runtime/Core/Public/HAL/MallocBinned3.h b/Engine/Source/Runtime/Core/Public/HAL/MallocBinned3.h
--- a/Engine/Source/Runtime/Core/Public/HAL/MallocBinned3.h
+++ b/Engine/Source/Runtime/Core/Public/HAL/MallocBinned3.h
@@ -30,12 +30,10 @@ struct FGenericMemoryStats;
 #define UE_MB3_BASE_PAGE_SIZE						4096			// Minimum "page size" for binned3
 
 
-#ifndef UE_MB3_MAX_SMALL_POOL_SIZE
-#	if UE_MB3_USE_CACHED_PAGE_ALLOCATOR_FOR_LARGE_ALLOCS
-#		define UE_MB3_MAX_SMALL_POOL_SIZE			(UE_MBC_MAX_LISTED_SMALL_POOL_SIZE)	// Maximum small bin size
-#	else
-#		define UE_MB3_MAX_SMALL_POOL_SIZE			(128 * 1024)	// Maximum small bin size
-#	endif
+#if AGGRESSIVE_MEMORY_SAVING || UE_MB3_USE_CACHED_PAGE_ALLOCATOR_FOR_LARGE_ALLOCS
+#	define UE_MB3_MAX_SMALL_POOL_SIZE				(UE_MBC_MAX_LISTED_SMALL_POOL_SIZE)	// Maximum small bin size
+#else
+#	define UE_MB3_MAX_SMALL_POOL_SIZE				(128 * 1024)	// Maximum small bin size
 #endif
 #define UE_MB3_SMALL_POOL_COUNT						(UE_MBC_NUM_LISTED_SMALL_POOLS + (UE_MB3_MAX_SMALL_POOL_SIZE - UE_MBC_MAX_LISTED_SMALL_POOL_SIZE) / UE_MB3_BASE_PAGE_SIZE)
 
diff --git a/Engine/Source/Runtime/Core/Public/HAL/MallocBinnedCommon.h b/Engine/Source/Runtime/Core/Public/HAL/MallocBinnedCommon.h
--- a/Engine/Source/Runtime/Core/Public/HAL/MallocBinnedCommon.h
+++ b/Engine/Source/Runtime/Core/Public/HAL/MallocBinnedCommon.h
@@ -14,22 +14,15 @@
 #include "Logging/LogMacros.h"
 
 
-// A project can define it's own UE_MBC_MAX_LISTED_SMALL_POOL_SIZE and UE_MBC_NUM_LISTED_SMALL_POOLS to reduce runtime memory usage
-// MallocBinnedCommon.cpp has a list of predefined bins that go up to 28672
-// By default allocators (i.e. MB3) that use these bins will rely on this number as a baseline for a small bins count
-// These allocators can increase the amount of small bins they want to manage by going over the default MBC bins list
-// In MB3 case that means defining BINNED3_MAX_SMALL_POOL_SIZE to something like 65536
-// Every bin over the UE_MBC_MAX_LISTED_SMALL_POOL_SIZE would come with a 4kb increment
-// These small bins would be kept in user mode, increasing application's memory footprint and reducing the time it takes to allocate memory from the said bins
-// If application needs to aggressively reduce it's memory footprint, potentially trading some perf due to an increased amount of kernel calls to allocate memory
-// it can redefine UE_MBC_MAX_LISTED_SMALL_POOL_SIZE and BINNED3_MAX_SMALL_POOL_SIZE to smaller numbers, the good value is 16384 for both
-// This, however, would require the app to redefine UE_MBC_NUM_LISTED_SMALL_POOLS too to match the number of bins that fall under the new define's threshold
-// In case of 16384, we'll skip 3 larger bins and so UE_MBC_NUM_LISTED_SMALL_POOLS should be set to 48 at the time of writing
-#if !defined(UE_MBC_MAX_LISTED_SMALL_POOL_SIZE)
-#	define UE_MBC_MAX_LISTED_SMALL_POOL_SIZE			28672
+#if !defined(AGGRESSIVE_MEMORY_SAVING)
+#	error "AGGRESSIVE_MEMORY_SAVING must be defined"
 #endif
 
-#if !defined(UE_MBC_NUM_LISTED_SMALL_POOLS)
+#if AGGRESSIVE_MEMORY_SAVING
+#	define UE_MBC_MAX_LISTED_SMALL_POOL_SIZE			16384
+#	define UE_MBC_NUM_LISTED_SMALL_POOLS				48
+#else
+#	define UE_MBC_MAX_LISTED_SMALL_POOL_SIZE			28672
 #	define UE_MBC_NUM_LISTED_SMALL_POOLS				51
 #endif
 
