Add missing include

#rnx

- This header could give a compiler error, depending on include order due to the use of TSizedHeapAllocator.

diff --git a/Engine/Source/Runtime/Core/Public/Containers/ContainerAllocationPolicies.h b/Engine/Source/Runtime/Core/Public/Containers/ContainerAllocationPolicies.h
--- a/Engine/Source/Runtime/Core/Public/Containers/ContainerAllocationPolicies.h
+++ b/Engine/Source/Runtime/Core/Public/Containers/ContainerAllocationPolicies.h
@@ -4,6 +4,7 @@
 
 #include "CoreTypes.h"
 #include "Containers/ContainerHelpers.h"
+#include "Containers/ContainerAllocationPolicies.h"
 #include "HAL/PlatformMath.h"
 #include "HAL/UnrealMemory.h"
 #include "Math/NumericLimits.h"
