[TextKey] Make TextKey SlabSize configurable with UE_TEXTKEY_SLAB_ALLOCATOR_SLAB_SIZE and bump the default size to 32k (64kb when using widechar). In tests this has saved up to 20Mb of memory.

[RN] TextKey SlabSize is now configurable with UE_TEXTKEY_SLAB_ALLOCATOR_SLAB_SIZE

diff --git a/Engine/Source/Runtime/Core/Private/Internationalization/TextKey.cpp b/Engine/Source/Runtime/Core/Private/Internationalization/TextKey.cpp
--- a/Engine/Source/Runtime/Core/Private/Internationalization/TextKey.cpp
+++ b/Engine/Source/Runtime/Core/Private/Internationalization/TextKey.cpp
@@ -19,6 +19,10 @@ DEFINE_LOG_CATEGORY_STATIC(LogTextKey, Log, All);
 	#define UE_TEXTKEY_USE_SLAB_ALLOCATOR (1)
 #endif
 
+#ifndef UE_TEXTKEY_SLAB_ALLOCATOR_SLAB_SIZE
+	#define UE_TEXTKEY_SLAB_ALLOCATOR_SLAB_SIZE (32768)
+#endif
+
 class FTextKeyState
 {
 public:
@@ -206,7 +210,7 @@ private:
 		}
 
 	private:
-		static const int32 SlabSizeInElements = 8192;
+		static const int32 SlabSizeInElements = UE_TEXTKEY_SLAB_ALLOCATOR_SLAB_SIZE;
 
 		struct FStringSlab
 		{
@@ -249,7 +253,7 @@ private:
 			// If no slabs have space then just allocate a new one
 			checkf(NumSlabElementsNeeded <= SlabSizeInElements, TEXT("Tried to allocate a FTextKey string of %d elements, which is larger than the allowed slab size of %d elements!"), NumSlabElementsNeeded, SlabSizeInElements);
 			FStringSlab& Slab = Slabs.AddDefaulted_GetRef();
-			Slab.Allocation = reinterpret_cast<TCHAR*>(FMemory::Malloc(SlabSizeInElements * sizeof(TCHAR)));
+			Slab.Allocation = reinterpret_cast<TCHAR*>(FMemory::Malloc(SlabSizeInElements * sizeof(TCHAR), 1));
 			return Slab;
 		}
 
