Use transactionally safe scope locks in ConsoleManager

diff --git a/Engine/Source/Runtime/Core/Private/HAL/ConsoleManager.cpp b/Engine/Source/Runtime/Core/Private/HAL/ConsoleManager.cpp
--- a/Engine/Source/Runtime/Core/Private/HAL/ConsoleManager.cpp
+++ b/Engine/Source/Runtime/Core/Private/HAL/ConsoleManager.cpp
@@ -20,6 +20,7 @@ ConsoleManager.cpp: console command handling
 #include "Misc/CoreDelegates.h"
 #include "Misc/OutputDeviceFile.h"
 #include "Misc/RemoteConfigIni.h"
+#include "Misc/TransactionallySafeScopeLock.h"
 #include "Modules/ModuleManager.h"
 #include "HAL/PlatformProcess.h"
 #include "ProfilingDebugging/CsvProfiler.h"
@@ -147,7 +148,7 @@ public:
 		* We don't aim to solve all concurrency problems (for example registering and unregistering a cvar on different threads, or reading a cvar from one thread while writing it from a different thread).
 		* Rather we just ensure that operations on a cvar from one thread will not conflict with operations on another cvar from another thread.
 	**/
-	mutable FCriticalSection ConsoleObjectsSynchronizationObject;
+	mutable FTransactionallySafeCriticalSection ConsoleObjectsSynchronizationObject;
 
 	/**
 	 * @param Name must not be 0, must not be empty
@@ -2283,7 +2284,7 @@ IConsoleObject* FConsoleManager::FindConsoleObject(const TCHAR* Name, bool bTrac
 
 IConsoleObject* FConsoleManager::FindConsoleObjectUnfiltered(const TCHAR* Name) const
 {
-	FScopeLock ScopeLock( &ConsoleObjectsSynchronizationObject );
+	FTransactionallySafeScopeLock ScopeLock( &ConsoleObjectsSynchronizationObject );
 	IConsoleObject* Var = ConsoleObjects.FindRef(Name);
 	return Var;
 }
@@ -2294,7 +2295,7 @@ void FConsoleManager::UnregisterConsoleObject(IConsoleObject* CVar, bool bKeepSt
 	{
 		return;
 	}
-	FScopeLock ScopeLock(&ConsoleObjectsSynchronizationObject);
+	FTransactionallySafeScopeLock ScopeLock(&ConsoleObjectsSynchronizationObject);
 
 	// Slow search for console object
 	const FString ObjName = FindConsoleObjectName( CVar );
@@ -2307,7 +2308,7 @@ void FConsoleManager::UnregisterConsoleObject(IConsoleObject* CVar, bool bKeepSt
 
 void FConsoleManager::UnregisterConsoleObject(const TCHAR* Name, bool bKeepState)
 {
-	FScopeLock ScopeLock(&ConsoleObjectsSynchronizationObject);
+	FTransactionallySafeScopeLock ScopeLock(&ConsoleObjectsSynchronizationObject);
 
 	IConsoleObject* Object = FindConsoleObject(Name);
 
@@ -2406,7 +2407,7 @@ void FConsoleManager::ForEachConsoleObjectThatStartsWith(const FConsoleObjectVis
 	check(ThatStartsWith);
 
 	//@caution, potential deadlock if the visitor tries to call back into the cvar system. Best not to do this, but we could capture and array of them, then release the lock, then dispatch the visitor.
-	FScopeLock ScopeLock( &ConsoleObjectsSynchronizationObject );
+	FTransactionallySafeScopeLock ScopeLock( &ConsoleObjectsSynchronizationObject );
 	for(TMap<FString, IConsoleObject*>::TConstIterator PairIt(ConsoleObjects); PairIt; ++PairIt)
 	{
 		const FString& Name = PairIt.Key();
@@ -2429,7 +2430,7 @@ void FConsoleManager::ForEachConsoleObjectThatContains(const FConsoleObjectVisit
 	int32 ContainsStringLength = FCString::Strlen(ThatContains);
 
 	//@caution, potential deadlock if the visitor tries to call back into the cvar system. Best not to do this, but we could capture and array of them, then release the lock, then dispatch the visitor.
-	FScopeLock ScopeLock( &ConsoleObjectsSynchronizationObject );
+	FTransactionallySafeScopeLock ScopeLock( &ConsoleObjectsSynchronizationObject );
 	for(TMap<FString, IConsoleObject*>::TConstIterator PairIt(ConsoleObjects); PairIt; ++PairIt)
 	{
 		const FString& Name = PairIt.Key();
@@ -2898,7 +2899,7 @@ IConsoleObject* FConsoleManager::AddConsoleObject(const TCHAR* Name, IConsoleObj
 	check(Obj);
 	LLM_SCOPE_BYNAME(TEXT("EngineMisc/ConsoleCommands"));
 
-	FScopeLock ScopeLock( &ConsoleObjectsSynchronizationObject ); // we will lock on the entire add process
+	FTransactionallySafeScopeLock ScopeLock( &ConsoleObjectsSynchronizationObject ); // we will lock on the entire add process
 	IConsoleObject* ExistingObj = ConsoleObjects.FindRef(Name);
 
 	if(Obj->GetFlags() & ECVF_Scalability)
@@ -3060,7 +3061,7 @@ FString FConsoleManager::FindConsoleObjectName(const IConsoleObject* InVar) cons
 {
 	check(InVar);
 
-	FScopeLock ScopeLock( &ConsoleObjectsSynchronizationObject );
+	FTransactionallySafeScopeLock ScopeLock( &ConsoleObjectsSynchronizationObject );
 	for(TMap<FString, IConsoleObject*>::TConstIterator PairIt(ConsoleObjects); PairIt; ++PairIt)
 	{
 		IConsoleObject* Var = PairIt.Value();
@@ -3171,7 +3172,7 @@ void FConsoleManager::GetConsoleHistory(const TCHAR* Key, TArray<FString>& Out)
 
 bool FConsoleManager::IsNameRegistered(const TCHAR* Name) const
 {
-	FScopeLock ScopeLock(&ConsoleObjectsSynchronizationObject);
+	FTransactionallySafeScopeLock ScopeLock(&ConsoleObjectsSynchronizationObject);
 	return ConsoleObjects.Contains(Name);
 }
 
