Fixed TNonRelocatableInlineAllocator::ForElementType::MoveToEmpty to leave the moved-from allocator pointing to the inline storage (like after a default construction) rather than unallocated heap storage.

diff --git a/Engine/Source/Runtime/Core/Public/Containers/ContainerAllocationPolicies.h b/Engine/Source/Runtime/Core/Public/Containers/ContainerAllocationPolicies.h
--- a/Engine/Source/Runtime/Core/Public/Containers/ContainerAllocationPolicies.h
+++ b/Engine/Source/Runtime/Core/Public/Containers/ContainerAllocationPolicies.h
@@ -1082,7 +1082,7 @@ public:
 			if (Other.HasAllocation())
 			{
 				Data = Other.Data;
-				Other.Data = nullptr;
+				Other.Data = Other.GetInlineElements();
 			}
 			else
 			{
