[AutoRTFM] Fix AutoRTFMTests on linux

The `FPlatformMisc.CreateGUID` test calls CreateGUID, which attempts a syscall on linux. Mark this as always open.

diff --git a/Engine/Source/Runtime/Core/Private/Unix/UnixPlatformMisc.cpp b/Engine/Source/Runtime/Core/Private/Unix/UnixPlatformMisc.cpp
--- a/Engine/Source/Runtime/Core/Private/Unix/UnixPlatformMisc.cpp
+++ b/Engine/Source/Runtime/Core/Private/Unix/UnixPlatformMisc.cpp
@@ -1654,7 +1654,7 @@ namespace
 
 // If we fail to create a Guid with urandom fallback to the generic platform.
 // This maybe need to be tweaked for Servers and hard fail here
-void FUnixPlatformMisc::CreateGuid(FGuid& Result)
+UE_AUTORTFM_ALWAYS_OPEN void FUnixPlatformMisc::CreateGuid(FGuid& Result)
 {
 	int BytesRead = SysGetRandom(&Result, sizeof(Result));
 
