Fixed dynamic delegate hashing of uninitialized padding by removing its GetTypeHash overload, so it will fall through to the base class implementation.

diff --git a/Engine/Source/Runtime/Core/Public/Delegates/DelegateSignatureImpl.inl b/Engine/Source/Runtime/Core/Public/Delegates/DelegateSignatureImpl.inl
--- a/Engine/Source/Runtime/Core/Public/Delegates/DelegateSignatureImpl.inl
+++ b/Engine/Source/Runtime/Core/Public/Delegates/DelegateSignatureImpl.inl
@@ -11,7 +11,6 @@
 
 #include "CoreTypes.h"
 #include "Misc/AssertionMacros.h"
-#include "Misc/Crc.h"
 #include "UObject/NameTypes.h"
 #include "UObject/ScriptDelegates.h"
 #include "UObject/WeakObjectPtrTemplates.h"
@@ -1153,11 +1152,6 @@ public:
 		__Internal_BindDynamic(ToRawPtr(InUserObject), InMethodPtr, InFunctionName);
 	}
 
-	friend uint32 GetTypeHash(const TBaseDynamicDelegate& Key)
-	{
-		return FCrc::MemCrc_DEPRECATED(&Key,sizeof(Key));
-	}
-
 	// NOTE:  Execute() method must be defined in derived classes
 
 	// NOTE:  ExecuteIfBound() method must be defined in derived classes
@@ -1300,3 +1294,7 @@ public:
 
 	// NOTE:  Broadcast() method must be defined in derived classes
 };
+
+#if UE_ENABLE_INCLUDE_ORDER_DEPRECATED_IN_5_5
+#include "Misc/Crc.h"
+#endif
diff --git a/Engine/Source/Runtime/Core/Public/UObject/ScriptDelegates.h b/Engine/Source/Runtime/Core/Public/UObject/ScriptDelegates.h
--- a/Engine/Source/Runtime/Core/Public/UObject/ScriptDelegates.h
+++ b/Engine/Source/Runtime/Core/Public/UObject/ScriptDelegates.h
@@ -2,15 +2,16 @@
 
 #pragma once
 
+#include "Containers/Array.h"
+#include "Containers/ContainerAllocationPolicies.h"
+#include "Containers/UnrealString.h"
 #include "CoreTypes.h"
+#include "Delegates/DelegateAccessHandler.h"
 #include "Misc/AssertionMacros.h"
+#include "Templates/SharedPointer.h"
+#include "Templates/TypeHash.h"
 #include "Templates/UnrealTypeTraits.h"
-#include "Containers/ContainerAllocationPolicies.h"
-#include "Containers/Array.h"
-#include "Containers/UnrealString.h"
 #include "UObject/NameTypes.h"
-#include "Templates/SharedPointer.h"
-#include "Delegates/DelegateAccessHandler.h"
 
 namespace UE::Core::Private
 {
