Crunched FTextKey down to a single int32

FTextKey used to be a pointer and embedded hash (when UE_TEXTKEY_STORE_EMBEDDED_HASH=1), where the pointer was a reference to an entry in an internal table.
This change updates FTextKey to just store the index to the entry in the internal table, and then look-up the data when needed. This look-up is fast enough that UE_TEXTKEY_STORE_EMBEDDED_HASH can now be disabled by default.

These changes take FTextKey from 12 bytes to 4 bytes, and FTextId from 24 bytes to 8 bytes. This can save a few MB of memory in the display string live table, plus makes anything else storing a FTextKey/FTextId (like the FText history types) smaller too.

diff --git a/Engine/Source/Runtime/Core/Private/Internationalization/TextKey.cpp b/Engine/Source/Runtime/Core/Private/Internationalization/TextKey.cpp
--- a/Engine/Source/Runtime/Core/Private/Internationalization/TextKey.cpp
+++ b/Engine/Source/Runtime/Core/Private/Internationalization/TextKey.cpp
@@ -4,28 +4,40 @@
 
 #include "AutoRTFM/AutoRTFM.h"
 #include "Containers/ContainerAllocationPolicies.h"
-#include "Containers/Map.h"
+#include "Containers/ChunkedArray.h"
 #include "CoreGlobals.h"
 #include "HAL/LowLevelMemTracker.h"
 #include "Hash/CityHash.h"
 #include "Logging/LogMacros.h"
 #include "Misc/ByteSwap.h"
+#include "Misc/Guid.h"
 #include "Misc/LazySingleton.h"
-#include "Misc/ScopeLock.h"
+#include "Misc/ScopeRWLock.h"
+#include "Modules/VisualizerDebuggingState.h"
 
 DEFINE_LOG_CATEGORY_STATIC(LogTextKey, Log, All);
 
-#ifndef UE_TEXTKEY_USE_SLAB_ALLOCATOR
-	#define UE_TEXTKEY_USE_SLAB_ALLOCATOR (1)
-#endif
-
 #ifndef UE_TEXTKEY_SLAB_ALLOCATOR_SLAB_SIZE
 	#define UE_TEXTKEY_SLAB_ALLOCATOR_SLAB_SIZE (32768)
 #endif
 
+#ifndef UE_TEXTKEY_ELEMENTS_CHUNK_SIZE_BYTES
+	#define UE_TEXTKEY_ELEMENTS_CHUNK_SIZE_BYTES (32768)
+#endif
+
+#ifndef UE_TEXTKEY_ELEMENTS_MIN_HASH_SIZE
+	#define UE_TEXTKEY_ELEMENTS_MIN_HASH_SIZE (32768)
+#endif
+
 class FTextKeyState
 {
 public:
+	FTextKeyState()
+	{
+		// Register the natvis data accessor
+		UE::Core::EVisualizerDebuggingStateResult Result = UE::Core::FVisualizerDebuggingState::Assign(FGuid(0xd31281c0, 0x182b4419, 0x814e25be, 0x4b7e7b41), this);
+	}
+
 	void FindOrAdd(FStringView InStr, FTextKey& OutTextKey)
 	{
 		check(!InStr.IsEmpty());
@@ -48,11 +60,29 @@ public:
 		};
 	}
 
+	const TCHAR* GetStringByIndex(const int32 InIndex) const
+	{
+		check(InIndex != INDEX_NONE);
+
+		FReadScopeLock ScopeLock(DataRW);
+
+		const FKeyData& KeyData = KeyDataAllocations.Get(InIndex);
+		return KeyData.Str;
+	}
+
+	uint32 GetHashByIndex(const int32 InIndex) const
+	{
+		check(InIndex != INDEX_NONE);
+
+		FReadScopeLock ScopeLock(DataRW);
+
+		const FKeyData& KeyData = KeyDataAllocations.Get(InIndex);
+		return KeyData.StrHash;
+	}
+
 	void Shrink()
 	{
-		FScopeLock ScopeLock(&SynchronizationObject);
-		LLM_SCOPE_BYNAME(TEXT("Localization/TextKeys"));
-		KeysTable.Shrink();
+		// Note: Nothing to shrink as things grow in chunks
 	}
 
 	static FTextKeyState& GetState()
@@ -68,30 +98,31 @@ public:
 private:
 	struct FKeyData
 	{
+		FKeyData()
+			: Str(TEXT(""))
+			, StrLen(0)
+			, StrHash(0)
+		{
+		}
+
 		explicit FKeyData(FStringView InStr)
 			: Str(InStr.GetData())
 			, StrLen(InStr.Len())
-#if UE_TEXTKEY_STORE_EMBEDDED_HASH
 			, StrHash(TextKeyUtil::HashString(InStr)) // Note: This hash gets serialized so *DO NOT* change it without fixing the serialization to discard the old hash method (also update FTextKey::GetTypeHash)
-#endif
 		{
 		}
 
 		FKeyData(FStringView InStr, const uint32 InStrHash)
 			: Str(InStr.GetData())
 			, StrLen(InStr.Len())
-#if UE_TEXTKEY_STORE_EMBEDDED_HASH
 			, StrHash(InStrHash)
-#endif
 		{
 		}
 
 		FKeyData(FStringView InStr, const FKeyData& InOther)
 			: Str(InStr.GetData())
 			, StrLen(InStr.Len())
-#if UE_TEXTKEY_STORE_EMBEDDED_HASH
 			, StrHash(InOther.StrHash)
-#endif
 		{
 		}
 
@@ -114,77 +145,175 @@ private:
 
 		friend FORCEINLINE uint32 GetTypeHash(const FKeyData& A)
 		{
-#if UE_TEXTKEY_STORE_EMBEDDED_HASH
 			return A.StrHash;
-#else
-			return TextKeyUtil::HashString(A.ToView());
-#endif
 		}
 
 		const TCHAR* Str;
 		int32 StrLen;
-#if UE_TEXTKEY_STORE_EMBEDDED_HASH
 		uint32 StrHash;
-#endif
 	};
 
 	void FindOrAddImpl(const FKeyData& KeyData, FTextKey& OutTextKey)
 	{
-		const TCHAR* StrPtr = FindOrAddString(KeyData);
-		check(StrPtr);
+		int32 Index = FindOrAddString(KeyData);
+		check(Index != INDEX_NONE);
 
-		OutTextKey.StrPtr = StrPtr;
+		OutTextKey.Index = Index;
 #if UE_TEXTKEY_STORE_EMBEDDED_HASH
 		OutTextKey.StrHash = KeyData.StrHash;
 #endif
 	}
 
-	const TCHAR* FindOrAddString(const FKeyData& KeyData)
+	int32 FindOrAddString(const FKeyData& KeyData)
 	{
-		FScopeLock ScopeLock(&SynchronizationObject);
-#if UE_TEXTKEY_USE_SLAB_ALLOCATOR
-		const TCHAR* StrPtr = nullptr;
-		if (const FKeyData* FoundKeyData = KeysTable.Find(KeyData))
+		// Read-only
+		int32 NumElementsOnRead = 0;
 		{
-			StrPtr = FoundKeyData->Str;
+			FReadScopeLock ScopeLock(DataRW);
+
+			if (int32 FoundIndex = KeyDataAllocations.Find(KeyData);
+				FoundIndex != INDEX_NONE)
+			{
+				return FoundIndex;
+			}
+
+			NumElementsOnRead = KeyDataAllocations.Num();
 		}
-		if (!StrPtr)
+
+		// Write
 		{
+			FWriteScopeLock ScopeLock(DataRW);
+
+			if (KeyDataAllocations.Num() > NumElementsOnRead)
+			{
+				// Find again in case another thread beat us to it!
+				if (int32 FoundIndex = KeyDataAllocations.Find(KeyData);
+					FoundIndex != INDEX_NONE)
+				{
+					return FoundIndex;
+				}
+			}
+
 			LLM_SCOPE_BYNAME(TEXT("Localization/TextKeys"));
-			StrPtr = StringAllocations.Add(KeyData.ToView());
-			KeysTable.Add(FKeyData(FStringView(StrPtr, KeyData.StrLen), KeyData));
+			const TCHAR* NewStrPtr = StringAllocations.Add(KeyData.ToView());
+			return KeyDataAllocations.Add(FKeyData(FStringView(NewStrPtr, KeyData.StrLen), KeyData));
 		}
-		return StrPtr;
-#else
-		const FString* StrPtr = KeysTable.Find(KeyData);
-		if (!StrPtr)
+	}
+
+	class FKeyDataAllocator
+	{
+	public:
+		FKeyDataAllocator() = default;
+
+		~FKeyDataAllocator()
 		{
-			LLM_SCOPE_BYNAME(TEXT("Localization/TextKeys"));
+			FreeHash();
+		}
+
+		int32 Add(const FKeyData& InKeyData)
+		{
+			ConditionalRehash(Elements.Num() + 1);
+
+			const uint32 KeyDataHash = GetTypeHash(InKeyData);
+			const uint32 HashIndex = KeyDataHash & (HashSize - 1);
+
+			const int32 NewElementIndex = Elements.Num();
+			Elements.AddElement(FElement{ InKeyData, Hash[HashIndex] });
+			Hash[HashIndex] = NewElementIndex;
+			return NewElementIndex;
+		}
 
-			// Need to copy the string here so we can reference its internal allocation as the key
-			FString StrCopy;
+		int32 Find(const FKeyData& InKeyData) const
+		{
+			if (HashSize > 0)
 			{
-				// We do this rather than use the FString constructor directly, 
-				// as this method avoids slack being added to the allocation
-				StrCopy.Reserve(KeyData.StrLen);
-				StrCopy.AppendChars(KeyData.Str, KeyData.StrLen);
+				const uint32 KeyDataHash = GetTypeHash(InKeyData);
+				const uint32 HashIndex = KeyDataHash & (HashSize - 1);
+
+				int32 ElementIndex = Hash[HashIndex];
+				while (ElementIndex != INDEX_NONE)
+				{
+					const FElement& Element = Elements[ElementIndex];
+					if (Element.Value == InKeyData)
+					{
+						return ElementIndex;
+					}
+					ElementIndex = Element.NextElementIndex;
+				}
 			}
+			return INDEX_NONE;
+		}
 
-			FKeyData DestKey(StrCopy, KeyData);
-			StrPtr = &KeysTable.Add(DestKey, MoveTemp(StrCopy));
-			check(DestKey.Str == **StrPtr); // The move must have moved the allocation we referenced in the key
+		const FKeyData& Get(int32 InIndex) const
+		{
+			return Elements[InIndex].Value;
 		}
-		return **StrPtr;
-#endif
-	}
 
-#if UE_TEXTKEY_USE_SLAB_ALLOCATOR
-	class FStringSlabAllocator
+		int32 Num() const
+		{
+			return Elements.Num();
+		}
+
+	private:
+		void AllocateHash()
+		{
+			Hash = static_cast<int32*>(FMemory::Malloc(HashSize * sizeof(int32)));
+			for (uint32 HashIndex = 0; HashIndex < HashSize; ++HashIndex)
+			{
+				Hash[HashIndex] = INDEX_NONE;
+			}
+		}
+
+		void FreeHash()
+		{
+			FMemory::Free(Hash);
+			Hash = nullptr;
+		}
+
+		void ConditionalRehash(const int32 NumElements)
+		{
+			const uint32 NewHashSize = FMath::Max<uint32>(MinHashSize, FPlatformMath::RoundUpToPowerOfTwo(NumElements / DEFAULT_NUMBER_OF_ELEMENTS_PER_HASH_BUCKET));
+			if (NewHashSize > HashSize)
+			{
+				FreeHash();
+				HashSize = NewHashSize;
+				AllocateHash();
+
+				int32 ElementIndex = 0;
+				for (FElement& Element : Elements)
+				{
+					const uint32 KeyDataHash = GetTypeHash(Element.Value);
+					const uint32 HashIndex = KeyDataHash & (HashSize - 1);
+
+					Element.NextElementIndex = Hash[HashIndex];
+					Hash[HashIndex] = ElementIndex++;
+				}
+			}
+		}
+
+		struct FElement
+		{
+			FKeyData Value;
+			/** Index of the next element in this hash bucket */
+			int32 NextElementIndex = INDEX_NONE;
+		};
+		static constexpr int32 ElementChunkSize = UE_TEXTKEY_ELEMENTS_CHUNK_SIZE_BYTES;
+		/** Values; indices are referenced by the hash and FTextKey */
+		TChunkedArray<FElement, ElementChunkSize> Elements;
+
+		static constexpr int32 MinHashSize = UE_TEXTKEY_ELEMENTS_MIN_HASH_SIZE;
+		/** Current size of the hash; if this changes the hash must be rebuilt */
+		uint32 HashSize = 0;
+		/** Index of the root element in each hash bucket; use FElement::NextElementIndex to walk the bucket */
+		int32* Hash = nullptr;
+	};
+
+	class FStringAllocator
 	{
 	public:
-		FStringSlabAllocator() = default;
+		FStringAllocator() = default;
 
-		~FStringSlabAllocator()
+		~FStringAllocator()
 		{
 			uint64 TotalNumElementsWasted = 0;
 			for (FStringSlab& Slab : Slabs)
@@ -210,7 +339,7 @@ private:
 		}
 
 	private:
-		static const int32 SlabSizeInElements = UE_TEXTKEY_SLAB_ALLOCATOR_SLAB_SIZE;
+		static constexpr int32 SlabSizeInElements = UE_TEXTKEY_SLAB_ALLOCATOR_SLAB_SIZE;
 
 		struct FStringSlab
 		{
@@ -253,21 +382,16 @@ private:
 			// If no slabs have space then just allocate a new one
 			checkf(NumSlabElementsNeeded <= SlabSizeInElements, TEXT("Tried to allocate a FTextKey string of %d elements, which is larger than the allowed slab size of %d elements!"), NumSlabElementsNeeded, SlabSizeInElements);
 			FStringSlab& Slab = Slabs.AddDefaulted_GetRef();
-			Slab.Allocation = reinterpret_cast<TCHAR*>(FMemory::Malloc(SlabSizeInElements * sizeof(TCHAR), 1));
+			Slab.Allocation = static_cast<TCHAR*>(FMemory::Malloc(SlabSizeInElements * sizeof(TCHAR), 1));
 			return Slab;
 		}
 
 		TArray<FStringSlab> Slabs;
 	};
-#endif
 
-	FCriticalSection SynchronizationObject;
-#if UE_TEXTKEY_USE_SLAB_ALLOCATOR
-	FStringSlabAllocator StringAllocations;
-	TSet<FKeyData> KeysTable;
-#else
-	TMap<FKeyData, FString> KeysTable;
-#endif
+	mutable FRWLock DataRW;
+	FStringAllocator StringAllocations;
+	FKeyDataAllocator KeyDataAllocations;
 };
 
 namespace TextKeyUtil
@@ -382,11 +506,6 @@ uint32 HashString(const FTCHARToUTF16& InStr)
 
 }
 
-FTextKey::FTextKey()
-{
-	Reset();
-}
-
 FTextKey::FTextKey(FStringView InStr)
 {
 	if (InStr.IsEmpty())
@@ -399,6 +518,24 @@ FTextKey::FTextKey(FStringView InStr)
 	}
 }
 
+const TCHAR* FTextKey::GetChars() const
+{
+	return Index != INDEX_NONE
+		? FTextKeyState::GetState().GetStringByIndex(Index)
+		: TEXT("");
+}
+
+uint32 GetTypeHash(const FTextKey& A)
+{
+#if UE_TEXTKEY_STORE_EMBEDDED_HASH
+	return A.StrHash;
+#else
+	return A.Index != INDEX_NONE
+		? FTextKeyState::GetState().GetHashByIndex(A.Index)
+		: 0;
+#endif
+}
+
 void FTextKey::SerializeAsString(FArchive& Ar)
 {
 	if (Ar.IsLoading())
@@ -586,14 +723,6 @@ void FTextKey::SerializeDiscardHash(FStructuredArchiveSlot Slot)
 	}
 }
 
-void FTextKey::Reset()
-{
-	StrPtr = TEXT("");
-#if UE_TEXTKEY_STORE_EMBEDDED_HASH
-	StrHash = 0;
-#endif
-}
-
 void FTextKey::CompactDataStructures()
 {
 	FTextKeyState::GetState().Shrink();
diff --git a/Engine/Source/Runtime/Core/Public/Internationalization/TextKey.h b/Engine/Source/Runtime/Core/Public/Internationalization/TextKey.h
--- a/Engine/Source/Runtime/Core/Public/Internationalization/TextKey.h
+++ b/Engine/Source/Runtime/Core/Public/Internationalization/TextKey.h
@@ -10,7 +10,7 @@
 #include "Serialization/StructuredArchive.h"
 
 #ifndef UE_TEXTKEY_STORE_EMBEDDED_HASH
-	#define UE_TEXTKEY_STORE_EMBEDDED_HASH (1)
+	#define UE_TEXTKEY_STORE_EMBEDDED_HASH (0)
 #endif
 
 namespace TextKeyUtil
@@ -74,7 +74,7 @@ namespace TextKeyUtil
 class FTextKey
 {
 public:
-	CORE_API FTextKey();
+	FTextKey() = default;
 	CORE_API FTextKey(FStringView InStr);
 	
 	FTextKey(const TCHAR* InStr)
@@ -88,32 +88,22 @@ public:
 	}
 
 	/** Get the underlying chars buffer this text key represents */
-	FORCEINLINE const TCHAR* GetChars() const
-	{
-		return StrPtr;
-	}
+	CORE_API const TCHAR* GetChars() const;
 
 	/** Compare for equality */
 	friend FORCEINLINE bool operator==(const FTextKey& A, const FTextKey& B)
 	{
-		return A.StrPtr == B.StrPtr;
+		return A.Index == B.Index;
 	}
 
 	/** Compare for inequality */
 	friend FORCEINLINE bool operator!=(const FTextKey& A, const FTextKey& B)
 	{
-		return A.StrPtr != B.StrPtr;
+		return A.Index != B.Index;
 	}
 
 	/** Get the hash of this text key */
-	friend FORCEINLINE uint32 GetTypeHash(const FTextKey& A)
-	{
-#if UE_TEXTKEY_STORE_EMBEDDED_HASH
-		return A.StrHash;
-#else
-		return TextKeyUtil::HashString(A.StrPtr);
-#endif
-	}
+	friend CORE_API uint32 GetTypeHash(const FTextKey& A);
 
 	/** Serialize this text key as if it were an FString */
 	CORE_API void SerializeAsString(FArchive& Ar);
@@ -136,11 +126,17 @@ public:
 	/** Is this text key empty? */
 	FORCEINLINE bool IsEmpty() const
 	{
-		return *StrPtr == 0;
+		return Index == INDEX_NONE;
 	}
 
 	/** Reset this text key to be empty */
-	CORE_API void Reset();
+	FORCEINLINE void Reset()
+	{
+		Index = INDEX_NONE;
+#if UE_TEXTKEY_STORE_EMBEDDED_HASH
+		StrHash = 0;
+#endif
+	}
 
 	/** Compact any slack within the internal table */
 	static CORE_API void CompactDataStructures();
@@ -149,23 +145,14 @@ public:
 	static CORE_API void TearDown();
 
 private:
-	/** Pointer to the null-terminated string buffer we reference */
-	const TCHAR* StrPtr;
+	/** Index of the internal FKeyData we reference */
+	int32 Index = INDEX_NONE;
 
 #if UE_TEXTKEY_STORE_EMBEDDED_HASH
-	/** Hash of this text key */
-	uint32 StrHash;
+	/** Local cache of FKeyData::StrHash to avoid indirection into the internal table */
+	uint32 StrHash = 0;
 #endif
 
-	FTextKey(const TCHAR* Str, uint32 Hash)
-		: StrPtr(Str)
-#if UE_TEXTKEY_STORE_EMBEDDED_HASH
-		, StrHash(Hash)
-#endif
-	{
-	}
-
-	friend class FTextId;
 	friend class FTextKeyState;
 };
 
@@ -175,93 +162,63 @@ private:
 class FTextId
 {
 public:
-	FTextId()
-	{
-		Reset();
-	}
+	FTextId() = default;
 
 	FTextId(const FTextKey& InNamespace, const FTextKey& InKey)
+		: Namespace(InNamespace)
+		, Key(InKey)
 	{
-		SetNamespace(InNamespace);
-		SetKey(InKey);
 	}
 
 	/** Get the namespace component of this text identity */
 	FORCEINLINE FTextKey GetNamespace() const
 	{
-#if UE_TEXTKEY_STORE_EMBEDDED_HASH
-		return FTextKey(NamespaceStr, NamespaceHash);
-#else
-		return FTextKey(NamespaceStr, 0);
-#endif
+		return Namespace;
 	}
 
 	/** Get the key component of this text identity */
 	FORCEINLINE FTextKey GetKey() const
 	{
-#if UE_TEXTKEY_STORE_EMBEDDED_HASH
-		return FTextKey(KeyStr, KeyHash);
-#else
-		return FTextKey(KeyStr, 0);
-#endif
+		return Key;
 	}
 
 	/** Compare for equality */
 	friend FORCEINLINE bool operator==(const FTextId& A, const FTextId& B)
 	{
-		return A.NamespaceStr == B.NamespaceStr && A.KeyStr == B.KeyStr;
+		return A.Namespace == B.Namespace && A.Key == B.Key;
 	}
 
 	/** Compare for inequality */
 	friend FORCEINLINE bool operator!=(const FTextId& A, const FTextId& B)
 	{
-		return A.NamespaceStr != B.NamespaceStr || A.KeyStr != B.KeyStr;
+		return A.Namespace != B.Namespace || A.Key != B.Key;
 	}
 
 	/** Get the hash of this text identity */
 	friend FORCEINLINE uint32 GetTypeHash(const FTextId& A)
 	{
-#if UE_TEXTKEY_STORE_EMBEDDED_HASH
-		return HashCombine(A.NamespaceHash, A.KeyHash);
-#else
-		return HashCombine(TextKeyUtil::HashString(A.NamespaceStr), TextKeyUtil::HashString(A.KeyStr));
-#endif
+		return HashCombine(GetTypeHash(A.Namespace), GetTypeHash(A.Key));
 	}
 
 	/** Serialize this text identity as if it were FStrings */
 	void SerializeAsString(FArchive& Ar)
 	{
-		FTextKey Namespace = GetNamespace();
 		Namespace.SerializeAsString(Ar);
-		SetNamespace(Namespace);
-
-		FTextKey Key = GetKey();
 		Key.SerializeAsString(Ar);
-		SetKey(Key);
 	}
 
 	/** Serialize this text identity including its hash values (this method is sensitive to hashing algorithm changes, so only use it for generated files that can be rebuilt from another source) */
 	void SerializeWithHash(FArchive& Ar)
 	{
-		FTextKey Namespace = GetNamespace();
 		Namespace.SerializeWithHash(Ar);
-		SetNamespace(Namespace);
-
-		FTextKey Key = GetKey();
 		Key.SerializeWithHash(Ar);
-		SetKey(Key);
 	}
 
 	/** Serialize this text identity including its hash values, discarding the hash on load (to upgrade from an older hashing algorithm) */
 	void SerializeDiscardHash(FArchive& Ar)
 	{
-		FTextKey Namespace = GetNamespace();
 		Namespace.SerializeDiscardHash(Ar);
-		SetNamespace(Namespace);
-
-		FTextKey Key = GetKey();
 		Key.SerializeDiscardHash(Ar);
-		SetKey(Key);
 	}
 
 	/** Serialize this text identity as if it were FStrings */
@@ -269,13 +226,8 @@ public:
 	{
 		FStructuredArchiveRecord Record = Slot.EnterRecord();
 
-		FTextKey Namespace = GetNamespace();
 		Namespace.SerializeAsString(Record.EnterField(TEXT("Namespace")));
-		SetNamespace(Namespace);
-
-		FTextKey Key = GetKey();
 		Key.SerializeAsString(Record.EnterField(TEXT("Key")));
-		SetKey(Key);
 	}
 
 	/** Serialize this text identity including its hash values (this method is sensitive to hashing algorithm changes, so only use it for generated files that can be rebuilt from another source) */
@@ -283,13 +235,8 @@ public:
 	{
 		FStructuredArchiveRecord Record = Slot.EnterRecord();
 
-		FTextKey Namespace = GetNamespace();
 		Namespace.SerializeWithHash(Record.EnterField(TEXT("Namespace")));
-		SetNamespace(Namespace);
-
-		FTextKey Key = GetKey();
 		Key.SerializeWithHash(Record.EnterField(TEXT("Key")));
-		SetKey(Key);
 	}
 
 	/** Serialize this text identity including its hash values, discarding the hash on load (to upgrade from an older hashing algorithm) */
@@ -297,51 +244,24 @@ public:
 	{
 		FStructuredArchiveRecord Record = Slot.EnterRecord();
 
-		FTextKey Namespace = GetNamespace();
 		Namespace.SerializeDiscardHash(Record.EnterField(TEXT("Namespace")));
-		SetNamespace(Namespace);
-
-		FTextKey Key = GetKey();
 		Key.SerializeDiscardHash(Record.EnterField(TEXT("Key")));
-		SetKey(Key);
 	}
 
 	/** Is this text identity empty? */
 	FORCEINLINE bool IsEmpty() const
 	{
-		return *NamespaceStr == 0 && *KeyStr == 0;
+		return Namespace.IsEmpty() && Key.IsEmpty();
 	}
 
 	/** Reset this text identity to be empty */
 	FORCEINLINE void Reset()
 	{
-		NamespaceStr = KeyStr = TEXT("");
-#if UE_TEXTKEY_STORE_EMBEDDED_HASH
-		NamespaceHash = KeyHash = 0;
-#endif
+		Namespace.Reset();
+		Key.Reset();
 	}
 
 private:
-	void SetNamespace(const FTextKey& InNamespace)
-	{
-		NamespaceStr = InNamespace.StrPtr;
-#if UE_TEXTKEY_STORE_EMBEDDED_HASH
-		NamespaceHash = InNamespace.StrHash;
-#endif
-	}
-
-	void SetKey(const FTextKey& InKey)
-	{
-		KeyStr = InKey.StrPtr;
-#if UE_TEXTKEY_STORE_EMBEDDED_HASH
-		KeyHash = InKey.StrHash;
-#endif
-	}
-
-	const TCHAR* NamespaceStr;
-	const TCHAR* KeyStr;
-#if UE_TEXTKEY_STORE_EMBEDDED_HASH
-	uint32 NamespaceHash;
-	uint32 KeyHash;
-#endif
+	FTextKey Namespace;
+	FTextKey Key;
 };
