Renamed container allocator function parameters to make them less ambiguous about whether they refer to a current state or new (desired) state, and refer to the 'num' (how many elements are constructed inside an allocation) or the 'max' (the total number of elements that the allocation could hold), as it was unclear.
Added a comment to clarify that TNonRelocatableInlineAllocator is itself non-trivially-relocatable, not that it's for non-trivially-relocatable elements.

diff --git a/Engine/Plugins/Experimental/Mutable/Source/MutableRuntime/Public/MuR/MemoryTrackingAllocationPolicy.h b/Engine/Plugins/Experimental/Mutable/Source/MutableRuntime/Public/MuR/MemoryTrackingAllocationPolicy.h
--- a/Engine/Plugins/Experimental/Mutable/Source/MutableRuntime/Public/MuR/MemoryTrackingAllocationPolicy.h
+++ b/Engine/Plugins/Experimental/Mutable/Source/MutableRuntime/Public/MuR/MemoryTrackingAllocationPolicy.h
@@ -76,14 +76,14 @@ namespace mu
 			}
 
 			FORCEINLINE void ResizeAllocation(
-				SizeType PreviousNumElements,
-				SizeType NumElements,
+				SizeType CurrentNum,
+				SizeType NewMax,
 				SIZE_T NumBytesPerElement
 			)
 			{
-				BaseAlloc::ForAnyElementType::ResizeAllocation(PreviousNumElements, NumElements, NumBytesPerElement);
+				BaseAlloc::ForAnyElementType::ResizeAllocation(CurrentNum, NewMax, NumBytesPerElement);
 
-				const SSIZE_T AllocatedSize = (SSIZE_T)BaseAlloc::ForAnyElementType::GetAllocatedSize(NumElements, NumBytesPerElement); 
+				const SSIZE_T AllocatedSize = (SSIZE_T)BaseAlloc::ForAnyElementType::GetAllocatedSize(NewMax, NumBytesPerElement); 
 				const SSIZE_T Differential = AllocatedSize - AllocSize;
 				const SSIZE_T PrevCounterValue = CounterType::Counter.fetch_add(Differential, std::memory_order_relaxed);
 				check(PrevCounterValue >= AllocSize);
@@ -97,15 +97,15 @@ namespace mu
 
 			template<class T = BaseAlloc>
 			FORCEINLINE typename TEnableIf<TAllocatorTraits<T>::SupportsElementAlignment, void>::Type ResizeAllocation(
-				SizeType PreviousNumElements,
-				SizeType NumElements,
+				SizeType CurrentNum,
+				SizeType NewMax,
 				SIZE_T NumBytesPerElement,
 				uint32 AlignmentOfElement
 			)
 			{
-				BaseAlloc::ForAnyElementType::ResizeAllocation(PreviousNumElements, NumElements, NumBytesPerElement, AlignmentOfElement);
+				BaseAlloc::ForAnyElementType::ResizeAllocation(CurrentNum, NewMax, NumBytesPerElement, AlignmentOfElement);
 
-				const SSIZE_T AllocatedSize = (SSIZE_T)BaseAlloc::ForAnyElementType::GetAllocatedSize(NumElements, NumBytesPerElement); 
+				const SSIZE_T AllocatedSize = (SSIZE_T)BaseAlloc::ForAnyElementType::GetAllocatedSize(NewMax, NumBytesPerElement); 
 				const SSIZE_T Differential = AllocatedSize - AllocSize;
 				const SSIZE_T PrevCounterValue = CounterType::Counter.fetch_add(Differential, std::memory_order_relaxed);
 				check(PrevCounterValue >= AllocSize);
diff --git a/Engine/Source/Developer/ShaderCompilerCommon/Private/HlslUtils.h b/Engine/Source/Developer/ShaderCompilerCommon/Private/HlslUtils.h
--- a/Engine/Source/Developer/ShaderCompilerCommon/Private/HlslUtils.h
+++ b/Engine/Source/Developer/ShaderCompilerCommon/Private/HlslUtils.h
@@ -131,40 +131,40 @@ namespace CrossCompiler
 			{
 				return Data;
 			}
-			void ResizeAllocation(SizeType PreviousNumElements, SizeType NumElements, SIZE_T NumBytesPerElement)
+			void ResizeAllocation(SizeType CurrentNum, SizeType NewMax, SIZE_T NumBytesPerElement)
 			{
 				void* OldData = Data;
-				if (NumElements)
+				if (NewMax)
 				{
 					// Allocate memory from the stack.
-					Data = (ElementType*)LinearAllocator->Alloc(NumElements * NumBytesPerElement,
+					Data = (ElementType*)LinearAllocator->Alloc(NewMax * NumBytesPerElement,
 						FMath::Max((uint32)sizeof(void*), (uint32)alignof(ElementType))
 						);
 
 					// If the container previously held elements, copy them into the new allocation.
-					if (OldData && PreviousNumElements)
+					if (OldData && CurrentNum)
 					{
-						const SizeType NumCopiedElements = FMath::Min(NumElements, PreviousNumElements);
+						const SizeType NumCopiedElements = FMath::Min(NewMax, CurrentNum);
 						FMemory::Memcpy(Data, OldData, NumCopiedElements * NumBytesPerElement);
 					}
 				}
 			}
-			SizeType CalculateSlackReserve(SizeType NumElements, SIZE_T NumBytesPerElement) const
+			SizeType CalculateSlackReserve(SizeType NewMax, SIZE_T NumBytesPerElement) const
 			{
-				return DefaultCalculateSlackReserve(NumElements, NumBytesPerElement, false);
+				return DefaultCalculateSlackReserve(NewMax, NumBytesPerElement, false);
 			}
-			SizeType CalculateSlackShrink(SizeType NumElements, SizeType NumAllocatedElements, SIZE_T NumBytesPerElement) const
+			SizeType CalculateSlackShrink(SizeType NewMax, SizeType CurrentMax, SIZE_T NumBytesPerElement) const
 			{
-				return DefaultCalculateSlackShrink(NumElements, NumAllocatedElements, NumBytesPerElement, false);
+				return DefaultCalculateSlackShrink(NewMax, CurrentMax, NumBytesPerElement, false);
 			}
-			SizeType CalculateSlackGrow(SizeType NumElements, SizeType NumAllocatedElements, SIZE_T NumBytesPerElement) const
+			SizeType CalculateSlackGrow(SizeType NewMax, SizeType CurrentMax, SIZE_T NumBytesPerElement) const
 			{
-				return DefaultCalculateSlackGrow(NumElements, NumAllocatedElements, NumBytesPerElement, false);
+				return DefaultCalculateSlackGrow(NewMax, CurrentMax, NumBytesPerElement, false);
 			}
 
-			SIZE_T GetAllocatedSize(SizeType NumAllocatedElements, SIZE_T NumBytesPerElement) const
+			SIZE_T GetAllocatedSize(SizeType CurrentMax, SIZE_T NumBytesPerElement) const
 			{
-				return NumAllocatedElements * NumBytesPerElement;
+				return CurrentMax * NumBytesPerElement;
 			}
 
 			bool HasAllocation() const
diff --git a/Engine/Source/Runtime/Core/Private/Containers/ContainerAllocationPolicies.cpp b/Engine/Source/Runtime/Core/Private/Containers/ContainerAllocationPolicies.cpp
--- a/Engine/Source/Runtime/Core/Private/Containers/ContainerAllocationPolicies.cpp
+++ b/Engine/Source/Runtime/Core/Private/Containers/ContainerAllocationPolicies.cpp
@@ -16,8 +16,8 @@ FORCENOINLINE void UE::Core::Private::OnInvalidSizedHeapAllocatorNum(int32 Index
 }
 
 #if !IS_MERGEDMODULES
-template void TSizedHeapAllocator<32, FMemory>::ForAnyElementType::ResizeAllocation(SizeType PreviousNumElements, SizeType NumElements, SIZE_T NumBytesPerElement);
-template void TSizedHeapAllocator<32, FMemory>::ForAnyElementType::ResizeAllocation(SizeType PreviousNumElements, SizeType NumElements, SIZE_T NumBytesPerElement, uint32 AlignmentOfElement);
-template void TSizedHeapAllocator<64, FMemory>::ForAnyElementType::ResizeAllocation(SizeType PreviousNumElements, SizeType NumElements, SIZE_T NumBytesPerElement);
-template void TSizedHeapAllocator<64, FMemory>::ForAnyElementType::ResizeAllocation(SizeType PreviousNumElements, SizeType NumElements, SIZE_T NumBytesPerElement, uint32 AlignmentOfElement);
+template void TSizedHeapAllocator<32, FMemory>::ForAnyElementType::ResizeAllocation(SizeType CurrentNum, SizeType NewMax, SIZE_T NumBytesPerElement);
+template void TSizedHeapAllocator<32, FMemory>::ForAnyElementType::ResizeAllocation(SizeType CurrentNum, SizeType NewMax, SIZE_T NumBytesPerElement, uint32 AlignmentOfElement);
+template void TSizedHeapAllocator<64, FMemory>::ForAnyElementType::ResizeAllocation(SizeType CurrentNum, SizeType NewMax, SIZE_T NumBytesPerElement);
+template void TSizedHeapAllocator<64, FMemory>::ForAnyElementType::ResizeAllocation(SizeType CurrentNum, SizeType NewMax, SIZE_T NumBytesPerElement, uint32 AlignmentOfElement);
 #endif
\ No newline at end of file
diff --git a/Engine/Source/Runtime/Core/Private/Containers/MemoryImageAllocator.cpp b/Engine/Source/Runtime/Core/Private/Containers/MemoryImageAllocator.cpp
--- a/Engine/Source/Runtime/Core/Private/Containers/MemoryImageAllocator.cpp
+++ b/Engine/Source/Runtime/Core/Private/Containers/MemoryImageAllocator.cpp
@@ -45,29 +45,29 @@ void FMemoryImageAllocatorBase::MoveToEmpty(FMemoryImageAllocatorBase& Other)
 	Other.Data = nullptr;
 }
 
-void FMemoryImageAllocatorBase::ResizeAllocation(int32 PreviousNumElements, int32 NumElements, SIZE_T NumBytesPerElement, uint32 Alignment)
+void FMemoryImageAllocatorBase::ResizeAllocation(int32 CurrentNum, int32 NewMax, SIZE_T NumBytesPerElement, uint32 Alignment)
 {
 	FScriptContainerElement* LocalData = Data.Get();
 	// Avoid calling FMemory::Realloc( nullptr, 0 ) as ANSI C mandates returning a valid pointer which is not what we want.
 	if (Data.IsFrozen())
 	{
 		// Can't grow a frozen array or shrink below zero
-		if (UNLIKELY((uint32)NumElements > (uint32)PreviousNumElements))
+		if (UNLIKELY((uint32)NewMax > (uint32)CurrentNum))
 		{
-			UE::Core::Private::OnInvalidMemoryImageAllocatorNum(NumElements, NumBytesPerElement);
+			UE::Core::Private::OnInvalidMemoryImageAllocatorNum(NewMax, NumBytesPerElement);
 		}
 	}
-	else if (LocalData || NumElements > 0)
+	else if (LocalData || NewMax > 0)
 	{
 		static_assert(sizeof(int32) <= sizeof(SIZE_T), "SIZE_T is expected to be larger than int32");
 
 		// Check for under/overflow
-		if (UNLIKELY(NumElements < 0 || NumBytesPerElement < 1 || NumBytesPerElement > (SIZE_T)MAX_int32))
+		if (UNLIKELY(NewMax < 0 || NumBytesPerElement < 1 || NumBytesPerElement > (SIZE_T)MAX_int32))
 		{
-			UE::Core::Private::OnInvalidMemoryImageAllocatorNum(NumElements, NumBytesPerElement);
+			UE::Core::Private::OnInvalidMemoryImageAllocatorNum(NewMax, NumBytesPerElement);
 		}
 
-		Data = (FScriptContainerElement*)FMemory::Realloc(LocalData, NumElements*NumBytesPerElement, Alignment);
+		Data = (FScriptContainerElement*)FMemory::Realloc(LocalData, NewMax*NumBytesPerElement, Alignment);
 	}
 }
 
diff --git a/Engine/Source/Runtime/Core/Private/HAL/Allocators/AnsiAllocator.cpp b/Engine/Source/Runtime/Core/Private/HAL/Allocators/AnsiAllocator.cpp
--- a/Engine/Source/Runtime/Core/Private/HAL/Allocators/AnsiAllocator.cpp
+++ b/Engine/Source/Runtime/Core/Private/HAL/Allocators/AnsiAllocator.cpp
@@ -9,20 +9,20 @@ FORCENOINLINE void UE::Core::Private::OnInvalidAnsiAllocatorNum(int32 NewNum, SI
 	for (;;);
 }
 
-void FAnsiAllocator::ForAnyElementType::ResizeAllocation(SizeType PreviousNumElements, SizeType NumElements, SIZE_T NumBytesPerElement)
+void FAnsiAllocator::ForAnyElementType::ResizeAllocation(SizeType CurrentNum, SizeType NewMax, SIZE_T NumBytesPerElement)
 {
 	// Avoid calling FMemory::Realloc( nullptr, 0 ) as ANSI C mandates returning a valid pointer which is not what we want.
-	if (NumElements)
+	if (NewMax)
 	{
 		static_assert(sizeof(int32) <= sizeof(SIZE_T), "SIZE_T is expected to be larger than int32");
 
 		// Check for under/overflow
-		if (UNLIKELY(NumElements < 0 || NumBytesPerElement < 1 || NumBytesPerElement > (SIZE_T)MAX_int32))
+		if (UNLIKELY(NewMax < 0 || NumBytesPerElement < 1 || NumBytesPerElement > (SIZE_T)MAX_int32))
 		{
-			UE::Core::Private::OnInvalidAnsiAllocatorNum(NumElements, NumBytesPerElement);
+			UE::Core::Private::OnInvalidAnsiAllocatorNum(NewMax, NumBytesPerElement);
 		}
 
-		void* NewRealloc = ::realloc(Data, NumElements*NumBytesPerElement);
+		void* NewRealloc = ::realloc(Data, NewMax*NumBytesPerElement);
 		Data = (FScriptContainerElement*)NewRealloc;
 	}
 	else
diff --git a/Engine/Source/Runtime/Core/Private/HAL/Allocators/AnsiAllocator.h b/Engine/Source/Runtime/Core/Private/HAL/Allocators/AnsiAllocator.h
--- a/Engine/Source/Runtime/Core/Private/HAL/Allocators/AnsiAllocator.h
+++ b/Engine/Source/Runtime/Core/Private/HAL/Allocators/AnsiAllocator.h
@@ -67,23 +67,23 @@ public:
 		{
 			return Data;
 		}
-		CORE_API void ResizeAllocation(SizeType PreviousNumElements, SizeType NumElements, SIZE_T NumBytesPerElement);
-		SizeType CalculateSlackReserve(SizeType NumElements, SIZE_T NumBytesPerElement) const
+		CORE_API void ResizeAllocation(SizeType CurrentNum, SizeType NewMax, SIZE_T NumBytesPerElement);
+		SizeType CalculateSlackReserve(SizeType NewMax, SIZE_T NumBytesPerElement) const
 		{
-			return DefaultCalculateSlackReserve(NumElements, NumBytesPerElement, false);
+			return DefaultCalculateSlackReserve(NewMax, NumBytesPerElement, false);
 		}
-		SizeType CalculateSlackShrink(SizeType NumElements, SizeType NumAllocatedElements, SIZE_T NumBytesPerElement) const
+		SizeType CalculateSlackShrink(SizeType NewMax, SizeType CurrentMax, SIZE_T NumBytesPerElement) const
 		{
-			return DefaultCalculateSlackShrink(NumElements, NumAllocatedElements, NumBytesPerElement, false);
+			return DefaultCalculateSlackShrink(NewMax, CurrentMax, NumBytesPerElement, false);
 		}
-		SizeType CalculateSlackGrow(SizeType NumElements, SizeType NumAllocatedElements, SIZE_T NumBytesPerElement) const
+		SizeType CalculateSlackGrow(SizeType NewMax, SizeType CurrentMax, SIZE_T NumBytesPerElement) const
 		{
-			return DefaultCalculateSlackGrow(NumElements, NumAllocatedElements, NumBytesPerElement, false);
+			return DefaultCalculateSlackGrow(NewMax, CurrentMax, NumBytesPerElement, false);
 		}
 
-		SIZE_T GetAllocatedSize(SizeType NumAllocatedElements, SIZE_T NumBytesPerElement) const
+		SIZE_T GetAllocatedSize(SizeType CurrentMax, SIZE_T NumBytesPerElement) const
 		{
-			return NumAllocatedElements * NumBytesPerElement;
+			return CurrentMax * NumBytesPerElement;
 		}
 
 		bool HasAllocation() const
diff --git a/Engine/Source/Runtime/Core/Private/HAL/LowLevelMemoryUtils.cpp b/Engine/Source/Runtime/Core/Private/HAL/LowLevelMemoryUtils.cpp
--- a/Engine/Source/Runtime/Core/Private/HAL/LowLevelMemoryUtils.cpp
+++ b/Engine/Source/Runtime/Core/Private/HAL/LowLevelMemoryUtils.cpp
@@ -11,7 +11,7 @@ FORCENOINLINE void UE::Core::Private::OnInvalidLLMAllocatorNum(int32 IndexSize,
 	for (;;);
 }
 
-template void TSizedLLMAllocator<32>::ForAnyElementType::ResizeAllocation(SizeType PreviousNumElements, SizeType NumElements, SIZE_T NumBytesPerElement);
-template void TSizedLLMAllocator<64>::ForAnyElementType::ResizeAllocation(SizeType PreviousNumElements, SizeType NumElements, SIZE_T NumBytesPerElement);
+template void TSizedLLMAllocator<32>::ForAnyElementType::ResizeAllocation(SizeType CurrentNum, SizeType NewMax, SIZE_T NumBytesPerElement);
+template void TSizedLLMAllocator<64>::ForAnyElementType::ResizeAllocation(SizeType CurrentNum, SizeType NewMax, SIZE_T NumBytesPerElement);
 
 #endif
diff --git a/Engine/Source/Runtime/Core/Private/HAL/LowLevelMemoryUtils.h b/Engine/Source/Runtime/Core/Private/HAL/LowLevelMemoryUtils.h
--- a/Engine/Source/Runtime/Core/Private/HAL/LowLevelMemoryUtils.h
+++ b/Engine/Source/Runtime/Core/Private/HAL/LowLevelMemoryUtils.h
@@ -959,46 +959,46 @@ public:
 			return Data;
 		}
 
-		FORCEINLINE void ResizeAllocation(SizeType PreviousNumElements, SizeType NumElements, SIZE_T NumBytesPerElement)
+		FORCEINLINE void ResizeAllocation(SizeType CurrentNum, SizeType NewMax, SIZE_T NumBytesPerElement)
 		{
 			// Avoid calling FMemory::Realloc( nullptr, 0 ) as ANSI C mandates returning a valid pointer which is not what we want.
-			if (Data || NumElements)
+			if (Data || NewMax)
 			{
 				static_assert(sizeof(SizeType) <= sizeof(SIZE_T), "SIZE_T is expected to handle all possible sizes");
 
 				// Check for under/overflow
-				bool bInvalidResize = NumElements < 0 || NumBytesPerElement < 1 || NumBytesPerElement > (SIZE_T)MAX_int32;
+				bool bInvalidResize = NewMax < 0 || NumBytesPerElement < 1 || NumBytesPerElement > (SIZE_T)MAX_int32;
 				if constexpr (sizeof(SizeType) == sizeof(SIZE_T))
 				{
-					bInvalidResize = bInvalidResize || (SIZE_T)(USizeType)NumElements > (SIZE_T)TNumericLimits<SizeType>::Max() / NumBytesPerElement;
+					bInvalidResize = bInvalidResize || (SIZE_T)(USizeType)NewMax > (SIZE_T)TNumericLimits<SizeType>::Max() / NumBytesPerElement;
 				}
 				if (UNLIKELY(bInvalidResize))
 				{
-					UE::Core::Private::OnInvalidLLMAllocatorNum(IndexSize, NumElements, NumBytesPerElement);
+					UE::Core::Private::OnInvalidLLMAllocatorNum(IndexSize, NewMax, NumBytesPerElement);
 				}
 
-				//checkSlow(((uint64)NumElements*(uint64)ElementTypeInfo.GetSize() < (uint64)INT_MAX));
-				size_t NewSize = NumElements * NumBytesPerElement;
+				//checkSlow(((uint64)NewMax*(uint64)ElementTypeInfo.GetSize() < (uint64)INT_MAX));
+				size_t NewSize = NewMax * NumBytesPerElement;
 				Data = UE::LLMPrivate::FLLMAllocator::Get()->Realloc(Data, Size, NewSize);
 				Size = NewSize;
 			}
 		}
-		FORCEINLINE SizeType CalculateSlackReserve(SizeType NumElements, SIZE_T NumBytesPerElement) const
+		FORCEINLINE SizeType CalculateSlackReserve(SizeType NewMax, SIZE_T NumBytesPerElement) const
 		{
-			return DefaultCalculateSlackReserve(NumElements, NumBytesPerElement, true);
+			return DefaultCalculateSlackReserve(NewMax, NumBytesPerElement, true);
 		}
-		FORCEINLINE SizeType CalculateSlackShrink(SizeType NumElements, SizeType NumAllocatedElements, SIZE_T NumBytesPerElement) const
+		FORCEINLINE SizeType CalculateSlackShrink(SizeType NewMax, SizeType CurrentMax, SIZE_T NumBytesPerElement) const
 		{
-			return DefaultCalculateSlackShrink(NumElements, NumAllocatedElements, NumBytesPerElement, true);
+			return DefaultCalculateSlackShrink(NewMax, CurrentMax, NumBytesPerElement, true);
 		}
-		FORCEINLINE SizeType CalculateSlackGrow(SizeType NumElements, SizeType NumAllocatedElements, SIZE_T NumBytesPerElement) const
+		FORCEINLINE SizeType CalculateSlackGrow(SizeType NewMax, SizeType CurrentMax, SIZE_T NumBytesPerElement) const
 		{
-			return DefaultCalculateSlackGrow(NumElements, NumAllocatedElements, NumBytesPerElement, true);
+			return DefaultCalculateSlackGrow(NewMax, CurrentMax, NumBytesPerElement, true);
 		}
 
-		SIZE_T GetAllocatedSize(SizeType NumAllocatedElements, SIZE_T NumBytesPerElement) const
+		SIZE_T GetAllocatedSize(SizeType CurrentMax, SIZE_T NumBytesPerElement) const
 		{
-			return NumAllocatedElements * NumBytesPerElement;
+			return CurrentMax * NumBytesPerElement;
 		}
 
 		bool HasAllocation() const
@@ -1038,8 +1038,8 @@ public:
 };
 
 // Define the ResizeAllocation functions with the regular allocators as exported to avoid bloat
-extern template CORE_API FORCENOINLINE void TSizedLLMAllocator<32>::ForAnyElementType::ResizeAllocation(SizeType PreviousNumElements, SizeType NumElements, SIZE_T NumBytesPerElement);
-extern template CORE_API FORCENOINLINE void TSizedLLMAllocator<64>::ForAnyElementType::ResizeAllocation(SizeType PreviousNumElements, SizeType NumElements, SIZE_T NumBytesPerElement);
+extern template CORE_API FORCENOINLINE void TSizedLLMAllocator<32>::ForAnyElementType::ResizeAllocation(SizeType CurrentNum, SizeType NewMax, SIZE_T NumBytesPerElement);
+extern template CORE_API FORCENOINLINE void TSizedLLMAllocator<64>::ForAnyElementType::ResizeAllocation(SizeType CurrentNum, SizeType NewMax, SIZE_T NumBytesPerElement);
 
 // The standard container-specific allocators based on TSizedLLMAllocator; these are copied from ContainerAllocationPolicies.h
 using FDefaultLLMAllocator = TSizedLLMAllocator<32>;
diff --git a/Engine/Source/Runtime/Core/Public/Containers/ContainerAllocationPolicies.h b/Engine/Source/Runtime/Core/Public/Containers/ContainerAllocationPolicies.h
--- a/Engine/Source/Runtime/Core/Public/Containers/ContainerAllocationPolicies.h
+++ b/Engine/Source/Runtime/Core/Public/Containers/ContainerAllocationPolicies.h
@@ -142,19 +142,19 @@ template<int IndexSize> class TSizedDefaultAllocator;
 using FDefaultAllocator = TSizedDefaultAllocator<32>;
 
 template <typename SizeType>
-FORCEINLINE SizeType DefaultCalculateSlackShrink(SizeType NumElements, SizeType NumAllocatedElements, SIZE_T BytesPerElement, bool bAllowQuantize, uint32 Alignment = DEFAULT_ALIGNMENT)
+FORCEINLINE SizeType DefaultCalculateSlackShrink(SizeType NewMax, SizeType CurrentMax, SIZE_T BytesPerElement, bool bAllowQuantize, uint32 Alignment = DEFAULT_ALIGNMENT)
 {
 	SizeType Retval;
-	checkSlow(NumElements < NumAllocatedElements);
+	checkSlow(NewMax < CurrentMax);
 
 	// If the container has too much slack, shrink it to exactly fit the number of elements.
-	const SizeType CurrentSlackElements = NumAllocatedElements - NumElements;
-	const SIZE_T CurrentSlackBytes = (NumAllocatedElements - NumElements)*BytesPerElement;
+	const SizeType CurrentSlackElements = CurrentMax - NewMax;
+	const SIZE_T CurrentSlackBytes = (CurrentMax - NewMax)*BytesPerElement;
 	const bool bTooManySlackBytes = CurrentSlackBytes >= 16384;
-	const bool bTooManySlackElements = 3 * NumElements < 2 * NumAllocatedElements;
-	if ((bTooManySlackBytes || bTooManySlackElements) && (CurrentSlackElements > 64 || !NumElements)) //  hard coded 64 :-(
+	const bool bTooManySlackElements = 3 * NewMax < 2 * CurrentMax;
+	if ((bTooManySlackBytes || bTooManySlackElements) && (CurrentSlackElements > 64 || !NewMax)) //  hard coded 64 :-(
 	{
-		Retval = NumElements;
+		Retval = NewMax;
 		if (Retval > 0)
 		{
 			if (bAllowQuantize)
@@ -165,14 +165,14 @@ FORCEINLINE SizeType DefaultCalculateSlackShrink(SizeType NumElements, SizeType
 	}
 	else
 	{
-		Retval = NumAllocatedElements;
+		Retval = CurrentMax;
 	}
 
 	return Retval;
 }
 
 template <typename SizeType>
-FORCEINLINE SizeType DefaultCalculateSlackGrow(SizeType NumElements, SizeType NumAllocatedElements, SIZE_T BytesPerElement, bool bAllowQuantize, uint32 Alignment = DEFAULT_ALIGNMENT)
+FORCEINLINE SizeType DefaultCalculateSlackGrow(SizeType NewMax, SizeType CurrentMax, SIZE_T BytesPerElement, bool bAllowQuantize, uint32 Alignment = DEFAULT_ALIGNMENT)
 {
 #if !defined(AGGRESSIVE_MEMORY_SAVING)
 	#error "AGGRESSIVE_MEMORY_SAVING must be defined"
@@ -186,25 +186,25 @@ FORCEINLINE SizeType DefaultCalculateSlackGrow(SizeType NumElements, SizeType Nu
 #endif
 
 	SizeType Retval;
-	checkSlow(NumElements > NumAllocatedElements && NumElements > 0);
+	checkSlow(NewMax > CurrentMax && NewMax > 0);
 
 	SIZE_T Grow = FirstGrow; // this is the amount for the first alloc
 
 #if CONTAINER_INITIAL_ALLOC_ZERO_SLACK
-	if (NumAllocatedElements)
+	if (CurrentMax)
 	{
 		// Allocate slack for the array proportional to its size.
-		Grow = SIZE_T(NumElements) + UE_CONTAINER_SLACK_GROWTH_FACTOR_NUMERATOR * SIZE_T(NumElements) / UE_CONTAINER_SLACK_GROWTH_FACTOR_DENOMINATOR + ConstantGrow;
+		Grow = SIZE_T(NewMax) + UE_CONTAINER_SLACK_GROWTH_FACTOR_NUMERATOR * SIZE_T(NewMax) / UE_CONTAINER_SLACK_GROWTH_FACTOR_DENOMINATOR + ConstantGrow;
 	}
-	else if (SIZE_T(NumElements) > Grow)
+	else if (SIZE_T(NewMax) > Grow)
 	{
-		Grow = SIZE_T(NumElements);
+		Grow = SIZE_T(NewMax);
 	}
 #else
-	if (NumAllocatedElements || SIZE_T(NumElements) > Grow)
+	if (CurrentMax || SIZE_T(NewMax) > Grow)
 	{
 		// Allocate slack for the array proportional to its size.
-		Grow = SIZE_T(NumElements) + UE_CONTAINER_SLACK_GROWTH_FACTOR_NUMERATOR * SIZE_T(NumElements) / UE_CONTAINER_SLACK_GROWTH_FACTOR_DENOMINATOR + ConstantGrow;
+		Grow = SIZE_T(NewMax) + UE_CONTAINER_SLACK_GROWTH_FACTOR_NUMERATOR * SIZE_T(NewMax) / UE_CONTAINER_SLACK_GROWTH_FACTOR_DENOMINATOR + ConstantGrow;
 	}
 #endif
 	
@@ -217,7 +217,7 @@ FORCEINLINE SizeType DefaultCalculateSlackGrow(SizeType NumElements, SizeType Nu
 		Retval = (SizeType)Grow;
 	}
 	// NumElements and MaxElements are stored in 32 bit signed integers so we must be careful not to overflow here.
-	if (NumElements > Retval)
+	if (NewMax > Retval)
 	{
 		Retval = TNumericLimits<SizeType>::Max();
 	}
@@ -226,15 +226,15 @@ FORCEINLINE SizeType DefaultCalculateSlackGrow(SizeType NumElements, SizeType Nu
 }
 
 template <typename SizeType>
-FORCEINLINE SizeType DefaultCalculateSlackReserve(SizeType NumElements, SIZE_T BytesPerElement, bool bAllowQuantize, uint32 Alignment = DEFAULT_ALIGNMENT)
+FORCEINLINE SizeType DefaultCalculateSlackReserve(SizeType NewMax, SIZE_T BytesPerElement, bool bAllowQuantize, uint32 Alignment = DEFAULT_ALIGNMENT)
 {
-	SizeType Retval = NumElements;
-	checkSlow(NumElements > 0);
+	SizeType Retval = NewMax;
+	checkSlow(NewMax > 0);
 	if (bAllowQuantize)
 	{
 		Retval = (SizeType)(FMemory::QuantizeSize(SIZE_T(Retval) * SIZE_T(BytesPerElement), Alignment) / BytesPerElement);
 		// NumElements and MaxElements are stored in 32 bit signed integers so we must be careful not to overflow here.
-		if (NumElements > Retval)
+		if (NewMax > Retval)
 		{
 			Retval = TNumericLimits<SizeType>::Max();
 		}
@@ -311,118 +311,118 @@ public:
 
 		/**
 		 * Resizes the container's allocation.
-		 * @param PreviousNumElements - The number of elements that were stored in the previous allocation.
-		 * @param NumElements - The number of elements to allocate space for.
+		 * @param CurrentNum - The number of elements that are currently constructed at the front of the allocation.
+		 * @param NewMax - The number of elements to allocate space for.
 		 * @param NumBytesPerElement - The number of bytes/element.
 		 */
 		void ResizeAllocation(
-			SizeType PreviousNumElements,
-			SizeType NumElements,
+			SizeType CurrentNum,
+			SizeType NewMax,
 			SIZE_T NumBytesPerElement
 		);
 
 		/**
 		 * Resizes the container's allocation.
-		 * @param PreviousNumElements - The number of elements that were stored in the previous allocation.
-		 * @param NumElements - The number of elements to allocate space for.
+		 * @param CurrentNum - The number of elements that are currently constructed at the front of the allocation.
+		 * @param NewMax - The number of elements to allocate space for.
 		 * @param NumBytesPerElement - The number of bytes/element.
 		 * @param AlignmentOfElement - The alignment of the element type.
 		 *
 		 * @note  This overload only exists if TAllocatorTraits<Allocator>::SupportsElementAlignment == true.
 		 */
 		void ResizeAllocation(
-			SizeType PreviousNumElements,
-			SizeType NumElements,
+			SizeType CurrentNum,
+			SizeType NewMax,
 			SIZE_T NumBytesPerElement,
 			uint32 AlignmentOfElement
 		);
 
 		/**
 		 * Calculates the amount of slack to allocate for an array that has just grown or shrunk to a given number of elements.
-		 * @param NumElements - The number of elements to allocate space for.
+		 * @param NewMax - The number of elements to allocate space for.
 		 * @param NumBytesPerElement - The number of bytes/element.
 		 */
 		SizeType CalculateSlackReserve(
-			SizeType NumElements,
+			SizeType NewMax,
 			SIZE_T NumBytesPerElement
 		) const;
 
 		/**
 		 * Calculates the amount of slack to allocate for an array that has just grown or shrunk to a given number of elements.
-		 * @param NumElements - The number of elements to allocate space for.
+		 * @param NewMax - The number of elements to allocate space for.
 		 * @param NumBytesPerElement - The number of bytes/element.
 		 * @param AlignmentOfElement - The alignment of the element type.
 		 *
 		 * @note  This overload only exists if TAllocatorTraits<Allocator>::SupportsElementAlignment == true.
 		 */
 		SizeType CalculateSlackReserve(
-			SizeType NumElements,
+			SizeType NewMax,
 			SIZE_T NumBytesPerElement,
 			uint32 AlignmentOfElement
 		) const;
 
 		/**
 		 * Calculates the amount of slack to allocate for an array that has just shrunk to a given number of elements.
-		 * @param NumElements - The number of elements to allocate space for.
-		 * @param CurrentNumSlackElements - The current number of elements allocated.
+		 * @param NewMax - The number of elements to allocate space for.
+		 * @param CurrentMax - The number of elements for which space is currently allocated.
 		 * @param NumBytesPerElement - The number of bytes/element.
 		 */
 		SizeType CalculateSlackShrink(
-			SizeType NumElements,
-			SizeType CurrentNumSlackElements,
+			SizeType NewMax,
+			SizeType CurrentMax,
 			SIZE_T NumBytesPerElement
 			) const;
 
 		/**
 		 * Calculates the amount of slack to allocate for an array that has just shrunk to a given number of elements.
-		 * @param NumElements - The number of elements to allocate space for.
-		 * @param CurrentNumSlackElements - The current number of elements allocated.
+		 * @param NewMax - The number of elements to allocate space for.
+		 * @param CurrentMax - The number of elements for which space is currently allocated.
 		 * @param NumBytesPerElement - The number of bytes/element.
 		 * @param AlignmentOfElement - The alignment of the element type.
 		 *
 		 * @note  This overload only exists if TAllocatorTraits<Allocator>::SupportsElementAlignment == true.
 		 */
 		SizeType CalculateSlackShrink(
-			SizeType NumElements,
-			SizeType CurrentNumSlackElements,
+			SizeType NewMax,
+			SizeType CurrentMax,
 			SIZE_T NumBytesPerElement,
 			uint32 AlignmentOfElement
 		) const;
 
 		/**
 		 * Calculates the amount of slack to allocate for an array that has just grown to a given number of elements.
-		 * @param NumElements - The number of elements to allocate space for.
-		 * @param CurrentNumSlackElements - The current number of elements allocated.
+		 * @param NewMax - The number of elements to allocate space for.
+		 * @param CurrentMax - The number of elements for which space is currently allocated.
 		 * @param NumBytesPerElement - The number of bytes/element.
 		 */
 		SizeType CalculateSlackGrow(
-			SizeType NumElements,
-			SizeType CurrentNumSlackElements,
+			SizeType NewMax,
+			SizeType CurrentMax,
 			SIZE_T NumBytesPerElement
 		) const;
 
 		/**
 		 * Calculates the amount of slack to allocate for an array that has just grown to a given number of elements.
-		 * @param NumElements - The number of elements to allocate space for.
-		 * @param CurrentNumSlackElements - The current number of elements allocated.
+		 * @param NewMax - The number of elements to allocate space for.
+		 * @param CurrentMax - The number of elements for which space is currently allocated.
 		 * @param NumBytesPerElement - The number of bytes/element.
 		 * @param AlignmentOfElement - The alignment of the element type.
 		 *
 		 * @note  This overload only exists if TAllocatorTraits<Allocator>::SupportsElementAlignment == true.
 		 */
 		SizeType CalculateSlackGrow(
-			SizeType NumElements,
-			SizeType CurrentNumSlackElements,
+			SizeType NewMax,
+			SizeType CurrentMax,
 			SIZE_T NumBytesPerElement,
 			uint32 AlignmentOfElement
 		) const;
 
 		/**
 		 * Returns the size of any requested heap allocation currently owned by the allocator.
-		 * @param NumAllocatedElements - The number of elements allocated by the container.
+		 * @param CurrentMax - The number of elements for which space is currently allocated.
 		 * @param NumBytesPerElement - The number of bytes/element.
 		 */
-		SIZE_T GetAllocatedSize(SizeType NumAllocatedElements, SIZE_T NumBytesPerElement) const;
+		SIZE_T GetAllocatedSize(SizeType CurrentMax, SIZE_T NumBytesPerElement) const;
 
 		/** Returns true if the allocator has made any heap allocations */
 		bool HasAllocation() const;
@@ -507,45 +507,45 @@ public:
 			return Data;
 		}
 		void ResizeAllocation(
-			SizeType PreviousNumElements,
-			SizeType NumElements,
+			SizeType CurrentNum,
+			SizeType NewMax,
 			SIZE_T NumBytesPerElement
 			)
 		{
 			// Avoid calling FMemory::Realloc( nullptr, 0 ) as ANSI C mandates returning a valid pointer which is not what we want.
-			if (Data || NumElements)
+			if (Data || NewMax)
 			{
 				static_assert(sizeof(int32) <= sizeof(SIZE_T), "SIZE_T is expected to be larger than int32");
 
 				// Check for under/overflow
-				if (UNLIKELY(NumElements < 0 || NumBytesPerElement < 1 || NumBytesPerElement > (SIZE_T)MAX_int32))
+				if (UNLIKELY(NewMax < 0 || NumBytesPerElement < 1 || NumBytesPerElement > (SIZE_T)MAX_int32))
 				{
-					UE::Core::Private::OnInvalidAlignedHeapAllocatorNum(NumElements, NumBytesPerElement);
+					UE::Core::Private::OnInvalidAlignedHeapAllocatorNum(NewMax, NumBytesPerElement);
 				}
 
 #if UE_ENABLE_ARRAY_SLACK_TRACKING
-				Data = (FScriptContainerElement*)FArraySlackTrackingHeader::Realloc(Data, NumElements, NumBytesPerElement, Alignment > alignof(FArraySlackTrackingHeader) ? Alignment : alignof(FArraySlackTrackingHeader));
+				Data = (FScriptContainerElement*)FArraySlackTrackingHeader::Realloc(Data, NewMax, NumBytesPerElement, Alignment > alignof(FArraySlackTrackingHeader) ? Alignment : alignof(FArraySlackTrackingHeader));
 #else
-				Data = (FScriptContainerElement*)FMemory::Realloc( Data, NumElements*NumBytesPerElement, Alignment );
+				Data = (FScriptContainerElement*)FMemory::Realloc( Data, NewMax*NumBytesPerElement, Alignment );
 #endif
 			}
 		}
-		FORCEINLINE SizeType CalculateSlackReserve(SizeType NumElements, SIZE_T NumBytesPerElement) const
+		FORCEINLINE SizeType CalculateSlackReserve(SizeType NewMax, SIZE_T NumBytesPerElement) const
 		{
-			return DefaultCalculateSlackReserve(NumElements, NumBytesPerElement, true, Alignment);
+			return DefaultCalculateSlackReserve(NewMax, NumBytesPerElement, true, Alignment);
 		}
-		FORCEINLINE SizeType CalculateSlackShrink(SizeType NumElements, SizeType NumAllocatedElements, SIZE_T NumBytesPerElement) const
+		FORCEINLINE SizeType CalculateSlackShrink(SizeType NewMax, SizeType CurrentMax, SIZE_T NumBytesPerElement) const
 		{
-			return DefaultCalculateSlackShrink(NumElements, NumAllocatedElements, NumBytesPerElement, true, Alignment);
+			return DefaultCalculateSlackShrink(NewMax, CurrentMax, NumBytesPerElement, true, Alignment);
 		}
-		FORCEINLINE SizeType CalculateSlackGrow(SizeType NumElements, SizeType NumAllocatedElements, SIZE_T NumBytesPerElement) const
+		FORCEINLINE SizeType CalculateSlackGrow(SizeType NewMax, SizeType CurrentMax, SIZE_T NumBytesPerElement) const
 		{
-			return DefaultCalculateSlackGrow(NumElements, NumAllocatedElements, NumBytesPerElement, true, Alignment);
+			return DefaultCalculateSlackGrow(NewMax, CurrentMax, NumBytesPerElement, true, Alignment);
 		}
 
-		SIZE_T GetAllocatedSize(SizeType NumAllocatedElements, SIZE_T NumBytesPerElement) const
+		SIZE_T GetAllocatedSize(SizeType CurrentMax, SIZE_T NumBytesPerElement) const
 		{
-			return NumAllocatedElements * NumBytesPerElement;
+			return CurrentMax * NumBytesPerElement;
 		}
 
 		bool HasAllocation() const
@@ -697,84 +697,84 @@ public:
 		{
 			return Data;
 		}
-		void ResizeAllocation(SizeType PreviousNumElements, SizeType NumElements, SIZE_T NumBytesPerElement)
+		void ResizeAllocation(SizeType CurrentNum, SizeType NewMax, SIZE_T NumBytesPerElement)
 		{
 			// Avoid calling FMemory::Realloc( nullptr, 0 ) as ANSI C mandates returning a valid pointer which is not what we want.
-			if (Data || NumElements)
+			if (Data || NewMax)
 			{
 				static_assert(sizeof(SizeType) <= sizeof(SIZE_T), "SIZE_T is expected to handle all possible sizes");
 
 				// Check for under/overflow
-				bool bInvalidResize = NumElements < 0 || NumBytesPerElement < 1 || NumBytesPerElement > (SIZE_T)MAX_int32;
+				bool bInvalidResize = NewMax < 0 || NumBytesPerElement < 1 || NumBytesPerElement > (SIZE_T)MAX_int32;
 				if constexpr (sizeof(SizeType) == sizeof(SIZE_T))
 				{
-					bInvalidResize = bInvalidResize || (SIZE_T)(USizeType)NumElements > (SIZE_T)TNumericLimits<SizeType>::Max() / NumBytesPerElement;
+					bInvalidResize = bInvalidResize || (SIZE_T)(USizeType)NewMax > (SIZE_T)TNumericLimits<SizeType>::Max() / NumBytesPerElement;
 				}
 				if (UNLIKELY(bInvalidResize))
 				{
-					UE::Core::Private::OnInvalidSizedHeapAllocatorNum(IndexSize, NumElements, NumBytesPerElement);
+					UE::Core::Private::OnInvalidSizedHeapAllocatorNum(IndexSize, NewMax, NumBytesPerElement);
 				}
 
 #if UE_ENABLE_ARRAY_SLACK_TRACKING
-				Data = (FScriptContainerElement*)FArraySlackTrackingHeader::Realloc(Data, NumElements, NumBytesPerElement, 0);
+				Data = (FScriptContainerElement*)FArraySlackTrackingHeader::Realloc(Data, NewMax, NumBytesPerElement, 0);
 #else
-				Data = (FScriptContainerElement*)BaseMallocType::Realloc( Data, NumElements*NumBytesPerElement );
+				Data = (FScriptContainerElement*)BaseMallocType::Realloc( Data, NewMax*NumBytesPerElement );
 #endif
 			}
 		}
-		void ResizeAllocation(SizeType PreviousNumElements, SizeType NumElements, SIZE_T NumBytesPerElement, uint32 AlignmentOfElement)
+		void ResizeAllocation(SizeType CurrentNum, SizeType NewMax, SIZE_T NumBytesPerElement, uint32 AlignmentOfElement)
 		{
 			// Avoid calling FMemory::Realloc( nullptr, 0 ) as ANSI C mandates returning a valid pointer which is not what we want.
-			if (Data || NumElements)
+			if (Data || NewMax)
 			{
 				static_assert(sizeof(SizeType) <= sizeof(SIZE_T), "SIZE_T is expected to handle all possible sizes");
 
 				// Check for under/overflow
-				bool bInvalidResize = NumElements < 0 || NumBytesPerElement < 1 || NumBytesPerElement > (SIZE_T)MAX_int32;
+				bool bInvalidResize = NewMax < 0 || NumBytesPerElement < 1 || NumBytesPerElement > (SIZE_T)MAX_int32;
 				if constexpr (sizeof(SizeType) == sizeof(SIZE_T))
 				{
-					bInvalidResize = bInvalidResize || ((SIZE_T)(USizeType)NumElements > (SIZE_T)TNumericLimits<SizeType>::Max() / NumBytesPerElement);
+					bInvalidResize = bInvalidResize || ((SIZE_T)(USizeType)NewMax > (SIZE_T)TNumericLimits<SizeType>::Max() / NumBytesPerElement);
 				}
 				if (UNLIKELY(bInvalidResize))
 				{
-					UE::Core::Private::OnInvalidSizedHeapAllocatorNum(IndexSize, NumElements, NumBytesPerElement);
+					UE::Core::Private::OnInvalidSizedHeapAllocatorNum(IndexSize, NewMax, NumBytesPerElement);
 				}
 
 #if UE_ENABLE_ARRAY_SLACK_TRACKING
-				Data = (FScriptContainerElement*)FArraySlackTrackingHeader::Realloc(Data, NumElements, NumBytesPerElement, AlignmentOfElement > alignof(FArraySlackTrackingHeader) ? AlignmentOfElement : alignof(FArraySlackTrackingHeader));
+				Data = (FScriptContainerElement*)FArraySlackTrackingHeader::Realloc(Data, NewMax, NumBytesPerElement, AlignmentOfElement > alignof(FArraySlackTrackingHeader) ? AlignmentOfElement : alignof(FArraySlackTrackingHeader));
 #else
-				Data = (FScriptContainerElement*)BaseMallocType::Realloc( Data, NumElements*NumBytesPerElement, AlignmentOfElement );
+				Data = (FScriptContainerElement*)BaseMallocType::Realloc( Data, NewMax*NumBytesPerElement, AlignmentOfElement );
 #endif
 			}
 		}
-		FORCEINLINE SizeType CalculateSlackReserve(SizeType NumElements, SIZE_T NumBytesPerElement) const
+		FORCEINLINE SizeType CalculateSlackReserve(SizeType NewMax, SIZE_T NumBytesPerElement) const
 		{
-			return DefaultCalculateSlackReserve(NumElements, NumBytesPerElement, true);
+			return DefaultCalculateSlackReserve(NewMax, NumBytesPerElement, true);
 		}
-		FORCEINLINE SizeType CalculateSlackReserve(SizeType NumElements, SIZE_T NumBytesPerElement, uint32 AlignmentOfElement) const
+		FORCEINLINE SizeType CalculateSlackReserve(SizeType NewMax, SIZE_T NumBytesPerElement, uint32 AlignmentOfElement) const
 		{
-			return DefaultCalculateSlackReserve(NumElements, NumBytesPerElement, true, (uint32)AlignmentOfElement);
+			return DefaultCalculateSlackReserve(NewMax, NumBytesPerElement, true, (uint32)AlignmentOfElement);
 		}
-		FORCEINLINE SizeType CalculateSlackShrink(SizeType NumElements, SizeType NumAllocatedElements, SIZE_T NumBytesPerElement) const
+		FORCEINLINE SizeType CalculateSlackShrink(SizeType NewMax, SizeType CurrentMax, SIZE_T NumBytesPerElement) const
 		{
-			return DefaultCalculateSlackShrink(NumElements, NumAllocatedElements, NumBytesPerElement, true);
+			return DefaultCalculateSlackShrink(NewMax, CurrentMax, NumBytesPerElement, true);
 		}
-		FORCEINLINE SizeType CalculateSlackShrink(SizeType NumElements, SizeType NumAllocatedElements, SIZE_T NumBytesPerElement, uint32 AlignmentOfElement) const
+		FORCEINLINE SizeType CalculateSlackShrink(SizeType NewMax, SizeType CurrentMax, SIZE_T NumBytesPerElement, uint32 AlignmentOfElement) const
 		{
-			return DefaultCalculateSlackShrink(NumElements, NumAllocatedElements, NumBytesPerElement, true, (uint32)AlignmentOfElement);
+			return DefaultCalculateSlackShrink(NewMax, CurrentMax, NumBytesPerElement, true, (uint32)AlignmentOfElement);
 		}
-		FORCEINLINE SizeType CalculateSlackGrow(SizeType NumElements, SizeType NumAllocatedElements, SIZE_T NumBytesPerElement) const
+		FORCEINLINE SizeType CalculateSlackGrow(SizeType NewMax, SizeType CurrentMax, SIZE_T NumBytesPerElement) const
 		{
-			return DefaultCalculateSlackGrow(NumElements, NumAllocatedElements, NumBytesPerElement, true);
+			return DefaultCalculateSlackGrow(NewMax, CurrentMax, NumBytesPerElement, true);
 		}
-		FORCEINLINE SizeType CalculateSlackGrow(SizeType NumElements, SizeType NumAllocatedElements, SIZE_T NumBytesPerElement, uint32 AlignmentOfElement) const
+		FORCEINLINE SizeType CalculateSlackGrow(SizeType NewMax, SizeType CurrentMax, SIZE_T NumBytesPerElement, uint32 AlignmentOfElement) const
 		{
-			return DefaultCalculateSlackGrow(NumElements, NumAllocatedElements, NumBytesPerElement, true, (uint32)AlignmentOfElement);
+			return DefaultCalculateSlackGrow(NewMax, CurrentMax, NumBytesPerElement, true, (uint32)AlignmentOfElement);
 		}
 
-		SIZE_T GetAllocatedSize(SizeType NumAllocatedElements, SIZE_T NumBytesPerElement) const
+		SIZE_T GetAllocatedSize(SizeType CurrentMax, SIZE_T NumBytesPerElement) const
 		{
-			return NumAllocatedElements * NumBytesPerElement;
+			return CurrentMax * NumBytesPerElement;
 		}
 
 		bool HasAllocation() const
@@ -827,10 +827,10 @@ public:
 
 // Define the ResizeAllocation functions with the regular allocator as exported to avoid bloat
 #if !IS_MERGEDMODULES
-extern template CORE_API FORCENOINLINE void TSizedHeapAllocator<32, FMemory>::ForAnyElementType::ResizeAllocation(SizeType PreviousNumElements, SizeType NumElements, SIZE_T NumBytesPerElement);
-extern template CORE_API FORCENOINLINE void TSizedHeapAllocator<32, FMemory>::ForAnyElementType::ResizeAllocation(SizeType PreviousNumElements, SizeType NumElements, SIZE_T NumBytesPerElement, uint32 AlignmentOfElement);
-extern template CORE_API FORCENOINLINE void TSizedHeapAllocator<64, FMemory>::ForAnyElementType::ResizeAllocation(SizeType PreviousNumElements, SizeType NumElements, SIZE_T NumBytesPerElement);
-extern template CORE_API FORCENOINLINE void TSizedHeapAllocator<64, FMemory>::ForAnyElementType::ResizeAllocation(SizeType PreviousNumElements, SizeType NumElements, SIZE_T NumBytesPerElement, uint32 AlignmentOfElement);
+extern template CORE_API FORCENOINLINE void TSizedHeapAllocator<32, FMemory>::ForAnyElementType::ResizeAllocation(SizeType CurrentNum, SizeType NewMax, SIZE_T NumBytesPerElement);
+extern template CORE_API FORCENOINLINE void TSizedHeapAllocator<32, FMemory>::ForAnyElementType::ResizeAllocation(SizeType CurrentNum, SizeType NewMax, SIZE_T NumBytesPerElement, uint32 AlignmentOfElement);
+extern template CORE_API FORCENOINLINE void TSizedHeapAllocator<64, FMemory>::ForAnyElementType::ResizeAllocation(SizeType CurrentNum, SizeType NewMax, SIZE_T NumBytesPerElement);
+extern template CORE_API FORCENOINLINE void TSizedHeapAllocator<64, FMemory>::ForAnyElementType::ResizeAllocation(SizeType CurrentNum, SizeType NewMax, SIZE_T NumBytesPerElement, uint32 AlignmentOfElement);
 #endif
 
 template <uint8 IndexSize>
@@ -906,15 +906,15 @@ public:
 			return GetInlineElements();
 		}
 
-		void ResizeAllocation(SizeType PreviousNumElements, SizeType NumElements,SIZE_T NumBytesPerElement)
+		void ResizeAllocation(SizeType CurrentNum, SizeType NewMax,SIZE_T NumBytesPerElement)
 		{
 			// Check if the new allocation will fit in the inline data area.
-			if(NumElements <= NumInlineElements)
+			if(NewMax <= NumInlineElements)
 			{
 				// If the old allocation wasn't in the inline data area, relocate it into the inline data area.
 				if(SecondaryData.GetAllocation())
 				{
-					RelocateConstructItems<ElementType>((void*)InlineData, (ElementType*)SecondaryData.GetAllocation(), PreviousNumElements);
+					RelocateConstructItems<ElementType>((void*)InlineData, (ElementType*)SecondaryData.GetAllocation(), CurrentNum);
 
 					// Free the old indirect allocation.
 					SecondaryData.ResizeAllocation(0,0,NumBytesPerElement);
@@ -925,49 +925,49 @@ public:
 				if(!SecondaryData.GetAllocation())
 				{
 					// Allocate new indirect memory for the data.
-					SecondaryData.ResizeAllocation(0,NumElements,NumBytesPerElement);
+					SecondaryData.ResizeAllocation(0,NewMax,NumBytesPerElement);
 
 					// Move the data out of the inline data area into the new allocation.
-					RelocateConstructItems<ElementType>((void*)SecondaryData.GetAllocation(), GetInlineElements(), PreviousNumElements);
+					RelocateConstructItems<ElementType>((void*)SecondaryData.GetAllocation(), GetInlineElements(), CurrentNum);
 				}
 				else
 				{
 					// Reallocate the indirect data for the new size.
-					SecondaryData.ResizeAllocation(PreviousNumElements, NumElements, NumBytesPerElement);
+					SecondaryData.ResizeAllocation(CurrentNum, NewMax, NumBytesPerElement);
 				}
 			}
 		}
 
-		FORCEINLINE SizeType CalculateSlackReserve(SizeType NumElements, SIZE_T NumBytesPerElement) const
+		FORCEINLINE SizeType CalculateSlackReserve(SizeType NewMax, SIZE_T NumBytesPerElement) const
 		{
 			// If the elements use less space than the inline allocation, only use the inline allocation as slack.
-			return NumElements <= NumInlineElements ?
+			return NewMax <= NumInlineElements ?
 				NumInlineElements :
-				SecondaryData.CalculateSlackReserve(NumElements, NumBytesPerElement);
+				SecondaryData.CalculateSlackReserve(NewMax, NumBytesPerElement);
 		}
-		FORCEINLINE SizeType CalculateSlackShrink(SizeType NumElements, SizeType NumAllocatedElements, SIZE_T NumBytesPerElement) const
+		FORCEINLINE SizeType CalculateSlackShrink(SizeType NewMax, SizeType CurrentMax, SIZE_T NumBytesPerElement) const
 		{
 			// If the elements use less space than the inline allocation, only use the inline allocation as slack.
-			return NumElements <= NumInlineElements ?
+			return NewMax <= NumInlineElements ?
 				NumInlineElements :
-				SecondaryData.CalculateSlackShrink(NumElements, NumAllocatedElements, NumBytesPerElement);
+				SecondaryData.CalculateSlackShrink(NewMax, CurrentMax, NumBytesPerElement);
 		}
-		FORCEINLINE SizeType CalculateSlackGrow(SizeType NumElements, SizeType NumAllocatedElements, SIZE_T NumBytesPerElement) const
+		FORCEINLINE SizeType CalculateSlackGrow(SizeType NewMax, SizeType CurrentMax, SIZE_T NumBytesPerElement) const
 		{
 			// If the elements use less space than the inline allocation, only use the inline allocation as slack.
 			// Also, when computing slack growth, don't count inline elements -- the slack algorithm has a special
 			// case to save memory on the initial heap allocation, versus subsequent reallocations, and we don't
 			// want the inline elements to be treated as if they were the first heap allocation.
-			return NumElements <= NumInlineElements ?
+			return NewMax <= NumInlineElements ?
 				NumInlineElements :
-				SecondaryData.CalculateSlackGrow(NumElements, NumAllocatedElements <= NumInlineElements ? 0 : NumAllocatedElements, NumBytesPerElement);
+				SecondaryData.CalculateSlackGrow(NewMax, CurrentMax <= NumInlineElements ? 0 : CurrentMax, NumBytesPerElement);
 		}
 
-		SIZE_T GetAllocatedSize(SizeType NumAllocatedElements, SIZE_T NumBytesPerElement) const
+		SIZE_T GetAllocatedSize(SizeType CurrentMax, SIZE_T NumBytesPerElement) const
 		{
-			if (NumAllocatedElements > NumInlineElements)
+			if (CurrentMax > NumInlineElements)
 			{
-				return SecondaryData.GetAllocatedSize(NumAllocatedElements, NumBytesPerElement);
+				return SecondaryData.GetAllocatedSize(CurrentMax, NumBytesPerElement);
 			}
 			return 0;
 		}
@@ -1027,7 +1027,9 @@ using TInlineAllocator64 = TSizedInlineAllocator<NumInlineElements, 64, Secondar
 /**
  * Implements a variant of TInlineAllocator with a secondary heap allocator that is allowed to store a pointer to its inline elements.
  * This allows caching a pointer to the elements which avoids any conditional logic in GetAllocation(), but prevents the allocator being trivially relocatable.
- * All UE4 allocators typically rely on elements being trivially relocatable, so instances of this allocator cannot be used in other containers.
+ * All UE allocators typically rely on elements being trivially relocatable, so instances of this allocator cannot be used in other containers.
+ *
+ * NOTE: instances of this allocator - or containers which use them - are non-trivially-relocatable, but the allocator still expects elements themselves to be trivially-relocatable.
  */
 template <uint32 NumInlineElements>
 class TNonRelocatableInlineAllocator
@@ -1097,15 +1099,15 @@ public:
 			return Data;
 		}
 
-		void ResizeAllocation(SizeType PreviousNumElements, SizeType NumElements,SIZE_T NumBytesPerElement)
+		void ResizeAllocation(SizeType CurrentNum, SizeType NewMax,SIZE_T NumBytesPerElement)
 		{
 			// Check if the new allocation will fit in the inline data area.
-			if(NumElements <= NumInlineElements)
+			if(NewMax <= NumInlineElements)
 			{
 				// If the old allocation wasn't in the inline data area, relocate it into the inline data area.
 				if(HasAllocation())
 				{
-					RelocateConstructItems<ElementType>(GetInlineElements(), Data, PreviousNumElements);
+					RelocateConstructItems<ElementType>(GetInlineElements(), Data, CurrentNum);
 #if UE_ENABLE_ARRAY_SLACK_TRACKING
 					FArraySlackTrackingHeader::Free(Data);
 #else
@@ -1120,47 +1122,47 @@ public:
 				{
 					// Reallocate the indirect data for the new size.
 #if UE_ENABLE_ARRAY_SLACK_TRACKING
-					Data = (ElementType*)FArraySlackTrackingHeader::Realloc(Data, (int32)NumElements, (int32)NumBytesPerElement, 0);
+					Data = (ElementType*)FArraySlackTrackingHeader::Realloc(Data, (int32)NewMax, (int32)NumBytesPerElement, 0);
 #else
-					Data = (ElementType*)FMemory::Realloc(Data, NumElements*NumBytesPerElement);
+					Data = (ElementType*)FMemory::Realloc(Data, NewMax*NumBytesPerElement);
 #endif
 				}
 				else
 				{
 					// Allocate new indirect memory for the data.
 #if UE_ENABLE_ARRAY_SLACK_TRACKING
-					Data = (ElementType*)FArraySlackTrackingHeader::Realloc(nullptr, (int32)NumElements, (int32)NumBytesPerElement, 0);
+					Data = (ElementType*)FArraySlackTrackingHeader::Realloc(nullptr, (int32)NewMax, (int32)NumBytesPerElement, 0);
 #else
-					Data = (ElementType*)FMemory::Realloc(nullptr, NumElements*NumBytesPerElement);
+					Data = (ElementType*)FMemory::Realloc(nullptr, NewMax*NumBytesPerElement);
 #endif
 
 					// Move the data out of the inline data area into the new allocation.
-					RelocateConstructItems<ElementType>(Data, GetInlineElements(), PreviousNumElements);
+					RelocateConstructItems<ElementType>(Data, GetInlineElements(), CurrentNum);
 				}
 			}
 		}
 
-		FORCEINLINE SizeType CalculateSlackReserve(SizeType NumElements, SIZE_T NumBytesPerElement) const
+		FORCEINLINE SizeType CalculateSlackReserve(SizeType NewMax, SIZE_T NumBytesPerElement) const
 		{
 			// If the elements use less space than the inline allocation, only use the inline allocation as slack.
-			return (NumElements <= NumInlineElements) ? NumInlineElements : DefaultCalculateSlackReserve(NumElements, NumBytesPerElement, true);
+			return (NewMax <= NumInlineElements) ? NumInlineElements : DefaultCalculateSlackReserve(NewMax, NumBytesPerElement, true);
 		}
 
-		FORCEINLINE SizeType CalculateSlackShrink(SizeType NumElements, SizeType NumAllocatedElements, SIZE_T NumBytesPerElement) const
+		FORCEINLINE SizeType CalculateSlackShrink(SizeType NewMax, SizeType CurrentMax, SIZE_T NumBytesPerElement) const
 		{
 			// If the elements use less space than the inline allocation, only use the inline allocation as slack.
-			return (NumElements <= NumInlineElements) ? NumInlineElements : DefaultCalculateSlackShrink(NumElements, NumAllocatedElements, NumBytesPerElement, true);
+			return (NewMax <= NumInlineElements) ? NumInlineElements : DefaultCalculateSlackShrink(NewMax, CurrentMax, NumBytesPerElement, true);
 		}
 
-		FORCEINLINE SizeType CalculateSlackGrow(SizeType NumElements, SizeType NumAllocatedElements, SIZE_T NumBytesPerElement) const
+		FORCEINLINE SizeType CalculateSlackGrow(SizeType NewMax, SizeType CurrentMax, SIZE_T NumBytesPerElement) const
 		{
 			// If the elements use less space than the inline allocation, only use the inline allocation as slack.
-			return (NumElements <= NumInlineElements) ? NumInlineElements : DefaultCalculateSlackGrow(NumElements, NumAllocatedElements, NumBytesPerElement, true);
+			return (NewMax <= NumInlineElements) ? NumInlineElements : DefaultCalculateSlackGrow(NewMax, CurrentMax, NumBytesPerElement, true);
 		}
 
-		SIZE_T GetAllocatedSize(SizeType NumAllocatedElements, SIZE_T NumBytesPerElement) const
+		SIZE_T GetAllocatedSize(SizeType CurrentMax, SIZE_T NumBytesPerElement) const
 		{
-			return HasAllocation()? (NumAllocatedElements * NumBytesPerElement) : 0;
+			return HasAllocation() ? (CurrentMax * NumBytesPerElement) : 0;
 		}
 
 		FORCEINLINE bool HasAllocation() const
@@ -1253,32 +1255,32 @@ public:
 			return GetInlineElements();
 		}
 
-		void ResizeAllocation(SizeType PreviousNumElements, SizeType NumElements,SIZE_T NumBytesPerElement)
+		void ResizeAllocation(SizeType CurrentNum, SizeType NewMax,SIZE_T NumBytesPerElement)
 		{
 			// Ensure the requested allocation will fit in the inline data area.
-			check(NumElements >= 0 && NumElements <= NumInlineElements);
+			check(NewMax >= 0 && NewMax <= NumInlineElements);
 		}
 
-		FORCEINLINE SizeType CalculateSlackReserve(SizeType NumElements, SIZE_T NumBytesPerElement) const
+		FORCEINLINE SizeType CalculateSlackReserve(SizeType NewMax, SIZE_T NumBytesPerElement) const
 		{
 			// Ensure the requested allocation will fit in the inline data area.
-			check(NumElements <= NumInlineElements);
+			check(NewMax <= NumInlineElements);
 			return NumInlineElements;
 		}
-		FORCEINLINE SizeType CalculateSlackShrink(SizeType NumElements, SizeType NumAllocatedElements, SIZE_T NumBytesPerElement) const
+		FORCEINLINE SizeType CalculateSlackShrink(SizeType NewMax, SizeType CurrentMax, SIZE_T NumBytesPerElement) const
 		{
 			// Ensure the requested allocation will fit in the inline data area.
-			check(NumAllocatedElements <= NumInlineElements);
+			check(CurrentMax <= NumInlineElements);
 			return NumInlineElements;
 		}
-		FORCEINLINE SizeType CalculateSlackGrow(SizeType NumElements, SizeType NumAllocatedElements, SIZE_T NumBytesPerElement) const
+		FORCEINLINE SizeType CalculateSlackGrow(SizeType NewMax, SizeType CurrentMax, SIZE_T NumBytesPerElement) const
 		{
 			// Ensure the requested allocation will fit in the inline data area.
-			check(NumElements <= NumInlineElements);
+			check(CurrentMax <= NumInlineElements);
 			return NumInlineElements;
 		}
 
-		SIZE_T GetAllocatedSize(SizeType NumAllocatedElements, SIZE_T NumBytesPerElement) const
+		SIZE_T GetAllocatedSize(SizeType CurrentMax, SIZE_T NumBytesPerElement) const
 		{
 			return 0;
 		}
diff --git a/Engine/Source/Runtime/Core/Public/Experimental/ConcurrentLinearAllocator.h b/Engine/Source/Runtime/Core/Public/Experimental/ConcurrentLinearAllocator.h
--- a/Engine/Source/Runtime/Core/Public/Experimental/ConcurrentLinearAllocator.h
+++ b/Engine/Source/Runtime/Core/Public/Experimental/ConcurrentLinearAllocator.h
@@ -614,34 +614,34 @@ public:
 		{
 			return Data;
 		}
-		void ResizeAllocation(SizeType PreviousNumElements, SizeType NumElements, SIZE_T NumBytesPerElement)
+		void ResizeAllocation(SizeType CurrentNum, SizeType NewMax, SIZE_T NumBytesPerElement)
 		{
 			static_assert(sizeof(int32) <= sizeof(SIZE_T), "SIZE_T is expected to be larger than int32");
 
 			// Check for under/overflow
-			if (UNLIKELY(NumElements < 0 || NumBytesPerElement < 1 || NumBytesPerElement > (SIZE_T)MAX_int32))
+			if (UNLIKELY(NewMax < 0 || NumBytesPerElement < 1 || NumBytesPerElement > (SIZE_T)MAX_int32))
 			{
-				UE::Core::Private::OnInvalidConcurrentLinearArrayAllocatorNum(NumElements, NumBytesPerElement);
+				UE::Core::Private::OnInvalidConcurrentLinearArrayAllocatorNum(NewMax, NumBytesPerElement);
 			}
 
-			Data = (ElementType*)TConcurrentLinearAllocator<BlockAllocationTag>::Realloc(Data, NumElements * NumBytesPerElement, alignof(ElementType));
+			Data = (ElementType*)TConcurrentLinearAllocator<BlockAllocationTag>::Realloc(Data, NewMax * NumBytesPerElement, alignof(ElementType));
 		}
-		SizeType CalculateSlackReserve(SizeType NumElements, SIZE_T NumBytesPerElement) const
+		SizeType CalculateSlackReserve(SizeType NewMax, SIZE_T NumBytesPerElement) const
 		{
-			return DefaultCalculateSlackReserve(NumElements, NumBytesPerElement, false);
+			return DefaultCalculateSlackReserve(NewMax, NumBytesPerElement, false);
 		}
-		SizeType CalculateSlackShrink(SizeType NumElements, SizeType NumAllocatedElements, SIZE_T NumBytesPerElement) const
+		SizeType CalculateSlackShrink(SizeType NewMax, SizeType CurrentMax, SIZE_T NumBytesPerElement) const
 		{
-			return DefaultCalculateSlackShrink(NumElements, NumAllocatedElements, NumBytesPerElement, false);
+			return DefaultCalculateSlackShrink(NewMax, CurrentMax, NumBytesPerElement, false);
 		}
-		SizeType CalculateSlackGrow(SizeType NumElements, SizeType NumAllocatedElements, SIZE_T NumBytesPerElement) const
+		SizeType CalculateSlackGrow(SizeType NewMax, SizeType CurrentMax, SIZE_T NumBytesPerElement) const
 		{
-			return DefaultCalculateSlackGrow(NumElements, NumAllocatedElements, NumBytesPerElement, false);
+			return DefaultCalculateSlackGrow(NewMax, CurrentMax, NumBytesPerElement, false);
 		}
 
-		SIZE_T GetAllocatedSize(SizeType NumAllocatedElements, SIZE_T NumBytesPerElement) const
+		SIZE_T GetAllocatedSize(SizeType CurrentMax, SIZE_T NumBytesPerElement) const
 		{
-			return NumAllocatedElements * NumBytesPerElement;
+			return CurrentMax * NumBytesPerElement;
 		}
 
 		bool HasAllocation() const
diff --git a/Engine/Source/Runtime/Core/Public/Misc/MemStack.h b/Engine/Source/Runtime/Core/Public/Misc/MemStack.h
--- a/Engine/Source/Runtime/Core/Public/Misc/MemStack.h
+++ b/Engine/Source/Runtime/Core/Public/Misc/MemStack.h
@@ -423,49 +423,49 @@ public:
 		}
 
 		//@TODO: FLOATPRECISION: Takes SIZE_T input but doesn't actually support it
-		void ResizeAllocation(SizeType PreviousNumElements, SizeType NumElements,SIZE_T NumBytesPerElement)
+		void ResizeAllocation(SizeType CurrentNum, SizeType NewMax,SIZE_T NumBytesPerElement)
 		{
 			void* OldData = Data;
-			if( NumElements )
+			if( NewMax )
 			{
 				static_assert(sizeof(int32) <= sizeof(SIZE_T), "SIZE_T is expected to be larger than int32");
 
 				// Check for under/overflow
-				if (UNLIKELY(NumElements < 0 || NumBytesPerElement < 1 || NumBytesPerElement > (SIZE_T)MAX_int32))
+				if (UNLIKELY(NewMax < 0 || NumBytesPerElement < 1 || NumBytesPerElement > (SIZE_T)MAX_int32))
 				{
-					UE::Core::Private::OnInvalidMemStackAllocatorNum(NumElements, NumBytesPerElement);
+					UE::Core::Private::OnInvalidMemStackAllocatorNum(NewMax, NumBytesPerElement);
 				}
 
 				// Allocate memory from the stack.
 				Data = (ElementType*)FMemStack::Get().PushBytes(
-					(int32)(NumElements * NumBytesPerElement),
+					(int32)(NewMax * NumBytesPerElement),
 					FMath::Max(Alignment,(uint32)alignof(ElementType))
 					);
 
-				// If the container previously held elements, copy them into the new allocation.
-				if(OldData && PreviousNumElements)
+				// If the container currently holds elements, copy them into the new allocation.
+				if(OldData && CurrentNum)
 				{
-					const SizeType NumCopiedElements = FMath::Min(NumElements,PreviousNumElements);
+					const SizeType NumCopiedElements = FMath::Min(NewMax,CurrentNum);
 					FMemory::Memcpy(Data,OldData,NumCopiedElements * NumBytesPerElement);
 				}
 			}
 		}
-		FORCEINLINE SizeType CalculateSlackReserve(SizeType NumElements, SIZE_T NumBytesPerElement) const
+		FORCEINLINE SizeType CalculateSlackReserve(SizeType NewMax, SIZE_T NumBytesPerElement) const
 		{
-			return DefaultCalculateSlackReserve(NumElements, NumBytesPerElement, false, Alignment);
+			return DefaultCalculateSlackReserve(NewMax, NumBytesPerElement, false, Alignment);
 		}
-		FORCEINLINE SizeType CalculateSlackShrink(SizeType NumElements, SizeType NumAllocatedElements, SIZE_T NumBytesPerElement) const
+		FORCEINLINE SizeType CalculateSlackShrink(SizeType NewMax, SizeType CurrentMax, SIZE_T NumBytesPerElement) const
 		{
-			return DefaultCalculateSlackShrink(NumElements, NumAllocatedElements, NumBytesPerElement, false, Alignment);
+			return DefaultCalculateSlackShrink(NewMax, CurrentMax, NumBytesPerElement, false, Alignment);
 		}
-		FORCEINLINE SizeType CalculateSlackGrow(SizeType NumElements, SizeType NumAllocatedElements, SIZE_T NumBytesPerElement) const
+		FORCEINLINE SizeType CalculateSlackGrow(SizeType NewMax, SizeType CurrentMax, SIZE_T NumBytesPerElement) const
 		{
-			return DefaultCalculateSlackGrow(NumElements, NumAllocatedElements, NumBytesPerElement, false, Alignment);
+			return DefaultCalculateSlackGrow(NewMax, CurrentMax, NumBytesPerElement, false, Alignment);
 		}
 
-		FORCEINLINE SIZE_T GetAllocatedSize(SizeType NumAllocatedElements, SIZE_T NumBytesPerElement) const
+		FORCEINLINE SIZE_T GetAllocatedSize(SizeType CurrentMax, SIZE_T NumBytesPerElement) const
 		{
-			return NumAllocatedElements * NumBytesPerElement;
+			return CurrentMax * NumBytesPerElement;
 		}
 
 		bool HasAllocation() const
diff --git a/Engine/Source/Runtime/Core/Public/Serialization/MemoryImage.h b/Engine/Source/Runtime/Core/Public/Serialization/MemoryImage.h
--- a/Engine/Source/Runtime/Core/Public/Serialization/MemoryImage.h
+++ b/Engine/Source/Runtime/Core/Public/Serialization/MemoryImage.h
@@ -617,9 +617,9 @@ public:
 		return Data.Get();
 	}
 
-	FORCEINLINE SIZE_T GetAllocatedSize(int32 NumAllocatedElements, SIZE_T NumBytesPerElement) const
+	FORCEINLINE SIZE_T GetAllocatedSize(int32 CurrentMax, SIZE_T NumBytesPerElement) const
 	{
-		return NumAllocatedElements * NumBytesPerElement;
+		return CurrentMax * NumBytesPerElement;
 	}
 	FORCEINLINE bool HasAllocation()
 	{
@@ -656,37 +656,37 @@ public:
 		{
 			return 0;
 		}
-		FORCEINLINE int32 CalculateSlackReserve(int32 NumElements, int32 NumBytesPerElement) const
+		FORCEINLINE int32 CalculateSlackReserve(int32 NewMax, int32 NumBytesPerElement) const
 		{
-			return DefaultCalculateSlackReserve(NumElements, NumBytesPerElement, true, Alignment);
+			return DefaultCalculateSlackReserve(NewMax, NumBytesPerElement, true, Alignment);
 		}
-		FORCEINLINE int32 CalculateSlackReserve(int32 NumElements, int32 NumBytesPerElement, uint32 AlignmentOfElement) const
+		FORCEINLINE int32 CalculateSlackReserve(int32 NewMax, int32 NumBytesPerElement, uint32 AlignmentOfElement) const
 		{
-			return DefaultCalculateSlackReserve(NumElements, NumBytesPerElement, true, AlignmentOfElement);
+			return DefaultCalculateSlackReserve(NewMax, NumBytesPerElement, true, AlignmentOfElement);
 		}
-		FORCEINLINE int32 CalculateSlackShrink(int32 NumElements, int32 NumAllocatedElements, int32 NumBytesPerElement) const
+		FORCEINLINE int32 CalculateSlackShrink(int32 NewMax, int32 CurrentMax, int32 NumBytesPerElement) const
 		{
-			return DefaultCalculateSlackShrink(NumElements, NumAllocatedElements, NumBytesPerElement, true, Alignment);
+			return DefaultCalculateSlackShrink(NewMax, CurrentMax, NumBytesPerElement, true, Alignment);
 		}
-		FORCEINLINE int32 CalculateSlackShrink(int32 NumElements, int32 NumAllocatedElements, int32 NumBytesPerElement, uint32 AlignmentOfElement) const
+		FORCEINLINE int32 CalculateSlackShrink(int32 NewMax, int32 CurrentMax, int32 NumBytesPerElement, uint32 AlignmentOfElement) const
 		{
-			return DefaultCalculateSlackShrink(NumElements, NumAllocatedElements, NumBytesPerElement, true, AlignmentOfElement);
+			return DefaultCalculateSlackShrink(NewMax, CurrentMax, NumBytesPerElement, true, AlignmentOfElement);
 		}
-		FORCEINLINE int32 CalculateSlackGrow(int32 NumElements, int32 NumAllocatedElements, int32 NumBytesPerElement) const
+		FORCEINLINE int32 CalculateSlackGrow(int32 NewMax, int32 CurrentMax, int32 NumBytesPerElement) const
 		{
-			return DefaultCalculateSlackGrow(NumElements, NumAllocatedElements, NumBytesPerElement, true, Alignment);
+			return DefaultCalculateSlackGrow(NewMax, CurrentMax, NumBytesPerElement, true, Alignment);
 		}
-		FORCEINLINE int32 CalculateSlackGrow(int32 NumElements, int32 NumAllocatedElements, int32 NumBytesPerElement, uint32 AlignmentOfElement) const
+		FORCEINLINE int32 CalculateSlackGrow(int32 NewMax, int32 CurrentMax, int32 NumBytesPerElement, uint32 AlignmentOfElement) const
 		{
-			return DefaultCalculateSlackGrow(NumElements, NumAllocatedElements, NumBytesPerElement, true, AlignmentOfElement);
+			return DefaultCalculateSlackGrow(NewMax, CurrentMax, NumBytesPerElement, true, AlignmentOfElement);
 		}
-		FORCEINLINE void ResizeAllocation(int32 PreviousNumElements, int32 NumElements, SIZE_T NumBytesPerElement)
+		FORCEINLINE void ResizeAllocation(int32 CurrentNum, int32 NewMax, SIZE_T NumBytesPerElement)
 		{
-			FMemoryImageAllocatorBase::ResizeAllocation(PreviousNumElements, NumElements, NumBytesPerElement, Alignment);
+			FMemoryImageAllocatorBase::ResizeAllocation(CurrentNum, NewMax, NumBytesPerElement, Alignment);
 		}
-		FORCEINLINE void ResizeAllocation(int32 PreviousNumElements, int32 NumElements, SIZE_T NumBytesPerElement, uint32 AlignmentOfElement)
+		FORCEINLINE void ResizeAllocation(int32 CurrentNum, int32 NewMax, SIZE_T NumBytesPerElement, uint32 AlignmentOfElement)
 		{
-			FMemoryImageAllocatorBase::ResizeAllocation(PreviousNumElements, NumElements, NumBytesPerElement, AlignmentOfElement);
+			FMemoryImageAllocatorBase::ResizeAllocation(CurrentNum, NewMax, NumBytesPerElement, AlignmentOfElement);
 		}
 
 		FORCEINLINE void WriteMemoryImage(FMemoryImageWriter& Writer, const FTypeLayoutDesc& TypeDesc, int32 NumAllocatedElements) const
diff --git a/Engine/Source/Runtime/Engine/Public/Animation/AnimCompressionTypes.h b/Engine/Source/Runtime/Engine/Public/Animation/AnimCompressionTypes.h
--- a/Engine/Source/Runtime/Engine/Public/Animation/AnimCompressionTypes.h
+++ b/Engine/Source/Runtime/Engine/Public/Animation/AnimCompressionTypes.h
@@ -585,14 +585,14 @@ public:
 			return Data;
 		}
 		void ResizeAllocation(
-			SizeType PreviousNumElements,
-			SizeType NumElements,
+			SizeType CurrentNum,
+			SizeType NewMax,
 			SIZE_T NumBytesPerElement
 		)
 		{
 			if (MappedRegion || MappedHandle)
 			{
-				check(NumElements == 0); // Currently we can only support resizing of memory mapped regions to 0 size (ie delete)
+				check(NewMax == 0); // Currently we can only support resizing of memory mapped regions to 0 size (ie delete)
 
 				delete MappedRegion;
 				delete MappedHandle;
@@ -600,40 +600,40 @@ public:
 				MappedHandle = nullptr;
 				Data = nullptr; // make sure we don't try to free this pointer
 			}
-			else if (Data || NumElements)
+			else if (Data || NewMax)
 			{
 				static_assert(sizeof(int32) <= sizeof(SIZE_T), "SIZE_T is expected to be larger than int32");
 
 				// Check for under/overflow
-				if (UNLIKELY(NumElements < 0 || NumBytesPerElement < 1 || NumBytesPerElement > (SIZE_T)MAX_int32))
+				if (UNLIKELY(NewMax < 0 || NumBytesPerElement < 1 || NumBytesPerElement > (SIZE_T)MAX_int32))
 				{
-				    UE::Animation::Private::OnInvalidMaybeMappedAllocatorNum(NumElements, NumBytesPerElement);
+				    UE::Animation::Private::OnInvalidMaybeMappedAllocatorNum(NewMax, NumBytesPerElement);
 				}
 
 				// Avoid calling FMemory::Realloc( nullptr, 0 ) as ANSI C mandates returning a valid pointer which is not what we want.
-				//checkSlow(((uint64)NumElements*(uint64)ElementTypeInfo.GetSize() < (uint64)INT_MAX));
-				Data = (FScriptContainerElement*)FMemory::Realloc(Data, NumElements*NumBytesPerElement, Alignment);
+				//checkSlow(((uint64)NewMax*(uint64)ElementTypeInfo.GetSize() < (uint64)INT_MAX));
+				Data = (FScriptContainerElement*)FMemory::Realloc(Data, NewMax*NumBytesPerElement, Alignment);
 			}
 		}
-		SizeType CalculateSlackReserve(SizeType NumElements, SIZE_T NumBytesPerElement) const
+		SizeType CalculateSlackReserve(SizeType NewMax, SIZE_T NumBytesPerElement) const
 		{
 			check(!MappedHandle && !MappedRegion); // this could be supported, but it probably is never what you want, so we will just assert.
-			return DefaultCalculateSlackReserve(NumElements, NumBytesPerElement, true, Alignment);
+			return DefaultCalculateSlackReserve(NewMax, NumBytesPerElement, true, Alignment);
 		}
-		SizeType CalculateSlackShrink(SizeType NumElements, SizeType NumAllocatedElements, SIZE_T NumBytesPerElement) const
+		SizeType CalculateSlackShrink(SizeType NewMax, SizeType CurrentMax, SIZE_T NumBytesPerElement) const
 		{
 			check(!MappedHandle && !MappedRegion); // this could be supported, but it probably is never what you want, so we will just assert.
-			return DefaultCalculateSlackShrink(NumElements, NumAllocatedElements, NumBytesPerElement, true, Alignment);
+			return DefaultCalculateSlackShrink(NewMax, CurrentMax, NumBytesPerElement, true, Alignment);
 		}
-		SizeType CalculateSlackGrow(SizeType NumElements, SizeType NumAllocatedElements, SIZE_T NumBytesPerElement) const
+		SizeType CalculateSlackGrow(SizeType NewMax, SizeType CurrentMax, SIZE_T NumBytesPerElement) const
 		{
 			check(!MappedHandle && !MappedRegion); // this could be supported, but it probably is never what you want, so we will just assert.
-			return DefaultCalculateSlackGrow(NumElements, NumAllocatedElements, NumBytesPerElement, true, Alignment);
+			return DefaultCalculateSlackGrow(NewMax, CurrentMax, NumBytesPerElement, true, Alignment);
 		}
 
-		SIZE_T GetAllocatedSize(SizeType NumAllocatedElements, SIZE_T NumBytesPerElement) const
+		SIZE_T GetAllocatedSize(SizeType CurrentMax, SIZE_T NumBytesPerElement) const
 		{
-			return NumAllocatedElements * NumBytesPerElement;
+			return CurrentMax * NumBytesPerElement;
 		}
 
 		bool HasAllocation() const
diff --git a/Engine/Source/Runtime/RenderCore/Public/RenderGraphAllocator.h b/Engine/Source/Runtime/RenderCore/Public/RenderGraphAllocator.h
--- a/Engine/Source/Runtime/RenderCore/Public/RenderGraphAllocator.h
+++ b/Engine/Source/Runtime/RenderCore/Public/RenderGraphAllocator.h
@@ -188,50 +188,50 @@ public:
 			return Data;
 		}
 
-		void ResizeAllocation(SizeType PreviousNumElements, SizeType NumElements, SIZE_T NumBytesPerElement)
+		void ResizeAllocation(SizeType CurrentNum, SizeType NewMax, SIZE_T NumBytesPerElement)
 		{
 			void* OldData = Data;
-			if (NumElements)
+			if (NewMax)
 			{
 				static_assert(sizeof(int32) <= sizeof(SIZE_T), "SIZE_T is expected to be larger than int32");
 
 				// Check for under/overflow
-				if (UNLIKELY(NumElements < 0 || NumBytesPerElement < 1 || NumBytesPerElement > (SIZE_T)MAX_int32))
+				if (UNLIKELY(NewMax < 0 || NumBytesPerElement < 1 || NumBytesPerElement > (SIZE_T)MAX_int32))
 				{
-					UE::RenderCore::Private::OnInvalidRDGAllocatorNum(NumElements, NumBytesPerElement);
+					UE::RenderCore::Private::OnInvalidRDGAllocatorNum(NewMax, NumBytesPerElement);
 				}
 
 				// Allocate memory from the allocator.
-				const int32 AllocSize = (int32)(NumElements * NumBytesPerElement);
+				const int32 AllocSize = (int32)(NewMax * NumBytesPerElement);
 				const int32 AllocAlignment = FMath::Max(Alignment, (uint32)alignof(ElementType));
 				Data = (ElementType*)FRDGAllocator::GetTLS().Alloc(AllocSize, FMath::Max(AllocSize >= 16 ? (int32)16 : (int32)8, AllocAlignment));
 
 				// If the container previously held elements, copy them into the new allocation.
-				if (OldData && PreviousNumElements)
+				if (OldData && CurrentNum)
 				{
-					const SizeType NumCopiedElements = FMath::Min(NumElements, PreviousNumElements);
+					const SizeType NumCopiedElements = FMath::Min(NewMax, CurrentNum);
 					FMemory::Memcpy(Data, OldData, NumCopiedElements * NumBytesPerElement);
 				}
 			}
 		}
-		FORCEINLINE SizeType CalculateSlackReserve(SizeType NumElements, SIZE_T NumBytesPerElement) const
+		FORCEINLINE SizeType CalculateSlackReserve(SizeType NewMax, SIZE_T NumBytesPerElement) const
 		{
-			return DefaultCalculateSlackReserve(NumElements, NumBytesPerElement, false, Alignment);
+			return DefaultCalculateSlackReserve(NewMax, NumBytesPerElement, false, Alignment);
 		}
 
-		FORCEINLINE SizeType CalculateSlackShrink(SizeType NumElements, SizeType NumAllocatedElements, SIZE_T NumBytesPerElement) const
+		FORCEINLINE SizeType CalculateSlackShrink(SizeType NewMax, SizeType CurrentMax, SIZE_T NumBytesPerElement) const
 		{
-			return DefaultCalculateSlackShrink(NumElements, NumAllocatedElements, NumBytesPerElement, false, Alignment);
+			return DefaultCalculateSlackShrink(NewMax, CurrentMax, NumBytesPerElement, false, Alignment);
 		}
 
-		FORCEINLINE SizeType CalculateSlackGrow(SizeType NumElements, SizeType NumAllocatedElements, SIZE_T NumBytesPerElement) const
+		FORCEINLINE SizeType CalculateSlackGrow(SizeType NewMax, SizeType CurrentMax, SIZE_T NumBytesPerElement) const
 		{
-			return DefaultCalculateSlackGrow(NumElements, NumAllocatedElements, NumBytesPerElement, false, Alignment);
+			return DefaultCalculateSlackGrow(NewMax, CurrentMax, NumBytesPerElement, false, Alignment);
 		}
 
-		FORCEINLINE SIZE_T GetAllocatedSize(SizeType NumAllocatedElements, SIZE_T NumBytesPerElement) const
+		FORCEINLINE SIZE_T GetAllocatedSize(SizeType CurrentMax, SIZE_T NumBytesPerElement) const
 		{
-			return NumAllocatedElements * NumBytesPerElement;
+			return CurrentMax * NumBytesPerElement;
 		}
 
 		bool HasAllocation() const
