Added ToString/AppendString functions to FTextKey

These are intended to replace GetChars() and allow us to investigate swapping the internal string storage to UTF-8, converting back to TCHAR as required.

This change doesn't yet deprecate GetChars() or replace any external usage of the function, but it does fix the internal usage during serialization.

diff --git a/Engine/Source/Runtime/Core/Private/Internationalization/TextKey.cpp b/Engine/Source/Runtime/Core/Private/Internationalization/TextKey.cpp
--- a/Engine/Source/Runtime/Core/Private/Internationalization/TextKey.cpp
+++ b/Engine/Source/Runtime/Core/Private/Internationalization/TextKey.cpp
@@ -397,8 +397,9 @@ private:
 namespace TextKeyUtil
 {
 
-static const int32 InlineStringSize = 128;
-typedef TArray<TCHAR, TInlineAllocator<InlineStringSize>> FInlineStringBuffer;
+static constexpr int32 InlineStringSize = 128;
+using FInlineStringBuffer = TArray<TCHAR, TInlineAllocator<InlineStringSize>>;
+using FInlineStringBuilder = TStringBuilder<InlineStringSize>;
 
 static_assert(PLATFORM_LITTLE_ENDIAN, "FTextKey serialization needs updating to support big-endian platforms!");
 
@@ -525,6 +526,35 @@ const TCHAR* FTextKey::GetChars() const
 		: TEXT("");
 }
 
+FString FTextKey::ToString() const
+{
+	FString Out;
+	AppendString(Out);
+	return Out;
+}
+
+void FTextKey::ToString(FString& Out) const
+{
+	Out.Reset();
+	AppendString(Out);
+}
+
+void FTextKey::ToString(FStringBuilderBase& Out) const
+{
+	Out.Reset();
+	AppendString(Out);
+}
+
+void FTextKey::AppendString(FString& Out) const
+{
+	Out += GetChars();
+}
+
+void FTextKey::AppendString(FStringBuilderBase& Out) const
+{
+	Out += GetChars();
+}
+
 uint32 GetTypeHash(const FTextKey& A)
 {
 #if UE_TEXTKEY_STORE_EMBEDDED_HASH
@@ -554,7 +584,9 @@ void FTextKey::SerializeAsString(FArchive& Ar)
 	}
 	else
 	{
-		TextKeyUtil::SaveKeyString(Ar, GetChars());
+		TextKeyUtil::FInlineStringBuilder StrBuilder;
+		AppendString(StrBuilder);
+		TextKeyUtil::SaveKeyString(Ar, StrBuilder.ToString());
 	}
 }
 
@@ -582,7 +614,9 @@ void FTextKey::SerializeWithHash(FArchive& Ar)
 		uint32 TmpStrHash = GetTypeHash(*this);
 		Ar << TmpStrHash;
 
-		TextKeyUtil::SaveKeyString(Ar, GetChars());
+		TextKeyUtil::FInlineStringBuilder StrBuilder;
+		AppendString(StrBuilder);
+		TextKeyUtil::SaveKeyString(Ar, StrBuilder.ToString());
 	}
 }
 
@@ -610,7 +644,9 @@ void FTextKey::SerializeDiscardHash(FArchive& Ar)
 		uint32 TmpStrHash = GetTypeHash(*this);
 		Ar << TmpStrHash;
 
-		TextKeyUtil::SaveKeyString(Ar, GetChars());
+		TextKeyUtil::FInlineStringBuilder StrBuilder;
+		AppendString(StrBuilder);
+		TextKeyUtil::SaveKeyString(Ar, StrBuilder.ToString());
 	}
 }
 
@@ -634,7 +670,7 @@ void FTextKey::SerializeAsString(FStructuredArchiveSlot Slot)
 		}
 		else
 		{
-			FString TmpStr = GetChars();
+			FString TmpStr = ToString();
 			Slot << TmpStr;
 		}
 	}
@@ -673,7 +709,7 @@ void FTextKey::SerializeWithHash(FStructuredArchiveSlot Slot)
 			uint32 TmpStrHash = GetTypeHash(*this);
 			Record << SA_VALUE(TEXT("Hash"), TmpStrHash);
 
-			FString TmpStr = GetChars();
+			FString TmpStr = ToString();
 			Record << SA_VALUE(TEXT("Str"), TmpStr);
 		}
 	}
@@ -712,7 +748,7 @@ void FTextKey::SerializeDiscardHash(FStructuredArchiveSlot Slot)
 			uint32 TmpStrHash = GetTypeHash(*this);
 			Record << SA_VALUE(TEXT("Hash"), TmpStrHash);
 
-			FString TmpStr = GetChars();
+			FString TmpStr = ToString();
 			Record << SA_VALUE(TEXT("Str"), TmpStr);
 		}
 	}
diff --git a/Engine/Source/Runtime/Core/Public/Internationalization/TextKey.h b/Engine/Source/Runtime/Core/Public/Internationalization/TextKey.h
--- a/Engine/Source/Runtime/Core/Public/Internationalization/TextKey.h
+++ b/Engine/Source/Runtime/Core/Public/Internationalization/TextKey.h
@@ -87,9 +87,19 @@ public:
 	{
 	}
 
-	/** Get the underlying chars buffer this text key represents */
+	/**
+	 * Get the underlying chars buffer this text key represents.
+	 * @note This function is due to be deprecated. Prefer ToString/AppendString for new code.
+	 */
 	CORE_API const TCHAR* GetChars() const;
 
+	/** Convert this text key back to its string representation */
+	CORE_API FString ToString() const;
+	CORE_API void ToString(FString& Out) const;
+	CORE_API void ToString(FStringBuilderBase& Out) const;
+	CORE_API void AppendString(FString& Out) const;
+	CORE_API void AppendString(FStringBuilderBase& Out) const;
+
 	/** Compare for equality */
 	friend FORCEINLINE bool operator==(const FTextKey& A, const FTextKey& B)
 	{
