Added the CAnyOf concept/trait

diff --git a/Engine/Source/Runtime/Core/Public/Concepts/AnyOf.h b/Engine/Source/Runtime/Core/Public/Concepts/AnyOf.h
new file mode 100644
--- /dev/null
+++ b/Engine/Source/Runtime/Core/Public/Concepts/AnyOf.h
@@ -0,0 +1,8 @@
+﻿// Copyright Epic Games, Inc. All Rights Reserved.
+
+#pragma once
+
+#include <type_traits>
+
+template <typename T, typename... Options>
+constexpr bool CAnyOf = (std::is_same_v<T, Options> || ...);
