Fix compile error with preview compiler

#rnx

diff --git a/Engine/Source/Runtime/Core/Public/Experimental/ConcurrentLinearAllocator.h b/Engine/Source/Runtime/Core/Public/Experimental/ConcurrentLinearAllocator.h
--- a/Engine/Source/Runtime/Core/Public/Experimental/ConcurrentLinearAllocator.h
+++ b/Engine/Source/Runtime/Core/Public/Experimental/ConcurrentLinearAllocator.h
@@ -28,7 +28,7 @@
 #include <sanitizer/asan_interface.h>
 #if defined(__SANITIZE_ADDRESS__)
 #define IS_ASAN_ENABLED 1
-#elif __has_feature(address_sanitizer)
+#elif defined(__has_feature) && __has_feature(address_sanitizer)
 #define IS_ASAN_ENABLED 1
 #else
 #define IS_ASAN_ENABLED 0
