Updated FTextKey::GetChars() usage to FTextKey::ToString()

This is in preparation to deprecate FTextKey::GetChars() in the future

#rnx

diff --git a/Engine/Source/Developer/TranslationEditor/Private/TranslationPickerEditWindow.cpp b/Engine/Source/Developer/TranslationEditor/Private/TranslationPickerEditWindow.cpp
--- a/Engine/Source/Developer/TranslationEditor/Private/TranslationPickerEditWindow.cpp
+++ b/Engine/Source/Developer/TranslationEditor/Private/TranslationPickerEditWindow.cpp
@@ -365,12 +365,12 @@ TSharedPtr<FTranslationPickerTextItem> FTranslationPickerTextItem::BuildTextItem
 #endif // WITH_EDITORONLY_DATA
 
 	// Clean the package localization ID from the namespace (to mirror what the text gatherer does when scraping for translation data)
-	Item->CleanNamespace = TextNamespaceUtil::StripPackageNamespace(Item->TextId.GetNamespace().GetChars());
+	Item->CleanNamespace = TextNamespaceUtil::StripPackageNamespace(Item->TextId.GetNamespace().ToString());
 
 	// Save the necessary data in UTranslationUnit for later.  This is what we pass to TranslationDataManager to save our edits
 	Item->TranslationUnit = NewObject<UTranslationUnit>();
 	Item->TranslationUnit->Namespace = Item->CleanNamespace;
-	Item->TranslationUnit->Key = Item->TextId.GetKey().GetChars();
+	Item->TranslationUnit->Key = Item->TextId.GetKey().ToString();
 	Item->TranslationUnit->Source = Item->SourceString;
 	Item->TranslationUnit->Translation = Item->TranslationString;
 	Item->TranslationUnit->LocresPath = LocResPath;
@@ -539,7 +539,7 @@ void STranslationPickerEditWidget::Construct(const FArguments& InArgs, const TSh
 			[
 				SNew(SEditableTextBox)
 				.IsReadOnly(true)
-				.Text(FText::AsCultureInvariant(Item->TextId.GetKey().GetChars()))
+				.Text(FText::AsCultureInvariant(Item->TextId.GetKey().ToString()))
 			];
 		
 		int32 Row = 4;
diff --git a/Engine/Source/Developer/TranslationEditor/Private/TranslationPickerFloatingWindow.cpp b/Engine/Source/Developer/TranslationEditor/Private/TranslationPickerFloatingWindow.cpp
--- a/Engine/Source/Developer/TranslationEditor/Private/TranslationPickerFloatingWindow.cpp
+++ b/Engine/Source/Developer/TranslationEditor/Private/TranslationPickerFloatingWindow.cpp
@@ -135,8 +135,8 @@ public:
 			FTextId TextId = FTextInspector::GetTextId(PickedText);
 
 			// Clean the package localization ID from the namespace (to mirror what the text gatherer does when scraping for translation data)
-			FString EntryNamespace = TextNamespaceUtil::StripPackageNamespace(TextId.GetNamespace().GetChars());
-			FString EntryKey = TextId.GetKey().GetChars();
+			FString EntryNamespace = TextNamespaceUtil::StripPackageNamespace(TextId.GetNamespace().ToString());
+			FString EntryKey = TextId.GetKey().ToString();
 
 			const FString CopyString = FString::Printf(TEXT("%s,%s"), *EntryNamespace, *EntryKey);
 	
diff --git a/Engine/Source/Editor/EditorWidgets/Private/STextPropertyEditableTextBox.cpp b/Engine/Source/Editor/EditorWidgets/Private/STextPropertyEditableTextBox.cpp
--- a/Engine/Source/Editor/EditorWidgets/Private/STextPropertyEditableTextBox.cpp
+++ b/Engine/Source/Editor/EditorWidgets/Private/STextPropertyEditableTextBox.cpp
@@ -1021,8 +1021,8 @@ FText STextPropertyEditableTextBox::GetToolTipText() const
 			{
 				check(SourceString);
 
-				const FString Namespace = TextId.GetNamespace().GetChars();
-				const FString Key = TextId.GetKey().GetChars();
+				const FString Namespace = TextId.GetNamespace().ToString();
+				const FString Key = TextId.GetKey().ToString();
 
 				const FString PackageNamespace = TextNamespaceUtil::ExtractPackageNamespace(Namespace);
 				const FString TextNamespace = TextNamespaceUtil::StripPackageNamespace(Namespace);
diff --git a/Engine/Source/Editor/KismetCompiler/Private/KismetCompilerVMBackend.cpp b/Engine/Source/Editor/KismetCompiler/Private/KismetCompilerVMBackend.cpp
--- a/Engine/Source/Editor/KismetCompiler/Private/KismetCompilerVMBackend.cpp
+++ b/Engine/Source/Editor/KismetCompiler/Private/KismetCompilerVMBackend.cpp
@@ -721,8 +721,8 @@ public:
 					{
 						// BP bytecode always removes the package localization ID to match how text works at runtime
 						// If we're gathering editor-only text then we'll pick up the version with the package localization ID from the property/pin rather than the bytecode
-						const FString Namespace = TextNamespaceUtil::StripPackageNamespace(TextId.GetNamespace().GetChars());
-						const FString Key = TextId.GetKey().GetChars();
+						const FString Namespace = TextNamespaceUtil::StripPackageNamespace(TextId.GetNamespace().ToString());
+						const FString Key = TextId.GetKey().ToString();
 
 						Writer << EBlueprintTextLiteralType::LocalizedText;
 						EmitStringLiteral(*SourceString);
diff --git a/Engine/Source/Editor/UnrealEd/Private/Commandlets/ExtractLocResCommandlet.cpp b/Engine/Source/Editor/UnrealEd/Private/Commandlets/ExtractLocResCommandlet.cpp
--- a/Engine/Source/Editor/UnrealEd/Private/Commandlets/ExtractLocResCommandlet.cpp
+++ b/Engine/Source/Editor/UnrealEd/Private/Commandlets/ExtractLocResCommandlet.cpp
@@ -99,9 +99,9 @@ int32 UExtractLocResCommandlet::Main(const FString& Params)
 		// Write each row
 		for (const auto& LocResEntryPair : LocResFile.Entries)
 		{
-			WriteCSVStringValue(LocResEntryPair.Key.GetNamespace().GetChars());
+			WriteCSVStringValue(LocResEntryPair.Key.GetNamespace().ToString());
 			LocResCSV += TEXT(",");
-			WriteCSVStringValue(LocResEntryPair.Key.GetKey().GetChars());
+			WriteCSVStringValue(LocResEntryPair.Key.GetKey().ToString());
 			LocResCSV += TEXT(",");
 			LocResCSV += FString::Printf(TEXT("0x%08x"), LocResEntryPair.Value.SourceStringHash);
 			LocResCSV += TEXT(",");
diff --git a/Engine/Source/Editor/UnrealEd/Private/Commandlets/GenerateTextLocalizationResourceCommandlet.cpp b/Engine/Source/Editor/UnrealEd/Private/Commandlets/GenerateTextLocalizationResourceCommandlet.cpp
--- a/Engine/Source/Editor/UnrealEd/Private/Commandlets/GenerateTextLocalizationResourceCommandlet.cpp
+++ b/Engine/Source/Editor/UnrealEd/Private/Commandlets/GenerateTextLocalizationResourceCommandlet.cpp
@@ -218,7 +218,7 @@ int32 UGenerateTextLocalizationResourceCommandlet::Main(const FString& Params)
 		const FTextKey LocResId = DestinationPath / CultureName / ResourceName;
 		if (!FTextLocalizationResourceGenerator::GenerateLocRes(LocTextHelper, CultureName, GenerateFlags, LocResId, PlatformAgnosticLocRes, PerPlatformLocRes))
 		{
-			UE_LOG(LogGenerateTextLocalizationResourceCommandlet, Error, TEXT("Failed to generate LocRes %s"), LocResId.GetChars());
+			UE_LOG(LogGenerateTextLocalizationResourceCommandlet, Error, TEXT("Failed to generate LocRes %s"), *LocResId.ToString());
 			return false;
 		}
 	
diff --git a/Engine/Source/Editor/UnrealEd/Private/Commandlets/StabilizeLocalizationKeys.cpp b/Engine/Source/Editor/UnrealEd/Private/Commandlets/StabilizeLocalizationKeys.cpp
--- a/Engine/Source/Editor/UnrealEd/Private/Commandlets/StabilizeLocalizationKeys.cpp
+++ b/Engine/Source/Editor/UnrealEd/Private/Commandlets/StabilizeLocalizationKeys.cpp
@@ -127,8 +127,8 @@ public:
 			return false;
 		}
 
-		const FString Namespace = TextId.GetNamespace().GetChars();
-		const FString Key = TextId.GetKey().GetChars();
+		const FString Namespace = TextId.GetNamespace().ToString();
+		const FString Key = TextId.GetKey().ToString();
 
 		const FString CurrentPackageNamespace = TextNamespaceUtil::ExtractPackageNamespace(Namespace);
 		if (CurrentPackageNamespace.Equals(PackageNamespace, ESearchCase::CaseSensitive))
diff --git a/Engine/Source/Runtime/Core/Private/Internationalization/StringTableCore.cpp b/Engine/Source/Runtime/Core/Private/Internationalization/StringTableCore.cpp
--- a/Engine/Source/Runtime/Core/Private/Internationalization/StringTableCore.cpp
+++ b/Engine/Source/Runtime/Core/Private/Internationalization/StringTableCore.cpp
@@ -133,7 +133,7 @@ void FStringTable::IsInternal(const bool bInIsInternal)
 
 FString FStringTable::GetNamespace() const
 {
-	return TableNamespace.GetChars();
+	return TableNamespace.ToString();
 }
 
 void FStringTable::SetNamespace(const FTextKey& InNamespace)
@@ -201,7 +201,7 @@ void FStringTable::EnumerateSourceStrings(const TFunctionRef<bool(const FString&
 {
 	EnumerateKeysAndSourceStrings([&InEnumerator](const FTextKey& InKey, const FString& InSourceString) -> bool
 	{
-		return InEnumerator(InKey.GetChars(), InSourceString);
+		return InEnumerator(InKey.ToString(), InSourceString);
 	});
 }
 
@@ -243,7 +243,7 @@ bool FStringTable::FindKey(const FStringTableEntryConstRef& InEntry, FString& Ou
 	FTextKey TmpKey;
 	if (FindKey(InEntry, TmpKey))
 	{
-		OutKey = TmpKey.GetChars();
+		TmpKey.ToString(OutKey);
 		return true;
 	}
 	return false;
@@ -361,7 +361,7 @@ void FStringTable::Serialize(FArchive& Ar)
 			TmpKeysToMetaData.Reserve(KeysToMetaData.Num());
 			for (const auto& KeyToMetaDataPair : KeysToMetaData)
 			{
-				TmpKeysToMetaData.Add(KeyToMetaDataPair.Key.GetChars(), KeyToMetaDataPair.Value);
+				TmpKeysToMetaData.Add(KeyToMetaDataPair.Key.ToString(), KeyToMetaDataPair.Value);
 			}
 
 			Ar << TmpKeysToMetaData;
@@ -433,7 +433,7 @@ bool FStringTable::ExportStrings(const FString& InFilename) const
 		// Write entries
 		for (const auto& KeyToEntryPair : KeysToEntries)
 		{
-			FString ExportedKey = KeyToEntryPair.Key.GetChars();
+			FString ExportedKey = KeyToEntryPair.Key.ToString();
 			ExportedKey.ReplaceCharWithEscapedCharInline();
 			ExportedKey.ReplaceInline(TEXT("\""), TEXT("\"\""));
 
diff --git a/Engine/Source/Runtime/Core/Private/Internationalization/StringTableRegistry.cpp b/Engine/Source/Runtime/Core/Private/Internationalization/StringTableRegistry.cpp
--- a/Engine/Source/Runtime/Core/Private/Internationalization/StringTableRegistry.cpp
+++ b/Engine/Source/Runtime/Core/Private/Internationalization/StringTableRegistry.cpp
@@ -154,7 +154,7 @@ void FStringTableRegistry::LogMissingStringTableEntry(const FName InTableId, con
 	FLocKeySet& LoggedMissingKeys = LoggedMissingEntries.FindOrAdd(InTableId);
 	LoggedMissingKeys.Add(InKey);
 
-	UE_LOG(LogStringTable, Warning, TEXT("Failed to find string table entry for '%s' '%s'. Did you forget to add a string table redirector?"), *InTableId.ToString(), InKey.GetChars());
+	UE_LOG(LogStringTable, Warning, TEXT("Failed to find string table entry for '%s' '%s'. Did you forget to add a string table redirector?"), *InTableId.ToString(), *InKey.ToString());
 }
 
 #if WITH_EDITOR
diff --git a/Engine/Source/Runtime/Core/Private/Internationalization/Text.cpp b/Engine/Source/Runtime/Core/Private/Internationalization/Text.cpp
--- a/Engine/Source/Runtime/Core/Private/Internationalization/Text.cpp
+++ b/Engine/Source/Runtime/Core/Private/Internationalization/Text.cpp
@@ -45,7 +45,7 @@ TOptional<FString> FTextInspector::GetNamespace(const FText& Text)
 	const FTextId TextId = FTextInspector::GetTextId(Text);
 	if (!TextId.IsEmpty())
 	{
-		return FString(TextId.GetNamespace().GetChars());
+		return TextId.GetNamespace().ToString();
 	}
 	return TOptional<FString>();
 }
@@ -55,7 +55,7 @@ TOptional<FString> FTextInspector::GetKey(const FText& Text)
 	const FTextId TextId = FTextInspector::GetTextId(Text);
 	if (!TextId.IsEmpty())
 	{
-		return FString(TextId.GetKey().GetChars());
+		return TextId.GetKey().ToString();
 	}
 	return TOptional<FString>();
 }
@@ -81,7 +81,7 @@ bool FTextInspector::GetTableIdAndKey(const FText& Text, FName& OutTableId, FStr
 	FTextKey TmpKey;
 	if (GetTableIdAndKey(Text, OutTableId, TmpKey))
 	{
-		OutKey = TmpKey.GetChars();
+		TmpKey.ToString(OutKey);
 		return true;
 	}
 
diff --git a/Engine/Source/Runtime/Core/Private/Internationalization/TextHistory.cpp b/Engine/Source/Runtime/Core/Private/Internationalization/TextHistory.cpp
--- a/Engine/Source/Runtime/Core/Private/Internationalization/TextHistory.cpp
+++ b/Engine/Source/Runtime/Core/Private/Internationalization/TextHistory.cpp
@@ -788,7 +788,7 @@ void FTextHistory_Base::Serialize(FStructuredArchive::FRecord Record)
 			const FString PackageNamespace = TextNamespaceUtil::GetPackageNamespace(BaseArchive);
 			if (!PackageNamespace.IsEmpty())
 			{
-				const FString NamespaceStr = Namespace.GetChars();
+				const FString NamespaceStr = Namespace.ToString();
 				const FString FullNamespace = TextNamespaceUtil::BuildFullNamespace(NamespaceStr, PackageNamespace);
 				if (!NamespaceStr.Equals(FullNamespace, ESearchCase::CaseSensitive))
 				{
@@ -804,7 +804,7 @@ void FTextHistory_Base::Serialize(FStructuredArchive::FRecord Record)
 		if (!GIsEditor)
 		{
 			// Strip the package localization ID to match how text works at runtime (properties do this when saving during cook)
-			Namespace = TextNamespaceUtil::StripPackageNamespace(Namespace.GetChars());
+			Namespace = TextNamespaceUtil::StripPackageNamespace(Namespace.ToString());
 		}
 #endif // WITH_EDITOR
 
@@ -820,7 +820,7 @@ void FTextHistory_Base::Serialize(FStructuredArchive::FRecord Record)
 		if (BaseArchive.IsCooking())
 		{
 			// We strip the package localization off the serialized text for a cooked game, as they're not used at runtime
-			Namespace = TextNamespaceUtil::StripPackageNamespace(Namespace.GetChars());
+			Namespace = TextNamespaceUtil::StripPackageNamespace(Namespace.ToString());
 		}
 		else
 		{
@@ -831,7 +831,7 @@ void FTextHistory_Base::Serialize(FStructuredArchive::FRecord Record)
 				const FString PackageNamespace = TextNamespaceUtil::GetPackageNamespace(BaseArchive);
 				if (!PackageNamespace.IsEmpty())
 				{
-					const FString NamespaceStr = Namespace.GetChars();
+					const FString NamespaceStr = Namespace.ToString();
 					const FString FullNamespace = TextNamespaceUtil::BuildFullNamespace(NamespaceStr, PackageNamespace);
 					if (!NamespaceStr.Equals(FullNamespace, ESearchCase::CaseSensitive))
 					{
@@ -1016,8 +1016,8 @@ bool FTextHistory_Base::WriteToBuffer(FString& Buffer, const bool bStripPackageN
 {
 	if (!TextId.IsEmpty())
 	{
-		FString Namespace = TextId.GetNamespace().GetChars();
-		FString Key = TextId.GetKey().GetChars();
+		FString Namespace = TextId.GetNamespace().ToString();
+		FString Key = TextId.GetKey().ToString();
 		if (bStripPackageNamespace)
 		{
 			TextNamespaceUtil::StripPackageNamespaceInline(Namespace);
@@ -2491,7 +2491,7 @@ bool FTextHistory_StringTableEntry::WriteToBuffer(FString& Buffer, const bool bS
 		FTextKey Key;
 		StringTableReferenceData->GetTableIdAndKey(TableId, Key);
 
-		FString KeyStr = Key.GetChars();
+		FString KeyStr = Key.ToString();
 
 #define LOC_DEFINE_REGION
 		// Produces LOCTABLE("...", "...")
diff --git a/Engine/Source/Runtime/Core/Private/Internationalization/TextLocalizationManager.cpp b/Engine/Source/Runtime/Core/Private/Internationalization/TextLocalizationManager.cpp
--- a/Engine/Source/Runtime/Core/Private/Internationalization/TextLocalizationManager.cpp
+++ b/Engine/Source/Runtime/Core/Private/Internationalization/TextLocalizationManager.cpp
@@ -182,7 +182,7 @@ FString KeyifyTextId(const FTextId& TextId)
 {
 	// We want to show the identity in terms of key, namespace. This is to try and fit into the constraints of UI text blocks and at least let the key component be visible to easily identify a piece of text.
 	// If the key/namespace pair is too long, the Slate.LogPaintedText cvar can be used to see the entire thing.
-	return FString::Printf(TEXT("%s, %s"), TextId.GetKey().GetChars(), TextId.GetNamespace().GetChars());
+	return FString::Printf(TEXT("%s, %s"), *TextId.GetKey().ToString(), *TextId.GetNamespace().ToString());
 }
 }
 
@@ -754,8 +754,8 @@ void FTextLocalizationManager::DumpLiveTableImpl(const FString* NamespaceFilter,
 		DisplayStringLookupTableToDump.Reserve(DisplayStringLookupTable.Num());
 		for (const auto& DisplayStringPair : DisplayStringLookupTable)
 		{
-			if (PassesFilter(DisplayStringPair.Key.GetNamespace().GetChars(), NamespaceFilter) &&
-				PassesFilter(DisplayStringPair.Key.GetKey().GetChars(), KeyFilter) &&
+			if (PassesFilter(DisplayStringPair.Key.GetNamespace().ToString(), NamespaceFilter) &&
+				PassesFilter(DisplayStringPair.Key.GetKey().ToString(), KeyFilter) &&
 				PassesFilter(**DisplayStringPair.Value.DisplayString, DisplayStringFilter))
 			{
 				DisplayStringLookupTableToDump.Add(DisplayStringPair.Key, DisplayStringPair.Value);
@@ -765,12 +765,12 @@ void FTextLocalizationManager::DumpLiveTableImpl(const FString* NamespaceFilter,
 
 	DisplayStringLookupTableToDump.KeySort([](const FTextId& A, const FTextId& B)
 	{
-		const int32 NamespaceResult = FCString::Strcmp(A.GetNamespace().GetChars(), B.GetNamespace().GetChars());
+		const int32 NamespaceResult = FCString::Strcmp(*A.GetNamespace().ToString(), *B.GetNamespace().ToString());
 		if (NamespaceResult != 0)
 		{
 			return NamespaceResult < 0;
 		}
-		return FCString::Strcmp(A.GetKey().GetChars(), B.GetKey().GetChars()) < 0;
+		return FCString::Strcmp(*A.GetKey().ToString(), *B.GetKey().ToString()) < 0;
 	});
 
 	for (const auto& DisplayStringPair : DisplayStringLookupTableToDump)
@@ -788,7 +788,7 @@ void FTextLocalizationManager::DumpLiveTable(const FString* NamespaceFilter, con
 
 	DumpLiveTableImpl(NamespaceFilter, KeyFilter, DisplayStringFilter, [&Category](const FTextId& Id, const FTextConstDisplayStringRef& DisplayString)
 	{
-		UE_LOG_REF(Category, Display, TEXT("LiveTableEntry: Namespace: '%s', Key: '%s', DisplayString: '%s'"), Id.GetNamespace().GetChars(), Id.GetKey().GetChars(), **DisplayString); //-V510
+		UE_LOG_REF(Category, Display, TEXT("LiveTableEntry: Namespace: '%s', Key: '%s', DisplayString: '%s'"), *Id.GetNamespace().ToString(), *Id.GetKey().ToString(), **DisplayString); //-V510
 	});
 
 	UE_LOG_REF(Category, Display, TEXT("----------------------------------------------------------------------"));
@@ -801,7 +801,7 @@ void FTextLocalizationManager::DumpLiveTable(const FString& OutputFilename, cons
 
 	DumpLiveTableImpl(NamespaceFilter, KeyFilter, DisplayStringFilter, [&DumpString](const FTextId& Id, const FTextConstDisplayStringRef& DisplayString)
 	{
-		DumpString += FString::Printf(TEXT("LiveTableEntry: Namespace: '%s', Key: '%s', DisplayString: '%s'"), Id.GetNamespace().GetChars(), Id.GetKey().GetChars(), **DisplayString);
+		DumpString += FString::Printf(TEXT("LiveTableEntry: Namespace: '%s', Key: '%s', DisplayString: '%s'"), *Id.GetNamespace().ToString(), *Id.GetKey().ToString(), **DisplayString);
 		DumpString += LINE_TERMINATOR;
 	});
 
@@ -1069,8 +1069,9 @@ FTextConstDisplayStringPtr FTextLocalizationManager::GetDisplayString(const FTex
 	// In builds with stable keys enabled, we want to find the display string for the "clean" version of the text (if the sources match) as this is the only version that is translated
 #if USE_STABLE_LOCALIZATION_KEYS
 	{
-		const FString DisplayNamespace = TextNamespaceUtil::StripPackageNamespace(TextId.GetNamespace().GetChars());
-		if (FCString::Strcmp(*DisplayNamespace, TextId.GetNamespace().GetChars()) != 0)
+		const FString FullNamespace = TextId.GetNamespace().ToString();
+		const FString DisplayNamespace = TextNamespaceUtil::StripPackageNamespace(FullNamespace);
+		if (!DisplayNamespace.Equals(FullNamespace, ESearchCase::CaseSensitive))
 		{
 			TextId = FTextId(DisplayNamespace, TextId.GetKey());
 		}
@@ -1122,7 +1123,7 @@ bool FTextLocalizationManager::GetLocResID(const FTextKey& Namespace, const FTex
 
 	if (LiveEntry != nullptr && !LiveEntry->LocResID.IsEmpty())
 	{
-		OutLocResId = LiveEntry->LocResID.GetChars();
+		LiveEntry->LocResID.ToString(OutLocResId);
 		return true;
 	}
 
diff --git a/Engine/Source/Runtime/Core/Private/Internationalization/TextLocalizationResource.cpp b/Engine/Source/Runtime/Core/Private/Internationalization/TextLocalizationResource.cpp
--- a/Engine/Source/Runtime/Core/Private/Internationalization/TextLocalizationResource.cpp
+++ b/Engine/Source/Runtime/Core/Private/Internationalization/TextLocalizationResource.cpp
@@ -36,7 +36,8 @@ int32 GetLocalizationTargetPathIdFromLocResId(const FTextKey& InLocResID)
 	if (!InLocResID.IsEmpty())
 	{
 		// LocResID would be "/Path/To/LocalizationTarget/Culture/LocalizationTarget.locres" so trim this back to "/Path/To/LocalizationTarget"
-		FStringView LocalizationTargetPath = InLocResID.GetChars();
+		const FString LocalizationTargetPathStr = InLocResID.ToString();
+		FStringView LocalizationTargetPath = LocalizationTargetPathStr;
 		LocalizationTargetPath = FPathViews::GetPath(LocalizationTargetPath); // Remove "LocalizationTarget.locres"
 		LocalizationTargetPath = FPathViews::GetPath(LocalizationTargetPath); // Remove "Culture"
 		LocalizationTargetPathId = FTextLocalizationManager::Get().GetLocalizationTargetPathId(LocalizationTargetPath);
@@ -270,13 +271,13 @@ bool FTextLocalizationResource::LoadFromArchive(FArchive& Archive, const FTextKe
 	{
 		// Legacy LocRes files lack the magic number, assume that's what we're dealing with, and seek back to the start of the file
 		Archive.Seek(0);
-		UE_LOG(LogTextLocalizationResource, Warning, TEXT("LocRes '%s' failed the magic number check! Assuming this is a legacy resource (please re-generate your localization resources!)"), LocResID.GetChars());
+		UE_LOG(LogTextLocalizationResource, Warning, TEXT("LocRes '%s' failed the magic number check! Assuming this is a legacy resource (please re-generate your localization resources!)"), *LocResID.ToString());
 	}
 
 	// Is this LocRes file too new to load?
 	if (VersionNumber > FTextLocalizationResourceVersion::ELocResVersion::Latest)
 	{
-		UE_LOG(LogTextLocalizationResource, Error, TEXT("LocRes '%s' is too new to be loaded (File Version: %d, Loader Version: %d)"), LocResID.GetChars(), (int32)VersionNumber, (int32)FTextLocalizationResourceVersion::ELocResVersion::Latest);
+		UE_LOG(LogTextLocalizationResource, Error, TEXT("LocRes '%s' is too new to be loaded (File Version: %d, Loader Version: %d)"), *LocResID.ToString(), (int32)VersionNumber, (int32)FTextLocalizationResourceVersion::ELocResVersion::Latest);
 		return false;
 	}
 
@@ -380,7 +381,7 @@ bool FTextLocalizationResource::LoadFromArchive(FArchive& Archive, const FTextKe
 				}
 				else
 				{
-					UE_LOG(LogTextLocalizationResource, Warning, TEXT("LocRes '%s' has an invalid localized string index for namespace '%s' and key '%s'. This entry will have no translation."), LocResID.GetChars(), Namespace.GetChars(), Key.GetChars());
+					UE_LOG(LogTextLocalizationResource, Warning, TEXT("LocRes '%s' has an invalid localized string index for namespace '%s' and key '%s'. This entry will have no translation."), *LocResID.ToString(), *Namespace.ToString(), *Key.ToString());
 				}
 			}
 			else
@@ -544,13 +545,13 @@ bool FTextLocalizationResource::ShouldReplaceEntry(const FTextKey& Namespace, co
 		if (bDidConflict)
 		{
 			const FString SummaryMessage = FString::Printf(TEXT("Text translation conflict for namespace \"%s\" and key \"%s\"."),
-				Namespace.GetChars(), Key.GetChars());
+				*Namespace.ToString(), *Key.ToString());
 			const FString DetailsMessage = FString::Printf(TEXT("The current translation is \"%s\" (from \"%s\" and source hash 0x%08x) and the conflicting translation of \"%s\" (from \"%s\" and source hash 0x%08x) will be ignored."), 
 				**CurrentEntry.LocalizedString,
-				CurrentEntry.LocResID.GetChars(),
+				*CurrentEntry.LocResID.ToString(),
 				CurrentEntry.SourceStringHash,
 				**NewEntry.LocalizedString,
-				NewEntry.LocResID.GetChars(),
+				*NewEntry.LocResID.ToString(),
 				NewEntry.SourceStringHash
 				);
 
diff --git a/Engine/Source/Runtime/CoreUObject/Private/Internationalization/TextPackageNamespaceUtil.cpp b/Engine/Source/Runtime/CoreUObject/Private/Internationalization/TextPackageNamespaceUtil.cpp
--- a/Engine/Source/Runtime/CoreUObject/Private/Internationalization/TextPackageNamespaceUtil.cpp
+++ b/Engine/Source/Runtime/CoreUObject/Private/Internationalization/TextPackageNamespaceUtil.cpp
@@ -291,7 +291,7 @@ bool TextNamespaceUtil::EditTextProperty(UObject* InTextOwner, const FTextProper
 		switch (InEditAction)
 		{
 		case ETextEditAction::Namespace:
-			if (const FString CurrentTextNamespace = StripPackageNamespace(CurrentTextId.GetNamespace().GetChars());
+			if (const FString CurrentTextNamespace = StripPackageNamespace(CurrentTextId.GetNamespace().ToString());
 				CurrentTextNamespace.Equals(InEditValue, ESearchCase::CaseSensitive))
 			{
 				return true;
@@ -299,7 +299,7 @@ bool TextNamespaceUtil::EditTextProperty(UObject* InTextOwner, const FTextProper
 			break;
 
 		case ETextEditAction::Key:
-			if (FCString::Strcmp(CurrentTextId.GetKey().GetChars(), *InEditValue) == 0)
+			if (CurrentTextId.GetKey().ToString().Equals(InEditValue, ESearchCase::CaseSensitive))
 			{
 				return true;
 			}
@@ -318,8 +318,8 @@ bool TextNamespaceUtil::EditTextProperty(UObject* InTextOwner, const FTextProper
 		}
 	}
 
-	const FString ProposedNamespace = (InEditAction == ETextEditAction::Namespace ? InEditValue : CurrentTextId.GetNamespace().GetChars());
-	const FString ProposedKey = (InEditAction == ETextEditAction::Key ? InEditValue : CurrentTextId.GetKey().GetChars());
+	const FString ProposedNamespace = (InEditAction == ETextEditAction::Namespace ? InEditValue : CurrentTextId.GetNamespace().ToString());
+	const FString ProposedKey = (InEditAction == ETextEditAction::Key ? InEditValue : CurrentTextId.GetKey().ToString());
 	const FString SourceString = (InEditAction == ETextEditAction::SourceString ? InEditValue : CurrentSourceString ? *CurrentSourceString : FString());
 
 	FString StableNamespace;
diff --git a/Engine/Source/Runtime/CoreUObject/Private/Serialization/PropertyLocalizationDataGathering.cpp b/Engine/Source/Runtime/CoreUObject/Private/Serialization/PropertyLocalizationDataGathering.cpp
--- a/Engine/Source/Runtime/CoreUObject/Private/Serialization/PropertyLocalizationDataGathering.cpp
+++ b/Engine/Source/Runtime/CoreUObject/Private/Serialization/PropertyLocalizationDataGathering.cpp
@@ -814,8 +814,8 @@ bool FPropertyLocalizationDataGatherer::ExtractTextIdentity(const FText& Text, F
 	const FTextId TextId = FTextInspector::GetTextId(Text);
 	if (!TextId.IsEmpty() && Text.ShouldGatherForLocalization())
 	{
-		OutNamespace = TextId.GetNamespace().GetChars();
-		OutKey = TextId.GetKey().GetChars();
+		TextId.GetNamespace().ToString(OutNamespace);
+		TextId.GetKey().ToString(OutKey);
 		if (bCleanNamespace)
 		{
 			OutNamespace = TextNamespaceUtil::StripPackageNamespace(OutNamespace);
diff --git a/Engine/Source/Runtime/Engine/Private/KismetTextLibrary.cpp b/Engine/Source/Runtime/Engine/Private/KismetTextLibrary.cpp
--- a/Engine/Source/Runtime/Engine/Private/KismetTextLibrary.cpp
+++ b/Engine/Source/Runtime/Engine/Private/KismetTextLibrary.cpp
@@ -332,8 +332,8 @@ bool UKismetTextLibrary::GetTextId(FText Text, FString& OutNamespace, FString& O
 	const FTextId TextId = FTextInspector::GetTextId(Text);
 	if (!TextId.IsEmpty())
 	{
-		OutNamespace = TextId.GetNamespace().GetChars();
-		OutKey = TextId.GetKey().GetChars();
+		TextId.GetNamespace().ToString(OutNamespace);
+		TextId.GetKey().ToString(OutKey);
 		return true;
 	}
 	return false;
