Move FName allocations from LinearVirtualMemoryAllocator into FMemory allocator as it belongs there

diff --git a/Engine/Source/Runtime/Core/Private/UObject/UnrealNames.cpp b/Engine/Source/Runtime/Core/Private/UObject/UnrealNames.cpp
--- a/Engine/Source/Runtime/Core/Private/UObject/UnrealNames.cpp
+++ b/Engine/Source/Runtime/Core/Private/UObject/UnrealNames.cpp
@@ -456,7 +456,7 @@ public:
 	{
 		for (int32 Index = CurrentBlock; Index >= 0; --Index)
 		{
-			GetPersistentLinearAllocator().TryDeallocate(Blocks[Index], BlockSizeBytes);
+			FMemory::Free(Blocks[Index]);
 		}
 	}
 
@@ -578,7 +578,7 @@ public:
 	{
 		return CurrentBlock + 1;
 	}
-	
+
 	uint8** GetBlocksForDebugVisualizer() { return Blocks; }
 
 	void DebugDump(TArray<const FNameEntry*>& Out) const
@@ -657,7 +657,7 @@ private:
 	static uint8* AllocBlock()
 	{
 		LLM_SCOPE(ELLMTag::FName);
-		return (uint8*)GetPersistentLinearAllocator().Allocate(BlockSizeBytes, alignof(FNameEntry));
+		return (uint8*)FMemory::Malloc(BlockSizeBytes, alignof(FNameEntry));
 	}
 	
 	void AllocateNewBlock()
