Minor refactoring of FPaddedBundlePointer in MallocBinnedCommon.h

diff --git a/Engine/Source/Runtime/Core/Public/HAL/MallocBinnedCommonUtils.h b/Engine/Source/Runtime/Core/Public/HAL/MallocBinnedCommonUtils.h
--- a/Engine/Source/Runtime/Core/Public/HAL/MallocBinnedCommonUtils.h
+++ b/Engine/Source/Runtime/Core/Public/HAL/MallocBinnedCommonUtils.h
@@ -79,20 +79,17 @@ namespace MallocBinnedPrivate
 		}
 
 	private:
-		struct FPaddedBundlePointer
+		struct FBundlePointer
 		{
-			FMallocBinnedCommonBase::FBundleNode* FreeBundles[UE_DEFAULT_GMallocBinnedMaxBundlesBeforeRecycle];
-#if (4 + (4 * PLATFORM_64BITS)) * UE_DEFAULT_GMallocBinnedMaxBundlesBeforeRecycle < PLATFORM_CACHE_LINE_SIZE
-#	define UE_MBC_BUNDLE_PADDING (PLATFORM_CACHE_LINE_SIZE - sizeof(FMallocBinnedCommonBase::FBundleNode*) * UE_DEFAULT_GMallocBinnedMaxBundlesBeforeRecycle)
-			uint8 Padding[UE_MBC_BUNDLE_PADDING];
-#endif
-			FPaddedBundlePointer()
+			alignas(PLATFORM_CACHE_LINE_SIZE) FMallocBinnedCommonBase::FBundleNode* FreeBundles[UE_DEFAULT_GMallocBinnedMaxBundlesBeforeRecycle];
+
+			FBundlePointer()
 			{
 				DefaultConstructItems<FMallocBinnedCommonBase::FBundleNode*>(FreeBundles, UE_DEFAULT_GMallocBinnedMaxBundlesBeforeRecycle);
 			}
 		};
-		static_assert(sizeof(FPaddedBundlePointer) == PLATFORM_CACHE_LINE_SIZE, "FPaddedBundlePointer should be the same size as a cache line");
-		MS_ALIGN(PLATFORM_CACHE_LINE_SIZE) FPaddedBundlePointer Bundles[NumSmallPools] GCC_ALIGN(PLATFORM_CACHE_LINE_SIZE);
+		static_assert(sizeof(FBundlePointer) == PLATFORM_CACHE_LINE_SIZE, "FBundlePointer should be the same size as a cache line");
+		alignas(PLATFORM_CACHE_LINE_SIZE) FBundlePointer Bundles[NumSmallPools];
 	};
 }
 
