Intrusive unset states for FScript* containers

#okforversepublic

diff --git a/Engine/Source/Runtime/Core/Public/Containers/Map.h b/Engine/Source/Runtime/Core/Public/Containers/Map.h
--- a/Engine/Source/Runtime/Core/Public/Containers/Map.h
+++ b/Engine/Source/Runtime/Core/Public/Containers/Map.h
@@ -1804,6 +1804,24 @@ public:
 	{
 	}
 
+	///////////////////////////////////////////////////
+	// Start - intrusive TOptional<TScriptMap> state //
+	///////////////////////////////////////////////////
+	constexpr static bool bHasIntrusiveUnsetOptionalState = true;
+	using IntrusiveUnsetOptionalStateType = TScriptMap;
+
+	explicit TScriptMap(FIntrusiveUnsetOptionalState Tag)
+		: Pairs(Tag)
+	{
+	}
+	bool operator==(FIntrusiveUnsetOptionalState Tag) const
+	{
+		return Pairs == Tag;
+	}
+	/////////////////////////////////////////////////
+	// End - intrusive TOptional<TScriptMap> state //
+	/////////////////////////////////////////////////
+
 	bool IsValidIndex(int32 Index) const
 	{
 		return Pairs.IsValidIndex(Index);
@@ -2014,6 +2032,14 @@ class FScriptMap : public TScriptMap<FDefaultSetAllocator, FScriptMap>
 
 public:
 	using Super::Super;
+
+	///////////////////////////////////////////////////
+	// Start - intrusive TOptional<FScriptMap> state //
+	///////////////////////////////////////////////////
+	using IntrusiveUnsetOptionalStateType = FScriptMap;
+	/////////////////////////////////////////////////
+	// End - intrusive TOptional<FScriptMap> state //
+	/////////////////////////////////////////////////
 };
 
 struct TMapPrivateFriend
diff --git a/Engine/Source/Runtime/Core/Public/Containers/ScriptArray.h b/Engine/Source/Runtime/Core/Public/Containers/ScriptArray.h
--- a/Engine/Source/Runtime/Core/Public/Containers/ScriptArray.h
+++ b/Engine/Source/Runtime/Core/Public/Containers/ScriptArray.h
@@ -288,6 +288,26 @@ private:
 	}
 
 public:
+	/////////////////////////////////////////////////////
+	// Start - intrusive TOptional<TScriptArray> state //
+	/////////////////////////////////////////////////////
+	constexpr static bool bHasIntrusiveUnsetOptionalState = true;
+	using IntrusiveUnsetOptionalStateType = TScriptArray;
+
+	explicit TScriptArray(FIntrusiveUnsetOptionalState Tag)
+		: ArrayNum(0)
+		, ArrayMax(-1)
+	{
+		// Use ArrayMax == -1 as our intrusive state so that the destructor still works without change, as it doesn't use ArrayMax.
+	}
+	bool operator==(FIntrusiveUnsetOptionalState Tag) const
+	{
+		return ArrayMax == -1;
+	}
+	///////////////////////////////////////////////////
+	// End - intrusive TOptional<TScriptArray> state //
+	///////////////////////////////////////////////////
+
 	// These should really be private, because they shouldn't be called, but there's a bunch of code
 	// that needs to be fixed first.
 	TScriptArray(const TScriptArray&) { check(false); }
@@ -315,6 +335,18 @@ protected:
 	}
 
 public:
+	/////////////////////////////////////////////////////
+	// Start - intrusive TOptional<FScriptArray> state //
+	/////////////////////////////////////////////////////
+	using IntrusiveUnsetOptionalStateType = FScriptArray;
+	explicit FScriptArray(FIntrusiveUnsetOptionalState Tag)
+		: TScriptArray(Tag)
+	{
+	}
+	///////////////////////////////////////////////////
+	// End - intrusive TOptional<FScriptArray> state //
+	///////////////////////////////////////////////////
+
 	// These should really be private, because they shouldn't be called, but there's a bunch of code
 	// that needs to be fixed first.
 	FScriptArray(const FScriptArray&) { check(false); }
diff --git a/Engine/Source/Runtime/Core/Public/Containers/Set.h b/Engine/Source/Runtime/Core/Public/Containers/Set.h
--- a/Engine/Source/Runtime/Core/Public/Containers/Set.h
+++ b/Engine/Source/Runtime/Core/Public/Containers/Set.h
@@ -1868,6 +1868,24 @@ public:
 	{
 	}
 
+	///////////////////////////////////////////////////
+	// Start - intrusive TOptional<TScriptSet> state //
+	///////////////////////////////////////////////////
+	constexpr static bool bHasIntrusiveUnsetOptionalState = true;
+	using IntrusiveUnsetOptionalStateType = TScriptSet;
+
+	explicit TScriptSet(FIntrusiveUnsetOptionalState Tag)
+		: Elements(Tag)
+	{
+	}
+	bool operator==(FIntrusiveUnsetOptionalState Tag) const
+	{
+		return Elements == Tag;
+	}
+	/////////////////////////////////////////////////
+	// End - intrusive TOptional<TScriptSet> state //
+	/////////////////////////////////////////////////
+
 	bool IsValidIndex(int32 Index) const
 	{
 		return Elements.IsValidIndex(Index);
@@ -2184,6 +2202,14 @@ class FScriptSet : public TScriptSet<FDefaultSetAllocator, FScriptSet>
 
 public:
 	using Super::Super;
+
+	///////////////////////////////////////////////////
+	// Start - intrusive TOptional<FScriptSet> state //
+	///////////////////////////////////////////////////
+	using IntrusiveUnsetOptionalStateType = FScriptSet;
+	/////////////////////////////////////////////////
+	// End - intrusive TOptional<FScriptSet> state //
+	/////////////////////////////////////////////////
 };
 
 struct TSetPrivateFriend
diff --git a/Engine/Source/Runtime/Core/Public/Containers/SparseArray.h b/Engine/Source/Runtime/Core/Public/Containers/SparseArray.h
--- a/Engine/Source/Runtime/Core/Public/Containers/SparseArray.h
+++ b/Engine/Source/Runtime/Core/Public/Containers/SparseArray.h
@@ -1325,6 +1325,24 @@ public:
 	{
 	}
 
+	/////////////////////////////////////////////////////
+	// Start - intrusive TOptional<TSparseArray> state //
+	/////////////////////////////////////////////////////
+	constexpr static bool bHasIntrusiveUnsetOptionalState = true;
+	using IntrusiveUnsetOptionalStateType = TScriptSparseArray;
+
+	explicit TScriptSparseArray(FIntrusiveUnsetOptionalState)
+		: NumFreeIndices(-1)
+	{
+	}
+	bool operator==(FIntrusiveUnsetOptionalState) const
+	{
+		return NumFreeIndices == -1;
+	}
+	///////////////////////////////////////////////////
+	// End - intrusive TOptional<TSparseArray> state //
+	///////////////////////////////////////////////////
+
 	bool IsValidIndex(int32 Index) const
 	{
 		return AllocationFlags.IsValidIndex(Index) && AllocationFlags[Index];
@@ -1506,6 +1524,14 @@ class FScriptSparseArray : public TScriptSparseArray<FDefaultSparseArrayAllocato
 
 public:
 	using Super::Super;
+
+	/////////////////////////////////////////////////////
+	// Start - intrusive TOptional<FSparseArray> state //
+	/////////////////////////////////////////////////////
+	using IntrusiveUnsetOptionalStateType = FScriptSparseArray;
+	///////////////////////////////////////////////////
+	// End - intrusive TOptional<FSparseArray> state //
+	///////////////////////////////////////////////////
 };
 
 /**
