Implement BuildAABB method for TBox2 similar to TBox

diff --git a/Engine/Source/Runtime/Core/Public/Math/Box2D.h b/Engine/Source/Runtime/Core/Public/Math/Box2D.h
--- a/Engine/Source/Runtime/Core/Public/Math/Box2D.h
+++ b/Engine/Source/Runtime/Core/Public/Math/Box2D.h
@@ -387,6 +387,20 @@ public:
 	 */
 	FString ToString() const;
 
+	/** 
+     * Utility function to build an AABB from Origin and Extent 
+     *
+     * @param Origin The location of the bounding box.
+     * @param Extent Half size of the bounding box.
+     * @return A new axis-aligned bounding box.
+     */
+    static TBox2<T> BuildAABB( const TVector2<T>& Origin, const TVector2<T>& Extent )
+    {
+    	TBox2<T> NewBox(Origin - Extent, Origin + Extent);
+
+    	return NewBox;
+    }
+
 public:
 
 	/**
