Added FindArbitraryElement to TSet, TMap and TSortedMap.
Added tests.

diff --git a/Engine/Source/Runtime/Core/Private/Containers/ContainersTest.cpp b/Engine/Source/Runtime/Core/Private/Containers/ContainersTest.cpp
--- a/Engine/Source/Runtime/Core/Private/Containers/ContainersTest.cpp
+++ b/Engine/Source/Runtime/Core/Private/Containers/ContainersTest.cpp
@@ -930,6 +930,141 @@ bool FContainersTSetTest::RunTest(const FString& Parameters)
 		}
 	}
 
+	// FindArbitraryElement
+	{
+		{
+			TSet<FString> Set;
+			FString* Found = Set.FindArbitraryElement();
+			TestNull(TEXT("FindArbitraryElement returns null on an empty set"), Found);
+		}
+
+		{
+			constexpr int32 ElementsToAdd = 100;
+			constexpr int32 IndexToKeep = 67;
+
+			TSet<FString> Set;
+			for (int32 I = 0; I != ElementsToAdd; ++I)
+			{
+				Set.Add(LexToString(I));
+			}
+			for (int32 I = 0; I != ElementsToAdd; ++I)
+			{
+				if (I != IndexToKeep)
+				{
+					Set.Remove(LexToString(I));
+				}
+			}
+
+			FString* Found = Set.FindArbitraryElement();
+			if (Found)
+			{
+				TestNotNull(TEXT("FindArbitraryElement finds a value on a highly sparse set"), Found);
+				TestEqual(TEXT("FindArbitraryElement finds the correct value on a highly sparse set"), *Found, LexToString(IndexToKeep));
+			}
+
+			Found = nullptr;
+			Set.Remove(LexToString(IndexToKeep));
+
+			FString* Found2 = Set.FindArbitraryElement();
+			TestNull(TEXT("FindArbitraryElement returns null on an emptied set"), Found2);
+		}
+	}
+
+	return !HasAnyErrors();
+}
+
+IMPLEMENT_SIMPLE_AUTOMATION_TEST(FContainersTMapTest, "System.Core.Containers.TMap", EAutomationTestFlags_ApplicationContextMask | EAutomationTestFlags::EngineFilter)
+bool FContainersTMapTest::RunTest(const FString& Parameters)
+{
+	// FindArbitraryElement
+	{
+		{
+			TMap<int32, FString> Map;
+			TPair<int32, FString>* Found = Map.FindArbitraryElement();
+			TestNull(TEXT("FindArbitraryElement returns null on an empty map"), Found);
+		}
+
+		{
+			constexpr int32 ElementsToAdd = 100;
+			constexpr int32 IndexToKeep = 23;
+
+			TMap<int32, FString> Map;
+			for (int32 I = 0; I != ElementsToAdd; ++I)
+			{
+				Map.Add(I, LexToString(I));
+			}
+			for (int32 I = 0; I != ElementsToAdd; ++I)
+			{
+				if (I != IndexToKeep)
+				{
+					Map.Remove(I);
+				}
+			}
+
+			TPair<int32, FString>* Found = Map.FindArbitraryElement();
+
+			TestNotNull(TEXT("FindArbitraryElement finds a value on a highly sparse map"), Found);
+			if (Found)
+			{
+				TestEqual(TEXT("FindArbitraryElement finds the correct key on a highly sparse map"), Found->Key, IndexToKeep);
+				TestEqual(TEXT("FindArbitraryElement finds the correct value on a highly sparse map"), Found->Value, LexToString(IndexToKeep));
+			}
+
+			Found = nullptr;
+			Map.Remove(IndexToKeep);
+
+			TPair<int32, FString>* Found2 = Map.FindArbitraryElement();
+			TestNull(TEXT("FindArbitraryElement returns null on an emptied map"), Found2);
+		}
+	}
+
+	return !HasAnyErrors();
+}
+
+IMPLEMENT_SIMPLE_AUTOMATION_TEST(FContainersTSortedMapTest, "System.Core.Containers.TSortedMap", EAutomationTestFlags_ApplicationContextMask | EAutomationTestFlags::EngineFilter)
+bool FContainersTSortedMapTest::RunTest(const FString& Parameters)
+{
+	// FindArbitraryElement
+	{
+		{
+			TSortedMap<int32, FString> Map;
+			TPair<int32, FString>* Found = Map.FindArbitraryElement();
+			TestNull(TEXT("FindArbitraryElement returns null on an empty map"), Found);
+		}
+
+		{
+			constexpr int32 ElementsToAdd = 100;
+			constexpr int32 IndexToKeep = 23;
+
+			TSortedMap<int32, FString> Map;
+			for (int32 I = 0; I != ElementsToAdd; ++I)
+			{
+				Map.Add(I, LexToString(I));
+			}
+			for (int32 I = 0; I != ElementsToAdd; ++I)
+			{
+				if (I != IndexToKeep)
+				{
+					Map.Remove(I);
+				}
+			}
+
+			TPair<int32, FString>* Found = Map.FindArbitraryElement();
+
+			TestNotNull(TEXT("FindArbitraryElement finds a value on a highly sparse map"), Found);
+			if (Found)
+			{
+				TestEqual(TEXT("FindArbitraryElement finds the correct key on a highly sparse map"), Found->Key, IndexToKeep);
+				TestEqual(TEXT("FindArbitraryElement finds the correct value on a highly sparse map"), Found->Value, LexToString(IndexToKeep));
+			}
+
+			Found = nullptr;
+			Map.Remove(IndexToKeep);
+
+			TPair<int32, FString>* Found2 = Map.FindArbitraryElement();
+			TestNull(TEXT("FindArbitraryElement returns null on an emptied map"), Found2);
+		}
+	}
 
 	return !HasAnyErrors();
 }
diff --git a/Engine/Source/Runtime/Core/Public/Containers/Map.h b/Engine/Source/Runtime/Core/Public/Containers/Map.h
--- a/Engine/Source/Runtime/Core/Public/Containers/Map.h
+++ b/Engine/Source/Runtime/Core/Public/Containers/Map.h
@@ -737,6 +737,22 @@ public:
 		return DefaultValue;
 	}
 
+	/**
+	 * Finds any pair in the map and returns a pointer to it.
+	 * Callers should not depend on particular patterns in the behaviour of this function.
+	 * @return A pointer to an arbitrary pair, or nullptr if the container is empty.
+	 */
+	ElementType* FindArbitraryElement()
+	{
+		// The goal of this function is to be fast, and so the implementation may be improved at any time even if it gives different results.
+
+		return Pairs.FindArbitraryElement();
+	}
+	const ElementType* FindArbitraryElement() const
+	{
+		return const_cast<TMapBase*>(this)->FindArbitraryElement();
+	}
+
 	/**
 	 * Check if map contains the specified key.
 	 *
diff --git a/Engine/Source/Runtime/Core/Public/Containers/Set.h b/Engine/Source/Runtime/Core/Public/Containers/Set.h
--- a/Engine/Source/Runtime/Core/Public/Containers/Set.h
+++ b/Engine/Source/Runtime/Core/Public/Containers/Set.h
@@ -914,6 +914,23 @@ private:
 	}
 
 public:
+	/**
+	 * Finds any element in the set and returns a pointer to it.
+	 * Callers should not depend on particular patterns in the behaviour of this function.
+	 * @return A pointer to an arbitrary element, or nullptr if the container is empty.
+	 */
+	ElementType* FindArbitraryElement()
+	{
+		// The goal of this function is to be fast, and so the implementation may be improved at any time even if it gives different results.
+
+		int32 Result = Elements.FindArbitraryElementIndex();
+		return (Result != INDEX_NONE) ? &Elements[Result].Value : nullptr;
+	}
+	const ElementType* FindArbitraryElement() const
+	{
+		return const_cast<TSet*>(this)->FindArbitraryElement();
+	}
+
 	/**
 	 * Finds an element with the given key in the set.
 	 * @param Key - The key to search for.
diff --git a/Engine/Source/Runtime/Core/Public/Containers/SortedMap.h b/Engine/Source/Runtime/Core/Public/Containers/SortedMap.h
--- a/Engine/Source/Runtime/Core/Public/Containers/SortedMap.h
+++ b/Engine/Source/Runtime/Core/Public/Containers/SortedMap.h
@@ -356,6 +356,22 @@ public:
 		return DefaultValue;
 	}
 
+	/**
+	 * Finds any pair in the map and returns a pointer to it.
+	 * Callers should not depend on particular patterns in the behaviour of this function.
+	 * @return A pointer to an arbitrary pair, or nullptr if the container is empty.
+	 */
+	ElementType* FindArbitraryElement()
+	{
+		// The goal of this function is to be fast, and so the implementation may be improved at any time even if it gives different results.
+
+		return Pairs.IsEmpty() ? nullptr : &Pairs.Last();
+	}
+	const ElementType* FindArbitraryElement() const
+	{
+		return const_cast<TSortedMap*>(this)->FindArbitraryElement();
+	}
+
 	/**
 	 * Checks if map contains the specified key.
 	 *
diff --git a/Engine/Source/Runtime/Core/Public/Containers/SparseArray.h b/Engine/Source/Runtime/Core/Public/Containers/SparseArray.h
--- a/Engine/Source/Runtime/Core/Public/Containers/SparseArray.h
+++ b/Engine/Source/Runtime/Core/Public/Containers/SparseArray.h
@@ -699,6 +699,23 @@ public:
 		return INDEX_NONE;
 	}
 
+	/**
+	 * Finds the index of an arbitrary element.
+	 * Callers should not depend on particular patterns in the behaviour of this function.
+	 * @return The index of any element in the container, or INDEX_NONE if it is empty.
+	 */
+	int32 FindArbitraryElementIndex() const
+	{
+		// The goal of this function is to be fast, and so the implementation may be improved at any time even if it gives different results.
+
+		if (NumFreeIndices == 0)
+		{
+			return Data.Num() - 1; // this will return INDEX_NONE if Data is empty.
+		}
+
+		return AllocationFlags.Find(true);
+	}
+
 	/** 
 	 * Helper function to return the amount of memory allocated by this container 
 	 * Only returns the size of allocations made directly by the container, not the elements themselves.
