Fix compile error with preview compiler

#rnx

diff --git a/Engine/Source/Runtime/Core/Public/Experimental/ConcurrentLinearAllocator.h b/Engine/Source/Runtime/Core/Public/Experimental/ConcurrentLinearAllocator.h
--- a/Engine/Source/Runtime/Core/Public/Experimental/ConcurrentLinearAllocator.h
+++ b/Engine/Source/Runtime/Core/Public/Experimental/ConcurrentLinearAllocator.h
@@ -28,12 +28,16 @@
 #include <sanitizer/asan_interface.h>
 #if defined(__SANITIZE_ADDRESS__)
 #define IS_ASAN_ENABLED 1
-#elif defined(__has_feature) && __has_feature(address_sanitizer)
+#elif defined(__has_feature)
+#if __has_feature(address_sanitizer)
 #define IS_ASAN_ENABLED 1
 #else
 #define IS_ASAN_ENABLED 0
 #endif
 #else
+#define IS_ASAN_ENABLED 0
+#endif
+#else
 #define ASAN_POISON_MEMORY_REGION(addr, size) ((void)(addr), (void)(size))
 #define ASAN_UNPOISON_MEMORY_REGION(addr, size) ((void)(addr), (void)(size))
 #define IS_ASAN_ENABLED 0
