Add test for gsl::not_null. Add simple utilities to hash, serialize and optional. Add conversion to/from TNotNullPtr.
Didn't add it to TNonNullPtr<TObjectPtr<>> on purpose. TObjectPtr is to track reference and gsl::not_null should only be used with UObject for argument and temporary variables.

diff --git a/Engine/Source/Runtime/Core/Core.Build.cs b/Engine/Source/Runtime/Core/Core.Build.cs
--- a/Engine/Source/Runtime/Core/Core.Build.cs
+++ b/Engine/Source/Runtime/Core/Core.Build.cs
@@ -33,6 +33,7 @@ public class Core : ModuleRules
 		PrivateDependencyModuleNames.Add("OodleDataCompression");
 		PrivateDependencyModuleNames.Add("xxhash");
 
+		PublicDependencyModuleNames.Add("GuidelinesSupportLibrary");
 		PublicDependencyModuleNames.Add("TraceLog");
 
 		PrivateIncludePathModuleNames.AddRange(
diff --git a/Engine/Source/Runtime/Core/Public/Templates/NonNullPointer.h b/Engine/Source/Runtime/Core/Public/Templates/NonNullPointer.h
--- a/Engine/Source/Runtime/Core/Public/Templates/NonNullPointer.h
+++ b/Engine/Source/Runtime/Core/Public/Templates/NonNullPointer.h
@@ -9,6 +9,7 @@
 #include "Templates/Requires.h"
 #include "Templates/UnrealTypeTraits.h"
 
+#include <gsl/pointers>
 #include <type_traits>
 
 class FArchive;
@@ -48,6 +49,15 @@ public:
 		ensureMsgf(InObject, TEXT("Tried to initialize TNonNullPtr with a null pointer!"));
 	}
 
+	/**
+	 * Converts from gsl::not_null
+	 */
+	template <typename OtherType, typename = std::enable_if_t<std::is_convertible<OtherType, ObjectType*>::value>>
+	FORCEINLINE TNonNullPtr(const gsl::not_null<OtherType>& Other)
+		: Object(Other.get())
+	{
+	}
+
 	/**
 	 * Constructs a non-null pointer from another non-null pointer
 	 */
@@ -93,6 +103,17 @@ public:
 		return *this;
 	}
 
+
+	/**
+	 * Assignment operator taking a gsl::not_null
+	 */
+	template <typename OtherType, typename = std::enable_if_t<std::is_convertible<OtherType, ObjectType*>::value>>
+	FORCEINLINE TNonNullPtr& operator=(const gsl::not_null<OtherType>& Other)
+	{
+		Object = Other.Get();
+		return *this;
+	}
+
 	/**
 	 * Comparison, will also handle default constructed state
 	 */
@@ -153,6 +174,15 @@ public:
 		ensureMsgf(Object, TEXT("Tried to access null pointer!"));
 		return Object;
 	}
+	
+	/**
+	 * Converts to a gsl::not_null
+	 */
+	FORCEINLINE operator gsl::not_null<ObjectType*>() const
+	{
+		ensureMsgf(Object, TEXT("Tried to access null pointer!"));
+		return gsl::make_not_null<ObjectType*>(Object);
+	}
 
 	/**
 	 * Returns the internal pointer
diff --git a/Engine/Source/Runtime/Core/Public/Templates/NotNull.h b/Engine/Source/Runtime/Core/Public/Templates/NotNull.h
new file mode 100644
--- /dev/null
+++ b/Engine/Source/Runtime/Core/Public/Templates/NotNull.h
@@ -0,0 +1,55 @@
+// Copyright Epic Games, Inc. All Rights Reserved.
+
+#pragma once
+
+#include "CoreTypes.h"
+#include "Misc/OptionalFwd.h"
+#include "Templates/UnrealTemplate.h"
+#include <gsl/pointers>
+
+class FArchive;
+
+/** Utility function to turn an `TOptional<gsl::not_null<T>>` back into a nullable T. */
+template<typename ObjectType>
+[[nodiscard]] inline ObjectType* GetRawPointerOrNull(const TOptional<gsl::not_null<ObjectType*>>& Optional)
+{
+	return Optional.IsSet() ? static_cast<ObjectType*>(Optional.GetValue()) : nullptr;
+}
+
+/** Utility function to turn an `TOptional<gsl::strict_not_null<T>>` back into a nullable T. */
+template<typename ObjectType>
+[[nodiscard]] inline ObjectType* GetRawPointerOrNull(const TOptional<gsl::strict_not_null<ObjectType*>>& Optional)
+{
+	return Optional.IsSet() ? static_cast<ObjectType*>(Optional.GetValue()) : nullptr;
+}
+
+
+/** Utility function to serialize a `gsl::not_null<T>. */
+template<typename ObjectType>
+FArchive& operator<<(FArchive& Ar, gsl::not_null<ObjectType>& NotNull)
+{
+	Ar << *NotNull;
+	return Ar;
+}
+
+/** Utility function to serialize a `gsl::strict_not_null<T>`. */
+template<typename ObjectType>
+FArchive& operator<<(FArchive& Ar, gsl::strict_not_null<ObjectType>& NotNull)
+{
+	Ar << *NotNull;
+	return Ar;
+}
+
+/** Utility function to hash a `gsl::not_null<T>`. */
+template<typename ObjectType>
+[[nodiscard]] inline auto GetTypeHash(const gsl::not_null<ObjectType>& NotNull) -> decltype(GetTypeHash(*NotNull))
+{
+	return GetTypeHash(*NotNull);
+}
+
+/** Utility function to hash a `gsl::strict_not_null<T>`. */
+template<typename ObjectType>
+[[nodiscard]] inline auto GetTypeHash(const gsl::strict_not_null<ObjectType>& NotNull) -> decltype(GetTypeHash(*NotNull))
+{
+	return GetTypeHash(*NotNull);
+}
diff --git a/Engine/Source/Runtime/Core/Tests/Templates/NotNullTest.cpp b/Engine/Source/Runtime/Core/Tests/Templates/NotNullTest.cpp
new file mode 100644
--- /dev/null
+++ b/Engine/Source/Runtime/Core/Tests/Templates/NotNullTest.cpp
@@ -0,0 +1,251 @@
+// Copyright Epic Games, Inc. All Rights Reserved.
+
+#if WITH_TESTS
+
+#include "Misc/Optional.h"
+#include "Templates/NonNullPointer.h"
+#include "Templates/NotNull.h"
+#include "Templates/SharedPointer.h"
+#include "Templates/Tuple.h"
+#include "Templates/UniquePtr.h"
+
+#include "Tests/TestHarnessAdapter.h"
+
+#include <gsl/pointers>
+#include <type_traits>
+#include <optional>
+
+TEST_CASE_NAMED(FTemplatesNotNullTest, "System::Core::Templates::NonNullPtr", "[Core][Templates][SmokeFilter]")
+{
+	SECTION("Static")
+	{
+		//~ gsl::not_null<*>
+		auto MakeInt = []()->int32* { return nullptr; };
+		STATIC_REQUIRE(!std::is_constructible<gsl::not_null<int32*>>::value);
+		STATIC_REQUIRE(std::is_constructible<gsl::not_null<int32*>, int32*>::value);
+		STATIC_REQUIRE(std::is_constructible<gsl::not_null<int32*>, decltype(gsl::make_not_null<int32*>(MakeInt()))>::value);
+		STATIC_REQUIRE(std::is_constructible<gsl::not_null<int32*>, gsl::not_null<int32*>>::value);
+		
+
+		STATIC_REQUIRE(std::is_copy_constructible<gsl::not_null<int32*>>::value);
+		STATIC_REQUIRE(std::is_move_constructible<gsl::not_null<int32*>>::value);
+		STATIC_REQUIRE(std::is_copy_assignable<gsl::not_null<int32*>>::value);
+		STATIC_REQUIRE(std::is_move_assignable<gsl::not_null<int32*>>::value);
+
+		//~ gsl::not_null<TUniquePtr<>>
+		STATIC_REQUIRE(!std::is_constructible<gsl::not_null<TUniquePtr<int32>>>::value);
+		STATIC_REQUIRE(std::is_constructible<gsl::not_null<TUniquePtr<int32>>, TUniquePtr<int32>>::value);
+		STATIC_REQUIRE(std::is_constructible<gsl::not_null<TUniquePtr<int32>>, TUniquePtr<int32>&&>::value);
+		STATIC_REQUIRE(std::is_constructible<gsl::not_null<TUniquePtr<int32>>, decltype(MakeUnique<int32>(0))>::value);
+		STATIC_REQUIRE(std::is_constructible<gsl::not_null<TUniquePtr<int32>>, decltype(gsl::make_not_null(MakeUnique<int32>(0)))>::value);
+		STATIC_REQUIRE(std::is_constructible<gsl::not_null<TUniquePtr<int32>>, gsl::not_null<TUniquePtr<int32>>>::value);
+
+		STATIC_REQUIRE(!std::is_copy_constructible<gsl::not_null<TUniquePtr<int32>>>::value);
+		STATIC_REQUIRE(std::is_move_constructible<gsl::not_null<TUniquePtr<int32>>>::value);
+		STATIC_REQUIRE(!std::is_copy_assignable<gsl::not_null<TUniquePtr<int32>>>::value);
+		STATIC_REQUIRE(!std::is_move_assignable<gsl::not_null<TUniquePtr<int32>>>::value);
+
+		STATIC_REQUIRE(!std::is_trivially_copy_constructible<TUniquePtr<int32>>::value);
+
+		//~ gsl::not_null<TSharedPtr<>>
+		STATIC_REQUIRE(!std::is_constructible<gsl::not_null<TSharedPtr<int32>>>::value);
+		STATIC_REQUIRE(std::is_constructible<gsl::not_null<TSharedPtr<int32>>, TSharedPtr<int32>>::value);
+		STATIC_REQUIRE(std::is_constructible<gsl::not_null<TSharedPtr<int32>>, decltype(MakeShared<int32>(0))>::value);
+		STATIC_REQUIRE(std::is_constructible<gsl::not_null<TSharedPtr<int32>>, decltype(gsl::make_not_null(TSharedPtr<int32>(MakeShareable<int32>(new int32()))))>::value);
+		STATIC_REQUIRE(std::is_constructible<gsl::not_null<TSharedPtr<int32>>, gsl::not_null<TSharedPtr<int32>>>::value);
+
+		STATIC_REQUIRE(std::is_copy_constructible<gsl::not_null<TSharedPtr<int32>>>::value);
+		STATIC_REQUIRE(std::is_move_constructible<gsl::not_null<TSharedPtr<int32>>>::value);
+		STATIC_REQUIRE(std::is_copy_assignable<gsl::not_null<TSharedPtr<int32>>>::value);
+		STATIC_REQUIRE(std::is_move_assignable<gsl::not_null<TSharedPtr<int32>>>::value);
+	}
+
+	struct FLocal
+	{
+		int32 CurrentValue = 0;
+		void Foo(int32* Value)
+		{
+			if (Value)
+			{
+				*Value = ++CurrentValue;
+			}
+		}
+		void Foo(const TUniquePtr<int32>& Value)
+		{
+			if (Value)
+			{
+				*Value = ++CurrentValue;
+			}
+		}
+		void Foo(const TSharedPtr<int32>& Value)
+		{
+			if (Value)
+			{
+				*Value = ++CurrentValue;
+			}
+		}
+
+		void Foo(gsl::not_null<int32*> Value)
+		{
+			*Value = ++CurrentValue;
+		}
+		void Foo(gsl::not_null<TUniquePtr<int32>>& Value)
+		{
+			*Value = ++CurrentValue;
+		}
+		void Foo(gsl::not_null<std::unique_ptr<int32>>& Value)
+		{
+			*Value = ++CurrentValue;
+		}
+		void Foo(gsl::not_null<TSharedPtr<int32>>& Value)
+		{
+			*Value = ++CurrentValue;
+		}
+
+		void NotNullPtr(gsl::not_null<int32*> Value)
+		{
+			*Value = ++CurrentValue;
+		}
+	};
+
+	SECTION("Strict")
+	{
+		FLocal Local;
+		int32 LocalValue = Local.CurrentValue;
+
+		int32* Value = &LocalValue;
+		Local.Foo(gsl::make_not_null(Value));
+		Local.Foo(gsl::make_strict_not_null(Value));
+		
+		gsl::strict_not_null<int32*> StrictNotNull = gsl::make_not_null(Value);
+		Local.Foo(StrictNotNull);
+		gsl::not_null<int32*> NotNull = StrictNotNull;
+		Local.Foo(NotNull);
+	}
+
+	SECTION("TNonNullPtr to not_null. not_null to TNonNullPtr")
+	{
+		FLocal Local;
+		int32 LocalValue = Local.CurrentValue;
+
+		gsl::not_null<int32*> NotNull = gsl::make_not_null(&LocalValue);
+		Local.Foo(NotNull);
+		Local.NotNullPtr(NotNull);
+		TNonNullPtr<int32> ImplicitNotNull = NotNull;
+		Local.NotNullPtr(ImplicitNotNull);
+		NotNull = ImplicitNotNull;
+		Local.Foo(NotNull);
+		CHECK(Local.CurrentValue == *NotNull);
+
+		gsl::strict_not_null<int32*> StrictNotNull = gsl::make_not_null(&LocalValue);
+		Local.Foo(StrictNotNull);
+		Local.NotNullPtr(StrictNotNull);
+		TNonNullPtr<int32> ImplicitStrictNotNull = StrictNotNull;
+		Local.NotNullPtr(ImplicitStrictNotNull);
+		StrictNotNull = ImplicitStrictNotNull;
+		Local.Foo(StrictNotNull);
+		CHECK(Local.CurrentValue == *StrictNotNull);
+	}
+
+	SECTION("gsl::make_not_null")
+	{
+		FLocal Local;
+		int32 LocalValue = Local.CurrentValue;
+
+		Local.Foo(&LocalValue);
+		Local.Foo(gsl::make_not_null(&LocalValue));
+		CHECK(Local.CurrentValue == LocalValue);
+
+		TUniquePtr<int32> UniqueValue = MakeUnique<int32>(Local.CurrentValue);
+		Local.Foo(UniqueValue);
+		CHECK(Local.CurrentValue == *UniqueValue);
+
+		TSharedPtr<int32> SharedValue = MakeShared<int32>(Local.CurrentValue);
+		Local.Foo(SharedValue);
+		Local.Foo(gsl::make_not_null(SharedValue));
+		CHECK(Local.CurrentValue == *SharedValue);
+	}
+
+	SECTION("Move")
+	{
+		FLocal Local;
+		int32 LocalValue = Local.CurrentValue;
+		
+		TUniquePtr<int32> UniqueValue = MakeUnique<int32>(Local.CurrentValue);
+		Local.Foo(UniqueValue);
+		CHECK(Local.CurrentValue == *UniqueValue);
+		gsl::not_null<TUniquePtr<int32>> NotNullUniqueValue1 = MoveTemp(UniqueValue);
+		Local.Foo(NotNullUniqueValue1);
+		CHECK(Local.CurrentValue == *NotNullUniqueValue1);
+
+		TSharedPtr<int32> SharedValue = MakeShared<int32>(Local.CurrentValue);
+		Local.Foo(SharedValue);
+		CHECK(Local.CurrentValue == *SharedValue);
+		gsl::not_null<TSharedPtr<int32>> NotNullSharedValue1 = MoveTemp(SharedValue);
+		Local.Foo(NotNullSharedValue1);
+		CHECK(Local.CurrentValue == *NotNullSharedValue1);
+		gsl::not_null<TSharedPtr<int32>> NotNullSharedValue2 = MoveTemp(NotNullSharedValue1);
+		Local.Foo(NotNullSharedValue2);
+		CHECK(Local.CurrentValue == *NotNullSharedValue2);
+	}
+
+	SECTION("Optional and GetRawPointerOrNull")
+	{
+		FLocal Local;
+		int32 LocalValue = Local.CurrentValue;
+
+		TOptional<int32*> ValuePtr = &LocalValue;
+		Local.Foo(ValuePtr.GetValue());
+		CHECK(Local.CurrentValue == *ValuePtr.GetValue());
+
+		TOptional<gsl::not_null<int32*>> NotNullValuePtr = gsl::make_not_null(&LocalValue);
+		Local.Foo(NotNullValuePtr.GetValue());
+		Local.Foo(GetRawPointerOrNull(NotNullValuePtr));
+		CHECK(Local.CurrentValue == *NotNullValuePtr.GetValue());
+
+		TOptional<TUniquePtr<int32>> UniqueValue = MakeUnique<int32>(Local.CurrentValue);
+		Local.Foo(UniqueValue.GetValue());
+		CHECK(Local.CurrentValue == *UniqueValue.GetValue());
+
+		// Works with std::optional but not with TOptional
+		{
+			std::optional<gsl::not_null<TUniquePtr<int32>>> NotNullUniqueValue1 = MakeUnique<int32>(Local.CurrentValue);
+			Local.Foo(NotNullUniqueValue1.value());
+			CHECK(Local.CurrentValue == *NotNullUniqueValue1.value());
+
+			//TOptional<gsl::not_null<TUniquePtr<int32>>> NotNullUniqueValue = MakeUnique<int32>(LocalValue);
+			//Local.Foo(NotNullUniqueValue.GetValue());
+			//CHECK(Local.CurrentValue == *NotNullUniqueValue.GetValue());
+		}
+
+		TOptional<gsl::not_null<TSharedPtr<int32>>> NotNullSharedValue = gsl::make_not_null(TSharedPtr<int32>(MakeShareable<int32>(new int32(Local.CurrentValue))));
+		Local.Foo(NotNullSharedValue.GetValue());
+		CHECK(Local.CurrentValue == *NotNullSharedValue.GetValue());
+	}
+
+	SECTION("Tuple")
+	{
+		FLocal Local;
+		int32 LocalValue = Local.CurrentValue;
+
+		TTuple<int32*> ValuePtr = &LocalValue;
+		Local.Foo(ValuePtr.Get<int32*>());
+		CHECK(Local.CurrentValue == *ValuePtr.Get<int32*>());
+
+		TTuple<gsl::not_null<int32*>> NotNullValuePtr = gsl::make_not_null(&LocalValue);
+		Local.Foo(NotNullValuePtr.Get<gsl::not_null<int32*>>());
+		CHECK(Local.CurrentValue == *NotNullValuePtr.Get<gsl::not_null<int32*>>());
+	}
+
+	SECTION("Archive")
+	{
+		FLocal Local;
+		int32 LocalValue = Local.CurrentValue;
+
+		FArchive Ar;
+		gsl::not_null<int32*> ValuePtr = gsl::make_not_null(&LocalValue);
+		Ar << ValuePtr;
+	}
+}
+
+#endif //WITH_TESTS
\ No newline at end of file
