[Backout] - CL35933311
Original CL Desc
-----------------------------------------------------------------
Add test for gsl::not_null. Add simple utilities to hash, serialize and optional. Add conversion to/from TNotNullPtr.
Didn't add it to TNonNullPtr<TObjectPtr<>> on purpose. TObjectPtr is to track reference and gsl::not_null should only be used with UObject for argument and temporary variables.

diff --git a/Engine/Source/Runtime/Core/Core.Build.cs b/Engine/Source/Runtime/Core/Core.Build.cs
--- a/Engine/Source/Runtime/Core/Core.Build.cs
+++ b/Engine/Source/Runtime/Core/Core.Build.cs
@@ -33,7 +33,6 @@ public class Core : ModuleRules
 		PrivateDependencyModuleNames.Add("OodleDataCompression");
 		PrivateDependencyModuleNames.Add("xxhash");
 
-		PublicDependencyModuleNames.Add("GuidelinesSupportLibrary");
 		PublicDependencyModuleNames.Add("TraceLog");
 
 		PrivateIncludePathModuleNames.AddRange(
diff --git a/Engine/Source/Runtime/Core/Public/Templates/NonNullPointer.h b/Engine/Source/Runtime/Core/Public/Templates/NonNullPointer.h
--- a/Engine/Source/Runtime/Core/Public/Templates/NonNullPointer.h
+++ b/Engine/Source/Runtime/Core/Public/Templates/NonNullPointer.h
@@ -9,7 +9,6 @@
 #include "Templates/Requires.h"
 #include "Templates/UnrealTypeTraits.h"
 
-#include <gsl/pointers>
 #include <type_traits>
 
 class FArchive;
@@ -49,15 +48,6 @@ public:
 		ensureMsgf(InObject, TEXT("Tried to initialize TNonNullPtr with a null pointer!"));
 	}
 
-	/**
-	 * Converts from gsl::not_null
-	 */
-	template <typename OtherType, typename = std::enable_if_t<std::is_convertible<OtherType, ObjectType*>::value>>
-	FORCEINLINE TNonNullPtr(const gsl::not_null<OtherType>& Other)
-		: Object(Other.get())
-	{
-	}
-
 	/**
 	 * Constructs a non-null pointer from another non-null pointer
 	 */
@@ -103,17 +93,6 @@ public:
 		return *this;
 	}
 
-
-	/**
-	 * Assignment operator taking a gsl::not_null
-	 */
-	template <typename OtherType, typename = std::enable_if_t<std::is_convertible<OtherType, ObjectType*>::value>>
-	FORCEINLINE TNonNullPtr& operator=(const gsl::not_null<OtherType>& Other)
-	{
-		Object = Other.Get();
-		return *this;
-	}
-
 	/**
 	 * Comparison, will also handle default constructed state
 	 */
@@ -174,15 +153,6 @@ public:
 		ensureMsgf(Object, TEXT("Tried to access null pointer!"));
 		return Object;
 	}
-	
-	/**
-	 * Converts to a gsl::not_null
-	 */
-	FORCEINLINE operator gsl::not_null<ObjectType*>() const
-	{
-		ensureMsgf(Object, TEXT("Tried to access null pointer!"));
-		return gsl::make_not_null<ObjectType*>(Object);
-	}
 
 	/**
 	 * Returns the internal pointer
diff --git a/Engine/Source/Runtime/Core/Public/Templates/NotNull.h b/Engine/Source/Runtime/Core/Public/Templates/NotNull.h
deleted file mode 100644
diff --git a/Engine/Source/Runtime/Core/Tests/Templates/NotNullTest.cpp b/Engine/Source/Runtime/Core/Tests/Templates/NotNullTest.cpp
deleted file mode 100644
