[AutoRTFM] Reland "Add PopAllOnAbortHandlers method."

This API deletes all pushed handlers for a given key. It is functionally equivalent to calling `PopOnAbortHandler` repeatedly, but more efficient.

diff --git a/Engine/Source/Programs/AutoRTFMTests/Private/Abort.cpp b/Engine/Source/Programs/AutoRTFMTests/Private/Abort.cpp
--- a/Engine/Source/Programs/AutoRTFMTests/Private/Abort.cpp
+++ b/Engine/Source/Programs/AutoRTFMTests/Private/Abort.cpp
@@ -250,13 +250,29 @@ TEST_CASE("Abort.PushOnAbortHandler_WithPop_NoAbort")
 {
 	int Value = 55;
 
+	AutoRTFM::Commit([&]
+		{
+			Value = 66;
+			AutoRTFM::PushOnAbortHandler(UIntToPointer(747), [&Value]() { Value = 77; });
+			Value = 88;
+
+			AutoRTFM::PopOnAbortHandler(UIntToPointer(747));
+		});
+
+	REQUIRE(Value == 88);
+}
+
+TEST_CASE("Abort.PushOnAbortHandler_WithPopAll_NoAbort")
+{
+	int Value = 55;
+
 	AutoRTFM::Commit([&]
 	{
 		Value = 66;
 		AutoRTFM::PushOnAbortHandler(UIntToPointer(747), [&Value](){ Value = 77; });
 		Value = 88;
 
-		AutoRTFM::PopOnAbortHandler(UIntToPointer(747));
+		AutoRTFM::PopAllOnAbortHandlers(UIntToPointer(747));
 	});
 
 	REQUIRE(Value == 88);
@@ -281,28 +297,69 @@ TEST_CASE("Abort.PushOnAbortHandler_WithPop_WithAbort")
 	REQUIRE(Value == 55);
 }
 
-TEST_CASE("Abort.PushOnAbortHandler_Duplicates1")
+TEST_CASE("Abort.PushOnAbortHandler_WithPopAll_WithAbort")
 {
 	int Value = 55;
 
 	const AutoRTFM::ETransactionResult Result = AutoRTFM::Transact([&]
 	{
 		Value = 66;
-		AutoRTFM::PushOnAbortHandler(UIntToPointer(747), [&Value](){ Value = 77; });
-		AutoRTFM::PushOnAbortHandler(UIntToPointer(747), [&Value](){ Value = 88; });
-		Value = 99;
+		AutoRTFM::PushOnAbortHandler(UIntToPointer(747), [&Value]() { Value = 77; });
+		Value = 88;
 
-		AutoRTFM::PopOnAbortHandler(UIntToPointer(747));
+		AutoRTFM::PopAllOnAbortHandlers(UIntToPointer(747));
 
 		AutoRTFM::AbortTransaction();
 	});
 
+	REQUIRE(AutoRTFM::ETransactionResult::AbortedByRequest == Result);
+	REQUIRE(Value == 55);
+}
+
+TEST_CASE("Abort.PushOnAbortHandler_Duplicates1")
+{
+	int Value = 55;
+
+	const AutoRTFM::ETransactionResult Result = AutoRTFM::Transact([&]
+		{
+			Value = 66;
+			AutoRTFM::PushOnAbortHandler(UIntToPointer(747), [&Value]() { Value = 77; });
+			AutoRTFM::PushOnAbortHandler(UIntToPointer(747), [&Value]() { Value = 88; });
+			Value = 99;
+
+			AutoRTFM::PopOnAbortHandler(UIntToPointer(747));
+
+			AutoRTFM::AbortTransaction();
+		});
+
 	REQUIRE(AutoRTFM::ETransactionResult::AbortedByRequest == Result);
 
 	// The first push on abort will still go through.
 	REQUIRE(Value == 77);
 }
 
+TEST_CASE("Abort.PushOnAbortHandler_PopAll_Duplicates")
+{
+	int Value = 55;
+
+	const AutoRTFM::ETransactionResult Result = AutoRTFM::Transact([&]
+	{
+		Value = 66;
+		AutoRTFM::PushOnAbortHandler(UIntToPointer(747), [&Value]() { Value = 77; });
+		AutoRTFM::PushOnAbortHandler(UIntToPointer(747), [&Value]() { Value = 88; });
+		Value = 99;
+
+		AutoRTFM::PopAllOnAbortHandlers(UIntToPointer(747));
+
+		AutoRTFM::AbortTransaction();
+	});
+
+	REQUIRE(AutoRTFM::ETransactionResult::AbortedByRequest == Result);
+
+	// No abort handlers should execute.
+	REQUIRE(Value == 55);
+}
+
 TEST_CASE("Abort.PushOnAbortHandler_Duplicates2")
 {
 	int Value = 55;
@@ -431,7 +488,7 @@ TEST_CASE("Abort.PushOnAbortHandler_Order")
 			const AutoRTFM::ETransactionResult Result = AutoRTFM::Transact([&]
 				{
 					AutoRTFM::OnAbort([&Value] { REQUIRE(40 == Value); Value += 1; });
-					
+
 					// Make a child transaction.
 					AutoRTFM::Commit([&]
 						{
diff --git a/Engine/Source/Runtime/Core/Private/AutoRTFM/API.cpp b/Engine/Source/Runtime/Core/Private/AutoRTFM/API.cpp
--- a/Engine/Source/Runtime/Core/Private/AutoRTFM/API.cpp
+++ b/Engine/Source/Runtime/Core/Private/AutoRTFM/API.cpp
@@ -336,6 +336,13 @@ UE_AUTORTFM_ALWAYS_OPEN void RTFM_PopOnAbortHandlerInternal(const void* Key)
 	Context->GetCurrentTransaction()->PopDeferUntilAbortHandler(Key);
 }
 
+UE_AUTORTFM_ALWAYS_OPEN void RTFM_PopAllOnAbortHandlersInternal(const void* Key)
+{
+	FContext* Context = FContext::Get();
+	ASSERT(Context->GetStatus() == EContextStatus::OnTrack);
+	Context->GetCurrentTransaction()->PopAllDeferUntilAbortHandlers(Key);
+}
+
 UE_AUTORTFM_ALWAYS_OPEN void RTFM_autortfm_on_commit(void (*Work)(void*), void* Arg)
 {
 	RTFM_OnCommitInternal([Work, Arg] { Work(Arg); });
@@ -589,6 +596,11 @@ extern "C" UE_AUTORTFM_NOAUTORTFM void autortfm_register_open_function(void* Ori
     FunctionMapAdd(OriginalFunction, NewFunction);
 }
 
+extern "C" bool autortfm_is_inner_transaction_stack(void* Ptr)
+{
+	return FContext::Get()->IsInnerTransactionStack(Ptr);
+}
+
 void ForTheRuntime::OnCommitInternal(TFunction<void()> && Work)
 {
 	if (autortfm_is_closed())
@@ -623,6 +635,14 @@ void ForTheRuntime::PopOnAbortHandlerInternal(const void* Key)
 	}
 }
 
+void ForTheRuntime::PopAllOnAbortHandlersInternal(const void* Key)
+{
+	if (autortfm_is_closed())
+	{
+		return RTFM_PopAllOnAbortHandlersInternal(Key);
+	}
+}
+
 extern "C" void autortfm_on_commit(void (*Work)(void*), void* Arg)
 {
 	if (autortfm_is_closed())
diff --git a/Engine/Source/Runtime/Core/Private/AutoRTFM/TaskArray.h b/Engine/Source/Runtime/Core/Private/AutoRTFM/TaskArray.h
--- a/Engine/Source/Runtime/Core/Private/AutoRTFM/TaskArray.h
+++ b/Engine/Source/Runtime/Core/Private/AutoRTFM/TaskArray.h
@@ -71,31 +71,19 @@ public:
 
 		auto EraseLastAddedKeyFromArray = [Key](FInternalArray& Array) -> bool
 		{
-			for(typename FInternalArray::SizeType Idx = 0; Idx < Array.Num(); Idx++)
+			typename FInternalArray::SizeType LastIdx = Array.FindLastByPredicate([Key](const SKeyValuePair& Pair)
 			{
-				typename FInternalArray::SizeType BackwardsIdx = Array.Num() - 1 - Idx;
+				return Pair.Key == Key;
+			});
 
-				if(Array[BackwardsIdx].Key != Key)
-				{
-					continue;
-				}
-
-				// We found our key to erase! Nuke it.
-
-				// We start at the element just after our element.
-				BackwardsIdx++;
-
-				for (; BackwardsIdx < Array.Num(); BackwardsIdx++)
-				{
-					Array[BackwardsIdx - 1] = MoveTemp(Array[BackwardsIdx]);
-				}
-
-				Array.SetNum(Array.Num() - 1, EAllowShrinking::No);
-
-				return true;
+			if (LastIdx == INDEX_NONE)
+			{
+				return false;
 			}
 
-			return false;
+			// We found our key to erase! Nuke it.
+			Array.RemoveAt(LastIdx);
+			return true;
 		};
 
 		if (EraseLastAddedKeyFromArray(Latest))
@@ -114,6 +102,28 @@ public:
 		return false;
 	}
 
+	bool DeleteAllMatchingKeys(SKey Key)
+	{
+		ASSERT(Key != SKey{});
+
+		auto EraseKeyFromArray = [Key](FInternalArray& Array) -> size_t
+		{
+			return Array.RemoveAll([Key](const SKeyValuePair& Pair)
+			{
+				return Pair.Key == Key;
+			});
+		};
+
+		size_t NumErased = EraseKeyFromArray(Latest);
+
+		for (FInternalArray& StashedVectorBox : Stash)
+		{
+			NumErased += EraseKeyFromArray(StashedVectorBox);
+		}
+
+		return NumErased > 0;
+	}
+
     void AddAll(TTaskArray&& Other)
     {
         Canonicalize();
diff --git a/Engine/Source/Runtime/Core/Private/AutoRTFM/Transaction.h b/Engine/Source/Runtime/Core/Private/AutoRTFM/Transaction.h
--- a/Engine/Source/Runtime/Core/Private/AutoRTFM/Transaction.h
+++ b/Engine/Source/Runtime/Core/Private/AutoRTFM/Transaction.h
@@ -64,7 +64,8 @@ public:
     void DeferUntilCommit(TFunction<void()>&&);
     void DeferUntilAbort(TFunction<void()>&&);
     void PushDeferUntilAbortHandler(const void* Key, TFunction<void()>&&);
-    bool PopDeferUntilAbortHandler(const void* Key);
+	bool PopDeferUntilAbortHandler(const void* Key);
+	bool PopAllDeferUntilAbortHandlers(const void* Key);
 
     [[noreturn]] void AbortAndThrow();
     void AbortWithoutThrowing();
diff --git a/Engine/Source/Runtime/Core/Private/AutoRTFM/TransactionInlines.h b/Engine/Source/Runtime/Core/Private/AutoRTFM/TransactionInlines.h
--- a/Engine/Source/Runtime/Core/Private/AutoRTFM/TransactionInlines.h
+++ b/Engine/Source/Runtime/Core/Private/AutoRTFM/TransactionInlines.h
@@ -160,7 +160,12 @@ UE_AUTORTFM_FORCEINLINE void FTransaction::PushDeferUntilAbortHandler(const void
 
 UE_AUTORTFM_FORCEINLINE bool FTransaction::PopDeferUntilAbortHandler(const void* Key)
 {
-    return AbortTasks.DeleteKey(Key);
+	return AbortTasks.DeleteKey(Key);
+}
+
+UE_AUTORTFM_FORCEINLINE bool FTransaction::PopAllDeferUntilAbortHandlers(const void* Key)
+{
+	return AbortTasks.DeleteAllMatchingKeys(Key);
 }
 
 UE_AUTORTFM_FORCEINLINE void FTransaction::CollectStats() const
diff --git a/Engine/Source/Runtime/Core/Public/AutoRTFM/AutoRTFM.h b/Engine/Source/Runtime/Core/Public/AutoRTFM/AutoRTFM.h
--- a/Engine/Source/Runtime/Core/Public/AutoRTFM/AutoRTFM.h
+++ b/Engine/Source/Runtime/Core/Public/AutoRTFM/AutoRTFM.h
@@ -279,6 +279,15 @@ static UE_AUTORTFM_FORCEINLINE void autortfm_register_open_function(void* origin
 }
 #endif
 
+#if UE_AUTORTFM_ENABLED
+UE_AUTORTFM_API bool autortfm_is_inner_transaction_stack(void* Ptr);
+#else
+static UE_AUTORTFM_FORCEINLINE bool autortfm_is_inner_transaction_stack(void* Ptr)
+{
+	return false;
+}
+#endif
+
 #if UE_AUTORTFM_ENABLED
 UE_AUTORTFM_API void autortfm_on_commit(void (*work)(void* arg), void* arg);
 #else
@@ -423,6 +432,7 @@ namespace ForTheRuntime
 	UE_AUTORTFM_API void OnAbortInternal(TFunction<void()>&& Work);
 	UE_AUTORTFM_API void PushOnAbortHandlerInternal(const void* Key, TFunction<void()>&& Work);
 	UE_AUTORTFM_API void PopOnAbortHandlerInternal(const void* Key);
+	UE_AUTORTFM_API void PopAllOnAbortHandlersInternal(const void* Key);
 } // namespace ForTheRuntime
 #endif
 
@@ -465,6 +475,14 @@ static UE_AUTORTFM_FORCEINLINE bool IsClosed() { return autortfm_is_closed(); }
 // in an on-abort or on-commit.
 static UE_AUTORTFM_FORCEINLINE bool IsCommittingOrAborting() { return autortfm_is_committing_or_aborting(); }
 
+// Returns true if the passed-in pointer is on the stack of the currently-executing transaction.
+// This is occasionally necessary when writing OnAbort handlers for objects on the stack, since 
+// we don't want to scribble on stack memory that might have been reused.
+static UE_AUTORTFM_FORCEINLINE bool IsInnerTransactionStack(void* Ptr) 
+{ 
+	return autortfm_is_inner_transaction_stack(Ptr);
+}
+
 // Run the functor in a transaction. Memory writes and other side effects get instrumented
 // and will be reversed if the transaction aborts.
 // 
@@ -652,8 +670,7 @@ template<typename TFunctor> static UE_AUTORTFM_FORCEINLINE void PushOnAbortHandl
 	}
 }
 
-// Unregister all handlers for transaction abort that were previously pushed
-// via PushOnAbortHandler with the given key
+// Unregister the most recently pushed handler (via PushOnAbortHandler) for the given key.
 static UE_AUTORTFM_FORCEINLINE void PopOnAbortHandler(const void* Key)
 {
 	if (autortfm_is_closed())
@@ -661,6 +678,15 @@ static UE_AUTORTFM_FORCEINLINE void PopOnAbortHandler(const void* Key)
 		ForTheRuntime::PopOnAbortHandlerInternal(Key);
 	}
 }
+
+// Unregister all pushed handlers (via PushOnAbortHandler) for the given key.
+static UE_AUTORTFM_FORCEINLINE void PopAllOnAbortHandlers(const void* Key)
+{
+	if (autortfm_is_closed())
+	{
+		ForTheRuntime::PopAllOnAbortHandlersInternal(Key);
+	}
+}
 #else
 // Have some work happen when this transaction aborts. If this is called
 // outside a transaction or from an open nest then the work is ignored.
@@ -672,9 +698,12 @@ template<typename TFunctor> static UE_AUTORTFM_FORCEINLINE void OnAbort(const TF
 // reaches the end of the relevant scope.
 template<typename TFunctor> static UE_AUTORTFM_FORCEINLINE void PushOnAbortHandler(const void* Key, const TFunctor&) {}
 
-// Unregister all handlers for transaction abort that were previously pushed
-// via PushOnAbortHandler with the given key
+// Unregister the most recently pushed handler (via PushOnAbortHandler) for the given key.
 static UE_AUTORTFM_FORCEINLINE void PopOnAbortHandler(const void* Key) {}
+
+
+// Unregister all pushed handlers (via PushOnAbortHandler) for the given key.
+static UE_AUTORTFM_FORCEINLINE void PopAllOnAbortHandlers(const void* Key) {}
 #endif
 
 // Inform the runtime that we have performed a new object allocation. It's only
