Csv profiler: add a cvar csv.AggregateTaskWorkerStats, defaulted to on. When enabled, stats recorded on task worker threads are aggregated (summed) instead of outputting a single stat per thread. This reduces CSV bloat when there are large numbers of worker threads and makes stat data more intelligible. The stats appear with an AllWorkers prefix instead of the individual thread name.
#tests tested on console

diff --git a/Engine/Source/Runtime/Core/Private/ProfilingDebugging/CsvProfiler.cpp b/Engine/Source/Runtime/Core/Private/ProfilingDebugging/CsvProfiler.cpp
--- a/Engine/Source/Runtime/Core/Private/ProfilingDebugging/CsvProfiler.cpp
+++ b/Engine/Source/Runtime/Core/Private/ProfilingDebugging/CsvProfiler.cpp
@@ -7,6 +7,7 @@
 */
 
 #include "ProfilingDebugging/CsvProfiler.h"
+#include "Async/Fundamental/Scheduler.h"
 #include "Async/TaskGraphInterfaces.h"
 #include "CoreGlobals.h"
 #include "HAL/RunnableThread.h"
@@ -84,6 +85,14 @@ TAutoConsoleVariable<int32> CVarCsvBlockOnCaptureEnd(
 	ECVF_Default
 );
 
+TAutoConsoleVariable<bool> CVarCsvAggregateTaskWorkerStats(
+	TEXT("csv.AggregateTaskWorkerStats"),
+	true,
+	TEXT("If enabled, stats recorded on task worker threads are aggregated instead of outputting a single stat per thread.\r\n")
+	TEXT("This reduces CSV bloat when there are large numbers of worker threads and makes stat data more intelligible"),
+	ECVF_Default
+);
+
 TAutoConsoleVariable<int32> CVarCsvContinuousWrites(
 	TEXT("csv.ContinuousWrites"),
 	1,
@@ -1846,6 +1855,8 @@ struct FCsvProcessThreadDataStats
 	uint32 EventCount;
 };
 
+class FCsvThreadGroupStatProcessor;
+
 class FCsvStreamWriter
 {
 	struct FCsvRow
@@ -1882,11 +1893,13 @@ class FCsvStreamWriter
 	TArray<FCsvStatSeries*> AllSeries;
 	TArray<class FCsvProfilerThreadDataProcessor*> DataProcessors;
 
+	TSharedPtr<FCsvThreadGroupStatProcessor> TaskWorkerThreadGroupStatProcessor;
+
 	uint32 RenderThreadId;
 	uint32 RHIThreadId;
 
 public:
-	FCsvStreamWriter(const TSharedRef<FArchive>& InOutputFile, bool bInContinuousWrites, int32 InBufferSize, int64 InNumFramesToBuffer, bool bInCompressOutput, uint32 RenderThreadId, uint32 RHIThreadId);
+	FCsvStreamWriter(const TSharedRef<FArchive>& InOutputFile, bool bInContinuousWrites, int32 InBufferSize, int64 InNumFramesToBuffer, bool bInCompressOutput, uint32 RenderThreadId, uint32 RHIThreadId, bool bAggregateTaskWorkerStats);
 	~FCsvStreamWriter();
 
 	void AddSeries(FCsvStatSeries* Series);
@@ -2038,6 +2051,7 @@ public:
 	FCsvProfilerThreadData()
 		: ThreadId(FPlatformTLS::GetCurrentThreadId())
 		, ThreadName(FThreadManager::GetThreadName(ThreadId))
+		, bIsTaskWorkerThread(LowLevelTasks::FScheduler::Get().IsWorkerThread())
 		, DataProcessor(nullptr)
 	{
 	}
@@ -2224,6 +2238,7 @@ public:
 
 	const uint32 ThreadId;
 	const FString ThreadName;
+	const bool bIsTaskWorkerThread;
 
 	class FCsvProfilerThreadDataProcessor* DataProcessor;
 	TArray<const char*> WaitStatNameStack;
@@ -2233,6 +2248,64 @@ uint32 FCsvProfilerThreadData::TlsSlot = FPlatformTLS::InvalidTlsSlot;
 FCriticalSection FCsvProfilerThreadData::TlsCS;
 TArray<FCsvProfilerThreadData::FWeakPtr> FCsvProfilerThreadData::TlsInstances;
 
+// Responsible to collating and accumulating stats for a thread, or group of threads
+class FCsvThreadGroupStatProcessor
+{
+	FCsvStreamWriter* Writer;
+	TArray<FCsvStatSeries*> StatSeriesArray;
+	FCsvStatRegister StatRegister;
+	FString Name;
+
+public:
+	FCsvThreadGroupStatProcessor(FCsvStreamWriter* InWriter, FString InName)
+		: Writer(InWriter)
+		, Name(InName)
+	{
+	}
+
+	~FCsvThreadGroupStatProcessor()
+	{
+		// Delete all the created stat series
+		for (FCsvStatSeries* Series : StatSeriesArray)
+		{
+			delete Series;
+		}
+	}
+
+	FCsvStatSeries* FindOrCreateStatSeries(const FCsvStatBase& Stat, FCsvStatSeries::EType SeriesType, bool bIsCountStat)
+	{
+		check(IsInCsvProcessingThread());
+		const int32 StatIndex = StatRegister.GetUniqueIndex(Stat.RawStatID, Stat.CategoryIndex, Stat.IsFNameStat(), bIsCountStat);
+		FCsvStatSeries* Series = nullptr;
+		if (StatSeriesArray.Num() <= StatIndex)
+		{
+			int32 GrowBy = StatIndex + 1 - StatSeriesArray.Num();
+			StatSeriesArray.AddZeroed(GrowBy);
+		}
+		if (StatSeriesArray[StatIndex] == nullptr)
+		{
+			Series = new FCsvStatSeries(SeriesType, StatIndex, Writer, StatRegister, Name);
+			StatSeriesArray[StatIndex] = Series;
+		}
+		else
+		{
+			Series = StatSeriesArray[StatIndex];
+#if DO_CHECK
+			checkf(SeriesType == Series->SeriesType, TEXT("Stat named %s was used in multiple stat types. Can't use same identifier for different stat types. Stat types are: Custom(Int), Custom(Float) and Timing"), *StatRegister.GetStatName(StatIndex));
+#endif
+		}
+		return Series;
+	}
+
+	inline uint64 GetAllocatedSize() const
+	{
+		return	((uint64)StatSeriesArray.GetAllocatedSize()) +
+				((uint64)StatSeriesArray.Num() * sizeof(FCsvStatSeries));
+	}
+
+};
+
+
 class FCsvProfilerThreadDataProcessor
 {
 	FCsvProfilerThreadData::FSharedPtr ThreadData;
@@ -2241,14 +2314,13 @@ class FCsvProfilerThreadDataProcessor
 	TArray<FCsvTimingMarker> MarkerStack;
 	TArray<FCsvTimingMarker> ExclusiveMarkerStack;
 
-	TArray<FCsvStatSeries*> StatSeriesArray;
-	FCsvStatRegister StatRegister;
-
 	uint64 LastProcessedTimestamp;
 
 	uint32 RenderThreadId;
 	uint32 RHIThreadId;
 
+	TSharedPtr<FCsvThreadGroupStatProcessor> StatProcessor;
+
 public:
 	FCsvProfilerThreadDataProcessor(FCsvProfilerThreadData::FSharedPtr InThreadData, FCsvStreamWriter* InWriter, uint32 InRenderThreadId, uint32 InRHIThreadId)
 		: ThreadData(InThreadData)
@@ -2256,6 +2328,7 @@ public:
 		, LastProcessedTimestamp(0)
 		, RenderThreadId(InRenderThreadId)
 		, RHIThreadId(InRHIThreadId)
+		, StatProcessor(new FCsvThreadGroupStatProcessor(InWriter, InThreadData->ThreadName))
 	{
 		check(ThreadData->DataProcessor == nullptr);
 		ThreadData->DataProcessor = this;
@@ -2265,12 +2338,6 @@ public:
 	{
 		check(ThreadData->DataProcessor == this);
 		ThreadData->DataProcessor = nullptr;
-
-		// Delete all the created stat series
-		for (FCsvStatSeries* Series : StatSeriesArray)
-		{
-			delete Series;
-		}
 	}
 
 	inline uint64 GetAllocatedSize() const
@@ -2278,55 +2345,32 @@ public:
 		return
 			((uint64)MarkerStack.GetAllocatedSize()) +
 			((uint64)ExclusiveMarkerStack.GetAllocatedSize()) +
-			((uint64)StatSeriesArray.GetAllocatedSize()) +
-			((uint64)StatSeriesArray.Num() * sizeof(FCsvStatSeries)) +
+			StatProcessor->GetAllocatedSize() +
 			((uint64)ThreadData->GetAllocatedSize());
 	}
 
-	void Process(FCsvProcessThreadDataStats& OutStats, int32& OutMinFrameNumberProcessed);
+	void Process(FCsvProcessThreadDataStats& OutStats, int32& OutMinFrameNumberProcessed, TSharedPtr<FCsvThreadGroupStatProcessor> TaskWorkerThreadGroupStatProcessor);
 
 private:
 	/** Temporary storage of data collected with every Process() call. */
 	TArray<FCsvTimingMarker> ThreadMarkers;
 	TArray<FCsvCustomStat> CustomStats;
 	TArray<FCsvEvent> Events;
-
-	FCsvStatSeries* FindOrCreateStatSeries(const FCsvStatBase& Stat, FCsvStatSeries::EType SeriesType, bool bIsCountStat)
-	{
-		check(IsInCsvProcessingThread());
-		const int32 StatIndex = StatRegister.GetUniqueIndex(Stat.RawStatID, Stat.CategoryIndex, Stat.IsFNameStat(), bIsCountStat);
-		FCsvStatSeries* Series = nullptr;
-		if (StatSeriesArray.Num() <= StatIndex)
-		{
-			int32 GrowBy = StatIndex + 1 - StatSeriesArray.Num();
-			StatSeriesArray.AddZeroed(GrowBy);
-		}
-		if (StatSeriesArray[StatIndex] == nullptr)
-		{
-			Series = new FCsvStatSeries(SeriesType, StatIndex, Writer, StatRegister, ThreadData->ThreadName);
-			StatSeriesArray[StatIndex] = Series;
-		}
-		else
-		{
-			Series = StatSeriesArray[StatIndex];
-#if DO_CHECK
-			checkf(SeriesType == Series->SeriesType, TEXT("Stat named %s was used in multiple stat types. Can't use same identifier for different stat types. Stat types are: Custom(Int), Custom(Float) and Timing"), *StatRegister.GetStatName(StatIndex));
-#endif
-		}
-		return Series;
-	}
 };
 
-FCsvStreamWriter::FCsvStreamWriter(const TSharedRef<FArchive>& InOutputFile, bool bInContinuousWrites, int32 InBufferSize, int64 InNumFramesToBuffer, bool bInCompressOutput, uint32 InRenderThreadId, uint32 InRHIThreadId)
+
+FCsvStreamWriter::FCsvStreamWriter(const TSharedRef<FArchive>& InOutputFile, bool bInContinuousWrites, int32 InBufferSize, int64 InNumFramesToBuffer, bool bInCompressOutput, uint32 InRenderThreadId, uint32 InRHIThreadId, bool bAggregateTaskWorkerStats)
 	: Stream(InOutputFile, InBufferSize, bInCompressOutput)
 	, NumFramesToBuffer(InNumFramesToBuffer)
 	, WriteFrameIndex(-1)
 	, ReadFrameIndex(-1)
 	, bContinuousWrites(bInContinuousWrites)
 	, bFirstRow(true)
+	, TaskWorkerThreadGroupStatProcessor(bAggregateTaskWorkerStats ? new FCsvThreadGroupStatProcessor(this, TEXT("AllWorkers")) : nullptr)
 	, RenderThreadId(InRenderThreadId)
 	, RHIThreadId(InRHIThreadId)
-{}
+{
+}
 
 FCsvStreamWriter::~FCsvStreamWriter()
 {
@@ -2492,14 +2536,13 @@ void FCsvStreamWriter::Process(FCsvProcessThreadDataStats& OutStats)
 			}
 		}
 	}
-	
 
 	int32 MinFrameNumberProcessed = MAX_int32;
 	{
 		QUICK_SCOPE_CYCLE_COUNTER(CSVProfiler_Writer_ProcessDataProcessors);
 		for (FCsvProfilerThreadDataProcessor* DataProcessor : DataProcessors)
 		{
-			DataProcessor->Process(OutStats, MinFrameNumberProcessed);
+			DataProcessor->Process(OutStats, MinFrameNumberProcessed, TaskWorkerThreadGroupStatProcessor);
 		}
 	}
 	
@@ -2523,6 +2566,11 @@ uint64 FCsvStreamWriter::GetAllocatedSize() const
 		((uint64)DataProcessors.GetAllocatedSize()) +
 		((uint64)Stream.GetAllocatedSize());
 
+	if ( TaskWorkerThreadGroupStatProcessor.IsValid() )
+	{
+		Size += TaskWorkerThreadGroupStatProcessor->GetAllocatedSize();
+	}
+
 	for (const auto& Pair          : Rows)           { Size += (uint64)Pair.Value.GetAllocatedSize();     }
 	for (const auto& Series        : AllSeries)      { Size += (uint64)Series->GetAllocatedSize();        }
 	for (const auto& DataProcessor : DataProcessors) { Size += (uint64)DataProcessor->GetAllocatedSize(); }
@@ -2622,7 +2670,7 @@ private:
 	FCsvProfiler& CsvProfiler;
 };
 
-void FCsvProfilerThreadDataProcessor::Process(FCsvProcessThreadDataStats& OutStats, int32& OutMinFrameNumberProcessed)
+void FCsvProfilerThreadDataProcessor::Process(FCsvProcessThreadDataStats& OutStats, int32& OutMinFrameNumberProcessed, TSharedPtr<FCsvThreadGroupStatProcessor> TaskWorkerThreadGroupStatProcessor)
 {
 	QUICK_SCOPE_CYCLE_COUNTER(STAT_FCsvProfilerThreadData_ProcessThreadData);
 
@@ -2635,6 +2683,13 @@ void FCsvProfilerThreadDataProcessor::Process(FCsvProcessThreadDataStats& OutSta
 	Events.Reset(0);
 	ThreadData->FlushResults(ThreadMarkers, CustomStats, Events);
 
+	// If we're aggregating task worker thread stats and this is a task worker thread, use the shared TaskWorkerThreadGroupStatProcessor instead of the per-thread one
+	FCsvThreadGroupStatProcessor* ThreadStatProcessor = StatProcessor.Get();
+	if (ThreadData->bIsTaskWorkerThread && TaskWorkerThreadGroupStatProcessor.IsValid())
+	{
+		ThreadStatProcessor = TaskWorkerThreadGroupStatProcessor.Get();
+	}
+
 	OutStats.TimestampCount += ThreadMarkers.Num();
 	OutStats.CustomStatCount += CustomStats.Num();
 	OutStats.EventCount += Events.Num();
@@ -2755,13 +2810,13 @@ void FCsvProfilerThreadDataProcessor::Process(FCsvProcessThreadDataStats& OutSta
 							uint64 ElapsedCycles = Marker.GetTimestamp() - StartMarker.GetTimestamp();
 
 							// Add the elapsed time to the table entry for this frame/stat
-							FCsvStatSeries* Series = FindOrCreateStatSeries(Marker, FCsvStatSeries::EType::TimerData, false);
+							FCsvStatSeries* Series = ThreadStatProcessor->FindOrCreateStatSeries(Marker, FCsvStatSeries::EType::TimerData, false);
 							Series->SetTimerValue(FrameNumber, ElapsedCycles);
 
 							// Add the COUNT/ series if enabled. Ignore artificial markers (inserted above)
 							if (GCsvStatCounts && !Marker.IsExclusiveArtificialMarker())
 							{
-								FCsvStatSeries* CountSeries = FindOrCreateStatSeries(Marker, FCsvStatSeries::EType::CustomStatInt, true);
+								FCsvStatSeries* CountSeries = ThreadStatProcessor->FindOrCreateStatSeries(Marker, FCsvStatSeries::EType::CustomStatInt, true);
 								CountSeries->SetCustomStatValue_Int(FrameNumber, ECsvCustomStatOp::Accumulate, 1);
 							}
 						}
@@ -2782,7 +2837,7 @@ void FCsvProfilerThreadDataProcessor::Process(FCsvProcessThreadDataStats& OutSta
 			if (FrameNumber >= 0)
 			{
 				bool bIsInteger = CustomStat.IsInteger();
-				FCsvStatSeries* Series = FindOrCreateStatSeries(CustomStat, bIsInteger ? FCsvStatSeries::EType::CustomStatInt : FCsvStatSeries::EType::CustomStatFloat, false);
+				FCsvStatSeries* Series = ThreadStatProcessor->FindOrCreateStatSeries(CustomStat, bIsInteger ? FCsvStatSeries::EType::CustomStatInt : FCsvStatSeries::EType::CustomStatFloat, false);
 				if (bIsInteger)
 				{
 					Series->SetCustomStatValue_Int(FrameNumber, CustomStat.GetCustomStatOp(), CustomStat.Value.AsInt);
@@ -2795,7 +2850,7 @@ void FCsvProfilerThreadDataProcessor::Process(FCsvProcessThreadDataStats& OutSta
 				// Add the COUNT/ series if enabled
 				if (GCsvStatCounts)
 				{
-					FCsvStatSeries* CountSeries = FindOrCreateStatSeries(CustomStat, FCsvStatSeries::EType::CustomStatInt, true);
+					FCsvStatSeries* CountSeries = ThreadStatProcessor->FindOrCreateStatSeries(CustomStat, FCsvStatSeries::EType::CustomStatInt, true);
 					CountSeries->SetCustomStatValue_Int(FrameNumber, ECsvCustomStatOp::Accumulate, 1);
 				}
 			}
@@ -3056,7 +3111,7 @@ void FCsvProfiler::BeginCaptureInternal(const FCsvCaptureCommand& CurrentCommand
 
 	// Actually start the capture
 	int64 NumFramesToBuffer = CVarCsvStreamFramesToBuffer.GetValueOnAnyThread();
-	CsvWriter = new FCsvStreamWriter(OutputFile.ToSharedRef(), bContinuousWrites, BufferSize, NumFramesToBuffer, bCompressOutput, RenderThreadId, RHIThreadId);
+	CsvWriter = new FCsvStreamWriter(OutputFile.ToSharedRef(), bContinuousWrites, BufferSize, NumFramesToBuffer, bCompressOutput, RenderThreadId, RHIThreadId, CVarCsvAggregateTaskWorkerStats.GetValueOnAnyThread());
 
 	NumFramesToCapture = CurrentCommand.Value;
 	GCsvRepeatFrameCount = NumFramesToCapture;
@@ -4235,4 +4290,4 @@ void CSVTest()
 
 #endif // CSV_PROFILER_ALLOW_DEBUG_FEATURES
 
-#endif // CSV_PROFILER
\ No newline at end of file
+#endif // CSV_PROFILER
