Make FConfigFile use a transactionally safe lock.

diff --git a/Engine/Source/Programs/AutoRTFMTests/Private/UECoreTests.cpp b/Engine/Source/Programs/AutoRTFMTests/Private/UECoreTests.cpp
--- a/Engine/Source/Programs/AutoRTFMTests/Private/UECoreTests.cpp
+++ b/Engine/Source/Programs/AutoRTFMTests/Private/UECoreTests.cpp
@@ -14,6 +14,7 @@
 #include "Misc/TransactionallySafeScopeLock.h"
 #include "Misc/TransactionallySafeRWScopeLock.h"
 #include "Containers/Queue.h"
+#include "Misc/ConfigCacheIni.h"
 
 TEST_CASE("UECore.FDelegateHandle")
 {
@@ -1735,3 +1736,22 @@ TEST_CASE("UECore.TQueue")
 		}
 	}
 }
+
+TEST_CASE("UECore.FConfigFile")
+{
+	SECTION("Empty")
+	{
+		FConfigFile Config;
+
+		Config.FindOrAddConfigSection(TEXT("WOW"));
+
+		REQUIRE(!Config.IsEmpty());
+	
+		AutoRTFM::Commit([&]
+			{
+				Config.Empty();
+			});
+		
+		REQUIRE(Config.IsEmpty());
+	}
+}
diff --git a/Engine/Source/Runtime/Core/Private/Misc/ConfigCacheIni.cpp b/Engine/Source/Runtime/Core/Private/Misc/ConfigCacheIni.cpp
--- a/Engine/Source/Runtime/Core/Private/Misc/ConfigCacheIni.cpp
+++ b/Engine/Source/Runtime/Core/Private/Misc/ConfigCacheIni.cpp
@@ -31,6 +31,7 @@
 #include "Async/Async.h"
 #include "Misc/OutputDeviceRedirector.h"
 #include "Logging/MessageLog.h"
+#include "Misc/TransactionallySafeRWLock.h"
 #include <limits>
 
 namespace
@@ -944,7 +945,7 @@ FConfigFile::FConfigFile(FConfigFile&& Other)
 
 FConfigFile& FConfigFile::operator=(const FConfigFile& Other)
 {
-	FWriteScopeLock ScopeLock(ConfigFileMapLock);
+	FTransactionallySafeWriteScopeLock ScopeLock(ConfigFileMapLock);
 	this->FConfigFileMap::operator=(Other);
 	Dirty = Other.Dirty;
 	NoSave = Other.NoSave;
@@ -991,7 +992,7 @@ FConfigFile& FConfigFile::operator=(const FConfigFile& Other)
 
 FConfigFile& FConfigFile::operator=(FConfigFile&& Other)
 {
-	FWriteScopeLock ScopeLock(ConfigFileMapLock);
+	FTransactionallySafeWriteScopeLock ScopeLock(ConfigFileMapLock);
 	this->FConfigFileMap::operator=(MoveTemp(Other));
 	Dirty = Other.Dirty;
 	NoSave = Other.NoSave;
@@ -1067,7 +1068,7 @@ UE::ConfigAccessTracking::FFile* FConfigFile::GetFileAccess() const
 
 bool FConfigFile::operator==( const FConfigFile& Other ) const
 {
-	FReadScopeLock ScopeLock(ConfigFileMapLock);
+	FTransactionallySafeReadScopeLock ScopeLock(ConfigFileMapLock);
 
 	if ( Pairs.Num() != Other.Pairs.Num() )
 		return 0;
@@ -1129,7 +1130,7 @@ void FConfigFile::Shrink()
 	if (IsInGameThread()) GConfigShrinkTime -= FPlatformTime::Seconds();
 #endif
 
-	FWriteScopeLock ScopeLock(ConfigFileMapLock);
+	FTransactionallySafeWriteScopeLock ScopeLock(ConfigFileMapLock);
 	FConfigFileMap::Shrink();
 
 	for (FConfigFileMap::TIterator It(*this); It; ++It)
@@ -6248,7 +6249,7 @@ FCriticalSection FConfigCacheIni::ConfigForPlatformLock;
 TMap<FName, FConfigCacheIni::FPluginInfo*> FConfigCacheIni::RegisteredPlugins;
 FTransactionallySafeCriticalSection FConfigCacheIni::RegisteredPluginsLock;
 
-FRWLock FConfigFile::ConfigFileMapLock;
+FTransactionallySafeRWLock FConfigFile::ConfigFileMapLock;
 
 void FConfigCacheIni::AddPluginToAllBranches(FName PluginName, FConfigModificationTracker* ModificationTracker)
 {
diff --git a/Engine/Source/Runtime/Core/Public/Misc/ConfigCacheIni.h b/Engine/Source/Runtime/Core/Public/Misc/ConfigCacheIni.h
--- a/Engine/Source/Runtime/Core/Public/Misc/ConfigCacheIni.h
+++ b/Engine/Source/Runtime/Core/Public/Misc/ConfigCacheIni.h
@@ -31,7 +31,7 @@
 #include "Misc/ConfigAccessTracking.h"
 #include "Misc/ConfigTypes.h"
 #include "Misc/Paths.h"
-#include "Misc/ScopeRWLock.h"
+#include "Misc/TransactionallySafeRWScopeLock.h"
 #include "Serialization/Archive.h"
 #include "Serialization/StructuredArchive.h"
 #include "Serialization/StructuredArchiveAdapters.h"
@@ -610,131 +610,131 @@ public:
 	UE_DEPRECATED(5.4, "Use FindSection, and/or use the new AddToSection, etc APIs to modify sections without retrieving the section. See top of ConfigCacheIni.h for more info.")
 	FORCEINLINE const FConfigSection* Find(const FString& SectionName) const
 	{
-		FReadScopeLock ScopeLock(ConfigFileMapLock);
+		FTransactionallySafeReadScopeLock ScopeLock(ConfigFileMapLock);
 		return FConfigFileMap::Find(SectionName);
 	}
 
 	UE_DEPRECATED(5.4, "Use FindSection, and/or use the new AddToSection, etc APIs to modify sections without retrieving the section. See top of ConfigCacheIni.h for more info.")
 	FORCEINLINE FConfigSection* Find(const FString& SectionName)
 	{
-		FReadScopeLock ScopeLock(ConfigFileMapLock);
+		FTransactionallySafeReadScopeLock ScopeLock(ConfigFileMapLock);
 		return FConfigFileMap::Find(SectionName);
 	}
 
 	FORCEINLINE const FConfigSection* FindSection(const FString& SectionName) const
 	{
-		FReadScopeLock ScopeLock(ConfigFileMapLock);
+		FTransactionallySafeReadScopeLock ScopeLock(ConfigFileMapLock);
 		return FConfigFileMap::Find(SectionName);
 	}
 
 	FORCEINLINE int32 Num() const
 	{
-		FReadScopeLock ScopeLock(ConfigFileMapLock);
+		FTransactionallySafeReadScopeLock ScopeLock(ConfigFileMapLock);
 		return FConfigFileMap::Num();
 	}
 
 	FORCEINLINE bool IsEmpty() const
 	{
-		FReadScopeLock ScopeLock(ConfigFileMapLock);
+		FTransactionallySafeReadScopeLock ScopeLock(ConfigFileMapLock);
 		return FConfigFileMap::IsEmpty();
 	}
 
 	FORCEINLINE void Empty(int32 ExpectedNumElements = 0)
 	{
-		FWriteScopeLock ScopeLock(ConfigFileMapLock);
+		FTransactionallySafeWriteScopeLock ScopeLock(ConfigFileMapLock);
 		FConfigFileMap::Empty(ExpectedNumElements);
 	}
 
 	FORCEINLINE bool Contains(const FString& SectionName) const
 	{
-		FReadScopeLock ScopeLock(ConfigFileMapLock);
+		FTransactionallySafeReadScopeLock ScopeLock(ConfigFileMapLock);
 		return FConfigFileMap::Contains(SectionName);
 	}
 
 	FORCEINLINE int32 GetKeys(TArray<FString>& Keys) const
 	{
-		FReadScopeLock ScopeLock(ConfigFileMapLock);
+		FTransactionallySafeReadScopeLock ScopeLock(ConfigFileMapLock);
 		return FConfigFileMap::GetKeys(Keys);
 	}
 
 	FORCEINLINE int32 GetKeys(TSet<FString>& Keys) const
 	{
-		FReadScopeLock ScopeLock(ConfigFileMapLock);
+		FTransactionallySafeReadScopeLock ScopeLock(ConfigFileMapLock);
 		return FConfigFileMap::GetKeys(Keys);
 	}
 
 	FORCEINLINE int32 Remove(KeyConstPointerType InKey)
 	{
-		FWriteScopeLock ScopeLock(ConfigFileMapLock);
+		FTransactionallySafeWriteScopeLock ScopeLock(ConfigFileMapLock);
 		return FConfigFileMap::Remove(InKey);
 	}
 
 	FORCEINLINE ValueType& Add(const KeyType& InKey, const ValueType& InValue)
 	{
-		FWriteScopeLock ScopeLock(ConfigFileMapLock);
+		FTransactionallySafeWriteScopeLock ScopeLock(ConfigFileMapLock);
 		return FConfigFileMap::Add(InKey, InValue);
 	}
 
 	FORCEINLINE ValueType& Add(const KeyType&  InKey, ValueType&& InValue)
 	{
-		FWriteScopeLock ScopeLock(ConfigFileMapLock);
+		FTransactionallySafeWriteScopeLock ScopeLock(ConfigFileMapLock);
 		return FConfigFileMap::Add(InKey, MoveTempIfPossible(InValue));
 	}
 
 	FORCEINLINE ValueType& Add(KeyType&& InKey, const ValueType&  InValue)
 	{
-		FWriteScopeLock ScopeLock(ConfigFileMapLock);
+		FTransactionallySafeWriteScopeLock ScopeLock(ConfigFileMapLock);
 		return FConfigFileMap::Add(MoveTempIfPossible(InKey), InValue);
 	}
 
 	FORCEINLINE ValueType& Add(KeyType&& InKey,	ValueType&& InValue)
 	{
-		FWriteScopeLock ScopeLock(ConfigFileMapLock);
+		FTransactionallySafeWriteScopeLock ScopeLock(ConfigFileMapLock);
 		return FConfigFileMap::Add(MoveTempIfPossible(InKey), MoveTempIfPossible(InValue));
 	}
 	
 	FORCEINLINE void Append(TMap<FString, FConfigSection> Other)
 	{
-		FWriteScopeLock ScopeLock(ConfigFileMapLock);
+		FTransactionallySafeWriteScopeLock ScopeLock(ConfigFileMapLock);
 		FConfigFileMap::Append(MoveTemp(Other));
 	}
 
 	FORCEINLINE void Reset()
 	{
-		FWriteScopeLock ScopeLock(ConfigFileMapLock);
+		FTransactionallySafeWriteScopeLock ScopeLock(ConfigFileMapLock);
 		FConfigFileMap::Reset();
 	}
 
 	UE_DEPRECATED(5.4, "Use FindOrAddConfigSection, and/or use the new AddToSection, etc APIs to modify sections without retrieving the section. See top of ConfigCacheIni.h for more info.")
 	FORCEINLINE ValueType& FindOrAdd(const FString& Key)
 	{
-		FWriteScopeLock ScopeLock(ConfigFileMapLock);
+		FTransactionallySafeWriteScopeLock ScopeLock(ConfigFileMapLock);
 		return FConfigFileMap::FindOrAdd(Key);
 	}
 
 	UE_DEPRECATED(5.4, "Use const ranged for iterators, (wrap your FConfigFile variable in AsConst to force the const iterator). See top of ConfigCacheIni.h for more info.")
 	FORCEINLINE TRangedForIterator begin()
 	{
-		FReadScopeLock ScopeLock(ConfigFileMapLock);
+		FTransactionallySafeReadScopeLock ScopeLock(ConfigFileMapLock);
 		return TRangedForIterator(Pairs.begin());
 	}
 
 	FORCEINLINE TRangedForConstIterator begin() const
 	{
-		FReadScopeLock ScopeLock(ConfigFileMapLock); 
+		FTransactionallySafeReadScopeLock ScopeLock(ConfigFileMapLock); 
 		return TRangedForConstIterator(Pairs.begin());
 	}
 
 	UE_DEPRECATED(5.4, "Use const ranged for iterators, (wrap your FConfigFile variable in AsConst to force the const iterator). See top of ConfigCacheIni.h for more info.")
 	FORCEINLINE TRangedForIterator end()
 	{
-		FReadScopeLock ScopeLock(ConfigFileMapLock); 
+		FTransactionallySafeReadScopeLock ScopeLock(ConfigFileMapLock); 
 		return TRangedForIterator(Pairs.end());
 	}
 
 	FORCEINLINE TRangedForConstIterator end() const
 	{
-		FReadScopeLock ScopeLock(ConfigFileMapLock);
+		FTransactionallySafeReadScopeLock ScopeLock(ConfigFileMapLock);
 		return TRangedForConstIterator(Pairs.end());
 	}
 	
@@ -811,7 +811,7 @@ private:
 
 	FORCEINLINE FConfigSection* FindInternal(const FString& SectionName)
 	{
-		FReadScopeLock ScopeLock(ConfigFileMapLock);
+		FTransactionallySafeReadScopeLock ScopeLock(ConfigFileMapLock);
 		return FConfigFileMap::Find(SectionName);
 	};
 
@@ -1010,7 +1010,7 @@ private:
 	friend class FConfigCacheIni;
 	friend FConfigContext;
 
-	static CORE_API FRWLock ConfigFileMapLock;
+	static CORE_API FTransactionallySafeRWLock ConfigFileMapLock;
 };
 
 /**
