Fixed FMallocBinned3::GetAllocatorStats to include all cached memory int he TotalOSAllocated stat when  UE_MB3_USE_CACHED_PAGE_ALLOCATOR_FOR_LARGE_ALLOCS is set to 1

diff --git a/Engine/Source/Runtime/Core/Private/HAL/MallocBinned3.cpp b/Engine/Source/Runtime/Core/Private/HAL/MallocBinned3.cpp
--- a/Engine/Source/Runtime/Core/Private/HAL/MallocBinned3.cpp
+++ b/Engine/Source/Runtime/Core/Private/HAL/MallocBinned3.cpp
@@ -1083,8 +1083,15 @@ void FMallocBinned3::GetAllocatorStats(FGenericMemoryStats& OutStats)
 	OutStats.Add(TEXT("Binned3AllocatedLargePoolMemory"), Binned3AllocatedLargePoolMemory);
 	OutStats.Add(TEXT("Binned3AllocatedLargePoolMemoryWAlignment"), Binned3AllocatedLargePoolMemoryWAlignment);
 
+	const uint64 CachedMem =  
+#if UE_MB3_USE_CACHED_PAGE_ALLOCATOR_FOR_LARGE_ALLOCS
+		GetCachedOSPageAllocator().GetCachedFreeTotal();
+#else
+		0;
+#endif
+
 	const uint64 TotalAllocated = TotalAllocatedSmallPoolMemory + Binned3AllocatedLargePoolMemory;
-	const uint64 TotalOSAllocated = Binned3AllocatedOSSmallPoolMemory + Binned3AllocatedLargePoolMemoryWAlignment;
+	const uint64 TotalOSAllocated = Binned3AllocatedOSSmallPoolMemory + Binned3AllocatedLargePoolMemoryWAlignment + CachedMem;
 
 	OutStats.Add(TEXT("TotalAllocated"), TotalAllocated);
 	OutStats.Add(TEXT("TotalOSAllocated"), TotalOSAllocated);
