[AutoRTFM] Fix memory leaks in the transactionally-safe locks

These were not releasing the final reference to the FState objects if aborted between ctor and dtor.

diff --git a/Engine/Source/Runtime/Core/Public/Misc/TransactionallySafeCriticalSection.h b/Engine/Source/Runtime/Core/Public/Misc/TransactionallySafeCriticalSection.h
--- a/Engine/Source/Runtime/Core/Public/Misc/TransactionallySafeCriticalSection.h
+++ b/Engine/Source/Runtime/Core/Public/Misc/TransactionallySafeCriticalSection.h
@@ -27,18 +27,15 @@ struct FTransactionallySafeCriticalSectionDefinition final
 {
 	// Always open because the constructor arguments will create the underlying critical section.
 	UE_AUTORTFM_ALWAYS_OPEN
-	FTransactionallySafeCriticalSectionDefinition() : State(new FState())
+	FTransactionallySafeCriticalSectionDefinition() : State(MakeShared<FState>())
 	{
 		if (AutoRTFM::IsTransactional())
 		{
 			const AutoRTFM::EContextStatus Status = AutoRTFM::Close([this]
 				{
-					// We explicitly copy the state here for the case that `this` was stack
-					// allocated and has already died before the on-abort is hit.
-					TSharedPtr<FState> LocalState = this->State;
-					AutoRTFM::OnAbort([LocalState]
+					AutoRTFM::PushOnAbortHandler(this, [this]
 						{
-							ensure(0 == LocalState->TransactionalLockCount);
+							this->~FTransactionallySafeCriticalSectionDefinition();
 						});
 				});
 
@@ -52,17 +49,21 @@ struct FTransactionallySafeCriticalSectionDefinition final
 		{
 			const AutoRTFM::EContextStatus Status = AutoRTFM::Close([this]
 				{
+					AutoRTFM::PopOnAbortHandler(this);
+
 					// We explicitly copy the state here for the case that `this` was stack
 					// allocated and has already died before the on-commit is hit.
-					TSharedPtr<FState> LocalState = this->State;
-					AutoRTFM::OnCommit([LocalState]
+					AutoRTFM::OnCommit([State = this->State]
 						{
-							ensure(0 == LocalState->TransactionalLockCount);
+							ensure(0 == State->TransactionalLockCount);
 						});
 				});
 
 			ensure(AutoRTFM::EContextStatus::OnTrack == Status);
 		}
+		
+		// As the State was constructed in the open, it must be released in the open.
+		AutoRTFM::Open([&] { State = nullptr; });
 	}
 
 	void Lock()
@@ -136,7 +137,7 @@ private:
 		uint32 TransactionalLockCount = 0;
 	};
 
-	const TSharedPtr<FState> State;
+	TSharedPtr<FState> State;
 };
 
 #if UE_AUTORTFM
diff --git a/Engine/Source/Runtime/Core/Public/Misc/TransactionallySafeRWLock.h b/Engine/Source/Runtime/Core/Public/Misc/TransactionallySafeRWLock.h
--- a/Engine/Source/Runtime/Core/Public/Misc/TransactionallySafeRWLock.h
+++ b/Engine/Source/Runtime/Core/Public/Misc/TransactionallySafeRWLock.h
@@ -38,16 +38,15 @@ struct FTransactionallySafeRWLockDefinition final
 {
 	// Always open because the constructor arguments will create the underlying critical section.
 	UE_AUTORTFM_ALWAYS_OPEN
-		FTransactionallySafeRWLockDefinition() : State(new FState())
+	FTransactionallySafeRWLockDefinition() : State(MakeShared<FState>())
 	{
 		if (AutoRTFM::IsTransactional())
 		{
 			const AutoRTFM::EContextStatus Status = AutoRTFM::Close([this]
 				{
-					// We explicitly copy the state here for the case that `this` was stack
-					// allocated and has already died before the on-abort is hit.
-					AutoRTFM::OnAbort([State = this->State]
+					AutoRTFM::PushOnAbortHandler(this, [this]
 						{
+							this->~FTransactionallySafeRWLockDefinition();
 						});
 				});
 
@@ -61,6 +60,8 @@ struct FTransactionallySafeRWLockDefinition final
 		{
 			const AutoRTFM::EContextStatus Status = AutoRTFM::Close([this]
 				{
+					AutoRTFM::PopOnAbortHandler(this);
+
 					// We explicitly copy the state here for the case that `this` was stack
 					// allocated and has already died before the on-commit is hit.
 					AutoRTFM::OnCommit([State = this->State]
@@ -71,6 +72,9 @@ struct FTransactionallySafeRWLockDefinition final
 
 			ensure(AutoRTFM::EContextStatus::OnTrack == Status);
 		}
+
+		// As the State was constructed in the open, it must be released in the open.
+		AutoRTFM::Open([&] { State = nullptr; });
 	}
 
 	void ReadLock()
@@ -172,7 +176,7 @@ private:
 		uint32 TransactionalLockCount = 0;
 	};
 
-	const TSharedPtr<FState> State;
+	TSharedPtr<FState> State;
 };
 
 #if UE_AUTORTFM
