[HAL] Support custom output devices for FMallocLeakDetection reports

diff --git a/Engine/Source/Runtime/Core/Private/HAL/MallocLeakDetection.cpp b/Engine/Source/Runtime/Core/Private/HAL/MallocLeakDetection.cpp
--- a/Engine/Source/Runtime/Core/Private/HAL/MallocLeakDetection.cpp
+++ b/Engine/Source/Runtime/Core/Private/HAL/MallocLeakDetection.cpp
@@ -307,14 +307,21 @@ int32 FMallocLeakDetection::DumpOpenCallstacks(const TCHAR* FileName, const FMal
 	FArchive* FileAr = nullptr;
 	FOutputDeviceArchiveWrapper* FileArWrapper = nullptr;
 
-	const FString PathName = *(FPaths::ProfilingDir() + TEXT("memreports/"));
-	IFileManager::Get().MakeDirectory(*PathName);
+	if (Options.OutputDevice)
+	{
+		ReportAr = Options.OutputDevice;
+	}
+	else
+	{
+		const FString PathName = *(FPaths::ProfilingDir() + TEXT("memreports/"));
+		IFileManager::Get().MakeDirectory(*PathName);
 
-	FString FilePath = PathName + CreateProfileFilename(FileName, TEXT(""), true);
+		FString FilePath = PathName + CreateProfileFilename(FileName, TEXT(""), true);
 
-	FileAr = IFileManager::Get().CreateDebugFileWriter(*FilePath);
-	FileArWrapper = new FOutputDeviceArchiveWrapper(FileAr);
-	ReportAr = FileArWrapper;
+		FileAr = IFileManager::Get().CreateDebugFileWriter(*FilePath);
+		FileArWrapper = new FOutputDeviceArchiveWrapper(FileAr);
+		ReportAr = FileArWrapper;
+	}
 
 	const float InvToMb = 1.0 / (1024 * 1024);
 	FPlatformMemoryStats MemoryStats = FPlatformMemory::GetStats();
diff --git a/Engine/Source/Runtime/Core/Public/HAL/MallocLeakDetection.h b/Engine/Source/Runtime/Core/Public/HAL/MallocLeakDetection.h
--- a/Engine/Source/Runtime/Core/Public/HAL/MallocLeakDetection.h
+++ b/Engine/Source/Runtime/Core/Public/HAL/MallocLeakDetection.h
@@ -55,6 +55,9 @@ struct FMallocLeakReportOptions
 
 	/** Sort allocations by this (default - size) */
 	ESortOption		SortBy;
+
+	/** Output device used for writing reports. If null, a memreport file will be written */
+	FOutputDevice*	OutputDevice;
 };
 
 #if MALLOC_LEAKDETECTION
