[AutoRTFM] Add framework for safe open -> closed assignment, use it to fix memory leak with FText

Replace AutoRTFM::TIsSafeToReturnFromOpen with AutoRTFM::TAssignFromOpenToClosed.
This provides the same "is-it-safe-to-return-T" functionality, but also supports custom logic for safely copy-assigning or move-assigning a value from the open to the closed.

This is used to fix a memory leak with FText.

diff --git a/Engine/Source/Programs/AutoRTFMTests/Private/Open.cpp b/Engine/Source/Programs/AutoRTFMTests/Private/Open.cpp
--- a/Engine/Source/Programs/AutoRTFMTests/Private/Open.cpp
+++ b/Engine/Source/Programs/AutoRTFMTests/Private/Open.cpp
@@ -105,13 +105,13 @@ TEST_CASE("Open.Atomics")
 
 TEST_CASE("Open.ReturnValue")
 {
-    static_assert(AutoRTFM::TIsSafeToReturnFromOpen<int>::value);
-    static_assert(AutoRTFM::TIsSafeToReturnFromOpen<float>::value);
-    static_assert(AutoRTFM::TIsSafeToReturnFromOpen<int*>::value);
-    static_assert(AutoRTFM::TIsSafeToReturnFromOpen<void>::value);
-    static_assert(AutoRTFM::TIsSafeToReturnFromOpen<std::tuple<int, float>>::value);
-    static_assert(!AutoRTFM::TIsSafeToReturnFromOpen<std::string>::value);
-    static_assert(!AutoRTFM::TIsSafeToReturnFromOpen<std::tuple<int, std::string>>::value);
+    static_assert(AutoRTFM::IsSafeToReturnFromOpen<int>);
+    static_assert(AutoRTFM::IsSafeToReturnFromOpen<float>);
+    static_assert(AutoRTFM::IsSafeToReturnFromOpen<int*>);
+    static_assert(AutoRTFM::IsSafeToReturnFromOpen<void>);
+    static_assert(AutoRTFM::IsSafeToReturnFromOpen<std::tuple<int, float>>);
+    static_assert(!AutoRTFM::IsSafeToReturnFromOpen<std::string>);
+    static_assert(!AutoRTFM::IsSafeToReturnFromOpen<std::tuple<int, std::string>>);
 
     SECTION("int")
     {
@@ -158,4 +158,75 @@ TEST_CASE("Open.ReturnValue")
         REQUIRE(42 == Int);
         REQUIRE("woof" == String);
     }
+
+    SECTION("Custom type")
+    {
+        SECTION("AutoRTFMAssignFromOpenToClosed() by value")
+        {
+            struct FMyStruct
+            {
+                int Value = 0;
+                static void AutoRTFMAssignFromOpenToClosed(FMyStruct& Closed, FMyStruct Open)
+                {
+                    Closed.Value = Open.Value;
+                }
+            };
+            static_assert(AutoRTFM::IsSafeToReturnFromOpen<FMyStruct>);
+
+            FMyStruct StructResult;
+            AutoRTFM::ETransactionResult TransactionResult = AutoRTFM::Transact([&]
+            {
+                FMyStruct Closed = AutoRTFM::Open([] { return FMyStruct{42}; });
+                AutoRTFM::Open([&] { StructResult = Closed; });
+            });
+            REQUIRE(AutoRTFM::ETransactionResult::Committed == TransactionResult);
+            REQUIRE(42 == StructResult.Value);
+        }
+        SECTION("AutoRTFMAssignFromOpenToClosed() by const-ref")
+        {
+            struct FMyStruct
+            {
+                int Value = 0;
+                static void AutoRTFMAssignFromOpenToClosed(FMyStruct& Closed, const FMyStruct& Open)
+                {
+                    Closed.Value = Open.Value;
+                }
+            };
+            static_assert(AutoRTFM::IsSafeToReturnFromOpen<FMyStruct>);
+
+            FMyStruct StructResult;
+            AutoRTFM::ETransactionResult TransactionResult = AutoRTFM::Transact([&]
+            {
+                FMyStruct Closed = AutoRTFM::Open([] { return FMyStruct{42}; });
+                AutoRTFM::Open([&] { StructResult = Closed; });
+            });
+            REQUIRE(AutoRTFM::ETransactionResult::Committed == TransactionResult);
+            REQUIRE(42 == StructResult.Value);
+        }
+        SECTION("AutoRTFMAssignFromOpenToClosed() by rvalue-ref")
+        {
+            struct FMyStruct
+            {
+                int Value = 0;
+                bool* WasMoved = nullptr;
+                static void AutoRTFMAssignFromOpenToClosed(FMyStruct& Closed, FMyStruct&& Open)
+                {
+                    Closed.Value = Open.Value;
+                    *Open.WasMoved = true;
+                }
+            };
+            static_assert(AutoRTFM::IsSafeToReturnFromOpen<FMyStruct>);
+
+            bool WasMoved = false;
+            FMyStruct StructResult;
+            AutoRTFM::ETransactionResult TransactionResult = AutoRTFM::Transact([&]
+            {
+                FMyStruct Closed = AutoRTFM::Open([&] { return FMyStruct{42, &WasMoved}; });
+                AutoRTFM::Open([&] { StructResult = Closed; });
+            });
+            REQUIRE(AutoRTFM::ETransactionResult::Committed == TransactionResult);
+            REQUIRE(true == WasMoved);
+            REQUIRE(42 == StructResult.Value);
+        }
+    }
 }
diff --git a/Engine/Source/Runtime/Core/Private/Internationalization/Text.cpp b/Engine/Source/Runtime/Core/Private/Internationalization/Text.cpp
--- a/Engine/Source/Runtime/Core/Private/Internationalization/Text.cpp
+++ b/Engine/Source/Runtime/Core/Private/Internationalization/Text.cpp
@@ -291,6 +291,17 @@ FText& FText::operator=(FText&& Other)
 	return *this;
 }
 
+void FText::AutoRTFMAssignFromOpenToClosed(FText& Closed, const FText& Open)
+{
+	// Copy (instead of move) in a closed transaction to ensure the ref-counts
+	// are correctly tracked in the transaction.
+	const AutoRTFM::EContextStatus Status = AutoRTFM::Close([&]
+	{
+		Closed = Open; 
+	});
+	ensure(AutoRTFM::EContextStatus::OnTrack == Status);
+}
+
 bool FText::IsEmpty() const
 {
 	return ToString().IsEmpty();
diff --git a/Engine/Source/Runtime/Core/Private/Internationalization/TextCache.cpp b/Engine/Source/Runtime/Core/Private/Internationalization/TextCache.cpp
--- a/Engine/Source/Runtime/Core/Private/Internationalization/TextCache.cpp
+++ b/Engine/Source/Runtime/Core/Private/Internationalization/TextCache.cpp
@@ -4,6 +4,7 @@
 
 #include "Containers/UnrealString.h"
 #include "HAL/LowLevelMemTracker.h"
+#include "Internationalization/Text.h"
 #include "Internationalization/TextKey.h"
 #include "Misc/CString.h"
 #include "Misc/LazySingleton.h"
@@ -33,42 +34,44 @@ FText FTextCache::FindOrCache(const TCHAR* InTextLiteral, const TCHAR* InNamespa
 	return FindOrCache(InTextLiteral, FTextId(InNamespace, InKey));
 }
 
-UE_AUTORTFM_ALWAYS_OPEN FText FTextCache::FindOrCache(const TCHAR* InTextLiteral, const FTextId& InTextId)
+FText FTextCache::FindOrCache(const TCHAR* InTextLiteral, const FTextId& InTextId)
 {
-	LLM_SCOPE(ELLMTag::Localization);
-
-	// First try and find a cached instance
+	return AutoRTFM::Open([&]
 	{
-		FText* ReturnFoundText = nullptr;
-
-		FScopeLock Lock(&CachedTextCS);
+		LLM_SCOPE(ELLMTag::Localization);
 
-		FText* FoundText = CachedText.Find(InTextId);
-		if (FoundText)
+		// First try and find a cached instance
 		{
-			const FString* FoundTextLiteral = FTextInspector::GetSourceString(*FoundText);
-			if (FoundTextLiteral && FCString::Strcmp(**FoundTextLiteral, InTextLiteral) == 0)
-			{
+			FText* ReturnFoundText = nullptr;
 
-				ReturnFoundText = FoundText;
+			FScopeLock Lock(&CachedTextCS);
+
+			FText* FoundText = CachedText.Find(InTextId);
+			if (FoundText)
+			{
+				const FString* FoundTextLiteral = FTextInspector::GetSourceString(*FoundText);
+				if (FoundTextLiteral && FCString::Strcmp(**FoundTextLiteral, InTextLiteral) == 0)
+				{
+					ReturnFoundText = FoundText;
+				}
 			}
-		}
 
-		if (ReturnFoundText)
-		{
-			return *ReturnFoundText;
+			if (ReturnFoundText)
+			{
+				return *ReturnFoundText;
+			}
 		}
-	}
 
-	// Not currently cached, make a new instance...
-	FText NewText = FText(InTextLiteral, InTextId.GetNamespace(), InTextId.GetKey(), ETextFlag::Immutable);
+		// Not currently cached, make a new instance...
+		FText NewText = FText(InTextLiteral, InTextId.GetNamespace(), InTextId.GetKey(), ETextFlag::Immutable);
 
-	// ... and add it to the cache
-	FScopeLock Lock(&CachedTextCS);
+		// ... and add it to the cache
+		FScopeLock Lock(&CachedTextCS);
 
-	CachedText.Emplace(InTextId, NewText);
+		CachedText.Emplace(InTextId, NewText);
 
-	return NewText;
+		return NewText;
+	});
 }
 
 UE_AUTORTFM_NOAUTORTFM void FTextCache::RemoveCache(const FTextId& InTextId)
diff --git a/Engine/Source/Runtime/Core/Public/AutoRTFM/AutoRTFM.h b/Engine/Source/Runtime/Core/Public/AutoRTFM/AutoRTFM.h
--- a/Engine/Source/Runtime/Core/Public/AutoRTFM/AutoRTFM.h
+++ b/Engine/Source/Runtime/Core/Public/AutoRTFM/AutoRTFM.h
@@ -368,6 +368,8 @@ static UE_AUTORTFM_FORCEINLINE void autortfm_check_abi(void* ptr, size_t size)
 #ifdef __cplusplus
 }
 
+#include <algorithm>
+#include <initializer_list>
 #include <tuple>
 
 namespace AutoRTFM
@@ -559,54 +561,148 @@ static UE_AUTORTFM_FORCEINLINE void AbortIfClosed()
     autortfm_abort_if_closed();
 }
 
-// Traits helper that is used to determine whether type T can be used as the 
-// return type of a function called by Open().
+namespace Detail
+{
+template <typename T, typename = void>
+struct THasAssignFromOpenToClosedMethod : std::false_type {};
+template <typename T>
+struct THasAssignFromOpenToClosedMethod<T, std::void_t<decltype(T::AutoRTFMAssignFromOpenToClosed(std::declval<T&>(), std::declval<T>()))>> : std::true_type {};
+}
+
+// Evaluates to true if the type T has a static method with the signature:
+//    static void AutoRTFMAssignFromOpenToClosed(T& Closed, U Open)
+// Where `U` is `T`, `const T&` or `T&&`. Supports both copy assignment and move assignment.
+template <typename T>
+static constexpr bool HasAssignFromOpenToClosedMethod = Detail::THasAssignFromOpenToClosedMethod<T>::value;
+
+// Template class used to declare a method for safely copying or moving an 
+// object of type T from open to closed transactions.
+// Specializations of TAssignFromOpenToClosed must have at least one static
+// method with the signature:
+//   static void Assign(T& Closed, U Open);
+// Where `U` is `T`, `const T&` or `T&&`. Supports both copy assignment and move assignment.
+//
+// TAssignFromOpenToClosed has pre-declared specializations for basic primitive
+// types, and can be extended with user-declared template specializations.
+//
+// TAssignFromOpenToClosed has a pre-declared specialization that detects and
+// calls a static method on T with the signature:
+//    static void AutoRTFMAssignFromOpenToClosed(T& Closed, U Open)
+// Where `U` is `T`, `const T&` or `T&&`. Supports both copy assignment and move assignment.
 template<typename T, typename = void>
-class TIsSafeToReturnFromOpen : public std::false_type {};
+struct TAssignFromOpenToClosed;
 
-// Specializations of TIsSafeToReturnFromOpen for fundamental types.
+namespace Detail
+{
+template <typename T, typename = void>
+struct THasAssignFromOpenToClosedTrait : std::false_type {};
 template <typename T>
-class TIsSafeToReturnFromOpen<T, std::enable_if_t<std::is_fundamental_v<T>>> : public std::true_type {};
+struct THasAssignFromOpenToClosedTrait<T, std::void_t<decltype(TAssignFromOpenToClosed<T>::Assign(std::declval<T&>(), std::declval<T>()))>> : std::true_type {};
+}
 
-// Specializations of TIsSafeToReturnFromOpen for raw pointer types.
+// Evaluates to true if the type T supports assigning from open to closed transactions.
 template <typename T>
-class TIsSafeToReturnFromOpen<T*, void> : public std::true_type {};
+static constexpr bool HasAssignFromOpenToClosedTrait = Detail::THasAssignFromOpenToClosedTrait<T>::value;
 
-// Specializations of TIsSafeToReturnFromOpen for std::tuple.
+// Specialization of TAssignFromOpenToClosed for fundamental types.
+template <typename T>
+struct TAssignFromOpenToClosed<T, std::enable_if_t<std::is_fundamental_v<T>>>
+{
+	UE_AUTORTFM_FORCEINLINE static void Assign(T& Closed, T Open) { Closed = Open; }
+};
+
+// Specialization of TAssignFromOpenToClosed for raw pointer types.
+template <typename T>
+struct TAssignFromOpenToClosed<T*, void>
+{
+	UE_AUTORTFM_FORCEINLINE static void Assign(T*& Closed, T* Open) { Closed = Open; }
+};
+
+// Specialization of TAssignFromOpenToClosed for std::tuple.
 template <typename ... TYPES>
-class TIsSafeToReturnFromOpen<std::tuple<TYPES...>, void> : public std::conditional_t<
-	(TIsSafeToReturnFromOpen<TYPES>::value && ...), std::true_type, std::false_type> {};
+struct TAssignFromOpenToClosed<std::tuple<TYPES...>, std::enable_if_t<(HasAssignFromOpenToClosedTrait<TYPES> && ...)>>
+{
+	template<size_t I = 0, typename SRC = void>
+	UE_AUTORTFM_FORCEINLINE static void AssignElements(std::tuple<TYPES...>& Closed, SRC&& Open)
+	{
+		if constexpr(I < sizeof...(TYPES))
+		{
+			using E = std::tuple_element_t<I, std::tuple<TYPES...>>;
+			TAssignFromOpenToClosed<E>::Assign(std::get<I>(Closed), std::get<I>(std::forward<SRC>(Open)));
+			AssignElements<I+1>(Closed, std::forward<SRC>(Open));
+		}
+	}
+
+	template<typename SRC>
+	UE_AUTORTFM_FORCEINLINE static void Assign(std::tuple<TYPES...>& Closed, SRC&& Open)
+	{
+		AssignElements(Closed, std::forward<SRC>(Open));
+	}
+};
+
+// Specialization of TAssignFromOpenToClosed for types that have a static method
+// with the signature:
+//    static void AutoRTFMAssignFromOpenToClosed(T& Closed, U Open)
+// Where `U` is `T`, `const T&` or `T&&`. Supports both copy assignment and move assignment.
+template <typename T>
+struct TAssignFromOpenToClosed<T, std::enable_if_t<HasAssignFromOpenToClosedMethod<T>>>
+{
+	template<typename OPEN>
+	UE_AUTORTFM_FORCEINLINE static void Assign(T& Closed, OPEN&& Open)
+	{
+		Closed.AutoRTFMAssignFromOpenToClosed(Closed, std::forward<OPEN>(Open));
+	}
+};
+
+// Specialization of TAssignFromOpenToClosed for `void` (used to make IsSafeToReturnFromOpen<void> work).
+template <>
+struct TAssignFromOpenToClosed<void, void>;
+
+// Evaluates to true if the type T is safe to return from Open().
+template <typename T>
+static constexpr bool IsSafeToReturnFromOpen = HasAssignFromOpenToClosedTrait<T> || std::is_same_v<T, void>;
 
 // Executes the given code non-transactionally regardless of whether we are in
 // a transaction or not. Returns the value returned by Functor.
+// TReturn must be void or a type that can be safely copied from the open to a closed transaction.
+// TAssignFromOpenToClosed must have a specialization for the type that is being returned.
 template<typename TFunctor, typename TReturn = decltype(std::declval<TFunctor>()())> 
 static UE_AUTORTFM_FORCEINLINE TReturn Open(const TFunctor& Functor)
 {
-	static_assert(TIsSafeToReturnFromOpen<TReturn>::value,
+	static_assert(IsSafeToReturnFromOpen<TReturn>,
 		"function return type is not safe to return from Open()");
 
-	if constexpr (std::is_same_v<void, TReturn>)
+	if (!autortfm_is_closed())
 	{
-		autortfm_open(
-			[] (void* Arg) { UE_AUTORTFM_CALLSITE_FORCEINLINE (*static_cast<const TFunctor*>(Arg))(); },
-			const_cast<void*>(static_cast<const void*>(&Functor)));
+		return Functor();
 	}
-	else
+
+	if constexpr (IsSafeToReturnFromOpen<TReturn>)
 	{
-		TReturn ReturnValue;
-		struct FData
+		if constexpr (std::is_same_v<void, TReturn>)
 		{
-			const TFunctor& Functor;
-			TReturn& ReturnValue;
-		};
-		FData Data{Functor, ReturnValue};
-		autortfm_open([](void* Arg)
+			autortfm_open(
+				[] (void* Arg) { UE_AUTORTFM_CALLSITE_FORCEINLINE (*static_cast<const TFunctor*>(Arg))(); },
+				const_cast<void*>(static_cast<const void*>(&Functor)));
+		}
+		else
+		{
+			TReturn ReturnValue;
+			struct FData
 			{
-				FData& Data = *reinterpret_cast<FData*>(Arg);
-				UE_AUTORTFM_CALLSITE_FORCEINLINE Data.ReturnValue = Data.Functor();
-			},
-			reinterpret_cast<void*>(&Data));
-		return Data.ReturnValue;
+				const TFunctor& Functor;
+				TReturn& ReturnValue;
+			};
+			FData Data{Functor, ReturnValue};
+			autortfm_open([](void* Arg)
+				{
+					FData& Data = *reinterpret_cast<FData*>(Arg);
+					UE_AUTORTFM_CALLSITE_FORCEINLINE TAssignFromOpenToClosed<TReturn>::Assign(Data.ReturnValue, std::move(Data.Functor()));
+				},
+				reinterpret_cast<void*>(&Data));
+
+			return Data.ReturnValue;
+		}
 	}
 }
 
diff --git a/Engine/Source/Runtime/Core/Public/Internationalization/Text.h b/Engine/Source/Runtime/Core/Public/Internationalization/Text.h
--- a/Engine/Source/Runtime/Core/Public/Internationalization/Text.h
+++ b/Engine/Source/Runtime/Core/Public/Internationalization/Text.h
@@ -394,6 +394,10 @@ public:
 	CORE_API FText(FText&& Other);
 	CORE_API FText& operator=(FText&& Other);
 
+	// Safely copies an FText from the open to a closed transaction.
+	// See AutoRTFM::Open().
+	static CORE_API void AutoRTFMAssignFromOpenToClosed(FText& Closed, const FText& Open);
+
 	/**
 	 * Generate an FText that represents the passed number in the current culture
 	 */
