[AutoRTFM] OnCommit: Free abort tasks, then each on-commit task after invoking

This provides symmetry with the on-abort flow (CL 35470941).

While this does not fix any bugs I've encountered, I find this greatly helps with debugging ref-count isssues with lambda-captured variables with commit / abort tasks.

diff --git a/Engine/Source/Runtime/Core/Private/AutoRTFM/TaskArray.h b/Engine/Source/Runtime/Core/Private/AutoRTFM/TaskArray.h
--- a/Engine/Source/Runtime/Core/Private/AutoRTFM/TaskArray.h
+++ b/Engine/Source/Runtime/Core/Private/AutoRTFM/TaskArray.h
@@ -166,6 +166,30 @@ public:
         return true;
     }
 
+    template<typename TFunc>
+    bool ForEachForward(const TFunc& Func)
+    {
+        for (FInternalArray& StashedVectorBox : Stash)
+        {
+            for (SKeyValuePair& EntryKVP : StashedVectorBox)
+            {
+                if (!Func(EntryKVP.Val))
+                {
+                    return false;
+                }
+            }
+        }
+
+        for (SKeyValuePair& EntryKVP : Latest)
+        {
+            if (!Func(EntryKVP.Val))
+            {
+                return false;
+            }
+        }
+        return true;
+    }
+
     template<typename TFunc>
     bool ForEachBackward(const TFunc& Func) const
     {
diff --git a/Engine/Source/Runtime/Core/Private/AutoRTFM/Transaction.cpp b/Engine/Source/Runtime/Core/Private/AutoRTFM/Transaction.cpp
--- a/Engine/Source/Runtime/Core/Private/AutoRTFM/Transaction.cpp
+++ b/Engine/Source/Runtime/Core/Private/AutoRTFM/Transaction.cpp
@@ -186,7 +186,14 @@ bool FTransaction::AttemptToCommitOuterNest()
 	Context->DumpState();
 	UE_LOG(LogAutoRTFM, Verbose, TEXT("Running commit tasks..."));
 
-    CommitTasks.ForEachForward([] (const TFunction<void()>& Task) -> bool { Task(); return true; });
+    AbortTasks.Reset();
+
+    CommitTasks.ForEachForward([] (TFunction<void()>& Task) -> bool
+    { 
+        Task(); 
+        Task.Reset();
+        return true; 
+    });
 
     return true;
 }
