[Core/Containers] Use MALLOCLEAK_IGNORE_SCOPE() for TLS allocations that are never freed

Without this, these allocations are reported as leaks. Technically they are, but I assume they're small and infrequent enough to be ignored.

diff --git a/Engine/Source/Runtime/Core/Private/Containers/LockFreeList.cpp b/Engine/Source/Runtime/Core/Private/Containers/LockFreeList.cpp
--- a/Engine/Source/Runtime/Core/Private/Containers/LockFreeList.cpp
+++ b/Engine/Source/Runtime/Core/Private/Containers/LockFreeList.cpp
@@ -1,6 +1,7 @@
 // Copyright Epic Games, Inc. All Rights Reserved.
 
 #include "Containers/LockFreeList.h"
+#include "HAL/MallocLeakDetection.h"
 #include "HAL/PlatformProcess.h"
 #include "HAL/IConsoleManager.h"
 #include "Stats/Stats.h"
@@ -192,6 +193,7 @@ private:
 		FThreadLocalCache* TLS = (FThreadLocalCache*)FPlatformTLS::GetTlsValue(TlsSlot);
 		if (!TLS)
 		{
+			MALLOCLEAK_IGNORE_SCOPE(); // TLS is never freed
 			TLS = new FThreadLocalCache();
 			FPlatformTLS::SetTlsValue(TlsSlot, TLS);
 		}
diff --git a/Engine/Source/Runtime/Core/Public/Containers/LockFreeFixedSizeAllocator.h b/Engine/Source/Runtime/Core/Public/Containers/LockFreeFixedSizeAllocator.h
--- a/Engine/Source/Runtime/Core/Public/Containers/LockFreeFixedSizeAllocator.h
+++ b/Engine/Source/Runtime/Core/Public/Containers/LockFreeFixedSizeAllocator.h
@@ -2,6 +2,7 @@
 
 #pragma once
 
+#include "HAL/MallocLeakDetection.h"
 #include "Misc/AssertionMacros.h"
 #include "HAL/UnrealMemory.h"
 #include "Misc/NoopCounter.h"
@@ -166,6 +167,7 @@ private:
 		FThreadLocalCache* TLS = (FThreadLocalCache*)FPlatformTLS::GetTlsValue(TlsSlot);
 		if (!TLS)
 		{
+			MALLOCLEAK_IGNORE_SCOPE(); // TLS is never freed
 			TLS = new FThreadLocalCache();
 			FPlatformTLS::SetTlsValue(TlsSlot, TLS);
 		}
