CSV profiler - temporarily disable aggregated task worker stats to fix an issue with custom stat reporting

diff --git a/Engine/Source/Runtime/Core/Private/ProfilingDebugging/CsvProfiler.cpp b/Engine/Source/Runtime/Core/Private/ProfilingDebugging/CsvProfiler.cpp
--- a/Engine/Source/Runtime/Core/Private/ProfilingDebugging/CsvProfiler.cpp
+++ b/Engine/Source/Runtime/Core/Private/ProfilingDebugging/CsvProfiler.cpp
@@ -87,7 +87,7 @@ TAutoConsoleVariable<int32> CVarCsvBlockOnCaptureEnd(
 
 TAutoConsoleVariable<bool> CVarCsvAggregateTaskWorkerStats(
 	TEXT("csv.AggregateTaskWorkerStats"),
-	true,
+	false,
 	TEXT("If enabled, stats recorded on task worker threads are aggregated instead of outputting a single stat per thread.\r\n")
 	TEXT("This reduces CSV bloat when there are large numbers of worker threads and makes stat data more intelligible"),
 	ECVF_Default
