UnitTests: Unified our unit testing backends to consistently use case insensitive string comparisons. Previously, different backends used different contentions (and were often inconsistent internally) which led to LowLevelTests and EngineTests returning different results for the same tests in a few cases.

- FAutomationTestBase (EngineTests) and LowLevelTestsRunner/Expectations.h use case insensitive comparisons by default.
- Deprecated the existing Test[Not]EqualInsensitive as not it always does the same thing as Test[Not]Equal
- Added Test[Not]EqualSensitive in case we need case sensitive comparisons somewhere.
- Fixed crashes in string comparison overloads for TCHAR*
- Removed superfluous overloads (FString vs TCHAR*) from LowLevelTestsRunner.
- Added TStringView and TUtf8StringView overloads to LowLevelTestsRunner to handle the above and more.
- Added a missing overload for TCHAR* to LowLevelTestsRunner (previously, TCHAR* vs TCHAR* would compare pointer values).
- Added a few additional tests.

#tests Executed some EngineTests and FoundationTests locally, ran EngineTest and LowLevelTests preflights with my change and compared results to the nightly runs.

diff --git a/Engine/Plugins/Tests/CQTest/Source/CQTest/Public/Assert/NoDiscardAsserter.inl b/Engine/Plugins/Tests/CQTest/Source/CQTest/Public/Assert/NoDiscardAsserter.inl
--- a/Engine/Plugins/Tests/CQTest/Source/CQTest/Public/Assert/NoDiscardAsserter.inl
+++ b/Engine/Plugins/Tests/CQTest/Source/CQTest/Public/Assert/NoDiscardAsserter.inl
@@ -206,7 +206,7 @@ template <>
 template <>
 [[nodiscard]] FORCEINLINE bool FNoDiscardAsserter::AreEqual(const FString& Expected, const FString& Actual, const FString& FailureMessage)
 {
-	return TestRunner.TestEqual(FailureMessage, *Actual, *Expected); //FAutomationTestBase expects Actual then Expected
+	return TestRunner.TestEqualSensitive(FailureMessage, *Actual, *Expected); //FAutomationTestBase expects Actual then Expected
 }
 
 template <typename TExpected, typename TActual>
diff --git a/Engine/Source/Developer/LowLevelTestsRunner/Public/TestCommon/Expectations.h b/Engine/Source/Developer/LowLevelTestsRunner/Public/TestCommon/Expectations.h
--- a/Engine/Source/Developer/LowLevelTestsRunner/Public/TestCommon/Expectations.h
+++ b/Engine/Source/Developer/LowLevelTestsRunner/Public/TestCommon/Expectations.h
@@ -15,7 +15,31 @@
 
 #include "../TestHarness.h"
 
-namespace {
+namespace
+{
+	FString GetStringValueToDisplay(FStringView Value)
+	{
+		if (Value.GetData())
+		{
+			return FString::Printf(TEXT("\"%.*s\""), Value.Len(), Value.GetData());
+		}
+		else
+		{
+			return TEXT("nullptr");
+		}
+	}
+
+	FString GetStringValueToDisplay(FUtf8StringView Value)
+	{
+		if (Value.GetData())
+		{
+			return FString::Printf(TEXT("\"%.*hs\""), Value.Len(), Value.GetData());
+		}
+		else
+		{
+			return TEXT("nullptr");
+		}
+	}
 
 	bool TestTrue(const TCHAR* What, bool Value)
 	{
@@ -151,24 +175,78 @@ namespace {
 
 	bool TestEqual(const TCHAR* What, const TCHAR* Actual, const TCHAR* Expected)
 	{
-		if (FCString::Strcmp(Actual, Expected) != 0)
+		bool bAreEqual = (Actual && Expected) ? (FCString::Stricmp(Actual, Expected) == 0) : (Actual == Expected);
+
+		if (!bAreEqual)
+		{
+			FAIL_CHECK(FString::Printf(TEXT("Expected '%s' to be %s, but it was %s."), What, *GetStringValueToDisplay(Expected), *GetStringValueToDisplay(Actual)));
+		}
+
+		return bAreEqual;
+	}
+
+	bool TestEqual(const TCHAR* What, FStringView Actual, FStringView Expected)
+	{
+		if (Actual.Compare(Expected, ESearchCase::IgnoreCase) != 0)
 		{
-			FAIL_CHECK(FString::Printf(TEXT("Expected '%s' to be \"%s\", but it was \"%s\"."), What, Expected, Actual));
+			FAIL_CHECK(FString::Printf(TEXT("Expected '%s' to be %s, but it was %s."), What, *GetStringValueToDisplay(Expected), *GetStringValueToDisplay(Actual)));
 			return false;
 		}
+
 		return true;
 	}
 
-	bool TestEqualInsensitive(const TCHAR* What, const TCHAR* Actual, const TCHAR* Expected)
+	bool TestEqual(const TCHAR* What, FUtf8StringView Actual, FUtf8StringView Expected)
 	{
-		if (FCString::Stricmp(Actual, Expected) != 0)
+		if (Actual.Compare(Expected, ESearchCase::IgnoreCase) != 0)
 		{
-			FAIL_CHECK(FString::Printf(TEXT("Expected '%s' to be \"%s\", but it was \"%s\"."), What, Expected, Actual));
+			FAIL_CHECK(FString::Printf(TEXT("Expected '%s' to be %s, but it was %s."), What, *GetStringValueToDisplay(Expected), *GetStringValueToDisplay(Actual)));
 			return false;
 		}
+
 		return true;
 	}
 
+	bool TestEqualSensitive(const TCHAR* What, const TCHAR* Actual, const TCHAR* Expected)
+	{
+		bool bAreEqual = (Actual && Expected) ? (FCString::Strcmp(Actual, Expected) == 0) : (Actual == Expected);
+
+		if (!bAreEqual)
+		{
+			FAIL_CHECK(FString::Printf(TEXT("Expected '%s' to be %s, but it was %s."), What, *GetStringValueToDisplay(Expected), *GetStringValueToDisplay(Actual)));
+		}
+
+		return bAreEqual;
+	}
+
+	bool TestEqualSensitive(const TCHAR* What, FStringView Actual, FStringView Expected)
+	{
+		if (Actual.Compare(Expected, ESearchCase::CaseSensitive) != 0)
+		{
+			FAIL_CHECK(FString::Printf(TEXT("Expected '%s' to be %s, but it was %s."), What, *GetStringValueToDisplay(Expected), *GetStringValueToDisplay(Actual)));
+			return false;
+		}
+
+		return true;
+	}
+
+	bool TestEqualSensitive(const TCHAR* What, FUtf8StringView Actual, FUtf8StringView Expected)
+	{
+		if (Actual.Compare(Expected, ESearchCase::CaseSensitive) != 0)
+		{
+			FAIL_CHECK(FString::Printf(TEXT("Expected '%s' to be %s, but it was %s."), What, *GetStringValueToDisplay(Expected), *GetStringValueToDisplay(Actual)));
+			return false;
+		}
+
+		return true;
+	}
+
+	UE_DEPRECATED(5.5, "Use TestEqual instead (string tests are case insensitive by default)")
+	bool TestEqualInsensitive(const TCHAR* What, const TCHAR* Actual, const TCHAR* Expected)
+	{
+		return TestEqual(What, Actual, Expected);
+	}
+
 	bool TestEqual(const FString& What, const int32 Actual, const int32 Expected)
 	{
 		return TestEqual(*What, Actual, Expected);
@@ -209,51 +287,99 @@ namespace {
 		return TestEqual(*What, Actual, Expected);
 	}
 
-	bool TestEqual(const TCHAR* What, const FString& Actual, const TCHAR* Expected)
+	bool TestEqual(const FString& What, FStringView Actual, FStringView Expected)
+	{
+		return TestEqual(*What, Actual, Expected);
+	}
+
+	bool TestEqual(const FString& What, FUtf8StringView Actual, FUtf8StringView Expected)
+	{
+		return TestEqual(*What, Actual, Expected);
+	}
+
+	template<typename ValueType>
+	bool TestEqual(const TCHAR* What, const ValueType& Actual, const ValueType& Expected)
 	{
-		return TestEqualInsensitive(What, *Actual, Expected);
+		if (Actual != Expected)
+		{
+			FAIL_CHECK(FString::Printf(TEXT("%s: The two values are not equal."), What));
+			return false;
+		}
+		return true;
 	}
 
-	bool TestEqual(const FString& What, const FString& Actual, const TCHAR* Expected)
+	template<typename ValueType>
+	bool TestEqual(const FString& What, const ValueType& Actual, const ValueType& Expected)
 	{
-		return TestEqualInsensitive(*What, *Actual, Expected);
+		return TestEqual(*What, Actual, Expected);
 	}
 
-	bool TestEqual(const TCHAR* What, const TCHAR* Actual, const FString& Expected)
+	bool TestNotEqual(const TCHAR* What, const TCHAR* Actual, const TCHAR* Expected)
 	{
-		return TestEqualInsensitive(What, Actual, *Expected);
+		bool bAreDifferent = (Actual && Expected) ? (FCString::Stricmp(Actual, Expected) != 0) : (Actual != Expected);
+
+		if (!bAreDifferent)
+		{
+			FAIL_CHECK(FString::Printf(TEXT("Expected '%s' to differ from %s, but it was %s."), What, *GetStringValueToDisplay(Expected), *GetStringValueToDisplay(Actual)));
+		}
+
+		return bAreDifferent;
 	}
 
-	bool TestEqual(const FString& What, const TCHAR* Actual, const FString& Expected)
+	bool TestNotEqual(const TCHAR* What, FStringView Actual, FStringView Expected)
 	{
-		return TestEqualInsensitive(*What, Actual, *Expected);
+		if (Actual.Compare(Expected, ESearchCase::IgnoreCase) == 0)
+		{
+			FAIL_CHECK(FString::Printf(TEXT("Expected '%s' to differ from %s, but it was %s."), What, *GetStringValueToDisplay(Expected), *GetStringValueToDisplay(Actual)));
+			return false;
+		}
+
+		return true;
 	}
 
-	bool TestEqual(const TCHAR* What, const FString& Actual, const FString& Expected)
+	bool TestNotEqual(const TCHAR* What, FUtf8StringView Actual, FUtf8StringView Expected)
 	{
-		return TestEqualInsensitive(What, *Actual, *Expected);
+		if (Actual.Compare(Expected, ESearchCase::IgnoreCase) == 0)
+		{
+			FAIL_CHECK(FString::Printf(TEXT("Expected '%s' to differ from %s, but it was %s."), What, *GetStringValueToDisplay(Expected), *GetStringValueToDisplay(Actual)));
+			return false;
+		}
+
+		return true;
 	}
 
-	bool TestEqual(const FString& What, const FString& Actual, const FString& Expected)
+	bool TestNotEqualSensitive(const TCHAR* What, const TCHAR* Actual, const TCHAR* Expected)
 	{
-		return TestEqualInsensitive(*What, *Actual, *Expected);
+		bool bAreDifferent = (Actual && Expected) ? (FCString::Strcmp(Actual, Expected) != 0) : (Actual != Expected);
+
+		if (!bAreDifferent)
+		{
+			FAIL_CHECK(FString::Printf(TEXT("Expected '%s' to differ from %s, but it was %s."), What, *GetStringValueToDisplay(Expected), *GetStringValueToDisplay(Actual)));
+		}
+
+		return bAreDifferent;
 	}
 
-	template<typename ValueType>
-	bool TestEqual(const TCHAR* What, const ValueType& Actual, const ValueType& Expected)
+	bool TestNotEqualSensitive(const TCHAR* What, FStringView Actual, FStringView Expected)
 	{
-		if (Actual != Expected)
+		if (Actual.Compare(Expected, ESearchCase::CaseSensitive) == 0)
 		{
-			FAIL_CHECK(FString::Printf(TEXT("%s: The two values are not equal."), What));
+			FAIL_CHECK(FString::Printf(TEXT("Expected '%s' to differ from %s, but it was %s."), What, *GetStringValueToDisplay(Expected), *GetStringValueToDisplay(Actual)));
 			return false;
 		}
+
 		return true;
 	}
 
-	template<typename ValueType>
-	bool TestEqual(const FString& What, const ValueType& Actual, const ValueType& Expected)
+	bool TestNotEqualSensitive(const TCHAR* What, FUtf8StringView Actual, FUtf8StringView Expected)
 	{
-		return TestEqual(*What, Actual, Expected);
+		if (Actual.Compare(Expected, ESearchCase::CaseSensitive) == 0)
+		{
+			FAIL_CHECK(FString::Printf(TEXT("Expected '%s' to differ from %s, but it was %s."), What, *GetStringValueToDisplay(Expected), *GetStringValueToDisplay(Actual)));
+			return false;
+		}
+
+		return true;
 	}
 
 	template<typename ValueType> bool TestNotEqual(const TCHAR* Description, const ValueType& Actual, const ValueType& Expected)
@@ -266,12 +392,44 @@ namespace {
 		return true;
 	}
 
+	bool TestNotEqual(const FString& What, const TCHAR* Actual, const TCHAR* Expected)
+	{
+		return TestNotEqual(*What, Actual, Expected);
+	}
+
+	bool TestNotEqual(const FString& What, FStringView Actual, FStringView Expected)
+	{
+		return TestNotEqual(*What, Actual, Expected);
+	}
+
+	bool TestNotEqual(const FString& What, FUtf8StringView Actual, FUtf8StringView Expected)
+	{
+		return TestNotEqual(*What, Actual, Expected);
+	}
+
+	bool TestNotEqualSensitive(const FString& What, const TCHAR* Actual, const TCHAR* Expected)
+	{
+		return TestNotEqualSensitive(*What, Actual, Expected);
+	}
+
+	bool TestNotEqualSensitive(const FString& What, FStringView Actual, FStringView Expected)
+	{
+		return TestNotEqualSensitive(*What, Actual, Expected);
+	}
+
+	bool TestNotEqualSensitive(const FString& What, FUtf8StringView Actual, FUtf8StringView Expected)
+	{
+		return TestNotEqualSensitive(*What, Actual, Expected);
+	}
+
 	template<typename ValueType> bool TestNotEqual(const FString& Description, const ValueType& Actual, const ValueType& Expected)
 	{
 		return TestNotEqual(*Description, Actual, Expected);
 	}
 
-	#define CHECK_EQUALS(What, X, Y) TestEqual(What, X, Y);
-	#define CHECK_NOT_EQUALS(What, X, Y) TestNotEqual(What, X, Y);
+	#define CHECK_EQUALS(What, X, Y)				TestEqual(What, X, Y);
+	#define CHECK_EQUALS_SENSITIVE(What, X, Y)		TestEqualSensitive(What, X, Y);
+	#define CHECK_NOT_EQUALS(What, X, Y)			TestNotEqual(What, X, Y);
+	#define CHECK_NOT_EQUALS_SENSITIVE(What, X, Y)	TestNotEqualSensitive(What, X, Y);
 
 }
\ No newline at end of file
diff --git a/Engine/Source/Runtime/Core/Private/Misc/AutomationTest.cpp b/Engine/Source/Runtime/Core/Private/Misc/AutomationTest.cpp
--- a/Engine/Source/Runtime/Core/Private/Misc/AutomationTest.cpp
+++ b/Engine/Source/Runtime/Core/Private/Misc/AutomationTest.cpp
@@ -1592,6 +1592,30 @@ void FAutomationTestBase::InternalSetSuccessState( bool bSuccessful )
 	ExecutionInfo.bSuccessful = bSuccessful;
 }
 
+FString FAutomationTestBase::GetStringValueToDisplay(FStringView Value) const
+{
+	if (Value.GetData())
+	{
+		return FString::Printf(TEXT("\"%.*s\""), Value.Len(), Value.GetData());
+	}
+	else
+	{
+		return TEXT("nullptr");
+	}
+}
+
+FString FAutomationTestBase::GetStringValueToDisplay(FUtf8StringView Value) const
+{
+	if (Value.GetData())
+	{
+		return FString::Printf(TEXT("\"%.*hs\""), Value.Len(), Value.GetData());
+	}
+	else
+	{
+		return TEXT("nullptr");
+	}
+}
+
 bool FAutomationTestBase::GetLastExecutionSuccessState()
 {
 	return ExecutionInfo.bSuccessful;
@@ -1859,76 +1883,72 @@ bool FAutomationTestBase::TestEqual(const TCHAR* What, const FLinearColor Actual
 
 bool FAutomationTestBase::TestEqual(const TCHAR* What, const TCHAR* Actual, const TCHAR* Expected)
 {
-	if (FCString::Strcmp(Actual, Expected) != 0)
-	{
-		AddError(FString::Printf(TEXT("Expected '%s' to be \"%s\", but it was \"%s\"."), What, Expected, Actual), 1);
-		return false;
-	}
-	return true;
+ 	bool bAreEqual = (Actual && Expected) ? (FCString::Stricmp(Actual, Expected) == 0) : (Actual == Expected);
+ 
+ 	if (!bAreEqual)
+ 	{
+ 		AddError(FString::Printf(TEXT("Expected '%s' to be %s, but it was %s."), What, *GetStringValueToDisplay(Expected), *GetStringValueToDisplay(Actual)), 1);
+ 	}
+ 
+ 	return bAreEqual;
 }
 
 bool FAutomationTestBase::TestEqual(const TCHAR* What, FUtf8StringView Actual, FUtf8StringView Expected)
 {
-	if (Actual.Compare(Expected, ESearchCase::CaseSensitive))
+	if (Actual.Compare(Expected, ESearchCase::IgnoreCase) != 0)
 	{
-		AddError(FString::Printf(TEXT("Expected '%s' to be \"%.*hs\", but it was \"%.*hs\"."), What, Expected.Len(), Expected.GetData(), Actual.Len(), Actual.GetData()), 1);
+		AddError(FString::Printf(TEXT("Expected '%s' to be %s, but it was %s."), What, *GetStringValueToDisplay(Expected), *GetStringValueToDisplay(Actual)), 1);
 		return false;
 	}
+
 	return true;
 }
 
 bool FAutomationTestBase::TestEqual(const TCHAR* What, FStringView Actual, FStringView Expected)
 {
-	if (Actual.Compare(Expected, ESearchCase::CaseSensitive))
+	if (Actual.Compare(Expected, ESearchCase::IgnoreCase) != 0)
 	{
-		AddError(FString::Printf(TEXT("Expected '%s' to be \"%.*s\", but it was \"%.*s\"."), What, Expected.Len(), Expected.GetData(), Actual.Len(), Actual.GetData()), 1);
+		AddError(FString::Printf(TEXT("Expected '%s' to be %s, but it was %s."), What, *GetStringValueToDisplay(Expected), *GetStringValueToDisplay(Actual)), 1);
 		return false;
 	}
+
 	return true;
 }
 
 bool FAutomationTestBase::TestNotEqual(const TCHAR* What, const TCHAR* Actual, const TCHAR* Expected)
 {
-	if (Actual && Expected)
-	{
-		if (FCString::Strcmp(Actual, Expected) == 0)
-		{
-			AddError(FString::Printf(TEXT("Expected '%s' not to be\"%s\", but it was \"%s\"."), What, Expected, Actual), 1);
-			return false;
-		}
-	}
-	else // null exists
+	bool bAreDifferent = (Actual && Expected) ? (FCString::Stricmp(Actual, Expected) != 0) : (Actual != Expected);
+
+	if (!bAreDifferent)
 	{
-		if (Actual == Expected)
-		{
-			AddError(FString::Printf(TEXT("Expected '%s' to differ but both values were unexpectedly null"), What), 1);
-			return false;
-		}
+		AddError(FString::Printf(TEXT("Expected '%s' to differ from %s, but it was %s."), What, *GetStringValueToDisplay(Expected), *GetStringValueToDisplay(Actual)), 1);
 	}
-	return true;
+
+ 	return bAreDifferent;
 }
 
 bool FAutomationTestBase::TestNotEqual(const TCHAR* What, FUtf8StringView Actual, FUtf8StringView Expected)
 {
-	if (Actual.Compare(Expected, ESearchCase::CaseSensitive) == 0)
+	if (Actual.Compare(Expected, ESearchCase::IgnoreCase) == 0)
 	{
-		AddError(FString::Printf(TEXT("Expected '%s' not to be \"%.*hs\", but it was \"%.*hs\"."), What, Expected.Len(), Expected.GetData(), Actual.Len(), Actual.GetData()), 1);
+		AddError(FString::Printf(TEXT("Expected '%s' to differ from %s, but it was %s."), What, *GetStringValueToDisplay(Expected), *GetStringValueToDisplay(Actual)), 1);
 		return false;
 	}
+
 	return true;
 }
 
 bool FAutomationTestBase::TestNotEqual(const TCHAR* What, FStringView Actual, FStringView Expected)
 {
-	if (Actual.Compare(Expected, ESearchCase::CaseSensitive) == 0)
+	if (Actual.Compare(Expected, ESearchCase::IgnoreCase) == 0)
 	{
-		AddError(FString::Printf(TEXT("Expected '%s' not to be \"%.*s\", but it was \"%.*s\"."), What, Expected.Len(), Expected.GetData(), Actual.Len(), Actual.GetData()), 1);
+		AddError(FString::Printf(TEXT("Expected '%s' to differ from %s, but it was %s."), What, *GetStringValueToDisplay(Expected), *GetStringValueToDisplay(Actual)), 1);
 		return false;
 	}
+
 	return true;
 }
 
-
 bool FAutomationTestBase::TestNotEqual(const TCHAR* What, const float Actual, const float Expected, float Tolerance)
 {
 	if (FMath::IsNearlyEqual(Actual, Expected, Tolerance))
@@ -1951,94 +1971,99 @@ bool FAutomationTestBase::TestNotEqual(const TCHAR* What, const double Actual, c
 
 bool FAutomationTestBase::TestEqualInsensitive(const TCHAR* What, const TCHAR* Actual, const TCHAR* Expected)
 {
-	if (FCString::Stricmp(Actual, Expected) != 0)
-	{
-		AddError(FString::Printf(TEXT("Expected '%s' to be \"%s\", but it was \"%s\"."), What, Expected, Actual), 1);
-		return false;
-	}
-	return true;
+	return TestEqual(What, Actual, Expected);
 }
 
 bool FAutomationTestBase::TestEqualInsensitive(const TCHAR* What, FStringView Actual, FStringView Expected)
 {
-	if (Actual.Compare(Expected, ESearchCase::IgnoreCase) != 0)
-	{
-		AddError(FString::Printf(TEXT("Expected '%s' to be \"%.*s\", but it was \"%.*s\"."), What, Expected.Len(), Expected.GetData(), Actual.Len(), Actual.GetData()), 1);
-		return false;
-	}
-	return true;
+	return TestEqual(What, Actual, Expected);
 }
 
 bool FAutomationTestBase::TestEqualInsensitive(const TCHAR* What, FUtf8StringView Actual, FUtf8StringView Expected)
 {
-	if (Actual.Compare(Expected, ESearchCase::IgnoreCase) != 0)
+	return TestEqual(What, Actual, Expected);
+}
+
+bool FAutomationTestBase::TestNotEqualInsensitive(const TCHAR* What, const TCHAR* Actual, const TCHAR* Expected)
+{
+	return TestNotEqual(What, Actual, Expected);
+}
+
+bool FAutomationTestBase::TestNotEqualInsensitive(const TCHAR* What, FStringView Actual, FStringView Expected)
+{
+	return TestNotEqual(What, Actual, Expected);
+}
+
+bool FAutomationTestBase::TestNotEqualInsensitive(const TCHAR* What, FUtf8StringView Actual, FUtf8StringView Expected)
+{
+	return TestNotEqual(What, Actual, Expected);
+}
+
+bool FAutomationTestBase::TestEqualSensitive(const TCHAR* What, const TCHAR* Actual, const TCHAR* Expected)
+{
+ 	bool bAreEqual = (Actual && Expected) ? (FCString::Strcmp(Actual, Expected) == 0) : (Actual == Expected);
+ 
+ 	if (!bAreEqual)
+ 	{
+ 		AddError(FString::Printf(TEXT("Expected '%s' to be %s, but it was %s."), What, *GetStringValueToDisplay(Expected), *GetStringValueToDisplay(Actual)), 1);
+ 	}
+ 
+ 	return bAreEqual;
+}
+
+bool FAutomationTestBase::TestEqualSensitive(const TCHAR* What, FStringView Actual, FStringView Expected)
+{
+	if (Actual.Compare(Expected, ESearchCase::CaseSensitive) != 0)
 	{
-		AddError(FString::Printf(TEXT("Expected '%s' to be \"%.*hs\", but it was \"%.*hs\"."), What, Expected.Len(), Expected.GetData(), Actual.Len(), Actual.GetData()), 1);
+		AddError(FString::Printf(TEXT("Expected '%s' to be %s, but it was %s."), What, *GetStringValueToDisplay(Expected), *GetStringValueToDisplay(Actual)), 1);
 		return false;
 	}
+
 	return true;
 }
 
-bool FAutomationTestBase::TestNotEqualInsensitive(const TCHAR* What, const TCHAR* Actual, const TCHAR* Expected)
+bool FAutomationTestBase::TestEqualSensitive(const TCHAR* What, FUtf8StringView Actual, FUtf8StringView Expected)
 {
-	if (Actual && Expected)
+	if (Actual.Compare(Expected, ESearchCase::CaseSensitive) != 0)
 	{
-		if (FCString::Stricmp(Actual, Expected) == 0)
-		{
-			AddError(FString::Printf(TEXT("Expected '%s' to differ from \"%s\", but it was \"%s\"."), What, Expected, Actual), 1);
-			return false;
-		}
-	}
-	else // null exists
-	{ 
-		if (Actual == Expected)
-		{
-			AddError(FString::Printf(TEXT("Expected '%s' to differ but both values were unexpectedly null"), What), 1);
-			return false;
-		}
+		AddError(FString::Printf(TEXT("Expected '%s' to be %s, but it was %s."), What, *GetStringValueToDisplay(Expected), *GetStringValueToDisplay(Actual)), 1);
+		return false;
 	}
+
 	return true;
 }
 
-bool FAutomationTestBase::TestNotEqualInsensitive(const TCHAR* What, FStringView Actual, FStringView Expected)
+bool FAutomationTestBase::TestNotEqualSensitive(const TCHAR* What, const TCHAR* Actual, const TCHAR* Expected)
 {
-	if (Actual.GetData() && Expected.GetData())
-	{
-		if (Actual.Compare(Expected, ESearchCase::IgnoreCase) == 0)
-		{
-			AddError(FString::Printf(TEXT("Expected '%s' to differ from \"%.*s\", but it was \"%.*s\"."), What, Expected.Len(), Expected.GetData(), Actual.Len(), Actual.GetData()), 1);
-			return false;
-		}
-	}
-	else // null exists
+ 	bool bAreDifferent = (Actual && Expected) ? (FCString::Strcmp(Actual, Expected) != 0) : (Actual != Expected);
+ 
+ 	if (!bAreDifferent)
+ 	{
+ 		AddError(FString::Printf(TEXT("Expected '%s' to differ from %s, but it was %s."), What, *GetStringValueToDisplay(Expected), *GetStringValueToDisplay(Actual)), 1);
+ 	}
+ 
+ 	return bAreDifferent;
+}
+
+bool FAutomationTestBase::TestNotEqualSensitive(const TCHAR* What, FStringView Actual, FStringView Expected)
+{
+	if (Actual.Compare(Expected, ESearchCase::CaseSensitive) == 0)
 	{
-		if (Actual == Expected)
-		{
-			AddError(FString::Printf(TEXT("Expected '%s' to differ but both values were unexpectedly null"), What), 1);
-			return false;
-		}
+		AddError(FString::Printf(TEXT("Expected '%s' to differ from %s, but it was %s."), What, *GetStringValueToDisplay(Expected), *GetStringValueToDisplay(Actual)), 1);
+		return false;
 	}
+
 	return true;
 }
 
-bool FAutomationTestBase::TestNotEqualInsensitive(const TCHAR* What, FUtf8StringView Actual, FUtf8StringView Expected)
+bool FAutomationTestBase::TestNotEqualSensitive(const TCHAR* What, FUtf8StringView Actual, FUtf8StringView Expected)
 {
-	if (Actual.GetData() && Expected.GetData())
-	{
-		if (Actual.Compare(Expected, ESearchCase::IgnoreCase) == 0)
-		{
-			AddError(FString::Printf(TEXT("Expected '%s' to differ from \"%.*hs\", but it was \"%.*hs\"."), What, Expected.Len(), Expected.GetData(), Actual.Len(), Actual.GetData()), 1);
-			return false;
-		}
-	}
-	else // null exists
+	if (Actual.Compare(Expected, ESearchCase::CaseSensitive) == 0)
 	{
-		if (Actual == Expected)
-		{
-			AddError(FString::Printf(TEXT("Expected '%s' to differ but both values were unexpectedly null"), What), 1);
-			return false;
-		}
+		AddError(FString::Printf(TEXT("Expected '%s' to differ from %s, but it was %s."), What, *GetStringValueToDisplay(Expected), *GetStringValueToDisplay(Actual)), 1);
+		return false;
 	}
+
 	return true;
 }
 
diff --git a/Engine/Source/Runtime/Core/Public/Misc/AutomationTest.h b/Engine/Source/Runtime/Core/Public/Misc/AutomationTest.h
--- a/Engine/Source/Runtime/Core/Public/Misc/AutomationTest.h
+++ b/Engine/Source/Runtime/Core/Public/Misc/AutomationTest.h
@@ -1944,13 +1944,27 @@ public:
 	CORE_API bool TestNotEqual(const TCHAR* What, const TCHAR* Actual, const TCHAR* Expected);
 	CORE_API bool TestNotEqual(const TCHAR* What, FStringView Actual, FStringView Expected);
 	CORE_API bool TestNotEqual(const TCHAR* What, FUtf8StringView Actual, FUtf8StringView Expected);
+
+	UE_DEPRECATED(5.5, "Use TestEqual instead (string tests are case insensitive by default)")
 	CORE_API bool TestEqualInsensitive(const TCHAR* What, const TCHAR* Actual, const TCHAR* Expected);
+	UE_DEPRECATED(5.5, "Use TestEqual instead (string tests are case insensitive by default)")
 	CORE_API bool TestEqualInsensitive(const TCHAR* What, FStringView Actual, FStringView Expected);
+	UE_DEPRECATED(5.5, "Use TestEqual instead (string tests are case insensitive by default)")
 	CORE_API bool TestEqualInsensitive(const TCHAR* What, FUtf8StringView Actual, FUtf8StringView Expected);
+	UE_DEPRECATED(5.5, "Use TestNotEqual instead (string tests are case insensitive by default)")
 	CORE_API bool TestNotEqualInsensitive(const TCHAR* What, const TCHAR* Actual, const TCHAR* Expected);
+	UE_DEPRECATED(5.5, "Use TestNotEqual instead (string tests are case insensitive by default)")
 	CORE_API bool TestNotEqualInsensitive(const TCHAR* What, FStringView Actual, FStringView Expected);
+	UE_DEPRECATED(5.5, "Use TestNotEqual instead (string tests are case insensitive by default)")
 	CORE_API bool TestNotEqualInsensitive(const TCHAR* What, FUtf8StringView Actual, FUtf8StringView Expected);
 
+	CORE_API bool TestEqualSensitive(const TCHAR* What, const TCHAR* Actual, const TCHAR* Expected);
+	CORE_API bool TestEqualSensitive(const TCHAR* What, FStringView Actual, FStringView Expected);
+	CORE_API bool TestEqualSensitive(const TCHAR* What, FUtf8StringView Actual, FUtf8StringView Expected);
+	CORE_API bool TestNotEqualSensitive(const TCHAR* What, const TCHAR* Actual, const TCHAR* Expected);
+	CORE_API bool TestNotEqualSensitive(const TCHAR* What, FStringView Actual, FStringView Expected);
+	CORE_API bool TestNotEqualSensitive(const TCHAR* What, FUtf8StringView Actual, FUtf8StringView Expected);
+
 	bool TestEqual(const FString& What, const int32 Actual, const int32 Expected)
 	{
 		return TestEqual(*What, Actual, Expected);
@@ -1991,28 +2005,86 @@ public:
 		return TestEqual(*What, Actual, Expected);
 	}
 
-	bool TestEqual(const FString& What, FUtf8StringView Actual, FUtf8StringView  Expected)
+	bool TestEqual(const FString& What, FUtf8StringView Actual, FUtf8StringView Expected)
 	{
-		return TestEqualInsensitive(*What, Actual, Expected);
+		return TestEqual(*What, Actual, Expected);
 	}
 
 	bool TestEqual(const FString& What, FStringView Actual, FStringView Expected)
 	{
-		return TestEqualInsensitive(*What, Actual, Expected);
+		return TestEqual(*What, Actual, Expected);
 	}
 
 	bool TestEqual(const FString& What, FStringView Actual, FUtf8StringView Expected)
 	{
 		// This overload is here because there are some tests that have raw string literals as the expected value.
 		FString Tmp(Expected);
-		return TestEqualInsensitive(*What, Actual, Tmp);
+		return TestEqual(*What, Actual, Tmp);
 	}
 
 	bool TestEqual(const FString& What, FUtf8StringView Actual, FStringView Expected)
 	{
 		// This overload is here because there are some tests that have raw string literals as the Actual value.
 		FString Tmp(Actual);
-		return TestEqualInsensitive(*What, Tmp, Expected);
+		return TestEqual(*What, Tmp, Expected);
+	}
+
+	bool TestEqualSensitive(const FString& What, const TCHAR* Actual, const TCHAR* Expected)
+	{
+		return TestEqualSensitive(*What, Actual, Expected);
+	}
+
+	bool TestEqualSensitive(const FString& What, FUtf8StringView Actual, FUtf8StringView Expected)
+	{
+		return TestEqualSensitive(*What, Actual, Expected);
+	}
+
+	bool TestEqualSensitive(const FString& What, FStringView Actual, FStringView Expected)
+	{
+		return TestEqualSensitive(*What, Actual, Expected);
+	}
+
+	bool TestEqualSensitive(const FString& What, FStringView Actual, FUtf8StringView Expected)
+	{
+		// This overload is here because there are some tests that have raw string literals as the expected value.
+		FString Tmp(Expected);
+		return TestEqualSensitive(*What, Actual, Tmp);
+	}
+
+	bool TestEqualSensitive(const FString& What, FUtf8StringView Actual, FStringView Expected)
+	{
+		// This overload is here because there are some tests that have raw string literals as the Actual value.
+		FString Tmp(Actual);
+		return TestEqualSensitive(*What, Tmp, Expected);
+	}
+
+	bool TestNotEqualSensitive(const FString& What, const TCHAR* Actual, const TCHAR* Expected)
+	{
+		return TestNotEqualSensitive(*What, Actual, Expected);
+	}
+
+	bool TestNotEqualSensitive(const FString& What, FUtf8StringView Actual, FUtf8StringView Expected)
+	{
+		return TestNotEqualSensitive(*What, Actual, Expected);
+	}
+
+	bool TestNotEqualSensitive(const FString& What, FStringView Actual, FStringView Expected)
+	{
+		return TestNotEqualSensitive(*What, Actual, Expected);
+	}
+
+	bool TestNotEqualSensitive(const FString& What, FStringView Actual, FUtf8StringView Expected)
+	{
+		// This overload is here because there are some tests that have raw string literals as the expected value.
+		FString Tmp(Expected);
+		return TestNotEqualSensitive(*What, Actual, Tmp);
+	}
+
+	bool TestNotEqualSensitive(const FString& What, FUtf8StringView Actual, FStringView Expected)
+	{
+		// This overload is here because there are some tests that have raw string literals as the Actual value.
+		FString Tmp(Actual);
+		return TestNotEqualSensitive(*What, Tmp, Expected);
 	}
 
 	/**
@@ -2513,6 +2585,14 @@ private:
 	 */
 	CORE_API void InternalSetSuccessState(bool bSuccessful);
 
+	/**
+	 * Returns a string description of the value stored in the passed string view in a form suitable for error messages.
+	 * Proper string values will be output with quotation marks e.g. "value".
+	 * Null will be represented as nullptr.
+	 */
+	FString GetStringValueToDisplay(FStringView Value) const;
+	FString GetStringValueToDisplay(FUtf8StringView Value) const;
+
 	/* Log messages to be expected while processing this test.*/
 	TSet<FAutomationExpectedMessage> ExpectedMessages;
 
@@ -4227,25 +4307,49 @@ public: \
 	}
 
 #define UTEST_EQUAL_INSENSITIVE(What, Actual, Expected)\
-	if (!TestEqualInsensitive(What, Actual, Expected))\
+	if (!TestEqual(What, Actual, Expected))\
 	{\
 		return false;\
 	}
 
 #define UTEST_EQUAL_INSENSITIVE_EXPR(Actual, Expected)\
-	if (!TestEqualInsensitive(TEXT(#Actual), Actual, Expected))\
+	if (!TestEqual(TEXT(#Actual), Actual, Expected))\
 	{\
 		return false;\
 	}
 
 #define UTEST_NOT_EQUAL_INSENSITIVE(What, Actual, Expected)\
-	if (!TestNotEqualInsensitive(What, Actual, Expected))\
+	if (!TestNotEqual(What, Actual, Expected))\
 	{\
 		return false;\
 	}
 
 #define UTEST_NOT_EQUAL_INSENSITIVE_EXPR(Actual, Expected)\
-	if (!TestNotEqualInsensitive(TEXT(#Actual), Actual, Expected))\
+	if (!TestNotEqual(TEXT(#Actual), Actual, Expected))\
+	{\
+		return false;\
+	}
+
+#define UTEST_EQUAL_SENSITIVE(What, Actual, Expected)\
+	if (!TestEqualSensitive(What, Actual, Expected))\
+	{\
+		return false;\
+	}
+
+#define UTEST_EQUAL_SENSITIVE_EXPR(Actual, Expected)\
+	if (!TestEqualSensitive(TEXT(#Actual), Actual, Expected))\
+	{\
+		return false;\
+	}
+
+#define UTEST_NOT_EQUAL_SENSITIVE(What, Actual, Expected)\
+	if (!TestNotEqualSensitive(What, Actual, Expected))\
+	{\
+		return false;\
+	}
+
+#define UTEST_NOT_EQUAL_SENSITIVE_EXPR(Actual, Expected)\
+	if (!TestNotEqualSensitive(TEXT(#Actual), Actual, Expected))\
 	{\
 		return false;\
 	}
diff --git a/Engine/Source/Runtime/Core/Public/Misc/LowLevelTestAdapter.h b/Engine/Source/Runtime/Core/Public/Misc/LowLevelTestAdapter.h
--- a/Engine/Source/Runtime/Core/Public/Misc/LowLevelTestAdapter.h
+++ b/Engine/Source/Runtime/Core/Public/Misc/LowLevelTestAdapter.h
@@ -140,7 +140,9 @@ FString CaptureExpressionsAndValues(const FString& InExpressions, ArgTypes&&...
 #define STATIC_CHECK_FALSE(...) static_assert(!(__VA_ARGS__), "!(" #__VA_ARGS__ ")");
 
 #define CHECK_EQUALS(What, X, Y) FAutomationTestFramework::Get().GetCurrentTest()->TestEqual(What, X, Y);
+#define CHECK_EQUALS_SENSITIVE(What, X, Y) FAutomationTestFramework::Get().GetCurrentTest()->TestEqualSensitive(What, X, Y);
 #define CHECK_NOT_EQUALS(What, X, Y) FAutomationTestFramework::Get().GetCurrentTest()->TestNotEqual(What, X, Y);
+#define CHECK_NOT_EQUALS_SENSITIVE(What, X, Y) FAutomationTestFramework::Get().GetCurrentTest()->TestNotEqualSensitive(What, X, Y);
 
 #define SECTION(Text) FAutomationTestFramework::Get().GetCurrentTest()->AddInfo(TEXT(Text));
 #define FAIL_CHECK(Message) FAutomationTestFramework::Get().GetCurrentTest()->AddError(Message);
diff --git a/Engine/Source/Runtime/Engine/Private/Tests/EngineAutomationTests.cpp b/Engine/Source/Runtime/Engine/Private/Tests/EngineAutomationTests.cpp
--- a/Engine/Source/Runtime/Engine/Private/Tests/EngineAutomationTests.cpp
+++ b/Engine/Source/Runtime/Engine/Private/Tests/EngineAutomationTests.cpp
@@ -1245,6 +1245,7 @@ bool FAutomationTestInequalityString::RunTest(const FString& Parameters)
 {
 	FString ExpectedString("Forty-two");
 	const TCHAR* ExpectedCharPtr = TEXT("Forty-two");
+	const UTF8CHAR* ExpectedCharPtrUtf8 = UTF8TEXT("Forty-two");
 	FString IdenticalString("Forty-two");
 	const TCHAR* IdenticalCharPtr = TEXT("Forty-two");
 	FString LowercaseString("forty-two");
@@ -1256,6 +1257,7 @@ bool FAutomationTestInequalityString::RunTest(const FString& Parameters)
 	FString DifferentString("42");
 	const TCHAR* DifferentCharPtr = TEXT("42");
 	const TCHAR* NullCharPtr = nullptr;
+	const UTF8CHAR* NullCharPtrUtf8 = nullptr;
 
 	TestEqual(TEXT("String identity equal"), IdenticalString, ExpectedString);
 	TestEqual(TEXT("char* identity equal"), IdenticalCharPtr, ExpectedCharPtr);
@@ -1270,16 +1272,88 @@ bool FAutomationTestInequalityString::RunTest(const FString& Parameters)
 	TestNotEqual(TEXT("char* unequal empty"), EmptyCharPtr, ExpectedCharPtr);
 	TestNotEqual(TEXT("char* unequal null"), NullCharPtr, ExpectedCharPtr);
 
-	TestEqualInsensitive(TEXT("String insensitive equal identity"), IdenticalString, ExpectedString);
-	TestEqualInsensitive(TEXT("char* insensitive equal identity"), IdenticalCharPtr, ExpectedCharPtr);
-	TestEqualInsensitive(TEXT("String insensitive equal lower"), LowercaseString, ExpectedString);
-	TestEqualInsensitive(TEXT("char* insensitive equal lower"), LowercaseCharPtr, ExpectedCharPtr);
-	TestEqualInsensitive(TEXT("String insensitive equal upper"), UppercaseString, ExpectedString);
-	TestEqualInsensitive(TEXT("char* insensitive equal upper"), UppercaseCharPtr, ExpectedCharPtr);
-	TestNotEqualInsensitive(TEXT("String insensitive unequal"), DifferentString, ExpectedString);
-	TestNotEqualInsensitive(TEXT("char* insensitive unequal"), DifferentCharPtr, ExpectedCharPtr);
-	TestNotEqualInsensitive(TEXT("char* insensitive unequal null"), NullCharPtr, ExpectedCharPtr);
-	
+	TestEqual(TEXT("String insensitive equal identity"), IdenticalString, ExpectedString);
+	TestEqual(TEXT("char* insensitive equal identity"), IdenticalCharPtr, ExpectedCharPtr);
+	TestEqual(TEXT("String insensitive equal lower"), LowercaseString, ExpectedString);
+	TestEqual(TEXT("char* insensitive equal lower"), LowercaseCharPtr, ExpectedCharPtr);
+	TestEqual(TEXT("String insensitive equal upper"), UppercaseString, ExpectedString);
+	TestEqual(TEXT("char* insensitive equal upper"), UppercaseCharPtr, ExpectedCharPtr);
+	TestNotEqual(TEXT("String insensitive unequal"), DifferentString, ExpectedString);
+	TestNotEqual(TEXT("char* insensitive unequal"), DifferentCharPtr, ExpectedCharPtr);
+	TestNotEqual(TEXT("char* insensitive unequal null"), NullCharPtr, ExpectedCharPtr);
+
+	TestEqualSensitive(TEXT("String sensitive equal identity"), IdenticalString, ExpectedString);
+	TestEqualSensitive(TEXT("char* sensitive equal identity"), IdenticalCharPtr, ExpectedCharPtr);	
+	TestNotEqualSensitive(TEXT("String sensitive unequal lower"), LowercaseString, ExpectedString);
+	TestNotEqualSensitive(TEXT("char* sensitive unequal lower"), LowercaseCharPtr, ExpectedCharPtr);
+	TestNotEqualSensitive(TEXT("String sensitive unequal upper"), UppercaseString, ExpectedString);
+	TestNotEqualSensitive(TEXT("char* sensitive unequal upper"), UppercaseCharPtr, ExpectedCharPtr);	
+	TestNotEqualSensitive(TEXT("String sensitive unequal"), DifferentString, ExpectedString);
+	TestNotEqualSensitive(TEXT("char* sensitive unequal"), DifferentCharPtr, ExpectedCharPtr);
+	TestNotEqualSensitive(TEXT("char* sensitive unequal null"), NullCharPtr, ExpectedCharPtr);
+
+	return true;
+}
+
+IMPLEMENT_SIMPLE_AUTOMATION_TEST(FAutomationTestInequalityStringNulls, "TestFramework.Validation.TestInequalityStringNulls", EAutomationTestFlags_ApplicationContextMask | EAutomationTestFlags::NegativeFilter);
+bool FAutomationTestInequalityStringNulls::RunTest(const FString& Parameters)
+{
+	const TCHAR* ExpectedCharPtr = TEXT("Forty-two");
+	const UTF8CHAR* ExpectedCharPtrUtf8 = UTF8TEXT("Forty-two");
+	const TCHAR* NullCharPtr = nullptr;
+	const UTF8CHAR* NullCharPtrUtf8 = nullptr;
+
+	// TCHAR*
+	TestEqual(TEXT("char* equal null null"), NullCharPtr, NullCharPtr);
+	TestEqual(TEXT("char* equal null string"), NullCharPtr, ExpectedCharPtr);
+	TestEqual(TEXT("char* equal string null"), ExpectedCharPtr, NullCharPtr);
+
+	TestNotEqual(TEXT("char* unequal null null"), NullCharPtr, NullCharPtr);
+	TestNotEqual(TEXT("char* unequal null string"), NullCharPtr, ExpectedCharPtr);
+	TestNotEqual(TEXT("char* unequal string null"), ExpectedCharPtr, NullCharPtr);
+
+	TestEqualSensitive(TEXT("char* equal(sensitive) null null"), NullCharPtr, NullCharPtr);
+	TestEqualSensitive(TEXT("char* equal(sensitive) null string"), NullCharPtr, ExpectedCharPtr);
+	TestEqualSensitive(TEXT("char* equal(sensitive) string null"), ExpectedCharPtr, NullCharPtr);
+
+	TestNotEqualSensitive(TEXT("char* unequal(sensitive) null null"), NullCharPtr, NullCharPtr);
+	TestNotEqualSensitive(TEXT("char* unequal(sensitive) null string"), NullCharPtr, ExpectedCharPtr);
+	TestNotEqualSensitive(TEXT("char* unequal(sensitive) string null"), ExpectedCharPtr, NullCharPtr);
+
+	// TStringView
+	TestEqual(TEXT("stringview equal null null"), MakeStringView(NullCharPtr), MakeStringView(NullCharPtr));
+	TestEqual(TEXT("stringview equal null string"), MakeStringView(NullCharPtr), MakeStringView(ExpectedCharPtr));
+	TestEqual(TEXT("stringview equal string null"), MakeStringView(ExpectedCharPtr), MakeStringView(NullCharPtr));
+
+	TestNotEqual(TEXT("stringview unequal null null"), MakeStringView(NullCharPtr), MakeStringView(NullCharPtr));
+	TestNotEqual(TEXT("stringview unequal null string"), MakeStringView(NullCharPtr), MakeStringView(ExpectedCharPtr));
+	TestNotEqual(TEXT("stringview unequal string null"), MakeStringView(ExpectedCharPtr), MakeStringView(NullCharPtr));
+
+	TestEqualSensitive(TEXT("stringview equal(sensitive) null null"), MakeStringView(NullCharPtr), MakeStringView(NullCharPtr));
+	TestEqualSensitive(TEXT("stringview equal(sensitive) null string"), MakeStringView(NullCharPtr), MakeStringView(ExpectedCharPtr));
+	TestEqualSensitive(TEXT("stringview equal(sensitive) string null"), MakeStringView(ExpectedCharPtr), MakeStringView(NullCharPtr));
+
+	TestNotEqualSensitive(TEXT("stringview unequal(sensitive) null null"), MakeStringView(NullCharPtr), MakeStringView(NullCharPtr));
+	TestNotEqualSensitive(TEXT("stringview unequal(sensitive) null string"), MakeStringView(NullCharPtr), MakeStringView(ExpectedCharPtr));
+	TestNotEqualSensitive(TEXT("stringview unequal(sensitive) string null"), MakeStringView(ExpectedCharPtr), MakeStringView(NullCharPtr));
+
+	// FUtf8StringView
+	TestEqual(TEXT("stringview8 equal null null"), MakeStringView(NullCharPtrUtf8), MakeStringView(NullCharPtrUtf8));
+	TestEqual(TEXT("stringview8 equal null string"), MakeStringView(NullCharPtrUtf8), MakeStringView(ExpectedCharPtrUtf8));
+	TestEqual(TEXT("stringview8 equal string null"), MakeStringView(ExpectedCharPtrUtf8), MakeStringView(NullCharPtrUtf8));
+
+	TestNotEqual(TEXT("stringview8 unequal null null"), MakeStringView(NullCharPtrUtf8), MakeStringView(NullCharPtrUtf8));
+	TestNotEqual(TEXT("stringview8 unequal null string"), MakeStringView(NullCharPtrUtf8), MakeStringView(ExpectedCharPtrUtf8));
+	TestNotEqual(TEXT("stringview8 unequal string null"), MakeStringView(ExpectedCharPtrUtf8), MakeStringView(NullCharPtrUtf8));
+
+	TestEqualSensitive(TEXT("stringview8 equal(sensitive) null null"), MakeStringView(NullCharPtrUtf8), MakeStringView(NullCharPtrUtf8));
+	TestEqualSensitive(TEXT("stringview8 equal(sensitive) null string"), MakeStringView(NullCharPtrUtf8), MakeStringView(ExpectedCharPtrUtf8));
+	TestEqualSensitive(TEXT("stringview8 equal(sensitive) string null"), MakeStringView(ExpectedCharPtrUtf8), MakeStringView(NullCharPtrUtf8));
+
+	TestNotEqualSensitive(TEXT("stringview8 unequal(sensitive) null null"), MakeStringView(NullCharPtrUtf8), MakeStringView(NullCharPtrUtf8));
+	TestNotEqualSensitive(TEXT("stringview8 unequal(sensitive) null string"), MakeStringView(NullCharPtrUtf8), MakeStringView(ExpectedCharPtrUtf8));
+	TestNotEqualSensitive(TEXT("stringview8 unequal(sensitive) string null"), MakeStringView(ExpectedCharPtrUtf8), MakeStringView(NullCharPtrUtf8));
+
 	return true;
 }
 
@@ -1302,6 +1376,7 @@ protected:
 	static const float ExpectedFloatValueLess;
 	static const float ExpectedFloatValueGreater;
 	static const FString ActualFStringValue;
+	static const FString ActualFStringValueCopy;
 	static const FString ExpectedFStringValueLowerCase;
 	static const FString UnexpectedFStringValueLowerCase;
 	static const FString CustomDescriptionString;
@@ -1316,6 +1391,7 @@ const float FAutomationUTestMacrosExpr::ExpectedFloatValueOutOfToleranceNegative
 const float FAutomationUTestMacrosExpr::ExpectedFloatValueLess(ActualFloatValue + (PositiveToleranceFloat*2)); //actual < expected
 const float FAutomationUTestMacrosExpr::ExpectedFloatValueGreater(ActualFloatValue - (PositiveToleranceFloat*2)); //actual > expected
 const FString FAutomationUTestMacrosExpr::ActualFStringValue(TEXT("EQUALS"));
+const FString FAutomationUTestMacrosExpr::ActualFStringValueCopy(TEXT("EQUALS"));
 const FString FAutomationUTestMacrosExpr::ExpectedFStringValueLowerCase(TEXT("equals"));
 const FString FAutomationUTestMacrosExpr::UnexpectedFStringValueLowerCase(TEXT("not-equals"));
 const FString FAutomationUTestMacrosExpr::CustomDescriptionString(TEXT("Error string appears when UTEST_ macro diverges from _EXPR variant"));
@@ -1336,6 +1412,10 @@ bool FAutomationEqualEXPR::RunTest(const FString& Parameters)
 	UTEST_EQUAL_INSENSITIVE(*CustomDescriptionString, *ActualFStringValue, *ExpectedFStringValueLowerCase);
 	UTEST_NOT_EQUAL_INSENSITIVE_EXPR(*ActualFStringValue, *UnexpectedFStringValueLowerCase);
 	UTEST_NOT_EQUAL_INSENSITIVE(*CustomDescriptionString, *ActualFStringValue, *UnexpectedFStringValueLowerCase);
+	UTEST_EQUAL_SENSITIVE_EXPR(*ActualFStringValue, *ActualFStringValueCopy);
+	UTEST_EQUAL_SENSITIVE(*CustomDescriptionString, *ActualFStringValue, *ActualFStringValueCopy);
+	UTEST_NOT_EQUAL_SENSITIVE_EXPR(*ActualFStringValue, *ExpectedFStringValueLowerCase);
+	UTEST_NOT_EQUAL_SENSITIVE(*CustomDescriptionString, *ActualFStringValue, *ExpectedFStringValueLowerCase);
 
 	return true;
 }
