[AutoRTFM] Update LockFreeFixedSizeAllocator to support AutoRTFM transactions.

This code is reachable from Verse by way of skeletal meshes/bones.

diff --git a/Engine/Source/Runtime/Core/Public/Containers/LockFreeFixedSizeAllocator.h b/Engine/Source/Runtime/Core/Public/Containers/LockFreeFixedSizeAllocator.h
--- a/Engine/Source/Runtime/Core/Public/Containers/LockFreeFixedSizeAllocator.h
+++ b/Engine/Source/Runtime/Core/Public/Containers/LockFreeFixedSizeAllocator.h
@@ -3,6 +3,7 @@
 #pragma once
 
 #include "HAL/MallocLeakDetection.h"
+#include "AutoRTFM/AutoRTFM.h"
 #include "Misc/AssertionMacros.h"
 #include "HAL/UnrealMemory.h"
 #include "Misc/NoopCounter.h"
@@ -10,8 +11,8 @@
 
 /**
 * Thread safe, lock free pooling allocator of fixed size blocks that
-* never returns free space, even at shutdown
-* alignment isn't handled, assumes FMemory::Malloc will work
+* never returns free space, even at shutdown.
+* Alignment isn't handled; assumes FMemory::Malloc will work.
 */
 
 #ifndef USE_NAIVE_TLockFreeFixedSizeAllocator_TLSCacheBase
@@ -189,19 +190,29 @@ private:
 
 /**
  * Thread safe, lock free pooling allocator of fixed size blocks that
- * never returns free space until program shutdown.
- * alignment isn't handled, assumes FMemory::Malloc will work
+ * only returns free space when the allocator is destroyed.
+ * Alignment isn't handled; assumes FMemory::Malloc will work.
  */
 template<int32 SIZE, int TPaddingForCacheContention, typename TTrackingCounter = FNoopCounter>
 class TLockFreeFixedSizeAllocator
 {
 public:
-
 	/** Destructor, returns all memory via FMemory::Free **/
 	~TLockFreeFixedSizeAllocator()
 	{
-		check(!NumUsed.GetValue());
-		Trim();
+		UE_AUTORTFM_OPEN2
+		{
+			check(!NumUsed.GetValue());
+			Trim();
+		};
+
+		// If we are in a transaction and this allocator exists on the transaction's stack, 
+		// running the ONABORT handler after destruction is dangerous and could stomp memory
+		// that belongs to someone else.
+		if (AutoRTFM::IsClosed() && AutoRTFM::IsInnerTransactionStack(this))
+		{
+			AutoRTFM::PopAllOnAbortHandlers(this);
+		}
 	}
 
 	/**
@@ -212,8 +223,24 @@ public:
 	 */
 	void* Allocate()
 	{
-		NumUsed.Increment();
-		void *Memory = FreeList.Pop();
+		UE_AUTORTFM_OPEN2
+		{
+			NumUsed.Increment();
+		};
+		AutoRTFM::PushOnAbortHandler(this, [this]() 
+		{
+			// TODO: investigate reusing the same OnAbort handler for multiple allocations/frees.
+			NumUsed.Decrement();
+		});
+		// When we are in an AutoRTFM transaction, it is simpler and faster to allocate a new block instead
+		// of reusing an existing one, since we don't need to track changes on newly allocated memory.
+		// If the transaction is aborted, AutoRTFM should clean it up for us automatically.
+		if (AutoRTFM::IsClosed())
+		{
+			return FMemory::Malloc(SIZE);
+		}
+		// Outside of a transaction, we prefer to reuse blocks from the FreeList.
+		void* Memory = FreeList.Pop();
 		if (Memory)
 		{
 			NumFree.Decrement();
@@ -231,23 +258,45 @@ public:
 	 * @param Item The item to free.
 	 * @see Allocate
 	 */
-	void Free(void *Item)
+	void Free(void* Item)
 	{
-		NumUsed.Decrement();
-		FreeList.Push(Item);
-		NumFree.Increment();
+		UE_AUTORTFM_OPEN2
+		{
+			NumUsed.Decrement();
+		};
+		AutoRTFM::PushOnAbortHandler(this, [this]()
+		{
+			// TODO: investigate reusing the same OnAbort handler for multiple allocations/frees.
+			NumUsed.Increment();
+		});
+		if (AutoRTFM::IsClosed())
+		{
+			// When we are in an AutoRTFM transaction, it is simpler to free blocks directly.
+			// Frees are deferred until the transaction is complete, and we don't need to
+			// worry about leaks occurring if transactions are aborted at inopportune times.
+			FMemory::Free(Item);
+		}
+		else
+		{
+			// Outside of a transaction, we push blocks onto the FreeList for reuse.
+			FreeList.Push(Item);
+			NumFree.Increment();
+		}
 	}
 
 	/**
-	* Returns all free memory to the heap
+	* Returns all free memory to the heap.
 	*/
 	void Trim()
 	{
-		while (void* Mem = FreeList.Pop())
+		UE_AUTORTFM_OPEN2
 		{
-			FMemory::Free(Mem);
-			NumFree.Decrement();
-		}
+			while (void* Mem = FreeList.Pop())
+			{
+				FMemory::Free(Mem);
+				NumFree.Decrement();
+			}
+		};
 	}
 
 	/**
@@ -278,7 +327,7 @@ private:
 	TLockFreePointerListUnordered<void, TPaddingForCacheContention> FreeList;
 
 	/** Total number of blocks outstanding and not in the free list. */
-	TTrackingCounter NumUsed; 
+	TTrackingCounter NumUsed;
 
 	/** Total number of blocks in the free list. */
 	TTrackingCounter NumFree;
diff --git a/Engine/Source/Runtime/Core/Tests/Containers/AutoRTFMLockFreeFixedSizeAllocatorTest.cpp b/Engine/Source/Runtime/Core/Tests/Containers/AutoRTFMLockFreeFixedSizeAllocatorTest.cpp
new file mode 100644
--- /dev/null
+++ b/Engine/Source/Runtime/Core/Tests/Containers/AutoRTFMLockFreeFixedSizeAllocatorTest.cpp
@@ -0,0 +1,135 @@
+// Copyright Epic Games, Inc. All Rights Reserved.
+
+#include "Misc/AutomationTest.h"
+#include "Misc/App.h"
+#include "AutoRTFM/AutoRTFM.h"
+#include "HAL/Platform.h"
+#include "Containers/LockFreeFixedSizeAllocator.h"
+
+#include <vector>
+
+#if WITH_DEV_AUTOMATION_TESTS
+
+#define CHECK_EQ(A, B) \
+	do { TestEqual(TEXT(__FILE__ ":" UE_STRINGIZE(__LINE__) ": TestEqual(" #A ", " #B ")"), (A), (B)); } while (0)
+
+IMPLEMENT_SIMPLE_AUTOMATION_TEST(FAutoRTFMLockFreeFixedSizeAllocator, "AutoRTFM + LockFreeFixedSizeAllocator", \
+	EAutomationTestFlags::EngineFilter | EAutomationTestFlags::ClientContext | \
+	EAutomationTestFlags::ServerContext | EAutomationTestFlags::CommandletContext)
+
+bool FAutoRTFMLockFreeFixedSizeAllocator::RunTest(const FString& Parameters)
+{
+	if (!AutoRTFM::ForTheRuntime::IsAutoRTFMRuntimeEnabled())
+	{
+		ExecutionInfo.AddEvent(FAutomationEvent(EAutomationEventType::Info,
+			TEXT("SKIPPED 'FAutoRTFMLockFreeFixedSizeAllocator' test. AutoRTFM disabled.")));
+		return true;
+	}
+
+	constexpr int BlockSize = 32;
+	using BasicAllocator = TLockFreeFixedSizeAllocator<BlockSize, PLATFORM_CACHE_LINE_SIZE, FThreadSafeCounter>;
+
+	AutoRTFM::Transact([&]
+		{
+			// It should be safe to instantiate a lock-free fixed-size allocator in a transaction.
+			BasicAllocator Allocator;
+			CHECK_EQ(Allocator.GetNumUsed().GetValue(), 0);
+			CHECK_EQ(Allocator.GetNumFree().GetValue(), 0);
+
+			// It should be safe to allocate from a lock-free fixed-size allocator inside a transaction.
+			void* Data = Allocator.Allocate();
+
+			// It should be safe to free the allocated object while still inside a transaction.
+			Allocator.Free(Data);
+		});
+
+	// It should be safe to instantiate a lock-free fixed-size allocator outside of a transaction, then
+	// use it inside a transaction.
+	{
+		BasicAllocator Allocator;
+
+		// It should be safe to allocate inside a transaction.
+		void* Data = nullptr;
+		AutoRTFM::Transact([&]
+			{
+				Data = Allocator.Allocate();
+			});
+		CHECK_EQ(Allocator.GetNumUsed().GetValue(), 1);
+		CHECK_EQ(Allocator.GetNumFree().GetValue(), 0);
+
+		// It should be safe to free data inside a transaction. Items that are freed inside of a transaction
+		// are freed immediately, instead of being added to the free-list.
+		AutoRTFM::Transact([&]
+			{
+				Allocator.Free(Data);
+			});
+		CHECK_EQ(Allocator.GetNumUsed().GetValue(), 0);
+		CHECK_EQ(Allocator.GetNumFree().GetValue(), 0);
+	}
+
+	// It should be safe to allocate an object inside a transaction, then free it outside of the transaction.
+	// These items will be added to the free-list, like normal.
+	{
+		BasicAllocator Allocator;
+		void* Data = nullptr;
+		AutoRTFM::Transact([&]
+			{
+				Data = Allocator.Allocate();
+			});
+		Allocator.Free(Data);
+		CHECK_EQ(Allocator.GetNumUsed().GetValue(), 0);
+		CHECK_EQ(Allocator.GetNumFree().GetValue(), 1);
+	}
+
+	// It should be safe to allocate an object inside a transaction, then abort the transaction.
+	// This should cause the allocation to be rolled back automatically by AutoRTFM, not be placed on the free-list.
+	{
+		BasicAllocator Allocator;
+		AutoRTFM::Transact([&]
+			{
+				[[maybe_unused]] void* Data = Allocator.Allocate();
+				AutoRTFM::AbortTransaction();
+			});
+		CHECK_EQ(Allocator.GetNumUsed().GetValue(), 0);
+		CHECK_EQ(Allocator.GetNumFree().GetValue(), 0);
+	}
+
+	// It should be safe to allocate an object inside a transaction, free it, then abort the transaction.
+	// This should cause the allocation to be rolled back automatically by AutoRTFM, not be placed on the free-list.
+	{
+		BasicAllocator Allocator;
+		AutoRTFM::Transact([&]
+			{
+				void* Data = Allocator.Allocate();
+				Allocator.Free(Data);
+				AutoRTFM::AbortTransaction();
+			});
+		CHECK_EQ(Allocator.GetNumUsed().GetValue(), 0);
+		CHECK_EQ(Allocator.GetNumFree().GetValue(), 0);
+	}
+
+	// It should be safe to allocate an object outside of a transaction, free it inside a transaction, and then 
+	// abort the transaction. The allocator should return a use-count of one and the object should be unchanged.
+	{
+		BasicAllocator Allocator;
+		void* Data = Allocator.Allocate();
+		memset(Data, 0x44, BlockSize);
+		AutoRTFM::Transact([&]
+			{
+				Allocator.Free(Data);
+				AutoRTFM::AbortTransaction();
+			});
+		CHECK_EQ(Allocator.GetNumUsed().GetValue(), 1);
+		CHECK_EQ(Allocator.GetNumFree().GetValue(), 0);
+
+		// The bytes inside Data shouldn't have been modified at all.
+		std::vector<char> ExpectedData(BlockSize, 0x44);
+		CHECK_EQ(memcmp(Data, ExpectedData.data(), BlockSize), 0);
+
+		Allocator.Free(Data);
+	}
+
+	return true;
+}
+
+#endif  // WITH_DEV_AUTOMATION_TESTS
