Only return the real status when we are the main thread, otherwise return idle since our Context is a global singleton
Assert we are the main thread when doing a Context::Reset

diff --git a/Engine/Source/Runtime/Core/Private/AutoRTFM/Context.cpp b/Engine/Source/Runtime/Core/Private/AutoRTFM/Context.cpp
--- a/Engine/Source/Runtime/Core/Private/AutoRTFM/Context.cpp
+++ b/Engine/Source/Runtime/Core/Private/AutoRTFM/Context.cpp
@@ -494,6 +494,8 @@ void FContext::AbortByLanguageAndThrow()
 
 void FContext::Reset()
 {
+	ASSERT(CurrentThreadId == FPlatformTLS::GetCurrentThreadId() || CurrentThreadId == FPlatformTLS::InvalidTlsSlot);
+
 	CurrentThreadId = FPlatformTLS::InvalidTlsSlot;
     OuterTransactStackAddress = nullptr;
 	CurrentTransactStackAddress = nullptr;
diff --git a/Engine/Source/Runtime/Core/Private/AutoRTFM/Context.h b/Engine/Source/Runtime/Core/Private/AutoRTFM/Context.h
--- a/Engine/Source/Runtime/Core/Private/AutoRTFM/Context.h
+++ b/Engine/Source/Runtime/Core/Private/AutoRTFM/Context.h
@@ -4,6 +4,7 @@
 
 #include "AutoRTFM/AutoRTFM.h"
 #include "ContextStatus.h"
+#include "HAL/PlatformTLS.h"
 
 namespace AutoRTFM
 {
@@ -51,7 +52,7 @@ public:
 	inline FCallNest* GetCurrentNest() const { return CurrentNest; }
     inline bool IsTransactionStack(const void* LogicalAddress) const { return LogicalAddress >= StackBegin && LogicalAddress < OuterTransactStackAddress; }
 	inline bool IsInnerTransactionStack(const void* LogicalAddress) const { return LogicalAddress >= StackBegin && LogicalAddress < CurrentTransactStackAddress; }
-	inline EContextStatus GetStatus() const { return Status; }
+	inline EContextStatus GetStatus() const { return CurrentThreadId == FPlatformTLS::GetCurrentThreadId() ? Status : EContextStatus::Idle; }
 	[[noreturn]] void Throw();
 	
     void DumpState() const;
@@ -84,7 +85,7 @@ private:
     void* OuterTransactStackAddress{nullptr};
     void* CurrentTransactStackAddress{nullptr};
     EContextStatus Status{EContextStatus::Idle};
-	uint32 CurrentThreadId{ ~0u };
+	uint32 CurrentThreadId{ FPlatformTLS::InvalidTlsSlot };
 };
 
 } // namespace AutoRTFM
