[AutoRTFM] Prevent memory corruption with stack writes in non-scoped transactions

Assert if you attempt to perform a write to the transaction-stack while in non-scoped transaction (writes to a Closed() scope are permitted).

Code changes:

Introduce FStackRange to hold a stack pointer range.

Use FStackRange in FContext for storing the true stack range for the thread.

Add a FStackRange to each FTransaction, and use this to ensure that write records aren't added to a transaction if the address is within the transaction's stack.
By having a serparate stack per transaction we can prevent stack writes making their way to a parent transaction's write log, reducing the amount of memory that may need to be tracked and simplifying the Undo logic.

Ensuring that the write log does not contain stack writes is required for the open-memory validator.

diff --git a/Engine/Source/Programs/AutoRTFMTests/Private/API.cpp b/Engine/Source/Programs/AutoRTFMTests/Private/API.cpp
--- a/Engine/Source/Programs/AutoRTFMTests/Private/API.cpp
+++ b/Engine/Source/Programs/AutoRTFMTests/Private/API.cpp
@@ -950,33 +950,33 @@ TEST_CASE("API.CheckConsistencyAssumingNoRaces")
     });
 }
 
-TEST_CASE("API.IsInnerTransactionStack")
+TEST_CASE("API.IsOnCurrentTransactionStack")
 {
 	{
 		int OnStackNotInTransaction = 1;
-		REQUIRE(!AutoRTFM::IsInnerTransactionStack(&OnStackNotInTransaction));
+		REQUIRE(!AutoRTFM::IsOnCurrentTransactionStack(&OnStackNotInTransaction));
 
 		int* OnHeapNotInTransaction = new int{2};
-		REQUIRE(!AutoRTFM::IsInnerTransactionStack(OnHeapNotInTransaction));
+		REQUIRE(!AutoRTFM::IsOnCurrentTransactionStack(OnHeapNotInTransaction));
 		delete OnHeapNotInTransaction;
 	}
 
 	AutoRTFM::Commit([&]
 	{
 		int OnStackInTransaction = 3;
-		REQUIRE(AutoRTFM::IsInnerTransactionStack(&OnStackInTransaction));
+		REQUIRE(AutoRTFM::IsOnCurrentTransactionStack(&OnStackInTransaction));
 
 		int* OnHeapInTransaction = new int{4};
-		REQUIRE(!AutoRTFM::IsInnerTransactionStack(OnHeapInTransaction));
+		REQUIRE(!AutoRTFM::IsOnCurrentTransactionStack(OnHeapInTransaction));
 		delete OnHeapInTransaction;
 
 		AutoRTFM::Commit([&]
 		{
 			// `OnStackInTransaction` is no longer in the innermost scope.
-			REQUIRE(!AutoRTFM::IsInnerTransactionStack(&OnStackInTransaction));
+			REQUIRE(!AutoRTFM::IsOnCurrentTransactionStack(&OnStackInTransaction));
 
 			int OnInnermostStackInTransaction = 5;
-			REQUIRE(AutoRTFM::IsInnerTransactionStack(&OnInnermostStackInTransaction));
+			REQUIRE(AutoRTFM::IsOnCurrentTransactionStack(&OnInnermostStackInTransaction));
 		});
 	});
-}
\ No newline at end of file
+}
diff --git a/Engine/Source/Programs/AutoRTFMTests/Private/OpenAPI.cpp b/Engine/Source/Programs/AutoRTFMTests/Private/OpenAPI.cpp
--- a/Engine/Source/Programs/AutoRTFMTests/Private/OpenAPI.cpp
+++ b/Engine/Source/Programs/AutoRTFMTests/Private/OpenAPI.cpp
@@ -355,7 +355,6 @@ TEST_CASE("OpenAPI.Commit_TransactOpenCloseCommit")
 	REQUIRE(!AutoRTFM::IsTransactional());
 }
 
-
 TEST_CASE("OpenAPI.Commit_TransactOpenCloseAbort")
 {
 	REQUIRE(!AutoRTFM::IsTransactional());
@@ -363,8 +362,8 @@ TEST_CASE("OpenAPI.Commit_TransactOpenCloseAbort")
 	// We're open
 	REQUIRE(!AutoRTFM::IsClosed());
 
-	int value = 10;
-	value++;
+	int Value = 10;
+	Value++;
 
 	// Close and start the top-level transaction
 	AutoRTFM::Transact([&]()
@@ -373,24 +372,22 @@ TEST_CASE("OpenAPI.Commit_TransactOpenCloseAbort")
 
 		AutoRTFM::Open([&]()
 		{
-			double valueLocal = 1.0;
 			AutoRTFM::ForTheRuntime::StartTransaction();
 
-			// Closing from the open doesn't work
 			REQUIRE(AutoRTFM::EContextStatus::OnTrack == AutoRTFM::Close([&]()
 			{
-				value = 42;
-				valueLocal = 10.0;
+				int Local = 0;
+				Local = 42;
+				Value = Local;
 			}));
 
-			AutoRTFM::AbortTransaction(); // undoes value = 42 in the open
-			REQUIRE(valueLocal == 1.0);
+			AutoRTFM::AbortTransaction(); // undoes Value = 42 in the open
 		});
 
 		FAIL("Should not reach here!");
 	});
 
-	REQUIRE(value == 11);
+	REQUIRE(Value == 11);
 	REQUIRE(!AutoRTFM::IsTransactional());
 }
 
@@ -479,9 +476,9 @@ TEST_CASE("OpenAPI.StackWriteCommitInTheOpen1")
 
 TEST_CASE("OpenAPI.StackWriteCommitInTheOpen2")
 {
+	int Value = 0;
 	AutoRTFM::Transact([&]()
 	{
-		int value = 0;
 		AutoRTFM::Open([&]()
 		{
 			AutoRTFM::ForTheRuntime::StartTransaction();
@@ -489,13 +486,13 @@ TEST_CASE("OpenAPI.StackWriteCommitInTheOpen2")
 			{
 				AutoRTFM::Open([&]()
 				{
-					AutoRTFM::RecordOpenWrite(&value);
-					value = 10;
+					AutoRTFM::RecordOpenWrite(&Value);
+					Value = 10;
 				});
 			}));
 
 			AutoRTFM::ForTheRuntime::CommitTransaction();
-			REQUIRE(value == 10);
+			REQUIRE(Value == 10);
 		});
 	});
 }
@@ -781,19 +778,20 @@ TEST_CASE("OpenAPI.TransOpenStartCloseAbortAbort")
 
 	REQUIRE(!AutoRTFM::IsTransactional());
 
+	int Value = 10;
+
 	AutoRTFM::Transact([&]() 
 	{
 		AutoRTFM::Open([&]()
 		{
 			AutoRTFM::ForTheRuntime::StartTransaction();
 
-			int value = 10;
-			value++;
+			Value++;
 
-			value = 42;
+			Value = 42;
 			AutoRTFM::EContextStatus CloseStatus = AutoRTFM::Close([&]()
 			{
-				value = 420;
+				Value = 420;
 				AutoRTFM::AbortTransaction();
 				bGetsToA = true;
 			});
@@ -805,10 +803,10 @@ TEST_CASE("OpenAPI.TransOpenStartCloseAbortAbort")
 			REQUIRE(bGetsToA == false);
 
 			bGetsToB = true;
-			REQUIRE(value == 42);
+			REQUIRE(Value == 42);
 			AutoRTFM::AbortTransaction();
 			bGetsToC = true;
-			REQUIRE(value == 42);
+			REQUIRE(Value == 42);
 		});
 
 		bGetsToD = true;
diff --git a/Engine/Source/Runtime/Core/Private/AutoRTFM/API.cpp b/Engine/Source/Runtime/Core/Private/AutoRTFM/API.cpp
--- a/Engine/Source/Runtime/Core/Private/AutoRTFM/API.cpp
+++ b/Engine/Source/Runtime/Core/Private/AutoRTFM/API.cpp
@@ -584,10 +584,7 @@ extern "C" void autortfm_record_open_write(void* Ptr, size_t Size)
 		return RTFM_autortfm_record_open_write(Ptr, Size);
 	}
 
-	FContext* const Context = FContext::Get();
-
-    Context->CheckOpenRecordWrite(Ptr);
-	Context->RecordWrite(Ptr, Size);
+	FContext::Get()->RecordWrite(Ptr, Size);
 }
 
 extern "C" UE_AUTORTFM_NOAUTORTFM void autortfm_register_open_function(void* OriginalFunction, void* NewFunction)
@@ -596,9 +593,13 @@ extern "C" UE_AUTORTFM_NOAUTORTFM void autortfm_register_open_function(void* Ori
     FunctionMapAdd(OriginalFunction, NewFunction);
 }
 
-extern "C" bool autortfm_is_inner_transaction_stack(void* Ptr)
+extern "C" bool autortfm_is_on_current_transaction_stack(void* Ptr)
 {
-	return FContext::Get()->IsInnerTransactionStack(Ptr);
+	if (FTransaction* CurrentTransaction = FContext::Get()->GetCurrentTransaction())
+	{
+		return CurrentTransaction->IsOnStack(Ptr);
+	}
+	return false;
 }
 
 void ForTheRuntime::OnCommitInternal(TFunction<void()> && Work)
diff --git a/Engine/Source/Runtime/Core/Private/AutoRTFM/Context.cpp b/Engine/Source/Runtime/Core/Private/AutoRTFM/Context.cpp
--- a/Engine/Source/Runtime/Core/Private/AutoRTFM/Context.cpp
+++ b/Engine/Source/Runtime/Core/Private/AutoRTFM/Context.cpp
@@ -63,12 +63,9 @@ bool FContext::IsCommittingOrAborting()
 
 bool FContext::StartTransaction()
 {
+	ensureMsgf(CurrentTransaction, TEXT("FContext::StartTransaction() can only be called within a scoped transaction"));
 	FTransaction* NewTransaction = new FTransaction(this);
-
-	void* TransactStackAddress = &NewTransaction; // get the stack-local position pointer.
-	ASSERT(TransactStackAddress > StackBegin);
-	ASSERT(TransactStackAddress < StackEnd);
-	ASSERT(TransactStackAddress < CurrentTransactStackAddress);
+	NewTransaction->SetStackRange(CurrentTransaction->GetStackRange());
 
 	// This form of transaction is always ultimately within a scoped Transact 
 	ASSERT(Status == EContextStatus::OnTrack);
@@ -163,8 +160,7 @@ bool FContext::IsAborting() const
 
 EContextStatus FContext::CallClosedNest(void (*ClosedFunction)(void* Arg), void* Arg)
 {
-	TScopedGuard<void*> CurrentNestStackAddressGuard(CurrentTransactStackAddress, &CurrentNestStackAddressGuard);
-
+	TScopedGuard<void*> ClosedStackAddressGuard(ClosedStackAddress, &ClosedStackAddressGuard);
 	PushCallNest(new FCallNest(this));
 
 	CurrentNest->Try([&]() { ClosedFunction(Arg); });
@@ -295,8 +291,7 @@ ETransactionResult FContext::Transact(void (*InstrumentedFunction)(void*), void*
 	// Transact requires a return from the lambda to commit the results
 	NewTransaction->SetIsScopedTransaction();
 
-	void* TransactStackAddress = &NewTransaction;
-	TScopedGuard<void*> CurrentNestStackAddressGuard(CurrentTransactStackAddress, TransactStackAddress);
+	void* TransactStackStart = &NewTransaction;
 
 	ETransactionResult Result = ETransactionResult::Committed; // Initialize to something to make the compiler happy.
 
@@ -307,30 +302,28 @@ ETransactionResult FContext::Transact(void (*InstrumentedFunction)(void*), void*
 		ASSERT(FPlatformTLS::InvalidTlsSlot == CurrentThreadId);
 		CurrentThreadId = FPlatformTLS::GetCurrentThreadId();
 
-		ASSERT(nullptr == StackBegin);
-		ASSERT(nullptr == StackEnd);
+		ASSERT(Stack == FStackRange{});
 
 #if PLATFORM_WINDOWS
-		GetCurrentThreadStackLimits(&StackBegin, &StackEnd);
+		GetCurrentThreadStackLimits(&Stack.Low, &Stack.High);
 #elif defined(__APPLE__)         
-		StackEnd = pthread_get_stackaddr_np(pthread_self());
+		Stack.High = pthread_get_stackaddr_np(pthread_self());
 		size_t StackSize = pthread_get_stacksize_np(pthread_self());
-		StackBegin = static_cast<char*>(StackEnd) - StackSize;
+		StackLow = static_cast<char*>(Stack.High) - StackSize;
 #else
 		pthread_attr_t Attr;
 		pthread_getattr_np(pthread_self(), &Attr);
-		size_t StackSize;
-		pthread_attr_getstack(&Attr, &StackBegin, &StackSize);
-		StackEnd = static_cast<char*>(StackBegin) + StackSize;
+		size_t StackSize = 0;
+		pthread_attr_getstack(&Attr, &Stack.Low, &StackSize);
+		Stack.High = static_cast<char*>(Stack.Low) + StackSize;
 #endif
-		ASSERT(StackEnd > StackBegin);
+		ASSERT(Stack.High > Stack.Low);
 
-		ASSERT(TransactStackAddress > StackBegin);
-		ASSERT(TransactStackAddress < StackEnd);
+		ASSERT(Stack.Contains(TransactStackStart));
+		NewTransaction->SetStackRange({Stack.Low, &TransactStackStart});
 
 		PushTransaction(NewTransaction);
 		PushCallNest(NewNest);
-        OuterTransactStackAddress = TransactStackAddress;
 
 		bool bTriedToRunOnce = false;
 
@@ -357,14 +350,14 @@ ETransactionResult FContext::Transact(void (*InstrumentedFunction)(void*), void*
 
 					// We've tried to run at least once if we get here!
 					bTriedToRunOnce = true;
-
 					continue;
 				}
-                else if (AttemptToCommitTransaction(CurrentTransaction))
-                {
-                    Result = ETransactionResult::Committed;
-                    break;
-                }
+
+				if (AttemptToCommitTransaction(CurrentTransaction))
+				{
+					Result = ETransactionResult::Committed;
+					break;
+				}
 
 				UE_LOG(LogAutoRTFM, Verbose, TEXT("Commit failed!"));
 
@@ -411,10 +404,8 @@ ETransactionResult FContext::Transact(void (*InstrumentedFunction)(void*), void*
 
 		ASSERT(CurrentThreadId == FPlatformTLS::GetCurrentThreadId());
 
-		ASSERT(nullptr != StackBegin);
-		ASSERT(nullptr != StackEnd);
-		ASSERT(TransactStackAddress > StackBegin);
-		ASSERT(TransactStackAddress < StackEnd);
+		ASSERT(Stack.Contains(TransactStackStart));
+		NewTransaction->SetStackRange({Stack.Low, &TransactStackStart});
 
 		PushTransaction(NewTransaction);
 		PushCallNest(NewNest);
@@ -497,10 +488,7 @@ void FContext::Reset()
 	ASSERT(CurrentThreadId == FPlatformTLS::GetCurrentThreadId() || CurrentThreadId == FPlatformTLS::InvalidTlsSlot);
 
 	CurrentThreadId = FPlatformTLS::InvalidTlsSlot;
-    OuterTransactStackAddress = nullptr;
-	CurrentTransactStackAddress = nullptr;
-	StackBegin = nullptr;
-	StackEnd = nullptr;
+	Stack = {};
     CurrentTransaction = nullptr;
 	CurrentNest = nullptr;
     Status = EContextStatus::Idle;
@@ -513,7 +501,7 @@ void FContext::Throw()
 
 void FContext::DumpState() const
 {
-	UE_LOG(LogAutoRTFM, Verbose, TEXT("Context at %p, transaction stack: %p..%p."), this, StackBegin, OuterTransactStackAddress);
+	UE_LOG(LogAutoRTFM, Verbose, TEXT("Context at %p"), this);
 }
 
 } // namespace AutoRTFM
diff --git a/Engine/Source/Runtime/Core/Private/AutoRTFM/Context.h b/Engine/Source/Runtime/Core/Private/AutoRTFM/Context.h
--- a/Engine/Source/Runtime/Core/Private/AutoRTFM/Context.h
+++ b/Engine/Source/Runtime/Core/Private/AutoRTFM/Context.h
@@ -5,6 +5,7 @@
 #include "AutoRTFM/AutoRTFM.h"
 #include "ContextStatus.h"
 #include "HAL/PlatformTLS.h"
+#include "StackRange.h"
 
 namespace AutoRTFM
 {
@@ -36,8 +37,6 @@ public:
 	void ClearTransactionStatus();
 	bool IsAborting() const;
 
-	void CheckOpenRecordWrite(void* LogicalAddress);
-
     // Record that a write is about to occur at the given LogicalAddress of Size bytes.
     void RecordWrite(void* LogicalAddress, size_t Size);
     template<unsigned SIZE> void RecordWrite(void* LogicalAddress);
@@ -50,11 +49,15 @@ public:
 
 	inline FTransaction* GetCurrentTransaction() const { return CurrentTransaction; }
 	inline FCallNest* GetCurrentNest() const { return CurrentNest; }
-    inline bool IsTransactionStack(const void* LogicalAddress) const { return LogicalAddress >= StackBegin && LogicalAddress < OuterTransactStackAddress; }
-	inline bool IsInnerTransactionStack(const void* LogicalAddress) const { return LogicalAddress >= StackBegin && LogicalAddress < CurrentTransactStackAddress; }
 	inline EContextStatus GetStatus() const { return CurrentThreadId == FPlatformTLS::GetCurrentThreadId() ? Status : EContextStatus::Idle; }
 	[[noreturn]] void Throw();
-	
+
+	// Returns the starting stack address of the innermost call to Closed(), or
+	// nullptr if there is no call to Closed. Used to assert that a stack memory
+	// write is safe to record.
+	// See FTransaction::ShouldRecordWrite()
+	inline const void* GetClosedStackAddress() const { return ClosedStackAddress; }
+
     void DumpState() const;
 
     static void InitializeGlobalData();
@@ -80,10 +83,8 @@ private:
     FTransaction* CurrentTransaction{nullptr};
 	FCallNest* CurrentNest{nullptr};
 
-    void* StackBegin{nullptr}; // begin as in the smaller of the two
-    void* StackEnd{nullptr};
-    void* OuterTransactStackAddress{nullptr};
-    void* CurrentTransactStackAddress{nullptr};
+    FStackRange Stack;
+	void* ClosedStackAddress = nullptr;
     EContextStatus Status{EContextStatus::Idle};
 	uint32 CurrentThreadId{ FPlatformTLS::InvalidTlsSlot };
 };
diff --git a/Engine/Source/Runtime/Core/Private/AutoRTFM/ContextInlines.h b/Engine/Source/Runtime/Core/Private/AutoRTFM/ContextInlines.h
--- a/Engine/Source/Runtime/Core/Private/AutoRTFM/ContextInlines.h
+++ b/Engine/Source/Runtime/Core/Private/AutoRTFM/ContextInlines.h
@@ -13,17 +13,6 @@ UE_AUTORTFM_FORCEINLINE FContext* FContext::Get()
 	return &FContext::ContextSingleton;
 }
 
-UE_AUTORTFM_FORCEINLINE void FContext::CheckOpenRecordWrite(void* LogicalAddress)
-{
-    // We don't record any writes to unscoped transactions when the write-address
-    // is within the nearest closed C++ function call nest. This may seem weird, but we
-    // have to prohibit this because we can't track all the comings and goings of
-    // local variables and all child C++ scopes. It's easy to get into a situation
-    // where a rollback will undo the recorded writes to stack-locals that actually
-    // corrupts the undo process itself.
-    ASSERT(CurrentTransaction->IsScopedTransaction() || !IsInnerTransactionStack(LogicalAddress));
-}
-
 AUTORTFM_NO_ASAN UE_AUTORTFM_FORCEINLINE void FContext::RecordWrite(void* LogicalAddress, size_t Size)
 {
     CurrentTransaction->RecordWrite(LogicalAddress, Size);
diff --git a/Engine/Source/Runtime/Core/Private/AutoRTFM/StackRange.h b/Engine/Source/Runtime/Core/Private/AutoRTFM/StackRange.h
new file mode 100644
--- /dev/null
+++ b/Engine/Source/Runtime/Core/Private/AutoRTFM/StackRange.h
@@ -0,0 +1,24 @@
+// Copyright Epic Games, Inc. All Rights Reserved.
+
+#pragma once
+
+namespace AutoRTFM
+{
+
+// FStackRange represents a stack memory range.
+// It is assumed the stack grows downwards.
+struct FStackRange
+{
+    void* Low = nullptr;  // One byte past the end of the stack range.
+    void* High = nullptr; // The first byte of the stack range.
+
+    // Returns true if the stack range contains Address
+    bool Contains(const void* Address) const { return Address > Low && Address <= High; }
+
+    // Equality operator
+    bool operator == (const FStackRange& Other) const { return Low == Other.Low && High == Other.High; }
+    // In-equality operator
+    bool operator != (const FStackRange& Other) const { return !(*this == Other); }
+};
+
+}  // namespace AutoRTFM
diff --git a/Engine/Source/Runtime/Core/Private/AutoRTFM/Transaction.cpp b/Engine/Source/Runtime/Core/Private/AutoRTFM/Transaction.cpp
--- a/Engine/Source/Runtime/Core/Private/AutoRTFM/Transaction.cpp
+++ b/Engine/Source/Runtime/Core/Private/AutoRTFM/Transaction.cpp
@@ -101,12 +101,8 @@ void FTransaction::Undo()
 		FWriteLogEntry& Entry = *Iter;
 		void* const Original = Entry.GetOriginal();
 
-        // Skip writes to our current transaction nest if we're scoped. We're about to
-		// leave so the changes don't matter. 
-        if (IsScopedTransaction() && Context->IsInnerTransactionStack(Original))
-        {
-            continue;
-        }
+        // No write records should be within the transaction's stack range.
+        ensure(!IsOnStack(Original));
 
 		const size_t Size = Entry.GetSize();
         void* const Copy = Entry.GetCopy();
@@ -149,20 +145,17 @@ void FTransaction::CommitNested()
 {
     ASSERT(Parent);
 
-    // We need to pass our write log to our parent transaction, but with care!
-    // We need to discard any writes to locations within the stack of our
-    // current transaction, which could be placed there if a child of the
-    // current transaction had written to stack local memory in the parent.
-
-    for (FWriteLogEntry& Write : WriteLog)
-    {
-        // Skip writes that are into our current transactions stack.
-        if (IsScopedTransaction() && Context->IsInnerTransactionStack(Write.GetOriginal()))
-        {
-            continue;
-        }
+	// We need to pass our write log to our parent transaction, but with care!
+	// We need to discard any writes if the memory location is on the parent
+	// transaction's stack range.
+	for (FWriteLogEntry& Write : WriteLog)
+	{
+		if (Parent->IsOnStack(Write.GetOriginal()))
+		{
+			continue;
+		}
 
-        Parent->WriteLog.Push(Write);
+		Parent->WriteLog.Push(Write);
 
 		FHitSet::Key HitSetEntry(Write.GetOriginal());
 		HitSetEntry.SetTopTag(static_cast<uint16_t>(Write.GetSize()));
diff --git a/Engine/Source/Runtime/Core/Private/AutoRTFM/Transaction.h b/Engine/Source/Runtime/Core/Private/AutoRTFM/Transaction.h
--- a/Engine/Source/Runtime/Core/Private/AutoRTFM/Transaction.h
+++ b/Engine/Source/Runtime/Core/Private/AutoRTFM/Transaction.h
@@ -5,6 +5,7 @@
 #include "HitSet.h"
 #include "IntervalTree.h"
 #include "LongJump.h"
+#include "StackRange.h"
 #include "Stats.h"
 #include "TaggedPtr.h"
 #include "TaskArray.h"
@@ -79,6 +80,13 @@ public:
     void DidAllocate(void* LogicalAddress, size_t Size);
     void DidFree(void* LogicalAddress);
 
+    // The stack range represents all stack memory inside the transaction scope
+    inline void SetStackRange(FStackRange Range) { StackRange = Range;} 
+    inline FStackRange GetStackRange() const { return StackRange; }
+
+    // Returns true if the LogicalAddress is within the stack of the transaction.
+    inline bool IsOnStack(const void* LogicalAddress) const;
+
 private:
     void Undo();
 
@@ -89,6 +97,8 @@ private:
 
     void CollectStats() const;
     
+    bool ShouldRecordWrite(void* LogicalAddress) const;
+
     FContext* Context;
     
     // Are we nested? Then this is the parent.
@@ -106,6 +116,7 @@ private:
     FWriteLog WriteLog;
     FWriteLogBumpAllocator WriteLogBumpAllocator;
     TStatStorage<uint64_t> StatDepth = 1;
+    FStackRange StackRange;
 };
 
 } // namespace AutoRTFM
diff --git a/Engine/Source/Runtime/Core/Private/AutoRTFM/TransactionInlines.h b/Engine/Source/Runtime/Core/Private/AutoRTFM/TransactionInlines.h
--- a/Engine/Source/Runtime/Core/Private/AutoRTFM/TransactionInlines.h
+++ b/Engine/Source/Runtime/Core/Private/AutoRTFM/TransactionInlines.h
@@ -6,10 +6,61 @@
 #include "Transaction.h"
 
 #include "HAL/Platform.h"
+#include "Utils.h"
 
 namespace AutoRTFM
 {
 
+UE_AUTORTFM_FORCEINLINE bool FTransaction::IsOnStack(const void* LogicalAddress) const
+{ 
+    return StackRange.Contains(LogicalAddress);
+}
+
+UE_AUTORTFM_FORCEINLINE bool FTransaction::ShouldRecordWrite(void* LogicalAddress) const
+{
+    // We cannot record writes to stack memory used within the transaction, as
+    // undoing the writes may corrupt stack memory that has been unwound or
+    // is now being used for a different variable from the one the write was
+    // made.
+    if (!IsOnStack(LogicalAddress))
+    {
+        return true;
+    }
+
+    // Writes to the stack under a scoped-transaction can be safely ignored,
+    // because the values on the stack are not visible outside of the scope of
+    // the transaction. In other words, if a scoped-transaction aborts that
+    // memory will cease to be meaningful anyway.
+
+    // Non-scoped transactions, as the name implies, do not impose a lexical
+    // scope that encompasses the transaction. Instead a non-scoped transaction
+    // is started with a call to StartTransaction() and ended with a call to
+    // either AbortTransaction() or CommitTransaction(). Unlike a
+    // scoped transaction, there's no precise stack range for a non-scoped 
+    // transaction, as the scope can freely grow or shrink between the calls to
+    // [Start|Abort|Commit]Transaction() and any recorded writes. The only
+    // guarantee we have is that a non-scoped transaction cannot shrink past the
+    // outer scoped transaction. For this reason, non-scoped transactions
+    // adopt the stack range of the outer transaction, as this is guaranteed
+    // to encompass the non-scoped transaction's scope range.
+
+    // For non-scoped transactions, we assert that we're not writing to a
+    // memory address that's in the transaction's stack range as this cannot be
+    // safely undone, and stack variables may be visible once the transaction is
+    // aborted. We make an exception for stack variables declared within the
+    // scope of a Close(), as writing to these stack variables can be safely
+    // ignored (they have the same constrained visibility as stack variables in
+    // a scoped transaction).
+
+    // Hitting this assert? 
+    // Consider moving the variable being written to an inner scoped 
+    // transaction, or move the variable outside of the nearest parent 
+    // scoped-transaction.
+    ASSERT(bIsStackScoped || LogicalAddress < FContext::Get()->GetClosedStackAddress()); 
+
+    return false;
+}
+
 AUTORTFM_NO_ASAN UE_AUTORTFM_FORCEINLINE void FTransaction::RecordWriteMaxPageSized(void* LogicalAddress, size_t Size)
 {
     void* CopyAddress = WriteLogBumpAllocator.Allocate(Size);
@@ -25,11 +76,7 @@ AUTORTFM_NO_ASAN UE_AUTORTFM_FORCEINLINE void FTransaction::RecordWrite(void* Lo
         return;
     }
 
-    // If we are recording a stack address that is relative to our current
-    // transactions stack location, we do not need to record the data in the
-    // write log because if that transaction aborted, that memory will cease to
-    // be meaningful anyway!
-    if (Context->IsInnerTransactionStack(LogicalAddress))
+    if (!ShouldRecordWrite(LogicalAddress))
     {
         Stats.Collect<EStatsKind::HitSetSkippedBecauseOfStackLocalMemory>();
         return;
@@ -79,11 +126,7 @@ template<unsigned SIZE> AUTORTFM_NO_ASAN UE_AUTORTFM_FORCEINLINE void FTransacti
 {
     static_assert(SIZE <= 8);
 
-    // If we are recording a stack address that is relative to our current
-    // transactions stack location, we do not need to record the data in the
-    // write log because if that transaction aborted, that memory will cease to
-    // be meaningful anyway!
-    if (Context->IsInnerTransactionStack(LogicalAddress))
+    if (!ShouldRecordWrite(LogicalAddress))
     {
         Stats.Collect<EStatsKind::HitSetSkippedBecauseOfStackLocalMemory>();
         return;
diff --git a/Engine/Source/Runtime/Core/Public/AutoRTFM/AutoRTFM.h b/Engine/Source/Runtime/Core/Public/AutoRTFM/AutoRTFM.h
--- a/Engine/Source/Runtime/Core/Public/AutoRTFM/AutoRTFM.h
+++ b/Engine/Source/Runtime/Core/Public/AutoRTFM/AutoRTFM.h
@@ -280,9 +280,9 @@ static UE_AUTORTFM_FORCEINLINE void autortfm_register_open_function(void* origin
 #endif
 
 #if UE_AUTORTFM_ENABLED
-UE_AUTORTFM_API bool autortfm_is_inner_transaction_stack(void* Ptr);
+UE_AUTORTFM_API bool autortfm_is_on_current_transaction_stack(void* Ptr);
 #else
-static UE_AUTORTFM_FORCEINLINE bool autortfm_is_inner_transaction_stack(void* Ptr)
+static UE_AUTORTFM_FORCEINLINE bool autortfm_is_on_current_transaction_stack(void* Ptr)
 {
 	return false;
 }
@@ -480,9 +480,9 @@ static UE_AUTORTFM_FORCEINLINE bool IsCommittingOrAborting() { return autortfm_i
 // Returns true if the passed-in pointer is on the stack of the currently-executing transaction.
 // This is occasionally necessary when writing OnAbort handlers for objects on the stack, since 
 // we don't want to scribble on stack memory that might have been reused.
-static UE_AUTORTFM_FORCEINLINE bool IsInnerTransactionStack(void* Ptr) 
+static UE_AUTORTFM_FORCEINLINE bool IsOnCurrentTransactionStack(void* Ptr) 
 { 
-	return autortfm_is_inner_transaction_stack(Ptr);
+	return autortfm_is_on_current_transaction_stack(Ptr);
 }
 
 // Run the functor in a transaction. Memory writes and other side effects get instrumented
diff --git a/Engine/Source/Runtime/Core/Public/Containers/LockFreeFixedSizeAllocator.h b/Engine/Source/Runtime/Core/Public/Containers/LockFreeFixedSizeAllocator.h
--- a/Engine/Source/Runtime/Core/Public/Containers/LockFreeFixedSizeAllocator.h
+++ b/Engine/Source/Runtime/Core/Public/Containers/LockFreeFixedSizeAllocator.h
@@ -209,7 +209,7 @@ public:
 		// If we are in a transaction and this allocator exists on the transaction's stack, 
 		// running the ONABORT handler after destruction is dangerous and could stomp memory
 		// that belongs to someone else.
-		if (AutoRTFM::IsClosed() && AutoRTFM::IsInnerTransactionStack(this))
+		if (AutoRTFM::IsClosed() && AutoRTFM::IsOnCurrentTransactionStack(this))
 		{
 			AutoRTFM::PopAllOnAbortHandlers(this);
 		}
