[StrongObjectPtr]
 - Make Attach/Detach private to avoid user errors

#rnx

diff --git a/Engine/Source/Runtime/Core/Public/UObject/StrongObjectPtrTemplates.h b/Engine/Source/Runtime/Core/Public/UObject/StrongObjectPtrTemplates.h
--- a/Engine/Source/Runtime/Core/Public/UObject/StrongObjectPtrTemplates.h
+++ b/Engine/Source/Runtime/Core/Public/UObject/StrongObjectPtrTemplates.h
@@ -131,6 +131,9 @@ public:
 		}
 	}
 	
+private:
+	template<class T, class TWeakObjectPtrBase> friend struct TWeakObjectPtr;
+
 	// Attach an object without incrementing its ref-count.
 	FORCEINLINE_DEBUGGABLE void Attach(ObjectType* InNewObject)
 	{
@@ -146,6 +149,7 @@ public:
 		return DetachedObject;
 	}
 
+public:
 	FORCEINLINE_DEBUGGABLE void Reset(ObjectType* InNewObject)
 	{
 		if (InNewObject)
