FMath: Fix ConstExprCeilLog2; it was formerly returning FloorLog2 instead of CeilLog2.
#rnx

diff --git a/Engine/Source/Runtime/Core/Public/GenericPlatform/GenericPlatformMath.h b/Engine/Source/Runtime/Core/Public/GenericPlatform/GenericPlatformMath.h
--- a/Engine/Source/Runtime/Core/Public/GenericPlatform/GenericPlatformMath.h
+++ b/Engine/Source/Runtime/Core/Public/GenericPlatform/GenericPlatformMath.h
@@ -795,7 +795,16 @@ struct FGenericPlatformMath
 	 */
 	static constexpr FORCEINLINE uint8 ConstExprCeilLogTwo(SIZE_T Arg)
 	{
-		return Arg <= 1 ? 0 : (1 + ConstExprCeilLogTwo(Arg / 2));
+		if (Arg <= 1)
+		{
+			return 0;
+		}
+		// Integer overflow if we tried to add 1 to maximum value, so handle that case separately
+		if (Arg + 1 < Arg)
+		{
+			return sizeof(Arg) * 8;
+		}
+		return 1 + ConstExprCeilLogTwo((Arg + 1) / 2);
 	}
 
 	/** @return Rounds the given number up to the next highest power of two. */
diff --git a/Engine/Source/Runtime/Core/Tests/Math/UnrealMathTest.cpp b/Engine/Source/Runtime/Core/Tests/Math/UnrealMathTest.cpp
--- a/Engine/Source/Runtime/Core/Tests/Math/UnrealMathTest.cpp
+++ b/Engine/Source/Runtime/Core/Tests/Math/UnrealMathTest.cpp
@@ -4348,6 +4348,31 @@ public:
 		return true;
 	}
 };
+
+// ConstExprCeilLogTwo
+static_assert(FMath::ConstExprCeilLogTwo(0) == 0, "0 -> 0");
+static_assert(FMath::ConstExprCeilLogTwo(1) == 0, "1 -> 0");
+static_assert(FMath::ConstExprCeilLogTwo(2) == 1, "2 -> 1");
+static_assert(FMath::ConstExprCeilLogTwo(3) == 2, "3 -> 2");
+static_assert(FMath::ConstExprCeilLogTwo(4) == 2, "4 -> 2");
+static_assert(FMath::ConstExprCeilLogTwo(5) == 3, "5 -> 3");
+static_assert(FMath::ConstExprCeilLogTwo(6) == 3, "6 -> 3");
+static_assert(FMath::ConstExprCeilLogTwo(7) == 3, "7 -> 3");
+static_assert(FMath::ConstExprCeilLogTwo(8) == 3, "8 -> 3");
+static_assert(FMath::ConstExprCeilLogTwo(16) == 4, "16 -> 4");
+static_assert(FMath::ConstExprCeilLogTwo(17) == 5, "17 -> 5");
+static_assert(FMath::ConstExprCeilLogTwo(31) == 5, "31 -> 5");
+static_assert(FMath::ConstExprCeilLogTwo(32) == 5, "32 -> 5");
+static_assert(FMath::ConstExprCeilLogTwo(33) == 6, "33 -> 6");
+static_assert(FMath::ConstExprCeilLogTwo(1U << 16) == 16, "2^16 -> 16");
+static_assert(FMath::ConstExprCeilLogTwo(uint64(1) << 31) == 31, "2^31 -> 31");
+static_assert(FMath::ConstExprCeilLogTwo((uint64(1) << 32) - 1) == 32, "2^32 - 1 -> 32");
+static_assert(FMath::ConstExprCeilLogTwo(uint64(1) << 32) == 32, "2^32 -> 32");
+static_assert(FMath::ConstExprCeilLogTwo(uint64(0xfffffffffffffffe)) == 64, "2^64 - 2 -> 64");
+static_assert(FMath::ConstExprCeilLogTwo(uint64(0xffffffffffffffff)) == 64, "2^64 - 1 -> 64");
+
+
+
 TEST_CASE_NAMED(FIsNearlyEqualByULPTest, "System::Core::Math::IsNearlyEqualByULP", "[ApplicationContextMask][SmokeFilter]")
 {
 	FIsNearlyEqualByULPTestClass Instance;
