Store the in-memory FTextKey string data as UTF-8 by default

This trades some CPU cost when converting FString -> FTextKey -> FString, in exchange for halving the text key string memory (as most, if not all, text keys are actually ASCII).

Note: You can set UE_TEXTKEY_USE_UTF8=0 to restore the previous behavior of storing TCHAR data internally

Note: This change formally deprecates FTextKey::GetChars(). Use FTextKey::ToString/AppendString instead.

diff --git a/Engine/Source/Runtime/Core/Private/Internationalization/TextKey.cpp b/Engine/Source/Runtime/Core/Private/Internationalization/TextKey.cpp
--- a/Engine/Source/Runtime/Core/Private/Internationalization/TextKey.cpp
+++ b/Engine/Source/Runtime/Core/Private/Internationalization/TextKey.cpp
@@ -5,6 +5,7 @@
 #include "AutoRTFM/AutoRTFM.h"
 #include "Containers/ContainerAllocationPolicies.h"
 #include "Containers/ChunkedArray.h"
+#include "Containers/Map.h"
 #include "CoreGlobals.h"
 #include "HAL/LowLevelMemTracker.h"
 #include "Hash/CityHash.h"
@@ -17,6 +18,11 @@
 
 DEFINE_LOG_CATEGORY_STATIC(LogTextKey, Log, All);
 
+// Note: If setting this to 0, you'll also want to update the FTextKey natvis to change ",s8" to ",su"
+#ifndef UE_TEXTKEY_USE_UTF8
+	#define UE_TEXTKEY_USE_UTF8 (1)
+#endif
+
 #ifndef UE_TEXTKEY_SLAB_ALLOCATOR_SLAB_SIZE
 	#define UE_TEXTKEY_SLAB_ALLOCATOR_SLAB_SIZE (32768)
 #endif
@@ -29,6 +35,13 @@ DEFINE_LOG_CATEGORY_STATIC(LogTextKey, Log, All);
 	#define UE_TEXTKEY_ELEMENTS_MIN_HASH_SIZE (32768)
 #endif
 
+#if UE_TEXTKEY_USE_UTF8
+using FTextKeyCharType = UTF8CHAR;
+#else
+using FTextKeyCharType = TCHAR;
+#endif
+using FTextKeyStringView = TStringView<FTextKeyCharType>;
+
 class FTextKeyState
 {
 public:
@@ -41,33 +54,103 @@ public:
 	void FindOrAdd(FStringView InStr, FTextKey& OutTextKey)
 	{
 		check(!InStr.IsEmpty());
+
+		// Note: This hash gets serialized so *DO NOT* change it without fixing the serialization to discard the old hash method (also update FTextKey::GetTypeHash)
+		const uint32 StrHash = TextKeyUtil::HashString(InStr);
+
 		// Open around adding this in a cache, if we abort just leak the value in the cache
 		// as the cache takes ownership
 		UE_AUTORTFM_OPEN2
 		{
-			FindOrAddImpl(FKeyData(InStr), OutTextKey);
+			FindOrAddImpl(InStr, StrHash, OutTextKey);
 		};
 	}
 
 	void FindOrAdd(FStringView InStr, const uint32 InStrHash, FTextKey& OutTextKey)
 	{
 		check(!InStr.IsEmpty());
+
 		// Open around adding this in a cache, if we abort just leak the value in the cache
 		// as the cache takes ownership
 		UE_AUTORTFM_OPEN2
 		{
-			FindOrAddImpl(FKeyData(InStr, InStrHash), OutTextKey);
+			FindOrAddImpl(InStr, InStrHash, OutTextKey);
 		};
 	}
 
-	const TCHAR* GetStringByIndex(const int32 InIndex) const
+	const TCHAR* GetLegacyTCHARPointerByIndex(const int32 InIndex)
+	{
+		check(InIndex != INDEX_NONE);
+
+#if UE_TEXTKEY_USE_UTF8
+		{
+			// Read-only
+			int32 NumElementsOnRead = 0;
+			{
+				FReadScopeLock ScopeLock(DataRW);
+
+				if (const FString* FoundString = LegacyTCHARState.Find(InIndex))
+				{
+					return **FoundString;
+				}
+
+				NumElementsOnRead = LegacyTCHARState.Num();
+			}
+
+			// Write
+			{
+				FWriteScopeLock ScopeLock(DataRW);
+
+				if (LegacyTCHARState.Num() > NumElementsOnRead)
+				{
+					// Find again in case another thread beat us to it!
+					if (const FString* FoundString = LegacyTCHARState.Find(InIndex))
+					{
+						return **FoundString;
+					}
+				}
+
+				const FKeyData& KeyData = KeyDataAllocations.Get(InIndex);
+
+				LLM_SCOPE_BYNAME(TEXT("Localization/Deprecated"));
+				// Open around adding this in a cache, if we abort just leak the value in the cache
+				// as the cache takes ownership
+				const TCHAR* AddedString = nullptr;
+				UE_AUTORTFM_OPEN2
+				{
+					AddedString = *LegacyTCHARState.Add(InIndex, FString(KeyData.ToView()));
+				};
+				return AddedString;
+			}
+		}
+#else
+		{
+			FReadScopeLock ScopeLock(DataRW);
+
+			const FKeyData& KeyData = KeyDataAllocations.Get(InIndex);
+			return KeyData.Str;
+		}
+#endif
+	}
+
+	void AppendStringByIndex(const int32 InIndex, FString& Out) const
 	{
 		check(InIndex != INDEX_NONE);
 
 		FReadScopeLock ScopeLock(DataRW);
 
 		const FKeyData& KeyData = KeyDataAllocations.Get(InIndex);
-		return KeyData.Str;
+		Out += KeyData.ToView();
+	}
+
+	void AppendStringByIndex(const int32 InIndex, FStringBuilderBase& Out) const
+	{
+		check(InIndex != INDEX_NONE);
+
+		FReadScopeLock ScopeLock(DataRW);
+
+		const FKeyData& KeyData = KeyDataAllocations.Get(InIndex);
+		Out += KeyData.ToView();
 	}
 
 	uint32 GetHashByIndex(const int32 InIndex) const
@@ -99,48 +182,45 @@ private:
 	struct FKeyData
 	{
 		FKeyData()
+#if UE_TEXTKEY_USE_UTF8
+			: Str(UTF8TEXT(""))
+#else
 			: Str(TEXT(""))
+#endif
 			, StrLen(0)
 			, StrHash(0)
 		{
 		}
 
-		explicit FKeyData(FStringView InStr)
-			: Str(InStr.GetData())
-			, StrLen(InStr.Len())
-			, StrHash(TextKeyUtil::HashString(InStr)) // Note: This hash gets serialized so *DO NOT* change it without fixing the serialization to discard the old hash method (also update FTextKey::GetTypeHash)
-		{
-		}
-
-		FKeyData(FStringView InStr, const uint32 InStrHash)
+		FKeyData(FTextKeyStringView InStr, const uint32 InStrHash)
 			: Str(InStr.GetData())
 			, StrLen(InStr.Len())
 			, StrHash(InStrHash)
 		{
 		}
 
-		FKeyData(FStringView InStr, const FKeyData& InOther)
+		FKeyData(FTextKeyStringView InStr, const FKeyData& InOther)
 			: Str(InStr.GetData())
 			, StrLen(InStr.Len())
 			, StrHash(InOther.StrHash)
 		{
 		}
 
-		FStringView ToView() const
+		FTextKeyStringView ToView() const
 		{
-			return FStringView(Str, StrLen);
+			return FTextKeyStringView(Str, StrLen);
 		}
 
 		friend FORCEINLINE bool operator==(const FKeyData& A, const FKeyData& B)
 		{
 			// We can use Memcmp here as we know we're comparing two blocks of the same size and don't care about lexical ordering
-			return A.StrLen == B.StrLen && FMemory::Memcmp(A.Str, B.Str, A.StrLen * sizeof(TCHAR)) == 0;
+			return A.StrLen == B.StrLen && FMemory::Memcmp(A.Str, B.Str, A.StrLen * sizeof(FTextKeyCharType)) == 0;
 		}
 
 		friend FORCEINLINE bool operator!=(const FKeyData& A, const FKeyData& B)
 		{
 			// We can use Memcmp here as we know we're comparing two blocks of the same size and don't care about lexical ordering
-			return A.StrLen != B.StrLen || FMemory::Memcmp(A.Str, B.Str, A.StrLen * sizeof(TCHAR)) != 0;
+			return A.StrLen != B.StrLen || FMemory::Memcmp(A.Str, B.Str, A.StrLen * sizeof(FTextKeyCharType)) != 0;
 		}
 
 		friend FORCEINLINE uint32 GetTypeHash(const FKeyData& A)
@@ -148,13 +228,16 @@ private:
 			return A.StrHash;
 		}
 
-		const TCHAR* Str;
+		const FTextKeyCharType* Str;
 		int32 StrLen;
 		uint32 StrHash;
 	};
 
-	void FindOrAddImpl(const FKeyData& KeyData, FTextKey& OutTextKey)
+	void FindOrAddImpl(FStringView InStr, const uint32 InStrHash, FTextKey& OutTextKey)
 	{
+		auto ConvertedString = StrCast<FTextKeyCharType>(InStr.GetData(), InStr.Len());
+		const FKeyData KeyData = FKeyData(FTextKeyStringView(ConvertedString.Get(), ConvertedString.Length()), InStrHash);
+
 		int32 Index = FindOrAddString(KeyData);
 		check(Index != INDEX_NONE);
 
@@ -195,8 +278,8 @@ private:
 			}
 
 			LLM_SCOPE_BYNAME(TEXT("Localization/TextKeys"));
-			const TCHAR* NewStrPtr = StringAllocations.Add(KeyData.ToView());
-			return KeyDataAllocations.Add(FKeyData(FStringView(NewStrPtr, KeyData.StrLen), KeyData));
+			const FTextKeyCharType* NewStrPtr = StringAllocations.Add(KeyData.ToView());
+			return KeyDataAllocations.Add(FKeyData(FTextKeyStringView(NewStrPtr, KeyData.StrLen), KeyData));
 		}
 	}
 
@@ -321,19 +404,19 @@ private:
 				TotalNumElementsWasted += (SlabSizeInElements - Slab.NumElementsUsed);
 				FMemory::Free(Slab.Allocation);
 			}
-			//UE_LOG(LogLocalization, Log, TEXT("FTextKey slab allocator allocated %d slabs (of %d elements) and wasted %d elements (%d bytes)"), Slabs.Num(), SlabSizeInElements, TotalNumElementsWasted, TotalNumElementsWasted * sizeof(TCHAR));
+			//UE_LOG(LogLocalization, Log, TEXT("FTextKey slab allocator allocated %d slabs (of %d elements) and wasted %d elements (%d bytes)"), Slabs.Num(), SlabSizeInElements, TotalNumElementsWasted, TotalNumElementsWasted * sizeof(FTextKeyCharType));
 		}
 
-		const TCHAR* Add(FStringView InStr)
+		const FTextKeyCharType* Add(FTextKeyStringView InStr)
 		{
 			const int32 NumSlabElementsNeeded = InStr.Len() + 1;
 			FStringSlab& Slab = GetSlab(NumSlabElementsNeeded);
 
-			TCHAR* StringPtr = Slab.Allocation + Slab.NumElementsUsed;
+			FTextKeyCharType* StringPtr = Slab.Allocation + Slab.NumElementsUsed;
 			Slab.NumElementsUsed += NumSlabElementsNeeded;
 
-			FMemory::Memcpy(StringPtr, InStr.GetData(), InStr.Len() * sizeof(TCHAR));
-			*(StringPtr + InStr.Len()) = 0;
+			FMemory::Memcpy(StringPtr, InStr.GetData(), InStr.Len() * sizeof(FTextKeyCharType));
+			*(StringPtr + InStr.Len()) = (FTextKeyCharType)0;
 
 			return StringPtr;
 		}
@@ -343,7 +426,7 @@ private:
 
 		struct FStringSlab
 		{
-			TCHAR* Allocation = nullptr;
+			FTextKeyCharType* Allocation = nullptr;
 			int32 NumElementsUsed = 0;
 		};
 
@@ -382,7 +465,7 @@ private:
 			// If no slabs have space then just allocate a new one
 			checkf(NumSlabElementsNeeded <= SlabSizeInElements, TEXT("Tried to allocate a FTextKey string of %d elements, which is larger than the allowed slab size of %d elements!"), NumSlabElementsNeeded, SlabSizeInElements);
 			FStringSlab& Slab = Slabs.AddDefaulted_GetRef();
-			Slab.Allocation = static_cast<TCHAR*>(FMemory::Malloc(SlabSizeInElements * sizeof(TCHAR), 1));
+			Slab.Allocation = static_cast<FTextKeyCharType*>(FMemory::Malloc(SlabSizeInElements * sizeof(FTextKeyCharType), 1));
 			return Slab;
 		}
 
@@ -392,6 +475,10 @@ private:
 	mutable FRWLock DataRW;
 	FStringAllocator StringAllocations;
 	FKeyDataAllocator KeyDataAllocations;
+#if UE_TEXTKEY_USE_UTF8
+	// Sparse TCHAR state; built on-demand by anything still using the deprecated FTextKey::GetChars function
+	TMap<int32, FString> LegacyTCHARState;
+#endif
 };
 
 namespace TextKeyUtil
@@ -522,7 +609,7 @@ FTextKey::FTextKey(FStringView InStr)
 const TCHAR* FTextKey::GetChars() const
 {
 	return Index != INDEX_NONE
-		? FTextKeyState::GetState().GetStringByIndex(Index)
+		? FTextKeyState::GetState().GetLegacyTCHARPointerByIndex(Index)
 		: TEXT("");
 }
 
@@ -547,12 +634,18 @@ void FTextKey::ToString(FStringBuilderBase& Out) const
 
 void FTextKey::AppendString(FString& Out) const
 {
-	Out += GetChars();
+	if (Index != INDEX_NONE)
+	{
+		FTextKeyState::GetState().AppendStringByIndex(Index, Out);
+	}
 }
 
 void FTextKey::AppendString(FStringBuilderBase& Out) const
 {
-	Out += GetChars();
+	if (Index != INDEX_NONE)
+	{
+		FTextKeyState::GetState().AppendStringByIndex(Index, Out);
+	}
 }
 
 uint32 GetTypeHash(const FTextKey& A)
diff --git a/Engine/Source/Runtime/Core/Public/Internationalization/TextKey.h b/Engine/Source/Runtime/Core/Public/Internationalization/TextKey.h
--- a/Engine/Source/Runtime/Core/Public/Internationalization/TextKey.h
+++ b/Engine/Source/Runtime/Core/Public/Internationalization/TextKey.h
@@ -87,10 +87,8 @@ public:
 	{
 	}
 
-	/**
-	 * Get the underlying chars buffer this text key represents.
-	 * @note This function is due to be deprecated. Prefer ToString/AppendString for new code.
-	 */
+	/** Get the underlying chars buffer this text key represents */
+	UE_DEPRECATED(5.5, "GetChars is deprecated as FTextKey may now store its internal data as UTF-8. Use ToString/AppendString instead.")
 	CORE_API const TCHAR* GetChars() const;
 
 	/** Convert this text key back to its string representation */
