[AutoRTFM] Clean up deprecated and unused AutoRTFM macros.

Our transaction macros existed in two forms--OPEN and OPEN2, ONCOMMIT and ONCOMMIT2, etc.
I searched the code and couldn't find any remaining usage of the old, deprecated form.
I've updated the deprecated form to behave identically to the modern equivalent, so now OPEN and OPEN2 have the exact same behavior.
In followup CLs, we can search-and-replace OPEN2 into OPEN (etc.), and then remove the OPEN2 form entirely.

diff --git a/Engine/Source/Runtime/Core/Public/AutoRTFM/AutoRTFM.h b/Engine/Source/Runtime/Core/Public/AutoRTFM/AutoRTFM.h
--- a/Engine/Source/Runtime/Core/Public/AutoRTFM/AutoRTFM.h
+++ b/Engine/Source/Runtime/Core/Public/AutoRTFM/AutoRTFM.h
@@ -838,11 +838,6 @@ template<typename TTYPE> static UE_AUTORTFM_FORCEINLINE void RecordOpenWrite(TTY
 // A collection of power-user functions that are reserved for use by the AutoRTFM runtime only.
 namespace ForTheRuntime
 {
-	[[deprecated("This macro is deprecated. Use UE_AUTORTFM_ONABORT2 instead!")]] UE_AUTORTFM_FORCEINLINE void DeprecatedUseOnAbortMacro() {}
-	[[deprecated("This macro is deprecated. Use UE_AUTORTFM_ONCOMMIT2 instead!")]] UE_AUTORTFM_FORCEINLINE void DeprecatedUseOnCommitMacro() {}
-	[[deprecated("This macro is deprecated. Use UE_AUTORTFM_OPEN2 instead!")]] UE_AUTORTFM_FORCEINLINE void DeprecatedUseOpenMacro() {}
-	[[deprecated("This macro is deprecated. Use UE_AUTORTFM_TRANSACT2 instead!")]] UE_AUTORTFM_FORCEINLINE void DeprecatedUseTransactMacro() {}
-
 	// An enum to represent the various ways we want to enable/disable the AutoRTFM runtime.
 	enum EAutoRTFMEnabledState
 	{
@@ -1055,56 +1050,42 @@ namespace AutoRTFM::Private
 #define UE_AUTORTFM_END_DISABLE_WARNINGS
 #endif
 
-// Older macros where the code is passed as a macro argument. These should be phased out as they make debugging more difficult
-#define UE_AUTORTFM_OPEN_IMPL(...) ::AutoRTFM::Open([&]() { __VA_ARGS__ })
-#define UE_AUTORTFM_ONABORT_IMPL(...) UE_AUTORTFM_BEGIN_DISABLE_WARNINGS ::AutoRTFM::OnAbort([=]() { __VA_ARGS__ }) UE_AUTORTFM_END_DISABLE_WARNINGS
-#define UE_AUTORTFM_ONCOMMIT_IMPL(...) UE_AUTORTFM_BEGIN_DISABLE_WARNINGS ::AutoRTFM::OnCommit([=]() { __VA_ARGS__ }) UE_AUTORTFM_END_DISABLE_WARNINGS
-#define UE_AUTORTFM_TRANSACT_IMPL(...) ::AutoRTFM::Transact([&]() { __VA_ARGS__ })
-
-#define UE_AUTORTFM_OPEN_IMPL2 ::AutoRTFM::Private::FOpenHelper{} + [&]()
-#define UE_AUTORTFM_ONABORT_IMPL2(...) ::AutoRTFM::Private::FOnAbortHelper{} + [__VA_ARGS__]()
-#define UE_AUTORTFM_ONCOMMIT_IMPL2(...) ::AutoRTFM::Private::FOnCommitHelper{} + [__VA_ARGS__]()
-#define UE_AUTORTFM_TRANSACT_IMPL2 ::AutoRTFM::Private::FTransactHelper{} + [&]() { __VA_ARGS__ })
+#define UE_AUTORTFM_OPEN_IMPL          ::AutoRTFM::Private::FOpenHelper{} + [&]()
+#define UE_AUTORTFM_ONABORT_IMPL(...)  ::AutoRTFM::Private::FOnAbortHelper{} + [__VA_ARGS__]()
+#define UE_AUTORTFM_ONCOMMIT_IMPL(...) ::AutoRTFM::Private::FOnCommitHelper{} + [__VA_ARGS__]()
+#define UE_AUTORTFM_TRANSACT_IMPL      ::AutoRTFM::Private::FTransactHelper{} + [&]() { __VA_ARGS__ })
 #else
 
-// Older macros where the code is passed as a macro argument. These should be phased out as they make debugging more difficult
-#define UE_AUTORTFM_OPEN_IMPL(...) do { __VA_ARGS__ } while (false)
-#define UE_AUTORTFM_ONABORT_IMPL(...) do { /* do nothing */ } while (false)
-#define UE_AUTORTFM_ONCOMMIT_IMPL(...) do { __VA_ARGS__ } while (false)
-#define UE_AUTORTFM_TRANSACT_IMPL(...) do { __VA_ARGS__ } while (false)
-
 // Do nothing, these should be followed by blocks that should be either executed or not executed
-#define UE_AUTORTFM_OPEN_IMPL2
-#define UE_AUTORTFM_ONABORT_IMPL2(...) while (false)
-#define UE_AUTORTFM_ONCOMMIT_IMPL2(...)
-#define UE_AUTORTFM_TRANSACT_IMPL2
+#define UE_AUTORTFM_OPEN_IMPL
+#define UE_AUTORTFM_ONABORT_IMPL(...) while (false)
+#define UE_AUTORTFM_ONCOMMIT_IMPL(...)
+#define UE_AUTORTFM_TRANSACT_IMPL
 #endif
 
 // Runs a block of code in the open, non-transactionally. Anything performed in the open will not be undone if a transaction fails.
-#define UE_AUTORTFM_OPEN(...) UE_AUTORTFM_OPEN_IMPL(AutoRTFM::ForTheRuntime::DeprecatedUseOpenMacro(); __VA_ARGS__)
-// This new version is used like UE_AUTORTFM_OPEN2 { ... code ... };
-#define UE_AUTORTFM_OPEN2 UE_AUTORTFM_OPEN_IMPL2
+// Calls should be written like this: UE_AUTORTFM_OPEN { ... code ... };
+#define UE_AUTORTFM_OPEN  UE_AUTORTFM_OPEN_IMPL
+#define UE_AUTORTFM_OPEN2 UE_AUTORTFM_OPEN_IMPL
 
 // Runs a block of code if a transaction aborts.
 // In non-transactional code paths the block of code will not be executed at all.
-// This captures any used variables from the parent function by-value.
-#define UE_AUTORTFM_ONABORT(...) UE_AUTORTFM_ONABORT_IMPL(AutoRTFM::ForTheRuntime::DeprecatedUseOnAbortMacro(); __VA_ARGS__)
-// In the new version of the macro, the macro arguments are the capture specification for the lambda
-// Used like UE_AUTORTFM_ONABORT2(=) { ... code ... };
-#define UE_AUTORTFM_ONABORT2(...) UE_AUTORTFM_ONABORT_IMPL2(__VA_ARGS__)
+// The macro arguments are the capture specification for the lambda.
+// Calls should be written like this: UE_AUTORTFM_ONABORT(=) { ... code ... };
+#define UE_AUTORTFM_ONABORT(...)  UE_AUTORTFM_ONABORT_IMPL(__VA_ARGS__)
+#define UE_AUTORTFM_ONABORT2(...) UE_AUTORTFM_ONABORT_IMPL(__VA_ARGS__)
 
 // Runs a block of code if a transaction commits successfully.
 // In non-transactional code paths the block of code will be executed immediately.
-// This captures any used variables from the parent function by-value.
-#define UE_AUTORTFM_ONCOMMIT(...) UE_AUTORTFM_ONCOMMIT_IMPL(AutoRTFM::ForTheRuntime::DeprecatedUseOnCommitMacro(); __VA_ARGS__)
-// In the new version of the macro, the macro arguments are the capture specification for the lambda
-// Used like UE_AUTORTFM_ONCOMMIT2(=) { ... code ... };
-#define UE_AUTORTFM_ONCOMMIT2(...) UE_AUTORTFM_ONCOMMIT_IMPL2(__VA_ARGS__)
+// The macro arguments are the capture specification for the lambda.
+// Calls should be written like this: UE_AUTORTFM_ONCOMMIT(=) { ... code ... };
+#define UE_AUTORTFM_ONCOMMIT(...)  UE_AUTORTFM_ONCOMMIT_IMPL(__VA_ARGS__)
+#define UE_AUTORTFM_ONCOMMIT2(...) UE_AUTORTFM_ONCOMMIT_IMPL(__VA_ARGS__)
 
 // Runs a block of code in the closed, transactionally, within a new transaction.
-#define UE_AUTORTFM_TRANSACT(...) UE_AUTORTFM_TRANSACT_IMPL(AutoRTFM::ForTheRuntime::DeprecatedUseTransactMacro(); __VA_ARGS__)
-// New version is used like a block: UE_AUTORTFM_TRANSACT2 { ... code ... };
-#define UE_AUTORTFM_TRANSACT2 UE_AUTORTFM_TRANSACT_IMPL2
+// Calls should be written like this: UE_AUTORTFM_TRANSACT { ... code ... };
+#define UE_AUTORTFM_TRANSACT  UE_AUTORTFM_TRANSACT_IMPL
+#define UE_AUTORTFM_TRANSACT2 UE_AUTORTFM_TRANSACT_IMPL
 
 #define UE_AUTORTFM_CONCAT_IMPL(A, B) A ## B
 #define UE_AUTORTFM_CONCAT(A, B) UE_AUTORTFM_CONCAT_IMPL(A, B)
