PR #12253: Add missing entry in ENUMERATE_SET_BY macro for ECVF_SetByGameOverride

#rnx

diff --git a/Engine/Source/Runtime/Core/Public/HAL/IConsoleManager.h b/Engine/Source/Runtime/Core/Public/HAL/IConsoleManager.h
--- a/Engine/Source/Runtime/Core/Public/HAL/IConsoleManager.h
+++ b/Engine/Source/Runtime/Core/Public/HAL/IConsoleManager.h
@@ -181,6 +181,7 @@ enum EConsoleVariableFlags
 	op(PluginLowPriority) \
 	op(DeviceProfile) \
 	op(PluginHighPriority) \
+	op(GameOverride) \
 	op(ConsoleVariablesIni) \
 	op(Hotfix) \
 	op(Preview) \
