Enable memory mapped animations and audio on Android

diff --git a/Engine/Source/Runtime/Core/Public/Android/AndroidPlatformProperties.h b/Engine/Source/Runtime/Core/Public/Android/AndroidPlatformProperties.h
--- a/Engine/Source/Runtime/Core/Public/Android/AndroidPlatformProperties.h
+++ b/Engine/Source/Runtime/Core/Public/Android/AndroidPlatformProperties.h
@@ -119,6 +119,16 @@ struct FAndroidPlatformProperties
 		return true;
 	}
 
+	static FORCEINLINE bool SupportsMemoryMappedAudio()
+	{
+		return true;
+	}
+
+	static FORCEINLINE bool SupportsMemoryMappedAnimation()
+	{
+		return true;
+	}
+
 	static FORCEINLINE bool SupportsVirtualTextureStreaming()
 	{
 		return true;
