Increasing the verbosity of the warning logged by SymLoadModuleExW failures to prevent log spam at crash time when using anticheat.

diff --git a/Engine/Source/Runtime/Core/Private/Windows/WindowsPlatformStackWalk.cpp b/Engine/Source/Runtime/Core/Private/Windows/WindowsPlatformStackWalk.cpp
--- a/Engine/Source/Runtime/Core/Private/Windows/WindowsPlatformStackWalk.cpp
+++ b/Engine/Source/Runtime/Core/Private/Windows/WindowsPlatformStackWalk.cpp
@@ -787,7 +787,7 @@ void LoadSymbolsForModule(HMODULE ModuleHandle, const FString& RemoteStorage)
 		{
 			TCHAR ErrorMsgBuffer[2048];
 			const TCHAR* ErrorMsg = FWindowsPlatformMisc::GetSystemErrorMessage(ErrorMsgBuffer, UE_ARRAY_COUNT(ErrorMsgBuffer), ErrorCode);
-			UE_LOG(LogWindows, Warning, TEXT("Failed to load module '%s' with SymLoadModuleExW. Error %d: %s"), ModuleName, ErrorCode, ErrorMsg);
+			UE_LOG(LogWindows, Verbose, TEXT("Failed to load module '%s' with SymLoadModuleExW. Error %d: %s"), ModuleName, ErrorCode, ErrorMsg);
 		}
 	}
 }
