Enable memory mapped animations on iOS

diff --git a/Engine/Source/Runtime/Core/Public/IOS/IOSPlatformProperties.h b/Engine/Source/Runtime/Core/Public/IOS/IOSPlatformProperties.h
--- a/Engine/Source/Runtime/Core/Public/IOS/IOSPlatformProperties.h
+++ b/Engine/Source/Runtime/Core/Public/IOS/IOSPlatformProperties.h
@@ -78,7 +78,7 @@ struct FIOSPlatformProperties
 	}
 	static FORCEINLINE bool SupportsMemoryMappedAnimation()
 	{
-		return false;
+		return true;
 	}
 	static FORCEINLINE int64 GetMemoryMappingAlignment()
 	{
