[AutoRTFM] Migrate existing ONCOMMIT2 blocks to ONCOMMIT.

As of https://p4-swarm.epicgames.net/changes/36019027 there is no longer any difference between ONCOMMIT and ONCOMMIT2.

diff --git a/Engine/Source/Runtime/Core/Private/FramePro/FrameProProfiler.cpp b/Engine/Source/Runtime/Core/Private/FramePro/FrameProProfiler.cpp
--- a/Engine/Source/Runtime/Core/Private/FramePro/FrameProProfiler.cpp
+++ b/Engine/Source/Runtime/Core/Private/FramePro/FrameProProfiler.cpp
@@ -149,7 +149,7 @@ public:
 	}
 	virtual void Free(void* p) override
 	{
-		UE_AUTORTFM_ONCOMMIT2(=, this)
+		UE_AUTORTFM_ONCOMMIT(=, this)
 		{
 			BaseMalloc.Free(p);
 		};
@@ -216,7 +216,7 @@ public:
 
 	static void Free(void* Original)
 	{
-		UE_AUTORTFM_ONCOMMIT2(Original)
+		UE_AUTORTFM_ONCOMMIT(Original)
 		{
 			FrameProAllocator::Get().GetBaseMalloc()->Free(Original);
 		};
diff --git a/Engine/Source/Runtime/Core/Public/HAL/FMemory.inl b/Engine/Source/Runtime/Core/Public/HAL/FMemory.inl
--- a/Engine/Source/Runtime/Core/Public/HAL/FMemory.inl
+++ b/Engine/Source/Runtime/Core/Public/HAL/FMemory.inl
@@ -147,7 +147,7 @@ FMEMORY_INLINE_FUNCTION_DECORATOR void FMemory::Free(void* Original)
 	// AutoRTFM: For transactional code, in order to support the transaction 
 	// aborting and needing to 'roll back' the Free, we defer the actual
 	// free until commit time.
-	UE_AUTORTFM_ONCOMMIT2(=)
+	UE_AUTORTFM_ONCOMMIT(=)
 	{
 		// optional tracking of every allocation
 		LLM_IF_ENABLED(FLowLevelMemTracker::Get().OnLowLevelFree(ELLMTracker::Default, Original, ELLMAllocType::FMalloc));
diff --git a/Engine/Source/Runtime/Core/Public/Templates/SharedPointerInternals.h b/Engine/Source/Runtime/Core/Public/Templates/SharedPointerInternals.h
--- a/Engine/Source/Runtime/Core/Public/Templates/SharedPointerInternals.h
+++ b/Engine/Source/Runtime/Core/Public/Templates/SharedPointerInternals.h
@@ -213,7 +213,7 @@ namespace SharedPointerInternals
 		{
 			if constexpr (Mode == ESPMode::ThreadSafe)
 			{
-				UE_AUTORTFM_ONCOMMIT2(this)
+				UE_AUTORTFM_ONCOMMIT(this)
 				{
 					// std::memory_order_acq_rel is used here so that, if we do end up executing the destructor, it's not possible
 					// for side effects from executing the destructor end up being visible before we've determined that the shared
@@ -286,7 +286,7 @@ namespace SharedPointerInternals
 		{
 			if constexpr (Mode == ESPMode::ThreadSafe)
 			{
-				UE_AUTORTFM_ONCOMMIT2(this)
+				UE_AUTORTFM_ONCOMMIT(this)
 					{
 						// See ReleaseSharedReference for the same reasons that std::memory_order_acq_rel is used in this function.
 
diff --git a/Engine/Source/Runtime/CoreUObject/Public/UObject/UObjectAnnotation.h b/Engine/Source/Runtime/CoreUObject/Public/UObject/UObjectAnnotation.h
--- a/Engine/Source/Runtime/CoreUObject/Public/UObject/UObjectAnnotation.h
+++ b/Engine/Source/Runtime/CoreUObject/Public/UObject/UObjectAnnotation.h
@@ -183,7 +183,7 @@ public:
 	 */
 	void RemoveAnnotation(const UObjectBase *Object)
 	{
-		UE_AUTORTFM_ONCOMMIT2(this, Object)
+		UE_AUTORTFM_ONCOMMIT(this, Object)
 		{
 			check(Object);
 			bool bHadElements = false;
diff --git a/Engine/Source/Runtime/CoreUObject/Public/UObject/UObjectBaseUtility.h b/Engine/Source/Runtime/CoreUObject/Public/UObject/UObjectBaseUtility.h
--- a/Engine/Source/Runtime/CoreUObject/Public/UObject/UObjectBaseUtility.h
+++ b/Engine/Source/Runtime/CoreUObject/Public/UObject/UObjectBaseUtility.h
@@ -180,7 +180,7 @@ public:
 
 		// only set these flags after we have finished our transaction
 		// GC should not happen during a transaction so delaying this wont cause GC issues
-		UE_AUTORTFM_ONCOMMIT2(this)
+		UE_AUTORTFM_ONCOMMIT(this)
 		{
 			AtomicallySetFlags(RF_MirroredGarbage);
 			GUObjectArray.IndexToObject(InternalIndex)->SetGarbage();
diff --git a/Engine/Source/Runtime/Engine/Private/Actor.cpp b/Engine/Source/Runtime/Engine/Private/Actor.cpp
--- a/Engine/Source/Runtime/Engine/Private/Actor.cpp
+++ b/Engine/Source/Runtime/Engine/Private/Actor.cpp
@@ -249,7 +249,7 @@ void AActor::InitializeDefaults()
 #if (CSV_PROFILER_STATS && !UE_BUILD_SHIPPING)
 	// Increment actor class count
 	// Update our ActorClassName count after the transaction finishes
-	UE_AUTORTFM_ONCOMMIT2(this)
+	UE_AUTORTFM_ONCOMMIT(this)
 	{
 		if (!HasAnyFlags(RF_ArchetypeObject | RF_ClassDefaultObject))
 		{
@@ -805,7 +805,7 @@ void AActor::BeginDestroy()
 #if (CSV_PROFILER_STATS && !UE_BUILD_SHIPPING)
 	// Decrement actor class count
 	// Update our ActorClassName count after the transaction finishes
-	UE_AUTORTFM_ONCOMMIT2(this)
+	UE_AUTORTFM_ONCOMMIT(this)
 	{
 		if (!HasAnyFlags(RF_ArchetypeObject | RF_ClassDefaultObject))
 		{
diff --git a/Engine/Source/Runtime/Experimental/Chaos/Public/Chaos/RefCountedObject.h b/Engine/Source/Runtime/Experimental/Chaos/Public/Chaos/RefCountedObject.h
--- a/Engine/Source/Runtime/Experimental/Chaos/Public/Chaos/RefCountedObject.h
+++ b/Engine/Source/Runtime/Experimental/Chaos/Public/Chaos/RefCountedObject.h
@@ -24,7 +24,7 @@ public:
 			// time during OnCommit. (If an ephemeral stack object is destroyed a non-zero refcount, this is user error;
 			// we might report this by check-failing here with a garbage value for its refcount, as its stack representation
 			// might already be overwritten.)
-			UE_AUTORTFM_ONCOMMIT2(this)
+			UE_AUTORTFM_ONCOMMIT(this)
 			{
 				check(GetRefCount() == 0);
 			};
@@ -62,7 +62,7 @@ public:
 
 	uint32 Release() const
 	{
-		UE_AUTORTFM_ONCOMMIT2(this)
+		UE_AUTORTFM_ONCOMMIT(this)
 		{
 			if (--NumRefs == 0 && RefCountMode == ERCM_Transient)
 			{
