Open around StackWalkAndDump when asserting and ensuring
For Windows open around all functions that take the GlobalLock
Try to make up the same ones on Linux so one doesnt fail on one while pass on the other platform. Some are missed but can address later

diff --git a/Engine/Source/Runtime/Core/Private/Unix/UnixPlatformStackWalk.cpp b/Engine/Source/Runtime/Core/Private/Unix/UnixPlatformStackWalk.cpp
--- a/Engine/Source/Runtime/Core/Private/Unix/UnixPlatformStackWalk.cpp
+++ b/Engine/Source/Runtime/Core/Private/Unix/UnixPlatformStackWalk.cpp
@@ -518,12 +518,9 @@ namespace
 	}
 }
 
-void FUnixPlatformStackWalk::ProgramCounterToSymbolInfo( uint64 ProgramCounter, FProgramCounterSymbolInfo& out_SymbolInfo )
+UE_AUTORTFM_ALWAYS_OPEN void FUnixPlatformStackWalk::ProgramCounterToSymbolInfo( uint64 ProgramCounter, FProgramCounterSymbolInfo& out_SymbolInfo )
 {
-	UE_AUTORTFM_OPEN2
-	{
-		PopulateProgramCounterSymbolInfoFromSymbolFile(ProgramCounter, out_SymbolInfo);
-	};
+	PopulateProgramCounterSymbolInfoFromSymbolFile(ProgramCounter, out_SymbolInfo);
 }
 
 bool FUnixPlatformStackWalk::ProgramCounterToHumanReadableString( int32 CurrentCallDepth, uint64 ProgramCounter, ANSICHAR* HumanReadableString, SIZE_T HumanReadableStringSize, FGenericCrashContext* Context )
@@ -667,7 +664,7 @@ bool FUnixPlatformStackWalk::ProgramCounterToHumanReadableString( int32 CurrentC
 	return true;
 }
 
-void FUnixPlatformStackWalk::StackWalkAndDump( ANSICHAR* HumanReadableString, SIZE_T HumanReadableStringSize, int32 IgnoreCount, void* Context )
+UE_AUTORTFM_ALWAYS_OPEN void FUnixPlatformStackWalk::StackWalkAndDump( ANSICHAR* HumanReadableString, SIZE_T HumanReadableStringSize, int32 IgnoreCount, void* Context )
 {
 	if (Context == nullptr)
 	{
@@ -709,7 +706,7 @@ namespace
 	};
 } // namespace
 
-void FUnixPlatformStackWalk::StackWalkAndDumpEx(ANSICHAR* HumanReadableString, SIZE_T HumanReadableStringSize, int32 IgnoreCount, uint32 Flags, void* Context)
+UE_AUTORTFM_ALWAYS_OPEN void FUnixPlatformStackWalk::StackWalkAndDumpEx(ANSICHAR* HumanReadableString, SIZE_T HumanReadableStringSize, int32 IgnoreCount, uint32 Flags, void* Context)
 {
 	const bool bHandlingEnsure = (Flags & EStackWalkFlags::FlagsUsedWhenHandlingEnsure) == EStackWalkFlags::FlagsUsedWhenHandlingEnsure;
 	GHandlingEnsure = bHandlingEnsure;
@@ -731,7 +728,7 @@ void FUnixPlatformStackWalk::StackWalkAndDumpEx(ANSICHAR* HumanReadableString, S
 	GHandlingEnsure = false;
 }
 
-void FUnixPlatformStackWalk::StackWalkAndDumpEx(ANSICHAR* HumanReadableString, SIZE_T HumanReadableStringSize, void* ProgramCounter, uint32 Flags, void* Context)
+UE_AUTORTFM_ALWAYS_OPEN void FUnixPlatformStackWalk::StackWalkAndDumpEx(ANSICHAR* HumanReadableString, SIZE_T HumanReadableStringSize, void* ProgramCounter, uint32 Flags, void* Context)
 {
 	const bool bHandlingEnsure = (Flags & EStackWalkFlags::FlagsUsedWhenHandlingEnsure) == EStackWalkFlags::FlagsUsedWhenHandlingEnsure;
 	GHandlingEnsure = bHandlingEnsure;
@@ -781,7 +778,7 @@ namespace
 	}
 }
 
-uint32 FUnixPlatformStackWalk::CaptureStackBackTrace( uint64* BackTrace, uint32 MaxDepth, void* Context )
+UE_AUTORTFM_ALWAYS_OPEN uint32 FUnixPlatformStackWalk::CaptureStackBackTrace( uint64* BackTrace, uint32 MaxDepth, void* Context )
 {
 	// Make sure we have place to store the information before we go through the process of raising
 	// an exception and handling it.
@@ -843,7 +840,7 @@ namespace
 	}
 }
 
-void FUnixPlatformStackWalk::ThreadStackWalkAndDump(ANSICHAR* HumanReadableString, SIZE_T HumanReadableStringSize, int32 IgnoreCount, uint32 ThreadId)
+UE_AUTORTFM_ALWAYS_OPEN void FUnixPlatformStackWalk::ThreadStackWalkAndDump(ANSICHAR* HumanReadableString, SIZE_T HumanReadableStringSize, int32 IgnoreCount, uint32 ThreadId)
 {
 	ThreadStackUserData ThreadCallStack;
 	ThreadCallStack.bCaptureCallStack = true;
@@ -855,7 +852,7 @@ void FUnixPlatformStackWalk::ThreadStackWalkAndDump(ANSICHAR* HumanReadableStrin
 	GatherCallstackFromThread(ThreadCallStack, ThreadId);
 }
 
-uint32 FUnixPlatformStackWalk::CaptureThreadStackBackTrace(uint64 ThreadId, uint64* BackTrace, uint32 MaxDepth, void* Context)
+UE_AUTORTFM_ALWAYS_OPEN uint32 FUnixPlatformStackWalk::CaptureThreadStackBackTrace(uint64 ThreadId, uint64* BackTrace, uint32 MaxDepth, void* Context)
 {
 	ThreadStackUserData ThreadBackTrace;
 	ThreadBackTrace.bCaptureCallStack = false;
@@ -885,7 +882,7 @@ namespace
 	}
 }
 
-int32 FUnixPlatformStackWalk::GetProcessModuleCount()
+UE_AUTORTFM_ALWAYS_OPEN int32 FUnixPlatformStackWalk::GetProcessModuleCount()
 {
 	int Size = 0;
 	dl_iterate_phdr(NumberOfDynamicLibrariesCallback, &Size);
@@ -950,7 +947,7 @@ namespace
 	}
 }
 
-int32 FUnixPlatformStackWalk::GetProcessModuleSignatures(FStackWalkModuleInfo *ModuleSignatures, const int32 ModuleSignaturesSize)
+UE_AUTORTFM_ALWAYS_OPEN int32 FUnixPlatformStackWalk::GetProcessModuleSignatures(FStackWalkModuleInfo *ModuleSignatures, const int32 ModuleSignaturesSize)
 {
 	if (ModuleSignatures == nullptr || ModuleSignaturesSize == 0)
 	{
diff --git a/Engine/Source/Runtime/Core/Private/Windows/WindowsPlatformStackWalk.cpp b/Engine/Source/Runtime/Core/Private/Windows/WindowsPlatformStackWalk.cpp
--- a/Engine/Source/Runtime/Core/Private/Windows/WindowsPlatformStackWalk.cpp
+++ b/Engine/Source/Runtime/Core/Private/Windows/WindowsPlatformStackWalk.cpp
@@ -136,7 +136,7 @@ void DetermineMaxCallstackDepth()
 
 #endif
 
-void FWindowsPlatformStackWalk::StackWalkAndDump( ANSICHAR* HumanReadableString, SIZE_T HumanReadableStringSize, int32 IgnoreCount, void* Context )
+UE_AUTORTFM_ALWAYS_OPEN void FWindowsPlatformStackWalk::StackWalkAndDump( ANSICHAR* HumanReadableString, SIZE_T HumanReadableStringSize, int32 IgnoreCount, void* Context )
 {
 	UE::TUniqueLock GlobalLock(GStackWalkingLock);
 
@@ -157,7 +157,7 @@ void FWindowsPlatformStackWalk::StackWalkAndDump( ANSICHAR* HumanReadableString,
 	}
 }
 
-void FWindowsPlatformStackWalk::StackWalkAndDump( ANSICHAR* HumanReadableString, SIZE_T HumanReadableStringSize, void* ProgramCounter, void* Context )
+UE_AUTORTFM_ALWAYS_OPEN void FWindowsPlatformStackWalk::StackWalkAndDump( ANSICHAR* HumanReadableString, SIZE_T HumanReadableStringSize, void* ProgramCounter, void* Context )
 {
 	UE::TUniqueLock GlobalLock(GStackWalkingLock);
 
@@ -165,7 +165,7 @@ void FWindowsPlatformStackWalk::StackWalkAndDump( ANSICHAR* HumanReadableString,
 	FGenericPlatformStackWalk::StackWalkAndDump(HumanReadableString, HumanReadableStringSize, ProgramCounter, Context);
 }
 
-FORCENOINLINE TArray<FProgramCounterSymbolInfo> FWindowsPlatformStackWalk::GetStack(int32 IgnoreCount, int32 MaxDepth, void* Context)
+UE_AUTORTFM_ALWAYS_OPEN TArray<FProgramCounterSymbolInfo> FWindowsPlatformStackWalk::GetStack(int32 IgnoreCount, int32 MaxDepth, void* Context)
 {
 	UE::TUniqueLock GlobalLock(GStackWalkingLock);
 
@@ -179,7 +179,7 @@ FORCENOINLINE TArray<FProgramCounterSymbolInfo> FWindowsPlatformStackWalk::GetSt
 	return FGenericPlatformStackWalk::GetStack(IgnoreCount, MaxDepth, Context);
 }
 
-void FWindowsPlatformStackWalk::ThreadStackWalkAndDump(ANSICHAR* HumanReadableString, SIZE_T HumanReadableStringSize, int32 IgnoreCount, uint32 ThreadId)
+UE_AUTORTFM_ALWAYS_OPEN void FWindowsPlatformStackWalk::ThreadStackWalkAndDump(ANSICHAR* HumanReadableString, SIZE_T HumanReadableStringSize, int32 IgnoreCount, uint32 ThreadId)
 {
 	UE::TUniqueLock GlobalLock(GStackWalkingLock);
 
@@ -313,7 +313,7 @@ void FWindowsPlatformStackWalk::CaptureStackTraceByProcess(uint64* OutBacktrace,
 	CaptureStackTraceExternalProcess(OutBacktrace, MaxDepth, reinterpret_cast<PCONTEXT>(InContext), reinterpret_cast<HANDLE>(InThreadHandle), OutDepth);	
 }
 
-uint32 FWindowsPlatformStackWalk::CaptureThreadStackBackTrace(uint64 ThreadId, uint64* BackTrace, uint32 MaxDepth, void* Context)
+UE_AUTORTFM_ALWAYS_OPEN uint32 FWindowsPlatformStackWalk::CaptureThreadStackBackTrace(uint64 ThreadId, uint64* BackTrace, uint32 MaxDepth, void* Context)
 {
 	UE::TUniqueLock GlobalLock(GStackWalkingLock);
 
@@ -383,7 +383,7 @@ uint32 FWindowsPlatformStackWalk::CaptureThreadStackBackTrace(uint64 ThreadId, u
  * @param	MaxDepth			Entries in BackTrace array
  * @param	Context				Optional thread context information (FWindowsThreadContextWrapper instance)
  */
-uint32 FWindowsPlatformStackWalk::CaptureStackBackTrace( uint64* BackTrace, uint32 MaxDepth, void* Context )
+UE_AUTORTFM_ALWAYS_OPEN uint32 FWindowsPlatformStackWalk::CaptureStackBackTrace( uint64* BackTrace, uint32 MaxDepth, void* Context )
 {
 	UE::TUniqueLock GlobalLock(GStackWalkingLock);
 
@@ -456,82 +456,79 @@ uint32 FWindowsPlatformStackWalk::CaptureStackBackTrace( uint64* BackTrace, uint
 	return Depth;
 }
 
-void FWindowsPlatformStackWalk::ProgramCounterToSymbolInfo( uint64 ProgramCounter, FProgramCounterSymbolInfo& out_SymbolInfo )
+UE_AUTORTFM_ALWAYS_OPEN void FWindowsPlatformStackWalk::ProgramCounterToSymbolInfo( uint64 ProgramCounter, FProgramCounterSymbolInfo& out_SymbolInfo )
 {
-	UE_AUTORTFM_OPEN2
-	{
-		UE::TUniqueLock GlobalLock(GStackWalkingLock);
+	UE::TUniqueLock GlobalLock(GStackWalkingLock);
 
-		// Initialize stack walking as it loads up symbol information which we require.
-		InitStackWalking();
+	// Initialize stack walking as it loads up symbol information which we require.
+	InitStackWalking();
 
 #if ON_DEMAND_SYMBOL_LOADING
-		// Load symbols for the module
-		bool bShouldReloadModuleMissingDebugSymbols = !FPlatformProperties::IsMonolithicBuild() && FPlatformStackWalk::WantsDetailedCallstacksInNonMonolithicBuilds();
-		LoadSymbolsForModuleByAddress(ProgramCounter, GetSymbolSearchPath(), bShouldReloadModuleMissingDebugSymbols);
+	// Load symbols for the module
+	bool bShouldReloadModuleMissingDebugSymbols = !FPlatformProperties::IsMonolithicBuild() && FPlatformStackWalk::WantsDetailedCallstacksInNonMonolithicBuilds();
+	LoadSymbolsForModuleByAddress(ProgramCounter, GetSymbolSearchPath(), bShouldReloadModuleMissingDebugSymbols);
 #endif
 
-		// Set the program counter.
-		out_SymbolInfo.ProgramCounter = ProgramCounter;
-
-		uint32 LastError = 0;
-		HANDLE ProcessHandle = GProcessHandle;
+	// Set the program counter.
+	out_SymbolInfo.ProgramCounter = ProgramCounter;
 
-		// Initialize symbol.
-		ANSICHAR SymbolBuffer[sizeof( SYMBOL_INFO ) + FProgramCounterSymbolInfo::MAX_NAME_LENGTH] = {0};
-		SYMBOL_INFO* Symbol = (SYMBOL_INFO*)SymbolBuffer;
-		Symbol->SizeOfStruct = sizeof(SYMBOL_INFO);
-		Symbol->MaxNameLen = FProgramCounterSymbolInfo::MAX_NAME_LENGTH;
+	uint32 LastError = 0;
+	HANDLE ProcessHandle = GProcessHandle;
 
-		// Get function name.
-		if( SymFromAddr( ProcessHandle, ProgramCounter, nullptr, Symbol ) )
-		{
-			// Skip any funky chars in the beginning of a function name.
-			int32 Offset = 0;
-			while( Symbol->Name[Offset] < 32 || Symbol->Name[Offset] > 127 )
-			{
-				Offset++;
-			}
+	// Initialize symbol.
+	ANSICHAR SymbolBuffer[sizeof( SYMBOL_INFO ) + FProgramCounterSymbolInfo::MAX_NAME_LENGTH] = {0};
+	SYMBOL_INFO* Symbol = (SYMBOL_INFO*)SymbolBuffer;
+	Symbol->SizeOfStruct = sizeof(SYMBOL_INFO);
+	Symbol->MaxNameLen = FProgramCounterSymbolInfo::MAX_NAME_LENGTH;
 
-			// Write out function name.
-			FCStringAnsi::Strncpy( out_SymbolInfo.FunctionName, Symbol->Name + Offset, FProgramCounterSymbolInfo::MAX_NAME_LENGTH ); 
-			FCStringAnsi::Strncat( out_SymbolInfo.FunctionName, "()", FProgramCounterSymbolInfo::MAX_NAME_LENGTH );
-		}
-		else
+	// Get function name.
+	if( SymFromAddr( ProcessHandle, ProgramCounter, nullptr, Symbol ) )
+	{
+		// Skip any funky chars in the beginning of a function name.
+		int32 Offset = 0;
+		while( Symbol->Name[Offset] < 32 || Symbol->Name[Offset] > 127 )
 		{
-			// No symbol found for this address.
-			LastError = GetLastError();
+			Offset++;
 		}
 
-		// Get filename and line number.
-		IMAGEHLP_LINE64	ImageHelpLine = {0};
-		ImageHelpLine.SizeOfStruct = sizeof( ImageHelpLine );
-		if( SymGetLineFromAddr64( ProcessHandle, ProgramCounter, (::DWORD *)&out_SymbolInfo.SymbolDisplacement, &ImageHelpLine ) )
-		{
-			FCStringAnsi::Strncpy( out_SymbolInfo.Filename, ImageHelpLine.FileName, FProgramCounterSymbolInfo::MAX_NAME_LENGTH );
-			out_SymbolInfo.LineNumber = ImageHelpLine.LineNumber;
-		}
-		else
-		{
-			LastError = GetLastError();
-		}
+		// Write out function name.
+		FCStringAnsi::Strncpy( out_SymbolInfo.FunctionName, Symbol->Name + Offset, FProgramCounterSymbolInfo::MAX_NAME_LENGTH ); 
+		FCStringAnsi::Strncat( out_SymbolInfo.FunctionName, "()", FProgramCounterSymbolInfo::MAX_NAME_LENGTH );
+	}
+	else
+	{
+		// No symbol found for this address.
+		LastError = GetLastError();
+	}
 
-		// Get module name.
-		IMAGEHLP_MODULE64 ImageHelpModule = {0};
-		ImageHelpModule.SizeOfStruct = sizeof( ImageHelpModule );
-		if( SymGetModuleInfo64( ProcessHandle, ProgramCounter, &ImageHelpModule) )
-		{
-			// Write out module information.
-			FCStringAnsi::Strncpy( out_SymbolInfo.ModuleName, ImageHelpModule.ImageName, FProgramCounterSymbolInfo::MAX_NAME_LENGTH );
-		}
-		else
-		{
-			LastError = GetLastError();
-		}
-	};
+	// Get filename and line number.
+	IMAGEHLP_LINE64	ImageHelpLine = {0};
+	ImageHelpLine.SizeOfStruct = sizeof( ImageHelpLine );
+	if( SymGetLineFromAddr64( ProcessHandle, ProgramCounter, (::DWORD *)&out_SymbolInfo.SymbolDisplacement, &ImageHelpLine ) )
+	{
+		FCStringAnsi::Strncpy( out_SymbolInfo.Filename, ImageHelpLine.FileName, FProgramCounterSymbolInfo::MAX_NAME_LENGTH );
+		out_SymbolInfo.LineNumber = ImageHelpLine.LineNumber;
+	}
+	else
+	{
+		LastError = GetLastError();
+	}
+
+	// Get module name.
+	IMAGEHLP_MODULE64 ImageHelpModule = {0};
+	ImageHelpModule.SizeOfStruct = sizeof( ImageHelpModule );
+	if( SymGetModuleInfo64( ProcessHandle, ProgramCounter, &ImageHelpModule) )
+	{
+		// Write out module information.
+		FCStringAnsi::Strncpy( out_SymbolInfo.ModuleName, ImageHelpModule.ImageName, FProgramCounterSymbolInfo::MAX_NAME_LENGTH );
+	}
+	else
+	{
+		LastError = GetLastError();
+	}
 }
 
-void FWindowsPlatformStackWalk::ProgramCounterToSymbolInfoEx(uint64 ProgramCounter, FProgramCounterSymbolInfoEx& out_SymbolInfo)
+UE_AUTORTFM_ALWAYS_OPEN void FWindowsPlatformStackWalk::ProgramCounterToSymbolInfoEx(uint64 ProgramCounter, FProgramCounterSymbolInfoEx& out_SymbolInfo)
 {
 	UE::TUniqueLock GlobalLock(GStackWalkingLock);
 
@@ -734,7 +731,7 @@ bool FWindowsPlatformStackWalk::UploadLocalSymbols()
 	return true;
 }
 
-void LoadSymbolsForModule(HMODULE ModuleHandle, const FString& RemoteStorage)
+UE_AUTORTFM_ALWAYS_OPEN void LoadSymbolsForModule(HMODULE ModuleHandle, const FString& RemoteStorage)
 {
 	HANDLE ProcessHandle = GProcessHandle;
 
@@ -819,7 +816,7 @@ void LoadSymbolsForProcessModules(const FString &RemoteStorage)
 	FMemory::Free(ModuleHandlePointer);
 }
 
-void LoadSymbolsForModuleByAddress(uint64 Address, const FString& RemoteStorage, bool bShouldReloadModuleMissingDebugSymbols)
+UE_AUTORTFM_ALWAYS_OPEN void LoadSymbolsForModuleByAddress(uint64 Address, const FString& RemoteStorage, bool bShouldReloadModuleMissingDebugSymbols)
 {
 	UE::TUniqueLock GlobalLock(GStackWalkingLock);
 
@@ -867,7 +864,7 @@ int32 FWindowsPlatformStackWalk::GetProcessModuleCount()
 	return ModuleCount;
 }
 
-int32 FWindowsPlatformStackWalk::GetProcessModuleSignatures(FStackWalkModuleInfo *ModuleSignatures, const int32 ModuleSignaturesSize)
+UE_AUTORTFM_ALWAYS_OPEN int32 FWindowsPlatformStackWalk::GetProcessModuleSignatures(FStackWalkModuleInfo *ModuleSignatures, const int32 ModuleSignaturesSize)
 {
 	UE::TUniqueLock GlobalLock(GStackWalkingLock);
 
@@ -1061,7 +1058,7 @@ FString GetSymbolSearchPath()
 /**
  * Initializes the symbol engine if needed.
  */
-bool FWindowsPlatformStackWalk::InitStackWalkingInternal(void* Process, bool bForceReinitOnProcessMismatch)
+UE_AUTORTFM_ALWAYS_OPEN bool FWindowsPlatformStackWalk::InitStackWalkingInternal(void* Process, bool bForceReinitOnProcessMismatch)
 {
 	UE::TUniqueLock GlobalLock(GStackWalkingLock);
 
@@ -1221,7 +1218,7 @@ void FWindowsPlatformStackWalk::RegisterOnModulesChanged()
 	FModuleManager::Get().OnModulesChanged().AddStatic( &OnModulesChanged );
 }
 
-bool FWindowsPlatformStackWalk::GetFunctionDefinitionLocation(const FString& FunctionSymbolName, const FString& FunctionModuleName, FString& OutPathname, uint32& OutLineNumber, uint32& OutColumnNumber)
+UE_AUTORTFM_ALWAYS_OPEN bool FWindowsPlatformStackWalk::GetFunctionDefinitionLocation(const FString& FunctionSymbolName, const FString& FunctionModuleName, FString& OutPathname, uint32& OutLineNumber, uint32& OutColumnNumber)
 {
 	UE::TUniqueLock GlobalLock(GStackWalkingLock);
 
