Updated free memory stat to include any cached free memory by the GMalloc where applicable

diff --git a/Engine/Source/Runtime/Core/Private/GenericPlatform/GenericPlatformMemory.cpp b/Engine/Source/Runtime/Core/Private/GenericPlatform/GenericPlatformMemory.cpp
--- a/Engine/Source/Runtime/Core/Private/GenericPlatform/GenericPlatformMemory.cpp
+++ b/Engine/Source/Runtime/Core/Private/GenericPlatform/GenericPlatformMemory.cpp
@@ -281,6 +281,11 @@ FPlatformMemoryStats FGenericPlatformMemory::GetStatsRaw()
 	return FPlatformMemory::GetStats();
 }
 
+uint64 FGenericPlatformMemory::GetMemoryUsedFast()
+{
+	return GetStats().UsedPhysical;
+}
+
 void FGenericPlatformMemory::GetStatsForMallocProfiler( FGenericMemoryStats& out_Stats )
 {
 #if	STATS
diff --git a/Engine/Source/Runtime/Core/Private/HAL/Allocators/CachedOSVeryLargePageAllocator.cpp b/Engine/Source/Runtime/Core/Private/HAL/Allocators/CachedOSVeryLargePageAllocator.cpp
--- a/Engine/Source/Runtime/Core/Private/HAL/Allocators/CachedOSVeryLargePageAllocator.cpp
+++ b/Engine/Source/Runtime/Core/Private/HAL/Allocators/CachedOSVeryLargePageAllocator.cpp
@@ -68,8 +68,8 @@ static FAutoConsoleVariableRef CVarMaxPageCountSmallPool(
 #if CSV_PROFILER_STATS
 CSV_DECLARE_CATEGORY_MODULE_EXTERN(CORE_API, FMemory);
 
-static volatile int32 GLargePageAllocatorCommitCount = 0;
-static volatile int32 GLargePageAllocatorDecommitCount = 0;
+static std::atomic<int32> GLargePageAllocatorCommitCount = 0;
+static std::atomic<int32> GLargePageAllocatorDecommitCount = 0;
 #endif
 
 void FCachedOSVeryLargePageAllocator::Init()
@@ -159,6 +159,7 @@ void FCachedOSVeryLargePageAllocator::Refresh()
 					}
 					LargePage->LinkHead(EmptyButAvailableLargePagesHead[AllocationHint]);
 					EmptyBackStoreCount[AllocationHint] += 1;
+					ImmediatelyFreeable += SizeOfLargePage;
 				}
 			}
 		}
@@ -225,6 +226,7 @@ FCachedOSVeryLargePageAllocator::FLargePage* FCachedOSVeryLargePageAllocator::Ge
 		LargePage->AllocationHint = AllocationHint;
 		LargePage->Unlink();
 		EmptyBackStoreCount[AllocationHint] -= 1;
+		ImmediatelyFreeable -= SizeOfLargePage;
 	}
 	else
 	{
@@ -281,7 +283,7 @@ FCachedOSVeryLargePageAllocator::FLargePage* FCachedOSVeryLargePageAllocator::Al
 			CachedFree += SizeOfLargePage;
 			CommittedLargePagesCount[AllocationHint] += 1;
 #if CSV_PROFILER_STATS
-			FPlatformAtomics::InterlockedIncrement(&GLargePageAllocatorCommitCount);
+			GLargePageAllocatorCommitCount.fetch_add(1, std::memory_order_relaxed);
 #endif
 		}
 	}
@@ -309,6 +311,7 @@ void FCachedOSVeryLargePageAllocator::Free(void* Ptr, SIZE_T Size, FCriticalSect
 			{
 				LargePage->LinkHead(EmptyButAvailableLargePagesHead[LargePage->AllocationHint]);
 				EmptyBackStoreCount[LargePage->AllocationHint] += 1;
+				ImmediatelyFreeable += SizeOfLargePage;
 			}
 			else
 			{
@@ -323,7 +326,7 @@ void FCachedOSVeryLargePageAllocator::Free(void* Ptr, SIZE_T Size, FCriticalSect
 
 				CommittedLargePagesCount[LargePage->AllocationHint] -= 1;
 #if CSV_PROFILER_STATS
-				FPlatformAtomics::InterlockedIncrement(&GLargePageAllocatorDecommitCount);
+				GLargePageAllocatorDecommitCount.fetch_add(1, std::memory_order_relaxed);
 #endif
 
 				LargePage->LinkHead(FreeLargePagesHead[LargePage->AllocationHint]);
@@ -374,8 +377,13 @@ void FCachedOSVeryLargePageAllocator::ShrinkEmptyBackStore(int32 NewEmptyBackSto
 		LLM_IF_ENABLED(FLowLevelMemTracker::Get().OnLowLevelFree(ELLMTracker::Platform, (void*)LargePage->BaseAddress));
 		LargePage->LinkHead(FreeLargePagesHead[LargePage->AllocationHint]);
 		CachedFree -= SizeOfLargePage;
+		ImmediatelyFreeable -= SizeOfLargePage;
 		EmptyBackStoreCount[AllocationHint] -= 1;
 		CommittedLargePagesCount[LargePage->AllocationHint] -= 1;
+
+#if CSV_PROFILER_STATS
+		GLargePageAllocatorDecommitCount.fetch_add(1, std::memory_order_relaxed);
+#endif
 	}
 }
 
@@ -401,15 +409,15 @@ void FCachedOSVeryLargePageAllocator::FreeAll(FCriticalSection* Mutex)
 void FCachedOSVeryLargePageAllocator::UpdateStats()
 {
 #if CSV_PROFILER_STATS
-	CSV_CUSTOM_STAT(FMemory, LargeAllocatorCommitCount, GLargePageAllocatorCommitCount, ECsvCustomStatOp::Set);
-	CSV_CUSTOM_STAT(FMemory, LargeAllocatorDecommitCount, GLargePageAllocatorDecommitCount, ECsvCustomStatOp::Set);
+	CSV_CUSTOM_STAT(FMemory, LargeAllocatorCommitCount, GLargePageAllocatorCommitCount.load(std::memory_order_relaxed), ECsvCustomStatOp::Set);
+	CSV_CUSTOM_STAT(FMemory, LargeAllocatorDecommitCount, GLargePageAllocatorDecommitCount.load(std::memory_order_relaxed), ECsvCustomStatOp::Set);
 	CSV_CUSTOM_STAT(FMemory, LargeAllocatorBackingStoreCountSmall, EmptyBackStoreCount[FMemory::AllocationHints::SmallPool], ECsvCustomStatOp::Set);
 	CSV_CUSTOM_STAT(FMemory, LargeAllocatorBackingStoreCountDefault, EmptyBackStoreCount[FMemory::AllocationHints::Default], ECsvCustomStatOp::Set);
 	CSV_CUSTOM_STAT(FMemory, LargeAllocatorPageCountSmall, CommittedLargePagesCount[FMemory::AllocationHints::SmallPool], ECsvCustomStatOp::Set);
 	CSV_CUSTOM_STAT(FMemory, LargeAllocatorPageCountDefault, CommittedLargePagesCount[FMemory::AllocationHints::Default], ECsvCustomStatOp::Set);
 
-	GLargePageAllocatorCommitCount = 0;
-	GLargePageAllocatorDecommitCount = 0;
+	GLargePageAllocatorCommitCount.store(0, std::memory_order_relaxed);
+	GLargePageAllocatorDecommitCount.store(0, std::memory_order_relaxed);
 #endif
 }
 #endif
diff --git a/Engine/Source/Runtime/Core/Private/HAL/MallocBinned2.cpp b/Engine/Source/Runtime/Core/Private/HAL/MallocBinned2.cpp
--- a/Engine/Source/Runtime/Core/Private/HAL/MallocBinned2.cpp
+++ b/Engine/Source/Runtime/Core/Private/HAL/MallocBinned2.cpp
@@ -31,6 +31,12 @@ PRAGMA_DISABLE_UNSAFE_TYPECAST_WARNINGS
 #	define GBinned2MoveOSFreesOffTimeCriticalThreads	UE_DEFAULT_GBinned2MoveOSFreesOffTimeCriticalThreads
 #endif
 
+int32 EnableLegacyCachedOSPageAllocatorFreeMemReporting = 0;
+static FAutoConsoleVariableRef GEnableLegacyCachedOSPageAllocatorFreeMemReportingCVar(
+	TEXT("MallocBinned2.EnableLegacyCachedOSPageAllocatorFreeMemReporting"),
+	EnableLegacyCachedOSPageAllocatorFreeMemReporting,
+	TEXT("Temporary. Whether or not report full amount of cached memory in CachedOSPageAllocator (legacy) or report only the amount of memory immediately available to return to the OS")
+);
 
 #if UE_MB2_ALLOCATOR_STATS
 	std::atomic<int64> AllocatedOSSmallPoolMemory(0);
@@ -911,10 +917,6 @@ void FMallocBinned2::CanaryFail(const FFreeBlock* Block) const
 
 void FMallocBinned2::GetAllocatorStats( FGenericMemoryStats& OutStats )
 {
-	// Even if we have MB2 stats off, cached slack needs to be included as it might be needed by some platforms
-	const uint64 OSPageAllocatorCachedFreeSize = CachedOSPageAllocator.GetCachedFreeTotal();
-	OutStats.Add(TEXT("PageAllocatorFreeCacheSize"), OSPageAllocatorCachedFreeSize);
-
 #if UE_MB2_ALLOCATOR_STATS
 	const int64  TotalAllocatedSmallPoolMemory           = GetTotalAllocatedSmallPoolMemory();
 	const int64  LocalAllocatedOSSmallPoolMemory         = AllocatedOSSmallPoolMemory.load(std::memory_order_relaxed);
@@ -927,7 +929,7 @@ void FMallocBinned2::GetAllocatorStats( FGenericMemoryStats& OutStats )
 	OutStats.Add(TEXT("AllocatedLargePoolMemoryWAlignment"), LocalAllocatedLargePoolMemoryWAlignment);
 
 	const uint64 TotalAllocated = TotalAllocatedSmallPoolMemory + LocalAllocatedLargePoolMemory;
-	const uint64 TotalOSAllocated = LocalAllocatedOSSmallPoolMemory + LocalAllocatedLargePoolMemoryWAlignment + OSPageAllocatorCachedFreeSize;
+	const uint64 TotalOSAllocated = LocalAllocatedOSSmallPoolMemory + LocalAllocatedLargePoolMemoryWAlignment + CachedOSPageAllocator.GetCachedFreeTotal();
 
 	OutStats.Add(TEXT("TotalAllocated"), TotalAllocated);
 	OutStats.Add(TEXT("TotalOSAllocated"), TotalOSAllocated);
@@ -983,6 +985,7 @@ void FMallocBinned2::DumpAllocatorStats(class FOutputDevice& Ar)
 
 void FMallocBinned2::UpdateStats()
 {
+	//TODO: this will report total cached free memory in the COSPA, however a big chunk of that memory can be immediately freed to the OS
 	CSV_CUSTOM_STAT(FMemory, AllocatorCachedSlackMB, (int32)(CachedOSPageAllocator.GetCachedFreeTotal()/(1024*1024)), ECsvCustomStatOp::Set);
 
 	CachedOSPageAllocator.UpdateStats();
diff --git a/Engine/Source/Runtime/Core/Public/GenericPlatform/GenericPlatformMemory.h b/Engine/Source/Runtime/Core/Public/GenericPlatform/GenericPlatformMemory.h
--- a/Engine/Source/Runtime/Core/Public/GenericPlatform/GenericPlatformMemory.h
+++ b/Engine/Source/Runtime/Core/Public/GenericPlatform/GenericPlatformMemory.h
@@ -364,10 +364,7 @@ struct FGenericPlatformMemory
 	/**
 	* @return memory used for platforms that can do it quickly (without affecting stat unit much)
 	*/
-	static uint64 GetMemoryUsedFast()
-	{
-		return 0;
-	}
+	static CORE_API uint64 GetMemoryUsedFast();
 
 	/**
 	 * Writes all platform specific current memory statistics in the format usable by the malloc profiler.
diff --git a/Engine/Source/Runtime/Core/Public/HAL/Allocators/CachedOSPageAllocator.h b/Engine/Source/Runtime/Core/Public/HAL/Allocators/CachedOSPageAllocator.h
--- a/Engine/Source/Runtime/Core/Public/HAL/Allocators/CachedOSPageAllocator.h
+++ b/Engine/Source/Runtime/Core/Public/HAL/Allocators/CachedOSPageAllocator.h
@@ -63,11 +63,16 @@ struct TCachedOSPageAllocator : private FCachedOSPageAllocator
 	{
 
 	}
-	uint64 GetCachedFreeTotal()
+	uint64 GetCachedFreeTotal() const
 	{
 		return CachedTotal;
 	}
 
+	uint64 GetCachedImmediatelyFreeable() const
+	{
+		return GetCachedFreeTotal();
+	}
+
 	bool IsOSAllocation(SIZE_T Size)
 	{
 		return FCachedOSPageAllocator::IsOSAllocation(Size, CachedByteLimit);
diff --git a/Engine/Source/Runtime/Core/Public/HAL/Allocators/CachedOSVeryLargePageAllocator.h b/Engine/Source/Runtime/Core/Public/HAL/Allocators/CachedOSVeryLargePageAllocator.h
--- a/Engine/Source/Runtime/Core/Public/HAL/Allocators/CachedOSVeryLargePageAllocator.h
+++ b/Engine/Source/Runtime/Core/Public/HAL/Allocators/CachedOSVeryLargePageAllocator.h
@@ -41,6 +41,7 @@ public:
 	FCachedOSVeryLargePageAllocator()
 		: bEnabled(true)
 		, CachedFree(0)
+		, ImmediatelyFreeable(0)
 	{
 		Init();
 	}
@@ -61,11 +62,16 @@ public:
 
 	void UpdateStats();
 
-	uint64 GetCachedFreeTotal()
+	uint64 GetCachedFreeTotal() const
 	{
 		return CachedFree + CachedOSPageAllocator.GetCachedFreeTotal();
 	}
 
+	uint64 GetCachedImmediatelyFreeable() const
+	{
+		return ImmediatelyFreeable + CachedOSPageAllocator.GetCachedImmediatelyFreeable();
+	}
+
 	FORCEINLINE bool IsSmallBlockAllocation(const void* Ptr) const
 	{
 		if (((uintptr_t)Ptr - AddressSpaceReserved) < AddressSpaceToReserveForSmallPool)
@@ -122,6 +128,7 @@ private:
 	uintptr_t	AddressSpaceReservedEndSmallPool;
 	uintptr_t	AddressSpaceReservedEnd;
 	uint64		CachedFree;
+	uint64		ImmediatelyFreeable;		// the amount of memory that can be immediately returned to the OS
 	int32		EmptyBackStoreCount[FMemory::AllocationHints::Max];
 	int32		CommittedLargePagesCount[FMemory::AllocationHints::Max];
 
diff --git a/Engine/Source/Runtime/Core/Public/HAL/MallocBinned2.h b/Engine/Source/Runtime/Core/Public/HAL/MallocBinned2.h
--- a/Engine/Source/Runtime/Core/Public/HAL/MallocBinned2.h
+++ b/Engine/Source/Runtime/Core/Public/HAL/MallocBinned2.h
@@ -52,6 +52,8 @@ struct FGenericMemoryStats;
 	extern int32 RecursionCounter;
 #endif
 
+extern int32 EnableLegacyCachedOSPageAllocatorFreeMemReporting;
+
 // Canary value used in FFreeBlock
 // A constant value unless we're compiled with fork support in which case there are two values identifying whether the page
 // was allocated pre- or post-fork
@@ -506,6 +508,15 @@ public:
 	CORE_API virtual void OnMallocInitialized() override;
 	CORE_API virtual void OnPreFork() override;
 	CORE_API virtual void OnPostFork() override;
+	virtual uint64 GetFreeCachedMemorySize() const
+	{
+		if (EnableLegacyCachedOSPageAllocatorFreeMemReporting)
+		{
+			return CachedOSPageAllocator.GetCachedFreeTotal();
+		}
+
+		return CachedOSPageAllocator.GetCachedImmediatelyFreeable();
+	}
 	// End FMalloc interface.
 
 	CORE_API void* MallocExternalSmall(SIZE_T Size, uint32 Alignment);
diff --git a/Engine/Source/Runtime/Core/Public/HAL/MemoryBase.h b/Engine/Source/Runtime/Core/Public/HAL/MemoryBase.h
--- a/Engine/Source/Runtime/Core/Public/HAL/MemoryBase.h
+++ b/Engine/Source/Runtime/Core/Public/HAL/MemoryBase.h
@@ -248,6 +248,14 @@ public:
 	 */
 	virtual void OnPostFork() {}
 
+	/**
+	 * Returns the amount of free memory cached by the allocator that can be returned to the system in case of a memory shortage
+	 */
+	virtual uint64 GetFreeCachedMemorySize() const
+	{
+		return 0;
+	}
+
 protected:
 	friend struct FCurrentFrameCalls;
 
