Fixed compilation error in PooledVirtualMemoryAllocator

diff --git a/Engine/Source/Runtime/Core/Private/HAL/Allocators/PooledVirtualMemoryAllocator.cpp b/Engine/Source/Runtime/Core/Private/HAL/Allocators/PooledVirtualMemoryAllocator.cpp
--- a/Engine/Source/Runtime/Core/Private/HAL/Allocators/PooledVirtualMemoryAllocator.cpp
+++ b/Engine/Source/Runtime/Core/Private/HAL/Allocators/PooledVirtualMemoryAllocator.cpp
@@ -152,7 +152,7 @@ void FPooledVirtualMemoryAllocator::Free(void* Ptr, SIZE_T Size, FCriticalSectio
 
 			if (UNLIKELY(Desc.Pool->WasAllocatedFromThisPool(Ptr, Size)))
 			{
-				// LLVM wants to be informed of the allocations of physical RAM.
+				// LLM wants to be informed of the allocations of physical RAM.
 				// This is the closest we can get.
 				LLM_IF_ENABLED(FLowLevelMemTracker::Get().OnLowLevelFree(ELLMTracker::Platform, Ptr));
 				Desc.Pool->Free(Ptr, Size);
@@ -279,7 +279,7 @@ void FPooledVirtualMemoryAllocator::FreeAll(FCriticalSection* /*Mutex = nullptr*
 	// That would be a speed optimization and not a size optimization so I'm not going for this at this point, this method is speedy enough.
 };
 
-uint64 FPooledVirtualMemoryAllocator::GetCachedFreeTotal()
+uint64 FPooledVirtualMemoryAllocator::GetCachedFreeTotal() const
 {
 	uint64 TotalFree = 0;
 
diff --git a/Engine/Source/Runtime/Core/Public/HAL/Allocators/PooledVirtualMemoryAllocator.h b/Engine/Source/Runtime/Core/Public/HAL/Allocators/PooledVirtualMemoryAllocator.h
--- a/Engine/Source/Runtime/Core/Public/HAL/Allocators/PooledVirtualMemoryAllocator.h
+++ b/Engine/Source/Runtime/Core/Public/HAL/Allocators/PooledVirtualMemoryAllocator.h
@@ -76,7 +76,11 @@ struct FPooledVirtualMemoryAllocator
 	};
 
 	/** Returns free memory in the pools */
-	uint64 GetCachedFreeTotal();
+	uint64 GetCachedFreeTotal() const;
+	uint64 GetCachedImmediatelyFreeable() const
+	{
+		return GetCachedFreeTotal();
+	}
 
 	/** Refresh allocator if needed. (does nothing in that implementation) */
 	void Refresh() {}
@@ -134,7 +138,7 @@ private:
 	FPoolDescriptorBase* ClassesListHeads[Limits::NumAllocationSizeClasses];
 
 	/** Per-class locks */
-	FCriticalSection     ClassesLocks[Limits::NumAllocationSizeClasses];
+	mutable FCriticalSection ClassesLocks[Limits::NumAllocationSizeClasses];
 
 	/** Increases the number of pooled allocations next time we need a pool
 	 *
