Fixed FGenericPlatformMemory::GetStats not implemented on this platform log spam

diff --git a/Engine/Source/Runtime/Core/Private/GenericPlatform/GenericPlatformMemory.cpp b/Engine/Source/Runtime/Core/Private/GenericPlatform/GenericPlatformMemory.cpp
--- a/Engine/Source/Runtime/Core/Private/GenericPlatform/GenericPlatformMemory.cpp
+++ b/Engine/Source/Runtime/Core/Private/GenericPlatform/GenericPlatformMemory.cpp
@@ -283,7 +283,7 @@ FPlatformMemoryStats FGenericPlatformMemory::GetStatsRaw()
 
 uint64 FGenericPlatformMemory::GetMemoryUsedFast()
 {
-	return GetStats().UsedPhysical;
+	return FPlatformMemory::GetStats().UsedPhysical;
 }
 
 void FGenericPlatformMemory::GetStatsForMallocProfiler( FGenericMemoryStats& out_Stats )
