CSV profiler: Fix underreporting issues with csv.AggregateTaskWorkerStats and re-enable by default
- Instead of updating aggregate AllWorker stats during each worker thread's Process call, we defer the aggregation until the final values are available. This fixes edge case bugs when processing stats across frame boundaries.
- Aggregation is performed in a new FCsvAggregateStatSeries which can hold values from multiple frames until they're ready to be finalized
- The new approach carries more overhead overall on the CSV profiler thread due to the additional processing, but it's still cheaper than with aggregation off

Debugging improvements:
-Add -csvTestTasks to test taskgraph worker stats
-Add -csvAllCategoriesDisabled and -csvTestCategoryOnly disable all categories and enable just the CsvTest category respectively. These even disable the global category. This is useful for debugging
-Add debug switches CSV_DEBUG_CUSTOM_STATS_INCLUDE_THREAD_NAME and CSV_DEBUG_EMIT_SEPARATE_THREAD_STATS_WHEN_TASK_AGGREGATION_ENABLED
-Improve accuracy of timer tests (precisely spin instead of sleeping)

#tests Tested on all the consoles

diff --git a/Engine/Source/Runtime/Core/Private/ProfilingDebugging/CsvProfiler.cpp b/Engine/Source/Runtime/Core/Private/ProfilingDebugging/CsvProfiler.cpp
--- a/Engine/Source/Runtime/Core/Private/ProfilingDebugging/CsvProfiler.cpp
+++ b/Engine/Source/Runtime/Core/Private/ProfilingDebugging/CsvProfiler.cpp
@@ -8,6 +8,7 @@
 
 #include "ProfilingDebugging/CsvProfiler.h"
 #include "Async/Fundamental/Scheduler.h"
+#include "Async/ParallelFor.h"
 #include "Async/TaskGraphInterfaces.h"
 #include "CoreGlobals.h"
 #include "HAL/RunnableThread.h"
@@ -51,6 +52,14 @@
 
 #define REPAIR_MARKER_STACKS 1
 
+// Note: Enabling this will break a lot of reporting but it may be useful for debugging. If disabled (default), the thread name is omitted from custom stat columns
+// This results in multiple columns with the same name (since we have 1 column per stat per thread), but these are automatically combined by the CSV tools
+#define CSV_DEBUG_CUSTOM_STATS_INCLUDE_THREAD_NAME 0
+
+// By default we hide per thread stats when aggregation is enabled (since the point is to reduce the number). Enabling this shows both
+#define CSV_DEBUG_EMIT_SEPARATE_THREAD_STATS_WHEN_TASK_AGGREGATION_ENABLED 0
+
+
 // Global CSV category (no prefix)
 FCsvCategory GGlobalCsvCategory(TEXT("GLOBAL"), true, true);
 
@@ -66,9 +75,10 @@ CSV_DEFINE_CATEGORY(CsvProfiler, true);
 	CSV_DEFINE_CATEGORY(CsvTest, true);
 	static bool GCsvTestingGT = false;
 	static bool GCsvTestingRT = false;
-
+	static bool GCsvTestCategoryOnly = false;
 	void CSVTest();
 #endif
+static bool GAllCategoriesStartDisabled = false;
 
 #define LIST_VALIDATION (DO_CHECK && 0)
 
@@ -87,7 +97,7 @@ TAutoConsoleVariable<int32> CVarCsvBlockOnCaptureEnd(
 
 TAutoConsoleVariable<bool> CVarCsvAggregateTaskWorkerStats(
 	TEXT("csv.AggregateTaskWorkerStats"),
-	false,
+	true,
 	TEXT("If enabled, stats recorded on task worker threads are aggregated instead of outputting a single stat per thread.\r\n")
 	TEXT("This reduces CSV bloat when there are large numbers of worker threads and makes stat data more intelligible"),
 	ECVF_Default
@@ -205,6 +215,13 @@ static bool GCsvExitOnCompletion = false;
 
 static thread_local bool GCsvThreadLocalWaitsEnabled = false;
 
+//
+// Forward declarations
+//
+struct FCsvStatSeries;
+struct FCsvAggregateStatSeries;
+struct FCsvProcessThreadDataStats;
+class FCsvThreadGroupStatProcessor;
 
 // A unique ID for a CSV stat, either ansi or FName
 union FCsvUniqueStatID
@@ -619,6 +636,10 @@ public:
 
 	void UpdateCategoryFromConfig(int32 CategoryIndex)
 	{
+		if (GAllCategoriesStartDisabled)
+		{
+			return;
+		}
 		for (FString const& EnabledCategory : CategoriesEnabledInConfig)
 		{
 			if (FWildcardString::IsMatch(*EnabledCategory, *CategoryNames[CategoryIndex]))
@@ -668,7 +689,10 @@ public:
 		{
 			return Index;
 		}
-
+		if (GAllCategoriesStartDisabled)
+		{
+			bEnableByDefault = false;
+		}
 		FScopeLock Lock(&CS);
 		{
 			Index = GetCategoryIndex(CategoryName);
@@ -1716,7 +1740,12 @@ struct FCsvStatSeries
 		CustomStatFloat
 	};
 
-	FCsvStatSeries(EType InSeriesType, const FCsvStatID& InStatID, FCsvStreamWriter* InWriter, FCsvStatRegister& StatRegister, const FString& ThreadName);
+	FCsvStatSeries(EType InSeriesType, const FCsvStatID& InStatID, FCsvStreamWriter* InWriter, FCsvStatRegister& StatRegister, const FString& ThreadName, FCsvAggregateStatSeries* InLinkedAggregateStatSeries);
+	
+	virtual ~FCsvStatSeries()
+	{
+	}
+
 	void FlushIfDirty();
 
 	void SetTimerValue(uint32 DataFrameNumber, uint64 ElapsedCycles)
@@ -1821,7 +1850,26 @@ struct FCsvStatSeries
 		return (SeriesType == EType::CustomStatFloat || SeriesType == EType::CustomStatInt);
 	}
 
-	inline uint64 GetAllocatedSize() const { return Name.GetAllocatedSize(); }
+	virtual void FinalizeFrame(int64 FrameNumber)
+	{
+		// Stat values are held in the series until a new value arrives.
+		// If we've caught up with the last value written to the series,
+		// we need to flush to get the correct value for this frame.
+		if (CurrentWriteFrameNumber == FrameNumber)
+		{
+			FlushIfDirty();
+		}
+	}
+
+	virtual bool IsAggregateSeries() const
+	{
+		return false;
+	}
+
+	virtual uint64 GetAllocatedSize() const 
+	{ 
+		return sizeof(*this) + Name.GetAllocatedSize();
+	}
 
 	FCsvStatID StatID;
 	const EType SeriesType;
@@ -1839,9 +1887,50 @@ struct FCsvStatSeries
 
 	int32 ColumnIndex;
 
+	FCsvAggregateStatSeries* LinkedAggregateStatSeries;
+
 	bool bDirty;
 };
 
+
+struct FCsvAggregateStatSeries : public FCsvStatSeries
+{
+	FCsvAggregateStatSeries(EType InSeriesType, const FCsvStatID& InStatID, FCsvStreamWriter* InWriter, FCsvStatRegister& StatRegister, const FString& ThreadName) 
+		: FCsvStatSeries(InSeriesType, InStatID, InWriter, StatRegister, ThreadName, nullptr)
+	{
+	}
+
+	// Accumulate the value for the linked series for a given frame
+	void AccumulateLinkedSeriesValue(int64 FrameNumber, const FCsvStatSeriesValue& Value)
+	{
+		FCsvStatSeriesValue& FrameValue = RowValues.FindOrAdd(FrameNumber);
+		switch (SeriesType)
+		{
+		case EType::TimerData:
+		case EType::CustomStatFloat:
+			FrameValue.Value.AsFloat += Value.Value.AsFloat;
+			break;
+		case EType::CustomStatInt:
+			FrameValue.Value.AsInt += Value.Value.AsInt;
+		}
+	}
+
+	virtual void FinalizeFrame(int64 FrameNumber) override;
+
+	virtual bool IsAggregateSeries() const override
+	{
+		return true;
+	}
+
+	virtual uint64 GetAllocatedSize() const override
+	{
+		return ((uint64)RowValues.GetAllocatedSize());
+	}
+
+	TMap<int64, FCsvStatSeriesValue> RowValues;
+};
+
+
 struct FCsvProcessThreadDataStats
 {
 	FCsvProcessThreadDataStats()
@@ -1890,7 +1979,12 @@ class FCsvStreamWriter
 	const bool bContinuousWrites;
 	bool bFirstRow;
 
+	// All series, including hidden series that feed into aggregate series rather than writing to the CSV directly
 	TArray<FCsvStatSeries*> AllSeries;
+	
+	// Subset of AllSeries - only visible
+	TArray<FCsvStatSeries*> VisibleSeries;
+
 	TArray<class FCsvProfilerThreadDataProcessor*> DataProcessors;
 
 	TSharedPtr<FCsvThreadGroupStatProcessor> TaskWorkerThreadGroupStatProcessor;
@@ -1902,7 +1996,7 @@ public:
 	FCsvStreamWriter(const TSharedRef<FArchive>& InOutputFile, bool bInContinuousWrites, int32 InBufferSize, int64 InNumFramesToBuffer, bool bInCompressOutput, uint32 RenderThreadId, uint32 RHIThreadId, bool bAggregateTaskWorkerStats);
 	~FCsvStreamWriter();
 
-	void AddSeries(FCsvStatSeries* Series);
+	void AddSeries(FCsvStatSeries* Series, bool bIsVisible);
 
 	void PushValue(FCsvStatSeries* Series, int64 FrameNumber, const FCsvStatSeriesValue& Value);
 	void PushEvent(const FCsvProcessedEvent& Event);
@@ -1915,12 +2009,13 @@ public:
 	inline uint64 GetAllocatedSize() const;
 };
 
-FCsvStatSeries::FCsvStatSeries(EType InSeriesType, const FCsvStatID& InStatID, FCsvStreamWriter* InWriter, FCsvStatRegister& StatRegister, const FString& ThreadName)
+FCsvStatSeries::FCsvStatSeries(EType InSeriesType, const FCsvStatID& InStatID, FCsvStreamWriter* InWriter, FCsvStatRegister& StatRegister, const FString& ThreadName, FCsvAggregateStatSeries* InLinkedAggregateStatSeries)
 	: StatID(InStatID)
 	, SeriesType(InSeriesType)
 	, CurrentWriteFrameNumber(-1)
 	, Writer(InWriter)
 	, ColumnIndex(-1)
+	, LinkedAggregateStatSeries(InLinkedAggregateStatSeries)
 	, bDirty(false)
 {
 	CurrentValue.AsTimerCycles = 0;
@@ -1930,7 +2025,7 @@ FCsvStatSeries::FCsvStatSeries(EType InSeriesType, const FCsvStatID& InStatID, F
 	Name = StatRegister.GetStatName(StatID);
 	bool bIsCountStat = StatRegister.IsCountStat(StatID);
 
-	if (!IsCustomStat() || bIsCountStat)
+	if (!IsCustomStat() || bIsCountStat || LinkedAggregateStatSeries != nullptr || (CSV_DEBUG_CUSTOM_STATS_INCLUDE_THREAD_NAME==1))
 	{
 		// Add a /<Threadname> prefix
 		Name = ThreadName + TEXT("/") + Name;
@@ -1948,7 +2043,9 @@ FCsvStatSeries::FCsvStatSeries(EType InSeriesType, const FCsvStatID& InStatID, F
 		Name = TEXT("COUNTS/") + Name;
 	}
 
-	Writer->AddSeries(this);
+	// If we have a linked stat series, don't write directly to the CSV
+	bool bVisible = LinkedAggregateStatSeries == nullptr || (CSV_DEBUG_EMIT_SEPARATE_THREAD_STATS_WHEN_TASK_AGGREGATION_ENABLED==1);
+	Writer->AddSeries(this, bVisible);
 }
 
 void FCsvStatSeries::FlushIfDirty()
@@ -1968,12 +2065,33 @@ void FCsvStatSeries::FlushIfDirty()
 			Value.Value.AsFloat = CurrentValue.AsFloatValue;
 			break;
 		}
-		Writer->PushValue(this, CurrentWriteFrameNumber, Value);
+		// If there's an aggregate stat, push data to that rather than the CSV writer directly
+		if (LinkedAggregateStatSeries)
+		{
+			LinkedAggregateStatSeries->AccumulateLinkedSeriesValue(CurrentWriteFrameNumber, Value);
+		}
+#if !CSV_DEBUG_EMIT_SEPARATE_THREAD_STATS_WHEN_TASK_AGGREGATION_ENABLED
+		else
+#endif
+		{
+			Writer->PushValue(this, CurrentWriteFrameNumber, Value);
+		}
 		CurrentValue.AsTimerCycles = 0;
 		bDirty = false;
 	}
 }
 
+void FCsvAggregateStatSeries::FinalizeFrame(int64 FrameNumber)
+{
+	// Write the final value for this row to the writer
+	FCsvStatSeriesValue Value;
+	if (RowValues.RemoveAndCopyValue(FrameNumber, Value))
+	{
+		Writer->PushValue(this, FrameNumber, Value);
+	}
+}
+
+
 class FCsvProfilerThreadData
 {
 public:
@@ -2255,11 +2373,13 @@ class FCsvThreadGroupStatProcessor
 	TArray<FCsvStatSeries*> StatSeriesArray;
 	FCsvStatRegister StatRegister;
 	FString Name;
+	FCsvThreadGroupStatProcessor* AggregateStatProcessor;
 
 public:
 	FCsvThreadGroupStatProcessor(FCsvStreamWriter* InWriter, FString InName)
 		: Writer(InWriter)
 		, Name(InName)
+		, AggregateStatProcessor(nullptr)
 	{
 	}
 
@@ -2272,9 +2392,30 @@ public:
 		}
 	}
 
-	FCsvStatSeries* FindOrCreateStatSeries(const FCsvStatBase& Stat, FCsvStatSeries::EType SeriesType, bool bIsCountStat)
+	void SetAggregateStatProcessor(FCsvThreadGroupStatProcessor* InAggregateStatProcessor)
+	{
+		AggregateStatProcessor = InAggregateStatProcessor;
+	}
+
+	void FinalizeStatSeriesFrame(int64 FrameNumber)
+	{
+		for (FCsvStatSeries* Series : StatSeriesArray)
+		{
+			Series->FinalizeFrame(FrameNumber);
+		}
+	}
+
+	FCsvStatSeries* FindOrCreateStatSeries(const FCsvStatBase& Stat, FCsvStatSeries::EType SeriesType, bool bIsCountStat, bool bIsAggregateSeries = false)
 	{
 		check(IsInCsvProcessingThread());
+
+		FCsvAggregateStatSeries* LinkedAggregateStatSeries = nullptr;
+		if ( AggregateStatProcessor )
+		{
+			// Find or create the linked aggregate stat series
+			LinkedAggregateStatSeries = (FCsvAggregateStatSeries*)AggregateStatProcessor->FindOrCreateStatSeries(Stat, SeriesType, bIsCountStat, true );
+		}
+
 		const int32 StatIndex = StatRegister.GetUniqueIndex(Stat.RawStatID, Stat.CategoryIndex, Stat.IsFNameStat(), bIsCountStat);
 		FCsvStatSeries* Series = nullptr;
 		if (StatSeriesArray.Num() <= StatIndex)
@@ -2284,7 +2425,14 @@ public:
 		}
 		if (StatSeriesArray[StatIndex] == nullptr)
 		{
-			Series = new FCsvStatSeries(SeriesType, StatIndex, Writer, StatRegister, Name);
+			if (bIsAggregateSeries)
+			{
+				Series = new FCsvAggregateStatSeries(SeriesType, StatIndex, Writer, StatRegister, Name);
+			}
+			else
+			{
+				Series = new FCsvStatSeries(SeriesType, StatIndex, Writer, StatRegister, Name, LinkedAggregateStatSeries);
+			}
 			StatSeriesArray[StatIndex] = Series;
 		}
 		else
@@ -2297,10 +2445,14 @@ public:
 		return Series;
 	}
 
-	inline uint64 GetAllocatedSize() const
+	uint64 GetAllocatedSize() const
 	{
-		return	((uint64)StatSeriesArray.GetAllocatedSize()) +
-				((uint64)StatSeriesArray.Num() * sizeof(FCsvStatSeries));
+		uint64 TotalSize = (uint64)StatSeriesArray.GetAllocatedSize();
+		for (const FCsvStatSeries* Series : StatSeriesArray)
+		{
+			TotalSize += Series->GetAllocatedSize();
+		}
+		return TotalSize;
 	}
 
 };
@@ -2381,10 +2533,14 @@ FCsvStreamWriter::~FCsvStreamWriter()
 	}
 }
 
-void FCsvStreamWriter::AddSeries(FCsvStatSeries* Series)
+void FCsvStreamWriter::AddSeries(FCsvStatSeries* Series, bool bIsVisible)
 {
 	check(Series->ColumnIndex == -1);
-	Series->ColumnIndex = AllSeries.Num();
+	if (bIsVisible)
+	{
+		Series->ColumnIndex = VisibleSeries.Num();
+		VisibleSeries.Add(Series);
+	}
 	AllSeries.Add(Series);
 }
 
@@ -2397,9 +2553,9 @@ void FCsvStreamWriter::PushValue(FCsvStatSeries* Series, int64 FrameNumber, cons
 	FCsvRow& Row = Rows.FindOrAdd(FrameNumber);
 
 	// Ensure the row is large enough to hold every series
-	if (Row.Values.Num() < AllSeries.Num())
+	if (Row.Values.Num() < VisibleSeries.Num())
 	{
-		Row.Values.SetNumZeroed(AllSeries.Num(), EAllowShrinking::No);
+		Row.Values.SetNumZeroed(VisibleSeries.Num(), EAllowShrinking::No);
 	}
 
 	Row.Values[Series->ColumnIndex] = Value;
@@ -2419,7 +2575,7 @@ void FCsvStreamWriter::FinalizeNextRow()
 		// Write the first header row
 		Stream.WriteString("EVENTS");
 
-		for (FCsvStatSeries* Series : AllSeries)
+		for (FCsvStatSeries* Series : VisibleSeries)
 		{
 			Stream.WriteString(Series->Name);
 		}
@@ -2450,14 +2606,22 @@ void FCsvStreamWriter::FinalizeNextRow()
 			Stream.WriteEmptyString();
 		}
 
+		// Finalize the series in dependency order (non-aggregate series followed by task worker aggregate series)
 		for (FCsvStatSeries* Series : AllSeries)
 		{
-			// Stat values are held in the series until a new value arrives.
-			// If we've caught up with the last value written to the series,
-			// we need to flush to get the correct value for this frame.
-			if (Series->CurrentWriteFrameNumber == ReadFrameIndex)
-				Series->FlushIfDirty();
+			if (!Series->IsAggregateSeries())
+			{
+				Series->FinalizeFrame(ReadFrameIndex);
+			}
+		}
+		if ( TaskWorkerThreadGroupStatProcessor )
+		{	
+			TaskWorkerThreadGroupStatProcessor->FinalizeStatSeriesFrame(ReadFrameIndex);
+		}
 
+		// Write the visible series to the CSV
+		for (FCsvStatSeries* Series : VisibleSeries)
+		{
 			if (Row->Values.IsValidIndex(Series->ColumnIndex))
 			{
 				const FCsvStatSeriesValue& Value = Row->Values[Series->ColumnIndex];
@@ -2493,7 +2657,7 @@ void FCsvStreamWriter::Finalize(const TMap<FString, FString>& Metadata)
 
 	// Write a final summary header row
 	Stream.WriteString("EVENTS");
-	for (FCsvStatSeries* Series : AllSeries)
+	for (FCsvStatSeries* Series : VisibleSeries)
 	{
 		Stream.WriteString(Series->Name);
 	}
@@ -2563,6 +2727,7 @@ uint64 FCsvStreamWriter::GetAllocatedSize() const
 	uint64 Size =
 		((uint64)Rows.GetAllocatedSize()) +
 		((uint64)AllSeries.GetAllocatedSize()) +
+		((uint64)VisibleSeries.GetAllocatedSize()) +
 		((uint64)DataProcessors.GetAllocatedSize()) +
 		((uint64)Stream.GetAllocatedSize());
 
@@ -2683,12 +2848,15 @@ void FCsvProfilerThreadDataProcessor::Process(FCsvProcessThreadDataStats& OutSta
 	Events.Reset(0);
 	ThreadData->FlushResults(ThreadMarkers, CustomStats, Events);
 
-	// If we're aggregating task worker thread stats and this is a task worker thread, use the shared TaskWorkerThreadGroupStatProcessor instead of the per-thread one
 	FCsvThreadGroupStatProcessor* ThreadStatProcessor = StatProcessor.Get();
+
+	// If we're aggregating task threads then link this thread's stat processor with the shared aggregate one
+	FCsvThreadGroupStatProcessor* AggregateStatProcessor = nullptr;
 	if (ThreadData->bIsTaskWorkerThread && TaskWorkerThreadGroupStatProcessor.IsValid())
 	{
-		ThreadStatProcessor = TaskWorkerThreadGroupStatProcessor.Get();
+		AggregateStatProcessor = TaskWorkerThreadGroupStatProcessor.Get();
 	}
+	ThreadStatProcessor->SetAggregateStatProcessor(AggregateStatProcessor);
 
 	OutStats.TimestampCount += ThreadMarkers.Num();
 	OutStats.CustomStatCount += CustomStats.Num();
@@ -3953,6 +4121,16 @@ void FCsvProfiler::Init()
 		GCsvTestingGT = true;
 		GCsvTestingRT = true;
 	}
+	if (FParse::Param(FCommandLine::Get(), TEXT("csvTestCategoryOnly")))
+	{
+		GAllCategoriesStartDisabled = true;
+		GCsvTestCategoryOnly = true;
+	}
+
+	if (FParse::Param(FCommandLine::Get(), TEXT("csvAllCategoriesDisabled")))
+	{
+		GAllCategoriesStartDisabled = true;
+	}
 
 	FString CsvCategoriesStr;
 	if (FParse::Value(FCommandLine::Get(), TEXT("csvCategories="), CsvCategoriesStr, /*bShouldStopOnSeparator*/false))
@@ -4058,6 +4236,11 @@ void FCsvProfiler::Init()
 
 #endif // CSV_PROFILER_ALLOW_DEBUG_FEATURES
 
+	if (GAllCategoriesStartDisabled)
+	{
+		FMemory::Memzero(GCsvCategoriesEnabled, sizeof(GCsvCategoriesEnabled));
+	}
+
 	// Always disable the CSV profiling thread if the platform does not support threading.
 	if (!FPlatformProcess::SupportsMultithreading())
 	{
@@ -4188,10 +4371,26 @@ TCsvPersistentCustomStat<float>* FCsvProfiler::GetOrCreatePersistentCustomStatFl
 
 #if CSV_PROFILER_ALLOW_DEBUG_FEATURES
 
+static void SpinWaitMs(double Milliseconds)
+{
+	double CurrentTime = FPlatformTime::Seconds();
+	double SecondsToWait = Milliseconds*0.001;
+	double TargetTime = CurrentTime + SecondsToWait;
+	while (CurrentTime < TargetTime)
+	{
+		CurrentTime = FPlatformTime::Seconds();
+	}
+}
+
 // Simple benchmarking and debugging tests for the csv profiler. Enable with -csvtest, e.g -csvtest -csvcaptureframes=400
 void CSVTest()
 {
-	TCsvPersistentCustomStat<float>* PersistentStatFloat = FCsvProfiler::Get()->GetOrCreatePersistentCustomStatFloat(TEXT("PersistentStatFloat"));
+	if (GCsvTestCategoryOnly)
+	{
+		FMemory::Memzero(GCsvCategoriesEnabled, sizeof(GCsvCategoriesEnabled));
+		GCsvCategoriesEnabled[CSV_CATEGORY_INDEX(CsvTest)] = true;
+	}
+	TCsvPersistentCustomStat<float>* PersistentStatFloat = FCsvProfiler::Get()->GetOrCreatePersistentCustomStatFloat(TEXT("PersistentStatFloat"), CSV_CATEGORY_INDEX(CsvTest));
 	PersistentStatFloat->Add(0.15f);
 	PersistentStatFloat->Sub(0.1f);
 
@@ -4200,6 +4399,32 @@ void CSVTest()
 	PersistentStatInt->Sub(1);
 
 	uint32 FrameNumber = FCsvProfiler::Get()->GetCaptureFrameNumber();
+
+	static bool bTaskStats = FParse::Param(FCommandLine::Get(), TEXT("csvTestTasks"));
+	if (bTaskStats)
+	{
+		ParallelFor(4, [FrameNumber](int32 Index)
+			{
+				CSV_SCOPED_TIMING_STAT(CsvTest, TaskTimer);
+				if (!IsInGameThread())
+				{
+					CSV_CUSTOM_STAT(CsvTest, TaskCustomStatSet, 0.5, ECsvCustomStatOp::Set);
+					CSV_CUSTOM_STAT(CsvTest, TaskCustomStatSet, 1.0, ECsvCustomStatOp::Set);
+					CSV_CUSTOM_STAT(CsvTest, TaskCustomStatAccumulate, 0.5, ECsvCustomStatOp::Accumulate);
+					CSV_CUSTOM_STAT(CsvTest, TaskCustomStatAccumulate, 0.5, ECsvCustomStatOp::Accumulate);
+					CSV_CUSTOM_STAT(CsvTest, TaskCustomStatMax, 0.5, ECsvCustomStatOp::Max);
+					CSV_CUSTOM_STAT(CsvTest, TaskCustomStatMax, 1.0, ECsvCustomStatOp::Max);
+
+					if (FrameNumber % 321 == 0)
+					{
+						CSV_CUSTOM_STAT(CsvTest, TaskSparse, 1.0, ECsvCustomStatOp::Set);
+					}
+				}
+				SpinWaitMs(1.0);
+			},
+			EParallelForFlags::BackgroundPriority);
+	}
+
 	CSV_SCOPED_TIMING_STAT(CsvTest, CsvTestStat);
 	CSV_CUSTOM_STAT(CsvTest, CaptureFrameNumber, int32(FrameNumber), ECsvCustomStatOp::Set);
 	CSV_CUSTOM_STAT(CsvTest, SameCustomStat, 1, ECsvCustomStatOp::Set);
@@ -4207,7 +4432,7 @@ void CSVTest()
 	for (int i = 0; i < 3; i++)
 	{
 		CSV_SCOPED_TIMING_STAT(CsvTest, RepeatStat1MS);
-		FPlatformProcess::Sleep(0.001f);
+		SpinWaitMs(1.0);
 	}
 
 	{
@@ -4254,13 +4479,6 @@ void CSVTest()
 			}
 		}
 	}
-	//for (int i = 0; i < 2048; i++)
-	//{
-	//	GCsvCategoriesEnabled[i] = false;
-	//} 
-	//GCsvCategoriesEnabled[CSV_CATEGORY_INDEX(Exclusive)] = true;
-	//GCsvCategoriesEnabled[CSV_CATEGORY_INDEX(CsvTest)] = true;
-
 	{
 		CSV_SCOPED_TIMING_STAT_EXCLUSIVE(ExclusiveLevel0);
 		{
@@ -4285,9 +4503,8 @@ void CSVTest()
 			CSV_SCOPED_TIMING_STAT_EXCLUSIVE(ExclusiveBeginEndbenchmarkInner3);
 		}
 	}
-
 }
 
 #endif // CSV_PROFILER_ALLOW_DEBUG_FEATURES
 
-#endif // CSV_PROFILER
+#endif // CSV_PROFILER
\ No newline at end of file
