[AutoRTFM] Fix linux debug builds

Clang does not like the [[clang::always_inline]] attribute used inside the lambda passed to autortfm_open, in linux + debug builds only. It warns (treated as error) with:

statement attribute 'always_inline' has higher precedence than function attribute 'noinline'

Instead of the lambda, use a function-scoped struct with a static method. This has the additional benefit of being able to use [[clang::noautortfm]] attribute on the callback stub, reducing the amount of code emitted.

diff --git a/Engine/Source/Runtime/Core/Public/AutoRTFM/AutoRTFM.h b/Engine/Source/Runtime/Core/Public/AutoRTFM/AutoRTFM.h
--- a/Engine/Source/Runtime/Core/Public/AutoRTFM/AutoRTFM.h
+++ b/Engine/Source/Runtime/Core/Public/AutoRTFM/AutoRTFM.h
@@ -666,7 +666,7 @@ static constexpr bool IsSafeToReturnFromOpen = HasAssignFromOpenToClosedTrait<T>
 // a transaction or not. Returns the value returned by Functor.
 // TReturn must be void or a type that can be safely copied from the open to a closed transaction.
 // TAssignFromOpenToClosed must have a specialization for the type that is being returned.
-template<typename TFunctor, typename TReturn = decltype(std::declval<TFunctor>()())> 
+template<typename TFunctor, typename TReturn = decltype(std::declval<TFunctor>()())>
 static UE_AUTORTFM_FORCEINLINE TReturn Open(const TFunctor& Functor)
 {
 	static_assert(IsSafeToReturnFromOpen<TReturn>,
@@ -681,27 +681,32 @@ static UE_AUTORTFM_FORCEINLINE TReturn Open(const TFunctor& Functor)
 	{
 		if constexpr (std::is_same_v<void, TReturn>)
 		{
-			autortfm_open(
-				[] (void* Arg) { UE_AUTORTFM_CALLSITE_FORCEINLINE (*static_cast<const TFunctor*>(Arg))(); },
-				const_cast<void*>(static_cast<const void*>(&Functor)));
+			struct FCallHelper
+			{
+				UE_AUTORTFM_NOAUTORTFM static void Call(void* Arg)
+				{
+					const TFunctor& Fn = *reinterpret_cast<TFunctor*>(Arg);
+					UE_AUTORTFM_CALLSITE_FORCEINLINE Fn();
+				}
+			};
+			autortfm_open(&FCallHelper::Call, const_cast<void*>(static_cast<const void*>(&Functor)));
 		}
 		else
 		{
-			TReturn ReturnValue;
-			struct FData
+			struct FCallHelper
 			{
+				UE_AUTORTFM_NOAUTORTFM static void Call(void* Arg)
+				{
+					FCallHelper& Self = *reinterpret_cast<FCallHelper*>(Arg);
+					UE_AUTORTFM_CALLSITE_FORCEINLINE
+						TAssignFromOpenToClosed<TReturn>::Assign(Self.ReturnValue, std::move(Self.Functor()));
+				}
 				const TFunctor& Functor;
-				TReturn& ReturnValue;
+				TReturn ReturnValue{};
 			};
-			FData Data{Functor, ReturnValue};
-			autortfm_open([](void* Arg)
-				{
-					FData& Data = *reinterpret_cast<FData*>(Arg);
-					UE_AUTORTFM_CALLSITE_FORCEINLINE TAssignFromOpenToClosed<TReturn>::Assign(Data.ReturnValue, std::move(Data.Functor()));
-				},
-				reinterpret_cast<void*>(&Data));
-
-			return Data.ReturnValue;
+			FCallHelper Helper{Functor};
+			autortfm_open(&FCallHelper::Call, reinterpret_cast<void*>(&Helper));
+			return Helper.ReturnValue;
 		}
 	}
 }
