[LLM] Remove FName block allocations from contributing to the MemQuery/Memtrace asset memory consumption

#tests local run

diff --git a/Engine/Source/Runtime/Core/Private/UObject/UnrealNames.cpp b/Engine/Source/Runtime/Core/Private/UObject/UnrealNames.cpp
--- a/Engine/Source/Runtime/Core/Private/UObject/UnrealNames.cpp
+++ b/Engine/Source/Runtime/Core/Private/UObject/UnrealNames.cpp
@@ -23,6 +23,7 @@
 #include "HAL/ThreadSafeCounter.h"
 #include "Misc/ScopeRWLock.h"
 #include "ProfilingDebugging/StringsTrace.h"
+#include "ProfilingDebugging/AssetMetadataTrace.h"
 #include "Containers/Set.h"
 #include "Internationalization/Text.h"
 #include "Internationalization/Internationalization.h"
@@ -30,6 +31,7 @@
 #include "HAL/FileManager.h"
 #include "HAL/IConsoleManager.h"
 #include "HAL/LowLevelMemTracker.h"
+#include "HAL/LowLevelMemStats.h"
 #include "Serialization/MemoryImage.h"
 #include "Serialization/MemoryReader.h"
 #include "Serialization/MemoryWriter.h"
@@ -382,6 +384,9 @@ private:
 		if (!Buffer)
 		{
 			LLM_SCOPE(ELLMTag::FName);
+			LLM_SCOPE_DYNAMIC_STAT_OBJECTPATH_FNAME(FName{NAME_Default}, ELLMTagSet::Assets);
+			LLM_SCOPE_DYNAMIC_STAT_OBJECTPATH_FNAME(FName{NAME_Default}, ELLMTagSet::AssetClasses);
+			UE_TRACE_METADATA_SCOPE_ASSET_FNAME(NAME_None, NAME_None, NAME_None);
 			constexpr uint32 BufferSizeBytes = NumWords * sizeof(WordType);
 			Buffer = (WordType*) FMemory::MallocZeroed(BufferSizeBytes, alignof(WordType));
 			WordType* Expected = nullptr;
@@ -657,6 +662,9 @@ private:
 	static uint8* AllocBlock()
 	{
 		LLM_SCOPE(ELLMTag::FName);
+		LLM_SCOPE_DYNAMIC_STAT_OBJECTPATH_FNAME(FName{NAME_Default}, ELLMTagSet::Assets);
+		LLM_SCOPE_DYNAMIC_STAT_OBJECTPATH_FNAME(FName{NAME_Default}, ELLMTagSet::AssetClasses);
+		UE_TRACE_METADATA_SCOPE_ASSET_FNAME(NAME_None, NAME_None, NAME_None);
 		return (uint8*)FMemory::Malloc(BlockSizeBytes, alignof(FNameEntry));
 	}
 	
@@ -996,6 +1004,9 @@ public:
 	void Initialize(FNameEntryAllocator& InEntries)
 	{
 		LLM_SCOPE(ELLMTag::FName);
+		LLM_SCOPE_DYNAMIC_STAT_OBJECTPATH_FNAME(FName{NAME_Default}, ELLMTagSet::Assets);
+		LLM_SCOPE_DYNAMIC_STAT_OBJECTPATH_FNAME(FName{NAME_Default}, ELLMTagSet::AssetClasses);
+		UE_TRACE_METADATA_SCOPE_ASSET_FNAME(NAME_None, NAME_None, NAME_None);
 		Entries = &InEntries;
 
 		Slots = (FNameSlot*)FMemory::Malloc(FNamePoolInitialSlotsPerShard * sizeof(FNameSlot), alignof(FNameSlot));
@@ -1398,6 +1409,9 @@ private:
 	void Grow(const uint32 NewCapacity)
 	{
 		LLM_SCOPE(ELLMTag::FName);
+		LLM_SCOPE_DYNAMIC_STAT_OBJECTPATH_FNAME(FName{NAME_Default}, ELLMTagSet::Assets);
+		LLM_SCOPE_DYNAMIC_STAT_OBJECTPATH_FNAME(FName{NAME_Default}, ELLMTagSet::AssetClasses);
+		UE_TRACE_METADATA_SCOPE_ASSET_FNAME(NAME_None, NAME_None, NAME_None);
 		TArrayView<FNameSlot> OldSlots(Slots, Capacity());
 		const uint32 OldUsedSlots = UsedSlots;
 
