Add support for merged modules loading on console platforms

#rnx

diff --git a/Engine/Source/Runtime/Core/Core.Build.cs b/Engine/Source/Runtime/Core/Core.Build.cs
--- a/Engine/Source/Runtime/Core/Core.Build.cs
+++ b/Engine/Source/Runtime/Core/Core.Build.cs
@@ -361,7 +361,7 @@ public class Core : ModuleRules
 
 		PrivateDefinitions.Add("UE_PERSISTENT_ALLOCATOR_RESERVE_SIZE=" + GetPersistentAllocatorReserveSize().ToString() + "ULL");
 
-		PrivateDefinitions.Add("UE_MERGED_MODULES=" + (Target.bMergeModules ? "1" : "0"));
+		PublicDefinitions.Add("UE_MERGED_MODULES=" + (Target.bMergeModules ? "1" : "0"));
 	}
 
 	protected virtual bool SupportsBinaryConfig(ReadOnlyTargetRules Target)
diff --git a/Engine/Source/Runtime/Core/Private/Modules/ModuleManager.cpp b/Engine/Source/Runtime/Core/Private/Modules/ModuleManager.cpp
--- a/Engine/Source/Runtime/Core/Private/Modules/ModuleManager.cpp
+++ b/Engine/Source/Runtime/Core/Private/Modules/ModuleManager.cpp
@@ -208,7 +208,7 @@ void FModuleManager::FindModules(const TCHAR* WildcardWithoutExtension, TArray<F
 void FModuleManager::FindModules(const TCHAR* WildcardWithoutExtension, TArray<FModuleDiskInfo>& OutModules) const
 {
 	// @todo plugins: Try to convert existing use cases to use plugins, and get rid of this function
-#if !IS_MONOLITHIC
+#if !IS_MONOLITHIC && !UE_MERGED_MODULES
 
 	TMap<FName, FString> ModulePaths;
 	FindModulePaths(WildcardWithoutExtension, ModulePaths);
@@ -257,7 +257,7 @@ void FModuleManager::FindModules(const TCHAR* WildcardWithoutExtension, TArray<F
 			OutModules.Add(FModuleDiskInfo{ WildcardName, FString() });
 		}
 	}
-#endif
+#endif //  !IS_MONOLITHIC && !UE_MERGED_MODULES
 }
 
 bool FModuleManager::ModuleExists(const TCHAR* ModuleName, FString* OutModuleFilePath) const
@@ -399,6 +399,22 @@ void FModuleManager::AddModule(const FName InModuleName)
 	FModuleManager::Get().AddModuleToModulesList(InModuleName, ModuleInfo);
 }
 
+#if CPUPROFILERTRACE_ENABLED
+
+UE_TRACE_EVENT_BEGIN(Cpu, LoadModule, NoSync)
+UE_TRACE_EVENT_FIELD(UE::Trace::WideString, Name)
+UE_TRACE_EVENT_END()
+
+UE_TRACE_EVENT_BEGIN(Cpu, FPlatformProcess_GetDllHandle, NoSync)
+UE_TRACE_EVENT_FIELD(UE::Trace::WideString, Name)
+UE_TRACE_EVENT_END()
+
+UE_TRACE_EVENT_BEGIN(Cpu, StartupModule, NoSync)
+UE_TRACE_EVENT_FIELD(UE::Trace::WideString, Name)
+UE_TRACE_EVENT_END()
+
+#endif // CPUPROFILERTRACE_ENABLED
+
 #if !IS_MONOLITHIC
 void FModuleManager::RefreshModuleFilenameFromManifestImpl(const FName InModuleName, FModuleInfo& ModuleInfo)
 {
@@ -462,6 +478,90 @@ void FModuleManager::RefreshModuleFilenameFromManifest(const FName InModuleName)
 		this->RefreshModuleFilenameFromManifestImpl(InModuleName, *ModuleInfoPtr);
 	}
 }
+
+void* FModuleManager::LoadLibrary(FName ModuleName, const FString& ModuleFileToLoad)
+{
+	UE_LOG(LogModuleManager, Verbose, TEXT("LoadLibrary: %s"), *ModuleName.ToString());
+
+	void* Handle = nullptr;
+
+#if UE_MERGED_MODULES
+
+	// First, attempt to find a cached library handle
+	FModuleManagerLibraryTracker& LoadedDynamicLibrary = LoadedDynamicLibraries.FindOrAdd(ModuleFileToLoad);
+	if (LoadedDynamicLibrary.Handle != nullptr)
+	{
+		Handle = LoadedDynamicLibrary.Handle;
+	}
+	else
+
+#endif // UE_MERGED_MODULES
+
+	// If no cached handle exists, then just load the library
+	{
+#if CPUPROFILERTRACE_ENABLED
+		UE_TRACE_LOG_SCOPED_T(Cpu, FPlatformProcess_GetDllHandle, CpuChannel)
+			<< FPlatformProcess_GetDllHandle.Name(*ModuleFileToLoad);
+#endif // CPUPROFILERTRACE_ENABLED
+
+		Handle = FPlatformProcess::GetDllHandle(*ModuleFileToLoad);
+	}
+
+#if UE_MERGED_MODULES
+
+	// While using merged modules, update the usage data for the loaded library
+	LoadedDynamicLibrary.Users.AddUnique(ModuleName);
+	if (LoadedDynamicLibrary.Handle == nullptr)
+	{
+		LoadedDynamicLibrary.Handle = Handle;
+		UE_LOG(LogModuleManager, Verbose, TEXT("LoadLibrary: cached library '%s'"), *ModuleFileToLoad);
+	}
+	else
+	{
+		UE_LOG(LogModuleManager, Verbose, TEXT("LoadLibrary: added module to users of '%s'"), *ModuleFileToLoad);
+	}
+
+#endif // UE_MERGED_MODULES
+
+	return Handle;
+}
+
+void FModuleManager::FreeLibrary(FName ModuleName, void* Handle)
+{
+	UE_LOG(LogModuleManager, Verbose, TEXT("FreeLibrary: %s"), *ModuleName.ToString());
+
+#if UE_MERGED_MODULES
+
+	// Find out if we should unload a dynamic library
+	const FString* LibraryToUnload = nullptr;
+	for (TPair<FString, FModuleManagerLibraryTracker>& LibraryNameAndTracker : LoadedDynamicLibraries)
+	{
+		FModuleManagerLibraryTracker& Tracker = LibraryNameAndTracker.Value;
+
+		if (Tracker.Handle == Handle)
+		{
+			Tracker.Users.Remove(ModuleName);
+			if (Tracker.Users.Num() == 0)
+			{
+				LibraryToUnload = &LibraryNameAndTracker.Key;
+			}
+			break;
+		}
+	}
+
+	// Unload the DLL
+	if (LibraryToUnload)
+	{
+		UE_LOG(LogModuleManager, Verbose, TEXT("FreeLibrary: unloading library '%s'"), **LibraryToUnload);
+		LoadedDynamicLibraries.Remove(*LibraryToUnload);
+		FPlatformProcess::FreeDllHandle(Handle);
+	}
+
+#else
+	FPlatformProcess::FreeDllHandle(Handle);
+#endif // UE_MERGED_MODULES
+}
+
 #endif	// !IS_MONOLITHIC
 
 IModuleInterface* FModuleManager::LoadModule(const FName InModuleName, ELoadModuleFlags InLoadModuleFlags)
@@ -520,22 +620,6 @@ IModuleInterface& FModuleManager::LoadModuleChecked( const FName InModuleName )
 	return *Module;
 }
 
-#if CPUPROFILERTRACE_ENABLED
-
-UE_TRACE_EVENT_BEGIN(Cpu, LoadModule, NoSync)
-	UE_TRACE_EVENT_FIELD(UE::Trace::WideString, Name)
-UE_TRACE_EVENT_END()
-
-UE_TRACE_EVENT_BEGIN(Cpu, FPlatformProcess_GetDllHandle, NoSync)
-	UE_TRACE_EVENT_FIELD(UE::Trace::WideString, Name)
-UE_TRACE_EVENT_END()
-
-UE_TRACE_EVENT_BEGIN(Cpu, StartupModule, NoSync)
-	UE_TRACE_EVENT_FIELD(UE::Trace::WideString, Name)
-UE_TRACE_EVENT_END()
-
-#endif // CPUPROFILERTRACE_ENABLED
-
 IModuleInterface* FModuleManager::LoadModuleWithFailureReason(const FName InModuleName, EModuleLoadResult& OutFailureReason, ELoadModuleFlags InLoadModuleFlags)
 {
 	IModuleInterface* LoadedModule = nullptr;
@@ -718,13 +802,7 @@ IModuleInterface* FModuleManager::LoadModuleWithFailureReason(const FName InModu
 		// Skip this check if file manager has not yet been initialized
 		if (FPaths::FileExists(ModuleFileToLoad))
 		{
-			{
-#if CPUPROFILERTRACE_ENABLED
-				UE_TRACE_LOG_SCOPED_T(Cpu, FPlatformProcess_GetDllHandle, CpuChannel)
-					<< FPlatformProcess_GetDllHandle.Name(*ModuleFileToLoad);
-#endif // CPUPROFILERTRACE_ENABLED
-				ModuleInfo->Handle = FPlatformProcess::GetDllHandle(*ModuleFileToLoad);
-			}
+			ModuleInfo->Handle = LoadLibrary(InModuleName, ModuleFileToLoad);
 			
 			if (ModuleInfo->Handle != nullptr)
 			{
@@ -793,21 +871,28 @@ IModuleInterface* FModuleManager::LoadModuleWithFailureReason(const FName InModu
 							UE_CLOG((InLoadModuleFlags & ELoadModuleFlags::LogFailures) != ELoadModuleFlags::None,
 								LogModuleManager, Warning, TEXT("ModuleManager: Unable to load module '%s' because InitializeModule function failed (returned nullptr.)"), *ModuleFileToLoad);
 
-							FPlatformProcess::FreeDllHandle(ModuleInfo->Handle);
+							FreeLibrary(InModuleName, ModuleInfo->Handle);
 							ModuleInfo->Handle = nullptr;
 							OutFailureReason = EModuleLoadResult::FailedToInitialize;
 						}
 					}
 				}
+
+#if !UE_MERGED_MODULES
+
+				// This is normal with merged modules, as we don't have a single module for each library
 				else
 				{
 					UE_CLOG((InLoadModuleFlags & ELoadModuleFlags::LogFailures) != ELoadModuleFlags::None,
 						LogModuleManager, Warning, TEXT("ModuleManager: Unable to load module '%s' because InitializeModule function was not found."), *ModuleFileToLoad);
 
-					FPlatformProcess::FreeDllHandle(ModuleInfo->Handle);
+					FreeLibrary(InModuleName, ModuleInfo->Handle);
 					ModuleInfo->Handle = nullptr;
 					OutFailureReason = EModuleLoadResult::FailedToInitialize;
 				}
+
+#endif // !UE_MERGED_MODULES
+
 			}
 			else
 			{
@@ -861,7 +946,7 @@ bool FModuleManager::UnloadModule( const FName InModuleName, bool bIsShutdown, b
 				if( !bIsShutdown && bAllowUnloadCode )
 				{
 					// Unload the DLL
-					FPlatformProcess::FreeDllHandle( ModuleInfo.Handle );
+					FreeLibrary( InModuleName, ModuleInfo.Handle );
 				}
 				ModuleInfo.Handle = nullptr;
 			}
@@ -1568,7 +1653,7 @@ void FModuleManager::LoadModuleBinaryOnly(FName ModuleName)
 	if (ModulePaths.Num() == 1)
 	{
 		FString ModuleFilename = MoveTemp(TMap<FName, FString>::TIterator(ModulePaths).Value());
-		FPlatformProcess::GetDllHandle(*ModuleFilename);
+		LoadLibrary(ModuleName, ModuleFilename);
 	}
 #endif
 }
diff --git a/Engine/Source/Runtime/Core/Public/Modules/ModuleManager.h b/Engine/Source/Runtime/Core/Public/Modules/ModuleManager.h
--- a/Engine/Source/Runtime/Core/Public/Modules/ModuleManager.h
+++ b/Engine/Source/Runtime/Core/Public/Modules/ModuleManager.h
@@ -660,6 +660,13 @@ private:
 
 	/** Refreshes the filename of a new module from the manifest */
 	CORE_API void RefreshModuleFilenameFromManifestImpl(const FName InModuleName, FModuleInfo& ModuleInfo);
+
+	/** Load the library for a given module */
+	void* LoadLibrary(FName ModuleName, const FString& ModuleFileToLoad);
+
+	/** Attempt to free the backing library for a module */
+	void FreeLibrary(FName ModuleName, void* Handle);
+
 #endif
 
 	/** Adds pending module initializer registrations to the StaticallyLinkedModuleInitializers map. */
@@ -669,6 +676,24 @@ private:
 	/** Map of all modules.  Maps the case-insensitive module name to information about that module, loaded or not. */
 	FModuleMap Modules;
 
+#if UE_MERGED_MODULES
+
+	struct FModuleManagerLibraryTracker
+	{
+		FModuleManagerLibraryTracker()
+			: Handle(nullptr)
+			, Users()
+		{}
+
+		void* Handle;
+		TArray<FName> Users;
+	};
+
+	/** Map of loaded DLL handles for when merged modular build is in use */
+	TMap<FString, FModuleManagerLibraryTracker> LoadedDynamicLibraries;
+
+#endif // UE_MERGED_MODULES
+
 	/** Pending registrations of module names */
 	/** We use an array here to stop comparisons (and thus FNames being constructed) when they are registered. */
 	/** Instead, we validate there are no duplicates when they're inserted into StaticallyLinkedModuleInitializers. */
@@ -843,7 +868,7 @@ namespace UE::Core::Private
 		}
 	}
 }
-#if IS_MONOLITHIC
+#if IS_MONOLITHIC || UE_MERGED_MODULES
 
 	// If we're linking monolithically we assume all modules are linked in with the main binary.
 	#define IMPLEMENT_MODULE( ModuleImplClass, ModuleName ) \
