Fix build in Unity stress test where LoadLibrary is confused for Win32's macro version

#rnx

diff --git a/Engine/Source/Runtime/Core/Private/Modules/ModuleManager.cpp b/Engine/Source/Runtime/Core/Private/Modules/ModuleManager.cpp
--- a/Engine/Source/Runtime/Core/Private/Modules/ModuleManager.cpp
+++ b/Engine/Source/Runtime/Core/Private/Modules/ModuleManager.cpp
@@ -479,9 +479,9 @@ void FModuleManager::RefreshModuleFilenameFromManifest(const FName InModuleName)
 	}
 }
 
-void* FModuleManager::LoadLibrary(FName ModuleName, const FString& ModuleFileToLoad)
+void* FModuleManager::InternalLoadLibrary(FName ModuleName, const FString& ModuleFileToLoad)
 {
-	UE_LOG(LogModuleManager, Verbose, TEXT("LoadLibrary: %s"), *ModuleName.ToString());
+	UE_LOG(LogModuleManager, Verbose, TEXT("InternalLoadLibrary: %s"), *ModuleName.ToString());
 
 	void* Handle = nullptr;
 
@@ -514,11 +514,11 @@ void* FModuleManager::LoadLibrary(FName ModuleName, const FString& ModuleFileToL
 	if (LoadedDynamicLibrary.Handle == nullptr)
 	{
 		LoadedDynamicLibrary.Handle = Handle;
-		UE_LOG(LogModuleManager, Verbose, TEXT("LoadLibrary: cached library '%s'"), *ModuleFileToLoad);
+		UE_LOG(LogModuleManager, Verbose, TEXT("InternalLoadLibrary: cached library '%s'"), *ModuleFileToLoad);
 	}
 	else
 	{
-		UE_LOG(LogModuleManager, Verbose, TEXT("LoadLibrary: added module to users of '%s'"), *ModuleFileToLoad);
+		UE_LOG(LogModuleManager, Verbose, TEXT("InternalLoadLibrary: added module to users of '%s'"), *ModuleFileToLoad);
 	}
 
 #endif // UE_MERGED_MODULES
@@ -526,9 +526,9 @@ void* FModuleManager::LoadLibrary(FName ModuleName, const FString& ModuleFileToL
 	return Handle;
 }
 
-void FModuleManager::FreeLibrary(FName ModuleName, void* Handle)
+void FModuleManager::InternalFreeLibrary(FName ModuleName, void* Handle)
 {
-	UE_LOG(LogModuleManager, Verbose, TEXT("FreeLibrary: %s"), *ModuleName.ToString());
+	UE_LOG(LogModuleManager, Verbose, TEXT("InternalFreeLibrary: %s"), *ModuleName.ToString());
 
 #if UE_MERGED_MODULES
 
@@ -552,7 +552,7 @@ void FModuleManager::FreeLibrary(FName ModuleName, void* Handle)
 	// Unload the DLL
 	if (LibraryToUnload)
 	{
-		UE_LOG(LogModuleManager, Verbose, TEXT("FreeLibrary: unloading library '%s'"), **LibraryToUnload);
+		UE_LOG(LogModuleManager, Verbose, TEXT("InternalFreeLibrary: unloading library '%s'"), **LibraryToUnload);
 		LoadedDynamicLibraries.Remove(*LibraryToUnload);
 		FPlatformProcess::FreeDllHandle(Handle);
 	}
@@ -802,7 +802,7 @@ IModuleInterface* FModuleManager::LoadModuleWithFailureReason(const FName InModu
 		// Skip this check if file manager has not yet been initialized
 		if (FPaths::FileExists(ModuleFileToLoad))
 		{
-			ModuleInfo->Handle = LoadLibrary(InModuleName, ModuleFileToLoad);
+			ModuleInfo->Handle = InternalLoadLibrary(InModuleName, ModuleFileToLoad);
 			
 			if (ModuleInfo->Handle != nullptr)
 			{
@@ -871,7 +871,7 @@ IModuleInterface* FModuleManager::LoadModuleWithFailureReason(const FName InModu
 							UE_CLOG((InLoadModuleFlags & ELoadModuleFlags::LogFailures) != ELoadModuleFlags::None,
 								LogModuleManager, Warning, TEXT("ModuleManager: Unable to load module '%s' because InitializeModule function failed (returned nullptr.)"), *ModuleFileToLoad);
 
-							FreeLibrary(InModuleName, ModuleInfo->Handle);
+							InternalFreeLibrary(InModuleName, ModuleInfo->Handle);
 							ModuleInfo->Handle = nullptr;
 							OutFailureReason = EModuleLoadResult::FailedToInitialize;
 						}
@@ -886,7 +886,7 @@ IModuleInterface* FModuleManager::LoadModuleWithFailureReason(const FName InModu
 					UE_CLOG((InLoadModuleFlags & ELoadModuleFlags::LogFailures) != ELoadModuleFlags::None,
 						LogModuleManager, Warning, TEXT("ModuleManager: Unable to load module '%s' because InitializeModule function was not found."), *ModuleFileToLoad);
 
-					FreeLibrary(InModuleName, ModuleInfo->Handle);
+					InternalFreeLibrary(InModuleName, ModuleInfo->Handle);
 					ModuleInfo->Handle = nullptr;
 					OutFailureReason = EModuleLoadResult::FailedToInitialize;
 				}
@@ -946,7 +946,7 @@ bool FModuleManager::UnloadModule( const FName InModuleName, bool bIsShutdown, b
 				if( !bIsShutdown && bAllowUnloadCode )
 				{
 					// Unload the DLL
-					FreeLibrary( InModuleName, ModuleInfo.Handle );
+					InternalFreeLibrary( InModuleName, ModuleInfo.Handle );
 				}
 				ModuleInfo.Handle = nullptr;
 			}
@@ -1653,7 +1653,7 @@ void FModuleManager::LoadModuleBinaryOnly(FName ModuleName)
 	if (ModulePaths.Num() == 1)
 	{
 		FString ModuleFilename = MoveTemp(TMap<FName, FString>::TIterator(ModulePaths).Value());
-		LoadLibrary(ModuleName, ModuleFilename);
+		InternalLoadLibrary(ModuleName, ModuleFilename);
 	}
 #endif
 }
diff --git a/Engine/Source/Runtime/Core/Public/Modules/ModuleManager.h b/Engine/Source/Runtime/Core/Public/Modules/ModuleManager.h
--- a/Engine/Source/Runtime/Core/Public/Modules/ModuleManager.h
+++ b/Engine/Source/Runtime/Core/Public/Modules/ModuleManager.h
@@ -662,10 +662,10 @@ private:
 	CORE_API void RefreshModuleFilenameFromManifestImpl(const FName InModuleName, FModuleInfo& ModuleInfo);
 
 	/** Load the library for a given module */
-	void* LoadLibrary(FName ModuleName, const FString& ModuleFileToLoad);
+	void* InternalLoadLibrary(FName ModuleName, const FString& ModuleFileToLoad);
 
 	/** Attempt to free the backing library for a module */
-	void FreeLibrary(FName ModuleName, void* Handle);
+	void InternalFreeLibrary(FName ModuleName, void* Handle);
 
 #endif
 
