FPaths::ValidatePath: Add a special case to return true for windows-style long paths \\?\K:\Path\Dir.
#rnx

diff --git a/Engine/Source/Runtime/Core/Private/Misc/Paths.cpp b/Engine/Source/Runtime/Core/Private/Misc/Paths.cpp
--- a/Engine/Source/Runtime/Core/Private/Misc/Paths.cpp
+++ b/Engine/Source/Runtime/Core/Private/Misc/Paths.cpp
@@ -1657,7 +1657,7 @@ FString FPaths::MakeValidFileName(const FString& InString, const TCHAR InReplace
 	return FString(Output.GetData());
 }
 
-bool FPaths::ValidatePath( const FString& InPath, FText* OutReason )
+bool FPaths::ValidatePath(const FString& InPath, FText* OutReason)
 {
 	const FString RestrictedChars = GetInvalidFileSystemChars();
 	static const TCHAR* RestrictedNames[] = {	TEXT("CON"), TEXT("PRN"), TEXT("AUX"), TEXT("CLOCK$"), TEXT("NUL"),
@@ -1667,76 +1667,104 @@ bool FPaths::ValidatePath( const FString& InPath, FText* OutReason )
 	FString Standardized = InPath;
 	NormalizeFilename(Standardized);
 	CollapseRelativeDirectories(Standardized);
-	RemoveDuplicateSlashes(Standardized);
-
-	// The loop below requires that the path not end with a /
-	if(Standardized.EndsWith(TEXT("/"), ESearchCase::CaseSensitive))
+	// Remove duplicate slashes, to normalize the path,
+	// but don't remove them at beginning so we can recognize network paths
+	if (Standardized.StartsWith(TEXT("//")))
 	{
-		Standardized.LeftChopInline(1, EAllowShrinking::No);
+		Standardized.RightChopInline(1, EAllowShrinking::No);
+		RemoveDuplicateSlashes(Standardized);
+		Standardized = TEXT("/") + MoveTemp(Standardized);
 	}
-
-	// Walk each part of the path looking for name errors
-	for(int32 StartPos = 0, EndPos = Standardized.Find(TEXT("/"), ESearchCase::CaseSensitive); ; 
-		StartPos = EndPos + 1, EndPos = Standardized.Find(TEXT("/"), ESearchCase::CaseSensitive, ESearchDir::FromStart, StartPos)
-		)
+	else
 	{
-		const bool bIsLastPart = EndPos == INDEX_NONE;
-		const FString PathPart = Standardized.Mid(StartPos, (bIsLastPart) ? MAX_int32 : EndPos - StartPos);
+		RemoveDuplicateSlashes(Standardized);
+	}
 
-		// If this is the first part of the path, it's possible for it to be a drive name and is allowed to contain a colon
-		if(StartPos == 0 && IsDrive(PathPart))
+	// Walk each part of the path looking for name errors
+	// Allow colons in the first component
+	bool bResult = true;
+	bool bAllowNextTokenToBeDrive = true;
+	FPathViews::IterateComponents(Standardized,
+		[&bAllowNextTokenToBeDrive,&Standardized,&RestrictedChars,&bResult,&OutReason]
+		(FStringView PathPart)
+	{
+		if (!bResult)
 		{
-			if(bIsLastPart)
-			{
-				break;
-			}
-			continue;
+			// Keep the first error and stop validating
+			return;
+		}
+		bool bAllowDrive = bAllowNextTokenToBeDrive;
+		bAllowNextTokenToBeDrive = false;
+		if (PathPart.IsEmpty())
+		{
+			// We can get an empty path part due to duplicate slashes at the beginning or a
+			// a terminating slash
+			return;
 		}
 
 		// Check for invalid characters
-		TCHAR CharString[] = { TEXT('\0'), TEXT('\0') };
 		FString MatchedInvalidChars;
-		for(const TCHAR* InvalidCharacters = *RestrictedChars; *InvalidCharacters; ++InvalidCharacters)
+		for (TCHAR InvalidCharacter : RestrictedChars)
 		{
-			CharString[0] = *InvalidCharacters;
-			if(PathPart.Contains(CharString))
+			int32 UnusedIndex;
+			if (PathPart.FindChar(InvalidCharacter, UnusedIndex))
 			{
-				MatchedInvalidChars += *InvalidCharacters;
+				if (InvalidCharacter == ':' && bAllowDrive && IsDrive(FString(PathPart)))
+				{
+					// Colons are allowed in drive specifiers
+					continue;
+				}
+				if (InvalidCharacter == '?')
+				{
+					int64 StartPos = PathPart.GetData() - *Standardized;
+					if (StartPos < 0 || StartPos > int64(Standardized.Len()))
+					{
+						StartPos = INDEX_NONE;
+					}
+					// ? is allowed at the beginning of the path to support windows-style long paths: "\\?\K:\LongPath"
+					if (StartPos == 2 && PathPart.Len() == 1 && Standardized.StartsWith(TEXT("//"))
+						&& Standardized.Len() > 4 && Standardized[3] == '/')
+					{
+						// Also allow a drive specifier in the next pathpart
+						bAllowNextTokenToBeDrive = true;
+						continue;
+					}
+				}
+
+				MatchedInvalidChars += InvalidCharacter;
 			}
 		}
-		if(MatchedInvalidChars.Len())
+
+		if (MatchedInvalidChars.Len())
 		{
-			if(OutReason)
+			if (OutReason)
 			{
 				FFormatNamedArguments Args;
 				Args.Add(TEXT("IllegalPathCharacters"), FText::FromString(MatchedInvalidChars));
 				*OutReason = FText::Format(NSLOCTEXT("Core", "PathContainsInvalidCharacters", "Path may not contain the following characters: {IllegalPathCharacters}"), Args);
 			}
-			return false;
+			bResult = false;
+			return;
 		}
 
 		// Check for reserved names
-		for(const TCHAR* RestrictedName : RestrictedNames)
+		for (const TCHAR* RestrictedName : RestrictedNames)
 		{
-			if(PathPart.Equals(RestrictedName, ESearchCase::IgnoreCase))
+			if (PathPart.Equals(RestrictedName, ESearchCase::IgnoreCase))
 			{
-				if(OutReason)
+				if (OutReason)
 				{
 					FFormatNamedArguments Args;
 					Args.Add(TEXT("RestrictedName"), FText::FromString(RestrictedName));
 					*OutReason = FText::Format(NSLOCTEXT("Core", "PathContainsRestrictedName", "Path may not contain a restricted name: {RestrictedName}"), Args);
 				}
-				return false;
+				bResult = false;
+				return;
 			}
 		}
+	});
 
-		if(bIsLastPart)
-		{
-			break;
-		}
-	}
-
-	return true;
+	return bResult;
 }
 
 void FPaths::Split( const FString& InPath, FString& PathPart, FString& FilenamePart, FString& ExtensionPart )
diff --git a/Engine/Source/Runtime/Core/Tests/Misc/PathsTest.cpp b/Engine/Source/Runtime/Core/Tests/Misc/PathsTest.cpp
--- a/Engine/Source/Runtime/Core/Tests/Misc/PathsTest.cpp
+++ b/Engine/Source/Runtime/Core/Tests/Misc/PathsTest.cpp
@@ -172,6 +172,36 @@ TEST_CASE_NAMED(FPathTests, "System::Core::Misc::Paths", "[ApplicationContextMas
 	RunSplitTest(TEXT("C:\\Folder\\First.Last\\File.txt"), TEXT("C:\\Folder\\First.Last"), TEXT("File"), TEXT("txt"));
 	RunSplitTest(TEXT("C:\\Folder\\First.Last\\File.tar.gz"), TEXT("C:\\Folder\\First.Last"), TEXT("File.tar"), TEXT("gz"));
 
+	// ValidatePath
+	auto RunValidatePathTest = [](const TCHAR* Path, bool bExpectedValue)
+		{
+			bool bValue = FPaths::ValidatePath(FString(Path));
+			if (bValue != bExpectedValue)
+			{
+				FAIL_CHECK(FString::Printf(TEXT("ValidatePath(\"%s\") == %s, expected %s."),
+					Path, *LexToString(bValue), *LexToString(bExpectedValue)));
+			}
+		};
+	RunValidatePathTest(TEXT("C:/Path/Dir"), true);
+	RunValidatePathTest(TEXT("C:\\Path\\Dir"), true);
+	RunValidatePathTest(TEXT("//Path/Dir"), true);
+	RunValidatePathTest(TEXT("\\\\Path\\Dir"), true);
+	RunValidatePathTest(TEXT("/Path/Dir"), true);
+	RunValidatePathTest(TEXT("\\Path\\Dir"), true);
+	RunValidatePathTest(TEXT("Path/Dir"), true);
+	RunValidatePathTest(TEXT("Path\\Dir"), true);
+	RunValidatePathTest(TEXT("Path"), true);
+	RunValidatePathTest(TEXT("/Path/"), true);
+	RunValidatePathTest(TEXT("\\Path\\"), true);
+	// LongPaths on windows allow the otherwise illegal character ? at the beginning, in "\\?\"
+	RunValidatePathTest(TEXT("\\\\?\\K:\\myfolder\\myfile.fbx"), true);
+	// Otherwise ? is not allowed
+	RunValidatePathTest(TEXT("C:\\Path\\Dir?"), false);
+	// : is not allowed outside of a drive specifier at beginning
+	RunValidatePathTest(TEXT("C:\\Path\\Dir:"), false);
+	// A few other banned characters from FPaths::GetInvalidFileSystemChars
+	RunValidatePathTest(TEXT("C:\\Path\\Dir&"), false);
+	RunValidatePathTest(TEXT("C:\\Path\\Dir^"), false);
 }
 
 #endif //WITH_TESTS
