Added a missing FAndroidPlatformProperties::GetMemoryMappingAlignment()

diff --git a/Engine/Source/Runtime/Core/Public/Android/AndroidPlatformProperties.h b/Engine/Source/Runtime/Core/Public/Android/AndroidPlatformProperties.h
--- a/Engine/Source/Runtime/Core/Public/Android/AndroidPlatformProperties.h
+++ b/Engine/Source/Runtime/Core/Public/Android/AndroidPlatformProperties.h
@@ -129,6 +129,11 @@ struct FAndroidPlatformProperties
 		return true;
 	}
 
+	static FORCEINLINE int64 GetMemoryMappingAlignment()
+	{
+		return 4096;
+	}
+
 	static FORCEINLINE bool SupportsVirtualTextureStreaming()
 	{
 		return true;
