[AutoRTFM] Enforce that the UE_AUTORTFM_REGISTER_OPEN_FUNCTION* function signatures match

Fix the places where they didn't actually match. Add more tests.

diff --git a/Engine/Source/Programs/AutoRTFMTests/Private/API.cpp b/Engine/Source/Programs/AutoRTFMTests/Private/API.cpp
--- a/Engine/Source/Programs/AutoRTFMTests/Private/API.cpp
+++ b/Engine/Source/Programs/AutoRTFMTests/Private/API.cpp
@@ -733,8 +733,8 @@ TEST_CASE("API.Close")
 TEST_CASE("API.RegisterOpenFunction")
 {
     AutoRTFM::ForTheRuntime::RegisterOpenFunction(
-        reinterpret_cast<void*>(NoAutoRTFM::DoSomethingCpp),
-        reinterpret_cast<void*>(NoAutoRTFM::DoSomethingInTransactionCpp));
+        NoAutoRTFM::DoSomethingCpp,
+        NoAutoRTFM::DoSomethingInTransactionCpp);
 
     int I = -42;
 
diff --git a/Engine/Source/Programs/AutoRTFMTests/Private/API.h b/Engine/Source/Programs/AutoRTFMTests/Private/API.h
--- a/Engine/Source/Programs/AutoRTFMTests/Private/API.h
+++ b/Engine/Source/Programs/AutoRTFMTests/Private/API.h
@@ -5,8 +5,8 @@
 namespace NoAutoRTFM
 {
     int DoSomethingC(int I);
-    int DoSomethingInTransactionC(int I, void*);
+    int DoSomethingInTransactionC(int I);
 
     int DoSomethingCpp(int I);
-    int DoSomethingInTransactionCpp(int I, void*);
+    int DoSomethingInTransactionCpp(int I);
 }
diff --git a/Engine/Source/Programs/AutoRTFMTests/Private/APINoAutoRTFM.cpp b/Engine/Source/Programs/AutoRTFMTests/Private/APINoAutoRTFM.cpp
--- a/Engine/Source/Programs/AutoRTFMTests/Private/APINoAutoRTFM.cpp
+++ b/Engine/Source/Programs/AutoRTFMTests/Private/APINoAutoRTFM.cpp
@@ -8,7 +8,7 @@ UE_AUTORTFM_NOAUTORTFM int NoAutoRTFM::DoSomethingC(int I)
     return I + 13;
 }
 
-UE_AUTORTFM_NOAUTORTFM int NoAutoRTFM::DoSomethingInTransactionC(int I, void*)
+UE_AUTORTFM_NOAUTORTFM int NoAutoRTFM::DoSomethingInTransactionC(int I)
 {
     return I + 42;
 }
@@ -18,7 +18,7 @@ UE_AUTORTFM_NOAUTORTFM int NoAutoRTFM::DoSomethingCpp(int I)
     return I + 13;
 }
 
-UE_AUTORTFM_NOAUTORTFM int NoAutoRTFM::DoSomethingInTransactionCpp(int I, void*)
+UE_AUTORTFM_NOAUTORTFM int NoAutoRTFM::DoSomethingInTransactionCpp(int I)
 {
     return I + 42;
 }
diff --git a/Engine/Source/Programs/AutoRTFMTests/Private/CString.cpp b/Engine/Source/Programs/AutoRTFMTests/Private/CString.cpp
--- a/Engine/Source/Programs/AutoRTFMTests/Private/CString.cpp
+++ b/Engine/Source/Programs/AutoRTFMTests/Private/CString.cpp
@@ -131,27 +131,59 @@ TEST_CASE("CString.memcpy")
 	const char* const From = "Kittie says meow";
 	char To[] = "Doggie says woof";
 
-	SECTION("With Abort")
+	SECTION("Direct")
 	{
-		CLANG_BEGIN_DISABLE_OPTIMIZATIONS
-		AutoRTFM::ETransactionResult Result = AutoRTFM::Transact([&]
+		SECTION("With Abort")
 		{
-			memcpy(To, From, 6);
-			AutoRTFM::AbortTransaction();
-		});
-		CLANG_END_DISABLE_OPTIMIZATIONS
+			CLANG_BEGIN_DISABLE_OPTIMIZATIONS
+			AutoRTFM::ETransactionResult Result = AutoRTFM::Transact([&]
+			{
+				memcpy(To, From, 6);
+				AutoRTFM::AbortTransaction();
+			});
+			CLANG_END_DISABLE_OPTIMIZATIONS
 
-		REQUIRE(AutoRTFM::ETransactionResult::AbortedByRequest == Result);
-		REQUIRE("Doggie says woof" == std::string_view(To));
+			REQUIRE(AutoRTFM::ETransactionResult::AbortedByRequest == Result);
+			REQUIRE("Doggie says woof" == std::string_view(To));
+		}
+
+		SECTION("With Commit")
+		{
+			CLANG_BEGIN_DISABLE_OPTIMIZATIONS
+			AutoRTFM::Commit([&] { memcpy(To, From, 6); });
+			CLANG_END_DISABLE_OPTIMIZATIONS
+
+			REQUIRE("Kittie says woof" == std::string_view(To));
+		}
 	}
 
-	SECTION("With Commit")
+	SECTION("Indirect")
 	{
-		CLANG_BEGIN_DISABLE_OPTIMIZATIONS
-		AutoRTFM::Commit([&] { memcpy(To, From, 6); });
-		CLANG_END_DISABLE_OPTIMIZATIONS
+		using FnPtr = void*(*)(void*, const void*, size_t);
+		FnPtr Fn = &memcpy;
 
-		REQUIRE("Kittie says woof" == std::string_view(To));
+		SECTION("With Abort")
+		{
+			CLANG_BEGIN_DISABLE_OPTIMIZATIONS
+			AutoRTFM::ETransactionResult Result = AutoRTFM::Transact([&]
+			{
+				Fn(To, From, 6);
+				AutoRTFM::AbortTransaction();
+			});
+			CLANG_END_DISABLE_OPTIMIZATIONS
+
+			REQUIRE(AutoRTFM::ETransactionResult::AbortedByRequest == Result);
+			REQUIRE("Doggie says woof" == std::string_view(To));
+		}
+
+		SECTION("With Commit")
+		{
+			CLANG_BEGIN_DISABLE_OPTIMIZATIONS
+			AutoRTFM::Commit([&] { Fn(To, From, 6); });
+			CLANG_END_DISABLE_OPTIMIZATIONS
+
+			REQUIRE("Kittie says woof" == std::string_view(To));
+		}
 	}
 }
 
@@ -159,27 +191,119 @@ TEST_CASE("CString.memmove")
 {
 	char To[] = "Hello, world!";
 
-	SECTION("With Abort")
+	SECTION("Direct")
 	{
-		CLANG_BEGIN_DISABLE_OPTIMIZATIONS
-		AutoRTFM::ETransactionResult Result = AutoRTFM::Transact([&]
+		SECTION("With Abort")
 		{
-			memmove(To + 7, To, 5);
-			AutoRTFM::AbortTransaction();
-		});
-		CLANG_END_DISABLE_OPTIMIZATIONS
+			CLANG_BEGIN_DISABLE_OPTIMIZATIONS
+			AutoRTFM::ETransactionResult Result = AutoRTFM::Transact([&]
+			{
+				memmove(To + 7, To, 5);
+				AutoRTFM::AbortTransaction();
+			});
+			CLANG_END_DISABLE_OPTIMIZATIONS
 
-		REQUIRE(AutoRTFM::ETransactionResult::AbortedByRequest == Result);
-		REQUIRE("Hello, world!" == std::string_view(To));
+			REQUIRE(AutoRTFM::ETransactionResult::AbortedByRequest == Result);
+			REQUIRE("Hello, world!" == std::string_view(To));
+		}
+
+		SECTION("With Commit")
+		{
+			CLANG_BEGIN_DISABLE_OPTIMIZATIONS
+			AutoRTFM::Commit([&] { memmove(To + 7, To, 5); });
+			CLANG_END_DISABLE_OPTIMIZATIONS
+
+			REQUIRE("Hello, Hello!" == std::string_view(To));
+		}
 	}
 
-	SECTION("With Commit")
+	SECTION("Indirect")
 	{
-		CLANG_BEGIN_DISABLE_OPTIMIZATIONS
-		AutoRTFM::Commit([&] { memmove(To + 7, To, 5); });
-		CLANG_END_DISABLE_OPTIMIZATIONS
+		using FnPtr = void*(*)(void*, const void*, size_t);
+		FnPtr Fn = &memmove;
+
+		SECTION("With Abort")
+		{
+			CLANG_BEGIN_DISABLE_OPTIMIZATIONS
+			AutoRTFM::ETransactionResult Result = AutoRTFM::Transact([&]
+			{
+				Fn(To + 7, To, 5);
+				AutoRTFM::AbortTransaction();
+			});
+			CLANG_END_DISABLE_OPTIMIZATIONS
 
-		REQUIRE("Hello, Hello!" == std::string_view(To));
+			REQUIRE(AutoRTFM::ETransactionResult::AbortedByRequest == Result);
+			REQUIRE("Hello, world!" == std::string_view(To));
+		}
+
+		SECTION("With Commit")
+		{
+			CLANG_BEGIN_DISABLE_OPTIMIZATIONS
+			AutoRTFM::Commit([&] { Fn(To + 7, To, 5); });
+			CLANG_END_DISABLE_OPTIMIZATIONS
+
+			REQUIRE("Hello, Hello!" == std::string_view(To));
+		}
+	}
+}
+
+TEST_CASE("CString.memset")
+{
+	char To[] = "Hello, world!";
+
+	SECTION("Direct")
+	{
+		SECTION("With Abort")
+		{
+			CLANG_BEGIN_DISABLE_OPTIMIZATIONS
+			AutoRTFM::ETransactionResult Result = AutoRTFM::Transact([&]
+			{
+				memset(To + 7, 'X', 5);
+				AutoRTFM::AbortTransaction();
+			});
+			CLANG_END_DISABLE_OPTIMIZATIONS
+
+			REQUIRE(AutoRTFM::ETransactionResult::AbortedByRequest == Result);
+			REQUIRE("Hello, world!" == std::string_view(To));
+		}
+
+		SECTION("With Commit")
+		{
+			CLANG_BEGIN_DISABLE_OPTIMIZATIONS
+			AutoRTFM::Commit([&] { memset(To + 7, 'X', 5); });
+			CLANG_END_DISABLE_OPTIMIZATIONS
+
+			REQUIRE("Hello, XXXXX!" == std::string_view(To));
+		}
+	}
+
+	SECTION("Indirect")
+	{
+		using FnPtr = void*(*)(void*, int, size_t);
+		FnPtr Fn = &memset;
+
+		SECTION("With Abort")
+		{
+			CLANG_BEGIN_DISABLE_OPTIMIZATIONS
+			AutoRTFM::ETransactionResult Result = AutoRTFM::Transact([&]
+			{
+				Fn(To + 7, 'X', 5);
+				AutoRTFM::AbortTransaction();
+			});
+			CLANG_END_DISABLE_OPTIMIZATIONS
+
+			REQUIRE(AutoRTFM::ETransactionResult::AbortedByRequest == Result);
+			REQUIRE("Hello, world!" == std::string_view(To));
+		}
+
+		SECTION("With Commit")
+		{
+			CLANG_BEGIN_DISABLE_OPTIMIZATIONS
+			AutoRTFM::Commit([&] { Fn(To + 7, 'X', 5); });
+			CLANG_END_DISABLE_OPTIMIZATIONS
+
+			REQUIRE("Hello, XXXXX!" == std::string_view(To));
+		}
 	}
 }
 
diff --git a/Engine/Source/Runtime/Core/Private/AutoRTFM/StdLibSupport.cpp b/Engine/Source/Runtime/Core/Private/AutoRTFM/StdLibSupport.cpp
--- a/Engine/Source/Runtime/Core/Private/AutoRTFM/StdLibSupport.cpp
+++ b/Engine/Source/Runtime/Core/Private/AutoRTFM/StdLibSupport.cpp
@@ -90,11 +90,25 @@ static void ThrowIfFormatContainsPercentN(const wchar_t* Format)
 
 }
 
-UE_AUTORTFM_REGISTER_OPEN_FUNCTION_EXPLICIT(memcpy, Memcpy);
-UE_AUTORTFM_REGISTER_OPEN_FUNCTION_EXPLICIT(memmove, Memmove);
-UE_AUTORTFM_REGISTER_OPEN_FUNCTION_EXPLICIT(memset, Memset);
+void* RTFM_memcpy(void* Dst, const void* Src, size_t Size) throw()
+{
+    return Memcpy(Dst, Src, Size, FContext::Get());
+}
+UE_AUTORTFM_REGISTER_OPEN_FUNCTION(memcpy);
+
+void* RTFM_memmove(void* Dst, const void* Src, size_t Size) throw()
+{
+    return Memmove(Dst, Src, Size, FContext::Get());
+}
+UE_AUTORTFM_REGISTER_OPEN_FUNCTION(memmove);
+
+void* RTFM_memset(void* Dst, int Value, size_t Size) throw()
+{
+    return Memset(Dst, Value, Size, FContext::Get());
+}
+UE_AUTORTFM_REGISTER_OPEN_FUNCTION(memset);
 
-void* RTFM_malloc(size_t Size)
+void* RTFM_malloc(size_t Size) throw()
 {
     void* Result = malloc(Size);
 	FContext* Context = FContext::Get();
@@ -108,7 +122,7 @@ void* RTFM_malloc(size_t Size)
 
 UE_AUTORTFM_REGISTER_OPEN_FUNCTION(malloc);
 
-void* RTFM_calloc(size_t Count, size_t Size)
+void* RTFM_calloc(size_t Count, size_t Size) throw()
 {
     void* Result = calloc(Count, Size);
 	FContext* Context = FContext::Get();
@@ -122,7 +136,7 @@ void* RTFM_calloc(size_t Count, size_t Size)
 
 UE_AUTORTFM_REGISTER_OPEN_FUNCTION(calloc);
 
-void RTFM_free(void* Ptr)
+void RTFM_free(void* Ptr) throw()
 {
     if (Ptr)
     {
@@ -135,7 +149,7 @@ void RTFM_free(void* Ptr)
 }
 UE_AUTORTFM_REGISTER_OPEN_FUNCTION(free);
 
-void* RTFM_realloc(void* Ptr, size_t Size)
+void* RTFM_realloc(void* Ptr, size_t Size) throw()
 {
     void* NewObject = RTFM_malloc(Size);
     if (Ptr)
@@ -155,7 +169,7 @@ void* RTFM_realloc(void* Ptr, size_t Size)
 }
 UE_AUTORTFM_REGISTER_OPEN_FUNCTION(realloc);
 
-char* RTFM_strcpy(char* const Dst, const char* const Src)
+char* RTFM_strcpy(char* const Dst, const char* const Src) throw()
 {
     const size_t SrcLen = strlen(Src);
 
@@ -165,7 +179,7 @@ char* RTFM_strcpy(char* const Dst, const char* const Src)
 }
 UE_AUTORTFM_REGISTER_OPEN_FUNCTION(strcpy);
 
-char* RTFM_strncpy(char* const Dst, const char* const Src, const size_t Num)
+char* RTFM_strncpy(char* const Dst, const char* const Src, const size_t Num) throw()
 {
 	FContext* Context = FContext::Get();
     Context->RecordWrite(Dst, Num);
@@ -173,7 +187,7 @@ char* RTFM_strncpy(char* const Dst, const char* const Src, const size_t Num)
 }
 UE_AUTORTFM_REGISTER_OPEN_FUNCTION(strncpy);
 
-char* RTFM_strcat(char* const Dst, const char* const Src)
+char* RTFM_strcat(char* const Dst, const char* const Src) throw()
 {
     const size_t DstLen = strlen(Dst);
     const size_t SrcLen = strlen(Src);
@@ -184,7 +198,7 @@ char* RTFM_strcat(char* const Dst, const char* const Src)
 }
 UE_AUTORTFM_REGISTER_OPEN_FUNCTION(strcat);
 
-char* RTFM_strncat(char* const Dst, const char* const Src, const size_t Num)
+char* RTFM_strncat(char* const Dst, const char* const Src, const size_t Num) throw()
 {
     const size_t DstLen = strlen(Dst);
 
@@ -197,12 +211,12 @@ UE_AUTORTFM_REGISTER_OPEN_FUNCTION(strncat);
 UE_AUTORTFM_REGISTER_SELF_FUNCTION(memcmp);
 UE_AUTORTFM_REGISTER_SELF_FUNCTION(strcmp);
 UE_AUTORTFM_REGISTER_SELF_FUNCTION(strncmp);
-UE_AUTORTFM_REGISTER_SELF_FUNCTION(static_cast<const char*(*)(const char*, int)>(&strchr));
-UE_AUTORTFM_REGISTER_SELF_FUNCTION(static_cast<const char* (*)(const char*, int)>(&strrchr));
-UE_AUTORTFM_REGISTER_SELF_FUNCTION(static_cast<const char* (*)(const char*, const char*)>(&strstr));
+UE_AUTORTFM_REGISTER_SELF_OVERLOADED_FUNCTION(const char* (const char*, int), strchr);
+UE_AUTORTFM_REGISTER_SELF_OVERLOADED_FUNCTION(const char* (const char*, int), strrchr);
+UE_AUTORTFM_REGISTER_SELF_OVERLOADED_FUNCTION(const char* (const char*, const char*), strstr);
 UE_AUTORTFM_REGISTER_SELF_FUNCTION(strlen);
 
-long int RTFM_strtol(const char* String, char** EndPtr, int Radix)
+long RTFM_strtol(char const* String, char** EndPtr, int Radix) throw()
 {
     if (nullptr != EndPtr)
     {
@@ -214,11 +228,12 @@ long int RTFM_strtol(const char* String, char** EndPtr, int Radix)
 }
 UE_AUTORTFM_REGISTER_OPEN_FUNCTION(strtol);
 
-UE_AUTORTFM_REGISTER_SELF_FUNCTION(static_cast<const wchar_t* (*)(const wchar_t*, wchar_t)>(&wcschr));
-UE_AUTORTFM_REGISTER_SELF_FUNCTION(static_cast<wchar_t* (*)(wchar_t*, wchar_t)>(&wcschr));
-UE_AUTORTFM_REGISTER_SELF_FUNCTION(static_cast<wchar_t* (*)(wchar_t*, const wchar_t*)>(&wcsstr));
-UE_AUTORTFM_REGISTER_SELF_FUNCTION(static_cast<const wchar_t* (*)(const wchar_t*, const wchar_t*)>(&wcsstr));
+UE_AUTORTFM_REGISTER_SELF_OVERLOADED_FUNCTION(const wchar_t* (const wchar_t*, wchar_t), wcschr);
+UE_AUTORTFM_REGISTER_SELF_OVERLOADED_FUNCTION(wchar_t* (wchar_t*, wchar_t), wcschr);
+UE_AUTORTFM_REGISTER_SELF_OVERLOADED_FUNCTION(wchar_t* (wchar_t*, const wchar_t*), wcsstr);
+UE_AUTORTFM_REGISTER_SELF_OVERLOADED_FUNCTION(const wchar_t*(const wchar_t*, const wchar_t*), wcsstr);
 UE_AUTORTFM_REGISTER_SELF_FUNCTION(wcscmp);
+UE_AUTORTFM_REGISTER_SELF_FUNCTION(wcslen);
 
 UE_AUTORTFM_REGISTER_SELF_FUNCTION(iswupper);
 UE_AUTORTFM_REGISTER_SELF_FUNCTION(iswlower);
@@ -255,46 +270,46 @@ UE_AUTORTFM_REGISTER_SELF_FUNCTION(fmod);
 UE_AUTORTFM_REGISTER_SELF_FUNCTION(modf);
 // Linux (likely Mac) have ambiguous overrides to these math functions
 #else
-UE_AUTORTFM_REGISTER_SELF_FUNCTION(static_cast<float(*)(float)>(&sqrt));
-UE_AUTORTFM_REGISTER_SELF_FUNCTION(static_cast<long double(*)(long double)>(&sqrt));
-UE_AUTORTFM_REGISTER_SELF_FUNCTION(static_cast<float(*)(float)>(&sin));
-UE_AUTORTFM_REGISTER_SELF_FUNCTION(static_cast<long double(*)(long double)>(&sin));
-UE_AUTORTFM_REGISTER_SELF_FUNCTION(static_cast<float(*)(float)>(&cos));
-UE_AUTORTFM_REGISTER_SELF_FUNCTION(static_cast<long double(*)(long double)>(&cos));
-UE_AUTORTFM_REGISTER_SELF_FUNCTION(static_cast<float(*)(float)>(&tan));
-UE_AUTORTFM_REGISTER_SELF_FUNCTION(static_cast<long double(*)(long double)>(&tan));
-UE_AUTORTFM_REGISTER_SELF_FUNCTION(static_cast<float(*)(float)>(&asin));
-UE_AUTORTFM_REGISTER_SELF_FUNCTION(static_cast<long double(*)(long double)>(&asin));
-UE_AUTORTFM_REGISTER_SELF_FUNCTION(static_cast<float(*)(float)>(&acos));
-UE_AUTORTFM_REGISTER_SELF_FUNCTION(static_cast<long double(*)(long double)>(&acos));
-UE_AUTORTFM_REGISTER_SELF_FUNCTION(static_cast<float(*)(float)>(&atan));
-UE_AUTORTFM_REGISTER_SELF_FUNCTION(static_cast<long double(*)(long double)>(&atan));
-UE_AUTORTFM_REGISTER_SELF_FUNCTION(static_cast<float(*)(float, float)>(&atan2));
-UE_AUTORTFM_REGISTER_SELF_FUNCTION(static_cast<long double(*)(long double, long double)>(&atan2));
-UE_AUTORTFM_REGISTER_SELF_FUNCTION(static_cast<float(*)(float)>(&sinh));
-UE_AUTORTFM_REGISTER_SELF_FUNCTION(static_cast<long double(*)(long double)>(&sinh));
-UE_AUTORTFM_REGISTER_SELF_FUNCTION(static_cast<float(*)(float)>(&cosh));
-UE_AUTORTFM_REGISTER_SELF_FUNCTION(static_cast<long double(*)(long double)>(&cosh));
-UE_AUTORTFM_REGISTER_SELF_FUNCTION(static_cast<float(*)(float)>(&tanh));
-UE_AUTORTFM_REGISTER_SELF_FUNCTION(static_cast<long double(*)(long double)>(&tanh));
-UE_AUTORTFM_REGISTER_SELF_FUNCTION(static_cast<float(*)(float)>(&asinh));
-UE_AUTORTFM_REGISTER_SELF_FUNCTION(static_cast<long double(*)(long double)>(&asinh));
-UE_AUTORTFM_REGISTER_SELF_FUNCTION(static_cast<float(*)(float)>(&acosh));
-UE_AUTORTFM_REGISTER_SELF_FUNCTION(static_cast<long double(*)(long double)>(&acosh));
-UE_AUTORTFM_REGISTER_SELF_FUNCTION(static_cast<float(*)(float)>(&atanh));
-UE_AUTORTFM_REGISTER_SELF_FUNCTION(static_cast<long double(*)(long double)>(&atanh));
-UE_AUTORTFM_REGISTER_SELF_FUNCTION(static_cast<float(*)(float)>(&exp));
-UE_AUTORTFM_REGISTER_SELF_FUNCTION(static_cast<long double(*)(long double)>(&exp));
-UE_AUTORTFM_REGISTER_SELF_FUNCTION(static_cast<float(*)(float)>(&log));
-UE_AUTORTFM_REGISTER_SELF_FUNCTION(static_cast<long double(*)(long double)>(&log));
-UE_AUTORTFM_REGISTER_SELF_FUNCTION(static_cast<float(*)(float, float)>(&pow));
-UE_AUTORTFM_REGISTER_SELF_FUNCTION(static_cast<long double(*)(long double, long double)>(&pow));
-UE_AUTORTFM_REGISTER_SELF_FUNCTION(static_cast<long long(*)(float)>(&llrint));
-UE_AUTORTFM_REGISTER_SELF_FUNCTION(static_cast<long long(*)(long double)>(&llrint));
-UE_AUTORTFM_REGISTER_SELF_FUNCTION(static_cast<float(*)(float, float)>(&fmod));
-UE_AUTORTFM_REGISTER_SELF_FUNCTION(static_cast<long double(*)(long double, long double)>(&fmod));
-UE_AUTORTFM_REGISTER_SELF_FUNCTION(static_cast<float(*)(float, float*)>(&modf));
-UE_AUTORTFM_REGISTER_SELF_FUNCTION(static_cast<long double(*)(long double, long double*)>(&modf));
+UE_AUTORTFM_REGISTER_SELF_OVERLOADED_FUNCTION(float (float), sqrt);
+UE_AUTORTFM_REGISTER_SELF_OVERLOADED_FUNCTION(long double (long double), sqrt);
+UE_AUTORTFM_REGISTER_SELF_OVERLOADED_FUNCTION(float (float), sin);
+UE_AUTORTFM_REGISTER_SELF_OVERLOADED_FUNCTION(long double (long double), sin);
+UE_AUTORTFM_REGISTER_SELF_OVERLOADED_FUNCTION(float (float), cos);
+UE_AUTORTFM_REGISTER_SELF_OVERLOADED_FUNCTION(long double (long double), cos);
+UE_AUTORTFM_REGISTER_SELF_OVERLOADED_FUNCTION(float (float), tan);
+UE_AUTORTFM_REGISTER_SELF_OVERLOADED_FUNCTION(long double (long double), tan);
+UE_AUTORTFM_REGISTER_SELF_OVERLOADED_FUNCTION(float (float), asin);
+UE_AUTORTFM_REGISTER_SELF_OVERLOADED_FUNCTION(long double (long double), asin);
+UE_AUTORTFM_REGISTER_SELF_OVERLOADED_FUNCTION(float (float), acos);
+UE_AUTORTFM_REGISTER_SELF_OVERLOADED_FUNCTION(long double (long double), acos);
+UE_AUTORTFM_REGISTER_SELF_OVERLOADED_FUNCTION(float (float), atan);
+UE_AUTORTFM_REGISTER_SELF_OVERLOADED_FUNCTION(long double (long double), atan);
+UE_AUTORTFM_REGISTER_SELF_OVERLOADED_FUNCTION(float (float, float), atan2);
+UE_AUTORTFM_REGISTER_SELF_OVERLOADED_FUNCTION(long double (long double, long double), atan2);
+UE_AUTORTFM_REGISTER_SELF_OVERLOADED_FUNCTION(float (float), sinh);
+UE_AUTORTFM_REGISTER_SELF_OVERLOADED_FUNCTION(long double (long double), sinh);
+UE_AUTORTFM_REGISTER_SELF_OVERLOADED_FUNCTION(float (float), cosh);
+UE_AUTORTFM_REGISTER_SELF_OVERLOADED_FUNCTION(long double (long double), cosh);
+UE_AUTORTFM_REGISTER_SELF_OVERLOADED_FUNCTION(float (float), tanh);
+UE_AUTORTFM_REGISTER_SELF_OVERLOADED_FUNCTION(long double (long double), tanh);
+UE_AUTORTFM_REGISTER_SELF_OVERLOADED_FUNCTION(float (float), asinh);
+UE_AUTORTFM_REGISTER_SELF_OVERLOADED_FUNCTION(long double (long double), asinh);
+UE_AUTORTFM_REGISTER_SELF_OVERLOADED_FUNCTION(float (float), acosh);
+UE_AUTORTFM_REGISTER_SELF_OVERLOADED_FUNCTION(long double (long double), acosh);
+UE_AUTORTFM_REGISTER_SELF_OVERLOADED_FUNCTION(float (float), atanh);
+UE_AUTORTFM_REGISTER_SELF_OVERLOADED_FUNCTION(long double (long double), atanh);
+UE_AUTORTFM_REGISTER_SELF_OVERLOADED_FUNCTION(float (float), exp);
+UE_AUTORTFM_REGISTER_SELF_OVERLOADED_FUNCTION(long double (long double), exp);
+UE_AUTORTFM_REGISTER_SELF_OVERLOADED_FUNCTION(float (float), log);
+UE_AUTORTFM_REGISTER_SELF_OVERLOADED_FUNCTION(long double (long double), log);
+UE_AUTORTFM_REGISTER_SELF_OVERLOADED_FUNCTION(float (float, float), pow);
+UE_AUTORTFM_REGISTER_SELF_OVERLOADED_FUNCTION(long double (long double, long double), pow);
+UE_AUTORTFM_REGISTER_SELF_OVERLOADED_FUNCTION(long long (float), llrint);
+UE_AUTORTFM_REGISTER_SELF_OVERLOADED_FUNCTION(long long (long double), llrint);
+UE_AUTORTFM_REGISTER_SELF_OVERLOADED_FUNCTION(float (float, float), fmod);
+UE_AUTORTFM_REGISTER_SELF_OVERLOADED_FUNCTION(long double (long double, long double), fmod);
+UE_AUTORTFM_REGISTER_SELF_OVERLOADED_FUNCTION(float (float, float*), modf);
+UE_AUTORTFM_REGISTER_SELF_OVERLOADED_FUNCTION(long double (long double, long double*), modf);
 #endif
 
 // Self register Math functions
@@ -323,7 +338,7 @@ UE_AUTORTFM_REGISTER_SELF_FUNCTION(modff);
 UE_AUTORTFM_REGISTER_SELF_FUNCTION(modfl);
 
 // FIXME: Does not currently support %n format specifiers.
-int RTFM_vsnprintf(char* Str, size_t Size, const char* Format, va_list ArgList)
+int RTFM_vsnprintf(char* Str, size_t Size, const char* Format, va_list ArgList) throw()
 {
     ThrowIfFormatContainsPercentN(Format);
 
@@ -373,9 +388,9 @@ int RTFM_vswprintf(wchar_t* Str, size_t Size, const wchar_t* Format, va_list Arg
 
     return vswprintf(Str, Size, Format, ArgList);
 }
-UE_AUTORTFM_REGISTER_OPEN_FUNCTION_EXPLICIT(
-    static_cast<int(*)(wchar_t*, size_t, const wchar_t*, va_list)>(&vswprintf),
-    RTFM_vswprintf);
+UE_AUTORTFM_REGISTER_OPEN_OVERLOADED_FUNCTION_EXPLICIT(
+    int (wchar_t*, size_t, const wchar_t*, va_list),
+    vswprintf, RTFM_vswprintf);
 
 // FIXME: Does not currently support %n format specifiers.
 int RTFM_swprintf(wchar_t* Buffer, size_t BufferCount, wchar_t const* Format, ...)
@@ -388,12 +403,12 @@ int RTFM_swprintf(wchar_t* Buffer, size_t BufferCount, wchar_t const* Format, ..
 
     return Count;
 }
-UE_AUTORTFM_REGISTER_OPEN_FUNCTION_EXPLICIT(
-    static_cast<int(*)(wchar_t*, size_t, wchar_t const*, ...)>(&swprintf),
-    RTFM_swprintf);
+UE_AUTORTFM_REGISTER_OPEN_OVERLOADED_FUNCTION_EXPLICIT(
+    int(wchar_t*, size_t, wchar_t const*, ...),
+    swprintf, RTFM_swprintf);
 
 // FIXME: Does not currently support %n format specifiers.
-int RTFM_snprintf(char* Str, size_t Size, const char* Format, ...)
+int RTFM_snprintf(char* Str, size_t Size, const char* Format, ...) throw()
 {
     va_list ArgList;
 
@@ -445,7 +460,7 @@ UE_AUTORTFM_REGISTER_OPEN_FUNCTION(puts);
 
 #if PLATFORM_WINDOWS
 
-FILE* RTFM___acrt_iob_func(int Index)
+FILE* RTFM___acrt_iob_func(unsigned Index)
 {
     switch (Index)
     {
@@ -554,8 +569,8 @@ extern "C" size_t _ZNSt3__112__next_primeEm(size_t N) __attribute__((weak));
 UE_AUTORTFM_REGISTER_SELF_FUNCTION(_ZNSt3__112__next_primeEm);
 #endif // PLATFORM_WINDOWS -> so end of !PLATFORM_WINDOWS
 
-UE_AUTORTFM_REGISTER_SELF_FUNCTION(static_cast<float(*)(float, float)>(&powf));
-UE_AUTORTFM_REGISTER_SELF_FUNCTION(static_cast<double(*)(double, double)>(&pow));
+UE_AUTORTFM_REGISTER_SELF_OVERLOADED_FUNCTION(float (float, float), powf);
+UE_AUTORTFM_REGISTER_SELF_OVERLOADED_FUNCTION(double (double, double), pow);
 
 #if PLATFORM_WINDOWS
 UE_AUTORTFM_REGISTER_SELF_FUNCTION(_tcsncmp);
@@ -599,7 +614,7 @@ UE_AUTORTFM_REGISTER_SELF_FUNCTION(bcmp);
 UE_AUTORTFM_REGISTER_SELF_FUNCTION(pthread_getspecific);
 #endif // PLATFORM_LINUX
 
-wchar_t* RTFM_wcscpy(wchar_t* Dst, const wchar_t* Src)
+wchar_t* RTFM_wcscpy(wchar_t* Dst, const wchar_t* Src) throw()
 {
     const size_t SrcLen = wcslen(Src);
 
@@ -608,7 +623,7 @@ wchar_t* RTFM_wcscpy(wchar_t* Dst, const wchar_t* Src)
     return wcscpy(Dst, Src);
 }
 
-wchar_t* RTFM_wcsncpy(wchar_t* Dst, const wchar_t* Src, size_t Count)
+wchar_t* RTFM_wcsncpy(wchar_t* Dst, const wchar_t* Src, size_t Count) throw()
 {
 	FContext* Context = FContext::Get();
 	Context->RecordWrite(Dst, Count * sizeof(wchar_t));
@@ -631,7 +646,7 @@ UE_AUTORTFM_REGISTER_OPEN_FUNCTION(wcsncpy);
 #pragma warning(pop)
 #endif
 
-int RTFM_atexit(void(__cdecl*Callback)(void))
+int RTFM_atexit(void(*Callback)(void)) throw()
 {
 	FContext* Context = FContext::Get();
 	Context->GetCurrentTransaction()->DeferUntilCommit([Callback]
diff --git a/Engine/Source/Runtime/Core/Public/AutoRTFM/AutoRTFM.h b/Engine/Source/Runtime/Core/Public/AutoRTFM/AutoRTFM.h
--- a/Engine/Source/Runtime/Core/Public/AutoRTFM/AutoRTFM.h
+++ b/Engine/Source/Runtime/Core/Public/AutoRTFM/AutoRTFM.h
@@ -930,7 +930,7 @@ namespace ForTheRuntime
 
 	// Register a transactional version of a function that wasn't compiled by the
 	// autortfm compiler. Normally, code is transactionalized by the compiler by
-	// emitting a clone that has transactional openation, with some magic to
+	// emitting a clone that has transactional operation, with some magic to
 	// redirect all function calls within a transaction to the transactional clone.
 	// This allows you to hook in your own transactionalized implementations of
 	// functions that the compiler did not see.
@@ -938,18 +938,20 @@ namespace ForTheRuntime
 	// Use with great caution!
 	//
 	// This results in calls to ClosedVariant to happen in open mode. We will call
-	// ClosedVariant's nontransactional version within the transaction. This happens
+	// ClosedVariant's non-transactional version within the transaction. This happens
 	// with the additional caveat that the function signatures must match.
-	static UE_AUTORTFM_FORCEINLINE void RegisterOpenFunction(void* const OpenFunction, void* const ClosedVariant)
+	template<typename FN_TYPE>
+	static UE_AUTORTFM_FORCEINLINE void RegisterOpenFunction(FN_TYPE* OPEN, FN_TYPE* CLOSED)
 	{
-		autortfm_register_open_function(OpenFunction, ClosedVariant);
+		autortfm_register_open_function(reinterpret_cast<void*>(OPEN), reinterpret_cast<void*>(CLOSED));
 	}
 
+	template<typename FN_TYPE, FN_TYPE* OPEN, FN_TYPE* CLOSED>
 	struct FRegisterOpenFunction
 	{
-		FRegisterOpenFunction(void* const OriginalFunction, void* const NewFunction)
+		FRegisterOpenFunction()
 		{
-			RegisterOpenFunction(OriginalFunction, NewFunction);
+			RegisterOpenFunction<FN_TYPE>(OPEN, CLOSED);
 		}
 	};
 
@@ -1096,22 +1098,35 @@ namespace AutoRTFM::Private
 #define UE_AUTORTFM_CONCAT(A, B) UE_AUTORTFM_CONCAT_IMPL(A, B)
 
 #if UE_AUTORTFM
-#define UE_AUTORTFM_REGISTER_OPEN_FUNCTION_EXPLICIT_IMPL(OriginalFunction, NewFunction) static const AutoRTFM::ForTheRuntime::FRegisterOpenFunction UE_AUTORTFM_CONCAT(AutoRTFMFunctionRegistration, __COUNTER__)(reinterpret_cast<void*>(OriginalFunction), reinterpret_cast<void*>(NewFunction))
-#define UE_AUTORTFM_REGISTER_OPEN_FUNCTION_IMPL(OriginalFunction) UE_AUTORTFM_REGISTER_OPEN_FUNCTION_EXPLICIT(OriginalFunction, RTFM_ ## OriginalFunction)
-#define UE_AUTORTFM_REGISTER_SELF_FUNCTION_IMPL(OriginalFunction) UE_AUTORTFM_REGISTER_OPEN_FUNCTION_EXPLICIT(OriginalFunction, OriginalFunction)
+#define UE_AUTORTFM_REGISTER_OPEN_FUNCTION_EXPLICIT_IMPL(Signature, OriginalFunction, NewFunction) \
+	static const AutoRTFM::ForTheRuntime::FRegisterOpenFunction<Signature, OriginalFunction, NewFunction> UE_AUTORTFM_CONCAT(AutoRTFMFunctionRegistration, __COUNTER__)
+#define UE_AUTORTFM_REGISTER_OPEN_FUNCTION_IMPL(Signature, OriginalFunction) UE_AUTORTFM_REGISTER_OPEN_FUNCTION_EXPLICIT_IMPL(Signature, OriginalFunction, RTFM_ ## OriginalFunction)
+#define UE_AUTORTFM_REGISTER_SELF_FUNCTION_IMPL(Signature, OriginalFunction) UE_AUTORTFM_REGISTER_OPEN_FUNCTION_EXPLICIT_IMPL(Signature, OriginalFunction, OriginalFunction)
 #else
-#define UE_AUTORTFM_REGISTER_OPEN_FUNCTION_EXPLICIT_IMPL(OriginalFunction, NewFunction)
-#define UE_AUTORTFM_REGISTER_OPEN_FUNCTION_IMPL(OriginalFunction)
-#define UE_AUTORTFM_REGISTER_SELF_FUNCTION_IMPL(OriginalFunction)
+#define UE_AUTORTFM_REGISTER_OPEN_FUNCTION_EXPLICIT_IMPL(Signature, OriginalFunction, NewFunction)
+#define UE_AUTORTFM_REGISTER_OPEN_FUNCTION_IMPL(Signature, OriginalFunction)
+#define UE_AUTORTFM_REGISTER_SELF_FUNCTION_IMPL(Signature, OriginalFunction)
 #endif
 
 // Register that a specific OpenFunction maps to a closed variant when called in closed code.
-#define UE_AUTORTFM_REGISTER_OPEN_FUNCTION_EXPLICIT(OpenFunction, ClosedVariant) UE_AUTORTFM_REGISTER_OPEN_FUNCTION_EXPLICIT_IMPL(OpenFunction, ClosedVariant)
+#define UE_AUTORTFM_REGISTER_OPEN_FUNCTION_EXPLICIT(OpenFunction, ClosedVariant) UE_AUTORTFM_REGISTER_OPEN_FUNCTION_EXPLICIT_IMPL(decltype(OpenFunction), OpenFunction, ClosedVariant)
+
+// Similar to UE_AUTORTFM_REGISTER_OPEN_FUNCTION_EXPLICIT, but takes an
+// additional Signature parameter used to distinguish the overload that is to
+// be registered.
+// Signature is in the form: ReturnType(ArgType1, ArgType2, ...)
+#define UE_AUTORTFM_REGISTER_OPEN_OVERLOADED_FUNCTION_EXPLICIT(Signature, OpenFunction, ClosedVariant) UE_AUTORTFM_REGISTER_OPEN_FUNCTION_EXPLICIT_IMPL(Signature, OpenFunction, ClosedVariant)
 
 // Tells the runtime that OpenFunction maps to an explicit closed variant with the "RTFM_" prefix.
-#define UE_AUTORTFM_REGISTER_OPEN_FUNCTION(OpenFunction) UE_AUTORTFM_REGISTER_OPEN_FUNCTION_IMPL(OpenFunction)
+#define UE_AUTORTFM_REGISTER_OPEN_FUNCTION(OpenFunction) UE_AUTORTFM_REGISTER_OPEN_FUNCTION_IMPL(decltype(OpenFunction), OpenFunction)
 
 // Tells the runtime that OpenFunction maps to itself in closed code (EG. it has no transactional semantics).
-#define UE_AUTORTFM_REGISTER_SELF_FUNCTION(OpenFunction) UE_AUTORTFM_REGISTER_SELF_FUNCTION_IMPL(OpenFunction)
+#define UE_AUTORTFM_REGISTER_SELF_FUNCTION(OpenFunction) UE_AUTORTFM_REGISTER_SELF_FUNCTION_IMPL(decltype(OpenFunction), OpenFunction)
+
+// Similar to UE_AUTORTFM_REGISTER_SELF_FUNCTION, but takes an additional
+// Signature parameter used to distinguish the overload that is to be
+// registered.
+// Signature is in the form: ReturnType(ArgType1, ArgType2, ...)
+#define UE_AUTORTFM_REGISTER_SELF_OVERLOADED_FUNCTION(Signature, OpenFunction) UE_AUTORTFM_REGISTER_SELF_FUNCTION_IMPL(Signature, OpenFunction)
 
 #endif // __cplusplus
