Sequencer MediaTrack requests parsing timecode meta data and fills in the media track section's source time code.

Also, Timecode is displayed in the media info tab in MediaSource editor.

diff --git a/Engine/Plugins/Media/MediaCompositing/Source/MediaCompositing/Private/MovieScene/MovieSceneMediaTrack.cpp b/Engine/Plugins/Media/MediaCompositing/Source/MediaCompositing/Private/MovieScene/MovieSceneMediaTrack.cpp
--- a/Engine/Plugins/Media/MediaCompositing/Source/MediaCompositing/Private/MovieScene/MovieSceneMediaTrack.cpp
+++ b/Engine/Plugins/Media/MediaCompositing/Source/MediaCompositing/Private/MovieScene/MovieSceneMediaTrack.cpp
@@ -6,6 +6,7 @@
 #include "IMediaClockSink.h"
 #include "IMediaModule.h"
 #include "MediaPlayer.h"
+#include "MediaPlaylist.h"
 #include "MediaSource.h"
 #include "Modules/ModuleManager.h"
 #include "MovieScene.h"
@@ -369,7 +370,11 @@ void UMovieSceneMediaTrack::StartGetDuration(UMediaSource* MediaSource, UMovieSc
 
 	// Open the media.
 	MediaPlayer->PlayOnOpen = false;
-	if (MediaPlayer->OpenSource(MediaSource))
+	FMediaPlayerOptions PlayerOptions;
+	PlayerOptions.SetAllAsOptional();
+	PlayerOptions.InternalCustomOptions.Emplace(MediaPlayerOptionValues::ParseTimecodeInfo(), FVariant());
+	
+	if (MediaPlayer->OpenSourceWithOptions(MediaSource, PlayerOptions))
 	{
 		NewSections.Emplace(MediaPlayer, Section);
 	}
@@ -389,29 +394,52 @@ void UMovieSceneMediaTrack::StartGetDuration(UMediaSource* MediaSource, UMovieSc
 	}
 }
 
+namespace UE::MovieSceneMediaTrack::Private
+{
+	FString GetTimecodeStringFromMediaPlayer(const UMediaPlayer* InMediaPlayer)
+	{
+		if (UMediaPlaylist* Playlist = InMediaPlayer->GetPlaylist())
+		{
+			if (const UMediaSource* MediaSource = Playlist->Get(InMediaPlayer->GetPlaylistIndex()))
+			{
+				return MediaSource->GetMediaOption(UMediaPlayer::MediaInfoNameStartTimecodeValue.Resolve(), FString());
+			}
+		}
+		return FString();
+	}
+}
+
 bool UMovieSceneMediaTrack::GetDuration(
 	const TStrongObjectPtr<UMediaPlayer>& MediaPlayer, TWeakObjectPtr<UMovieSceneSection>& NewSection)
 {
-	bool bIsDone = false;
-
 	// Check everything is ok.
-	if ((MediaPlayer.IsValid() == false) || (MediaPlayer->HasError()) || (MediaPlayer->IsClosed()) ||
-		(NewSection.IsValid() == false))
-	{
-		bIsDone = true;
-	}
-	else
+	bool bIsDone = (MediaPlayer.IsValid() == false) || (MediaPlayer->HasError()) || (MediaPlayer->IsClosed()) || (NewSection.IsValid() == false);
+
+	if (!bIsDone)
 	{
 		// Get the duration.
-		FTimespan Duration = MediaPlayer->GetDuration();
+		const FTimespan Duration = MediaPlayer->GetDuration();
 		if (Duration != 0)
 		{
 			// Once it is non zero, then set the length of the section.
-			FFrameRate TickResolution = NewSection->GetTypedOuter<UMovieScene>()->GetTickResolution();
-			FFrameNumber StartFrame = NewSection->GetInclusiveStartFrame();
-			FFrameNumber EndFrame = StartFrame + (Duration.GetTotalSeconds() * TickResolution).FrameNumber;
+			const FFrameRate TickResolution = NewSection->GetTypedOuter<UMovieScene>()->GetTickResolution();
+			const FFrameNumber StartFrame = NewSection->GetInclusiveStartFrame();
+			const FFrameNumber EndFrame = StartFrame + (Duration.GetTotalSeconds() * TickResolution).FrameNumber;
 			NewSection->SetEndFrame(TRangeBound<FFrameNumber>::Exclusive(EndFrame));
 			bIsDone = true;
+
+			// Timecode from the source. Either player's meta data or media source.
+			using namespace UE::MovieSceneMediaTrack::Private;
+			const FVariant Value = MediaPlayer->GetMediaInfo(UMediaPlayer::MediaInfoNameStartTimecodeValue.Resolve());
+			const FString TimecodeString = Value.IsEmpty() ? GetTimecodeStringFromMediaPlayer(MediaPlayer.Get()) : Value.GetValue<FString>();  
+			if (!TimecodeString.IsEmpty())
+			{
+				const TOptional<FTimecode> Timecode = FTimecode::ParseTimecode(*TimecodeString);
+				if (Timecode.IsSet())
+				{
+					NewSection->TimecodeSource = FMovieSceneTimecodeSource(Timecode.GetValue());
+				}
+			}
 		}
 	}
 
diff --git a/Engine/Plugins/Media/MediaPlayerEditor/Source/MediaPlayerEditor/Private/Toolkits/MediaSourceEditorToolkit.cpp b/Engine/Plugins/Media/MediaPlayerEditor/Source/MediaPlayerEditor/Private/Toolkits/MediaSourceEditorToolkit.cpp
--- a/Engine/Plugins/Media/MediaPlayerEditor/Source/MediaPlayerEditor/Private/Toolkits/MediaSourceEditorToolkit.cpp
+++ b/Engine/Plugins/Media/MediaPlayerEditor/Source/MediaPlayerEditor/Private/Toolkits/MediaSourceEditorToolkit.cpp
@@ -294,6 +294,7 @@ void FMediaSourceEditorToolkit::BindCommands()
 			FMediaPlayerOptions Options;
 			Options.SetAllAsOptional();
 			Options.InternalCustomOptions.Emplace(MediaPlayerOptionValues::Environment(), MediaPlayerOptionValues::Environment_Preview());
+			Options.InternalCustomOptions.Emplace(MediaPlayerOptionValues::ParseTimecodeInfo(), FVariant());
 			MediaPlayer->OpenSourceWithOptions(MediaSource, Options);
 		}),
 		FCanExecuteAction::CreateLambda([this] { return true; })
diff --git a/Engine/Plugins/Media/MediaPlayerEditor/Source/MediaPlayerEditor/Private/Widgets/SMediaPlayerEditorMediaDetails.cpp b/Engine/Plugins/Media/MediaPlayerEditor/Source/MediaPlayerEditor/Private/Widgets/SMediaPlayerEditorMediaDetails.cpp
--- a/Engine/Plugins/Media/MediaPlayerEditor/Source/MediaPlayerEditor/Private/Widgets/SMediaPlayerEditorMediaDetails.cpp
+++ b/Engine/Plugins/Media/MediaPlayerEditor/Source/MediaPlayerEditor/Private/Widgets/SMediaPlayerEditorMediaDetails.cpp
@@ -111,6 +111,14 @@ void SMediaPlayerEditorMediaDetails::Construct(const FArguments& InArgs,
 								[
 									SAssignNew(NumTilesText, STextBlock)
 								]
+							// Start Timecode.
+							+ SVerticalBox::Slot()
+								.AutoHeight()
+								.VAlign(VAlign_Center)
+								.Padding(4.0f)
+								[
+									SAssignNew(StartTimecodeText, STextBlock)
+								]
 						]
 				]
 
@@ -139,6 +147,7 @@ void SMediaPlayerEditorMediaDetails::UpdateDetails()
 	int64 ResourceSize = 0;
 	int32 SurfaceWidth = 0;
 	int32 SurfaceHeight = 0;
+	FString StartTimecode;
 
 	// Get player info.
 	if (MediaPlayer != nullptr)
@@ -148,6 +157,7 @@ void SMediaPlayerEditorMediaDetails::UpdateDetails()
 		FIntPoint NumTiles(EForceInit::ForceInitToZero);
 		MediaPlayer->GetMediaInfo<FIntPoint>(NumTiles, UMediaPlayer::MediaInfoNameSourceNumTiles.Resolve());
 		NumTotalTiles = NumTiles.X * NumTiles.Y;
+		MediaPlayer->GetMediaInfo<FString>(StartTimecode, UMediaPlayer::MediaInfoNameStartTimecodeValue.Resolve());
 	}
 
 	// Get texture info,
@@ -172,14 +182,20 @@ void SMediaPlayerEditorMediaDetails::UpdateDetails()
 	LODBiasText->SetText(FText::Format(LOCTEXT("LODBias", "Combined LOD Bias: {0}"),
 		FText::AsNumber(LODBias)));
 	MethodText->SetText(FText::Format(LOCTEXT("Method", "Method: {0}"), Method));
-	NumMipsText->SetText(FText::Format(LOCTEXT("NumberOfMips", "Number Of Mips: {0}"),
+	NumMipsText->SetText(FText::Format(LOCTEXT("NumberOfMips", "Mips Qty: {0}"),
 		FText::AsNumber(NumMips)));
-	NumTilesText->SetText(FText::Format(LOCTEXT("NumberOfTiles", "Number Of Tiles: {0}"),
+	NumTilesText->SetText(FText::Format(LOCTEXT("NumberOfTiles", "Tiles Qty: {0}"),
 		FText::AsNumber(NumTotalTiles)));
 	ResolutionText->SetText(FText::Format(LOCTEXT("Resolution", "Resolution: {0}x{1}"),
 		FText::AsNumber(SurfaceWidth), FText::AsNumber(SurfaceHeight)));
 	ResourceSizeText->SetText(FText::Format(LOCTEXT("ResourceSize", "Resource Size: {0} KB"),
 		FText::AsNumber(ResourceSize)));
+	
+	if (!StartTimecode.IsEmpty())
+	{
+		StartTimecodeText->SetText(FText::Format(LOCTEXT("StartTimecode", "Start Timecode: {0}"), FText::FromString(StartTimecode)));
+	}
+	StartTimecodeText->SetVisibility(!StartTimecode.IsEmpty() ? EVisibility::Visible : EVisibility::Collapsed);
 }
 
 #undef LOCTEXT_NAMESPACE
diff --git a/Engine/Plugins/Media/MediaPlayerEditor/Source/MediaPlayerEditor/Public/Widgets/SMediaPlayerEditorMediaDetails.h b/Engine/Plugins/Media/MediaPlayerEditor/Source/MediaPlayerEditor/Public/Widgets/SMediaPlayerEditorMediaDetails.h
--- a/Engine/Plugins/Media/MediaPlayerEditor/Source/MediaPlayerEditor/Public/Widgets/SMediaPlayerEditorMediaDetails.h
+++ b/Engine/Plugins/Media/MediaPlayerEditor/Source/MediaPlayerEditor/Public/Widgets/SMediaPlayerEditorMediaDetails.h
@@ -57,4 +57,5 @@ private:
 	TSharedPtr<STextBlock> NumTilesText;
 	TSharedPtr<STextBlock> ResolutionText;
 	TSharedPtr<STextBlock> ResourceSizeText;
+	TSharedPtr<STextBlock> StartTimecodeText;
 };
diff --git a/Engine/Source/Runtime/Core/Private/Misc/Timecode.cpp b/Engine/Source/Runtime/Core/Private/Misc/Timecode.cpp
--- a/Engine/Source/Runtime/Core/Private/Misc/Timecode.cpp
+++ b/Engine/Source/Runtime/Core/Private/Misc/Timecode.cpp
@@ -16,3 +16,195 @@ bool FTimecode::UseDropFormatTimecodeByDefaultWhenSupported()
 {
 	return CVarUseDropFormatTimecodeByDefaultWhenSupported.GetValueOnAnyThread();
 }
+
+namespace UE::Core::Private
+{
+	/**
+	 * Parses a string representation of a timecode.
+	 *
+	 * Supported SMPTE drop frame (DF) and non-drop frame (NDF) formats variations:
+	 * - NDF: HH:MM:SS:FF
+	 * - DF: HH:MM:SS;FF or HH:MM:SS.FF or HH;MM;SS;FF or HH.MM.SS.FF.
+	 *
+	 * The 2 digits per number is not enforced. It is possible to parse high frame numbers (above 60), such as for audio timecodes.
+	 * Full SMPTE compliance is not ensured by the parser (i.e. greater than 24h, negative time and any number of frames per second).
+	 *
+	 * Sub-frame variation:
+	 * Supports the sub-frame variation where the frame number is a decimal number: HH:MM:SS:FF.ZZ.
+	 *
+	 * Side effect of supporting sub-frame is that this parser can't unambiguously parse partial timecodes.
+	 * For instance, "00:00.00" is ambiguous because it is undetermined if the last value is frame or sub-frame.
+	 */
+	class FTimecodeParser
+	{
+	public:		
+		// Expected frame separators
+		const TArray<TCHAR, TInlineAllocator<3>> AllSeparators;
+		const TArray<TCHAR, TInlineAllocator<3>> SubFrameSeparators;
+		const TArray<TCHAR, TInlineAllocator<3>> DropFrameSeparators;
+
+		// Set of separators for each positions.
+		TArray<const TArray<TCHAR, TInlineAllocator<3>>*, TInlineAllocator<4>> ValueSeparators;
+
+		FTimecodeParser()
+			: AllSeparators({ TEXT(':'), TEXT(';'), TEXT('.') })
+			, SubFrameSeparators({TEXT('.')})
+			, DropFrameSeparators({ TEXT(';'), TEXT('.') })
+		{
+			// Expected separators for each positions: HH[:;.]MM[:;,]SS[:;.]FF.DD
+			ValueSeparators.Add(&AllSeparators);
+			ValueSeparators.Add(&AllSeparators);
+			ValueSeparators.Add(&AllSeparators);
+			ValueSeparators.Add(&SubFrameSeparators);
+		}
+
+		TOptional<FTimecode> Evaluate(const FStringView InString) const
+		{
+			TArray<TCHAR, TInlineAllocator<64>> Buffer;
+			constexpr int32 MaxNumberOfValues = 5;
+			constexpr int32 MinNumberOfValues = 4;
+			TArray<double, TInlineAllocator<MaxNumberOfValues>> Values;
+			TArray<int32, TInlineAllocator<MaxNumberOfValues>> Signs;
+
+			// Separator positions: HH(0):MM(1):SS(2):FF(3).DD
+			constexpr int32 FrameSeparatorPosition = 2;
+			constexpr int32 SubFrameSeparatorPosition = 3;
+
+			bool bDropFrameFormat = false;
+			bool bDropFrameSeparatorEncountered = false;
+
+			auto TryPushValue = [&Buffer, &Values, &Signs]() -> bool
+			{
+				Buffer.Push(TEXT('\0'));
+
+				// Validation
+				if (!FCString::IsNumeric(Buffer.GetData()))
+				{
+					return false;
+				}
+
+				// Keep track of sign independently of values because Atoi can't represent -0;
+				Signs.Push(Buffer[0] == TEXT('-') ? -1 : 1);
+
+				if (Values.Num() >= MinNumberOfValues)
+				{
+					// Sub-frame fraction is parsed as a positive-only decimal fraction.
+					// It will be made negative only if there is no other way to preserve the sign (i.e all other values are zero).
+					Buffer.RemoveAll([](TCHAR InChar) { return InChar == TEXT('-') || InChar == TEXT('+'); });
+					Buffer.Insert(TEXT('.'), 0);
+					Values.Push(FCString::Atof(Buffer.GetData()));
+				}
+				else
+				{
+					// hours, minutes, seconds and frames are integers
+					Values.Push(FCString::Atoi(Buffer.GetData()));
+				}
+
+				Buffer.Reset();
+				return true;
+			};
+
+			// Parse the string and convert values.
+			for (const TCHAR Char : InString)
+			{
+				if (AllSeparators.Contains(Char))
+				{
+					const uint32 SeparatorPosition = Values.Num();
+					
+					// Validate expected separator for this value position.
+					if (!ValueSeparators.IsValidIndex(SeparatorPosition) || !ValueSeparators[SeparatorPosition]->Contains(Char))
+					{
+						// Unexpected separator.
+						return {};
+					}
+
+					if (Buffer.IsEmpty() || !TryPushValue())
+					{
+						// value was not a number (or empty)
+						return {};
+					}
+
+					// The DF separator can be between all digit pairs or just the 3rd position.
+					// However, if it is present but not at the 3rd position, this is ambiguous.
+					if (DropFrameSeparators.Contains(Char) && SeparatorPosition < SubFrameSeparatorPosition)
+					{
+						bDropFrameSeparatorEncountered = true;	// This will detect ambiguous cases.
+					
+						// Encountering ';' or '.' on the 3rd separator is a confirmed DF format.
+						if (SeparatorPosition == FrameSeparatorPosition)
+						{
+							bDropFrameFormat = true;
+						}
+					}
+				}
+				else if (!FChar::IsWhitespace(Char)) // Ignore whitespaces
+				{
+					Buffer.Push(Char);
+				}
+			}
+
+			// Reject ambiguous separator combinations (DF vs NDF).
+			if (bDropFrameSeparatorEncountered && !bDropFrameFormat)
+			{
+				return {};
+			}
+			
+			// Convert last value. Fail on empty or non-numeric string.
+			if (Buffer.IsEmpty() || !TryPushValue())
+			{
+				return {};
+			}
+
+			// Validate that we have the expected number of parsed values.
+			if (Values.Num() < MinNumberOfValues || Values.Num() > MaxNumberOfValues)
+			{
+				return {};
+			}
+
+			// Ensure the sign is preserved in case of zero values.
+			if (Signs.Contains(-1))
+			{
+				// A negative sign has been encountered while parsing, see if we have a negative value already.
+				const bool bFoundNegative = Values.ContainsByPredicate([](double InValue) { return InValue < 0.0; });
+
+				// Here, we are doing our best to preserve the sign by applying it to the first non-zero value.
+				if (!bFoundNegative)
+				{
+					if (double* FirstNonZeroValue = Values.FindByPredicate([](double InValue) { return InValue > 0.0; }))
+					{
+						*FirstNonZeroValue = -*FirstNonZeroValue;
+					}
+				}
+			}
+
+			FTimecode TC;
+			TC.Hours = static_cast<int32>(Values[0]);
+			TC.Minutes = static_cast<int32>(Values[1]);
+			TC.Seconds = static_cast<int32>(Values[2]);
+			TC.Frames = static_cast<int32>(Values[3]);
+
+			if (Values.IsValidIndex(4))
+			{
+				// Note: For valid timecode math, it is necessary to allow a negative sub-frame fraction (if all other values are zero).
+				TC.Subframe = FMath::Clamp(static_cast<float>(Values[4]), -1.0f, 1.0f);
+			}
+
+			TC.bDropFrameFormat = bDropFrameFormat;
+
+			return TC;
+		}
+		
+		static FTimecodeParser& Get()
+		{
+			static FTimecodeParser StaticTimecodeParser;
+			return StaticTimecodeParser;
+		}
+	};
+}
+
+TOptional<FTimecode> FTimecode::ParseTimecode(const FStringView InTimecodeString)
+{
+	using namespace UE::Core::Private;
+	FTimecodeParser& StaticTimecodeParser = FTimecodeParser::Get();
+	return StaticTimecodeParser.Evaluate(InTimecodeString);
+}
\ No newline at end of file
diff --git a/Engine/Source/Runtime/Core/Public/Misc/Timecode.h b/Engine/Source/Runtime/Core/Public/Misc/Timecode.h
--- a/Engine/Source/Runtime/Core/Public/Misc/Timecode.h
+++ b/Engine/Source/Runtime/Core/Public/Misc/Timecode.h
@@ -383,6 +383,27 @@ public:
 		return Builder.ToString();
 	}
 
+	/**
+	 * Parses a string representation of a timecode.
+	 *
+	 * Supported SMPTE drop frame (DF) and non-drop frame (NDF) formats variations:
+	 * - NDF: HH:MM:SS:FF
+	 * - DF: HH:MM:SS;FF or HH:MM:SS.FF or HH;MM;SS;FF or HH.MM.SS.FF.
+	 *
+	 * The 2 digits per number is not enforced. It is possible to parse high frame numbers (above 60), such as for audio timecodes.
+	 * Full SMPTE compliance is not ensured by the parser (i.e. greater than 24h, negative time and any number of frames per second).
+	 *
+	 * Sub-frame variation:
+	 * Supports the sub-frame variation where the frame number is a decimal number: HH:MM:SS:FF.ZZ.
+	 *
+	 * Side effect of supporting sub-frame is that this function can't unambiguously parse partial timecodes.
+	 * 
+	 * @param InTimecodeString A string representation of a timecode
+	 * 
+	 * @return Parsed timecode if valid.
+	 */
+	static CORE_API TOptional<FTimecode> ParseTimecode(const FStringView InTimecodeString);
+
 	/**
 	 * Will return true if the timecode represents a valid timecode value where
 	 *
@@ -390,7 +411,7 @@ public:
 	 * Minutes is +/- [0, 59]
 	 * Seconds is +/- [0, 59]
 	 * Frames [INT_MIN, INT_MAX]
-	 * Subframes > 0
+	 * Subframes >= 0
 	 *
 	 */
 	bool IsValid() const
diff --git a/Engine/Source/Runtime/Core/Tests/Misc/TimecodeTest.cpp b/Engine/Source/Runtime/Core/Tests/Misc/TimecodeTest.cpp
--- a/Engine/Source/Runtime/Core/Tests/Misc/TimecodeTest.cpp
+++ b/Engine/Source/Runtime/Core/Tests/Misc/TimecodeTest.cpp
@@ -245,5 +245,175 @@ TEST_CASE_NAMED(FTimecodeTest, "System::Core::Misc::Timecode", "[ApplicationCont
 	REQUIRE(bSuccessfully);
 }
 
+namespace UE::TimecodeParserTest::Private
+{
+	/**
+	 * Bypass the FTimecode constructor to avoid the checkSlow for non conforming timecode tests
+	 * i.e. we want to test parser results even if not conforming to FTimecode constructor.
+	 */
+	static FTimecode MakeTimecodeNoCheck(int32 InHours, int32 InMinutes, int32 InSeconds, int32 InFrames, float InSubframe, bool bInDropFrame)
+	{
+		FTimecode Timecode;
+		Timecode.Hours = InHours;
+		Timecode.Minutes = InMinutes;
+		Timecode.Seconds = InSeconds;
+		Timecode.Frames = InFrames;
+		Timecode.Subframe = InSubframe;
+		Timecode.bDropFrameFormat = bInDropFrame;
+		return Timecode;
+	};
+
+	/**
+	 * Formats a timecode to string for testing.
+	 * Including the signed sub-frame and up to 6 decimal of precision on sub-frame fraction.
+	 */
+	static FString TimecodeToString(const FTimecode& InTimecode)
+	{
+		const bool bHasNegativeComponent = InTimecode.Hours < 0 || InTimecode.Minutes < 0 || InTimecode.Seconds < 0 || InTimecode.Frames < 0 || InTimecode.Subframe < 0.0;
+		const TCHAR* SignText = bHasNegativeComponent ? TEXT("- ") : TEXT("");
+
+		TStringBuilder<64> Builder;	// (Allow for up to 10 digits per values + sign and separator) x 5 values = 60 + terminating char.
+		Builder.Appendf(TEXT("%s%02d:%02d:%02d%c%02d"), SignText,
+			FMath::Abs(InTimecode.Hours), FMath::Abs(InTimecode.Minutes), FMath::Abs(InTimecode.Seconds), InTimecode.bDropFrameFormat ? TEXT(';') : TEXT(':'), FMath::Abs(InTimecode.Frames));
+
+		if (InTimecode.Subframe != 0.0f)
+		{
+			// Allow for up to 6 decimals of precision.
+			TStringBuilder<32> SubFrameBuilder;
+			SubFrameBuilder.Appendf(TEXT("%02d"), FMath::Clamp(static_cast<int32>(1000000 * FMath::Abs(InTimecode.Subframe)),0,999999));
+
+			// Trim trailing zeros, but leave at least 2 digits.
+			const FStringView SubFrameValue = SubFrameBuilder.ToView();
+			int NumDigits = SubFrameValue.Len();
+			while (NumDigits > 2 && SubFrameValue[NumDigits - 1] == TEXT('0'))
+			{
+				NumDigits--;
+			}
+			Builder.Appendf(TEXT(".%.*s"), NumDigits, SubFrameValue.GetData());
+		}
+		return Builder.ToString();
+	};
+
+	static bool IsSame(const FTimecode& InTimecode, const FTimecode& InOtherTimecode)
+	{
+		// Remark: FTimecode == operator ignores bDropFrameFormat.
+		return InTimecode == InOtherTimecode && InTimecode.bDropFrameFormat == InOtherTimecode.bDropFrameFormat;
+	}
+}
+
+TEST_CASE_NAMED(FTimecodeParserTest, "System::Core::Misc::TimecodeParser", "[ApplicationContextMask][EngineFilter]")
+{
+	using namespace UE::TimecodeParserTest::Private;
+	
+	struct FTimecodeParserTestEntry
+	{
+		FString TimecodeString;
+		FTimecode ExpectedTimecode;
+	};
+
+	constexpr bool bDropFrame = true;	// a.k.a DF
+	constexpr bool bStandard = false;	// a.k.a NDF
+
+	TArray<FTimecodeParserTestEntry> TimecodeParseSuccessTests =
+	{
+		{ TEXT("00:00:00:00"), FTimecode(0,0,0,0,bStandard)},
+		{ TEXT("00:00:00;00"), FTimecode(0,0,0,0,bDropFrame)},
+		{ TEXT("00:00:00:00.00"), FTimecode(0,0,0,0, 0.0f, bStandard)},
+		{ TEXT("00:00:00;00.00"), FTimecode(0,0,0,0, 0.0f, bDropFrame)},
+		{ TEXT("10:20:30:40.50"), FTimecode(10,20,30,40, 0.5f, bStandard)},
+		{ TEXT("15:55:22;09.90"), FTimecode(15,55,22,9, 0.9f, bDropFrame)},
+		{ TEXT("10: 11: 12; 191"), FTimecode(10,11,12,191, bDropFrame)},
+		// drop frame variations separator
+		{ TEXT("01;02;03;04"), FTimecode(1,2,3,4,bDropFrame)},
+		{ TEXT("01:02;03;04"), FTimecode(1,2,3,4,bDropFrame)},
+		{ TEXT("01;02;03;04.50"), FTimecode(1,2,3,4, 0.5f, bDropFrame)},
+		// drop frame variations with '.' separator
+		{ TEXT("01:02:03.04"), FTimecode(1,2,3,4,bDropFrame)},
+		{ TEXT("01.02.03.04"), FTimecode(1,2,3,4,bDropFrame)},
+		{ TEXT("01:02.03.04"), FTimecode(1,2,3,4,bDropFrame)},
+		{ TEXT("01.02.03.04.50"), FTimecode(1,2,3,4, 0.5f, bDropFrame)},
+		// higher precision sub-frame
+		{ TEXT("01:02:03:04.777"), FTimecode(1,2,3,4, 0.777f, bStandard)},
+		{ TEXT("01.02.03.04.555"), FTimecode(1,2,3,4, 0.555f, bDropFrame)},
+		// sign tests
+		{ TEXT("+ 15:55:22;09.90"), FTimecode(15,55,22,9, 0.9f, bDropFrame)},
+		{ TEXT("+15:55:22;09.90"), FTimecode(15,55,22,9, 0.9f, bDropFrame)},
+		{ TEXT("- 15:55:22;09.90"), FTimecode(-15,55,22,9, 0.9f, bDropFrame)},
+		{ TEXT("-15:55:22;09.90"), FTimecode(-15,55,22,9, 0.9f, bDropFrame)},
+		{ TEXT("15:-55:22;09.90"), FTimecode(15,-55,22,9, 0.9f, bDropFrame)},
+		{ TEXT("15:55:-22;09.90"), FTimecode(15,55,-22,9, 0.9f, bDropFrame)},
+		{ TEXT("15:55:22;-09.90"), FTimecode(15,55,22,-9, 0.9f, bDropFrame)},
+		// sign tests - sign gets applied to first non-zero value.
+		{ TEXT("- 00:55:22:09"), FTimecode(0,-55,22,9, bStandard)},
+		{ TEXT("- 00:00:22:09"), FTimecode(0,0,-22,9, bStandard)},
+		{ TEXT("- 00:00:00:09"), FTimecode(0,0,0,-9, bStandard)},
+		// sign tests - negative on the sub-frame. The sub-frame is negative only if there is no other way to preserve sign.
+		// -- Note: this would only happen if manually entered, but we still want the parser to do something with it.
+		{ TEXT("15:55:22;09.-90"), FTimecode(-15,55,22,9, 0.9f, bDropFrame)},
+		{ TEXT("00:55:22;09.-90"), FTimecode(0,-55,22,9, 0.9f, bDropFrame)},
+		{ TEXT("00:00:22;09.-90"), FTimecode(0,0,-22,9, 0.9f, bDropFrame)},
+		{ TEXT("00:00:00;09.-90"), FTimecode(0,0,0,-9, 0.9f, bDropFrame)},
+		{ TEXT("00:00:00;00.-90"), MakeTimecodeNoCheck(0,0,0,0, -0.9f, bDropFrame)},
+		{ TEXT("- 00:00:00:00.90"), MakeTimecodeNoCheck(0,0,0,0, -0.9f, bStandard)},
+		// High frame number (ex for audio timecodes) supported.
+		{ TEXT("15:55:22:43999"), FTimecode(15,55,22,43999, bStandard)},
+		// Values out of normal range
+		{ TEXT("200:210:220:44000"), MakeTimecodeNoCheck(200,210,220,44000, 0.0f, bStandard)}
+	};
+
+	TArray<FTimecodeParserTestEntry> TimecodeParseFailureTests =
+	{
+		{ TEXT(""), FTimecode()},
+		{ TEXT("00:00"), FTimecode()},			// not enough values
+		{ TEXT("01:02:1d:25"), FTimecode()},	// value is not a valid 'base10' number
+		{ TEXT("00.00:00:00"), FTimecode()},	// wrong separator
+		{ TEXT(":00:00:00:00"), FTimecode()},	// doesn't begin with a number
+		{ TEXT("...0"), FTimecode()},			// ""
+		{ TEXT("00;00:00:00"), FTimecode()},	// drop frame separator at the wrong place (ambiguous)
+		{ TEXT("00:00:00:00:00"), FTimecode()},	// wrong subframe separator
+	};
+
+	auto ExecuteTimecodeTests = [](const TArray<FTimecodeParserTestEntry>& InTests, bool bInExpectedSuccess)
+	{
+		for (const FTimecodeParserTestEntry& Test : InTests)
+		{
+			TOptional<FTimecode> TimecodeOptional = FTimecode::ParseTimecode(*Test.TimecodeString);
+			if (TimecodeOptional.IsSet())
+			{
+				const FTimecode& Timecode = TimecodeOptional.GetValue();
+				if (bInExpectedSuccess)
+				{
+					if (IsSame(Timecode, Test.ExpectedTimecode))
+					{
+						INFO(FString::Printf(TEXT("Parsing '%s' to parsed timecode '%s' -> OK"), *Test.TimecodeString, *TimecodeToString(Timecode)));
+					}
+					else
+					{
+						FAIL_CHECK(FString::Printf(TEXT("Parsing '%s' (parsed timecode '%s') was different than expected: '%s'"),
+							*Test.TimecodeString, *TimecodeToString(Timecode), *TimecodeToString(Test.ExpectedTimecode)));
+					}
+				}
+				else
+				{
+					FAIL_CHECK(FString::Printf(TEXT("Parsing '%s' was expected to failed, but it succeeded (parsed timecode '%s')"), *Test.TimecodeString, *TimecodeToString(Timecode)));
+				}
+			}
+			else
+			{
+				if (bInExpectedSuccess)
+				{
+					FAIL_CHECK(FString::Printf(TEXT("Parsing '%s' failed, but was expected to succeed (Expected: '%s'."), *Test.TimecodeString, *TimecodeToString(Test.ExpectedTimecode)));
+				}
+				else
+				{
+					INFO(FString::Printf(TEXT("Parsing '%s' failed as expected."), *Test.TimecodeString));
+				}
+			}
+		}
+	};
+
+	ExecuteTimecodeTests(TimecodeParseSuccessTests, true /*bInExpectedSuccess*/);
+	ExecuteTimecodeTests(TimecodeParseFailureTests, false /*bInExpectedSuccess*/);
+}
 
 #endif //WITH_TESTS
diff --git a/Engine/Source/Runtime/MediaAssets/Private/Assets/MediaPlayer.cpp b/Engine/Source/Runtime/MediaAssets/Private/Assets/MediaPlayer.cpp
--- a/Engine/Source/Runtime/MediaAssets/Private/Assets/MediaPlayer.cpp
+++ b/Engine/Source/Runtime/MediaAssets/Private/Assets/MediaPlayer.cpp
@@ -35,6 +35,8 @@
 
 FLazyName UMediaPlayer::MediaInfoNameSourceNumMips(TEXT("SourceNumMips"));
 FLazyName UMediaPlayer::MediaInfoNameSourceNumTiles(TEXT("SourceNumTiles"));
+FLazyName UMediaPlayer::MediaInfoNameStartTimecodeValue(TEXT("StartTimecodeValue"));
+FLazyName UMediaPlayer::MediaInfoNameStartTimecodeFrameRate(TEXT("StartTimecodeFrameRate"));
 
 /* UMediaPlayer structors
  *****************************************************************************/
diff --git a/Engine/Source/Runtime/MediaAssets/Public/MediaPlayer.h b/Engine/Source/Runtime/MediaAssets/Public/MediaPlayer.h
--- a/Engine/Source/Runtime/MediaAssets/Public/MediaPlayer.h
+++ b/Engine/Source/Runtime/MediaAssets/Public/MediaPlayer.h
@@ -1287,6 +1287,10 @@ public:
 	static MEDIAASSETS_API FLazyName MediaInfoNameSourceNumMips;
 	/** Number of tiles (X and Y) in the source. IntPoint.*/
 	static MEDIAASSETS_API FLazyName MediaInfoNameSourceNumTiles;
+	/** Start Timecode. SMPTE format (ex: HH:MM:SS:FF) string. */
+	static MEDIAASSETS_API FLazyName MediaInfoNameStartTimecodeValue;
+	/** Start Timecode frame rate. String. */
+	static MEDIAASSETS_API FLazyName MediaInfoNameStartTimecodeFrameRate;
 
 	/**
 	 * Duration of samples to cache ahead of the play head.
