Simplify thread-safe counter code a bit.

- Rather than three dedicated constructors that did the same work, we now have one complete and one delegating constructor.
- Reset() is implemented in terms of Set().
- An unnecessary cast was removed.

diff --git a/Engine/Source/Runtime/Core/Public/HAL/ThreadSafeCounter.h b/Engine/Source/Runtime/Core/Public/HAL/ThreadSafeCounter.h
--- a/Engine/Source/Runtime/Core/Public/HAL/ThreadSafeCounter.h
+++ b/Engine/Source/Runtime/Core/Public/HAL/ThreadSafeCounter.h
@@ -6,26 +6,26 @@
 #include "HAL/PlatformAtomics.h"
 
 ////////////////////////////////////////////////////////////////////////////////////////////////////////////////////////////
-// DEPRECATED. Please use `std::atomic<int32>`
+// DEPRECATED. Please use `std::atomic<int32>`.
 ////////////////////////////////////////////////////////////////////////////////////////////////////////////////////////////
 
 /** Thread safe counter */
 class FThreadSafeCounter
 {
 public:
-	typedef int32 IntegerType;
+	using IntegerType = int32;
 
 	/**
-	 * Default constructor.
+	 * Constructor, initializing counter to passed in value.
 	 *
-	 * Initializes the counter to 0.
+	 * @param Value	Value to initialize counter to
 	 */
-	FThreadSafeCounter()
+	FThreadSafeCounter(int32 Value = 0)
 	{
 #if USING_THREAD_SANITISER
-		FPlatformAtomics::AtomicStore(&Counter, 0);
+		FPlatformAtomics::AtomicStore(&Counter, Value);
 #else
-		Counter = 0;
+		Counter = Value;
 #endif
 	}
 
@@ -39,26 +39,8 @@ public:
 	 * @param Other The other thread safe counter to copy
 	 */
 	FThreadSafeCounter( const FThreadSafeCounter& Other )
+		: FThreadSafeCounter(Other.GetValue())
 	{
-#if USING_THREAD_SANITISER
-		FPlatformAtomics::AtomicStore(&Counter, Other.GetValue());
-#else
-		Counter = Other.GetValue();
-#endif
-	}
-
-	/**
-	 * Constructor, initializing counter to passed in value.
-	 *
-	 * @param Value	Value to initialize counter to
-	 */
-	FThreadSafeCounter( int32 Value )
-	{
-#if USING_THREAD_SANITISER
-		FPlatformAtomics::AtomicStore(&Counter, Value);
-#else
-		Counter = Value;
-#endif
 	}
 
 	/**
@@ -127,7 +109,7 @@ public:
 	 */
 	int32 Reset()
 	{
-		return FPlatformAtomics::InterlockedExchange(&Counter, 0);
+		return Set(0);
 	}
 
 	/**
@@ -137,13 +119,13 @@ public:
 	 */
 	int32 GetValue() const
 	{
-		return FPlatformAtomics::AtomicRead(&const_cast<FThreadSafeCounter*>(this)->Counter);
+		return FPlatformAtomics::AtomicRead(&Counter);
 	}
 
 private:
 
-	/** Hidden on purpose as usage wouldn't be thread safe. */
-	void operator=( const FThreadSafeCounter& Other ){}
+	/** Assignment is intentionally disallowed, as usage wouldn't be thread safe. */
+	void operator=(const FThreadSafeCounter& Other) = delete;
 
 	/** Thread-safe counter */
 	volatile int32 Counter;
