[LLM] Fixed extended Assets/AssetClasses tagset Tags report their existance to UnrealInsights, bloating the trace file and tanking the performance of analysis

#rnx
#tests local memory trace capture and analysis

diff --git a/Engine/Source/Runtime/Core/Private/ProfilingDebugging/TagTrace.cpp b/Engine/Source/Runtime/Core/Private/ProfilingDebugging/TagTrace.cpp
--- a/Engine/Source/Runtime/Core/Private/ProfilingDebugging/TagTrace.cpp
+++ b/Engine/Source/Runtime/Core/Private/ProfilingDebugging/TagTrace.cpp
@@ -164,6 +164,11 @@ public:
 private:
 
 	static constexpr int32 FNAME_INDEX_OFFSET = 512;
+
+	static int32 GetFNameTagID(FName UniqueName)
+	{
+		return UniqueName.GetDisplayIndex().ToUnstableInt() + FNAME_INDEX_OFFSET;
+	}
 	
 	struct FTagNameSetEntry
 	{
@@ -271,6 +276,15 @@ int32 FTagTrace::AnnounceLLMExtendedTag(FLowLevelMemTracker& Tracker, const UE::
 		}
 		// Not a generic tag, so this is a Platform or Project tag.
 	}
+	
+	FName UniqueName = Tracker.GetTagUniqueName(LLMTag);
+	int32 TagId = bEnumTag ? static_cast<int32>(EnumTag) : GetFNameTagID(UniqueName);
+
+	if (!LLMTag->IsReportable())
+	{
+		return TagId;
+	}
+
 	int32 ParentTagId = -1;
 	const UE::LLMPrivate::FTagData* ParentLLMTag = Tracker.GetTagParent(LLMTag);
 	if (ParentLLMTag)
@@ -285,8 +299,8 @@ int32 FTagTrace::AnnounceLLMExtendedTag(FLowLevelMemTracker& Tracker, const UE::
 	{
 		DisplayNamePtr = nullptr; // Use the UniqueName instead
 	}
-	int32 TagId = bEnumTag ? static_cast<int32>(EnumTag) : -1;
-	return AnnounceFNameTag(Tracker.GetTagUniqueName(LLMTag), TagId, DisplayNamePtr, ParentTagId);
+
+	return AnnounceFNameTag(UniqueName, TagId, DisplayNamePtr, ParentTagId);
 }
 
 #endif
@@ -312,7 +326,7 @@ int32 FTagTrace::AnnounceFNameTag(FName UniqueName, int32 TagId, const TCHAR* Di
 {
 	if (TagId == -1)
 	{
-		TagId = UniqueName.GetDisplayIndex().ToUnstableInt() + FNAME_INDEX_OFFSET;
+		TagId = GetFNameTagID(UniqueName);
 	}
 
 	// Find or add the item
