Updated FMallocBinned2 trim callback to call FCoreDelegates::GetMemoryTrimDelegate() when very large page allocator is used

diff --git a/Engine/Source/Runtime/Core/Private/HAL/MallocBinned2.cpp b/Engine/Source/Runtime/Core/Private/HAL/MallocBinned2.cpp
--- a/Engine/Source/Runtime/Core/Private/HAL/MallocBinned2.cpp
+++ b/Engine/Source/Runtime/Core/Private/HAL/MallocBinned2.cpp
@@ -467,6 +467,8 @@ void FMallocBinned2::OnMallocInitialized()
 #if UE_USE_VERYLARGEPAGEALLOCATOR
 	FCoreDelegates::GetLowLevelAllocatorMemoryTrimDelegate().AddLambda([this]()
 		{
+			FCoreDelegates::GetMemoryTrimDelegate().Broadcast();
+
 			FScopeLock Lock(&ExternalAllocMutex);
 			CachedOSPageAllocator.FreeAll(&ExternalAllocMutex);
 		}
