Avoid writing to bool input when serializing with FArchive during a save operation

diff --git a/Engine/Source/Runtime/Core/Private/Serialization/Archive.cpp b/Engine/Source/Runtime/Core/Private/Serialization/Archive.cpp
--- a/Engine/Source/Runtime/Core/Private/Serialization/Archive.cpp
+++ b/Engine/Source/Runtime/Core/Private/Serialization/Archive.cpp
@@ -511,13 +511,21 @@ void FArchive::SerializeBool( bool& D )
 		OldUBoolValue = D ? 1 : 0;
 		this->Serialize(&OldUBoolValue, sizeof(OldUBoolValue));
 	}
+
 	if (OldUBoolValue > 1)
 	{
 		UE_LOG(LogSerialization, Error, TEXT("Invalid boolean encountered while reading archive %s - stream is most likely corrupted."), *GetArchiveName());
 
 		this->SetError();
 	}
-	D = !!OldUBoolValue;
+
+	if (this->IsLoading())
+	{
+		// Only write to our input if we are loading to make sure we don't
+		// write to 'read-only' memory (e.g. FMemoryWriter) as this can cause TSAN
+		// validation races
+		D = !!OldUBoolValue;
+	}
 }
 #endif
 
diff --git a/Engine/Source/Runtime/Core/Public/Serialization/Archive.h b/Engine/Source/Runtime/Core/Public/Serialization/Archive.h
--- a/Engine/Source/Runtime/Core/Public/Serialization/Archive.h
+++ b/Engine/Source/Runtime/Core/Public/Serialization/Archive.h
@@ -1456,7 +1456,13 @@ public:
 				Ar.SetError();
 			}
 
-			D = !!OldUBoolValue;
+			if (Ar.IsLoading())
+			{
+				// Only write to our input if we are loading to make sure we don't
+				// write to 'read-only' memory (e.g. FMemoryWriter) as this can cause TSAN
+				// validation races
+				D = !!OldUBoolValue;
+			}
 		}
 		return Ar;
 	}
