Prevent GModuleSymbolFileMemory from being freed if it's size exceeds SSIZE_MAX. Cleaned up some logging and code duplication.

#testing Built locally and ran forked linux server
#sh-09-06-pre-sh-31-20-unable-to-load-dm-ugc-maps-ext

diff --git a/Engine/Source/Runtime/Core/Private/Unix/UnixPlatformStackWalk.cpp b/Engine/Source/Runtime/Core/Private/Unix/UnixPlatformStackWalk.cpp
--- a/Engine/Source/Runtime/Core/Private/Unix/UnixPlatformStackWalk.cpp
+++ b/Engine/Source/Runtime/Core/Private/Unix/UnixPlatformStackWalk.cpp
@@ -52,7 +52,7 @@ void CORE_API UnixPlatformStackWalk_PreloadModuleSymbolFile()
 		if (SymbolFileFD == -1)
 		{
 			int ErrNo = errno;
-			UE_LOG(LogHAL, Warning, TEXT("UnixPlatformStackWalk_UnloadPreloadedModuleSymbol: open() failed on path %s errno=%d (%s)"),
+			UE_LOG(LogHAL, Warning, TEXT("UnixPlatformStackWalk_PreloadedModuleSymbol: open() failed on path %s errno=%d (%s)"),
 				*ModuleSymbolPath,
 				ErrNo,
 				UTF8_TO_TCHAR(strerror(ErrNo)));
@@ -65,16 +65,24 @@ void CORE_API UnixPlatformStackWalk_PreloadModuleSymbolFile()
 
 			// Allocate and jump by an extra page size so we can make sure we read only *our* memory and dont read only someone elses.
 			GModuleSymbolFileMemory = (uint8_t*)FMemory::Malloc(GModuleSymbolFileMemorySize +  2 * (FPlatformMemory::GetConstants().PageSize));
+			UE_LOG(LogHAL, Warning, TEXT("UnixPlatformStackWalk_PreloadModuleSymbolFile: GModuleSymbolFileMemory = 0x%x, GModuleSymbolFileMemorySize = %ld bytes"), GModuleSymbolFileMemory, GModuleSymbolFileMemorySize);
+
 			GModuleSymbolFileMemory += FPlatformMemory::GetConstants().PageSize;
-			
+			UE_LOG(LogHAL, Warning, TEXT("UnixPlatformStackWalk_PreloadModuleSymbolFile: GModuleSymbolFileMemory = 0x%x (After adding additional memory page"), GModuleSymbolFileMemory);
+	
+			// On linux read() will transfer at most 2,147,479,552 bytes
+			const int32 MaxBytesToRead = 0x7FFFF000;
+
 			ssize_t BytesRead = 0;
 			{
-				ssize_t RemainingBytes = GModuleSymbolFileMemorySize;
+				// RemainingBytes must be size_t to prevent wrap around if GModuleSymbolFileMemorySize is too large
+				size_t RemainingBytes = GModuleSymbolFileMemorySize;
 				uint8_t* CurrentModulePos = GModuleSymbolFileMemory;
 				ssize_t CurrentBytesRead = 0;
-				while(RemainingBytes > SSIZE_MAX)
+
+				while(RemainingBytes > MaxBytesToRead)
 				{
-					CurrentBytesRead = read(SymbolFileFD, CurrentModulePos, SSIZE_MAX);
+					CurrentBytesRead = read(SymbolFileFD, CurrentModulePos, MaxBytesToRead);
 					if(CurrentBytesRead < 0)
 					{
 						break;
@@ -92,21 +100,22 @@ void CORE_API UnixPlatformStackWalk_PreloadModuleSymbolFile()
 			// Did not read expected amount of bytes
 			if (BytesRead != GModuleSymbolFileMemorySize)
 			{
-				UE_LOG(LogHAL, Warning, TEXT("UnixPlatformStackWalk_UnloadPreloadedModuleSymbol: BytesRead %d Expected %d"), BytesRead, GModuleSymbolFileMemorySize);
-				GModuleSymbolFileMemory -= FPlatformMemory::GetConstants().PageSize;
-				FMemory::Free(GModuleSymbolFileMemory);
+				UE_LOG(LogHAL, Warning, TEXT("UnixPlatformStackWalk_PreloadedModuleSymbol: BytesRead %d Expected %ld"), BytesRead, GModuleSymbolFileMemorySize);
+				UnixPlatformStackWalk_UnloadPreloadedModuleSymbol();
 
 				if (BytesRead == -1)
 				{
 					int ErrNo = errno;
-					UE_LOG(LogHAL, Warning, TEXT("UnixPlatformStackWalk_UnloadPreloadedModuleSymbol: read() failed, errno=%d (%s)"),
+					UE_LOG(LogHAL, Warning, TEXT("UnixPlatformStackWalk_PreloadedModuleSymbol: read() failed, errno=%d (%s)"),
 						ErrNo,
 						UTF8_TO_TCHAR(strerror(ErrNo)));
 				}
 			}
-
-			// Mark our selfs to the left most page boundary read only, we allocated and moved down our memory by a page to give us some slack                                             
-			mprotect(reinterpret_cast<void*>(reinterpret_cast<uint64>(GModuleSymbolFileMemory) & ~(FPlatformMemory::GetConstants().PageSize - 1)), GModuleSymbolFileMemorySize, PROT_READ);
+			else
+			{
+				// Mark our selfs to the left most page boundary read only, we allocated and moved down our memory by a page to give us some slack. Only do this if we've not freed GModuleSymbolFileMemory!
+				mprotect(reinterpret_cast<void*>(reinterpret_cast<uint64>(GModuleSymbolFileMemory) & ~(FPlatformMemory::GetConstants().PageSize - 1)), GModuleSymbolFileMemorySize, PROT_READ);
+			}
 		}
 	}
 }
