Parse port if specified for trace host.

diff --git a/Engine/Source/Runtime/Core/Private/ProfilingDebugging/TraceAuxiliary.cpp b/Engine/Source/Runtime/Core/Private/ProfilingDebugging/TraceAuxiliary.cpp
--- a/Engine/Source/Runtime/Core/Private/ProfilingDebugging/TraceAuxiliary.cpp
+++ b/Engine/Source/Runtime/Core/Private/ProfilingDebugging/TraceAuxiliary.cpp
@@ -552,19 +552,29 @@ void FTraceAuxiliaryImpl::SetTruncateFile(bool bNewTruncateFileState)
 }
 
 ////////////////////////////////////////////////////////////////////////////////
-bool FTraceAuxiliaryImpl::SendToHost(const TCHAR* Host, const FTraceAuxiliary::FLogCategoryAlias& LogCategory, uint16 SendFlags)
+bool FTraceAuxiliaryImpl::SendToHost(const TCHAR* InHost, const FTraceAuxiliary::FLogCategoryAlias& LogCategory, uint16 SendFlags)
 {
-	uint32 Port = 0; // "0" indicates that the default should be used.
-	if (!UE::Trace::SendTo(Host, Port, SendFlags))
+	FStringView HostView(InHost);
+	// Parse port if specified. Default is "0" which indicates default port will be used.
+	uint32 Port = 0; 
+	int32 Separator = INDEX_NONE;
+	if (HostView.FindChar(TEXT(':'), Separator))
 	{
-		UE_LOG_REF(LogCategory, Warning, TEXT("Unable to trace to host '%s'"), Host);
+		LexFromString(Port, HostView.RightChop(Separator + 1).GetData());
+		HostView.LeftInline(Separator);
+	}
+	
+	FString Host(HostView);
+	if (!UE::Trace::SendTo(*Host, Port, SendFlags))
+	{
+		UE_LOG_REF(LogCategory, Warning, TEXT("Unable to trace to host '%s'"), InHost);
 		return false;
 	}
 
 	{
 		FWriteScopeLock _(CurrentTargetLock);
 		CurrentTraceTarget.TraceType = FTraceAuxiliary::EConnectionType::Network;
-		CurrentTraceTarget.TraceDest = MoveTemp(Host);
+		CurrentTraceTarget.TraceDest = MoveTemp(InHost);
 	}
 	return true;
 }
