Adding new settings to distinguish between streaming on cloud and on a local device

-Added new flag to see if we're streaming on cloud or remote
-Added a cvar to turn on/off the option

#rnx

diff --git a/Engine/Source/Runtime/Core/Public/GenericPlatform/GenericPlatformMisc.h b/Engine/Source/Runtime/Core/Public/GenericPlatform/GenericPlatformMisc.h
--- a/Engine/Source/Runtime/Core/Public/GenericPlatform/GenericPlatformMisc.h
+++ b/Engine/Source/Runtime/Core/Public/GenericPlatform/GenericPlatformMisc.h
@@ -1920,6 +1920,14 @@ public:
 		return false;
 	}
 
+	/**
+	*	Return true if the game is being streamed to a separate client device
+	*/
+	static bool WasLaunchedStreaming()
+	{
+		return false;
+	}
+
 	static void DisableScreenTimeout()
 	{
 	}
