Renamed FLinearVirtualMemoryAllocator to FLinearAllocator

diff --git a/Engine/Source/Runtime/Core/Private/Memory/LinearAllocator.cpp b/Engine/Source/Runtime/Core/Private/Memory/LinearAllocator.cpp
new file mode 100644
--- /dev/null
+++ b/Engine/Source/Runtime/Core/Private/Memory/LinearAllocator.cpp
@@ -0,0 +1,142 @@
+// Copyright Epic Games, Inc. All Rights Reserved.
+
+#include "Memory/LinearAllocator.h"
+
+
+CORE_API FPersistentLinearAllocatorExtends GPersistentLinearAllocatorExtends;
+
+static struct FInitializer
+{
+	FInitializer()
+	{
+		GPersistentLinearAllocatorExtends.Address = (uint64)GetPersistentLinearAllocator().GetBasePointer();
+		GPersistentLinearAllocatorExtends.Size = (uint64)GetPersistentLinearAllocator().GetReservedMemorySize();
+	}
+} Initializer;
+
+
+#if UE_ENABLE_LINEAR_VIRTUAL_ALLOCATOR
+
+#include "CoreGlobals.h"
+#include "Misc/ScopeLock.h"
+#include "HAL/LowLevelMemTracker.h"
+
+FLinearAllocator::FLinearAllocator(SIZE_T ReserveMemorySize)
+	: Reserved(ReserveMemorySize)
+{
+	if (FPlatformMemory::CanOverallocateVirtualMemory() && ReserveMemorySize)
+	{
+		VirtualMemory = VirtualMemory.AllocateVirtual(ReserveMemorySize);
+		if (!VirtualMemory.GetVirtualPointer())
+		{
+			UE_LOG(LogMemory, Warning, TEXT("LinearVirtualMemoryAllocator failed to reserve %u MB and will default to FMemory::Malloc instead"), ReserveMemorySize / 1024 / 1024);
+			Reserved = 0;
+		}
+	}
+	else
+	{
+#if PLATFORM_IOS || PLATFORM_TVOS
+		UE_LOG(LogMemory, Warning, TEXT("LinearVirtualMemoryAllocator requires com.apple.developer.kernel.extended-virtual-addressing entitlement to work"));
+#else
+		UE_LOG(LogMemory, Warning, TEXT("This platform does not allow to allocate more virtual memory than there is physical memory. LinearVirtualMemoryAllocator will default to FMemory::Malloc instead"));
+#endif
+		Reserved = 0;
+	}
+
+	GPersistentLinearAllocatorExtends.Address = (uint64)VirtualMemory.GetVirtualPointer();
+	GPersistentLinearAllocatorExtends.Size = (uint64)Reserved;
+}
+
+void* FLinearAllocator::Allocate(SIZE_T Size, uint32 Alignment)
+{
+	Alignment = FMath::Max(Alignment, 8u);
+	{
+		void* Mem = nullptr;
+		{
+			FScopeLock AutoLock(&Lock);
+			if (CanFit(Size, Alignment))
+			{
+				CurrentOffset = Align(CurrentOffset, Alignment);
+				const SIZE_T NewOffset = CurrentOffset + Size;
+				if (NewOffset > Committed)
+				{
+					LLM_PLATFORM_SCOPE(ELLMTag::FMalloc);
+					const size_t CommitGranularity = UE_USE_VERYLARGEPAGEALLOCATOR ? 2ul * 1024 * 1024 : 65536ul;
+					const SIZE_T ToCommit = Align(NewOffset - Committed, FMath::Max(VirtualMemory.GetCommitAlignment(), CommitGranularity));
+				#if UE_USE_VERYLARGEPAGEALLOCATOR
+					if (!VirtualMemory.Commit(Committed, ToCommit, false))
+					{
+						// do not try to use linear allocator anymore and fallback to FMemory::Malloc from now on
+						Reserved = Committed;
+						return FMemory::Malloc(Size, Alignment);
+					}
+				#else
+					VirtualMemory.Commit(Committed, ToCommit);
+				#endif
+					Committed += ToCommit;
+					LLM_IF_ENABLED(FLowLevelMemTracker::Get().OnLowLevelAlloc(ELLMTracker::Platform, (uint8*)VirtualMemory.GetVirtualPointer() + CurrentOffset, ToCommit));
+				}
+				Mem = (uint8*)VirtualMemory.GetVirtualPointer() + CurrentOffset;
+				CurrentOffset += Size;
+			}
+		}
+		if (Mem)
+		{
+			LLM(FLowLevelMemTracker::Get().OnLowLevelAlloc(ELLMTracker::Default, Mem, Size, ELLMTag::Untagged, ELLMAllocType::FMalloc));
+			return Mem;
+		}
+	}
+
+	ExceedsReservation.fetch_add(Size, std::memory_order_relaxed);
+	return FMemory::Malloc(Size, Alignment);
+}
+
+void FLinearAllocator::PreAllocate(SIZE_T Size, uint32 Alignment)
+{
+	Alignment = FMath::Max(Alignment, 8u);
+
+	FScopeLock AutoLock(&Lock);
+	if (CanFit(Size, Alignment))
+	{
+		CurrentOffset = Align(CurrentOffset, Alignment);
+		const SIZE_T NewOffset = CurrentOffset + Size;
+		if (NewOffset > Committed)
+		{
+			const SIZE_T ToCommit = Align(NewOffset - Committed, VirtualMemory.GetCommitAlignment());
+			VirtualMemory.Commit(Committed, ToCommit);
+			Committed += ToCommit;
+		}
+	}
+}
+
+bool FLinearAllocator::TryDeallocate(void* Ptr, SIZE_T Size)
+{
+	if (ContainsPointer(Ptr))
+	{
+		FScopeLock AutoLock(&Lock);
+		if ((uint8*)Ptr + Size == (uint8*)VirtualMemory.GetVirtualPointer() + CurrentOffset)
+		{
+			LLM(FLowLevelMemTracker::Get().OnLowLevelFree(ELLMTracker::Default, Ptr, ELLMAllocType::FMalloc));
+			CurrentOffset -= Size;
+			return true;
+		}
+
+		return false;
+	}
+
+	FMemory::Free(Ptr);
+	return true;
+}
+
+bool FLinearAllocator::CanFit(SIZE_T Size, uint32 Alignment) const
+{
+	return (Reserved - Align(CurrentOffset, Alignment)) >= Size;
+}
+
+FLinearAllocator& GetPersistentLinearAllocator()
+{
+	static FLinearAllocator GPersistentLinearAllocator(UE_PERSISTENT_ALLOCATOR_RESERVE_SIZE);
+	return GPersistentLinearAllocator;
+}
+
+#endif //~UE_ENABLE_LINEAR_VIRTUAL_ALLOCATOR
\ No newline at end of file
diff --git a/Engine/Source/Runtime/Core/Private/Memory/LinearVirtualMemoryAllocator.cpp b/Engine/Source/Runtime/Core/Private/Memory/LinearVirtualMemoryAllocator.cpp
deleted file mode 100644
diff --git a/Engine/Source/Runtime/Core/Private/UObject/UnrealNames.cpp b/Engine/Source/Runtime/Core/Private/UObject/UnrealNames.cpp
--- a/Engine/Source/Runtime/Core/Private/UObject/UnrealNames.cpp
+++ b/Engine/Source/Runtime/Core/Private/UObject/UnrealNames.cpp
@@ -10,7 +10,6 @@
 #include "Math/NumericLimits.h"
 #include "Math/UnrealMathUtility.h"
 #include "HAL/UnrealMemory.h"
-#include "Memory/LinearVirtualMemoryAllocator.h"
 #include "Templates/UnrealTemplate.h"
 #include "Misc/CString.h"
 #include "Misc/Crc.h"
diff --git a/Engine/Source/Runtime/Core/Public/Memory/LinearAllocator.h b/Engine/Source/Runtime/Core/Public/Memory/LinearAllocator.h
new file mode 100644
--- /dev/null
+++ b/Engine/Source/Runtime/Core/Public/Memory/LinearAllocator.h
@@ -0,0 +1,103 @@
+// Copyright Epic Games, Inc. All Rights Reserved.
+
+#pragma once
+
+#include "HAL/Platform.h"
+#include "HAL/UnrealMemory.h"
+
+#ifndef UE_ENABLE_LINEAR_VIRTUAL_ALLOCATOR
+	#define UE_ENABLE_LINEAR_VIRTUAL_ALLOCATOR PLATFORM_HAS_FPlatformVirtualMemoryBlock
+#endif
+
+#if UE_ENABLE_LINEAR_VIRTUAL_ALLOCATOR
+#include <atomic>
+#include "HAL/CriticalSection.h"
+
+struct FLinearAllocator
+{
+	CORE_API FLinearAllocator(SIZE_T ReserveMemorySize);
+
+	FORCEINLINE ~FLinearAllocator()
+	{
+		VirtualMemory.FreeVirtual();
+	}
+
+	CORE_API void* Allocate(SIZE_T Size, uint32 Alignment = 8);
+	CORE_API void  PreAllocate(SIZE_T Size, uint32 Alignment = 8);
+
+	// This will succeed only when deallocating the last allocation
+	CORE_API bool  TryDeallocate(void* Ptr, SIZE_T Size);
+
+	FORCEINLINE bool ContainsPointer(const void* Ptr) const
+	{
+		return (uintptr_t)Ptr - (uintptr_t)VirtualMemory.GetVirtualPointer() < Reserved;
+	}
+
+	FORCEINLINE SIZE_T GetExceedingSize() const
+	{
+		return ExceedsReservation.load(std::memory_order_relaxed);
+	}
+
+	FORCEINLINE bool IsInitialized() const
+	{
+		return Reserved != 0;
+	}
+
+	FORCEINLINE SIZE_T GetAllocatedMemorySize() const
+	{
+		return Committed;
+	}
+
+	FORCEINLINE SIZE_T GetReservedMemorySize() const
+	{
+		return Reserved;
+	}
+
+	FORCEINLINE const void* GetBasePointer() const
+	{
+		return VirtualMemory.GetVirtualPointer();
+	}
+
+private:
+	FCriticalSection Lock;
+	FPlatformMemory::FPlatformVirtualMemoryBlock VirtualMemory;
+	SIZE_T Reserved;
+	SIZE_T Committed = 0;
+	SIZE_T CurrentOffset = 0;
+	std::atomic<SIZE_T> ExceedsReservation = 0;
+
+	bool CanFit(SIZE_T Size, uint32 Alignment) const;
+};
+
+CORE_API FLinearAllocator& GetPersistentLinearAllocator();
+
+#else
+// stub implementation with most functions being nop
+struct FLinearAllocator
+{
+	FORCEINLINE FLinearAllocator(SIZE_T) {}
+
+	FORCEINLINE void*	Allocate(SIZE_T Size, uint32 Alignment = 8) { return FMemory::Malloc(Size, Alignment); }
+	FORCEINLINE void	PreAllocate(SIZE_T, uint32) {}
+	FORCEINLINE bool	TryDeallocate(void* Ptr, SIZE_T Size)	{ FMemory::Free(Ptr); return true; }
+
+	FORCEINLINE bool	ContainsPointer(const void*) const	{ return false; }
+	FORCEINLINE SIZE_T	GetExceedingSize() const			{ return 0; }
+	FORCEINLINE bool	IsInitialized() const				{ return false; }
+	FORCEINLINE SIZE_T	GetAllocatedMemorySize() const		{ return 0; }
+	FORCEINLINE SIZE_T	GetReservedMemorySize() const		{ return 0; }
+	FORCEINLINE const void* GetBasePointer() const			{ return nullptr; }
+};
+
+FORCEINLINE FLinearAllocator GetPersistentLinearAllocator() { return FLinearAllocator(0); }
+
+#endif	//~UE_ENABLE_LINEAR_VIRTUAL_ALLOCATOR
+
+struct FPersistentLinearAllocatorExtends
+{
+	uint64 Address;
+	uint64 Size;
+};
+
+// Special case for the FPermanentObjectPoolExtents to reduce the amount of pointer dereferencing
+extern CORE_API FPersistentLinearAllocatorExtends GPersistentLinearAllocatorExtends;
\ No newline at end of file
diff --git a/Engine/Source/Runtime/Core/Public/Memory/LinearVirtualMemoryAllocator.h b/Engine/Source/Runtime/Core/Public/Memory/LinearVirtualMemoryAllocator.h
deleted file mode 100644
diff --git a/Engine/Source/Runtime/CoreUObject/Public/UObject/UObjectAllocator.h b/Engine/Source/Runtime/CoreUObject/Public/UObject/UObjectAllocator.h
--- a/Engine/Source/Runtime/CoreUObject/Public/UObject/UObjectAllocator.h
+++ b/Engine/Source/Runtime/CoreUObject/Public/UObject/UObjectAllocator.h
@@ -8,7 +8,7 @@
 
 #include "CoreMinimal.h"
 #include "HAL/PlatformMath.h"
-#include "Memory/LinearVirtualMemoryAllocator.h"
+#include "Memory/LinearAllocator.h"
 
 class UObjectBase;
 
@@ -20,7 +20,7 @@ public:
 	 *
 	 * @param InPermanentObjectPoolSize size of permanent object pool
 	 */
-	UE_DEPRECATED(5.5, "Permanent Object Pool is handled by the global instance of FLinearVirtualMemoryAllocator now")
+	UE_DEPRECATED(5.5, "Permanent Object Pool is handled by the global instance of FLinearAllocator now")
 	COREUOBJECT_API void AllocatePermanentObjectPool(int32 InPermanentObjectPoolSize) {}
 
 	/**
@@ -59,7 +59,7 @@ extern COREUOBJECT_API FUObjectAllocator GUObjectAllocator;
 class FPermanentObjectPoolExtents
 {
 public:
-	FORCEINLINE FPermanentObjectPoolExtents(const FLinearVirtualMemoryAllocatorExtends& InAllocatorExtends = GPersistentLinearAllocatorExtends)
+	FORCEINLINE FPermanentObjectPoolExtents(const FPersistentLinearAllocatorExtends& InAllocatorExtends = GPersistentLinearAllocatorExtends)
 		: Address(InAllocatorExtends.Address)
 		, Size(InAllocatorExtends.Size)
 	{}
