Fixed compilation in UnrealMemory.cpp

diff --git a/Engine/Source/Runtime/Core/Private/HAL/UnrealMemory.cpp b/Engine/Source/Runtime/Core/Private/HAL/UnrealMemory.cpp
--- a/Engine/Source/Runtime/Core/Private/HAL/UnrealMemory.cpp
+++ b/Engine/Source/Runtime/Core/Private/HAL/UnrealMemory.cpp
@@ -23,7 +23,7 @@
 #include "HAL/MallocDoubleFreeFinder.h"
 #include "HAL/MallocFrameProfiler.h"
 #include "HAL/MallocStomp2.h"
-#include "Memory/LinearVirtualMemoryAllocator.h"
+#include "Memory/LinearAllocator.h"
 
 #if MALLOC_GT_HOOKS
 
