Runtime IoStore trace support

#rnx

diff --git a/Engine/Source/Runtime/Core/Private/ProfilingDebugging/IoStoreTrace.cpp b/Engine/Source/Runtime/Core/Private/ProfilingDebugging/IoStoreTrace.cpp
new file mode 100644
--- /dev/null
+++ b/Engine/Source/Runtime/Core/Private/ProfilingDebugging/IoStoreTrace.cpp
@@ -0,0 +1,113 @@
+// Copyright Epic Games, Inc. All Rights Reserved.
+#include "ProfilingDebugging/IoStoreTrace.h"
+
+#if UE_TRACE_IOSTORE_ENABLED
+
+#include "ProfilingDebugging/CallstackTrace.h"
+#include "IO/IoDispatcherBackend.h"
+#include "Trace/Trace.h"
+
+// Capturing callstacks is disabled by default. 
+// It is most useful for identifying locations in the code where an IoStore load may be missing metadata instrumentation.
+// There is also a cpu & memory cost to capturing the callstacks too.
+#if !defined(UE_TRACE_IOSTORE_CAPTURE_CALLSTACKS)
+	#define UE_TRACE_IOSTORE_CAPTURE_CALLSTACKS 0
+#endif
+
+
+UE_TRACE_CHANNEL_DEFINE(IoStoreChannel)
+UE_TRACE_METADATA_EVENT_DEFINE(IoStoreTag)
+
+UE_TRACE_EVENT_BEGIN(IoStore, BackendName, NoSync)
+	UE_TRACE_EVENT_FIELD(uint64, BackendHandle)
+	UE_TRACE_EVENT_FIELD(UE::Trace::WideString, Name)
+UE_TRACE_EVENT_END()
+
+UE_TRACE_EVENT_BEGIN(IoStore, RequestCreate)
+	UE_TRACE_EVENT_FIELD(uint64, Cycle)
+	UE_TRACE_EVENT_FIELD(uint64, RequestHandle)
+	UE_TRACE_EVENT_FIELD(uint64, BatchHandle)
+	UE_TRACE_EVENT_FIELD(uint32, ChunkIdHash)
+	UE_TRACE_EVENT_FIELD(uint8,  ChunkType)
+	UE_TRACE_EVENT_FIELD(uint32, CallstackId)
+	UE_TRACE_EVENT_FIELD(uint64, Offset)
+	UE_TRACE_EVENT_FIELD(uint64, Size)
+UE_TRACE_EVENT_END()
+
+UE_TRACE_EVENT_BEGIN(IoStore, RequestStarted)
+	UE_TRACE_EVENT_FIELD(uint64, Cycle)
+	UE_TRACE_EVENT_FIELD(uint64, RequestHandle)
+	UE_TRACE_EVENT_FIELD(uint64, BackendHandle)
+UE_TRACE_EVENT_END()
+
+UE_TRACE_EVENT_BEGIN(IoStore, RequestCompleted)
+	UE_TRACE_EVENT_FIELD(uint64, Cycle)
+	UE_TRACE_EVENT_FIELD(uint64, RequestHandle)
+	UE_TRACE_EVENT_FIELD(uint64, Size)
+UE_TRACE_EVENT_END()
+
+UE_TRACE_EVENT_BEGIN(IoStore, RequestFailed)
+	UE_TRACE_EVENT_FIELD(uint64, Cycle)
+	UE_TRACE_EVENT_FIELD(uint64, RequestHandle)
+UE_TRACE_EVENT_END()
+
+void FIoStoreTrace::BackendName(IIoDispatcherBackend* IoDispatcherBackend, const TCHAR* Name)
+{
+	UE_TRACE_LOG(IoStore, BackendName, IoStoreChannel)
+		<< BackendName.BackendHandle((uint64)IoDispatcherBackend)
+		<< BackendName.Name(Name)
+		;
+}
+
+void FIoStoreTrace::RequestCreate(FIoBatch* IoBatch, FIoRequestImpl* IoRequestImpl)
+{
+	const uint32 CallstackId = 0;
+	const uint32 ChunkIdHash = GetTypeHash(IoRequestImpl->ChunkId);
+
+#if UE_TRACE_IOSTORE_CAPTURE_CALLSTACKS
+	if (UE_TRACE_CHANNELEXPR_IS_ENABLED(IoStoreChannel))
+	{
+		CallstackId = CallstackTrace_GetCurrentId();
+	}
+#endif
+
+	UE_TRACE_LOG(IoStore, RequestCreate, IoStoreChannel)
+		<< RequestCreate.Cycle(FPlatformTime::Cycles64())
+		<< RequestCreate.RequestHandle((uint64)IoRequestImpl)
+		<< RequestCreate.BatchHandle((uint64)IoBatch)
+		<< RequestCreate.ChunkIdHash(ChunkIdHash)
+		<< RequestCreate.ChunkType((uint8)IoRequestImpl->ChunkId.GetChunkType())
+		<< RequestCreate.CallstackId(CallstackId)
+		<< RequestCreate.Offset(IoRequestImpl->Options.GetOffset())
+		<< RequestCreate.Size(IoRequestImpl->Options.GetSize())
+		;
+}
+
+void FIoStoreTrace::RequestStarted(FIoRequestImpl* IoRequestImpl, IIoDispatcherBackend* IoDispatcherBackend)
+{
+	UE_TRACE_LOG(IoStore, RequestStarted, IoStoreChannel)
+		<< RequestStarted.Cycle(FPlatformTime::Cycles64())
+		<< RequestStarted.RequestHandle((uint64)IoRequestImpl)
+		<< RequestStarted.BackendHandle((uint64)IoDispatcherBackend)
+		;
+}
+
+void FIoStoreTrace::RequestCompleted(FIoRequestImpl* IoRequestImpl, uint64 Size)
+{
+	UE_TRACE_LOG(IoStore, RequestCompleted, IoStoreChannel)
+		<< RequestCompleted.Cycle(FPlatformTime::Cycles64())
+		<< RequestCompleted.RequestHandle((uint64)IoRequestImpl)
+		<< RequestCompleted.Size(Size)
+		;
+}
+
+void FIoStoreTrace::RequestFailed(FIoRequestImpl* IoRequestImpl)
+{
+	UE_TRACE_LOG(IoStore, RequestFailed, IoStoreChannel)
+		<< RequestFailed.Cycle(FPlatformTime::Cycles64())
+		<< RequestFailed.RequestHandle((uint64)IoRequestImpl)
+		;
+}
+
+
+#endif // UE_TRACE_IOSTORE_ENABLED
diff --git a/Engine/Source/Runtime/Core/Public/ProfilingDebugging/IoStoreTrace.h b/Engine/Source/Runtime/Core/Public/ProfilingDebugging/IoStoreTrace.h
new file mode 100644
--- /dev/null
+++ b/Engine/Source/Runtime/Core/Public/ProfilingDebugging/IoStoreTrace.h
@@ -0,0 +1,79 @@
+// Copyright Epic Games, Inc. All Rights Reserved.
+#pragma once
+
+#include "CoreTypes.h"
+#include "Misc/Build.h"
+#include "Trace/Config.h"
+
+// requires -trace=metadata,assetmetadata,iostore
+
+#if !defined(UE_TRACE_IOSTORE_ENABLED) // to allow users to manually disable it from build.cs files
+	#if !UE_BUILD_SHIPPING
+		#define UE_TRACE_IOSTORE_ENABLED UE_TRACE_ENABLED
+	#else
+		#define UE_TRACE_IOSTORE_ENABLED 0
+	#endif
+#endif
+
+#if UE_TRACE_IOSTORE_ENABLED
+
+
+#include "ProfilingDebugging/MetadataTrace.h"
+#include "ProfilingDebugging/StringsTrace.h"
+
+class FIoBatch;
+class FIoRequestImpl;
+struct IIoDispatcherBackend;
+
+struct FIoStoreTrace
+{
+	static CORE_API void BackendName(IIoDispatcherBackend* IoDispatcherBackend, const TCHAR* Name);
+	static CORE_API void RequestCreate(FIoBatch* IoBatch, FIoRequestImpl* IoRequestImpl);
+	static CORE_API void RequestStarted(FIoRequestImpl* IoRequestImpl, IIoDispatcherBackend* IoDispatcherBackend);
+	static CORE_API void RequestCompleted(FIoRequestImpl* IoRequestImpl, uint64 Size);
+	static CORE_API void RequestFailed(FIoRequestImpl* IoRequestImpl);
+};
+
+UE_TRACE_CHANNEL_EXTERN(IoStoreChannel, CORE_API);
+
+#define TRACE_IOSTORE_BACKEND_NAME(IoDispatcherBackend, Name) \
+	FIoStoreTrace::BackendName(IoDispatcherBackend, Name);
+
+#define TRACE_IOSTORE_REQUEST_CREATE(IoBatch,IoRequestImpl) \
+	FIoStoreTrace::RequestCreate(IoBatch,IoRequestImpl);
+
+#define TRACE_IOSTORE_BACKEND_REQUEST_STARTED(IoRequestImpl, IoDispatcherBackend) \
+	FIoStoreTrace::RequestStarted(IoRequestImpl, IoDispatcherBackend);
+
+#define TRACE_IOSTORE_BACKEND_REQUEST_COMPLETED(IoRequestImpl, Size) \
+	FIoStoreTrace::RequestCompleted(IoRequestImpl, Size);
+
+#define TRACE_IOSTORE_BACKEND_REQUEST_FAILED(IoRequestImpl) \
+	FIoStoreTrace::RequestFailed(IoRequestImpl);
+
+// creates a metadata scope that includes a tag string for identifying upackage data - used for IoStore usage where an asset/package cannot be determined (e.g. bulk data)
+#if UE_TRACE_METADATA_ENABLED
+	CORE_API UE_TRACE_METADATA_EVENT_BEGIN_EXTERN(IoStoreTag)
+		UE_TRACE_METADATA_EVENT_REFERENCE_FIELD(Strings, FName, Tag)
+	UE_TRACE_METADATA_EVENT_END()
+
+	#define TRACE_IOSTORE_METADATA_SCOPE_TAG(TagName) \
+		auto TagNameRef = UE_TRACE_CHANNELEXPR_IS_ENABLED(MetadataChannel|IoStoreChannel) ? FStringTrace::GetNameRef(TagName) : UE::Trace::FEventRef32(0,0);\
+		UE_TRACE_METADATA_SCOPE(IoStoreTag, IoStoreChannel) \
+			<< IoStoreTag.Tag(TagNameRef);
+#else
+	#define TRACE_IOSTORE_METADATA_SCOPE_TAG(...)
+#endif
+
+#else
+
+#define TRACE_IOSTORE_BACKEND_NAME(...)
+#define TRACE_IOSTORE_REQUEST_CREATE(...)
+#define TRACE_IOSTORE_BACKEND_REQUEST_STARTED(...)
+#define TRACE_IOSTORE_BACKEND_REQUEST_COMPLETED(...)
+#define TRACE_IOSTORE_BACKEND_REQUEST_FAILED(...)
+
+#define TRACE_IOSTORE_METADATA_SCOPE_TAG(...)
+
+#endif //UE_TRACE_IOSTORE_ENABLED
+
