Containers/RingBuffer.h: Implement use of the AllocatorT template parameter to TRingBuffer instead of bypassing it and always allocating from FMemory.
#rnx

diff --git a/Engine/Source/Runtime/Core/Public/Containers/RingBuffer.h b/Engine/Source/Runtime/Core/Public/Containers/RingBuffer.h
--- a/Engine/Source/Runtime/Core/Public/Containers/RingBuffer.h
+++ b/Engine/Source/Runtime/Core/Public/Containers/RingBuffer.h
@@ -138,6 +138,11 @@ public:
 	typedef T ElementType;
 	/* The Allocator type being used */
 	typedef AllocatorT Allocator;
+	using ElementAllocatorType = std::conditional_t<
+		Allocator::NeedsElementType,
+		typename Allocator::template ForElementType<ElementType>,
+		typename Allocator::ForAnyElementType
+	>;
 	/** Type used to request values at a given index in the container. */
 	typedef std::make_signed_t<typename Allocator::SizeType> IndexType;
 	/** Type used to communicate size and capacity and counts */
@@ -157,8 +162,7 @@ public:
 
 	/** Construct Empty Queue with capacity 0. */
 	TRingBuffer()
-		: AllocationData(nullptr)
-		, IndexMask(static_cast<StorageModuloType>(-1))
+		: IndexMask(static_cast<StorageModuloType>(-1))
 		, Front(0u)
 		, AfterBack(0u)
 	{
@@ -189,7 +193,12 @@ public:
 
 	TRingBuffer& operator=(TRingBuffer&& Other)
 	{
-		Swap(AllocationData, Other.AllocationData);
+		if (this == &Other)
+		{
+			return *this;
+		}
+		Empty();
+		AllocatorInstance.MoveToEmpty(Other.AllocatorInstance);
 		Swap(IndexMask, Other.IndexMask);
 		Swap(Front, Other.Front);
 		Swap(AfterBack, Other.AfterBack);
@@ -327,7 +336,12 @@ public:
 	IndexType AddUninitialized()
 	{
 		ConditionalIncrementCapacity();
-		return static_cast<IndexType>(AfterBack++ - Front); // Note this may overflow and set AfterBack = 0.  This overflow is legal; the constraint ((AfterBack - Front) == Num()) will still be true despite Front and AfterBack being on opposite sides of 0.
+		// Note this increment may overflow and set AfterBack = 0.  This overflow is legal;
+		// the constraint ((AfterBack - Front) == Num()) will still be true despite Front and AfterBack being
+		// on opposite sides of 0.
+		IndexType Result = static_cast<IndexType>(AfterBack++ - Front); 
+		SlackTrackerNumChanged();
+		return Result;
 	}
 
 	/** Add a new element after the back pointer of the RingBuffer, resizing if necessary.  The constructor is not called on the new element and its values in memory are arbitrary. Returns a reference to the added element. */
@@ -373,6 +387,7 @@ public:
 			MoveConstructItems(Data + MaskedMoveRangeStart, OtherData, OtherNum);
 		}
 		AfterBack += OtherNum;
+		SlackTrackerNumChanged();
 	}
 
 	/** Add a new element before the front pointer of the RingBuffer, resizing if necessary.  The new element is move constructed from the argument. Returns the index of the added element. */
@@ -433,7 +448,11 @@ public:
 	IndexType AddFrontUninitialized()
 	{
 		ConditionalIncrementCapacity();
-		--Front; // Note this may underflow and set Front = 0xffffffff.  This underflow is legal; the constraint ((AfterBack - Front) == Num()) will still be true despite Front and AfterBack being on opposite sides of 0.
+		// Note this decrement may underflow and set Front = 0xffffffff.  This underflow is legal;
+		// the constraint ((AfterBack - Front) == Num()) will still be true despite Front and AfterBack being on
+		// opposite sides of 0.
+		--Front;
+		SlackTrackerNumChanged();
 		return 0;
 	}
 
@@ -478,7 +497,9 @@ public:
 	void PopFrontNoCheck(SizeType PopCount=1)
 	{
 		DestructRange(Front, Front + PopCount);
-		Front += PopCount; // Note this may overflow (wrapping around to 0xffffffff) if AfterBack has already underflowed; this is valid.
+		// Note this increment may overflow (wrapping around to 0xffffffff) if AfterBack has already underflowed; this is valid.
+		Front += PopCount;
+		SlackTrackerNumChanged();
 	}
 
 	/* Pop one element from the front pointer of the RingBuffer and return the popped value. Invalid to call when the RingBuffer is empty. */
@@ -501,7 +522,9 @@ public:
 	void PopNoCheck(SizeType PopCount=1)
 	{
 		DestructRange(AfterBack - PopCount, AfterBack);
-		AfterBack -= PopCount; // Note this may underflow (wrapping around to 0xffffffff) if Front has already underflowed; this is valid.
+		// Note this decrement may underflow (wrapping around to 0xffffffff) if Front has already underflowed; this is valid.
+		AfterBack -= PopCount;
+		SlackTrackerNumChanged();
 	}
 
 	/* Pop one element from the back pointer of the RingBuffer and return the popped value. Invalid to call when the RingBuffer is empty. */
@@ -711,6 +734,7 @@ public:
 				++WriteIndex;
 			}
 		}
+		SlackTrackerNumChanged();
 		return NumDeleted;
 	}
 
@@ -783,60 +807,80 @@ private:
 		check(NormalizeCapacity(NewCapacity) == NewCapacity);
 		check(NewCapacity >= SrcNum);
 
-		ElementType* NewData;
 		StorageModuloType NewIndexMask;
-		if (NewCapacity > 0)
+		if (NewCapacity > 0 && SrcNum > 0)
 		{
-			NewData = static_cast<ElementType*>(FMemory::Malloc(sizeof(ElementType) * NewCapacity, alignof(ElementType)));
-			NewIndexMask = NewCapacity - 1;
-			if (SrcNum > 0)
-			{
-				// move data to new storage
-				const StorageModuloType MaskedFront = Front & IndexMask;
-				const StorageModuloType MaskedAfterBack = AfterBack & IndexMask;
+			// Allocate SwapStorage
+			ElementAllocatorType SwapStorageBuffer;
+			AllocatorResizeAllocation(SwapStorageBuffer, 0, NewCapacity);
+			ElementType* SwapStorage = SwapStorageBuffer.GetAllocation();
 
-				// MaskedFront equal to MaskedAfterBack will occur if the queue's Num equals Capacity or if the Num == 0.  We checked Num == 0 above, so here it means Num == Capacity.
-				if (MaskedFront >= MaskedAfterBack)
+			// Move data to swap storage
+			const StorageModuloType MaskedFront = Front & IndexMask;
+			const StorageModuloType MaskedAfterBack = AfterBack & IndexMask;
+
+			// MaskedFront equal to MaskedAfterBack will occur if the queue's Num equals Capacity or if the Num == 0.  We checked Num == 0 above, so here it means Num == Capacity.
+			if (MaskedFront >= MaskedAfterBack)
+			{
+				StorageModuloType WriteIndex = 0;
+				for (StorageModuloType ReadIndex = MaskedFront; ReadIndex < SrcCapacity; ++ReadIndex)
 				{
-					StorageModuloType WriteIndex = 0;
-					for (StorageModuloType ReadIndex = MaskedFront; ReadIndex < SrcCapacity; ++ReadIndex)
-					{
-						::new ((void*)&NewData[WriteIndex++]) ElementType(MoveTemp(SrcData[ReadIndex]));
-					}
-					DestructRange(MaskedFront, SrcCapacity);
-					for (StorageModuloType ReadIndex = 0; ReadIndex < MaskedAfterBack; ++ReadIndex)
-					{
-						::new ((void*)&NewData[WriteIndex++]) ElementType(MoveTemp(SrcData[ReadIndex]));
-					}
-					DestructRange(0, MaskedAfterBack);
+					::new ((void*)&SwapStorage[WriteIndex++]) ElementType(MoveTemp(SrcData[ReadIndex]));
 				}
-				else
+				DestructRange(MaskedFront, SrcCapacity);
+				for (StorageModuloType ReadIndex = 0; ReadIndex < MaskedAfterBack; ++ReadIndex)
 				{
-					StorageModuloType WriteIndex = 0;
-					for (StorageModuloType ReadIndex = MaskedFront; ReadIndex < MaskedAfterBack; ++ReadIndex)
-					{
-						::new((void*)&NewData[WriteIndex++]) ElementType(MoveTemp(SrcData[ReadIndex]));
-					}
-					DestructRange(MaskedFront, MaskedAfterBack);
+					::new ((void*)&SwapStorage[WriteIndex++]) ElementType(MoveTemp(SrcData[ReadIndex]));
 				}
+				DestructRange(0, MaskedAfterBack);
 			}
+			else
+			{
+				StorageModuloType WriteIndex = 0;
+				for (StorageModuloType ReadIndex = MaskedFront; ReadIndex < MaskedAfterBack; ++ReadIndex)
+				{
+					::new((void*)&SwapStorage[WriteIndex++]) ElementType(MoveTemp(SrcData[ReadIndex]));
+				}
+				DestructRange(MaskedFront, MaskedAfterBack);
+			}
+			// Move the swap storage into this->AllocatorInstance. "Empty" in MoveToEmpty means the elements have been
+			// destructed, but the array memory might still be allocated; MoveToEmpty reallocates or frees it.
+			AllocatorInstance.MoveToEmpty(SwapStorageBuffer);
+			NewIndexMask = NewCapacity - 1;
 		}
 		else
 		{
-			NewData = nullptr;
-			NewIndexMask = static_cast<StorageModuloType>(-1);
-		}
-		if (AllocationData)
-		{
-			FMemory::Free(AllocationData);
+			AllocatorResizeAllocation(AllocatorInstance, 0, NewCapacity);
+			NewIndexMask = static_cast<StorageModuloType>(NewCapacity - 1);
 		}
 
-		AllocationData = NewData;
 		IndexMask = NewIndexMask;
 		Front = 0u;
 		AfterBack = SrcNum;
 	}
 
+	void AllocatorResizeAllocation(ElementAllocatorType& InAllocatorInstance, SizeType CurrentNum, SizeType NewArrayMax)
+	{
+		if constexpr (TAllocatorTraits<Allocator>::SupportsElementAlignment)
+		{
+			InAllocatorInstance.ResizeAllocation(CurrentNum, NewArrayMax, sizeof(ElementType), alignof(ElementType));
+		}
+		else
+		{
+			InAllocatorInstance.ResizeAllocation(CurrentNum, NewArrayMax, sizeof(ElementType));
+		}
+	}
+
+	void SlackTrackerNumChanged()
+	{
+#if UE_ENABLE_ARRAY_SLACK_TRACKING
+		if constexpr (TAllocatorTraits<Allocator>::SupportsSlackTracking)
+		{
+			AllocatorInstance.SlackTrackerLogNum(Num());
+		}
+#endif
+	}
+
 	/**
 	 * Destruct all elements in the RingBuffer from Index RangeStart to Index RangeEnd.
 	 * RangeStart and RangeEnd are given in StorageModulo space.
@@ -934,13 +978,13 @@ private:
 	/** Return a pointer to the underlying storage of the RingBuffer */
 	const ElementType* GetStorage() const
 	{
-		return AllocationData;
+		return static_cast<const ElementType*>(AllocatorInstance.GetAllocation());
 	}
 
 	/** Return a const pointer to the underlying storage of the RingBuffer */
 	ElementType* GetStorage()
 	{
-		return AllocationData;
+		return static_cast<ElementType*>(AllocatorInstance.GetAllocation());
 	}
 
 	/* Check and return whether the given Index is within range. */
@@ -964,14 +1008,14 @@ private:
 	friend class FRingBufferTest;
 
 	/**
-	 * The underlying storage of the RingBuffer.
+	 * The underlying storage of the RingBuffer is in the c-style array provided by this AllocatorInstance.GetAllocation().
 	 * Elements in this array are uninitialized until the front pointer moves before them or the back pointer moves after them.
 	 * The front pointer and the back pointer can be at arbitrary offsets in this array.
 	 */
-	ElementType* AllocationData;
+	ElementAllocatorType AllocatorInstance;
 	/**
 	 * A bitmask used to convert from StorageModulo space into an index into Storage.
-	 * (X & IndexMask) is a valid index into AllocationData for any value of X, as long as the RingBuffer is non-empty
+	 * (X & IndexMask) is a valid index into GetStorage() for any value of X, as long as the RingBuffer is non-empty.
 	 * Tightly tied to capacity; IndexMask == capacity - 1
 	 */
 	StorageModuloType IndexMask;
diff --git a/Engine/Source/Runtime/Core/Tests/Containers/RingBufferTest.cpp b/Engine/Source/Runtime/Core/Tests/Containers/RingBufferTest.cpp
--- a/Engine/Source/Runtime/Core/Tests/Containers/RingBufferTest.cpp
+++ b/Engine/Source/Runtime/Core/Tests/Containers/RingBufferTest.cpp
@@ -4,8 +4,52 @@
 
 #if WITH_TESTS
 
+#include "Containers/UnrealString.h"
 #include "Tests/TestHarnessAdapter.h"
 
+namespace UE::Containers::RingBuffer::Test
+{
+int32 NumCallsOfResizeAllocation = 0;
+}
+
+template <uint32 NumInlineElements>
+class TCoreTRingBufferTestInlineAllocator :
+	public TInlineAllocator<NumInlineElements, FDefaultAllocator>
+{
+public:
+	using SuperAllocator = TInlineAllocator<NumInlineElements, FDefaultAllocator>;
+	template <typename ElementType>
+	using SuperForElementType = typename SuperAllocator::template ForElementType<ElementType>;
+	using SizeType = SuperAllocator::SizeType;
+
+	template<typename ElementType>
+	class ForElementType : public SuperForElementType<ElementType>
+	{
+	public:
+		// Allocator API that is used and we need to test
+		void ResizeAllocation(SizeType CurrentNum, SizeType NewMax, SIZE_T NumBytesPerElement)
+		{
+			++UE::Containers::RingBuffer::Test::NumCallsOfResizeAllocation;
+			SuperForElementType<ElementType>::ResizeAllocation(CurrentNum, NewMax, NumBytesPerElement);
+		}
+
+		// Allocator API unused by TRingBuffer
+		// void MoveToEmpty(ForElementType& Other)
+		// SizeType CalculateSlackReserve(SizeType NumElements, SIZE_T NumBytesPerElement) const
+		// SizeType CalculateSlackShrink(SizeType NumElements, SizeType NumAllocatedElements, SIZE_T NumBytesPerElement) const
+		// SizeType CalculateSlackGrow(SizeType NumElements, SizeType NumAllocatedElements, SIZE_T NumBytesPerElement) const
+
+		// Allocator API that is used but we don't need to test
+		// FORCEINLINE ElementType* GetAllocation() const
+		// SIZE_T GetAllocatedSize(SizeType NumAllocatedElements, SIZE_T NumBytesPerElement) const
+		// bool HasAllocation() const
+		// SizeType GetInitialCapacity() const
+#if UE_ENABLE_ARRAY_SLACK_TRACKING
+		// void SlackTrackerLogNum(SizeType NewNumUsed)
+#endif
+	};
+};
+
 class FRingBufferTest
 {
 public:
@@ -1034,12 +1078,12 @@ public:
 				{
 					Q.Add(It);
 				}
-				uint32* SavedData = Q.AllocationData;
+				uint32* SavedData = Q.GetStorage();
 				CHECK_MESSAGE(TEXT("Compact - Full array front at start - setup"), (Q.Front == 0) == true);
 				View = Q.Compact();
 				CHECK_MESSAGE(TEXT("Compact - Full array front at start - values"),
 					ArrayViewsEqual(View, TArrayView<const uint32>({ 0,1,2,3,4,5,6,7 })) == true);
-				CHECK_MESSAGE(TEXT("Compact - Full array front at start - no reallocate"), (Q.AllocationData == SavedData) == true);
+				CHECK_MESSAGE(TEXT("Compact - Full array front at start - no reallocate"), (Q.GetStorage() == SavedData) == true);
 			}
 		}
 
@@ -1122,6 +1166,64 @@ public:
 			}
 		}
 
+		// Test use of a specified allocator
+		{
+			using namespace UE::Containers::RingBuffer::Test;
+			{
+				NumCallsOfResizeAllocation = 0;
+				TRingBuffer<int32, TCoreTRingBufferTestInlineAllocator<10>> Q;
+				for (int32 AddIndex = 0; AddIndex < 20; ++AddIndex)
+				{
+					Q.Add(AddIndex);
+					for (int32 TestIndex = 0; TestIndex < Q.Num(); ++TestIndex)
+					{
+						CHECK_MESSAGE(TEXT("Specified Allocator - Expected values present with inline allocator"),
+							Q[TestIndex] == TestIndex);
+					}
+				}
+				CHECK_MESSAGE(TEXT("Specified Allocator - ResizeAllocation called"), NumCallsOfResizeAllocation > 0);
+			}
+		}
+
+		// Test some Reallocate cases
+		{
+			struct FLargeElement
+			{
+				FLargeElement(int32 InValue)
+					: Value(InValue)
+				{
+				}
+				int32 Value;
+				// To be a proper test, this buffer must be larger than InlineBufferSize in Reallocate
+				uint8 Buffer[16384];
+			};
+			{
+				TRingBuffer<FLargeElement> Q;
+				constexpr int32 Count = 10;
+				for (int32 AddIndex = 0; AddIndex < Count; ++AddIndex)
+				{
+					Q.Emplace(AddIndex);
+				}
+				for (int32 TestIndex = 0; TestIndex < Count; ++TestIndex)
+				{
+					CHECK_MESSAGE(TEXT("Reallocate FLargeElement - Expected values present"), Q[TestIndex].Value == TestIndex);
+				}
+			}
+			{
+				// To be a proper test, sizeof(int32)*Count must be larger than InlineBufferSize in Reallocate 
+				constexpr int32 Count = 10000;
+				TRingBuffer<FLargeElement> Q;
+				for (int32 AddIndex = 0; AddIndex < Count; ++AddIndex)
+				{
+					Q.Emplace(AddIndex);
+				}
+				for (int32 TestIndex = 0; TestIndex < Count; ++TestIndex)
+				{
+					CHECK_MESSAGE(TEXT("Reallocate large count - Expected values present"), Q[TestIndex].Value == TestIndex);
+				}
+			}
+		}
+
 		return true;
 	}
 
