Containers/RingBuffer.h: Fix compile error due to missing typename in using statement to a typedef on a templated class.
#rnx

diff --git a/Engine/Source/Runtime/Core/Tests/Containers/RingBufferTest.cpp b/Engine/Source/Runtime/Core/Tests/Containers/RingBufferTest.cpp
--- a/Engine/Source/Runtime/Core/Tests/Containers/RingBufferTest.cpp
+++ b/Engine/Source/Runtime/Core/Tests/Containers/RingBufferTest.cpp
@@ -20,7 +20,7 @@ public:
 	using SuperAllocator = TInlineAllocator<NumInlineElements, FDefaultAllocator>;
 	template <typename ElementType>
 	using SuperForElementType = typename SuperAllocator::template ForElementType<ElementType>;
-	using SizeType = SuperAllocator::SizeType;
+	using SizeType = typename SuperAllocator::SizeType;
 
 	template<typename ElementType>
 	class ForElementType : public SuperForElementType<ElementType>
@@ -34,12 +34,12 @@ public:
 		}
 
 		// Allocator API unused by TRingBuffer
-		// void MoveToEmpty(ForElementType& Other)
 		// SizeType CalculateSlackReserve(SizeType NumElements, SIZE_T NumBytesPerElement) const
 		// SizeType CalculateSlackShrink(SizeType NumElements, SizeType NumAllocatedElements, SIZE_T NumBytesPerElement) const
 		// SizeType CalculateSlackGrow(SizeType NumElements, SizeType NumAllocatedElements, SIZE_T NumBytesPerElement) const
 
 		// Allocator API that is used but we don't need to test
+		// void MoveToEmpty(ForElementType& Other)
 		// FORCEINLINE ElementType* GetAllocation() const
 		// SIZE_T GetAllocatedSize(SizeType NumAllocatedElements, SIZE_T NumBytesPerElement) const
 		// bool HasAllocation() const
