Fix a compiler error because `ALLOW_DEBUG_FILES` wasn't being checked.

diff --git a/Engine/Source/Runtime/Core/Private/HAL/MallocLeakDetection.cpp b/Engine/Source/Runtime/Core/Private/HAL/MallocLeakDetection.cpp
--- a/Engine/Source/Runtime/Core/Private/HAL/MallocLeakDetection.cpp
+++ b/Engine/Source/Runtime/Core/Private/HAL/MallocLeakDetection.cpp
@@ -313,6 +313,7 @@ int32 FMallocLeakDetection::DumpOpenCallstacks(const TCHAR* FileName, const FMal
 	}
 	else
 	{
+#if ALLOW_DEBUG_FILES
 		const FString PathName = *(FPaths::ProfilingDir() + TEXT("memreports/"));
 		IFileManager::Get().MakeDirectory(*PathName);
 
@@ -321,6 +322,10 @@ int32 FMallocLeakDetection::DumpOpenCallstacks(const TCHAR* FileName, const FMal
 		FileAr = IFileManager::Get().CreateDebugFileWriter(*FilePath);
 		FileArWrapper = new FOutputDeviceArchiveWrapper(FileAr);
 		ReportAr = FileArWrapper;
+#else
+		// We can't create debug files because we aren't allowed, so bail!
+		return 0;
+#endif
 	}
 
 	const float InvToMb = 1.0 / (1024 * 1024);
