[AutoRTFM] Add TTransactionallySafeCounter wrapper class.

This is a template which wraps a FThreadSafeCounter (normal, 64, or no-op) with transactionally-safe behavior; atomic operations are performed in the open and will be reversed if the transaction is aborted.

This is now used in the LockFreeFixedSizeAllocator rather than bespoke logic for NumUsed.

diff --git a/Engine/Source/Runtime/Core/Public/Containers/LockFreeFixedSizeAllocator.h b/Engine/Source/Runtime/Core/Public/Containers/LockFreeFixedSizeAllocator.h
--- a/Engine/Source/Runtime/Core/Public/Containers/LockFreeFixedSizeAllocator.h
+++ b/Engine/Source/Runtime/Core/Public/Containers/LockFreeFixedSizeAllocator.h
@@ -7,6 +7,7 @@
 #include "Misc/AssertionMacros.h"
 #include "HAL/UnrealMemory.h"
 #include "Misc/NoopCounter.h"
+#include "Misc/TransactionallySafeCounter.h"
 #include "Containers/LockFreeList.h"
 
 /**
@@ -205,14 +206,6 @@ public:
 			check(!NumUsed.GetValue());
 			Trim();
 		};
-
-		// If we are in a transaction and this allocator exists on the transaction's stack, 
-		// running the ONABORT handler after destruction is dangerous and could stomp memory
-		// that belongs to someone else.
-		if (AutoRTFM::IsClosed() && AutoRTFM::IsOnCurrentTransactionStack(this))
-		{
-			AutoRTFM::PopAllOnAbortHandlers(this);
-		}
 	}
 
 	/**
@@ -223,15 +216,8 @@ public:
 	 */
 	void* Allocate()
 	{
-		UE_AUTORTFM_OPEN
-		{
-			NumUsed.Increment();
-		};
-		AutoRTFM::PushOnAbortHandler(this, [this]() 
-		{
-			// TODO: investigate reusing the same OnAbort handler for multiple allocations/frees.
-			NumUsed.Decrement();
-		});
+		NumUsed.Increment();
+
 		// When we are in an AutoRTFM transaction, it is simpler and faster to allocate a new block instead
 		// of reusing an existing one, since we don't need to track changes on newly allocated memory.
 		// If the transaction is aborted, AutoRTFM should clean it up for us automatically.
@@ -260,15 +246,8 @@ public:
 	 */
 	void Free(void* Item)
 	{
-		UE_AUTORTFM_OPEN
-		{
-			NumUsed.Decrement();
-		};
-		AutoRTFM::PushOnAbortHandler(this, [this]()
-		{
-			// TODO: investigate reusing the same OnAbort handler for multiple allocations/frees.
-			NumUsed.Increment();
-		});
+		NumUsed.Decrement();
+
 		if (AutoRTFM::IsClosed())
 		{
 			// When we are in an AutoRTFM transaction, it is simpler to free blocks directly.
@@ -307,7 +286,7 @@ public:
 	 */
 	const TTrackingCounter& GetNumUsed() const
 	{
-		return NumUsed;
+		return NumUsed.GetCounter();
 	}
 
 	/**
@@ -327,7 +306,7 @@ private:
 	TLockFreePointerListUnordered<void, TPaddingForCacheContention> FreeList;
 
 	/** Total number of blocks outstanding and not in the free list. */
-	TTrackingCounter NumUsed;
+	TTransactionallySafeCounter<TTrackingCounter> NumUsed;
 
 	/** Total number of blocks in the free list. */
 	TTrackingCounter NumFree;
diff --git a/Engine/Source/Runtime/Core/Public/Misc/NoopCounter.h b/Engine/Source/Runtime/Core/Public/Misc/NoopCounter.h
--- a/Engine/Source/Runtime/Core/Public/Misc/NoopCounter.h
+++ b/Engine/Source/Runtime/Core/Public/Misc/NoopCounter.h
@@ -8,6 +8,7 @@
 class FNoopCounter
 {
 public:
+	using IntegerType = int32;
 
 	FNoopCounter() { }
 	FNoopCounter( const FNoopCounter& Other ) { }
diff --git a/Engine/Source/Runtime/Core/Public/Misc/TransactionallySafeCounter.h b/Engine/Source/Runtime/Core/Public/Misc/TransactionallySafeCounter.h
new file mode 100644
--- /dev/null
+++ b/Engine/Source/Runtime/Core/Public/Misc/TransactionallySafeCounter.h
@@ -0,0 +1,210 @@
+// Copyright Epic Games, Inc. All Rights Reserved.
+
+#pragma once
+
+#include "AutoRTFM/AutoRTFM.h"
+#include "Templates/Function.h"
+
+#if !UE_AUTORTFM
+
+/** When AutoRTFM is disabled, the template becomes a trivial subclass. */
+#include "HAL/ThreadSafeCounter.h"
+
+template <typename TBaseCounter>
+class TTransactionallySafeCounter : public TBaseCounter
+{
+public:
+	using TBaseCounter::TBaseCounter;
+
+	TBaseCounter& GetCounter()
+	{
+		return *this;
+	}
+
+	const TBaseCounter& GetCounter() const
+	{
+		return *this;
+	}
+};
+
+#else
+
+/**
+ * Wraps a thread safe counter so that it can be safely used in an AutoRTFM transaction.
+ * If the transaction is aborted, any adjustments made in the transaction's thread will 
+ * be undone, but any adjustments made in a non-transactional context will be preserved.
+ */
+template <typename TBaseCounter>
+class TTransactionallySafeCounter
+{
+public:
+	using IntegerType = typename TBaseCounter::IntegerType;
+
+	/**
+	 * Constructor. Initializes counter to passed-in value.
+	 *
+	 * @param Value	Value to initialize counter to
+	 */
+	TTransactionallySafeCounter(int32 Value = 0) : Counter(Value)
+	{
+	}
+
+	/**
+	 * Copy Constructor.
+	 *
+	 * If the counter in the Other parameter is changing from other threads, there are no
+	 * guarantees as to which values you will get up to the caller to not care, synchronize
+	 * or other way to make those guarantees.
+	 *
+	 * @param Other The other thread-safe counter to copy
+	 */
+	template <typename TAnyCounterType>
+	TTransactionallySafeCounter(const TTransactionallySafeCounter<TAnyCounterType>& Other)
+		: TTransactionallySafeCounter(Other.GetValue())
+	{
+	}
+
+	/**
+	 * Destructor.
+	 */
+	~TTransactionallySafeCounter()
+	{
+		// If we are in a transaction and this object exists on the current transaction's stack, 
+		// running the OnAbort handlers after destruction is dangerous and could stomp memory
+		// that belongs to someone else.
+		if (AutoRTFM::IsClosed() && AutoRTFM::IsOnCurrentTransactionStack(this))
+		{
+			AutoRTFM::PopAllOnAbortHandlers(this);
+		}
+	}
+
+	/**
+	 * Increment and return new value.
+	 * If we are running in the closed, track our change so that it can later be undone if the transaction is aborted.
+	 * 
+	 * @return the new, incremented value
+	 * @see Add, Decrement, Reset, Set, Subtract
+	 */
+	IntegerType Increment()
+	{
+		AutoRTFM::PushOnAbortHandler(this, [this]
+		{
+			// TODO: can we consolidate the OnAbort handling of a transaction scope?
+			Counter.Decrement();
+		});
+		IntegerType Result;
+		UE_AUTORTFM_OPEN
+		{
+			Result = Counter.Increment();
+		};
+		return Result;
+	}
+
+	/**
+	 * Adds an amount and returns the old value.
+	 * If we are running in the closed, track our change so that it can later be undone if the transaction is aborted.
+	 *
+	 * @param Amount Amount to increase the counter by
+	 * @return the old value
+	 * @see Decrement, Increment, Reset, Set, Subtract
+	 */
+	IntegerType Add(IntegerType Amount)
+	{
+		AutoRTFM::PushOnAbortHandler(this, [this, Amount]
+		{
+			// TODO: can we consolidate the OnAbort handling of a transaction scope?
+			Counter.Add(-Amount);
+		});
+		IntegerType Result;
+		UE_AUTORTFM_OPEN
+		{
+			Result = Counter.Add(Amount);
+		};
+		return Result;
+	}
+
+	/**
+	 * Decrement and return new value.
+	 * If we are running in the closed, track our change so that it can later be undone if the transaction is aborted.
+	 *
+	 * @return the new, decremented value
+	 * @see Add, Increment, Reset, Set, Subtract
+	 */
+	IntegerType Decrement()
+	{
+		AutoRTFM::PushOnAbortHandler(this, [this]
+		{
+			// TODO: can we consolidate the OnAbort handling of a transaction scope?
+			Counter.Increment();
+		});
+		IntegerType Result;
+		UE_AUTORTFM_OPEN
+		{ 
+			Result = Counter.Decrement();
+		};
+		return Result;
+	}
+
+	/**
+	 * Subtracts an amount and returns the old value.
+	 * If we are running in the closed, track our change so that it can later be undone if the transaction is aborted.
+	 *
+	 * @param Amount Amount to decrease the counter by
+	 * @return the old value
+	 * @see Add, Decrement, Increment, Reset, Set
+	 */
+	IntegerType Subtract(int32 Amount)
+	{
+		AutoRTFM::PushOnAbortHandler(this, [this, Amount]
+		{
+			// TODO: can we consolidate the OnAbort handling of a transaction scope?
+			Counter.Add(Amount);
+		});
+		IntegerType Result;
+		UE_AUTORTFM_OPEN
+		{ 
+			Result = Counter.Add(-Amount);
+		};
+		return Result;
+	}
+
+	/**
+	 * Gets the current value.
+	 *
+	 * @return the current value
+	 */
+	IntegerType GetValue() const
+	{
+		IntegerType Result;
+		UE_AUTORTFM_OPEN
+		{
+			Result = Counter.GetValue();
+		};
+		return Result;
+	}
+
+	/**
+	 * Allows direct access to the base counter.
+	 *
+	 * @return the wrapped counter object
+	 */
+	TBaseCounter& GetCounter()
+	{
+		return Counter;
+	}
+
+	const TBaseCounter& GetCounter() const
+	{
+		return Counter;
+	}
+
+private:
+
+	/** Assignment is intentionally disallowed, as usage wouldn't be thread safe. */
+	void operator=(const TTransactionallySafeCounter<TBaseCounter>& Other) = delete;
+
+	/** A thread-safe counter to wrap */
+	TBaseCounter Counter;
+};
+
+#endif
\ No newline at end of file
diff --git a/Engine/Source/Runtime/Core/Tests/HAL/AutoRTFMTransactionallySafeCounterTest.cpp b/Engine/Source/Runtime/Core/Tests/HAL/AutoRTFMTransactionallySafeCounterTest.cpp
new file mode 100644
--- /dev/null
+++ b/Engine/Source/Runtime/Core/Tests/HAL/AutoRTFMTransactionallySafeCounterTest.cpp
@@ -0,0 +1,151 @@
+// Copyright Epic Games, Inc. All Rights Reserved.
+
+#include "Misc/AutomationTest.h"
+#include "Misc/App.h"
+#include "AutoRTFM/AutoRTFM.h"
+#include "Misc/TransactionallySafeCounter.h"
+#include "HAL/ThreadSafeCounter.h"
+#include "HAL/ThreadSafeCounter64.h"
+
+#if WITH_DEV_AUTOMATION_TESTS
+
+#define CHECK_EQ(A, B) \
+	do { TestEqual(TEXT(__FILE__ ":" UE_STRINGIZE(__LINE__) ": TestEqual(" #A ", " #B ")"), int(A), int(B)); } while (0)
+
+IMPLEMENT_SIMPLE_AUTOMATION_TEST(FAutoRTFMTransactionallySafeCounter, "AutoRTFM + TransactionallySafeCounter", \
+	EAutomationTestFlags::EngineFilter | EAutomationTestFlags::ClientContext | \
+	EAutomationTestFlags::ServerContext | EAutomationTestFlags::CommandletContext)
+
+bool FAutoRTFMTransactionallySafeCounter::RunTest(const FString& Parameters)
+{
+	if (!AutoRTFM::ForTheRuntime::IsAutoRTFMRuntimeEnabled())
+	{
+		ExecutionInfo.AddEvent(FAutomationEvent(EAutomationEventType::Info,
+			TEXT("SKIPPED 'FAutoRTFMTransactionallySafeCounter' test. AutoRTFM disabled.")));
+		return true;
+	}
+
+	auto TestWithType = [&](auto TypedObject)
+	{
+		using CounterType = decltype(TypedObject);
+
+		{
+			// It should be safe to increment a counter.
+			CounterType Counter;
+			CHECK_EQ(Counter.GetValue(), 0);
+			CHECK_EQ(Counter.Increment(), 1);
+			CHECK_EQ(Counter.GetValue(), 1);
+		}
+
+		AutoRTFM::Transact([&]
+		{
+			// It should be safe to increment a counter inside a transaction.
+			CounterType Counter;
+			CHECK_EQ(Counter.GetValue(), 0);
+			CHECK_EQ(Counter.Increment(), 1);
+			CHECK_EQ(Counter.GetValue(), 1);
+		});
+
+		// Incrementing a counter is cancelled out when AbortTransaction is called.
+		{
+			CounterType Counter;
+			CHECK_EQ(Counter.Increment(), 1);
+
+			AutoRTFM::Transact([&]
+			{
+				// It should be safe to increment a counter inside a transaction.
+				CHECK_EQ(Counter.Increment(), 2);
+				CHECK_EQ(Counter.Add(2), 2);
+				CHECK_EQ(Counter.GetValue(), 4);
+				AutoRTFM::AbortTransaction();
+			});
+
+			// Rollback will bring the counter value back to 1.
+			CHECK_EQ(Counter.GetValue(), 1);
+
+			AutoRTFM::Transact([&]
+			{
+				// A second aborted transaction should work just as well as the first time.
+				CHECK_EQ(Counter.Decrement(), 0);
+				CHECK_EQ(Counter.Subtract(2), 0);
+				CHECK_EQ(Counter.GetValue(), -2);
+				AutoRTFM::AbortTransaction();
+			});
+
+			// Again we should have rolled back to 1.
+			CHECK_EQ(Counter.GetValue(), 1);
+		}
+
+		// Verify that transaction nests work normally when committed.
+		{
+			CounterType Counter;
+
+			AutoRTFM::Transact([&]
+			{
+				Counter.Add(100);
+				CHECK_EQ(Counter.GetValue(), 100);
+
+				AutoRTFM::Transact([&]
+				{
+					Counter.Add(1000);
+					CHECK_EQ(Counter.GetValue(), 1100);
+
+					AutoRTFM::Transact([&]
+					{
+						Counter.Add(9999);
+						Counter.Increment();
+						CHECK_EQ(Counter.GetValue(), 11100);
+					});
+
+					CHECK_EQ(Counter.GetValue(), 11100);
+				});
+
+				CHECK_EQ(Counter.GetValue(), 11100);
+			});
+
+			CHECK_EQ(Counter.GetValue(), 11100);
+		}
+
+
+		// Verify that transaction nests work normally when aborted.
+		{
+			CounterType Counter;
+
+			AutoRTFM::Transact([&]
+			{
+				Counter.Add(100);
+				CHECK_EQ(Counter.GetValue(), 100);
+
+				AutoRTFM::Transact([&]
+				{
+					Counter.Add(1000);
+					CHECK_EQ(Counter.GetValue(), 1100);
+
+					AutoRTFM::Transact([&]
+					{
+						Counter.Add(9999);
+						Counter.Increment();
+						CHECK_EQ(Counter.GetValue(), 11100);
+
+						AutoRTFM::AbortTransaction();
+					});
+
+					CHECK_EQ(Counter.GetValue(), 1100);
+					AutoRTFM::AbortTransaction();
+				});
+
+				CHECK_EQ(Counter.GetValue(), 100);
+				AutoRTFM::AbortTransaction();
+			});
+
+			CHECK_EQ(Counter.GetValue(), 0);
+		}
+	};
+
+	TestWithType(TTransactionallySafeCounter<FThreadSafeCounter>());
+	TestWithType(TTransactionallySafeCounter<FThreadSafeCounter64>());
+
+	return true;
+}
+
+#endif  // WITH_DEV_AUTOMATION_TESTS
