ScopedCSVStat with conditional invocation

diff --git a/Engine/Source/Runtime/Core/Public/ProfilingDebugging/CsvProfiler.h b/Engine/Source/Runtime/Core/Public/ProfilingDebugging/CsvProfiler.h
--- a/Engine/Source/Runtime/Core/Public/ProfilingDebugging/CsvProfiler.h
+++ b/Engine/Source/Runtime/Core/Public/ProfilingDebugging/CsvProfiler.h
@@ -103,6 +103,9 @@ struct FCsvDeclaredStat;
 	#define CSV_SCOPED_TIMING_STAT_GLOBAL(StatName) \
 		TRACE_CSV_PROFILER_INLINE_STAT(#StatName, CSV_CATEGORY_INDEX_GLOBAL); \
 		FScopedCsvStat _ScopedCsvStat_ ## StatName (#StatName, CSV_CATEGORY_INDEX_GLOBAL, "CSV_"#StatName);
+	#define CSV_SCOPED_TIMING_STAT_CONDITIONAL(Category,StatName, Condition) \
+		TRACE_CSV_PROFILER_INLINE_STAT(#StatName, CSV_CATEGORY_INDEX(Category)); \
+		FScopedCsvStatConditional _ScopedCsvStat_ ## StatName (#StatName, CSV_CATEGORY_INDEX(Category), Condition, "CSV_"#StatName);
 	#define CSV_SCOPED_TIMING_STAT_EXCLUSIVE(StatName) \
 		TRACE_CSV_PROFILER_INLINE_STAT_EXCLUSIVE(#StatName); \
 		FScopedCsvStatExclusive _ScopedCsvStatExclusive_ ## StatName (#StatName, "CSV_"#StatName);
@@ -577,6 +580,33 @@ public:
 	uint32 CategoryIndex;
 };
 
+class FScopedCsvStatConditional
+{
+public:
+	FScopedCsvStatConditional(const char* InStatName, uint32 InCategoryIndex, bool bInCondition, const char* InNamedEventName = nullptr)
+		: StatName(InStatName)
+		, CategoryIndex(InCategoryIndex)
+		, bCondition(bInCondition)
+	{
+		if (bCondition)
+		{
+			FCsvProfiler::BeginStat(StatName, CategoryIndex, InNamedEventName);
+		}
+	}
+
+	~FScopedCsvStatConditional()
+	{
+		if (bCondition)
+		{
+			FCsvProfiler::EndStat(StatName, CategoryIndex);
+		}
+	}
+
+	const char* StatName;
+	uint32 CategoryIndex;
+	bool bCondition;
+};
+
 class FScopedCsvStatExclusive 
 {
 public:
