Changing TScriptDelegate to only take a read lock during execution (ProcessDelegate). This allows script delegates to call functions which may go wide and do additional reads on the running delegate. Now, the invocation list is compacted after the execution is complete rather than inline with execution.

diff --git a/Engine/Source/Runtime/Core/Public/UObject/ScriptDelegates.h b/Engine/Source/Runtime/Core/Public/UObject/ScriptDelegates.h
--- a/Engine/Source/Runtime/Core/Public/UObject/ScriptDelegates.h
+++ b/Engine/Source/Runtime/Core/Public/UObject/ScriptDelegates.h
@@ -892,37 +892,42 @@ public:
 	/**
 	 * Executes a multi-cast delegate by calling all functions on objects bound to the delegate.  Always
 	 * safe to call, even if when no objects are bound, or if objects have expired.  In general, you should
-	 * never call this function directly.  Instead, call Broadcast() on a derived class.
+	 * never call this function directly.  Instead, call Broadcast() on a derived class. Note that this function
+	 * is not truly const because it will clean up any compactable entries in the invocation list.
 	 *
 	 * @param	Params				Parameter structure
 	 */
 	template <class UObjectTemplate>
 	void ProcessMulticastDelegate(void* Parameters) const
 	{
-		// the `const` on the method is a lie
-		FWriteAccessScope WriteScope = const_cast<TMulticastScriptDelegate*>(this)->GetWriteAccessScope();
-
-		if( InvocationList.Num() > 0 )
 		{
-			// Create a copy of the invocation list, just in case the list is modified by one of the callbacks during the broadcast
-			typedef TArray< UnicastDelegateType, TInlineAllocator< 4 > > FInlineInvocationList;
-			FInlineInvocationList InvocationListCopy = FInlineInvocationList(InvocationList);
-	
-			// Invoke each bound function
-			for( typename FInlineInvocationList::TConstIterator FunctionIt( InvocationListCopy ); FunctionIt; ++FunctionIt )
+			FReadAccessScope ReadScope = const_cast<TMulticastScriptDelegate*>(this)->GetReadAccessScope();
+
+			if( InvocationList.Num() > 0 )
 			{
-				if( FunctionIt->IsBound() )
-				{
-					// Invoke this delegate!
-					FunctionIt->template ProcessDelegate<UObjectTemplate>(Parameters);
-				}
-				else if ( FunctionIt->IsCompactable() )
+				// Create a copy of the invocation list, just in case the list is modified by one of the callbacks during the broadcast
+				typedef TArray< UnicastDelegateType, TInlineAllocator< 4 > > FInlineInvocationList;
+				FInlineInvocationList InvocationListCopy = FInlineInvocationList(InvocationList);
+		
+				// Invoke each bound function
+				for( typename FInlineInvocationList::TConstIterator FunctionIt( InvocationListCopy ); FunctionIt; ++FunctionIt )
 				{
-					// Function couldn't be executed, so remove it.  Note that because the original list could have been modified by one of the callbacks, we have to search for the function to remove here.
-					RemoveInternal( *FunctionIt );
+					if( FunctionIt->IsBound() )
+					{
+						// Invoke this delegate!
+						FunctionIt->template ProcessDelegate<UObjectTemplate>(Parameters);
+					}
 				}
 			}
 		}
+
+		{
+			FWriteAccessScope WriteScope = const_cast<TMulticastScriptDelegate*>(this)->GetWriteAccessScope();
+			// Removes need to occur under a separate write scope because we don't want to hold a write lock during execution
+			// We want to take the least restrictive lock (guard, really) possible in order to permit the callbacks to
+			// inspect the multicast delegate itself (e.g., when using the serialization system to inspect/explore data or find references)
+			CompactInvocationList();
+		}
 	}
 
 	/**
