Add flag to allow writes to IPlatformFile::OpenAsyncRead.  Applies FILE_SHARE_WRITE flag on microsoft platforms.

diff --git a/Engine/Plugins/Developer/Concert/ConcertSync/ConcertSyncClient/Source/ConcertSyncClient/Private/ConcertSandboxPlatformFile.cpp b/Engine/Plugins/Developer/Concert/ConcertSync/ConcertSyncClient/Source/ConcertSyncClient/Private/ConcertSandboxPlatformFile.cpp
--- a/Engine/Plugins/Developer/Concert/ConcertSync/ConcertSyncClient/Source/ConcertSyncClient/Private/ConcertSandboxPlatformFile.cpp
+++ b/Engine/Plugins/Developer/Concert/ConcertSync/ConcertSyncClient/Source/ConcertSyncClient/Private/ConcertSandboxPlatformFile.cpp
@@ -622,7 +622,7 @@ bool FConcertSandboxPlatformFile::IterateDirectoryStat(const TCHAR* Directory, F
 	return LowerLevel->IterateDirectoryStat(Directory, Visitor); // Note: Using the path we were given here to ensure the calling code gets the expected path
 }
 
-IAsyncReadFileHandle* FConcertSandboxPlatformFile::OpenAsyncRead(const TCHAR* Filename)
+IAsyncReadFileHandle* FConcertSandboxPlatformFile::OpenAsyncRead(const TCHAR* Filename, bool bAllowWrite /*= false*/)
 {
 	const FConcertSandboxPlatformFilePath ResolvedPath = ToSandboxPath(Filename);
 
@@ -630,11 +630,11 @@ IAsyncReadFileHandle* FConcertSandboxPlatformFile::OpenAsyncRead(const TCHAR* Fi
 	{
 		if (IsPathDeleted(ResolvedPath) || LowerLevel->FileExists(*ResolvedPath.GetSandboxPath()))
 		{
-			return LowerLevel->OpenAsyncRead(*ResolvedPath.GetSandboxPath());
+			return LowerLevel->OpenAsyncRead(*ResolvedPath.GetSandboxPath(), bAllowWrite);
 		}
 	}
 
-	return LowerLevel->OpenAsyncRead(*ResolvedPath.GetNonSandboxPath());
+	return LowerLevel->OpenAsyncRead(*ResolvedPath.GetNonSandboxPath(), bAllowWrite);
 }
 
 void FConcertSandboxPlatformFile::SetAsyncMinimumPriority(EAsyncIOPriorityAndFlags MinPriority)
diff --git a/Engine/Plugins/Developer/Concert/ConcertSync/ConcertSyncClient/Source/ConcertSyncClient/Private/ConcertSandboxPlatformFile.h b/Engine/Plugins/Developer/Concert/ConcertSync/ConcertSyncClient/Source/ConcertSyncClient/Private/ConcertSandboxPlatformFile.h
--- a/Engine/Plugins/Developer/Concert/ConcertSync/ConcertSyncClient/Source/ConcertSyncClient/Private/ConcertSandboxPlatformFile.h
+++ b/Engine/Plugins/Developer/Concert/ConcertSync/ConcertSyncClient/Source/ConcertSyncClient/Private/ConcertSandboxPlatformFile.h
@@ -146,7 +146,7 @@ public:
 	virtual FFileStatData GetStatData(const TCHAR* FilenameOrDirectory) override;
 	virtual bool IterateDirectory(const TCHAR* Directory, FDirectoryVisitor& Visitor) override;
 	virtual bool IterateDirectoryStat(const TCHAR* Directory, FDirectoryStatVisitor& Visitor) override;
-	virtual IAsyncReadFileHandle* OpenAsyncRead(const TCHAR* Filename) override;
+	virtual IAsyncReadFileHandle* OpenAsyncRead(const TCHAR* Filename, bool bAllowWrite = false) override;
 	virtual void SetAsyncMinimumPriority(EAsyncIOPriorityAndFlags MinPriority) override;
 	virtual FString ConvertToAbsolutePathForExternalAppForRead(const TCHAR* Filename) override;
 	virtual FString ConvertToAbsolutePathForExternalAppForWrite(const TCHAR* Filename) override;
diff --git a/Engine/Source/Runtime/Core/Private/GenericPlatform/GenericPlatformFile.cpp b/Engine/Source/Runtime/Core/Private/GenericPlatform/GenericPlatformFile.cpp
--- a/Engine/Source/Runtime/Core/Private/GenericPlatform/GenericPlatformFile.cpp
+++ b/Engine/Source/Runtime/Core/Private/GenericPlatform/GenericPlatformFile.cpp
@@ -243,12 +243,14 @@ class FGenericAsyncReadFileHandle final : public IAsyncReadFileHandle
 	IFileHandle* HandleCache[MAX_CACHED_SYNC_FILE_HANDLES_PER_GENERIC_ASYNC_FILE_HANDLE];
 	bool bOpenFailed;
 	bool bDisableHandleCaching;
+	bool bAllowWrite;
 public:
-	FGenericAsyncReadFileHandle(IPlatformFile* InLowerLevel, const TCHAR* InFilename)
+	FGenericAsyncReadFileHandle(IPlatformFile* InLowerLevel, const TCHAR* InFilename, bool bInAllowWrite)
 		: LowerLevel(InLowerLevel)
 		, Filename(InFilename)
 		, bOpenFailed(false)
 		, bDisableHandleCaching(!!DISABLE_HANDLE_CACHING)
+		, bAllowWrite(bInAllowWrite)
 	{
 #if !WITH_EDITOR
 		if (GCacheHandleForPakFilesOnly && !Filename.EndsWith(TEXT(".pak")))
@@ -327,9 +329,9 @@ public:
 		if (bDisableHandleCaching)
 		{
 #if DISABLE_BUFFERING_ON_GENERIC_ASYNC_FILE_HANDLE
-			return LowerLevel->OpenReadNoBuffering(*Filename);
+			return LowerLevel->OpenReadNoBuffering(*Filename, bAllowWrite);
 #else
-			return LowerLevel->OpenRead(*Filename);
+			return LowerLevel->OpenRead(*Filename, bAllowWrite);
 #endif
 		}
 		IFileHandle* Result = nullptr;
@@ -340,9 +342,9 @@ public:
 			if (!HandleCache[0] && !bOpenFailed)
 			{
 #if DISABLE_BUFFERING_ON_GENERIC_ASYNC_FILE_HANDLE
-				HandleCache[0] = LowerLevel->OpenReadNoBuffering(*Filename);
+				HandleCache[0] = LowerLevel->OpenReadNoBuffering(*Filename, bAllowWrite);
 #else
-				HandleCache[0] = LowerLevel->OpenRead(*Filename);
+				HandleCache[0] = LowerLevel->OpenRead(*Filename, bAllowWrite);
 #endif
 				bOpenFailed = (HandleCache[0] == nullptr);
 			}
@@ -367,9 +369,9 @@ public:
 			if (!Result && !bOpenFailed)
 			{
 #if DISABLE_BUFFERING_ON_GENERIC_ASYNC_FILE_HANDLE
-				Result = LowerLevel->OpenReadNoBuffering(*Filename);
+				Result = LowerLevel->OpenReadNoBuffering(*Filename, bAllowWrite);
 #else
-				Result = LowerLevel->OpenRead(*Filename);
+				Result = LowerLevel->OpenRead(*Filename, bAllowWrite);
 #endif
 				bOpenFailed = (Result == nullptr);
 			}
@@ -512,9 +514,27 @@ FFileOpenResult IPlatformFile::OpenRead(const TCHAR* Filename, IPlatformFile::EO
 	}
 }
 
-IAsyncReadFileHandle* IPlatformFile::OpenAsyncRead(const TCHAR* Filename)
+FFileOpenAsyncResult IPlatformFile::OpenAsyncRead(const TCHAR* Filename, EOpenReadFlags Flags)
 {
-	return new FGenericAsyncReadFileHandle(this, Filename);
+	if (EnumHasAnyFlags(Flags, EOpenReadFlags::AllowDelete))
+	{
+		return MakeError(TEXTVIEW("The current IPlatformFile implementation does not support the EReadFlags::AllowDelete flag"));
+	}
+
+	const bool bWrite = EnumHasAnyFlags(Flags, EOpenReadFlags::AllowWrite);
+
+	IAsyncReadFileHandle* Handle = OpenAsyncRead(Filename, bWrite);
+	if (Handle != nullptr)
+	{
+		return MakeValue(Handle);
+	}
+
+	return MakeError(TEXTVIEW("[Unknown Error] Current IPlatformFile implementation does not support returning error messages"));
+}
+
+IAsyncReadFileHandle* IPlatformFile::OpenAsyncRead(const TCHAR* Filename, bool bAllowWrite /*= false*/)
+{
+	return new FGenericAsyncReadFileHandle(this, Filename, bAllowWrite);
 }
 
 DEFINE_STAT(STAT_AsyncFileMemory);
diff --git a/Engine/Source/Runtime/Core/Private/Windows/WindowsPlatformFile.cpp b/Engine/Source/Runtime/Core/Private/Windows/WindowsPlatformFile.cpp
--- a/Engine/Source/Runtime/Core/Private/Windows/WindowsPlatformFile.cpp
+++ b/Engine/Source/Runtime/Core/Private/Windows/WindowsPlatformFile.cpp
@@ -1230,10 +1230,10 @@ bool FWindowsPlatformFile::SetMarkOfTheWeb(FStringView Filename, bool bNewStatus
 }
 
 #if USE_WINDOWS_ASYNC_IMPL
-IAsyncReadFileHandle* FWindowsPlatformFile::OpenAsyncRead(const TCHAR* Filename)
+IAsyncReadFileHandle* FWindowsPlatformFile::OpenAsyncRead(const TCHAR* Filename, bool bAllowWrite /*= false*/)
 {
 	uint32  Access = GENERIC_READ;
-	uint32  WinFlags = FILE_SHARE_READ;
+	uint32  WinFlags = FILE_SHARE_READ | (bAllowWrite ? FILE_SHARE_WRITE : 0);
 	uint32  Create = OPEN_EXISTING;
 
 
diff --git a/Engine/Source/Runtime/Core/Private/Windows/WindowsPlatformFilePrivate.h b/Engine/Source/Runtime/Core/Private/Windows/WindowsPlatformFilePrivate.h
--- a/Engine/Source/Runtime/Core/Private/Windows/WindowsPlatformFilePrivate.h
+++ b/Engine/Source/Runtime/Core/Private/Windows/WindowsPlatformFilePrivate.h
@@ -61,7 +61,7 @@ public:
 	virtual bool SetMarkOfTheWeb(FStringView Filename, bool bNewStatus, const FString* InSourceURL = nullptr) override;
 
 #if USE_WINDOWS_ASYNC_IMPL
-	virtual IAsyncReadFileHandle* OpenAsyncRead(const TCHAR* Filename) override;
+	virtual IAsyncReadFileHandle* OpenAsyncRead(const TCHAR* Filename, bool bAllowWrite = false) override;
 #endif
 
 	virtual FFileOpenResult OpenRead(const TCHAR* Filename, EOpenReadFlags Flags) override;
diff --git a/Engine/Source/Runtime/Core/Public/GenericPlatform/GenericPlatformFile.h b/Engine/Source/Runtime/Core/Public/GenericPlatform/GenericPlatformFile.h
--- a/Engine/Source/Runtime/Core/Public/GenericPlatform/GenericPlatformFile.h
+++ b/Engine/Source/Runtime/Core/Public/GenericPlatform/GenericPlatformFile.h
@@ -277,7 +277,11 @@ class FFileSystemError
 {
 public:
 	FFileSystemError() = delete;
-	UE_NONCOPYABLE(FFileSystemError);
+
+	FFileSystemError(const FFileSystemError&) = delete;
+	FFileSystemError(FFileSystemError&&) = default;	
+	FFileSystemError& operator=(const FFileSystemError&) = delete;
+	FFileSystemError& operator=(FFileSystemError&&) = default;
 
 	explicit FFileSystemError(FString&& InErrorMessage, int32 InSystemErrorCode = 0)
 		: ErrorMessage(MoveTemp(InErrorMessage))
@@ -325,6 +329,7 @@ private:
 
 /** Data structure returned by calls to IPlatformFile::OpenRead */
 using FFileOpenResult = TValueOrError<TUniquePtr<IFileHandle>, FFileSystemError>;
+using FFileOpenAsyncResult = TValueOrError<TUniquePtr<IAsyncReadFileHandle>, FFileSystemError>;
 
 /**
 * File I/O Interface
@@ -708,12 +713,19 @@ public:
 	/////////// Generally, these do not need to be implemented per platform.
 	/////////////////////////////////////////////////////////////////////////////////////////////////////////
 
+	/** Open a file for async reading. This call does not hit the disk or block.
+	* 
+	* @param Flags		Allows specialization of the open operation, @see EReadFlags
+	* @return	The return value will either contain the valid file handle or an error message. @see FFileOpenAsyncResult
+	*/
+	CORE_API virtual FFileOpenAsyncResult OpenAsyncRead(const TCHAR* Filename, EOpenReadFlags Flags);
+
 	/** Open a file for async reading. This call does not hit the disk or block.
 	*
 	* @param Filename file to be opened
 	* @return Close the file by delete'ing the handle. A non-null return value does not mean the file exists, since that may not be determined yet.
 	*/
-	CORE_API virtual IAsyncReadFileHandle* OpenAsyncRead(const TCHAR* Filename);
+	CORE_API virtual IAsyncReadFileHandle* OpenAsyncRead(const TCHAR* Filename, bool bAllowWrite = false);
 
 	/** Controls if the pak precacher should process precache requests.
 	* Requests below this threshold will not get precached. Without this throttle, quite a lot of memory
diff --git a/Engine/Source/Runtime/Core/Public/HAL/IPlatformFileCachedWrapper.h b/Engine/Source/Runtime/Core/Public/HAL/IPlatformFileCachedWrapper.h
--- a/Engine/Source/Runtime/Core/Public/HAL/IPlatformFileCachedWrapper.h
+++ b/Engine/Source/Runtime/Core/Public/HAL/IPlatformFileCachedWrapper.h
@@ -445,9 +445,9 @@ public:
 	{
 		return LowerLevel->SendMessageToServer(Message, Handler);
 	}
-	virtual IAsyncReadFileHandle* OpenAsyncRead(const TCHAR* Filename) override
+	virtual IAsyncReadFileHandle* OpenAsyncRead(const TCHAR* Filename, bool bAllowWrite = false) override
 	{
-		return LowerLevel->OpenAsyncRead(Filename);
+		return LowerLevel->OpenAsyncRead(Filename, bAllowWrite);
 	}
 	virtual IMappedFileHandle* OpenMapped(const TCHAR* Filename) override
 	{
diff --git a/Engine/Source/Runtime/Core/Public/HAL/IPlatformFileLogWrapper.h b/Engine/Source/Runtime/Core/Public/HAL/IPlatformFileLogWrapper.h
--- a/Engine/Source/Runtime/Core/Public/HAL/IPlatformFileLogWrapper.h
+++ b/Engine/Source/Runtime/Core/Public/HAL/IPlatformFileLogWrapper.h
@@ -482,13 +482,13 @@ public:
 	}
 	CORE_API void HandleDumpCommand(const TCHAR* Cmd, FOutputDevice& Ar);
 #endif
-	virtual IAsyncReadFileHandle* OpenAsyncRead(const TCHAR* Filename) override
+	virtual IAsyncReadFileHandle* OpenAsyncRead(const TCHAR* Filename, bool bAllowWrite = false) override
 	{
 		FString DataStr = FString::Printf(TEXT("OpenAsyncRead %s"), Filename);
 		SCOPED_NAMED_EVENT_FSTRING(DataStr, FColor::Emerald);
 		FILE_LOG(LogPlatformFile, Log, TEXT("%s"), *DataStr);
 		double StartTime = FPlatformTime::Seconds();
-		IAsyncReadFileHandle* Result = LowerLevel->OpenAsyncRead(Filename);
+		IAsyncReadFileHandle* Result = LowerLevel->OpenAsyncRead(Filename, bAllowWrite);
 		double ThisTime = (FPlatformTime::Seconds() - StartTime) / 1000.0;
 		FILE_LOG(LogPlatformFile, Log, TEXT("OpenAsyncRead return %llx [%fms]"), uint64(Result), ThisTime);
 		//@todo no wrapped logging for async file handles (yet)
diff --git a/Engine/Source/Runtime/Core/Public/HAL/IPlatformFileOpenLogWrapper.h b/Engine/Source/Runtime/Core/Public/HAL/IPlatformFileOpenLogWrapper.h
--- a/Engine/Source/Runtime/Core/Public/HAL/IPlatformFileOpenLogWrapper.h
+++ b/Engine/Source/Runtime/Core/Public/HAL/IPlatformFileOpenLogWrapper.h
@@ -218,10 +218,10 @@ public:
 	{
 		return LowerLevel->SendMessageToServer(Message, Handler);
 	}
-	virtual IAsyncReadFileHandle* OpenAsyncRead(const TCHAR* Filename) override
+	virtual IAsyncReadFileHandle* OpenAsyncRead(const TCHAR* Filename, bool bAllowWrite = false) override
 	{
 		// we must not record the "open" here...what matters is when we start reading the file!
-		return new FLoggingAsyncReadFileHandle(this, Filename, LowerLevel->OpenAsyncRead(Filename));
+		return new FLoggingAsyncReadFileHandle(this, Filename, LowerLevel->OpenAsyncRead(Filename, bAllowWrite));
 	}
 	virtual IMappedFileHandle* OpenMapped(const TCHAR* Filename) override
 	{
diff --git a/Engine/Source/Runtime/PakFile/Private/IPlatformFilePak.cpp b/Engine/Source/Runtime/PakFile/Private/IPlatformFilePak.cpp
--- a/Engine/Source/Runtime/PakFile/Private/IPlatformFilePak.cpp
+++ b/Engine/Source/Runtime/PakFile/Private/IPlatformFilePak.cpp
@@ -4636,7 +4636,7 @@ IMappedFileHandle* FPakPlatformFile::OpenMapped(const TCHAR* Filename)
 	return nullptr;
 }
 
-IAsyncReadFileHandle* FPakPlatformFile::OpenAsyncRead(const TCHAR* Filename)
+IAsyncReadFileHandle* FPakPlatformFile::OpenAsyncRead(const TCHAR* Filename, bool bAllowWrite /*= false*/)
 {
 	CSV_SCOPED_TIMING_STAT(FileIOVerbose, PakOpenAsyncRead);
 #if USE_PAK_PRECACHE
@@ -4668,7 +4668,8 @@ IAsyncReadFileHandle* FPakPlatformFile::OpenAsyncRead(const TCHAR* Filename)
 		}
 	}
 #endif
-	return IPlatformFile::OpenAsyncRead(Filename);
+	// TODO: This seems wrong?  Shouldn't it be LowerLevel->OpenAsyncRead() ?
+	return IPlatformFile::OpenAsyncRead(Filename, bAllowWrite);
 }
 
 void FPakPlatformFile::SetAsyncMinimumPriority(EAsyncIOPriorityAndFlags Priority)
diff --git a/Engine/Source/Runtime/PakFile/Public/IPlatformFilePak.h b/Engine/Source/Runtime/PakFile/Public/IPlatformFilePak.h
--- a/Engine/Source/Runtime/PakFile/Public/IPlatformFilePak.h
+++ b/Engine/Source/Runtime/PakFile/Public/IPlatformFilePak.h
@@ -2484,7 +2484,7 @@ public:
 
 	PAKFILE_API virtual bool CopyFile(const TCHAR* To, const TCHAR* From, EPlatformFileRead ReadFlags = EPlatformFileRead::None, EPlatformFileWrite WriteFlags = EPlatformFileWrite::None) override;
 
-	PAKFILE_API virtual IAsyncReadFileHandle* OpenAsyncRead(const TCHAR* Filename) override;
+	PAKFILE_API virtual IAsyncReadFileHandle* OpenAsyncRead(const TCHAR* Filename, bool bAllowWrite = false) override;
 	PAKFILE_API virtual void SetAsyncMinimumPriority(EAsyncIOPriorityAndFlags Priority) override;
 
 	virtual IMappedFileHandle* OpenMapped(const TCHAR* Filename) override;
diff --git a/Engine/Source/Runtime/SandboxFile/Private/IPlatformFileSandboxWrapper.cpp b/Engine/Source/Runtime/SandboxFile/Private/IPlatformFileSandboxWrapper.cpp
--- a/Engine/Source/Runtime/SandboxFile/Private/IPlatformFileSandboxWrapper.cpp
+++ b/Engine/Source/Runtime/SandboxFile/Private/IPlatformFileSandboxWrapper.cpp
@@ -970,14 +970,14 @@ bool FSandboxPlatformFile::CopyFile(const TCHAR* To, const TCHAR* From, EPlatfor
 	return Result;
 }
 
-IAsyncReadFileHandle* FSandboxPlatformFile::OpenAsyncRead(const TCHAR* Filename)
+IAsyncReadFileHandle* FSandboxPlatformFile::OpenAsyncRead(const TCHAR* Filename, bool bAllowWrite /*= false*/)
 {
 	FString UserFilename(*ConvertToSandboxPath(Filename));
 	if (!OkForInnerAccess(Filename) || LowerLevel->FileExists(*UserFilename))
 	{
-		return LowerLevel->OpenAsyncRead(*UserFilename);
+		return LowerLevel->OpenAsyncRead(*UserFilename, bAllowWrite);
 	}
-	return LowerLevel->OpenAsyncRead(Filename);
+	return LowerLevel->OpenAsyncRead(Filename, bAllowWrite);
 }
 void FSandboxPlatformFile::SetAsyncMinimumPriority(EAsyncIOPriorityAndFlags Priority)
 {
diff --git a/Engine/Source/Runtime/SandboxFile/Public/IPlatformFileSandboxWrapper.h b/Engine/Source/Runtime/SandboxFile/Public/IPlatformFileSandboxWrapper.h
--- a/Engine/Source/Runtime/SandboxFile/Public/IPlatformFileSandboxWrapper.h
+++ b/Engine/Source/Runtime/SandboxFile/Public/IPlatformFileSandboxWrapper.h
@@ -221,7 +221,7 @@ public:
 	SANDBOXFILE_API virtual FString					ConvertToAbsolutePathForExternalAppForRead( const TCHAR* Filename ) override;
 	SANDBOXFILE_API virtual FString					ConvertToAbsolutePathForExternalAppForWrite( const TCHAR* Filename ) override;
 
-	SANDBOXFILE_API virtual IAsyncReadFileHandle*	OpenAsyncRead(const TCHAR* Filename) override;
+	SANDBOXFILE_API virtual IAsyncReadFileHandle*	OpenAsyncRead(const TCHAR* Filename, bool bAllowWrite = false) override;
 	SANDBOXFILE_API virtual void					SetAsyncMinimumPriority(EAsyncIOPriorityAndFlags Priority) override;
 	SANDBOXFILE_API virtual IMappedFileHandle*		OpenMapped(const TCHAR* Filename) override;
 
