Update comments for OpenAsyncRead

diff --git a/Engine/Source/Runtime/Core/Public/GenericPlatform/GenericPlatformFile.h b/Engine/Source/Runtime/Core/Public/GenericPlatform/GenericPlatformFile.h
--- a/Engine/Source/Runtime/Core/Public/GenericPlatform/GenericPlatformFile.h
+++ b/Engine/Source/Runtime/Core/Public/GenericPlatform/GenericPlatformFile.h
@@ -723,6 +723,7 @@ public:
 	/** Open a file for async reading. This call does not hit the disk or block.
 	*
 	* @param Filename file to be opened
+	* @param bAllowWrite (applies to certain platforms only) whether this file is allowed to be written to by other processes. This flag is needed to open files that are currently being written to as well.
 	* @return Close the file by delete'ing the handle. A non-null return value does not mean the file exists, since that may not be determined yet.
 	*/
 	CORE_API virtual IAsyncReadFileHandle* OpenAsyncRead(const TCHAR* Filename, bool bAllowWrite = false);
