Restrict MakeStringView to disallow non-character pointer arguments

diff --git a/Engine/Source/Runtime/Core/Public/Containers/StringView.h b/Engine/Source/Runtime/Core/Public/Containers/StringView.h
--- a/Engine/Source/Runtime/Core/Public/Containers/StringView.h
+++ b/Engine/Source/Runtime/Core/Public/Containers/StringView.h
@@ -492,16 +492,20 @@ template <typename CharRangeType>
 TStringView(CharRangeType&& Range)
 	-> TStringView<TElementType_T<CharRangeType>>;
 
-template <typename CharPtrOrRangeType>
-constexpr inline auto MakeStringView(CharPtrOrRangeType&& CharPtrOrRange UE_LIFETIMEBOUND) -> decltype(TStringView(Forward<CharPtrOrRangeType>(CharPtrOrRange)))
+template <typename CharPtrOrRangeType,
+	// This uses enable_if instead of concepts because it tries to instantiate TStringView<void> with concepts.
+	// This can use concepts once TStringView constrains its CharType to character types.
+	std::enable_if_t<!std::is_pointer_v<CharPtrOrRangeType> || TIsCharType_V<std::remove_pointer_t<CharPtrOrRangeType>>>* = nullptr>
+constexpr inline auto MakeStringView(const CharPtrOrRangeType& CharPtrOrRange UE_LIFETIMEBOUND) -> decltype(TStringView(CharPtrOrRange))
 {
-	return TStringView(Forward<CharPtrOrRangeType>(CharPtrOrRange));
+	return TStringView(CharPtrOrRange);
 }
 
-template <typename CharPtrType>
-constexpr inline auto MakeStringView(CharPtrType&& CharPtr UE_LIFETIMEBOUND, int32 Size) -> decltype(TStringView(Forward<CharPtrType>(CharPtr), Size))
+template <typename CharType
+	UE_REQUIRES(TIsCharType_V<CharType>)>
+constexpr inline auto MakeStringView(const CharType* CharPtr UE_LIFETIMEBOUND, int32 Size) -> decltype(TStringView(CharPtr, Size))
 {
-	return TStringView(Forward<CharPtrType>(CharPtr), Size);
+	return TStringView(CharPtr, Size);
 }
 
 //////////////////////////////////////////////////////////////////////////
