Add trace event for unresolved I/O store requests

#rnx

diff --git a/Engine/Plugins/IoStoreInsights/Source/IoStoreInsights/Private/Analyzers/IoStoreInsightsAnalyzer.cpp b/Engine/Plugins/IoStoreInsights/Source/IoStoreInsights/Private/Analyzers/IoStoreInsightsAnalyzer.cpp
--- a/Engine/Plugins/IoStoreInsights/Source/IoStoreInsights/Private/Analyzers/IoStoreInsightsAnalyzer.cpp
+++ b/Engine/Plugins/IoStoreInsights/Source/IoStoreInsights/Private/Analyzers/IoStoreInsightsAnalyzer.cpp
@@ -78,6 +78,27 @@ namespace UE::IoStoreInsights
 			}
 			break;
 
+			case RouteId_RequestUnresolved:
+			{
+				const uint64 Cycle			= EventData.GetValue<uint64>("Cycle");
+				const double Time			= Context.EventTime.AsSeconds(Cycle);
+
+				Session.UpdateDurationSeconds(Time);
+
+				const uint64 RequestHandle	= EventData.GetValue<uint64>("RequestHandle");
+				const uint32 ThreadId		= Context.ThreadInfo.GetId();
+				const uint64 Size			= 0;
+				const bool bFailed			= true;
+
+				if (const FPendingActivity* FindRead = ActiveReadsMap.Find(RequestHandle))
+				{
+					Provider.EndIoStoreActivity(FindRead->IoStoreRequestIndex, FindRead->ActivityIndex, Size, bFailed, Time);
+					ActiveReadsMap.Remove(RequestHandle);
+				}
+				ActiveRequestsMap.Remove(RequestHandle); // note: we don't have a separate "request destroy", so clean everything up once it has been read
+			}
+			break;
+
 			case RouteId_RequestStarted:
 			{
 				const uint64 Cycle = EventData.GetValue<uint64>("Cycle");
diff --git a/Engine/Plugins/IoStoreInsights/Source/IoStoreInsights/Private/Analyzers/IoStoreInsightsAnalyzer.h b/Engine/Plugins/IoStoreInsights/Source/IoStoreInsights/Private/Analyzers/IoStoreInsightsAnalyzer.h
--- a/Engine/Plugins/IoStoreInsights/Source/IoStoreInsights/Private/Analyzers/IoStoreInsightsAnalyzer.h
+++ b/Engine/Plugins/IoStoreInsights/Source/IoStoreInsights/Private/Analyzers/IoStoreInsightsAnalyzer.h
@@ -33,6 +33,7 @@ namespace UE::IoStoreInsights
 			RouteId_RequestCompleted,
 			RouteId_RequestFailed,
 			RouteId_PackageMapping,
+			RouteId_RequestUnresolved,
 		};
 
 		struct FPendingActivity
diff --git a/Engine/Source/Runtime/Core/Private/IO/IoDispatcher.cpp b/Engine/Source/Runtime/Core/Private/IO/IoDispatcher.cpp
--- a/Engine/Source/Runtime/Core/Private/IO/IoDispatcher.cpp
+++ b/Engine/Source/Runtime/Core/Private/IO/IoDispatcher.cpp
@@ -838,6 +838,7 @@ private:
 				{
 					check(Request->NextRequest == nullptr);
 					Request->Backend = nullptr;
+					TRACE_IOSTORE_REQUEST_UNRESOLVED(Request);
 					CompleteRequest(Request, EIoErrorCode::NotFound);
 					Request->ReleaseRef();
 					UnresolvedCount++;
diff --git a/Engine/Source/Runtime/Core/Private/ProfilingDebugging/IoStoreTrace.cpp b/Engine/Source/Runtime/Core/Private/ProfilingDebugging/IoStoreTrace.cpp
--- a/Engine/Source/Runtime/Core/Private/ProfilingDebugging/IoStoreTrace.cpp
+++ b/Engine/Source/Runtime/Core/Private/ProfilingDebugging/IoStoreTrace.cpp
@@ -34,6 +34,11 @@ UE_TRACE_EVENT_BEGIN(IoStore, RequestCreate)
 	UE_TRACE_EVENT_FIELD(uint64, Size)
 UE_TRACE_EVENT_END()
 
+UE_TRACE_EVENT_BEGIN(IoStore, RequestUnresolved)
+	UE_TRACE_EVENT_FIELD(uint64, Cycle)
+	UE_TRACE_EVENT_FIELD(uint64, RequestHandle)
+UE_TRACE_EVENT_END()
+
 UE_TRACE_EVENT_BEGIN(IoStore, RequestStarted)
 	UE_TRACE_EVENT_FIELD(uint64, Cycle)
 	UE_TRACE_EVENT_FIELD(uint64, RequestHandle)
@@ -83,6 +88,14 @@ void FIoStoreTrace::RequestCreate(FIoBatch* IoBatch, FIoRequestImpl* IoRequestIm
 		;
 }
 
+void FIoStoreTrace::RequestUnresolved(FIoRequestImpl* IoRequestImpl)
+{
+	UE_TRACE_LOG(IoStore, RequestUnresolved, IoStoreChannel)
+		<< RequestUnresolved.Cycle(FPlatformTime::Cycles64())
+		<< RequestUnresolved.RequestHandle((uint64)IoRequestImpl)
+		;
+}
+
 void FIoStoreTrace::RequestStarted(FIoRequestImpl* IoRequestImpl, IIoDispatcherBackend* IoDispatcherBackend)
 {
 	UE_TRACE_LOG(IoStore, RequestStarted, IoStoreChannel)
diff --git a/Engine/Source/Runtime/Core/Public/ProfilingDebugging/IoStoreTrace.h b/Engine/Source/Runtime/Core/Public/ProfilingDebugging/IoStoreTrace.h
--- a/Engine/Source/Runtime/Core/Public/ProfilingDebugging/IoStoreTrace.h
+++ b/Engine/Source/Runtime/Core/Public/ProfilingDebugging/IoStoreTrace.h
@@ -29,6 +29,7 @@ struct FIoStoreTrace
 {
 	static CORE_API void BackendName(IIoDispatcherBackend* IoDispatcherBackend, const TCHAR* Name);
 	static CORE_API void RequestCreate(FIoBatch* IoBatch, FIoRequestImpl* IoRequestImpl);
+	static CORE_API void RequestUnresolved(FIoRequestImpl* IoRequestImpl);
 	static CORE_API void RequestStarted(FIoRequestImpl* IoRequestImpl, IIoDispatcherBackend* IoDispatcherBackend);
 	static CORE_API void RequestCompleted(FIoRequestImpl* IoRequestImpl, uint64 Size);
 	static CORE_API void RequestFailed(FIoRequestImpl* IoRequestImpl);
@@ -42,6 +43,9 @@ UE_TRACE_CHANNEL_EXTERN(IoStoreChannel, CORE_API);
 #define TRACE_IOSTORE_REQUEST_CREATE(IoBatch,IoRequestImpl) \
 	FIoStoreTrace::RequestCreate(IoBatch,IoRequestImpl);
 
+#define TRACE_IOSTORE_REQUEST_UNRESOLVED(IoRequestImpl) \
+	FIoStoreTrace::RequestUnresolved(IoRequestImpl);
+
 #define TRACE_IOSTORE_BACKEND_REQUEST_STARTED(IoRequestImpl, IoDispatcherBackend) \
 	FIoStoreTrace::RequestStarted(IoRequestImpl, IoDispatcherBackend);
 
@@ -69,6 +73,7 @@ UE_TRACE_CHANNEL_EXTERN(IoStoreChannel, CORE_API);
 
 #define TRACE_IOSTORE_BACKEND_NAME(...)
 #define TRACE_IOSTORE_REQUEST_CREATE(...)
+#define TRACE_IOSTORE_REQUEST_UNRESOLVED(...)
 #define TRACE_IOSTORE_BACKEND_REQUEST_STARTED(...)
 #define TRACE_IOSTORE_BACKEND_REQUEST_COMPLETED(...)
 #define TRACE_IOSTORE_BACKEND_REQUEST_FAILED(...)
