Use the full path for GPU crash dumps when copying them to the crash output folder.

Also move instead of copying them, so that old files are not re-uploaded for future crashes.

diff --git a/Engine/Source/Runtime/Core/Private/Windows/WindowsPlatformCrashContext.cpp b/Engine/Source/Runtime/Core/Private/Windows/WindowsPlatformCrashContext.cpp
--- a/Engine/Source/Runtime/Core/Private/Windows/WindowsPlatformCrashContext.cpp
+++ b/Engine/Source/Runtime/Core/Private/Windows/WindowsPlatformCrashContext.cpp
@@ -286,6 +286,11 @@ void FWindowsPlatformCrashContext::CopyPlatformSpecificFiles(const TCHAR* Output
 	// Best effort, so don't care about result: couldn't copy -> tough, no GPU crash dump
 	TArray<FString> GPUDumpFiles;
 	IFileManager::Get().FindFiles(GPUDumpFiles, *FPaths::ProjectLogDir(), NvAftermathDumpExtension);
+	for (FString& GPUDumpFilename : GPUDumpFiles)
+	{
+		GPUDumpFilename = FPaths::Combine(*FPaths::ProjectLogDir(), *GPUDumpFilename);
+	}
+
 	if (GPUDumpFiles.Num())
 	{
 		GPUDumpFiles.Sort([](FString LHS, FString RHS)
@@ -301,7 +306,10 @@ void FWindowsPlatformCrashContext::CopyPlatformSpecificFiles(const TCHAR* Output
 		{
 			FString GPUMiniDumpFilename = FPaths::GetCleanFilename(SelectedGPUCrashDump);
 			const FString GPUMiniDumpDstAbsolute = FPaths::Combine(OutputDirectory, *GPUMiniDumpFilename);
-			static_cast<void>(IFileManager::Get().Copy(*GPUMiniDumpDstAbsolute, *SelectedGPUCrashDump));
+			if (!IFileManager::Get().Move(*GPUMiniDumpDstAbsolute, *SelectedGPUCrashDump))
+			{
+				UE_LOG(LogWindows, Error, TEXT("Error moving GPU crash dump file %s to output crash directory"), *GPUMiniDumpFilename);
+			}
 		}
 	}
 }
