Cap the size of the hitset to 128MB to prevent huge arrays of 1/2 byte values swamping the set.

diff --git a/Engine/Source/Runtime/Core/Private/AutoRTFM/HitSet.h b/Engine/Source/Runtime/Core/Private/AutoRTFM/HitSet.h
--- a/Engine/Source/Runtime/Core/Private/AutoRTFM/HitSet.h
+++ b/Engine/Source/Runtime/Core/Private/AutoRTFM/HitSet.h
@@ -9,18 +9,23 @@ namespace AutoRTFM
 {
 	class FHitSet final
 	{
+	public:
+		using Key = TTaggedPtr<void>;
+
+	private:
 		// TODO: Revisit a good probe depth for the hashset.
 		static constexpr uint64_t LinearProbeDepth = 16;
 
 		// TODO: Revisit a good initial capacity for the hitset.
 		static constexpr uint64_t LogInitialCapacity = 4;
 
-	public:
-		using Key = TTaggedPtr<void>;
+		// TODO: Revisit a good max capacity for the hitset.
+		static constexpr uint64_t MaxCapacity = (128 * 1024 * 1024) / sizeof(Key);
 
+	public:
 		explicit FHitSet()
 		{
-			constexpr uint64_t InitialCapacity = 1u << LogInitialCapacity;
+			constexpr uint64_t InitialCapacity = static_cast<uint64_t>(1) << LogInitialCapacity;
 
 			// Do not want to deal with null payloads.
 			static_assert(0 != InitialCapacity);
@@ -40,9 +45,10 @@ namespace AutoRTFM
 			FMemory::Free(Payload);
 		}
 
-		// Insert something in the HitSet, returning true if the put succeeded
-		// (EG. the key was not already in the set).
-		bool Insert(Key K)
+		// Insert something in the HitSet, returning true if we found the key.
+		// If false is returned, we could have inserted or not - depending on
+		// whether the HitSet has reached its max capacity.
+		bool FindOrTryInsert(Key K)
 		{
 			while (true)
 			{
@@ -51,11 +57,14 @@ namespace AutoRTFM
 				default:
 					AutoRTFM::Unreachable();
 				case EInsertResult::Inserted:
-					return true;
-				case EInsertResult::Exists:
 					return false;
+				case EInsertResult::Exists:
+					return true;
 				case EInsertResult::NotInserted:
-					Resize();
+					if (!Resize())
+					{
+						return false;
+					}
 					continue;
 				}
 			}
@@ -105,15 +114,21 @@ namespace AutoRTFM
 			ASSERT(0 != SixtyFourMinusLogCapacity);
 		}
 
-		void Resize()
+		bool Resize()
 		{
-			uintptr_t* const OldPayload = Payload;
 			const uint64_t OldCapacity = Capacity();
+
+			if (OldCapacity >= MaxCapacity)
+			{
+				return false;
+			}
+
+			uintptr_t* const OldPayload = Payload;
 			const uint64_t OldSize = Size;
 
 			while (true)
 			{
-				bool NeedAnotherResize = false;
+				bool bNeedAnotherResize = false;
 
 				Size = 0;
 
@@ -122,8 +137,7 @@ namespace AutoRTFM
 				Payload = static_cast<uintptr_t*>(FMemory::MallocZeroed(Capacity() * sizeof(uintptr_t)));
 				ASSERT(nullptr != Payload);
 
-				// Now we need to rehash and reinsert all the items. We need to
-				// remember about the extra linear probe section here too.
+				// Now we need to rehash and reinsert all the items.
 				for (size_t I = 0; I < OldCapacity; I++)
 				{
 					// Skip empty locations.
@@ -136,7 +150,7 @@ namespace AutoRTFM
 
 					if (EInsertResult::NotInserted == Result)
 					{
-						NeedAnotherResize = true;
+						bNeedAnotherResize = true;
 						break;
 					}
 					else
@@ -145,7 +159,7 @@ namespace AutoRTFM
 					}
 				}
 
-				if (NeedAnotherResize)
+				if (bNeedAnotherResize)
 				{
 					FMemory::Free(Payload);
 					continue;
@@ -158,7 +172,7 @@ namespace AutoRTFM
 
 			FMemory::Free(OldPayload);
 
-			return;
+			return true;
 		}
 
 		enum class EInsertResult : uint8_t
diff --git a/Engine/Source/Runtime/Core/Private/AutoRTFM/Transaction.cpp b/Engine/Source/Runtime/Core/Private/AutoRTFM/Transaction.cpp
--- a/Engine/Source/Runtime/Core/Private/AutoRTFM/Transaction.cpp
+++ b/Engine/Source/Runtime/Core/Private/AutoRTFM/Transaction.cpp
@@ -160,7 +160,7 @@ void FTransaction::CommitNested()
 		FHitSet::Key HitSetEntry(Write.GetOriginal());
 		HitSetEntry.SetTopTag(static_cast<uint16_t>(Write.GetSize()));
 
-        Parent->HitSet.Insert(HitSetEntry);
+        Parent->HitSet.FindOrTryInsert(HitSetEntry);
     }
 
     Parent->WriteLogBumpAllocator.Merge(MoveTemp(WriteLogBumpAllocator));
diff --git a/Engine/Source/Runtime/Core/Private/AutoRTFM/TransactionInlines.h b/Engine/Source/Runtime/Core/Private/AutoRTFM/TransactionInlines.h
--- a/Engine/Source/Runtime/Core/Private/AutoRTFM/TransactionInlines.h
+++ b/Engine/Source/Runtime/Core/Private/AutoRTFM/TransactionInlines.h
@@ -92,7 +92,7 @@ AUTORTFM_NO_ASAN UE_AUTORTFM_FORCEINLINE void FTransaction::RecordWrite(void* Lo
         FMemoryLocation Key(LogicalAddress);
         Key.SetTopTag(static_cast<uint16_t>(Size));
 
-        if (!HitSet.Insert(Key))
+        if (HitSet.FindOrTryInsert(Key))
         {
             Stats.Collect<EStatsKind::HitSetHit>();
             return;
@@ -135,7 +135,7 @@ template<unsigned SIZE> AUTORTFM_NO_ASAN UE_AUTORTFM_FORCEINLINE void FTransacti
     FMemoryLocation Key(LogicalAddress);
     Key.SetTopTag(static_cast<uint16_t>(SIZE));
 
-    if (!HitSet.Insert(Key))
+    if (HitSet.FindOrTryInsert(Key))
     {
         Stats.Collect<EStatsKind::HitSetHit>();
         return;
