[AutoRTFM] Allow OnCommit and OnAbort handlers to use mutable lambdas.

This allows us to capture elements by value and mutate the captured value, or pass it to a function which requires a non-const argument. This can be perilous if the callback can be invoked repeatedly, but we only invoke an individual OnCommit or OnAbort handler once apiece, so it is safe.

Note that it isn't possible to do something similar for the macro forms, as the macro forms discard the capture list entirely when AutoRTFM is off.

diff --git a/Engine/Source/Programs/AutoRTFMTests/Private/API.cpp b/Engine/Source/Programs/AutoRTFMTests/Private/API.cpp
--- a/Engine/Source/Programs/AutoRTFMTests/Private/API.cpp
+++ b/Engine/Source/Programs/AutoRTFMTests/Private/API.cpp
@@ -877,6 +877,20 @@ TEST_CASE("API.OnCommit")
     REQUIRE(AutoRTFM::ETransactionResult::AbortedByRequest == NestResult);
 }
 
+TEST_CASE("API.OnCommit_MutableCapture")
+{
+	FString Message = "Hello";
+
+	AutoRTFM::Transact([&]
+	{
+		AutoRTFM::OnCommit([MessageCopy = Message]() mutable
+		{
+			MessageCopy += " World!";
+			REQUIRE(MessageCopy == "Hello World!");
+		});
+	});
+}
+
 TEST_CASE("API.OnCommitMacro_NoAbort")
 {
 	int Value = 123;
@@ -1002,6 +1016,22 @@ TEST_CASE("API.OnAbort")
     REQUIRE(AutoRTFM::ETransactionResult::AbortedByRequest == NestResult);
 }
 
+TEST_CASE("API.OnAbort_MutableCapture")
+{
+	FString Message = "Hello";
+
+	AutoRTFM::Transact([&]
+	{
+		AutoRTFM::OnAbort([MessageCopy = Message]() mutable
+		{
+			MessageCopy += " World!";
+			REQUIRE(MessageCopy == "Hello World!");
+		});
+
+		AutoRTFM::AbortTransaction();
+	});
+}
+
 TEST_CASE("API.OnAbortMacro_NoAbort")
 {
 	int Value = 123;
diff --git a/Engine/Source/Programs/AutoRTFMTests/Private/Abort.cpp b/Engine/Source/Programs/AutoRTFMTests/Private/Abort.cpp
--- a/Engine/Source/Programs/AutoRTFMTests/Private/Abort.cpp
+++ b/Engine/Source/Programs/AutoRTFMTests/Private/Abort.cpp
@@ -292,6 +292,22 @@ TEST_CASE("Abort.PushOnAbortHandler_WithAbort")
 	REQUIRE(Value == 77);
 }
 
+TEST_CASE("Abort.PushOnAbortHandler_MutableCapture")
+{
+	FString Message = "Hello";
+
+	AutoRTFM::Transact([&]
+	{
+		AutoRTFM::PushOnAbortHandler(UIntToPointer(12345), [MessageCopy = Message]() mutable
+		{
+			MessageCopy += " World!";
+			REQUIRE(MessageCopy == "Hello World!");
+		});
+
+		AutoRTFM::AbortTransaction();
+	});
+}
+
 TEST_CASE("Abort.PushOnAbortHandler_WithPop_NoAbort")
 {
 	int Value = 55;
diff --git a/Engine/Source/Runtime/Core/Public/AutoRTFM/AutoRTFM.h b/Engine/Source/Runtime/Core/Public/AutoRTFM/AutoRTFM.h
--- a/Engine/Source/Runtime/Core/Public/AutoRTFM/AutoRTFM.h
+++ b/Engine/Source/Runtime/Core/Public/AutoRTFM/AutoRTFM.h
@@ -729,11 +729,11 @@ template<typename TFunctor> [[nodiscard]] static UE_AUTORTFM_FORCEINLINE EContex
 // this just adds the work to the work deferred until the outer nest's commit.
 // If this is called outside a transaction or from an open nest then the work
 // happens immediately.
-template<typename TFunctor> static UE_AUTORTFM_FORCEINLINE void OnCommit(const TFunctor& Work)
+template<typename TFunctor> static UE_AUTORTFM_FORCEINLINE void OnCommit(TFunctor&& Work)
 {
 	if (autortfm_is_closed())
 	{
-		ForTheRuntime::OnCommitInternal(Work);
+		ForTheRuntime::OnCommitInternal(std::forward<TFunctor>(Work));
 	}
 	else
 	{
@@ -745,17 +745,17 @@ template<typename TFunctor> static UE_AUTORTFM_FORCEINLINE void OnCommit(const T
 // this just adds the work to the work deferred until the outer nest's commit.
 // If this is called outside a transaction or from an open nest then the work
 // happens immediately.
-template<typename TFunctor> static UE_AUTORTFM_FORCEINLINE void OnCommit(const TFunctor& Work) { Work(); }
+template<typename TFunctor> static UE_AUTORTFM_FORCEINLINE void OnCommit(TFunctor&& Work) { Work(); }
 #endif
 
 #if UE_AUTORTFM
 // Have some work happen when this transaction aborts. If this is called
 // outside a transaction or from an open nest then the work is ignored.
-template<typename TFunctor> static UE_AUTORTFM_FORCEINLINE void OnAbort(const TFunctor& Work)
+template<typename TFunctor> static UE_AUTORTFM_FORCEINLINE void OnAbort(TFunctor&& Work)
 {
 	if (autortfm_is_closed())
 	{
-		ForTheRuntime::OnAbortInternal(Work);
+		ForTheRuntime::OnAbortInternal(std::forward<TFunctor>(Work));
 	}
 }
 
@@ -763,11 +763,11 @@ template<typename TFunctor> static UE_AUTORTFM_FORCEINLINE void OnAbort(const TF
 // the handler can be unregistered (see PopOnAbortHandler). This is useful
 // for scoped mutations that need an abort handler present unless execution
 // reaches the end of the relevant scope.
-template<typename TFunctor> static UE_AUTORTFM_FORCEINLINE void PushOnAbortHandler(const void* Key, const TFunctor& Work)
+template<typename TFunctor> static UE_AUTORTFM_FORCEINLINE void PushOnAbortHandler(const void* Key, TFunctor&& Work)
 {
 	if (autortfm_is_closed())
 	{
-		ForTheRuntime::PushOnAbortHandlerInternal(Key, Work);
+		ForTheRuntime::PushOnAbortHandlerInternal(Key, std::forward<TFunctor>(Work));
 	}
 }
 
@@ -791,13 +791,13 @@ static UE_AUTORTFM_FORCEINLINE void PopAllOnAbortHandlers(const void* Key)
 #else
 // Have some work happen when this transaction aborts. If this is called
 // outside a transaction or from an open nest then the work is ignored.
-template<typename TFunctor> static UE_AUTORTFM_FORCEINLINE void OnAbort(const TFunctor&) {}
+template<typename TFunctor> static UE_AUTORTFM_FORCEINLINE void OnAbort(TFunctor&&) {}
 
 // Register a handler for transaction abort. Takes a key parameter so that
 // the handler can be unregistered (see PopOnAbortHandler). This is useful
 // for scoped mutations that need an abort handler present unless execution
 // reaches the end of the relevant scope.
-template<typename TFunctor> static UE_AUTORTFM_FORCEINLINE void PushOnAbortHandler(const void* Key, const TFunctor&) {}
+template<typename TFunctor> static UE_AUTORTFM_FORCEINLINE void PushOnAbortHandler(const void* Key, TFunctor&&) {}
 
 // Unregister the most recently pushed handler (via PushOnAbortHandler) for the given key.
 static UE_AUTORTFM_FORCEINLINE void PopOnAbortHandler(const void* Key) {}
