Added FCulture::GetVerseIdentifier, FCulture::CultureNameToVerseIdentifier, and FCulture::CultureNameFromVerseIdentifier

Verse identifiers cannot contain a dash, so will need to use "en_US" style culture names (like ICU) rather than the "en-US" style culture names that UE uses internally

UE will continue to use "en-US" style internally, but plans to transition to using "en_US" style for content stored on disk (which may need to be a valid Verse path)

diff --git a/Engine/Source/Runtime/Core/Private/Internationalization/Culture.cpp b/Engine/Source/Runtime/Core/Private/Internationalization/Culture.cpp
--- a/Engine/Source/Runtime/Core/Private/Internationalization/Culture.cpp
+++ b/Engine/Source/Runtime/Core/Private/Internationalization/Culture.cpp
@@ -95,6 +95,7 @@ FCultureRef FCulture::Create(TUniquePtr<FCultureImplementation>&& InImplementati
 FCulture::FCulture(TUniquePtr<FCultureImplementation>&& InImplementation)
 	: Implementation(MoveTemp(InImplementation))
 	, CachedName(Implementation->GetName())
+	, CachedVerseIdentifier(FCulture::CultureNameToVerseIdentifier(CachedName))
 	, CachedUnrealLegacyThreeLetterISOLanguageName(Implementation->GetUnrealLegacyThreeLetterISOLanguageName())
 	, CachedThreeLetterISOLanguageName(Implementation->GetThreeLetterISOLanguageName())
 	, CachedTwoLetterISOLanguageName(Implementation->GetTwoLetterISOLanguageName())
@@ -202,11 +203,32 @@ FString FCulture::GetCanonicalName(const FString& Name)
 	return FCultureImplementation::GetCanonicalName(Name, FInternationalization::Get());
 }
 
+FString FCulture::CultureNameToVerseIdentifier(const FString& Name)
+{
+	// Verse uses "en_US" style while UE uses "en-US"
+	FString Result = Name;
+	Result.ReplaceCharInline(TEXT('-'), TEXT('_'), ESearchCase::CaseSensitive);
+	return Result;
+}
+
+FString FCulture::CultureNameFromVerseIdentifier(const FString& VerseIdentifer)
+{
+	// UE uses "en-US" style while Verse uses "en_US"
+	FString Result = VerseIdentifer;
+	Result.ReplaceCharInline(TEXT('_'), TEXT('-'), ESearchCase::CaseSensitive);
+	return Result;
+}
+
 const FString& FCulture::GetName() const
 {
 	return CachedName;
 }
 
+const FString& FCulture::GetVerseIdentifier() const
+{
+	return CachedVerseIdentifier;
+}
+
 const FString& FCulture::GetNativeName() const
 {
 	return CachedNativeName;
diff --git a/Engine/Source/Runtime/Core/Private/Internationalization/ICUCulture.cpp b/Engine/Source/Runtime/Core/Private/Internationalization/ICUCulture.cpp
--- a/Engine/Source/Runtime/Core/Private/Internationalization/ICUCulture.cpp
+++ b/Engine/Source/Runtime/Core/Private/Internationalization/ICUCulture.cpp
@@ -252,8 +252,9 @@ FString FICUCultureImplementation::GetCanonicalName(const FString& Name, FIntern
 
 FString FICUCultureImplementation::GetName() const
 {
+	// UE uses "en-US" style while ICU uses "en_US"
 	FString Result = ICULocale.getName();
-	Result.ReplaceInline(TEXT("_"), TEXT("-"), ESearchCase::IgnoreCase);
+	Result.ReplaceCharInline(TEXT('_'), TEXT('-'), ESearchCase::CaseSensitive);
 	return Result;
 }
 
diff --git a/Engine/Source/Runtime/Core/Private/Internationalization/InternationalizationUtilities.cpp b/Engine/Source/Runtime/Core/Private/Internationalization/InternationalizationUtilities.cpp
--- a/Engine/Source/Runtime/Core/Private/Internationalization/InternationalizationUtilities.cpp
+++ b/Engine/Source/Runtime/Core/Private/Internationalization/InternationalizationUtilities.cpp
@@ -253,7 +253,7 @@ namespace InternationalizationUtilities
 				);
 
 				NameTag = SanitizedName.Left(NameTagEndIndex);
-				NameTag.ReplaceInline(TEXT("_"), TEXT("-"), ESearchCase::CaseSensitive);
+				NameTag.ReplaceCharInline(TEXT('_'), TEXT('-'), ESearchCase::CaseSensitive);
 
 				if (NameKeywordsSplitIndex != INDEX_NONE)
 				{
diff --git a/Engine/Source/Runtime/Core/Public/Internationalization/Culture.h b/Engine/Source/Runtime/Core/Public/Internationalization/Culture.h
--- a/Engine/Source/Runtime/Core/Public/Internationalization/Culture.h
+++ b/Engine/Source/Runtime/Core/Public/Internationalization/Culture.h
@@ -49,8 +49,27 @@ public:
 
 	static CORE_API FString GetCanonicalName(const FString& Name);
 
+	/**
+	 * Given a well-formed UE culture name (eg, "en-US"), convert it to its equivalent Verse identifier (eg, "en_US").
+	 */
+	static CORE_API FString CultureNameToVerseIdentifier(const FString& Name);
+
+	/**
+	 * Given a well-formed Verse culture identifier (eg, "en_US"), convert it to its equivalent UE culture name (eg, "en-US").
+	 * @note GetCanonicalName will also perform this conversion, with additional validation.
+	 */
+	static CORE_API FString CultureNameFromVerseIdentifier(const FString& VerseIdentifer);
+
+	/**
+	 * Get the name of this culture, eg, "en-US"
+	 */
 	CORE_API const FString& GetName() const;
 	
+	/**
+	 * Get the name of this culture as a Verse identifier (for a Verse path), eg, "en_US"
+	 */
+	CORE_API const FString& GetVerseIdentifier() const;
+
 	CORE_API const FString& GetNativeName() const;
 
 	CORE_API const FString& GetUnrealLegacyThreeLetterISOLanguageName() const;
@@ -109,6 +128,7 @@ private:
 	FString CachedDisplayName;
 	FString CachedEnglishName;
 	FString CachedName;
+	FString CachedVerseIdentifier;
 	FString CachedNativeName;
 	FString CachedUnrealLegacyThreeLetterISOLanguageName;
 	FString CachedThreeLetterISOLanguageName;
