Add a TransactionallySafeMutex which can be used with UE::TScopeLocks

diff --git a/Engine/Source/Programs/AutoRTFMTests/Private/UECoreTests.cpp b/Engine/Source/Programs/AutoRTFMTests/Private/UECoreTests.cpp
--- a/Engine/Source/Programs/AutoRTFMTests/Private/UECoreTests.cpp
+++ b/Engine/Source/Programs/AutoRTFMTests/Private/UECoreTests.cpp
@@ -14,6 +14,7 @@
 #include "MyAutoRTFMTestObject.h"
 #include "Misc/TransactionallySafeScopeLock.h"
 #include "Misc/TransactionallySafeRWScopeLock.h"
+#include "Misc/TransactionallySafeMutex.h"
 #include "Containers/Queue.h"
 #include "Misc/ConfigCacheIni.h"
 
@@ -1477,6 +1478,108 @@ TEST_CASE("UECore.TransactionallySafeWriteScopeLock")
 	}
 }
 
+TEST_CASE("UECore.TransactionallySafeMutex")
+{
+	SECTION("Outside Transaction")
+	{
+		FTransactionallySafeMutex Mutex;
+
+		AutoRTFM::ETransactionResult Result = AutoRTFM::Transact([&]
+			{
+				UE::TScopeLock Lock(Mutex);
+				AutoRTFM::AbortTransaction();
+			});
+
+		REQUIRE(AutoRTFM::ETransactionResult::AbortedByRequest == Result);
+
+		Result = AutoRTFM::Transact([&]
+			{
+				UE::TScopeLock Lock(Mutex);
+			});
+
+		REQUIRE(AutoRTFM::ETransactionResult::Committed == Result);
+	}
+
+	SECTION("Inside Transaction")
+	{
+		AutoRTFM::ETransactionResult Result = AutoRTFM::Transact([&]
+			{
+				FTransactionallySafeMutex Mutex;
+				UE::TScopeLock Lock(Mutex);
+				AutoRTFM::AbortTransaction();
+			});
+
+		REQUIRE(AutoRTFM::ETransactionResult::AbortedByRequest == Result);
+
+		Result = AutoRTFM::Transact([&]
+			{
+				FTransactionallySafeMutex Mutex;
+				UE::TScopeLock Lock(Mutex);
+			});
+
+		REQUIRE(AutoRTFM::ETransactionResult::Committed == Result);
+	}
+
+	SECTION("Inside Transaction Used In Nested Transaction")
+	{
+		AutoRTFM::ETransactionResult InnerResult = AutoRTFM::ETransactionResult::Committed;
+
+		AutoRTFM::ETransactionResult Result = AutoRTFM::Transact([&]
+			{
+				FTransactionallySafeMutex Mutex;
+
+				InnerResult = AutoRTFM::Transact([&]
+					{
+						UE::TScopeLock Lock(Mutex);
+						AutoRTFM::CascadingAbortTransaction();
+					});
+			});
+
+		REQUIRE(AutoRTFM::ETransactionResult::AbortedByCascade == Result);
+
+		Result = AutoRTFM::Transact([&]
+			{
+				FTransactionallySafeMutex Mutex;
+
+				InnerResult = AutoRTFM::Transact([&]
+					{
+						UE::TScopeLock Lock(Mutex);
+						AutoRTFM::AbortTransaction();
+					});
+			});
+
+		REQUIRE(AutoRTFM::ETransactionResult::AbortedByRequest == InnerResult);
+		REQUIRE(AutoRTFM::ETransactionResult::Committed == Result);
+
+		Result = AutoRTFM::Transact([&]
+			{
+				FTransactionallySafeMutex Mutex;
+
+				InnerResult = AutoRTFM::Transact([&]
+					{
+						UE::TScopeLock Lock(Mutex);
+					});
+
+				AutoRTFM::AbortTransaction();
+			});
+
+		REQUIRE(AutoRTFM::ETransactionResult::AbortedByRequest == Result);
+
+		Result = AutoRTFM::Transact([&]
+			{
+				FTransactionallySafeMutex Mutex;
+
+				InnerResult = AutoRTFM::Transact([&]
+					{
+						UE::TScopeLock Lock(Mutex);
+					});
+			});
+
+		REQUIRE(AutoRTFM::ETransactionResult::Committed == InnerResult);
+		REQUIRE(AutoRTFM::ETransactionResult::Committed == Result);
+	}
+}
+
 TEST_CASE("UECore.FTextFormatPatternDefinition")
 {
 	FTextFormatPatternDefinitionConstPtr Ptr;
diff --git a/Engine/Source/Runtime/Core/Private/Internationalization/TextFormatter.cpp b/Engine/Source/Runtime/Core/Private/Internationalization/TextFormatter.cpp
--- a/Engine/Source/Runtime/Core/Private/Internationalization/TextFormatter.cpp
+++ b/Engine/Source/Runtime/Core/Private/Internationalization/TextFormatter.cpp
@@ -29,6 +29,7 @@
 #include "Templates/ValueOrError.h"
 #include "Trace/Detail/Channel.h"
 #include "AutoRTFM/AutoRTFM.h"
+#include "Misc/TransactionallySafeMutex.h"
 #include "Async/Mutex.h"
 
 #define LOCTEXT_NAMESPACE "TextFormatter"
@@ -381,14 +382,8 @@ public:
 	 */
 	FORCEINLINE bool IsValid() const
 	{
-		bool bIsValid = false;
-		UE_AUTORTFM_OPEN
-		{
-			UE::TScopeLock Lock(CompiledDataMutex);
-			bIsValid = IsValid_NoLock();
-		};
-
-		return bIsValid;
+		UE::TScopeLock Lock(CompiledDataMutex);
+		return IsValid_NoLock();
 	}
 
 	/**
@@ -402,14 +397,8 @@ public:
 	 */
 	FORCEINLINE bool ValidatePattern(const FCulturePtr& InCulture, TArray<FString>& OutValidationErrors)
 	{
-		bool bIsValid = false;
-		UE_AUTORTFM_OPEN
-		{
-			UE::TScopeLock Lock(CompiledDataMutex);
-			bIsValid = ValidatePattern_NoLock(InCulture, OutValidationErrors);
-		};
-
-		return bIsValid;
+		UE::TScopeLock Lock(CompiledDataMutex);
+		return ValidatePattern_NoLock(InCulture, OutValidationErrors);
 	}
 
 	/**
@@ -417,14 +406,8 @@ public:
 	 */
 	FORCEINLINE FString Format(const FPrivateTextFormatArguments& InFormatArgs)
 	{
-		FString RetString;
-		UE_AUTORTFM_OPEN
-		{
-			UE::TScopeLock Lock(CompiledDataMutex);
-			RetString = Format_NoLock(InFormatArgs);
-		};
-
-		return RetString;
+		UE::TScopeLock Lock(CompiledDataMutex);
+		return Format_NoLock(InFormatArgs);
 	}
 
 	/**
@@ -432,11 +415,8 @@ public:
 	 */
 	FORCEINLINE void GetFormatArgumentNames(TArray<FString>& OutArgumentNames)
 	{
-		UE_AUTORTFM_OPEN
-		{
-			UE::TScopeLock Lock(CompiledDataMutex);
-			GetFormatArgumentNames_NoLock(OutArgumentNames);
-		};
+		UE::TScopeLock Lock(CompiledDataMutex);
+		GetFormatArgumentNames_NoLock(OutArgumentNames);
 	}
 
 	/**
@@ -462,14 +442,8 @@ public:
 	 */
 	FORCEINLINE FTextFormat::EExpressionType GetExpressionType() const
 	{
-		FTextFormat::EExpressionType RetType = FTextFormat::EExpressionType::Invalid;
-		UE_AUTORTFM_OPEN
-		{
-			UE::TScopeLock Lock(CompiledDataMutex);
-			RetType = CompiledExpressionType;
-		};
-
-		return RetType;
+		UE::TScopeLock Lock(CompiledDataMutex);
+		return CompiledExpressionType;
 	}
 
 	/**
@@ -554,7 +528,7 @@ private:
 	/**
 	 * Mutex protecting the compiled data from being modified concurrently.
 	 */
-	mutable UE::FMutex CompiledDataMutex;
+	mutable FTransactionallySafeMutex CompiledDataMutex;
 
 	/**
 	 * Copy of the string that was last compiled.
diff --git a/Engine/Source/Runtime/Core/Public/Misc/TransactionallySafeMutex.h b/Engine/Source/Runtime/Core/Public/Misc/TransactionallySafeMutex.h
new file mode 100644
--- /dev/null
+++ b/Engine/Source/Runtime/Core/Public/Misc/TransactionallySafeMutex.h
@@ -0,0 +1,151 @@
+// Copyright Epic Games, Inc. All Rights Reserved.
+
+#pragma once
+
+#include "Async/Mutex.h"
+#include "AutoRTFM/AutoRTFM.h"
+#include "Templates/SharedPointer.h"
+
+#if UE_AUTORTFM
+
+// this was initially copied from wrapping critical section in a transactionally safe way
+
+// A transactionally safe mutex that works in the following novel ways:
+// - In the open (non-transactional):
+//   - Take the lock like before. Simple!
+//   - Free the lock like before too.
+// - In the closed (transactional):
+//   - During locking we query `TransactionalLockCount`:
+//     - 0 means we haven't taken the lock within our transaction nest and need to acquire the lock.
+//     - Otherwise we already have the lock (and are preventing non-transactional code seeing any
+//       modifications we've made while holding the lock), so just bump `TransactionalLockCount`.
+//     - We also register an on-abort handler to release the lock should we abort (but we need to
+//       query `TransactionalLockCount` even there because we could be aborting an inner transaction
+//       and the parent transaction still wants to have the lock held!).
+//   - During unlocking we defer doing the unlock until the transaction commits.
+//
+// Thus with this approach we will hold this lock for the *entirety* of the transactional nest should
+// we take the lock during the transaction, thus preventing non-transactional code from seeing any
+// modifications we should make.
+struct FTransactionallySafeMutexDefinition final
+{
+	// Always open because the constructor arguments will create the underlying critical section.
+	UE_AUTORTFM_ALWAYS_OPEN
+	FTransactionallySafeMutexDefinition() : State(MakeShared<FState>())
+	{
+		if (AutoRTFM::IsTransactional())
+		{
+			const AutoRTFM::EContextStatus Status = AutoRTFM::Close([this]
+				{
+					AutoRTFM::PushOnAbortHandler(this, [this]
+						{
+							this->~FTransactionallySafeMutexDefinition();
+						});
+				});
+
+			ensure(AutoRTFM::EContextStatus::OnTrack == Status);
+		}
+	}
+
+	~FTransactionallySafeMutexDefinition()
+	{
+		if (AutoRTFM::IsTransactional())
+		{
+			const AutoRTFM::EContextStatus Status = AutoRTFM::Close([this]
+				{
+					AutoRTFM::PopOnAbortHandler(this);
+
+					// We explicitly copy the state here for the case that `this` was stack
+					// allocated and has already died before the on-commit is hit.
+					AutoRTFM::OnCommit([State = this->State]
+						{
+							ensure(0 == State->TransactionalLockCount);
+						});
+				});
+
+			ensure(AutoRTFM::EContextStatus::OnTrack == Status);
+		}
+
+		// As the State was constructed in the open, it must be released in the open.
+		AutoRTFM::Open([&] { State = nullptr; });
+	}
+
+	void Lock()
+	{
+		if (AutoRTFM::IsTransactional() || AutoRTFM::IsCommittingOrAborting())
+		{
+			AutoRTFM::Open([this]
+				{
+					// The transactional system which can increment TransactionalLockCount
+					// is always single-threaded, thus this is safe to check without atomicity.
+					if (0 == State->TransactionalLockCount)
+					{
+						State->Mutex.Lock();
+					}
+
+					State->TransactionalLockCount += 1;
+				});
+
+			// We explicitly copy the state here for the case that `this` was stack
+			// allocated and has already died before the on-abort is hit.
+			AutoRTFM::OnAbort([State = this->State]
+				{
+					ensure(0 != State->TransactionalLockCount);
+
+					State->TransactionalLockCount -= 1;
+
+					if (0 == State->TransactionalLockCount)
+					{
+						State->Mutex.Unlock();
+					}
+				});
+		}
+		else
+		{
+			State->Mutex.Lock();
+			ensure(0 == State->TransactionalLockCount);
+		}
+	}
+
+	void Unlock()
+	{
+		if (AutoRTFM::IsTransactional() || AutoRTFM::IsCommittingOrAborting())
+		{
+			// We explicitly copy the state here for the case that `this` was stack
+			// allocated and has already died before the on-commit is hit.
+			AutoRTFM::OnCommit([State = this->State]
+				{
+					ensure(0 != State->TransactionalLockCount);
+
+					State->TransactionalLockCount -= 1;
+
+					if (0 == State->TransactionalLockCount)
+					{
+						State->Mutex.Unlock();
+					}
+				});
+		}
+		else
+		{
+			ensure(0 == State->TransactionalLockCount);
+			State->Mutex.Unlock();
+		}
+	}
+
+private:
+	UE_NONCOPYABLE(FTransactionallySafeMutexDefinition)
+
+	struct FState final
+	{
+		UE::FMutex Mutex;
+		uint32 TransactionalLockCount = 0;
+	};
+
+	TSharedPtr<FState> State;
+};
+
+using FTransactionallySafeMutex = FTransactionallySafeMutexDefinition;
+
+#else
+using FTransactionallySafeMutex = UE::FMutex;
+#endif // UE_AUTORTFM
