CsvProfiler - Report explicit event wait names as trace scopes.

#tests Replayrun

diff --git a/Engine/Source/Runtime/Core/Private/ProfilingDebugging/CsvProfiler.cpp b/Engine/Source/Runtime/Core/Private/ProfilingDebugging/CsvProfiler.cpp
--- a/Engine/Source/Runtime/Core/Private/ProfilingDebugging/CsvProfiler.cpp
+++ b/Engine/Source/Runtime/Core/Private/ProfilingDebugging/CsvProfiler.cpp
@@ -84,9 +84,6 @@ static bool GAllCategoriesStartDisabled = false;
 
 DEFINE_LOG_CATEGORY_STATIC(LogCsvProfiler, Log, All);
 
-const char * GDefaultWaitStatName = "EventWait";
-const char * GIgnoreWaitStatName =  "[IGNORE]";
-
 TAutoConsoleVariable<int32> CVarCsvBlockOnCaptureEnd(
 	TEXT("csv.BlockOnCaptureEnd"), 
 	1,
@@ -2091,6 +2088,30 @@ void FCsvAggregateStatSeries::FinalizeFrame(int64 FrameNumber)
 	}
 }
 
+struct FCsvWaitStatName
+{
+	explicit constexpr FCsvWaitStatName(
+		const char* InStatName,
+		const char* InFormattedStatName = nullptr,
+		const char* InFormattedStatNameNonCP = nullptr)
+		: StatName(InStatName)
+		, FormattedStatName(InFormattedStatName)
+		, FormattedStatNameNonCP(InFormattedStatNameNonCP)
+	{
+	}
+
+	constexpr friend bool operator==(const FCsvWaitStatName& Lhs, const FCsvWaitStatName& Rhs)
+	{
+		return Lhs.StatName == Rhs.StatName;
+	}
+
+	const char* StatName = nullptr;
+	const char* FormattedStatName = nullptr;
+	const char* FormattedStatNameNonCP = nullptr;
+};
+
+static constexpr FCsvWaitStatName GDefaultWaitStatName("EventWait");
+static constexpr FCsvWaitStatName GIgnoreWaitStatName("[IGNORE]");
 
 class FCsvProfilerThreadData
 {
@@ -2331,23 +2352,23 @@ public:
 			((uint64)Events.GetAllocatedSize());
 	}
 
-	CSV_PROFILER_INLINE const char * GetWaitStatName() const
+	CSV_PROFILER_INLINE const FCsvWaitStatName& GetWaitStatName() const
 	{
 		return WaitStatNameStack.Num() == 0 ? GDefaultWaitStatName : WaitStatNameStack.Last();
 	}
 
-	CSV_PROFILER_INLINE void PushWaitStatName(const char * WaitStatName)
+	CSV_PROFILER_INLINE void PushWaitStatName(const FCsvWaitStatName& WaitStatName)
 	{
 		LLM_SCOPE(ELLMTag::CsvProfiler);
 		WaitStatNameStack.Push(WaitStatName);
 	}
-	CSV_PROFILER_INLINE const char* PopWaitStatName()
+	CSV_PROFILER_INLINE TOptional<FCsvWaitStatName> PopWaitStatName()
 	{
 		if (WaitStatNameStack.Num() > 0)
 		{
 			return WaitStatNameStack.Pop();
 		}
-		return nullptr;
+		return NullOpt;
 	}
 	// Raw stat data (written from the thread)
 	TSingleProducerSingleConsumerList<FCsvTimingMarker, 256> TimingMarkers;
@@ -2359,7 +2380,7 @@ public:
 	const bool bIsTaskWorkerThread;
 
 	class FCsvProfilerThreadDataProcessor* DataProcessor;
-	TArray<const char*> WaitStatNameStack;
+	TArray<FCsvWaitStatName> WaitStatNameStack;
 };
 
 uint32 FCsvProfilerThreadData::TlsSlot = FPlatformTLS::InvalidTlsSlot;
@@ -3348,7 +3369,7 @@ void FCsvProfiler::BeginCaptureInternal(const FCsvCaptureCommand& CurrentCommand
 
 	// Check TLS is initialized before starting the capture. This should have happened in BeginCapture
 	check(FCsvProfilerThreadData::IsTlsSlotInitialized());
-	TRACE_CSV_PROFILER_BEGIN_CAPTURE(*Filename, RenderThreadId, RHIThreadId, GDefaultWaitStatName, GCsvStatCounts);
+	TRACE_CSV_PROFILER_BEGIN_CAPTURE(*Filename, RenderThreadId, RHIThreadId, GDefaultWaitStatName.StatName, GCsvStatCounts);
 	GCsvProfilerIsCapturing = true;
 	CaptureStartTime = FPlatformTime::Seconds();
 }
@@ -3800,7 +3821,7 @@ void FCsvProfiler::EndExclusiveStat(const char * StatName)
 }
 
 
-void FCsvProfiler::BeginSetWaitStat(const char * StatName)
+void FCsvProfiler::BeginSetWaitStat(const char* StatName, const char* FormattedStatName, const char* FormattedStatNameNonCP)
 {
 #if !CSV_PROFILER_MINIMAL
 	if (GCsvProfilerIsCapturing && GCsvCategoriesEnabled[CSV_CATEGORY_INDEX(Exclusive)])
@@ -3811,7 +3832,15 @@ void FCsvProfiler::BeginSetWaitStat(const char * StatName)
 			FPlatformMisc::BeginNamedEvent(FColor::Yellow, *FString::Printf(TEXT("CsvWaitStat_%s"), StringCast<TCHAR>(StatName).Get()));
 		}
 #endif
-		FCsvProfilerThreadData::Get().PushWaitStatName(StatName == nullptr ? GIgnoreWaitStatName : StatName);
+		
+		if (StatName == nullptr)
+		{
+			FCsvProfilerThreadData::Get().PushWaitStatName(GIgnoreWaitStatName);
+		}
+		else
+		{
+			FCsvProfilerThreadData::Get().PushWaitStatName(FCsvWaitStatName(StatName, FormattedStatName, FormattedStatNameNonCP));
+		}
 	}
 #endif
 }
@@ -3837,7 +3866,7 @@ void FCsvProfiler::BeginWait()
 #if !CSV_PROFILER_MINIMAL
 	if (GCsvProfilerIsCapturing && GCsvCategoriesEnabled[CSV_CATEGORY_INDEX(Exclusive)])
 	{
-		const char* WaitStatName = FCsvProfilerThreadData::Get().GetWaitStatName();
+		const FCsvWaitStatName& WaitStatName = FCsvProfilerThreadData::Get().GetWaitStatName();
 		if (WaitStatName != GIgnoreWaitStatName)
 		{
 #if CSV_PROFILER_SUPPORT_NAMED_EVENTS
@@ -3846,14 +3875,24 @@ void FCsvProfiler::BeginWait()
 				if ( FThreadIdleStats::Get().IsCriticalPath() )
 				{
 					CsvBeginNamedEvent(FColor(192, 96, 96), "CsvEventWait");
+
+					if (WaitStatName.FormattedStatName != nullptr)
+					{
+						CsvBeginNamedEvent(FColor(192, 96, 96), WaitStatName.FormattedStatName);
+					}
 				}
 				else
 				{
 					CsvBeginNamedEvent(FColor(255, 128, 128), "CsvEventWait (Non-CP)");
+
+					if (WaitStatName.FormattedStatNameNonCP != nullptr)
+					{
+						CsvBeginNamedEvent(FColor(255, 128, 128), WaitStatName.FormattedStatNameNonCP);
+					}
 				}
 			}
 #endif
-			FCsvProfilerThreadData::Get().AddTimestampExclusiveBegin(WaitStatName);
+			FCsvProfilerThreadData::Get().AddTimestampExclusiveBegin(WaitStatName.StatName);
 		}
 	}
 #endif
@@ -3864,14 +3903,19 @@ void FCsvProfiler::EndWait()
 #if !CSV_PROFILER_MINIMAL
 	if (GCsvProfilerIsCapturing && GCsvCategoriesEnabled[CSV_CATEGORY_INDEX(Exclusive)])
 	{
-		const char* WaitStatName = FCsvProfilerThreadData::Get().GetWaitStatName();
+		const FCsvWaitStatName& WaitStatName = FCsvProfilerThreadData::Get().GetWaitStatName();
 		if (WaitStatName != GIgnoreWaitStatName)
 		{
-			FCsvProfilerThreadData::Get().AddTimestampExclusiveEnd(FCsvProfilerThreadData::Get().GetWaitStatName());
+			FCsvProfilerThreadData::Get().AddTimestampExclusiveEnd(WaitStatName.StatName);
 #if CSV_PROFILER_SUPPORT_NAMED_EVENTS
 			if (UNLIKELY(GCsvProfilerNamedEventsExclusive))
 			{
 				CsvEndNamedEvent();
+
+				if (WaitStatName.FormattedStatName != nullptr || WaitStatName.FormattedStatNameNonCP != nullptr)
+				{
+					CsvEndNamedEvent();
+				}
 			}
 #endif
 		}
diff --git a/Engine/Source/Runtime/Core/Public/ProfilingDebugging/CsvProfiler.h b/Engine/Source/Runtime/Core/Public/ProfilingDebugging/CsvProfiler.h
--- a/Engine/Source/Runtime/Core/Public/ProfilingDebugging/CsvProfiler.h
+++ b/Engine/Source/Runtime/Core/Public/ProfilingDebugging/CsvProfiler.h
@@ -127,7 +127,7 @@ struct FCsvDeclaredStat;
 
 	#define CSV_SCOPED_SET_WAIT_STAT(StatName) \
 		TRACE_CSV_PROFILER_INLINE_STAT_EXCLUSIVE("EventWait/"#StatName); \
-		FScopedCsvSetWaitStat _ScopedCsvSetWaitStat ## StatName("EventWait/"#StatName);
+		FScopedCsvSetWaitStat _ScopedCsvSetWaitStat ## StatName("EventWait/"#StatName, "CsvEventWait/"#StatName, "CsvEventWait/"#StatName" (Non-CP)");
 
 	#define CSV_SCOPED_SET_WAIT_STAT_IGNORE()						FScopedCsvSetWaitStat _ScopedCsvSetWaitStat ## StatName();
 
@@ -411,7 +411,7 @@ public:
 		RecordEventfInternal(CategoryIndex, (const TCHAR*)Fmt, Args...);
 	}
 
-	CORE_API static void BeginSetWaitStat(const char * StatName);
+	CORE_API static void BeginSetWaitStat(const char* StatName, const char* FormattedStatName, const char* FormattedStatNameNonCP);
 	CORE_API static void EndSetWaitStat();
 
 	CORE_API static void BeginWait();
@@ -737,17 +737,25 @@ public:
 class FScopedCsvSetWaitStat
 {
 public:
-	FScopedCsvSetWaitStat(const char * InStatName = nullptr)
+	FScopedCsvSetWaitStat(
+		const char* InStatName = nullptr,
+		const char* InFormattedStatName = nullptr,
+		const char* InFormattedStatNameNonCP = nullptr)
 		: StatName(InStatName)
+		, FormattedStatName(InFormattedStatName)
+		, FormattedStatNameNonCP(InFormattedStatNameNonCP)
 	{
-		FCsvProfiler::BeginSetWaitStat(StatName);
+		FCsvProfiler::BeginSetWaitStat(StatName, FormattedStatName, FormattedStatNameNonCP);
 	}
 
 	~FScopedCsvSetWaitStat()
 	{
 		FCsvProfiler::EndSetWaitStat();
 	}
-	const char * StatName;
+
+	const char* StatName;
+	const char* FormattedStatName;
+	const char* FormattedStatNameNonCP;
 };
 
 struct FCsvCategory
