GameFeedback adding CustomFields for analytics
Additionally added 'GetCSVID' to the CSVProfiler so we can access the unique GUID for a CSV instance.

#tests Tested against Win64 Server Cooked
#rnx

diff --git a/Engine/Source/Runtime/Core/Private/ProfilingDebugging/CsvProfiler.cpp b/Engine/Source/Runtime/Core/Private/ProfilingDebugging/CsvProfiler.cpp
--- a/Engine/Source/Runtime/Core/Private/ProfilingDebugging/CsvProfiler.cpp
+++ b/Engine/Source/Runtime/Core/Private/ProfilingDebugging/CsvProfiler.cpp
@@ -34,7 +34,6 @@
 #include "Misc/Compression.h"
 #include "Misc/ConfigCacheIni.h"
 #include "Misc/Fork.h"
-#include "Misc/Guid.h"
 #include "Misc/WildcardString.h"
 #include "Modules/ModuleManager.h"
 #include "HAL/PlatformMemoryHelpers.h"
@@ -3091,6 +3090,7 @@ FCsvProfiler::FCsvProfiler()
 	, CaptureFrameNumber(0)
 	, CaptureFrameNumberRT(0)
 	, CaptureOnEventFrameCount(-1)
+	, CsvGUID(FGuid(0, 0, 0, 0))
 	, bInsertEndFrameAtFrameStart(false)
 	, bNamedEventsWasEnabled(false)
 	, LastEndFrameTimestamp(0)
@@ -3323,9 +3323,10 @@ void FCsvProfiler::BeginCaptureInternal(const FCsvCaptureCommand& CurrentCommand
 	}
 
 	// Set the CSV ID and mirror it to the log
-	FString CsvId = FGuid::NewGuid().ToString();
-	SetMetadataInternal(TEXT("CsvID"), *CsvId);
-	UE_LOG(LogCsvProfiler, Display, TEXT("Capture started. CSV ID: %s"), *CsvId);
+	CsvGUID = FGuid::NewGuid();
+	FString CsvIdString = CsvGUID.ToString();
+	SetMetadataInternal(TEXT("CsvID"), *CsvIdString);
+	UE_LOG(LogCsvProfiler, Display, TEXT("Capture started. CSV ID: %s"), *CsvIdString);
 
 	int32 TargetFPS = FPlatformMisc::GetMaxRefreshRate();
 	static IConsoleVariable* CsvTargetFrameRateCVar = IConsoleManager::Get().FindConsoleVariable(TEXT("csv.TargetFrameRateOverride"));
@@ -4347,6 +4348,11 @@ int32 FCsvProfiler::GetNumFrameToCaptureOnEvent() const
 	return CaptureOnEventFrameCount;
 }
 
+const FGuid& FCsvProfiler::GetCsvID()
+{
+	return CsvGUID;
+}
+
 bool FCsvProfiler::EnableCategoryByString(const FString& CategoryName) const
 {
 	int32 Category = FCsvCategoryData::Get()->GetCategoryIndex(CategoryName);
diff --git a/Engine/Source/Runtime/Core/Public/ProfilingDebugging/CsvProfiler.h b/Engine/Source/Runtime/Core/Public/ProfilingDebugging/CsvProfiler.h
--- a/Engine/Source/Runtime/Core/Public/ProfilingDebugging/CsvProfiler.h
+++ b/Engine/Source/Runtime/Core/Public/ProfilingDebugging/CsvProfiler.h
@@ -21,6 +21,7 @@
 #include "HAL/PreprocessorHelpers.h"
 #include "HAL/ThreadSafeCounter.h"
 #include "Misc/EnumClassFlags.h"
+#include "Misc/Guid.h"
 #include "ProfilingDebugging/CsvProfilerConfig.h"
 #include "ProfilingDebugging/CsvProfilerTrace.h"
 #include "ProfilingDebugging/MiscTrace.h"
@@ -426,6 +427,7 @@ public:
 	CORE_API int32 GetCaptureFrameNumber() const;
 	CORE_API int32 GetCaptureFrameNumberRT() const;
 	CORE_API int32 GetNumFrameToCaptureOnEvent() const;
+	CORE_API const FGuid& GetCsvID();
 
 	CORE_API bool EnableCategoryByString(const FString& CategoryName) const;
 	CORE_API void EnableCategoryByIndex(uint32 CategoryIndex, bool bEnable) const;
@@ -525,6 +527,8 @@ private:
 	int32 CaptureFrameNumberRT;
 	int32 CaptureOnEventFrameCount;
 
+	FGuid CsvGUID;
+
 	bool bInsertEndFrameAtFrameStart;
 	bool bNamedEventsWasEnabled;
 
