Changes to apple platform specific calls to make shadercompileworker work remotely.

* FApplePlatformFile::CreateDirectory is using mkdir instead of NSFileManager::createDirectoryAtPath when IS_PROGRAM is 1 (ShaderCompileWorker)
* FApplePlatformProcess::UserTempDir returns getenv("TMPDIR") instead of NSTemporaryDirectory when IS_PROGRAM is 1 (ShaderCompileWorker)
* Fixed so FApplePlatformFile::FileSize returns -1 if Stat fails (it kind of did that before too by accident but only if stat didn't touch the struct it writes the results into)

Change has been tested manually by human testers

diff --git a/Engine/Source/Runtime/Core/Private/Apple/ApplePlatformFile.cpp b/Engine/Source/Runtime/Core/Private/Apple/ApplePlatformFile.cpp
--- a/Engine/Source/Runtime/Core/Private/Apple/ApplePlatformFile.cpp
+++ b/Engine/Source/Runtime/Core/Private/Apple/ApplePlatformFile.cpp
@@ -460,7 +460,11 @@ int64 FApplePlatformFile::FileSize(const TCHAR* Filename)
 {
 	struct stat FileInfo;
 	FileInfo.st_size = -1;
-	Stat(Filename, &FileInfo);
+	if (Stat(Filename, &FileInfo) != 0)
+	{
+		return -1;
+	}
+	
 	// make sure to return -1 for directories
 	if (S_ISDIR(FileInfo.st_mode))
 	{
@@ -704,10 +708,17 @@ bool FApplePlatformFile::CreateDirectory(const TCHAR* Directory)
 {
 	SCOPED_AUTORELEASE_POOL;
 
+	// Using native functions on program targets to be able to use detour and remote helpers
+#if IS_PROGRAM
+	if (mkdir(TCHAR_TO_UTF8(*NormalizeFilename(Directory)), 0775) == 0)
+		return true;
+	return errno == EEXIST;
+#else
 	CFStringRef CFDirectory = FPlatformString::TCHARToCFString(*NormalizeFilename(Directory));
 	bool Result = [[NSFileManager defaultManager] createDirectoryAtPath:(NSString*)CFDirectory withIntermediateDirectories:true attributes:nil error:nil];
 	CFRelease(CFDirectory);
 	return Result;
+#endif
 }
 
 bool FApplePlatformFile::DeleteDirectory(const TCHAR* Directory)
diff --git a/Engine/Source/Runtime/Core/Private/Apple/ApplePlatformProcess.cpp b/Engine/Source/Runtime/Core/Private/Apple/ApplePlatformProcess.cpp
--- a/Engine/Source/Runtime/Core/Private/Apple/ApplePlatformProcess.cpp
+++ b/Engine/Source/Runtime/Core/Private/Apple/ApplePlatformProcess.cpp
@@ -22,7 +22,12 @@ const TCHAR* FApplePlatformProcess::UserTempDir()
 	static FString TempDir;
 	if (!TempDir.Len())
 	{
-		TempDir = NSTemporaryDirectory();
+		// The reason we use TMPDIR here is to support applications being detoured/remotely distributed (ShaderCompileWorker)
+#if IS_PROGRAM
+		TempDir = getenv("TMPDIR");
+#else
+		TempDir = NSTemporaryDirectory(); // NSTemporaryDirectory() can return a different temp directory than TMPDIR and detouring is not supported.
+#endif
 	}
 	return *TempDir;
 }
