Handle iOS stat() file failure returns.

#rnx

diff --git a/Engine/Source/Runtime/Core/Private/IOS/IOSPlatformFile.cpp b/Engine/Source/Runtime/Core/Private/IOS/IOSPlatformFile.cpp
--- a/Engine/Source/Runtime/Core/Private/IOS/IOSPlatformFile.cpp
+++ b/Engine/Source/Runtime/Core/Private/IOS/IOSPlatformFile.cpp
@@ -603,7 +603,10 @@ int64 FIOSPlatformFile::FileSize(const TCHAR* Filename)
 		if(stat(TCHAR_TO_UTF8(*ConvertToIOSPath(NormalizedFilename, true, false)), &FileInfo) == -1)
 		{
 			// if not in the private write path, check the public write path
-			stat(TCHAR_TO_UTF8(*ConvertToIOSPath(NormalizedFilename, true, true)), &FileInfo);
+			if(stat(TCHAR_TO_UTF8(*ConvertToIOSPath(NormalizedFilename, true, true)), &FileInfo) == -1)
+			{
+				return -1;
+			}
 		}
 	}
 
@@ -722,7 +725,6 @@ FDateTime FIOSPlatformFile::GetTimeStamp(const TCHAR* Filename)
 	// convert _stat time to FDateTime
 	FTimespan TimeSinceEpoch(0, 0, FileInfo.st_mtime);
 	return IOSEpoch + TimeSinceEpoch;
-
 }
 
 void FIOSPlatformFile::SetTimeStamp(const TCHAR* Filename, const FDateTime DateTime)
@@ -891,7 +893,10 @@ IMappedFileHandle* FIOSPlatformFile::OpenMapped(const TCHAR* Filename)
 			struct stat FileInfo;
 			FileInfo.st_size = -1;
 			// check the read path
-			fstat(Handle, &FileInfo);
+			if(fstat(Handle, &FileInfo) == -1)
+			{
+				return NULL;
+			}
 			uint64 FileSize = FileInfo.st_size;
 
 			return new FIOSMappedFileHandle(Handle, FileSize, FinalPath);
