Fix GFP hotfixing for GFPs without Engine ini contents

Allow hotfixes for any GFP on any branch, regardless of existing config - previously Engine/DeviceProfiles ini file was required
Fix CDO instances of the hotfix manager that registered for the dynamic branch HF

#tests QA confirmed hotifx applied correctly on entering a game mode and was removed correctly on leaving
#rnx

diff --git a/Engine/Plugins/Online/OnlineFramework/Source/Hotfix/Private/OnlineHotfixManager.cpp b/Engine/Plugins/Online/OnlineFramework/Source/Hotfix/Private/OnlineHotfixManager.cpp
--- a/Engine/Plugins/Online/OnlineFramework/Source/Hotfix/Private/OnlineHotfixManager.cpp
+++ b/Engine/Plugins/Online/OnlineFramework/Source/Hotfix/Private/OnlineHotfixManager.cpp
@@ -185,12 +185,15 @@ UOnlineHotfixManager::UOnlineHotfixManager() :
 #endif
 	GameContentPath = FString() / FApp::GetProjectName() / TEXT("Content");
 
-	if (!UObject::IsGarbageEliminationEnabled())
+	if (this != GetClass()->GetDefaultObject())
 	{
-		FCoreUObjectDelegates::GetPreGarbageCollectDelegate().AddUObject(this, &UOnlineHotfixManager::StopTrackingInvalidHotfixedAssets);
-	}
+		if (!UObject::IsGarbageEliminationEnabled())
+		{
+			FCoreUObjectDelegates::GetPreGarbageCollectDelegate().AddUObject(this, &UOnlineHotfixManager::StopTrackingInvalidHotfixedAssets);
+		}
 
-	UE::DynamicConfig::HotfixBranch.AddUObject(this, &UOnlineHotfixManager::HotfixDynamicBranch);
+		UE::DynamicConfig::HotfixPluginForBranch.AddUObject(this, &UOnlineHotfixManager::HotfixDynamicBranch);
+	}
 }
 
 UOnlineHotfixManager::UOnlineHotfixManager(FVTableHelper& Helper)
@@ -1986,6 +1989,7 @@ void UOnlineHotfixManager::HotfixDynamicBranch(const FName& Tag, const FName& Br
 	FConfigBranch* BranchToHotfix = GConfig->FindBranch(Branch, FString());
 	if (BranchToHotfix)
 	{
+		// Check in dynamic hotfix contents that we can find a given ini file and apply it as a dynamic layer
 		const auto TryApplyHotfix = [&](const FString Prefix)
 		{
 			const FString HotfixFileName = Prefix + Tag.ToString() + Branch.ToString() + TEXT(".ini");
@@ -2002,8 +2006,11 @@ void UOnlineHotfixManager::HotfixDynamicBranch(const FName& Tag, const FName& Br
 		};
 
 		const FString PlatformName = FPlatformProperties::IniPlatformName();
+
+		// Valid hotfix filenames: $Plugin$Branch.ini, Default$Plugin$Branch.ini, $Platform$PluginBranch.ini
 		TryApplyHotfix(FString());
-		TryApplyHotfix(PlatformName + TEXT("_"));		
+		TryApplyHotfix("Default");
+		TryApplyHotfix(PlatformName);
 	}
 }
 
diff --git a/Engine/Plugins/Runtime/GameFeatures/Source/GameFeatures/Private/GameFeatureData.cpp b/Engine/Plugins/Runtime/GameFeatures/Source/GameFeatures/Private/GameFeatureData.cpp
--- a/Engine/Plugins/Runtime/GameFeatures/Source/GameFeatures/Private/GameFeatureData.cpp
+++ b/Engine/Plugins/Runtime/GameFeatures/Source/GameFeatures/Private/GameFeatureData.cpp
@@ -191,6 +191,13 @@ void UGameFeatureData::InitializeHierarchicalPluginIniFiles(const FString& Plugi
 
 				// apply plugin modifications from this plugin to the everything
 				FConfigCacheIni::AddPluginToAllBranches(PluginName, ChangeTracker);
+
+				// give hotfix a chance to run at that point
+				{
+#define HOTFIX_BRANCH(Ini) UE::DynamicConfig::HotfixPluginForBranch.Broadcast(PluginName, #Ini, ChangeTracker);
+					ENUMERATE_KNOWN_INI_FILES(HOTFIX_BRANCH);
+#undef HOTFIX_BRANCH
+				}
 			});
 
 		return;
diff --git a/Engine/Source/Runtime/Core/Private/Misc/ConfigCacheIni.cpp b/Engine/Source/Runtime/Core/Private/Misc/ConfigCacheIni.cpp
--- a/Engine/Source/Runtime/Core/Private/Misc/ConfigCacheIni.cpp
+++ b/Engine/Source/Runtime/Core/Private/Misc/ConfigCacheIni.cpp
@@ -3255,7 +3255,8 @@ bool FConfigBranch::AddDynamicLayerStringToHierarchy(const FString& Filename, co
 			bInsertedAtEnd = true;
 		}
 	}
-	
+
+	DynamicLayer->Tag = Tag;
 	FillFileFromBuffer(DynamicLayer, Contents, true, Filename);
 
 	// track modified section names if desired
diff --git a/Engine/Source/Runtime/Core/Private/Misc/ConfigUtilities.cpp b/Engine/Source/Runtime/Core/Private/Misc/ConfigUtilities.cpp
--- a/Engine/Source/Runtime/Core/Private/Misc/ConfigUtilities.cpp
+++ b/Engine/Source/Runtime/Core/Private/Misc/ConfigUtilities.cpp
@@ -19,7 +19,7 @@
 
 namespace UE::DynamicConfig
 {
-TMulticastDelegate<void(const FName&, const FName&, class FConfigModificationTracker*)> HotfixBranch;
+TMulticastDelegate<void(const FName&, const FName&, class FConfigModificationTracker*)> HotfixPluginForBranch;
 TMulticastDelegate<void(const class FConfigModificationTracker* ChangeTracker)> ReloadObjects;
 TMulticastDelegate<void(const TSet<FString>&)> UpdateDeviceProfiles;
 
@@ -32,7 +32,7 @@ void PerformDynamicConfig(FName Tag, TFunction<void(class FConfigModificationTra
 	// run the callback
 	PerformModification(&ChangeTracker);
 	
-	// now update everything if anythign was read in!
+	// now update everything if anything was read in!
 	if (ChangeTracker.ModifiedSectionsPerBranch.Num() > 0)
 	{
 		// reload objects that had their configs changed
@@ -40,22 +40,16 @@ void PerformDynamicConfig(FName Tag, TFunction<void(class FConfigModificationTra
 		{
 			const FConfigModificationTracker::FCVarTracker& CVars = SectionPair.Value;
 			EConsoleVariableFlags Priority = (EConsoleVariableFlags)CVars.CVarPriority;
-	
+
 			// now walk over the updated cvars and set them based on the priority
 			for (const auto& BranchPair : CVars.CVarEntriesPerBranch)
 			{
 				const FConfigSection& Section = BranchPair.Value;
-	
+
 				for (const auto& CVarPair : Section)
 				{
-//					void OnSetCVarFromIniEntry(const TCHAR *IniFile, const TCHAR *Key, const TCHAR* Value, uint32 SetBy, bool bAllowCheating, bool bNoLogging)
-
-					UE::ConfigUtilities::OnSetCVarFromIniEntry(TEXT("DynamicLayer"), *CVarPair.Key.ToString(),*CVarPair.Value.GetValue(), Priority, false, false, Tag);
-//					IConsoleVariable* CVar = IConsoleManager::Get().FindConsoleVariable(*CVarPair.Key.ToString());
-//					CVar->Set(*CVarPair.Value.GetValue(), Priority, Tag);
+					UE::ConfigUtilities::OnSetCVarFromIniEntry(TEXT("DynamicLayer"), *CVarPair.Key.ToString(), *CVarPair.Value.GetValue(), Priority, false, false, Tag);
 				}
-
-				HotfixBranch.Broadcast(Tag, BranchPair.Key, &ChangeTracker);
 			}
 		}
 
diff --git a/Engine/Source/Runtime/Core/Public/Misc/ConfigUtilities.h b/Engine/Source/Runtime/Core/Public/Misc/ConfigUtilities.h
--- a/Engine/Source/Runtime/Core/Public/Misc/ConfigUtilities.h
+++ b/Engine/Source/Runtime/Core/Public/Misc/ConfigUtilities.h
@@ -21,7 +21,8 @@ namespace UE::DynamicConfig
 {
 	extern CORE_API void PerformDynamicConfig(FName Tag, TFunction<void(class FConfigModificationTracker*)> PerformModification, TFunction<void(FConfigModificationTracker*)> PostModification=nullptr);
 
-	extern CORE_API TMulticastDelegate<void(const FName& Tag, const FName& Branch, class FConfigModificationTracker* ModificationTracker)> HotfixBranch;
+	// this isn't directly used in this module, but the OnlineHotfixManager and GameFeatures plugin use it to talk to each other
+	extern CORE_API TMulticastDelegate<void(const FName& Tag, const FName& Branch, class FConfigModificationTracker* ModificationTracker)> HotfixPluginForBranch;
 
 	// this calls the UObjectBaseUtility from code where object system is not linked (it also calls TSOnConfigSectionsChanged()!)
 	extern CORE_API TMulticastDelegate<void(const FConfigModificationTracker* ChangeTracker)> ReloadObjects;
