Enabled GPUScene on mobile platfroms by default

diff --git a/Engine/Source/Runtime/Core/Private/HAL/ConsoleManager.cpp b/Engine/Source/Runtime/Core/Private/HAL/ConsoleManager.cpp
--- a/Engine/Source/Runtime/Core/Private/HAL/ConsoleManager.cpp
+++ b/Engine/Source/Runtime/Core/Private/HAL/ConsoleManager.cpp
@@ -3793,7 +3793,7 @@ static TAutoConsoleVariable<int32> CVarMobileForwardEnableClusteredReflections(
 
 static TAutoConsoleVariable<int32> CVarMobileSupportGPUScene(
 	TEXT("r.Mobile.SupportGPUScene"),
-	0,
+	1,
 	TEXT("Whether to support GPU scene, required for auto-instancing (only Mobile feature level)"),
 	ECVF_ReadOnly | ECVF_RenderThreadSafe
 );
