Fixed typo.

diff --git a/Engine/Source/Runtime/Core/Public/Templates/SharedPointerFwd.h b/Engine/Source/Runtime/Core/Public/Templates/SharedPointerFwd.h
--- a/Engine/Source/Runtime/Core/Public/Templates/SharedPointerFwd.h
+++ b/Engine/Source/Runtime/Core/Public/Templates/SharedPointerFwd.h
@@ -19,7 +19,7 @@ enum class ESPMode : uint8
 
 
 // Forward declarations.  By default, thread safety features are turned on. (Mode = ESPMode::ThreadSafe).
-// If you need more concerned with performance of ref-counting, you should use ESPMode::NotThreadSafe.
+// If you are more concerned with performance of ref-counting, you should use ESPMode::NotThreadSafe.
 template< class ObjectType, ESPMode Mode = ESPMode::ThreadSafe > class TSharedRef;
 template< class ObjectType, ESPMode Mode = ESPMode::ThreadSafe > class TSharedPtr;
 template< class ObjectType, ESPMode Mode = ESPMode::ThreadSafe > class TWeakPtr;
