Fix spline nearest find fail in weird cases
- previously, LastValue could inherit between different Newton iterations, causing the converge to happen immaturely, and returning an invalid result with local (0,0,0) point

diff --git a/Engine/Source/Runtime/Core/Public/Math/InterpCurve.h b/Engine/Source/Runtime/Core/Public/Math/InterpCurve.h
--- a/Engine/Source/Runtime/Core/Public/Math/InterpCurve.h
+++ b/Engine/Source/Runtime/Core/Public/Math/InterpCurve.h
@@ -685,7 +685,8 @@ float FInterpCurve<T>::FindNearestOnSegment(const T& PointInSpace, int32 PtIdx,
 			[&PointsT, &PtIdx, &NextPtIdx, &PointInSpace, &Diff, MaxIteration, InvThree]
 			(float& Value, float& Move, TFunctionRef<bool(float&, const float)> BreakBeforeEvaluate, TFunctionRef<bool(float&, const float)> BreakAfterEvaluate, const float Tolerance) -> float
 		{
-			T FoundPoint = {};
+			// Initialize FoundPoint to some meaningful points on the curve
+			T FoundPoint = Value > 0.5 ? PointsT[NextPtIdx].OutVal : PointsT[PtIdx].OutVal;
 
 			for (int32 Iter = 0; Iter < MaxIteration; ++Iter)
 			{
@@ -721,6 +722,9 @@ float FInterpCurve<T>::FindNearestOnSegment(const T& PointInSpace, int32 PtIdx,
 			float& Move = NextMovesT[Index];
 
 			DistancesSq[Index] = Newton(Value, Move, BreakIfConverged, NeverBreak, ToleranceStep1);
+
+			// Reset LastValue after calculation so that each newton iteration actually converge
+			LastValue = UE_BIG_NUMBER;
 		}
 
 		// Find the index of the best approximation
