[AutoRTFM] Add const iterators for FWriteLog

diff --git a/Engine/Source/Runtime/Core/Private/AutoRTFM/WriteLog.h b/Engine/Source/Runtime/Core/Private/AutoRTFM/WriteLog.h
--- a/Engine/Source/Runtime/Core/Private/AutoRTFM/WriteLog.h
+++ b/Engine/Source/Runtime/Core/Private/AutoRTFM/WriteLog.h
@@ -5,6 +5,7 @@
 #include "AutoRTFM/HitSet.h"
 #include "Utils.h"
 #include <stdint.h>
+#include <type_traits>
 
 namespace AutoRTFM
 {
@@ -50,7 +51,7 @@ namespace AutoRTFM
 			return Entry;
 		}
 
-		UE_AUTORTFM_FORCEINLINE void* GetOriginal()
+		UE_AUTORTFM_FORCEINLINE void* GetOriginal() const
 		{
 			return OriginalAndIsSmallAndSize.Get();
 		}
@@ -67,6 +68,18 @@ namespace AutoRTFM
 			}
 		}
 
+		UE_AUTORTFM_FORCEINLINE const void* GetCopy() const
+		{
+			if (IsSmall())
+			{
+				return reinterpret_cast<const void*>(&Copy);
+			}
+			else
+			{
+				return reinterpret_cast<const void*>(Copy);
+			}
+		}
+
 		UE_AUTORTFM_FORCEINLINE size_t GetSize() const
 		{
 			return OriginalAndIsSmallAndSize.GetTopTag() & ~IsSmallBit;
@@ -120,85 +133,84 @@ namespace AutoRTFM
             Reset();
         }
 
-        struct Iterator final
-        {
-            explicit Iterator(FWriteLogEntryBucket* const Bucket) : Bucket(Bucket) {}
+		template<bool IS_FORWARD, bool IS_CONST>
+		struct TIterator final
+		{
+			using BucketType = std::conditional_t<IS_CONST, const FWriteLogEntryBucket, FWriteLogEntryBucket>;
+			using EntryType = std::conditional_t<IS_CONST, const FWriteLogEntry, FWriteLogEntry>;
 
-            FWriteLogEntry& operator*()
-            {
-                return Bucket->Entries[Offset];
-            }
+			explicit TIterator(BucketType* Bucket) : Bucket(Bucket) {}
 
-            void operator++()
-            {
-                Offset++;
+			EntryType& operator*() const
+			{
+				return IS_FORWARD ? Bucket->Entries[Offset] : Bucket->Entries[Bucket->Size - Offset - 1];
+			}
 
-                if (Offset == Bucket->Size)
-                {
-                    Bucket = Bucket->Next;
-                    Offset = 0;
-                }
-            }
+			void operator++()
+			{
+				Offset++;
 
-            bool operator!=(Iterator& Other) const
-            {
-                return (Other.Bucket != Bucket) || (Other.Offset != Offset);
-            }
+				if (Offset == Bucket->Size)
+				{
+					Bucket = IS_FORWARD ? Bucket->Next : Bucket->Prev;
+					Offset = 0;
+				}
+			}
 
-        private:
-            FWriteLogEntryBucket* Bucket;
-            size_t Offset = 0;
-        };
+			bool operator!=(const TIterator& Other) const
+			{
+				return (Other.Bucket != Bucket) || (Other.Offset != Offset);
+			}
 
-        struct ReverseIterator final
-        {
-            explicit ReverseIterator(FWriteLogEntryBucket* const Bucket) : Bucket(Bucket) {}
+		private:
+			BucketType* Bucket;
+			size_t Offset = 0;
+		};
 
-            FWriteLogEntry& operator*()
-            {
-                return Bucket->Entries[Bucket->Size - Offset - 1];
-            }
+		using Iterator = TIterator</* IS_FORWARD */ true, /* IS_CONST */ false>;
+		using ReverseIterator = TIterator</* IS_FORWARD */ false, /* IS_CONST */ false>;
+		using ConstIterator = TIterator</* IS_FORWARD */ true, /* IS_CONST */ true>;
+		using ConstReverseIterator = TIterator</* IS_FORWARD */ false, /* IS_CONST */ true>;
 
-            void operator++()
-            {
-                Offset++;
+		Iterator begin()
+		{
+			return Iterator(Start);
+		}
 
-                if (Offset == Bucket->Size)
-                {
-                    Bucket = Bucket->Prev;
-                    Offset = 0;
-                }
-            }
+		ConstIterator begin() const
+		{
+			return ConstIterator(Start);
+		}
 
-            bool operator!=(const ReverseIterator& Other) const
-            {
-                return (Other.Bucket != Bucket) || (Other.Offset != Offset);
-            }
+		ReverseIterator rbegin()
+		{
+			return ReverseIterator(Current);
+		}
 
-        private:
-            FWriteLogEntryBucket* Bucket;
-            size_t Offset = 0;
-        };
+		ConstReverseIterator rbegin() const
+		{
+			return ConstReverseIterator(Current);
+		}
 
-        Iterator begin()
-        {
-            return Iterator(Start);
-        }
+		Iterator end()
+		{
+			return Iterator(nullptr);
+		}
 
-        ReverseIterator rbegin()
-        {
-            return ReverseIterator(Current);
-        }
+		ConstIterator end() const
+		{
+			return ConstIterator(nullptr);
+		}
 
-        Iterator end()
-        {
-            return Iterator(nullptr);
-        }
+		ReverseIterator rend()
+		{
+			return ReverseIterator(nullptr);
+		}
 
-        ReverseIterator rend()
-        {
-            return ReverseIterator(nullptr);
-        }
+		ConstReverseIterator rend() const
+		{
+			return ConstReverseIterator(nullptr);
+		}
 
         void Reset()
         {
