[MemoryStats]
 - Fix potential race condition and accumulation of queued memory stats tasks

#rnx

diff --git a/Engine/Source/Runtime/Core/Private/GenericPlatform/GenericPlatformMemory.cpp b/Engine/Source/Runtime/Core/Private/GenericPlatform/GenericPlatformMemory.cpp
--- a/Engine/Source/Runtime/Core/Private/GenericPlatform/GenericPlatformMemory.cpp
+++ b/Engine/Source/Runtime/Core/Private/GenericPlatform/GenericPlatformMemory.cpp
@@ -12,6 +12,7 @@
 #include "Containers/Ticker.h"
 #include "Misc/FeedbackContext.h"
 #include "Async/Async.h"
+#include "Async/Mutex.h"
 #include "HAL/MallocAnsi.h"
 #include "GenericPlatform/GenericPlatformMemoryPoolStats.h"
 #include "HAL/MemoryMisc.h"
@@ -108,34 +109,43 @@ struct FGenericStatsUpdater
 	/** Gathers and sets all platform memory statistics into the corresponding stats. */
 	static void DoUpdateStats()
 	{
-        QUICK_SCOPE_CYCLE_COUNTER(STAT_FGenericStatsUpdater_DoUpdateStats);
-
-		// This is slow, so do it on the task graph.
-		FPlatformMemoryStats MemoryStats = FPlatformMemory::GetStats();
-		SET_MEMORY_STAT( STAT_TotalPhysical, MemoryStats.TotalPhysical );
-		SET_MEMORY_STAT( STAT_TotalVirtual, MemoryStats.TotalVirtual );
-		SET_MEMORY_STAT( STAT_PageSize, MemoryStats.PageSize );
-		SET_MEMORY_STAT( STAT_TotalPhysicalGB, MemoryStats.TotalPhysicalGB );
-
-		SET_MEMORY_STAT( STAT_AvailablePhysical, MemoryStats.AvailablePhysical );
-		SET_MEMORY_STAT( STAT_AvailableVirtual, MemoryStats.AvailableVirtual );
-		SET_MEMORY_STAT( STAT_UsedPhysical, MemoryStats.UsedPhysical );
-		SET_MEMORY_STAT( STAT_PeakUsedPhysical, MemoryStats.PeakUsedPhysical );
-		SET_MEMORY_STAT( STAT_UsedVirtual, MemoryStats.UsedVirtual );
-		SET_MEMORY_STAT( STAT_PeakUsedVirtual, MemoryStats.PeakUsedVirtual );
-
-		TRACE_COUNTER_SET(PlatformMemoryTotalPhysical, MemoryStats.TotalPhysical);
-		TRACE_COUNTER_SET(PlatformMemoryTotalVirtual, MemoryStats.TotalVirtual);
-		TRACE_COUNTER_SET(PlatformMemoryPageSize, MemoryStats.PageSize);
-		TRACE_COUNTER_SET(PlatformMemoryAvailablePhysical, MemoryStats.AvailablePhysical);
-		TRACE_COUNTER_SET(PlatformMemoryAvailableVirtual, MemoryStats.AvailableVirtual);
-		TRACE_COUNTER_SET(PlatformMemoryUsedPhysical, MemoryStats.UsedPhysical);
-		TRACE_COUNTER_SET(PlatformMemoryPeakUsedPhysical, MemoryStats.PeakUsedPhysical);
-		TRACE_COUNTER_SET(PlatformMemoryUsedVirtual, MemoryStats.UsedVirtual);
-		TRACE_COUNTER_SET(PlatformMemoryPeakUsedVirtual, MemoryStats.PeakUsedVirtual);
-
-		// Platform specific stats.
-		FPlatformMemory::InternalUpdateStats( MemoryStats );
+		static UE::FMutex Mutex;
+		// Avoid race in case the ticker runs faster than the background tasks
+		// and just skip the update if another one is already running so that
+		// we can keep up with the ticker.
+		if (Mutex.TryLock())
+		{
+			QUICK_SCOPE_CYCLE_COUNTER(STAT_FGenericStatsUpdater_DoUpdateStats);
+
+			// This is slow, so do it on the task graph.
+			FPlatformMemoryStats MemoryStats = FPlatformMemory::GetStats();
+			SET_MEMORY_STAT(STAT_TotalPhysical, MemoryStats.TotalPhysical);
+			SET_MEMORY_STAT(STAT_TotalVirtual, MemoryStats.TotalVirtual);
+			SET_MEMORY_STAT(STAT_PageSize, MemoryStats.PageSize);
+			SET_MEMORY_STAT(STAT_TotalPhysicalGB, MemoryStats.TotalPhysicalGB);
+
+			SET_MEMORY_STAT(STAT_AvailablePhysical, MemoryStats.AvailablePhysical);
+			SET_MEMORY_STAT(STAT_AvailableVirtual, MemoryStats.AvailableVirtual);
+			SET_MEMORY_STAT(STAT_UsedPhysical, MemoryStats.UsedPhysical);
+			SET_MEMORY_STAT(STAT_PeakUsedPhysical, MemoryStats.PeakUsedPhysical);
+			SET_MEMORY_STAT(STAT_UsedVirtual, MemoryStats.UsedVirtual);
+			SET_MEMORY_STAT(STAT_PeakUsedVirtual, MemoryStats.PeakUsedVirtual);
+
+			TRACE_COUNTER_SET(PlatformMemoryTotalPhysical, MemoryStats.TotalPhysical);
+			TRACE_COUNTER_SET(PlatformMemoryTotalVirtual, MemoryStats.TotalVirtual);
+			TRACE_COUNTER_SET(PlatformMemoryPageSize, MemoryStats.PageSize);
+			TRACE_COUNTER_SET(PlatformMemoryAvailablePhysical, MemoryStats.AvailablePhysical);
+			TRACE_COUNTER_SET(PlatformMemoryAvailableVirtual, MemoryStats.AvailableVirtual);
+			TRACE_COUNTER_SET(PlatformMemoryUsedPhysical, MemoryStats.UsedPhysical);
+			TRACE_COUNTER_SET(PlatformMemoryPeakUsedPhysical, MemoryStats.PeakUsedPhysical);
+			TRACE_COUNTER_SET(PlatformMemoryUsedVirtual, MemoryStats.UsedVirtual);
+			TRACE_COUNTER_SET(PlatformMemoryPeakUsedVirtual, MemoryStats.PeakUsedVirtual);
+
+			// Platform specific stats.
+			FPlatformMemory::InternalUpdateStats(MemoryStats);
+
+			Mutex.Unlock();
+		}
 	}
 };
 
