Improve non-uniform scale handling in FTransform::GetRelativeTransformReverse

diff --git a/Engine/Source/Runtime/Core/Private/Math/Transform.cpp b/Engine/Source/Runtime/Core/Private/Math/Transform.cpp
--- a/Engine/Source/Runtime/Core/Private/Math/Transform.cpp
+++ b/Engine/Source/Runtime/Core/Private/Math/Transform.cpp
@@ -97,26 +97,26 @@ bool TTransform<T>::InitFromString( const FString& Source )
 template<typename T>
 TTransform<T> TTransform<T>::GetRelativeTransformReverse(const TTransform<T>& Other) const
 {
-	// A (-1) * B = VQS(B)(VQS (A)(-1))
+	// A(-1) * B = VQS(B) * VQS(A)(-1)
 	// 
-	// Scale = S(B)/S(A)
+	// Scale = S(B) / S(A)
 	// Rotation = Q(B) * Q(A)(-1)
-	// Translation = T(B)-S(B)/S(A) *[Q(B)*Q(A)(-1)*T(A)*Q(A)*Q(B)(-1)]
+	// Translation = T(B) - [Q(B) * Q(A)(-1) * S(B)/S(A) * T(A) * Q(A) * Q(B)(-1)]
 	// where A = this, and B = Other
 	TTransform<T> Result;
 
 	TVector<T> SafeRecipScale3D = GetSafeScaleReciprocal(Scale3D);
-	Result.Scale3D = Other.Scale3D*SafeRecipScale3D;	
+	Result.Scale3D = Other.Scale3D * SafeRecipScale3D;
 
-	Result.Rotation = Other.Rotation*Rotation.Inverse();
+	Result.Rotation = Other.Rotation * Rotation.Inverse();
 
-	Result.Translation = Other.Translation - Result.Scale3D * ( Result.Rotation * Translation );
+	Result.Translation = Other.Translation - (Result.Rotation * (Result.Scale3D * Translation));
 
 #if DEBUG_INVERSE_TRANSFORM
 	TMatrix<T> AM = ToMatrixWithScale();
 	TMatrix<T> BM = Other.ToMatrixWithScale();
 
-	Result.DebugEqualMatrix(AM.InverseFast() *  BM);
+	Result.DebugEqualMatrix(AM.InverseFast() * BM);
 #endif
 
 	return Result;
diff --git a/Engine/Source/Runtime/Core/Private/Math/TransformVectorized.cpp b/Engine/Source/Runtime/Core/Private/Math/TransformVectorized.cpp
--- a/Engine/Source/Runtime/Core/Private/Math/TransformVectorized.cpp
+++ b/Engine/Source/Runtime/Core/Private/Math/TransformVectorized.cpp
@@ -104,39 +104,39 @@ bool TTransform<T>::InitFromString( const FString& Source )
 template<typename T>
 TTransform<T> TTransform<T>::GetRelativeTransformReverse(const TTransform<T>& Other) const
 {
-	// A (-1) * B = VQS(B)(VQS (A)(-1))
+	// A(-1) * B = VQS(B) * VQS(A)(-1)
 	// 
-	// Scale = S(B)/S(A)
+	// Scale = S(B) / S(A)
 	// Rotation = Q(B) * Q(A)(-1)
-	// Translation = T(B)-S(B)/S(A) *[Q(B)*Q(A)(-1)*T(A)*Q(A)*Q(B)(-1)]
+	// Translation = T(B) - [Q(B) * Q(A)(-1) * S(B)/S(A) * T(A) * Q(A) * Q(B)(-1)]
 	// where A = this, and B = Other
 	TTransform<T> Result;
 
-	// Scale = S(B)/S(A)	
-	TransformVectorRegister VSafeScale3D	= VectorSet_W0(GetSafeScaleReciprocal(Scale3D));
+	// Scale = S(B) / S(A)
+	TransformVectorRegister VSafeScale3D = VectorSet_W0(GetSafeScaleReciprocal(Scale3D));
 	TransformVectorRegister VScale3D = VectorMultiply(Other.Scale3D, VSafeScale3D);
-	
-	// Rotation = Q(B) * Q(A)(-1)	
+
+	// Rotation = Q(B) * Q(A)(-1)
 	TransformVectorRegister VInverseRot = VectorQuaternionInverse(Rotation);
-	TransformVectorRegister VRotation = VectorQuaternionMultiply2(Other.Rotation, VInverseRot );
-	
-	// RotatedTranslation
-	TransformVectorRegister VR = VectorQuaternionRotateVector(VRotation, Translation);
+	TransformVectorRegister VRotation = VectorQuaternionMultiply2(Other.Rotation, VInverseRot);
+
+	// RotatedScaledTranslation
+	TransformVectorRegister VRS = VectorQuaternionRotateVector(VRotation, VectorMultiply(VScale3D, Translation));
 
-	// Translation = T(B)-S(B)/S(A) *[Q(B)*Q(A)(-1)*T(A)*Q(A)*Q(B)(-1)]	
-	TransformVectorRegister VTranslation = VectorSet_W0(VectorNegateMultiplyAdd(VScale3D, VR, Other.Translation));
+	// Translation = T(B) - [Q(B) * Q(A)(-1) * S(B)/S(A) * T(A) * Q(A) * Q(B)(-1)]
+	TransformVectorRegister VTranslation = VectorSet_W0(VectorSubtract(Other.Translation, VRS));
 
-	Result.Scale3D = VScale3D;	
+	Result.Scale3D = VScale3D;
 	Result.Translation = VTranslation;
 	Result.Rotation = VRotation;
 		
-	Result.DiagnosticCheckNaN_All(); 
+	Result.DiagnosticCheckNaN_All();
 
 #if DEBUG_INVERSE_TRANSFORM
 	TMatrix<T> AM = ToMatrixWithScale();
 	TMatrix<T> BM = Other.ToMatrixWithScale();
 
-	Result.DebugEqualMatrix(AM.InverseFast() *  BM);
+	Result.DebugEqualMatrix(AM.InverseFast() * BM);
 #endif
 
 	return Result;
diff --git a/Engine/Source/Runtime/Core/Tests/Math/UnrealMathTest.cpp b/Engine/Source/Runtime/Core/Tests/Math/UnrealMathTest.cpp
--- a/Engine/Source/Runtime/Core/Tests/Math/UnrealMathTest.cpp
+++ b/Engine/Source/Runtime/Core/Tests/Math/UnrealMathTest.cpp
@@ -1522,6 +1522,34 @@ FORCENOINLINE void TestVectorExpLogFunctions()
 	}
 }
 
+template<typename TransformType>
+FORCENOINLINE void TestTransformFunctions()
+{
+	using FloatType = typename TransformType::FReal;
+	using VectorType = UE::Math::TVector<FloatType>;
+	using RotatorType = UE::Math::TRotator<FloatType>;
+	using Matrix44Type = UE::Math::TMatrix<FloatType>;
+
+	// Negative scale isn't behaving nicely with FTransform, stick to positive non-uniform for now
+	VectorType ScaleA(2.5f, 0.5f, 1.2f);
+	VectorType ScaleB(2.215f, 0.75f, 1.215f);
+
+	VectorType TranslationA(10.34f, -188.1f, 0.017f);
+	VectorType TranslationB(-110.354f, 18.16f, 0.17017f);
+
+	RotatorType RotationA(19.5f, 71.81f, 0.18f);
+	RotatorType RotationB(-13.5f, 1.781f, 7.18f);
+
+	TransformType TransformA(RotationA, TranslationA, ScaleA);
+	TransformType TransformB(RotationB, TranslationB, ScaleB);
+
+	TransformType TransformAB = TransformA * TransformB;
+
+	TransformType ExpectedTransformB = TransformA.GetRelativeTransformReverse(TransformAB);
+
+	LogTest<FloatType>(TEXT("FTransform::GetRelativeTransformReverse"), ExpectedTransformB.Equals(TransformB));
+}
+
 
 bool RunDoubleVectorTest()
 {
@@ -2197,6 +2225,8 @@ bool RunDoubleVectorTest()
 		LogTest<double>(TEXT("VectorTransformVector"), TestVectorsEqual(V1, V2, 1e-8));
 	}
 
+	// FTransform
+	TestTransformFunctions<FTransform3d>();
 
 	// Quat / Rotator conversion to vectors, matrices
 	{
@@ -3054,6 +3084,9 @@ TEST_CASE_NAMED(FVectorRegisterAbstractionTest, "System::Core::Math::Vector Regi
 		LogTest<float>(TEXT("VectorTransformVector"), TestVectorsEqual(V1, V2, 1e-8f));
 	}
 
+	// FTransform
+	TestTransformFunctions<FTransform3f>();
+
 	// NaN / Inf tests
 	SetScratch(0.0f, 0.0f, 0.0f, 0.0f);
 	LogTest<float>(TEXT("VectorContainsNaNOrInfinite true"), VectorContainsNaNOrInfinite(MakeVectorRegister(NaN, NaN, NaN, NaN)));
