Add SCOPED_NAMED_EVENT_STATID to cover for cases where FScopeCycleCounter was used explicitly when ENABLE_STATNAMEDEVENTS is set and add it to Tickables.cpp
#rnx

diff --git a/Engine/Source/Runtime/Core/Public/Stats/Stats.h b/Engine/Source/Runtime/Core/Public/Stats/Stats.h
--- a/Engine/Source/Runtime/Core/Public/Stats/Stats.h
+++ b/Engine/Source/Runtime/Core/Public/Stats/Stats.h
@@ -274,6 +274,10 @@ FORCEINLINE void StatsPrimaryEnableSubtract(int32 Value = 1)
 	FScopeCycleCounter StatNamedEventsScope_##Stat(TStatId(ANSI_TO_PROFILING(#Stat))); \
 	SCOPE_CYCLE_COUNTER_TO_TRACE(#Stat, Stat, true);
 
+#define SCOPE_CYCLE_COUNTER_STATID(StatId) \
+	FScopeCycleCounter StatNamedEventsScope_STATID(StatId); \
+	TRACE_CPUPROFILER_EVENT_SCOPE_TEXT_CONDITIONAL(StatId.StatString, GCycleStatsShouldEmitNamedEvents > 0);
+
 #define CONDITIONAL_SCOPE_CYCLE_COUNTER(Stat,bCondition) \
 	FScopeCycleCounter StatNamedEventsScope_##Stat(bCondition ? ANSI_TO_PROFILING(#Stat) : nullptr); \
 	SCOPE_CYCLE_COUNTER_TO_TRACE(#Stat, Stat, bCondition);
@@ -315,6 +319,9 @@ public:
 #define SCOPE_CYCLE_COUNTER(Stat) \
 	FLightweightStatScope LightweightStatScope_##Stat(TEXT(#Stat));
 
+#define SCOPE_CYCLE_COUNTER_STATID(StatId) \
+	FLightweightStatScope LightweightStatScope_##Stat(TEXT("Lightweight StatId Scope"));
+
 #define CONDITIONAL_SCOPE_CYCLE_COUNTER(Stat,bCondition) \
 	FLightweightStatScope LightweightStatScope_##Stat(bCondition ? TEXT(#Stat) : nullptr);
 
@@ -323,6 +330,7 @@ public:
 
 #else
 #define SCOPE_CYCLE_COUNTER(Stat)
+#define SCOPE_CYCLE_COUNTER_STATID(StatId)
 #define QUICK_SCOPE_CYCLE_COUNTER(Stat)
 #define DECLARE_SCOPE_CYCLE_COUNTER(CounterName,StatId,GroupId)
 #define CONDITIONAL_SCOPE_CYCLE_COUNTER(Stat,bCondition)
diff --git a/Engine/Source/Runtime/Core/Public/Stats/Stats2.h b/Engine/Source/Runtime/Core/Public/Stats/Stats2.h
--- a/Engine/Source/Runtime/Core/Public/Stats/Stats2.h
+++ b/Engine/Source/Runtime/Core/Public/Stats/Stats2.h
@@ -2115,6 +2115,9 @@ struct FStat_##StatName\
 #define SCOPE_CYCLE_COUNTER(Stat) \
 	FScopeCycleCounter CycleCount_##Stat(GET_STATID(Stat), GET_STATFLAGS(Stat));
 
+#define SCOPE_CYCLE_COUNTER_STATID(StatId) \
+	FScopeCycleCounter CycleCount_FromStatId(StatId);
+
 #define CONDITIONAL_SCOPE_CYCLE_COUNTER(Stat,bCondition) \
 	FScopeCycleCounter CycleCount_##Stat(bCondition ? GET_STATID(Stat) : TStatId(), GET_STATFLAGS(Stat));
 
diff --git a/Engine/Source/Runtime/Engine/Private/Tickable.cpp b/Engine/Source/Runtime/Engine/Private/Tickable.cpp
--- a/Engine/Source/Runtime/Engine/Private/Tickable.cpp
+++ b/Engine/Source/Runtime/Engine/Private/Tickable.cpp
@@ -192,7 +192,7 @@ void FTickableGameObject::TickObjects(UWorld* World, const ELevelTick LevelTickT
 					if ((GIsEditor && TickableObject->IsTickableInEditor()) ||
 						(bIsGameWorld && ((!bIsPaused && LevelTickType != LEVELTICK_TimeOnly) || (bIsPaused && TickableObject->IsTickableWhenPaused()))))
 					{
-						FScopeCycleCounter Context(TickableObject->GetStatId());
+						SCOPE_CYCLE_COUNTER_STATID(TickableObject->GetStatId());
 						TickableObject->Tick(DeltaSeconds);
 					}
 				}
