#UE Add ability to opt-in to use threads before the fork on forking servers

diff --git a/Engine/Source/Runtime/Core/Private/HAL/ThreadingBase.cpp b/Engine/Source/Runtime/Core/Private/HAL/ThreadingBase.cpp
--- a/Engine/Source/Runtime/Core/Private/HAL/ThreadingBase.cpp
+++ b/Engine/Source/Runtime/Core/Private/HAL/ThreadingBase.cpp
@@ -1513,6 +1513,7 @@ void FTlsAutoCleanup::Register()
  * This thread starts as a fake thread and gets ticked like it was in a single-threaded environment.
  * Once it receives the OnPostFork event it creates and holds a real thread that
  * will cause the RunnableObject to be executed in it's own thread.
+ * Optionally, it can create the real thread immediately, but it also asserts that the real thread has exited by the fork point.
  */
 class FForkableThread : public FFakeThread
 {
@@ -1529,6 +1530,10 @@ private:
 
 public:
 
+	FForkableThread(bool InAllowPreForkRealThread)
+		: bAllowPreForkRealThread(InAllowPreForkRealThread)
+	{}
+
 	virtual ~FForkableThread()
 	{
 		delete RealThread;
@@ -1598,8 +1603,20 @@ public:
 		checkf(FForkProcessHelper::SupportsMultithreadingPostFork(), TEXT("ForkableThreads should only be created when -PostForkThreading is enabled"));
 		checkf(FForkProcessHelper::IsForkedMultithreadInstance() == false, TEXT("Once forked we create a real runnable thread instead of a ForkableThread"));
 
-		// Call the fake thread creator
-		bool bCreated = Super::CreateInternal(InRunnable, InThreadName, InStackSize, InThreadPri, InThreadAffinityMask, InCreateFlags);
+		bool bCreated = false;
+		if (bAllowPreForkRealThread)
+		{
+			ThreadName = InThreadName;
+			ThreadAffinityMask = InThreadAffinityMask;
+			Runnable = InRunnable;
+			InRunnable->Init();
+			bCreated = CreateRealThread();
+		}
+		else
+		{
+			// Call the fake thread creator
+			bCreated = Super::CreateInternal(InRunnable, InThreadName, InStackSize, InThreadPri, InThreadAffinityMask, InCreateFlags);
+		}
 
 		// Cache the target values until we create the real thread
 		CachedStackSize = InStackSize;
@@ -1610,10 +1627,19 @@ public:
 
 protected:
 
+	virtual void OnPreFork() override
+	{
+		checkf(!bAllowPreForkRealThread, TEXT("A forkable thread that allows real threading before the fork point was left alive. You must delete all forkable threads that were created with bAllowPrefork == true before forking. ThreadName:%s ThreadID:%d"), *GetThreadName(), GetThreadID());
+	}
+
 	virtual void OnPostFork() override
 	{
 		check(FForkProcessHelper::IsForkedMultithreadInstance());
+		CreateRealThread();
+	}
 
+	bool CreateRealThread()
+	{
 		check(RealThread == nullptr);
 		RealThread = FPlatformProcess::CreateRunnableThread();
 		bool bCreated = RealThread->CreateInternal(Runnable, *GetThreadName(), CachedStackSize, CachedPriority, ThreadAffinityMask, EThreadCreateFlags::None);
@@ -1634,10 +1660,14 @@ protected:
 			delete RealThread;
 			RealThread = nullptr;
 		}
+
+		return bCreated;
 	}
+
+	bool bAllowPreForkRealThread;
 };
 
-FRunnableThread* FForkProcessHelper::CreateForkableThread(class FRunnable* InRunnable, const TCHAR* InThreadName, uint32 InStackSize, EThreadPriority InThreadPri, uint64 InThreadAffinityMask, EThreadCreateFlags InCreateFlags)
+FRunnableThread* FForkProcessHelper::CreateForkableThread(class FRunnable* InRunnable, const TCHAR* InThreadName, uint32 InStackSize, EThreadPriority InThreadPri, uint64 InThreadAffinityMask, EThreadCreateFlags InCreateFlags, bool bAllowPrefork)
 {
 	bool bCreateRealThread = FPlatformProcess::SupportsMultithreading();
 	bool bCreateForkableThread(false);
@@ -1669,9 +1699,13 @@ FRunnableThread* FForkProcessHelper::CreateForkableThread(class FRunnable* InRun
 	}
 	else if (bCreateForkableThread)
 	{
-		if( InRunnable->GetSingleThreadInterface() )
+		if (bAllowPrefork)
+		{
+			NewThread = new FForkableThread(true);
+		}
+		else if( InRunnable->GetSingleThreadInterface() )
 		{
-			NewThread = new FForkableThread();
+			NewThread = new FForkableThread(false);
 		}
 	}
 	else
diff --git a/Engine/Source/Runtime/Core/Private/Misc/Fork.cpp b/Engine/Source/Runtime/Core/Private/Misc/Fork.cpp
--- a/Engine/Source/Runtime/Core/Private/Misc/Fork.cpp
+++ b/Engine/Source/Runtime/Core/Private/Misc/Fork.cpp
@@ -110,6 +110,16 @@ void FForkProcessHelper::OnForkingOccured()
 	}
 }
 
+void FForkProcessHelper::OnPreFork()
+{
+	// Use a local list of forkable threads so we don't keep a lock on the global list during thread creation
+	TArray<FRunnableThread*> ForkableThreads = FThreadManager::Get().GetForkableThreads();
+	for (FRunnableThread* ForkableThread : ForkableThreads)
+	{
+		ForkableThread->OnPreFork();
+	}
+}
+
 bool FForkProcessHelper::IsForkedMultithreadInstance()
 {
 	return ForkInternal::bIsForkedMultithreadInstance;
diff --git a/Engine/Source/Runtime/Core/Public/HAL/RunnableThread.h b/Engine/Source/Runtime/Core/Public/HAL/RunnableThread.h
--- a/Engine/Source/Runtime/Core/Public/HAL/RunnableThread.h
+++ b/Engine/Source/Runtime/Core/Public/HAL/RunnableThread.h
@@ -196,6 +196,14 @@ private:
 	/** Used by the thread manager to tick threads in single-threaded mode */
 	virtual void Tick() {}
 
+	/**
+	 * Called on the parent process before the fork occurs
+	 */
+	virtual void OnPreFork()
+	{
+		checkf(false, TEXT("Only forkable threads should receive OnPreFork."));
+	}
+
 	/**
 	 * Called on the forked process when the forkable thread can create a real thread
 	 */
diff --git a/Engine/Source/Runtime/Core/Public/Misc/Fork.h b/Engine/Source/Runtime/Core/Public/Misc/Fork.h
--- a/Engine/Source/Runtime/Core/Public/Misc/Fork.h
+++ b/Engine/Source/Runtime/Core/Public/Misc/Fork.h
@@ -69,6 +69,11 @@ public:
 	 */
 	static CORE_API void OnForkingOccured();
 
+	/**
+	 * Event triggered when a fork is about to occur on the parent process
+	 */
+	static CORE_API void OnPreFork();
+
 	/**
 	 * Tells if we allow multithreading on forked processes.
 	 * Default is set to false but can be configured to always be true via DEFAULT_MULTITHREAD_FORKED_PROCESSES
@@ -91,7 +96,8 @@ public:
 		uint32 InStackSize = 0,
 		EThreadPriority InThreadPri = TPri_Normal,
 		uint64 InThreadAffinityMask = FPlatformAffinity::GetNoAffinityMask(),
-		EThreadCreateFlags InCreateFlags = EThreadCreateFlags::None
+		EThreadCreateFlags InCreateFlags = EThreadCreateFlags::None,
+		bool bAllowPreFork = false
 	);
 
 	/**
