[AutoRTFM] Disable force-inlining for debug builds.

To help improve the debugger experience.

diff --git a/Engine/Source/Runtime/Core/Public/AutoRTFM/AutoRTFM.h b/Engine/Source/Runtime/Core/Public/AutoRTFM/AutoRTFM.h
--- a/Engine/Source/Runtime/Core/Public/AutoRTFM/AutoRTFM.h
+++ b/Engine/Source/Runtime/Core/Public/AutoRTFM/AutoRTFM.h
@@ -74,6 +74,14 @@
 #define UE_AUTORTFM_MOVE MoveTemp
 #endif
 
+#if defined(UE_BUILD_DEBUG) && UE_BUILD_DEBUG
+// Force-inlining can make debugging glitchy. Disable this if we're running a debug build.
+#undef UE_AUTORTFM_CALLSITE_FORCEINLINE
+#define UE_AUTORTFM_CALLSITE_FORCEINLINE
+#undef UE_AUTORTFM_FORCEINLINE
+#define UE_AUTORTFM_FORCEINLINE inline
+#endif
+
 #if UE_AUTORTFM
 template <typename FuncType> class TFunction;
 #endif
