[AutoRTFM] Split variadic UE_AUTORTFM_ALWAYS_OPEN functions to va_list / variadic pair

The new memory validator does not support variadic functions annotated with UE_AUTORTFM_ALWAYS_OPEN. Move the body of these functions to a new va_list form, and call that from the variadic function.

diff --git a/Engine/Source/Runtime/Core/Private/Logging/StructuredLog.cpp b/Engine/Source/Runtime/Core/Private/Logging/StructuredLog.cpp
--- a/Engine/Source/Runtime/Core/Private/Logging/StructuredLog.cpp
+++ b/Engine/Source/Runtime/Core/Private/Logging/StructuredLog.cpp
@@ -24,6 +24,7 @@
 #include "Serialization/CompactBinaryValue.h"
 #include "Serialization/VarInt.h"
 #include "AutoRTFM/AutoRTFM.h"
+#include <cstdarg>
 
 void VARARGS StaticFailDebug (const TCHAR* Error, const ANSICHAR* Expression, const ANSICHAR* File, int32 Line, bool bIsEnsure, void* ProgramCounter, const TCHAR* DescriptionFormat, ...);
 void         StaticFailDebugV(const TCHAR* Error, const ANSICHAR* Expression, const ANSICHAR* File, int32 Line, bool bIsEnsure, void* ProgramCounter, const TCHAR* DescriptionFormat, va_list DescriptionArgs);
@@ -1037,17 +1038,14 @@ FORCENOINLINE static FCbObject SerializeBasicLogMessage(const FStaticBasicLogRec
 }
 
 UE_AUTORTFM_ALWAYS_OPEN
-void BasicLog(const FLogCategoryBase& Category, const FStaticBasicLogRecord* Log, ...)
+static void BasicLogV(const FLogCategoryBase& Category, const FStaticBasicLogRecord* Log, va_list Args)
 {
 #if !NO_LOGGING
 	EnsureBasicLogMessageSpec(Category, *Log);
 
 	if (GConvertBasicLogToLogRecord)
 	{
-		va_list Args;
-		va_start(Args, Log);
 		FCbObject Fields = SerializeBasicLogMessage(*Log, Args);
-		va_end(Args);
 
 		FLogRecord Record;
 		Record.SetFormat(GStaticBasicLogFormat);
@@ -1075,36 +1073,49 @@ void BasicLog(const FLogCategoryBase& Category, const FStaticBasicLogRecord* Log
 	}
 	else
 	{
-		va_list Args;
 	#if LOGTRACE_ENABLED
 		if (UE_TRACE_CHANNELEXPR_IS_ENABLED(LogChannel))
 		{
-			va_start(Args, Log);
-			BasicLogToTrace(Log, Log->Format, Args);
-			va_end(Args);
+			va_list Args2;
+			va_copy(Args2, Args);
+			BasicLogToTrace(Log, Log->Format, Args2);
 		}
 	#endif
-		va_start(Args, Log);
 		FMsg::LogV(Log->File, Log->Line, Category.GetCategoryName(), Log->Verbosity, Log->Format, Args);
-		va_end(Args);
 	}
 #endif
 }
 
-UE_AUTORTFM_ALWAYS_OPEN
-void BasicFatalLog(const FLogCategoryBase& Category, const FStaticBasicLogRecord* Log, ...)
+void BasicLog(const FLogCategoryBase& Category, const FStaticBasicLogRecord* Log, ...)
 {
 #if !NO_LOGGING
 	va_list Args;
 	va_start(Args, Log);
-	StaticFailDebugV(TEXT("Fatal error:"), "", Log->File, Log->Line, /*bIsEnsure*/ false, PLATFORM_RETURN_ADDRESS(), Log->Format, Args);
+	BasicLogV(Category, Log, Args);
 	va_end(Args);
+#endif
+}
+
+UE_AUTORTFM_ALWAYS_OPEN
+static void BasicFatalLogV(const FLogCategoryBase& Category, const FStaticBasicLogRecord* Log, va_list Args)
+{
+#if !NO_LOGGING
+	StaticFailDebugV(TEXT("Fatal error:"), "", Log->File, Log->Line, /*bIsEnsure*/ false, PLATFORM_RETURN_ADDRESS(), Log->Format, Args);
 
 	UE_DEBUG_BREAK_AND_PROMPT_FOR_REMOTE();
 	FDebug::ProcessFatalError(PLATFORM_RETURN_ADDRESS());
 #endif
 }
 
+void BasicFatalLog(const FLogCategoryBase& Category, const FStaticBasicLogRecord* Log, ...)
+{
+#if !NO_LOGGING
+	va_list Args;
+	va_start(Args, Log);
+	BasicFatalLogV(Category, Log, Args);
+	va_end(Args);
+#endif
+}
 } // UE::Logging::Private
 
 ///////////////////////////////////////////////////////////////////////////////////////////////////////////////////////
diff --git a/Engine/Source/Runtime/Core/Private/Misc/AssertionMacros.cpp b/Engine/Source/Runtime/Core/Private/Misc/AssertionMacros.cpp
--- a/Engine/Source/Runtime/Core/Private/Misc/AssertionMacros.cpp
+++ b/Engine/Source/Runtime/Core/Private/Misc/AssertionMacros.cpp
@@ -621,18 +621,14 @@ bool FORCENOINLINE FDebug::CheckVerifyFailedImpl(
 }
 
 UE_AUTORTFM_ALWAYS_OPEN
-bool FORCENOINLINE FDebug::CheckVerifyFailedImpl2(
+bool FORCENOINLINE FDebug::CheckVerifyFailedImpl2V(
 	const ANSICHAR* Expr,
 	const ANSICHAR* File,
 	int32 Line,
 	const TCHAR* Format,
-	...)
+	va_list Args)
 {
-	va_list Args;
-
-	va_start(Args, Format);
 	FDebug::LogAssertFailedMessageImplV(Expr, File, Line, PLATFORM_RETURN_ADDRESS(), Format, Args);
-	va_end(Args);
 
 	if (GLog)
 	{
@@ -645,9 +641,9 @@ bool FORCENOINLINE FDebug::CheckVerifyFailedImpl2(
 	{
 		FPlatformMisc::PromptForRemoteDebugging(false);
 
-		va_start(Args, Format);
-		AssertFailedImplV(Expr, File, Line, PLATFORM_RETURN_ADDRESS(), Format, Args);
-		va_end(Args);
+		va_list Args2;
+		va_copy(Args2, Args);
+		AssertFailedImplV(Expr, File, Line, PLATFORM_RETURN_ADDRESS(), Format, Args2);
 
 		return false;
 	}
@@ -659,6 +655,21 @@ bool FORCENOINLINE FDebug::CheckVerifyFailedImpl2(
 #endif
 }
 
+bool FORCENOINLINE FDebug::CheckVerifyFailedImpl2(
+	const ANSICHAR* Expr,
+	const ANSICHAR* File,
+	int32 Line,
+	const TCHAR* Format,
+	...)
+{
+	va_list Args;
+	va_start(Args, Format);
+	bool bResult = CheckVerifyFailedImpl2V(Expr, File, Line, Format, Args);
+	va_end(Args);
+
+	return bResult;
+}
+
 #endif // DO_CHECK || DO_GUARD_SLOW || DO_ENSURE
 
 void VARARGS FDebug::AssertFailed(const ANSICHAR* Expr, const ANSICHAR* File, int32 Line, const TCHAR* Format/* = TEXT("")*/, ...)
@@ -797,16 +808,21 @@ bool UE_COLD UE_DEBUG_SECTION UE::Assert::Private::ExecCheckImplInternal(std::at
 }
 
 UE_AUTORTFM_ALWAYS_OPEN
-bool UE_COLD UE_DEBUG_SECTION VARARGS UE::Assert::Private::EnsureFailed(std::atomic<bool>& bExecuted, const FStaticEnsureRecord* Ensure, ...)
+bool UE_COLD UE_DEBUG_SECTION UE::Assert::Private::EnsureFailedV(std::atomic<bool>& bExecuted, const FStaticEnsureRecord* Ensure, void* ProgramCounter, va_list Args)
 {
 	if (bExecuted.exchange(true, std::memory_order_release) && !(Ensure->bAlways && GEnsureAlwaysEnabled))
 	{
 		return false;
 	}
 
+	return CheckVerifyImpl(bExecuted, Ensure->bAlways, Ensure->File, Ensure->Line, ProgramCounter, Ensure->Expression, Ensure->Format, Args);
+}
+
+bool UE_COLD UE_DEBUG_SECTION VARARGS UE::Assert::Private::EnsureFailed(std::atomic<bool>& bExecuted, const FStaticEnsureRecord* Ensure, ...)
+{
 	va_list Args;
 	va_start(Args, Ensure);
-	const bool bResult = CheckVerifyImpl(bExecuted, Ensure->bAlways, Ensure->File, Ensure->Line, PLATFORM_RETURN_ADDRESS(), Ensure->Expression, Ensure->Format, Args);
+	const bool bResult = EnsureFailedV(bExecuted, Ensure, PLATFORM_RETURN_ADDRESS(), Args);
 	va_end(Args);
 
 	return bResult;
diff --git a/Engine/Source/Runtime/Core/Public/Misc/AssertionMacros.h b/Engine/Source/Runtime/Core/Public/Misc/AssertionMacros.h
--- a/Engine/Source/Runtime/Core/Public/Misc/AssertionMacros.h
+++ b/Engine/Source/Runtime/Core/Public/Misc/AssertionMacros.h
@@ -96,6 +96,7 @@ public:
 private:
 	static CORE_API void VARARGS LogAssertFailedMessageImpl(const ANSICHAR* Expr, const ANSICHAR* File, int32 Line, void* ProgramCounter, const TCHAR* Fmt, ...);
 	static CORE_API void LogAssertFailedMessageImplV(const ANSICHAR* Expr, const ANSICHAR* File, int32 Line, void* ProgramCounter, const TCHAR* Fmt, va_list Args);
+	static CORE_API bool CheckVerifyFailedImpl2V(const ANSICHAR* Expr, const ANSICHAR* File, int32 Line, const TCHAR* Format, va_list Args);
 
 public:
 //	/**
@@ -390,6 +391,7 @@ RetType UE_COLD UE_DEBUG_SECTION DispatchCheckVerify (InnerType&& Inner, ArgType
 	};
 
 	CORE_API bool UE_COLD UE_DEBUG_SECTION VARARGS EnsureFailed(std::atomic<bool>& bExecuted, const FStaticEnsureRecord* Ensure, ...);
+	CORE_API bool UE_COLD UE_DEBUG_SECTION EnsureFailedV(std::atomic<bool>& bExecuted, const FStaticEnsureRecord* Ensure, void* ProgramCounter, va_list Args);
 
 	CORE_API bool UE_COLD UE_DEBUG_SECTION ExecCheckImplInternal(std::atomic<bool>& bExecuted, bool bAlways, const ANSICHAR* File, int32 Line, const ANSICHAR* Expr);
 
