[Insights] Fixed timing regions emitted from SlowTask to use unique ids. Matching by name was not reliable when multiple slow tasks with same name are active at the same time.

diff --git a/Engine/Source/Runtime/Core/Private/Misc/SlowTask.cpp b/Engine/Source/Runtime/Core/Private/Misc/SlowTask.cpp
--- a/Engine/Source/Runtime/Core/Private/Misc/SlowTask.cpp
+++ b/Engine/Source/Runtime/Core/Private/Misc/SlowTask.cpp
@@ -49,6 +49,7 @@ FSlowTask::FSlowTask(float InAmountOfWork, const FText& InDefaultMessage, bool b
 	, bCreatedDialog(false)		// only set to true if we create a dialog
 	, Context(InContext)
 	, bSkipRecursiveDialogCreation(false)
+	, RegionId(0)
 {
 	// If we have no work to do ourselves, create an arbitrary scope so that any actions performed underneath this still contribute to this one.
 	if (TotalAmountOfWork == 0.f)
@@ -100,17 +101,12 @@ void FSlowTask::Initialize()
 	if (bEnabled)
 	{
 		Context.ScopeStack.Push(this);
+#if MISCTRACE_ENABLED
 		if (Context.ScopeStack.Num() <= GSlowTaskMaxTraceRegionDepth)
 		{
-			if (!DefaultMessage.IsEmpty())
-			{
-				TRACE_BEGIN_REGION(*DefaultMessage.ToString());
-			}
-			else
-			{
-				TRACE_BEGIN_REGION(TEXT("<SlowTask>"));
-			}
+			RegionId = TRACE_BEGIN_REGION_WITH_ID(!DefaultMessage.IsEmpty() ? *DefaultMessage.ToString() : TEXT("<SlowTask>"));
 		}
+#endif
 	}
 }
 
@@ -139,17 +135,15 @@ void FSlowTask::Destroy()
 		FSlowTaskStack& Stack = Context.ScopeStack;
 		if (ensure(Stack.Num() != 0))
 		{
+#if MISCTRACE_ENABLED
 			if (Context.ScopeStack.Num() <= GSlowTaskMaxTraceRegionDepth)
 			{
-				if (!DefaultMessage.IsEmpty())
-				{
-					TRACE_END_REGION(*DefaultMessage.ToString());
-				}
-				else
+				if (RegionId != 0)
 				{
-					TRACE_END_REGION(TEXT("<SlowTask>"));
+					TRACE_END_REGION_WITH_ID(RegionId);
 				}
 			}
+#endif
 
 			FSlowTask* Task = Stack.Last();
 			if (ensureMsgf(Task == this, TEXT("Out-of-order slow task construction/destruction: destroying '%s' but '%s' is at the top of the stack"), *DefaultMessage.ToString(), *Task->DefaultMessage.ToString()))
diff --git a/Engine/Source/Runtime/Core/Public/Misc/SlowTask.h b/Engine/Source/Runtime/Core/Public/Misc/SlowTask.h
--- a/Engine/Source/Runtime/Core/Public/Misc/SlowTask.h
+++ b/Engine/Source/Runtime/Core/Public/Misc/SlowTask.h
@@ -71,6 +71,9 @@ private:
 	/** Avoid calling hierarchical dialog creation if not required */
 	bool bSkipRecursiveDialogCreation : 1;
 
+	/** Id of the traced timing region */
+	uint64 RegionId;
+
 	/** Prevent copying */
 	FSlowTask(const FSlowTask&);
 
