[Cleanup] Make sure variables are initialized even if APIs fail.

diff --git a/Engine/Source/Runtime/Core/Private/AutoRTFM/Context.cpp b/Engine/Source/Runtime/Core/Private/AutoRTFM/Context.cpp
--- a/Engine/Source/Runtime/Core/Private/AutoRTFM/Context.cpp
+++ b/Engine/Source/Runtime/Core/Private/AutoRTFM/Context.cpp
@@ -358,10 +358,11 @@ ETransactionResult FContext::Transact(void (*InstrumentedFunction)(void*), void*
 #elif defined(__APPLE__)         
 		Stack.High = pthread_get_stackaddr_np(pthread_self());
 		size_t StackSize = pthread_get_stacksize_np(pthread_self());
-		StackLow = static_cast<char*>(Stack.High) - StackSize;
+		Stack.Low = static_cast<char*>(Stack.High) - StackSize;
 #else
-		pthread_attr_t Attr;
+		pthread_attr_t Attr{};
 		pthread_getattr_np(pthread_self(), &Attr);
+		Stack.Low = 0;
 		size_t StackSize = 0;
 		pthread_attr_getstack(&Attr, &Stack.Low, &StackSize);
 		Stack.High = static_cast<char*>(Stack.Low) + StackSize;
