Add missing va_end() calls for va_copy() and fix copy for CheckVerifyFailedImpl2V()

Spotted by [REDACTED_PERSON] in post-submit review of 36529254.

diff --git a/Engine/Source/Runtime/Core/Private/AutoRTFM/StdLibSupport.cpp b/Engine/Source/Runtime/Core/Private/AutoRTFM/StdLibSupport.cpp
--- a/Engine/Source/Runtime/Core/Private/AutoRTFM/StdLibSupport.cpp
+++ b/Engine/Source/Runtime/Core/Private/AutoRTFM/StdLibSupport.cpp
@@ -413,14 +413,16 @@ int RTFM_vsnprintf(char* Str, size_t Size, const char* Format, va_list ArgList)
     {
         va_list ArgList2;
         va_copy(ArgList2, ArgList);
-
-        FContext* Context = FContext::Get();
         int Count = vsnprintf(nullptr, 0, Format, ArgList2);
+        va_end(ArgList2);
+
         if (Count >= 0)
         {
             size_t NumBytes = FMath::Min<size_t>(Size, (1 + Count)) * sizeof(char);
+            FContext* Context = FContext::Get();
             Context->RecordWrite(Str, NumBytes);
         }
+
     }
 
     return vsnprintf(Str, Size, Format, ArgList);
@@ -444,6 +446,8 @@ int RTFM_vswprintf(wchar_t* Str, size_t Size, const wchar_t* Format, va_list Arg
         int Count = vfwprintf(FNullFile::Get(), Format, ArgList2);
 #endif
 
+        va_end(ArgList2);
+
         size_t NumChars = FMath::Min<size_t>(Size, (1 + FMath::Max(Count, 0)));
         size_t NumBytes = NumChars * sizeof(wchar_t);
         if (NumBytes >= 0)
@@ -574,12 +578,13 @@ int RTFM___stdio_common_vsprintf(
     {
         va_list ArgList2;
         va_copy(ArgList2, ArgList);
-
-        FContext* Context = FContext::Get();
         int Count = __stdio_common_vsprintf(Options, nullptr, 0, Format, Locale, ArgList2);
+        va_end(ArgList2);
+
         if (Count >= 0)
         {
             size_t NumBytes = FMath::Min<size_t>(BufferCount, (1 + Count)) * sizeof(char);
+            FContext* Context = FContext::Get();
             Context->RecordWrite(Buffer, NumBytes);
         }
     }
@@ -603,12 +608,13 @@ int RTFM___stdio_common_vswprintf(
     {
         va_list ArgList2;
         va_copy(ArgList2, ArgList);
-
-        FContext* Context = FContext::Get();
         int Count = __stdio_common_vswprintf(Options, nullptr, 0, Format, Locale, ArgList2);
+        va_end(ArgList2);
+
         if (Count >= 0)
         {
             size_t NumBytes = FMath::Min<size_t>(BufferCount, (1 + Count)) * sizeof(wchar_t);
+            FContext* Context = FContext::Get();
             Context->RecordWrite(Buffer, NumBytes);
         }
     }
diff --git a/Engine/Source/Runtime/Core/Private/Logging/StructuredLog.cpp b/Engine/Source/Runtime/Core/Private/Logging/StructuredLog.cpp
--- a/Engine/Source/Runtime/Core/Private/Logging/StructuredLog.cpp
+++ b/Engine/Source/Runtime/Core/Private/Logging/StructuredLog.cpp
@@ -1079,6 +1079,7 @@ static void BasicLogV(const FLogCategoryBase& Category, const FStaticBasicLogRec
 			va_list Args2;
 			va_copy(Args2, Args);
 			BasicLogToTrace(Log, Log->Format, Args2);
+			va_end(Args2);
 		}
 	#endif
 		FMsg::LogV(Log->File, Log->Line, Category.GetCategoryName(), Log->Verbosity, Log->Format, Args);
diff --git a/Engine/Source/Runtime/Core/Private/Misc/AssertionMacros.cpp b/Engine/Source/Runtime/Core/Private/Misc/AssertionMacros.cpp
--- a/Engine/Source/Runtime/Core/Private/Misc/AssertionMacros.cpp
+++ b/Engine/Source/Runtime/Core/Private/Misc/AssertionMacros.cpp
@@ -628,7 +628,10 @@ bool FORCENOINLINE FDebug::CheckVerifyFailedImpl2V(
 	const TCHAR* Format,
 	va_list Args)
 {
-	FDebug::LogAssertFailedMessageImplV(Expr, File, Line, PLATFORM_RETURN_ADDRESS(), Format, Args);
+	va_list Args2;
+	va_copy(Args2, Args);
+	FDebug::LogAssertFailedMessageImplV(Expr, File, Line, PLATFORM_RETURN_ADDRESS(), Format, Args2);
+	va_end(Args2);
 
 	if (GLog)
 	{
@@ -641,9 +644,7 @@ bool FORCENOINLINE FDebug::CheckVerifyFailedImpl2V(
 	{
 		FPlatformMisc::PromptForRemoteDebugging(false);
 
-		va_list Args2;
-		va_copy(Args2, Args);
-		AssertFailedImplV(Expr, File, Line, PLATFORM_RETURN_ADDRESS(), Format, Args2);
+		AssertFailedImplV(Expr, File, Line, PLATFORM_RETURN_ADDRESS(), Format, Args);
 
 		return false;
 	}
@@ -801,7 +802,8 @@ bool UE_COLD UE_DEBUG_SECTION UE::Assert::Private::ExecCheckImplInternal(std::at
 		}
 
 		va_list Args = {};
-		return CheckVerifyImpl(bExecuted, bAlways, File, Line, PLATFORM_RETURN_ADDRESS(), Expr, TEXT(""), Args);
+		const bool bResult = CheckVerifyImpl(bExecuted, bAlways, File, Line, PLATFORM_RETURN_ADDRESS(), Expr, TEXT(""), Args);
+		return bResult;
 	}
 
 	return false;
@@ -815,7 +817,8 @@ bool UE_COLD UE_DEBUG_SECTION UE::Assert::Private::EnsureFailedV(std::atomic<boo
 		return false;
 	}
 
-	return CheckVerifyImpl(bExecuted, Ensure->bAlways, Ensure->File, Ensure->Line, ProgramCounter, Ensure->Expression, Ensure->Format, Args);
+	const bool bResult = CheckVerifyImpl(bExecuted, Ensure->bAlways, Ensure->File, Ensure->Line, ProgramCounter, Ensure->Expression, Ensure->Format, Args);
+	return bResult;
 }
 
 bool UE_COLD UE_DEBUG_SECTION VARARGS UE::Assert::Private::EnsureFailed(std::atomic<bool>& bExecuted, const FStaticEnsureRecord* Ensure, ...)
