Color: Added more convenience around sRGB color space, including CDPR change.

#rnx

diff --git a/Engine/Source/Runtime/Core/Private/ColorManagement/ColorSpace.cpp b/Engine/Source/Runtime/Core/Private/ColorManagement/ColorSpace.cpp
--- a/Engine/Source/Runtime/Core/Private/ColorManagement/ColorSpace.cpp
+++ b/Engine/Source/Runtime/Core/Private/ColorManagement/ColorSpace.cpp
@@ -69,6 +69,13 @@ void FColorSpace::SetWorking(FColorSpace ColorSpace)
 	bIsWorkingColorSpaceReadyForUse = true;
 }
 
+const FColorSpace& FColorSpace::GetSRGB()
+{
+	static const FColorSpace ColorSpaceSRGB = FColorSpace(EColorSpace::sRGB);
+
+	return ColorSpaceSRGB;
+}
+
 static bool IsSRGBChromaticities(const TStaticArray<FVector2d, 4>& Chromaticities, double Tolerance = 1.e-7)
 {
 	return	Chromaticities[0].Equals(FVector2d(0.64, 0.33), Tolerance) &&
@@ -318,13 +325,20 @@ FMatrix44d FColorSpaceTransform::CalcChromaticAdaptionMatrix(FVector3d SourceXYZ
 	return XyzToRgb * ScaleMat * RgbToXyz;
 }
 
-FColorSpaceTransform FColorSpaceTransform::GetSRGBToWorkingColorSpace()
+const FColorSpaceTransform& FColorSpaceTransform::GetSRGBToWorkingColorSpace()
 {
-	static FColorSpaceTransform CachedTransform = FColorSpaceTransform(FColorSpace(EColorSpace::sRGB), FColorSpace::GetWorking());
+	static FColorSpaceTransform CachedTransform = FColorSpaceTransform(FColorSpace::GetSRGB(), FColorSpace::GetWorking());
 	
 	return CachedTransform;
 }
 
+const FColorSpaceTransform& FColorSpaceTransform::GetWorkingColorSpaceToSRGB()
+{
+	static FColorSpaceTransform CachedTransform = FColorSpaceTransform(FColorSpace::GetWorking(), FColorSpace::GetSRGB());
+
+	return CachedTransform;
+}
+
 static FMatrix44d CalcColorSpaceTransformMatrix(const FColorSpace& Src, const FColorSpace& Dst, EChromaticAdaptationMethod Method)
 {
 	if (Method == UE::Color::EChromaticAdaptationMethod::None)
diff --git a/Engine/Source/Runtime/Core/Public/ColorManagement/ColorSpace.h b/Engine/Source/Runtime/Core/Public/ColorManagement/ColorSpace.h
--- a/Engine/Source/Runtime/Core/Public/ColorManagement/ColorSpace.h
+++ b/Engine/Source/Runtime/Core/Public/ColorManagement/ColorSpace.h
@@ -111,6 +111,12 @@ public:
 	 */
 	static CORE_API void SetWorking(FColorSpace ColorSpace);
 
+	/**
+	 * Get a (statically cached) sRGB color space.
+	 *
+	 * @return FColorSpace sRGB color space
+	 */
+	static CORE_API const FColorSpace& GetSRGB();
 
 	/** Constructor */
 	FColorSpace() {}
@@ -360,7 +366,12 @@ struct FColorSpaceTransform : FMatrix44d
 	/**
 	* Convenience function to get a (statically cached) conversion from sRGB/Rec709 to the working color space.
 	*/
-	CORE_API static FColorSpaceTransform GetSRGBToWorkingColorSpace();
+	CORE_API static const FColorSpaceTransform& GetSRGBToWorkingColorSpace();
+
+	/**
+	 * Convenience function to get a (statically cached) conversion from the working color space to sRGB/Rec709.
+	*/
+	CORE_API static const FColorSpaceTransform& GetWorkingColorSpaceToSRGB();
 };
 
 } }  // end namespace UE::Color
