Fix FUnixPlatformProcess::GetDllHandle issue when non-module shared libraries are reloaded

#pr 12144
#12144

diff --git a/Engine/Source/Runtime/Core/Private/Unix/UnixPlatformProcess.cpp b/Engine/Source/Runtime/Core/Private/Unix/UnixPlatformProcess.cpp
--- a/Engine/Source/Runtime/Core/Private/Unix/UnixPlatformProcess.cpp
+++ b/Engine/Source/Runtime/Core/Private/Unix/UnixPlatformProcess.cpp
@@ -119,7 +119,8 @@ void* FUnixPlatformProcess::GetDllHandle( const TCHAR* Filename )
 
 		if (UpgradeToGlobal)
 		{
-			Handle = dlopen( TCHAR_TO_UTF8(*AbsolutePath), DlOpenMode | RTLD_NOLOAD | RTLD_GLOBAL );
+			dlclose( Handle );
+			Handle = dlopen( TCHAR_TO_UTF8(*AbsolutePath), DlOpenMode | RTLD_GLOBAL );
 		}
 	} 
 	else if (!FString(Filename).Contains(TEXT("/")))
