Fixed a missing #include in FormatTest.cpp

#rnx

diff --git a/Engine/Source/Runtime/Core/Tests/String/FormatTest.cpp b/Engine/Source/Runtime/Core/Tests/String/FormatTest.cpp
--- a/Engine/Source/Runtime/Core/Tests/String/FormatTest.cpp
+++ b/Engine/Source/Runtime/Core/Tests/String/FormatTest.cpp
@@ -2,6 +2,7 @@
 
 #if WITH_TESTS
 
+#include "Containers/Map.h"
 #include "Containers/UnrealString.h"
 #include "Tests/TestHarnessAdapter.h"
 
