ColorSpace: Added comments pertaining to the required project relaunch for cached static matrices.

diff --git a/Engine/Source/Runtime/Core/Private/ColorManagement/ColorSpace.cpp b/Engine/Source/Runtime/Core/Private/ColorManagement/ColorSpace.cpp
--- a/Engine/Source/Runtime/Core/Private/ColorManagement/ColorSpace.cpp
+++ b/Engine/Source/Runtime/Core/Private/ColorManagement/ColorSpace.cpp
@@ -327,6 +327,7 @@ FMatrix44d FColorSpaceTransform::CalcChromaticAdaptionMatrix(FVector3d SourceXYZ
 
 const FColorSpaceTransform& FColorSpaceTransform::GetSRGBToWorkingColorSpace()
 {
+	// Note: An editor restart is needed for the working color space shader compiler definitions (WORKING_COLOR_SPACE_RGB_TO_XYZ_MAT, XYZ_TO_RGB_WORKING_COLOR_SPACE_MAT, etc), so we can currently rely on the same assumption here.
 	static FColorSpaceTransform CachedTransform = FColorSpaceTransform(FColorSpace::GetSRGB(), FColorSpace::GetWorking());
 	
 	return CachedTransform;
@@ -334,6 +335,7 @@ const FColorSpaceTransform& FColorSpaceTransform::GetSRGBToWorkingColorSpace()
 
 const FColorSpaceTransform& FColorSpaceTransform::GetWorkingColorSpaceToSRGB()
 {
+	// Note: An editor restart is needed for the working color space shader compiler definitions (WORKING_COLOR_SPACE_RGB_TO_XYZ_MAT, XYZ_TO_RGB_WORKING_COLOR_SPACE_MAT, etc), so we can currently rely on the same assumption here.
 	static FColorSpaceTransform CachedTransform = FColorSpaceTransform(FColorSpace::GetWorking(), FColorSpace::GetSRGB());
 
 	return CachedTransform;
