IFileHandle seekless ReadAt

#rnx

diff --git a/Engine/Source/Runtime/Core/Private/Android/AndroidPlatformFile.cpp b/Engine/Source/Runtime/Core/Private/Android/AndroidPlatformFile.cpp
--- a/Engine/Source/Runtime/Core/Private/Android/AndroidPlatformFile.cpp
+++ b/Engine/Source/Runtime/Core/Private/Android/AndroidPlatformFile.cpp
@@ -220,7 +220,7 @@ public:
 	int64 Length;
 	int64 CurrentOffset;
 
-	FORCEINLINE void CheckValid()
+	FORCEINLINE void CheckValid() const
 	{
 		check(File.IsValid() && File->Handle != -1);
 	}
@@ -365,6 +365,12 @@ public:
 		return BytesToRead == 0;
 	}
 
+	virtual bool ReadAt(uint8* Destination, int64 BytesToRead, int64 Offset) override
+	{
+		int64 TrueOffset = Start + Offset;
+		return ReadInternal(Destination, BytesToRead, TrueOffset);
+	}
+
 	virtual bool Write(const uint8* Source, int64 BytesToWrite) override
 	{
 		CheckValid();
@@ -452,6 +458,55 @@ public:
 	{
 		return Length;
 	}
+
+private:
+	bool ReadInternal(uint8* Destination, int64 BytesToRead, int64 Offset)
+	{
+		CheckValid();
+#if LOG_ANDROID_FILE
+		FPlatformMisc::LowLevelOutputDebugStringf(
+			TEXT("(%d/%d) FFileHandleAndroid:Read => Path = %s, BytesToRead = %d"),
+			FAndroidTLS::GetCurrentThreadId(), File->Handle,
+			*(File->Path), int32(BytesToRead));
+#endif
+		if (BytesToRead < 0 || Offset - Start < 0 || (BytesToRead + Offset - Start) > Size())
+		{
+			return false;
+		}
+
+		if (BytesToRead == 0)
+		{
+			return true;
+		}
+
+		check(Destination);
+
+		while (BytesToRead > 0)
+		{
+			int64 ThisSize = FMath::Min<int64>(READWRITE_SIZE, BytesToRead);
+
+			ThisSize = __pread(File->Handle, Destination, ThisSize, Offset);
+#if LOG_ANDROID_FILE
+			FPlatformMisc::LowLevelOutputDebugStringf(
+				TEXT("(%d/%d) FFileHandleAndroid:Read => Path = %s, ThisSize = %d, destination = %X"),
+				FAndroidTLS::GetCurrentThreadId(), File->Handle,
+				*(File->Path), int32(ThisSize), Destination);
+#endif
+			if (ThisSize < 0)
+			{
+				return false;
+			}
+			else if (ThisSize == 0)
+			{
+				break;
+			}
+			Offset += ThisSize;
+			Destination += ThisSize;
+			BytesToRead -= ThisSize;
+		}
+
+		return BytesToRead == 0;
+	}
 };
 
 
diff --git a/Engine/Source/Runtime/Core/Private/Apple/ApplePlatformFile.cpp b/Engine/Source/Runtime/Core/Private/Apple/ApplePlatformFile.cpp
--- a/Engine/Source/Runtime/Core/Private/Apple/ApplePlatformFile.cpp
+++ b/Engine/Source/Runtime/Core/Private/Apple/ApplePlatformFile.cpp
@@ -9,8 +9,10 @@
 #include "Containers/UnrealString.h"
 #include "Containers/StringConv.h"
 #include "CoreGlobals.h"
+#include "HAL/CriticalSection.h"
 #include "HAL/PlatformFile.h"
 #include "HAL/PlatformTime.h"
+#include "Misc/ScopeLock.h"
 #include "ProfilingDebugging/PlatformFileTrace.h"
 #include "Templates/Function.h"
 #include <sys/stat.h>
@@ -193,6 +195,50 @@ public:
 			return ReadInternal(Destination, BytesToRead) == BytesToRead;
 		}
 	}
+	virtual bool ReadAt(uint8* Destination, int64 BytesToRead, int64 Offset) override
+	{
+		if (BytesToRead < 0 || Offset < 0)
+		{
+			return false;
+		}
+
+		if (BytesToRead == 0)
+		{
+			return true;
+		}
+
+#if MANAGE_FILE_HANDLES
+		if (IsManaged())
+		{
+			ActivateSlot();
+		}
+#endif //MANAGE_FILE_HANDLES
+
+		int64 TotalBytesRead = 0;
+		TRACE_PLATFORMFILE_BEGIN_READ(this, FileHandle, Offset, BytesToRead);
+
+		do
+		{
+			size_t BytesToRead32 = static_cast<size_t>(FMath::Min<int64>(READWRITE_SIZE, BytesToRead));
+			ssize_t BytesRead = pread(FileHandle, Destination, BytesToRead, Offset);
+
+			TotalBytesRead += BytesRead;
+
+			if (BytesRead != BytesToRead32)
+			{
+				TRACE_PLATFORMFILE_END_READ(this, TotalBytesRead);
+				return false;
+			}
+
+			Offset += BytesRead;
+			BytesToRead -= BytesToRead32;
+
+		} while (BytesToRead > 0);
+
+		TRACE_PLATFORMFILE_END_READ(this, TotalBytesRead);
+
+		return true;
+	}
 	virtual bool Write(const uint8* Source, int64 BytesToWrite) override
 	{
 		check(IsValid());
diff --git a/Engine/Source/Runtime/Core/Private/HAL/PlatformFileManager.cpp b/Engine/Source/Runtime/Core/Private/HAL/PlatformFileManager.cpp
--- a/Engine/Source/Runtime/Core/Private/HAL/PlatformFileManager.cpp
+++ b/Engine/Source/Runtime/Core/Private/HAL/PlatformFileManager.cpp
@@ -28,12 +28,28 @@ IPlatformFile& FPlatformFileManager::GetPlatformFile()
 	return *TopmostPlatformFile;
 }
 
+CORE_API IPlatformFile& FPlatformFileManager::GetPlatformPhysical()
+{
+	if (PhysicalPlatformFile == nullptr)
+	{
+		PhysicalPlatformFile = &IPlatformFile::GetPlatformPhysical();
+	}
+
+	return *PhysicalPlatformFile;
+}
+
 void FPlatformFileManager::SetPlatformFile(IPlatformFile& NewTopmostPlatformFile)
 {
 	TopmostPlatformFile = &NewTopmostPlatformFile;
 	NewTopmostPlatformFile.InitializeAfterSetActive();
 }
 
+void FPlatformFileManager::SetPlatformPhysical(IPlatformFile& NewPhysicalPlatformFile)
+{
+	PhysicalPlatformFile = &NewPhysicalPlatformFile;
+	NewPhysicalPlatformFile.InitializeAfterSetActive();
+}
+
 IPlatformFile* FPlatformFileManager::FindPlatformFile(const TCHAR* Name)
 {
 	check(TopmostPlatformFile != NULL);
diff --git a/Engine/Source/Runtime/Core/Private/IOS/IOSPlatformFile.cpp b/Engine/Source/Runtime/Core/Private/IOS/IOSPlatformFile.cpp
--- a/Engine/Source/Runtime/Core/Private/IOS/IOSPlatformFile.cpp
+++ b/Engine/Source/Runtime/Core/Private/IOS/IOSPlatformFile.cpp
@@ -5,12 +5,14 @@
 #include "Containers/StringConv.h"
 #include "HAL/PlatformTime.h"
 #include "Templates/Function.h"
+#include "HAL/CriticalSection.h"
 #include "HAL/PlatformProcess.h"
 #include "Misc/Parse.h"
 #include "Misc/CoreMisc.h"
 #include "Misc/CommandLine.h"
 #include "Misc/App.h"
 #include "Misc/Paths.h"
+#include "Misc/ScopeLock.h"
 
 #include "Async/MappedFileHandle.h"
 #include <sys/mman.h>
@@ -184,6 +186,43 @@ public:
         }
 	}
 
+	virtual bool ReadAt(uint8* Destination, int64 BytesToRead, int64 Offset) override
+	{
+		if (BytesToRead < 0 || Offset < 0)
+		{
+			return false;
+		}
+
+		if (BytesToRead == 0)
+		{
+			return true;
+		}
+
+#if MANAGE_FILE_HANDLES_IOS
+		if (IsManaged())
+		{
+			ActivateSlot();
+		}
+#endif //MANAGE_FILE_HANDLES_IOS
+
+		do
+		{
+			size_t BytesToRead32 = static_cast<size_t>(FMath::Min<int64>(READWRITE_SIZE, BytesToRead));
+			ssize_t BytesRead = pread(FileHandle, Destination, BytesToRead, Offset);
+
+			if (BytesRead != BytesToRead32)
+			{
+				return false;
+			}
+
+			Offset += BytesRead;
+			BytesToRead -= BytesToRead32;
+
+		} while (BytesToRead > 0);
+
+		return true;
+	}
+
 	virtual bool Seek( int64 NewPosition ) override
 	{
 		check(NewPosition >= 0);
@@ -315,6 +354,9 @@ private:
 
     void ActivateSlot()
     {
+		static FCriticalSection LockHandles;
+		FScopeLock Lock(&LockHandles);
+
         if( IsManaged() )
         {
             if( ManagedFiles[HandleSlot].ID != FileID )
diff --git a/Engine/Source/Runtime/Core/Private/Unix/UnixPlatformFile.cpp b/Engine/Source/Runtime/Core/Private/Unix/UnixPlatformFile.cpp
--- a/Engine/Source/Runtime/Core/Private/Unix/UnixPlatformFile.cpp
+++ b/Engine/Source/Runtime/Core/Private/Unix/UnixPlatformFile.cpp
@@ -21,6 +21,8 @@ DEFINE_LOG_CATEGORY_STATIC(LogUnixPlatformFile, Log, All);
 // make an FTimeSpan object that represents the "epoch" for time_t (from a stat struct)
 const FDateTime UnixEpoch(1970, 1, 1);
 
+class FFileHandleUnix;
+
 namespace
 {
 	FFileStatData UnixStatToUEFileData(struct stat& FileInfo)
@@ -217,6 +219,57 @@ public:
 		}
 	}
 
+	virtual bool ReadAt(uint8* Destination, int64 BytesToRead, int64 Offset) override
+	{
+		if (BytesToRead < 0 || Offset < 0)
+		{
+			return false;
+		}
+
+		if (BytesToRead == 0)
+		{
+			return true;
+		}
+
+		FFileHandleRegistryReadTracker TrackRead(GFileRegistry, *this, true);
+		if (!FileOpenAsWrite && !TrackRead.IsValid())
+		{
+			return false;
+		}
+
+		do
+		{
+			size_t BytesToRead32 = static_cast<size_t>(FMath::Min<int64>(READWRITE_SIZE, BytesToRead));
+			ssize_t BytesRead = pread(FileHandle, Destination, BytesToRead, Offset);
+
+			if (BytesRead == -1 && errno == EFAULT)
+			{
+				// Hacky workaround @see ReadInternal
+				void* TempDest = FMemory::Malloc(BytesToRead);
+				if (TempDest)
+				{
+					BytesRead = pread(FileHandle, TempDest, BytesToRead, Offset);
+					if (BytesRead > 0 && BytesRead <= BytesToRead)
+					{
+						FMemory::Memcpy(Destination, TempDest, BytesRead);
+					}
+					FMemory::Free(TempDest);
+				}
+			}
+
+			if (BytesRead != BytesToRead32)
+			{
+				return false;
+			}
+
+			Offset += BytesRead;
+			BytesToRead -= BytesToRead32;
+
+		} while (BytesToRead > 0);
+
+		return true;
+	}
+
 	virtual bool Write(const uint8* Source, int64 BytesToWrite) override
 	{
 		check(IsValid());
diff --git a/Engine/Source/Runtime/Core/Private/Windows/WindowsPlatformFile.cpp b/Engine/Source/Runtime/Core/Private/Windows/WindowsPlatformFile.cpp
--- a/Engine/Source/Runtime/Core/Private/Windows/WindowsPlatformFile.cpp
+++ b/Engine/Source/Runtime/Core/Private/Windows/WindowsPlatformFile.cpp
@@ -27,6 +27,7 @@
 
 #include "Windows/AllowWindowsPlatformTypes.h" // Required immediately before MicrosoftAsyncIO.h and the block below it
 #include "Microsoft/MicrosoftAsyncIO.h"
+
 namespace FileConstants
 {
 uint32 WIN_INVALID_SET_FILE_POINTER = INVALID_SET_FILE_POINTER;
@@ -119,8 +120,19 @@ FILETIME UEDateTimeToWindowsFileTime(const FDateTime& InDateTime)
 	return FileTime;
 }
 
+void UpdateOverlappedPos(OVERLAPPED& Overlapped, int64 UpdatedPos)
+{
+	ULARGE_INTEGER LI;
+	LI.QuadPart = UpdatedPos;
+
+	Overlapped.Offset = LI.LowPart;
+	Overlapped.OffsetHigh = LI.HighPart;
+}
+
 } // namespace UE::WindowsPlatformFile::Private
 
+#if USE_OVERLAPPED_IO
+
 /**
  * This file reader uses overlapped i/o and double buffering to asynchronously read from files
  */
@@ -483,6 +495,69 @@ public:
 		return true;
 	}
 
+	virtual bool ReadAt(uint8* Destination, int64 BytesToRead, int64 Offset) override
+	{
+		if (BytesToRead < 0 || Offset < 0 || (BytesToRead + Offset) > FileSize)
+		{
+			return false;
+		}
+
+		if (BytesToRead == 0)
+		{
+			return true;
+		}
+
+		OVERLAPPED LocalOverlappedIO;
+		FMemory::Memzero(&LocalOverlappedIO, sizeof(OVERLAPPED));
+
+		// Calls to ::ReadAt should be rare enough that we can create/destroy events for each read as needed
+		LocalOverlappedIO.hEvent = CreateEvent(NULL, 0, 0, NULL);
+
+		if (LocalOverlappedIO.hEvent == NULL)
+		{
+			// TODO Error
+			return false;
+		}
+
+		ON_SCOPE_EXIT
+		{
+			CloseHandle(LocalOverlappedIO.hEvent);
+		};
+
+		TRACE_PLATFORMFILE_BEGIN_READ(&LocalOverlappedIO, Handle, Offset, BytesToRead);
+		int64 TotalNumRead = 0;
+		do
+		{
+			uint32 BytesToRead32 = (uint32)FMath::Min<int64>(BytesToRead, int64(UINT32_MAX));
+			uint32 NumRead = 0;
+
+			UE::WindowsPlatformFile::Private::UpdateOverlappedPos(LocalOverlappedIO, Offset);
+
+			if (!ReadFile(Handle, Destination, BytesToRead32, (::DWORD*)&NumRead, &LocalOverlappedIO))
+			{
+				if (GetLastError() != ERROR_IO_PENDING)
+				{
+					// TODO Error
+					TRACE_PLATFORMFILE_END_READ(&LocalOverlappedIO, TotalNumRead);
+					return false;
+				}
+
+				if (!GetOverlappedResult(Handle, &LocalOverlappedIO, (::DWORD*)&NumRead, true))
+				{
+					// TODO Error
+					TRACE_PLATFORMFILE_END_READ(&LocalOverlappedIO, TotalNumRead);
+					return false;
+				}
+			}
+
+			BytesToRead -= BytesToRead32;
+			TotalNumRead += NumRead;
+		} while (BytesToRead > 0);
+
+		TRACE_PLATFORMFILE_END_READ(&LocalOverlappedIO, TotalNumRead);
+		return true;
+	}
+
 	virtual bool Write(const uint8* Source, int64 BytesToWrite) override
 	{
 		check(0 && "This is an async reader only and doesn't support writing");
@@ -537,6 +612,8 @@ public:
 	}
 };
 
+#endif // USE_OVERLAPPED_IO
+
 /** 
  * Windows file handle implementation
 **/
@@ -561,14 +638,6 @@ class FFileHandleWindows : public IFileHandle
 	 */
 	uint32 Flags;
 
-	FORCEINLINE void UpdateOverlappedPos()
-	{
-		ULARGE_INTEGER LI;
-		LI.QuadPart = FilePos;
-		OverlappedIO.Offset = LI.LowPart;
-		OverlappedIO.OffsetHigh = LI.HighPart;
-	}
-
 	FORCEINLINE bool UpdatedNonOverlappedPos()
 	{
 		LARGE_INTEGER LI;
@@ -628,7 +697,7 @@ public:
 #endif
 		FileHandle = NULL;
 	}
-	bool IsValid()
+	bool IsValid() const
 	{
 		return FileHandle != NULL && FileHandle != INVALID_HANDLE_VALUE && FileSize != -1;
 	}
@@ -648,7 +717,7 @@ public:
 		check(NewPosition >= 0);
 
 		FilePos = NewPosition;
-		UpdateOverlappedPos();
+		UE::WindowsPlatformFile::Private::UpdateOverlappedPos(OverlappedIO, FilePos);
 		return true;
 	}
 	virtual bool SeekFromEnd(int64 NewPositionRelativeToEnd = 0) override
@@ -695,7 +764,7 @@ public:
 			TotalNumRead += NumRead;
 			// Update where we are in the file
 			FilePos += NumRead;
-			UpdateOverlappedPos();
+			UE::WindowsPlatformFile::Private::UpdateOverlappedPos(OverlappedIO, FilePos);
 			
 			// Early out as a failure case if we did not read all of the bytes that we expected to read
 			if (BytesToRead32 != NumRead)
@@ -708,6 +777,53 @@ public:
 		TRACE_PLATFORMFILE_END_READ(&OverlappedIO, TotalNumRead);
 		return true;
 	}
+	virtual bool ReadAt(uint8* Destination, int64 BytesToRead, int64 Offset) override
+	{
+		check(IsValid());
+
+		if (BytesToRead < 0 || Offset < 0 || (BytesToRead + Offset) > FileSize)
+		{
+			return false;
+		}
+
+		if (BytesToRead == 0)
+		{
+			return true;
+		}
+
+		OVERLAPPED LocalOverlappedIO;
+		FMemory::Memzero(&LocalOverlappedIO, sizeof(OVERLAPPED));
+
+		TRACE_PLATFORMFILE_BEGIN_READ(&LocalOverlappedIO, FileHandle, Offset, BytesToRead);
+		int64 TotalNumRead = 0;
+		do
+		{
+			uint32 BytesToRead32 = (uint32)FMath::Min<int64>(BytesToRead, int64(UINT32_MAX));
+			uint32 NumRead = 0;
+
+			UE::WindowsPlatformFile::Private::UpdateOverlappedPos(LocalOverlappedIO, Offset);
+
+			if (!ReadFile(FileHandle, Destination, BytesToRead32, (::DWORD*)&NumRead, &LocalOverlappedIO))
+			{
+				TRACE_PLATFORMFILE_END_READ(&LocalOverlappedIO, TotalNumRead);
+				return false;
+			}
+
+			BytesToRead -= BytesToRead32;
+			TotalNumRead += NumRead;
+
+			// Early out as a failure case if we did not read all of the bytes that we expected to read
+			if (BytesToRead32 != NumRead)
+			{
+				TRACE_PLATFORMFILE_END_READ(&OverlappedIO, TotalNumRead);
+				return false;
+			}
+
+		} while (BytesToRead > 0);
+
+		TRACE_PLATFORMFILE_END_READ(&LocalOverlappedIO, TotalNumRead);
+		return true;
+	}
 	virtual bool Write(const uint8* Source, int64 BytesToWrite) override
 	{
 		check(IsValid());
@@ -744,7 +860,7 @@ public:
 			TotalNumWritten += NumWritten;
 			// Update where we are in the file
 			FilePos += NumWritten;
-			UpdateOverlappedPos();
+			UE::WindowsPlatformFile::Private::UpdateOverlappedPos(OverlappedIO, FilePos);
 			FileSize = FMath::Max(FilePos, FileSize);
 			
 			// Early out as a failure case if we didn't write all of the data we expected
diff --git a/Engine/Source/Runtime/Core/Public/GenericPlatform/GenericPlatformFile.h b/Engine/Source/Runtime/Core/Public/GenericPlatform/GenericPlatformFile.h
--- a/Engine/Source/Runtime/Core/Public/GenericPlatform/GenericPlatformFile.h
+++ b/Engine/Source/Runtime/Core/Public/GenericPlatform/GenericPlatformFile.h
@@ -145,6 +145,8 @@ public:
 	**/
 	virtual bool		Read(uint8* Destination, int64 BytesToRead) = 0;
 
+	virtual bool		ReadAt(uint8* Destination, int64 BytesToRead, int64 Offset) = 0;
+
 	/** 
 	 * Write bytes to the file.
 	 * @param Source		Buffer to write, should be at least BytesToWrite in size.
diff --git a/Engine/Source/Runtime/Core/Public/GenericPlatform/GenericPlatformMisc.h b/Engine/Source/Runtime/Core/Public/GenericPlatform/GenericPlatformMisc.h
--- a/Engine/Source/Runtime/Core/Public/GenericPlatform/GenericPlatformMisc.h
+++ b/Engine/Source/Runtime/Core/Public/GenericPlatform/GenericPlatformMisc.h
@@ -25,6 +25,7 @@ class GenericApplication;
 class IInstallBundleManager;
 class IPlatformChunkInstall;
 class IPlatformCompression;
+class IWrappedFileHandle;
 class IPlatformHostCommunication;
 struct FDefaultDelegateUserPolicy;
 struct FGenericCrashContext;
@@ -1264,6 +1265,8 @@ public:
 	 */
 	static CORE_API IPlatformCompression* GetPlatformCompression();
 
+	static CORE_API IWrappedFileHandle* CreateCompressedPlatformFileHandle(IFileHandle* InLowerLevelHandle);
+
 	/**
 	 * Returns the platform specific interface for communication with processes running on the host pc.
 	 * 
diff --git a/Engine/Source/Runtime/Core/Public/HAL/IPlatformFileCachedWrapper.h b/Engine/Source/Runtime/Core/Public/HAL/IPlatformFileCachedWrapper.h
--- a/Engine/Source/Runtime/Core/Public/HAL/IPlatformFileCachedWrapper.h
+++ b/Engine/Source/Runtime/Core/Public/HAL/IPlatformFileCachedWrapper.h
@@ -135,6 +135,12 @@ public:
 		return Result;
 	}
 
+	virtual bool ReadAt(uint8* Destination, int64 BytesToRead, int64 Offset) override
+	{
+		// concurrent reads won't be cached yet
+		return FileHandle->ReadAt(Destination, BytesToRead, Offset);
+	}
+
 	virtual bool		Write(const uint8* Source, int64 BytesToWrite) override
 	{
 		if (!bWritable || BytesToWrite < 0)
diff --git a/Engine/Source/Runtime/Core/Public/HAL/IPlatformFileLogWrapper.h b/Engine/Source/Runtime/Core/Public/HAL/IPlatformFileLogWrapper.h
--- a/Engine/Source/Runtime/Core/Public/HAL/IPlatformFileLogWrapper.h
+++ b/Engine/Source/Runtime/Core/Public/HAL/IPlatformFileLogWrapper.h
@@ -87,6 +87,15 @@ public:
 		FILE_LOG(LogPlatformFile, VeryVerbose, TEXT("Read return %d [%fms]"), int32(Result), ThisTime);
 		return Result;
 	}
+	virtual bool		ReadAt(uint8* Destination, int64 BytesToRead, int64 Offset) override
+	{
+		FILE_LOG(LogPlatformFile, VeryVerbose, TEXT("ReadAt %s %lld %lld"), *Filename, BytesToRead, Offset);
+		double StartTime = FPlatformTime::Seconds();
+		bool Result = FileHandle->ReadAt(Destination, BytesToRead, Offset);
+		double ThisTime = (FPlatformTime::Seconds() - StartTime) / 1000.0;
+		FILE_LOG(LogPlatformFile, VeryVerbose, TEXT("ReadAt return %d [%fms]"), int32(Result), ThisTime);
+		return Result;
+	}
 	virtual bool		Write(const uint8* Source, int64 BytesToWrite) override
 	{
 		FILE_LOG(LogPlatformFile, VeryVerbose, TEXT("Write %s %lld"), *Filename, BytesToWrite);
diff --git a/Engine/Source/Runtime/Core/Public/HAL/IPlatformFileManagedStorageWrapper.h b/Engine/Source/Runtime/Core/Public/HAL/IPlatformFileManagedStorageWrapper.h
--- a/Engine/Source/Runtime/Core/Public/HAL/IPlatformFileManagedStorageWrapper.h
+++ b/Engine/Source/Runtime/Core/Public/HAL/IPlatformFileManagedStorageWrapper.h
@@ -718,6 +718,11 @@ public:
 		return FileHandle->Read(Destination, BytesToRead);
 	}
 
+	virtual bool ReadAt(uint8* Destination, int64 BytesToRead, int64 Offset) override
+	{
+		return FileHandle->ReadAt(Destination, BytesToRead, Offset);
+	}
+
 	virtual bool Write(const uint8* Source, int64 BytesToWrite) override
 	{
 		if (!TryManageFile())
diff --git a/Engine/Source/Runtime/Core/Public/HAL/PlatformFileManager.h b/Engine/Source/Runtime/Core/Public/HAL/PlatformFileManager.h
--- a/Engine/Source/Runtime/Core/Public/HAL/PlatformFileManager.h
+++ b/Engine/Source/Runtime/Core/Public/HAL/PlatformFileManager.h
@@ -18,9 +18,11 @@ class FPlatformFileManager
 #if USE_ATOMIC_PLATFORM_FILE
 	/** Currently used platform file. */
 	TAtomic<class IPlatformFile*> TopmostPlatformFile;
+	TAtomic<class IPlatformFile*> PhysicalPlatformFile;
 #else
 	/** Currently used platform file. */
 	class IPlatformFile* TopmostPlatformFile;
+	class IPlatformFile* PhysicalPlatformFile;
 #endif
 
 public:
@@ -35,6 +37,8 @@ public:
 	 */
 	CORE_API IPlatformFile& GetPlatformFile( );
 
+	CORE_API IPlatformFile& GetPlatformPhysical();
+
 	/**
 	 * Sets the current platform file.
 	 *
@@ -42,6 +46,8 @@ public:
 	 */
 	CORE_API void SetPlatformFile( IPlatformFile& NewTopmostPlatformFile );
 
+	CORE_API void SetPlatformPhysical(IPlatformFile& NewPhysicalPlatformFile);
+
 	/**
 	 * Finds a platform file in the chain of active platform files.
 	 *
diff --git a/Engine/Source/Runtime/CoreUObject/Tests/PluginManagerTest.cpp b/Engine/Source/Runtime/CoreUObject/Tests/PluginManagerTest.cpp
--- a/Engine/Source/Runtime/CoreUObject/Tests/PluginManagerTest.cpp
+++ b/Engine/Source/Runtime/CoreUObject/Tests/PluginManagerTest.cpp
@@ -55,6 +55,12 @@ namespace UE::CoreUObject::Private::Tests
 			return true;
 		}
 
+		virtual bool ReadAt(uint8* Destination, int64 BytesToRead, int64 Offset) override
+		{
+			FMemory::Memcpy(Destination, Data + Offset, BytesToRead);
+			return true;
+		}
+
 		virtual int64 Tell() { return Pos; }
 		virtual bool Write(const uint8* Source, int64 BytesToWrite) { return false; }
 		virtual bool Flush(const bool bFullFlush = false){ return false; }
diff --git a/Engine/Source/Runtime/PakFile/Private/IPlatformFilePak.cpp b/Engine/Source/Runtime/PakFile/Private/IPlatformFilePak.cpp
--- a/Engine/Source/Runtime/PakFile/Private/IPlatformFilePak.cpp
+++ b/Engine/Source/Runtime/PakFile/Private/IPlatformFilePak.cpp
@@ -4167,7 +4167,7 @@ public:
 		return PakEntry.UncompressedSize;
 	}
 
-	void Serialize(int64 DesiredPosition, void* V, int64 Length)
+	void Serialize(int64 DesiredPosition, void* V, int64 Length) const
 	{
 		const int32 CompressionBlockSize = PakEntry.CompressionBlockSize;
 		uint32 CompressionBlockIndex = (uint32)(DesiredPosition / CompressionBlockSize);
diff --git a/Engine/Source/Runtime/PakFile/Public/IPlatformFilePak.h b/Engine/Source/Runtime/PakFile/Public/IPlatformFilePak.h
--- a/Engine/Source/Runtime/PakFile/Public/IPlatformFilePak.h
+++ b/Engine/Source/Runtime/PakFile/Public/IPlatformFilePak.h
@@ -1744,7 +1744,7 @@ public:
 		return PakEntry.Size;
 	}
 
-	void Serialize(int64 DesiredPosition, void* V, int64 Length)
+	void Serialize(int64 DesiredPosition, void* V, int64 Length) const
 	{
 		FGuid EncryptionKeyGuid = PakFile.GetInfo().EncryptionKeyGuid;
 		const constexpr int64 Alignment = (int64)EncryptionPolicy::Alignment;
@@ -1858,30 +1858,8 @@ public:
 	}
 	virtual bool Read(uint8* Destination, int64 BytesToRead) override
 	{
-		SCOPE_SECONDS_ACCUMULATOR(STAT_PakFile_Read);
-
-		// Check that the file header is OK
-		if (!Reader.PakEntry.Verified)
-		{
-			FPakEntry FileHeader;
-			FSharedPakReader PakReader = Reader.AcquirePakReader();
-			PakReader->Seek(Reader.PakEntry.Offset);
-			FileHeader.Serialize(PakReader.GetArchive(), Reader.PakFile.GetInfo().Version);
-			if (FPakEntry::VerifyPakEntriesMatch(Reader.PakEntry, FileHeader))
-			{
-				Reader.PakEntry.Verified = true;
-			}
-			else
-			{
-				//Header is corrupt, fail the read
-				return false;
-			}
-		}
-		//
-		if (Reader.FileSize() >= (ReadPos + BytesToRead))
+		if (ReadInternal(Destination, BytesToRead, ReadPos))
 		{
-			// Read directly from Pak.
-			Reader.Serialize(ReadPos, Destination, BytesToRead);
 			ReadPos += BytesToRead;
 			return true;
 		}
@@ -1890,6 +1868,15 @@ public:
 			return false;
 		}
 	}
+	virtual bool ReadAt(uint8* Destination, int64 BytesToRead, int64 Offset) override
+	{
+		if (BytesToRead < 0 || Offset < 0 || (BytesToRead + Offset) > Reader.FileSize())
+		{
+			return false;
+		}
+
+		return ReadInternal(Destination, BytesToRead, Offset);
+	}
 	virtual bool Write(const uint8* Source, int64 BytesToWrite) override
 	{
 		// Writing in pak files is not allowed.
@@ -1910,6 +1897,37 @@ public:
 		return false;
 	}
 	///~ End IFileHandle Interface
+private:
+	bool ReadInternal(uint8* Destination, int64 BytesToRead, int64 Offset)
+	{
+		SCOPE_SECONDS_ACCUMULATOR(STAT_PakFile_Read);
+
+		if (!Reader.PakEntry.Verified)
+		{
+			FPakEntry FileHeader;
+			FSharedPakReader PakReader = Reader.AcquirePakReader();
+			PakReader->Seek(Reader.PakEntry.Offset);
+			FileHeader.Serialize(PakReader.GetArchive(), Reader.PakFile.GetInfo().Version);
+			if (FPakEntry::VerifyPakEntriesMatch(Reader.PakEntry, FileHeader))
+			{
+				Reader.PakEntry.Verified = true;
+			}
+			else
+			{
+				return false;
+			}
+		}
+
+		if (Reader.FileSize() >= (Offset + BytesToRead))
+		{
+			Reader.Serialize(Offset, Destination, BytesToRead);
+			return true;
+		}
+		else
+		{
+			return false;
+		}
+	}
 };
 
 /**
diff --git a/Engine/Source/Runtime/Sockets/Public/NetworkMessage.h b/Engine/Source/Runtime/Sockets/Public/NetworkMessage.h
--- a/Engine/Source/Runtime/Sockets/Public/NetworkMessage.h
+++ b/Engine/Source/Runtime/Sockets/Public/NetworkMessage.h
@@ -46,6 +46,7 @@ namespace NFS_Messages
 		CopyFile,
 		GetFileInfo,
 		Read,
+		ReadAt,
 		Write,
 		Close,
 		Seek,
diff --git a/Engine/Source/Runtime/StorageServerClient/Private/StorageServerPlatformFile.cpp b/Engine/Source/Runtime/StorageServerClient/Private/StorageServerPlatformFile.cpp
--- a/Engine/Source/Runtime/StorageServerClient/Private/StorageServerPlatformFile.cpp
+++ b/Engine/Source/Runtime/StorageServerClient/Private/StorageServerPlatformFile.cpp
@@ -378,11 +378,19 @@ private:
 class FStorageServerFileHandle
 	: public IFileHandle
 {
+	enum
+	{
+		BufferSize = 64 << 10
+	};
 	FStorageServerPlatformFile& Owner;
 	FIoChunkId FileChunkId;
 	FString Filename;
 	int64 FilePos = 0;
 	int64 FileSize = -1;
+	int64 BufferStart = -1;
+	int64 BufferEnd = -1;
+	uint8 Buffer[BufferSize];
+	FCriticalSection BufferCS;
 
 public:
 	FStorageServerFileHandle(FStorageServerPlatformFile& InOwner, FIoChunkId InFileChunkId, int64 InFileSize, const TCHAR* InFilename)
@@ -456,6 +464,60 @@ public:
 		return BytesRemaining == 0;
 	}
 
+	virtual bool ReadAt(uint8* Destination, int64 BytesToRead, int64 Offset)
+	{
+		if (BytesToRead == 0)
+		{
+			return true;
+		}
+
+		if (BytesToRead > BufferSize)
+		{
+			const int64 BytesRead = Owner.SendReadMessage(Destination, FileChunkId, Offset, BytesToRead);
+			if (BytesRead == BytesToRead)
+			{
+				STORAGESERVER_CACHEMISS(BytesRead);
+				return true;
+			}
+			return false;
+		}
+
+		{
+			FScopeLock BufferLock(&BufferCS);
+
+			int64 BytesReadFromBuffer = 0;
+			if (Offset >= BufferStart && Offset < BufferEnd)
+			{
+				const int64 BufferOffset = Offset - BufferStart;
+				check(BufferOffset < BufferSize);
+				BytesReadFromBuffer = FMath::Min(BufferSize - BufferOffset, BytesToRead);
+				FMemory::Memcpy(Destination, Buffer + BufferOffset, BytesReadFromBuffer);
+				STORAGESERVER_CACHEHIT(BytesReadFromBuffer);
+				if (BytesReadFromBuffer == BytesToRead)
+				{
+					Offset += BytesReadFromBuffer;
+					return true;
+				}
+			}
+
+			const int64 BytesRead = Owner.SendReadMessage(Buffer, FileChunkId, Offset + BytesReadFromBuffer, BufferSize);
+			BufferStart = Offset + BytesReadFromBuffer;
+			BufferEnd = BufferStart + BytesRead;
+
+			const int64 BytesToReadFromBuffer = FMath::Min(BytesRead, BytesToRead - BytesReadFromBuffer);
+			FMemory::Memcpy(Destination + BytesReadFromBuffer, Buffer, BytesToReadFromBuffer);
+			BytesReadFromBuffer += BytesToReadFromBuffer;
+			if (BytesReadFromBuffer == BytesToRead)
+			{
+				Offset += BytesReadFromBuffer;
+				STORAGESERVER_CACHEMISS(BytesReadFromBuffer);
+				return true;
+			}
+
+			return false;
+		}
+	}
+
 	virtual bool Write(const uint8* Source, int64 BytesToWrite) override
 	{
 		check(false);
@@ -681,6 +743,7 @@ bool FStorageServerPlatformFile::Initialize(IPlatformFile* Inner, const TCHAR* C
 		{
 			UE_LOG(LogStorageServerPlatformFile, Display, TEXT("Using settings from command line: -ZenStoreBaseURI='%s'"), *BaseURI);
 		}
+
 		return true;
 	}
 	return false;
@@ -688,6 +751,22 @@ bool FStorageServerPlatformFile::Initialize(IPlatformFile* Inner, const TCHAR* C
 
 void FStorageServerPlatformFile::InitializeAfterProjectFilePath()
 {
+	InitializeConnection();
+
+	// optional debugging module depends on a valid Connection
+	if (FModuleManager::Get().ModuleExists(TEXT("StorageServerClientDebug")))
+	{
+		FModuleManager::Get().LoadModule("StorageServerClientDebug");
+	}
+}
+
+void FStorageServerPlatformFile::InitializeConnection()
+{
+	if (Connection)
+	{
+		return;
+	}
+
 #if WITH_COTF
 	UE::Cook::ICookOnTheFlyModule& CookOnTheFlyModule = FModuleManager::LoadModuleChecked<UE::Cook::ICookOnTheFlyModule>(TEXT("CookOnTheFly"));
 	CookOnTheFlyServerConnection = CookOnTheFlyModule.GetDefaultServerConnection();
@@ -725,12 +804,6 @@ void FStorageServerPlatformFile::InitializeAfterProjectFilePath()
 			FStringView HostAddr = Connection->GetHostAddr();
 			UE_LOG(LogStorageServerPlatformFile, Fatal, TEXT("Failed to get file list from Zen at '%.*s'"), HostAddr.Len(), HostAddr.GetData());
 		}
-
-		// optional debugging module depends on a valid Connection
-		if (FModuleManager::Get().ModuleExists(TEXT("StorageServerClientDebug")))
-		{
-			FModuleManager::Get().LoadModule("StorageServerClientDebug");
-		}
 	}
 	else if (bAbortOnConnectionFailure)
 	{
@@ -857,6 +930,7 @@ IFileHandle* FStorageServerPlatformFile::InternalOpenFile(const FIoChunkId& File
 IFileHandle* FStorageServerPlatformFile::OpenRead(const TCHAR* Filename, bool bAllowWrite)
 {
 	TStringBuilder<1024> StorageServerFilename;
+
 	if (MakeStorageServerPath(Filename, StorageServerFilename))
 	{
 		FIoChunkId FileChunkId;
diff --git a/Engine/Source/Runtime/StorageServerClient/Private/StorageServerPlatformFile.h b/Engine/Source/Runtime/StorageServerClient/Private/StorageServerPlatformFile.h
--- a/Engine/Source/Runtime/StorageServerClient/Private/StorageServerPlatformFile.h
+++ b/Engine/Source/Runtime/StorageServerClient/Private/StorageServerPlatformFile.h
@@ -124,6 +124,7 @@ public:
 private:
 	friend class FStorageServerFileHandle;
 
+	void InitializeConnection();
 	bool IsNonServerFilenameAllowed(FStringView InFilename);
 	bool IsAssumedImmutableTimeStampFilename(FStringView InFilename) const;
 	bool MakeStorageServerPath(const TCHAR* LocalFilenameOrDirectory, FStringBuilderBase& OutPath) const;
diff --git a/Engine/Source/Runtime/StreamingFile/Private/StreamingNetworkPlatformFile.cpp b/Engine/Source/Runtime/StreamingFile/Private/StreamingNetworkPlatformFile.cpp
--- a/Engine/Source/Runtime/StreamingFile/Private/StreamingNetworkPlatformFile.cpp
+++ b/Engine/Source/Runtime/StreamingFile/Private/StreamingNetworkPlatformFile.cpp
@@ -251,6 +251,21 @@ public:
 		return Result;
 	}
 
+	virtual bool ReadAt(uint8* Destination, int64 BytesToRead, int64 Offset) override
+	{
+		if (!bReadable || BytesToRead < 0 || Offset < 0 || (BytesToRead + Offset) > Size())
+		{
+			return false;
+		}
+
+		if (BytesToRead == 0)
+		{
+			return true;
+		}
+
+		return Network.SendReadMessage(HandleId, Destination, BytesToRead, Offset);
+	}
+
 	virtual bool Write(const uint8* Source, int64 BytesToWrite) override
 	{
 		bool Result = false;
@@ -847,6 +862,34 @@ bool FStreamingNetworkPlatformFile::SendReadMessage(uint64 HandleId, uint8* Dest
 }
 
 
+STREAMINGFILE_API bool FStreamingNetworkPlatformFile::SendReadMessage(uint64 HandleId, uint8* Destination, int64 BytesToRead, int64 Offset)
+{
+	FStreamingNetworkFileArchive Payload(NFS_Messages::ReadAt);
+	Payload << HandleId;
+	Payload << BytesToRead;
+	Payload << Offset;
+
+	FArrayReader Response;
+
+	if (SendPayloadAndReceiveResponse(Payload, Response) == false)
+	{
+		return false;
+	}
+
+	int64 ServerBytesRead = 0;
+	Response << ServerBytesRead;
+
+	if (ServerBytesRead == BytesToRead)
+	{
+		Response.Serialize(Destination, BytesToRead);
+		return true;
+	}
+	else
+	{
+		return false;
+	}
+}
+
 bool FStreamingNetworkPlatformFile::SendWriteMessage(uint64 HandleId, const uint8* Source, int64 BytesToWrite)
 {
 	FScopeLock ScopeLock(&SynchronizationObject);
diff --git a/Engine/Source/Runtime/StreamingFile/Public/StreamingNetworkPlatformFile.h b/Engine/Source/Runtime/StreamingFile/Public/StreamingNetworkPlatformFile.h
--- a/Engine/Source/Runtime/StreamingFile/Public/StreamingNetworkPlatformFile.h
+++ b/Engine/Source/Runtime/StreamingFile/Public/StreamingNetworkPlatformFile.h
@@ -159,6 +159,8 @@ public:
 	/** Sends Read message to the server. */
 	STREAMINGFILE_API bool SendReadMessage(uint64 HandleId, uint8* Destination, int64 BytesToRead);
 
+	STREAMINGFILE_API bool SendReadMessage(uint64 HandleId, uint8* Destination, int64 BytesToRead, int64 Offset);
+
 	/** Sends Write message to the server. */
 	STREAMINGFILE_API bool SendWriteMessage(uint64 HandleId, const uint8* Source, int64 BytesToWrite);	
 
