Logging: Added UE_LOG_CONTEXT to add thread-local context to structured log records

diff --git a/Engine/Source/Runtime/Core/Private/Logging/StructuredLog.cpp b/Engine/Source/Runtime/Core/Private/Logging/StructuredLog.cpp
--- a/Engine/Source/Runtime/Core/Private/Logging/StructuredLog.cpp
+++ b/Engine/Source/Runtime/Core/Private/Logging/StructuredLog.cpp
@@ -5,6 +5,8 @@
 #include "Algo/Accumulate.h"
 #include "Algo/AllOf.h"
 #include "Algo/NoneOf.h"
+#include "AutoRTFM/AutoRTFM.h"
+#include "Containers/AnsiString.h"
 #include "Containers/StringConv.h"
 #include "Containers/StringView.h"
 #include "CoreGlobals.h"
@@ -23,7 +25,7 @@
 #include "Serialization/CompactBinarySerialization.h"
 #include "Serialization/CompactBinaryValue.h"
 #include "Serialization/VarInt.h"
-#include "AutoRTFM/AutoRTFM.h"
+#include "Templates/Function.h"
 #include <cstdarg>
 
 void VARARGS StaticFailDebug (const TCHAR* Error, const ANSICHAR* Expression, const ANSICHAR* File, int32 Line, bool bIsEnsure, void* ProgramCounter, const TCHAR* DescriptionFormat, ...);
@@ -652,6 +654,60 @@ void FLogRecord::FormatMessageTo(FWideStringBuilderBase& Out) const
 namespace UE::Logging::Private
 {
 
+thread_local FLogContext* LogContextHead = nullptr;
+thread_local FLogContext* LogContextTail = nullptr;
+
+FLogContext::FLogContext(const FLogField& InField)
+: Prev(LogContextTail)
+{
+	(Prev ? Prev->Next : LogContextHead) = this;
+	LogContextTail = this;
+
+	const FUtf8StringView Name = InField.Name;
+
+	TCbWriter<256> Writer;
+	Writer.SetName(Name);
+	if (InField.WriteValue)
+	{
+		InField.WriteValue(Writer, InField.Value);
+	}
+	else
+	{
+		Writer.AddNull();
+	}
+	Field = Writer.Save();
+	Field.MakeOwned();
+
+	for (FLogContext* Node = Prev; Node; Node = Node->Prev)
+	{
+		if (Node->Field.GetName().Equals(Name))
+		{
+			Node->bDisabledByNewerContext = true;
+			bDisabledOlderContext = true;
+		}
+	}
+}
+
+FLogContext::~FLogContext()
+{
+	(Prev ? Prev->Next : LogContextHead) = Next;
+	(Next ? Next->Prev : LogContextTail) = Prev;
+
+	if (bDisabledOlderContext)
+	{
+		const FUtf8StringView Name = Field.GetName();
+		for (FLogContext* Node = Prev; Node; Node = Node->Prev)
+		{
+			if (Node->Field.GetName().Equals(Name))
+			{
+				Node->bDisabledByNewerContext = false;
+			}
+		}
+	}
+}
+
+///////////////////////////////////////////////////////////////////////////////////////////////////////////////////////
+
 class FLogTemplateFieldIterator
 {
 public:
@@ -734,7 +790,8 @@ FORCENOINLINE static FCbObject SerializeLogFields(
 	const FLogField* Fields,
 	const int32 FieldCount)
 {
-	if (FieldCount == 0)
+	const bool bHasContext = !!LogContextHead;
+	if (FieldCount == 0 && !bHasContext)
 	{
 		return FCbObject();
 	}
@@ -742,6 +799,8 @@ FORCENOINLINE static FCbObject SerializeLogFields(
 	TCbWriter<1024> Writer;
 	Writer.BeginObject();
 
+	TArray<FAnsiString, TInlineAllocator<16>> FieldNames;
+
 	// Anonymous. Extract names from Template.
 	if (!Fields->Name)
 	{
@@ -750,7 +809,12 @@ FORCENOINLINE static FCbObject SerializeLogFields(
 		{
 			check(It);
 			const auto Name = StringCast<ANSICHAR>(It.GetName().GetData(), It.GetName().Len());
-			Fields->WriteValue(Writer.SetName(MakeStringView(Name.Get(), Name.Length())), Fields->Value);
+			const FAnsiStringView NameView(Name.Get(), Name.Length());
+			Fields->WriteValue(Writer.SetName(NameView), Fields->Value);
+			if (bHasContext)
+			{
+				FieldNames.Emplace(NameView);
+			}
 		}
 		check(!It);
 	}
@@ -760,7 +824,48 @@ FORCENOINLINE static FCbObject SerializeLogFields(
 		for (const FLogField* FieldsEnd = Fields + FieldCount; Fields != FieldsEnd; ++Fields)
 		{
 			Fields->WriteValue(Writer.SetName(Fields->Name), Fields->Value);
+			if (bHasContext)
+			{
+				FieldNames.Emplace(Fields->Name);
+			}
+		}
+	}
+
+	if (bHasContext)
+	{
+		TBitArray<> ActiveContext;
+		int32 ContextIndex = 0;
+
+		// Traverse contexts backward and activate any which have a name that has not been seen yet.
+		for (FLogContext* Node = LogContextTail; Node; Node = Node->Prev)
+		{
+			if (!Node->bDisabledByNewerContext)
+			{
+				const FUtf8StringView NodeName = Node->Field.GetName();
+				ActiveContext.Add(!FieldNames.FindByPredicate([NodeName](FAnsiStringView Name) { return NodeName.Equals(Name); }));
+				++ContextIndex;
+			}
+		}
+
+		// Traverse contexts forward and copy any which were activated above.
+		for (FLogContext* Node = LogContextHead; Node; Node = Node->Next)
+		{
+			if (!Node->bDisabledByNewerContext && ActiveContext[--ContextIndex])
+			{
+				Writer.AddField(Node->Field.GetName(), Node->Field);
+			}
+		}
+
+		// Traverse contexts forward and build an array of names in $Contexts.
+		Writer.BeginArray(ANSITEXTVIEW("$Contexts"));
+		for (FLogContext* Node = LogContextHead; Node; Node = Node->Next)
+		{
+			if (!Node->bDisabledByNewerContext)
+			{
+				Writer.AddString(Node->Field.GetName());
+			}
 		}
+		Writer.EndArray();
 	}
 
 	Writer.EndObject();
@@ -1032,7 +1137,34 @@ FORCENOINLINE static FCbObject SerializeBasicLogMessage(const FStaticBasicLogRec
 
 	TCbWriter<512> Writer;
 	Writer.BeginObject();
-	Writer.AddString(ANSITEXTVIEW("Message"), Message);
+
+	const FUtf8StringView MessageName = UTF8TEXTVIEW("Message");
+	Writer.AddString(MessageName, Message);
+
+	if (LogContextHead)
+	{
+		// Traverse contexts forward and copy any that Message did not override.
+		for (FLogContext* Node = LogContextHead; Node; Node = Node->Next)
+		{
+			const FUtf8StringView Name = Node->Field.GetName();
+			if (!Node->bDisabledByNewerContext && !Name.Equals(MessageName))
+			{
+				Writer.AddField(Name, Node->Field);
+			}
+		}
+
+		// Traverse contexts forward and build an array of names in $Contexts.
+		Writer.BeginArray(ANSITEXTVIEW("$Contexts"));
+		for (FLogContext* Node = LogContextHead; Node; Node = Node->Next)
+		{
+			if (!Node->bDisabledByNewerContext)
+			{
+				Writer.AddString(Node->Field.GetName());
+			}
+		}
+		Writer.EndArray();
+	}
+
 	Writer.EndObject();
 	return Writer.Save().AsObject();
 }
@@ -1057,19 +1189,7 @@ static void BasicLogV(const FLogCategoryBase& Category, const FStaticBasicLogRec
 		Record.SetVerbosity(Log->Verbosity);
 		Record.SetTime(FLogTime::Now());
 
-		FOutputDevice* OutputDevice = nullptr;
-		switch (Log->Verbosity)
-		{
-		case ELogVerbosity::Error:
-		case ELogVerbosity::Warning:
-		case ELogVerbosity::Display:
-		case ELogVerbosity::SetColor:
-			OutputDevice = GWarn;
-			break;
-		default:
-			break;
-		}
-		(OutputDevice ? OutputDevice : GLog)->SerializeRecord(Record);
+		DispatchStaticLogRecord(*Log, Record);
 	}
 	else
 	{
@@ -1129,6 +1249,18 @@ void DispatchDynamicLogRecord(const FLogRecord& Record)
 	Logging::Private::DispatchLogRecord(Record);
 }
 
+void VisitLogContext(TFunctionRef<void (const FCbField&)> Visitor)
+{
+	using namespace UE::Logging::Private;
+	for (const FLogContext* Node = LogContextHead; Node; Node = Node->Next)
+	{
+		if (!Node->bDisabledByNewerContext)
+		{
+			Visitor(Node->Field);
+		}
+	}
+}
+
 } // UE
 
 #endif // !NO_LOGGING
diff --git a/Engine/Source/Runtime/Core/Public/Logging/StructuredLog.h b/Engine/Source/Runtime/Core/Public/Logging/StructuredLog.h
--- a/Engine/Source/Runtime/Core/Public/Logging/StructuredLog.h
+++ b/Engine/Source/Runtime/Core/Public/Logging/StructuredLog.h
@@ -19,6 +19,7 @@
 #include <type_traits>
 
 class FTextFormat;
+template <typename FuncType> class TFunctionRef;
 
 #define UE_API CORE_API
 
@@ -105,6 +106,24 @@ class FTextFormat;
 /** Expands to a structured log value. UE_LOGFMT_VALUE(Value) */
 #define UE_LOGFMT_VALUE(Value) Value
 
+/**
+ * Registers a context on the calling thread, by name, with an optional value.
+ *
+ * The optional value can be anything that serializes to compact binary using FCbWriter or SerializeForLog.
+ * Context is unregistered when it goes out of scope.
+ * Context is overridden by a newer context with the same name.
+ * Context is overridden by a field of the same name in the log record.
+ * Context is copied to FLogRecord::Fields for every log record created in the lifetime of the context.
+ * Context names are written to FLogRecord::Fields in an array named $Context.
+ * Contexts can be visited by calling VisitLogContext.
+ *
+ * Examples:
+ * UE_LOG_CONTEXT("Loading");
+ * UE_LOG_CONTEXT("Count", 123.0);
+ * UE_LOG_CONTEXT("Asset", FAssetLog(AssetPath));
+ */
+#define UE_LOG_CONTEXT(Name, ...) ::UE::Logging::Private::FLogContext ANONYMOUS_VARIABLE(LogContext_)(Name, ##__VA_ARGS__)
+
 #else // #if !NO_LOGGING
 
 #define UE_LOGFMT(...)
@@ -118,6 +137,8 @@ class FTextFormat;
 #define UE_LOGFMT_FIELD(Name, Value)
 #define UE_LOGFMT_VALUE(Value)
 
+#define UE_LOG_CONTEXT(Name, Value)
+
 #endif
 
 struct FDateTime;
@@ -217,9 +238,18 @@ private:
  * Dynamic dispatch bypasses many optimizations provided by the macros.
  * Anything pointed to by the record MUST remain valid until threaded logs have been flushed.
  * Filtering by category or verbosity is the responsibility of the caller.
+ * Active log contexts will not be added to this record. If they are needed that is the responsibility of the caller.
  */
 UE_API void DispatchDynamicLogRecord(const FLogRecord& Record);
 
+/**
+ * Advanced function to visit the log contexts for the calling thread.
+ *
+ * Invokes the visitor once for each context, in order from the oldest context added on this thread to the newest.
+ * If there are multiple contexts with the same name, only the newest will be visited.
+ */
+UE_API void VisitLogContext(TFunctionRef<void (const FCbField&)> Visitor);
+
 /**
  * Serializes the value to be used in a log message.
  *
@@ -393,14 +423,14 @@ struct FLogFieldCreator
 	template <typename ValueType, typename... FieldArgTypes, std::enable_if_t<ValueCount<ValueType>>* = nullptr>
 	inline static void Create(FLogField* Fields, const ValueType& Value, FieldArgTypes&&... FieldArgs)
 	{
-		new(Fields) FLogField{nullptr, &Value, FLogField::Write<std::remove_reference_t<ValueType>>};
+		new(Fields) FLogField{nullptr, &Value, FLogField::Write<ValueType>};
 		Create(Fields + 1, (FieldArgTypes&&)FieldArgs...);
 	}
 
 	template <typename NameType, typename ValueType, typename... FieldArgTypes>
 	inline static void Create(FLogField* Fields, TLogFieldName<NameType> Name, const ValueType& Value, FieldArgTypes&&... FieldArgs)
 	{
-		new(Fields) FLogField{Name.Name, &Value, FLogField::Write<std::remove_reference_t<ValueType>>};
+		new(Fields) FLogField{Name.Name, &Value, FLogField::Write<ValueType>};
 		Create(Fields + 1, (FieldArgTypes&&)FieldArgs...);
 	}
 };
@@ -474,6 +504,32 @@ inline void LogIfActive(const LogCategoryType& Category, const LogRecordType& Lo
 	}
 }
 
+struct FLogContext
+{
+	inline explicit FLogContext(const ANSICHAR* Name)
+		: FLogContext({Name, nullptr, nullptr})
+	{
+	}
+
+	template <typename ValueType>
+	inline explicit FLogContext(const ANSICHAR* Name, const ValueType& Value)
+		: FLogContext({Name, &Value, FLogField::Write<ValueType>})
+	{
+	}
+
+	UE_API explicit FLogContext(const FLogField& Field);
+	UE_API ~FLogContext();
+
+	FLogContext(const FLogContext&) = delete;
+	FLogContext& operator=(const FLogContext&) = delete;
+
+	FCbField Field;
+	FLogContext* Next = nullptr;
+	FLogContext* Prev = nullptr;
+	bool bDisabledByNewerContext = false;
+	bool bDisabledOlderContext = false;
+};
+
 } // UE::Logging::Private
 
 #define UE_PRIVATE_LOGFMT(CategoryName, Verbosity, Format, ...) \
