[EventCount] Race fix

There was a race possible between FPipe::WaitUntilEmpty and FPipe::ClearTask() which was caused by a race in EventCount between Notify & Wait.

That led to a dead-lock on IOS.

Fixed by adding a StoreLoad (kinda) barrier into Notify to ensure the proper memory ordering for Counter value.
RMW is used as a barrier, but could also be solved with a few seq_cst fences (considered to be less optimal and unsupported by TSan).

More details:

Thread 1:
- ClearTask: fetch_sub-s TaskCount to 0 and goes Notify
- Notify: reads count as 0 and exits without any wakes (first bit is not set -> nobody is waiting)
Thead 2:
- WaitUntilEmpty: first check fails, second check check fails, calls PrepareWait (sets count to 1, returns 0), calls WaitFor with 0
- WaitFor: gets count as 1, count without 1 bit is 0, matches token -> wait for Notify which will never happen.

One thread was reading stale data due to memory reordering.

#rnx

diff --git a/Engine/Source/Runtime/Core/Public/Async/EventCount.h b/Engine/Source/Runtime/Core/Public/Async/EventCount.h
--- a/Engine/Source/Runtime/Core/Public/Async/EventCount.h
+++ b/Engine/Source/Runtime/Core/Public/Async/EventCount.h
@@ -135,7 +135,19 @@ public:
 	 */
 	inline void Notify()
 	{
-		CounterType Value = Count.load(std::memory_order_relaxed);
+		//
+		// .fetch_add(0, acq_rel) is used to have a StoreLoad barrier,
+		// which we can't express in C++. That works by making the load
+		// also be store (via RMW) and relying on a StoreStore barrier to
+		// get the desired ordering.
+		//
+		// Previously, this code was:
+		//   CounterType Value = Count.load(std::memory_order_relaxed);
+		//
+		// which had a memory re-ordering and stale values being read,
+		// leading to a missed Wake and dead-locked waiter, as a result.
+		//
+		CounterType Value = Count.fetch_add(0, std::memory_order_acq_rel);
 		if ((Value & 1) && Count.compare_exchange_strong(Value, Value + 1, std::memory_order_release))
 		{
 			ParkingLot::WakeAll(&Count);
