API improvements for creating and finding localized FText

== Added FText::AsLocalizable_Advanced ==
* This is the replacement for FInternationalization::ForUseOnlyByLocMacroAndGraphNodeTextLiterals_CreateText within the formal FText API.
* It is the implementation of the LOCTEXT macro, and is NOT a replacement for creating text literals in C++!
* It is intended for advanced use-cases that were already using FInternationalization::ForUseOnlyByLocMacroAndGraphNodeTextLiterals_CreateText to create a localizable FText that was then being assigned to something like a FText property on an asset, OR as an alternative to FText::FindTextInLiveTable_Advanced for cases where the text exists, but isn't yet translated.

== Added FText::FindTextInLiveTable_Advanced ==
* This is the replacement for FText::FindText.
* It has the same behavior as before, but is clearer about what it can and cannot do.
* FText::FindText will be deprecated shortly so that people can audit their existing FText::FindText usages and decide if they want to continue to use FText::FindTextInLiveTable_Advanced or swap to using FText::AsLocalizable_Advanced on a case-by-case basis.

== Misc ==
* Added FText::AsCultureInvariant overload to construct from a FStringView.

diff --git a/Engine/Source/Runtime/Core/Private/Internationalization/Internationalization.cpp b/Engine/Source/Runtime/Core/Private/Internationalization/Internationalization.cpp
--- a/Engine/Source/Runtime/Core/Private/Internationalization/Internationalization.cpp
+++ b/Engine/Source/Runtime/Core/Private/Internationalization/Internationalization.cpp
@@ -9,7 +9,6 @@
 #include "Misc/Paths.h"
 #include "Internationalization/Culture.h"
 #include "Internationalization/CustomCultureImplementation.h"
-#include "Internationalization/TextCache.h"
 
 #if UE_ENABLE_ICU
 #include "Internationalization/ICUInternationalization.h"
@@ -35,11 +34,6 @@ void FInternationalization::TearDown()
 	TLazySingleton<FInternationalization>::TearDown();
 }
 
-FText FInternationalization::ForUseOnlyByLocMacroAndGraphNodeTextLiterals_CreateText(const TCHAR* InTextLiteral, const TCHAR* InNamespace, const TCHAR* InKey)
-{
-	return FTextCache::Get().FindOrCache(InTextLiteral, InNamespace, InKey);
-}
-
 bool FInternationalization::SetCurrentCulture(const FString& InCultureName)
 {
 	FCulturePtr NewCulture = GetCulture(InCultureName);
diff --git a/Engine/Source/Runtime/Core/Private/Internationalization/PolyglotTextData.cpp b/Engine/Source/Runtime/Core/Private/Internationalization/PolyglotTextData.cpp
--- a/Engine/Source/Runtime/Core/Private/Internationalization/PolyglotTextData.cpp
+++ b/Engine/Source/Runtime/Core/Private/Internationalization/PolyglotTextData.cpp
@@ -184,7 +184,7 @@ void FPolyglotTextData::CacheText(FText* OutFailureReason)
 	if (IsValid(OutFailureReason))
 	{
 		FTextLocalizationManager::Get().RegisterPolyglotTextData(*this);
-		if (!FText::FindText(Namespace, Key, CachedText, &NativeString))
+		if (!FText::FindTextInLiveTable_Advanced(Namespace, Key, CachedText, &NativeString))
 		{
 			ClearCache();
 		}
diff --git a/Engine/Source/Runtime/Core/Private/Internationalization/Text.cpp b/Engine/Source/Runtime/Core/Private/Internationalization/Text.cpp
--- a/Engine/Source/Runtime/Core/Private/Internationalization/Text.cpp
+++ b/Engine/Source/Runtime/Core/Private/Internationalization/Text.cpp
@@ -12,6 +12,7 @@
 
 #include "Internationalization/TextHistory.h"
 #include "Misc/Guid.h"
+#include "Internationalization/TextCache.h"
 #include "Internationalization/TextFormatter.h"
 #include "Internationalization/TextChronoFormatter.h"
 #include "Internationalization/TextTransformer.h"
@@ -820,19 +821,6 @@ FString FText::GetInvariantTimeZone()
 	return TEXT("Etc/Unknown");
 }
 
-bool FText::FindText(const FTextKey& Namespace, const FTextKey& Key, FText& OutText, const FString* const SourceString)
-{
-	FTextConstDisplayStringPtr FoundString = FTextLocalizationManager::Get().FindDisplayString( Namespace, Key, SourceString );
-
-	if ( FoundString.IsValid() )
-	{
-		OutText = FText(MakeRefCount<FTextHistory_Base>(FTextId(Namespace, Key), SourceString ? FString(*SourceString) : FString(), FoundString.ToSharedRef()));
-		return true;
-	}
-
-	return false;
-}
-
 void FText::SerializeText(FArchive& Ar, FText& Value)
 {
 	SerializeText(FStructuredArchiveFromArchive(Ar).GetSlot(), Value);
@@ -1107,9 +1095,14 @@ FText FText::FromStringView(FStringView InString)
 	return NewText;
 }
 
-FText FText::AsCultureInvariant( const FString& String )
+FText FText::AsCultureInvariant( const TCHAR* String )
 {
-	FText NewText = String.IsEmpty() ? FText::GetEmpty() : FText(CopyTemp(String));
+	return FText::AsCultureInvariant(FStringView(String));
+}
+
+FText FText::AsCultureInvariant( FStringView String )
+{
+	FText NewText = String.IsEmpty() ? FText::GetEmpty() : FText(FString(String));
 	NewText.Flags |= ETextFlag::CultureInvariant;
 
 	return NewText;
@@ -1131,6 +1124,38 @@ FText FText::AsCultureInvariant( FText Text )
 	return NewText;
 }
 
+FText FText::AsLocalizable_Advanced(const FTextKey& Namespace, const FTextKey& Key, const TCHAR* String)
+{
+	return FText::AsLocalizable_Advanced(Namespace, Key, FStringView(String));
+}
+
+FText FText::AsLocalizable_Advanced(const FTextKey& Namespace, const FTextKey& Key, FStringView String)
+{
+	return String.IsEmpty() 
+		? FText::GetEmpty() 
+		: FTextCache::Get().FindOrCache(String, FTextId(Namespace, Key));
+}
+
+FText FText::AsLocalizable_Advanced(const FTextKey& Namespace, const FTextKey& Key, FString&& String)
+{
+	return String.IsEmpty() 
+		? FText::GetEmpty() 
+		: FTextCache::Get().FindOrCache(MoveTemp(String), FTextId(Namespace, Key));
+}
+
+bool FText::FindTextInLiveTable_Advanced(const FTextKey& Namespace, const FTextKey& Key, FText& OutText, const FString* const SourceString)
+{
+	FTextConstDisplayStringPtr FoundString = FTextLocalizationManager::Get().FindDisplayString(Namespace, Key, SourceString);
+
+	if (FoundString.IsValid())
+	{
+		OutText = FText(MakeRefCount<FTextHistory_Base>(FTextId(Namespace, Key), SourceString ? FString(*SourceString) : FString(), FoundString.ToSharedRef()));
+		return true;
+	}
+
+	return false;
+}
+
 const FString& FText::ToString() const
 {
 	Rebuild();
diff --git a/Engine/Source/Runtime/Core/Private/Internationalization/TextCache.cpp b/Engine/Source/Runtime/Core/Private/Internationalization/TextCache.cpp
--- a/Engine/Source/Runtime/Core/Private/Internationalization/TextCache.cpp
+++ b/Engine/Source/Runtime/Core/Private/Internationalization/TextCache.cpp
@@ -29,12 +29,52 @@ void FTextCache::TearDown()
 	return TLazySingleton<FTextCache>::TearDown();
 }
 
-FText FTextCache::FindOrCache(const TCHAR* InTextLiteral, const TCHAR* InNamespace, const TCHAR* InKey)
+FText FTextCache::FindOrCache(const TCHAR* InTextLiteral, const FTextId& InTextId)
 {
-	return FindOrCache(InTextLiteral, FTextId(InNamespace, InKey));
+	return FindOrCache(FStringView(InTextLiteral), InTextId);
 }
 
-FText FTextCache::FindOrCache(const TCHAR* InTextLiteral, const FTextId& InTextId)
+FText FTextCache::FindOrCache(FStringView InTextLiteral, const FTextId& InTextId)
+{
+	return AutoRTFM::Open([&]
+	{
+		LLM_SCOPE(ELLMTag::Localization);
+
+		// First try and find a cached instance
+		{
+			FText* ReturnFoundText = nullptr;
+
+			FScopeLock Lock(&CachedTextCS);
+
+			FText* FoundText = CachedText.Find(InTextId);
+			if (FoundText)
+			{
+				const FString* FoundTextLiteral = FTextInspector::GetSourceString(*FoundText);
+				if (FoundTextLiteral && InTextLiteral.Equals(*FoundTextLiteral, ESearchCase::CaseSensitive))
+				{
+					ReturnFoundText = FoundText;
+				}
+			}
+
+			if (ReturnFoundText)
+			{
+				return *ReturnFoundText;
+			}
+		}
+
+		// Not currently cached, make a new instance...
+		FText NewText = FText(FString(InTextLiteral), InTextId.GetNamespace(), InTextId.GetKey(), ETextFlag::Immutable);
+
+		// ... and add it to the cache
+		FScopeLock Lock(&CachedTextCS);
+
+		CachedText.Emplace(InTextId, NewText);
+
+		return NewText;
+	});
+}
+
+FText FTextCache::FindOrCache(FString&& InTextLiteral, const FTextId& InTextId)
 {
 	return AutoRTFM::Open([&]
 	{
@@ -50,7 +90,7 @@ FText FTextCache::FindOrCache(const TCHAR* InTextLiteral, const FTextId& InTextI
 			if (FoundText)
 			{
 				const FString* FoundTextLiteral = FTextInspector::GetSourceString(*FoundText);
-				if (FoundTextLiteral && FCString::Strcmp(**FoundTextLiteral, InTextLiteral) == 0)
+				if (FoundTextLiteral && InTextLiteral.Equals(*FoundTextLiteral, ESearchCase::CaseSensitive))
 				{
 					ReturnFoundText = FoundText;
 				}
@@ -63,7 +103,7 @@ FText FTextCache::FindOrCache(const TCHAR* InTextLiteral, const FTextId& InTextI
 		}
 
 		// Not currently cached, make a new instance...
-		FText NewText = FText(InTextLiteral, InTextId.GetNamespace(), InTextId.GetKey(), ETextFlag::Immutable);
+		FText NewText = FText(MoveTemp(InTextLiteral), InTextId.GetNamespace(), InTextId.GetKey(), ETextFlag::Immutable);
 
 		// ... and add it to the cache
 		FScopeLock Lock(&CachedTextCS);
diff --git a/Engine/Source/Runtime/Core/Private/Internationalization/TextCache.h b/Engine/Source/Runtime/Core/Private/Internationalization/TextCache.h
--- a/Engine/Source/Runtime/Core/Private/Internationalization/TextCache.h
+++ b/Engine/Source/Runtime/Core/Private/Internationalization/TextCache.h
@@ -28,8 +28,9 @@ public:
 	/**
 	 * Try and find an existing cached entry for the given data, or construct and cache a new entry if one cannot be found.
 	 */
-	FText FindOrCache(const TCHAR* InTextLiteral, const TCHAR* InNamespace, const TCHAR* InKey);
 	FText FindOrCache(const TCHAR* InTextLiteral, const FTextId& InTextId);
+	FText FindOrCache(FStringView InTextLiteral, const FTextId& InTextId);
+	FText FindOrCache(FString&& InTextLiteral, const FTextId& InTextId);
 
 	/**
 	 * Remove any cached entries for the given text IDs.
diff --git a/Engine/Source/Runtime/Core/Public/Internationalization/Internationalization.h b/Engine/Source/Runtime/Core/Public/Internationalization/Internationalization.h
--- a/Engine/Source/Runtime/Core/Public/Internationalization/Internationalization.h
+++ b/Engine/Source/Runtime/Core/Public/Internationalization/Internationalization.h
@@ -32,7 +32,15 @@ public:
 
 	static CORE_API void TearDown();
 
-	static CORE_API FText ForUseOnlyByLocMacroAndGraphNodeTextLiterals_CreateText(const TCHAR* InTextLiteral, const TCHAR* InNamespace, const TCHAR* InKey);
+	/**
+	 * Legacy alias for FText::AsLocalizable_Advanced.
+	 * @note FText::AsLocalizable_Advanced takes its arguments in a different order!
+	 * Will be deprecated in the future. DO NOT USE IN NEW CODE!
+	 */
+	static FORCEINLINE FText ForUseOnlyByLocMacroAndGraphNodeTextLiterals_CreateText(const TCHAR* InTextLiteral, const TCHAR* InNamespace, const TCHAR* InKey)
+	{
+		return FText::AsLocalizable_Advanced(InNamespace, InKey, InTextLiteral);
+	}
 
 	/**
 	 * Struct that can be used to capture a snapshot of the active culture state in a way that can be re-applied losslessly.
@@ -288,12 +296,12 @@ namespace UE::Private::LocGen
  * Creates an FText. All parameters must be string literals. All literals will be passed through the localization system.
  * The global LOCTEXT_NAMESPACE macro must be first set to a string literal to specify this localization key's namespace.
  */
-#define LOCTEXT(InKey, InTextLiteral) FInternationalization::ForUseOnlyByLocMacroAndGraphNodeTextLiterals_CreateText(TEXT(InTextLiteral), TEXT(LOCTEXT_NAMESPACE), TEXT(InKey))
+#define LOCTEXT(InKey, InTextLiteral) FText::AsLocalizable_Advanced(TEXT(LOCTEXT_NAMESPACE), TEXT(InKey), TEXT(InTextLiteral))
 
 /**
  * Creates an FText. All parameters must be string literals. All literals will be passed through the localization system.
  */
-#define NSLOCTEXT(InNamespace, InKey, InTextLiteral) FInternationalization::ForUseOnlyByLocMacroAndGraphNodeTextLiterals_CreateText(TEXT(InTextLiteral), TEXT(InNamespace), TEXT(InKey))
+#define NSLOCTEXT(InNamespace, InKey, InTextLiteral) FText::AsLocalizable_Advanced(TEXT(InNamespace), TEXT(InKey), TEXT(InTextLiteral))
 
 /**
  * Creates a culture invariant FText from the given string literal.
diff --git a/Engine/Source/Runtime/Core/Public/Internationalization/Text.h b/Engine/Source/Runtime/Core/Public/Internationalization/Text.h
--- a/Engine/Source/Runtime/Core/Public/Internationalization/Text.h
+++ b/Engine/Source/Runtime/Core/Public/Internationalization/Text.h
@@ -465,12 +465,6 @@ public:
 	 */
 	static CORE_API FText AsMemory(uint64 NumBytes, EMemoryUnitStandard UnitStandard);
 
-	/**
-	 * Attempts to find an existing FText using the representation found in the loc tables for the specified namespace and key
-	 * @return true if OutText was properly set; otherwise false and OutText will be untouched
-	 */
-	static CORE_API bool FindText( const FTextKey& Namespace, const FTextKey& Key, FText& OutText, const FString* const SourceString = nullptr );
-
 	/**
 	 * Attempts to create an FText instance from a string table ID and key (this is the same as the LOCTABLE macro, except this can also work with non-literal string values).
 	 * @return The found text, or a dummy FText if not found.
@@ -496,7 +490,8 @@ public:
 	/**
 	 * Generate a culture invariant FText representing the passed in string
 	 */
-	static CORE_API FText AsCultureInvariant( const FString& String );
+	static CORE_API FText AsCultureInvariant( const TCHAR* String );
+	static CORE_API FText AsCultureInvariant( FStringView String );
 	static CORE_API FText AsCultureInvariant( FString&& String );
 
 	/**
@@ -504,6 +499,45 @@ public:
 	 */
 	static CORE_API FText AsCultureInvariant( FText Text );
 
+	/**
+	 * === !! This is an ADVANCED function. USE WITH CAUTION !! ===
+	 * 
+	 * Generate a localizable FText for the given ID and source string, either to create a NEW localizable FText, or to dynamically reference an EXISTING FText.
+	 * 
+	 * When using it to create a NEW localizable FText (eg, when editing an asset), YOU are responsible for ensuring that the resultant FText is gathered for localization (eg, by assigning it to a FText property or via a CUSTOM gather handler).
+	 * Note: DO NOT use this as a general replacement for the LOCTEXT macro in C++!
+	 * 
+	 * When using it to dynamically reference an EXISTING FText (such as known LOCTEXT in code) it can be thought of as an alternative to FText::FindTextInLiveTable_Advanced that also handles untranslated text.
+	 * Note: Direct dynamic references to FText are EXTREMELY FRAGILE, and you may want to use a string table instead!
+	 */
+	static CORE_API FText AsLocalizable_Advanced(const FTextKey& Namespace, const FTextKey& Key, const TCHAR* String);
+	static CORE_API FText AsLocalizable_Advanced(const FTextKey& Namespace, const FTextKey& Key, FStringView String);
+	static CORE_API FText AsLocalizable_Advanced(const FTextKey& Namespace, const FTextKey& Key, FString&& String);
+
+	/**
+	 * === !! This is an ADVANCED function. USE WITH CAUTION !! ===
+	 *
+	 * Attempt to dynamically reference an EXISTING FText via its active display string in the live table.
+	 * Note: This can ONLY find text that is currently localized (gathered, translated, and has an active display string in FTextLocalizationManager). If you need to find a localizable but untranslated text, see FText::AsLocalizable_Advanced.
+	 * Note: Direct dynamic references to FText are EXTREMELY FRAGILE, and you may want to use a string table instead!
+	 * 
+	 * @param Namespace The namespace of the text to find (if any).
+	 * @param Key The key of the text to find.
+	 * @param SourceString If set (not empty) then the found text must also have been created from this source string.
+	 * 
+	 * @return true if OutText was set; otherwise false and OutText will be unmodified
+	 */
+	static CORE_API bool FindTextInLiveTable_Advanced(const FTextKey& Namespace, const FTextKey& Key, FText& OutText, const FString* const SourceString = nullptr);
+
+	/**
+	 * Legacy alias for FText::FindTextInLiveTable_Advanced.
+	 * Will be deprecated in the future. DO NOT USE IN NEW CODE!
+	 */
+	static FORCEINLINE bool FindText(const FTextKey& Namespace, const FTextKey& Key, FText& OutText, const FString* const SourceString = nullptr)
+	{
+		return FText::FindTextInLiveTable_Advanced(Namespace, Key, OutText, SourceString);
+	}
+
 	CORE_API const FString& ToString() const;
 
 	/** Deep build of the source string for this FText, climbing the history hierarchy */
diff --git a/Engine/Source/Runtime/Engine/Classes/Kismet/KismetTextLibrary.h b/Engine/Source/Runtime/Engine/Classes/Kismet/KismetTextLibrary.h
--- a/Engine/Source/Runtime/Engine/Classes/Kismet/KismetTextLibrary.h
+++ b/Engine/Source/Runtime/Engine/Classes/Kismet/KismetTextLibrary.h
@@ -195,12 +195,17 @@ class UKismetTextLibrary : public UBlueprintFunctionLibrary
 	static ENGINE_API FText GetEmptyText();
 
 	/**
-	 * Attempts to find existing Text using the representation found in the loc tables for the specified namespace and key.
-	 * @param Namespace The namespace of the text to find (if any).
-	 * @param Key The key of the text to find.
-	 * @param SourceString If set (not empty) then the found text must also have been created from this source string.
-	 */
-	UFUNCTION(BlueprintPure, Category="Utilities|Text", meta = (AdvancedDisplay = "SourceString"))
+	  * === !! This is an ADVANCED function. USE WITH CAUTION !! ===
+	  *
+	  * Attempt to dynamically reference an EXISTING Text via its active display string in the live table.
+	  * Note: This can ONLY find text that is currently localized (gathered, translated, and has an active display string in TextLocalizationManager). If you need to find a localizable but untranslated text, see 'Make Literal Text'.
+	  * Note: Direct dynamic references to Text are EXTREMELY FRAGILE, and you may want to use a string table instead!
+	  *
+	  * @param Namespace The namespace of the text to find (if any).
+	  * @param Key The key of the text to find.
+	  * @param SourceString If set (not empty) then the found text must also have been created from this source string.
+	  */
+	UFUNCTION(BlueprintPure, Category="Utilities|Text", DisplayName="Find Text in Live Table (Advanced)", meta = (AdvancedDisplay = "SourceString", ScriptName="FindTextInLiveTable_Advanced;FindTextInLocalizationTable"))
 	static ENGINE_API bool FindTextInLocalizationTable(const FString& Namespace, const FString& Key, FText& OutText, const FString& SourceString = TEXT(""));
 
 	/** Returns true if A and B are linguistically equal (A == B). */
diff --git a/Engine/Source/Runtime/Engine/Private/KismetTextLibrary.cpp b/Engine/Source/Runtime/Engine/Private/KismetTextLibrary.cpp
--- a/Engine/Source/Runtime/Engine/Private/KismetTextLibrary.cpp
+++ b/Engine/Source/Runtime/Engine/Private/KismetTextLibrary.cpp
@@ -147,7 +147,7 @@ FText UKismetTextLibrary::GetEmptyText()
 
 bool UKismetTextLibrary::FindTextInLocalizationTable(const FString& Namespace, const FString& Key, FText& OutText, const FString& SourceString)
 {
-	return FText::FindText(Namespace, Key, OutText, SourceString.IsEmpty() ? nullptr : &SourceString);
+	return FText::FindTextInLiveTable_Advanced(Namespace, Key, OutText, SourceString.IsEmpty() ? nullptr : &SourceString);
 }
 
 bool UKismetTextLibrary::EqualEqual_IgnoreCase_TextText(const FText& A, const FText& B)
