#UE Use pre-existing OnParentPreFork delegate instead of making a new one

diff --git a/Engine/Source/Runtime/Core/Private/HAL/ThreadingBase.cpp b/Engine/Source/Runtime/Core/Private/HAL/ThreadingBase.cpp
--- a/Engine/Source/Runtime/Core/Private/HAL/ThreadingBase.cpp
+++ b/Engine/Source/Runtime/Core/Private/HAL/ThreadingBase.cpp
@@ -4,6 +4,7 @@
 #include "UObject/NameTypes.h"
 #include "Stats/Stats.h"
 #include "Misc/CommandLine.h"
+#include "Misc/CoreDelegates.h"
 #include "Misc/CoreStats.h"
 #include "Misc/EventPool.h"
 #include "Misc/LazySingleton.h"
@@ -450,6 +451,16 @@ public:
 };
 uint32 FFakeThread::ThreadIdCounter = 0xffff;
 
+FThreadManager::FThreadManager()
+{
+	FCoreDelegates::OnParentPreFork.AddRaw(this, &FThreadManager::HandleOnParentPreFork);
+}
+
+FThreadManager::~FThreadManager()
+{
+	FCoreDelegates::OnParentPreFork.RemoveAll(this);
+}
+
 bool FThreadManager::CheckThreadListSafeToContinueIteration()
 {
 	if (bIsThreadListDirty)
@@ -670,6 +681,16 @@ TArray<FRunnableThread*> FThreadManager::GetForkableThreads()
 	return ForkableThreads;
 }
 
+void FThreadManager::HandleOnParentPreFork()
+{
+	// Use a local list of forkable threads so we don't keep a lock on the global list during thread creation
+	TArray<FRunnableThread*> ForkableThreads = GetForkableThreads();
+	for (FRunnableThread* ForkableThread : ForkableThreads)
+	{
+		ForkableThread->OnPreFork();
+	}
+}
+
 /*-----------------------------------------------------------------------------
 	FEvent, FScopedEvent
 -----------------------------------------------------------------------------*/
diff --git a/Engine/Source/Runtime/Core/Private/Misc/Fork.cpp b/Engine/Source/Runtime/Core/Private/Misc/Fork.cpp
--- a/Engine/Source/Runtime/Core/Private/Misc/Fork.cpp
+++ b/Engine/Source/Runtime/Core/Private/Misc/Fork.cpp
@@ -110,16 +110,6 @@ void FForkProcessHelper::OnForkingOccured()
 	}
 }
 
-void FForkProcessHelper::OnPreFork()
-{
-	// Use a local list of forkable threads so we don't keep a lock on the global list during thread creation
-	TArray<FRunnableThread*> ForkableThreads = FThreadManager::Get().GetForkableThreads();
-	for (FRunnableThread* ForkableThread : ForkableThreads)
-	{
-		ForkableThread->OnPreFork();
-	}
-}
-
 bool FForkProcessHelper::IsForkedMultithreadInstance()
 {
 	return ForkInternal::bIsForkedMultithreadInstance;
@@ -158,4 +148,4 @@ void FForkProcessHelper::LowLevelPostForkChild(uint16 ChildIndex)
 {
 	FForkProcessHelper::SetIsForkedChildProcess(ChildIndex);
 	GMalloc->OnPostFork();
-}
+}
\ No newline at end of file
diff --git a/Engine/Source/Runtime/Core/Public/HAL/ThreadManager.h b/Engine/Source/Runtime/Core/Public/HAL/ThreadManager.h
--- a/Engine/Source/Runtime/Core/Public/HAL/ThreadManager.h
+++ b/Engine/Source/Runtime/Core/Public/HAL/ThreadManager.h
@@ -33,6 +33,9 @@ class FThreadManager
 
 public:
 
+	FThreadManager();
+	~FThreadManager();
+
 	/**
 	* Used internally to add a new thread object.
 	*
@@ -123,4 +126,7 @@ private:
 
 	/** Returns internal name of a the thread given its TLS id */
 	CORE_API const FString& GetThreadNameInternal(uint32 ThreadId);
+
+	/** Notification that the parent is about to fork */
+	void HandleOnParentPreFork();
 };
diff --git a/Engine/Source/Runtime/Core/Public/Misc/Fork.h b/Engine/Source/Runtime/Core/Public/Misc/Fork.h
--- a/Engine/Source/Runtime/Core/Public/Misc/Fork.h
+++ b/Engine/Source/Runtime/Core/Public/Misc/Fork.h
@@ -69,11 +69,6 @@ public:
 	 */
 	static CORE_API void OnForkingOccured();
 
-	/**
-	 * Event triggered when a fork is about to occur on the parent process
-	 */
-	static CORE_API void OnPreFork();
-
 	/**
 	 * Tells if we allow multithreading on forked processes.
 	 * Default is set to false but can be configured to always be true via DEFAULT_MULTITHREAD_FORKED_PROCESSES
