Added "Localization.AlwaysLoadNativeLocalizationDataDuringInitialization" (default: false)

This allows games to opt-in to loading the native localization data during initialization, even if the game isn't starting in the native language. This ensures that all gathered text will load some localization data, even if not fully translated.

diff --git a/Engine/Source/Runtime/Core/Private/Internationalization/TextLocalizationManager.cpp b/Engine/Source/Runtime/Core/Private/Internationalization/TextLocalizationManager.cpp
--- a/Engine/Source/Runtime/Core/Private/Internationalization/TextLocalizationManager.cpp
+++ b/Engine/Source/Runtime/Core/Private/Internationalization/TextLocalizationManager.cpp
@@ -64,6 +64,9 @@ static FAutoConsoleVariableRef CVarAsyncLoadLocalizationData(TEXT("Localization.
 static bool AsyncLoadLocalizationDataOnLanguageChange = false;
 static FAutoConsoleVariableRef CVarAsyncLoadLocalizationDataOnLanguageChange(TEXT("Localization.AsyncLoadLocalizationDataOnLanguageChange"), AsyncLoadLocalizationDataOnLanguageChange, TEXT("True to load localization data asynchronously (non-blocking) when the language changes, or False to load it synchronously (blocking)"));
 
+static bool AlwaysLoadNativeLocalizationDataDuringInitialization = false;
+static FAutoConsoleVariableRef CVarAlwaysLoadNativeLocalizationDataDuringInitialization(TEXT("Localization.AlwaysLoadNativeLocalizationDataDuringInitialization"), AlwaysLoadNativeLocalizationDataDuringInitialization, TEXT("True to load the native localization data during initialization, even if we're not starting in the native language. This ensures that all gathered text will load some localization data, even if not fully translated."));
+
 #if WITH_EDITOR
 static bool ForceLoadGameLocalizationInEditor = false;
 static FAutoConsoleVariableRef CVarForceLoadGameLocalizationInEditor(TEXT("Localization.ForceLoadGameLocalizationInEditor"), ForceLoadGameLocalizationInEditor, TEXT("True to force load game localization data in an editor"));
@@ -522,6 +525,7 @@ void InitEngineTextLocalization()
 	ELocalizationLoadFlags LocLoadFlags = ELocalizationLoadFlags::None;
 	LocLoadFlags |= (WITH_EDITOR ? ELocalizationLoadFlags::Editor : ELocalizationLoadFlags::None);
 	LocLoadFlags |= ELocalizationLoadFlags::Engine;
+	LocLoadFlags |= TextLocalizationManager::AlwaysLoadNativeLocalizationDataDuringInitialization ? ELocalizationLoadFlags::Native : ELocalizationLoadFlags::None;
 	LocLoadFlags |= ELocalizationLoadFlags::Additional;
 	
 	ELocalizationLoadFlags ApplyLocLoadFlags = LocLoadFlags;
@@ -586,6 +590,7 @@ void InitGameTextLocalization()
 	TextLocalizationResourceUtil::ClearNativeProjectCultureName();
 
 	ELocalizationLoadFlags LocLoadFlags = ELocalizationLoadFlags::Game;
+	LocLoadFlags |= TextLocalizationManager::AlwaysLoadNativeLocalizationDataDuringInitialization ? ELocalizationLoadFlags::Native : ELocalizationLoadFlags::None;
 	if (PreviousLanguage != CurrentLanguage)
 	{
 		// If the active language changed, then we also need to reload the Engine and Additional localization data 
