[AutoRTFM] Add a memory validator which checks for writes to memory locations in both the open and closed.

Writes made by both the open and closed are highly likely to cause corruption if the transaction is aborted.
To detect this, record a hash of the memory writes when transitioning from non-open -> open, and check the hash is the same when transitioning back from open -> non-open.

Functions marked with AUTORTFM_ALWAYS_OPEN no longer register the open function as the closed. Instead a new trampoline function is created that calls autortfm_pre_open(), the open function and then autortfm_post_open().

diff --git a/Engine/Source/Programs/AutoRTFMEngineTests/Private/AutoRTFMEngineTests.cpp b/Engine/Source/Programs/AutoRTFMEngineTests/Private/AutoRTFMEngineTests.cpp
--- a/Engine/Source/Programs/AutoRTFMEngineTests/Private/AutoRTFMEngineTests.cpp
+++ b/Engine/Source/Programs/AutoRTFMEngineTests/Private/AutoRTFMEngineTests.cpp
@@ -231,6 +231,8 @@ static void PreInit()
 
 	FDelayedAutoRegisterHelper::RunAndClearDelayedAutoRegisterDelegates(EDelayedRegisterRunPhase::StatSystemReady);
 
+	// Unexpected memory validation errors should be full assertions.
+	AutoRTFM::ForTheRuntime::SetMemoryValidationLevel(AutoRTFM::EMemoryValidationLevel::Enabled);
 
 	// Ensure GCachedScalabilityCVars.bInitialized is set (required by some mesh components)
 	FTaskTagScope::SwapTag(ETaskTag::EGameThread);
diff --git a/Engine/Source/Programs/AutoRTFMTests/Private/AutoRTFMTests.cpp b/Engine/Source/Programs/AutoRTFMTests/Private/AutoRTFMTests.cpp
--- a/Engine/Source/Programs/AutoRTFMTests/Private/AutoRTFMTests.cpp
+++ b/Engine/Source/Programs/AutoRTFMTests/Private/AutoRTFMTests.cpp
@@ -1,5 +1,6 @@
 // Copyright Epic Games, Inc. All Rights Reserved.
 
+#include "AutoRTFM/AutoRTFM.h"
 #include "CoreMinimal.h"
 #include "Misc/EventPool.h"
 #include "Misc/LazySingleton.h"
@@ -98,6 +99,9 @@ int RunTests(int ArgC, const char* ArgV[])
 	// We don't want to trigger ensure's on abort because we are going to test that.
 	AutoRTFM::ForTheRuntime::SetEnsureOnAbortByLanguage(false);
 
+	// Unexpected memory validation errors should be full assertions.
+	AutoRTFM::ForTheRuntime::SetMemoryValidationLevel(AutoRTFM::EMemoryValidationLevel::Enabled);
+
 	int Result = Session.run();
 
 	FPlatformMisc::RequestExit(false);
diff --git a/Engine/Source/Programs/AutoRTFMTests/Private/FakeVM.cpp b/Engine/Source/Programs/AutoRTFMTests/Private/FakeVM.cpp
--- a/Engine/Source/Programs/AutoRTFMTests/Private/FakeVM.cpp
+++ b/Engine/Source/Programs/AutoRTFMTests/Private/FakeVM.cpp
@@ -1,6 +1,7 @@
 // Copyright Epic Games, Inc. All Rights Reserved.
 
 #include "Catch2Includes.h"
+#include "Utils.h"
 #include <AutoRTFM/AutoRTFM.h>
 #include <vector>
 
@@ -57,6 +58,10 @@ void Bad(int*, int)
 
 AutoRTFM::ETransactionResult FakeVM(const FOpcode* Opcodes, const int Length, int& Value)
 {
+	// SOL-6743: These tests perform writes within open and closed scopes to the same memory
+	// locations, which will cause memory validation errors.
+	AUTORTFM_SCOPED_DISABLE_MEMORY_VALIDATION();
+
 	return AutoRTFM::Transact([Opcodes, Length, &Value]
 	{
 		AutoRTFM::Open([Opcodes, Length, &Value]
diff --git a/Engine/Source/Programs/AutoRTFMTests/Private/Open.cpp b/Engine/Source/Programs/AutoRTFMTests/Private/Open.cpp
--- a/Engine/Source/Programs/AutoRTFMTests/Private/Open.cpp
+++ b/Engine/Source/Programs/AutoRTFMTests/Private/Open.cpp
@@ -3,8 +3,20 @@
 #include "Catch2Includes.h"
 #include <AutoRTFM/AutoRTFM.h>
 #include <atomic>
+#include <catch_amalgamated.hpp>
 #include <map>
 #include <vector>
+#include "Utils.h"
+
+namespace
+{
+
+UE_AUTORTFM_ALWAYS_OPEN static void AssignIntPointer(int* Pointer, int Value)
+{
+    *Pointer = Value;
+}
+
+}
 
 TEST_CASE("Open")
 {
@@ -230,3 +242,239 @@ TEST_CASE("Open.ReturnValue")
         }
     }
 }
+
+TEST_CASE("Open.Collision")
+{
+    AUTORTFM_SCOPED_ENABLE_MEMORY_VALIDATION_AS_WARNING();
+
+    AutoRTFMTestUtils::FCaptureWarningContext WarningContext;
+    int I = 0;
+    int J = 0;
+
+    SECTION("NoCollision")
+    {
+        SECTION("Different Memory Locations")
+        {
+            AutoRTFM::Transact([&] ()
+            {
+                I = 42;
+                AutoRTFM::Open([&] () { J = 24; });
+            });
+
+            REQUIRE(I == 42);
+            REQUIRE(J == 24);
+            REQUIRE(WarningContext.GetWarnings().empty());
+        }
+
+        SECTION("Transact(Open(Assign), Assign)")
+        {
+            AutoRTFM::Transact([&] ()
+            {
+                AutoRTFM::Open([&] () { I = 24; });
+                I = 42;
+            });
+
+            REQUIRE(I == 42);
+            REQUIRE(WarningContext.GetWarnings().empty());
+        }
+
+        SECTION("Transact(Assign, Open(Transact(Assign)))")
+        {
+            AutoRTFM::Transact([&] ()
+            {
+                I = 24;
+                AutoRTFM::Open([&] ()
+                {
+                    AutoRTFM::Transact([&] ()
+                    {
+                        I = 42;
+                    });
+                });
+            });
+
+            REQUIRE(I == 42);
+            REQUIRE(WarningContext.GetWarnings().empty());
+        }
+
+        SECTION("Transact(Open(Transact(Assign)), Assign)")
+        {
+            AutoRTFM::Transact([&] ()
+            {
+                AutoRTFM::Open([&] ()
+                {
+                    AutoRTFM::Transact([&] ()
+                    {
+                        I = 24;
+                    });
+                });
+                I = 42;
+            });
+
+            REQUIRE(I == 42);
+            REQUIRE(WarningContext.GetWarnings().empty());
+        }
+    }
+
+    SECTION("Transact(Assign, Open(Assign))")
+    {
+        AutoRTFM::Transact([&] ()
+        {
+            I = 42;
+            AutoRTFM::Open([&] () { I = 24; });
+        });
+
+        REQUIRE_THAT(WarningContext.GetWarnings(),
+            Catch::Matchers::VectorContains(FString(AutoRTFMTestUtils::kMemoryModifiedWarning)));
+    }
+
+    SECTION("Transact(Assign, Transact(Open(Assign)))")
+    {
+        AutoRTFM::Transact([&] ()
+        {
+            I = 42;
+            AutoRTFM::Transact([&] ()
+            {
+                AutoRTFM::Open([&] () { I = 24; });
+            });
+        });
+
+        REQUIRE_THAT(WarningContext.GetWarnings(),
+            Catch::Matchers::VectorContains(FString(AutoRTFMTestUtils::kMemoryModifiedWarning)));
+    }
+
+    SECTION("Transact(Assign, Transact(Transact(Open(Assign))))")
+    {
+        AutoRTFM::Transact([&] ()
+        {
+            I = 42;
+            AutoRTFM::Transact([&] ()
+            {
+                AutoRTFM::Transact([&] ()
+                {
+                    AutoRTFM::Open([&] ()
+                    {
+                        I = 24;
+                    });
+                });
+            });
+        });
+
+        REQUIRE_THAT(WarningContext.GetWarnings(),
+            Catch::Matchers::VectorContains(FString(AutoRTFMTestUtils::kMemoryModifiedWarning)));
+    }
+
+    SECTION("Transact(Transact(Assign, Transact(Open(Assign))))")
+    {
+        AutoRTFM::Transact([&] ()
+        {
+            AutoRTFM::Transact([&] ()
+            {
+                I = 42;
+                AutoRTFM::Transact([&] ()
+                {
+                    AutoRTFM::Open([&] ()
+                    {
+                        I = 24;
+                    });
+                });
+            });
+        });
+
+        REQUIRE_THAT(WarningContext.GetWarnings(),
+            Catch::Matchers::VectorContains(FString(AutoRTFMTestUtils::kMemoryModifiedWarning)));
+    }
+
+    SECTION("Transact(Transact(Assign, Open(Assign, Transact())))")
+    {
+        AutoRTFM::Transact([&] ()
+        {
+            I = 42;
+            AutoRTFM::Open([&] ()
+            {
+                I = 24;
+                AutoRTFM::Transact([&] () {});
+            });
+        });
+
+        REQUIRE_THAT(WarningContext.GetWarnings(),
+            Catch::Matchers::VectorContains(FString(AutoRTFMTestUtils::kMemoryModifiedWarning)));
+    }
+
+    SECTION("Transact(Transact(Assign, Open(Transact(), Assign)))")
+    {
+        AutoRTFM::Transact([&] ()
+        {
+            I = 42;
+            AutoRTFM::Open([&] ()
+            {
+                AutoRTFM::Transact([&] () {});
+                I = 24;
+            });
+        });
+
+        REQUIRE_THAT(WarningContext.GetWarnings(),
+            Catch::Matchers::VectorContains(FString(AutoRTFMTestUtils::kMemoryModifiedWarning)));
+    }
+
+    SECTION("Transact(Assign, CallOpen(Assign))")
+    {
+        AutoRTFM::Transact([&] ()
+        {
+            I = 42;
+            AssignIntPointer(&I, 24);
+        });
+
+        REQUIRE_THAT(WarningContext.GetWarnings(),
+            Catch::Matchers::VectorContains(FString(AutoRTFMTestUtils::kMemoryModifiedWarning)));
+    }
+
+    SECTION("Transact(Assign, Transact(CallOpen(Assign)))")
+    {
+        AutoRTFM::Transact([&] ()
+        {
+            I = 42;
+            AutoRTFM::Transact([&] ()
+            {
+                AssignIntPointer(&I, 24);
+            });
+        });
+
+        REQUIRE_THAT(WarningContext.GetWarnings(),
+            Catch::Matchers::VectorContains(FString(AutoRTFMTestUtils::kMemoryModifiedWarning)));
+    }
+
+    SECTION("Transact(Assign, Open(Assign, OpenNoValidation()))")
+    {
+        AutoRTFM::Transact([&] ()
+        {
+            I = 42;
+            AutoRTFM::Open([&]
+            {
+                I = 10;
+                AutoRTFM::Open<AutoRTFM::EMemoryValidationLevel::Disabled>([&]{});
+            });
+        });
+
+        REQUIRE_THAT(WarningContext.GetWarnings(),
+            Catch::Matchers::VectorContains(FString(AutoRTFMTestUtils::kMemoryModifiedWarning)));
+    }
+    
+    SECTION("Transact(Assign, Open(Assign, Close(OpenNoValidation())))")
+    {
+        AutoRTFM::Transact([&] ()
+        {
+            I = 42;
+            AutoRTFM::Open([&]
+            {
+                I = 10;
+                std::ignore = AutoRTFM::Close([&]
+                {
+                    AutoRTFM::Open<AutoRTFM::EMemoryValidationLevel::Disabled>([&]{});
+                });
+            });
+        });
+
+        REQUIRE_THAT(WarningContext.GetWarnings(),
+            Catch::Matchers::VectorContains(FString(AutoRTFMTestUtils::kMemoryModifiedWarning)));
+    }
+}
diff --git a/Engine/Source/Programs/AutoRTFMTests/Private/OpenAPI.cpp b/Engine/Source/Programs/AutoRTFMTests/Private/OpenAPI.cpp
--- a/Engine/Source/Programs/AutoRTFMTests/Private/OpenAPI.cpp
+++ b/Engine/Source/Programs/AutoRTFMTests/Private/OpenAPI.cpp
@@ -1,6 +1,7 @@
 // Copyright Epic Games, Inc. All Rights Reserved.
 
 #include "Catch2Includes.h"
+#include "Utils.h"
 #include <AutoRTFM/AutoRTFM.h>
 
 #include <map>
@@ -66,6 +67,9 @@ TEST_CASE("OpenAPI.RecordDataClosed_Illegal", "[.]")
 
 TEST_CASE("OpenAPI.WriteDataInTheOpen")
 {
+	AUTORTFM_SCOPED_ENABLE_MEMORY_VALIDATION_AS_WARNING();
+	AutoRTFMTestUtils::FCaptureWarningContext WarningContext;
+
 	int value = 0;
 	auto transactResult = AutoRTFM::Transact([&]()
 	{
@@ -78,6 +82,8 @@ TEST_CASE("OpenAPI.WriteDataInTheOpen")
 
 	REQUIRE(transactResult == AutoRTFM::ETransactionResult::Committed);
 	REQUIRE(value == 1);
+	REQUIRE_THAT(WarningContext.GetWarnings(),
+		Catch::Matchers::VectorContains(FString(AutoRTFMTestUtils::kMemoryModifiedWarning)));
 }
 
 TEST_CASE("OpenAPI.AbortTransactionScopedFromOpen")
@@ -476,6 +482,8 @@ TEST_CASE("OpenAPI.StackWriteCommitInTheOpen1")
 
 TEST_CASE("OpenAPI.StackWriteCommitInTheOpen2")
 {
+	AUTORTFM_SCOPED_ENABLE_MEMORY_VALIDATION_AS_WARNING();
+	AutoRTFMTestUtils::FCaptureWarningContext WarningContext;
 	AutoRTFM::EContextStatus Status = AutoRTFM::EContextStatus::Idle;
 
 	int Value = 0;
@@ -503,6 +511,9 @@ TEST_CASE("OpenAPI.StackWriteCommitInTheOpen2")
 		REQUIRE(AutoRTFM::EContextStatus::OnTrack == Status);
 		REQUIRE(Value == 10);
 	});
+
+	REQUIRE_THAT(WarningContext.GetWarnings(),
+		Catch::Matchers::VectorContains(FString(AutoRTFMTestUtils::kMemoryModifiedWarning)));
 }
 
 TEST_CASE("OpenAPI.StackWriteAbortInTheOpen1")
@@ -683,6 +694,9 @@ TEST_CASE("OpenAPI.WriteTrivialStructure2")
 
 TEST_CASE("OpenAPI.Footgun1")
 {
+	AUTORTFM_SCOPED_ENABLE_MEMORY_VALIDATION_AS_WARNING();
+	AutoRTFMTestUtils::FCaptureWarningContext WarningContext;
+
 	int valueA = 0;
 	int valueB = 0;
 
@@ -710,10 +724,16 @@ TEST_CASE("OpenAPI.Footgun1")
 	// We rollback the transaction to the value we had when we first recorded the address
 	REQUIRE(valueA == 10);
 	REQUIRE(valueB == 0);
+
+	REQUIRE_THAT(WarningContext.GetWarnings(),
+		Catch::Matchers::VectorContains(FString(AutoRTFMTestUtils::kMemoryModifiedWarning)));
 }
 
 TEST_CASE("OpenAPI.Footgun2")
 {
+	AUTORTFM_SCOPED_ENABLE_MEMORY_VALIDATION_AS_WARNING();
+	AutoRTFMTestUtils::FCaptureWarningContext WarningContext;
+
 	int valueB = 0;
 	int valueC = 0;
 	AutoRTFM::Transact([&]()
@@ -742,6 +762,9 @@ TEST_CASE("OpenAPI.Footgun2")
 	// We rollback the transaction to the value we had when we first recorded the address
 	REQUIRE(valueB == 0);
 	REQUIRE(valueC == 10);
+
+	REQUIRE_THAT(WarningContext.GetWarnings(),
+		Catch::Matchers::VectorContains(FString(AutoRTFMTestUtils::kMemoryModifiedWarning)));
 }
 
 #if 0
diff --git a/Engine/Source/Programs/AutoRTFMTests/Private/Utils.h b/Engine/Source/Programs/AutoRTFMTests/Private/Utils.h
new file mode 100644
--- /dev/null
+++ b/Engine/Source/Programs/AutoRTFMTests/Private/Utils.h
@@ -0,0 +1,93 @@
+// Copyright Epic Games, Inc. All Rights Reserved.
+
+#pragma once
+
+#include "Misc/FeedbackContext.h"
+#include <AutoRTFM/AutoRTFM.h>
+#include <vector>
+
+namespace AutoRTFMTestUtils
+{
+
+// A helper class that for the lifetime of the object, intercepts and records UE_LOG warnings.
+class FCaptureWarningContext : private FFeedbackContext
+{
+public:
+	FCaptureWarningContext() : OldContext(GWarn) { GWarn = this; }
+	~FCaptureWarningContext() { GWarn = OldContext; }
+
+	void Serialize(const TCHAR* V, ELogVerbosity::Type Verbosity, const FName& Category) override
+	{
+		if (Verbosity == ELogVerbosity::Warning)
+		{
+			Warnings.push_back(V);
+		}
+		else
+		{
+			OldContext->Serialize(V, Verbosity, Category);
+		}
+	}
+
+	void Serialize(const TCHAR* V, ELogVerbosity::Type Verbosity, const FName& Category, double Time) override
+	{
+		if (Verbosity == ELogVerbosity::Warning)
+		{
+			Warnings.push_back(V);
+		}
+		else
+		{
+			OldContext->Serialize(V, Verbosity, Category, Time);
+		}
+	}
+
+	const std::vector<FString>& GetWarnings() const
+	{
+		return Warnings;
+	}
+
+private:
+	FCaptureWarningContext(const FCaptureWarningContext&) = delete;
+	FCaptureWarningContext& operator = (const FCaptureWarningContext&) = delete;
+
+	FFeedbackContext* OldContext = nullptr;
+	std::vector<FString> Warnings;
+};
+
+
+// A helper class that temporarily changes the memory validation level for the lifetime of the
+// object, restoring the original level on destruction.
+class ScopedMemoryValidationLevel
+{
+public:
+	ScopedMemoryValidationLevel(AutoRTFM::EMemoryValidationLevel NewLevel)
+		: PrevLevel(AutoRTFM::ForTheRuntime::GetMemoryValidationLevel())
+	{
+		AutoRTFM::ForTheRuntime::SetMemoryValidationLevel(NewLevel);
+	}
+
+	~ScopedMemoryValidationLevel()
+	{
+		AutoRTFM::ForTheRuntime::SetMemoryValidationLevel(PrevLevel);
+	}
+private:
+	ScopedMemoryValidationLevel(const ScopedMemoryValidationLevel&) = delete;
+	ScopedMemoryValidationLevel& operator = (const ScopedMemoryValidationLevel&) = delete;
+
+	const AutoRTFM::EMemoryValidationLevel PrevLevel;
+};
+
+#define AUTORTFM_SCOPED_DISABLE_MEMORY_VALIDATION() \
+	AutoRTFMTestUtils::ScopedMemoryValidationLevel DisableMemoryValidation \
+		{AutoRTFM::EMemoryValidationLevel::Disabled}; \
+	static_assert(true) /* require semicolon */
+
+#define AUTORTFM_SCOPED_ENABLE_MEMORY_VALIDATION_AS_WARNING() \
+	AutoRTFMTestUtils::ScopedMemoryValidationLevel EnableMemoryValidationAsWarning \
+		{AutoRTFM::EMemoryValidationLevel::EnabledAsWarning}; \
+	static_assert(true) /* require semicolon */
+
+static constexpr const TCHAR* kMemoryModifiedWarning =
+    TEXT("Memory modified in a transaction was also modified in an call to AutoRTFM::Open(). ")
+    TEXT("This may lead to memory corruption if the transaction is aborted.");
+
+}
diff --git a/Engine/Source/Runtime/Core/Private/AutoRTFM/API.cpp b/Engine/Source/Runtime/Core/Private/AutoRTFM/API.cpp
--- a/Engine/Source/Runtime/Core/Private/AutoRTFM/API.cpp
+++ b/Engine/Source/Runtime/Core/Private/AutoRTFM/API.cpp
@@ -4,6 +4,8 @@
 #include "CoreGlobals.h"
 #include "HAL/IConsoleManager.h"
 #include "GenericPlatform/GenericPlatformCrashContext.h"
+#include "HAL/Platform.h"
+#include "Transaction.h"
 
 #if UE_AUTORTFM
 static_assert(UE_AUTORTFM_ENABLED, "AutoRTFM/API.cpp requires the compiler flag '-fautortfm'");
@@ -27,6 +29,10 @@ namespace
 
 	int GAutoRTFMRetryTransactions = AutoRTFM::ForTheRuntime::EAutoRTFMRetryTransactionState::NoRetry;
 
+	static constexpr AutoRTFM::EMemoryValidationLevel GDefaultMemoryValidationLevel = AutoRTFM::EMemoryValidationLevel::Disabled;
+
+	int GMemoryValidationLevel = static_cast<int>(GDefaultMemoryValidationLevel);
+
 	void UpdateAutoRTFMRetryTransactionsData()
 	{
 		switch (GAutoRTFMRetryTransactions)
@@ -39,6 +45,19 @@ namespace
 			return FGenericCrashContext::SetGameData(TEXT("AutoRTFMRetryTransactionState"), TEXT("RetryNestedToo"));
 		}
 	}
+
+	void UpdateAutoRTFMMemoryValidationLevelData()
+	{
+		switch (static_cast<AutoRTFM::EMemoryValidationLevel>(GMemoryValidationLevel))
+		{
+		case AutoRTFM::EMemoryValidationLevel::Disabled:
+			return FGenericCrashContext::SetGameData(TEXT("AutoRTFMMemoryValidationLevel"), TEXT("Disabled"));
+		case AutoRTFM::EMemoryValidationLevel::EnabledAsWarning:
+			return FGenericCrashContext::SetGameData(TEXT("AutoRTFMMemoryValidationLevel"), TEXT("EnabledAsWarning"));
+		case AutoRTFM::EMemoryValidationLevel::Enabled:
+			return FGenericCrashContext::SetGameData(TEXT("AutoRTFMMemoryValidationLevel"), TEXT("Enabled"));
+		}
+	}
 }
 
 #if UE_AUTORTFM
@@ -58,6 +77,14 @@ static FAutoConsoleVariableRef CVarAutoRTFMRetryTransactions(
 	ECVF_Default
 );
 
+static FAutoConsoleVariableRef CVarAutoRTFMMemoryValidationLevel(
+	TEXT("AutoMemoryValidationLevel"),
+	GMemoryValidationLevel,
+	TEXT("Detects potential memory corruption due to writes made both by a transaction and open-code"),
+	FConsoleVariableDelegate::CreateLambda([](IConsoleVariable*) { UpdateAutoRTFMMemoryValidationLevelData(); }),
+	ECVF_Default
+);
+
 static FDelayedAutoRegisterHelper DelayedAutoRegister(EDelayedRegisterRunPhase::EndOfEngineInit, []
 	{
 		UpdateAutoRTFMRuntimeCrashData();
@@ -184,6 +211,24 @@ namespace AutoRTFM
 			return false;
 #endif
 		}
+
+		EMemoryValidationLevel GetMemoryValidationLevel()
+		{
+#if UE_AUTORTFM
+			return static_cast<EMemoryValidationLevel>(GMemoryValidationLevel);
+#else
+			return EMemoryValidationLevel::Disabled;
+#endif
+		}
+
+		void SetMemoryValidationLevel(EMemoryValidationLevel Level)
+		{
+#if UE_AUTORTFM
+			GMemoryValidationLevel = static_cast<int>(
+				Level == EMemoryValidationLevel::Default ? GDefaultMemoryValidationLevel : Level);
+#endif
+		}
+
 	}
 }
 
@@ -203,6 +248,8 @@ namespace AutoRTFM
 // This is the implementation of the AutoRTFM.h API. Ideally, functions here should just delegate to some internal API.
 // For now, I have these functions also perform some error checking.
 
+#define AUTORTFM_CORE_API [[clang::autortfm_core_api, clang::noinline]]
+
 namespace AutoRTFM
 {
 
@@ -210,12 +257,12 @@ namespace
 {
 
 // Internal closed-variant implementations.
-UE_AUTORTFM_ALWAYS_OPEN bool RTFM_autortfm_is_transactional()
+AUTORTFM_CORE_API bool RTFM_autortfm_is_transactional()
 {
 	return true;
 }
 
-UE_AUTORTFM_ALWAYS_OPEN autortfm_result RTFM_autortfm_transact(void (*UninstrumentedWork)(void*), void (*InstrumentedWork)(void*), void* Arg)
+AUTORTFM_CORE_API autortfm_result RTFM_autortfm_transact(void (*UninstrumentedWork)(void*), void (*InstrumentedWork)(void*), void* Arg)
 {
 	FContext* Context = FContext::Get();
 	return static_cast<autortfm_result>(Context->Transact(InstrumentedWork, Arg));
@@ -230,63 +277,52 @@ UE_AUTORTFM_FORCEINLINE autortfm_result TransactThenOpenImpl(void (*Uninstrument
 		}));
 }
 
-UE_AUTORTFM_ALWAYS_OPEN autortfm_result RTFM_autortfm_transact_then_open(void (*UninstrumentedWork)(void*), void (*InstrumentedWork)(void*), void* Arg)
+AUTORTFM_CORE_API autortfm_result RTFM_autortfm_transact_then_open(void (*UninstrumentedWork)(void*), void (*InstrumentedWork)(void*), void* Arg)
 {
 	return TransactThenOpenImpl(UninstrumentedWork, Arg);
 }
 
-UE_AUTORTFM_ALWAYS_OPEN void RTFM_autortfm_commit(void (*UninstrumentedWork)(void*), void (*InstrumentedWork)(void*), void* Arg)
+AUTORTFM_CORE_API void RTFM_autortfm_commit(void (*UninstrumentedWork)(void*), void (*InstrumentedWork)(void*), void* Arg)
 {
 	autortfm_result Result = autortfm_transact(UninstrumentedWork, InstrumentedWork, Arg);
 	UE_CLOG(Result != autortfm_committed, LogAutoRTFM, Fatal, TEXT("Unexpected transaction result: %u."), Result);
 }
 
-UE_AUTORTFM_ALWAYS_OPEN void RTFM_autortfm_abort()
+AUTORTFM_CORE_API void RTFM_autortfm_abort()
 {
 	FContext* Context = FContext::Get();
 	Context->AbortByRequestAndThrow();
 }
 
-UE_AUTORTFM_ALWAYS_OPEN void RTFM_autortfm_start_transaction()
+AUTORTFM_CORE_API void RTFM_autortfm_start_transaction()
 {
 	UE_LOG(LogAutoRTFM, Fatal, TEXT("The function `autortfm_start_transaction` was called from closed code."));
 }
 
-UE_AUTORTFM_ALWAYS_OPEN void RTFM_autortfm_commit_transaction()
+AUTORTFM_CORE_API void RTFM_autortfm_commit_transaction()
 {
 	UE_LOG(LogAutoRTFM, Fatal, TEXT("The function `RTFM_autortfm_commit_transaction` was called from closed code."));
 }
 
-UE_AUTORTFM_ALWAYS_OPEN autortfm_result RTFM_autortfm_abort_transaction()
+AUTORTFM_CORE_API autortfm_result RTFM_autortfm_abort_transaction()
 {
 	FContext* const Context = FContext::Get();
 	return static_cast<autortfm_result>(Context->AbortTransaction(/* bIsClosed */ true, /* bIsCascading */ false));
 }
 
-UE_AUTORTFM_ALWAYS_OPEN autortfm_result RTFM_autortfm_cascading_abort_transaction()
+AUTORTFM_CORE_API autortfm_result RTFM_autortfm_cascading_abort_transaction()
 {
 	FContext* const Context = FContext::Get();
 	return static_cast<autortfm_result>(Context->AbortTransaction(/* bIsClosed */ true, /* bIsCascading */ true));
 }
 
-UE_AUTORTFM_ALWAYS_OPEN void RTFM_autortfm_clear_transaction_status()
+AUTORTFM_CORE_API void RTFM_autortfm_clear_transaction_status()
 {
 	UE_LOG(LogAutoRTFM, Fatal, TEXT("The function `autortfm_clear_transaction_status` was called from closed code."));
 	AutoRTFM::Unreachable();
 }
 
-UE_AUTORTFM_ALWAYS_OPEN void RTFM_autortfm_open(void (*Work)(void*), void* Arg)
-{
-	Work(Arg);
-
-	FContext* Context = FContext::Get();
-	if (Context->IsAborting())
-	{
-		Context->Throw();
-	}
-}
-
-UE_AUTORTFM_ALWAYS_OPEN autortfm_status RTFM_autortfm_close(void (*UninstrumentedWork)(void*), void (*InstrumentedWork)(void*), void* Arg)
+AUTORTFM_CORE_API autortfm_status RTFM_autortfm_close(void (*UninstrumentedWork)(void*), void (*InstrumentedWork)(void*), void* Arg)
 {
 	FContext* const Context = FContext::Get();
 
@@ -303,81 +339,128 @@ UE_AUTORTFM_ALWAYS_OPEN autortfm_status RTFM_autortfm_close(void (*Uninstrumente
 	return static_cast<autortfm_status>(Context->GetStatus());
 }
 
-UE_AUTORTFM_ALWAYS_OPEN void RTFM_autortfm_record_open_write(void*, size_t)
+extern "C" UE_AUTORTFM_NOAUTORTFM void autortfm_pre_open()
+{
+	if (FTransaction* Transaction = FContext::Get()->GetCurrentTransaction())
+	{
+		Transaction->SetMemoryValidationLevel(static_cast<AutoRTFM::EMemoryValidationLevel>(GMemoryValidationLevel));
+		Transaction->SetState(FTransaction::EState::OpenActive);
+	}
+}
+
+extern "C" UE_AUTORTFM_NOAUTORTFM void autortfm_post_open()
+{
+	if (FTransaction* Transaction = FContext::Get()->GetCurrentTransaction())
+	{
+		if (!Transaction->IsDone())  // Transaction may have been aborted.
+		{
+			Transaction->SetState(FTransaction::EState::ClosedActive);
+		}
+	}
+}
+
+AUTORTFM_CORE_API void RTFM_autortfm_open(autortfm_memory_validation_level ValidationLevel, void (*Work)(void*), void* Arg)
+{
+	FTransaction* Transaction = FContext::Get()->GetCurrentTransaction();
+	if (Transaction)
+	{
+		AutoRTFM::EMemoryValidationLevel Level 
+			= ValidationLevel == autortfm_memory_validation_level_default 
+			? static_cast<AutoRTFM::EMemoryValidationLevel>(GMemoryValidationLevel)
+			: static_cast<AutoRTFM::EMemoryValidationLevel>(ValidationLevel);
+		Transaction->SetMemoryValidationLevel(Level);
+		Transaction->SetState(FTransaction::EState::OpenActive);
+	}
+
+	Work(Arg);
+
+	FContext* Context = FContext::Get();
+	if (Transaction == Context->GetCurrentTransaction() && !Transaction->IsDone())  // Transaction may have been aborted.
+	{
+		Transaction->SetState(FTransaction::EState::ClosedActive);
+	}
+
+	if (Context->IsAborting())
+	{
+		Context->Throw();
+	}
+}
+
+AUTORTFM_CORE_API void RTFM_autortfm_record_open_write(void*, size_t)
 {
 	UE_LOG(LogAutoRTFM, Fatal, TEXT("The function `autortfm_record_open_write` was called from closed code."));
 }
 
-UE_AUTORTFM_ALWAYS_OPEN void RTFM_OnCommitInternal(TFunction<void()> && Work)
+AUTORTFM_CORE_API void RTFM_OnCommitInternal(TFunction<void()> && Work)
 {
 	FContext* Context = FContext::Get();
 	ASSERT(Context->GetStatus() == EContextStatus::OnTrack);
 	Context->GetCurrentTransaction()->DeferUntilCommit(MoveTemp(Work));
 }
 
-UE_AUTORTFM_ALWAYS_OPEN void RTFM_OnAbortInternal(TFunction<void()>&& Work)
+AUTORTFM_CORE_API void RTFM_OnAbortInternal(TFunction<void()>&& Work)
 {
 	FContext* Context = FContext::Get();
 	ASSERT(Context->GetStatus() == EContextStatus::OnTrack);
 	Context->GetCurrentTransaction()->DeferUntilAbort(MoveTemp(Work));
 }
 
-UE_AUTORTFM_ALWAYS_OPEN void RTFM_PushOnAbortHandlerInternal(const void* Key, TFunction<void()>&& Work)
+AUTORTFM_CORE_API void RTFM_PushOnAbortHandlerInternal(const void* Key, TFunction<void()>&& Work)
 {
 	FContext* Context = FContext::Get();
 	ASSERT(Context->GetStatus() == EContextStatus::OnTrack);
 	Context->GetCurrentTransaction()->PushDeferUntilAbortHandler(Key, MoveTemp(Work));
 }
 
-UE_AUTORTFM_ALWAYS_OPEN void RTFM_PopOnAbortHandlerInternal(const void* Key)
+AUTORTFM_CORE_API void RTFM_PopOnAbortHandlerInternal(const void* Key)
 {
 	FContext* Context = FContext::Get();
 	ASSERT(Context->GetStatus() == EContextStatus::OnTrack);
 	Context->GetCurrentTransaction()->PopDeferUntilAbortHandler(Key);
 }
 
-UE_AUTORTFM_ALWAYS_OPEN void RTFM_PopAllOnAbortHandlersInternal(const void* Key)
+AUTORTFM_CORE_API void RTFM_PopAllOnAbortHandlersInternal(const void* Key)
 {
 	FContext* Context = FContext::Get();
 	ASSERT(Context->GetStatus() == EContextStatus::OnTrack);
 	Context->GetCurrentTransaction()->PopAllDeferUntilAbortHandlers(Key);
 }
 
-UE_AUTORTFM_ALWAYS_OPEN void RTFM_autortfm_on_commit(void (*Work)(void*), void* Arg)
+AUTORTFM_CORE_API void RTFM_autortfm_on_commit(void (*Work)(void*), void* Arg)
 {
 	RTFM_OnCommitInternal([Work, Arg] { Work(Arg); });
 }
 
-UE_AUTORTFM_ALWAYS_OPEN void RTFM_autortfm_on_abort(void (*Work)(void*), void* Arg)
+AUTORTFM_CORE_API void RTFM_autortfm_on_abort(void (*Work)(void*), void* Arg)
 {
 	RTFM_OnAbortInternal([Work, Arg] { Work(Arg); });
 }
 
-UE_AUTORTFM_ALWAYS_OPEN void RTFM_autortfm_push_on_abort_handler(const void* Key, void (*Work)(void*), void* Arg)
+AUTORTFM_CORE_API void RTFM_autortfm_push_on_abort_handler(const void* Key, void (*Work)(void*), void* Arg)
 {
 	RTFM_PushOnAbortHandlerInternal(Key, [Work, Arg] { Work(Arg); });
 }
 
-UE_AUTORTFM_ALWAYS_OPEN void RTFM_autortfm_pop_on_abort_handler(const void* Key)
+AUTORTFM_CORE_API void RTFM_autortfm_pop_on_abort_handler(const void* Key)
 {
 	RTFM_PopOnAbortHandlerInternal(Key);
 }
 
-UE_AUTORTFM_ALWAYS_OPEN void* RTFM_autortfm_did_allocate(void* Ptr, size_t Size)
+AUTORTFM_CORE_API void* RTFM_autortfm_did_allocate(void* Ptr, size_t Size)
 {
 	FContext* Context = FContext::Get();
 	Context->DidAllocate(Ptr, Size);
 	return Ptr;
 }
 
-UE_AUTORTFM_ALWAYS_OPEN void RTFM_autortfm_did_free(void* Ptr)
+AUTORTFM_CORE_API void RTFM_autortfm_did_free(void* Ptr)
 {
 	// We should never-ever-ever actually free memory from within closed code of
 	// a transaction.
 	AutoRTFM::Unreachable();
 }
 
-UE_AUTORTFM_ALWAYS_OPEN void RTFM_autortfm_check_consistency_assuming_no_races()
+AUTORTFM_CORE_API void RTFM_autortfm_check_consistency_assuming_no_races()
 {
 }
 
@@ -534,11 +617,21 @@ extern "C" UE_AUTORTFM_NOAUTORTFM bool autortfm_current_nest_throw()
 	return true;
 }
 
+extern "C" void autortfm_open_explicit_validation(autortfm_memory_validation_level ValidationLevel, void (*Work)(void*), void* Arg)
+{
+	if (autortfm_is_closed())
+	{
+		return RTFM_autortfm_open(ValidationLevel, Work, Arg);
+	}
+
+	Work(Arg);
+}
+
 extern "C" void autortfm_open(void (*Work)(void*), void* Arg)
 {
 	if (autortfm_is_closed())
 	{
-		return RTFM_autortfm_open(Work, Arg);
+		return RTFM_autortfm_open(autortfm_memory_validation_level_default, Work, Arg);
 	}
 
 	Work(Arg);
@@ -559,7 +652,7 @@ extern "C" autortfm_status autortfm_close(void (*UninstrumentedWork)(void*), voi
 
 		FContext* const Context = FContext::Get();
 
-		if (InstrumentedWork)
+		if (LIKELY(InstrumentedWork))
 		{
 			Result = static_cast<autortfm_status>(Context->CallClosedNest(InstrumentedWork, Arg));
 		}
diff --git a/Engine/Source/Runtime/Core/Private/AutoRTFM/Context.cpp b/Engine/Source/Runtime/Core/Private/AutoRTFM/Context.cpp
--- a/Engine/Source/Runtime/Core/Private/AutoRTFM/Context.cpp
+++ b/Engine/Source/Runtime/Core/Private/AutoRTFM/Context.cpp
@@ -2,14 +2,18 @@
 
 #if (defined(__AUTORTFM) && __AUTORTFM)
 #include "Context.h"
+
+#include "AutoRTFM/AutoRTFMMetrics.h"
 #include "CallNestInlines.h"
 #include "ContextInlines.h"
 #include "FunctionMap.h"
 #include "GlobalData.h"
+#include "Misc/AssertionMacros.h"
 #include "ScopedGuard.h"
+#include "StackRange.h"
 #include "Stats.h"
+#include "Transaction.h"
 #include "TransactionInlines.h"
-#include "AutoRTFM/AutoRTFMMetrics.h"
 
 #include "Templates/UniquePtr.h"
 #include "Containers/StringConv.h"
@@ -74,7 +78,7 @@ void FContext::MaterializeDeferredTransactions()
 bool FContext::StartTransaction()
 {
 	NumDeferredTransactions += 1;
-	ensureMsgf(RawCurrentTransaction(), TEXT("FContext::StartTransaction() can only be called within a scoped transaction"));
+	ensureMsgf(CurrentTransaction, TEXT("FContext::StartTransaction() can only be called within a scoped transaction"));
 	ASSERT(Status == EContextStatus::OnTrack);
 
 	GAutoRTFMMetrics.NumTransactionsStarted++;
@@ -85,13 +89,15 @@ bool FContext::StartTransaction()
 void FContext::StartNonDeferredTransaction()
 {
 	ASSERT(GetNumDeferredTransactions() == 0);
-	ensureMsgf(RawCurrentTransaction(), TEXT("FContext::StartNonDeferredTransaction() can only be called within a scoped transaction"));
-	FTransaction* NewTransaction = new FTransaction(this);
-	NewTransaction->SetStackRange(RawCurrentTransaction()->GetStackRange());
+	ensureMsgf(CurrentTransaction, TEXT("FContext::StartNonDeferredTransaction() can only be called within a scoped transaction"));
+
+	FTransaction* NewTransaction = PushTransaction(
+		/* Closed */ false,
+		/* bIsScoped */ false,
+		/* StackRange */ CurrentTransaction->GetStackRange());
 
 	// This form of transaction is always ultimately within a scoped Transact 
 	ASSERT(Status == EContextStatus::OnTrack);
-	PushTransaction(NewTransaction);
 
 	GAutoRTFMMetrics.NumTransactionsStarted++;
 }
@@ -112,11 +118,11 @@ ETransactionResult FContext::CommitTransaction()
 	else
 	{
 		// Scoped transactions commit on return, so committing explicitly isn't allowed
-		ASSERT(RawCurrentTransaction()->IsScopedTransaction() == false);
+		ASSERT(CurrentTransaction->IsScopedTransaction() == false);
 
-		if (RawCurrentTransaction()->IsNested())
+		if (CurrentTransaction->IsNested())
 		{
-			Result = ResolveNestedTransaction(RawCurrentTransaction());
+			Result = ResolveNestedTransaction(CurrentTransaction);
 		}
 		else
 		{
@@ -124,7 +130,7 @@ ETransactionResult FContext::CommitTransaction()
 			DumpState();
 			UE_LOG(LogAutoRTFM, Verbose, TEXT("Committing..."));
 
-			if (AttemptToCommitTransaction(RawCurrentTransaction()))
+			if (AttemptToCommitTransaction(CurrentTransaction))
 			{
 				Result = ETransactionResult::Committed;
 			}
@@ -164,15 +170,15 @@ ETransactionResult FContext::AbortTransaction(bool bIsClosed, bool bIsCascading)
 	}
 	else
 	{
-		ASSERT(nullptr != RawCurrentTransaction());
+		ASSERT(nullptr != CurrentTransaction);
 
 		// Sort out how aborts work
-		RawCurrentTransaction()->AbortWithoutThrowing();
+		CurrentTransaction->AbortWithoutThrowing();
 
 		// Non-scoped transactions are ended immediately, but scoped need to get to the end scope before being popped
-		if (!RawCurrentTransaction()->IsScopedTransaction())
+		if (!CurrentTransaction->IsScopedTransaction())
 		{
-			Result = ResolveNestedTransaction(RawCurrentTransaction());
+			Result = ResolveNestedTransaction(CurrentTransaction);
 			PopTransaction();
 		}
 	}
@@ -202,12 +208,22 @@ bool FContext::IsAborting() const
 EContextStatus FContext::CallClosedNest(void (*ClosedFunction)(void* Arg), void* Arg)
 {
 	TScopedGuard<void*> ClosedStackAddressGuard(ClosedStackAddress, &ClosedStackAddressGuard);
+
+	FTransaction* const Transaction = GetCurrentTransaction();
+	ASSERT(Transaction != nullptr);
+	FTransaction::EState PreviousState = Transaction->SetState(FTransaction::EState::ClosedActive);
+
 	PushCallNest(new FCallNest(this));
 
 	CurrentNest->Try([&]() { ClosedFunction(Arg); });
 
 	PopCallNest();
 
+	if (Transaction == CurrentTransaction && !Transaction->IsDone())  // Transaction may have been aborted.
+	{
+		Transaction->SetState(PreviousState);
+	}
+
 	return GetStatus();
 }
 
@@ -229,29 +245,43 @@ void FContext::PopCallNest()
 	delete OldCallNest;
 }
 
-void FContext::PushTransaction(FTransaction* NewTransaction)
+FTransaction* FContext::PushTransaction(bool bClosed, bool bIsScoped, FStackRange StackRange)
 {
 	ASSERT(!GetNumDeferredTransactions());
-	ASSERT(NewTransaction != nullptr);
-	ASSERT(!NewTransaction->IsDone());
-	ASSERT(NewTransaction->GetParent() == nullptr);
 
-	ASSERT(RawCurrentTransaction() == nullptr || !RawCurrentTransaction()->IsDone());
-	
-	NewTransaction->SetParent(RawCurrentTransaction());
+	if (CurrentTransaction != nullptr)
+	{
+		ASSERT(CurrentTransaction->IsActive());
+		CurrentTransaction->SetInactive();
+	}
+
+	FTransaction* NewTransaction = new FTransaction(this);
+	NewTransaction->Initialize(
+		/* State */ bClosed ? FTransaction::EState::ClosedActive : FTransaction::EState::OpenActive,
+		/* Parent */ CurrentTransaction, 
+		/* bIsScoped */ bIsScoped, 
+		/* StackRange */ StackRange);
+
 	CurrentTransaction = NewTransaction;
 
 	// Collect stats that we've got a new transaction.
 	Stats.Collect<EStatsKind::Transaction>();
+
+	return NewTransaction;
 }
 
 void FContext::PopTransaction()
 {
 	ASSERT(!GetNumDeferredTransactions());
-	ASSERT(RawCurrentTransaction() != nullptr);
-	ASSERT(RawCurrentTransaction()->IsDone());
-	FTransaction* OldTransaction = RawCurrentTransaction();
-	CurrentTransaction = RawCurrentTransaction()->GetParent();
+	ASSERT(CurrentTransaction != nullptr);
+	ASSERT(CurrentTransaction->IsDone());
+	FTransaction* OldTransaction = CurrentTransaction;
+	CurrentTransaction = CurrentTransaction->GetParent();
+	if (CurrentTransaction != nullptr)
+	{
+		ASSERT(CurrentTransaction->IsInactive());
+		CurrentTransaction->SetActive();
+	}
 	delete OldTransaction;
 }
 
@@ -274,11 +304,6 @@ void FContext::ClearTransactionStatus()
 
 ETransactionResult FContext::ResolveNestedTransaction(FTransaction* NewTransaction)
 {
-	// We just use this bit to help assertions for now (though we could use it more strongly). Because of how we use this right now,
-	// it's OK that it's set before we commit but after we abort.
-	ASSERT(!NewTransaction->IsDone());
-	NewTransaction->SetIsDone();
-
 	if (Status == EContextStatus::OnTrack)
 	{
 		bool bCommitResult = AttemptToCommitTransaction(NewTransaction);
@@ -287,6 +312,8 @@ ETransactionResult FContext::ResolveNestedTransaction(FTransaction* NewTransacti
 		return ETransactionResult::Committed;
 	}
 
+	ASSERT(NewTransaction->IsDone());
+
 	switch (Status)
 	{
 	case EContextStatus::AbortedByRequest:
@@ -334,17 +361,13 @@ ETransactionResult FContext::Transact(void (*InstrumentedFunction)(void*), void*
 	MaterializeDeferredTransactions();
 	ASSERT(!GetNumDeferredTransactions());
     
-	FTransaction* NewTransaction = new FTransaction(this);
 	FCallNest* NewNest = new FCallNest(this);
 
-	// Transact requires a return from the lambda to commit the results
-	NewTransaction->SetIsScopedTransaction();
-
-	void* TransactStackStart = &NewTransaction;
+	void* TransactStackStart = &NewNest;
 
 	ETransactionResult Result = ETransactionResult::Committed; // Initialize to something to make the compiler happy.
 
-    if (!RawCurrentTransaction())
+    if (!CurrentTransaction)
     {
         ASSERT(Status == EContextStatus::Idle);
 
@@ -370,9 +393,12 @@ ETransactionResult FContext::Transact(void (*InstrumentedFunction)(void*), void*
 		ASSERT(Stack.High > Stack.Low);
 
 		ASSERT(Stack.Contains(TransactStackStart));
-		NewTransaction->SetStackRange({Stack.Low, &TransactStackStart});
 
-		PushTransaction(NewTransaction);
+		FTransaction* NewTransaction = PushTransaction(
+			/* Closed */ true,
+			/* bIsScoped */ true,
+			/* StackRange */ {Stack.Low, &TransactStackStart});
+
 		PushCallNest(NewNest);
 
 		bool bTriedToRunOnce = false;
@@ -380,9 +406,9 @@ ETransactionResult FContext::Transact(void (*InstrumentedFunction)(void*), void*
         for (;;)
         {
             Status = EContextStatus::OnTrack;
-            ASSERT(RawCurrentTransaction()->IsFresh());
+            ASSERT(CurrentTransaction->IsFresh());
 			CurrentNest->Try([&] () { InstrumentedFunction(Arg); });
-			ASSERT(RawCurrentTransaction() == NewTransaction); // The transaction lambda should have unwound any nested transactions.
+			ASSERT(CurrentTransaction == NewTransaction); // The transaction lambda should have unwound any nested transactions.
             ASSERT(Status != EContextStatus::Idle);
 
             if (Status == EContextStatus::OnTrack)
@@ -395,15 +421,16 @@ ETransactionResult FContext::Transact(void (*InstrumentedFunction)(void*), void*
 				{
 					// We skip trying to commit this time, and instead re-run the transaction.
 					Status = EContextStatus::AbortedByFailedLockAcquisition;
-					RawCurrentTransaction()->AbortWithoutThrowing();
+					CurrentTransaction->AbortWithoutThrowing();
 					ClearTransactionStatus();
 
 					// We've tried to run at least once if we get here!
+					CurrentTransaction->Reset();
 					bTriedToRunOnce = true;
 					continue;
 				}
 
-				if (AttemptToCommitTransaction(RawCurrentTransaction()))
+				if (AttemptToCommitTransaction(CurrentTransaction))
 				{
 					Result = ETransactionResult::Committed;
 					break;
@@ -436,14 +463,14 @@ ETransactionResult FContext::Transact(void (*InstrumentedFunction)(void*), void*
             ASSERT(Status == EContextStatus::AbortedByFailedLockAcquisition);
         }
 
-		NewTransaction->SetIsDone();
+		NewTransaction->SetState(FTransaction::EState::Done);
 
 		PopCallNest();
 		PopTransaction();
 		ClearTransactionStatus();
 
 		ASSERT(CurrentNest == nullptr);
-		ASSERT(RawCurrentTransaction() == nullptr);
+		ASSERT(CurrentTransaction == nullptr);
 
         Reset();
 	}
@@ -455,9 +482,12 @@ ETransactionResult FContext::Transact(void (*InstrumentedFunction)(void*), void*
 		ASSERT(CurrentThreadId == FPlatformTLS::GetCurrentThreadId());
 
 		ASSERT(Stack.Contains(TransactStackStart));
-		NewTransaction->SetStackRange({Stack.Low, &TransactStackStart});
 
-		PushTransaction(NewTransaction);
+		FTransaction* NewTransaction = PushTransaction(
+			/* Closed */ true,
+			/* bIsScoped */ true,
+			/* StackRange */ {Stack.Low, &TransactStackStart});
+
 		PushCallNest(NewNest);
 
 		bool bTriedToRunOnce = false;
@@ -465,7 +495,7 @@ ETransactionResult FContext::Transact(void (*InstrumentedFunction)(void*), void*
 		for (;;)
 		{
 			CurrentNest->Try([&]() { InstrumentedFunction(Arg); });
-			ASSERT(RawCurrentTransaction() == NewTransaction);
+			ASSERT(CurrentTransaction == NewTransaction);
 
 			if (Status == EContextStatus::OnTrack)
 			{
@@ -477,8 +507,8 @@ ETransactionResult FContext::Transact(void (*InstrumentedFunction)(void*), void*
 					ClearTransactionStatus();
 
 					// We've tried to run at least once if we get here!
+					CurrentTransaction->Reset();
 					bTriedToRunOnce = true;
-
 					continue;
 				}
 			}
@@ -491,7 +521,7 @@ ETransactionResult FContext::Transact(void (*InstrumentedFunction)(void*), void*
 		PopTransaction();
 
 		ASSERT(CurrentNest != nullptr);
-		ASSERT(RawCurrentTransaction() != nullptr);
+		ASSERT(CurrentTransaction != nullptr);
 
 		// Cascading aborts should cause all transactions to abort!
 		switch (Result)
@@ -499,7 +529,7 @@ ETransactionResult FContext::Transact(void (*InstrumentedFunction)(void*), void*
 		default:
 			break;
 		case ETransactionResult::AbortedByCascade:
-			RawCurrentTransaction()->AbortAndThrow();
+			CurrentTransaction->AbortAndThrow();
 			break;
 		}
 
diff --git a/Engine/Source/Runtime/Core/Private/AutoRTFM/Context.h b/Engine/Source/Runtime/Core/Private/AutoRTFM/Context.h
--- a/Engine/Source/Runtime/Core/Private/AutoRTFM/Context.h
+++ b/Engine/Source/Runtime/Core/Private/AutoRTFM/Context.h
@@ -51,17 +51,13 @@ public:
 	{
 		return NumDeferredTransactions;
 	}
-	inline FTransaction* RawCurrentTransaction()
-	{
-		return CurrentTransaction;
-	}
 	inline FTransaction* GetCurrentTransaction()
 	{
 		if (UNLIKELY(GetNumDeferredTransactions()))
 		{
 			MaterializeDeferredTransactions();
 		}
-		return RawCurrentTransaction();
+		return CurrentTransaction;
 	}
 	inline FCallNest* GetCurrentNest() const { return CurrentNest; }
 	inline EContextStatus GetStatus() const { return CurrentThreadId == FPlatformTLS::GetCurrentThreadId() ? Status : EContextStatus::Idle; }
@@ -89,7 +85,7 @@ private:
 	void PushCallNest(FCallNest* NewCallNest);
 	void PopCallNest();
 
-	void PushTransaction(FTransaction* NewTransaction);
+	FTransaction* PushTransaction(bool bClosed, bool bIsScoped, FStackRange StackRange);
 	void PopTransaction();
 
 	ETransactionResult ResolveNestedTransaction(FTransaction* NewTransaction);
diff --git a/Engine/Source/Runtime/Core/Private/AutoRTFM/Transaction.cpp b/Engine/Source/Runtime/Core/Private/AutoRTFM/Transaction.cpp
--- a/Engine/Source/Runtime/Core/Private/AutoRTFM/Transaction.cpp
+++ b/Engine/Source/Runtime/Core/Private/AutoRTFM/Transaction.cpp
@@ -2,9 +2,20 @@
 
 #if (defined(__AUTORTFM) && __AUTORTFM)
 #include "Transaction.h"
-#include "TransactionInlines.h"
+
 #include "CallNestInlines.h"
 #include "GlobalData.h"
+#include "HAL/Platform.h"
+#include "Hash/xxhash.h"
+#include "Logging/LogMacros.h"
+#include "Logging/LogVerbosity.h"
+#include "Misc/AssertionMacros.h"
+#include "TransactionInlines.h"
+#include "Utils.h"
+
+#include <cstddef>
+#include <cstdint>
+#include <type_traits>
 
 namespace AutoRTFM
 {
@@ -14,6 +25,42 @@ FTransaction::FTransaction(FContext* Context)
 {
 }
 
+FTransaction::~FTransaction()
+{
+	ASSERT(IsDone());
+}
+
+void FTransaction::Initialize(EState State, FTransaction* Parent_, bool bIsStackScoped_, FStackRange StackRange_)
+{
+	ASSERT(CurrentState == EState::Uninitialized);
+	ASSERT(State == EState::ClosedActive || State == EState::OpenActive);
+
+	CurrentState = State;
+	Parent = Parent_;
+	bIsStackScoped = bIsStackScoped_;
+	StackRange = StackRange_;
+
+	// For stats, record the nested depth of the transaction.
+	if (Parent)
+	{
+		StatDepth = Parent->StatDepth + 1;
+	}
+
+	Stats.Collect<EStatsKind::AverageTransactionDepth>(StatDepth);
+	Stats.Collect<EStatsKind::MaximumTransactionDepth>(StatDepth);
+}
+
+void FTransaction::Reset()
+{
+	// Complete the transaction, which will reset all members.
+	SetState(EState::Done);
+
+	// Reset to the initial state.
+	CurrentState = EState::ClosedActive;
+
+	ASSERT(IsFresh());
+}
+
 bool FTransaction::IsFresh() const
 {
     return HitSet.IsEmpty()
@@ -21,7 +68,7 @@ bool FTransaction::IsFresh() const
         && WriteLog.IsEmpty()
         && CommitTasks.IsEmpty()
         && AbortTasks.IsEmpty()
-        && !bIsDone;
+        && !IsDone();
 }
 
 void FTransaction::AbortWithoutThrowing()
@@ -29,11 +76,14 @@ void FTransaction::AbortWithoutThrowing()
 	UE_LOG(LogAutoRTFM, Verbose, TEXT("Aborting '%hs'!"), GetContextStatusName(Context->GetStatus()));
 
     ASSERT(Context->IsAborting());
-    ASSERT(Context->GetCurrentTransaction() == this);
 
     Stats.Collect<EStatsKind::Abort>();
     CollectStats();
 
+	// Ensure that we're closed before applying the undo, as this will ensure the open-memory
+	// validation is performed before the memory is reset.
+	SetState(EState::ClosedActive);
+
     // Call the destructors of all the OnCommit functors before undoing the transactional memory and
     // calling the OnAbort callbacks. This is important as the callback functions may have captured
     // variables that are depending on the allocated memory. 
@@ -59,7 +109,8 @@ void FTransaction::AbortWithoutThrowing()
 		ASSERT(Context->IsAborting());
     }
 
-    Reset();
+	// Transaction is now done.
+	SetState(EState::Done);
 }
 
 void FTransaction::AbortAndThrow()
@@ -86,13 +137,16 @@ bool FTransaction::AttemptToCommit()
     {
         bResult = AttemptToCommitOuterNest();
     }
-    Reset();
+
+	SetState(EState::Done);
+
     return bResult;
 }
 
 void FTransaction::Undo()
 {
 	UE_LOG(LogAutoRTFM, Verbose, TEXT("Undoing a transaction..."));
+	ASSERT(IsClosed());
 
 	int VerboseCounter = 0;
 	int Num = WriteLog.Num();
@@ -191,14 +245,179 @@ bool FTransaction::AttemptToCommitOuterNest()
     return true;
 }
 
-void FTransaction::Reset()
+FTransaction::EState FTransaction::SetState(EState NewState)
 {
-    CommitTasks.Reset();
-    AbortTasks.Reset();
-	HitSet.Reset();
-    NewMemoryTracker.Reset();
-	WriteLog.Reset();
-	WriteLogBumpAllocator.Reset();
+	if (NewState == State())
+	{
+		return NewState;
+	}
+
+	// Once Done, the transaction cannot change state without a call to Reset().
+	ASSERT(!IsDone());
+
+	switch (CurrentState)
+	{
+		// OpenActive -> OpenInactive, ClosedActive or Done
+		case EState::OpenActive:
+			ASSERT(NewState == EState::OpenInactive || NewState == EState::ClosedActive || NewState == EState::Done);
+			if (CurrentMemoryValidationLevel != EMemoryValidationLevel::Disabled)
+			{
+				ValidateWriteHash();
+				RecordedWriteHash = 0;
+			}
+			break;
+
+		// ClosedActive -> ClosedInactive, OpenActive or Done
+		case EState::ClosedActive:
+			ASSERT(NewState == EState::ClosedInactive || NewState == EState::OpenActive || NewState == EState::Done);
+			break;
+
+		// ClosedActive -> OpenActive
+		case EState::OpenInactive:
+			ASSERT(NewState == EState::OpenActive);
+			break;
+
+		// ClosedInactive -> ClosedActive
+		case EState::ClosedInactive:
+			ASSERT(NewState == EState::ClosedActive);
+			break;
+
+		default:
+			UE_LOG(LogAutoRTFM, Fatal, TEXT("Invalid state"));
+			break;
+	}
+
+	switch (NewState) {
+		// OpenInactive, ClosedActive or Done -> OpenActive
+		case EState::OpenActive:
+			CurrentMemoryValidationLevel = PendingMemoryValidationLevel;
+			if (CurrentMemoryValidationLevel != EMemoryValidationLevel::Disabled)
+			{
+				ASSERT(RecordedWriteHash == 0);
+				RecordWriteHash();
+			}
+			break;
+
+		// OpenActive, ClosedActive -> Done
+		case EState::Done:
+			ASSERT(IsActive());
+			CommitTasks.Reset();
+			AbortTasks.Reset();
+			HitSet.Reset();
+			NewMemoryTracker.Reset();
+			WriteLog.Reset();
+			WriteLogBumpAllocator.Reset();
+			break;
+
+		default:
+			break;
+	}
+
+	EState OldState = CurrentState;
+	CurrentState = NewState;
+	return OldState;
+}
+
+void FTransaction::SetMemoryValidationLevel(EMemoryValidationLevel Level)
+{
+	PendingMemoryValidationLevel = Level;
+}
+
+void FTransaction::SetActive()
+{
+	switch (CurrentState)
+	{
+		case EState::OpenActive:
+		case EState::ClosedActive:
+			break;
+		case EState::OpenInactive:
+			SetState(EState::OpenActive);
+			break;
+		case EState::ClosedInactive:
+			SetState(EState::ClosedActive);
+			break;
+		default:
+			UE_LOG(LogAutoRTFM, Fatal, TEXT("Invalid state"));
+	}
+}
+
+void FTransaction::SetInactive()
+{
+	switch (CurrentState)
+	{
+		case EState::OpenInactive:
+		case EState::ClosedInactive:
+			break;
+		case EState::OpenActive:
+			SetState(EState::OpenInactive);
+			break;
+		case EState::ClosedActive:
+			SetState(EState::ClosedInactive);
+			break;
+		default:
+			UE_LOG(LogAutoRTFM, Fatal, TEXT("Invalid state"));
+	}
+}
+
+void FTransaction::RecordWriteHash()
+{
+	RecordedWriteHash = CalculateNestedWriteHash();
+}
+
+void FTransaction::ValidateWriteHash() const
+{
+	TWriteHash OldHash = RecordedWriteHash;
+	TWriteHash NewHash = CalculateNestedWriteHash();
+	if (LIKELY(OldHash == NewHash))
+	{
+		return;
+	}
+
+	static constexpr TCHAR Message[] =
+		TEXT("Memory modified in a transaction was also modified in an call to AutoRTFM::Open(). ")
+		TEXT("This may lead to memory corruption if the transaction is aborted.");
+	if (CurrentMemoryValidationLevel == EMemoryValidationLevel::EnabledAsWarning)
+	{
+		UE_LOG(LogAutoRTFM, Warning, Message);
+	}
+	else
+	{
+		UE_LOG(LogAutoRTFM, Error, Message);
+	}
+}
+
+FTransaction::TWriteHash FTransaction::CalculateNestedWriteHash() const
+{
+	FXxHash64Builder Hasher;
+	if (nullptr != Parent)
+	{
+		TWriteHash ParentHash = Parent->CalculateNestedWriteHash();
+		Hasher.Update(&ParentHash, sizeof(ParentHash));
+	}
+
+	using FSmallValue = uint64;
+	FSmallValue SmallValue = 0;
+	for(auto Iter = WriteLog.begin(); Iter != WriteLog.end(); ++Iter)
+	{
+		const FWriteLogEntry& Entry = *Iter;
+		const void* Original = Entry.GetOriginal();
+		size_t Size = Entry.GetSize();
+		if (Size <= sizeof(FSmallValue))
+		{
+			FSmallValue Value = 0;
+			memcpy(&Value, Original, Size);
+			SmallValue = SmallValue * 31 + Value;
+		}
+		else
+		{
+			Hasher.Update(Original, Size);
+		}
+	}
+	Hasher.Update(&SmallValue, sizeof(SmallValue));
+
+	auto Hash = Hasher.Finalize().Hash;
+	static_assert(std::is_same_v<decltype(Hash), TWriteHash>, "TWriteHash type mismatch");
+	return Hash;
 }
 
 } // namespace AutoRTFM
diff --git a/Engine/Source/Runtime/Core/Private/AutoRTFM/Transaction.h b/Engine/Source/Runtime/Core/Private/AutoRTFM/Transaction.h
--- a/Engine/Source/Runtime/Core/Private/AutoRTFM/Transaction.h
+++ b/Engine/Source/Runtime/Core/Private/AutoRTFM/Transaction.h
@@ -2,14 +2,15 @@
 
 #pragma once
 
+#include "AutoRTFM/AutoRTFM.h"
+#include "HAL/Platform.h"
+#include "Templates/Function.h"
+#include "CoreTypes.h"
 #include "HitSet.h"
 #include "IntervalTree.h"
-#include "LongJump.h"
 #include "StackRange.h"
 #include "Stats.h"
-#include "TaggedPtr.h"
 #include "TaskArray.h"
-#include "Templates/Function.h"
 #include "WriteLog.h"
 #include "WriteLogBumpAllocator.h"
 
@@ -20,24 +21,57 @@ class FContext;
 class FTransaction final
 {
 public:
+    // State flow diagram for a transaction:
+    //                       ┌──────────────────────────────────────┐
+    //                       │             Uninitialized            │
+    //                       └──────────────────────────────────────┘
+    //                               │                       │
+    //                               ▼                       ▼
+    // ┌────────────────┐    ┌────────────────┐    ┌────────────────┐    ┌────────────────┐
+    // │  OpenInactive  │ ←→ │   OpenActive   │ ←→ │  ClosedActive  │ ←→ │ ClosedInactive │
+    // └────────────────┘    └────────────────┘    └────────────────┘    └────────────────┘
+    //                               │                       │
+    //                               ▼                       ▼
+    //                       ┌──────────────────────────────────────┐
+    //                       │                 Done                 │
+    //                       └──────────────────────────────────────┘
+    enum class EState
+    {
+        // The initial state for the transaction.
+        // Can only transition to OpenActive or ClosedActive.
+        Uninitialized,
+        // The transaction is open (not recording writes) and is the current transcation.
+        // Can only transition to OpenInactive, ClosedActive or Done.
+        OpenActive,
+        // The transaction is closed (recording writes) and is the current transcation.
+        // Can only transition to ClosedInactive, OpenActive or Done.
+        ClosedActive,
+        // The transaction is open and the current transaction is a descendant.
+        // Can only transition to OpenActive.
+        OpenInactive,
+        // The transaction is closed and the current transaction is a descendant.
+        // Can only transition to ClosedActive.
+        ClosedInactive,
+        // The transaction is committed or aborted.
+        // Once in this state the transaction must be reset with Reset() or destructed.
+        Done,
+    };
+
+    // Constructor.
+    // The transaction is constructed in an initial Closed state.
     FTransaction(FContext* Context);
-    
-    bool IsNested() const { return !!Parent; }
-    FTransaction* GetParent() const { return Parent; }
 
-	void SetParent(FTransaction* NewParent)
-    {
-        Parent = NewParent;
+    // Destructor.
+    // The transaction must we in a Done state.
+    ~FTransaction();
 
-        // For stats, record the nested depth of the transaction.
-		if (NewParent)
-		{
-			StatDepth = NewParent->StatDepth + 1;
-		}
+    void Initialize(EState State, FTransaction* Parent, bool bIsScoped, FStackRange StackRange);
 
-		Stats.Collect<EStatsKind::AverageTransactionDepth>(StatDepth);
-		Stats.Collect<EStatsKind::MaximumTransactionDepth>(StatDepth);
-    }
+    // Clears the tracked transaction state and resets back to the default Uninitialized state.
+    void Reset();
+
+    bool IsNested() const { return !!Parent; }
+    FTransaction* GetParent() const { return Parent; }
 
     // This should just use type displays or ranges. Maybe ranges could even work out great.
     bool IsNestedWithin(const FTransaction* Other) const
@@ -55,12 +89,7 @@ public:
         }
     }
 
-    bool IsFresh() const;
-    bool IsDone() const { return bIsDone; }
-    void SetIsDone() { bIsDone = true; }
-    
 	inline bool IsScopedTransaction() const { return bIsStackScoped; }
-	inline void SetIsScopedTransaction() { bIsStackScoped = true; }
 
     void DeferUntilCommit(TFunction<void()>&&);
     void DeferUntilAbort(TFunction<void()>&&);
@@ -80,27 +109,60 @@ public:
     void DidAllocate(void* LogicalAddress, size_t Size);
     void DidFree(void* LogicalAddress);
 
+    // Changes the state of the transaction to NewState and returns the previous state.
+    EState SetState(EState NewState);
+    void SetActive();
+    void SetInactive();
+
+    // Only takes effect at the next state transition to OpenActive.
+    void SetMemoryValidationLevel(EMemoryValidationLevel Level);
+
+    // State querying
+    EState State() const { return CurrentState; }
+    bool IsFresh() const;
+    bool IsOpenActive() const { return CurrentState == EState::OpenActive; }
+    bool IsClosedActive() const { return CurrentState == EState::ClosedActive; }
+    bool IsOpenInactive() const { return CurrentState == EState::OpenInactive; }
+    bool IsClosedInactive() const { return CurrentState == EState::ClosedInactive; }
+    bool IsOpen() const { return IsOpenActive() || IsOpenInactive(); }
+    bool IsClosed() const { return IsClosedActive() || IsClosedInactive(); }
+    bool IsActive() const { return IsOpenActive() || IsClosedActive(); }
+    bool IsInactive() const { return IsOpenInactive() || IsClosedInactive(); }
+    bool IsDone() const { return CurrentState == EState::Done; }
+
     // The stack range represents all stack memory inside the transaction scope
-    inline void SetStackRange(FStackRange Range) { StackRange = Range;} 
     inline FStackRange GetStackRange() const { return StackRange; }
 
     // Returns true if the LogicalAddress is within the stack of the transaction.
     inline bool IsOnStack(const void* LogicalAddress) const;
 
 private:
+    using TWriteHash = uint64;
+
     void Undo();
 
     void CommitNested();
     bool AttemptToCommitOuterNest();
 
-    void Reset(); // Frees memory and sets us up for possible retry.
+    // Calculates the hash of all the memory written during the transaction and stores it to
+    // RecordedWriteHash.
+    void RecordWriteHash();
+
+    // Compares the hash of all the memory that will be written on abort against the
+    // RecordedWriteHash, raising a warning or error if the hashes are not equal.
+    // See CalculateNestedWriteHash().
+    void ValidateWriteHash() const;
+
+    // Returns a hash of all the memory written during this transaction and all parent transactions.
+    // Can be used to ensure that no transaction written memory is modified during an open.
+    TWriteHash CalculateNestedWriteHash() const;
 
     void CollectStats() const;
-    
+
     bool ShouldRecordWrite(void* LogicalAddress) const;
 
-    FContext* Context;
-    
+    FContext* Context{nullptr};
+
     // Are we nested? Then this is the parent.
     FTransaction* Parent{nullptr};
 
@@ -108,15 +170,17 @@ private:
     TTaskArray<TFunction<void()>> CommitTasks;
     TTaskArray<TFunction<void()>> AbortTasks;
 
-    bool bIsDone{false};
-	bool bIsStackScoped{false};
-
     FHitSet HitSet;
     FIntervalTree NewMemoryTracker;
     FWriteLog WriteLog;
     FWriteLogBumpAllocator WriteLogBumpAllocator;
     TStatStorage<uint64_t> StatDepth = 1;
+    TWriteHash RecordedWriteHash = 0;
+    EMemoryValidationLevel CurrentMemoryValidationLevel = EMemoryValidationLevel::Disabled;
+    EMemoryValidationLevel PendingMemoryValidationLevel = EMemoryValidationLevel::Disabled;
+    EState CurrentState{EState::Uninitialized};
     FStackRange StackRange;
+    bool bIsStackScoped = false;
 };
 
 } // namespace AutoRTFM
diff --git a/Engine/Source/Runtime/Core/Private/Hash/Blake3.cpp b/Engine/Source/Runtime/Core/Private/Hash/Blake3.cpp
--- a/Engine/Source/Runtime/Core/Private/Hash/Blake3.cpp
+++ b/Engine/Source/Runtime/Core/Private/Hash/Blake3.cpp
@@ -22,7 +22,7 @@ void FBlake3::Reset()
 	// Call the blake3 function in the open (non-transactionally), and manually
 	// record the update to the Hasher structure.
 	bool bIsClosed = AutoRTFM::IsClosed();
-	UE_AUTORTFM_OPEN
+	UE_AUTORTFM_OPEN_NO_VALIDATION
 	{
 		if (bIsClosed)
 		{
@@ -44,7 +44,7 @@ void FBlake3::Update(const void* Data, uint64 Size)
 	// blake3 is a precompiled external library without AutoRTFM support.
 	// Call the blake3 function in the open (non-transactionally), and manually
 	// record the update to the Hasher structure.
-	UE_AUTORTFM_OPEN
+	UE_AUTORTFM_OPEN_NO_VALIDATION
 	{
 		if (bIsClosed)
 		{
diff --git a/Engine/Source/Runtime/Core/Public/AutoRTFM/AutoRTFM.h b/Engine/Source/Runtime/Core/Public/AutoRTFM/AutoRTFM.h
--- a/Engine/Source/Runtime/Core/Public/AutoRTFM/AutoRTFM.h
+++ b/Engine/Source/Runtime/Core/Public/AutoRTFM/AutoRTFM.h
@@ -64,6 +64,7 @@
 #define UE_AUTORTFM_MOVE std::move
 #define PRAGMA_DISABLE_UNREACHABLE_CODE_WARNINGS
 #define PRAGMA_RESTORE_UNREACHABLE_CODE_WARNINGS
+#define LIKELY(EXPR) __builtin_expect(!!(EXPR), 1)
 #else
 #include <HAL/Platform.h>
 #include <HAL/PlatformMemory.h>
@@ -131,6 +132,25 @@ typedef enum
 	autortfm_status_aborted_by_cascade
 } autortfm_status;
 
+// An enumerator of transactional memory validation levels.
+// Memory validation is used to detect modification by open-code to memory that was written by a
+// transaction. In this situation, aborting the transaction can corrupt memory as the undo will
+// overwrite the writes made in the open-code.
+typedef enum
+{
+	// Use the default memory validation level.
+	autortfm_memory_validation_level_default,
+
+	// Disable memory validation.
+	autortfm_memory_validation_level_disabled,
+
+	// Enable memory validation. Memory validation failures are treated as warnings.
+	autortfm_memory_validation_level_warning,
+
+	// Enable memory validation. Memory validation are treated as fatal.
+	autortfm_memory_validation_level_enabled,
+} autortfm_memory_validation_level;
+
 #if UE_AUTORTFM_ENABLED
 // Note: There is no implementation of this function.
 // The AutoRTFM compiler will replace all calls to this function with a constant boolean value.
@@ -252,6 +272,12 @@ UE_AUTORTFM_API void autortfm_open(void (*work)(void* arg), void* arg);
 static UE_AUTORTFM_FORCEINLINE void autortfm_open(void (*work)(void* arg), void* arg) { work(arg); }
 #endif
 
+#if UE_AUTORTFM_ENABLED
+UE_AUTORTFM_API void autortfm_open_explicit_validation(autortfm_memory_validation_level, void (*work)(void* arg), void* arg);
+#else
+static UE_AUTORTFM_FORCEINLINE void autortfm_open_explicit_validation(autortfm_memory_validation_level, void (*work)(void* arg), void* arg) { work(arg); }
+#endif
+
 #if UE_AUTORTFM_ENABLED
 [[nodiscard]] UE_AUTORTFM_API autortfm_status autortfm_close(void (*UninstrumentedWork)(void*), void (*InstrumentedWork)(void*), void* Arg);
 #else
@@ -435,6 +461,25 @@ enum class EContextStatus : uint8_t
 	AbortedByCascade = autortfm_status_aborted_by_cascade
 };
 
+// An enumerator of transactional memory validation levels.
+// Memory validation is used to detect modification by open-code to memory that was written by a
+// transaction. In this situation, aborting the transaction can corrupt memory as the undo will
+// overwrite the writes made in the open-code.
+enum class EMemoryValidationLevel : uint8_t
+{
+	// The default memory validation level.
+	Default = autortfm_memory_validation_level_default,
+
+	// Disable memory validation.
+	Disabled = autortfm_memory_validation_level_disabled,
+
+	// Enable memory validation. Memory validation failures are treated as warnings.
+	EnabledAsWarning = autortfm_memory_validation_level_warning,
+
+	// Enable memory validation. Memory validation are treated as fatal.
+	Enabled = autortfm_memory_validation_level_enabled,
+};
+
 #if UE_AUTORTFM
 namespace ForTheRuntime
 {
@@ -674,7 +719,12 @@ static constexpr bool IsSafeToReturnFromOpen = HasAssignFromOpenToClosedTrait<T>
 // a transaction or not. Returns the value returned by Functor.
 // TReturn must be void or a type that can be safely copied from the open to a closed transaction.
 // TAssignFromOpenToClosed must have a specialization for the type that is being returned.
-template<typename TFunctor, typename TReturn = decltype(std::declval<TFunctor>()())>
+template
+<
+	EMemoryValidationLevel VALIDATION_LEVEL = EMemoryValidationLevel::Default,
+	typename TFunctor = void,
+	typename TReturn = decltype(std::declval<TFunctor>()())
+> 
 static UE_AUTORTFM_FORCEINLINE TReturn Open(const TFunctor& Functor)
 {
 	static_assert(IsSafeToReturnFromOpen<TReturn>,
@@ -697,7 +747,18 @@ static UE_AUTORTFM_FORCEINLINE TReturn Open(const TFunctor& Functor)
 					UE_AUTORTFM_CALLSITE_FORCEINLINE Fn();
 				}
 			};
-			autortfm_open(&FCallHelper::Call, const_cast<void*>(static_cast<const void*>(&Functor)));
+
+			if constexpr (VALIDATION_LEVEL == EMemoryValidationLevel::Default)
+			{
+				autortfm_open(&FCallHelper::Call, const_cast<void*>(static_cast<const void*>(&Functor)));
+			}
+			else
+			{
+				autortfm_open_explicit_validation(
+					static_cast<autortfm_memory_validation_level>(VALIDATION_LEVEL),
+					&FCallHelper::Call,
+					const_cast<void*>(static_cast<const void*>(&Functor)));
+			}
 		}
 		else
 		{
@@ -713,7 +774,17 @@ static UE_AUTORTFM_FORCEINLINE TReturn Open(const TFunctor& Functor)
 				TReturn ReturnValue{};
 			};
 			FCallHelper Helper{Functor};
-			autortfm_open(&FCallHelper::Call, reinterpret_cast<void*>(&Helper));
+			if constexpr (VALIDATION_LEVEL == EMemoryValidationLevel::Default)
+			{
+				autortfm_open(&FCallHelper::Call, reinterpret_cast<void*>(&Helper));
+			}
+			else
+			{
+				autortfm_open_explicit_validation(
+					static_cast<autortfm_memory_validation_level>(VALIDATION_LEVEL),
+					&FCallHelper::Call,
+					reinterpret_cast<void*>(&Helper));
+			}
 			return Helper.ReturnValue;
 		}
 	}
@@ -914,6 +985,12 @@ namespace ForTheRuntime
 	// Returns true if we should also retry nested transactions.
 	UE_AUTORTFM_API bool ShouldRetryNestedTransactionsToo();
 
+	// Returns the memory validation level currently enabled.
+	UE_AUTORTFM_API EMemoryValidationLevel GetMemoryValidationLevel();
+
+	// Sets the memory validation level. See IsWriteValidationEnabled().
+	UE_AUTORTFM_API void SetMemoryValidationLevel(EMemoryValidationLevel Level);
+
 	// Manually create a new transaction from open code and push it as a transaction nest.
 	// Can only be called within an already active parent transaction (EG. this cannot start
 	// a transaction nest itself).
@@ -1031,6 +1108,14 @@ namespace AutoRTFM::Private
 			AutoRTFM::Open(UE_AUTORTFM_MOVE(F));
 		}
 	};
+	struct FOpenNoMemoryValidationHelper
+	{
+		template<typename FunctorType>
+		void operator+(FunctorType F)
+		{
+			AutoRTFM::Open<EMemoryValidationLevel::Disabled>(UE_AUTORTFM_MOVE(F));
+		}
+	};
 	struct FOnAbortHelper
 	{
 		template<typename FunctorType>
@@ -1066,6 +1151,7 @@ namespace AutoRTFM::Private
 #endif
 
 #define UE_AUTORTFM_OPEN_IMPL          ::AutoRTFM::Private::FOpenHelper{} + [&]()
+#define UE_AUTORTFM_OPEN_NO_VALIDATION_IMPL ::AutoRTFM::Private::FOpenNoMemoryValidationHelper{} + [&]()
 #define UE_AUTORTFM_ONABORT_IMPL(...)  ::AutoRTFM::Private::FOnAbortHelper{} + [__VA_ARGS__]()
 #define UE_AUTORTFM_ONCOMMIT_IMPL(...) ::AutoRTFM::Private::FOnCommitHelper{} + [__VA_ARGS__]()
 #define UE_AUTORTFM_TRANSACT_IMPL      ::AutoRTFM::Private::FTransactHelper{} + [&]()
@@ -1073,6 +1159,7 @@ namespace AutoRTFM::Private
 
 // Do nothing, these should be followed by blocks that should be either executed or not executed
 #define UE_AUTORTFM_OPEN_IMPL
+#define UE_AUTORTFM_OPEN_NO_VALIDATION_IMPL
 #define UE_AUTORTFM_ONABORT_IMPL(...) while (false)
 #define UE_AUTORTFM_ONCOMMIT_IMPL(...)
 #define UE_AUTORTFM_TRANSACT_IMPL
@@ -1083,6 +1170,8 @@ namespace AutoRTFM::Private
 #define UE_AUTORTFM_OPEN  UE_AUTORTFM_OPEN_IMPL
 #define UE_AUTORTFM_OPEN2 UE_AUTORTFM_OPEN_IMPL
 
+#define UE_AUTORTFM_OPEN_NO_VALIDATION UE_AUTORTFM_OPEN_NO_VALIDATION_IMPL
+
 // Runs a block of code if a transaction aborts.
 // In non-transactional code paths the block of code will not be executed at all.
 // The macro arguments are the capture specification for the lambda.
