Make AutoRTFM 12.48% faster by obsessing over the codegen on the most taken path in recording writes.

diff --git a/Engine/Source/Runtime/Core/Private/AutoRTFM/Context.h b/Engine/Source/Runtime/Core/Private/AutoRTFM/Context.h
--- a/Engine/Source/Runtime/Core/Private/AutoRTFM/Context.h
+++ b/Engine/Source/Runtime/Core/Private/AutoRTFM/Context.h
@@ -40,6 +40,7 @@ public:
     // Record that a write is about to occur at the given LogicalAddress of Size bytes.
     void RecordWrite(void* LogicalAddress, size_t Size);
     template<unsigned SIZE> void RecordWrite(void* LogicalAddress);
+	template<unsigned SIZE> void RecordWriteSlow(void* LogicalAddress);
 
     void DidAllocate(void* LogicalAddress, size_t Size);
     void DidFree(void* LogicalAddress);
@@ -47,18 +48,28 @@ public:
     // The rest of this is internalish.
     [[noreturn]] void AbortByLanguageAndThrow();
 
-	inline uint64 GetNumDeferredTransactions()
+	inline uint64 GetNumDeferredTransactions() const
 	{
 		return NumDeferredTransactions;
 	}
+	inline bool MustMaterializeDeferredTransactions() const
+	{
+		return (UNLIKELY(0 < GetNumDeferredTransactions()));
+	}
 	inline FTransaction* GetCurrentTransaction()
 	{
-		if (UNLIKELY(GetNumDeferredTransactions()))
+		if (MustMaterializeDeferredTransactions())
 		{
 			MaterializeDeferredTransactions();
 		}
 		return CurrentTransaction;
 	}
+	inline FTransaction* GetCurrentTransactionAssumingNoMaterialization()
+	{
+		ASSERT(!MustMaterializeDeferredTransactions());
+		return CurrentTransaction;
+	}
+
 	inline FCallNest* GetCurrentNest() const { return CurrentNest; }
 	inline EContextStatus GetStatus() const { return CurrentThreadId == FPlatformTLS::GetCurrentThreadId() ? Status : EContextStatus::Idle; }
 	[[noreturn]] void Throw();
diff --git a/Engine/Source/Runtime/Core/Private/AutoRTFM/ContextInlines.h b/Engine/Source/Runtime/Core/Private/AutoRTFM/ContextInlines.h
--- a/Engine/Source/Runtime/Core/Private/AutoRTFM/ContextInlines.h
+++ b/Engine/Source/Runtime/Core/Private/AutoRTFM/ContextInlines.h
@@ -20,7 +20,21 @@ AUTORTFM_NO_ASAN UE_AUTORTFM_FORCEINLINE void FContext::RecordWrite(void* Logica
 
 template<unsigned SIZE> AUTORTFM_NO_ASAN UE_AUTORTFM_FORCEINLINE void FContext::RecordWrite(void* LogicalAddress)
 {
-    GetCurrentTransaction()->RecordWrite<SIZE>(LogicalAddress);
+	if (MustMaterializeDeferredTransactions())
+	{
+		AUTORTFM_MUST_TAIL return RecordWriteSlow<SIZE>(LogicalAddress);
+	}
+	else
+	{
+		GetCurrentTransactionAssumingNoMaterialization()->RecordWrite<SIZE>(LogicalAddress);
+	}
+}
+
+template<unsigned SIZE> AUTORTFM_NO_ASAN FORCENOINLINE void FContext::RecordWriteSlow(void* LogicalAddress)
+{
+	// Going through get current transaction will do the materialization of deferred transactions, but it is slow
+	// so that is why we use this carve out.
+	GetCurrentTransaction()->RecordWrite<SIZE>(LogicalAddress);
 }
 
 UE_AUTORTFM_FORCEINLINE void FContext::DidAllocate(void* LogicalAddress, size_t Size)
diff --git a/Engine/Source/Runtime/Core/Private/AutoRTFM/FunctionMapInlines.h b/Engine/Source/Runtime/Core/Private/AutoRTFM/FunctionMapInlines.h
--- a/Engine/Source/Runtime/Core/Private/AutoRTFM/FunctionMapInlines.h
+++ b/Engine/Source/Runtime/Core/Private/AutoRTFM/FunctionMapInlines.h
@@ -32,10 +32,7 @@ inline void* FunctionMapLookup(void* OldFunction, const char* Where)
 
 	if (UNLIKELY(!Result))
 	{
-#ifdef __clang__
-		[[clang::musttail]]
-#endif
-		return FunctionMapReportError(OldFunction, Where);
+		AUTORTFM_MUST_TAIL return FunctionMapReportError(OldFunction, Where);
 	}
 
 	return Result;
diff --git a/Engine/Source/Runtime/Core/Private/AutoRTFM/HitSet.h b/Engine/Source/Runtime/Core/Private/AutoRTFM/HitSet.h
--- a/Engine/Source/Runtime/Core/Private/AutoRTFM/HitSet.h
+++ b/Engine/Source/Runtime/Core/Private/AutoRTFM/HitSet.h
@@ -45,10 +45,24 @@ namespace AutoRTFM
 			FMemory::Free(Payload);
 		}
 
+		enum class EInsertResult : uint8_t
+		{
+			Exists,
+			Inserted,
+			NotInserted
+		};
+
+		// Insert something in the HitSet, returning whether the key already
+		// exists, it was inserted, or not.
+		UE_AUTORTFM_FORCEINLINE EInsertResult FindOrTryInsertNoResize(Key K)
+		{
+			return InsertNoResize(K.GetPayload());
+		}
+
 		// Insert something in the HitSet, returning true if we found the key.
 		// If false is returned, we could have inserted or not - depending on
 		// whether the HitSet has reached its max capacity.
-		bool FindOrTryInsert(Key K)
+		UE_AUTORTFM_FORCEINLINE bool FindOrTryInsert(Key K)
 		{
 			while (true)
 			{
@@ -175,13 +189,6 @@ namespace AutoRTFM
 			return true;
 		}
 
-		enum class EInsertResult : uint8_t
-		{
-			Exists,
-			Inserted,
-			NotInserted
-		};
-
 		UE_AUTORTFM_FORCEINLINE uintptr_t FirstHashInRange(const uintptr_t Hashee) const
 		{
 			constexpr uintptr_t Fibonacci = 0x9E3779B97F4A7C15;
diff --git a/Engine/Source/Runtime/Core/Private/AutoRTFM/Transaction.h b/Engine/Source/Runtime/Core/Private/AutoRTFM/Transaction.h
--- a/Engine/Source/Runtime/Core/Private/AutoRTFM/Transaction.h
+++ b/Engine/Source/Runtime/Core/Private/AutoRTFM/Transaction.h
@@ -105,6 +105,8 @@ public:
     void RecordWrite(void* LogicalAddress, size_t Size);
     void RecordWriteMaxPageSized(void* LogicalAddress, size_t Size);
     template<unsigned SIZE> void RecordWrite(void* LogicalAddress);
+	template<unsigned SIZE> void RecordWriteInsertedSlow(void* LogicalAddress);
+	template<unsigned SIZE> void RecordWriteNotInsertedSlow(void* LogicalAddress);
 
     void DidAllocate(void* LogicalAddress, size_t Size);
     void DidFree(void* LogicalAddress);
diff --git a/Engine/Source/Runtime/Core/Private/AutoRTFM/TransactionInlines.h b/Engine/Source/Runtime/Core/Private/AutoRTFM/TransactionInlines.h
--- a/Engine/Source/Runtime/Core/Private/AutoRTFM/TransactionInlines.h
+++ b/Engine/Source/Runtime/Core/Private/AutoRTFM/TransactionInlines.h
@@ -135,14 +135,36 @@ template<unsigned SIZE> AUTORTFM_NO_ASAN UE_AUTORTFM_FORCEINLINE void FTransacti
     FMemoryLocation Key(LogicalAddress);
     Key.SetTopTag(static_cast<uint16_t>(SIZE));
 
-    if (HitSet.FindOrTryInsert(Key))
-    {
-        Stats.Collect<EStatsKind::HitSetHit>();
-        return;
-    }
+	switch (HitSet.FindOrTryInsertNoResize(Key))
+	{
+	case FHitSet::EInsertResult::Exists:
+		Stats.Collect<EStatsKind::HitSetHit>();
+		return;
+	case FHitSet::EInsertResult::Inserted:
+		AUTORTFM_MUST_TAIL return FTransaction::RecordWriteInsertedSlow<SIZE>(LogicalAddress);
+	case FHitSet::EInsertResult::NotInserted:
+		AUTORTFM_MUST_TAIL return FTransaction::RecordWriteNotInsertedSlow<SIZE>(LogicalAddress);
+	}
+}
+
+template<unsigned SIZE> AUTORTFM_NO_ASAN FORCENOINLINE void FTransaction::RecordWriteNotInsertedSlow(void* LogicalAddress)
+{
+	FMemoryLocation Key(LogicalAddress);
+	Key.SetTopTag(static_cast<uint16_t>(SIZE));
+
+	if (HitSet.FindOrTryInsert(Key))
+	{
+		Stats.Collect<EStatsKind::HitSetHit>();
+		return;
+	}
+
+	Stats.Collect<EStatsKind::HitSetMiss>();
 
-    Stats.Collect<EStatsKind::HitSetMiss>();
+	return RecordWriteInsertedSlow<SIZE>(LogicalAddress);
+}
 
+template<unsigned SIZE> AUTORTFM_NO_ASAN FORCENOINLINE void FTransaction::RecordWriteInsertedSlow(void* LogicalAddress)
+{
 	if (NewMemoryTracker.Contains(LogicalAddress, SIZE))
 	{
 		Stats.Collect<EStatsKind::NewMemoryTrackerHit>();
diff --git a/Engine/Source/Runtime/Core/Private/AutoRTFM/Utils.h b/Engine/Source/Runtime/Core/Private/AutoRTFM/Utils.h
--- a/Engine/Source/Runtime/Core/Private/AutoRTFM/Utils.h
+++ b/Engine/Source/Runtime/Core/Private/AutoRTFM/Utils.h
@@ -21,7 +21,7 @@ namespace AutoRTFM
     __assume(false);
 #else
     __builtin_unreachable();
-#endif // PLATFORM_WINDOWS
+#endif
 }
 
 FString GetFunctionDescription(void* FunctionPtr);
@@ -32,16 +32,54 @@ FString GetFunctionDescription(TReturnType (*FunctionPtr)(TParameterTypes...))
     return GetFunctionDescription(reinterpret_cast<void*>(FunctionPtr));
 }
 
+// We use a special calling convention for this call to minimize
+// the cost of calling this function (which is super unlikely to
+// be called because it is the asserting true codepath, EG. only
+// fatal process destroying explosions happen here).
+#if UE_AUTORTFM
+[[clang::preserve_most]]
+#endif
+[[noreturn]] FORCENOINLINE inline void DoAssert(void(*Logger)())
+{
+	Logger();
+
+#if defined(_MSC_VER) && !defined(__clang__)
+	__assume(false);
+#else
+	__builtin_unreachable();
+#endif
+}
+
 } // namespace AutoRTFM
 
-#define ASSERT(exp) UE_CLOG(UNLIKELY(!(exp)), LogAutoRTFM, Fatal, TEXT("AutoRTFM assert!"))
+// This is all a bit funky - but for good reason! We want `ASSERT`
+// to be as close to *zero* cost if the assert wouldn't trigger.
+// We will always pay the cost of the unlikely branch, but we need
+// to make the body of taking the assert happen in another function
+// so as not to affect codegen, register allocation, and stack
+// reservation. But we also want the `ASSERT` to give us accurate
+// `__FILE__` and `__LINE__` information for the line that it was
+// triggered on. So what we do is have a lambda at the line that
+// actually does the assert (but crucially gets the correct file
+// and line numbers), but then pass this to another function
+// `DoAssert` which has a special calling convention that makes
+// the caller as optimal as possible (at the expense of the callee).
+#define ASSERT(exp) if (UNLIKELY(!(exp))) { auto Lambda = []() { UE_LOG(LogAutoRTFM, Fatal, TEXT("AutoRTFM assert %s!"), TEXT(#exp)); }; AutoRTFM::DoAssert(Lambda); }
 
 #if defined(__has_feature)
-	#if __has_feature(address_sanitizer)
-		#define AUTORTFM_NO_ASAN [[clang::no_sanitize("address")]]
-	#endif
+#if __has_feature(address_sanitizer)
+#define AUTORTFM_NO_ASAN [[clang::no_sanitize("address")]]
+#endif
 #endif
 
 #if !defined(AUTORTFM_NO_ASAN)
-	#define AUTORTFM_NO_ASAN
+#define AUTORTFM_NO_ASAN
+#endif
+
+#if defined(__clang__)
+#define AUTORTFM_MUST_TAIL [[clang::musttail]]
+#endif
+
+#if !defined(AUTORTFM_MUST_TAIL)
+#define AUTORTFM_MUST_TAIL
 #endif
