Adding Const to GetCsvID function

#rnx
#tests Built and tested against Win64 Dedicated Server

diff --git a/Engine/Source/Runtime/Core/Private/ProfilingDebugging/CsvProfiler.cpp b/Engine/Source/Runtime/Core/Private/ProfilingDebugging/CsvProfiler.cpp
--- a/Engine/Source/Runtime/Core/Private/ProfilingDebugging/CsvProfiler.cpp
+++ b/Engine/Source/Runtime/Core/Private/ProfilingDebugging/CsvProfiler.cpp
@@ -4348,7 +4348,7 @@ int32 FCsvProfiler::GetNumFrameToCaptureOnEvent() const
 	return CaptureOnEventFrameCount;
 }
 
-const FGuid& FCsvProfiler::GetCsvID()
+const FGuid& FCsvProfiler::GetCsvID() const
 {
 	return CsvGUID;
 }
diff --git a/Engine/Source/Runtime/Core/Public/ProfilingDebugging/CsvProfiler.h b/Engine/Source/Runtime/Core/Public/ProfilingDebugging/CsvProfiler.h
--- a/Engine/Source/Runtime/Core/Public/ProfilingDebugging/CsvProfiler.h
+++ b/Engine/Source/Runtime/Core/Public/ProfilingDebugging/CsvProfiler.h
@@ -427,7 +427,7 @@ public:
 	CORE_API int32 GetCaptureFrameNumber() const;
 	CORE_API int32 GetCaptureFrameNumberRT() const;
 	CORE_API int32 GetNumFrameToCaptureOnEvent() const;
-	CORE_API const FGuid& GetCsvID();
+	CORE_API const FGuid& GetCsvID() const;
 
 	CORE_API bool EnableCategoryByString(const FString& CategoryName) const;
 	CORE_API void EnableCategoryByIndex(uint32 CategoryIndex, bool bEnable) const;
