Deprecated FText::FindText

FText::FindText is deprecated as it can no longer find untranslated text. Consider FText::FindTextInLiveTable_Advanced or FText::AsLocalizable_Advanced as a replacement, depending on your use-case.

diff --git a/Engine/Source/Runtime/Core/Public/Internationalization/Text.h b/Engine/Source/Runtime/Core/Public/Internationalization/Text.h
--- a/Engine/Source/Runtime/Core/Public/Internationalization/Text.h
+++ b/Engine/Source/Runtime/Core/Public/Internationalization/Text.h
@@ -529,10 +529,7 @@ public:
 	 */
 	static CORE_API bool FindTextInLiveTable_Advanced(const FTextKey& Namespace, const FTextKey& Key, FText& OutText, const FString* const SourceString = nullptr);
 
-	/**
-	 * Legacy alias for FText::FindTextInLiveTable_Advanced.
-	 * Will be deprecated in the future. DO NOT USE IN NEW CODE!
-	 */
+	UE_DEPRECATED(5.5, "FText::FindText is deprecated as it can no longer find untranslated text. Consider FText::FindTextInLiveTable_Advanced or FText::AsLocalizable_Advanced as a replacement, depending on your use-case.")
 	static FORCEINLINE bool FindText(const FTextKey& Namespace, const FTextKey& Key, FText& OutText, const FString* const SourceString = nullptr)
 	{
 		return FText::FindTextInLiveTable_Advanced(Namespace, Key, OutText, SourceString);
