[AutoRTFM] Fully remove [[clang::autortfm]] support and [[clang::noinline]] from RTFM macros

[[clang::autortfm]] has been a no-op attribute for some time, but the frontend still consumed it. This is now properly removed.

Remove [[clang::noinline]] from the UE_AUTORTFM_NOAUTORTFM, UE_AUTORTFM_ALWAYS_OPEN and AUTORTFM_CORE_API macros. This attribute is implicitly added by the other clang::noautortfm, clang::autortfm_always_open and clang::autortfm_core_api attributes.

diff --git a/Engine/Source/Runtime/Core/Private/AutoRTFM/API.cpp b/Engine/Source/Runtime/Core/Private/AutoRTFM/API.cpp
--- a/Engine/Source/Runtime/Core/Private/AutoRTFM/API.cpp
+++ b/Engine/Source/Runtime/Core/Private/AutoRTFM/API.cpp
@@ -248,7 +248,7 @@ namespace AutoRTFM
 // This is the implementation of the AutoRTFM.h API. Ideally, functions here should just delegate to some internal API.
 // For now, I have these functions also perform some error checking.
 
-#define AUTORTFM_CORE_API [[clang::autortfm_core_api, clang::noinline]]
+#define AUTORTFM_CORE_API [[clang::autortfm_core_api]]
 
 namespace AutoRTFM
 {
diff --git a/Engine/Source/Runtime/Core/Public/AutoRTFM/AutoRTFM.h b/Engine/Source/Runtime/Core/Public/AutoRTFM/AutoRTFM.h
--- a/Engine/Source/Runtime/Core/Public/AutoRTFM/AutoRTFM.h
+++ b/Engine/Source/Runtime/Core/Public/AutoRTFM/AutoRTFM.h
@@ -35,9 +35,8 @@
 #endif
 
 #if UE_AUTORTFM
-#define UE_AUTORTFM_NOAUTORTFM [[clang::noautortfm, clang::noinline]]
-// #jira SOL-6589: remove clang::noinline once this JIRA is fixed.
-#define UE_AUTORTFM_ALWAYS_OPEN [[clang::autortfm_always_open, clang::noinline]]
+#define UE_AUTORTFM_NOAUTORTFM [[clang::noautortfm]]
+#define UE_AUTORTFM_ALWAYS_OPEN [[clang::autortfm_always_open]]
 #define UE_AUTORTFM_CALLSITE_FORCEINLINE [[clang::always_inline]]
 #else
 #define UE_AUTORTFM_NOAUTORTFM
