[AutoRTFM] Handle RecordOpenWrite() with the memory validator.

Keep track of the number of write records when calculating the snapshot hash. Hash the same number of entries when validating.
This means new writes added with RecordOpenWrite() will not cause the hashes to differ.

diff --git a/Engine/Source/Programs/AutoRTFMTests/Private/OpenAPI.cpp b/Engine/Source/Programs/AutoRTFMTests/Private/OpenAPI.cpp
--- a/Engine/Source/Programs/AutoRTFMTests/Private/OpenAPI.cpp
+++ b/Engine/Source/Programs/AutoRTFMTests/Private/OpenAPI.cpp
@@ -67,9 +67,6 @@ TEST_CASE("OpenAPI.RecordDataClosed_Illegal", "[.]")
 
 TEST_CASE("OpenAPI.WriteDataInTheOpen")
 {
-	AUTORTFM_SCOPED_ENABLE_MEMORY_VALIDATION_AS_WARNING();
-	AutoRTFMTestUtils::FCaptureWarningContext WarningContext;
-
 	int value = 0;
 	auto transactResult = AutoRTFM::Transact([&]()
 	{
@@ -82,8 +79,6 @@ TEST_CASE("OpenAPI.WriteDataInTheOpen")
 
 	REQUIRE(transactResult == AutoRTFM::ETransactionResult::Committed);
 	REQUIRE(value == 1);
-	REQUIRE_THAT(WarningContext.GetWarnings(),
-		Catch::Matchers::VectorContains(FString(AutoRTFMTestUtils::kMemoryModifiedWarning)));
 }
 
 TEST_CASE("OpenAPI.AbortTransactionScopedFromOpen")
@@ -482,8 +477,6 @@ TEST_CASE("OpenAPI.StackWriteCommitInTheOpen1")
 
 TEST_CASE("OpenAPI.StackWriteCommitInTheOpen2")
 {
-	AUTORTFM_SCOPED_ENABLE_MEMORY_VALIDATION_AS_WARNING();
-	AutoRTFMTestUtils::FCaptureWarningContext WarningContext;
 	AutoRTFM::EContextStatus Status = AutoRTFM::EContextStatus::Idle;
 
 	int Value = 0;
@@ -511,9 +504,6 @@ TEST_CASE("OpenAPI.StackWriteCommitInTheOpen2")
 		REQUIRE(AutoRTFM::EContextStatus::OnTrack == Status);
 		REQUIRE(Value == 10);
 	});
-
-	REQUIRE_THAT(WarningContext.GetWarnings(),
-		Catch::Matchers::VectorContains(FString(AutoRTFMTestUtils::kMemoryModifiedWarning)));
 }
 
 TEST_CASE("OpenAPI.StackWriteAbortInTheOpen1")
diff --git a/Engine/Source/Runtime/Core/Private/AutoRTFM/Transaction.cpp b/Engine/Source/Runtime/Core/Private/AutoRTFM/Transaction.cpp
--- a/Engine/Source/Runtime/Core/Private/AutoRTFM/Transaction.cpp
+++ b/Engine/Source/Runtime/Core/Private/AutoRTFM/Transaction.cpp
@@ -127,6 +127,10 @@ bool FTransaction::AttemptToCommit()
     Stats.Collect<EStatsKind::Commit>();
     CollectStats();
 
+	// Ensure that we're closed before applying the commit, as this will ensure the open-memory
+	// validation is performed before the write log is cleared.
+	SetState(EState::ClosedActive);
+
     bool bResult;
     if (IsNested())
     {
@@ -264,6 +268,7 @@ FTransaction::EState FTransaction::SetState(EState NewState)
 			{
 				ValidateWriteHash();
 				RecordedWriteHash = 0;
+				NumWriteLogsHashed = 0;
 			}
 			break;
 
@@ -361,13 +366,14 @@ void FTransaction::SetInactive()
 
 void FTransaction::RecordWriteHash()
 {
-	RecordedWriteHash = CalculateNestedWriteHash();
+	NumWriteLogsHashed = WriteLog.Num();
+	RecordedWriteHash = CalculateNestedWriteHash(NumWriteLogsHashed);
 }
 
 void FTransaction::ValidateWriteHash() const
 {
 	TWriteHash OldHash = RecordedWriteHash;
-	TWriteHash NewHash = CalculateNestedWriteHash();
+	TWriteHash NewHash = CalculateNestedWriteHash(NumWriteLogsHashed);
 	if (LIKELY(OldHash == NewHash))
 	{
 		return;
@@ -386,19 +392,25 @@ void FTransaction::ValidateWriteHash() const
 	}
 }
 
-FTransaction::TWriteHash FTransaction::CalculateNestedWriteHash() const
+FTransaction::TWriteHash FTransaction::CalculateNestedWriteHash(size_t NumWriteEntries) const
 {
 	FXxHash64Builder Hasher;
 	if (nullptr != Parent)
 	{
-		TWriteHash ParentHash = Parent->CalculateNestedWriteHash();
+		TWriteHash ParentHash = Parent->CalculateNestedWriteHash(Parent->WriteLog.Num());
 		Hasher.Update(&ParentHash, sizeof(ParentHash));
 	}
 
+	size_t WriteIndex = 0;
 	using FSmallValue = uint64;
 	FSmallValue SmallValue = 0;
 	for(auto Iter = WriteLog.begin(); Iter != WriteLog.end(); ++Iter)
 	{
+		if (WriteIndex++ == NumWriteEntries)
+		{
+			break;
+		}
+
 		const FWriteLogEntry& Entry = *Iter;
 		const void* Original = Entry.GetOriginal();
 		size_t Size = Entry.GetSize();
diff --git a/Engine/Source/Runtime/Core/Private/AutoRTFM/Transaction.h b/Engine/Source/Runtime/Core/Private/AutoRTFM/Transaction.h
--- a/Engine/Source/Runtime/Core/Private/AutoRTFM/Transaction.h
+++ b/Engine/Source/Runtime/Core/Private/AutoRTFM/Transaction.h
@@ -157,7 +157,8 @@ private:
 
     // Returns a hash of all the memory written during this transaction and all parent transactions.
     // Can be used to ensure that no transaction written memory is modified during an open.
-    TWriteHash CalculateNestedWriteHash() const;
+    // NumWriteEntries is the number of write entries to hash before stopping.
+    TWriteHash CalculateNestedWriteHash(size_t NumWriteEntries) const;
 
     void CollectStats() const;
 
@@ -178,6 +179,7 @@ private:
     FWriteLogBumpAllocator WriteLogBumpAllocator;
     TStatStorage<uint64_t> StatDepth = 1;
     TWriteHash RecordedWriteHash = 0;
+    size_t NumWriteLogsHashed = 0;
     EMemoryValidationLevel CurrentMemoryValidationLevel = EMemoryValidationLevel::Disabled;
     EMemoryValidationLevel PendingMemoryValidationLevel = EMemoryValidationLevel::Disabled;
     EState CurrentState{EState::Uninitialized};
diff --git a/Engine/Source/Runtime/Core/Private/Hash/Blake3.cpp b/Engine/Source/Runtime/Core/Private/Hash/Blake3.cpp
--- a/Engine/Source/Runtime/Core/Private/Hash/Blake3.cpp
+++ b/Engine/Source/Runtime/Core/Private/Hash/Blake3.cpp
@@ -22,7 +22,7 @@ void FBlake3::Reset()
 	// Call the blake3 function in the open (non-transactionally), and manually
 	// record the update to the Hasher structure.
 	bool bIsClosed = AutoRTFM::IsClosed();
-	UE_AUTORTFM_OPEN_NO_VALIDATION
+	UE_AUTORTFM_OPEN
 	{
 		if (bIsClosed)
 		{
@@ -44,7 +44,7 @@ void FBlake3::Update(const void* Data, uint64 Size)
 	// blake3 is a precompiled external library without AutoRTFM support.
 	// Call the blake3 function in the open (non-transactionally), and manually
 	// record the update to the Hasher structure.
-	UE_AUTORTFM_OPEN_NO_VALIDATION
+	UE_AUTORTFM_OPEN
 	{
 		if (bIsClosed)
 		{
