Fix missing nullcheck in SCOPE_CYCLE_COUNTER_STATID
#rnx

diff --git a/Engine/Source/Runtime/Core/Public/Stats/Stats.h b/Engine/Source/Runtime/Core/Public/Stats/Stats.h
--- a/Engine/Source/Runtime/Core/Public/Stats/Stats.h
+++ b/Engine/Source/Runtime/Core/Public/Stats/Stats.h
@@ -276,7 +276,7 @@ FORCEINLINE void StatsPrimaryEnableSubtract(int32 Value = 1)
 
 #define SCOPE_CYCLE_COUNTER_STATID(StatId) \
 	FScopeCycleCounter StatNamedEventsScope_STATID(StatId); \
-	TRACE_CPUPROFILER_EVENT_SCOPE_TEXT_CONDITIONAL(StatId.StatString, GCycleStatsShouldEmitNamedEvents > 0);
+	TRACE_CPUPROFILER_EVENT_SCOPE_TEXT_CONDITIONAL(StatId.StatString, StatId.StatString != nullptr && GCycleStatsShouldEmitNamedEvents > 0);
 
 #define CONDITIONAL_SCOPE_CYCLE_COUNTER(Stat,bCondition) \
 	FScopeCycleCounter StatNamedEventsScope_##Stat(bCondition ? ANSI_TO_PROFILING(#Stat) : nullptr); \
diff --git a/Engine/Source/Runtime/Engine/Private/Tickable.cpp b/Engine/Source/Runtime/Engine/Private/Tickable.cpp
--- a/Engine/Source/Runtime/Engine/Private/Tickable.cpp
+++ b/Engine/Source/Runtime/Engine/Private/Tickable.cpp
@@ -192,7 +192,7 @@ void FTickableGameObject::TickObjects(UWorld* World, const ELevelTick LevelTickT
 					if ((GIsEditor && TickableObject->IsTickableInEditor()) ||
 						(bIsGameWorld && ((!bIsPaused && LevelTickType != LEVELTICK_TimeOnly) || (bIsPaused && TickableObject->IsTickableWhenPaused()))))
 					{
-						FScopeCycleCounter Context(TickableObject->GetStatId());
+						SCOPE_CYCLE_COUNTER_STATID(TickableObject->GetStatId());
 						TickableObject->Tick(DeltaSeconds);
 					}
 				}
