Deprecated FInternationalization::ForUseOnlyByLocMacroAndGraphNodeTextLiterals_CreateText in favor of FText::AsLocalizable_Advanced

Note: FText::AsLocalizable_Advanced takes its arguments in a different order!

diff --git a/Engine/Source/Runtime/Core/Public/Internationalization/Internationalization.h b/Engine/Source/Runtime/Core/Public/Internationalization/Internationalization.h
--- a/Engine/Source/Runtime/Core/Public/Internationalization/Internationalization.h
+++ b/Engine/Source/Runtime/Core/Public/Internationalization/Internationalization.h
@@ -32,11 +32,7 @@ public:
 
 	static CORE_API void TearDown();
 
-	/**
-	 * Legacy alias for FText::AsLocalizable_Advanced.
-	 * @note FText::AsLocalizable_Advanced takes its arguments in a different order!
-	 * Will be deprecated in the future. DO NOT USE IN NEW CODE!
-	 */
+	UE_DEPRECATED(5.5, "FInternationalization::ForUseOnlyByLocMacroAndGraphNodeTextLiterals_CreateText is deprecated in favor of FText::AsLocalizable_Advanced (NOTE: FText::AsLocalizable_Advanced takes its arguments in a different order!).")
 	static FORCEINLINE FText ForUseOnlyByLocMacroAndGraphNodeTextLiterals_CreateText(const TCHAR* InTextLiteral, const TCHAR* InNamespace, const TCHAR* InKey)
 	{
 		return FText::AsLocalizable_Advanced(InNamespace, InKey, InTextLiteral);
