Clarify help text around FQuat::ToRotationVector and ToAxisAndAngle, and make a helper GetShortestArcWith to help people use these functions without accidentally getting the "long" rotation.
Simply some of the calling places to use the new helper.

diff --git a/Engine/Plugins/Experimental/PhysicsControl/Source/PhysicsControl/Private/AnimNode_RigidBodyWithControl.cpp b/Engine/Plugins/Experimental/PhysicsControl/Source/PhysicsControl/Private/AnimNode_RigidBodyWithControl.cpp
--- a/Engine/Plugins/Experimental/PhysicsControl/Source/PhysicsControl/Private/AnimNode_RigidBodyWithControl.cpp
+++ b/Engine/Plugins/Experimental/PhysicsControl/Source/PhysicsControl/Private/AnimNode_RigidBodyWithControl.cpp
@@ -608,8 +608,7 @@ void FAnimNode_RigidBodyWithControl::EvaluateSkeletalControl_AnyThread(FComponen
 								BodyData.TransferedBoneLinearVelocity = ((NextSSTM.GetLocation() - PrevSSTM.GetLocation()) / DeltaSeconds);
 								
 								// Angular Velocity
-								FQuat DeltaRotation = (NextSSTM.GetRotation().Inverse() * PrevSSTM.GetRotation());
-								DeltaRotation.EnforceShortestArcWith(FQuat::Identity);
+								const FQuat DeltaRotation = (NextSSTM.GetRotation().Inverse() * PrevSSTM.GetRotation()).GetShortestArcWith(FQuat::Identity);
 								BodyData.TransferedBoneAngularVelocity = DeltaRotation.ToRotationVector() / DeltaSeconds;
 							}
 							else
diff --git a/Engine/Plugins/Experimental/PhysicsControl/Source/PhysicsControl/Private/PhysicsControlComponent.cpp b/Engine/Plugins/Experimental/PhysicsControl/Source/PhysicsControl/Private/PhysicsControlComponent.cpp
--- a/Engine/Plugins/Experimental/PhysicsControl/Source/PhysicsControl/Private/PhysicsControlComponent.cpp
+++ b/Engine/Plugins/Experimental/PhysicsControl/Source/PhysicsControl/Private/PhysicsControlComponent.cpp
@@ -1207,11 +1207,10 @@ bool UPhysicsControlComponent::SetControlTargetOrientation(
 	{
 		if (AngularVelocityDeltaTime != 0)
 		{
-			FQuat OldQ = Record->ControlTarget.TargetOrientation.Quaternion();
+			const FQuat OldQ = Record->ControlTarget.TargetOrientation.Quaternion();
 			const FQuat OrientationQ = Orientation.Quaternion();
-			OldQ.EnforceShortestArcWith(OrientationQ);
 			// Note that quats multiply in the opposite order to TMs
-			const FQuat DeltaQ = OrientationQ * OldQ.Inverse();
+			const FQuat DeltaQ = (OrientationQ * OldQ.Inverse()).GetShortestArcWith(FQuat::Identity);
 			Record->ControlTarget.TargetAngularVelocity =
 				DeltaQ.ToRotationVector() / (UE_TWO_PI * AngularVelocityDeltaTime);
 		}
@@ -1358,12 +1357,10 @@ bool UPhysicsControlComponent::SetControlTargetPoses(
 
 		if (VelocityDeltaTime != 0)
 		{
-			FQuat OldQ = Record->ControlTarget.TargetOrientation.Quaternion();
-			OldQ.EnforceShortestArcWith(OrientationQ);
+			const FQuat OldQ = Record->ControlTarget.TargetOrientation.Quaternion();
 			// Note that quats multiply in the opposite order to TMs
-			FQuat DeltaQ = OrientationQ * OldQ.Inverse();
-			Record->ControlTarget.TargetAngularVelocity =
-				DeltaQ.ToRotationVector() / (UE_TWO_PI * VelocityDeltaTime);
+			FQuat DeltaQ = (OrientationQ * OldQ.Inverse()).GetShortestArcWith(FQuat::Identity);
+			Record->ControlTarget.TargetAngularVelocity = DeltaQ.ToRotationVector() / (UE_TWO_PI * VelocityDeltaTime);
 
 			Record->ControlTarget.TargetVelocity =
 				(Position - Record->ControlTarget.TargetPosition) / VelocityDeltaTime;
diff --git a/Engine/Plugins/Experimental/PhysicsControl/Source/PhysicsControl/Private/PhysicsControlPoseData.cpp b/Engine/Plugins/Experimental/PhysicsControl/Source/PhysicsControl/Private/PhysicsControlPoseData.cpp
--- a/Engine/Plugins/Experimental/PhysicsControl/Source/PhysicsControl/Private/PhysicsControlPoseData.cpp
+++ b/Engine/Plugins/Experimental/PhysicsControl/Source/PhysicsControl/Private/PhysicsControlPoseData.cpp
@@ -96,11 +96,8 @@ FVector FBoneData::CalculateAngularVelocity(float Dt) const
 	}
 
 	// Note that quats multiply in the opposite order to TMs, and must be in the same hemisphere.
-	const FQuat Q = CurrentTM.GetRotation();
-	FQuat PrevQ = PreviousTM.GetRotation();
-	PrevQ.EnforceShortestArcWith(Q);
-	const FQuat DeltaQ = Q * PrevQ.Inverse();
-	return DeltaQ.ToRotationVector() / Dt;
+	const FQuat DeltaQ = CurrentTM.GetRotation() * PreviousTM.GetRotation().Inverse();
+	return (DeltaQ.GetShortestArcWith(FQuat::Identity).ToRotationVector()) / Dt;
 }
 
 } // namespace PhysicsControl
diff --git a/Engine/Source/Runtime/Core/Public/Math/Quat.h b/Engine/Source/Runtime/Core/Public/Math/Quat.h
--- a/Engine/Source/Runtime/Core/Public/Math/Quat.h
+++ b/Engine/Source/Runtime/Core/Public/Math/Quat.h
@@ -402,8 +402,11 @@ public:
 	 */
 	FORCEINLINE T SizeSquared() const;
 
-
-	/** Get the angle in radians of this quaternion */
+	/** 
+	 * Get the angle in radians of this quaternion 
+	 * @warning : The rotation returned may not be the shortest version.
+	 *            You may wish to call GetShortestArcWith(FQuat::Identity) first.
+	 */
 	FORCEINLINE T GetAngle() const;
 
 	/** 
@@ -412,6 +415,8 @@ public:
 	 * @param Axis{out] Normalized rotation axis of the quaternion
 	 * @param Angle{out] Angle of the quaternion in radians
 	 * @warning : Requires this quaternion to be normalized.
+	 * @warning : The rotation returned may not be the shortest version.
+	 *            You may wish to call GetShortestArcWith(FQuat::Identity) first.
 	 */
 	void ToAxisAndAngle(TVector<T>& Axis, float& Angle) const;
 	void ToAxisAndAngle(TVector<T>& Axis, double& Angle) const;
@@ -420,7 +425,9 @@ public:
 	 * Get the rotation vector corresponding to this quaternion. 
 	 * The direction of the vector represents the rotation axis, 
 	 * and the magnitude the angle in radians.
-	 * @warning : Requires this quaternion to be normalized.
+	 * @warning : Requires this quaternion to be normalized. 
+	 * @warning : The rotation returned may not be the shortest version. 
+	 *            You may wish to call GetShortestArcWith(FQuat::Identity) first.
 	 */
 	TVector<T> ToRotationVector() const;
 
@@ -468,6 +475,8 @@ public:
 
 	/**
 	 * @return quaternion with W=0 and V=theta*v.
+	 * @warning : The rotation returned may not be the shortest version.
+	 *            You may wish to call GetShortestArcWith(FQuat::Identity) first.
 	 */
 	CORE_API TQuat<T> Log() const;
 
@@ -489,7 +498,13 @@ public:
 	 * the shortest possible rotation angle
 	 */
 	void EnforceShortestArcWith(const TQuat<T>& OtherQuat);
-	
+
+	/**
+	 * Returns the version of this quaternion that represents the shortest rotation angle 
+	 * to the other quaternion
+	 */
+	TQuat<T> GetShortestArcWith(const TQuat<T>& OtherQuat) const;
+
 	/** Get the forward direction (X axis) after it has been rotated by this Quaternion. */
 	FORCEINLINE TVector<T> GetAxisX() const;
 
@@ -1248,6 +1263,13 @@ FORCEINLINE void TQuat<T>::EnforceShortestArcWith(const TQuat<T>& OtherQuat)
 	W *= Bias;
 }
 
+template<typename T>
+FORCEINLINE TQuat<T> TQuat<T>::GetShortestArcWith(const TQuat<T>& OtherQuat) const
+{
+	TQuat<T> Result(*this);
+	Result.EnforceShortestArcWith(OtherQuat);
+	return Result;
+}
 
 template<typename T>
 FORCEINLINE TVector<T> TQuat<T>::GetAxisX() const
