[TaskGraph/Pipe]
 - Fix potential memory ordering issue between tasks inside the same pipe

diff --git a/Engine/Source/Runtime/Core/Private/Tasks/Pipe.cpp b/Engine/Source/Runtime/Core/Private/Tasks/Pipe.cpp
--- a/Engine/Source/Runtime/Core/Private/Tasks/Pipe.cpp
+++ b/Engine/Source/Runtime/Core/Private/Tasks/Pipe.cpp
@@ -5,6 +5,8 @@
 #include "Containers/Array.h"
 #include "Misc/ScopeExit.h"
 
+DEFINE_LOG_CATEGORY_STATIC(LogTaskGraph, Log, All);
+
 namespace UE::Tasks
 {
 	Private::FTaskBase* FPipe::PushIntoPipe(Private::FTaskBase& Task)
@@ -119,10 +121,25 @@ namespace UE::Tasks
 	void FPipe::ExecutionStarted()
 	{
 		FPipeCallStack::Push(*this);
+
+		// Since different threads can execute the same pipe, we need a single atomic
+		// that each thread are going to synchronize memory ordering with.
+		// Acquire any unsynchronized memory modifications that happened
+		// in the last pipe task execution.
+		if (Executing.load(std::memory_order_acquire))
+		{
+			UE_LOG(LogTaskGraph, Fatal, TEXT("Pipe execution must be sequential"));
+		}
+
+		Executing.store(true, std::memory_order_relaxed);
 	}
 
 	void FPipe::ExecutionFinished()
 	{
+		// Publish any unsynchronized memory modification that happened during 
+		// this pipe task execution before another task is executed on the same pipe.
+		Executing.store(false, std::memory_order_release);
+
 		FPipeCallStack::Pop(*this);
 	}
 
diff --git a/Engine/Source/Runtime/Core/Public/Tasks/Pipe.h b/Engine/Source/Runtime/Core/Public/Tasks/Pipe.h
--- a/Engine/Source/Runtime/Core/Public/Tasks/Pipe.h
+++ b/Engine/Source/Runtime/Core/Public/Tasks/Pipe.h
@@ -134,6 +134,7 @@ namespace UE::Tasks
 		// pipe builds a chain (a linked list) of tasks and so needs to store only the last one. the last task is null if the pipe is not blocked
 		std::atomic<Private::FTaskBase*> LastTask{ nullptr };
 		std::atomic<uint64> TaskCount { 0 };
+		std::atomic<bool> Executing{ false }; // Used for memory ordering between pipe tasks.
 		UE::FEventCount EmptyEvent;
 	public:
 		FORCENOINLINE const TCHAR* GetDebugName() const
