[EventCount]
 - Bring back the x86 version of the code to avoid performance cost on platforms that don't need the additional costly instruction

#rnx

diff --git a/Engine/Source/Runtime/Core/Public/Async/EventCount.h b/Engine/Source/Runtime/Core/Public/Async/EventCount.h
--- a/Engine/Source/Runtime/Core/Public/Async/EventCount.h
+++ b/Engine/Source/Runtime/Core/Public/Async/EventCount.h
@@ -135,6 +135,7 @@ public:
 	 */
 	inline void Notify()
 	{
+#if PLATFORM_WEAKLY_CONSISTENT_MEMORY
 		//
 		// .fetch_add(0, acq_rel) is used to have a StoreLoad barrier,
 		// which we can't express in C++. That works by making the load
@@ -148,6 +149,12 @@ public:
 		// leading to a missed Wake and dead-locked waiter, as a result.
 		//
 		CounterType Value = Count.fetch_add(0, std::memory_order_acq_rel);
+#else
+		// On x86 and other non weak memory model, the fetch_or inside PrepareWait
+		// is a serializing instruction that will flush the store buffer. 
+		// We can omit the expensive locked op here and just do a relaxed read.
+		CounterType Value = Count.load(std::memory_order_relaxed);
+#endif
 		if ((Value & 1) && Count.compare_exchange_strong(Value, Value + 1, std::memory_order_release))
 		{
 			ParkingLot::WakeAll(&Count);
