Fix the RecordOpenWrite bug where it was incorrectly recording if another thread was in a transaction.

diff --git a/Engine/Source/Programs/AutoRTFMTests/Private/OpenAPI.cpp b/Engine/Source/Programs/AutoRTFMTests/Private/OpenAPI.cpp
--- a/Engine/Source/Programs/AutoRTFMTests/Private/OpenAPI.cpp
+++ b/Engine/Source/Programs/AutoRTFMTests/Private/OpenAPI.cpp
@@ -3,7 +3,9 @@
 #include "Catch2Includes.h"
 #include <AutoRTFM/AutoRTFM.h>
 
+#include <atomic>
 #include <map>
+#include <thread>
 #include <vector>
 
 TEST_CASE("OpenAPI.StartAbortAndStartAgain")
@@ -897,3 +899,32 @@ TEST_CASE("OpenAPI.DeferredStartTransactionOverflow")
 		}
 	});
 }
+
+TEST_CASE("OpenAPI.CheckRaceAgainstOtherThread")
+{
+	bool bHit = false;
+	std::atomic_uint Handshake = 0;
+	AutoRTFM::ETransactionResult Result = AutoRTFM::Transact([&]
+		{
+			AutoRTFM::Open([&]
+				{
+					std::thread([&]
+						{
+							// This should be a no-op in the spawnee thread.
+							AutoRTFM::RecordOpenWrite(&bHit);
+							bHit = true;
+
+							// Unblock the main thread.
+							Handshake += 1;
+						}).detach();
+
+					// Wait for the spawnee.
+					while (1 != Handshake) {}
+				});
+
+			AutoRTFM::AbortTransaction();
+		});
+
+	REQUIRE(AutoRTFM::ETransactionResult::AbortedByRequest == Result);
+	REQUIRE(bHit);
+}
\ No newline at end of file
diff --git a/Engine/Source/Runtime/Core/Private/AutoRTFM/API.cpp b/Engine/Source/Runtime/Core/Private/AutoRTFM/API.cpp
--- a/Engine/Source/Runtime/Core/Private/AutoRTFM/API.cpp
+++ b/Engine/Source/Runtime/Core/Private/AutoRTFM/API.cpp
@@ -583,7 +583,11 @@ extern "C" void autortfm_record_open_write(void* Ptr, size_t Size)
 	{
 		return RTFM_autortfm_record_open_write(Ptr, Size);
 	}
-	if (FTransaction* CurrentTransaction = FContext::Get()->GetCurrentTransaction())
+	else if (!FContext::IsTransactional())	
+	{
+		return;
+	}
+	else if (FTransaction* CurrentTransaction = FContext::Get()->GetCurrentTransaction())
 	{
 		CurrentTransaction->RecordWrite(Ptr, Size);
 	}
@@ -597,11 +601,18 @@ extern "C" UE_AUTORTFM_NOAUTORTFM void autortfm_register_open_function(void* Ori
 
 extern "C" bool autortfm_is_on_current_transaction_stack(void* Ptr)
 {
-	if (FTransaction* CurrentTransaction = FContext::Get()->GetCurrentTransaction())
+	if (!FContext::IsTransactional())
+	{
+		return false;
+	}
+	else if (FTransaction* CurrentTransaction = FContext::Get()->GetCurrentTransaction())
 	{
 		return CurrentTransaction->IsOnStack(Ptr);
 	}
-	return false;
+	else
+	{
+		return false;
+	}
 }
 
 void ForTheRuntime::OnCommitInternal(TFunction<void()> && Work)
