[Pipe]
 - Fix double definition for LogTaskGraph

#rnx

diff --git a/Engine/Source/Runtime/Core/Private/Tasks/Pipe.cpp b/Engine/Source/Runtime/Core/Private/Tasks/Pipe.cpp
--- a/Engine/Source/Runtime/Core/Private/Tasks/Pipe.cpp
+++ b/Engine/Source/Runtime/Core/Private/Tasks/Pipe.cpp
@@ -5,7 +5,7 @@
 #include "Containers/Array.h"
 #include "Misc/ScopeExit.h"
 
-DEFINE_LOG_CATEGORY_STATIC(LogTaskGraph, Log, All);
+DEFINE_LOG_CATEGORY_STATIC(LogTaskGraphPipe, Log, All);
 
 namespace UE::Tasks
 {
@@ -128,7 +128,7 @@ namespace UE::Tasks
 		// in the last pipe task execution.
 		if (Executing.load(std::memory_order_acquire))
 		{
-			UE_LOG(LogTaskGraph, Fatal, TEXT("Pipe execution must be sequential"));
+			UE_LOG(LogTaskGraphPipe, Fatal, TEXT("Pipe execution must be sequential"));
 		}
 
 		Executing.store(true, std::memory_order_relaxed);
