[TaskGraph]
 - Remove older deprecations

#rnx

diff --git a/Engine/Source/Runtime/Core/Public/Async/Fundamental/Scheduler.h b/Engine/Source/Runtime/Core/Public/Async/Fundamental/Scheduler.h
--- a/Engine/Source/Runtime/Core/Public/Async/Fundamental/Scheduler.h
+++ b/Engine/Source/Runtime/Core/Public/Async/Fundamental/Scheduler.h
@@ -139,10 +139,6 @@ namespace LowLevelTasks
 		// returns true if the current thread execution is in the context of busy-waiting
 		CORE_API static bool IsBusyWaiting();
 
-		// returns the AffinityIndex of the thread LocalQueue
-		UE_DEPRECATED(5.4, "This method will be removed.")
-		CORE_API static uint32 GetAffinityIndex();
-
 	protected:
 		inline static bool IsBackgroundWorker()
 		{
@@ -181,9 +177,6 @@ namespace LowLevelTasks
 		//try to launch the task, the return value will specify if the task was in the ready state and has been launched
 		inline bool TryLaunch(FTask& Task, EQueuePreference QueuePreference = EQueuePreference::DefaultPreference, bool bWakeUpWorker = true);
 
-		UE_DEPRECATED(5.4, "Use TryLaunch instead")
-		inline bool TryLaunchAffinity(FTask& Task, uint32 AffinityIndex);
-
 		//tries to do some work until the Task is completed
 		template<typename TaskType>
 		UE_DEPRECATED(5.5, "This method will be removed.")
@@ -373,12 +366,6 @@ namespace LowLevelTasks
 		return false;
 	}
 
-	inline bool FScheduler::TryLaunchAffinity(FTask& Task, uint32 AffinityIndex)
-	{
-		// Redirect until we remove the deprecated function.
-		return FScheduler::Get().TryLaunch(Task);
-	}
-
 	inline uint32 FScheduler::GetNumWorkers() const
 	{
 		return ActiveWorkers.load(std::memory_order_relaxed);
diff --git a/Engine/Source/Runtime/Core/Public/Async/Fundamental/Task.h b/Engine/Source/Runtime/Core/Public/Async/Fundamental/Task.h
--- a/Engine/Source/Runtime/Core/Public/Async/Fundamental/Task.h
+++ b/Engine/Source/Runtime/Core/Public/Async/Fundamental/Task.h
@@ -664,25 +664,4 @@ namespace LowLevelTasks
 		FPackedData LocalPackedData = PackedData.load(std::memory_order_relaxed);
 		return { LocalPackedData.GetDebugName(), LocalPackedData.GetPriority(), LocalPackedData.GetFlags() };
 	}
-
-	enum class UE_DEPRECATED(5.4, "This was meant for internal use only and will be removed") ESleepState
-	{
-		Affinity,
-		Running,
-		Drowsing,
-		Sleeping,
-	};
-
-	/*
-	* the struct is naturally 64 bytes aligned, the extra alignment just 
-	* re-enforces this assumption and will error if it changes in the future
-	*/
-	struct UE_DEPRECATED(5.4, "This was meant for internal use only and will be removed") alignas(64) FSleepEvent
-	{
-		FEventRef SleepEvent;
-		PRAGMA_DISABLE_DEPRECATION_WARNINGS
-		std::atomic<ESleepState> SleepState { ESleepState::Running };
-		std::atomic<FSleepEvent*> Next { nullptr };
-		PRAGMA_ENABLE_DEPRECATION_WARNINGS
-	};
 }
\ No newline at end of file
