Mac: Fix framepro includes.

diff --git a/Engine/Source/Runtime/Core/Private/FramePro/FrameProUE4.cpp b/Engine/Source/Runtime/Core/Private/FramePro/FrameProUE4.cpp
--- a/Engine/Source/Runtime/Core/Private/FramePro/FrameProUE4.cpp
+++ b/Engine/Source/Runtime/Core/Private/FramePro/FrameProUE4.cpp
@@ -72,11 +72,11 @@
 	#define FRAMEPRO_USE_RTLCAPTURESTACKBACKTRACE 0	// system version of FRAMEPRO_USE_RTLVIRTUALUNWIND - only available on x64 builds
 
 	//------------------------------------------------------------------------
-	#if PLATFORM_ANDROID
+	//@EPIC BEGIN: arm mac needs to use this one
+	#if PLATFORM_ANDROID || PLATFORM_MAC
 		#include <sys/syscall.h>
-	#elif PLATFORM_MAC
-		#include <cpuid.h>
 	#endif
+	//@EPIC END
 
 	//@EPIC BEGIN: workaround for -nothreading
 	#if PLATFORM_WINDOWS
