[core] add some missing constexpr's to type traits

diff --git a/Engine/Source/Runtime/Core/Public/Templates/UnrealTypeTraits.h b/Engine/Source/Runtime/Core/Public/Templates/UnrealTypeTraits.h
--- a/Engine/Source/Runtime/Core/Public/Templates/UnrealTypeTraits.h
+++ b/Engine/Source/Runtime/Core/Public/Templates/UnrealTypeTraits.h
@@ -93,7 +93,7 @@ template<typename T>
 struct TFormatSpecifier
 {
 	template <typename CharType = TCHAR>
-	FORCEINLINE static TCHAR const* GetFormatSpecifier()
+	FORCEINLINE static constexpr TCHAR const* GetFormatSpecifier()
 	{
 		// Force the template instantiation to be dependent upon T so the compiler cannot automatically decide that this template can never be instantiated.
 		// If the static_assert below were a constant 0 or something not dependent on T, compilers are free to detect this and fail to compile the template.
@@ -108,7 +108,7 @@ template<> \
 struct TFormatSpecifier<type> \
 {  \
 	template <typename CharType = TCHAR> \
-	FORCEINLINE static decltype(auto) GetFormatSpecifier() \
+	FORCEINLINE static constexpr decltype(auto) GetFormatSpecifier() \
 	{ \
         if constexpr (std::is_same_v<CharType, WIDECHAR>) \
         { \
