Parallelizing the initialization of platform configs to optimize: FConfigCacheIni::AddPluginToAllBranches.

#rnx

diff --git a/Engine/Source/Developer/TargetPlatform/Private/TargetPlatformManagerModule.cpp b/Engine/Source/Developer/TargetPlatform/Private/TargetPlatformManagerModule.cpp
--- a/Engine/Source/Developer/TargetPlatform/Private/TargetPlatformManagerModule.cpp
+++ b/Engine/Source/Developer/TargetPlatform/Private/TargetPlatformManagerModule.cpp
@@ -5,7 +5,6 @@
 #include "Misc/App.h"
 #include "Misc/ConfigCacheIni.h"
 #include "Misc/CommandLine.h"
-#include "Misc/ConfigCacheIni.h"
 #include "Misc/Paths.h"
 #include "Misc/OutputDeviceRedirector.h"
 #include "ShaderCompilerCore.h"
@@ -885,6 +884,8 @@ protected:
 		}
 #endif
 
+		FConfigCacheIni::PreInitializePlatformPlugins();
+
 		// find a set of valid target platform names (the platform DataDrivenPlatformInfo.ini file was found indicates support for the platform 
 		// exists on disk, so the TP is expected to work)
 		FScopedSlowTask SlowTask((float)FDataDrivenPlatformInfoRegistry::GetAllPlatformInfos().Num());
diff --git a/Engine/Source/Runtime/Core/Private/Misc/ConfigCacheIni.cpp b/Engine/Source/Runtime/Core/Private/Misc/ConfigCacheIni.cpp
--- a/Engine/Source/Runtime/Core/Private/Misc/ConfigCacheIni.cpp
+++ b/Engine/Source/Runtime/Core/Private/Misc/ConfigCacheIni.cpp
@@ -29,6 +29,7 @@
 #include "Serialization/MemoryWriter.h"
 #include "Serialization/LargeMemoryReader.h"
 #include "Async/Async.h"
+#include "Async/ParallelFor.h"
 #include "Misc/OutputDeviceRedirector.h"
 #include "Logging/MessageLog.h"
 #include "Misc/TransactionallySafeRWLock.h"
@@ -6422,6 +6423,45 @@ void FConfigCacheIni::AsyncInitializeConfigForPlatforms()
 #endif
 #endif
 
+void FConfigCacheIni::PreInitializePlatformPlugins()
+{
+#if ALLOW_OTHER_PLATFORM_CONFIG
+#if WITH_EDITOR
+
+	const TMap<FName, FDataDrivenPlatformInfo>& AllPlatformInfos = FDataDrivenPlatformInfoRegistry::GetAllPlatformInfos();
+
+	TArray<FName> PlatformNames;
+	TArray<FDataDrivenPlatformInfo> PlatformInfos;
+	AllPlatformInfos.GetKeys(PlatformNames);
+
+	ParallelFor(PlatformNames.Num(), [PlatformNames, AllPlatformInfos](int32 Index)
+	{
+		FName PlatformName = PlatformNames[Index];
+
+		// if we have the editor and we are using -game
+		// only need to instantiate the current platform 
+		if (IsRunningGame())
+		{
+			if (PlatformName != FPlatformProperties::IniPlatformName())
+			{
+				return;
+			}
+		}
+
+		const FDataDrivenPlatformInfo* Info = AllPlatformInfos.Find(PlatformName);
+
+		if (Info->bEnabledForUse)
+		{
+			// Calling ForPlatform will invoke AddPluginToBranches on any pending 
+			// plugins found for this platform's config. We do this beforehand on
+			// worker threads to speed up platform initialization.
+			FConfigCacheIni::ForPlatform(PlatformName);
+		}
+	}, /*bForceSingleThread*/ false);
+
+#endif
+#endif
+}
 
 FConfigCacheIni* FConfigCacheIni::ForPlatform(FName PlatformName)
 {
@@ -6448,29 +6488,39 @@ FConfigCacheIni* FConfigCacheIni::ForPlatform(FName PlatformName)
 	}
 #endif
 
-	// protect against other threads clearing the array, or two threads trying to read in a missing platform at the same time
-	FScopeLock Lock(&ConfigForPlatformLock);
-	FConfigCacheIni* PlatformConfig = ConfigForPlatform.FindRef(PlatformName);
+	FConfigCacheIni* PlatformConfig = nullptr;
+	TArray<FName> PendingModificationPlugins;
 
-	// read any missing platform configs now, on demand (this will happen when WITH_EDITOR is 0)
-	if (PlatformConfig == nullptr)
 	{
-		double Start = FPlatformTime::Seconds();
-		
-		PlatformConfig = ConfigForPlatform.Add(PlatformName, new FConfigCacheIni(EConfigCacheType::Temporary, PlatformName, true /* bInGloballyRegistered */));
-		FConfigContext Context = FConfigContext::ReadIntoConfigSystem(PlatformConfig, PlatformName.ToString());
-		InitializeKnownConfigFiles(Context);
+		// protect against other threads clearing the array, or two threads trying to read in a missing platform at the same time
+		FScopeLock Lock(&ConfigForPlatformLock);
+
+		PlatformConfig = ConfigForPlatform.FindRef(PlatformName);
+
+		// read any missing platform configs now, on demand (this will happen when WITH_EDITOR is 0)
+		if (PlatformConfig == nullptr)
+		{
+			double Start = FPlatformTime::Seconds();
 
-		UE_LOG(LogConfig, Display, TEXT("Read in platform %s ini files took %.2f seconds"), *PlatformName.ToString(), FPlatformTime::Seconds() - Start);
+			PlatformConfig = ConfigForPlatform.Add(PlatformName, new FConfigCacheIni(EConfigCacheType::Temporary, PlatformName, true /* bInGloballyRegistered */));
+			FConfigContext Context = FConfigContext::ReadIntoConfigSystem(PlatformConfig, PlatformName.ToString());
+			InitializeKnownConfigFiles(Context);
+
+			UE_LOG(LogConfig, Display, TEXT("Read in platform %s ini files took %.2f seconds"), *PlatformName.ToString(), FPlatformTime::Seconds() - Start);
+		}
+
+		PendingModificationPlugins = MoveTemp(PlatformConfig->PendingModificationPlugins);
+		PlatformConfig->PendingModificationPlugins.Empty();
 	}
 
-	for (FName PluginName : PlatformConfig->PendingModificationPlugins)
 	{
-		// delayed plugin injection
-		PlatformConfig->AddPluginToBranches(PluginName, nullptr);
+		for (FName PluginName : PendingModificationPlugins)
+		{
+			// delayed plugin injection
+			PlatformConfig->AddPluginToBranches(PluginName, nullptr);
+		}
 	}
-	PlatformConfig->PendingModificationPlugins.Empty();
-	
+
 	return PlatformConfig;
 
 #else
diff --git a/Engine/Source/Runtime/Core/Public/Misc/ConfigCacheIni.h b/Engine/Source/Runtime/Core/Public/Misc/ConfigCacheIni.h
--- a/Engine/Source/Runtime/Core/Public/Misc/ConfigCacheIni.h
+++ b/Engine/Source/Runtime/Core/Public/Misc/ConfigCacheIni.h
@@ -1885,6 +1885,12 @@ public:
 	static CORE_API void AddPluginToAllBranches(FName PluginName, FConfigModificationTracker* ModificationTracker=nullptr);
 	static CORE_API void RemoveTagFromAllBranches(FName Tag, FConfigModificationTracker* ModificationTracker=nullptr);
 
+	/**
+	* Initializes configs and plugins for every supported platform on worker threads. This function can be
+	* called prior to normal platform initialization to speed up the full initialization process.
+	*/
+	static CORE_API void PreInitializePlatformPlugins();
+
 private:
 #if WITH_EDITOR
 	/** We only auto-initialize other platform configs in the editor to not slow down programs like ShaderCOmpileWorker */
