Fix for field shadow error of PlatformName.

#rnx

diff --git a/Engine/Source/Runtime/Core/Private/Misc/ConfigCacheIni.cpp b/Engine/Source/Runtime/Core/Private/Misc/ConfigCacheIni.cpp
--- a/Engine/Source/Runtime/Core/Private/Misc/ConfigCacheIni.cpp
+++ b/Engine/Source/Runtime/Core/Private/Misc/ConfigCacheIni.cpp
@@ -6436,26 +6436,26 @@ void FConfigCacheIni::PreInitializePlatformPlugins()
 
 	ParallelFor(PlatformNames.Num(), [PlatformNames, AllPlatformInfos](int32 Index)
 	{
-		FName PlatformName = PlatformNames[Index];
+		FName PlatfName = PlatformNames[Index];
 
 		// if we have the editor and we are using -game
 		// only need to instantiate the current platform 
 		if (IsRunningGame())
 		{
-			if (PlatformName != FPlatformProperties::IniPlatformName())
+			if (PlatfName != FPlatformProperties::IniPlatformName())
 			{
 				return;
 			}
 		}
 
-		const FDataDrivenPlatformInfo* Info = AllPlatformInfos.Find(PlatformName);
+		const FDataDrivenPlatformInfo* Info = AllPlatformInfos.Find(PlatfName);
 
 		if (Info->bEnabledForUse)
 		{
 			// Calling ForPlatform will invoke AddPluginToBranches on any pending 
 			// plugins found for this platform's config. We do this beforehand on
 			// worker threads to speed up platform initialization.
-			FConfigCacheIni::ForPlatform(PlatformName);
+			FConfigCacheIni::ForPlatform(PlatfName);
 		}
 	}, /*bForceSingleThread*/ false);
 
