[AutoRTFM] Optimize FHitSet by avoiding heap allocation for first InitialCapacity writes

Also rename GetSize() to GetCount(), as the former is slightly ambiguous about whether it means 'count' or 'count*sizeof(item)'

diff --git a/Engine/Source/Runtime/Core/Private/AutoRTFM/HitSet.h b/Engine/Source/Runtime/Core/Private/AutoRTFM/HitSet.h
--- a/Engine/Source/Runtime/Core/Private/AutoRTFM/HitSet.h
+++ b/Engine/Source/Runtime/Core/Private/AutoRTFM/HitSet.h
@@ -4,6 +4,7 @@
 
 #include "TaggedPtr.h"
 #include "Utils.h"
+#include <array>
 
 namespace AutoRTFM
 {
@@ -19,30 +20,15 @@ namespace AutoRTFM
 		// TODO: Revisit a good initial capacity for the hitset.
 		static constexpr uint64_t LogInitialCapacity = 4;
 
+		static constexpr uint64_t InitialCapacity = static_cast<uint64_t>(1) << LogInitialCapacity;
+
 		// TODO: Revisit a good max capacity for the hitset.
 		static constexpr uint64_t MaxCapacity = (128 * 1024 * 1024) / sizeof(Key);
 
 	public:
-		explicit FHitSet()
-		{
-			constexpr uint64_t InitialCapacity = static_cast<uint64_t>(1) << LogInitialCapacity;
-
-			// Do not want to deal with null payloads.
-			static_assert(0 != InitialCapacity);
-
-			// The capacity is always a power of two so that the range reduction is optimal.
-			static_assert(0 == (InitialCapacity & (InitialCapacity - 1)));
-
-			Payload = static_cast<uintptr_t*>(FMemory::MallocZeroed(InitialCapacity * sizeof(uintptr_t)));
-			ASSERT(nullptr != Payload);
-
-			SixtyFourMinusLogCapacity = 64 - LogInitialCapacity;
-			Size = 0;
-		}
-
 		~FHitSet()
 		{
-			FMemory::Free(Payload);
+			Reset();
 		}
 
 		enum class EInsertResult : uint8_t
@@ -86,14 +72,23 @@ namespace AutoRTFM
 
 		UE_AUTORTFM_FORCEINLINE bool IsEmpty() const
 		{
-			return 0 == Size;
+			return 0 == Count;
 		}
 
 		// Clear out the data stored in the set, but does not reduce the capacity.
 		void Reset()
 		{
-			memset(Payload, 0, sizeof(uintptr_t) * Capacity());
-			Size = 0;
+			if (Count)
+			{
+				if (Payload != SmallPayload.data())
+				{
+					FMemory::Free(Payload);
+					Payload = SmallPayload.data();
+				}
+				SmallPayload = {};
+				SixtyFourMinusLogCapacity = 64 - LogInitialCapacity;
+				Count = 0;
+			}
 		}
 
 		uint64_t GetCapacity() const
@@ -101,17 +96,16 @@ namespace AutoRTFM
 			return Capacity();
 		}
 
-		uint64_t GetSize() const
+		uint64_t GetCount() const
 		{
-			return Size;
+			return Count;
 		}
 
 	private:
-		uintptr_t* Payload;
-
-		uint64_t SixtyFourMinusLogCapacity;
-
-		uint64_t Size;
+		std::array<uint64_t, InitialCapacity> SmallPayload{};
+		uintptr_t* Payload = SmallPayload.data();
+		uint64_t SixtyFourMinusLogCapacity = 64 - LogInitialCapacity;
+		uint64_t Count = 0;
 
 		UE_AUTORTFM_FORCEINLINE uint64_t Capacity() const
 		{
@@ -138,13 +132,13 @@ namespace AutoRTFM
 			}
 
 			uintptr_t* const OldPayload = Payload;
-			const uint64_t OldSize = Size;
+			const uint64_t OldSize = Count;
 
 			while (true)
 			{
 				bool bNeedAnotherResize = false;
 
-				Size = 0;
+				Count = 0;
 
 				IncreaseCapacity();
 
@@ -175,16 +169,22 @@ namespace AutoRTFM
 
 				if (bNeedAnotherResize)
 				{
-					FMemory::Free(Payload);
+					if (Payload != SmallPayload.data())
+					{
+						FMemory::Free(Payload);
+					}
 					continue;
 				}
 
 				break;
 			}
 
-			ASSERT(OldSize == Size);
+			ASSERT(OldSize == Count);
 
-			FMemory::Free(OldPayload);
+			if (OldPayload != SmallPayload.data())
+			{
+				FMemory::Free(OldPayload);
+			}
 
 			return true;
 		}
@@ -216,7 +216,7 @@ namespace AutoRTFM
 			if (0 == Payload[I])
 			{
 				Payload[I] = Raw;
-				Size++;
+				Count++;
 				return EInsertResult::Inserted;
 			}
 
diff --git a/Engine/Source/Runtime/Core/Private/AutoRTFM/TransactionInlines.h b/Engine/Source/Runtime/Core/Private/AutoRTFM/TransactionInlines.h
--- a/Engine/Source/Runtime/Core/Private/AutoRTFM/TransactionInlines.h
+++ b/Engine/Source/Runtime/Core/Private/AutoRTFM/TransactionInlines.h
@@ -247,7 +247,7 @@ UE_AUTORTFM_FORCEINLINE void FTransaction::CollectStats() const
     Stats.Collect<EStatsKind::AverageAbortTasks>(AbortTasks.Num());
     Stats.Collect<EStatsKind::MaximumAbortTasks>(AbortTasks.Num());
 
-    Stats.Collect<EStatsKind::AverageHitSetSize>(HitSet.GetSize());
+    Stats.Collect<EStatsKind::AverageHitSetSize>(HitSet.GetCount());
     Stats.Collect<EStatsKind::AverageHitSetCapacity>(HitSet.GetCapacity());
 }
 
