[AutoRTFM] Fix uintptr_t <-> uint64_t mismatch, drop usage of <array>

diff --git a/Engine/Source/Runtime/Core/Private/AutoRTFM/HitSet.h b/Engine/Source/Runtime/Core/Private/AutoRTFM/HitSet.h
--- a/Engine/Source/Runtime/Core/Private/AutoRTFM/HitSet.h
+++ b/Engine/Source/Runtime/Core/Private/AutoRTFM/HitSet.h
@@ -4,7 +4,6 @@
 
 #include "TaggedPtr.h"
 #include "Utils.h"
-#include <array>
 
 namespace AutoRTFM
 {
@@ -80,12 +79,12 @@ namespace AutoRTFM
 		{
 			if (Count)
 			{
-				if (Payload != SmallPayload.data())
+				if (Payload != SmallPayload)
 				{
 					FMemory::Free(Payload);
-					Payload = SmallPayload.data();
+					Payload = SmallPayload;
 				}
-				SmallPayload = {};
+				memset(SmallPayload, 0, sizeof(SmallPayload));
 				SixtyFourMinusLogCapacity = 64 - LogInitialCapacity;
 				Count = 0;
 			}
@@ -102,8 +101,8 @@ namespace AutoRTFM
 		}
 
 	private:
-		std::array<uint64_t, InitialCapacity> SmallPayload{};
-		uintptr_t* Payload = SmallPayload.data();
+		uintptr_t SmallPayload[InitialCapacity]{};
+		uintptr_t* Payload = SmallPayload;
 		uint64_t SixtyFourMinusLogCapacity = 64 - LogInitialCapacity;
 		uint64_t Count = 0;
 
@@ -169,7 +168,7 @@ namespace AutoRTFM
 
 				if (bNeedAnotherResize)
 				{
-					if (Payload != SmallPayload.data())
+					if (Payload != SmallPayload)
 					{
 						FMemory::Free(Payload);
 					}
@@ -181,7 +180,7 @@ namespace AutoRTFM
 
 			ASSERT(OldSize == Count);
 
-			if (OldPayload != SmallPayload.data())
+			if (OldPayload != SmallPayload)
 			{
 				FMemory::Free(OldPayload);
 			}
