[AutoRTFM] Use transactionally-safe locks in TrackedActivity.

diff --git a/Engine/Source/Runtime/Core/Private/Misc/TrackedActivity.cpp b/Engine/Source/Runtime/Core/Private/Misc/TrackedActivity.cpp
--- a/Engine/Source/Runtime/Core/Private/Misc/TrackedActivity.cpp
+++ b/Engine/Source/Runtime/Core/Private/Misc/TrackedActivity.cpp
@@ -2,9 +2,9 @@
 
 #include "Misc/TrackedActivity.h"
 #include "Containers/UnrealString.h"
-#include "HAL/CriticalSection.h"
 #include "Misc/ScopeLock.h"
 #include "Misc/StringBuilder.h"
+#include "Misc/TransactionallySafeCriticalSection.h"
 
 // Singleton class that keep track of all the existing activities and their stacks
 class FTrackedActivityManager
@@ -36,7 +36,7 @@ public:
 		uint32 Id;
 	};
 
-	FCriticalSection ActivitiesCs;
+	FTransactionallySafeCriticalSection ActivitiesCs;
 	TArray<FActivity*> Activities;
 	TUniqueFunction<void(FTrackedActivity::EEvent Event, const FTrackedActivity::FInfo& Info)> EventFunc;
 	uint32 EventFuncMaxDepth = ~0u;
@@ -50,7 +50,7 @@ public:
 		Activity->SortValue = SortValue;
 		Activity->Stack.Add({ Status, Light, false });
 
-		FScopeLock _(&ActivitiesCs);
+		UE::TScopeLock _(ActivitiesCs);
 		Activity->Id = IdCounter++;
 		Activities.Add(Activity);
 		SendEvent(FTrackedActivity::EEvent::Added, *Activity);
@@ -61,7 +61,7 @@ public:
 	{
 		check(Activity->Stack.Num() == 1);
 		{
-			FScopeLock _(&ActivitiesCs);
+			UE::TScopeLock _(ActivitiesCs);
 			Activities.Remove(Activity);
 			SendEvent(FTrackedActivity::EEvent::Removed, *Activity);
 		}
@@ -70,7 +70,7 @@ public:
 
 	uint32 Push(FActivity& Activity, const TCHAR* Status, bool bShowParent, ELight Light)
 	{
-		FScopeLock _(&ActivitiesCs);
+		UE::TScopeLock _(ActivitiesCs);
 		uint32 Index = Activity.Stack.Num();
 		Activity.Stack.Add({ Status, Light, bShowParent });
 		SendEvent(FTrackedActivity::EEvent::Changed, Activity);
@@ -79,14 +79,14 @@ public:
 
 	void Pop(FActivity& Activity)
 	{
-		FScopeLock _(&ActivitiesCs);
+		UE::TScopeLock _(ActivitiesCs);
 		Activity.Stack.SetNum(Activity.Stack.Num() - 1);
 		SendEvent(FTrackedActivity::EEvent::Changed, Activity);
 	}
 
 	void Update(FActivity& Activity, const TCHAR* Status, uint32 Index)
 	{
-		FScopeLock _(&ActivitiesCs);
+		UE::TScopeLock _(ActivitiesCs);
 		FTrackedActivityManager::FInfo& Info = Index == ~0u ? Activity.Stack.Last() : Activity.Stack[Index];
 		Info.Status = Status;
 		SendEvent(FTrackedActivity::EEvent::Changed, Activity);
@@ -94,7 +94,7 @@ public:
 
 	void Update(FActivity& Activity, const TCHAR* Status, ELight Light, uint32 Index)
 	{
-		FScopeLock _(&ActivitiesCs);
+		UE::TScopeLock _(ActivitiesCs);
 		FTrackedActivityManager::FInfo& Info = Index == ~0u ? Activity.Stack.Last() : Activity.Stack[Index];
 		Info.Status = Status;
 		Info.Light = Light;
@@ -103,7 +103,7 @@ public:
 
 	void Update(FActivity& Activity, ELight Light, uint32 Index)
 	{
-		FScopeLock _(&ActivitiesCs);
+		UE::TScopeLock _(ActivitiesCs);
 		FTrackedActivityManager::FInfo& Info = Index == ~0u ? Activity.Stack.Last() : Activity.Stack[Index];
 		Info.Light = Light;
 		SendEvent(FTrackedActivity::EEvent::Changed, Activity);
@@ -139,7 +139,7 @@ public:
 
 	void TraverseActivities(const TFunction<void(const FTrackedActivity::FInfo& Info)>& Func)
 	{
-		FScopeLock _(&ActivitiesCs);
+		UE::TScopeLock _(ActivitiesCs);
 
 		for (auto A : Activities)
 		{
@@ -151,7 +151,7 @@ public:
 	void RegisterEventListener(TUniqueFunction<void(FTrackedActivity::EEvent Event, const FTrackedActivity::FInfo& Info)>&& Func, uint32 MaxDepth)
 	{
 		checkf(!EventFunc, TEXT("TrackedActivity system is only supporting one event listener at the time. Please add support"));
-		FScopeLock _(&ActivitiesCs);
+		UE::TScopeLock _(ActivitiesCs);
 
 		EventFunc = MoveTemp(Func);
 		EventFuncMaxDepth = MaxDepth;
