Add missing API declaration to some AutoRTFM functions to fix a link error.

diff --git a/Engine/Source/Runtime/Core/Private/AutoRTFM/API.cpp b/Engine/Source/Runtime/Core/Private/AutoRTFM/API.cpp
--- a/Engine/Source/Runtime/Core/Private/AutoRTFM/API.cpp
+++ b/Engine/Source/Runtime/Core/Private/AutoRTFM/API.cpp
@@ -339,7 +339,7 @@ AUTORTFM_CORE_API autortfm_status RTFM_autortfm_close(void (*UninstrumentedWork)
 	return static_cast<autortfm_status>(Context->GetStatus());
 }
 
-extern "C" UE_AUTORTFM_NOAUTORTFM void autortfm_pre_open()
+extern "C" UE_AUTORTFM_NOAUTORTFM UE_AUTORTFM_API void autortfm_pre_open()
 {
 	if (FTransaction* Transaction = FContext::Get()->GetCurrentTransaction())
 	{
@@ -348,7 +348,7 @@ extern "C" UE_AUTORTFM_NOAUTORTFM void autortfm_pre_open()
 	}
 }
 
-extern "C" UE_AUTORTFM_NOAUTORTFM void autortfm_post_open()
+extern "C" UE_AUTORTFM_NOAUTORTFM UE_AUTORTFM_API void autortfm_post_open()
 {
 	if (FTransaction* Transaction = FContext::Get()->GetCurrentTransaction())
 	{
