[Core/ShaderCompileWorker]
* Added option to disable collecting of storage information using define
* Set define for ShaderCompileWorker target to disable collecting storage information since it seems to hang for minutes some times

diff --git a/Engine/Source/Runtime/Core/Private/Windows/WindowsPlatformMisc.cpp b/Engine/Source/Runtime/Core/Private/Windows/WindowsPlatformMisc.cpp
--- a/Engine/Source/Runtime/Core/Private/Windows/WindowsPlatformMisc.cpp
+++ b/Engine/Source/Runtime/Core/Private/Windows/WindowsPlatformMisc.cpp
@@ -79,6 +79,12 @@ THIRD_PARTY_INCLUDES_END
 
 #include "FramePro/FrameProProfiler.h"
 
+// Collect extended storage information is only enabled if target has explicitly enabled it
+// Reason is because we've seen 6+ minutes stalls in the dcom ConnectServer call when compiling shaders
+#ifndef ALLOW_COLLECT_STORAGE_INFORMATION
+#define ALLOW_COLLECT_STORAGE_INFORMATION 0
+#endif
+
 // This might not be defined by Windows when maintaining backwards-compatibility to pre-Win8 builds
 #ifndef SM_CONVERTIBLESLATEMODE
 #define SM_CONVERTIBLESLATEMODE			0x2003
@@ -444,6 +450,8 @@ namespace
 
 	static bool CollectStorageInformation()
 	{
+#if ALLOW_COLLECT_STORAGE_INFORMATION
+
 		IWbemLocator* WbemLocator = nullptr;
 		IWbemServices* WbemServices = nullptr;
 
@@ -680,6 +688,8 @@ namespace
 		FWindowsPlatformMisc::CoUninitialize();
 		FWindowsPlatformMisc::UpdateDriveFreeSpace();
 
+#endif // ALLOW_COLLECT_STORAGE_INFORMATION
+
 		return true;
 	}
 }
@@ -938,6 +948,7 @@ void FWindowsPlatformMisc::PlatformInit()
 	// Register on the game thread.
 	FWindowsPlatformStackWalk::RegisterOnModulesChanged();
 
+	// Collecting storage information using dcom services.
 	CollectStorageInformation();
 }
 
