Logging: Added experimental prepending of log context to log messages

diff --git a/Engine/Source/Runtime/Core/Private/Logging/StructuredLog.cpp b/Engine/Source/Runtime/Core/Private/Logging/StructuredLog.cpp
--- a/Engine/Source/Runtime/Core/Private/Logging/StructuredLog.cpp
+++ b/Engine/Source/Runtime/Core/Private/Logging/StructuredLog.cpp
@@ -12,7 +12,7 @@
 #include "CoreGlobals.h"
 #include "HAL/PlatformMath.h"
 #include "HAL/PlatformTime.h"
-#include "Internationalization/Internationalization.h"
+#include "Internationalization/Text.h"
 #include "Logging/LogTrace.h"
 #include "Logging/StructuredLogFormat.h"
 #include "Misc/AsciiSet.h"
@@ -37,7 +37,12 @@ namespace UE::Logging::Private
 // Temporary override until performance and functionality are sufficient for this to be the default.
 CORE_API bool GConvertBasicLogToLogRecord = false;
 
-static constexpr FAsciiSet ValidLogFieldName("0123456789ABCDEFGHIJKLMNOPQRSTUVWXYZabcdefghijklmnopqrstuvwxyz_");
+// Experimental feature to prepend log context to the log message during formatting.
+CORE_API bool GPrependLogContextToLogMessage = false;
+
+static constexpr FAsciiSet GValidLogFieldName("0123456789ABCDEFGHIJKLMNOPQRSTUVWXYZabcdefghijklmnopqrstuvwxyz_");
+
+static constexpr FAnsiStringView GLogContextsFieldName = ANSITEXTVIEW("$Contexts");
 
 ///////////////////////////////////////////////////////////////////////////////////////////////////////////////////////
 
@@ -276,7 +281,7 @@ FLogTemplate* FLogTemplate::Create(const TCHAR* Format, const FLogField* Fields,
 	checkf(bPositional || Algo::AllOf(FieldsView, &FLogField::Name),
 		TEXT("Log fields must be entirely named or entirely anonymous. [[%s]]"), Format);
 	checkf(bPositional || Algo::AllOf(FieldsView,
-		[](const FLogField& Field) { return *Field.Name && *Field.Name != '_' && FAsciiSet::HasOnly(Field.Name, ValidLogFieldName); }),
+		[](const FLogField& Field) { return *Field.Name && *Field.Name != '_' && FAsciiSet::HasOnly(Field.Name, GValidLogFieldName); }),
 		TEXT("Log field names must match \"[A-Za-z0-9][A-Za-z0-9_]*\" in [[%s]]."), Format);
 
 	TArray<FLogTemplateOp, TInlineAllocator<16>> Ops;
@@ -329,7 +334,7 @@ FLogTemplate* FLogTemplate::Create(const TCHAR* Format, const FLogField* Fields,
 
 		// Field
 		const TCHAR* const FieldStart = TextStart;
-		const TCHAR* const FieldNameEnd = FAsciiSet::Skip(FieldStart + 1, ValidLogFieldName);
+		const TCHAR* const FieldNameEnd = FAsciiSet::Skip(FieldStart + 1, GValidLogFieldName);
 		checkf(*FieldNameEnd, TEXT("Log format has an unterminated field reference. Use '{{' to escape '{' if needed. [[%s]]"), Format);
 		checkf(*FieldNameEnd == TEXT('}'), TEXT("Log format has invalid character '%c' in field name. "
 			"Use '{{' to escape '{' if needed. Names must match \"[A-Za-z0-9][A-Za-z0-9_]*\". [[%s]]"), *FieldNameEnd, Format);
@@ -627,6 +632,29 @@ FORCENOINLINE static void FormatDynamicRecordMessageTo(TStringBuilderBase<CharTy
 template <typename CharType>
 static void FormatRecordMessageTo(TStringBuilderBase<CharType>& Out, const FLogRecord& Record)
 {
+	using namespace Logging::Private;
+
+	if (GPrependLogContextToLogMessage)
+	{
+		const FCbObject& Fields = Record.GetFields();
+		for (FCbFieldView NameField : Fields[GLogContextsFieldName])
+		{
+			const FUtf8StringView NameView = NameField.AsString();
+			if (FCbFieldView ContextField = Fields[NameView])
+			{
+				Out.Append(NameView);
+				if (!ContextField.IsNull())
+				{
+					Out.AppendChar('(');
+					CompactBinaryToCompactJson(ContextField.RemoveName(), Out);
+					Out.AppendChar(')');
+				}
+				Out.AppendChar(':');
+				Out.AppendChar(' ');
+			}
+		}
+	}
+
 	const FLogTemplate* Template = Record.GetTemplate();
 	if (LIKELY(Template))
 	{
@@ -857,7 +885,7 @@ FORCENOINLINE static FCbObject SerializeLogFields(
 		}
 
 		// Traverse contexts forward and build an array of names in $Contexts.
-		Writer.BeginArray(ANSITEXTVIEW("$Contexts"));
+		Writer.BeginArray(GLogContextsFieldName);
 		for (FLogContext* Node = LogContextHead; Node; Node = Node->Next)
 		{
 			if (!Node->bDisabledByNewerContext)
@@ -1154,7 +1182,7 @@ FORCENOINLINE static FCbObject SerializeBasicLogMessage(const FStaticBasicLogRec
 		}
 
 		// Traverse contexts forward and build an array of names in $Contexts.
-		Writer.BeginArray(ANSITEXTVIEW("$Contexts"));
+		Writer.BeginArray(GLogContextsFieldName);
 		for (FLogContext* Node = LogContextHead; Node; Node = Node->Next)
 		{
 			if (!Node->bDisabledByNewerContext)
