Add double type for modf as we hit this from some functions
We only had float, long double covered as that is what the math header we had for libc++ had, but there seems to be other overloads around

diff --git a/Engine/Source/Programs/AutoRTFMTests/Private/Math.cpp b/Engine/Source/Programs/AutoRTFMTests/Private/Math.cpp
--- a/Engine/Source/Programs/AutoRTFMTests/Private/Math.cpp
+++ b/Engine/Source/Programs/AutoRTFMTests/Private/Math.cpp
@@ -270,7 +270,7 @@ TEST_CASE("Math.modf")
 			return std::make_pair(A, B); 
 		});
 	}
-	SECTION("double")
+	SECTION("long double")
 	{
 		Check([]
 		{
@@ -279,6 +279,15 @@ TEST_CASE("Math.modf")
 			return std::make_pair(A, B); 
 		});
 	}
+	SECTION("double")
+	{
+		Check([]
+		{
+			double A = 0;
+			double B = modf(4.2, &A);
+			return std::make_pair(A, B); 
+		});
+	}
 }
 
 TEST_CASE("Math.modff")
diff --git a/Engine/Source/Runtime/Core/Private/AutoRTFM/StdLibSupport.cpp b/Engine/Source/Runtime/Core/Private/AutoRTFM/StdLibSupport.cpp
--- a/Engine/Source/Runtime/Core/Private/AutoRTFM/StdLibSupport.cpp
+++ b/Engine/Source/Runtime/Core/Private/AutoRTFM/StdLibSupport.cpp
@@ -377,6 +377,7 @@ UE_AUTORTFM_REGISTER_SELF_OVERLOADED_FUNCTION(float (float, float), fmod);
 UE_AUTORTFM_REGISTER_SELF_OVERLOADED_FUNCTION(long double (long double, long double), fmod);
 UE_AUTORTFM_REGISTER_SELF_OVERLOADED_FUNCTION(float (float, float*), modf);
 UE_AUTORTFM_REGISTER_SELF_OVERLOADED_FUNCTION(long double (long double, long double*), modf);
+UE_AUTORTFM_REGISTER_SELF_OVERLOADED_FUNCTION(double (double, double*), modf);
 #endif
 
 // Self register Math functions
