remove the deprecated FCustomChunkMapping struct

diff --git a/Engine/Source/Runtime/Core/Public/GenericPlatform/GenericPlatformChunkInstall.h b/Engine/Source/Runtime/Core/Public/GenericPlatform/GenericPlatformChunkInstall.h
--- a/Engine/Source/Runtime/Core/Public/GenericPlatform/GenericPlatformChunkInstall.h
+++ b/Engine/Source/Runtime/Core/Public/GenericPlatform/GenericPlatformChunkInstall.h
@@ -111,23 +111,6 @@ DECLARE_MULTICAST_DELEGATE_TwoParams(FPlatformNamedChunkInstallMultiDelegate, FN
 DECLARE_DELEGATE_OneParam(FPlatformNamedChunkCompleteDelegate, const FNamedChunkCompleteCallbackParam&);
 DECLARE_MULTICAST_DELEGATE_OneParam(FPlatformNamedChunkCompleteMultiDelegate, const FNamedChunkCompleteCallbackParam&);
 
-struct FCustomChunkMapping
-{
-	enum class CustomChunkMappingType : uint8
-	{
-		Main,
-		Optional
-	};
-
-	FString Pattern;
-	uint32	ChunkID;
-	CustomChunkMappingType MappingType;
-
-	FCustomChunkMapping(FString InPattern, uint32 InChunkID, CustomChunkMappingType InMappingType) :
-		Pattern(InPattern), ChunkID(InChunkID), MappingType(InMappingType)
-	{}
-};
-
 enum class ENamedChunkType : uint8
 {
 	Invalid,
