[AutoRTFM] Update the activity tracker to send its events in the open.

This callback mechanism is used by the Windows console, which isn't AutoRTFM-safe.

diff --git a/Engine/Source/Runtime/Core/Private/Misc/TrackedActivity.cpp b/Engine/Source/Runtime/Core/Private/Misc/TrackedActivity.cpp
--- a/Engine/Source/Runtime/Core/Private/Misc/TrackedActivity.cpp
+++ b/Engine/Source/Runtime/Core/Private/Misc/TrackedActivity.cpp
@@ -1,6 +1,7 @@
 // Copyright Epic Games, Inc. All Rights Reserved.
 
 #include "Misc/TrackedActivity.h"
+#include "AutoRTFM/AutoRTFM.h"
 #include "Containers/UnrealString.h"
 #include "Misc/ScopeLock.h"
 #include "Misc/StringBuilder.h"
@@ -162,7 +163,14 @@ public:
 		if (!EventFunc)
 			return;
 		TStringBuilder<256> Temp;
-		EventFunc(Event, GetInfo(Activity, EventFuncMaxDepth, Temp));
+		FTrackedActivity::FInfo ActivityInfo{GetInfo(Activity, EventFuncMaxDepth, Temp)};
+		
+		// This callback mechanism is used to update the console output windows. These are wildly AutoRTFM-unsafe
+		// (since they are manipulating HWNDs directly, etc.), so we always run the callback in the open.
+		UE_AUTORTFM_OPEN
+		{
+			EventFunc(Event, ActivityInfo);
+		};
 	}
 };
 
