[core] turn on FString::Appendf compile-time format string checking

diff --git a/Engine/Source/Runtime/Core/Public/Containers/UnrealString.h.inl b/Engine/Source/Runtime/Core/Public/Containers/UnrealString.h.inl
--- a/Engine/Source/Runtime/Core/Public/Containers/UnrealString.h.inl
+++ b/Engine/Source/Runtime/Core/Public/Containers/UnrealString.h.inl
@@ -1574,11 +1574,22 @@ public:
 	 * @returns A string that was constructed using format and additional parameters.
 	 */
 	template <typename... Types>
-	[[nodiscard]] static UE_STRING_CLASS Printf(::UE::Core::TCheckedFormatString<FmtCharType, Types...> Fmt, Types... Args)
+	[[nodiscard]] static UE_STRING_CLASS Printf(UE::Core::TCheckedFormatString<FmtCharType, Types...> Fmt, Types... Args)
 	{
 		return PrintfImpl((const ElementType*)Fmt.FormatString, Args...);
 	}
 
+	/**
+	 * Just like Printf, but appends the formatted text to the existing string instead.
+	 * @return a reference to the modified string, so that it can be chained
+	 */
+	template <typename... Types>
+	UE_STRING_CLASS& Appendf(UE::Core::TCheckedFormatString<FmtCharType, Types...> Fmt, Types... Args)
+	{
+		AppendfImpl(*this, (const ElementType*)Fmt.FormatString, Args...);
+		return *this;
+	}
+
 #else
 
 	template <typename FmtType, typename... Types>
@@ -1590,12 +1601,6 @@ public:
 		return PrintfImpl((const ElementType*)Fmt, Args...);
 	}
 
-#endif // UE_VALIDATE_FORMAT_STRINGS
-
-	/**
-	 * Just like Printf, but appends the formatted text to the existing string instead.
-	 * @return a reference to the modified string, so that it can be chained
-	 */
 	template <typename FmtType, typename... Types>
 	UE_STRING_CLASS& Appendf(const FmtType& Fmt, Types... Args)
 	{
@@ -1606,6 +1611,8 @@ public:
 		return *this;
 	}
 
+#endif // UE_VALIDATE_FORMAT_STRINGS
+
 private:
 	static CORE_API UE_STRING_CLASS VARARGS PrintfImpl(const ElementType* Fmt, ...);
 	static CORE_API void VARARGS AppendfImpl(UE_STRING_CLASS& AppendToMe, const ElementType* Fmt, ...);
