Fix Unreal Insight Bookmarks not showing up in case the bookmark channel is not yet enabled but the OutputBookmarkSpec, in TRACE_BOOKMARK, is called

diff --git a/Engine/Source/Runtime/Core/Private/ProfilingDebugging/MiscTrace.cpp b/Engine/Source/Runtime/Core/Private/ProfilingDebugging/MiscTrace.cpp
--- a/Engine/Source/Runtime/Core/Private/ProfilingDebugging/MiscTrace.cpp
+++ b/Engine/Source/Runtime/Core/Private/ProfilingDebugging/MiscTrace.cpp
@@ -9,7 +9,7 @@
 #include "HAL/PlatformTime.h"
 
 UE_TRACE_CHANNEL(FrameChannel)
-UE_TRACE_CHANNEL(BookmarkChannel)
+UE_TRACE_CHANNEL_DEFINE(BookmarkChannel)
 UE_TRACE_CHANNEL(RegionChannel)
 UE_TRACE_CHANNEL(ScreenshotChannel)
 
diff --git a/Engine/Source/Runtime/Core/Public/ProfilingDebugging/MiscTrace.h b/Engine/Source/Runtime/Core/Public/ProfilingDebugging/MiscTrace.h
--- a/Engine/Source/Runtime/Core/Public/ProfilingDebugging/MiscTrace.h
+++ b/Engine/Source/Runtime/Core/Public/ProfilingDebugging/MiscTrace.h
@@ -17,6 +17,8 @@
 #endif
 #endif
 
+UE_TRACE_CHANNEL_EXTERN(BookmarkChannel, CORE_API)
+
 enum ETraceFrameType
 {
 	TraceFrameType_Game,
@@ -92,6 +94,8 @@ private:
 };
 
 #define TRACE_BOOKMARK(Format, ...) \
+if (UE_TRACE_CHANNELEXPR_IS_ENABLED(BookmarkChannel)) \
+{ \
 	static bool PREPROCESSOR_JOIN(__BookmarkPoint, __LINE__); \
 	if (!PREPROCESSOR_JOIN(__BookmarkPoint, __LINE__)) \
 	{ \
@@ -101,9 +105,12 @@ private:
 		FMiscTrace::OutputBookmarkSpec(&PREPROCESSOR_JOIN(__BookmarkPoint, __LINE__), __FILE__, __LINE__, Format); \
 		PREPROCESSOR_JOIN(__BookmarkPoint, __LINE__) = true; \
 	} \
-	FMiscTrace::OutputBookmark(&PREPROCESSOR_JOIN(__BookmarkPoint, __LINE__), ##__VA_ARGS__);
+	FMiscTrace::OutputBookmark(&PREPROCESSOR_JOIN(__BookmarkPoint, __LINE__), ##__VA_ARGS__); \
+}
 
 #define TRACE_BOOKMARK_CYCLES(Cycles, Format, ...) \
+if (UE_TRACE_CHANNELEXPR_IS_ENABLED(BookmarkChannel)) \
+{ \
 	static bool PREPROCESSOR_JOIN(__BookmarkPoint, __LINE__); \
 	if (!PREPROCESSOR_JOIN(__BookmarkPoint, __LINE__)) \
 	{ \
@@ -113,7 +120,8 @@ private:
 		FMiscTrace::OutputBookmarkSpec(&PREPROCESSOR_JOIN(__BookmarkPoint, __LINE__), __FILE__, __LINE__, Format); \
 		PREPROCESSOR_JOIN(__BookmarkPoint, __LINE__) = true; \
 	} \
-	FMiscTrace::OutputBookmarkCycles(Cycles, &PREPROCESSOR_JOIN(__BookmarkPoint, __LINE__), ##__VA_ARGS__);
+	FMiscTrace::OutputBookmarkCycles(Cycles, &PREPROCESSOR_JOIN(__BookmarkPoint, __LINE__), ##__VA_ARGS__); \
+}
 
 #define TRACE_BEGIN_REGION(RegionName) \
 	FMiscTrace::OutputBeginRegion(RegionName);
