Fix cocoaThread causing UE fails to exit normally

#rnx

diff --git a/Engine/Source/Runtime/Core/Private/Mac/CocoaThread.cpp b/Engine/Source/Runtime/Core/Private/Mac/CocoaThread.cpp
--- a/Engine/Source/Runtime/Core/Private/Mac/CocoaThread.cpp
+++ b/Engine/Source/Runtime/Core/Private/Mac/CocoaThread.cpp
@@ -385,6 +385,7 @@ FCocoaRunLoopSource* FCocoaRunLoopSource::GameRunLoopSource = nullptr;
 	// We have exited the game thread, so any UE code running now should treat the Main thread
 	// as the game thread, so we don't crash in static destructors.
 	GGameThreadId = GMainThreadId;
+	GCocoaGameThread = nil;
 	
 	// Tell the main thread we are OK to quit, but don't wait for it.
 	if (IsEngineExitRequested())
