Make a cvar that forces a crash on an AutoRTFM hit

diff --git a/Engine/Source/Runtime/Core/Private/AutoRTFM/API.cpp b/Engine/Source/Runtime/Core/Private/AutoRTFM/API.cpp
--- a/Engine/Source/Runtime/Core/Private/AutoRTFM/API.cpp
+++ b/Engine/Source/Runtime/Core/Private/AutoRTFM/API.cpp
@@ -33,6 +33,8 @@ namespace
 
 	int GMemoryValidationLevel = static_cast<int>(GDefaultMemoryValidationLevel);
 
+	bool GAutoRTFMAssertOnError = false;
+
 	void UpdateAutoRTFMRetryTransactionsData()
 	{
 		switch (GAutoRTFMRetryTransactions)
@@ -85,6 +87,13 @@ static FAutoConsoleVariableRef CVarAutoRTFMMemoryValidationLevel(
 	ECVF_Default
 );
 
+static FAutoConsoleVariableRef CVarAutoRTFMAssertOnError(
+	TEXT("AutoRTFMAssertOnError"),
+	GAutoRTFMAssertOnError,
+	TEXT("If true when we hit an AutoRTFM issue assert over ensuring"),
+	ECVF_Default
+);
+
 static FDelayedAutoRegisterHelper DelayedAutoRegister(EDelayedRegisterRunPhase::EndOfEngineInit, []
 	{
 		UpdateAutoRTFMRuntimeCrashData();
@@ -164,6 +173,15 @@ namespace AutoRTFM
 #endif
 		}
 
+		bool IsAutoRTFMAssertOnError()
+		{
+#if UE_AUTORTFM
+			return GAutoRTFMAssertOnError;
+#else
+			return false;
+#endif
+		}
+
 		void SetRetryTransaction(EAutoRTFMRetryTransactionState State)
 		{
 #if UE_AUTORTFM
@@ -332,7 +350,14 @@ AUTORTFM_CORE_API autortfm_status RTFM_autortfm_close(void (*UninstrumentedWork)
 	}
 	else
 	{
-		ensureMsgf(!ForTheRuntime::IsEnsureOnAbortByLanguageEnabled(), TEXT("Could not find function %p '%s' where '%s'."), UninstrumentedWork, *GetFunctionDescription(UninstrumentedWork), ANSI_TO_TCHAR("autortfm_close"));
+		if (ForTheRuntime::IsAutoRTFMAssertOnError())
+		{
+			UE_LOG(LogAutoRTFM, Fatal, TEXT("Could not find function %p '%s' where '%s'."), UninstrumentedWork, *GetFunctionDescription(UninstrumentedWork), ANSI_TO_TCHAR("autortfm_close"));
+		}
+		else
+		{
+			ensureMsgf(!ForTheRuntime::IsEnsureOnAbortByLanguageEnabled(), TEXT("Could not find function %p '%s' where '%s'."), UninstrumentedWork, *GetFunctionDescription(UninstrumentedWork), ANSI_TO_TCHAR("autortfm_close"));
+		}
 		Context->AbortByLanguageAndThrow();
 	}
 
@@ -658,7 +683,14 @@ extern "C" autortfm_status autortfm_close(void (*UninstrumentedWork)(void*), voi
 		}
 		else
 		{
-			ensureMsgf(!ForTheRuntime::IsEnsureOnAbortByLanguageEnabled(), TEXT("Could not find function %p '%s' where '%s'."), UninstrumentedWork, *GetFunctionDescription(UninstrumentedWork), ANSI_TO_TCHAR("autortfm_close"));
+			if (ForTheRuntime::IsAutoRTFMAssertOnError())
+			{
+				UE_LOG(LogAutoRTFM, Fatal, TEXT("Could not find function %p '%s' where '%s'."), UninstrumentedWork, *GetFunctionDescription(UninstrumentedWork), ANSI_TO_TCHAR("autortfm_close"));
+			}
+			else
+			{
+				ensureMsgf(!ForTheRuntime::IsEnsureOnAbortByLanguageEnabled(), TEXT("Could not find function %p '%s' where '%s'."), UninstrumentedWork, *GetFunctionDescription(UninstrumentedWork), ANSI_TO_TCHAR("autortfm_close"));
+			}
 	        Context->AbortByLanguageAndThrow();
 		}
 	}
diff --git a/Engine/Source/Runtime/Core/Private/AutoRTFM/CompilerInterface.cpp b/Engine/Source/Runtime/Core/Private/AutoRTFM/CompilerInterface.cpp
--- a/Engine/Source/Runtime/Core/Private/AutoRTFM/CompilerInterface.cpp
+++ b/Engine/Source/Runtime/Core/Private/AutoRTFM/CompilerInterface.cpp
@@ -144,11 +144,25 @@ extern "C" UE_AUTORTFM_API void autortfm_llvm_fail(const char* Message)
 {
     if (Message)
     {
-		ensureMsgf(!ForTheRuntime::IsEnsureOnAbortByLanguageEnabled(), TEXT("Transaction failing because of language issue '%s'."), ANSI_TO_TCHAR(Message));
+		if (ForTheRuntime::IsAutoRTFMAssertOnError())
+		{
+			UE_LOG(LogAutoRTFM, Fatal, TEXT("Transaction failing because of language issue '%s'."), ANSI_TO_TCHAR(Message));
+		}
+		else
+		{
+			ensureMsgf(!ForTheRuntime::IsEnsureOnAbortByLanguageEnabled(), TEXT("Transaction failing because of language issue '%s'."), ANSI_TO_TCHAR(Message));
+		}
     }
     else
     {
-		ensureMsgf(!ForTheRuntime::IsEnsureOnAbortByLanguageEnabled(), TEXT("Transaction failing because of language issue."));
+		if (ForTheRuntime::IsAutoRTFMAssertOnError())
+		{
+			UE_LOG(LogAutoRTFM, Fatal, TEXT("Transaction failing because of language issue."));
+		}
+		else
+		{
+			ensureMsgf(!ForTheRuntime::IsEnsureOnAbortByLanguageEnabled(), TEXT("Transaction failing because of language issue."));
+		}
 	}
 
 	FContext* Context = FContext::Get();
@@ -157,7 +171,14 @@ extern "C" UE_AUTORTFM_API void autortfm_llvm_fail(const char* Message)
 
 extern "C" UE_AUTORTFM_API void autortfm_llvm_missing_function()
 {
-	ensureMsgf(!ForTheRuntime::IsEnsureOnAbortByLanguageEnabled(), TEXT("Transaction failing because of missing function."));
+	if (ForTheRuntime::IsAutoRTFMAssertOnError())
+	{
+		UE_LOG(LogAutoRTFM, Fatal, TEXT("Transaction failing because of missing function."));
+	}
+	else
+	{
+		ensureMsgf(!ForTheRuntime::IsEnsureOnAbortByLanguageEnabled(), TEXT("Transaction failing because of missing function."));
+	}
 
 	FContext* Context = FContext::Get();
     Context->AbortByLanguageAndThrow();
diff --git a/Engine/Source/Runtime/Core/Private/AutoRTFM/FunctionMap.cpp b/Engine/Source/Runtime/Core/Private/AutoRTFM/FunctionMap.cpp
--- a/Engine/Source/Runtime/Core/Private/AutoRTFM/FunctionMap.cpp
+++ b/Engine/Source/Runtime/Core/Private/AutoRTFM/FunctionMap.cpp
@@ -110,11 +110,25 @@ void* FunctionMapReportError(void* OldFunction, const char* Where)
 {
 	if (Where)
 	{
-		ensureMsgf(!ForTheRuntime::IsEnsureOnAbortByLanguageEnabled(), TEXT("Could not find function %p '%s' where '%s'."), OldFunction, *GetFunctionDescription(OldFunction), ANSI_TO_TCHAR(Where));
+		if (ForTheRuntime::IsAutoRTFMAssertOnError())
+		{
+			UE_LOG(LogAutoRTFM, Fatal, TEXT("Could not find function %p '%s' where '%s'."), OldFunction, *GetFunctionDescription(OldFunction), ANSI_TO_TCHAR(Where));
+		}
+		else
+		{
+			ensureMsgf(!ForTheRuntime::IsEnsureOnAbortByLanguageEnabled(), TEXT("Could not find function %p '%s' where '%s'."), OldFunction, *GetFunctionDescription(OldFunction), ANSI_TO_TCHAR(Where));
+		}
 	}
 	else
 	{
-		ensureMsgf(!ForTheRuntime::IsEnsureOnAbortByLanguageEnabled(), TEXT("Could not find function %p '%s'."), OldFunction, *GetFunctionDescription(OldFunction));
+		if (ForTheRuntime::IsAutoRTFMAssertOnError())
+		{
+			UE_LOG(LogAutoRTFM, Fatal, TEXT("Could not find function %p '%s'."), OldFunction, *GetFunctionDescription(OldFunction));
+		}
+		else
+		{
+			ensureMsgf(!ForTheRuntime::IsEnsureOnAbortByLanguageEnabled(), TEXT("Could not find function %p '%s'."), OldFunction, *GetFunctionDescription(OldFunction));
+		}
 	}
 
 	FContext* Context = FContext::Get();
diff --git a/Engine/Source/Runtime/Core/Public/AutoRTFM/AutoRTFM.h b/Engine/Source/Runtime/Core/Public/AutoRTFM/AutoRTFM.h
--- a/Engine/Source/Runtime/Core/Public/AutoRTFM/AutoRTFM.h
+++ b/Engine/Source/Runtime/Core/Public/AutoRTFM/AutoRTFM.h
@@ -973,6 +973,9 @@ namespace ForTheRuntime
 	// Returns whether the runtime will trigger an ensure on an abort-by-language, or not.
 	UE_AUTORTFM_API bool IsEnsureOnAbortByLanguageEnabled();
 
+	// Returns whether we want to assert or ensure on a Language Error
+	UE_AUTORTFM_API bool IsAutoRTFMAssertOnError();
+
 	// Set whether we should retry transactions.
 	UE_AUTORTFM_API void SetRetryTransaction(EAutoRTFMRetryTransactionState State);
 
