[AutoRTFM] Refactor and tighten transaction state tracking

More accurately track transaction state W.R.T. deferred transactions. Fix tests that now correctly report memory validation failures.

Assert if attempting to transition a transation to its existing state - this is usually a sign that things are not right.

Split the public FTransaction::SetState() method into a number of SetXXX() methods, so that SetOpenActive() can take the memory valdiation level. This is only relevant to the open-active state, and removes the rather icky 'PendingMemoryValidationLevel' member.

Also add DebugBreakIfMemoryValidationFails(), useful for isolating where the memory validator write hash changes.

diff --git a/Engine/Source/Programs/AutoRTFMTests/Private/OpenAPI.cpp b/Engine/Source/Programs/AutoRTFMTests/Private/OpenAPI.cpp
--- a/Engine/Source/Programs/AutoRTFMTests/Private/OpenAPI.cpp
+++ b/Engine/Source/Programs/AutoRTFMTests/Private/OpenAPI.cpp
@@ -11,6 +11,9 @@
 
 TEST_CASE("OpenAPI.StartAbortAndStartAgain")
 {
+    AUTORTFM_SCOPED_ENABLE_MEMORY_VALIDATION_AS_WARNING();
+    AutoRTFMTestUtils::FCaptureWarningContext WarningContext;
+
 	int valueB = 0;
 	int valueC = 0;
 	AutoRTFM::Transact([&]()
@@ -37,6 +40,9 @@ TEST_CASE("OpenAPI.StartAbortAndStartAgain")
 	// We rollback the transaction to the value we had when we first recorded the address
 	REQUIRE(valueB == 20);
 	REQUIRE(valueC == 0);
+
+	REQUIRE_THAT(WarningContext.GetWarnings(),
+		Catch::Matchers::VectorContains(FString(AutoRTFMTestUtils::kMemoryModifiedWarning)));
 }
 
 TEST_CASE("OpenAPI.CommitScopedFromOpen_Illegal", "[.]")
@@ -891,6 +897,8 @@ TEST_CASE("OpenAPI.TransOpenTransCloseAbortAbort")
 
 TEST_CASE("OpenAPI.DeferredStartTransactionOverflow")
 {
+    AUTORTFM_SCOPED_DISABLE_MEMORY_VALIDATION(); // Avoid stack overflow in hashing.
+
 	AutoRTFM::TransactThenOpen([&]
 	{
 		uint64 NestingCount = static_cast<uint64>(std::numeric_limits<uint16>::max()) + 42;
@@ -940,4 +948,4 @@ TEST_CASE("OpenAPI.CheckRaceAgainstOtherThread")
 
 	REQUIRE(AutoRTFM::ETransactionResult::AbortedByRequest == Result);
 	REQUIRE(bHit);
-}
\ No newline at end of file
+}
diff --git a/Engine/Source/Runtime/Core/Private/AutoRTFM/API.cpp b/Engine/Source/Runtime/Core/Private/AutoRTFM/API.cpp
--- a/Engine/Source/Runtime/Core/Private/AutoRTFM/API.cpp
+++ b/Engine/Source/Runtime/Core/Private/AutoRTFM/API.cpp
@@ -1,6 +1,7 @@
 // Copyright Epic Games, Inc. All Rights Reserved.
 
 #include "AutoRTFM/AutoRTFM.h"
+#include "AutoRTFM/Context.h"
 #include "CoreGlobals.h"
 #include "HAL/IConsoleManager.h"
 #include "GenericPlatform/GenericPlatformCrashContext.h"
@@ -247,6 +248,19 @@ namespace AutoRTFM
 #endif
 		}
 
+		UE_AUTORTFM_ALWAYS_OPEN void DebugBreakIfMemoryValidationFails()
+		{
+#if UE_AUTORTFM
+			// Check memory validation is enabled, otherwise we won't have a
+			// hash to compare against.
+			ASSERT(GetMemoryValidationLevel() == EMemoryValidationLevel::Enabled ||
+				GetMemoryValidationLevel() == EMemoryValidationLevel::EnabledAsWarning);
+			if (FTransaction* Transaction = AutoRTFM::FContext::Get()->GetMaterializedTransaction()) {
+				Transaction->DebugBreakIfMemoryValidationFails();
+			}
+#endif
+		}
+
 	}
 }
 
@@ -368,8 +382,7 @@ extern "C" UE_AUTORTFM_NOAUTORTFM UE_AUTORTFM_API void autortfm_pre_open()
 {
 	if (FTransaction* Transaction = FContext::Get()->GetCurrentTransaction())
 	{
-		Transaction->SetMemoryValidationLevel(static_cast<AutoRTFM::EMemoryValidationLevel>(GMemoryValidationLevel));
-		Transaction->SetState(FTransaction::EState::OpenActive);
+		Transaction->SetOpenActive(ForTheRuntime::GetMemoryValidationLevel());
 	}
 }
 
@@ -377,9 +390,9 @@ extern "C" UE_AUTORTFM_NOAUTORTFM UE_AUTORTFM_API void autortfm_post_open()
 {
 	if (FTransaction* Transaction = FContext::Get()->GetCurrentTransaction())
 	{
-		if (!Transaction->IsDone())  // Transaction may have been aborted.
+		if (Transaction->IsOpenActive())  // Transaction may have been aborted.
 		{
-			Transaction->SetState(FTransaction::EState::ClosedActive);
+			Transaction->SetClosedActive();
 		}
 	}
 }
@@ -391,18 +404,17 @@ AUTORTFM_CORE_API void RTFM_autortfm_open(autortfm_memory_validation_level Valid
 	{
 		AutoRTFM::EMemoryValidationLevel Level 
 			= ValidationLevel == autortfm_memory_validation_level_default 
-			? static_cast<AutoRTFM::EMemoryValidationLevel>(GMemoryValidationLevel)
+			? ForTheRuntime::GetMemoryValidationLevel()
 			: static_cast<AutoRTFM::EMemoryValidationLevel>(ValidationLevel);
-		Transaction->SetMemoryValidationLevel(Level);
-		Transaction->SetState(FTransaction::EState::OpenActive);
+		Transaction->SetOpenActive(Level);
 	}
 
 	Work(Arg);
 
 	FContext* Context = FContext::Get();
-	if (Transaction == Context->GetCurrentTransaction() && !Transaction->IsDone())  // Transaction may have been aborted.
+	if (Transaction == Context->GetCurrentTransaction() && Transaction->IsOpenActive())  // Transaction may have been aborted.
 	{
-		Transaction->SetState(FTransaction::EState::ClosedActive);
+		Transaction->SetClosedActive();
 	}
 
 	if (Context->IsAborting())
@@ -576,7 +588,8 @@ extern "C" bool autortfm_start_transaction()
 	}
 
 	UE_CLOG(!FContext::IsTransactional(), LogAutoRTFM, Fatal, TEXT("The function `autortfm_start_transaction` was called from outside a transact."));
-	return FContext::Get()->StartTransaction();
+	FContext::Get()->StartTransaction(ForTheRuntime::GetMemoryValidationLevel());
+	return true;
 }
 
 extern "C" autortfm_result autortfm_commit_transaction()
diff --git a/Engine/Source/Runtime/Core/Private/AutoRTFM/Context.cpp b/Engine/Source/Runtime/Core/Private/AutoRTFM/Context.cpp
--- a/Engine/Source/Runtime/Core/Private/AutoRTFM/Context.cpp
+++ b/Engine/Source/Runtime/Core/Private/AutoRTFM/Context.cpp
@@ -71,22 +71,27 @@ void FContext::MaterializeDeferredTransactions()
 	NumDeferredTransactions = 0;
 	for (uint64 I = 0; I < NumToAllocate; ++I)
 	{
-		StartNonDeferredTransaction();
+		StartNonDeferredTransaction(EMemoryValidationLevel::Disabled);
 	}
 }
 
-bool FContext::StartTransaction()
+void FContext::StartTransaction(EMemoryValidationLevel MemoryValidationLevel)
 {
+	if (MemoryValidationLevel != EMemoryValidationLevel::Disabled)
+	{
+		MaterializeDeferredTransactions();
+		StartNonDeferredTransaction(MemoryValidationLevel);
+		return;
+	}
+
 	NumDeferredTransactions += 1;
 	ensureMsgf(CurrentTransaction, TEXT("FContext::StartTransaction() can only be called within a scoped transaction"));
 	ASSERT(Status == EContextStatus::OnTrack);
 
 	GAutoRTFMMetrics.NumTransactionsStarted++;
-
-	return true;
 }
 
-void FContext::StartNonDeferredTransaction()
+void FContext::StartNonDeferredTransaction(EMemoryValidationLevel MemoryValidationLevel)
 {
 	ASSERT(GetNumDeferredTransactions() == 0);
 	ensureMsgf(CurrentTransaction, TEXT("FContext::StartNonDeferredTransaction() can only be called within a scoped transaction"));
@@ -94,7 +99,8 @@ void FContext::StartNonDeferredTransaction()
 	FTransaction* NewTransaction = PushTransaction(
 		/* Closed */ false,
 		/* bIsScoped */ false,
-		/* StackRange */ CurrentTransaction->GetStackRange());
+		/* StackRange */ CurrentTransaction->GetStackRange(),
+		/* MemoryValidationLevel */ MemoryValidationLevel);
 
 	// This form of transaction is always ultimately within a scoped Transact 
 	ASSERT(Status == EContextStatus::OnTrack);
@@ -211,7 +217,9 @@ EContextStatus FContext::CallClosedNest(void (*ClosedFunction)(void* Arg), void*
 
 	FTransaction* const Transaction = GetCurrentTransaction();
 	ASSERT(Transaction != nullptr);
-	FTransaction::EState PreviousState = Transaction->SetState(FTransaction::EState::ClosedActive);
+	ASSERT(FTransaction::EState::OpenActive == Transaction->State());
+	EMemoryValidationLevel PreviousValidationLevel = Transaction->MemoryValidationLevel();
+	Transaction->SetClosedActive();
 
 	PushCallNest(new FCallNest(this));
 
@@ -219,9 +227,9 @@ EContextStatus FContext::CallClosedNest(void (*ClosedFunction)(void* Arg), void*
 
 	PopCallNest();
 
-	if (Transaction == CurrentTransaction && !Transaction->IsDone())  // Transaction may have been aborted.
+	if (Transaction == CurrentTransaction && Transaction->IsClosedActive())  // Transaction may have been aborted.
 	{
-		Transaction->SetState(PreviousState);
+		Transaction->SetOpenActive(PreviousValidationLevel);
 	}
 
 	return GetStatus();
@@ -245,7 +253,11 @@ void FContext::PopCallNest()
 	delete OldCallNest;
 }
 
-FTransaction* FContext::PushTransaction(bool bClosed, bool bIsScoped, FStackRange StackRange)
+FTransaction* FContext::PushTransaction(
+	bool bClosed,
+	bool bIsScoped,
+	FStackRange StackRange,
+	EMemoryValidationLevel MemoryValidationLevel)
 {
 	ASSERT(!GetNumDeferredTransactions());
 
@@ -257,11 +269,19 @@ FTransaction* FContext::PushTransaction(bool bClosed, bool bIsScoped, FStackRang
 
 	FTransaction* NewTransaction = new FTransaction(this);
 	NewTransaction->Initialize(
-		/* State */ bClosed ? FTransaction::EState::ClosedActive : FTransaction::EState::OpenActive,
 		/* Parent */ CurrentTransaction, 
 		/* bIsScoped */ bIsScoped, 
 		/* StackRange */ StackRange);
 
+	if (bClosed)
+	{
+		NewTransaction->SetClosedActive();
+	}
+	else
+	{
+		NewTransaction->SetOpenActive(MemoryValidationLevel);
+	}
+
 	CurrentTransaction = NewTransaction;
 
 	// Collect stats that we've got a new transaction.
@@ -397,7 +417,8 @@ ETransactionResult FContext::Transact(void (*InstrumentedFunction)(void*), void*
 		FTransaction* NewTransaction = PushTransaction(
 			/* Closed */ true,
 			/* bIsScoped */ true,
-			/* StackRange */ {Stack.Low, &TransactStackStart});
+			/* StackRange */ {Stack.Low, &TransactStackStart},
+			/* MemoryValidationLevel */ EMemoryValidationLevel::Disabled);
 
 		PushCallNest(NewNest);
 
@@ -411,8 +432,9 @@ ETransactionResult FContext::Transact(void (*InstrumentedFunction)(void*), void*
 			ASSERT(CurrentTransaction == NewTransaction); // The transaction lambda should have unwound any nested transactions.
             ASSERT(Status != EContextStatus::Idle);
 
-            if (Status == EContextStatus::OnTrack)
-            {
+			switch (Status)
+			{
+			case EContextStatus::OnTrack:
 				UE_LOG(LogAutoRTFM, Verbose, TEXT("About to commit; my state is:"));
 				DumpState();
 				UE_LOG(LogAutoRTFM, Verbose, TEXT("Committing..."));
@@ -426,6 +448,7 @@ ETransactionResult FContext::Transact(void (*InstrumentedFunction)(void*), void*
 
 					// We've tried to run at least once if we get here!
 					CurrentTransaction->Reset();
+					CurrentTransaction->SetClosedActive();
 					bTriedToRunOnce = true;
 					continue;
 				}
@@ -438,32 +461,36 @@ ETransactionResult FContext::Transact(void (*InstrumentedFunction)(void*), void*
 
 				UE_LOG(LogAutoRTFM, Verbose, TEXT("Commit failed!"));
 
-                ASSERT(Status != EContextStatus::OnTrack);
-                ASSERT(Status != EContextStatus::Idle);
-            }
+				ASSERT(Status != EContextStatus::OnTrack);
+				ASSERT(Status != EContextStatus::Idle);
+				break;
+
+			case EContextStatus::AbortedByRequest:
+				Result = ETransactionResult::AbortedByRequest;
+				break;
+
+			case EContextStatus::AbortedByLanguage:
+				Result = ETransactionResult::AbortedByLanguage;
+				break;
 
-            if (Status == EContextStatus::AbortedByRequest)
-            {
-                Result = ETransactionResult::AbortedByRequest;
-                break;
-            }
+			case EContextStatus::AbortedByCascade:
+				Result = ETransactionResult::AbortedByCascade;
+				break;
 
-            if (Status == EContextStatus::AbortedByLanguage)
-            {
-                Result = ETransactionResult::AbortedByLanguage;
-                break;
-            }
+			case EContextStatus::AbortedByFailedLockAcquisition:
+				continue; // Retry the transaction
 
-            if (Status == EContextStatus::AbortedByCascade)
-            {
-                Result = ETransactionResult::AbortedByCascade;
-                break;
-            }
+			default:
+				Unreachable();
+			}
 
-            ASSERT(Status == EContextStatus::AbortedByFailedLockAcquisition);
-        }
+			break;
+		}
 
-		NewTransaction->SetState(FTransaction::EState::Done);
+		if (!NewTransaction->IsDone())
+		{
+			NewTransaction->SetDone();
+		}
 
 		PopCallNest();
 		PopTransaction();
@@ -486,7 +513,8 @@ ETransactionResult FContext::Transact(void (*InstrumentedFunction)(void*), void*
 		FTransaction* NewTransaction = PushTransaction(
 			/* Closed */ true,
 			/* bIsScoped */ true,
-			/* StackRange */ {Stack.Low, &TransactStackStart});
+			/* StackRange */ {Stack.Low, &TransactStackStart},
+			/* MemoryValidationLevel */ EMemoryValidationLevel::Disabled);
 
 		PushCallNest(NewNest);
 
@@ -508,6 +536,7 @@ ETransactionResult FContext::Transact(void (*InstrumentedFunction)(void*), void*
 
 					// We've tried to run at least once if we get here!
 					CurrentTransaction->Reset();
+					CurrentTransaction->SetClosedActive();
 					bTriedToRunOnce = true;
 					continue;
 				}
diff --git a/Engine/Source/Runtime/Core/Private/AutoRTFM/Context.h b/Engine/Source/Runtime/Core/Private/AutoRTFM/Context.h
--- a/Engine/Source/Runtime/Core/Private/AutoRTFM/Context.h
+++ b/Engine/Source/Runtime/Core/Private/AutoRTFM/Context.h
@@ -30,7 +30,7 @@ public:
 	void AbortByRequestWithoutThrowing();
 
 	// Open API - no throw
-	bool StartTransaction();
+	void StartTransaction(EMemoryValidationLevel MemoryValidationLevel);
 
 	ETransactionResult CommitTransaction();
 	ETransactionResult AbortTransaction(bool bIsClosed, bool bIsCascading);
@@ -64,6 +64,10 @@ public:
 		}
 		return CurrentTransaction;
 	}
+	inline FTransaction* GetMaterializedTransaction()
+	{
+		return CurrentTransaction;
+	}
 	inline FTransaction* GetCurrentTransactionAssumingNoMaterialization()
 	{
 		ASSERT(!MustMaterializeDeferredTransactions());
@@ -91,12 +95,12 @@ private:
     FContext(const FContext&) = delete;
 
 	void MaterializeDeferredTransactions();
-	void StartNonDeferredTransaction();
+	void StartNonDeferredTransaction(EMemoryValidationLevel MemoryValidationLevel);
 
 	void PushCallNest(FCallNest* NewCallNest);
 	void PopCallNest();
 
-	FTransaction* PushTransaction(bool bClosed, bool bIsScoped, FStackRange StackRange);
+	FTransaction* PushTransaction(bool bClosed, bool bIsScoped, FStackRange, EMemoryValidationLevel);
 	void PopTransaction();
 
 	ETransactionResult ResolveNestedTransaction(FTransaction* NewTransaction);
diff --git a/Engine/Source/Runtime/Core/Private/AutoRTFM/Transaction.cpp b/Engine/Source/Runtime/Core/Private/AutoRTFM/Transaction.cpp
--- a/Engine/Source/Runtime/Core/Private/AutoRTFM/Transaction.cpp
+++ b/Engine/Source/Runtime/Core/Private/AutoRTFM/Transaction.cpp
@@ -30,12 +30,10 @@ FTransaction::~FTransaction()
 	ASSERT(IsDone());
 }
 
-void FTransaction::Initialize(EState State, FTransaction* Parent_, bool bIsStackScoped_, FStackRange StackRange_)
+void FTransaction::Initialize(FTransaction* Parent_, bool bIsStackScoped_, FStackRange StackRange_)
 {
 	ASSERT(CurrentState == EState::Uninitialized);
-	ASSERT(State == EState::ClosedActive || State == EState::OpenActive);
 
-	CurrentState = State;
 	Parent = Parent_;
 	bIsStackScoped = bIsStackScoped_;
 	StackRange = StackRange_;
@@ -52,11 +50,18 @@ void FTransaction::Initialize(EState State, FTransaction* Parent_, bool bIsStack
 
 void FTransaction::Reset()
 {
-	// Complete the transaction, which will reset all members.
-	SetState(EState::Done);
+	ASSERT(IsDone());
+
+	CommitTasks.Reset();
+	AbortTasks.Reset();
+	HitSet.Reset();
+	NewMemoryTracker.Reset();
+	WriteLog.Reset();
+	WriteLogBumpAllocator.Reset();
+	CurrentMemoryValidationLevel = EMemoryValidationLevel::Disabled;
 
 	// Reset to the initial state.
-	CurrentState = EState::ClosedActive;
+	CurrentState = EState::Uninitialized;
 
 	ASSERT(IsFresh());
 }
@@ -80,9 +85,10 @@ void FTransaction::AbortWithoutThrowing()
     Stats.Collect<EStatsKind::Abort>();
     CollectStats();
 
-	// Ensure that we're closed before applying the undo, as this will ensure the open-memory
-	// validation is performed before the memory is reset.
-	SetState(EState::ClosedActive);
+	// Ensure that we enter the done state before applying the commit, as this
+	// will ensure the open-memory validation is performed before the write log
+	// is cleared.
+	SetDone();
 
     // Call the destructors of all the OnCommit functors before undoing the transactional memory and
     // calling the OnAbort callbacks. This is important as the callback functions may have captured
@@ -108,9 +114,6 @@ void FTransaction::AbortWithoutThrowing()
     {
 		ASSERT(Context->IsAborting());
     }
-
-	// Transaction is now done.
-	SetState(EState::Done);
 }
 
 void FTransaction::AbortAndThrow()
@@ -127,9 +130,10 @@ bool FTransaction::AttemptToCommit()
     Stats.Collect<EStatsKind::Commit>();
     CollectStats();
 
-	// Ensure that we're closed before applying the commit, as this will ensure the open-memory
-	// validation is performed before the write log is cleared.
-	SetState(EState::ClosedActive);
+	// Ensure that we enter the done state before applying the commit, as this
+	// will ensure the open-memory validation is performed before the write log
+	// is cleared.
+	SetDone();
 
     bool bResult;
     if (IsNested())
@@ -142,15 +146,13 @@ bool FTransaction::AttemptToCommit()
         bResult = AttemptToCommitOuterNest();
     }
 
-	SetState(EState::Done);
-
     return bResult;
 }
 
 void FTransaction::Undo()
 {
 	UE_LOG(LogAutoRTFM, Verbose, TEXT("Undoing a transaction..."));
-	ASSERT(IsClosed());
+	ASSERT(IsDone());
 
 	int VerboseCounter = 0;
 	int Num = WriteLog.Num();
@@ -249,118 +251,147 @@ bool FTransaction::AttemptToCommitOuterNest()
     return true;
 }
 
-FTransaction::EState FTransaction::SetState(EState NewState)
+void FTransaction::SetOpenActive(EMemoryValidationLevel NewMemoryValidationLevel)
 {
-	if (NewState == State())
+	CurrentMemoryValidationLevel = NewMemoryValidationLevel;
+	if (CurrentMemoryValidationLevel == EMemoryValidationLevel::Disabled)
 	{
-		return NewState;
+		// TODO: Validate if open -> open with different validation levels.
+		RecordedWriteHash = 0;
+		NumWriteLogsHashed = 0;
 	}
+	SetState<EState::OpenActive>();
+}
+
+void FTransaction::SetClosedActive()
+{
+	SetState<EState::ClosedActive>();
+}
 
-	// Once Done, the transaction cannot change state without a call to Reset().
-	ASSERT(!IsDone());
+void FTransaction::SetOpenInactive()
+{
+	SetState<EState::OpenInactive>();
+}
+
+void FTransaction::SetClosedInactive()
+{
+	SetState<EState::ClosedInactive>();
+}
 
+void FTransaction::SetActive()
+{
 	switch (CurrentState)
 	{
-		// OpenActive -> OpenInactive, ClosedActive or Done
 		case EState::OpenActive:
-			ASSERT(NewState == EState::OpenInactive || NewState == EState::ClosedActive || NewState == EState::Done);
-			if (CurrentMemoryValidationLevel != EMemoryValidationLevel::Disabled)
-			{
-				ValidateWriteHash();
-				RecordedWriteHash = 0;
-				NumWriteLogsHashed = 0;
-			}
-			break;
-
-		// ClosedActive -> ClosedInactive, OpenActive or Done
 		case EState::ClosedActive:
-			ASSERT(NewState == EState::ClosedInactive || NewState == EState::OpenActive || NewState == EState::Done);
 			break;
-
-		// ClosedActive -> OpenActive
 		case EState::OpenInactive:
-			ASSERT(NewState == EState::OpenActive);
+			SetState<EState::OpenActive>();
 			break;
-
-		// ClosedInactive -> ClosedActive
 		case EState::ClosedInactive:
-			ASSERT(NewState == EState::ClosedActive);
+			SetState<EState::ClosedActive>();
 			break;
-
 		default:
 			UE_LOG(LogAutoRTFM, Fatal, TEXT("Invalid state"));
-			break;
 	}
+}
 
-	switch (NewState) {
-		// OpenInactive, ClosedActive or Done -> OpenActive
+void FTransaction::SetInactive()
+{
+	switch (CurrentState)
+	{
+		case EState::OpenInactive:
+		case EState::ClosedInactive:
+			break;
 		case EState::OpenActive:
-			CurrentMemoryValidationLevel = PendingMemoryValidationLevel;
-			if (CurrentMemoryValidationLevel != EMemoryValidationLevel::Disabled)
-			{
-				ASSERT(RecordedWriteHash == 0);
-				RecordWriteHash();
-			}
+			SetState<EState::OpenInactive>();
 			break;
-
-		// OpenActive, ClosedActive -> Done
-		case EState::Done:
-			ASSERT(IsActive());
-			CommitTasks.Reset();
-			AbortTasks.Reset();
-			HitSet.Reset();
-			NewMemoryTracker.Reset();
-			WriteLog.Reset();
-			WriteLogBumpAllocator.Reset();
+		case EState::ClosedActive:
+			SetState<EState::ClosedInactive>();
 			break;
-
 		default:
-			break;
+			UE_LOG(LogAutoRTFM, Fatal, TEXT("Invalid state"));
 	}
-
-	EState OldState = CurrentState;
-	CurrentState = NewState;
-	return OldState;
 }
 
-void FTransaction::SetMemoryValidationLevel(EMemoryValidationLevel Level)
+void FTransaction::SetDone()
 {
-	PendingMemoryValidationLevel = Level;
+	SetState<EState::Done>();
 }
 
-void FTransaction::SetActive()
+template<FTransaction::EState NewState>
+void FTransaction::SetState()
 {
+	ASSERT(NewState != CurrentState);
+
 	switch (CurrentState)
 	{
+		case EState::Uninitialized:
+			ASSERT(NewState == EState::OpenActive || NewState == EState::ClosedActive);
+			break;
+
+		// OpenActive -> OpenInactive, ClosedActive or Done
 		case EState::OpenActive:
+			ASSERT(NewState == EState::OpenInactive || NewState == EState::ClosedActive || NewState == EState::Done);
+			if (CurrentMemoryValidationLevel != EMemoryValidationLevel::Disabled)
+			{
+				ValidateWriteHash();
+				RecordedWriteHash = 0;
+				NumWriteLogsHashed = 0;
+			}
+			else
+			{
+				ASSERT(RecordedWriteHash == 0 && NumWriteLogsHashed == 0);
+			}
+			break;
+
+		// ClosedActive -> ClosedInactive, OpenActive or Done
 		case EState::ClosedActive:
+			ASSERT(NewState == EState::ClosedInactive || NewState == EState::OpenActive || NewState == EState::Done);
 			break;
+
+		// ClosedActive -> OpenActive
 		case EState::OpenInactive:
-			SetState(EState::OpenActive);
+			ASSERT(NewState == EState::OpenActive);
 			break;
+
+		// ClosedInactive -> ClosedActive
 		case EState::ClosedInactive:
-			SetState(EState::ClosedActive);
+			ASSERT(NewState == EState::ClosedActive);
 			break;
+
+		case EState::Done:
+			UE_LOG(LogAutoRTFM, Fatal, TEXT("Once Done, the transaction cannot change state without a call to Reset()"));
+			break;
+
 		default:
 			UE_LOG(LogAutoRTFM, Fatal, TEXT("Invalid state"));
+			break;
+	}
+
+	// OpenInactive, ClosedActive or Done -> OpenActive
+	if (NewState == EState::OpenActive) {
+		if (CurrentMemoryValidationLevel != EMemoryValidationLevel::Disabled)
+		{
+			ASSERT(RecordedWriteHash == 0 && NumWriteLogsHashed == 0);
+			RecordWriteHash();
+		}
 	}
+
+	CurrentState = NewState;
 }
 
-void FTransaction::SetInactive()
+void FTransaction::DebugBreakIfMemoryValidationFails()
 {
-	switch (CurrentState)
+	if (CurrentMemoryValidationLevel != EMemoryValidationLevel::Disabled)
 	{
-		case EState::OpenInactive:
-		case EState::ClosedInactive:
-			break;
-		case EState::OpenActive:
-			SetState(EState::OpenInactive);
-			break;
-		case EState::ClosedActive:
-			SetState(EState::ClosedInactive);
-			break;
-		default:
-			UE_LOG(LogAutoRTFM, Fatal, TEXT("Invalid state"));
+		TWriteHash OldHash = RecordedWriteHash;
+		TWriteHash NewHash = CalculateNestedWriteHash(NumWriteLogsHashed);
+		if (OldHash != NewHash)
+		{
+			UE_LOG(LogAutoRTFM, Warning, TEXT("DebugBreakIfInvalidMemoryHash() detected a change in hash"));
+			PLATFORM_BREAK();
+		}
 	}
 }
 
@@ -374,21 +405,20 @@ void FTransaction::ValidateWriteHash() const
 {
 	TWriteHash OldHash = RecordedWriteHash;
 	TWriteHash NewHash = CalculateNestedWriteHash(NumWriteLogsHashed);
-	if (LIKELY(OldHash == NewHash))
-	{
-		return;
-	}
 
 	static constexpr TCHAR Message[] =
 		TEXT("Memory modified in a transaction was also modified in an call to AutoRTFM::Open(). ")
 		TEXT("This may lead to memory corruption if the transaction is aborted.");
 	if (CurrentMemoryValidationLevel == EMemoryValidationLevel::EnabledAsWarning)
 	{
-		UE_LOG(LogAutoRTFM, Warning, Message);
+		if (UNLIKELY(OldHash != NewHash))
+		{
+			UE_LOG(LogAutoRTFM, Warning, Message);
+		}
 	}
 	else
 	{
-		UE_LOG(LogAutoRTFM, Error, Message);
+		ensureMsgf(OldHash == NewHash, Message);
 	}
 }
 
diff --git a/Engine/Source/Runtime/Core/Private/AutoRTFM/Transaction.h b/Engine/Source/Runtime/Core/Private/AutoRTFM/Transaction.h
--- a/Engine/Source/Runtime/Core/Private/AutoRTFM/Transaction.h
+++ b/Engine/Source/Runtime/Core/Private/AutoRTFM/Transaction.h
@@ -65,7 +65,7 @@ public:
     // The transaction must we in a Done state.
     ~FTransaction();
 
-    void Initialize(EState State, FTransaction* Parent, bool bIsScoped, FStackRange StackRange);
+    void Initialize(FTransaction* Parent, bool bIsScoped, FStackRange StackRange);
 
     // Clears the tracked transaction state and resets back to the default Uninitialized state.
     void Reset();
@@ -111,13 +111,17 @@ public:
     void DidAllocate(void* LogicalAddress, size_t Size);
     void DidFree(void* LogicalAddress);
 
-    // Changes the state of the transaction to NewState and returns the previous state.
-    EState SetState(EState NewState);
+    void SetOpenActive(EMemoryValidationLevel MemoryValidationLevel);
+    void SetClosedActive();
+    void SetOpenInactive();
+    void SetClosedInactive();
     void SetActive();
     void SetInactive();
+    void SetDone();
 
-    // Only takes effect at the next state transition to OpenActive.
-    void SetMemoryValidationLevel(EMemoryValidationLevel Level);
+    // A debug helper that will break to the debugger if the memory validation
+    // hash no longer matches. Useful for isolating where the open write happened.
+    void DebugBreakIfMemoryValidationFails();
 
     // State querying
     EState State() const { return CurrentState; }
@@ -132,6 +136,8 @@ public:
     bool IsInactive() const { return IsOpenInactive() || IsClosedInactive(); }
     bool IsDone() const { return CurrentState == EState::Done; }
 
+    EMemoryValidationLevel MemoryValidationLevel() const { return CurrentMemoryValidationLevel; }
+
     // The stack range represents all stack memory inside the transaction scope
     inline FStackRange GetStackRange() const { return StackRange; }
 
@@ -164,6 +170,9 @@ private:
 
     bool ShouldRecordWrite(void* LogicalAddress) const;
 
+    template<EState NewState>
+    void SetState();
+
     FContext* Context{nullptr};
 
     // Are we nested? Then this is the parent.
@@ -181,7 +190,6 @@ private:
     TWriteHash RecordedWriteHash = 0;
     size_t NumWriteLogsHashed = 0;
     EMemoryValidationLevel CurrentMemoryValidationLevel = EMemoryValidationLevel::Disabled;
-    EMemoryValidationLevel PendingMemoryValidationLevel = EMemoryValidationLevel::Disabled;
     EState CurrentState{EState::Uninitialized};
     FStackRange StackRange;
     bool bIsStackScoped = false;
diff --git a/Engine/Source/Runtime/Core/Public/AutoRTFM/AutoRTFM.h b/Engine/Source/Runtime/Core/Public/AutoRTFM/AutoRTFM.h
--- a/Engine/Source/Runtime/Core/Public/AutoRTFM/AutoRTFM.h
+++ b/Engine/Source/Runtime/Core/Public/AutoRTFM/AutoRTFM.h
@@ -994,6 +994,12 @@ namespace ForTheRuntime
 	// Sets the memory validation level. See IsWriteValidationEnabled().
 	UE_AUTORTFM_API void SetMemoryValidationLevel(EMemoryValidationLevel Level);
 
+	// A debug helper that will break to the debugger if the hash of the memory
+	// write locations no longer matches the hash recorded when the transaction
+	// was opened. Useful for isolating where the open write happened.
+	// Requires the memory validation to be enabled to be called.
+	UE_AUTORTFM_API void DebugBreakIfMemoryValidationFails();
+
 	// Manually create a new transaction from open code and push it as a transaction nest.
 	// Can only be called within an already active parent transaction (EG. this cannot start
 	// a transaction nest itself).
