[AutoRTFM] Avoid stack overflow when enabling bTrackAllocationLocations

diff --git a/Engine/Source/Runtime/Core/Private/AutoRTFM/Transaction.h b/Engine/Source/Runtime/Core/Private/AutoRTFM/Transaction.h
--- a/Engine/Source/Runtime/Core/Private/AutoRTFM/Transaction.h
+++ b/Engine/Source/Runtime/Core/Private/AutoRTFM/Transaction.h
@@ -193,6 +193,7 @@ private:
     EState CurrentState{EState::Uninitialized};
     FStackRange StackRange;
     bool bIsStackScoped = false;
+    bool bIsInDidAllocate = false;
 };
 
 } // namespace AutoRTFM
diff --git a/Engine/Source/Runtime/Core/Private/AutoRTFM/TransactionInlines.h b/Engine/Source/Runtime/Core/Private/AutoRTFM/TransactionInlines.h
--- a/Engine/Source/Runtime/Core/Private/AutoRTFM/TransactionInlines.h
+++ b/Engine/Source/Runtime/Core/Private/AutoRTFM/TransactionInlines.h
@@ -3,6 +3,7 @@
 #pragma once
 
 #include "Context.h"
+#include "ScopedGuard.h"
 #include "Transaction.h"
 
 #include "HAL/Platform.h"
@@ -178,11 +179,12 @@ template<unsigned SIZE> AUTORTFM_NO_ASAN FORCENOINLINE void FTransaction::Record
 
 UE_AUTORTFM_FORCEINLINE void FTransaction::DidAllocate(void* LogicalAddress, const size_t Size)
 {
-	if (0 == Size)
+	if (0 == Size || bIsInDidAllocate)
 	{
 		return;
 	}
 
+	AutoRTFM::TScopedGuard<bool> RecursionGuard(bIsInDidAllocate, true);
     const bool DidInsert = NewMemoryTracker.Insert(LogicalAddress, Size);
     ASSERT(DidInsert);
 }
@@ -190,10 +192,10 @@ UE_AUTORTFM_FORCEINLINE void FTransaction::DidAllocate(void* LogicalAddress, con
 UE_AUTORTFM_FORCEINLINE void FTransaction::DidFree(void* LogicalAddress)
 {
     ASSERT(bTrackAllocationLocations);
-    
+
     // Checking if one byte is in the interval map is enough to ascertain if it
     // is new memory and we should be worried.
-    ASSERT(!NewMemoryTracker.Contains(LogicalAddress, 1));
+    ASSERT(bIsInDidAllocate || !NewMemoryTracker.Contains(LogicalAddress, 1));
 }
 
 UE_AUTORTFM_FORCEINLINE void FTransaction::DeferUntilCommit(TFunction<void()>&& Callback)
