Allow to mmap FName blocks on iOS and Android platforms

diff --git a/Engine/Build/Android/Java/src/com/epicgames/unreal/GameActivity.java.template b/Engine/Build/Android/Java/src/com/epicgames/unreal/GameActivity.java.template
--- a/Engine/Build/Android/Java/src/com/epicgames/unreal/GameActivity.java.template
+++ b/Engine/Build/Android/Java/src/com/epicgames/unreal/GameActivity.java.template
@@ -4873,6 +4873,16 @@ public class GameActivity extends $${gameActivitySuperClass}$$ implements Surfac
 		return appPackageName;
 	}
 
+	public String AndroidThunkJava_GetCacheDir()
+	{
+		try {
+			return getApplicationContext().getCacheDir().getPath();
+		} catch (Exception e) {
+			Log.debug("getCacheDir failed: " + e.toString());
+		}
+		return "";
+	}
+
 	public boolean AndroidThunkJava_IsMusicActive()
 	{
 		AudioManager audioManager = (AudioManager)getThisActivity().getSystemService(Context.AUDIO_SERVICE);
diff --git a/Engine/Source/Runtime/Core/Private/UObject/UnrealNames.cpp b/Engine/Source/Runtime/Core/Private/UObject/UnrealNames.cpp
--- a/Engine/Source/Runtime/Core/Private/UObject/UnrealNames.cpp
+++ b/Engine/Source/Runtime/Core/Private/UObject/UnrealNames.cpp
@@ -40,6 +40,14 @@
 #include "Misc/AsciiSet.h"
 #include "AutoRTFM/AutoRTFM.h"
 
+#define USE_FNAME_MMAP (PLATFORM_ANDROID || PLATFORM_IOS)
+#if USE_FNAME_MMAP
+#include <sys/mman.h>
+#if PLATFORM_ANDROID
+extern FString AndroidThunkCpp_GetCacheDir();
+#endif
+#endif
+
 PRAGMA_DISABLE_UNSAFE_TYPECAST_WARNINGS
 
 DEFINE_LOG_CATEGORY_STATIC(LogUnrealNames, Log, All);
@@ -97,6 +105,12 @@ namespace UE::Name::Private
 	);
 }
 
+bool GFNameUseMmap = false;
+static FAutoConsoleVariableRef CVarFNameUseMmap(
+	TEXT("FName.UseMmap"),
+	GFNameUseMmap,
+	TEXT("Whether to use mmap instead of malloc for block allocations"));
+
 const TCHAR* LexToString(EName Ename)
 {
 	switch (Ename)
@@ -449,6 +463,25 @@ public:
 	enum { Stride = alignof(FNameEntry) };
 	enum { BlockSizeBytes = Stride * FNameBlockOffsets };
 
+#if USE_FNAME_MMAP
+	const uint32 PageAlignedBlockSizeBytes = Align((uint32)BlockSizeBytes, (uint32)sysconf(_SC_PAGESIZE));
+	// mmap block aligned to FName block size and page size
+	const uint32 MmapSizeBytes = Align(128 * 1024 * 1024, PageAlignedBlockSizeBytes);
+
+	bool IsBlockMmapped(uint8* BlockAddress)
+	{
+		for (uint8* Address : MmappedAddresses)
+		{
+			if (BlockAddress >= Address && BlockAddress < (Address + MmapSizeBytes))
+			{
+				return true;
+			}
+		}
+		return false;
+	}
+
+#endif
+
 	/** Initializes all member variables. */
 	FNameEntryAllocator()
 	{
@@ -459,8 +492,24 @@ public:
 	{
 		for (int32 Index = CurrentBlock; Index >= 0; --Index)
 		{
+#if USE_FNAME_MMAP
+			if (IsBlockMmapped(Blocks[Index]))
+			{
+				continue;
+			}
+#endif
 			FMemory::Free(Blocks[Index]);
 		}
+#if USE_FNAME_MMAP
+		for (uint8* Address : MmappedAddresses)
+		{
+			munmap(Address, MmapSizeBytes);
+		}
+		if (MmapFile != -1)
+		{
+			close(MmapFile);
+		}
+#endif
 	}
 
 	void ReserveBlocks(uint32 Num)
@@ -657,12 +706,66 @@ private:
 		}
 	}
 
-	static uint8* AllocBlock()
+	uint8* AllocBlock()
 	{
 		LLM_SCOPE(ELLMTag::FName);
 		LLM_TAGSET_SCOPE_CLEAR(ELLMTagSet::Assets);
 		LLM_TAGSET_SCOPE_CLEAR(ELLMTagSet::AssetClasses);
 		UE_TRACE_METADATA_CLEAR_SCOPE();
+
+#if USE_FNAME_MMAP
+		if (GFNameUseMmap)
+		{
+#if PLATFORM_ANDROID
+			if (MmapFile == -1)
+			{
+				FString CacheDirectoryPath = AndroidThunkCpp_GetCacheDir();
+				if (!CacheDirectoryPath.IsEmpty())
+				{
+					auto CacheDirectoryPathCStr = StringCast<char>(*CacheDirectoryPath);
+					MmapFile = open(CacheDirectoryPathCStr.Get(), O_TMPFILE | O_EXCL | O_RDWR, S_IRUSR | S_IWUSR);
+				}
+			}
+#else
+			if (MmapFile == -1)
+			{
+				MmapFile = fileno(tmpfile());
+			}
+#endif
+
+			if (MmapFile != -1)
+			{
+				// Increase file size, write a single 0 at the end
+				char Zero = 0;
+				pwrite(MmapFile, &Zero, 1, MmapFileOffset + PageAlignedBlockSizeBytes - 1);
+
+				if (CurrentMmapAddress == nullptr || CurrentMmapOffset >= MmapSizeBytes)
+				{
+					void* AddressHint = NULL;
+					if (CurrentMmapAddress != nullptr)
+					{
+						AddressHint = CurrentMmapAddress + MmapSizeBytes;
+					}
+
+					CurrentMmapAddress = (uint8*)mmap(AddressHint, MmapSizeBytes, PROT_NONE, MAP_SHARED, MmapFile, MmapFileOffset);
+
+					
+					CurrentMmapOffset = 0;
+					MmappedAddresses.Add(CurrentMmapAddress);
+				}
+
+				MmapFileOffset += PageAlignedBlockSizeBytes;
+				
+				uint8* MappedPointer = CurrentMmapAddress + CurrentMmapOffset;
+				CurrentMmapOffset += PageAlignedBlockSizeBytes;
+
+				// enable read/write for the new block
+				mprotect(MappedPointer, PageAlignedBlockSizeBytes, PROT_READ | PROT_WRITE);
+
+				return MappedPointer;
+			}
+		}
+#endif
 		return (uint8*)FMemory::Malloc(BlockSizeBytes, alignof(FNameEntry));
 	}
 	
@@ -684,6 +787,14 @@ private:
 		}
 #endif
 
+#if USE_FNAME_MMAP
+		// mark last block as readonly
+		if (IsBlockMmapped(Blocks[CurrentBlock]))
+		{
+			mprotect(Blocks[CurrentBlock], PageAlignedBlockSizeBytes, PROT_READ);
+		}
+#endif
+
 		++CurrentBlock;
 		CurrentByteCursor = 0;
 
@@ -697,13 +808,23 @@ private:
 			Blocks[CurrentBlock] = AllocBlock();
 		}
 
-		FPlatformMisc::Prefetch(Blocks[CurrentBlock]);
+		if (!GFNameUseMmap)
+		{
+			FPlatformMisc::Prefetch(Blocks[CurrentBlock]);
+		}
 	}
 
 	mutable FRWLock Lock;
 	uint32 CurrentBlock = 0;
 	uint32 CurrentByteCursor = 0;
 	uint8* Blocks[FNameMaxBlocks] = {};
+#if USE_FNAME_MMAP
+	TArray<uint8*> MmappedAddresses;
+	int32 MmapFile = -1;
+	int32 MmapFileOffset = 0;
+	int32 CurrentMmapOffset = 0;
+	uint8* CurrentMmapAddress = nullptr;
+#endif
 };
 
 // Increasing shards reduces contention but uses more memory and adds cache pressure.
diff --git a/Engine/Source/Runtime/Launch/Private/Android/AndroidJNI.cpp b/Engine/Source/Runtime/Launch/Private/Android/AndroidJNI.cpp
--- a/Engine/Source/Runtime/Launch/Private/Android/AndroidJNI.cpp
+++ b/Engine/Source/Runtime/Launch/Private/Android/AndroidJNI.cpp
@@ -143,6 +143,7 @@ void FJavaWrapper::FindClassesAndMethods(JNIEnv* Env)
 	AndroidThunkJava_GetSharedPreferenceString = FindMethod(Env, GameActivityClassID, "AndroidThunkJava_GetSharedPreferenceString", "(Ljava/lang/String;Ljava/lang/String;Ljava/lang/String;)Ljava/lang/String;", bIsOptional);
 	AndroidThunkJava_DeleteSharedPreference = FindMethod(Env, GameActivityClassID, "AndroidThunkJava_DeleteSharedPreference", "(Ljava/lang/String;Ljava/lang/String;)V", bIsOptional);
 	AndroidThunkJava_DeleteSharedPreferenceGroup = FindMethod(Env, GameActivityClassID, "AndroidThunkJava_DeleteSharedPreferenceGroup", "(Ljava/lang/String;)V", bIsOptional);
+	AndroidThunkJava_GetCacheDir = FindMethod(Env, GameActivityClassID, "AndroidThunkJava_GetCacheDir", "()Ljava/lang/String;", bIsOptional);
 
 	// Screen capture/recording permission
 	AndroidThunkJava_IsScreenCaptureDisabled = FindMethod(Env, GameActivityClassID, "AndroidThunkJava_IsScreenCaptureDisabled", "()Z", bIsOptional);
@@ -555,6 +556,7 @@ jmethodID FJavaWrapper::AndroidThunkJava_SetSharedPreferenceString;
 jmethodID FJavaWrapper::AndroidThunkJava_GetSharedPreferenceString;
 jmethodID FJavaWrapper::AndroidThunkJava_DeleteSharedPreference;
 jmethodID FJavaWrapper::AndroidThunkJava_DeleteSharedPreferenceGroup;
+jmethodID FJavaWrapper::AndroidThunkJava_GetCacheDir;
 
 jclass FJavaWrapper::InputDeviceInfoClass;
 jfieldID FJavaWrapper::InputDeviceInfo_VendorId;
@@ -1569,6 +1571,16 @@ FString AndroidThunkCpp_GetSharedPreferenceString(const FString& Group, const FS
 	return Result;
 }
 
+FString AndroidThunkCpp_GetCacheDir()
+{
+	FString Result;
+	if (JNIEnv* Env = FAndroidApplication::GetJavaEnv())
+	{
+		Result = FJavaHelper::FStringFromLocalRef(Env, (jstring)FJavaWrapper::CallObjectMethod(Env, FJavaWrapper::GameActivityThis, FJavaWrapper::AndroidThunkJava_GetCacheDir));
+	}
+	return Result;
+}
+
 bool AndroidThunkCpp_GetSharedPreferenceStringTypeSafe(const FString& Group, const FString& Key, FString& OutValue)
 {
 	if (JNIEnv* Env = FAndroidApplication::GetJavaEnv())
diff --git a/Engine/Source/Runtime/Launch/Public/Android/AndroidJNI.h b/Engine/Source/Runtime/Launch/Public/Android/AndroidJNI.h
--- a/Engine/Source/Runtime/Launch/Public/Android/AndroidJNI.h
+++ b/Engine/Source/Runtime/Launch/Public/Android/AndroidJNI.h
@@ -82,6 +82,7 @@ public:
 	static jmethodID AndroidThunkJava_GetSharedPreferenceString;
 	static jmethodID AndroidThunkJava_DeleteSharedPreference;
 	static jmethodID AndroidThunkJava_DeleteSharedPreferenceGroup;
+	static jmethodID AndroidThunkJava_GetCacheDir;
 
 
 	// Screen capture/recording permission
