[AutoRTFM] Enable allocation tracking in test builds.

This will shake out incorrect allocations and frees more quickly during testing.

diff --git a/Engine/Source/Runtime/Core/Private/AutoRTFM/Toggles.h b/Engine/Source/Runtime/Core/Private/AutoRTFM/Toggles.h
--- a/Engine/Source/Runtime/Core/Private/AutoRTFM/Toggles.h
+++ b/Engine/Source/Runtime/Core/Private/AutoRTFM/Toggles.h
@@ -2,13 +2,19 @@
 
 #pragma once
 
+#include "Misc/Build.h"
+
 namespace AutoRTFM
 {
 
 // Should we collect stats on the AutoRTFM runtime or not.
 static constexpr bool bCollectStats = false;
 
-// Should we track the locations of allocations.
+// Tracking the locations of allocations is helpful for catching invalid frees.
+#if UE_BUILD_SHIPPING
 static constexpr bool bTrackAllocationLocations = false;
+#else
+static constexpr bool bTrackAllocationLocations = true;
+#endif
 
 } // namespace AutoRTFM
