[Insights] Added "ConsoleCmd" CPU scoped trace event for console input processing (includes the console command execution). The CPU timer has the actual commad string as metadata.

diff --git a/Engine/Source/Runtime/Core/Private/HAL/ConsoleManager.cpp b/Engine/Source/Runtime/Core/Private/HAL/ConsoleManager.cpp
--- a/Engine/Source/Runtime/Core/Private/HAL/ConsoleManager.cpp
+++ b/Engine/Source/Runtime/Core/Private/HAL/ConsoleManager.cpp
@@ -26,6 +26,8 @@ ConsoleManager.cpp: console command handling
 #include "ProfilingDebugging/CsvProfiler.h"
 #include "HAL/FileManager.h"
 #include "Logging/MessageLog.h"
+#include "Trace/Trace.h"
+#include "Trace/Trace.inl"
 
 #include <clocale>
 
@@ -2683,9 +2685,17 @@ void UnsetCVarTag(const TCHAR* Params, FOutputDevice& Ar)
 	IConsoleManager::Get().UnsetAllConsoleVariablesWithTag(*TagName);
 }
 
+UE_TRACE_EVENT_BEGIN(Cpu, ConsoleCmd, NoSync)
+	UE_TRACE_EVENT_FIELD(UE::Trace::WideString, Cmd)
+UE_TRACE_EVENT_END()
+
 bool FConsoleManager::ProcessUserConsoleInput(const TCHAR* InInput, FOutputDevice& Ar, UWorld* InWorld)
 {
 	check(InInput);
+
+	UE_TRACE_LOG_SCOPED_T(Cpu, ConsoleCmd, CpuChannel)
+		<< ConsoleCmd.Cmd(InInput);
+
 	CSV_EVENT_GLOBAL(TEXT("Cmd: %s"), InInput);
 
 	if (FParse::Command(&InInput, TEXT("dumpcvars")))
