[AutoRTFM] Fix the 'AutoRTFMMemoryValidationLevel' command line flag name

Also log whenever the AutoRTFM flags change (including on startup). Makes it far easier to tell whether the flags are actually working.

diff --git a/Engine/Source/Runtime/Core/Private/AutoRTFM/API.cpp b/Engine/Source/Runtime/Core/Private/AutoRTFM/API.cpp
--- a/Engine/Source/Runtime/Core/Private/AutoRTFM/API.cpp
+++ b/Engine/Source/Runtime/Core/Private/AutoRTFM/API.cpp
@@ -21,11 +21,6 @@ namespace
 	int GAutoRTFMRuntimeEnabled = AutoRTFM::ForTheRuntime::EAutoRTFMEnabledState::AutoRTFM_Disabled;
 #endif // UE_AUTORTFM_ENABLED_RUNTIME_BY_DEFAULT
 
-	void UpdateAutoRTFMRuntimeCrashData()
-	{
-		FGenericCrashContext::SetGameData(TEXT("IsAutoRTFMRuntimeEnabled"), AutoRTFM::ForTheRuntime::IsAutoRTFMRuntimeEnabled() ? TEXT("true") : TEXT("false"));
-	}
-
 	bool GAutoRTFMEnsureOnAbortByLanguage = true;
 
 	int GAutoRTFMRetryTransactions = AutoRTFM::ForTheRuntime::EAutoRTFMRetryTransactionState::NoRetry;
@@ -36,28 +31,41 @@ namespace
 
 	bool GAutoRTFMAssertOnError = false;
 
-	void UpdateAutoRTFMRetryTransactionsData()
+	void OnIsAutoRTFMEnabledChanged()
+	{
+		const bool bEnabled = AutoRTFM::ForTheRuntime::IsAutoRTFMRuntimeEnabled();
+		UE_LOG(LogCore, Log, TEXT("AutoRTFM: %s"), bEnabled ? TEXT("enabled") : TEXT("disabled"));
+		FGenericCrashContext::SetGameData(TEXT("IsAutoRTFMRuntimeEnabled"), bEnabled ? TEXT("true") : TEXT("false"));
+	}
+
+	void OnAutoRTFMRetryTransactionsChanged()
 	{
 		switch (GAutoRTFMRetryTransactions)
 		{
 		case AutoRTFM::ForTheRuntime::EAutoRTFMRetryTransactionState::NoRetry:
+			UE_LOG(LogCore, Log, TEXT("AutoRTFM Retry Transactions: disabled"));
 			return FGenericCrashContext::SetGameData(TEXT("AutoRTFMRetryTransactionState"), TEXT("NoRetry"));
 		case AutoRTFM::ForTheRuntime::EAutoRTFMRetryTransactionState::RetryNonNested:
+			UE_LOG(LogCore, Log, TEXT("AutoRTFM Retry Transactions: retry-non-nested"));
 			return FGenericCrashContext::SetGameData(TEXT("AutoRTFMRetryTransactionState"), TEXT("RetryNonNested"));
 		case AutoRTFM::ForTheRuntime::EAutoRTFMRetryTransactionState::RetryNestedToo:
+			UE_LOG(LogCore, Log, TEXT("AutoRTFM Retry Transactions: retry-nested-too"));
 			return FGenericCrashContext::SetGameData(TEXT("AutoRTFMRetryTransactionState"), TEXT("RetryNestedToo"));
 		}
 	}
 
-	void UpdateAutoRTFMMemoryValidationLevelData()
+	void OnAutoRTFMMemoryValidationLevelChanged()
 	{
 		switch (static_cast<AutoRTFM::EMemoryValidationLevel>(GMemoryValidationLevel))
 		{
 		case AutoRTFM::EMemoryValidationLevel::Disabled:
+			UE_LOG(LogCore, Log, TEXT("AutoRTFM Memory Validation: disabled"));
 			return FGenericCrashContext::SetGameData(TEXT("AutoRTFMMemoryValidationLevel"), TEXT("Disabled"));
 		case AutoRTFM::EMemoryValidationLevel::EnabledAsWarning:
+			UE_LOG(LogCore, Log, TEXT("AutoRTFM Memory Validation: enabled-as-warning"));
 			return FGenericCrashContext::SetGameData(TEXT("AutoRTFMMemoryValidationLevel"), TEXT("EnabledAsWarning"));
 		case AutoRTFM::EMemoryValidationLevel::Enabled:
+			UE_LOG(LogCore, Log, TEXT("AutoRTFM Memory Validation: enabled"));
 			return FGenericCrashContext::SetGameData(TEXT("AutoRTFMMemoryValidationLevel"), TEXT("Enabled"));
 		}
 	}
@@ -68,7 +76,7 @@ static FAutoConsoleVariableRef CVarAutoRTFMRuntimeEnabled(
 	TEXT("AutoRTFMRuntimeEnabled"),
 	GAutoRTFMRuntimeEnabled,
 	TEXT("Enables the AutoRTFM runtime"),
-	FConsoleVariableDelegate::CreateLambda([] (IConsoleVariable*) { UpdateAutoRTFMRuntimeCrashData(); }),
+	FConsoleVariableDelegate::CreateLambda([] (IConsoleVariable*) { OnIsAutoRTFMEnabledChanged(); }),
 	ECVF_Default
 );
 
@@ -76,15 +84,15 @@ static FAutoConsoleVariableRef CVarAutoRTFMRetryTransactions(
 	TEXT("AutoRTFMRetryTransactions"),
 	GAutoRTFMRetryTransactions,
 	TEXT("Enables the AutoRTFM sanitizer-like mode where we can force an abort-and-retry on transactions (useful to test abort codepaths work as intended)"),
-	FConsoleVariableDelegate::CreateLambda([](IConsoleVariable*) { UpdateAutoRTFMRetryTransactionsData(); }),
+	FConsoleVariableDelegate::CreateLambda([](IConsoleVariable*) { OnAutoRTFMRetryTransactionsChanged(); }),
 	ECVF_Default
 );
 
 static FAutoConsoleVariableRef CVarAutoRTFMMemoryValidationLevel(
-	TEXT("AutoMemoryValidationLevel"),
+	TEXT("AutoRTFMMemoryValidationLevel"),
 	GMemoryValidationLevel,
 	TEXT("Detects potential memory corruption due to writes made both by a transaction and open-code"),
-	FConsoleVariableDelegate::CreateLambda([](IConsoleVariable*) { UpdateAutoRTFMMemoryValidationLevelData(); }),
+	FConsoleVariableDelegate::CreateLambda([](IConsoleVariable*) { OnAutoRTFMMemoryValidationLevelChanged(); }),
 	ECVF_Default
 );
 
@@ -97,8 +105,9 @@ static FAutoConsoleVariableRef CVarAutoRTFMAssertOnError(
 
 static FDelayedAutoRegisterHelper DelayedAutoRegister(EDelayedRegisterRunPhase::EndOfEngineInit, []
 	{
-		UpdateAutoRTFMRuntimeCrashData();
-		UpdateAutoRTFMRetryTransactionsData();
+		OnIsAutoRTFMEnabledChanged();
+		OnAutoRTFMRetryTransactionsChanged();
+		OnAutoRTFMMemoryValidationLevelChanged();
 	});
 #endif
 
@@ -124,7 +133,7 @@ namespace AutoRTFM
 
 			GAutoRTFMRuntimeEnabled = State;
 
-			UpdateAutoRTFMRuntimeCrashData();
+			OnIsAutoRTFMEnabledChanged();
 
 			return true;
 #else
@@ -187,7 +196,7 @@ namespace AutoRTFM
 		{
 #if UE_AUTORTFM
 			GAutoRTFMRetryTransactions = State;
-			UpdateAutoRTFMRetryTransactionsData();
+			OnAutoRTFMRetryTransactionsChanged();
 #endif
 		}
 
