[AutoRTFM] Handle logging in the open.

The output consoles are not RTFM-aware.

diff --git a/Engine/Source/Runtime/Core/Private/Logging/LogMacros.cpp b/Engine/Source/Runtime/Core/Private/Logging/LogMacros.cpp
--- a/Engine/Source/Runtime/Core/Private/Logging/LogMacros.cpp
+++ b/Engine/Source/Runtime/Core/Private/Logging/LogMacros.cpp
@@ -10,6 +10,7 @@
 #include "Misc/VarargsHelper.h"
 #include "Stats/Stats.h"
 #include "ProfilingDebugging/CsvProfiler.h"
+#include "AutoRTFM/AutoRTFM.h"
 
 void StaticFailDebugV(const TCHAR* Error, const ANSICHAR* Expression, const ANSICHAR* File, int32 Line, bool bIsEnsure, void* ProgramCounter, const TCHAR* DescriptionFormat, va_list DescriptionArgs);
 
@@ -34,8 +35,13 @@ void FMsg::LogfImpl(const ANSICHAR* File, int32 Line, const FLogCategoryName& Ca
 		default:
 			break;
 		}
-		GROWABLE_LOGF(LogOverride	? LogOverride->Log(Category, Verbosity, Buffer)
-									: GLog->RedirectLog(Category, Verbosity, Buffer))
+
+		// Logging is always done in the open as we want logs even with transactionalized code.
+		GROWABLE_LOGF(UE_AUTORTFM_OPEN
+		{
+			LogOverride ? LogOverride->Log(Category, Verbosity, Buffer)
+			  	 	    : GLog->RedirectLog(Category, Verbosity, Buffer);
+		})
 	}
 	else
 	{
@@ -61,24 +67,27 @@ void FMsg::LogV(const ANSICHAR* File, int32 Line, const FLogCategoryName& Catego
 		TStringBuilder<512> Buffer;
 		Buffer.AppendV(Fmt, Args);
 		const TCHAR* Message = *Buffer;
-		FOutputDevice* OutputDevice = nullptr;
+		FOutputDevice* OutputDevice = GLog;
 		switch (Verbosity)
 		{
 		case ELogVerbosity::Error:
 		case ELogVerbosity::Warning:
 		case ELogVerbosity::Display:
 		case ELogVerbosity::SetColor:
-			OutputDevice = GWarn;
+			if (GWarn != nullptr)
+			{
+				OutputDevice = GWarn;
+			}
 			break;
 		default:
 			break;
 		}
 
 		// Logging is always done in the open as we want logs even with transactionalized code.
-		AutoRTFM::Open([OutputDevice, Message, Verbosity, Category]
+		UE_AUTORTFM_OPEN
 		{
-			(OutputDevice ? OutputDevice : GLog)->Serialize(Message, Verbosity, Category);
-		});
+			OutputDevice->Serialize(Message, Verbosity, Category);
+		};
 
 #if CSV_PROFILER_STATS
 		// Only update the CSV stat if we're not crashing, otherwise things can get messy
