[AutoRTFM] Implement FTransactionallySafeScopeLock in terms of UE::TScopeLock.

These two classes were almost exact 1:1 replicas of each other, except for one snag: TScopeLock takes a reference to the lock, where FTransactionallySafeScopeLock took a pointer. A custom ctor papers over the difference here.

Separately, I will replace existing usages of FTransactionallySafeScopeLock with UE::TScopeLock, and then this file can be removed completely.

diff --git a/Engine/Source/Runtime/Core/Public/Misc/TransactionallySafeScopeLock.h b/Engine/Source/Runtime/Core/Public/Misc/TransactionallySafeScopeLock.h
--- a/Engine/Source/Runtime/Core/Public/Misc/TransactionallySafeScopeLock.h
+++ b/Engine/Source/Runtime/Core/Public/Misc/TransactionallySafeScopeLock.h
@@ -2,35 +2,13 @@
 
 #pragma once
 
-#include "TransactionallySafeCriticalSection.h"
+#include "Misc/ScopeLock.h"
+#include "Misc/TransactionallySafeCriticalSection.h"
 
 // A transactionally safe scope lock - that uses a transactionally safe critical section to back it.
-struct FTransactionallySafeScopeLock final
+class FTransactionallySafeScopeLock : public UE::TScopeLock<FTransactionallySafeCriticalSection>
 {
+public:
 	UE_NODISCARD_CTOR FTransactionallySafeScopeLock(FTransactionallySafeCriticalSection* InSynchObject)
-		: SynchObject(InSynchObject)
-	{
-		check(SynchObject);
-		SynchObject->Lock();
-	}
-
-	~FTransactionallySafeScopeLock()
-	{
-		Unlock();
-	}
-
-	void Unlock()
-	{
-		if (SynchObject)
-		{
-			SynchObject->Unlock();
-			SynchObject = nullptr;
-		}
-	}
-
-private:
-	FTransactionallySafeScopeLock() = delete;
-	FTransactionallySafeScopeLock(const FTransactionallySafeScopeLock&) = delete;
-
-	FTransactionallySafeCriticalSection* SynchObject;
+		: UE::TScopeLock<FTransactionallySafeCriticalSection>(*InSynchObject) {}
 };
