[Trace] Added CPU scopes (TRACE_CPUPROFILER_EVENT_SCOPE) to FTraceAuxiliary API.

diff --git a/Engine/Source/Runtime/Core/Private/ProfilingDebugging/TraceAuxiliary.cpp b/Engine/Source/Runtime/Core/Private/ProfilingDebugging/TraceAuxiliary.cpp
--- a/Engine/Source/Runtime/Core/Private/ProfilingDebugging/TraceAuxiliary.cpp
+++ b/Engine/Source/Runtime/Core/Private/ProfilingDebugging/TraceAuxiliary.cpp
@@ -8,21 +8,21 @@
 #include "Misc/ScopeRWLock.h"
 
 #if PLATFORM_WINDOWS
-#	include "Windows/AllowWindowsPlatformTypes.h"
-#	include <Windows.h>
-#	include "Windows/HideWindowsPlatformTypes.h"
+#include "Windows/AllowWindowsPlatformTypes.h"
+#include <Windows.h>
+#include "Windows/HideWindowsPlatformTypes.h"
 #endif
 #if PLATFORM_LINUX || PLATFORM_MAC
-# 	include <sys/wait.h>
-#   include <semaphore.h>
+#include <sys/wait.h>
+#include <semaphore.h>
 #endif
 
 #if defined(WITH_UNREAL_TRACE_LAUNCH)
-  #define UE_TRACE_SERVER_LAUNCH_ENABLED WITH_UNREAL_TRACE_LAUNCH
-  UE_DEPRECATED_MACRO(5.4, "The WITH_UNREAL_TRACE_LAUNCH macro has been deprecated in favor of UE_TRACE_SERVER_LAUNCH_ENABLED.")
+	#define UE_TRACE_SERVER_LAUNCH_ENABLED WITH_UNREAL_TRACE_LAUNCH
+	UE_DEPRECATED_MACRO(5.4, "The WITH_UNREAL_TRACE_LAUNCH macro has been deprecated in favor of UE_TRACE_SERVER_LAUNCH_ENABLED.")
 #elif !defined(UE_TRACE_SERVER_LAUNCH_ENABLED)
-  #define UE_TRACE_SERVER_LAUNCH_ENABLED (PLATFORM_DESKTOP && !UE_BUILD_SHIPPING && !IS_PROGRAM)
-  #define WITH_UNREAL_TRACE_LAUNCH UE_DEPRECATED_MACRO(5.4, "The WITH_UNREAL_TRACE_LAUNCH macro has been deprecated in favor of UE_TRACE_SERVER_LAUNCH_ENABLED.") UE_TRACE_SERVER_LAUNCH_ENABLED
+	#define UE_TRACE_SERVER_LAUNCH_ENABLED (PLATFORM_DESKTOP && !UE_BUILD_SHIPPING && !IS_PROGRAM)
+	#define WITH_UNREAL_TRACE_LAUNCH UE_DEPRECATED_MACRO(5.4, "The WITH_UNREAL_TRACE_LAUNCH macro has been deprecated in favor of UE_TRACE_SERVER_LAUNCH_ENABLED.") UE_TRACE_SERVER_LAUNCH_ENABLED
 #endif
 
 #if !defined(UE_TRACE_AUTOSTART)
@@ -304,8 +304,8 @@ void FTraceAuxiliaryImpl::ForEachChannel(const TCHAR* ChannelList, bool bResolve
 }
 
 ////////////////////////////////////////////////////////////////////////////////
-template <typename IdentifierType>
-bool FTraceAuxiliaryImpl::EnableChannel(IdentifierType ChannelIdentifier, const FTraceAuxiliary::FLogCategoryAlias& LogCategory)
+template <typename ChannelType>
+bool FTraceAuxiliaryImpl::EnableChannel(ChannelType ChannelIdentifier, const FTraceAuxiliary::FLogCategoryAlias& LogCategory)
 {
 	// Channel names have been provided by the user and may not exist yet. As
 	// we want to maintain bActive accurately (channels toggles are reference
@@ -325,7 +325,7 @@ bool FTraceAuxiliaryImpl::EnableChannel(IdentifierType ChannelIdentifier, const
 		const uint32 ChannelNameLen = Channel->GetName(&ChannelNameA);
 		ChannelName << FAnsiStringView(ChannelNameA, ChannelNameLen);
 	}
-	
+
 	// It is not possible to change read only channels once trace is initialized.
 	if (bReadOnlyChannelsFrozen && Channel->IsReadOnly())
 	{
@@ -365,7 +365,7 @@ bool FTraceAuxiliaryImpl::DisableChannel(ChannelType ChannelIdentifier, const FT
 		const uint32 ChannelNameLen = Channel->GetName(&ChannelNameA);
 		ChannelName << FAnsiStringView(ChannelNameA, ChannelNameLen);
 	}
-	
+
 	// It is not possible to change read only channels once trace is initialized.
 	if (bReadOnlyChannelsFrozen && Channel->IsReadOnly())
 	{
@@ -405,6 +405,7 @@ bool FTraceAuxiliaryImpl::Connect(FTraceAuxiliary::EConnectionType Type, const T
 	{
 		if (Type == FTraceAuxiliary::EConnectionType::Network)
 		{
+			TRACE_CPUPROFILER_EVENT_SCOPE(FTraceAux_SendToHost);
 			bConnected = SendToHost(Parameter, LogCategory, SendFlags);
 			if (bConnected)
 			{
@@ -419,6 +420,7 @@ bool FTraceAuxiliaryImpl::Connect(FTraceAuxiliary::EConnectionType Type, const T
 
 		else if (Type == FTraceAuxiliary::EConnectionType::File)
 		{
+			TRACE_CPUPROFILER_EVENT_SCOPE(FTraceAux_WriteToFile);
 			bConnected = WriteToFile(Parameter, LogCategory, SendFlags);
 			if (bConnected)
 			{
@@ -442,6 +444,7 @@ bool FTraceAuxiliaryImpl::Connect(FTraceAuxiliary::EConnectionType Type, const T
 				StartedType = CurrentTraceTarget.TraceType;
 			}
 
+			TRACE_CPUPROFILER_EVENT_SCOPE(FTraceAux_OnTraceStarted);
 			FTraceAuxiliary::OnTraceStarted.Broadcast(StartedType, StartedDest);
 		}
 	}
@@ -461,6 +464,8 @@ bool FTraceAuxiliaryImpl::Stop()
 		return false;
 	}
 
+	TRACE_CPUPROFILER_EVENT_SCOPE(FTraceAux_Stop);
+
 	if (!UE::Trace::Stop())
 	{
 		return false;
@@ -478,6 +483,7 @@ bool FTraceAuxiliaryImpl::Stop()
 		CurrentTraceTarget.TraceDest.Reset();
 	}
 
+	TRACE_CPUPROFILER_EVENT_SCOPE(FTraceAux_OnTraceStopped);
 	FTraceAuxiliary::OnTraceStopped.Broadcast(StopedType, StopedDest);
 
 	return true;
@@ -556,14 +562,14 @@ bool FTraceAuxiliaryImpl::SendToHost(const TCHAR* InHost, const FTraceAuxiliary:
 {
 	FStringView HostView(InHost);
 	// Parse port if specified. Default is "0" which indicates default port will be used.
-	uint32 Port = 0; 
+	uint32 Port = 0;
 	int32 Separator = INDEX_NONE;
 	if (HostView.FindChar(TEXT(':'), Separator))
 	{
 		LexFromString(Port, HostView.RightChop(Separator + 1).GetData());
 		HostView.LeftInline(Separator);
 	}
-	
+
 	FString Host(HostView);
 	if (!UE::Trace::SendTo(*Host, Port, SendFlags))
 	{
@@ -695,6 +701,7 @@ bool FTraceAuxiliaryImpl::WriteSnapshot(const TCHAR* InFilePath, const FTraceAux
 
 	if (bResult)
 	{
+		TRACE_CPUPROFILER_EVENT_SCOPE(FTraceAux_OnSnapshotSaved);
 		FTraceAuxiliary::OnSnapshotSaved.Broadcast(FTraceAuxiliary::EConnectionType::File, NativePath);
 		UE_LOG_REF(LogCategory, Display, TEXT("Trace snapshot generated in %.3f seconds to \"%s\"."), FPlatformTime::Seconds() - StartTime, *NativePath);
 	}
@@ -723,6 +730,7 @@ bool FTraceAuxiliaryImpl::SendSnapshot(const TCHAR* InHost, uint32 InPort, const
 
 	if (bResult)
 	{
+		TRACE_CPUPROFILER_EVENT_SCOPE(FTraceAux_OnSnapshotSaved);
 		FTraceAuxiliary::OnSnapshotSaved.Broadcast(FTraceAuxiliary::EConnectionType::Network, InHost);
 		UE_LOG_REF(LogCategory, Display, TEXT("Trace snapshot generated in %.3f seconds to \"%s\"."), FPlatformTime::Seconds() - StartTime, InHost);
 	}
@@ -807,6 +815,7 @@ void FTraceAuxiliaryImpl::StartWorkerThread()
 {
 	if (!bWorkerThreadStarted)
 	{
+		TRACE_CPUPROFILER_EVENT_SCOPE(FTraceAux_StartWorkerThread);
 		UE::Trace::StartWorkerThread();
 		bWorkerThreadStarted = true;
 	}
@@ -848,6 +857,7 @@ void FTraceAuxiliaryImpl::StartEndFramePump()
 ////////////////////////////////////////////////////////////////////////////////
 void OnConnectionCallback()
 {
+	TRACE_CPUPROFILER_EVENT_SCOPE(FTraceAux_OnConnection);
 	FTraceAuxiliary::OnConnection.Broadcast();
 }
 
@@ -867,7 +877,7 @@ void OnMessageCallback(const UE::Trace::FMessageEvent& Message)
 	{
 		LogStr << Message.TypeStr;
 	}
-	
+
 	switch(Message.Type)
 	{
 	case UE::Trace::EMessageType::Log:
@@ -893,7 +903,7 @@ void OnMessageCallback(const UE::Trace::FMessageEvent& Message)
 		}
 		break;
 	}
-	
+
 	#undef LOG_WITH_VERBOSITY
 }
 
@@ -1440,6 +1450,8 @@ bool FTraceAuxiliary::Start(EConnectionType Type, const TCHAR* Target, const TCH
 		return false;
 	}
 
+	TRACE_CPUPROFILER_EVENT_SCOPE(FTraceAux_Start);
+
 	if (GTraceAuxiliary.IsConnected())
 	{
 		UE_LOG_REF(LogCategory, Error, TEXT("Unable to start trace, already tracing to %s"), *GTraceAuxiliary.GetDest());
@@ -1457,6 +1469,7 @@ bool FTraceAuxiliary::Start(EConnectionType Type, const TCHAR* Target, const TCH
 
 	if (Channels)
 	{
+		TRACE_CPUPROFILER_EVENT_SCOPE(FTraceAux_EnableChannels);
 		UE_LOG_REF(LogCategory, Display, TEXT("Requested channels: '%s'"), Channels);
 		GTraceAuxiliary.ResetCommandlineChannels();
 		GTraceAuxiliary.AddCommandlineChannels(Channels);
@@ -1522,6 +1535,7 @@ bool FTraceAuxiliary::Resume()
 bool FTraceAuxiliary::WriteSnapshot(const TCHAR* InFilePath)
 {
 #if UE_TRACE_ENABLED
+	TRACE_CPUPROFILER_EVENT_SCOPE(FTraceAux_WriteSnapshot);
 	return GTraceAuxiliary.WriteSnapshot(InFilePath, LogTrace);
 #else
 	return true;
@@ -1532,6 +1546,7 @@ bool FTraceAuxiliary::WriteSnapshot(const TCHAR* InFilePath)
 bool FTraceAuxiliary::SendSnapshot(const TCHAR* Host, uint32 Port)
 {
 #if UE_TRACE_ENABLED
+	TRACE_CPUPROFILER_EVENT_SCOPE(FTraceAux_SendSnapshot);
 	return GTraceAuxiliary.SendSnapshot(Host, Port, LogTrace);
 #else
 	return true;
@@ -1638,7 +1653,7 @@ void FTraceAuxiliary::Initialize(const TCHAR* CommandLine)
 	{
 		InstanceId[Index] = InstanceGuid[Index];
 	}
-	
+
 	UE_TRACE_LOG(Diagnostics, Session2, UE::Trace::TraceLogChannel, DataSize)
 		<< Session2.Platform(PREPROCESSOR_TO_STRING(UBT_COMPILED_PLATFORM), PlatformLen)
 		<< Session2.AppName(AppName, AppNameLen)
@@ -1937,6 +1952,7 @@ FTraceAuxiliary::ETraceSystemStatus FTraceAuxiliary::GetTraceSystemStatus()
 void FTraceAuxiliary::TryAutoConnect()
 {
 #if UE_TRACE_ENABLED
+	TRACE_CPUPROFILER_EVENT_SCOPE(FTraceAux_TryAutoConnect);
 #if PLATFORM_WINDOWS
 	if (GTraceAutoStart && !IsConnected())
 	{
@@ -1951,16 +1967,16 @@ void FTraceAuxiliary::TryAutoConnect()
 		}
 	}
 #elif PLATFORM_MAC || PLATFORM_LINUX
-    if (GTraceAutoStart && !IsConnected())
-    {
-        sem_t* AutoConnectSemaphore = sem_open("/UnrealInsightsAutoConnect", O_RDONLY);
-        if (AutoConnectSemaphore != SEM_FAILED)
-        {
-            UE_LOG(LogTrace, Display, TEXT("Unreal Insights instance detected, auto-connecting to local trace server..."));
-            Start(EConnectionType::Network, TEXT("127.0.0.1"), GTraceAuxiliary.HasCommandlineChannels() ? nullptr : TEXT("default"), nullptr);
-            sem_close(AutoConnectSemaphore);
-        }
-    }
+	if (GTraceAutoStart && !IsConnected())
+	{
+		sem_t* AutoConnectSemaphore = sem_open("/UnrealInsightsAutoConnect", O_RDONLY);
+		if (AutoConnectSemaphore != SEM_FAILED)
+		{
+			UE_LOG(LogTrace, Display, TEXT("Unreal Insights instance detected, auto-connecting to local trace server..."));
+			Start(EConnectionType::Network, TEXT("127.0.0.1"), GTraceAuxiliary.HasCommandlineChannels() ? nullptr : TEXT("default"), nullptr);
+			sem_close(AutoConnectSemaphore);
+		}
+	}
 #endif // PLATFORMS
 #endif // UE_TRACE_ENABLED
 }
