Make PopOnAbortHandler 17.96x faster.

diff --git a/Engine/Source/Programs/AutoRTFMTests/Private/Abort.cpp b/Engine/Source/Programs/AutoRTFMTests/Private/Abort.cpp
--- a/Engine/Source/Programs/AutoRTFMTests/Private/Abort.cpp
+++ b/Engine/Source/Programs/AutoRTFMTests/Private/Abort.cpp
@@ -1,7 +1,7 @@
 // Copyright Epic Games, Inc. All Rights Reserved.
 
 #include "Catch2Includes.h"
-#include <AutoRTFM/AutoRTFM.h>
+#include <AutoRTFM/AutoRTFMTesting.h>
 #include <map>
 #include <vector>
 
@@ -733,3 +733,73 @@ TEST_CASE("Abort.StackWriteToOuterOpen")
 	REQUIRE(AutoRTFM::EContextStatus::OnTrack == CloseStatus);
 	REQUIRE("Success" == TestResult);
 }
+
+TEST_CASE("Abort.PushOnAbortOrdering")
+{
+	SECTION("InParent")
+	{
+		unsigned Orderer = 0;
+		AutoRTFM::Testing::Abort([&]
+		{
+			AutoRTFM::OnAbort([&] { REQUIRE(6 == Orderer++); });
+			AutoRTFM::PushOnAbortHandler(&Orderer, [&] { REQUIRE(5 == Orderer++); });
+			AutoRTFM::OnAbort([&] { REQUIRE(4 == Orderer++); });
+			AutoRTFM::PushOnAbortHandler(&Orderer, [&] { REQUIRE(false); });
+			AutoRTFM::OnAbort([&] { REQUIRE(3 == Orderer++); });
+			AutoRTFM::PopOnAbortHandler(&Orderer);
+			AutoRTFM::OnAbort([&] { REQUIRE(2 == Orderer++); });
+			AutoRTFM::PushOnAbortHandler(&Orderer, [&] { REQUIRE(1 == Orderer++); });
+			AutoRTFM::OnAbort([&] { REQUIRE(0 == Orderer++); });
+			AutoRTFM::AbortTransaction();
+		});
+	}
+
+	SECTION("InChild")
+	{
+		// The inner transaction cannot fail in this example so lets skip the test if we are testing that.
+		if (AutoRTFM::ForTheRuntime::ShouldRetryNestedTransactionsToo())
+		{
+			return;
+		}
+
+		unsigned Orderer = 0;
+		AutoRTFM::Testing::Abort([&]
+		{
+			AutoRTFM::OnAbort([&] { REQUIRE(20 == Orderer++); });
+			AutoRTFM::PushOnAbortHandler(&Orderer, [&] { REQUIRE(19 == Orderer++); });
+			AutoRTFM::OnAbort([&] { REQUIRE(18 == Orderer++); });
+			AutoRTFM::PushOnAbortHandler(&Orderer, [&] { REQUIRE(false); });
+			AutoRTFM::OnAbort([&] { REQUIRE(17 == Orderer++); });
+			AutoRTFM::PopOnAbortHandler(&Orderer);
+			AutoRTFM::OnAbort([&] { REQUIRE(16 == Orderer++); });
+			AutoRTFM::PushOnAbortHandler(&Orderer, [&] { REQUIRE(15 == Orderer++); });
+			AutoRTFM::OnAbort([&] { REQUIRE(14 == Orderer++); });
+
+			// This commits which will add the on abort handlers to the parent scope.
+			AutoRTFM::Commit([&]
+			{
+				AutoRTFM::OnAbort([&] { REQUIRE(13 == Orderer++); });
+				AutoRTFM::PushOnAbortHandler(&Orderer, [&] { REQUIRE(12 == Orderer++); });
+				AutoRTFM::OnAbort([&] { REQUIRE(11 == Orderer++); });
+				AutoRTFM::PushOnAbortHandler(&Orderer, [&] { REQUIRE(false); });
+				AutoRTFM::OnAbort([&] { REQUIRE(10 == Orderer++); });
+				AutoRTFM::PopOnAbortHandler(&Orderer);
+				AutoRTFM::OnAbort([&] { REQUIRE(9 == Orderer++); });
+				AutoRTFM::PushOnAbortHandler(&Orderer, [&] { REQUIRE(8 == Orderer++); });
+				AutoRTFM::OnAbort([&] { REQUIRE(7 == Orderer++); });
+			});
+
+			AutoRTFM::OnAbort([&] { REQUIRE(6 == Orderer++); });
+			AutoRTFM::PushOnAbortHandler(&Orderer, [&] { REQUIRE(5 == Orderer++); });
+			AutoRTFM::OnAbort([&] { REQUIRE(4 == Orderer++); });
+			AutoRTFM::PushOnAbortHandler(&Orderer, [&] { REQUIRE(false); });
+			AutoRTFM::OnAbort([&] { REQUIRE(3 == Orderer++); });
+			AutoRTFM::PopOnAbortHandler(&Orderer);
+			AutoRTFM::OnAbort([&] { REQUIRE(2 == Orderer++); });
+			AutoRTFM::PushOnAbortHandler(&Orderer, [&] { REQUIRE(1 == Orderer++); });
+			AutoRTFM::OnAbort([&] { REQUIRE(0 == Orderer++); });
+
+			AutoRTFM::AbortTransaction();
+		});
+	}
+}
diff --git a/Engine/Source/Programs/AutoRTFMTests/Private/Benchmarks.cpp b/Engine/Source/Programs/AutoRTFMTests/Private/Benchmarks.cpp
--- a/Engine/Source/Programs/AutoRTFMTests/Private/Benchmarks.cpp
+++ b/Engine/Source/Programs/AutoRTFMTests/Private/Benchmarks.cpp
@@ -1,7 +1,7 @@
 // Copyright Epic Games, Inc. All Rights Reserved.
 
 #include "Catch2Includes.h"
-#include <AutoRTFM/AutoRTFM.h>
+#include <AutoRTFM/AutoRTFMTesting.h>
 #include <vector>
 #include <unordered_set>
 
@@ -336,3 +336,160 @@ TEST_CASE("Graph")
 		AutoRTFM::Commit([&]() { Graph = BuildGraph(); WalkGraph(Graph); });
 	};
 }
+
+TEST_CASE("Benchmarks.PopOnAbort")
+{
+	constexpr unsigned Count = 16 * 128;
+
+	BENCHMARK("PopOnAbortSingleKey")
+	{
+		bool bHit = false;
+		AutoRTFM::Testing::Abort([&]
+		{
+			for (unsigned int Index = 0; Index < Count; Index++)
+			{
+				AutoRTFM::PushOnAbortHandler(&bHit, [&] { bHit = true; });
+				AutoRTFM::PopOnAbortHandler(&bHit);
+			}
+
+			AutoRTFM::AbortTransaction();
+		});
+		REQUIRE(!bHit);
+	};
+
+	BENCHMARK("PopOnAbortMultiKey")
+	{
+		bool bHits[Count] = { false };
+		AutoRTFM::Testing::Abort([&]
+		{
+			for (unsigned int Index = 0; Index < Count; Index++)
+			{
+				AutoRTFM::PushOnAbortHandler(&bHits[Index], [&bHits, Index] { bHits[Index] = true; });
+				AutoRTFM::PopOnAbortHandler(&bHits[Index]);
+			}
+
+			AutoRTFM::AbortTransaction();
+		});
+
+		for (unsigned int Index = 0; Index < Count; Index++)
+		{
+			REQUIRE(!bHits[Index]);
+		}
+	};
+
+	BENCHMARK("PopOnAbortShortSingleKey")
+	{
+		bool bHit = false;
+		unsigned AbortCounter = 0;
+		AutoRTFM::Testing::Abort([&]
+		{
+			for (unsigned int Index = 0; Index < Count; Index++)
+			{
+				AutoRTFM::OnAbort([&] { AbortCounter++; });
+			}
+
+			for (unsigned int Index = 0; Index < Count; Index++)
+			{
+				AutoRTFM::PushOnAbortHandler(&bHit, [&] { bHit = true; });
+			}
+
+			for (unsigned int Index = 0; Index < Count; Index++)
+			{
+				AutoRTFM::PopOnAbortHandler(&bHit);
+			}
+
+			AutoRTFM::AbortTransaction();
+		});
+		REQUIRE(!bHit);
+		REQUIRE(Count == AbortCounter);
+	};
+
+	BENCHMARK("PopOnAbortLongSingleKey")
+	{
+		bool bHit = false;
+		unsigned AbortCounter = 0;
+		AutoRTFM::Testing::Abort([&]
+		{
+			for (unsigned int Index = 0; Index < Count; Index++)
+			{
+				AutoRTFM::PushOnAbortHandler(&bHit, [&] { bHit = true; });
+			}
+
+			for (unsigned int Index = 0; Index < Count; Index++)
+			{
+				AutoRTFM::OnAbort([&] { AbortCounter++; });
+			}
+
+			for (unsigned int Index = 0; Index < Count; Index++)
+			{
+				AutoRTFM::PopOnAbortHandler(&bHit);
+			}
+
+			AutoRTFM::AbortTransaction();
+		});
+		REQUIRE(!bHit);
+		REQUIRE(Count == AbortCounter);
+	};
+
+	BENCHMARK("PopOnAbortShortMultiKey")
+	{
+		bool bHits[Count] = { false };
+		unsigned AbortCounter = 0;
+		AutoRTFM::Testing::Abort([&]
+		{
+			for (unsigned int Index = 0; Index < Count; Index++)
+			{
+				AutoRTFM::OnAbort([&] { AbortCounter++; });
+			}
+
+			for (unsigned int Index = 0; Index < Count; Index++)
+			{
+				AutoRTFM::PushOnAbortHandler(&bHits[Index], [&bHits, Index] { bHits[Index] = true; });
+			}
+
+			for (unsigned int Index = 0; Index < Count; Index++)
+			{
+				AutoRTFM::PopOnAbortHandler(&bHits[Index]);
+			}
+
+			AutoRTFM::AbortTransaction();
+		});
+
+		for (unsigned int Index = 0; Index < Count; Index++)
+		{
+			REQUIRE(!bHits[Index]);
+		}
+		REQUIRE(Count == AbortCounter);
+	};
+
+	BENCHMARK("PopOnAbortLongMultiKey")
+	{
+		bool bHits[Count] = { false };
+		unsigned AbortCounter = 0;
+		AutoRTFM::Testing::Abort([&]
+		{
+			for (unsigned int Index = 0; Index < Count; Index++)
+			{
+				AutoRTFM::PushOnAbortHandler(&bHits[Index], [&bHits, Index] { bHits[Index] = true; });
+			}
+
+			for (unsigned int Index = 0; Index < Count; Index++)
+			{
+				AutoRTFM::OnAbort([&] { AbortCounter++; });
+			}
+
+			for (unsigned int Index = 0; Index < Count; Index++)
+			{
+				AutoRTFM::PopOnAbortHandler(&bHits[Index]);
+			}
+
+			AutoRTFM::AbortTransaction();
+		});
+
+		for (unsigned int Index = 0; Index < Count; Index++)
+		{
+			REQUIRE(!bHits[Index]);
+		}
+		REQUIRE(Count == AbortCounter);
+	};
+}
diff --git a/Engine/Source/Runtime/Core/Private/AutoRTFM/TaskArray.h b/Engine/Source/Runtime/Core/Private/AutoRTFM/TaskArray.h
--- a/Engine/Source/Runtime/Core/Private/AutoRTFM/TaskArray.h
+++ b/Engine/Source/Runtime/Core/Private/AutoRTFM/TaskArray.h
@@ -3,6 +3,8 @@
 #pragma once
 
 #include "Containers/Array.h"
+#include "Containers/BitArray.h"
+#include "Containers/Map.h"
 #include "Templates/Function.h"
 #include "Utils.h"
 
@@ -25,103 +27,196 @@ struct TTaskArrayKeyValuePair
 	InValType Val;
 
 	TTaskArrayKeyValuePair() = default;
-	TTaskArrayKeyValuePair(const InKeyType& Key, const InValType& Val) : Key(Key), Val(Val){}
-	TTaskArrayKeyValuePair(InKeyType&& Key, InValType&& Val) : Key(MoveTemp(Key)), Val(MoveTemp(Val)){}
+	TTaskArrayKeyValuePair(const InKeyType& Key, const InValType& Val) : Key(Key), Val(Val) {}
+	TTaskArrayKeyValuePair(InKeyType&& Key, InValType&& Val) : Key(MoveTemp(Key)), Val(MoveTemp(Val)) {}
 };
 
 template<typename InElementType>
-class TTaskArray
+class TTaskArray final
 {
 	using SKey = const void*;
-	using SKeyValuePair = TTaskArrayKeyValuePair<SKey, InElementType>;
-	using FInternalArray = TArray<SKeyValuePair>;
-public:
 
+	using KeyedElementType = TBitArray<>;
+
+	using PairType = TTaskArrayKeyValuePair<SKey*, InElementType*>;
+
+	struct FInternalArray final
+	{
+		using IsKeyedArrayType = TBitArray<>;
+		using KeyArrayType = TArray<SKey>;
+		using ElementArrayType = TArray<InElementType>;
+		using SizeType = typename ElementArrayType::SizeType;
+
+		// This iterator is just a zip of the two data structures that make up the array.
+		template<typename IsKeyedIterType, typename KeyIterType, typename ElementIterType> struct TIterator final
+		{
+			explicit TIterator(IsKeyedIterType IsKeyedIter, KeyIterType KeyIter, ElementIterType ElementIter) :
+				IsKeyedIter(IsKeyedIter), KeyIter(KeyIter), ElementIter(ElementIter) {}
+
+			PairType operator*() const
+			{
+				if (*IsKeyedIter)
+				{
+					return PairType(&*KeyIter, &*ElementIter);
+				}
+				else
+				{
+					return PairType(nullptr, &*ElementIter);
+				}
+			}
+
+			TIterator& operator++()
+			{
+				// We only bump the key iterator if we are currently on a keyed iterator!
+				if (*IsKeyedIter)
+				{
+					++KeyIter;
+				}
+				++IsKeyedIter;
+				++ElementIter;
+				return *this;
+			}
+
+			bool operator!=(const TIterator& Other) const
+			{
+				return (Other.IsKeyedIter != IsKeyedIter) || (Other.KeyIter != KeyIter) || (Other.ElementIter != ElementIter);
+			}
+
+		private:
+			IsKeyedIterType IsKeyedIter;
+			KeyIterType KeyIter;
+			ElementIterType ElementIter;
+		};
+
+		IsKeyedArrayType IsKeyed;
+		KeyArrayType Keys;
+		ElementArrayType Elements;
+		
+		void AddKeyed(const SKey Key, const InElementType& Element)
+		{
+			IsKeyed.Add(true);
+			Keys.Add(Key);
+			Elements.Add(Element);
+		}
+
+		void AddKeyed(const SKey Key, InElementType&& Element)
+		{
+			IsKeyed.Add(true);
+			Keys.Add(Key);
+			Elements.Add(MoveTemp(Element));
+		}
+
+		void AddElement(const InElementType& Element)
+		{
+			IsKeyed.Add(false);
+			Elements.Add(Element);
+		}
+
+		void AddElement(InElementType&& Element)
+		{
+			IsKeyed.Add(false);
+			Elements.Add(MoveTemp(Element));
+		}
+
+		bool IsEmpty() const
+		{
+			const bool bResult = Elements.IsEmpty();
+			ASSERT(bResult == IsKeyed.IsEmpty());
+			return bResult;
+		}
+
+		SizeType Num() const
+		{
+			const SizeType Num = Elements.Num();
+			ASSERT(Num == IsKeyed.Num());
+			return Num;
+		}
+
+		void Empty()
+		{
+			IsKeyed.Empty();
+			Keys.Empty();
+			Elements.Empty();
+		}
+
+		using FIterator = TIterator<IsKeyedArrayType::FIterator, typename KeyArrayType::RangedForIteratorType, typename ElementArrayType::RangedForIteratorType>;
+		using FConstIterator = TIterator<IsKeyedArrayType::FConstIterator, typename KeyArrayType::RangedForConstIteratorType, typename ElementArrayType::RangedForConstIteratorType>;
+		using FReverseIterator = TIterator<IsKeyedArrayType::FReverseIterator, typename KeyArrayType::RangedForReverseIteratorType, typename ElementArrayType::RangedForReverseIteratorType>;
+		using FConstReverseIterator = TIterator<IsKeyedArrayType::FConstReverseIterator, typename KeyArrayType::RangedForConstReverseIteratorType, typename ElementArrayType::RangedForConstReverseIteratorType>;
+		
+		FIterator begin() { return TIterator(IsKeyed.begin(), Keys.begin(), Elements.begin()); }
+		FIterator end() { return TIterator(IsKeyed.end(), Keys.end(), Elements.end()); }
+
+		FConstIterator begin() const { return TIterator(IsKeyed.begin(), Keys.begin(), Elements.begin()); }
+		FConstIterator end() const { return TIterator(IsKeyed.end(), Keys.end(), Elements.end()); }
+
+		FReverseIterator rbegin() { return TIterator(IsKeyed.rbegin(), Keys.rbegin(), Elements.rbegin()); }
+		FReverseIterator rend() { return TIterator(IsKeyed.rend(), Keys.rend(), Elements.rend()); }
+
+		FConstReverseIterator rbegin() const { return TIterator(IsKeyed.rbegin(), Keys.rbegin(), Elements.rbegin()); }
+		FConstReverseIterator rend() const { return TIterator(IsKeyed.rend(), Keys.rend(), Elements.rend()); }
+	};
+
+public:
 	TTaskArray() = default;
-	TTaskArray(const TTaskArray&) = delete;
-	void operator=(const TTaskArray&) = delete;
 
     bool IsEmpty() const { return Latest.IsEmpty() && Stash.IsEmpty(); }
 
     void Add(InElementType&& Value)
     {
-		Latest.Push(TTaskArrayKeyValuePair(SKey{}, MoveTemp(Value)));
+		Latest.AddElement(MoveTemp(Value));
     }
 
 	void Add(const InElementType& Value)
 	{
-		Latest.Push(TTaskArrayKeyValuePair(SKey{}, Value));
+		Latest.AddElement(Value);
 	}
 
 	void AddKeyed(SKey Key, InElementType&& Value)
     {
 		ASSERT(Key != SKey{});
-		Latest.Push(TTaskArrayKeyValuePair(Key, MoveTemp(Value)));
+		Latest.AddKeyed(MoveTemp(Key), MoveTemp(Value));
+		KeyedElements.FindOrAdd(Key).Add(false);
     }
 
 	void AddKeyed(SKey Key, const InElementType& Value)
     {
 		ASSERT(Key != SKey{});
-		Latest.Push(TTaskArrayKeyValuePair(Key, Value));
+		Latest.AddKeyed(MoveTemp(Key), Value);
+		KeyedElements.FindOrAdd(Key).Add(false);
     }
 
 	bool DeleteKey(SKey Key)
 	{
 		ASSERT(Key != SKey{});
 
-		auto EraseLastAddedKeyFromArray = [Key](FInternalArray& Array) -> bool
-		{
-			typename FInternalArray::SizeType LastIdx = Array.FindLastByPredicate([Key](const SKeyValuePair& Pair)
-			{
-				return Pair.Key == Key;
-			});
-
-			if (LastIdx == INDEX_NONE)
-			{
-				return false;
-			}
-
-			// We found our key to erase! Nuke it.
-			Array.RemoveAt(LastIdx);
-			return true;
-		};
-
-		if (EraseLastAddedKeyFromArray(Latest))
-		{
-			return true;
-		}
+		KeyedElementType* const KeyedElement = KeyedElements.Find(Key);
 
-		for (FInternalArray& StashedVectorBox : TBackwards(Stash))
+		if (UNLIKELY(nullptr == KeyedElement))
 		{
-			if (EraseLastAddedKeyFromArray(StashedVectorBox))
-			{
-				return true;
-			}
+			return false;
 		}
 
-		return false;
+		return INDEX_NONE != KeyedElement->FindAndSetLastZeroBit();
 	}
 
 	bool DeleteAllMatchingKeys(SKey Key)
 	{
 		ASSERT(Key != SKey{});
 
-		auto EraseKeyFromArray = [Key](FInternalArray& Array) -> size_t
-		{
-			return Array.RemoveAll([Key](const SKeyValuePair& Pair)
-			{
-				return Pair.Key == Key;
-			});
-		};
+		KeyedElementType* const KeyedElement = KeyedElements.Find(Key);
 
-		size_t NumErased = EraseKeyFromArray(Latest);
+		if (UNLIKELY(nullptr == KeyedElement))
+		{
+			return false;
+		}
 
-		for (FInternalArray& StashedVectorBox : Stash)
+		for (FBitReference Bit : *KeyedElement)
 		{
-			NumErased += EraseKeyFromArray(StashedVectorBox);
+			Bit = true;
 		}
 
-		return NumErased > 0;
+		return !KeyedElement->IsEmpty();
 	}
 
     void AddAll(TTaskArray&& Other)
@@ -132,7 +227,6 @@ public:
         {
             Stash.Push(MoveTemp(StashedVectorBox));
         }
-
         Other.Stash.Empty();
 
         if (!Other.Latest.IsEmpty())
@@ -140,110 +234,58 @@ public:
             Stash.Push(MoveTemp(Other.Latest));
 			Other.Latest = FInternalArray{};
         }
-    }
 
-    template<typename TFunc>
-    bool ForEachForward(const TFunc& Func) const
-    {
-        for (FInternalArray& StashedVectorBox : Stash)
-        {
-            for (const SKeyValuePair& EntryKVP : StashedVectorBox)
-            {
-                if (!Func(EntryKVP.Val))
-                {
-                    return false;
-                }
-            }
-        }
-
-        for (const SKeyValuePair& EntryKVP : Latest)
-        {
-            if (!Func(EntryKVP.Val))
-            {
-                return false;
-            }
-        }
-        return true;
+		for (auto& Pair : Other.KeyedElements)
+		{
+			KeyedElementType& KeyedElement = KeyedElements.FindOrAdd(Pair.Key);
+			
+			for (FBitReference Bit : Pair.Value)
+			{
+				KeyedElement.Add(Bit);
+			}
+		}
+		Other.KeyedElements.Empty();
     }
 
-    template<typename TFunc>
-    bool ForEachForward(const TFunc& Func)
+    template<typename FunctorType>
+    void ForEachForward(const FunctorType& Func)
     {
         for (FInternalArray& StashedVectorBox : Stash)
         {
-            for (SKeyValuePair& EntryKVP : StashedVectorBox)
+            for (PairType Pair : StashedVectorBox)
             {
-                if (!Func(EntryKVP.Val))
-                {
-                    return false;
-                }
+				ForwardProcessEntry(Pair, Func);
             }
         }
 
-        for (SKeyValuePair& EntryKVP : Latest)
+        for (PairType Pair : Latest)
         {
-            if (!Func(EntryKVP.Val))
-            {
-                return false;
-            }
+			ForwardProcessEntry(Pair, Func);
         }
-        return true;
     }
 
-    template<typename TFunc>
-    bool ForEachBackward(const TFunc& Func) const
+    template<typename FunctorType>
+    void ForEachBackward(const FunctorType& Func)
     {
-		for (const SKeyValuePair& EntryKVP : TBackwards(Latest))
-		{
-			if (!Func(EntryKVP.Val))
-			{
-				return false;
-			}
-		}
-
-		for (const FInternalArray& StashedVectorBox : TBackwards(Stash))
+		for (PairType Pair : TBackwards(Latest))
 		{
-			for (const SKeyValuePair& EntryKVP : TBackwards(StashedVectorBox))
-			{
-				if (!Func(EntryKVP.Val))
-				{
-					return false;
-				}
-			}
-		}
-
-        return true;
-    }
-
-	template<typename TFunc>
-	bool ForEachBackward(const TFunc& Func)
-	{
-		for (SKeyValuePair& EntryKVP : TBackwards(Latest))
-		{
-			if (!Func(EntryKVP.Val))
-			{
-				return false;
-			}
+			BackwardProcessEntry(Pair, Func);
 		}
 
 		for (FInternalArray& StashedVectorBox : TBackwards(Stash))
 		{
-			for (SKeyValuePair& EntryKVP : TBackwards(StashedVectorBox))
+			for (PairType Pair : TBackwards(StashedVectorBox))
 			{
-				if (!Func(EntryKVP.Val))
-				{
-					return false;
-				}
+				BackwardProcessEntry(Pair, Func);
 			}
 		}
-
-		return true;
-	}
+    }
 
     void Reset()
     {
         Latest.Empty();
         Stash.Empty();
+		KeyedElements.Empty();
     }
 
     size_t Num() const
@@ -260,6 +302,8 @@ public:
     }
 
 private:
+	UE_NONCOPYABLE(TTaskArray)
+
     // We don't want to do this too often; currently we just do it where it's asymptitically relevant like AddAll. This doesn't
     // logically change the TaskArray but it changes its internal representation. Hence the use of `mutable` and hence why this
     // method is `const`.
@@ -271,9 +315,55 @@ private:
 			Latest = FInternalArray{};
 		}
     }
+
+	template<typename FunctorType>
+	void ForwardProcessEntry(PairType Pair, const FunctorType& Func)
+	{
+		if (Pair.Key)
+		{
+			KeyedElementType* const KeyedElement = KeyedElements.Find(*Pair.Key);
+
+			// We don't support keyed elements in the forward iterator (yet?).
+			ASSERT(!KeyedElement);
+		}
+
+		Func(*Pair.Val);
+	}
+
+	template<typename FunctorType>
+	void BackwardProcessEntry(PairType Pair, const FunctorType& Func)
+	{
+		if (Pair.Key)
+		{
+			KeyedElementType* const KeyedElement = KeyedElements.Find(*Pair.Key);
+
+			const int32 Num = KeyedElement->Num();
+
+			if (0 < Num)
+			{
+				const bool bLast = (*KeyedElement)[Num - 1];
+				KeyedElement->RemoveAt(Num - 1);
+
+				if (bLast)
+				{
+					static TFunction<void()> EmptyFunctor = [] {};
+
+					// We have a deferred signal to pop this element, so do it now.
+					*Pair.Key = SKey{};
+					Pair.Val->Reset();
+					*Pair.Val = EmptyFunctor;
+					return;
+				}
+			}
+		}
+
+		ASSERT(Pair.Val);
+		Func(*Pair.Val);
+	}
     
     mutable FInternalArray Latest;
     mutable TArray<FInternalArray> Stash;
+	mutable TMap<SKey, KeyedElementType> KeyedElements;
 };
 
 } // namespace AutoRTFM
diff --git a/Engine/Source/Runtime/Core/Private/AutoRTFM/Transaction.cpp b/Engine/Source/Runtime/Core/Private/AutoRTFM/Transaction.cpp
--- a/Engine/Source/Runtime/Core/Private/AutoRTFM/Transaction.cpp
+++ b/Engine/Source/Runtime/Core/Private/AutoRTFM/Transaction.cpp
@@ -96,14 +96,13 @@ void FTransaction::AbortWithoutThrowing()
     CommitTasks.Reset();
 
     Undo();
-	AbortTasks.ForEachBackward([&](TFunction<void()>& Task) -> bool 
+	AbortTasks.ForEachBackward([&](TFunction<void()>& Task)
     { 
         // Call and then reset each of the tasks in reverse order.
         // This ensures that the task and its destructor are called in reverse chronological order,
         // which is important if the function has captures with non-trivial destructors.
         Task();
         Task.Reset();
-        return true; 
     });
 
     if (IsNested())
@@ -241,11 +240,10 @@ bool FTransaction::AttemptToCommitOuterNest()
 
     AbortTasks.Reset();
 
-    CommitTasks.ForEachForward([] (TFunction<void()>& Task) -> bool
+    CommitTasks.ForEachForward([] (TFunction<void()>& Task)
     { 
         Task(); 
         Task.Reset();
-        return true; 
     });
 
     return true;
diff --git a/Engine/Source/Runtime/Core/Public/Containers/BitArray.h b/Engine/Source/Runtime/Core/Public/Containers/BitArray.h
--- a/Engine/Source/Runtime/Core/Public/Containers/BitArray.h
+++ b/Engine/Source/Runtime/Core/Public/Containers/BitArray.h
@@ -1511,6 +1511,61 @@ public:
 		int32 Index;
 	};
 
+	/** Reverse iterator. */
+	class FReverseIterator : public FRelativeBitReference
+	{
+	public:
+		FORCEINLINE explicit FReverseIterator(TBitArray<Allocator>& InArray UE_LIFETIMEBOUND)
+			: FRelativeBitReference(InArray.Num() - 1)
+			, Array(InArray)
+			, Index(InArray.Num() - 1)
+		{
+		}
+
+		FORCEINLINE explicit FReverseIterator(TBitArray<Allocator>& InArray UE_LIFETIMEBOUND, int32 StartIndex)
+			: FRelativeBitReference(-1)
+			, Array(InArray)
+			, Index(-1)
+		{
+		}
+
+		FORCEINLINE FReverseIterator& operator++()
+		{
+			--Index;
+			this->Mask >>= 1;
+			if (!this->Mask)
+			{
+				// Advance to the next uint32.
+				this->Mask = (1 << (NumBitsPerDWORD - 1));
+				--this->WordIndex;
+			}
+			return *this;
+		}
+
+		FORCEINLINE FBitReference operator*() const
+		{
+			return GetValue();
+		}
+
+		/** conversion to "bool" returning true if the iterator is valid. */
+		FORCEINLINE explicit operator bool() const
+		{
+			return Index >= 0;
+		}
+
+		/** inverse of the "bool" operator */
+		FORCEINLINE bool operator !() const
+		{
+			return !(bool)*this;
+		}
+
+		FORCEINLINE FBitReference GetValue() const { return FBitReference(Array.GetData()[this->WordIndex], this->Mask); }
+		FORCEINLINE int32 GetIndex() const { return Index; }
+	private:
+		TBitArray<Allocator>& Array;
+		int32 Index;
+	};
+
 	/** Const BitArray iterator. */
 	class FConstIterator : public FRelativeBitReference
 	{
@@ -1567,6 +1622,12 @@ public:
 			,	Index(InArray.Num() - 1)
 		{
 		}
+		FORCEINLINE explicit FConstReverseIterator(const TBitArray<Allocator>& InArray UE_LIFETIMEBOUND, int32 StartIndex)
+			: FRelativeBitReference(-1)
+			, Array(InArray)
+			, Index(-1)
+		{
+		}
 		FORCEINLINE FConstReverseIterator& operator++()
 		{
 			--Index;
@@ -1580,6 +1641,11 @@ public:
 			return *this;
 		}
 
+		FORCEINLINE FConstBitReference operator*() const
+		{
+			return GetValue();
+		}
+
 		/** conversion to "bool" returning true if the iterator is valid. */
 		FORCEINLINE explicit operator bool() const
 		{ 
@@ -1735,6 +1801,11 @@ public:
 	FORCEINLINE FIterator		end()			{ return FIterator(*this, NumBits); }
 	FORCEINLINE FConstIterator	end() const		{ return FConstIterator(*this, NumBits); }
 
+	FORCEINLINE FReverseIterator		rbegin() { return FReverseIterator(*this); }
+	FORCEINLINE FConstReverseIterator	rbegin() const { return FConstReverseIterator(*this); }
+	FORCEINLINE FReverseIterator		rend() { return FReverseIterator(*this, 0); }
+	FORCEINLINE FConstReverseIterator	rend() const { return FConstReverseIterator(*this, 0); }
+
 private:
 	AllocatorType AllocatorInstance;
 	int32         NumBits;
diff --git a/Engine/Source/Runtime/Core/Tests/Containers/BitArrayTests.cpp b/Engine/Source/Runtime/Core/Tests/Containers/BitArrayTests.cpp
--- a/Engine/Source/Runtime/Core/Tests/Containers/BitArrayTests.cpp
+++ b/Engine/Source/Runtime/Core/Tests/Containers/BitArrayTests.cpp
@@ -1574,6 +1574,34 @@ TEST_CASE_NAMED(FBitArrayTestFor, "System::Core::Containers::BitArray::For", "[A
 	}
 }
 
+TEST_CASE_NAMED(FBitArrayTestReverseFor, "System::Core::Containers::BitArray::ReverseFor", "[ApplicationContextMask][EngineFilter]")
+{
+	class FBackwards final
+	{
+		TBitArray<>& Obj;
+	public:
+		FBackwards(TBitArray<>& Obj) : Obj(Obj) {}
+		auto begin() { return Obj.rbegin(); }
+		auto end() { return Obj.rend(); }
+	};
+
+	std::initializer_list<bool> Expected = { true, false, true, true, false, false };
+
+	TBitArray<> Bits;
+	for (bool bBit : Expected)
+	{
+		Bits.Add(bBit);
+	}
+
+	unsigned Index = 0;
+
+	for (bool bActual : FBackwards(Bits))
+	{
+		const bool bExpected = Expected.begin()[Expected.size() - 1 - Index++];
+		CHECK(bActual == bExpected);
+	}
+}
+
 TEST_CASE_NAMED(FBitArrayInvariantsTest, "System::Core::Containers::BitArray::Invariants", "[ApplicationContextMask][SmokeFilter]")
 {
 	using namespace UE::BitArrayTest;
