Workaround the wrong PVS studio warning.

diff --git a/Engine/Source/Runtime/Core/Private/AutoRTFM/TaskArray.h b/Engine/Source/Runtime/Core/Private/AutoRTFM/TaskArray.h
--- a/Engine/Source/Runtime/Core/Private/AutoRTFM/TaskArray.h
+++ b/Engine/Source/Runtime/Core/Private/AutoRTFM/TaskArray.h
@@ -213,7 +213,7 @@ public:
 
 		for (FBitReference Bit : *KeyedElement)
 		{
-			Bit = true;
+			Bit = true; //-V714 PVS doesn't seem to understand that `FBitReference` is backed by a ref so this assignment does change the underlying value.
 		}
 
 		return !KeyedElement->IsEmpty();
