[AutoRTFM] Migrate to UE::TScopeLock instead of FTransactionallySafeScopeLock.

These are functionally identical.

diff --git a/Engine/Source/Runtime/Core/Private/HAL/ConsoleManager.cpp b/Engine/Source/Runtime/Core/Private/HAL/ConsoleManager.cpp
--- a/Engine/Source/Runtime/Core/Private/HAL/ConsoleManager.cpp
+++ b/Engine/Source/Runtime/Core/Private/HAL/ConsoleManager.cpp
@@ -10,17 +10,17 @@ ConsoleManager.cpp: console command handling
 #include "Containers/UnrealString.h"
 #include "CoreTypes.h"
 #include "HAL/CriticalSection.h"
-#include "Misc/ScopeLock.h"
-#include "Misc/ScopeRWLock.h"
-#include "Misc/Paths.h"
 #include "Stats/Stats.h"
 #include "Misc/ConfigCacheIni.h"
 #include "Misc/ConfigContext.h"
 #include "Misc/ConfigUtilities.h"
 #include "Misc/CoreDelegates.h"
 #include "Misc/OutputDeviceFile.h"
+#include "Misc/Paths.h"
 #include "Misc/RemoteConfigIni.h"
-#include "Misc/TransactionallySafeScopeLock.h"
+#include "Misc/ScopeLock.h"
+#include "Misc/ScopeRWLock.h"
+#include "Misc/TransactionallySafeCriticalSection.h"
 #include "Modules/ModuleManager.h"
 #include "HAL/PlatformProcess.h"
 #include "ProfilingDebugging/CsvProfiler.h"
@@ -2286,7 +2286,7 @@ IConsoleObject* FConsoleManager::FindConsoleObject(const TCHAR* Name, bool bTrac
 
 IConsoleObject* FConsoleManager::FindConsoleObjectUnfiltered(const TCHAR* Name) const
 {
-	FTransactionallySafeScopeLock ScopeLock( &ConsoleObjectsSynchronizationObject );
+	UE::TScopeLock ScopeLock(ConsoleObjectsSynchronizationObject);
 	IConsoleObject* Var = ConsoleObjects.FindRef(Name);
 	return Var;
 }
@@ -2297,7 +2297,7 @@ void FConsoleManager::UnregisterConsoleObject(IConsoleObject* CVar, bool bKeepSt
 	{
 		return;
 	}
-	FTransactionallySafeScopeLock ScopeLock(&ConsoleObjectsSynchronizationObject);
+	UE::TScopeLock ScopeLock(ConsoleObjectsSynchronizationObject);
 
 	// Slow search for console object
 	const FString ObjName = FindConsoleObjectName( CVar );
@@ -2310,7 +2310,7 @@ void FConsoleManager::UnregisterConsoleObject(IConsoleObject* CVar, bool bKeepSt
 
 void FConsoleManager::UnregisterConsoleObject(const TCHAR* Name, bool bKeepState)
 {
-	FTransactionallySafeScopeLock ScopeLock(&ConsoleObjectsSynchronizationObject);
+	UE::TScopeLock ScopeLock(ConsoleObjectsSynchronizationObject);
 
 	IConsoleObject* Object = FindConsoleObject(Name);
 
@@ -2409,7 +2409,7 @@ void FConsoleManager::ForEachConsoleObjectThatStartsWith(const FConsoleObjectVis
 	check(ThatStartsWith);
 
 	//@caution, potential deadlock if the visitor tries to call back into the cvar system. Best not to do this, but we could capture and array of them, then release the lock, then dispatch the visitor.
-	FTransactionallySafeScopeLock ScopeLock( &ConsoleObjectsSynchronizationObject );
+	UE::TScopeLock ScopeLock(ConsoleObjectsSynchronizationObject);
 	for(TMap<FString, IConsoleObject*>::TConstIterator PairIt(ConsoleObjects); PairIt; ++PairIt)
 	{
 		const FString& Name = PairIt.Key();
@@ -2432,7 +2432,7 @@ void FConsoleManager::ForEachConsoleObjectThatContains(const FConsoleObjectVisit
 	int32 ContainsStringLength = FCString::Strlen(ThatContains);
 
 	//@caution, potential deadlock if the visitor tries to call back into the cvar system. Best not to do this, but we could capture and array of them, then release the lock, then dispatch the visitor.
-	FTransactionallySafeScopeLock ScopeLock( &ConsoleObjectsSynchronizationObject );
+	UE::TScopeLock ScopeLock(ConsoleObjectsSynchronizationObject);
 	for(TMap<FString, IConsoleObject*>::TConstIterator PairIt(ConsoleObjects); PairIt; ++PairIt)
 	{
 		const FString& Name = PairIt.Key();
@@ -2909,7 +2909,7 @@ IConsoleObject* FConsoleManager::AddConsoleObject(const TCHAR* Name, IConsoleObj
 	check(Obj);
 	LLM_SCOPE_BYNAME(TEXT("EngineMisc/ConsoleCommands"));
 
-	FTransactionallySafeScopeLock ScopeLock( &ConsoleObjectsSynchronizationObject ); // we will lock on the entire add process
+	UE::TScopeLock ScopeLock(ConsoleObjectsSynchronizationObject); // we will lock on the entire add process
 	IConsoleObject* ExistingObj = ConsoleObjects.FindRef(Name);
 
 	if(Obj->GetFlags() & ECVF_Scalability)
@@ -3071,7 +3071,7 @@ FString FConsoleManager::FindConsoleObjectName(const IConsoleObject* InVar) cons
 {
 	check(InVar);
 
-	FTransactionallySafeScopeLock ScopeLock( &ConsoleObjectsSynchronizationObject );
+	UE::TScopeLock ScopeLock(ConsoleObjectsSynchronizationObject);
 	for(TMap<FString, IConsoleObject*>::TConstIterator PairIt(ConsoleObjects); PairIt; ++PairIt)
 	{
 		IConsoleObject* Var = PairIt.Value();
@@ -3182,7 +3182,7 @@ void FConsoleManager::GetConsoleHistory(const TCHAR* Key, TArray<FString>& Out)
 
 bool FConsoleManager::IsNameRegistered(const TCHAR* Name) const
 {
-	FTransactionallySafeScopeLock ScopeLock(&ConsoleObjectsSynchronizationObject);
+	UE::TScopeLock ScopeLock(ConsoleObjectsSynchronizationObject);
 	return ConsoleObjects.Contains(Name);
 }
 
