Change FName mmap define so that it has to be set explicitly

diff --git a/Engine/Source/Runtime/Core/Private/UObject/UnrealNames.cpp b/Engine/Source/Runtime/Core/Private/UObject/UnrealNames.cpp
--- a/Engine/Source/Runtime/Core/Private/UObject/UnrealNames.cpp
+++ b/Engine/Source/Runtime/Core/Private/UObject/UnrealNames.cpp
@@ -40,7 +40,10 @@
 #include "Misc/AsciiSet.h"
 #include "AutoRTFM/AutoRTFM.h"
 
-#define USE_FNAME_MMAP (PLATFORM_ANDROID || PLATFORM_IOS)
+#ifndef USE_FNAME_MMAP
+#define USE_FNAME_MMAP 0
+#endif
+
 #if USE_FNAME_MMAP
 #include <sys/mman.h>
 #if PLATFORM_ANDROID
