Follow-ups for the TaskArray change.

diff --git a/Engine/Source/Programs/AutoRTFMTests/Private/Benchmarks.cpp b/Engine/Source/Programs/AutoRTFMTests/Private/Benchmarks.cpp
--- a/Engine/Source/Programs/AutoRTFMTests/Private/Benchmarks.cpp
+++ b/Engine/Source/Programs/AutoRTFMTests/Private/Benchmarks.cpp
@@ -337,7 +337,7 @@ TEST_CASE("Graph")
 	};
 }
 
-TEST_CASE("Benchmarks.PopOnAbort")
+TEST_CASE("Benchmarks.PopOnAbortHandler")
 {
 	constexpr unsigned Count = 16 * 128;
 
@@ -493,3 +493,44 @@ TEST_CASE("Benchmarks.PopOnAbort")
 		REQUIRE(Count == AbortCounter);
 	};
 }
+
+TEST_CASE("Benchmarks.PopAllOnAbortHandlers")
+{
+	constexpr unsigned Count = 16 * 128;
+
+	BENCHMARK("Pop")
+	{
+		bool bHit = false;
+		AutoRTFM::Testing::Abort([&]
+		{
+			for (unsigned int Index = 0; Index < Count; Index++)
+			{
+				AutoRTFM::PushOnAbortHandler(&bHit, [&] { bHit = true; });
+			}
+
+			for (unsigned int Index = 0; Index < Count; Index++)
+			{
+				AutoRTFM::PopOnAbortHandler(&bHit);
+			}
+
+			AutoRTFM::AbortTransaction();
+		});
+		REQUIRE(!bHit);
+	};
+
+	BENCHMARK("PopAll")
+	{
+		bool bHit = false;
+		AutoRTFM::Testing::Abort([&]
+		{
+			for (unsigned int Index = 0; Index < Count; Index++)
+			{
+				AutoRTFM::PushOnAbortHandler(&bHit, [&] { bHit = true; });
+			}
+
+			AutoRTFM::PopAllOnAbortHandlers(&bHit);
+			AutoRTFM::AbortTransaction();
+		});
+		REQUIRE(!bHit);
+	};
+}
diff --git a/Engine/Source/Runtime/Core/Private/AutoRTFM/TaskArray.h b/Engine/Source/Runtime/Core/Private/AutoRTFM/TaskArray.h
--- a/Engine/Source/Runtime/Core/Private/AutoRTFM/TaskArray.h
+++ b/Engine/Source/Runtime/Core/Private/AutoRTFM/TaskArray.h
@@ -36,7 +36,7 @@ class TTaskArray final
 {
 	using SKey = const void*;
 
-	using KeyedElementType = TBitArray<>;
+	using BitMaskPerKeyType = TBitArray<>;
 
 	using PairType = TTaskArrayKeyValuePair<SKey*, InElementType*>;
 
@@ -47,7 +47,7 @@ class TTaskArray final
 		using ElementArrayType = TArray<InElementType>;
 		using SizeType = typename ElementArrayType::SizeType;
 
-		// This iterator is just a zip of the two data structures that make up the array.
+		// This iterator is just a zip of the data structures that make up the array.
 		template<typename IsKeyedIterType, typename KeyIterType, typename ElementIterType> struct TIterator final
 		{
 			explicit TIterator(IsKeyedIterType IsKeyedIter, KeyIterType KeyIter, ElementIterType ElementIter) :
@@ -55,14 +55,8 @@ class TTaskArray final
 
 			PairType operator*() const
 			{
-				if (*IsKeyedIter)
-				{
-					return PairType(&*KeyIter, &*ElementIter);
-				}
-				else
-				{
-					return PairType(nullptr, &*ElementIter);
-				}
+				SKey* const Key = (*IsKeyedIter) ? &*KeyIter : nullptr;
+				return PairType(Key, &*ElementIter);
 			}
 
 			TIterator& operator++()
@@ -77,6 +71,11 @@ class TTaskArray final
 				return *this;
 			}
 
+			bool operator==(const TIterator& Other) const
+			{
+				return !(*this != Other);
+			}
+
 			bool operator!=(const TIterator& Other) const
 			{
 				return (Other.IsKeyedIter != IsKeyedIter) || (Other.KeyIter != KeyIter) || (Other.ElementIter != ElementIter);
@@ -92,27 +91,27 @@ class TTaskArray final
 		KeyArrayType Keys;
 		ElementArrayType Elements;
 		
-		void AddKeyed(const SKey Key, const InElementType& Element)
+		void Add(const SKey Key, const InElementType& Element)
 		{
 			IsKeyed.Add(true);
 			Keys.Add(Key);
 			Elements.Add(Element);
 		}
 
-		void AddKeyed(const SKey Key, InElementType&& Element)
+		void Add(const SKey Key, InElementType&& Element)
 		{
 			IsKeyed.Add(true);
 			Keys.Add(Key);
 			Elements.Add(MoveTemp(Element));
 		}
 
-		void AddElement(const InElementType& Element)
+		void Add(const InElementType& Element)
 		{
 			IsKeyed.Add(false);
 			Elements.Add(Element);
 		}
 
-		void AddElement(InElementType&& Element)
+		void Add(InElementType&& Element)
 		{
 			IsKeyed.Add(false);
 			Elements.Add(MoveTemp(Element));
@@ -132,11 +131,11 @@ class TTaskArray final
 			return Num;
 		}
 
-		void Empty()
+		void Reset()
 		{
-			IsKeyed.Empty();
-			Keys.Empty();
-			Elements.Empty();
+			IsKeyed.Reset();
+			Keys.Reset();
+			Elements.Reset();
 		}
 
 		using FIterator = TIterator<IsKeyedArrayType::FIterator, typename KeyArrayType::RangedForIteratorType, typename ElementArrayType::RangedForIteratorType>;
@@ -164,33 +163,33 @@ public:
 
     void Add(InElementType&& Value)
     {
-		Latest.AddElement(MoveTemp(Value));
+		Latest.Add(MoveTemp(Value));
     }
 
 	void Add(const InElementType& Value)
 	{
-		Latest.AddElement(Value);
+		Latest.Add(Value);
 	}
 
 	void AddKeyed(SKey Key, InElementType&& Value)
     {
 		ASSERT(Key != SKey{});
-		Latest.AddKeyed(MoveTemp(Key), MoveTemp(Value));
-		KeyedElements.FindOrAdd(Key).Add(false);
+		Latest.Add(MoveTemp(Key), MoveTemp(Value));
+		BitMasksPerKey.FindOrAdd(Key).Add(false);
     }
 
 	void AddKeyed(SKey Key, const InElementType& Value)
     {
 		ASSERT(Key != SKey{});
-		Latest.AddKeyed(MoveTemp(Key), Value);
-		KeyedElements.FindOrAdd(Key).Add(false);
+		Latest.Add(MoveTemp(Key), Value);
+		BitMasksPerKey.FindOrAdd(Key).Add(false);
     }
 
 	bool DeleteKey(SKey Key)
 	{
 		ASSERT(Key != SKey{});
 
-		KeyedElementType* const KeyedElement = KeyedElements.Find(Key);
+		BitMaskPerKeyType* const KeyedElement = BitMasksPerKey.Find(Key);
 
 		if (UNLIKELY(nullptr == KeyedElement))
 		{
@@ -204,17 +203,14 @@ public:
 	{
 		ASSERT(Key != SKey{});
 
-		KeyedElementType* const KeyedElement = KeyedElements.Find(Key);
+		BitMaskPerKeyType* const KeyedElement = BitMasksPerKey.Find(Key);
 
 		if (UNLIKELY(nullptr == KeyedElement))
 		{
 			return false;
 		}
 
-		for (FBitReference Bit : *KeyedElement)
-		{
-			Bit = true; //-V714 PVS doesn't seem to understand that `FBitReference` is backed by a ref so this assignment does change the underlying value.
-		}
+		KeyedElement->SetRange(0, KeyedElement->Num(), true);
 
 		return !KeyedElement->IsEmpty();
 	}
@@ -227,7 +223,7 @@ public:
         {
             Stash.Push(MoveTemp(StashedVectorBox));
         }
-        Other.Stash.Empty();
+        Other.Stash.Reset();
 
         if (!Other.Latest.IsEmpty())
         {
@@ -235,16 +231,16 @@ public:
 			Other.Latest = FInternalArray{};
         }
 
-		for (auto& Pair : Other.KeyedElements)
+		for (auto& Pair : Other.BitMasksPerKey)
 		{
-			KeyedElementType& KeyedElement = KeyedElements.FindOrAdd(Pair.Key);
+			BitMaskPerKeyType& KeyedElement = BitMasksPerKey.FindOrAdd(Pair.Key);
 			
 			for (FBitReference Bit : Pair.Value)
 			{
 				KeyedElement.Add(Bit);
 			}
 		}
-		Other.KeyedElements.Empty();
+		Other.BitMasksPerKey.Reset();
     }
 
     template<typename FunctorType>
@@ -283,9 +279,9 @@ public:
 
     void Reset()
     {
-        Latest.Empty();
-        Stash.Empty();
-		KeyedElements.Empty();
+        Latest.Reset();
+        Stash.Reset();
+		BitMasksPerKey.Reset();
     }
 
     size_t Num() const
@@ -301,13 +297,11 @@ public:
 		return static_cast<size_t>(Result);
     }
 
-private:
 	UE_NONCOPYABLE(TTaskArray)
 
-    // We don't want to do this too often; currently we just do it where it's asymptitically relevant like AddAll. This doesn't
-    // logically change the TaskArray but it changes its internal representation. Hence the use of `mutable` and hence why this
-    // method is `const`.
-    void Canonicalize() const
+private:
+    // We don't want to do this too often, as it'll involve a few allocations.
+    void Canonicalize()
     {
         if (!Latest.IsEmpty())
         {
@@ -321,7 +315,7 @@ private:
 	{
 		if (Pair.Key)
 		{
-			KeyedElementType* const KeyedElement = KeyedElements.Find(*Pair.Key);
+			BitMaskPerKeyType* const KeyedElement = BitMasksPerKey.Find(*Pair.Key);
 
 			// We don't support keyed elements in the forward iterator (yet?).
 			ASSERT(!KeyedElement);
@@ -335,7 +329,7 @@ private:
 	{
 		if (Pair.Key)
 		{
-			KeyedElementType* const KeyedElement = KeyedElements.Find(*Pair.Key);
+			BitMaskPerKeyType* const KeyedElement = BitMasksPerKey.Find(*Pair.Key);
 
 			const int32 Num = KeyedElement->Num();
 
@@ -361,9 +355,9 @@ private:
 		Func(*Pair.Val);
 	}
     
-    mutable FInternalArray Latest;
-    mutable TArray<FInternalArray> Stash;
-	mutable TMap<SKey, KeyedElementType> KeyedElements;
+    FInternalArray Latest;
+    TArray<FInternalArray> Stash;
+	TMap<SKey, BitMaskPerKeyType> BitMasksPerKey;
 };
 
 } // namespace AutoRTFM
