Fix the coding convention template arg violations in the AutoRTFM header.

diff --git a/Engine/Source/Runtime/Core/Public/AutoRTFM/AutoRTFM.h b/Engine/Source/Runtime/Core/Public/AutoRTFM/AutoRTFM.h
--- a/Engine/Source/Runtime/Core/Public/AutoRTFM/AutoRTFM.h
+++ b/Engine/Source/Runtime/Core/Public/AutoRTFM/AutoRTFM.h
@@ -84,7 +84,7 @@
 #endif
 
 #if UE_AUTORTFM
-template <typename FuncType> class TFunction;
+template<typename FuncType> class TFunction;
 #endif
 
 #define UE_AUTORTFM_UNUSED(UNUSEDVAR) (void)UNUSEDVAR
@@ -491,25 +491,25 @@ namespace ForTheRuntime
 } // namespace ForTheRuntime
 #endif
 
-template<typename TFunctor>
-void AutoRTFMFunctorInvoker(void* Arg) { (*static_cast<const TFunctor*>(Arg))(); }
+template<typename FunctorType>
+void AutoRTFMFunctorInvoker(void* Arg) { (*static_cast<const FunctorType*>(Arg))(); }
 
 #if UE_AUTORTFM
 extern "C" UE_AUTORTFM_API void* autortfm_lookup_function(void* OriginalFunction, const char* Where);
 
-template<typename TFunctor>
-auto AutoRTFMLookupInstrumentedFunctorInvoker(const TFunctor& Functor) -> void(*)(void*)
+template<typename FunctorType>
+auto AutoRTFMLookupInstrumentedFunctorInvoker(const FunctorType& Functor) -> void(*)(void*)
 {
 	// keep this as a single expression to help ensure that even Debug builds optimize this.
 	// if we put intermediate results in local variables then the compiler emits loads
 	// and stores to the stack which confuse our custom pass that tries to strip away
 	// the actual call to autortfm_lookup_function
-	void (*Result)(void*) = reinterpret_cast<void(*)(void*)>(autortfm_lookup_function(reinterpret_cast<void*>(&AutoRTFMFunctorInvoker<TFunctor>), "AutoRTFMLookupInstrumentedFunctorInvoker"));
+	void (*Result)(void*) = reinterpret_cast<void(*)(void*)>(autortfm_lookup_function(reinterpret_cast<void*>(&AutoRTFMFunctorInvoker<FunctorType>), "AutoRTFMLookupInstrumentedFunctorInvoker"));
 	return Result;
 }
 #else
-template<typename TFunctor>
-auto AutoRTFMLookupInstrumentedFunctorInvoker(const TFunctor& Functor) -> void(*)(void*)
+template<typename FunctorType>
+auto AutoRTFMLookupInstrumentedFunctorInvoker(const FunctorType& Functor) -> void(*)(void*)
 {
 	return nullptr;
 }
@@ -546,14 +546,14 @@ static UE_AUTORTFM_FORCEINLINE bool IsOnCurrentTransactionStack(void* Ptr)
 // if this nested transaction succeeds.
 //
 // If AutoRTFM is disabled, the code will be ran non-transactionally.
-template<typename TFunctor>
-static UE_AUTORTFM_FORCEINLINE ETransactionResult Transact(const TFunctor& Functor)
+template<typename FunctorType>
+static UE_AUTORTFM_FORCEINLINE ETransactionResult Transact(const FunctorType& Functor)
 {
 	ETransactionResult Result =
 		static_cast<ETransactionResult>(
 			autortfm_transact(
-				&AutoRTFMFunctorInvoker<TFunctor>,
-				AutoRTFMLookupInstrumentedFunctorInvoker<TFunctor>(Functor),
+				&AutoRTFMFunctorInvoker<FunctorType>,
+				AutoRTFMLookupInstrumentedFunctorInvoker<FunctorType>(Functor),
 				const_cast<void*>(static_cast<const void*>(&Functor))));
 
 	return Result;
@@ -563,14 +563,14 @@ static UE_AUTORTFM_FORCEINLINE ETransactionResult Transact(const TFunctor& Funct
 // The reason we expose it is that it allows the caller's module to not
 // be compiled with the AutoRTFM instrumentation of functions if the only
 // thing that's being invoked is a function in the open.
-template<typename TFunctor>
-static UE_AUTORTFM_FORCEINLINE ETransactionResult TransactThenOpen(const TFunctor& Functor)
+template<typename FunctorType>
+static UE_AUTORTFM_FORCEINLINE ETransactionResult TransactThenOpen(const FunctorType& Functor)
 {
 	ETransactionResult Result =
 		static_cast<ETransactionResult>(
 			autortfm_transact_then_open(
-				&AutoRTFMFunctorInvoker<TFunctor>,
-				AutoRTFMLookupInstrumentedFunctorInvoker<TFunctor>(Functor),
+				&AutoRTFMFunctorInvoker<FunctorType>,
+				AutoRTFMLookupInstrumentedFunctorInvoker<FunctorType>(Functor),
 				const_cast<void*>(static_cast<const void*>(&Functor))));
 
 	return Result;
@@ -579,12 +579,12 @@ static UE_AUTORTFM_FORCEINLINE ETransactionResult TransactThenOpen(const TFuncto
 // Run the callback in a transaction like Transact, but abort program
 // execution if the result is anything other than autortfm_committed.
 // Useful for testing.
-template<typename TFunctor>
-static UE_AUTORTFM_FORCEINLINE void Commit(const TFunctor& Functor)
+template<typename FunctorType>
+static UE_AUTORTFM_FORCEINLINE void Commit(const FunctorType& Functor)
 {
     autortfm_commit(
-		&AutoRTFMFunctorInvoker<TFunctor>,
-		AutoRTFMLookupInstrumentedFunctorInvoker<TFunctor>(Functor),
+		&AutoRTFMFunctorInvoker<FunctorType>,
+		AutoRTFMLookupInstrumentedFunctorInvoker<FunctorType>(Functor),
 		const_cast<void*>(static_cast<const void*>(&Functor)));
 }
 
@@ -616,16 +616,16 @@ static UE_AUTORTFM_FORCEINLINE void AbortIfClosed()
 
 namespace Detail
 {
-template <typename T, typename = void>
+template<typename T, typename = void>
 struct THasAssignFromOpenToClosedMethod : std::false_type {};
-template <typename T>
+template<typename T>
 struct THasAssignFromOpenToClosedMethod<T, std::void_t<decltype(T::AutoRTFMAssignFromOpenToClosed(std::declval<T&>(), std::declval<T>()))>> : std::true_type {};
 }
 
 // Evaluates to true if the type T has a static method with the signature:
 //    static void AutoRTFMAssignFromOpenToClosed(T& Closed, U Open)
 // Where `U` is `T`, `const T&` or `T&&`. Supports both copy assignment and move assignment.
-template <typename T>
+template<typename T>
 static constexpr bool HasAssignFromOpenToClosedMethod = Detail::THasAssignFromOpenToClosedMethod<T>::value;
 
 // Template class used to declare a method for safely copying or moving an 
@@ -647,32 +647,32 @@ struct TAssignFromOpenToClosed;
 
 namespace Detail
 {
-template <typename T, typename = void>
+template<typename T, typename = void>
 struct THasAssignFromOpenToClosedTrait : std::false_type {};
-template <typename T>
+template<typename T>
 struct THasAssignFromOpenToClosedTrait<T, std::void_t<decltype(TAssignFromOpenToClosed<T>::Assign(std::declval<T&>(), std::declval<T>()))>> : std::true_type {};
 }
 
 // Evaluates to true if the type T supports assigning from open to closed transactions.
-template <typename T>
+template<typename T>
 static constexpr bool HasAssignFromOpenToClosedTrait = Detail::THasAssignFromOpenToClosedTrait<T>::value;
 
 // Specialization of TAssignFromOpenToClosed for fundamental types.
-template <typename T>
+template<typename T>
 struct TAssignFromOpenToClosed<T, std::enable_if_t<std::is_fundamental_v<T>>>
 {
 	UE_AUTORTFM_FORCEINLINE static void Assign(T& Closed, T Open) { Closed = Open; }
 };
 
 // Specialization of TAssignFromOpenToClosed for raw pointer types.
-template <typename T>
+template<typename T>
 struct TAssignFromOpenToClosed<T*, void>
 {
 	UE_AUTORTFM_FORCEINLINE static void Assign(T*& Closed, T* Open) { Closed = Open; }
 };
 
 // Specialization of TAssignFromOpenToClosed for std::tuple.
-template <typename ... TYPES>
+template<typename ... TYPES>
 struct TAssignFromOpenToClosed<std::tuple<TYPES...>, std::enable_if_t<(HasAssignFromOpenToClosedTrait<TYPES> && ...)>>
 {
 	template<size_t I = 0, typename SRC = void>
@@ -697,7 +697,7 @@ struct TAssignFromOpenToClosed<std::tuple<TYPES...>, std::enable_if_t<(HasAssign
 // with the signature:
 //    static void AutoRTFMAssignFromOpenToClosed(T& Closed, U Open)
 // Where `U` is `T`, `const T&` or `T&&`. Supports both copy assignment and move assignment.
-template <typename T>
+template<typename T>
 struct TAssignFromOpenToClosed<T, std::enable_if_t<HasAssignFromOpenToClosedMethod<T>>>
 {
 	template<typename OPEN>
@@ -708,26 +708,26 @@ struct TAssignFromOpenToClosed<T, std::enable_if_t<HasAssignFromOpenToClosedMeth
 };
 
 // Specialization of TAssignFromOpenToClosed for `void` (used to make IsSafeToReturnFromOpen<void> work).
-template <>
+template<>
 struct TAssignFromOpenToClosed<void, void>;
 
 // Evaluates to true if the type T is safe to return from Open().
-template <typename T>
+template<typename T>
 static constexpr bool IsSafeToReturnFromOpen = HasAssignFromOpenToClosedTrait<T> || std::is_same_v<T, void>;
 
 // Executes the given code non-transactionally regardless of whether we are in
 // a transaction or not. Returns the value returned by Functor.
-// TReturn must be void or a type that can be safely copied from the open to a closed transaction.
+// ReturnType must be void or a type that can be safely copied from the open to a closed transaction.
 // TAssignFromOpenToClosed must have a specialization for the type that is being returned.
 template
 <
 	EMemoryValidationLevel VALIDATION_LEVEL = EMemoryValidationLevel::Default,
-	typename TFunctor = void,
-	typename TReturn = decltype(std::declval<TFunctor>()())
+	typename FunctorType = void,
+	typename ReturnType = decltype(std::declval<FunctorType>()())
 > 
-static UE_AUTORTFM_FORCEINLINE TReturn Open(const TFunctor& Functor)
+static UE_AUTORTFM_FORCEINLINE ReturnType Open(const FunctorType& Functor)
 {
-	static_assert(IsSafeToReturnFromOpen<TReturn>,
+	static_assert(IsSafeToReturnFromOpen<ReturnType>,
 		"function return type is not safe to return from Open()");
 
 	if (!autortfm_is_closed())
@@ -735,15 +735,15 @@ static UE_AUTORTFM_FORCEINLINE TReturn Open(const TFunctor& Functor)
 		return Functor();
 	}
 
-	if constexpr (IsSafeToReturnFromOpen<TReturn>)
+	if constexpr (IsSafeToReturnFromOpen<ReturnType>)
 	{
-		if constexpr (std::is_same_v<void, TReturn>)
+		if constexpr (std::is_same_v<void, ReturnType>)
 		{
 			struct FCallHelper
 			{
 				UE_AUTORTFM_NOAUTORTFM static void Call(void* Arg)
 				{
-					const TFunctor& Fn = *reinterpret_cast<TFunctor*>(Arg);
+					const FunctorType& Fn = *reinterpret_cast<FunctorType*>(Arg);
 					UE_AUTORTFM_CALLSITE_FORCEINLINE Fn();
 				}
 			};
@@ -768,10 +768,10 @@ static UE_AUTORTFM_FORCEINLINE TReturn Open(const TFunctor& Functor)
 				{
 					FCallHelper& Self = *reinterpret_cast<FCallHelper*>(Arg);
 					UE_AUTORTFM_CALLSITE_FORCEINLINE
-						TAssignFromOpenToClosed<TReturn>::Assign(Self.ReturnValue, std::move(Self.Functor()));
+						TAssignFromOpenToClosed<ReturnType>::Assign(Self.ReturnValue, std::move(Self.Functor()));
 				}
-				const TFunctor& Functor;
-				TReturn ReturnValue{};
+				const FunctorType& Functor;
+				ReturnType ReturnValue{};
 			};
 			FCallHelper Helper{Functor};
 			if constexpr (VALIDATION_LEVEL == EMemoryValidationLevel::Default)
@@ -794,12 +794,12 @@ static UE_AUTORTFM_FORCEINLINE TReturn Open(const TFunctor& Functor)
 // (whether we are in open or closed code).
 //
 // Will crash if called outside of a transaction nest.
-template<typename TFunctor> [[nodiscard]] static UE_AUTORTFM_FORCEINLINE EContextStatus Close(const TFunctor& Functor)
+template<typename FunctorType> [[nodiscard]] static UE_AUTORTFM_FORCEINLINE EContextStatus Close(const FunctorType& Functor)
 {
     return static_cast<EContextStatus>(
 		autortfm_close(
-			&AutoRTFMFunctorInvoker<TFunctor>,
-			AutoRTFMLookupInstrumentedFunctorInvoker<TFunctor>(Functor),
+			&AutoRTFMFunctorInvoker<FunctorType>,
+			AutoRTFMLookupInstrumentedFunctorInvoker<FunctorType>(Functor),
 			const_cast<void*>(static_cast<const void*>(&Functor))));
 }
 
@@ -808,11 +808,11 @@ template<typename TFunctor> [[nodiscard]] static UE_AUTORTFM_FORCEINLINE EContex
 // this just adds the work to the work deferred until the outer nest's commit.
 // If this is called outside a transaction or from an open nest then the work
 // happens immediately.
-template<typename TFunctor> static UE_AUTORTFM_FORCEINLINE void OnCommit(TFunctor&& Work)
+template<typename FunctorType> static UE_AUTORTFM_FORCEINLINE void OnCommit(FunctorType&& Work)
 {
 	if (autortfm_is_closed())
 	{
-		ForTheRuntime::OnCommitInternal(std::forward<TFunctor>(Work));
+		ForTheRuntime::OnCommitInternal(std::forward<FunctorType>(Work));
 	}
 	else
 	{
@@ -824,23 +824,23 @@ template<typename TFunctor> static UE_AUTORTFM_FORCEINLINE void OnCommit(TFuncto
 // this just adds the work to the work deferred until the outer nest's commit.
 // If this is called outside a transaction or from an open nest then the work
 // happens immediately.
-template<typename TFunctor> static UE_AUTORTFM_FORCEINLINE void OnCommit(TFunctor&& Work) { Work(); }
+template<typename FunctorType> static UE_AUTORTFM_FORCEINLINE void OnCommit(FunctorType&& Work) { Work(); }
 #endif
 
 #if UE_AUTORTFM
 // Have some work happen when this transaction aborts. If this is called
 // outside a transaction or from an open nest then the work is ignored.
-template<typename TFunctor> static UE_AUTORTFM_FORCEINLINE void OnAbort(TFunctor&& Work)
+template<typename FunctorType> static UE_AUTORTFM_FORCEINLINE void OnAbort(FunctorType&& Work)
 {
 	if (autortfm_is_closed())
 	{
-		ForTheRuntime::OnAbortInternal(std::forward<TFunctor>(Work));
+		ForTheRuntime::OnAbortInternal(std::forward<FunctorType>(Work));
 	}
 }
 #else
 // Have some work happen when this transaction aborts. If this is called
 // outside a transaction or from an open nest then the work is ignored.
-template<typename TFunctor> static UE_AUTORTFM_FORCEINLINE void OnAbort(TFunctor&&) {}
+template<typename FunctorType> static UE_AUTORTFM_FORCEINLINE void OnAbort(FunctorType&&) {}
 #endif
 
 #if UE_AUTORTFM
@@ -848,11 +848,11 @@ template<typename TFunctor> static UE_AUTORTFM_FORCEINLINE void OnAbort(TFunctor
 // the handler can be unregistered (see PopOnAbortHandler). This is useful
 // for scoped mutations that need an abort handler present unless execution
 // reaches the end of the relevant scope.
-template<typename TFunctor> static UE_AUTORTFM_FORCEINLINE void PushOnAbortHandler(const void* Key, TFunctor&& Work)
+template<typename FunctorType> static UE_AUTORTFM_FORCEINLINE void PushOnAbortHandler(const void* Key, FunctorType&& Work)
 {
 	if (autortfm_is_closed())
 	{
-		ForTheRuntime::PushOnAbortHandlerInternal(Key, std::forward<TFunctor>(Work));
+		ForTheRuntime::PushOnAbortHandlerInternal(Key, std::forward<FunctorType>(Work));
 	}
 }
 #else
@@ -860,7 +860,7 @@ template<typename TFunctor> static UE_AUTORTFM_FORCEINLINE void PushOnAbortHandl
 // the handler can be unregistered (see PopOnAbortHandler). This is useful
 // for scoped mutations that need an abort handler present unless execution
 // reaches the end of the relevant scope.
-template<typename TFunctor> static UE_AUTORTFM_FORCEINLINE void PushOnAbortHandler(const void* Key, TFunctor&&) {}
+template<typename FunctorType> static UE_AUTORTFM_FORCEINLINE void PushOnAbortHandler(const void* Key, FunctorType&&) {}
 #endif
 
 #if UE_AUTORTFM
@@ -919,9 +919,9 @@ static UE_AUTORTFM_FORCEINLINE void RecordOpenWrite(void* Ptr, size_t Size)
 }
 
 // Informs the runtime that a block of memory is about to be overwritten.
-template<typename TTYPE> static UE_AUTORTFM_FORCEINLINE void RecordOpenWrite(TTYPE* Ptr)
+template<typename Type> static UE_AUTORTFM_FORCEINLINE void RecordOpenWrite(Type* Ptr)
 {
-	autortfm_record_open_write(Ptr, sizeof(TTYPE));
+	autortfm_record_open_write(Ptr, sizeof(Type));
 }
 
 // A collection of power-user functions that are reserved for use by the AutoRTFM runtime only.
@@ -1058,7 +1058,7 @@ namespace ForTheRuntime
 	static UE_AUTORTFM_FORCEINLINE void RecordOpenRead(void const*, size_t) {}
 
 	// Reserved for future.
-	template<typename TTYPE> static UE_AUTORTFM_FORCEINLINE void RecordOpenRead(TTYPE*) {}
+	template<typename Type> static UE_AUTORTFM_FORCEINLINE void RecordOpenRead(Type*) {}
 
 	// WriteMemory first records the memory span as written (see RecordOpenWrite) and then copies the specified value into it.
 	static UE_AUTORTFM_FORCEINLINE void WriteMemory(void* DestPtr, void const* SrcPtr, size_t Size)
@@ -1068,30 +1068,30 @@ namespace ForTheRuntime
 	}
 
 	// WriteMemory first records the memory span as written (see RecordOpenWrite) and then copies the specified value into it.
-	template<typename TTYPE> static UE_AUTORTFM_FORCEINLINE void WriteMemory(TTYPE* DestPtr, TTYPE const* SrcPtr)
+	template<typename Type> static UE_AUTORTFM_FORCEINLINE void WriteMemory(Type* DestPtr, Type const* SrcPtr)
 	{
-		if constexpr (std::is_trivially_copyable<TTYPE>::value)
+		if constexpr (std::is_trivially_copyable<Type>::value)
 		{
-			RecordOpenWrite(DestPtr, sizeof(TTYPE));
+			RecordOpenWrite(DestPtr, sizeof(Type));
 			*DestPtr = *SrcPtr;
 		}
 		else
 		{
-			WriteMemory(DestPtr, SrcPtr, sizeof(TTYPE));
+			WriteMemory(DestPtr, SrcPtr, sizeof(Type));
 		}
 	}
 
 	// WriteMemory first records the memory span as written (see RecordOpenWrite) and then copies the specified value into it.
-	template<typename TTYPE> static UE_AUTORTFM_FORCEINLINE void WriteMemory(TTYPE* DestPtr, TTYPE const SrcValue)
+	template<typename Type> static UE_AUTORTFM_FORCEINLINE void WriteMemory(Type* DestPtr, Type const SrcValue)
 	{
-		if constexpr (std::is_trivially_copyable<TTYPE>::value)
+		if constexpr (std::is_trivially_copyable<Type>::value)
 		{
-			RecordOpenWrite(DestPtr, sizeof(TTYPE));
+			RecordOpenWrite(DestPtr, sizeof(Type));
 			*DestPtr = SrcValue;
 		}
 		else
 		{
-			WriteMemory(DestPtr, &SrcValue, sizeof(TTYPE));
+			WriteMemory(DestPtr, &SrcValue, sizeof(Type));
 		}
 	}
 
