Prevent FTraceAuxiliary to initalized twice.

diff --git a/Engine/Source/Runtime/Core/Private/ProfilingDebugging/TraceAuxiliary.cpp b/Engine/Source/Runtime/Core/Private/ProfilingDebugging/TraceAuxiliary.cpp
--- a/Engine/Source/Runtime/Core/Private/ProfilingDebugging/TraceAuxiliary.cpp
+++ b/Engine/Source/Runtime/Core/Private/ProfilingDebugging/TraceAuxiliary.cpp
@@ -1585,6 +1585,13 @@ static void TraceAuxiliaryAddPostForkCallback(const TCHAR* CommandLine)
 ////////////////////////////////////////////////////////////////////////////////
 void FTraceAuxiliary::Initialize(const TCHAR* CommandLine)
 {
+	static bool bInitialized = false;
+	checkf(!bInitialized, TEXT("FTraceAuxiliary may only be initialized once."));
+	if (bInitialized)
+	{
+		return;
+	}
+	
 	TRACE_CPUPROFILER_EVENT_SCOPE(FTraceAux_Init);
 	UE_MEMSCOPE(TRACE_TAG);
 
@@ -1598,7 +1605,6 @@ void FTraceAuxiliary::Initialize(const TCHAR* CommandLine)
 #endif
 
 #if UE_TRACE_ENABLED
-	UE_LOG(LogTrace, Log, TEXT("Initializing trace..."));
 
 	// Setup message callback so we get feedback from TraceLog
 	UE::Trace::SetMessageCallback(&OnMessageCallback);
@@ -1617,6 +1623,11 @@ void FTraceAuxiliary::Initialize(const TCHAR* CommandLine)
 		return;
 	}
 
+	// Only set this post fork if used.
+	bInitialized = true;
+
+	UE_LOG(LogTrace, Log, TEXT("Initializing trace..."));
+	
 	const TCHAR* AppName = TEXT(UE_APP_NAME);
 	const TCHAR* ProjectName = TEXT("");
 #if IS_MONOLITHIC && !IS_PROGRAM
