Make `FIntervalTree::Contains` 2x faster by swizzling the backing data structure from array of structs (with gaps) to struct of arrays.

This results in an 8.45% uplift in the BASIC.verse benchmark.

#rnx

diff --git a/Engine/Source/Runtime/Core/Private/AutoRTFM/IntervalTree.h b/Engine/Source/Runtime/Core/Private/AutoRTFM/IntervalTree.h
--- a/Engine/Source/Runtime/Core/Private/AutoRTFM/IntervalTree.h
+++ b/Engine/Source/Runtime/Core/Private/AutoRTFM/IntervalTree.h
@@ -2,8 +2,8 @@
 
 #pragma once
 
-#include <cstdint>
-#include "Containers/PagedArray.h"
+#include "Containers/Array.h"
+#include "Containers/BitArray.h"
 #include "Utils.h"
 
 namespace AutoRTFM
@@ -41,6 +41,11 @@ template<typename T> struct TIntervalTreeArray final
 		return Payload.Pop();
 	}
 
+	FORCEINLINE void Reset()
+	{
+		Payload.Reset();
+	}
+
 	FORCEINLINE T& operator[](uint32 Index)
 	{
 		// A bit nasty - but since we know we are in bounds, skip the check.
@@ -85,9 +90,7 @@ struct FIntervalTree final
 
         do
         {
-            const FIntervalTreeNode& Node = Nodes[Current];
-
-            const FRange Range = Node.Range;
+            const FRange Range = NodeRanges[Current];
 
             // This check does not need to prove that NewRange is entirely
             // enclosed within Range, because if any byte of NewRange was in the
@@ -97,7 +100,8 @@ struct FIntervalTree final
                 return true;
             }
 
-			Current = (NewRange.Start < Range.Start) ? Node.Left : Node.Right;
+			const ArrayType<FIntervalTreeNodeIndex>& Next = (NewRange.Start < Range.End) ? NodeLefts : NodeRights;
+			Current = Next[Current];
         } while (IntervalTreeNodeIndexNone != Current);
 
         return false;
@@ -111,7 +115,11 @@ struct FIntervalTree final
     void Reset()
     {
         Root = IntervalTreeNodeIndexNone;
-        Nodes.Empty();
+		NodeRanges.Reset();
+		NodeLefts.Reset();
+		NodeRights.Reset();
+		NodeParents.Reset();
+		NodeColors.Reset();
     }
 
     void Merge(const FIntervalTree& Other)
@@ -129,12 +137,12 @@ struct FIntervalTree final
         {
             const FIntervalTreeNodeIndex Current = ToProcess.Pop();
 
-            FRange Range = Other.Nodes[Current].Range;
+            FRange Range = Other.NodeRanges[Current];
 
             Insert(MoveTemp(Range));
 
-            const FIntervalTreeNodeIndex Left = Other.Nodes[Current].Left;
-            const FIntervalTreeNodeIndex Right = Other.Nodes[Current].Right;
+            const FIntervalTreeNodeIndex Left = Other.NodeLefts[Current];
+            const FIntervalTreeNodeIndex Right = Other.NodeRights[Current];
 
             if (IntervalTreeNodeIndexNone != Left)
             {
@@ -164,32 +172,47 @@ private:
 
     static constexpr FIntervalTreeNodeIndex IntervalTreeNodeIndexNone = UINT32_MAX;
 
-    struct FIntervalTreeNode final
-    {
-        explicit FIntervalTreeNode(const FRange Range) :
-            Range(Range), Parent(IntervalTreeNodeIndexNone), bIsBlack(true) {}
-
-		FIntervalTreeNode(const FRange Range, FIntervalTreeNodeIndex Parent) :
-			Range(Range), Parent(Parent), bIsBlack(false) {}
-
-		FRange Range;
-        FIntervalTreeNodeIndex Left = IntervalTreeNodeIndexNone;
-        FIntervalTreeNodeIndex Right = IntervalTreeNodeIndexNone;
-        FIntervalTreeNodeIndex Parent;
-        bool bIsBlack;
-    };
+	enum EColor : bool
+	{
+		Red = false,
+		Black = true
+	};
 
     FIntervalTreeNodeIndex Root = IntervalTreeNodeIndexNone;
-	ArrayType<FIntervalTreeNode> Nodes;
+	ArrayType<FRange> NodeRanges;
+	ArrayType<FIntervalTreeNodeIndex> NodeLefts;
+	ArrayType<FIntervalTreeNodeIndex> NodeRights;
+	ArrayType<FIntervalTreeNodeIndex> NodeParents;
+	TBitArray<> NodeColors;
 
 	static constexpr bool bExtraDebugging = false;
 
+	FORCEINLINE FIntervalTreeNodeIndex AddNode(FRange&& NewRange)
+	{
+		const FIntervalTreeNodeIndex Index = NodeRanges.Add(NewRange);
+		NodeLefts.Add(IntervalTreeNodeIndexNone);
+		NodeRights.Add(IntervalTreeNodeIndexNone);
+		NodeParents.Add(IntervalTreeNodeIndexNone);
+		NodeColors.Add(EColor::Black);
+		return Index;
+	}
+
+	FORCEINLINE FIntervalTreeNodeIndex AddNode(FRange&& NewRange, FIntervalTreeNodeIndex Parent)
+	{
+		const FIntervalTreeNodeIndex Index = NodeRanges.Add(NewRange);
+		NodeLefts.Add(IntervalTreeNodeIndexNone);
+		NodeRights.Add(IntervalTreeNodeIndexNone);
+		NodeParents.Add(Parent);
+		NodeColors.Add(EColor::Red);
+		return Index;
+	}
+
     FORCENOINLINE bool Insert(FRange&& NewRange)
     {
         if (UNLIKELY(IntervalTreeNodeIndexNone == Root))
         {
-            ASSERT(0 == Nodes.Num());
-            Root = Nodes.Add(FIntervalTreeNode(NewRange));
+            ASSERT(NodeRanges.IsEmpty());
+			Root = AddNode(MoveTemp(NewRange));
             return true;
         }
 
@@ -197,7 +220,7 @@ private:
 
         for(;;)
         {
-            const FRange Range = Nodes[Current].Range;
+            const FRange Range = NodeRanges[Current];
 
             if (UNLIKELY((NewRange.Start < Range.End) && (Range.Start < NewRange.End)))
             {
@@ -211,18 +234,18 @@ private:
                     ASSERT(NewRange.Start < Range.Start);
 
 					// We can just modify the existing node in place.
-					Nodes[Current].Range.Start = NewRange.Start;
+					NodeRanges[Current].Start = NewRange.Start;
                     return true;
                 }
-                else if (IntervalTreeNodeIndexNone == Nodes[Current].Left)
+                else if (IntervalTreeNodeIndexNone == NodeLefts[Current])
                 {
-                    const FIntervalTreeNodeIndex Index = Nodes.Add(FIntervalTreeNode(NewRange, Current));
-					Nodes[Current].Left = Index;
+					const FIntervalTreeNodeIndex Index = AddNode(MoveTemp(NewRange), Current);
+					NodeLefts[Current] = Index;
                     Current = Index;
                     break;
                 }
 
-                Current = Nodes[Current].Left;
+                Current = NodeLefts[Current];
             }
             else
             {
@@ -231,18 +254,18 @@ private:
                     ASSERT(NewRange.End > Range.End);
 
 					// We can just modify the existing node in place.
-					Nodes[Current].Range.End = NewRange.End;
+					NodeRanges[Current].End = NewRange.End;
                     return true;
                 }
-                else if (IntervalTreeNodeIndexNone == Nodes[Current].Right)
+                else if (IntervalTreeNodeIndexNone == NodeRights[Current])
                 {
-                    const FIntervalTreeNodeIndex Index = Nodes.Add(FIntervalTreeNode(NewRange, Current));
-					Nodes[Current].Right = Index;
+                    const FIntervalTreeNodeIndex Index = AddNode(MoveTemp(NewRange), Current);
+					NodeRights[Current] = Index;
                     Current = Index;
                     break;
                 }
 
-                Current = Nodes[Current].Right;
+                Current = NodeRights[Current];
             }
 
             ASSERT(Root != Current);
@@ -250,19 +273,19 @@ private:
 
         auto IsBlack = [this](FIntervalTreeNodeIndex Index)
         {
-            return IntervalTreeNodeIndexNone == Index || Nodes[Index].bIsBlack;
+			return (IntervalTreeNodeIndexNone == Index) || (EColor::Black == NodeColors[Index]);
         };
 
         for(;;)
         {
-            FIntervalTreeNodeIndex Parent = Nodes[Current].Parent;
+            FIntervalTreeNodeIndex Parent = NodeParents[Current];
 
             UE_ASSUME(Current != Parent);
 
             // The root will always have a black parent, so this check covers both.
             if (Parent == Root)
             {
-                Nodes[Parent].bIsBlack = true;
+                NodeColors[Parent] = EColor::Black;
                 break;
             }
             else if (IsBlack(Parent))
@@ -270,32 +293,30 @@ private:
                 break;
             }
 
-            const FIntervalTreeNodeIndex GrandParent = Nodes[Parent].Parent;
+            const FIntervalTreeNodeIndex GrandParent = NodeParents[Parent];
 
             UE_ASSUME((Parent != GrandParent) && (Current != GrandParent));
 
-            const bool bGrandParentIsRoot = GrandParent == Root;
-
-            const bool bParentIsLeft = Nodes[GrandParent].Left == Parent;
+            const bool bParentIsLeft = NodeLefts[GrandParent] == Parent;
 
             // The uncle is the other node of our parent.
-            const FIntervalTreeNodeIndex Uncle = bParentIsLeft ? Nodes[GrandParent].Right : Nodes[GrandParent].Left;
+            const FIntervalTreeNodeIndex Uncle = bParentIsLeft ? NodeRights[GrandParent] : NodeLefts[GrandParent];
 
             UE_ASSUME((GrandParent != Uncle) && (Parent != Uncle) && (Current != Uncle));
 
             if (!IsBlack(Uncle))
             {
                 ASSERT(IsBlack(GrandParent));
-                Nodes[Parent].bIsBlack = true;
-                Nodes[Uncle].bIsBlack = true;
+                NodeColors[Parent] = EColor::Black;
+                NodeColors[Uncle] = EColor::Black;
 
-                if (bGrandParentIsRoot)
+                if (GrandParent == Root)
                 {
                     break;
                 }
                 else
                 {
-                    Nodes[GrandParent].bIsBlack = false;
+					NodeColors[GrandParent] = EColor::Red;
                 }
 
                 Current = GrandParent;
@@ -303,85 +324,90 @@ private:
             }
 
             // Our uncle is black, so we need to swizzle around.
-            const bool CurrentIsLeft = Nodes[Parent].Left == Current;
+            const bool bCurrentIsLeft = NodeLefts[Parent] == Current;
 
             if (bParentIsLeft)
             {
-                if (!CurrentIsLeft)
+                if (!bCurrentIsLeft)
                 {
-                    Nodes[Parent].Right = Nodes[Current].Left;
+                    NodeRights[Parent] = NodeLefts[Current];
 
-                    if (IntervalTreeNodeIndexNone != Nodes[Parent].Right)
+                    if (IntervalTreeNodeIndexNone != NodeRights[Parent])
                     {
-                        Nodes[Nodes[Parent].Right].Parent = Parent;
+                        NodeParents[NodeRights[Parent]] = Parent;
                     }
 
-                    Nodes[Parent].Parent = Current;
-                    Nodes[Current].Parent = GrandParent;
-                    Nodes[Current].Left = Parent;
-                    Nodes[GrandParent].Left = Current;
+                    NodeParents[Parent] = Current;
+                    NodeParents[Current] = GrandParent;
+                    NodeLefts[Current] = Parent;
+                    NodeLefts[GrandParent] = Current;
                     std::swap(Parent, Current);
                 }
 
-                Nodes[Parent].Parent = Nodes[GrandParent].Parent;
-                Nodes[GrandParent].Left = Nodes[Parent].Right;
+                NodeParents[Parent] = NodeParents[GrandParent];
+                NodeLefts[GrandParent] = NodeRights[Parent];
 
-                if (IntervalTreeNodeIndexNone != Nodes[GrandParent].Left)
+                if (IntervalTreeNodeIndexNone != NodeLefts[GrandParent])
                 {
-                    Nodes[Nodes[GrandParent].Left].Parent = GrandParent;
+                    NodeParents[NodeLefts[GrandParent]] = GrandParent;
                 }
 
-                Nodes[Parent].Right = GrandParent;
-                Nodes[GrandParent].Parent = Parent;
+                NodeRights[Parent] = GrandParent;
+                NodeParents[GrandParent] = Parent;
 
-                std::swap(Nodes[GrandParent].bIsBlack, Nodes[Parent].bIsBlack);
+				const bool bGrandParentIsBlack = NodeColors[GrandParent];
+				NodeColors[GrandParent] = NodeColors[Parent];
+				NodeColors[Parent] = bGrandParentIsBlack;
             }
             else
             {
-                if (CurrentIsLeft)
+                if (bCurrentIsLeft)
                 {
-                    Nodes[Parent].Left = Nodes[Current].Right;
+                    NodeLefts[Parent] = NodeRights[Current];
                     
-                    if (IntervalTreeNodeIndexNone != Nodes[Parent].Left)
+                    if (IntervalTreeNodeIndexNone != NodeLefts[Parent])
                     {
-                        Nodes[Nodes[Parent].Left].Parent = Parent;
+                        NodeParents[NodeLefts[Parent]] = Parent;
                     }
 
-                    Nodes[Parent].Parent = Current;
-                    Nodes[Current].Parent = GrandParent;
-                    Nodes[Current].Right = Parent;
-                    Nodes[GrandParent].Right = Current;
+                    NodeParents[Parent] = Current;
+                    NodeParents[Current] = GrandParent;
+                    NodeRights[Current] = Parent;
+                    NodeRights[GrandParent] = Current;
                     std::swap(Parent, Current);
                 }
 
-                Nodes[Parent].Parent = Nodes[GrandParent].Parent;
-                Nodes[GrandParent].Right = Nodes[Parent].Left;
+                NodeParents[Parent] = NodeParents[GrandParent];
+                NodeRights[GrandParent] = NodeLefts[Parent];
 
-                if (IntervalTreeNodeIndexNone != Nodes[GrandParent].Right)
+                if (IntervalTreeNodeIndexNone != NodeRights[GrandParent])
                 {
-                    Nodes[Nodes[GrandParent].Right].Parent = GrandParent;
+                    NodeParents[NodeRights[GrandParent]] = GrandParent;
                 }
 
-                Nodes[Parent].Left = GrandParent;
-                Nodes[GrandParent].Parent = Parent;
+                NodeLefts[Parent] = GrandParent;
+                NodeParents[GrandParent] = Parent;
 
-                std::swap(Nodes[GrandParent].bIsBlack, Nodes[Parent].bIsBlack);
+				const bool bGrandParentIsBlack = NodeColors[GrandParent];
+				NodeColors[GrandParent] = NodeColors[Parent];
+				NodeColors[Parent] = bGrandParentIsBlack;
             }
 
-            if (bGrandParentIsRoot)
+            if (IntervalTreeNodeIndexNone == NodeParents[Parent])
             {
                 Root = Parent;
             }
             else
             {
-                FIntervalTreeNode& GreatGrandParent = Nodes[Nodes[Parent].Parent];
-                if (GrandParent == GreatGrandParent.Left)
+				const FIntervalTreeNodeIndex GreatGrandParent = NodeParents[Parent];
+
+                if (GrandParent == NodeLefts[GreatGrandParent])
                 {
-                    GreatGrandParent.Left = Parent;
+					NodeLefts[GreatGrandParent] = Parent;
                 }
                 else
                 {
-                    GreatGrandParent.Right = Parent;
+					NodeRights[GreatGrandParent] = Parent;
                 }
             }
 
@@ -411,11 +437,11 @@ private:
 		if constexpr (bExtraDebugging)
 		{
 			ASSERT(IntervalTreeNodeIndexNone != Index);
-			ASSERT(Index < Nodes.Num());
+			ASSERT(Index < NodeRanges.Num());
 
-			const FIntervalTreeNodeIndex Parent = Nodes[Index].Parent;
-			const FIntervalTreeNodeIndex Left = Nodes[Index].Left;
-			const FIntervalTreeNodeIndex Right = Nodes[Index].Right;
+			const FIntervalTreeNodeIndex Parent = NodeParents[Index];
+			const FIntervalTreeNodeIndex Left = NodeLefts[Index];
+			const FIntervalTreeNodeIndex Right = NodeRights[Index];
 
 			if (IntervalTreeNodeIndexNone == Parent)
 			{
@@ -423,8 +449,8 @@ private:
 			}
 			else
 			{
-				ASSERT(Nodes[Parent].bIsBlack || Nodes[Index].bIsBlack);
-				ASSERT((Nodes[Parent].Left == Index) ^ (Nodes[Parent].Right == Index));
+				ASSERT(NodeColors[Parent] || NodeColors[Index]);
+				ASSERT((NodeLefts[Parent] == Index) ^ (NodeRights[Parent] == Index));
 			}
 
 			ASSERT(Left != Index);
