[AutoRTFM] Fully remove [[clang::autortfm]] support and [[clang::noinline]] from RTFM macros

A reland of 36578620

[[clang::autortfm]] has been a no-op attribute for some time, but the frontend still consumed it. This is now properly removed.

Remove [[clang::noinline]] from the UE_AUTORTFM_NOAUTORTFM, UE_AUTORTFM_ALWAYS_OPEN and AUTORTFM_CORE_API macros. This attribute is implicitly added by the other clang::noautortfm, clang::autortfm_always_open and clang::autortfm_core_api attributes.

diff --git a/Engine/Source/Runtime/Core/Private/AutoRTFM/API.cpp b/Engine/Source/Runtime/Core/Private/AutoRTFM/API.cpp
--- a/Engine/Source/Runtime/Core/Private/AutoRTFM/API.cpp
+++ b/Engine/Source/Runtime/Core/Private/AutoRTFM/API.cpp
@@ -289,7 +289,7 @@ namespace AutoRTFM
 // This is the implementation of the AutoRTFM.h API. Ideally, functions here should just delegate to some internal API.
 // For now, I have these functions also perform some error checking.
 
-#define AUTORTFM_CORE_API [[clang::autortfm_core_api, clang::noinline]]
+#define AUTORTFM_CORE_API [[clang::autortfm_core_api]]
 
 namespace AutoRTFM
 {
@@ -411,8 +411,8 @@ AUTORTFM_CORE_API void RTFM_autortfm_open(autortfm_memory_validation_level Valid
 	FTransaction* Transaction = FContext::Get()->GetCurrentTransaction();
 	if (Transaction)
 	{
-		AutoRTFM::EMemoryValidationLevel Level 
-			= ValidationLevel == autortfm_memory_validation_level_default 
+		AutoRTFM::EMemoryValidationLevel Level
+			= ValidationLevel == autortfm_memory_validation_level_default
 			? ForTheRuntime::GetMemoryValidationLevel()
 			: static_cast<AutoRTFM::EMemoryValidationLevel>(ValidationLevel);
 		Transaction->SetOpenActive(Level);
@@ -730,7 +730,7 @@ extern "C" void autortfm_record_open_write(void* Ptr, size_t Size)
 	{
 		return RTFM_autortfm_record_open_write(Ptr, Size);
 	}
-	else if (!FContext::IsTransactional())	
+	else if (!FContext::IsTransactional())
 	{
 		return;
 	}
diff --git a/Engine/Source/Runtime/Core/Public/AutoRTFM/AutoRTFM.h b/Engine/Source/Runtime/Core/Public/AutoRTFM/AutoRTFM.h
--- a/Engine/Source/Runtime/Core/Public/AutoRTFM/AutoRTFM.h
+++ b/Engine/Source/Runtime/Core/Public/AutoRTFM/AutoRTFM.h
@@ -35,9 +35,8 @@
 #endif
 
 #if UE_AUTORTFM
-#define UE_AUTORTFM_NOAUTORTFM [[clang::noautortfm, clang::noinline]]
-// #jira SOL-6589: remove clang::noinline once this JIRA is fixed.
-#define UE_AUTORTFM_ALWAYS_OPEN [[clang::autortfm_always_open, clang::noinline]]
+#define UE_AUTORTFM_NOAUTORTFM [[clang::noautortfm]]
+#define UE_AUTORTFM_ALWAYS_OPEN [[clang::autortfm_always_open]]
 #define UE_AUTORTFM_CALLSITE_FORCEINLINE [[clang::always_inline]]
 #else
 #define UE_AUTORTFM_NOAUTORTFM
@@ -103,7 +102,7 @@ extern "C"
 //   for AutoRTFM. So very focused tests work best when written in C.
 // - It makes compiler optimizations much easier as there is no mangling to
 // 	 consider when looking for functions in the runtime we can optimize.
-// 
+//
 // We use snake_case for C API surface area to make it easy to distinguish.
 //
 // The C API should not be used directly - it is here purely as an
@@ -306,8 +305,8 @@ static UE_AUTORTFM_FORCEINLINE void autortfm_record_open_write(void* Ptr, size_t
 #if UE_AUTORTFM_ENABLED
 UE_AUTORTFM_API void autortfm_register_open_function(void* original_function, void* new_function);
 #else
-static UE_AUTORTFM_FORCEINLINE void autortfm_register_open_function(void* original_function, void* new_function) 
-{ 
+static UE_AUTORTFM_FORCEINLINE void autortfm_register_open_function(void* original_function, void* new_function)
+{
 	UE_AUTORTFM_UNUSED(original_function);
 	UE_AUTORTFM_UNUSED(new_function);
 }
@@ -433,7 +432,7 @@ enum class ETransactionResult
 	AbortedByCascade = autortfm_aborted_by_cascade
 };
 
-// The context status shows what state the AutoRTFM context is currently in. 
+// The context status shows what state the AutoRTFM context is currently in.
 enum class EContextStatus : uint8_t
 {
 	// An Idle status means we are not in transactional code.
@@ -531,16 +530,16 @@ static UE_AUTORTFM_FORCEINLINE bool IsClosed() { return autortfm_is_closed(); }
 static UE_AUTORTFM_FORCEINLINE bool IsCommittingOrAborting() { return autortfm_is_committing_or_aborting(); }
 
 // Returns true if the passed-in pointer is on the stack of the currently-executing transaction.
-// This is occasionally necessary when writing OnAbort handlers for objects on the stack, since 
+// This is occasionally necessary when writing OnAbort handlers for objects on the stack, since
 // we don't want to scribble on stack memory that might have been reused.
-static UE_AUTORTFM_FORCEINLINE bool IsOnCurrentTransactionStack(void* Ptr) 
-{ 
+static UE_AUTORTFM_FORCEINLINE bool IsOnCurrentTransactionStack(void* Ptr)
+{
 	return autortfm_is_on_current_transaction_stack(Ptr);
 }
 
 // Run the functor in a transaction. Memory writes and other side effects get instrumented
 // and will be reversed if the transaction aborts.
-// 
+//
 // If this begins a nested transaction, the instrumented effects are logged onto the root
 // transaction, so the effects can be reversed later if the root transaction aborts, even
 // if this nested transaction succeeds.
@@ -559,7 +558,7 @@ static UE_AUTORTFM_FORCEINLINE ETransactionResult Transact(const FunctorType& Fu
 	return Result;
 }
 
-// This is just like calling Transact([&] { Open([&] { Functor(); }); });  
+// This is just like calling Transact([&] { Open([&] { Functor(); }); });
 // The reason we expose it is that it allows the caller's module to not
 // be compiled with the AutoRTFM instrumentation of functions if the only
 // thing that's being invoked is a function in the open.
@@ -628,7 +627,7 @@ struct THasAssignFromOpenToClosedMethod<T, std::void_t<decltype(T::AutoRTFMAssig
 template<typename T>
 static constexpr bool HasAssignFromOpenToClosedMethod = Detail::THasAssignFromOpenToClosedMethod<T>::value;
 
-// Template class used to declare a method for safely copying or moving an 
+// Template class used to declare a method for safely copying or moving an
 // object of type T from open to closed transactions.
 // Specializations of TAssignFromOpenToClosed must have at least one static
 // method with the signature:
@@ -724,7 +723,7 @@ template
 	EMemoryValidationLevel VALIDATION_LEVEL = EMemoryValidationLevel::Default,
 	typename FunctorType = void,
 	typename ReturnType = decltype(std::declval<FunctorType>()())
-> 
+>
 static UE_AUTORTFM_FORCEINLINE ReturnType Open(const FunctorType& Functor)
 {
 	static_assert(IsSafeToReturnFromOpen<ReturnType>,
@@ -911,7 +910,7 @@ static UE_AUTORTFM_FORCEINLINE void DidFree(void* Ptr)
 // Informs the runtime that a block of memory is about to be overwritten in the open.
 // During a transaction, this allows the runtime to copy the data in preparation for
 // a possible abort. Normally, tracking memory overwrites should be automatically
-// handled by AutoRTFM, but manual overwrite tracking may be required for third-party 
+// handled by AutoRTFM, but manual overwrite tracking may be required for third-party
 // libraries or outside compilers (such as ISPC).
 static UE_AUTORTFM_FORCEINLINE void RecordOpenWrite(void* Ptr, size_t Size)
 {
@@ -946,7 +945,7 @@ namespace ForTheRuntime
 	};
 
 	// An enum to represent whether we should abort and retry transactions (for testing purposes).
-	enum EAutoRTFMRetryTransactionState 
+	enum EAutoRTFMRetryTransactionState
 	{
 		// Do not abort and retry transactions (the default).
 		NoRetry = 0,
@@ -1012,7 +1011,7 @@ namespace ForTheRuntime
 	{
 		return autortfm_start_transaction();
 	}
-	
+
 	// Manually commit the top transaction nest, popping it from the execution scope.
 	// Can only be called within an already active parent transaction (EG. this cannot end
 	// a transaction nest itself).
