[AutoRTFM] Record write made by CreateGUID

And disable memory validation, as this is hot code.

diff --git a/Engine/Source/Programs/AutoRTFMTests/Private/UEPlatform.cpp b/Engine/Source/Programs/AutoRTFMTests/Private/UEPlatform.cpp
--- a/Engine/Source/Programs/AutoRTFMTests/Private/UEPlatform.cpp
+++ b/Engine/Source/Programs/AutoRTFMTests/Private/UEPlatform.cpp
@@ -3,17 +3,67 @@
 #include "Catch2Includes.h"
 #include "AutoRTFM/AutoRTFM.h"
 #include "HAL/PlatformMisc.h"
+#include "Utils.h"
 
 TEST_CASE("FPlatformMisc.CreateGUID")
 {
-	AutoRTFM::ForTheRuntime::SetEnsureOnAbortByLanguage(true);
-	FGuid Guid;
-	AutoRTFM::ETransactionResult Result = AutoRTFM::Transact([&]
+	AUTORTFM_SCOPED_SET_ENSURE_ON_ABORT_BY_LANGUAGE(true);
+
+	SECTION("Commit")
+	{
+		FGuid Guid;
+		SECTION("GuidOutsideTransactionStack")
+		{
+			AutoRTFM::ETransactionResult Result = AutoRTFM::Transact([&]
+			{
+				FPlatformMisc::CreateGuid(Guid);
+			});
+
+			REQUIRE(AutoRTFM::ETransactionResult::Committed == Result);
+		}
+
+		SECTION("GuidInsideTransactionStack")
+		{
+			AutoRTFM::ETransactionResult Result = AutoRTFM::Transact([&]
+			{
+				FGuid InnerGuid;
+				FPlatformMisc::CreateGuid(InnerGuid);
+				AutoRTFM::Open([&]{ Guid = InnerGuid; });
+			});
+
+			REQUIRE(AutoRTFM::ETransactionResult::Committed == Result);
+		}
+
+		REQUIRE(Guid != FGuid());
+	}
+
+	SECTION("Abort")
 	{
-		FPlatformMisc::CreateGuid(Guid);
-	});
-	AutoRTFM::ForTheRuntime::SetEnsureOnAbortByLanguage(false);
+		FGuid Guid;
+		SECTION("GuidOutsideTransactionStack")
+		{
+			AutoRTFM::ETransactionResult Result = AutoRTFM::Transact([&]
+			{
+				FPlatformMisc::CreateGuid(Guid);
+				AutoRTFM::AbortTransaction();
+			});
+
+			REQUIRE(AutoRTFM::ETransactionResult::AbortedByRequest == Result);
+		}
+
+		SECTION("GuidInsideTransactionStack")
+		{
+			AutoRTFM::ETransactionResult Result = AutoRTFM::Transact([&]
+			{
+				FGuid InnerGuid;
+				FPlatformMisc::CreateGuid(InnerGuid);
+				Guid = InnerGuid;
+				AutoRTFM::AbortTransaction();
+			});
+
+			REQUIRE(AutoRTFM::ETransactionResult::AbortedByRequest == Result);
+		}
 
-	REQUIRE(AutoRTFM::ETransactionResult::Committed == Result);
-	REQUIRE(Guid != FGuid());
+		REQUIRE(Guid == FGuid());
+	}
 }
diff --git a/Engine/Source/Programs/AutoRTFMTests/Private/UObjectTests.cpp b/Engine/Source/Programs/AutoRTFMTests/Private/UObjectTests.cpp
--- a/Engine/Source/Programs/AutoRTFMTests/Private/UObjectTests.cpp
+++ b/Engine/Source/Programs/AutoRTFMTests/Private/UObjectTests.cpp
@@ -364,7 +364,7 @@ TEST_CASE("UObject.FUObjectThreadContext")
 		
 		static void CtorCreateInnerClosed(const FObjectInitializer& ObjectInitializer, UMyAutoRTFMTestObject& Object)
 		{
-			AutoRTFM::Close([&]
+			(void) AutoRTFM::Close([&]
 			{
 				REQUIRE(1 == FUObjectThreadContext::Get().IsInConstructor);
 				CreateObjectWithCtor(nullptr);
@@ -385,7 +385,7 @@ TEST_CASE("UObject.FUObjectThreadContext")
 		static void CtorAbort(const FObjectInitializer& ObjectInitializer, UMyAutoRTFMTestObject& Object)
 		{
 			REQUIRE(1 == FUObjectThreadContext::Get().IsInConstructor);
-			AutoRTFM::Close([&]
+			(void) AutoRTFM::Close([&]
 			{
 				AutoRTFM::AbortTransaction();
 				FAIL(/* unreachable */);
diff --git a/Engine/Source/Programs/AutoRTFMTests/Private/Utils.h b/Engine/Source/Programs/AutoRTFMTests/Private/Utils.h
--- a/Engine/Source/Programs/AutoRTFMTests/Private/Utils.h
+++ b/Engine/Source/Programs/AutoRTFMTests/Private/Utils.h
@@ -53,7 +53,6 @@ private:
 	std::vector<FString> Warnings;
 };
 
-
 // A helper class that temporarily changes the memory validation level for the lifetime of the
 // object, restoring the original level on destruction.
 class ScopedMemoryValidationLevel
@@ -90,4 +89,31 @@ static constexpr const TCHAR* kMemoryModifiedWarning =
     TEXT("Memory modified in a transaction was also modified in an call to AutoRTFM::Open(). ")
     TEXT("This may lead to memory corruption if the transaction is aborted.");
 
+// A helper class that temporarily changes the 'EnsureOnAbortByLanguage' setting
+// for the lifetime of the object, restoring the original setting on destruction.
+class ScopedEnsureOnAbortByLanguage
+{
+public:
+	ScopedEnsureOnAbortByLanguage(bool NewValue)
+		: PrevValue(AutoRTFM::ForTheRuntime::IsEnsureOnAbortByLanguageEnabled())
+	{
+		AutoRTFM::ForTheRuntime::SetEnsureOnAbortByLanguage(NewValue);
+	}
+
+	~ScopedEnsureOnAbortByLanguage()
+	{
+		AutoRTFM::ForTheRuntime::SetEnsureOnAbortByLanguage(PrevValue);
+	}
+private:
+	ScopedEnsureOnAbortByLanguage(const ScopedEnsureOnAbortByLanguage&) = delete;
+	ScopedEnsureOnAbortByLanguage& operator = (const ScopedEnsureOnAbortByLanguage&) = delete;
+
+	const bool PrevValue;
+};
+
+#define AUTORTFM_SCOPED_SET_ENSURE_ON_ABORT_BY_LANGUAGE(VALUE) \
+	AutoRTFMTestUtils::ScopedEnsureOnAbortByLanguage ScopedSetEnsureOnAbortByLanguage \
+		{VALUE}; \
+	static_assert(true) /* require semicolon */
+
 }
diff --git a/Engine/Source/Runtime/Core/Private/Unix/UnixPlatformMisc.cpp b/Engine/Source/Runtime/Core/Private/Unix/UnixPlatformMisc.cpp
--- a/Engine/Source/Runtime/Core/Private/Unix/UnixPlatformMisc.cpp
+++ b/Engine/Source/Runtime/Core/Private/Unix/UnixPlatformMisc.cpp
@@ -1650,29 +1650,47 @@ namespace
 	}
 
 #endif // !SYS_getrandom
+
+	void CreateGuidImpl(FGuid& Result)
+	{
+		int BytesRead = SysGetRandom(&Result, sizeof(Result));
+
+		if (BytesRead == sizeof(Result))
+		{
+			// https://tools.ietf.org/html/rfc4122#section-4.4
+			// https://en.wikipedia.org/wiki/Universally_unique_identifier
+			//
+			// The 4 bits of digit M indicate the UUID version, and the 1–3
+			//   most significant bits of digit N indicate the UUID variant.
+			// xxxxxxxx-xxxx-Mxxx-Nxxx-xxxxxxxxxxxx
+			Result[1] = (Result[1] & 0xffff0fff) | 0x00004000; // version 4
+			Result[2] = (Result[2] & 0x3fffffff) | 0x80000000; // variant 1
+		}
+		else
+		{
+			// Fall back to generic CreateGuid
+			FGenericPlatformMisc::CreateGuid(Result);
+		}
+	}
 }
 
 // If we fail to create a Guid with urandom fallback to the generic platform.
 // This maybe need to be tweaked for Servers and hard fail here
-UE_AUTORTFM_ALWAYS_OPEN void FUnixPlatformMisc::CreateGuid(FGuid& Result)
+void FUnixPlatformMisc::CreateGuid(FGuid& Result)
 {
-	int BytesRead = SysGetRandom(&Result, sizeof(Result));
-
-	if (BytesRead == sizeof(Result))
+	if (AutoRTFM::IsTransactional())
 	{
-		// https://tools.ietf.org/html/rfc4122#section-4.4
-		// https://en.wikipedia.org/wiki/Universally_unique_identifier
-		//
-		// The 4 bits of digit M indicate the UUID version, and the 1–3
-		//   most significant bits of digit N indicate the UUID variant.
-		// xxxxxxxx-xxxx-Mxxx-Nxxx-xxxxxxxxxxxx
-		Result[1] = (Result[1] & 0xffff0fff) | 0x00004000; // version 4
-		Result[2] = (Result[2] & 0x3fffffff) | 0x80000000; // variant 1
+		// Memory validation disabled as this is hot code.
+		Result = AutoRTFM::Open<AutoRTFM::EMemoryValidationLevel::Disabled>([]
+		{
+			FGuid Guid;
+			CreateGuidImpl(Guid);
+			return Guid;
+		});
 	}
 	else
 	{
-		// Fall back to generic CreateGuid
-		FGenericPlatformMisc::CreateGuid(Result);
+		CreateGuidImpl(Result);
 	}
 }
 
diff --git a/Engine/Source/Runtime/Core/Private/Windows/WindowsPlatformMisc.cpp b/Engine/Source/Runtime/Core/Private/Windows/WindowsPlatformMisc.cpp
--- a/Engine/Source/Runtime/Core/Private/Windows/WindowsPlatformMisc.cpp
+++ b/Engine/Source/Runtime/Core/Private/Windows/WindowsPlatformMisc.cpp
@@ -1415,9 +1415,22 @@ const TCHAR* FWindowsPlatformMisc::GetSystemErrorMessage(TCHAR* OutBuffer, int32
 	return OutBuffer;
 }
 
-UE_AUTORTFM_ALWAYS_OPEN void FWindowsPlatformMisc::CreateGuid(FGuid& Result)
+void FWindowsPlatformMisc::CreateGuid(FGuid& Result)
 {
-	verify( CoCreateGuid( (GUID*)&Result )==S_OK );
+	if (AutoRTFM::IsTransactional())
+	{
+		// Memory validation disabled as this is hot code.
+		Result = AutoRTFM::Open<AutoRTFM::EMemoryValidationLevel::Disabled>([]
+		{
+			FGuid Guid;
+			verify( CoCreateGuid( (GUID*)&Guid )==S_OK );
+			return Guid;
+		});
+	}
+	else
+	{
+		verify( CoCreateGuid( (GUID*)&Result )==S_OK );
+	}
 }
 
 class FWindowsDialog
diff --git a/Engine/Source/Runtime/Core/Public/Misc/Guid.h b/Engine/Source/Runtime/Core/Public/Misc/Guid.h
--- a/Engine/Source/Runtime/Core/Public/Misc/Guid.h
+++ b/Engine/Source/Runtime/Core/Public/Misc/Guid.h
@@ -263,6 +263,13 @@ public:
 		return true;
 	}
 
+	// Allows a FGuid to be returned from AutoRTFM::Open().
+	// See AutoRTFM::Open().
+	static inline void AutoRTFMAssignFromOpenToClosed(FGuid& Closed, FGuid Open)
+	{
+		Closed = Open;
+	}
+
 public:
 
 	/**
