Reduced MB2 bins from 32KB-16 to 13104. This saved us around 16 MB

diff --git a/Engine/Source/Runtime/Core/Private/HAL/MallocBinned2.cpp b/Engine/Source/Runtime/Core/Private/HAL/MallocBinned2.cpp
--- a/Engine/Source/Runtime/Core/Private/HAL/MallocBinned2.cpp
+++ b/Engine/Source/Runtime/Core/Private/HAL/MallocBinned2.cpp
@@ -95,9 +95,11 @@ static constexpr uint16 SmallBinSizes[] =
 	9360,		// /7	16b
 	10912,		// /6	64b
 	13104,		// /5	16b
+#if !AGGRESSIVE_MEMORY_SAVING
 	16384-16,	// /4
 	21840,		// /3	16b
 	32768-16	// /2
+#endif
 };
 
 MS_ALIGN(PLATFORM_CACHE_LINE_SIZE) static uint8 UnusedAlignPadding[PLATFORM_CACHE_LINE_SIZE] GCC_ALIGN(PLATFORM_CACHE_LINE_SIZE) = { 0 };
diff --git a/Engine/Source/Runtime/Core/Public/HAL/MallocBinned2.h b/Engine/Source/Runtime/Core/Public/HAL/MallocBinned2.h
--- a/Engine/Source/Runtime/Core/Public/HAL/MallocBinned2.h
+++ b/Engine/Source/Runtime/Core/Public/HAL/MallocBinned2.h
@@ -25,8 +25,14 @@ struct FGenericMemoryStats;
 #endif
 
 #define UE_MB2_LARGE_ALLOC					65536		// Alignment of OS-allocated pointer - pool-allocated pointers will have a non-aligned pointer
-#define UE_MB2_MAX_SMALL_POOL_SIZE			(32768-16)	// Maximum bin size in SmallBinSizes in cpp file
-#define UE_MB2_SMALL_POOL_COUNT				51
+
+#if AGGRESSIVE_MEMORY_SAVING
+#	define UE_MB2_MAX_SMALL_POOL_SIZE		(13104)		// Maximum bin size in SmallBinSizes in cpp file
+#	define UE_MB2_SMALL_POOL_COUNT			48
+#else
+#	define UE_MB2_MAX_SMALL_POOL_SIZE		(32768-16)	// Maximum bin size in SmallBinSizes in cpp file
+#	define UE_MB2_SMALL_POOL_COUNT			51
+#endif
 
 
 // When book keeping is at the end of FFreeBlock, MallocBinned2 cannot tell if the allocation comes from a large allocation (higher than 64KB, also named as "OSAllocation") 
