[FormatStringSan] decay all types before hitting the checker's body

diff --git a/Engine/Source/Runtime/Core/Public/String/FormatStringSan.h b/Engine/Source/Runtime/Core/Public/String/FormatStringSan.h
--- a/Engine/Source/Runtime/Core/Public/String/FormatStringSan.h
+++ b/Engine/Source/Runtime/Core/Public/String/FormatStringSan.h
@@ -33,7 +33,9 @@ namespace UE::Core::Private
 	{
 		// Returns true when the type is a const char*, const TCHAR*, const char[] or const TCHAR[]
 		template<typename T>
-		inline constexpr bool bIsAConstString = !(std::is_convertible_v<std::decay_t<T>, char*> || std::is_convertible_v<std::decay_t<T>, TCHAR*>) && (std::is_convertible_v<std::decay_t<T>, const char*> || std::is_convertible_v<std::decay_t<T>, const TCHAR*>);
+		inline constexpr bool bIsAConstString =
+			   !(std::is_convertible_v<T, char*> || std::is_convertible_v<T, TCHAR*>)
+			&& (std::is_convertible_v<T, const char*> || std::is_convertible_v<T, const TCHAR*>);
 
 		enum EFormatStringSanStatus
 		{
@@ -129,7 +131,7 @@ namespace UE::Core::Private
 		struct TFmtArgTypes {};
 
 		template <typename... Ts>
-		constexpr TFmtArgTypes<Ts...> GetFmtArgTypes(Ts...)
+		constexpr TFmtArgTypes<std::decay_t<Ts>...> GetFmtArgTypes(Ts...)
 		{
 			return {};
 		}
@@ -494,7 +496,7 @@ namespace UE::Core::Private
 			{
 				if constexpr (std::is_convertible_v<const CharType*, const TCHAR*>)
 				{
-					typedef TCheckFormatString<TFmtArgTypes<ArgTypes...>> Checker;
+					typedef TCheckFormatString<TFmtArgTypes<std::decay_t<ArgTypes>...>> Checker;
 					FResult Result = Checker::Check(false, 0, Fmt);
 
 #define UE_FORMAT_STRING_ERROR_BEGIN								\
