[AutoRTFM] Chase function pointers through relocation thunks (Win64)

Calling a function implemented in an imported DLL will pass through a thunk function which is emitted by the dynamic linker. This thunk function will usually perform a jump to load of a PC-relative function pointer.

AutoRTFM looks up the closed function at runtime, either through the "magic mike" function prefix data, or via a hashmap of open function pointer to closed function pointer. Either way, AutoRTFM requires the true open function to find the closed.

To handle AutoRTFM across DLL boundaries, if the lookup fails, then attempt to inspect the first x64 instruction of the function, assuming its a relocation thunk. If the first instruction is a jump to a load of a relative function pointer, then decode that pointer and attempt to look up that function address.

diff --git a/Engine/Source/Runtime/Core/Private/AutoRTFM/FunctionMap.cpp b/Engine/Source/Runtime/Core/Private/AutoRTFM/FunctionMap.cpp
--- a/Engine/Source/Runtime/Core/Private/AutoRTFM/FunctionMap.cpp
+++ b/Engine/Source/Runtime/Core/Private/AutoRTFM/FunctionMap.cpp
@@ -2,6 +2,7 @@
 
 #if (defined(__AUTORTFM) && __AUTORTFM)
 #include "FunctionMap.h"
+#include "FunctionMapInlines.h"
 #include "ContextInlines.h"
 #include "Atomic.h"
 #include "FastLock.h"
@@ -106,6 +107,29 @@ void FunctionMapAdd(void* OldFunction, void* NewFunction)
     }
 }
 
+void* FunctionMapLookupTryFollowingRelocation(void* Function, const char* Where)
+{
+#if defined(_x86_64) || defined(__x86_64__) || defined(_M_X64) || defined(_M_AMD64)
+	// Note: Windows has multiple ways to perform relocations:
+	// https://learn.microsoft.com/en-us/windows/win32/debug/pe-format#base-relocation-types
+	// This handles the only relocation mode currently observed, but may need to
+	// support more modes in the future.
+	std::byte* Ptr = reinterpret_cast<std::byte*>(Function);
+	uint16 U16;
+	memcpy(&U16, Ptr, 2);
+	if (U16 == 0x25ff) // jmp qword ptr [rip+<relative-address>]
+	{
+		uint32 RelativeAddress;
+		memcpy(&RelativeAddress, Ptr+2, 4);
+		Ptr += RelativeAddress + 6; // apply relative offset
+		memcpy(&Ptr, Ptr, 8);       // load pointer
+		AUTORTFM_MUST_TAIL return FunctionMapLookup(Ptr, Where);
+	}
+#endif
+
+	AUTORTFM_MUST_TAIL return FunctionMapReportError(Function, Where);
+}
+
 void* FunctionMapReportError(void* OldFunction, const char* Where)
 {
 	if (Where)
diff --git a/Engine/Source/Runtime/Core/Private/AutoRTFM/FunctionMap.h b/Engine/Source/Runtime/Core/Private/AutoRTFM/FunctionMap.h
--- a/Engine/Source/Runtime/Core/Private/AutoRTFM/FunctionMap.h
+++ b/Engine/Source/Runtime/Core/Private/AutoRTFM/FunctionMap.h
@@ -34,7 +34,11 @@ struct FFunctionMap {
     FFunctionEntry Entries[1];
 };
 
-void FunctionMapAdd(void* OldFunction, void* NewFuncton);
+void FunctionMapAdd(void* OldFunction, void* NewFunction);
+
+// Attempts to lookup the function through a relocation thunk.
+// This is called by FunctionMapLookup(), and should not be called directly.
+void* FunctionMapLookupTryFollowingRelocation(void* Function, const char* Where);
 
 void* FunctionMapReportError(void* OldFunction, const char* Where);
 
diff --git a/Engine/Source/Runtime/Core/Private/AutoRTFM/FunctionMapInlines.h b/Engine/Source/Runtime/Core/Private/AutoRTFM/FunctionMapInlines.h
--- a/Engine/Source/Runtime/Core/Private/AutoRTFM/FunctionMapInlines.h
+++ b/Engine/Source/Runtime/Core/Private/AutoRTFM/FunctionMapInlines.h
@@ -32,7 +32,7 @@ inline void* FunctionMapLookup(void* OldFunction, const char* Where)
 
 	if (UNLIKELY(!Result))
 	{
-		AUTORTFM_MUST_TAIL return FunctionMapReportError(OldFunction, Where);
+		AUTORTFM_MUST_TAIL return FunctionMapLookupTryFollowingRelocation(OldFunction, Where);
 	}
 
 	return Result;
