FConfigContext::GenerateDestIniFile() - Config file wouldn't be created if its sole content was user settings

diff --git a/Engine/Source/Runtime/Core/Private/Misc/ConfigContext.cpp b/Engine/Source/Runtime/Core/Private/Misc/ConfigContext.cpp
--- a/Engine/Source/Runtime/Core/Private/Misc/ConfigContext.cpp
+++ b/Engine/Source/Runtime/Core/Private/Misc/ConfigContext.cpp
@@ -1031,11 +1031,7 @@ bool FConfigContext::GenerateDestIniFile()
 	Branch->FinalCombinedLayers.Empty();
 
 	// read the static files into the branch
-	bool bResult = LoadIniFileHierarchy();
-	if (bResult == false)
-	{
-		return false;
-	}
+	LoadIniFileHierarchy();
 	
 #if !IS_PROGRAM
 	// Don't try to load any generated files from disk in cooked builds. We will always use the re-generated INIs.
