[AutoRTFM] Move all of LogfImpl into the open.

As part of this change, the bulk of GROWABLE_LOGF has moved out of a macro and into a templated function. The macro is now a small wrapper which builds a va_list, then calls the function.

An older version of this CL can be seen at https://p4-swarm.epicgames.net/reviews/36908345 . This seems like a superior implementation to me, but I'm leaving both versions up so reviewers can decide.

(Worth noting--this code is extremely similar to StringBuilder.AppendV, but using FMemory::SystemMalloc instead of plain FMemory::Malloc. It might be possible to remove most of the remaining logic here entirely if we were to make StringBuilder's allocator a little more flexible.)

diff --git a/Engine/Source/Runtime/Core/Private/Logging/LogMacros.cpp b/Engine/Source/Runtime/Core/Private/Logging/LogMacros.cpp
--- a/Engine/Source/Runtime/Core/Private/Logging/LogMacros.cpp
+++ b/Engine/Source/Runtime/Core/Private/Logging/LogMacros.cpp
@@ -16,9 +16,10 @@ void StaticFailDebugV(const TCHAR* Error, const ANSICHAR* Expression, const ANSI
 
 CSV_DEFINE_CATEGORY(FMsgLogf, true);
 
-void FMsg::LogfImpl(const ANSICHAR* File, int32 Line, const FLogCategoryName& Category, ELogVerbosity::Type Verbosity, const TCHAR* Fmt, ...)
+UE_AUTORTFM_ALWAYS_OPEN
+void FMsg::LogfImplV(const ANSICHAR* File, int32 Line, const FLogCategoryName& Category, ELogVerbosity::Type Verbosity, void* ReturnAddress, const TCHAR* Fmt, va_list Args)
 {
-#if !NO_LOGGING
+	// This function always executes in the open, because our loggers and crash handlers are not AutoRTFM-safe.
 	if (LIKELY(Verbosity != ELogVerbosity::Fatal))
 	{
 		// SetColour is routed to GWarn just like the other verbosities and handled in the
@@ -36,23 +37,30 @@ void FMsg::LogfImpl(const ANSICHAR* File, int32 Line, const FLogCategoryName& Ca
 			break;
 		}
 
-		// Logging is always done in the open as we want logs even with transactionalized code.
-		GROWABLE_LOGF(UE_AUTORTFM_OPEN
+		GrowableLogfV(Fmt, Args, [&](const TCHAR* Buffer)
 		{
 			LogOverride ? LogOverride->Log(Category, Verbosity, Buffer)
-			  	 	    : GLog->RedirectLog(Category, Verbosity, Buffer);
-		})
+		                : GLog->RedirectLog(Category, Verbosity, Buffer);
+		});
 	}
 	else
 	{
-		va_list Args;
-		va_start(Args, Fmt);
-		StaticFailDebugV(TEXT("Fatal error:"), "", File, Line, /*bIsEnsure*/ false, PLATFORM_RETURN_ADDRESS(), Fmt, Args);
-		va_end(Args);
-		va_start(Args, Fmt);
+		va_list ArgsCopy;
+		va_copy(ArgsCopy, Args);
+		StaticFailDebugV(TEXT("Fatal error:"), "", File, Line, /*bIsEnsure*/ false, ReturnAddress, Fmt, ArgsCopy);
+		va_end(ArgsCopy);
+
 		FDebug::AssertFailedV("", File, Line, Fmt, Args);
-		va_end(Args);
 	}
+}
+
+void FMsg::LogfImpl(const ANSICHAR* File, int32 Line, const FLogCategoryName& Category, ELogVerbosity::Type Verbosity, const TCHAR* Fmt, ...)
+{
+#if !NO_LOGGING
+	va_list Args;
+	va_start(Args, Fmt);
+	LogfImplV(File, Line, Category, Verbosity, PLATFORM_RETURN_ADDRESS(), Fmt, Args);
+	va_end(Args);
 #endif
 }
 
diff --git a/Engine/Source/Runtime/Core/Private/Misc/VarargsHelper.h b/Engine/Source/Runtime/Core/Private/Misc/VarargsHelper.h
--- a/Engine/Source/Runtime/Core/Private/Misc/VarargsHelper.h
+++ b/Engine/Source/Runtime/Core/Private/Misc/VarargsHelper.h
@@ -3,41 +3,59 @@
 #pragma once
 
 #include "CoreTypes.h"
-#include "Misc/VarArgs.h"
+#include "Misc/CString.h"
 #include "HAL/UnrealMemory.h"
-#include "Templates/UnrealTemplate.h"
-
-// Pulled the two FOutputDevice::Logf functions into shared code. Needs to be a #define
-// since it uses GET_TYPED_VARARGS_RESULT which uses the va_list stuff which operates on the
-// current function, so we can't easily call a function
-#define GROWABLE_LOGF(SerializeFunc) \
-	int32	BufferSize	= 1024; \
-	TCHAR*	Buffer		= NULL; \
-	int32	Result		= -1; \
-	/* allocate some stack space to use on the first pass, which matches most strings */ \
-	TCHAR	StackBuffer[512]; \
-	TCHAR*	AllocatedBuffer = NULL; \
-\
-	/* first, try using the stack buffer */ \
-	Buffer = StackBuffer; \
-	GET_TYPED_VARARGS_RESULT( TCHAR, Buffer, UE_ARRAY_COUNT(StackBuffer), UE_ARRAY_COUNT(StackBuffer) - 1, Fmt, Fmt, Result ); \
-\
-	/* if that fails, then use heap allocation to make enough space */ \
-	while(Result == -1) \
-	{ \
-		FMemory::SystemFree(AllocatedBuffer); \
-		/* We need to use malloc here directly as GMalloc might not be safe. */ \
-		Buffer = AllocatedBuffer = (TCHAR*) FMemory::SystemMalloc( BufferSize * sizeof(TCHAR) ); \
-		if (Buffer == NULL) \
-		{ \
-			return; \
-		} \
-		GET_TYPED_VARARGS_RESULT( TCHAR, Buffer, BufferSize, BufferSize-1, Fmt, Fmt, Result ); \
-		BufferSize *= 2; \
-	}; \
-	Buffer[Result] = TEXT('\0'); \
-	; \
-\
-	SerializeFunc; \
+
+#include <cstdarg>
+
+template <typename SerializeFuncType>
+bool GrowableLogfV(const TCHAR* Fmt, va_list Args, const SerializeFuncType& SerializeFunc)
+{
+	// Allocate some stack space to use on the first pass. This should be sufficient for most strings.
+	constexpr int DefaultBufferSize = 512;
+
+	TCHAR	StackBuffer[DefaultBufferSize];
+	TCHAR*	Buffer = StackBuffer;
+	int32	BufferSize = DefaultBufferSize;
+	TCHAR*	AllocatedBuffer = nullptr;
+	int32	Length = -1;
+
+	for (;;)
+	{
+		va_list ArgsCopy;
+		va_copy(ArgsCopy, Args);
+		Length = FCString::GetVarArgs(Buffer, BufferSize, Fmt, ArgsCopy);
+		va_end(ArgsCopy);
+
+		if (Length >= 0 && Length < BufferSize - 1)
+		{
+			break;
+		}
+
+		// Make increasingly large allocations on the heap until we can hold the formatted string.
+		// We need to use SystemMalloc here, as GMalloc might not be safe.
+		FMemory::SystemFree(AllocatedBuffer);
+		BufferSize *= 2;
+		Buffer = AllocatedBuffer = (TCHAR*) FMemory::SystemMalloc(BufferSize * sizeof(TCHAR));
+		if (Buffer == nullptr)
+		{
+			return false;
+		}
+	}
+	Buffer[Length] = TEXT('\0');
+
+	SerializeFunc(Buffer);
 	FMemory::SystemFree(AllocatedBuffer);
+	return true;
+}
 
+// This macro expects to be used in a variadic function where the parameter list ends with
+// a `Fmt` text string, followed by the `...` variadic indicator. The passed-in code snippet
+// will have access to the resulting string via a variable named `Buffer`.
+#define GROWABLE_LOGF(SerializeSnippet) \
+	{ \
+		va_list Args; \
+		va_start(Args, Fmt); \
+		GrowableLogfV(Fmt, Args, [&](TCHAR* Buffer) { SerializeSnippet; }); \
+		va_end(Args); \
+	}
diff --git a/Engine/Source/Runtime/Core/Public/Logging/LogMacros.h b/Engine/Source/Runtime/Core/Public/Logging/LogMacros.h
--- a/Engine/Source/Runtime/Core/Public/Logging/LogMacros.h
+++ b/Engine/Source/Runtime/Core/Public/Logging/LogMacros.h
@@ -69,6 +69,7 @@ struct FMsg
 	CORE_API static void LogV(const ANSICHAR* File, int32 Line, const FLogCategoryName& Category, ELogVerbosity::Type Verbosity, const TCHAR* Fmt, va_list Args);
 
 private:
+	CORE_API static void VARARGS LogfImplV(const ANSICHAR* File, int32 Line, const FLogCategoryName& Category, ELogVerbosity::Type Verbosity, void* ReturnAddress, const TCHAR* Fmt, va_list Args);
 	CORE_API static void VARARGS LogfImpl(const ANSICHAR* File, int32 Line, const FLogCategoryName& Category, ELogVerbosity::Type Verbosity, const TCHAR* Fmt, ...);
 	CORE_API static void VARARGS Logf_InternalImpl(const ANSICHAR* File, int32 Line, const FLogCategoryName& Category, ELogVerbosity::Type Verbosity, const TCHAR* Fmt, ...);
 	CORE_API static void VARARGS SendNotificationStringfImpl(const TCHAR* Fmt, ...);
