[FormatStringSan] add some user-friendly special-case error messages for FString -> %s and TCHAR(not pointer) -> %s

diff --git a/Engine/Source/Runtime/Core/Public/String/FormatStringSan.h b/Engine/Source/Runtime/Core/Public/String/FormatStringSan.h
--- a/Engine/Source/Runtime/Core/Public/String/FormatStringSan.h
+++ b/Engine/Source/Runtime/Core/Public/String/FormatStringSan.h
@@ -5,6 +5,7 @@
 #include "CoreTypes.h"
 #include "Templates/Requires.h"
 #include "Templates/Identity.h"
+#include "Containers/ContainersFwd.h"
 
 #define UE_CHECK_FORMAT_STRING(Fmt, ...) do {														\
 		typedef ::UE::Core::Private::FormatStringSan::TCheckFormatString<decltype(::UE::Core::Private::FormatStringSan::GetFmtArgTypes(__VA_ARGS__))> UE_FMT_STR_Checker; \
@@ -41,6 +42,8 @@ namespace UE::Core::Private
 		{
 			StatusOk,
 			StatusSNeedsTCHARPtrArg,
+			StatusSNeedsDereferencedFString,
+			StatusSNeedsTCharPtrArgButGotTCHAR,
 			StatusHSNeedsCharPtrArg,
 			StatusIncompleteFormatSpecifierOrUnescapedPercent,
 			StatusCNeedsTCHARArg,
@@ -274,6 +277,14 @@ namespace UE::Core::Private
 					{
 						return {StatusUTF8NeedsStringCast, CurArgPos};
 					}
+					else if constexpr (std::is_same_v<Arg, FString>)
+					{
+						return {StatusSNeedsDereferencedFString, CurArgPos};
+					}
+					else if constexpr (std::is_same_v<Arg, TCHAR>)
+					{
+						return {StatusSNeedsTCharPtrArgButGotTCHAR, CurArgPos};
+					}
 					else if constexpr (!bIsTCharPtr<Arg>)
 					{
 						return {StatusSNeedsTCHARPtrArg, CurArgPos};
diff --git a/Engine/Source/Runtime/Core/Public/String/FormatStringSanErrors.h b/Engine/Source/Runtime/Core/Public/String/FormatStringSanErrors.h
--- a/Engine/Source/Runtime/Core/Public/String/FormatStringSanErrors.h
+++ b/Engine/Source/Runtime/Core/Public/String/FormatStringSanErrors.h
@@ -17,6 +17,8 @@ auto X = WRONG_USE_OF_FORMAT_STRING_SAN_ERRORS_H();
 
 UE_FORMAT_STRING_ERROR_BEGIN
   CASE(StatusSNeedsTCHARPtrArg, "'%s' expects `TCHAR*`; use '%hs' for `const char*`.");
+  CASE(StatusSNeedsDereferencedFString, "FString should be dereferenced before passing to '%s'.");
+  CASE(StatusSNeedsTCharPtrArgButGotTCHAR, "'%s' expects `TCHAR*` but got `TCHAR`. (Is there an extra '*' somewhere?)");
   CASE(StatusHSNeedsCharPtrArg, "'%hs' expects `char*`; use '%s' for `const TCHAR*`, `TEXT(...)`, etc.");
   CASE(StatusIncompleteFormatSpecifierOrUnescapedPercent,
        "incomplete '%' escape sequence. (to print a percent sign, write '%%'.)");
