FileExistsInPlatformPackage
- Early out if not passed a relative path. Passing a bad path was leading to a crash on a particular platform.

diff --git a/Engine/Source/Runtime/Core/Private/Android/AndroidPlatformMisc.cpp b/Engine/Source/Runtime/Core/Private/Android/AndroidPlatformMisc.cpp
--- a/Engine/Source/Runtime/Core/Private/Android/AndroidPlatformMisc.cpp
+++ b/Engine/Source/Runtime/Core/Private/Android/AndroidPlatformMisc.cpp
@@ -5,6 +5,7 @@
 #include "HAL/PlatformStackWalk.h"
 #include "Misc/FileHelper.h"
 #include "Misc/App.h"
+#include "Misc/PathViews.h"
 #include "Misc/ConfigCacheIni.h"
 #include "Misc/CoreDelegates.h"
 #include "Misc/FeedbackContext.h"
@@ -1755,6 +1756,12 @@ FString FAndroidMisc::LoadTextFileFromPlatformPackage(const FString& RelativePat
 bool FAndroidMisc::FileExistsInPlatformPackage(const FString& RelativePath)
 {
 #if USE_ANDROID_JNI
+	if (!FPathViews::IsRelativePath(RelativePath))
+	{
+		UE_LOG(LogAndroid, Warning, TEXT("FileExistsInPlatformPackage: expected a relative path, but received %s"), *RelativePath);
+		return false;
+	}
+
 	AAssetManager* AssetMgr = AndroidThunkCpp_GetAssetManager();
 	AAsset* asset = AAssetManager_open(AssetMgr, TCHAR_TO_UTF8(*RelativePath), AASSET_MODE_UNKNOWN);
 	if (asset)
diff --git a/Engine/Source/Runtime/Core/Private/GenericPlatform/GenericPlatformMisc.cpp b/Engine/Source/Runtime/Core/Private/GenericPlatform/GenericPlatformMisc.cpp
--- a/Engine/Source/Runtime/Core/Private/GenericPlatform/GenericPlatformMisc.cpp
+++ b/Engine/Source/Runtime/Core/Private/GenericPlatform/GenericPlatformMisc.cpp
@@ -1947,6 +1947,12 @@ FString FGenericPlatformMisc::LoadTextFileFromPlatformPackage(const FString& Rel
 
 bool FGenericPlatformMisc::FileExistsInPlatformPackage(const FString& RelativePath)
 {
+	if (!FPathViews::IsRelativePath(RelativePath))
+	{
+		UE_LOG(LogGenericPlatformMisc, Warning, TEXT("FileExistsInPlatformPackage: expected a relative path, but received %s"), *RelativePath);
+		return false;
+	}
+
 	FString Path = RootDir() / RelativePath;
 	return IPlatformFile::GetPlatformPhysical().FileExists(*Path);
 }
diff --git a/Engine/Source/Runtime/Core/Private/IOS/IOSPlatformMisc.cpp b/Engine/Source/Runtime/Core/Private/IOS/IOSPlatformMisc.cpp
--- a/Engine/Source/Runtime/Core/Private/IOS/IOSPlatformMisc.cpp
+++ b/Engine/Source/Runtime/Core/Private/IOS/IOSPlatformMisc.cpp
@@ -28,6 +28,7 @@
 #include "Misc/FeedbackContext.h"
 #include "Misc/OutputDeviceError.h"
 #include "Misc/OutputDeviceRedirector.h"
+#include "Misc/PathViews.h"
 #include "Misc/ScopeExit.h"
 #include "Misc/SecureHash.h"
 #include "Modules/ModuleManager.h"
@@ -1211,6 +1212,12 @@ FString FIOSPlatformMisc::LoadTextFileFromPlatformPackage(const FString& Relativ
 
 bool FIOSPlatformMisc::FileExistsInPlatformPackage(const FString& RelativePath)
 {
+	if (!FPathViews::IsRelativePath(RelativePath))
+	{
+		UE_LOG(LogIOS, Warning, TEXT("FileExistsInPlatformPackage: expected a relative path, but received %s"), *RelativePath);
+		return false;
+	}
+
 	FString FilePath = FString([[NSBundle mainBundle] bundlePath]) / RelativePath;
 
 	return 0 == access(TCHAR_TO_UTF8(*FilePath), F_OK);
diff --git a/Engine/Source/Runtime/Experimental/IoStore/OnDemand/Private/IoStoreOnDemand.cpp b/Engine/Source/Runtime/Experimental/IoStore/OnDemand/Private/IoStoreOnDemand.cpp
--- a/Engine/Source/Runtime/Experimental/IoStore/OnDemand/Private/IoStoreOnDemand.cpp
+++ b/Engine/Source/Runtime/Experimental/IoStore/OnDemand/Private/IoStoreOnDemand.cpp
@@ -970,11 +970,12 @@ bool LoadFromCompactBinary(FCbFieldView Field, FOnDemandToc& OutToc)
 TIoStatusOr<FOnDemandToc> FOnDemandToc::LoadFromFile(const FString& FilePath, bool bValidate)
 {
 	TUniquePtr<FArchive> Ar;
-	if (FPlatformMisc::FileExistsInPlatformPackage(FilePath))
+	if (FPathViews::IsRelativePath(FilePath) && FPlatformMisc::FileExistsInPlatformPackage(FilePath))
 	{
 		Ar = CreateReaderFromPlatformPackage(FilePath);
 	}
-	else
+	
+	if (Ar.IsValid() == false)
 	{
 		Ar.Reset(IFileManager::Get().CreateFileReader(*FilePath));
 	}
