CA_ASSUME changes to improve PVS null checking

diff --git a/Engine/Source/Runtime/Core/Public/Microsoft/MicrosoftPlatformCodeAnalysis.h b/Engine/Source/Runtime/Core/Public/Microsoft/MicrosoftPlatformCodeAnalysis.h
--- a/Engine/Source/Runtime/Core/Public/Microsoft/MicrosoftPlatformCodeAnalysis.h
+++ b/Engine/Source/Runtime/Core/Public/Microsoft/MicrosoftPlatformCodeAnalysis.h
@@ -70,8 +70,13 @@
 	#define CA_SUPPRESS( WarningNumber ) __pragma( warning( suppress: WarningNumber ) )
 
 	// Tells the code analysis engine to assume the statement to be true.  Useful for suppressing false positive warnings.
+#if defined( PVS_STUDIO )
+	__declspec(dllimport, noreturn) void CA_AssumeNoReturn();
+	#define CA_ASSUME( Expr )  (__builtin_expect(!bool(Expr), 0) ? CA_AssumeNoReturn() : (void)0)
+#else
 	// NOTE: We use a double operator not here to avoid issues with passing certain class objects directly into __analysis_assume (which may cause a bogus compiler warning)
 	#define CA_ASSUME( Expr ) __analysis_assume( !!( Expr ) )
+#endif
 
 	// Does a simple 'if (Condition)', but disables warnings about using constants in the condition.  Helps with some macro expansions.
 	#define CA_CONSTANT_IF(Condition) __pragma(warning(push)) __pragma(warning(disable:6326)) if (Condition) __pragma(warning(pop))
