[Core]
* Fixed shipping compile error in ConsoleManager.cpp

diff --git a/Engine/Source/Runtime/Core/Private/HAL/ConsoleManager.cpp b/Engine/Source/Runtime/Core/Private/HAL/ConsoleManager.cpp
--- a/Engine/Source/Runtime/Core/Private/HAL/ConsoleManager.cpp
+++ b/Engine/Source/Runtime/Core/Private/HAL/ConsoleManager.cpp
@@ -2693,8 +2693,10 @@ bool FConsoleManager::ProcessUserConsoleInput(const TCHAR* InInput, FOutputDevic
 {
 	check(InInput);
 
+#if CPUPROFILERTRACE_ENABLED
 	UE_TRACE_LOG_SCOPED_T(Cpu, ConsoleCmd, CpuChannel)
 		<< ConsoleCmd.Cmd(InInput);
+#endif
 
 	CSV_EVENT_GLOBAL(TEXT("Cmd: %s"), InInput);
 
