Fix AutoRTFM::IsClosed and AutoRTFM::IsTransactional to always return false within a static local initializer.

diff --git a/Engine/Source/Programs/AutoRTFMTests/Private/Statics.cpp b/Engine/Source/Programs/AutoRTFMTests/Private/Statics.cpp
--- a/Engine/Source/Programs/AutoRTFMTests/Private/Statics.cpp
+++ b/Engine/Source/Programs/AutoRTFMTests/Private/Statics.cpp
@@ -1,7 +1,7 @@
 // Copyright Epic Games, Inc. All Rights Reserved.
 
 #include "Catch2Includes.h"
-#include <AutoRTFM/AutoRTFM.h>
+#include <AutoRTFM/AutoRTFMTesting.h>
 
 TEST_CASE("Statics.Int")
 {
@@ -103,3 +103,192 @@ TEST_CASE("Statics.StructAbort")
 	// been rolled back, but it should still be initialized correctly.
 	REQUIRE(1 == *GetSlot());
 }
+
+using FuncPtrType = bool(*)();
+
+FuncPtrType GIsClosed = &AutoRTFM::IsClosed;
+FuncPtrType GIsTransactional = &AutoRTFM::IsTransactional;
+
+FORCENOINLINE static bool GetIsClosed()
+{
+	return AutoRTFM::IsClosed();
+}
+
+FORCENOINLINE static bool GetIsTransactional()
+{
+	return AutoRTFM::IsTransactional();
+}
+
+TEST_CASE("Static.IsClosed")
+{
+	SECTION("Normal")
+	{
+		struct MyStruct final
+		{
+			bool bWasClosedAtConstruction = AutoRTFM::IsClosed();
+		};
+
+		AutoRTFM::Testing::Commit([&]
+		{
+			static MyStruct Mine;
+			REQUIRE(false == Mine.bWasClosedAtConstruction);
+		});
+	}
+
+	SECTION("FromGlobal")
+	{
+		struct MyStruct final
+		{
+			bool bWasClosedAtConstruction = GIsClosed();
+		};
+
+		AutoRTFM::Testing::Commit([&]
+		{
+			static MyStruct Mine;
+			REQUIRE(false == Mine.bWasClosedAtConstruction);
+		});
+	}
+
+	SECTION("InNoInlineCall")
+	{
+		struct MyStruct final
+		{
+			bool bWasClosedAtConstruction = GetIsClosed();
+		};
+
+		AutoRTFM::Testing::Commit([&]
+		{
+			static MyStruct Mine;
+			REQUIRE(false == Mine.bWasClosedAtConstruction);
+		});
+	}
+
+	SECTION("InNestedStaticLocalInitializer")
+	{
+		struct MyStruct final
+		{
+			struct MyInner final
+			{
+				bool bWasClosedAtConstruction = AutoRTFM::IsClosed();
+			};
+
+			MyStruct()
+			{
+				static MyInner Mine;
+				bWasClosedAtConstruction = Mine.bWasClosedAtConstruction;
+			}
+
+			bool bWasClosedAtConstruction;
+		};
+
+		AutoRTFM::Testing::Commit([&]
+		{
+			static MyStruct Mine;
+			REQUIRE(false == Mine.bWasClosedAtConstruction);
+		});
+	}
+
+	SECTION("InOpenStaticLocalInitializer")
+	{
+		struct MyStruct final
+		{
+			bool bWasClosedAtConstruction = AutoRTFM::IsClosed();
+		};
+
+		AutoRTFM::Testing::Commit([&]
+		{
+			AutoRTFM::Open([&]
+			{
+				static MyStruct Mine;
+				REQUIRE(false == Mine.bWasClosedAtConstruction);
+			});
+		});
+	}
+}
+
+TEST_CASE("Static.IsTransactional")
+{
+	SECTION("Normal")
+	{
+		struct MyStruct final
+		{
+			bool bWasTransactionalAtConstruction = AutoRTFM::IsTransactional();
+		};
+
+		AutoRTFM::Testing::Commit([&]
+		{
+			static MyStruct Mine;
+			REQUIRE(false == Mine.bWasTransactionalAtConstruction);
+		});
+	}
+
+	SECTION("FromGlobal")
+	{
+		struct MyStruct final
+		{
+			bool bWasTransactionalAtConstruction = GIsTransactional();
+		};
+
+		AutoRTFM::Testing::Commit([&]
+		{
+			static MyStruct Mine;
+			REQUIRE(false == Mine.bWasTransactionalAtConstruction);
+		});
+	}
+
+	SECTION("InNoInlineCall")
+	{
+		struct MyStruct final
+		{
+			bool bWasTransactionalAtConstruction = GetIsTransactional();
+		};
+
+		AutoRTFM::Testing::Commit([&]
+		{
+			static MyStruct Mine;
+			REQUIRE(false == Mine.bWasTransactionalAtConstruction);
+		});
+	}
+
+	SECTION("InNestedStaticLocalInitializer")
+	{
+		struct MyStruct final
+		{
+			struct MyInner final
+			{
+				bool bWasTransactionalAtConstruction = AutoRTFM::IsTransactional();
+			};
+
+			MyStruct()
+			{
+				static MyInner Mine;
+				bWasTransactionalAtConstruction = Mine.bWasTransactionalAtConstruction;
+			}
+
+			bool bWasTransactionalAtConstruction;
+		};
+
+		AutoRTFM::Testing::Commit([&]
+		{
+			static MyStruct Mine;
+			REQUIRE(false == Mine.bWasTransactionalAtConstruction);
+		});
+	}
+
+	SECTION("InOpenStaticLocalInitializer")
+	{
+		struct MyStruct final
+		{
+			bool bWasTransactionalAtConstruction = AutoRTFM::IsTransactional();
+		};
+
+		AutoRTFM::Testing::Commit([&]
+		{
+			AutoRTFM::Open([&]
+			{
+				static MyStruct Mine;
+				REQUIRE(false == Mine.bWasTransactionalAtConstruction);
+			});
+		});
+	}
+}
diff --git a/Engine/Source/Programs/AutoRTFMTests/Private/UECoreTests.cpp b/Engine/Source/Programs/AutoRTFMTests/Private/UECoreTests.cpp
--- a/Engine/Source/Programs/AutoRTFMTests/Private/UECoreTests.cpp
+++ b/Engine/Source/Programs/AutoRTFMTests/Private/UECoreTests.cpp
@@ -1071,6 +1071,67 @@ TEST_CASE("UECore.TScopeLock_TransactionallySafeCriticalSection")
 		REQUIRE(AutoRTFM::ETransactionResult::Committed == InnerResult);
 		REQUIRE(AutoRTFM::ETransactionResult::Committed == Result);
 	}
+
+
+	SECTION("In Static Local Initializer")
+	{
+		struct MyStruct final
+		{
+			FTransactionallySafeCriticalSection CriticalSection;
+		};
+
+		auto Lambda = []()
+		{
+			static MyStruct Mine;
+			UE::TScopeLock _(Mine.CriticalSection);
+			return 42;
+		};
+
+		AutoRTFM::Testing::Abort([&]
+		{
+			REQUIRE(42 == Lambda());
+			AutoRTFM::AbortTransaction();
+		});
+
+		REQUIRE(42 == Lambda());
+
+		AutoRTFM::Testing::Commit([&]
+		{
+			REQUIRE(42 == Lambda());
+		});
+
+		REQUIRE(42 == Lambda());
+	}
+
+	SECTION("In Static Local Initializer Called From Open")
+	{
+		struct MyStruct final
+		{
+			FTransactionallySafeCriticalSection CriticalSection;
+		};
+
+		auto Lambda = []()
+		{
+			static MyStruct Mine;
+			UE::TScopeLock _(Mine.CriticalSection);
+			return 42;
+		};
+
+		AutoRTFM::Testing::Abort([&]
+		{
+			AutoRTFM::Open([&] { REQUIRE(42 == Lambda()); });
+			AutoRTFM::AbortTransaction();
+		});
+
+		REQUIRE(42 == Lambda());
+
+		AutoRTFM::Testing::Commit([&]
+		{
+			AutoRTFM::Open([&] { REQUIRE(42 == Lambda()); });
+		});
+
+		REQUIRE(42 == Lambda());
+	}
 }
 
 TEST_CASE("UECore.TransactionallySafeRWScopeLock")
@@ -1610,6 +1671,66 @@ TEST_CASE("UECore.TransactionallySafeMutex")
 		REQUIRE(AutoRTFM::ETransactionResult::Committed == InnerResult);
 		REQUIRE(AutoRTFM::ETransactionResult::Committed == Result);
 	}
+
+	SECTION("In Static Local Initializer")
+	{
+		struct MyStruct final
+		{
+			FTransactionallySafeMutex Mutex;
+		};
+
+		auto Lambda = []()
+		{
+			static MyStruct Mine;
+			UE::TScopeLock _(Mine.Mutex);
+			return 42;
+		};
+
+		AutoRTFM::Testing::Abort([&]
+		{
+			REQUIRE(42 == Lambda());
+			AutoRTFM::AbortTransaction();
+		});
+
+		REQUIRE(42 == Lambda());
+
+		AutoRTFM::Testing::Commit([&]
+		{
+			REQUIRE(42 == Lambda());
+		});
+
+		REQUIRE(42 == Lambda());
+	}
+
+	SECTION("In Static Local Initializer Called From Open")
+	{
+		struct MyStruct final
+		{
+			FTransactionallySafeMutex Mutex;
+		};
+
+		auto Lambda = []()
+		{
+			static MyStruct Mine;
+			UE::TScopeLock _(Mine.Mutex);
+			return 42;
+		};
+
+		AutoRTFM::Testing::Abort([&]
+		{
+			AutoRTFM::Open([&] { REQUIRE(42 == Lambda()); });
+			AutoRTFM::AbortTransaction();
+		});
+
+		REQUIRE(42 == Lambda());
+
+		AutoRTFM::Testing::Commit([&]
+		{
+			AutoRTFM::Open([&] { REQUIRE(42 == Lambda()); });
+		});
+
+		REQUIRE(42 == Lambda());
+	}
 }
 
 TEST_CASE("UECore.FTextFormatPatternDefinition")
diff --git a/Engine/Source/Runtime/Core/Private/AutoRTFM/API.cpp b/Engine/Source/Runtime/Core/Private/AutoRTFM/API.cpp
--- a/Engine/Source/Runtime/Core/Private/AutoRTFM/API.cpp
+++ b/Engine/Source/Runtime/Core/Private/AutoRTFM/API.cpp
@@ -406,6 +406,16 @@ extern "C" UE_AUTORTFM_NOAUTORTFM UE_AUTORTFM_API void autortfm_post_open()
 	}
 }
 
+extern "C" UE_AUTORTFM_NOAUTORTFM UE_AUTORTFM_API void autortfm_pre_static_local_initializer()
+{
+	FContext::Get()->EnteringStaticLocalInitializer();
+}
+
+extern "C" UE_AUTORTFM_NOAUTORTFM UE_AUTORTFM_API void autortfm_post_static_local_initializer()
+{
+	FContext::Get()->LeavingStaticLocalInitializer();
+}
+
 AUTORTFM_CORE_API void RTFM_autortfm_open(autortfm_memory_validation_level ValidationLevel, void (*Work)(void*), void* Arg)
 {
 	FTransaction* Transaction = FContext::Get()->GetCurrentTransaction();
@@ -524,7 +534,7 @@ extern "C" bool autortfm_is_transactional()
 
 	if (ForTheRuntime::IsAutoRTFMRuntimeEnabled())
 	{
-		return FContext::IsTransactional();
+		return FContext::Get()->IsTransactional();
 	}
 
 	return false;
@@ -534,7 +544,7 @@ extern "C" bool autortfm_is_committing_or_aborting()
 {
 	if (ForTheRuntime::IsAutoRTFMRuntimeEnabled())
 	{
-		return FContext::IsCommittingOrAborting();
+		return FContext::Get()->IsCommittingOrAborting();
 	}
 
 	return false;
@@ -584,7 +594,7 @@ extern "C" void autortfm_abort()
 		return RTFM_autortfm_abort();
 	}
 
-	UE_CLOG(!FContext::IsTransactional(), LogAutoRTFM, Fatal, TEXT("The function `autortfm_abort` was called from outside a transaction."));
+	UE_CLOG(!FContext::Get()->IsTransactional(), LogAutoRTFM, Fatal, TEXT("The function `autortfm_abort` was called from outside a transaction."));
 	FContext::Get()->AbortByRequestAndThrow();
 }
 
@@ -596,7 +606,7 @@ extern "C" bool autortfm_start_transaction()
 		return false;
 	}
 
-	UE_CLOG(!FContext::IsTransactional(), LogAutoRTFM, Fatal, TEXT("The function `autortfm_start_transaction` was called from outside a transact."));
+	UE_CLOG(!FContext::Get()->IsTransactional(), LogAutoRTFM, Fatal, TEXT("The function `autortfm_start_transaction` was called from outside a transact."));
 	FContext::Get()->StartTransaction(ForTheRuntime::GetMemoryValidationLevel());
 	return true;
 }
@@ -609,7 +619,7 @@ extern "C" autortfm_result autortfm_commit_transaction()
 		return autortfm_aborted_by_language;
 	}
 
-	UE_CLOG(!FContext::IsTransactional(), LogAutoRTFM, Fatal, TEXT("The function `autortfm_commit_transaction` was called from outside a transact."));
+	UE_CLOG(!FContext::Get()->IsTransactional(), LogAutoRTFM, Fatal, TEXT("The function `autortfm_commit_transaction` was called from outside a transact."));
 	return static_cast<autortfm_result>(FContext::Get()->CommitTransaction());
 }
 
@@ -620,7 +630,7 @@ extern "C" autortfm_result autortfm_abort_transaction()
 		return RTFM_autortfm_abort_transaction();
 	}
 
-	UE_CLOG(!FContext::IsTransactional(), LogAutoRTFM, Fatal, TEXT("The function `autortfm_abort_transaction` was called from outside a transact."));
+	UE_CLOG(!FContext::Get()->IsTransactional(), LogAutoRTFM, Fatal, TEXT("The function `autortfm_abort_transaction` was called from outside a transact."));
 	FContext* const Context = FContext::Get();
 	return static_cast<autortfm_result>(Context->AbortTransaction(false, false));
 }
@@ -632,7 +642,7 @@ extern "C" autortfm_result autortfm_cascading_abort_transaction()
 		return RTFM_autortfm_cascading_abort_transaction();
 	}
 
-	UE_CLOG(!FContext::IsTransactional(), LogAutoRTFM, Fatal, TEXT("The function `autortfm_cascading_abort_transaction` was called from outside a transact."));
+	UE_CLOG(!FContext::Get()->IsTransactional(), LogAutoRTFM, Fatal, TEXT("The function `autortfm_cascading_abort_transaction` was called from outside a transact."));
 	FContext* const Context = FContext::Get();
 	return static_cast<autortfm_result>(Context->AbortTransaction(false, true));
 }
@@ -695,7 +705,7 @@ extern "C" autortfm_status autortfm_close(void (*UninstrumentedWork)(void*), voi
 
 	if (ForTheRuntime::IsAutoRTFMRuntimeEnabled())
 	{
-		UE_CLOG(!FContext::IsTransactional(), LogAutoRTFM, Fatal, TEXT("Close called from an outside a transaction."));
+		UE_CLOG(!FContext::Get()->IsTransactional(), LogAutoRTFM, Fatal, TEXT("Close called from an outside a transaction."));
 
 		FContext* const Context = FContext::Get();
 
@@ -730,7 +740,7 @@ extern "C" void autortfm_record_open_write(void* Ptr, size_t Size)
 	{
 		return RTFM_autortfm_record_open_write(Ptr, Size);
 	}
-	else if (!FContext::IsTransactional())
+	else if (!FContext::Get()->IsTransactional())
 	{
 		return;
 	}
@@ -748,7 +758,7 @@ extern "C" UE_AUTORTFM_NOAUTORTFM void autortfm_register_open_function(void* Ori
 
 extern "C" bool autortfm_is_on_current_transaction_stack(void* Ptr)
 {
-	if (!FContext::IsTransactional())
+	if (!FContext::Get()->IsTransactional())
 	{
 		return false;
 	}
@@ -866,7 +876,7 @@ extern "C" void autortfm_did_free(void* Ptr)
 			return;
 		}
 
-		if (FContext::IsTransactional())
+		if (FContext::Get()->IsTransactional())
 		{
 			FContext* const Context = FContext::Get();
 
@@ -887,7 +897,7 @@ extern "C" void autortfm_check_consistency_assuming_no_races()
 		return RTFM_autortfm_check_consistency_assuming_no_races();
 	}
 
-    if (FContext::IsTransactional())
+    if (FContext::Get()->IsTransactional())
     {
         AutoRTFM::Unreachable();
     }
diff --git a/Engine/Source/Runtime/Core/Private/AutoRTFM/CompilerInterface.cpp b/Engine/Source/Runtime/Core/Private/AutoRTFM/CompilerInterface.cpp
--- a/Engine/Source/Runtime/Core/Private/AutoRTFM/CompilerInterface.cpp
+++ b/Engine/Source/Runtime/Core/Private/AutoRTFM/CompilerInterface.cpp
@@ -116,7 +116,7 @@ extern "C" UE_AUTORTFM_API AUTORTFM_NO_ASAN void autortfm_record_masked_write(vo
 	}
 }
 
-extern "C" UE_AUTORTFM_API void* autortfm_lookup_function(void* OriginalFunction, const char* Where)
+extern "C" UE_AUTORTFM_API UE_AUTORTFM_ALWAYS_OPEN void* autortfm_lookup_function(void* OriginalFunction, const char* Where)
 {
 	FContext* Context = FContext::Get();
     return FunctionMapLookup(OriginalFunction, Where);
diff --git a/Engine/Source/Runtime/Core/Private/AutoRTFM/Context.cpp b/Engine/Source/Runtime/Core/Private/AutoRTFM/Context.cpp
--- a/Engine/Source/Runtime/Core/Private/AutoRTFM/Context.cpp
+++ b/Engine/Source/Runtime/Core/Private/AutoRTFM/Context.cpp
@@ -48,14 +48,14 @@ void FContext::InitializeGlobalData()
 
 FContext FContext::ContextSingleton;
 
-bool FContext::IsTransactional()
+bool FContext::IsTransactional() const
 {
-    return Get()->GetStatus() == EContextStatus::OnTrack;
+    return GetStatus() == EContextStatus::OnTrack;
 }
 
-bool FContext::IsCommittingOrAborting()
+bool FContext::IsCommittingOrAborting() const
 {
-	switch (Get()->GetStatus())
+	switch (GetStatus())
 	{
 	default:
 		return true;
@@ -601,6 +601,7 @@ void FContext::Reset()
     CurrentTransaction = nullptr;
 	CurrentNest = nullptr;
     Status = EContextStatus::Idle;
+	StackLocalInitializerDepth = 0;
 }
 
 void FContext::Throw()
diff --git a/Engine/Source/Runtime/Core/Private/AutoRTFM/Context.h b/Engine/Source/Runtime/Core/Private/AutoRTFM/Context.h
--- a/Engine/Source/Runtime/Core/Private/AutoRTFM/Context.h
+++ b/Engine/Source/Runtime/Core/Private/AutoRTFM/Context.h
@@ -18,8 +18,6 @@ class FContext
 {
 public:
     static FContext* Get();
-    static bool IsTransactional();
-	static bool IsCommittingOrAborting();
     
     // This is public API
     ETransactionResult Transact(void (*InstrumentedFunction)(void*), void* Arg);
@@ -36,6 +34,8 @@ public:
 	ETransactionResult AbortTransaction(bool bIsClosed, bool bIsCascading);
 	void ClearTransactionStatus();
 	bool IsAborting() const;
+	bool IsTransactional() const;
+	bool IsCommittingOrAborting() const;
 
     // Record that a write is about to occur at the given LogicalAddress of Size bytes.
     void RecordWrite(void* LogicalAddress, size_t Size);
@@ -88,6 +88,45 @@ public:
 
     static void InitializeGlobalData();
 
+	inline void EnteringStaticLocalInitializer()
+	{
+		if (LIKELY(GetStatus() == EContextStatus::Idle))
+		{
+			return;
+		}
+
+		if (Status == EContextStatus::OnTrack)
+		{
+			ASSERT(0 == StackLocalInitializerDepth);
+			Status = EContextStatus::InStaticLocalInitializer;
+			StackLocalInitializerDepth++;
+		}
+		else if (Status == EContextStatus::InStaticLocalInitializer)
+		{
+			StackLocalInitializerDepth++;
+		}
+	}
+
+	inline void LeavingStaticLocalInitializer()
+	{
+		if (LIKELY(GetStatus() == EContextStatus::Idle))
+		{
+			return;
+		}
+
+		ASSERT(Status != EContextStatus::OnTrack);
+
+		if (Status == EContextStatus::InStaticLocalInitializer)
+		{
+			StackLocalInitializerDepth--;
+
+			if (0 == StackLocalInitializerDepth)
+			{
+				Status = EContextStatus::OnTrack;
+			}
+		}
+	}
+
 private:
 	static FContext ContextSingleton;
 
@@ -119,6 +158,7 @@ private:
 	void* ClosedStackAddress = nullptr;
     EContextStatus Status{EContextStatus::Idle};
 	uint32 CurrentThreadId{ FPlatformTLS::InvalidTlsSlot };
+	uint32 StackLocalInitializerDepth = 0;
 };
 
 } // namespace AutoRTFM
diff --git a/Engine/Source/Runtime/Core/Private/AutoRTFM/Transaction.h b/Engine/Source/Runtime/Core/Private/AutoRTFM/Transaction.h
--- a/Engine/Source/Runtime/Core/Private/AutoRTFM/Transaction.h
+++ b/Engine/Source/Runtime/Core/Private/AutoRTFM/Transaction.h
@@ -40,10 +40,10 @@ public:
         // The initial state for the transaction.
         // Can only transition to OpenActive or ClosedActive.
         Uninitialized,
-        // The transaction is open (not recording writes) and is the current transcation.
+        // The transaction is open (not recording writes) and is the current transaction.
         // Can only transition to OpenInactive, ClosedActive or Done.
         OpenActive,
-        // The transaction is closed (recording writes) and is the current transcation.
+        // The transaction is closed (recording writes) and is the current transaction.
         // Can only transition to ClosedInactive, OpenActive or Done.
         ClosedActive,
         // The transaction is open and the current transaction is a descendant.
diff --git a/Engine/Source/Runtime/Core/Public/AutoRTFM/AutoRTFM.h b/Engine/Source/Runtime/Core/Public/AutoRTFM/AutoRTFM.h
--- a/Engine/Source/Runtime/Core/Public/AutoRTFM/AutoRTFM.h
+++ b/Engine/Source/Runtime/Core/Public/AutoRTFM/AutoRTFM.h
@@ -128,7 +128,8 @@ typedef enum
 	autortfm_status_aborted_by_language,
 	autortfm_status_aborted_by_request,
 	autortfm_status_committing,
-	autortfm_status_aborted_by_cascade
+	autortfm_status_aborted_by_cascade,
+	autortfm_status_in_static_local_initializer
 } autortfm_status;
 
 // An enumerator of transactional memory validation levels.
@@ -457,7 +458,11 @@ enum class EContextStatus : uint8_t
 
 	// An AbortedByCascade status means that a call to CascadingAbortTransaction occurred and
 	// we are currently aborting because of it.
-	AbortedByCascade = autortfm_status_aborted_by_cascade
+	AbortedByCascade = autortfm_status_aborted_by_cascade,
+
+	// Means we are in a static local initializer which always run in the open.
+	// `IsTransactional()` will return `false` when in this state.
+	InStaticLocalInitializer = autortfm_status_in_static_local_initializer,
 };
 
 // An enumerator of transactional memory validation levels.
