Add portable crash stacks for non-crashing threads to crash report

#rnx

diff --git a/Engine/Source/Runtime/Core/Private/Mac/MacPlatformMisc.cpp b/Engine/Source/Runtime/Core/Private/Mac/MacPlatformMisc.cpp
--- a/Engine/Source/Runtime/Core/Private/Mac/MacPlatformMisc.cpp
+++ b/Engine/Source/Runtime/Core/Private/Mac/MacPlatformMisc.cpp
@@ -2382,7 +2382,8 @@ void FMacCrashContext::GenerateCrashInfoAndLaunchReporter() const
 
 		// Do not inline this! The lifetime of this object needs to extend over the usage of Argv in posix_spawn() call below.
 		auto CrashInfoFolderUTF8 = TStringConversion<FTCHARToUTF8_Convert>(*CrashInfoFolder);
-
+		
+		CaptureAllThreadContext(0);
 		GenerateInfoInFolder(CrashInfoFolderUTF8.Get());
 
 		CrashInfoFolder += TEXT("/");
@@ -2522,6 +2523,7 @@ void FMacCrashContext::GenerateEnsureInfoAndLaunchReporter() const
 		FString GameName = FApp::GetProjectName();
 		FString EnsureLogFolder = FString(GMacAppInfo.CrashReportPath) / FString::Printf(TEXT("EnsureReport-%s-%s"), *GameName, *Guid.ToString(EGuidFormats::Digits));
 		
+		CaptureAllThreadContext(0);
 		GenerateInfoInFolder(TCHAR_TO_UTF8(*EnsureLogFolder));
 		
 		FString Arguments;
@@ -2545,11 +2547,13 @@ void FMacCrashContext::GenerateEnsureInfoAndLaunchReporter() const
 	}
 }
 
+static FString AllThreadContexts;
+
 void FMacCrashContext::AddThreadContext(
 	uint32 ThreadIdEnteredOn,
 	uint32 ThreadId,
 	const FString& ThreadName,
-	const TArray<FCrashStackFrame>& PortableCallStack)
+	const TArray<FCrashStackFrame>& PortableCallStack) const
 {
 	AllThreadContexts += TEXT("<Thread>");
 	{
@@ -2580,7 +2584,7 @@ void FMacCrashContext::AddThreadContext(
 	AllThreadContexts += LINE_TERMINATOR;
 }
 
-void FMacCrashContext::CaptureAllThreadContext(uint32 ThreadIdEnteredOn)
+void FMacCrashContext::CaptureAllThreadContext(uint32 ThreadIdEnteredOn) const
 {
 	TArray<typename FThreadManager::FThreadStackBackTrace> StackTraces;
 	FThreadManager::Get().GetAllThreadStackBackTraces(StackTraces);
@@ -2591,16 +2595,25 @@ void FMacCrashContext::CaptureAllThreadContext(uint32 ThreadIdEnteredOn)
 		const uint32 ThreadId     = ThreadStTrace.ThreadId;
 		const FString& ThreadName = ThreadStTrace.ThreadName;
 
-		TArray<FCrashStackFrame> PortableStack;
-		GetPortableCallStack(ThreadStTrace.ProgramCounters.GetData(), ThreadStTrace.ProgramCounters.Num(), PortableStack);
+		if (Idx == 0)
+		{
+			// current thread, which is crashing
+			AddThreadContext(ThreadId, ThreadId, ThreadName, CallStack);
+		}
+		else
+		{
+			TArray<FCrashStackFrame> PortableStack;
+			GetPortableCallStack(ThreadStTrace.ProgramCounters.GetData(), ThreadStTrace.ProgramCounters.Num(), PortableStack);
 
-		AddThreadContext(ThreadIdEnteredOn, ThreadId, ThreadName, PortableStack);
+			AddThreadContext(ThreadIdEnteredOn, ThreadId, ThreadName, PortableStack);
+		}
 	}
 }
 
 bool FMacCrashContext::GetPlatformAllThreadContextsString(FString& OutStr) const
 {
 	OutStr = AllThreadContexts;
+	AllThreadContexts = TEXT("");
 	return !OutStr.IsEmpty();
 }
 
diff --git a/Engine/Source/Runtime/Core/Public/Mac/MacPlatformCrashContext.h b/Engine/Source/Runtime/Core/Public/Mac/MacPlatformCrashContext.h
--- a/Engine/Source/Runtime/Core/Public/Mac/MacPlatformCrashContext.h
+++ b/Engine/Source/Runtime/Core/Public/Mac/MacPlatformCrashContext.h
@@ -22,7 +22,7 @@ struct CORE_API FMacCrashContext : public FApplePlatformCrashContext
 	void GenerateEnsureInfoAndLaunchReporter() const;
 
 	/** Captures all information about all threads */
-	void CaptureAllThreadContext(uint32 ThreadIdEnteredOn);
+	void CaptureAllThreadContext(uint32 ThreadIdEnteredOn) const;
 
 protected:
 	virtual bool GetPlatformAllThreadContextsString(FString& OutStr) const override;
@@ -32,7 +32,7 @@ private:
 		uint32 ThreadIdEnteredOn, 
 		uint32 ThreadId,
 		const FString& ThreadName,
-		const TArray<FCrashStackFrame>& PortableCallStack);
+		const TArray<FCrashStackFrame>& PortableCallStack) const;
 
 	/**
 	* <Thread>
@@ -45,7 +45,6 @@ private:
 	* <Thread>...</Thread>
 	* ...
 	*/
-	FString AllThreadContexts;
 };
 
 typedef FMacCrashContext FPlatformCrashContext;
