Fix the Verse tests build breakage.

#rnx

diff --git a/Engine/Source/Runtime/Core/Private/AutoRTFM/CompilerInterface.cpp b/Engine/Source/Runtime/Core/Private/AutoRTFM/CompilerInterface.cpp
--- a/Engine/Source/Runtime/Core/Private/AutoRTFM/CompilerInterface.cpp
+++ b/Engine/Source/Runtime/Core/Private/AutoRTFM/CompilerInterface.cpp
@@ -116,7 +116,7 @@ extern "C" UE_AUTORTFM_API AUTORTFM_NO_ASAN void autortfm_record_masked_write(vo
 	}
 }
 
-extern "C" UE_AUTORTFM_API UE_AUTORTFM_ALWAYS_OPEN void* autortfm_lookup_function(void* OriginalFunction, const char* Where)
+extern "C" UE_AUTORTFM_ALWAYS_OPEN UE_AUTORTFM_API void* autortfm_lookup_function(void* OriginalFunction, const char* Where)
 {
 	FContext* Context = FContext::Get();
     return FunctionMapLookup(OriginalFunction, Where);
