A new CVar which can contain a list of comma separated Config Branch names that should never be unloaded.
At time of adding a new Branch, checks to see if the name is in the never unload list and sets member variable accordingly.
Skip ConfigBranches that are set to never unload when checking to do a SafeUnload.
Enabling SafeUnload timer for all platforms by default. Set to 30 seconds. Provide initial list of Branch names to never unload.
#tests Local cooked build running against replays and measuring results and log inspection.

diff --git a/Engine/Source/Runtime/Core/Private/Misc/ConfigCacheIni.cpp b/Engine/Source/Runtime/Core/Private/Misc/ConfigCacheIni.cpp
--- a/Engine/Source/Runtime/Core/Private/Misc/ConfigCacheIni.cpp
+++ b/Engine/Source/Runtime/Core/Private/Misc/ConfigCacheIni.cpp
@@ -77,6 +77,12 @@ static FAutoConsoleVariableRef CVarTimeToUnloadConfig(
 	GTimeToUnloadConfig,
 	TEXT("If > 0, when a config branch hasn't been accessed in this many seconds, SafeUnload the branch"));
 
+static FString GConfigBranchesToNeverUnload;
+static FAutoConsoleVariableRef CVarConfigBranchesToNeverUnload(
+	TEXT("ini.ConfigBranchesToNeverUnload"),
+	GConfigBranchesToNeverUnload,
+	TEXT("A comma separated list of config branch names that should never be unloaded."));
+
 
 #if WITH_EDITOR
 // editor wants full replay - and we can't put this in a cvar since we need way too early for cvar processing!
@@ -3544,6 +3550,29 @@ FConfigCacheIni::~FConfigCacheIni()
 	Flush( 1 );
 }
 
+bool FConfigCacheIni::IsConfigBranchNameInNeverUnloadList(const FName& ConfigBranchName)
+{
+	// No branch names to filter
+	if (GConfigBranchesToNeverUnload.Len() == 0)
+	{
+		return false;
+	}
+
+	// Fill out the list first time encountered
+	if (ConfigBranchNamesToNeverUnload.Num() == 0)
+	{
+		GConfigBranchesToNeverUnload.ParseIntoArray(ConfigBranchNamesToNeverUnload, TEXT(","));
+	}
+
+	if (ConfigBranchNamesToNeverUnload.Num() == 0)
+	{
+		return false;
+	}
+
+	const FString IniName = ConfigBranchName.ToString();
+	return ConfigBranchNamesToNeverUnload.Contains(IniName);
+}
+
 void FConfigCacheIni::Tick(float DeltaSeconds)
 {
 	if (GTimeToUnloadConfig == 0)
@@ -3575,8 +3604,14 @@ void FConfigCacheIni::Tick(float DeltaSeconds)
 	double Now = FPlatformTime::Seconds();
 	for (FConfigBranch* Branch : BranchesToCheck)
 	{
-		if (Branch == nullptr || Branch->bIsSafeUnloaded)
+		if (Branch == nullptr || Branch->bIsSafeUnloaded || !Branch->bAllowedToRemove)
+		{
+			continue;
+		}
+
+		if (IsConfigBranchNameInNeverUnloadList(Branch->IniName))
 		{
+			Branch->bAllowedToRemove = false;
 			continue;
 		}
 		
diff --git a/Engine/Source/Runtime/Core/Public/Misc/ConfigCacheIni.h b/Engine/Source/Runtime/Core/Public/Misc/ConfigCacheIni.h
--- a/Engine/Source/Runtime/Core/Public/Misc/ConfigCacheIni.h
+++ b/Engine/Source/Runtime/Core/Public/Misc/ConfigCacheIni.h
@@ -1070,6 +1070,9 @@ public:
 	
 	bool bIsSafeUnloaded;
 	bool bIsHierarchical;
+	// If set true, this ConfigBranch will be allowed to be unloaded.
+	bool bAllowedToRemove = true;
+
 	EBranchReplayMethod ReplayMethod = EBranchReplayMethod::NoReplay;
 	
 	FConfigFileHierarchy Hierarchy;
@@ -1905,6 +1908,9 @@ private:
 	void AddPluginToBranches(FName PluginName, FConfigModificationTracker* ModificationTracker);
 	void RemoveTagFromBranches(FName Tag, FConfigModificationTracker* ModificationTracker);
 
+	/** Does this ConfigBranch IniName appear in our disallowed removal list **/
+	bool IsConfigBranchNameInNeverUnloadList(const FName& ConfigBranchName);
+
 	/** true if file operations should not be performed */
 	bool bAreFileOperationsDisabled;
 
@@ -1943,6 +1949,9 @@ private:
 	static TMap<FName, FConfigCacheIni*> ConfigForPlatform;
 	static FCriticalSection ConfigForPlatformLock;
 #endif
+
+	/** List of ConfigBranch names that should never be unloaded. Read from parsing CVar. See ini.ConfigBranchesToNeverUnload **/
+	TArray<FString> ConfigBranchNamesToNeverUnload;
 	
 	friend FConfigContext;
 };
