Remove code that marks completed FName block as readonly

diff --git a/Engine/Source/Runtime/Core/Private/UObject/UnrealNames.cpp b/Engine/Source/Runtime/Core/Private/UObject/UnrealNames.cpp
--- a/Engine/Source/Runtime/Core/Private/UObject/UnrealNames.cpp
+++ b/Engine/Source/Runtime/Core/Private/UObject/UnrealNames.cpp
@@ -790,14 +790,6 @@ private:
 		}
 #endif
 
-#if USE_FNAME_MMAP
-		// mark last block as readonly
-		if (IsBlockMmapped(Blocks[CurrentBlock]))
-		{
-			mprotect(Blocks[CurrentBlock], PageAlignedBlockSizeBytes, PROT_READ);
-		}
-#endif
-
 		++CurrentBlock;
 		CurrentByteCursor = 0;
 
