Don't null-terminate the internal FTextKey string data

This is no longer required as you never have direct access to it, unlike the old API that returned a const TCHAR* that had to be null-terminated, so removing it saves a small amount of redundant memory.

#rnx

diff --git a/Engine/Source/Runtime/Core/Private/Internationalization/TextKey.cpp b/Engine/Source/Runtime/Core/Private/Internationalization/TextKey.cpp
--- a/Engine/Source/Runtime/Core/Private/Internationalization/TextKey.cpp
+++ b/Engine/Source/Runtime/Core/Private/Internationalization/TextKey.cpp
@@ -82,55 +82,44 @@ public:
 	{
 		check(InIndex != INDEX_NONE);
 
-#if UE_TEXTKEY_USE_UTF8
+		// Read-only
+		int32 NumElementsOnRead = 0;
 		{
-			// Read-only
-			int32 NumElementsOnRead = 0;
+			FReadScopeLock ScopeLock(DataRW);
+
+			if (const FString* FoundString = LegacyTCHARState.Find(InIndex))
 			{
-				FReadScopeLock ScopeLock(DataRW);
+				return **FoundString;
+			}
 
-				if (const FString* FoundString = LegacyTCHARState.Find(InIndex))
-				{
-					return **FoundString;
-				}
+			NumElementsOnRead = LegacyTCHARState.Num();
+		}
 
-				NumElementsOnRead = LegacyTCHARState.Num();
-			}
+		// Write
+		{
+			FWriteScopeLock ScopeLock(DataRW);
 
-			// Write
+			if (LegacyTCHARState.Num() > NumElementsOnRead)
 			{
-				FWriteScopeLock ScopeLock(DataRW);
-
-				if (LegacyTCHARState.Num() > NumElementsOnRead)
+				// Find again in case another thread beat us to it!
+				if (const FString* FoundString = LegacyTCHARState.Find(InIndex))
 				{
-					// Find again in case another thread beat us to it!
-					if (const FString* FoundString = LegacyTCHARState.Find(InIndex))
-					{
-						return **FoundString;
-					}
+					return **FoundString;
 				}
-
-				const FKeyData& KeyData = KeyDataAllocations.Get(InIndex);
-
-				LLM_SCOPE_BYNAME(TEXT("Localization/Deprecated"));
-				// Open around adding this in a cache, if we abort just leak the value in the cache
-				// as the cache takes ownership
-				const TCHAR* AddedString = nullptr;
-				UE_AUTORTFM_OPEN
-				{
-					AddedString = *LegacyTCHARState.Add(InIndex, FString(KeyData.ToView()));
-				};
-				return AddedString;
 			}
-		}
-#else
-		{
-			FReadScopeLock ScopeLock(DataRW);
 
 			const FKeyData& KeyData = KeyDataAllocations.Get(InIndex);
-			return KeyData.Str;
+
+			LLM_SCOPE_BYNAME(TEXT("Localization/Deprecated"));
+			// Open around adding this in a cache, if we abort just leak the value in the cache
+			// as the cache takes ownership
+			const TCHAR* AddedString = nullptr;
+			UE_AUTORTFM_OPEN
+			{
+				AddedString = *LegacyTCHARState.Add(InIndex, FString(KeyData.ToView()));
+			};
+			return AddedString;
 		}
-#endif
 	}
 
 	void AppendStringByIndex(const int32 InIndex, FString& Out) const
@@ -228,7 +217,7 @@ private:
 			return A.StrHash;
 		}
 
-		const FTextKeyCharType* Str;
+		const FTextKeyCharType* Str; // Note: May not be null-terminated; use ToView() to get a valid TStringView
 		int32 StrLen;
 		uint32 StrHash;
 	};
@@ -278,8 +267,8 @@ private:
 			}
 
 			LLM_SCOPE_BYNAME(TEXT("Localization/TextKeys"));
-			const FTextKeyCharType* NewStrPtr = StringAllocations.Add(KeyData.ToView());
-			return KeyDataAllocations.Add(FKeyData(FTextKeyStringView(NewStrPtr, KeyData.StrLen), KeyData));
+			FTextKeyStringView NewStr = StringAllocations.Add(KeyData.ToView());
+			return KeyDataAllocations.Add(FKeyData(NewStr, KeyData));
 		}
 	}
 
@@ -407,18 +396,17 @@ private:
 			//UE_LOG(LogLocalization, Log, TEXT("FTextKey slab allocator allocated %d slabs (of %d elements) and wasted %d elements (%d bytes)"), Slabs.Num(), SlabSizeInElements, TotalNumElementsWasted, TotalNumElementsWasted * sizeof(FTextKeyCharType));
 		}
 
-		const FTextKeyCharType* Add(FTextKeyStringView InStr)
+		FTextKeyStringView Add(FTextKeyStringView InStr)
 		{
-			const int32 NumSlabElementsNeeded = InStr.Len() + 1;
+			const int32 NumSlabElementsNeeded = InStr.Len();
 			FStringSlab& Slab = GetSlab(NumSlabElementsNeeded);
 
 			FTextKeyCharType* StringPtr = Slab.Allocation + Slab.NumElementsUsed;
 			Slab.NumElementsUsed += NumSlabElementsNeeded;
 
 			FMemory::Memcpy(StringPtr, InStr.GetData(), InStr.Len() * sizeof(FTextKeyCharType));
-			*(StringPtr + InStr.Len()) = (FTextKeyCharType)0;
 
-			return StringPtr;
+			return FTextKeyStringView(StringPtr, InStr.Len());
 		}
 
 	private:
@@ -475,10 +463,9 @@ private:
 	mutable FRWLock DataRW;
 	FStringAllocator StringAllocations;
 	FKeyDataAllocator KeyDataAllocations;
-#if UE_TEXTKEY_USE_UTF8
+
 	// Sparse TCHAR state; built on-demand by anything still using the deprecated FTextKey::GetChars function
 	TMap<int32, FString> LegacyTCHARState;
-#endif
 };
 
 namespace TextKeyUtil
