Fix copy pasted comment

#rnx

diff --git a/Engine/Source/Runtime/Core/Public/HAL/UnrealMemory.h b/Engine/Source/Runtime/Core/Public/HAL/UnrealMemory.h
--- a/Engine/Source/Runtime/Core/Public/HAL/UnrealMemory.h
+++ b/Engine/Source/Runtime/Core/Public/HAL/UnrealMemory.h
@@ -263,8 +263,7 @@ struct FMemory
 	*/
 	static CORE_API void EnablePurgatoryTests();
 	/**
-	* Called once main is started and we have -purgatorymallocproxy.
-	* This uses the purgatory malloc proxy to check if things are writing to stale pointers.
+	* Called once main is started and we have -poisonmallocproxy.
 	*/
 	static CORE_API void EnablePoisonTests();
 	/**
