FCycleCounter: Pass file and line to trace. UnrealInsights displays this information for each scope, and it was missing the information for the scopes created by FCycleCounter.
#rnx

diff --git a/Engine/Source/Runtime/Core/Public/Misc/SourceLocation.h b/Engine/Source/Runtime/Core/Public/Misc/SourceLocation.h
--- a/Engine/Source/Runtime/Core/Public/Misc/SourceLocation.h
+++ b/Engine/Source/Runtime/Core/Public/Misc/SourceLocation.h
@@ -60,10 +60,10 @@ namespace UE
 			}
 
 		private:
-			uint_least32_t Line;
-			uint_least32_t Column;
-			const char* FileName;
-			const char* FunctionName;
+			uint_least32_t Line = 0;
+			uint_least32_t Column = 0;
+			const char* FileName = nullptr;
+			const char* FunctionName = nullptr;
 		};
 	#endif // __cpp_lib_source_location
 #else // UE_INCLUDE_SOURCE_LOCATION
diff --git a/Engine/Source/Runtime/Core/Public/Stats/Stats.h b/Engine/Source/Runtime/Core/Public/Stats/Stats.h
--- a/Engine/Source/Runtime/Core/Public/Stats/Stats.h
+++ b/Engine/Source/Runtime/Core/Public/Stats/Stats.h
@@ -17,6 +17,7 @@
 #include "ProfilingDebugging/UMemoryDefines.h"
 #include "Stats/Stats2.h"
 #include "AutoRTFM/AutoRTFM.h"
+#include "Misc/SourceLocation.h"
 
 /** 
  *	Learn about the Stats System at docs.unrealengine.com
@@ -65,14 +66,14 @@ public:
 	 * Pushes the specified stat onto the hierarchy for this thread. Starts
 	 * the timing of the cycles used
 	 */
-	FORCEINLINE_STATS FScopeCycleCounter( TStatId StatId, EStatFlags StatFlags, bool bAlways = false)
+	FORCEINLINE_STATS FScopeCycleCounter( TStatId StatId, EStatFlags StatFlags, bool bAlways = false, UE::FSourceLocation SourceLocation = UE::FSourceLocation::Current())
 	{
-		AutoRTFM::Open([&]{ Start( StatId, StatFlags, bAlways); });
+		AutoRTFM::Open([&]{ Start(StatId, StatFlags, bAlways, SourceLocation); });
 		AutoRTFM::PushOnAbortHandler(this, [this](){ this->Stop(); });
 	}
 
-	FORCEINLINE_STATS FScopeCycleCounter(TStatId StatId, bool bAlways = false)
-		: FScopeCycleCounter(StatId, EStatFlags::None, bAlways)
+	FORCEINLINE_STATS FScopeCycleCounter(TStatId StatId, bool bAlways = false, UE::FSourceLocation SourceLocation = UE::FSourceLocation::Current())
+		: FScopeCycleCounter(StatId, EStatFlags::None, bAlways, SourceLocation)
 	{
 	}
 
diff --git a/Engine/Source/Runtime/Core/Public/Stats/Stats2.h b/Engine/Source/Runtime/Core/Public/Stats/Stats2.h
--- a/Engine/Source/Runtime/Core/Public/Stats/Stats2.h
+++ b/Engine/Source/Runtime/Core/Public/Stats/Stats2.h
@@ -27,6 +27,7 @@
 #include "Misc/Build.h"
 #include "Misc/CString.h"
 #include "Misc/EnumClassFlags.h"
+#include "Misc/SourceLocation.h"
 #include "ProfilingDebugging/CpuProfilerTrace.h"
 #include "ProfilingDebugging/MiscTrace.h"
 #include "StatsCommon.h"
@@ -1666,7 +1667,7 @@ public:
 	 * Pushes the specified stat onto the hierarchy for this thread. Starts
 	 * the timing of the cycles used
 	 */
-	FORCEINLINE_STATS void Start(TStatId InStatId, EStatFlags InStatFlags, bool bAlways = false)
+	FORCEINLINE_STATS void Start(TStatId InStatId, EStatFlags InStatFlags, bool bAlways = false, UE::FSourceLocation SourceLocation = UE::FSourceLocation::Current())
 	{
 		FMinimalName StatMinimalName = InStatId.GetMinimalName(EMemoryOrder::Relaxed);
 		if (StatMinimalName.IsNone())
@@ -1689,7 +1690,7 @@ public:
 			if (UE_TRACE_CHANNELEXPR_IS_ENABLED(CpuChannel))
 			{
 				FName StatName = MinimalNameToName(StatMinimalName);
-				FCpuProfilerTrace::OutputBeginDynamicEventWithId(StatName, InStatId.GetStatDescriptionWIDE());
+				FCpuProfilerTrace::OutputBeginDynamicEventWithId(StatName, InStatId.GetStatDescriptionWIDE(), SourceLocation.GetFileName(), SourceLocation.GetLine());
 				EmittedEvent |= TraceEvent;
 			}
 #endif
@@ -1704,9 +1705,9 @@ public:
 		}
 	}
 
-	FORCEINLINE_STATS void Start(TStatId InStatId, bool bAlways = false)
+	FORCEINLINE_STATS void Start(TStatId InStatId, bool bAlways = false, UE::FSourceLocation SourceLocation = UE::FSourceLocation::Current())
 	{
-		Start(InStatId, EStatFlags::None, bAlways);
+		Start(InStatId, EStatFlags::None, bAlways, SourceLocation);
 	}
 
 	FORCEINLINE_STATS void StartTrace(const FName Name)
