Added UTF-8 overloads of LexFromString for numeric types.

diff --git a/Engine/Source/Runtime/Core/Private/String/LexFromString.cpp b/Engine/Source/Runtime/Core/Private/String/LexFromString.cpp
--- a/Engine/Source/Runtime/Core/Private/String/LexFromString.cpp
+++ b/Engine/Source/Runtime/Core/Private/String/LexFromString.cpp
@@ -3,6 +3,7 @@
 #include "String/LexFromString.h"
 
 #include "Containers/UnrealString.h"
+#include "Containers/Utf8String.h"
 #include "Misc/StringBuilder.h"
 
 void LexFromString(int8& OutValue, const FStringView& InString)
@@ -70,3 +71,69 @@ void LexFromString(bool& OutValue, const FStringView& InString)
 	TStringBuilder<64> Builder;
 	LexFromString(OutValue, (Builder << InString).ToString());
 }
+
+void LexFromString(int8& OutValue, const FUtf8StringView& InString)
+{
+	TUtf8StringBuilder<64> Builder;
+	LexFromString(OutValue, (Builder << InString).ToString());
+}
+
+void LexFromString(int16& OutValue, const FUtf8StringView& InString)
+{
+	TUtf8StringBuilder<64> Builder;
+	LexFromString(OutValue, (Builder << InString).ToString());
+}
+
+void LexFromString(int32& OutValue, const FUtf8StringView& InString)
+{
+	TUtf8StringBuilder<64> Builder;
+	LexFromString(OutValue, (Builder << InString).ToString());
+}
+
+void LexFromString(int64& OutValue, const FUtf8StringView& InString)
+{
+	TUtf8StringBuilder<64> Builder;
+	LexFromString(OutValue, (Builder << InString).ToString());
+}
+
+void LexFromString(uint8& OutValue, const FUtf8StringView& InString)
+{
+	TUtf8StringBuilder<64> Builder;
+	LexFromString(OutValue, (Builder << InString).ToString());
+}
+
+void LexFromString(uint16& OutValue, const FUtf8StringView& InString)
+{
+	TUtf8StringBuilder<64> Builder;
+	LexFromString(OutValue, (Builder << InString).ToString());
+}
+
+void LexFromString(uint32& OutValue, const FUtf8StringView& InString)
+{
+	TUtf8StringBuilder<64> Builder;
+	LexFromString(OutValue, (Builder << InString).ToString());
+}
+
+void LexFromString(uint64& OutValue, const FUtf8StringView& InString)
+{
+	TUtf8StringBuilder<64> Builder;
+	LexFromString(OutValue, (Builder << InString).ToString());
+}
+
+void LexFromString(float& OutValue, const FUtf8StringView& InString)
+{
+	TUtf8StringBuilder<64> Builder;
+	LexFromString(OutValue, (Builder << InString).ToString());
+}
+
+void LexFromString(double& OutValue, const FUtf8StringView& InString)
+{
+	TUtf8StringBuilder<64> Builder;
+	LexFromString(OutValue, (Builder << InString).ToString());
+}
+
+void LexFromString(bool& OutValue, const FUtf8StringView& InString)
+{
+	TUtf8StringBuilder<64> Builder;
+	LexFromString(OutValue, (Builder << InString).ToString());
+}
diff --git a/Engine/Source/Runtime/Core/Public/String/LexFromString.h b/Engine/Source/Runtime/Core/Public/String/LexFromString.h
--- a/Engine/Source/Runtime/Core/Public/String/LexFromString.h
+++ b/Engine/Source/Runtime/Core/Public/String/LexFromString.h
@@ -16,3 +16,15 @@ CORE_API void LexFromString(uint64& OutValue, const FStringView& InString);
 CORE_API void LexFromString(float& OutValue,  const FStringView& InString);
 CORE_API void LexFromString(double& OutValue, const FStringView& InString);
 CORE_API void LexFromString(bool& OutValue,   const FStringView& InString);
+
+CORE_API void LexFromString(int8& OutValue,   const FUtf8StringView& InString);
+CORE_API void LexFromString(int16& OutValue,  const FUtf8StringView& InString);
+CORE_API void LexFromString(int32& OutValue,  const FUtf8StringView& InString);
+CORE_API void LexFromString(int64& OutValue,  const FUtf8StringView& InString);
+CORE_API void LexFromString(uint8& OutValue,  const FUtf8StringView& InString);
+CORE_API void LexFromString(uint16& OutValue, const FUtf8StringView& InString);
+CORE_API void LexFromString(uint32& OutValue, const FUtf8StringView& InString);
+CORE_API void LexFromString(uint64& OutValue, const FUtf8StringView& InString);
+CORE_API void LexFromString(float& OutValue,  const FUtf8StringView& InString);
+CORE_API void LexFromString(double& OutValue, const FUtf8StringView& InString);
+CORE_API void LexFromString(bool& OutValue,   const FUtf8StringView& InString);
