VulkanRHI: Temp fix for packaging SM5 only projects in Linux (limit the NullRHI to ShaderPlatforms supported by the project).  Add debug message to show the selected Vulkan shader platform.

#rnx

diff --git a/Engine/Source/Runtime/Core/Private/Unix/UnixPlatformMisc.cpp b/Engine/Source/Runtime/Core/Private/Unix/UnixPlatformMisc.cpp
--- a/Engine/Source/Runtime/Core/Private/Unix/UnixPlatformMisc.cpp
+++ b/Engine/Source/Runtime/Core/Private/Unix/UnixPlatformMisc.cpp
@@ -767,10 +767,27 @@ const TCHAR* FUnixPlatformMisc::GetNullRHIShaderFormat()
 	{
 		return TEXT("SF_VULKAN_SM5");
 	}
-	else
+	else if (FParse::Param(FCommandLine::Get(), TEXT("sm6")))
 	{
 		return TEXT("SF_VULKAN_SM6");
 	}
+	else
+	{
+		// Limit to the highest SM supported by the project
+		TArray<FString> TargetedShaderFormats;
+		GConfig->GetArray(TEXT("/Script/LinuxTargetPlatform.LinuxTargetSettings"), TEXT("TargetedRHIs"), TargetedShaderFormats, GEngineIni);
+		if (TargetedShaderFormats.Contains(TEXT("SF_VULKAN_SM6")))
+		{
+			return TEXT("SF_VULKAN_SM6");
+		}
+		if (TargetedShaderFormats.Contains(TEXT("SF_VULKAN_SM5")))
+		{
+			return TEXT("SF_VULKAN_SM5");
+		}
+	}
+
+	// Default to SM6
+	return TEXT("SF_VULKAN_SM6");
 }
 
 #define CPUINFO_TOKENS                            \
diff --git a/Engine/Source/Runtime/VulkanRHI/Private/VulkanRHI.cpp b/Engine/Source/Runtime/VulkanRHI/Private/VulkanRHI.cpp
--- a/Engine/Source/Runtime/VulkanRHI/Private/VulkanRHI.cpp
+++ b/Engine/Source/Runtime/VulkanRHI/Private/VulkanRHI.cpp
@@ -399,6 +399,8 @@ FDynamicRHI* FVulkanDynamicRHIModule::CreateRHI(ERHIFeatureLevel::Type InRequest
 	GMaxRHIShaderPlatform = ShaderPlatformForFeatureLevel[GMaxRHIFeatureLevel];
 	checkf(GMaxRHIShaderPlatform != SP_NumPlatforms, TEXT("Requested feature level [%s] mapped to unsupported shader platform!"), *LexToString(InRequestedFeatureLevel));
 
+	UE_LOG(LogVulkanRHI, Display, TEXT("Vulkan RHI ShaderPlatform for %s: %s."), *LexToString(InRequestedFeatureLevel), *LexToString(GMaxRHIShaderPlatform, false));
+
 	GVulkanRHI = new FVulkanDynamicRHI();
 	FDynamicRHI* FinalRHI = GVulkanRHI;
 
