Enable DiskUtilizationTracker on XB1/PS4

#tests ReplayRuns with the stat for mb loaded being reported for boot/map load.
#rnx

diff --git a/Engine/Source/Runtime/Core/Public/HAL/DiskUtilizationTracker.h b/Engine/Source/Runtime/Core/Public/HAL/DiskUtilizationTracker.h
--- a/Engine/Source/Runtime/Core/Public/HAL/DiskUtilizationTracker.h
+++ b/Engine/Source/Runtime/Core/Public/HAL/DiskUtilizationTracker.h
@@ -13,8 +13,6 @@
 #include "HAL/ThreadSafeBool.h"
 #include "ProfilingDebugging/CsvProfiler.h"
 
-#include <nn/nn_Log.h>
-
 CSV_DECLARE_CATEGORY_EXTERN(DiskIO);
 
 #ifndef SPEW_DISK_UTILIZATION
