Don't re-key empty, invariant, or string table texts when using TextNamespaceUtil::CopyTextToPackage

These are safe to copy directly

#rnx

diff --git a/Engine/Source/Runtime/Core/Private/Internationalization/TextNamespaceUtil.cpp b/Engine/Source/Runtime/Core/Private/Internationalization/TextNamespaceUtil.cpp
--- a/Engine/Source/Runtime/Core/Private/Internationalization/TextNamespaceUtil.cpp
+++ b/Engine/Source/Runtime/Core/Private/Internationalization/TextNamespaceUtil.cpp
@@ -63,6 +63,12 @@ void TextNamespaceUtil::StripPackageNamespaceInline(FString& InOutTextNamespace)
 
 FText TextNamespaceUtil::CopyTextToPackage(const FText& InText, const FString& InPackageNamespace, const ETextCopyMethod InCopyMethod, const bool bAlwaysApplyPackageNamespace)
 {
+	if (InText.IsEmpty() || InText.IsCultureInvariant() || InText.IsFromStringTable())
+	{
+		// Safe to use these as-is
+		return InText;
+	}
+
 #if USE_STABLE_LOCALIZATION_KEYS
 	if (!InPackageNamespace.IsEmpty() && (InCopyMethod == ETextCopyMethod::NewKey || InCopyMethod == ETextCopyMethod::PreserveKey))
 	{
