[Config] Allow parsing config tokens from commandline with brackets to isolate delimiters.

This fixes the issue where multiple key-value pairs are specified in brackets for INI configurations in the commandline.
Such a commandline argument can look like this: "-ini:Engine:[UbaController]:Horde=(Pool=Name,MaxCores=500)".
The current implementation parses the value for the INI configuration "Horde" as "(Pool=Name" because the following comma separates the tokens, even though brackets are meant to prevent this.

diff --git a/Engine/Source/Runtime/Core/Private/Misc/ConfigCacheIni.cpp b/Engine/Source/Runtime/Core/Private/Misc/ConfigCacheIni.cpp
--- a/Engine/Source/Runtime/Core/Private/Misc/ConfigCacheIni.cpp
+++ b/Engine/Source/Runtime/Core/Private/Misc/ConfigCacheIni.cpp
@@ -1871,7 +1871,7 @@ void FConfigFile::OverrideFromCommandline(FConfigCommandStream* File, const FStr
 			{
 				FString NextSettingToken;
 				const TCHAR* SettingsStream = *SettingsString;
-				while (FParse::Token(SettingsStream, NextSettingToken, false, CommandlineOverrideSpecifiers::PropertySeperator))
+				while (FParse::Expression(SettingsStream, NextSettingToken, false, CommandlineOverrideSpecifiers::PropertySeperator))
 				{
 					SettingPairs.Add(MoveTemp(NextSettingToken));
 				}
diff --git a/Engine/Source/Runtime/Core/Private/Misc/Parse.cpp b/Engine/Source/Runtime/Core/Private/Misc/Parse.cpp
--- a/Engine/Source/Runtime/Core/Private/Misc/Parse.cpp
+++ b/Engine/Source/Runtime/Core/Private/Misc/Parse.cpp
@@ -818,17 +818,22 @@ void FParse::Next( const TCHAR** Stream )
 	// Upon exit, *Stream either points to valid Stream or a nul.
 }
 
+static bool IsDelimiterOrWhitespace(TCHAR Character, const TCHAR SingleCharacterDelimiter)
+{
+	return (SingleCharacterDelimiter == TEXT('\0') && FChar::IsWhitespace(Character))
+		|| (SingleCharacterDelimiter != TEXT('\0') && Character == SingleCharacterDelimiter);
+}
+
 //
 // Grab the next space-delimited (or SingleCharacterDelimiter-delimited) string from the input stream.
 // If quoted, gets entire quoted string.
 //
-bool FParse::Token( const TCHAR*& Str, TCHAR* Result, int32 MaxLen, bool UseEscape, const TCHAR SingleCharacterDelimiter/* = TEXT('\0')*/)
+bool FParse::Token(const TCHAR*& Str, TCHAR* Result, int32 MaxLen, bool bUseEscape, const TCHAR SingleCharacterDelimiter/* = TEXT('\0')*/)
 {
 	int32 Len=0;
 
 	// Skip preceeding delimiters (either spaces and tabs or custom delimiters)
-	while((SingleCharacterDelimiter == TEXT('\0') && (FChar::IsWhitespace(*Str)))
-		|| ((SingleCharacterDelimiter != TEXT('\0') && *Str == SingleCharacterDelimiter)))
+	while (IsDelimiterOrWhitespace(*Str, SingleCharacterDelimiter))
 	{
 		Str++;
 	}
@@ -837,24 +842,24 @@ bool FParse::Token( const TCHAR*& Str, TCHAR* Result, int32 MaxLen, bool UseEsca
 	{
 		// Get quoted string.
 		Str++;
-		while( *Str && *Str!=TEXT('"') && (Len+1)<MaxLen )
+		while (*Str && *Str != TEXT('"') && (Len + 1) < MaxLen)
 		{
-			TCHAR c = *Str++;
-			if( c==TEXT('\\') && UseEscape )
+			TCHAR Character = *Str++;
+			if (Character == TEXT('\\') && bUseEscape)
 			{
 				// Get escape.
-				c = *Str++;
-				if( !c )
+				Character = *Str++;
+				if (!Character)
 				{
 					break;
 				}
 			}
-			if( (Len+1)<MaxLen )
+			if ((Len + 1) < MaxLen)
 			{
-				Result[Len++] = c;
+				Result[Len++] = Character;
 			}
 		}
-		if( *Str==TEXT('"') )
+		if (*Str == TEXT('"'))
 		{
 			Str++;
 		}
@@ -874,16 +879,16 @@ bool FParse::Token( const TCHAR*& Str, TCHAR* Result, int32 MaxLen, bool UseEsca
 			}
 			if (!bInQuote)
 			{
-				if ((SingleCharacterDelimiter != TEXT('\0') && Character == SingleCharacterDelimiter)
-					|| (SingleCharacterDelimiter == TEXT('\0') && (FChar::IsWhitespace(Character))))
+				if (IsDelimiterOrWhitespace(Character, SingleCharacterDelimiter))
 				{
+					// Intentionally leave trailing delimiters unchanged
 					break;
 				}
 			}
 			Str++;
 
 			// Preserve escapes if they're in a quoted string (the check for " is in the else to let \" work as expected)
-			if (Character == TEXT('\\') && UseEscape && bInQuote)
+			if (Character == TEXT('\\') && bUseEscape && bInQuote)
 			{
 				if ((Len+1) < MaxLen)
 				{
@@ -912,12 +917,12 @@ bool FParse::Token( const TCHAR*& Str, TCHAR* Result, int32 MaxLen, bool UseEsca
 	return Len != 0;
 }
 
-bool FParse::Token( const TCHAR*& Str, FString& Arg, bool UseEscape, const TCHAR SingleCharacterDelimiter/* = TEXT('\0')*/)
+bool FParse::Token( const TCHAR*& Str, FString& Arg, bool bUseEscape, const TCHAR SingleCharacterDelimiter/* = TEXT('\0')*/)
 {
 	Arg.Reset();
 
-	// Skip preceeding spaces and tabs.
-	while( FChar::IsWhitespace(*Str) )
+	// Skip preceeding delimiters (either spaces and tabs or custom delimiters)
+	while (IsDelimiterOrWhitespace(*Str, SingleCharacterDelimiter))
 	{
 		Str++;
 	}
@@ -926,23 +931,23 @@ bool FParse::Token( const TCHAR*& Str, FString& Arg, bool UseEscape, const TCHAR
 	{
 		// Get quoted string.
 		Str++;
-		while( *Str && *Str != TCHAR('"') )
+		while (*Str && *Str != TEXT('"'))
 		{
-			TCHAR c = *Str++;
-			if( c==TEXT('\\') && UseEscape )
+			TCHAR Character = *Str++;
+			if (Character == TEXT('\\') && bUseEscape)
 			{
 				// Get escape.
-				c = *Str++;
-				if( !c )
+				Character = *Str++;
+				if (!Character)
 				{
 					break;
 				}
 			}
 
-			Arg += c;
+			Arg += Character;
 		}
 
-		if ( *Str == TEXT('"') )
+		if (*Str == TEXT('"'))
 		{
 			Str++;
 		}
@@ -962,8 +967,7 @@ bool FParse::Token( const TCHAR*& Str, FString& Arg, bool UseEscape, const TCHAR
 			}
 			if (!bInQuote)
 			{
-				if ((SingleCharacterDelimiter != TEXT('\0') && Character == SingleCharacterDelimiter)
-					|| (SingleCharacterDelimiter == TEXT('\0') && (FChar::IsWhitespace(Character))))
+				if (IsDelimiterOrWhitespace(Character, SingleCharacterDelimiter))
 				{
 					// Consume the delimiter. If it's whitespace this isn't critical since we'll consume it at the start
 					// of the next call to Token() but if it's not whitespace we won't, so we better do it now.
@@ -974,7 +978,7 @@ bool FParse::Token( const TCHAR*& Str, FString& Arg, bool UseEscape, const TCHAR
 			Str++;
 
 			// Preserve escapes if they're in a quoted string (the check for " is in the else to let \" work as expected)
-			if (Character == TEXT('\\') && UseEscape && bInQuote)
+			if (Character == TEXT('\\') && bUseEscape && bInQuote)
 			{
 				Arg += Character;
 
@@ -1348,6 +1352,115 @@ bool FParse::SchemeNameFromURI(const TCHAR* URI, FString& OutSchemeName)
 	}
 }
 
+static TCHAR GetCloseBracketCharacterForOpenBracket(TCHAR Character)
+{
+	switch (Character)
+	{
+	case TEXT('('):	return TEXT(')');
+	case TEXT('['):	return TEXT(']');
+	case TEXT('{'):	return TEXT('}');
+	default:		return TEXT('\0'); // no match
+	}
+}
+
+bool FParse::Expression(const TCHAR*& Str, FString& OutExpression, bool bUseEscape, const TCHAR SingleCharacterDelimiter/* = TEXT('\0')*/)
+{
+	OutExpression.Reset();
+
+	const TCHAR* OriginalStr = Str;
+
+	// Skip preceeding delimiters (either spaces and tabs or custom delimiters)
+	while (IsDelimiterOrWhitespace(*Str, SingleCharacterDelimiter))
+	{
+		Str++;
+	}
+
+	if (*Str == TEXT('"'))
+	{
+		// Forward to Token() to get quoted string.
+		return FParse::Token(Str, OutExpression, bUseEscape, SingleCharacterDelimiter);
+	}
+
+	// Get unquoted string (that might contain a quoted part, which will be left intact).
+	// For example, -ARG="foo bar baz", will be treated as one token, with quotes intact
+	bool bInQuote = false;
+	TArray<TCHAR> BracketStack;
+
+	while (1)
+	{
+		TCHAR Character = *Str;
+		if (Character == 0)
+		{
+			break;
+		}
+		if (!bInQuote)
+		{
+			// Increase and decrease bracket level but we don't care about matching pairs.
+			// This is only to isolate delimiters, not to provide syntax validation.
+			if (Character == TEXT('(') || Character == TEXT('[') || Character == TEXT('{'))
+			{
+				BracketStack.Add(Character);
+			}
+			else if (Character == TEXT(')') || Character == TEXT(']') || Character == TEXT('}'))
+			{
+				if (BracketStack.IsEmpty())
+				{
+					// If bracket stack is empty, interpret closing bracket as delimiter when parsed as interleaved expression.
+					// Example: "A=(B=(C,D))"
+					//              ^      ^__ This bracket ends parsing the expression "B=(C,D)"
+					//              |__ Start position
+					break;
+				}
+				if (GetCloseBracketCharacterForOpenBracket(BracketStack.Top()) != Character)
+				{
+					// Could not match closing bracket with previous open bracket
+					break;
+				}
+				BracketStack.Pop();
+			}
+
+			if (BracketStack.IsEmpty())
+			{
+				if (IsDelimiterOrWhitespace(Character, SingleCharacterDelimiter))
+				{
+					// Intentionally leave trailing delimiters unchanged
+					break;
+				}
+			}
+		}
+		Str++;
+
+		// Preserve escapes if they're in a quoted string (the check for " is in the else to let \" work as expected)
+		if (Character == TEXT('\\') && bUseEscape && bInQuote)
+		{
+			OutExpression += Character;
+
+			Character = *Str;
+			if (!Character)
+			{
+				break;
+			}
+			Str++;
+		}
+		else if (Character == TEXT('"'))
+		{
+			bInQuote = !bInQuote;
+		}
+
+		OutExpression += Character;
+	}
+
+	// If brackets or quotation marks are not balanced, return false as parsing has failed and reset the input string
+	if (!BracketStack.IsEmpty() || bInQuote)
+	{
+		Str = OriginalStr;
+		OutExpression.Reset();
+		return false;
+	}
+
+	return OutExpression.Len() > 0;
+}
+
 
 #if !(UE_BUILD_SHIPPING || UE_BUILD_TEST)
 
@@ -1393,6 +1506,77 @@ bool FParseLineExtendedTest::RunTest(const FString& Parameters)
 	return true;
 }
 
+IMPLEMENT_SIMPLE_AUTOMATION_TEST(FParseExpressionSimpleTest, "System.Core.Misc.ParseExpressionSimple", EAutomationTestFlags_ApplicationContextMask | EAutomationTestFlags::SmokeFilter)
+bool FParseExpressionSimpleTest::RunTest(const FString& Parameters)
+{
+	const TCHAR* Tests[] =
+	{
+		TEXT("Test string"),                                    // Normal string
+		TEXT("{Test string}"),                                  // Braced string
+		TEXT("\"Test string\""),                                // Quoted string -> quotation marks will be trimmed
+		TEXT("\"Test \\\"string\\\"\""),                        // Quoted string w/ escaped quotes
+		TEXT("a=\"Test\",b=\"Test\""),                          // Quoted value list -> quotation marks are left in tact
+		TEXT("a=\"Test\\\\\",b=\"{Test}\""),                    // Quoted value list w/ escaped backslash preceeding closing quote
+		TEXT("a=\"Test\\\\\\\" String\",b=\"{Test}\""),         // Quoted value list w/ escaped backslash preceeding escaped quote
+		TEXT("Test=(Inner=\"{content}\")"),                     // Nested value list
+		TEXT("Test=(Inner=\"{content}\",Inner2=[\"{{{{\"])"),   // Double nested value list
+		TEXT("B=(C,D))"),                                       // Trailing brackets when parsed from interleaved expressions, e.g. "A=(B=(C,D))"
+		TEXT("(Inner=\"{content}\",Inner2=[\"{{{{\"]"),         // Failed expression due to unbalanced brackets
+		TEXT(")Inner("),                                        // Failed expression due to wrong bracket orientation
+		TEXT("([MismatchedBrackets)]"),                         // Failed expression due to mismatched brackets
+	};
+
+	struct FParseExpressionExpectedTestResult
+	{
+		bool bSuccess;
+		const TCHAR Delimiter;
+		TArray<const TCHAR*> ExpectedExpressions;
+	};
+
+	const FParseExpressionExpectedTestResult ExpectedResults[] =
+	{
+		FParseExpressionExpectedTestResult{ true, TEXT(' '), { TEXT("Test"), TEXT("string") } },
+		FParseExpressionExpectedTestResult{ true, TEXT(' '), { TEXT("{Test string}") } },
+		FParseExpressionExpectedTestResult{ true, TEXT(' '), { TEXT("Test string") } },
+		FParseExpressionExpectedTestResult{ true, TEXT(' '), { TEXT("Test \"string\"") } },
+		FParseExpressionExpectedTestResult{ true, TEXT(','), { TEXT("a=\"Test\""), TEXT("b=\"Test\"") } },
+		FParseExpressionExpectedTestResult{ true, TEXT(','), { TEXT("a=\"Test\\\\\""), TEXT("b=\"{Test}\"") } },
+		FParseExpressionExpectedTestResult{ true, TEXT(','), { TEXT("a=\"Test\\\\\\\" String\""), TEXT("b=\"{Test}\"") } },
+		FParseExpressionExpectedTestResult{ true, TEXT('='), { TEXT("Test"), TEXT("(Inner=\"{content}\")") } },
+		FParseExpressionExpectedTestResult{ true, TEXT('='), { TEXT("Test"), TEXT("(Inner=\"{content}\",Inner2=[\"{{{{\"])") } },
+		FParseExpressionExpectedTestResult{ true, TEXT('='), { TEXT("B"), TEXT("(C,D)")}},
+		FParseExpressionExpectedTestResult{ false, TEXT('='), { } },
+		FParseExpressionExpectedTestResult{ false, TEXT('\0'), { } },
+		FParseExpressionExpectedTestResult{ false, TEXT('\0'), { } },
+	};
+
+	FString Result;
+	for (int32 TestIndex = 0; TestIndex < UE_ARRAY_COUNT(ExpectedResults); ++TestIndex)
+	{
+		const TCHAR* Stream = Tests[TestIndex];
+		const FParseExpressionExpectedTestResult& Expected = ExpectedResults[TestIndex];
+		if (Expected.bSuccess)
+		{
+			for (int32 ExpressionIndex = 0; ExpressionIndex < Expected.ExpectedExpressions.Num(); ++ExpressionIndex)
+			{
+				Result.Reset();
+				const TCHAR* ExpectedExpression = Expected.ExpectedExpressions[ExpressionIndex];
+				bool bSuccess = FParse::Expression(Stream, Result, true, Expected.Delimiter);
+				TestTrue(*FString::Printf(TEXT("Expecting parsed expression [%s] %d/%d to be [%s]. Result was [%s]."),
+					Tests[TestIndex], ExpressionIndex + 1, Expected.ExpectedExpressions.Num(), ExpectedExpression, *Result), bSuccess && Result == ExpectedExpression);
+			}
+		}
+		else
+		{
+			Result.Reset();
+			bool bSuccess = FParse::Expression(Stream, Result, true, Expected.Delimiter);
+			TestFalse(*FString::Printf(TEXT("Expecting parsed expression [%s] to fail. Result was [%s]."), Tests[TestIndex], *Result), bSuccess);
+		}
+	}
+
+	return true;
+}
+
 #endif //!(UE_BUILD_SHIPPING || UE_BUILD_TEST)
 
 
diff --git a/Engine/Source/Runtime/Core/Public/Misc/Parse.h b/Engine/Source/Runtime/Core/Public/Misc/Parse.h
--- a/Engine/Source/Runtime/Core/Public/Misc/Parse.h
+++ b/Engine/Source/Runtime/Core/Public/Misc/Parse.h
@@ -88,9 +88,9 @@ struct FParse
 	static CORE_API bool LineExtended(const TCHAR** Stream, FString& Result, int32& LinesConsumed, bool Exact = 0);
 	static CORE_API bool LineExtended(const TCHAR** Stream, FStringBuilderBase& Result, int32& LinesConsumed, bool Exact = 0);
 	/** Grabs the next space-delimited (unless SingleCharacterDelimiter is provided) string from the input stream. If quoted, gets entire quoted string. */
-	static CORE_API bool Token( const TCHAR*& Str, TCHAR* Result, int32 MaxLen, bool UseEscape, const TCHAR SingleCharacterDelimiter = TEXT('\0'));
+	static CORE_API bool Token( const TCHAR*& Str, TCHAR* Result, int32 MaxLen, bool bUseEscape, const TCHAR SingleCharacterDelimiter = TEXT('\0'));
 	/** Grabs the next space-delimited (unless SingleCharacterDelimiter is provided) string from the input stream. If quoted, gets entire quoted string. */
-	static CORE_API bool Token( const TCHAR*& Str, FString& Arg, bool UseEscape, const TCHAR SingleCharacterDelimiter = TEXT('\0'));
+	static CORE_API bool Token( const TCHAR*& Str, FString& Arg, bool bUseEscape, const TCHAR SingleCharacterDelimiter = TEXT('\0'));
 	/** Grabs the next alpha-numeric space-delimited token from the input stream. */
 	static CORE_API bool AlnumToken(const TCHAR*& Str, FString& Arg);
 	/** Grabs the next space-delimited string from the input stream. If quoted, gets entire quoted string. */
@@ -104,6 +104,17 @@ struct FParse
 	/** Parse a quoted string token. */
 	static CORE_API bool QuotedString( const TCHAR* Stream, FString& Value, int32* OutNumCharsRead = nullptr );
 	static CORE_API bool QuotedString( const TCHAR* Stream, FStringBuilderBase& Value, int32* OutNumCharsRead = nullptr );
+	/**
+	* Parse the next space-delimited (or SingleCharacterDelmiter-delimited) expression, using brackets
+	* (), {}, or [] to allow subexpressions to be part of the single expression even if they contain the
+	* delimiter.
+	* The parsed expression can be either an atomic token such as "Value1" in the stream "Value1 Value2 Value3"
+	* or a bracketed expression such as "(Value1.1 Value1.2)" in the stream "(Value1.1 Value1.2) Value2 Value3".
+	* If quotes are encountered, text inside of the quotes is not interpreted (brackets and delimiters are ignored)
+	* and it is included verbatim in the parsed expression.
+	* Trailing delimiters are not consumed, only leading delimiters.
+	*/
+	static CORE_API bool Expression(const TCHAR*& Str, FString& OutExpression, bool bUseEscape, const TCHAR SingleCharacterDelimiter = TEXT('\0'));
 
 	//
 	// Parse a hex digit.
