Friending Global comparison operators

diff --git a/Engine/Source/Runtime/Core/Public/Templates/PointerVariants.h b/Engine/Source/Runtime/Core/Public/Templates/PointerVariants.h
--- a/Engine/Source/Runtime/Core/Public/Templates/PointerVariants.h
+++ b/Engine/Source/Runtime/Core/Public/Templates/PointerVariants.h
@@ -25,6 +25,22 @@ struct TWeakPtrVariant;
 template<typename BaseType>
 struct TStrongPtrVariant;
 
+namespace UE::Core::Private
+{
+	template<typename BaseType, bool bIsStrong>
+	struct TPtrVariantBase;
+}
+
+template <typename LhsType, bool bLhsStrength, typename RhsType, bool bRhsStrength,
+	decltype((LhsType*)nullptr == (RhsType*)nullptr) = false  /* this will prevent the template instantiating if the pointer types are not comparable*/>
+bool operator==(const UE::Core::Private::TPtrVariantBase<LhsType, bLhsStrength>& Lhs, const UE::Core::Private::TPtrVariantBase<RhsType, bRhsStrength>& Rhs);
+
+template<typename LhsType, typename RhsType, bool bRhsStrength, decltype((LhsType*)nullptr == (RhsType*)nullptr) = false>
+bool operator==(LhsType* Lhs, const UE::Core::Private::TPtrVariantBase<RhsType, bRhsStrength>& Rhs);
+
+template<typename LhsType, bool bLhsStrength, typename RhsType, decltype((LhsType*)nullptr == (RhsType*)nullptr) = false>
+bool operator==(const UE::Core::Private::TPtrVariantBase<LhsType, bLhsStrength>& Lhs, RhsType* Rhs);
+
 namespace UE::Core::Private
 {
 	template<typename BaseType, bool bIsStrong>
@@ -118,6 +134,40 @@ namespace UE::Core::Private
 				}, PtrVariant);
 		}
 
+	private:
+		template <typename LhsType, bool bLhsStrength, typename RhsType, bool bRhsStrength, decltype((LhsType*)nullptr == (RhsType*)nullptr)>
+		friend bool(::operator==)(const TPtrVariantBase<LhsType, bLhsStrength>& Lhs, const TPtrVariantBase<RhsType, bRhsStrength>& Rhs);
+
+		template<typename LhsType, typename RhsType, bool bRhsStrength, decltype((LhsType*)nullptr == (RhsType*)nullptr)>
+		friend bool(::operator==)(LhsType* Lhs, const TPtrVariantBase<RhsType, bRhsStrength>& Rhs);
+
+		template<typename LhsType, bool bLhsStrength, typename RhsType, decltype((LhsType*)nullptr == (RhsType*)nullptr)>
+		friend bool(::operator==)(const TPtrVariantBase<LhsType, bLhsStrength>& Lhs, RhsType* Rhs);
+
+	private:
+		const BaseType* GetRawPtrValue_Internal() const
+		{
+			return ::Visit([this](auto& StoredValue)
+					{
+						using StoredValueType = std::decay_t<decltype(StoredValue)>;
+						if constexpr (std::is_same_v<StoredValueType, SharedType>)
+						{
+							if constexpr (bIsStrong)
+							{
+								return IsValid() ? StoredValue.Get() : (const BaseType*)nullptr;
+							}
+							else
+							{
+								return IsValid() ? StoredValue.Pin().Get() : (const BaseType*)nullptr;
+							}
+						}
+						else
+						{
+							return IsValid() ? StoredValue.CastedPtr : (const BaseType*)nullptr;
+						}
+					}, PtrVariant);
+		}
+
 	private:
 		TVariant<SharedType, FObjectPtrWrapper> PtrVariant;
 	};
@@ -125,38 +175,37 @@ namespace UE::Core::Private
 
 /* Global Comparisons */
 
-template <typename LhsType, bool bLhsStrength, typename RhsType, bool bRhsStrength, 
-	decltype((LhsType*)nullptr == (RhsType*)nullptr) = false /* this will prevent the template instantiating if the pointer types are not comparable*/>
+template <typename LhsType, bool bLhsStrength, typename RhsType, bool bRhsStrength, decltype((LhsType*)nullptr == (RhsType*)nullptr)>
 bool operator==(const UE::Core::Private::TPtrVariantBase<LhsType, bLhsStrength>& Lhs, const UE::Core::Private::TPtrVariantBase<RhsType, bRhsStrength>& Rhs)
 {
-	return Lhs.Get() == Rhs.Get();
+	return Lhs.GetRawPtrValue_Internal() == Rhs.GetRawPtrValue_Internal();
 }
 
-template <typename LhsType, bool bLhsStrength, typename RhsType, bool bRhsStrength, decltype((LhsType*)nullptr == (RhsType*)nullptr) = false>
+template <typename LhsType, bool bLhsStrength, typename RhsType, bool bRhsStrength, decltype((LhsType*)nullptr == (RhsType*)nullptr)>
 bool operator!=(const UE::Core::Private::TPtrVariantBase<LhsType, bLhsStrength>& Lhs, const UE::Core::Private::TPtrVariantBase<RhsType, bRhsStrength>& Rhs)
 {
 	return !(Lhs == Rhs);
 }
 
-template<typename LhsType, typename RhsType, bool bRhsStrength, decltype((LhsType*)nullptr == (RhsType*)nullptr) = false>
+template<typename LhsType, typename RhsType, bool bRhsStrength, decltype((LhsType*)nullptr == (RhsType*)nullptr)>
 bool operator==(LhsType* Lhs, const UE::Core::Private::TPtrVariantBase<RhsType, bRhsStrength>& Rhs)
 {
-	return Lhs == Rhs.Get();
+	return Lhs == Rhs.GetRawPtrValue_Internal();
 }
 
-template<typename LhsType, typename RhsType, bool bRhsStrength, decltype((LhsType*)nullptr == (RhsType*)nullptr) = false>
+template<typename LhsType, typename RhsType, bool bRhsStrength, decltype((LhsType*)nullptr == (RhsType*)nullptr)>
 bool operator!=(LhsType* Lhs, const UE::Core::Private::TPtrVariantBase<RhsType, bRhsStrength>& Rhs)
 {
 	return !(Lhs == Rhs);
 }
 
-template<typename LhsType, bool bLhsStrength, typename RhsType, decltype((LhsType*)nullptr == (RhsType*)nullptr) = false>
+template<typename LhsType, bool bLhsStrength, typename RhsType, decltype((LhsType*)nullptr == (RhsType*)nullptr)>
 bool operator==(const UE::Core::Private::TPtrVariantBase<LhsType, bLhsStrength>& Lhs, RhsType* Rhs)
 {
-	return Lhs.Get() == Rhs;
+	return Lhs.GetRawPtrValue_Internal() == Rhs;
 }
 
-template<typename LhsType, bool bLhsStrength, typename RhsType, decltype((LhsType*)nullptr == (RhsType*)nullptr) = false>
+template<typename LhsType, bool bLhsStrength, typename RhsType, decltype((LhsType*)nullptr == (RhsType*)nullptr)>
 bool operator!=(const UE::Core::Private::TPtrVariantBase<LhsType, bLhsStrength>& Lhs, RhsType* Rhs)
 {
 	return !(Lhs == Rhs);
