Ensure we have the disable pragma clang optimization on/off defined appropriately for Mac and iOS clients

#tests Verified that disable compiler optimizations flag functioned as expected.

diff --git a/Engine/Source/Runtime/Core/Public/IOS/IOSPlatform.h b/Engine/Source/Runtime/Core/Public/IOS/IOSPlatform.h
--- a/Engine/Source/Runtime/Core/Public/IOS/IOSPlatform.h
+++ b/Engine/Source/Runtime/Core/Public/IOS/IOSPlatform.h
@@ -110,6 +110,13 @@ typedef FIOSPlatformTypes FPlatformTypes;
 
 #define ABSTRACT abstract
 
+// We can use pragma optimisation's on and off as of Apple LLVM 7.3.0 but not before.
+#if (__clang_major__ > 7) || (__clang_major__ == 7 && __clang_minor__ >= 3)
+#define PRAGMA_DISABLE_OPTIMIZATION_ACTUAL _Pragma("clang optimize off")
+#define PRAGMA_ENABLE_OPTIMIZATION_ACTUAL  _Pragma("clang optimize on")
+#endif
+
+// Strings.
 // Strings.
 #define LINE_TERMINATOR TEXT("\n")
 #define LINE_TERMINATOR_ANSI "\n"
diff --git a/Engine/Source/Runtime/Core/Public/Mac/MacPlatform.h b/Engine/Source/Runtime/Core/Public/Mac/MacPlatform.h
--- a/Engine/Source/Runtime/Core/Public/Mac/MacPlatform.h
+++ b/Engine/Source/Runtime/Core/Public/Mac/MacPlatform.h
@@ -120,6 +120,13 @@ typedef FMacPlatformTypes FPlatformTypes;
 
 #define ABSTRACT abstract
 
+// We can use pragma optimisation's on and off as of Apple LLVM 7.3.0 but not before.
+#if (__clang_major__ > 7) || (__clang_major__ == 7 && __clang_minor__ >= 3)
+#define PRAGMA_DISABLE_OPTIMIZATION_ACTUAL _Pragma("clang optimize off")
+#define PRAGMA_ENABLE_OPTIMIZATION_ACTUAL  _Pragma("clang optimize on")
+#endif
+
+// Strings.
 // Alignment.
 #define GCC_PACK(n) __attribute__((packed,aligned(n)))
 #define GCC_ALIGN(n) __attribute__((aligned(n)))
