Updated FGuid::Parse and FGuid::ParseExact to accept a FStringView

Also cleaned up temporary FString usage in these functions, and FGuid::ExportTextItem and FGuid::ImportTextItem

diff --git a/Engine/Source/Runtime/Core/Private/Misc/Guid.cpp b/Engine/Source/Runtime/Core/Private/Misc/Guid.cpp
--- a/Engine/Source/Runtime/Core/Private/Misc/Guid.cpp
+++ b/Engine/Source/Runtime/Core/Private/Misc/Guid.cpp
@@ -25,7 +25,7 @@
 
 bool FGuid::ExportTextItem(FString& ValueStr, FGuid const& DefaultValue, UObject* Parent, int32 PortFlags, class UObject* ExportRootScope) const
 {
-	ValueStr += ToString();
+	AppendString(ValueStr, EGuidFormats::Digits);
 
 	return true;
 }
@@ -33,12 +33,14 @@ bool FGuid::ExportTextItem(FString& ValueStr, FGuid const& DefaultValue, UObject
 
 bool FGuid::ImportTextItem(const TCHAR*& Buffer, int32 PortFlags, class UObject* Parent, FOutputDevice* ErrorText)
 {
-	if (FPlatformString::Strlen(Buffer) < 32)
+	FStringView BufferView = Buffer;
+
+	if (BufferView.Len() < 32)
 	{
 		return false;
 	}
 
-	if (!ParseExact(FString(Buffer).Left(32), EGuidFormats::Digits, *this))
+	if (!ParseExact(BufferView.Left(32), EGuidFormats::Digits, *this))
 	{
 		return false;
 	}
@@ -282,7 +284,17 @@ FGuid FGuid::Combine(const FGuid& GuidA, const FGuid& GuidB)
 	);
 }
 
+bool FGuid::Parse(const TCHAR* GuidString, FGuid& OutGuid)
+{
+	return Parse(FStringView(GuidString), OutGuid);
+}
+
 bool FGuid::Parse(const FString& GuidString, FGuid& OutGuid)
+{
+	return Parse(FStringView(GuidString), OutGuid);
+}
+
+bool FGuid::Parse(FStringView GuidString, FGuid& OutGuid)
 {
 	if (GuidString.Len() == 32)
 	{
@@ -296,7 +308,7 @@ bool FGuid::Parse(const FString& GuidString, FGuid& OutGuid)
 
 	if (GuidString.Len() == 38)
 	{
-		if (GuidString.StartsWith("{"))
+		if (GuidString[0] == TEXT('{'))
 		{
 			return ParseExact(GuidString, EGuidFormats::DigitsWithHyphensInBraces, OutGuid);
 		}
@@ -327,8 +339,17 @@ bool FGuid::Parse(const FString& GuidString, FGuid& OutGuid)
 	return false;
 }
 
+bool FGuid::ParseExact(const TCHAR* GuidString, EGuidFormats Format, FGuid& OutGuid)
+{
+	return ParseExact(FStringView(GuidString), Format, OutGuid);
+}
 
 bool FGuid::ParseExact(const FString& GuidString, EGuidFormats Format, FGuid& OutGuid)
+{
+	return ParseExact(FStringView(GuidString), Format, OutGuid);
+}
+
+bool FGuid::ParseExact(FStringView GuidString, EGuidFormats Format, FGuid& OutGuid)
 {
 	// The below normalizing into the Digits format doesn't work with Short and Base36 Guids
 	if (Format == EGuidFormats::Short)
@@ -341,19 +362,20 @@ bool FGuid::ParseExact(const FString& GuidString, EGuidFormats Format, FGuid& Ou
 		uint32 Data[4] = {};
 
 		// We don't need to do replacements to get an accurate size (deferring allocations till later)
-		if (FBase64::GetDecodedDataSize(GuidString) != sizeof(Data))
+		if (FBase64::GetDecodedDataSize(GuidString.GetData(), GuidString.Len()) != sizeof(Data))
 		{
 			// This isn't a Short GUID if it's not 128 bits / 16 bytes (the size of Data)
 			return false;
 		}
 
 		// Replace the characters we replaced going out (but not the padding, UE discards it immediately)
-		FString GuidCopy = GuidString;
-		GuidCopy.ReplaceCharInline(TEXT('-'), TEXT('+'), ESearchCase::CaseSensitive);
-		GuidCopy.ReplaceCharInline(TEXT('_'), TEXT('/'), ESearchCase::CaseSensitive);
+		TStringBuilder<22> GuidCopy;
+		GuidCopy += GuidString;
+		Algo::Replace(MakeArrayView(GuidCopy), TEXT('-'), TEXT('+'));
+		Algo::Replace(MakeArrayView(GuidCopy), TEXT('_'), TEXT('/'));
 
 		// Decode the data
-		if (!FBase64::Decode(*GuidCopy, GuidCopy.Len(), reinterpret_cast<uint8*>(&Data)))
+		if (!FBase64::Decode(GuidCopy.GetData(), GuidCopy.Len(), reinterpret_cast<uint8*>(&Data)))
 		{
 			// Data is not valid in some way
 			return false;
@@ -374,39 +396,34 @@ bool FGuid::ParseExact(const FString& GuidString, EGuidFormats Format, FGuid& Ou
 		}
 
 		FUInt128 Value;
-		const TCHAR* c = *GuidString;
-
-		while (*c)
+		for (const TCHAR GuidChar : GuidString)
 		{
 			Value *= 36;
-			if (*c >= '0' && *c <= '9')
+			if (GuidChar >= '0' && GuidChar <= '9')
 			{
-				Value += *c - '0';
+				Value += GuidChar - '0';
 			}
-			else if (*c >= 'A' && *c <= 'Z')
+			else if (GuidChar >= 'A' && GuidChar <= 'Z')
 			{
-				Value += *c - 'A' + 10;
+				Value += GuidChar - 'A' + 10;
 			}
 			else
 			{
 				return false;
 			}
-			c++;
 		}
 
 		OutGuid = FGuid(
-			Value.GetQuadPart(3), 
-			Value.GetQuadPart(2), 
-			Value.GetQuadPart(1), 
+			Value.GetQuadPart(3),
+			Value.GetQuadPart(2),
+			Value.GetQuadPart(1),
 			Value.GetQuadPart(0)
 		);
 
 		return true;
 	}
 
-	FString NormalizedGuidString;
-
-	NormalizedGuidString.Empty(32);
+	TStringBuilder<32> NormalizedGuidString;
 
 	if (Format == EGuidFormats::Digits || Format == EGuidFormats::DigitsLower)
 	{
@@ -563,19 +580,20 @@ bool FGuid::ParseExact(const FString& GuidString, EGuidFormats Format, FGuid& Ou
 		NormalizedGuidString += GuidString.Mid(27, 8);
 	}
 
-	for (int32 Index = 0; Index < NormalizedGuidString.Len(); ++Index)
+	FStringView NormalizedGuidStringView = NormalizedGuidString.ToView();
+	for (const TCHAR GuidChar : NormalizedGuidStringView)
 	{
-		if (!FChar::IsHexDigit(NormalizedGuidString[Index]))
+		if (!FChar::IsHexDigit(GuidChar))
 		{
 			return false;
 		}
 	}
 
 	OutGuid = FGuid(
-		FParse::HexNumber(*NormalizedGuidString.Mid(0, 8)),
-		FParse::HexNumber(*NormalizedGuidString.Mid(8, 8)),
-		FParse::HexNumber(*NormalizedGuidString.Mid(16, 8)),
-		FParse::HexNumber(*NormalizedGuidString.Mid(24, 8))
+		FParse::HexNumber(NormalizedGuidStringView.Mid(0, 8)),
+		FParse::HexNumber(NormalizedGuidStringView.Mid(8, 8)),
+		FParse::HexNumber(NormalizedGuidStringView.Mid(16, 8)),
+		FParse::HexNumber(NormalizedGuidStringView.Mid(24, 8))
 	);
 
 	return true;
diff --git a/Engine/Source/Runtime/Core/Private/Misc/Parse.cpp b/Engine/Source/Runtime/Core/Private/Misc/Parse.cpp
--- a/Engine/Source/Runtime/Core/Private/Misc/Parse.cpp
+++ b/Engine/Source/Runtime/Core/Private/Misc/Parse.cpp
@@ -1310,27 +1310,27 @@ bool FParse::LineExtended(const TCHAR** Stream, FStringBuilderBase& Result, int3
 	return ParseLineExtended(Stream, Result, LinesConsumed, bExact);
 }
 
-uint32 FParse::HexNumber(const TCHAR* HexString)
+uint32 FParse::HexNumber(FStringView HexString)
 {
 	uint32 Ret = 0;
 
-	while (*HexString)
+	for (const TCHAR HexChar : HexString)
 	{
 		Ret *= 16;
-		Ret += FParse::HexDigit(*HexString++);
+		Ret += FParse::HexDigit(HexChar);
 	}
 
 	return Ret;
 }
 
-uint64 FParse::HexNumber64(const TCHAR* HexString)
+uint64 FParse::HexNumber64(FStringView HexString)
 {
 	uint64 Ret = 0;
 
-	while (*HexString)
+	for (const TCHAR HexChar : HexString)
 	{
 		Ret *= 16;
-		Ret += FParse::HexDigit(*HexString++);
+		Ret += FParse::HexDigit(HexChar);
 	}
 
 	return Ret;
diff --git a/Engine/Source/Runtime/Core/Public/Misc/Guid.h b/Engine/Source/Runtime/Core/Public/Misc/Guid.h
--- a/Engine/Source/Runtime/Core/Public/Misc/Guid.h
+++ b/Engine/Source/Runtime/Core/Public/Misc/Guid.h
@@ -393,7 +393,9 @@ public:
 	 * @return true if the string was converted successfully, false otherwise.
 	 * @see ParseExact, ToString
 	 */
+	static CORE_API bool Parse(const TCHAR* GuidString, FGuid& OutGuid);
 	static CORE_API bool Parse(const FString& GuidString, FGuid& OutGuid);
+	static CORE_API bool Parse(FStringView GuidString, FGuid& OutGuid);
 
 	/**
 	 * Converts a string with the specified format to a GUID.
@@ -404,7 +406,9 @@ public:
 	 * @return true if the string was converted successfully, false otherwise.
 	 * @see Parse, ToString
 	 */
+	static CORE_API bool ParseExact(const TCHAR* GuidString, EGuidFormats Format, FGuid& OutGuid);
 	static CORE_API bool ParseExact(const FString& GuidString, EGuidFormats Format, FGuid& OutGuid);
+	static CORE_API bool ParseExact(FStringView GuidString, EGuidFormats Format, FGuid& OutGuid);
 
 //private:
 public:
diff --git a/Engine/Source/Runtime/Core/Public/Misc/Parse.h b/Engine/Source/Runtime/Core/Public/Misc/Parse.h
--- a/Engine/Source/Runtime/Core/Public/Misc/Parse.h
+++ b/Engine/Source/Runtime/Core/Public/Misc/Parse.h
@@ -144,8 +144,8 @@ struct FParse
 	}
 
 	/** Parses a hexadecimal string value. */
-	static CORE_API uint32 HexNumber(const TCHAR* HexString);
-	static CORE_API uint64 HexNumber64(const TCHAR* HexString);
+	static CORE_API uint32 HexNumber(FStringView HexString);
+	static CORE_API uint64 HexNumber64(FStringView HexString);
 
 	/** Parses the scheme name from a URI */
 	static CORE_API bool SchemeNameFromURI(const TCHAR* InURI, FString& OutSchemeName);
