Add a two functions being called by FPak area, these both dont mutate the data and should be read only so just opening around them

diff --git a/Engine/Source/Runtime/Core/Private/AutoRTFM/StdLibSupport.cpp b/Engine/Source/Runtime/Core/Private/AutoRTFM/StdLibSupport.cpp
--- a/Engine/Source/Runtime/Core/Private/AutoRTFM/StdLibSupport.cpp
+++ b/Engine/Source/Runtime/Core/Private/AutoRTFM/StdLibSupport.cpp
@@ -220,6 +220,8 @@ UE_AUTORTFM_REGISTER_SELF_OVERLOADED_FUNCTION(const wchar_t* (const wchar_t*, wc
 UE_AUTORTFM_REGISTER_SELF_OVERLOADED_FUNCTION(wchar_t* (wchar_t*, wchar_t), wcschr);
 UE_AUTORTFM_REGISTER_SELF_OVERLOADED_FUNCTION(wchar_t* (wchar_t*, const wchar_t*), wcsstr);
 UE_AUTORTFM_REGISTER_SELF_OVERLOADED_FUNCTION(const wchar_t*(const wchar_t*, const wchar_t*), wcsstr);
+UE_AUTORTFM_REGISTER_SELF_OVERLOADED_FUNCTION(wchar_t*(wchar_t*, wchar_t), wcsrchr);
+UE_AUTORTFM_REGISTER_SELF_OVERLOADED_FUNCTION(const wchar_t*(const wchar_t*, wchar_t), wcsrchr);
 UE_AUTORTFM_REGISTER_SELF_FUNCTION(wcscmp);
 UE_AUTORTFM_REGISTER_SELF_FUNCTION(wcslen);
 
@@ -657,6 +659,8 @@ UE_AUTORTFM_REGISTER_SELF_FUNCTION(GetCurrentThreadId);
 UE_AUTORTFM_REGISTER_SELF_FUNCTION(TlsGetValue);
 UE_AUTORTFM_REGISTER_SELF_FUNCTION(GetLocalTime);
 
+UE_AUTORTFM_REGISTER_SELF_FUNCTION(GetFileAttributesW);
+
 BOOL RTFM_TlsSetValue(DWORD dwTlsIndex, LPVOID lpTlsValue)
 {
 	LPVOID CurrentValue = TlsGetValue(dwTlsIndex);
