Added a way to override the EAllowShrinking defaults.

diff --git a/Engine/Plugins/Experimental/AnimNext/Source/AnimNext/Internal/AnimNextPool.h b/Engine/Plugins/Experimental/AnimNext/Source/AnimNext/Internal/AnimNextPool.h
--- a/Engine/Plugins/Experimental/AnimNext/Source/AnimNext/Internal/AnimNextPool.h
+++ b/Engine/Plugins/Experimental/AnimNext/Source/AnimNext/Internal/AnimNextPool.h
@@ -86,7 +86,7 @@ public:
 		return Handle;
 	}
 
-	void Release(TPoolHandle<ElementType> InHandle, EAllowShrinking InAllowShrinking = EAllowShrinking::Yes)
+	void Release(TPoolHandle<ElementType> InHandle, EAllowShrinking InAllowShrinking = EAllowShrinking::Default)
 	{
 		if(IsValidHandle(InHandle))
 		{
diff --git a/Engine/Plugins/Experimental/AnimNext/Source/AnimNext/Public/TraitCore/TraitEventList.h b/Engine/Plugins/Experimental/AnimNext/Source/AnimNext/Public/TraitCore/TraitEventList.h
--- a/Engine/Plugins/Experimental/AnimNext/Source/AnimNext/Public/TraitCore/TraitEventList.h
+++ b/Engine/Plugins/Experimental/AnimNext/Source/AnimNext/Public/TraitCore/TraitEventList.h
@@ -25,7 +25,7 @@ namespace UE::AnimNext
 		int32 Num() const { return Events.Num(); }
 		bool IsEmpty() const { return Events.IsEmpty(); }
 
-		void SetNum(int32 NewNum, EAllowShrinking AllowShrinking = EAllowShrinking::Yes) { Events.SetNum(NewNum, AllowShrinking); }
+		void SetNum(int32 NewNum, EAllowShrinking AllowShrinking = EAllowShrinking::Default) { Events.SetNum(NewNum, AllowShrinking); }
 
 		void Append(const FTraitEventList& Source) { Events.Append(Source.Events); }
 
diff --git a/Engine/Plugins/Experimental/AnimNext/Source/AnimNext/Public/TransformArray.h b/Engine/Plugins/Experimental/AnimNext/Source/AnimNext/Public/TransformArray.h
--- a/Engine/Plugins/Experimental/AnimNext/Source/AnimNext/Public/TransformArray.h
+++ b/Engine/Plugins/Experimental/AnimNext/Source/AnimNext/Public/TransformArray.h
@@ -39,7 +39,7 @@ struct TTransformArrayAoS
 		Transforms.Reset(NumTransforms);
 	}
 
-	void SetNum(int32 NumTransforms, EAllowShrinking AllowShrinking = EAllowShrinking::Yes)
+	void SetNum(int32 NumTransforms, EAllowShrinking AllowShrinking = EAllowShrinking::Default)
 	{
 		Transforms.SetNum(NumTransforms, AllowShrinking);
 	}
@@ -49,7 +49,7 @@ struct TTransformArrayAoS
 		SetNum(NumTransforms, bAllowShrinking ? EAllowShrinking::Yes : EAllowShrinking::No);
 	}
 
-	void SetNumUninitialized(int32 NumTransforms, EAllowShrinking AllowShrinking = EAllowShrinking::Yes)
+	void SetNumUninitialized(int32 NumTransforms, EAllowShrinking AllowShrinking = EAllowShrinking::Default)
 	{
 		Transforms.SetNumUninitialized(NumTransforms, AllowShrinking);
 	}
@@ -335,7 +335,7 @@ struct TTransformArraySoA
 		AllocatedMemory.Reset(NumTransforms * TransformSize);
 	}
 
-	void SetNum(int32 NumTransforms, EAllowShrinking AllowShrinking = EAllowShrinking::Yes)
+	void SetNum(int32 NumTransforms, EAllowShrinking AllowShrinking = EAllowShrinking::Default)
 	{
 		constexpr int32 TransformSize = sizeof(FVector) + sizeof(FQuat) + sizeof(FVector);
 		AllocatedMemory.SetNum(NumTransforms * TransformSize, AllowShrinking);
@@ -348,7 +348,7 @@ struct TTransformArraySoA
 		SetNum(NumTransforms, bAllowShrinking ? EAllowShrinking::Yes : EAllowShrinking::No);
 	}
 
-	void SetNumUninitialized(int32 NumTransforms, EAllowShrinking AllowShrinking = EAllowShrinking::Yes)
+	void SetNumUninitialized(int32 NumTransforms, EAllowShrinking AllowShrinking = EAllowShrinking::Default)
 	{
 		constexpr int32 TransformSize = sizeof(FVector) + sizeof(FQuat) + sizeof(FVector);
 		AllocatedMemory.SetNumUninitialized(NumTransforms * TransformSize, AllowShrinking);
diff --git a/Engine/Plugins/Runtime/ReplicationGraph/Source/Public/ReplicationGraphTypes.h b/Engine/Plugins/Runtime/ReplicationGraph/Source/Public/ReplicationGraphTypes.h
--- a/Engine/Plugins/Runtime/ReplicationGraph/Source/Public/ReplicationGraphTypes.h
+++ b/Engine/Plugins/Runtime/ReplicationGraph/Source/Public/ReplicationGraphTypes.h
@@ -510,7 +510,7 @@ struct REPLICATIONGRAPH_API FActorRepListRefView
 	}
 
 	/** Removes the element quickly but changes the list order */
-	bool RemoveFast(const FActorRepListType& ElementToRemove, EAllowShrinking AllowShrink = EAllowShrinking::Yes)
+	bool RemoveFast(const FActorRepListType& ElementToRemove, EAllowShrinking AllowShrink = EAllowShrinking::Default)
 	{
 		return RepList.RemoveSingleSwap(ElementToRemove, AllowShrink) > 0;
 	}
diff --git a/Engine/Source/Runtime/Core/Public/Containers/AllowShrinking.h b/Engine/Source/Runtime/Core/Public/Containers/AllowShrinking.h
--- a/Engine/Source/Runtime/Core/Public/Containers/AllowShrinking.h
+++ b/Engine/Source/Runtime/Core/Public/Containers/AllowShrinking.h
@@ -8,7 +8,9 @@
 enum class EAllowShrinking : uint8
 {
 	No,
-	Yes
+	Yes,
+
+	Default = Yes
 };
 
 #define UE_ALLOWSHRINKING_BOOL_DEPRECATED(FunctionName) UE_DEPRECATED_FORENGINE(5.5, FunctionName " with a boolean bAllowShrinking has been deprecated - please use the EAllowShrinking enum instead")
diff --git a/Engine/Source/Runtime/Core/Public/Containers/Array.h b/Engine/Source/Runtime/Core/Public/Containers/Array.h
--- a/Engine/Source/Runtime/Core/Public/Containers/Array.h
+++ b/Engine/Source/Runtime/Core/Public/Containers/Array.h
@@ -881,7 +881,7 @@ public:
 	 * @param AllowShrinking If this call allows shrinking of the array during element remove.
 	 * @returns Popped element.
 	 */
-	ElementType Pop(EAllowShrinking AllowShrinking = EAllowShrinking::Yes)
+	ElementType Pop(EAllowShrinking AllowShrinking = EAllowShrinking::Default)
 	{
 		RangeCheck(0);
 		ElementType Result = MoveTempIfPossible(GetData()[ArrayNum - 1]);
@@ -1784,7 +1784,7 @@ public:
 	 * @param Index Location in array of the element to remove.
 	 * @param AllowShrinking (Optional) Tells if this call can shrink array if suitable after remove. Default is yes.
 	 */
-	void RemoveAt(SizeType Index, EAllowShrinking AllowShrinking = EAllowShrinking::Yes)
+	void RemoveAt(SizeType Index, EAllowShrinking AllowShrinking = EAllowShrinking::Default)
 	{
 		RangeCheck(Index);
 		RemoveAtImpl(Index);
@@ -1806,7 +1806,7 @@ public:
 		typename CountType
 		UE_REQUIRES(std::is_integral_v<CountType>)
 	>
-	FORCEINLINE void RemoveAt(SizeType Index, CountType Count, EAllowShrinking AllowShrinking = EAllowShrinking::Yes)
+	FORCEINLINE void RemoveAt(SizeType Index, CountType Count, EAllowShrinking AllowShrinking = EAllowShrinking::Default)
 	{
 		static_assert(!std::is_same_v<CountType, bool>, "TArray::RemoveAt: unexpected bool passed as the Count argument");
 		RangeCheck(Index, Count);
@@ -1877,7 +1877,7 @@ public:
 	 * @param AllowShrinking (Optional) Tells if this call can shrink array if
 	 *                        suitable after remove. Default is yes.
 	 */
-	FORCEINLINE void RemoveAtSwap(SizeType Index, EAllowShrinking AllowShrinking = EAllowShrinking::Yes)
+	FORCEINLINE void RemoveAtSwap(SizeType Index, EAllowShrinking AllowShrinking = EAllowShrinking::Default)
 	{
 		RangeCheck(Index);
 		RemoveAtSwapImpl(Index);
@@ -1903,7 +1903,7 @@ public:
 		typename CountType
 		UE_REQUIRES(std::is_integral_v<CountType>)
 	>
-	FORCEINLINE void RemoveAtSwap(SizeType Index, CountType Count, EAllowShrinking AllowShrinking = EAllowShrinking::Yes)
+	FORCEINLINE void RemoveAtSwap(SizeType Index, CountType Count, EAllowShrinking AllowShrinking = EAllowShrinking::Default)
 	{
 		static_assert(!std::is_same_v<CountType, bool>, "TArray::RemoveAtSwap: unexpected bool passed as the Count argument");
 		RangeCheck(Index, Count);
@@ -1983,7 +1983,7 @@ public:
 	 * @param NewNum New size of the array.
 	 * @param AllowShrinking Tell if this function can shrink the memory in-use if suitable.
 	 */
-	void SetNum(SizeType NewNum, EAllowShrinking AllowShrinking = EAllowShrinking::Yes)
+	void SetNum(SizeType NewNum, EAllowShrinking AllowShrinking = EAllowShrinking::Default)
 	{
 		if (NewNum > Num())
 		{
@@ -2014,7 +2014,7 @@ public:
 	 * @param NewNum New size of the array.
 	 * @param AllowShrinking Tell if this function can shrink the memory in-use if suitable.
 	 */
-	void SetNumZeroed(SizeType NewNum, EAllowShrinking AllowShrinking = EAllowShrinking::Yes)
+	void SetNumZeroed(SizeType NewNum, EAllowShrinking AllowShrinking = EAllowShrinking::Default)
 	{
 		if (NewNum > Num())
 		{
@@ -2042,7 +2042,7 @@ public:
 	 * @param NewNum New size of the array.
 	 * @param AllowShrinking Tell if this function can shrink the memory in-use if suitable.
 	 */
-	void SetNumUninitialized(SizeType NewNum, EAllowShrinking AllowShrinking = EAllowShrinking::Yes)
+	void SetNumUninitialized(SizeType NewNum, EAllowShrinking AllowShrinking = EAllowShrinking::Default)
 	{
 		if (NewNum > Num())
 		{
@@ -2776,7 +2776,7 @@ public:
 	 * @see Remove, RemoveSingle, RemoveSingleSwap, RemoveSwap
 	 */
 	template <class PREDICATE_CLASS>
-	SizeType RemoveAllSwap(const PREDICATE_CLASS& Predicate, EAllowShrinking AllowShrinking = EAllowShrinking::Yes)
+	SizeType RemoveAllSwap(const PREDICATE_CLASS& Predicate, EAllowShrinking AllowShrinking = EAllowShrinking::Default)
 	{
 		bool bRemoved = false;
 		const SizeType OriginalNum = ArrayNum;
@@ -2817,7 +2817,7 @@ public:
 	 * @returns The number of items removed. For RemoveSingleItem, this is always either 0 or 1.
 	 * @see Add, Insert, Remove, RemoveAll, RemoveAllSwap, RemoveSwap
 	 */
-	SizeType RemoveSingleSwap(const ElementType& Item, EAllowShrinking AllowShrinking = EAllowShrinking::Yes)
+	SizeType RemoveSingleSwap(const ElementType& Item, EAllowShrinking AllowShrinking = EAllowShrinking::Default)
 	{
 		SizeType Index = Find(Item);
 		if (Index == INDEX_NONE)
@@ -2849,7 +2849,7 @@ public:
 	 * @returns Number of elements removed.
 	 * @see Add, Insert, Remove, RemoveAll, RemoveAllSwap
 	 */
-	SizeType RemoveSwap(const ElementType& Item, EAllowShrinking AllowShrinking = EAllowShrinking::Yes)
+	SizeType RemoveSwap(const ElementType& Item, EAllowShrinking AllowShrinking = EAllowShrinking::Default)
 	{
 		CheckAddress(&Item);
 
@@ -3432,7 +3432,7 @@ public:
 	 *        The auto-dereferencing behavior does not occur with smart pointers.
 	 */
 	template <class PREDICATE_CLASS>
-	void HeapPop(ElementType& OutItem, const PREDICATE_CLASS& Predicate, EAllowShrinking AllowShrinking = EAllowShrinking::Yes)
+	void HeapPop(ElementType& OutItem, const PREDICATE_CLASS& Predicate, EAllowShrinking AllowShrinking = EAllowShrinking::Default)
 	{
 		OutItem = MoveTemp((*this)[0]);
 		RemoveAtSwap(0, 1, AllowShrinking);
@@ -3458,7 +3458,7 @@ public:
 	 *        Therefore, your array will be heapified by the values being pointed to, rather than the pointers' values.
 	 *        The auto-dereferencing behavior does not occur with smart pointers.
 	 */
-	void HeapPop(ElementType& OutItem, EAllowShrinking AllowShrinking = EAllowShrinking::Yes)
+	void HeapPop(ElementType& OutItem, EAllowShrinking AllowShrinking = EAllowShrinking::Default)
 	{
 		HeapPop(OutItem, TLess<ElementType>(), AllowShrinking);
 	}
@@ -3490,7 +3490,7 @@ public:
 	 *        The auto-dereferencing behavior does not occur with smart pointers.
 	 */
 	template <class PREDICATE_CLASS>
-	void HeapPopDiscard(const PREDICATE_CLASS& Predicate, EAllowShrinking AllowShrinking = EAllowShrinking::Yes)
+	void HeapPopDiscard(const PREDICATE_CLASS& Predicate, EAllowShrinking AllowShrinking = EAllowShrinking::Default)
 	{
 		RemoveAtSwap(0, 1, AllowShrinking);
 		TDereferenceWrapper< ElementType, PREDICATE_CLASS> PredicateWrapper(Predicate);
@@ -3513,7 +3513,7 @@ public:
 	 *        Therefore, your array will be heapified by the values being pointed to, rather than the pointers' values.
 	 *        The auto-dereferencing behavior does not occur with smart pointers.
 	 */
-	void HeapPopDiscard(EAllowShrinking AllowShrinking = EAllowShrinking::Yes)
+	void HeapPopDiscard(EAllowShrinking AllowShrinking = EAllowShrinking::Default)
 	{
 		HeapPopDiscard(TLess<ElementType>(), AllowShrinking);
 	}
@@ -3558,7 +3558,7 @@ public:
 	 *        The auto-dereferencing behavior does not occur with smart pointers.
 	 */
 	template <class PREDICATE_CLASS>
-	void HeapRemoveAt(SizeType Index, const PREDICATE_CLASS& Predicate, EAllowShrinking AllowShrinking = EAllowShrinking::Yes)
+	void HeapRemoveAt(SizeType Index, const PREDICATE_CLASS& Predicate, EAllowShrinking AllowShrinking = EAllowShrinking::Default)
 	{
 		RemoveAtSwap(Index, 1, AllowShrinking);
 
@@ -3584,7 +3584,7 @@ public:
 	 *        Therefore, your array will be heapified by the values being pointed to, rather than the pointers' values.
 	 *        The auto-dereferencing behavior does not occur with smart pointers.
 	 */
-	void HeapRemoveAt(SizeType Index, EAllowShrinking AllowShrinking = EAllowShrinking::Yes)
+	void HeapRemoveAt(SizeType Index, EAllowShrinking AllowShrinking = EAllowShrinking::Default)
 	{
 		HeapRemoveAt(Index, TLess< ElementType >(), AllowShrinking);
 	}
diff --git a/Engine/Source/Runtime/Core/Public/Containers/IndirectArray.h b/Engine/Source/Runtime/Core/Public/Containers/IndirectArray.h
--- a/Engine/Source/Runtime/Core/Public/Containers/IndirectArray.h
+++ b/Engine/Source/Runtime/Core/Public/Containers/IndirectArray.h
@@ -250,7 +250,7 @@ public:
 	 * @param AllowShrinking (Optional) Tells if this call can shrink array if
 	 *                        suitable after remove. Default is yes.
 	 */
-	void RemoveAt(int32 Index, EAllowShrinking AllowShrinking = EAllowShrinking::Yes)
+	void RemoveAt(int32 Index, EAllowShrinking AllowShrinking = EAllowShrinking::Default)
 	{
 		check(Index >= 0);
 		check(Index < Array.Num());
@@ -268,7 +268,7 @@ public:
 	 * @param AllowShrinking (Optional) Tells if this call can shrink array if
 	 *                        suitable after remove. Default is yes.
 	 */
-	void RemoveAt(int32 Index, int32 Count, EAllowShrinking AllowShrinking = EAllowShrinking::Yes)
+	void RemoveAt(int32 Index, int32 Count, EAllowShrinking AllowShrinking = EAllowShrinking::Default)
 	{
 		check(Index >= 0);
 		check(Index <= Array.Num());
@@ -298,7 +298,7 @@ public:
 	 * @param AllowShrinking (Optional) Tells if this call can shrink array if
 	 *                        suitable after remove. Default is yes.
 	 */
-	void RemoveAtSwap(int32 Index, EAllowShrinking AllowShrinking = EAllowShrinking::Yes)
+	void RemoveAtSwap(int32 Index, EAllowShrinking AllowShrinking = EAllowShrinking::Default)
 	{
 		check(Index >= 0);
 		check(Index < Array.Num());
@@ -319,7 +319,7 @@ public:
 	 * @param AllowShrinking (Optional) Tells if this call can shrink array if
 	 *                        suitable after remove. Default is yes.
 	 */
-	void RemoveAtSwap(int32 Index, int32 Count, EAllowShrinking AllowShrinking = EAllowShrinking::Yes)
+	void RemoveAtSwap(int32 Index, int32 Count, EAllowShrinking AllowShrinking = EAllowShrinking::Default)
 	{
 		check(Index >= 0);
 		check(Index <= Array.Num());
diff --git a/Engine/Source/Runtime/Core/Public/Containers/PagedArray.h b/Engine/Source/Runtime/Core/Public/Containers/PagedArray.h
--- a/Engine/Source/Runtime/Core/Public/Containers/PagedArray.h
+++ b/Engine/Source/Runtime/Core/Public/Containers/PagedArray.h
@@ -321,7 +321,7 @@ public:
 	 * Resizes array to the parameter number of elements.
 	 * The allow shrinking parameter indicates whether the container page allocation can be reduced if possible.
 	 */
-	void SetNum(SizeType NewNum, EAllowShrinking AllowShrinking = EAllowShrinking::Yes)
+	void SetNum(SizeType NewNum, EAllowShrinking AllowShrinking = EAllowShrinking::Default)
 	{
 		const SizeType RequiredPageCount = NumRequiredPages(NewNum);
 		if (NewNum > Num())
@@ -523,7 +523,7 @@ public:
 	/*
 	 * Removes the last element in the container.
 	 */
-	void Pop(EAllowShrinking AllowShrinking = EAllowShrinking::Yes)
+	void Pop(EAllowShrinking AllowShrinking = EAllowShrinking::Default)
 	{
 		CheckValidIndex(0);
 		const SizeType LastIndex = Num() - 1;
@@ -547,7 +547,7 @@ public:
 	 * to ensure the range is contiguous. This method provides efficient removal O(1) but it doesn't preserve the
 	 * insertion order.
 	 */
-	void RemoveAtSwap(SizeType Index, EAllowShrinking AllowShrinking = EAllowShrinking::Yes)
+	void RemoveAtSwap(SizeType Index, EAllowShrinking AllowShrinking = EAllowShrinking::Default)
 	{
 		CheckValidIndex(Index);
 		const SizeType TargetPageIndex = GetPageIndex(Index);
diff --git a/Engine/Source/Runtime/Core/Public/Containers/ScriptArray.h b/Engine/Source/Runtime/Core/Public/Containers/ScriptArray.h
--- a/Engine/Source/Runtime/Core/Public/Containers/ScriptArray.h
+++ b/Engine/Source/Runtime/Core/Public/Containers/ScriptArray.h
@@ -105,7 +105,7 @@ public:
 			ResizeTo(ArrayNum, NumBytesPerElement, AlignmentOfElement);
 		}
 	}
-	void SetNumUninitialized(int32 NewNum, int32 NumBytesPerElement, uint32 AlignmentOfElement, EAllowShrinking AllowShrinking = EAllowShrinking::Yes)
+	void SetNumUninitialized(int32 NewNum, int32 NumBytesPerElement, uint32 AlignmentOfElement, EAllowShrinking AllowShrinking = EAllowShrinking::Default)
 	{
 		checkSlow(NewNum >= 0);
 		int32 OldNum = Num();
@@ -188,7 +188,7 @@ public:
 		return ArrayMax - ArrayNum;
 	}
 
-	void Remove( int32 Index, int32 Count, int32 NumBytesPerElement, uint32 AlignmentOfElement, EAllowShrinking AllowShrinking = EAllowShrinking::Yes )
+	void Remove( int32 Index, int32 Count, int32 NumBytesPerElement, uint32 AlignmentOfElement, EAllowShrinking AllowShrinking = EAllowShrinking::Default )
 	{
 		if (Count)
 		{
diff --git a/Engine/Source/Runtime/Core/Public/Containers/UnrealString.h.inl b/Engine/Source/Runtime/Core/Public/Containers/UnrealString.h.inl
--- a/Engine/Source/Runtime/Core/Public/Containers/UnrealString.h.inl
+++ b/Engine/Source/Runtime/Core/Public/Containers/UnrealString.h.inl
@@ -474,7 +474,7 @@ public:
 	 * @param Index          The index of the character to remove.
 	 * @param AllowShrinking Whether or not to reallocate to shrink the storage after removal.
 	 */
-	CORE_API void RemoveAt(int32 Index, EAllowShrinking AllowShrinking = EAllowShrinking::Yes);
+	CORE_API void RemoveAt(int32 Index, EAllowShrinking AllowShrinking = EAllowShrinking::Default);
 
 	/**
 	 * Removes characters within the string.
@@ -483,7 +483,7 @@ public:
 	 * @param Count          The number of characters to remove.
 	 * @param AllowShrinking Whether or not to reallocate to shrink the storage after removal.
 	 */
-	CORE_API void RemoveAt(int32 Index, int32 Count, EAllowShrinking AllowShrinking = EAllowShrinking::Yes);
+	CORE_API void RemoveAt(int32 Index, int32 Count, EAllowShrinking AllowShrinking = EAllowShrinking::Default);
 	UE_ALLOWSHRINKING_BOOL_DEPRECATED("RemoveAt")
 	FORCEINLINE void RemoveAt(int32 Index, int32 Count, bool bAllowShrinking)
 	{
@@ -1142,7 +1142,7 @@ public:
 	}
 
 	/** Modifies the string such that it is now the left most given number of characters */
-	FORCEINLINE void LeftInline(int32 Count, EAllowShrinking AllowShrinking = EAllowShrinking::Yes)
+	FORCEINLINE void LeftInline(int32 Count, EAllowShrinking AllowShrinking = EAllowShrinking::Default)
 	{
 		const int32 Length = Len();
 		Count = FMath::Clamp(Count, 0, Length);
@@ -1168,7 +1168,7 @@ public:
 	}
 
 	/** Modifies the string such that it is now the left most characters chopping the given number of characters from the end */
-	FORCEINLINE void LeftChopInline(int32 Count, EAllowShrinking AllowShrinking = EAllowShrinking::Yes)
+	FORCEINLINE void LeftChopInline(int32 Count, EAllowShrinking AllowShrinking = EAllowShrinking::Default)
 	{
 		const int32 Length = Len();
 		RemoveAt(FMath::Clamp(Length-Count, 0, Length), Count, AllowShrinking);
@@ -1193,7 +1193,7 @@ public:
 	}
 
 	/** Modifies the string such that it is now the right most given number of characters */
-	FORCEINLINE void RightInline(int32 Count, EAllowShrinking AllowShrinking = EAllowShrinking::Yes)
+	FORCEINLINE void RightInline(int32 Count, EAllowShrinking AllowShrinking = EAllowShrinking::Default)
 	{
 		const int32 Length = Len();
 		RemoveAt(0, Length-FMath::Clamp(Count,0,Length), AllowShrinking);
@@ -1214,7 +1214,7 @@ public:
 	}
 
 	/** Modifies the string such that it is now the string to the right of the specified location, counting forward from the left (from the beginning of the word). */
-	FORCEINLINE void RightChopInline(int32 Count, EAllowShrinking AllowShrinking = EAllowShrinking::Yes)
+	FORCEINLINE void RightChopInline(int32 Count, EAllowShrinking AllowShrinking = EAllowShrinking::Default)
 	{
 		RemoveAt(0, Count, AllowShrinking);
 	}
@@ -1233,7 +1233,7 @@ public:
 	[[nodiscard]] FORCEINLINE UE_STRING_CLASS Mid(int32 Start) && { return ((UE_STRING_CLASS&&)*this).RightChop(Start); }
 
 	/** Modifies the string such that it is now the substring from Start position for Count characters. */
-	FORCEINLINE void MidInline(int32 Start, int32 Count = MAX_int32, EAllowShrinking AllowShrinking = EAllowShrinking::Yes)
+	FORCEINLINE void MidInline(int32 Start, int32 Count = MAX_int32, EAllowShrinking AllowShrinking = EAllowShrinking::Default)
 	{
 		if (Count != MAX_int32 && int64(Start) + Count < MAX_int32)
 		{
diff --git a/Engine/Source/Runtime/CoreUObject/Private/Serialization/AsyncLoading.h b/Engine/Source/Runtime/CoreUObject/Private/Serialization/AsyncLoading.h
--- a/Engine/Source/Runtime/CoreUObject/Private/Serialization/AsyncLoading.h
+++ b/Engine/Source/Runtime/CoreUObject/Private/Serialization/AsyncLoading.h
@@ -329,7 +329,7 @@ struct FAsyncLoadEventArgs
 class FImportOrImportIndexArray : public TArray<int32>
 {
 public:
-	void HeapPop(int32& OutItem, EAllowShrinking AllowShrinking = EAllowShrinking::Yes);
+	void HeapPop(int32& OutItem, EAllowShrinking AllowShrinking = EAllowShrinking::Default);
 	UE_ALLOWSHRINKING_BOOL_DEPRECATED("HeapPop")
 	FORCEINLINE void HeapPop(int32& OutItem, bool bAllowShrinking)
 	{
diff --git a/Engine/Source/Runtime/Engine/Classes/Animation/PreviewAssetAttachComponent.h b/Engine/Source/Runtime/Engine/Classes/Animation/PreviewAssetAttachComponent.h
--- a/Engine/Source/Runtime/Engine/Classes/Animation/PreviewAssetAttachComponent.h
+++ b/Engine/Source/Runtime/Engine/Classes/Animation/PreviewAssetAttachComponent.h
@@ -139,8 +139,8 @@ public:
 	/** 
 	 * RemoveAtSwap passthrough
 	 */
-	ENGINE_API void RemoveAtSwap( int32 Index, EAllowShrinking AllowShrinking = EAllowShrinking::Yes);
-	ENGINE_API void RemoveAtSwap( int32 Index, int32 Count, EAllowShrinking AllowShrinking = EAllowShrinking::Yes);
+	ENGINE_API void RemoveAtSwap( int32 Index, EAllowShrinking AllowShrinking = EAllowShrinking::Default);
+	ENGINE_API void RemoveAtSwap( int32 Index, int32 Count, EAllowShrinking AllowShrinking = EAllowShrinking::Default);
 	UE_ALLOWSHRINKING_BOOL_DEPRECATED("RemoveAtSwap")
 	FORCEINLINE void RemoveAtSwap(int32 Index, int32 Count, bool bAllowShrinking)
 	{
diff --git a/Engine/Source/Runtime/Engine/Private/Animation/PreviewAssetAttachComponent.cpp b/Engine/Source/Runtime/Engine/Private/Animation/PreviewAssetAttachComponent.cpp
--- a/Engine/Source/Runtime/Engine/Private/Animation/PreviewAssetAttachComponent.cpp
+++ b/Engine/Source/Runtime/Engine/Private/Animation/PreviewAssetAttachComponent.cpp
@@ -84,12 +84,12 @@ TIterator FPreviewAssetAttachContainer::CreateIterator()
 	return AttachedObjects.CreateIterator();
 }
 
-void FPreviewAssetAttachContainer::RemoveAtSwap( int32 Index, EAllowShrinking AllowShrinking /*=EAllowShrinking::Yes */ )
+void FPreviewAssetAttachContainer::RemoveAtSwap( int32 Index, EAllowShrinking AllowShrinking /*=EAllowShrinking::Default */ )
 {
 	AttachedObjects.RemoveAtSwap( Index, AllowShrinking );
 }
 
-void FPreviewAssetAttachContainer::RemoveAtSwap( int32 Index, int32 Count, EAllowShrinking AllowShrinking /*=EAllowShrinking::Yes */ )
+void FPreviewAssetAttachContainer::RemoveAtSwap( int32 Index, int32 Count, EAllowShrinking AllowShrinking /*=EAllowShrinking::Default */ )
 {
 	AttachedObjects.RemoveAtSwap( Index, Count, AllowShrinking );
 }
diff --git a/Engine/Source/Runtime/Experimental/Chaos/Public/RewindData.h b/Engine/Source/Runtime/Experimental/Chaos/Public/RewindData.h
--- a/Engine/Source/Runtime/Experimental/Chaos/Public/RewindData.h
+++ b/Engine/Source/Runtime/Experimental/Chaos/Public/RewindData.h
@@ -127,7 +127,7 @@ struct FBaseRewindHistory
 	virtual const bool HasDataInHistory() const { return false; }
 
 	/** Resize the history */
-	virtual void ResizeDataHistory(const int32 FrameCount, const EAllowShrinking AllowShrinking = EAllowShrinking::Yes) {}
+	virtual void ResizeDataHistory(const int32 FrameCount, const EAllowShrinking AllowShrinking = EAllowShrinking::Default) {}
 
 	/** Perform a fast reset, marking the data history as reset but not clearing the data or resetting collections */
 	virtual void ResetFast() {}
@@ -386,7 +386,7 @@ public :
 	}
 
 	/** Resize the history */
-	FORCEINLINE void ResizeDataHistory(const int32 FrameCount, const EAllowShrinking AllowShrinking = EAllowShrinking::Yes) override
+	FORCEINLINE void ResizeDataHistory(const int32 FrameCount, const EAllowShrinking AllowShrinking = EAllowShrinking::Default) override
 	{
 		if (FrameCount > 0 && NumFrames != FrameCount)
 		{
@@ -1537,7 +1537,7 @@ public:
 		return Managers[Frame].DeltaTime;
 	}
 
-	void RemoveObject(const FGeometryParticleHandle* Particle, const EAllowShrinking AllowShrinking=EAllowShrinking::Yes)
+	void RemoveObject(const FGeometryParticleHandle* Particle, const EAllowShrinking AllowShrinking=EAllowShrinking::Default)
 	{
 		DirtyParticles.Remove(Particle, AllowShrinking);
 	}
@@ -1547,7 +1547,7 @@ public:
 		RemoveObject(Particle, bAllowShrinking ? EAllowShrinking::Yes : EAllowShrinking::No);
 	}
 
-	void RemoveObject(const FPBDJointConstraintHandle* Joint, const EAllowShrinking AllowShrinking = EAllowShrinking::Yes)
+	void RemoveObject(const FPBDJointConstraintHandle* Joint, const EAllowShrinking AllowShrinking = EAllowShrinking::Default)
 	{
 		DirtyJoints.Remove(Joint, AllowShrinking);
 	}
diff --git a/Engine/Source/Runtime/GeometryCore/Public/Spatial/DenseGrid2.h b/Engine/Source/Runtime/GeometryCore/Public/Spatial/DenseGrid2.h
--- a/Engine/Source/Runtime/GeometryCore/Public/Spatial/DenseGrid2.h
+++ b/Engine/Source/Runtime/GeometryCore/Public/Spatial/DenseGrid2.h
@@ -63,7 +63,7 @@ public:
 		return FVector2i((int32)DimensionX, (int32)DimensionY);
 	}
 
-	void Resize(int32 DimX, int32 DimY, EAllowShrinking AllowShrinking = EAllowShrinking::Yes)
+	void Resize(int32 DimX, int32 DimY, EAllowShrinking AllowShrinking = EAllowShrinking::Default)
 	{
 		DimensionX = (int64)DimX;
 		DimensionY = (int64)DimY;
diff --git a/Engine/Source/Runtime/GeometryCore/Public/Spatial/DenseGrid3.h b/Engine/Source/Runtime/GeometryCore/Public/Spatial/DenseGrid3.h
--- a/Engine/Source/Runtime/GeometryCore/Public/Spatial/DenseGrid3.h
+++ b/Engine/Source/Runtime/GeometryCore/Public/Spatial/DenseGrid3.h
@@ -58,7 +58,7 @@ public:
 		return Dimensions;
 	}
 
-	void Resize(int DimX, int DimY, int DimZ, EAllowShrinking AllowShrinking = EAllowShrinking::Yes)
+	void Resize(int DimX, int DimY, int DimZ, EAllowShrinking AllowShrinking = EAllowShrinking::Default)
 	{
 		check((int64)DimX * (int64)DimY * (int64)DimZ < INT_MAX);
 		Buffer.SetNumUninitialized(DimX * DimY * DimZ, AllowShrinking);
diff --git a/Engine/Source/Runtime/RenderCore/Public/ShaderSource.h b/Engine/Source/Runtime/RenderCore/Public/ShaderSource.h
--- a/Engine/Source/Runtime/RenderCore/Public/ShaderSource.h
+++ b/Engine/Source/Runtime/RenderCore/Public/ShaderSource.h
@@ -24,7 +24,7 @@
 class FShaderSource
 {
 private:
-	inline void SetLen(int32 Num, EAllowShrinking AllowShrinking = EAllowShrinking::Yes)
+	inline void SetLen(int32 Num, EAllowShrinking AllowShrinking = EAllowShrinking::Default)
 	{
 		Source.SetNumUninitialized(Num + ShaderSourceSimdPadding, AllowShrinking);
 		FMemory::Memzero(Source.GetData() + Num, sizeof(CharType) * ShaderSourceSimdPadding);
@@ -95,7 +95,7 @@ public:
 	 * @param Num the desired allocation size (padding bytes will be added on top of this)
 	 * @param AllowShrinking whether to reallocate or keep the existing larger size allocation
 	 */
-	inline void ShrinkToLen(int32 Num, EAllowShrinking AllowShrinking = EAllowShrinking::Yes)
+	inline void ShrinkToLen(int32 Num, EAllowShrinking AllowShrinking = EAllowShrinking::Default)
 	{
 		checkf(Num <= Len(), TEXT("Trying to shrink to %d characters but existing allocation is smaller (%d characters)"), Num, Len());
 		SetLen(Num, AllowShrinking);
diff --git a/Engine/Source/Runtime/SlateCore/Public/Containers/ObservableArray.h b/Engine/Source/Runtime/SlateCore/Public/Containers/ObservableArray.h
--- a/Engine/Source/Runtime/SlateCore/Public/Containers/ObservableArray.h
+++ b/Engine/Source/Runtime/SlateCore/Public/Containers/ObservableArray.h
@@ -299,7 +299,7 @@ public:
 		}
 	}
 
-	SizeType RemoveSingle(const ElementType& Item, EAllowShrinking AllowShrinking = EAllowShrinking::Yes)
+	SizeType RemoveSingle(const ElementType& Item, EAllowShrinking AllowShrinking = EAllowShrinking::Default)
 	{
 		SizeType Index = Array.Find(Item);
 		if (Index == INDEX_NONE)
@@ -315,7 +315,7 @@ public:
 		return RemoveSingle(Item, bAllowShrinking ? EAllowShrinking::Yes : EAllowShrinking::No);
 	}
 
-	SizeType RemoveSingleSwap(const ElementType& Item, EAllowShrinking AllowShrinking = EAllowShrinking::Yes)
+	SizeType RemoveSingleSwap(const ElementType& Item, EAllowShrinking AllowShrinking = EAllowShrinking::Default)
 	{
 		SizeType Index = Array.Find(Item);
 		if (Index == INDEX_NONE)
@@ -331,7 +331,7 @@ public:
 		return RemoveSingleSwap(Item, bAllowShrinking ? EAllowShrinking::Yes : EAllowShrinking::No);
 	}
 
-	void RemoveAt(SizeType Index, EAllowShrinking AllowShrinking = EAllowShrinking::Yes)
+	void RemoveAt(SizeType Index, EAllowShrinking AllowShrinking = EAllowShrinking::Default)
 	{
 		check((Index >= 0) & (Index < Num()));
 
@@ -344,7 +344,7 @@ public:
 		typename CountType
 		UE_REQUIRES(std::is_integral_v<CountType>)
 	>
-	void RemoveAt(SizeType Index, CountType NumToRemove, EAllowShrinking AllowShrinking = EAllowShrinking::Yes)
+	void RemoveAt(SizeType Index, CountType NumToRemove, EAllowShrinking AllowShrinking = EAllowShrinking::Default)
 	{
 		static_assert(!std::is_same_v<CountType, bool>, "TObservableArray::RemoveAt: unexpected bool passed as the Count argument");
 		check((NumToRemove > 0) & (Index >= 0) & (Index + NumToRemove <= Num()));
@@ -378,7 +378,7 @@ public:
 		RemoveAt(Index, NumToRemove, bAllowShrinking ? EAllowShrinking::Yes : EAllowShrinking::No);
 	}
 
-	void RemoveAtSwap(SizeType Index, EAllowShrinking AllowShrinking = EAllowShrinking::Yes)
+	void RemoveAtSwap(SizeType Index, EAllowShrinking AllowShrinking = EAllowShrinking::Default)
 	{
 		SizeType PreviousNum = Array.Num();
 		SizeType SwapAmount = FPlatformMath::Min(1, PreviousNum - (Index + 1));
@@ -392,7 +392,7 @@ public:
 		typename CountType
 		UE_REQUIRES(std::is_integral_v<CountType>)
 	>
-	void RemoveAtSwap(SizeType Index, CountType NumToRemove, EAllowShrinking AllowShrinking = EAllowShrinking::Yes)
+	void RemoveAtSwap(SizeType Index, CountType NumToRemove, EAllowShrinking AllowShrinking = EAllowShrinking::Default)
 	{
 		static_assert(!std::is_same_v<CountType, bool>, "TObservableArray::RemoveAtSwap: unexpected bool passed as the Count argument");
 		check((NumToRemove > 0) & (Index >= 0) & (Index + NumToRemove <= Num()));
