Added UE_EXPERIMENTAL macro for performing "forward-deprecation", i.e. for marking public APIs from official modules as "experimental". Using the API will result in a compilation warning for the API user (except for engine code), that states "This API is experimental. As such, there is no guarantee that it won't change or be removed in future releases. Please use at your own risk."
It can be useful for fresh new features from official modules that are not quite ready for prime-time at the moment of the engine release. Marking them as experimental saves the developer the trouble of having to deprecate them in a future engine version, which is likely to happen as new APIs are rarely appropriate the first time around.
Similar to UE_DEPRECATED, UE_EXPERIMENTAL compilation warnings can be suppressed by adorning the calling code with PRAGMA_DISABLE_EXPERIMENTAL_WARNINGS / PRAGMA_ENABLE_EXPERIMENTAL_WARNINGS.

The macro follows closely the behavior of UE_INTERNAL :
* It is using the [[deprecated]] c++ attibute internally to report the compilation warning
* It is defined if UE_VALIDATE_EXPERIMENTAL_API is defined to 1
* UE_VALIDATE_EXPERIMENTAL_API is defined to 1 if a the module rule bValidateExperimentalApi is true
* By default, bValidateExperimentalApi is set to false for all engine modules (see : bTreatAsEngineModule), but can be forced to true (can be useful to force a given engine module to report experimental API usage as failures) or false (can be used to force a non-engine module not to report experimental API usage as failures)
* Similar to UE_DEPRECATED, UE_EXPERIMENTAL is recognized by UHT and can therefore be used on reflected code (UPROPERTY, UFUNCTION...). It also requires a version number (the engine version at which the API was introduced, which will help later on, when moving APIs out of experimental or removing them altogether) and an optional message.

Misc:
* Updated the various UE_DEPRECATED examples to include a proper deprecation message since it is strongly recommended for the user to provide it
* Added/expanded documentation for UE_INTERNAL and UE_DEPRECATED (e.g. the UE_DEPRECATED(all, ...) syntax for deprecating reflected code : indicates that the deprecation grace period is now over but that this property/function will remain deprecated for all time for data backward compatibility reasons)
* Updated the warning that comes with UE_INTERNAL to better reflect the intention : code that 'may' (rather than 'will') change in a future release
* Added PRAGMA_DISABLE_INTERNAL_WARNINGS / PRAGMA_ENABLE_INTERNAL_WARNINGS to prevent internal warnings on a case-by-case basis

#rnx

diff --git a/Engine/Source/Programs/Shared/EpicGames.UHT/Tokenizer/UhtTokenReaderGeneralExtensions.cs b/Engine/Source/Programs/Shared/EpicGames.UHT/Tokenizer/UhtTokenReaderGeneralExtensions.cs
--- a/Engine/Source/Programs/Shared/EpicGames.UHT/Tokenizer/UhtTokenReaderGeneralExtensions.cs
+++ b/Engine/Source/Programs/Shared/EpicGames.UHT/Tokenizer/UhtTokenReaderGeneralExtensions.cs
@@ -238,6 +238,17 @@ namespace EpicGames.UHT.Tokenizer
 						.Require(')', "deprecation macro");
 					attributeAction?.Invoke("deprecated");
 				}
+				else if (token.IsIdentifier("UE_EXPERIMENTAL"))
+				{
+					tokenReader.ConsumeToken();
+					tokenReader
+						.Require('(', "experimental macro")
+						.RequireConstFloat()
+						.Require(',', "experimental macro")
+						.RequireConstString("message in experimental macro")
+						.Require(')', "experimental macro");
+					attributeAction?.Invoke("experimental");
+				}
 				else if (token.IsIdentifier("UE_INTERNAL"))
 				{
 					tokenReader.ConsumeToken();
diff --git a/Engine/Source/Programs/UnrealBuildTool/Configuration/ModuleRules.cs b/Engine/Source/Programs/UnrealBuildTool/Configuration/ModuleRules.cs
--- a/Engine/Source/Programs/UnrealBuildTool/Configuration/ModuleRules.cs
+++ b/Engine/Source/Programs/UnrealBuildTool/Configuration/ModuleRules.cs
@@ -828,6 +828,16 @@ namespace UnrealBuildTool
 		}
 		private bool? bValidateInternalApiPrivate;
 
+		/// <summary>
+		/// Emits deprecated warnings\errors for experimental API usage for non-engine modules 
+		/// </summary>
+		public bool bValidateExperimentalApi
+		{
+			get => bValidateExperimentalApiPrivate ?? !bTreatAsEngineModule;
+			set => bValidateExperimentalApiPrivate = value;
+		}
+		private bool? bValidateExperimentalApiPrivate;
+
 		/// <summary>
 		/// Which engine version's build settings to use by default. 
 		/// </summary>
diff --git a/Engine/Source/Programs/UnrealBuildTool/Configuration/UEBuildModuleCPP.cs b/Engine/Source/Programs/UnrealBuildTool/Configuration/UEBuildModuleCPP.cs
--- a/Engine/Source/Programs/UnrealBuildTool/Configuration/UEBuildModuleCPP.cs
+++ b/Engine/Source/Programs/UnrealBuildTool/Configuration/UEBuildModuleCPP.cs
@@ -381,7 +381,7 @@ namespace UnrealBuildTool
 					{
 						Logger.LogInformation("Dependent modules '{OtherModules}'", String.Join(' ', ReferenceStackParentModules.OrderBy(x => x.Name).Select(x => x.Name)));
 					}
-					throw new BuildLogEventException("Missing precompiled manifest for '{Name}', '{Manifest}" , Name, PrecompiledManifestLocation);
+					throw new BuildLogEventException("Missing precompiled manifest for '{Name}', '{Manifest}", Name, PrecompiledManifestLocation);
 				}
 
 				PrecompiledManifest Manifest = PrecompiledManifest.Read(PrecompiledManifestLocation);
@@ -1187,6 +1187,7 @@ namespace UnrealBuildTool
 				if (Definition.Contains("UE_IS_ENGINE_MODULE", StringComparison.Ordinal) ||
 					Definition.Contains("UE_VALIDATE_FORMAT_STRINGS", StringComparison.Ordinal) ||
 					Definition.Contains("UE_VALIDATE_INTERNAL_API", StringComparison.Ordinal) ||
+					Definition.Contains("UE_VALIDATE_EXPERIMENTAL_API", StringComparison.Ordinal) ||
 					Definition.Contains("UE_DEPRECATED_FORGAME", StringComparison.Ordinal) ||
 					Definition.Contains("UE_DEPRECATED_FORENGINE", StringComparison.Ordinal) ||
 					Definition.Contains("UE_ENABLE_INCLUDE_ORDER_DEPRECATED_IN_", StringComparison.Ordinal))
@@ -1290,6 +1291,14 @@ namespace UnrealBuildTool
 					Definitions.Add($"UE_VALIDATE_INTERNAL_API={(ModuleCompileEnvironment.bValidateInternalApi ? "1" : "0")}");
 				}
 
+				// Modify definitions if we need to create a new shared pch for validating experimental api usage
+				if (ModuleCompileEnvironment.bValidateExperimentalApi != Template.BaseCompileEnvironment.bValidateExperimentalApi)
+				{
+					Definitions = new List<string>(Definitions);
+					Definitions.RemoveAll(x => x.Contains("UE_VALIDATE_EXPERIMENTAL_API", StringComparison.Ordinal));
+					Definitions.Add($"UE_VALIDATE_EXPERIMENTAL_API={(ModuleCompileEnvironment.bValidateExperimentalApi ? "1" : "0")}");
+				}
+
 				// Create a suffix to distinguish this shared PCH variant from any others.
 				string Variant = GetSuffixForSharedPCH(ModuleCompileEnvironment, Template.BaseCompileEnvironment);
 
@@ -1395,14 +1404,17 @@ namespace UnrealBuildTool
 			{
 				return false;
 			}
+
 			if (ModuleCompileEnvironment.bUseRTTI != CompileEnvironment.bUseRTTI)
 			{
 				return false;
 			}
+
 			if (ModuleCompileEnvironment.bEnableExceptions != CompileEnvironment.bEnableExceptions)
 			{
 				return false;
 			}
+
 			if (ModuleCompileEnvironment.CppStandard != CompileEnvironment.CppStandard)
 			{
 				return false;
@@ -1423,6 +1435,11 @@ namespace UnrealBuildTool
 				return false;
 			}
 
+			if (ModuleCompileEnvironment.bValidateExperimentalApi != CompileEnvironment.bValidateExperimentalApi)
+			{
+				return false;
+			}
+
 			if (ModuleCompileEnvironment.bEnableAutoRTFMInstrumentation != CompileEnvironment.bEnableAutoRTFMInstrumentation)
 			{
 				return false;
@@ -1519,6 +1536,18 @@ namespace UnrealBuildTool
 				}
 			}
 
+			if (CompileEnvironment.bValidateExperimentalApi != BaseCompileEnvironment.bValidateExperimentalApi)
+			{
+				if (CompileEnvironment.bValidateExperimentalApi)
+				{
+					Variant += ".ValExpApi";
+				}
+				else
+				{
+					Variant += ".NoValExpApi";
+				}
+			}
+
 			if (CompileEnvironment.bDeterministic != BaseCompileEnvironment.bDeterministic)
 			{
 				if (CompileEnvironment.bDeterministic)
@@ -1572,6 +1601,7 @@ namespace UnrealBuildTool
 			CompileEnvironment.bAutoRTFMVerify = ModuleCompileEnvironment.bAutoRTFMVerify;
 			CompileEnvironment.bValidateFormatStrings = ModuleCompileEnvironment.bValidateFormatStrings;
 			CompileEnvironment.bValidateInternalApi = ModuleCompileEnvironment.bValidateInternalApi;
+			CompileEnvironment.bValidateExperimentalApi = ModuleCompileEnvironment.bValidateExperimentalApi;
 		}
 
 		/// <summary>
@@ -2164,6 +2194,7 @@ namespace UnrealBuildTool
 			Result.DeterministicWarningLevel = Rules.DeterministicWarningLevel;
 			Result.bValidateFormatStrings = Rules.bValidateFormatStrings;
 			Result.bValidateInternalApi = Rules.bValidateInternalApi;
+			Result.bValidateExperimentalApi = Rules.bValidateExperimentalApi;
 			Result.bUseAutoRTFMCompiler = Target.bUseAutoRTFMCompiler;
 
 			CompileEnvironmentDebugInfoSettings(Target, Result);
@@ -2281,6 +2312,7 @@ namespace UnrealBuildTool
 
 			Result.Definitions.Add($"UE_VALIDATE_FORMAT_STRINGS={(Rules.bValidateFormatStrings ? "1" : "0")}");
 			Result.Definitions.Add($"UE_VALIDATE_INTERNAL_API={(Rules.bValidateInternalApi ? "1" : "0")}");
+			Result.Definitions.Add($"UE_VALIDATE_EXPERIMENTAL_API={(Rules.bValidateExperimentalApi ? "1" : "0")}");
 
 			Result.Definitions.AddRange(EngineIncludeOrderHelper.GetDeprecationDefines(Rules.IncludeOrderVersion));
 
@@ -2349,6 +2381,7 @@ namespace UnrealBuildTool
 			CompileEnvironment.bTreatAsEngineModule = Rules.bTreatAsEngineModule;
 			CompileEnvironment.bValidateFormatStrings = Rules.bValidateFormatStrings;
 			CompileEnvironment.bValidateInternalApi = Rules.bValidateInternalApi;
+			CompileEnvironment.bValidateExperimentalApi = Rules.bValidateExperimentalApi;
 
 			// Override compile environment
 			CompileEnvironment.bIsBuildingDLL = !Target.ShouldCompileMonolithic();
@@ -2365,6 +2398,7 @@ namespace UnrealBuildTool
 
 			CompileEnvironment.Definitions.Add($"UE_VALIDATE_FORMAT_STRINGS={(Rules.bValidateFormatStrings ? "1" : "0")}");
 			CompileEnvironment.Definitions.Add($"UE_VALIDATE_INTERNAL_API={(Rules.bValidateInternalApi ? "1" : "0")}");
+			CompileEnvironment.Definitions.Add($"UE_VALIDATE_EXPERIMENTAL_API={(Rules.bValidateExperimentalApi ? "1" : "0")}");
 
 			CompileEnvironment.Definitions.AddRange(EngineIncludeOrderHelper.GetDeprecationDefines(Rules.IncludeOrderVersion));
 
diff --git a/Engine/Source/Programs/UnrealBuildTool/System/CppCompileEnvironment.cs b/Engine/Source/Programs/UnrealBuildTool/System/CppCompileEnvironment.cs
--- a/Engine/Source/Programs/UnrealBuildTool/System/CppCompileEnvironment.cs
+++ b/Engine/Source/Programs/UnrealBuildTool/System/CppCompileEnvironment.cs
@@ -725,6 +725,11 @@ namespace UnrealBuildTool
 		/// </summary>
 		public bool bValidateInternalApi = false;
 
+		/// <summary>
+		/// Emits deprecated warnings\errors for experimental API usage for non-engine modules.
+		/// </summary>
+		public bool bValidateExperimentalApi = false;
+
 		/// <summary>
 		/// Directory where to put crash report files for platforms that support it
 		/// </summary>
@@ -848,6 +853,7 @@ namespace UnrealBuildTool
 			CrashDiagnosticDirectory = Other.CrashDiagnosticDirectory;
 			bValidateFormatStrings = Other.bValidateFormatStrings;
 			bValidateInternalApi = Other.bValidateInternalApi;
+			bValidateExperimentalApi = Other.bValidateExperimentalApi;
 		}
 
 		public CppCompileEnvironment(CppCompileEnvironment Other, UnrealArch OverrideArchitecture)
diff --git a/Engine/Source/Runtime/Core/Public/GenericPlatform/GenericPlatformCompilerPreSetup.h b/Engine/Source/Runtime/Core/Public/GenericPlatform/GenericPlatformCompilerPreSetup.h
--- a/Engine/Source/Runtime/Core/Public/GenericPlatform/GenericPlatformCompilerPreSetup.h
+++ b/Engine/Source/Runtime/Core/Public/GenericPlatform/GenericPlatformCompilerPreSetup.h
@@ -32,9 +32,9 @@
 
 #ifndef UE_DEPRECATED_MACRO
 	#if UE_WARNINGS_AS_ERRORS
-		#define UE_DEPRECATED_MACRO(Version, Message) EMIT_CUSTOM_ERROR(Message " Please update your code to the new API before upgrading to the next release, otherwise your project will no longer compile.")
+		#define UE_DEPRECATED_MACRO(Version, Message) EMIT_CUSTOM_ERROR(Message " - Please update your code to the new API before upgrading to the next release, otherwise your project will no longer compile.")
 	#else
-		#define UE_DEPRECATED_MACRO(Version, Message) EMIT_CUSTOM_WARNING(Message " Please update your code to the new API before upgrading to the next release, otherwise your project will no longer compile.")
+		#define UE_DEPRECATED_MACRO(Version, Message) EMIT_CUSTOM_WARNING(Message " - Please update your code to the new API before upgrading to the next release, otherwise your project will no longer compile.")
 	#endif
 #endif
 
diff --git a/Engine/Source/Runtime/Core/Public/Misc/CoreMiscDefines.h b/Engine/Source/Runtime/Core/Public/Misc/CoreMiscDefines.h
--- a/Engine/Source/Runtime/Core/Public/Misc/CoreMiscDefines.h
+++ b/Engine/Source/Runtime/Core/Public/Misc/CoreMiscDefines.h
@@ -168,6 +168,18 @@ enum EPerElement {PerElement};
 /** Thread-safe call once helper for void functions, similar to std::call_once without the std::once_flag */
 #define UE_CALL_ONCE(Func, ...) static int32 ANONYMOUS_VARIABLE(ThreadSafeOnce) = ((Func)(__VA_ARGS__), 1)
 
+#if defined (__INTELLISENSE__)
+
+	#define UE_DEPRECATED(Version, Message)
+	#define UE_INTERNAL
+	#define PRAGMA_DISABLE_INTERNAL_WARNINGS 
+	#define PRAGMA_ENABLE_INTERNAL_WARNINGS 
+	#define UE_EXPERIMENTAL(Version, Message)
+	#define PRAGMA_DISABLE_EXPERIMENTAL_WARNINGS 
+	#define PRAGMA_ENABLE_EXPERIMENTAL_WARNINGS 
+
+#else // defined (__INTELLISENSE__)
+	
 /**
  * Macro for marking up deprecated code, functions and types.
  *
@@ -181,18 +193,18 @@ enum EPerElement {PerElement};
  *
  * Sample usage (note the slightly different syntax for classes and structures):
  *
- *		UE_DEPRECATED(5.xx, "Message")
+ *		UE_DEPRECATED(5.xx, "MyFunction has been deprecated, please use MyNewFunction instead.")
  *		void MyFunction();
  *
- *		UE_DEPRECATED(5.xx, "Message")
- *		typedef FThing MyType;
+ *		UE_DEPRECATED(5.xx, "FMyType has been deprecated, please use MyNewType instead.")
+ *		typedef FThing FMyType;
  *
- *		using MyAlias UE_DEPRECATED(5.xx, "Message") = FThing;
+ *		using FMyAlias UE_DEPRECATED(5.xx, "FMyAlias has been deprecated, please use MyNewAlias instead.") = FThing;
  *
- *		UE_DEPRECATED(5.xx, "Message")
+ *		UE_DEPRECATED(5.xx, "MyVariable has been deprecated and serves no purpose anymore")
  *		int32 MyVariable;
  *
- *		namespace UE_DEPRECATED(5.xx, "Message") MyNamespace
+ *		namespace UE_DEPRECATED(5.xx, "MyNamespace has been renamed MyNewNamespace, please use this instead.") MyNamespace
  *		{
  *		}
  *
@@ -201,27 +213,27 @@ enum EPerElement {PerElement};
  *		deprecated, then declare the type with the visibility macro.  Note that macros like
  *		USTRUCT must immediately precede the the type declaration, not the forward declaration.
  *
- *		struct UE_DEPRECATED(5.xx, "Message") FMyStruct;
+ *		struct UE_DEPRECATED(5.xx, "FMyStruct has been deprecated and shall not be used anymore.") FMyStruct;
  *		USTRUCT()
  *		struct MODULE_API FMyStruct
  *		{
  *		};
  *
- *		class UE_DEPRECATED(5.xx, "Message") FMyClass;
+ *		class UE_DEPRECATED(5.xx, "FMyClass has been deprecated and shall not be used anymore.") FMyClass;
  *		class MODULE_API FMyClass
  *		{
  *		};
  *
- *		enum class UE_DEPRECATED(5.xx, "Message") EMyEnumeration
+ *		enum class UE_DEPRECATED(5.xx, "EMyEnumeration has been deprecated, use EMyEnumeration instead.") EMyEnumeration
  *		{
  *			Zero = 0,
- *			One UE_DEPRECATED(5.xx, "Message") = 1,
+ *			One UE_DEPRECATED(5.xx, "EMyEnumeration::One is not a valid enum entry anymore, please use another one.") = 1,
  *			Two = 2
  *		};
  *
  *		Unfortunately, VC++ will complain about using member functions and fields from deprecated
  *		class/structs even for class/struct implementation e.g.:
- *		class UE_DEPRECATED(5.xx, "") DeprecatedClass
+ *		class UE_DEPRECATED(5.xx, "DeprecatedClass has been deprecated and shall not be used anymore.") DeprecatedClass
  *		{
  *		public:
  *			DeprecatedClass() {}
@@ -244,7 +256,7 @@ enum EPerElement {PerElement};
  *			float MyFloat;
  *		};
  *
- *		class UE_DEPRECATED(5.xx, "") DeprecatedClass : DeprecatedClass_Base_DEPRECATED
+ *		class UE_DEPRECATED(5.xx, "DeprecatedClass has been deprecated and shall not be used anymore.") DeprecatedClass : DeprecatedClass_Base_DEPRECATED
  *		{
  *		public:
  *			DeprecatedClass() {}
@@ -256,25 +268,69 @@ enum EPerElement {PerElement};
  *		};
  *
  *		template <typename T>
- *		class UE_DEPRECATED(5.xx, "") DeprecatedClassTemplate
+ *		class UE_DEPRECATED(5.xx, "DeprecatedClassTemplate has been deprecated and shall not be used anymore.") DeprecatedClassTemplate
  *		{
  *		};
  *
  *		template <typename T>
- *		UE_DEPRECATED(5.xx, "")
+ *		UE_DEPRECATED(5.xx, "DeprecatedFunctionTemplate has been deprecated and shall not be used anymore.")
  *		void DeprecatedFunctionTemplate()
  *		{
  *		}
  *
- * @param VERSION The release number in which the feature was marked deprecated.
- * @param MESSAGE A message containing upgrade notes.
+ * @param Version The release number in which the feature was marked deprecated. Note: 'all', as in UE_DEPRECATED(all, ...), can be used for reflected code which has to stick around forever and 
+    is marked with meta-data like DeprecatedFunction or DeprecatedProperty. It indicates that the deprecation grace period is now over but that this property/function will remain deprecated for all time.
+ * @param Message A message containing upgrade notes.
  */
+	#define UE_DEPRECATED(Version, Message) [[deprecated(Message " - Please update your code to the new API before upgrading to the next release, otherwise your project will no longer compile.")]]
 
-#if defined (__INTELLISENSE__)
-#define UE_DEPRECATED(Version, Message)
-#else
-#define UE_DEPRECATED(Version, Message) [[deprecated(Message " Please update your code to the new API before upgrading to the next release, otherwise your project will no longer compile.")]]
-#endif
+/**
+ * Macro for marking up engine-internal code, functions and types. This informs the users that there is no guarantee that the API won't change or be removed in 
+ *  future releases and that they should use at their risk. 
+ *  Using an API marked as such in a non-engine module will trigger a compilation warning (this relies on the c++ [[deprecated]] attribute internally) that can be suppressed by either : 
+ *   - Using PRAGMA_DISABLE_INTERNAL_WARNINGS and PRAGMA_ENABLE_INTERNAL_WARNINGS around the call site (preferred method)
+ *   - Setting bValidateInternalApi to true in the calling module's build rules, thus telling the entire module that any internal API call is accepted
+ *   - Setting bTreatAsEngineModule to true in the calling module's build rules, thus telling the entire module that any internal (OR experimental) API call is accepted
+ * 
+ * It is advised to simply avoid using APIs marked as such in non-engine modules. 
+ *
+ * Please refer to UE_DEPRECATED for how to use the macro. 
+ */
+	#if UE_VALIDATE_INTERNAL_API 
+		#define UE_INTERNAL [[deprecated("This API is internal to the engine. As such, there is no guarantee that it won't change or be removed in future releases. Please use at your own risk.")]]
+		#define PRAGMA_DISABLE_INTERNAL_WARNINGS PRAGMA_DISABLE_DEPRECATION_WARNINGS
+		#define PRAGMA_ENABLE_INTERNAL_WARNINGS PRAGMA_ENABLE_DEPRECATION_WARNINGS
+	#else // UE_VALIDATE_INTERNAL_API 
+		#define UE_INTERNAL
+		#define PRAGMA_DISABLE_INTERNAL_WARNINGS 
+		#define PRAGMA_ENABLE_INTERNAL_WARNINGS 
+	#endif // !UE_VALIDATE_INTERNAL_API 
+
+/**
+ * Macro for marking up experimental code, functions and types from non-experimental modules. This informs the users that there is no guarantee that the API won't change or be removed in 
+ *  future releases and that they should use at their own risk. 
+ *  Using an API marked as such in a non-engine module will trigger a compilation warning (this relies on the c++ [[deprecated]] attribute internally) that can be suppressed by either : 
+ *   - Using PRAGMA_DISABLE_EXPERIMENTAL_WARNINGS and PRAGMA_ENABLE_EXPERIMENTAL_WARNINGS around the call site (preferred method)
+ *   - Setting bValidateExperimentalApi to true in the calling module's build rules, thus telling the entire module that any experimental API call is accepted
+ *   - Setting bTreatAsEngineModule to true in the calling module's build rules, thus telling the entire module that any experimental (OR internal) API call is accepted
+ *
+ * Please refer to UE_DEPRECATED for how to use the macro. 
+ *
+ * @param Version The release number in which the feature was introduced. This does *not* give an indication for how long this API will exist as "experimental". It is only there to inform 
+ *  the users about when the experimental API was first introduced. After some time, this API may be removed or it may become official, at which point it will stop being marked as UE_EXPERIMENTAL.
+ * @param Message A message that may contain additional details about the usage of this experimental API.
+ */
+	#if UE_VALIDATE_EXPERIMENTAL_API
+		#define UE_EXPERIMENTAL(Version, Message)  [[deprecated(Message " - This API is experimental. As such, there is no guarantee that it won't change or be removed in future releases. Please use at your own risk.")]]
+		#define PRAGMA_DISABLE_EXPERIMENTAL_WARNINGS PRAGMA_DISABLE_DEPRECATION_WARNINGS
+		#define PRAGMA_ENABLE_EXPERIMENTAL_WARNINGS PRAGMA_ENABLE_DEPRECATION_WARNINGS
+	#else // UE_VALIDATE_EXPERIMENTAL_API
+		#define UE_EXPERIMENTAL(Version, Message)
+		#define PRAGMA_DISABLE_EXPERIMENTAL_WARNINGS 
+		#define PRAGMA_ENABLE_EXPERIMENTAL_WARNINGS 
+	#endif // !UE_VALIDATE_EXPERIMENTAL_API
+
+#endif // !defined(__INTELLISENSE__)
 
 #ifndef UE_DEPRECATED_FORGAME
 	#define UE_DEPRECATED_FORGAME PREPROCESSOR_NOTHING_FUNCTION
@@ -284,12 +340,6 @@ enum EPerElement {PerElement};
 	#define UE_DEPRECATED_FORENGINE UE_DEPRECATED
 #endif
 
-#if UE_VALIDATE_INTERNAL_API
-	#define UE_INTERNAL [[deprecated("Please remove usage of this internal API before upgrading to the next release, otherwise your project will no longer compile.")]]
-#else
-	#define UE_INTERNAL
-#endif
-
 
 /**
  * Macro which can be placed in a header to throw a deprecation warning when it is included.
@@ -328,7 +378,7 @@ struct TStaticDeprecateExpression
  */
 #define UE_STATIC_DEPRECATE(Version, bExpression, Message) \
 	struct PREPROCESSOR_JOIN(FDeprecationMsg_, __LINE__) { \
-		[[deprecated(Message " Please update your code to the new API before upgrading to the next release, otherwise your project will no longer compile.")]] \
+		[[deprecated(Message " - Please update your code to the new API before upgrading to the next release, otherwise your project will no longer compile.")]] \
 		static constexpr int condition(TStaticDeprecateExpression<true>) { return 1; } \
 		static constexpr int condition(TStaticDeprecateExpression<false>) { return 1; } \
 	}; \
