[core] Added some simple expression parser tests

diff --git a/Engine/Source/Runtime/Core/Private/Misc/ExpressionParser.cpp b/Engine/Source/Runtime/Core/Private/Misc/ExpressionParser.cpp
--- a/Engine/Source/Runtime/Core/Private/Misc/ExpressionParser.cpp
+++ b/Engine/Source/Runtime/Core/Private/Misc/ExpressionParser.cpp
@@ -784,7 +784,7 @@ namespace ExpressionParser
 	}
 }
 
-#if !(UE_BUILD_SHIPPING || UE_BUILD_TEST)
+#if WITH_DEV_AUTOMATION_TESTS
 
 namespace Tests
 {
@@ -1092,6 +1092,121 @@ bool FShortCircuitParserTest::RunTest(const FString& Parameters)
 	return true;
 }
 
-#endif
+namespace Tests
+{
+	template <class ResultTokenType>
+	static TOptional<FExpressionError> ConsumeCharRangeGreedy(FExpressionTokenConsumer& Consumer, TCHAR Start, TCHAR EndInclusive)
+	{
+		TOptional<FStringToken> Token = Consumer.GetStream().ParseToken([=](TCHAR Ch)
+		{
+			return Ch >= Start && Ch <= EndInclusive ? EParseState::Continue : EParseState::StopBefore;
+		});
+		if (Token)
+		{
+			check(Token->IsValid());
+			Consumer.Add(*Token, ResultTokenType());
+		}
+		return {};
+	}
+
+	struct FLowerAlphaTokenTag {};
+	static TOptional<FExpressionError> ConsumeLowerAlphaGreedy(FExpressionTokenConsumer& Consumer)
+	{
+		return ConsumeCharRangeGreedy<FLowerAlphaTokenTag>(Consumer, 'a', 'z');
+	}
+
+	struct FDigitsTokenTag {};
+	static TOptional<FExpressionError> ConsumeDigitsGreedy(FExpressionTokenConsumer& Consumer)
+	{
+		return ConsumeCharRangeGreedy<FDigitsTokenTag>(Consumer, '0', '9');
+	}
+
+	struct FStringTokenTag {};
+} // end namespace Tests
+
+DEFINE_EXPRESSION_NODE_TYPE(Tests::FLowerAlphaTokenTag, 0x01772467, 0xb30c4b0f, 0xb7863e3f, 0x5a52360b);
+DEFINE_EXPRESSION_NODE_TYPE(Tests::FDigitsTokenTag, 0x51e2dea9, 0xa09247ec, 0x93651811, 0x8f4df950);
+DEFINE_EXPRESSION_NODE_TYPE(Tests::FStringTokenTag, 0x3365d89e, 0xa4344abe, 0xbcf04bdc, 0xdcd728c8);
+
+IMPLEMENT_SIMPLE_AUTOMATION_TEST(FBasicLexerTest, "System.Core.Expression Parser.Lexer.Basic", EAutomationTestFlags_ApplicationContextMask | EAutomationTestFlags::SmokeFilter | EAutomationTestFlags::HighPriority)
+bool FBasicLexerTest::RunTest(const FString& Parameters)
+{
+	FTokenDefinitions TokenDefinitions;
+	TokenDefinitions.DefineToken(Tests::ConsumeLowerAlphaGreedy);
+
+	ExpressionParser::LexResultType Result = ExpressionParser::Lex(TEXT("abc"), TokenDefinitions);
+	UTEST_TRUE(TEXT("Lex succeeded"), Result.HasValue());
+	UTEST_EQUAL(TEXT("Found one token"), Result.GetValue().Num(), 1);
+	UTEST_EQUAL(TEXT("Token contains expected 3 characters"), Result.GetValue()[0].Context.GetString().Len(), 3);
+
+	return true;
+}
+
+IMPLEMENT_SIMPLE_AUTOMATION_TEST(FLexerIgnoreWhitespaceTest, "System.Core.Expression Parser.Lexer.Ignore whitespace", EAutomationTestFlags_ApplicationContextMask | EAutomationTestFlags::SmokeFilter | EAutomationTestFlags::HighPriority)
+bool FLexerIgnoreWhitespaceTest::RunTest(const FString& Parameters)
+{
+	FTokenDefinitions TokenDefinitions;
+	TokenDefinitions.DefineToken(Tests::ConsumeLowerAlphaGreedy);
+	TokenDefinitions.DefineToken(Tests::ConsumeDigitsGreedy);
+	TokenDefinitions.IgnoreWhitespace();
+
+	ExpressionParser::LexResultType Result = ExpressionParser::Lex(TEXT("abc 123"), TokenDefinitions);
+	UTEST_TRUE(TEXT("Lex succeeded"), Result.HasValue());
+	UTEST_EQUAL(TEXT("Found two tokens"), Result.GetValue().Num(), 2);
+	UTEST_EQUAL(TEXT("First token contains expected 3 characters"), Result.GetValue()[0].Context.GetString().Len(), 3);
+	UTEST_EQUAL(TEXT("Second token contains expected 3 characters"), Result.GetValue()[1].Context.GetString().Len(), 3);
+
+	return true;
+}
+
+IMPLEMENT_SIMPLE_AUTOMATION_TEST(FLexerParseStringTest, "System.Core.Expression Parser.Lexer.Parse string", EAutomationTestFlags_ApplicationContextMask | EAutomationTestFlags::SmokeFilter | EAutomationTestFlags::HighPriority)
+bool FLexerParseStringTest::RunTest(const FString& Parameters)
+{
+	// an example of matching but ignoring stuff before and after a token
+
+	auto ConsumeString = [](FExpressionTokenConsumer& Consumer) -> TOptional<FExpressionError> {
+
+		FTokenStream& Stream = Consumer.GetStream();
+		if (TOptional<FStringToken> InitialQuoteToken = Stream.ParseSymbol('"'))
+		{
+			TOptional<FStringToken> ContentsToken = Stream.ParseToken([](TCHAR C)
+			{
+				return C == '"' ? EParseState::StopBefore : EParseState::Continue;
+			}, &*InitialQuoteToken);
+
+			// note: continue to accumulate to initial token, leaving contents token untouched
+			TOptional<FStringToken> CloseTokenToIgnore = Stream.ParseSymbol('"', &*InitialQuoteToken);
+			if (!ContentsToken || !CloseTokenToIgnore)
+			{
+				return FExpressionError(FText::AsCultureInvariant(TEXT("Unterminated string")));
+			}
+			Consumer.Add(*ContentsToken, Tests::FStringTokenTag());
+			// skip close quote
+			Stream.SetReadPos(*CloseTokenToIgnore);
+		}
+		return {};
+	};
+
+	FTokenDefinitions TokenDefinitions;
+	TokenDefinitions.DefineToken(ConsumeString);
+	TokenDefinitions.DefineToken(Tests::ConsumeDigitsGreedy);
+	TokenDefinitions.IgnoreWhitespace();
+
+	ExpressionParser::LexResultType Result = ExpressionParser::Lex(TEXT("\"1 a\" 123"), TokenDefinitions);
+	UTEST_TRUE(TEXT("Lex succeeded"), Result.HasValue());
+	UTEST_EQUAL(TEXT("Found two tokens"), Result.GetValue().Num(), 2);
+	UTEST_EQUAL(TEXT("First token contains expected 3 characters"), Result.GetValue()[0].Context.GetString().Len(), 3);
+	UTEST_EQUAL(TEXT("Second token contains expected 3 characters"), Result.GetValue()[1].Context.GetString().Len(), 3);
+
+	ExpressionParser::LexResultType ExpectedFailureResult = ExpressionParser::Lex(TEXT("\"abc"), TokenDefinitions);
+	UTEST_TRUE(TEXT("Unterminated string"), ExpectedFailureResult.HasError());
+
+	ExpectedFailureResult = ExpressionParser::Lex(TEXT("\""), TokenDefinitions);
+	UTEST_TRUE(TEXT("Just quote"), ExpectedFailureResult.HasError());
+
+	return true;
+}
+
+#endif // WITH_DEV_AUTOMATION_TESTS
 
 #undef LOCTEXT_NAMESPACE
