Added Apple XCode Instruments memory tracking

diff --git a/Engine/Source/Runtime/Core/Private/Apple/ApplePlatformMemory.cpp b/Engine/Source/Runtime/Core/Private/Apple/ApplePlatformMemory.cpp
--- a/Engine/Source/Runtime/Core/Private/Apple/ApplePlatformMemory.cpp
+++ b/Engine/Source/Runtime/Core/Private/Apple/ApplePlatformMemory.cpp
@@ -39,6 +39,24 @@
 
 NS_ASSUME_NONNULL_BEGIN
 
+// Enable to use XCode Instuments hooks to propagate UE allocs.
+// It has impact on runtime performance, disabled by default. 
+#define USE_APPLE_SUPPORT_INSTRUMENTED_ALLOCS (0 && APPLE_SUPPORT_INSTRUMENTED_ALLOCS)
+
+// This name is used in XCode instruments
+enum class MmapTag : int32
+{
+	BinnedAllocFromOS			= VM_MAKE_TAG(VM_MEMORY_APPLICATION_SPECIFIC_1 + 0), // "Memory Tag 240"
+	AllocateVirtualMemoryBlock	= VM_MAKE_TAG(VM_MEMORY_APPLICATION_SPECIFIC_1 + 1), // "Memory Tag 241"
+	LLMAlloc					= VM_MAKE_TAG(VM_MEMORY_APPLICATION_SPECIFIC_1 + 2), // "Memory Tag 242"
+};
+
+static void* MmapWithTag(void* __nullable Addr, const size_t Len, const int Prot, const int Flags, const MmapTag Tag, const off_t Offset)
+{
+	// Replace Tag with -1 to disable
+	return mmap(Addr, Len, Prot, Flags, (int32)Tag, Offset);
+}
+
 /** 
  * Zombie object implementation so that we can implement NSZombie behaviour for our custom allocated objects.
  * Will leak memory - just like Cocoa's NSZombie - but allows for debugging of invalid usage of the pooled types.
@@ -579,7 +597,7 @@ void* _Nullable FApplePlatformMemory::BinnedAllocFromOS(SIZE_T Size)
 
 	// the remainder of the map will be used for the descriptor, if any.
 	// we always allocate at least one page more than needed
-	void* PointerWeGotFromMMap = mmap(nullptr, ActualSizeMapped, PROT_READ | PROT_WRITE, MAP_PRIVATE | MAP_ANON, -1, 0);
+	void* PointerWeGotFromMMap = MmapWithTag(nullptr, ActualSizeMapped, PROT_READ | PROT_WRITE, MAP_PRIVATE | MAP_ANON, MmapTag::BinnedAllocFromOS, 0);
 	// store these values, to unmap later
 
 	Pointer = PointerWeGotFromMMap;
@@ -661,7 +679,7 @@ void* _Nullable FApplePlatformMemory::BinnedAllocFromOS(SIZE_T Size)
 	LLM_IF_ENABLED(FLowLevelMemTracker::Get().OnLowLevelAlloc(ELLMTracker::Platform, Pointer, Size));
 	return Pointer;
 #else
-	void* Ptr = mmap(nullptr, Size, PROT_READ | PROT_WRITE, MAP_PRIVATE | MAP_ANON, -1, 0);
+	void* Ptr = MmapWithTag(nullptr, Size, PROT_READ | PROT_WRITE, MAP_PRIVATE | MAP_ANON, MmapTag::BinnedAllocFromOS, 0);
 	if (Ptr == (void*)-1)
 	{
 		UE_LOG(LogTemp, Warning, TEXT("mmap failure allocating %d, error code: %d"), Size, errno);
@@ -777,7 +795,7 @@ FApplePlatformMemory::FPlatformVirtualMemoryBlock FApplePlatformMemory::FPlatfor
 	size_t Alignment = FMath::Max(InAlignment, GetVirtualSizeAlignment());
 	check(Alignment <= GetVirtualSizeAlignment());
 
-	Result.Ptr = mmap(nullptr, Result.GetActualSize(), PROT_READ | PROT_WRITE, MAP_PRIVATE | MAP_ANON, -1, 0);
+	Result.Ptr = MmapWithTag(nullptr, Result.GetActualSize(), PROT_READ | PROT_WRITE, MAP_PRIVATE | MAP_ANON, MmapTag::AllocateVirtualMemoryBlock, 0);
 	if (!LIKELY(Result.Ptr != MAP_FAILED))
 	{
 		FPlatformMemory::OnOutOfMemory(Result.GetActualSize(), InAlignment);
@@ -836,7 +854,7 @@ int64 LLMMallocTotal = 0;
 
 void* LLMAlloc(size_t Size)
 {
-    void* Ptr = mmap(nullptr, Size, PROT_READ | PROT_WRITE, MAP_PRIVATE | MAP_ANON, -1, 0);
+    void* Ptr = MmapWithTag(nullptr, Size, PROT_READ | PROT_WRITE, MAP_PRIVATE | MAP_ANON, MmapTag::LLMAlloc, 0);
 
     LLMMallocTotal += Size;
     
@@ -868,6 +886,96 @@ bool FApplePlatformMemory::GetLLMAllocFunctions(void*_Nonnull(*_Nonnull&OutAlloc
 #endif
 }
 
+#if APPLE_SUPPORT_INSTRUMENTED_ALLOCS
+
+#if USE_APPLE_SUPPORT_INSTRUMENTED_ALLOCS
+
+// [REDACTED_PROFILE_URL]/CF/blob/main/CFRuntime.c#L84-L92
+enum class CFTraceEvent : int32_t
+{
+	ObjectRetained = 12,
+	ObjectReleased = 13,
+	Malloc = 16,
+	_Free = 19,
+	Zombie = 21,
+	VMalloc = 23,
+	_Free2 = 26, // ?
+	Retain = 28,
+	Released = 29,
+	_Free3 = 30, // ?
+};
+extern "C" void __CFRecordAllocationEvent(int32_t CFTraceType, void* Ptr, int64_t Size, uint64_t Data, const char* Name);
+extern "C" void __CFSetLastAllocationEventName(void* Ptr, const char* Name);
+
+// [REDACTED_PROFILE_URL]/libmalloc/blob/main/private/stack_logging.h#L71
+enum class MallocLogger : uint32_t
+{
+	Free = 0,
+	Generic = 1,
+	Alloc = 2,
+	Dealloc = 4,
+	FlagZone = 8, // ?
+	VMAlloc = 16,
+	VMDealloc = 32,
+	FlagCleared = 64, // ?
+	FileMapOrSharedMem = 128
+};
+typedef void(FnMallocLogger)(uint32_t MallocLoggerType, uintptr_t Arg1, uintptr_t Arg2, uintptr_t Arg3, uintptr_t Result, uint32_t NumFramesToSkip);
+extern "C" FnMallocLogger *malloc_logger;
+
+void FApplePlatformMemory::OnLowLevelMemory_Alloc(void const* Pointer, uint64 Size, uint64 Tag)
+{
+#if ENABLE_LOW_LEVEL_MEM_TRACKER
+	const char* EventName = FLowLevelMemTracker::IsEnabled() ? LLMGetTagNameANSI((ELLMTag)Tag) : "UE";
+
+	// mmap will already call logger internally
+	// see [REDACTED_PROFILE_URL]/xnu/blob/main/libsyscall/wrappers/unix03/mmap.c#L72
+	if (Tag == (uint64)ELLMTracker::Platform)
+	{
+		__CFSetLastAllocationEventName((void*)Pointer, EventName);
+		return;
+	}
+#else
+	const char* EventName = "UE";
+#endif // ENABLE_LOW_LEVEL_MEM_TRACKER
+
+	if (malloc_logger)
+	{
+		malloc_logger((uint32_t)MallocLogger::Alloc, 0, (uintptr_t)Size, 0, (uintptr_t)Pointer, 0);
+		__CFSetLastAllocationEventName((void*)Pointer, EventName);
+	}
+}
+
+void FApplePlatformMemory::OnLowLevelMemory_Free(void const* Pointer, uint64 Size, uint64 Tag)
+{
+#if ENABLE_LOW_LEVEL_MEM_TRACKER
+	// munmap will already call logger internally
+	if (Tag == (uint64)ELLMTracker::Platform)
+	{
+		return;
+	}
+#endif // ENABLE_LOW_LEVEL_MEM_TRACKER
+
+	if (malloc_logger)
+	{
+		malloc_logger((uint32_t)MallocLogger::Dealloc, 0, (uintptr_t)Pointer, 0, 0, 0);
+	}
+}
+
+#else
+
+void FApplePlatformMemory::OnLowLevelMemory_Alloc(void const* Pointer, uint64 Size, uint64 Tag)
+{
+}
+
+void FApplePlatformMemory::OnLowLevelMemory_Free(void const* Pointer, uint64 Size, uint64 Tag)
+{
+}
+
+#endif // USE_APPLE_SUPPORT_INSTRUMENTED_ALLOCS
+
+#endif // APPLE_SUPPORT_INSTRUMENTED_ALLOCS
+
 bool FApplePlatformMemory::CanOverallocateVirtualMemory()
 {
 #if PLATFORM_IOS || PLATFORM_TVOS
diff --git a/Engine/Source/Runtime/Core/Public/Apple/ApplePlatformMemory.h b/Engine/Source/Runtime/Core/Public/Apple/ApplePlatformMemory.h
--- a/Engine/Source/Runtime/Core/Public/Apple/ApplePlatformMemory.h
+++ b/Engine/Source/Runtime/Core/Public/Apple/ApplePlatformMemory.h
@@ -25,6 +25,13 @@ NS_ASSUME_NONNULL_BEGIN
 #define USE_MALLOC_BINNED2 !USE_MALLOC_BINNED3
 #endif // USE_MALLOC_BINNED2
 
+// Support for passing LLM alloc data to XCode Instruments.
+// Use USE_APPLE_SUPPORT_INSTRUMENTED_ALLOCS to enable it.
+// As APPLE_SUPPORT_INSTRUMENTED_ALLOCS define only enables the LLM API, because changes to this header trigger wide recompilation.
+#ifndef APPLE_SUPPORT_INSTRUMENTED_ALLOCS
+	#define APPLE_SUPPORT_INSTRUMENTED_ALLOCS (!UE_BUILD_SHIPPING)
+#endif // APPLE_SUPPORT_INSTRUMENTED_ALLOCS
+
 /**
  * NSObject subclass that can be used to override the allocation functions to go through UE4's memory allocator.
  * This ensures that memory allocated by custom Objective-C types can be tracked by UE4's tools and 
@@ -158,6 +165,10 @@ struct CORE_API FApplePlatformMemory : public FGenericPlatformMemory
 	};
 
     static bool GetLLMAllocFunctions(void* _Nonnull (* _Nonnull &OutAllocFunction)(size_t), void(* _Nonnull &OutFreeFunction)(void*, size_t), int32& OutAlignment);
+#if APPLE_SUPPORT_INSTRUMENTED_ALLOCS
+	static void OnLowLevelMemory_Alloc(void const* Pointer, uint64 Size, uint64 Tag);
+	static void OnLowLevelMemory_Free(void const* Pointer, uint64 Size, uint64 Tag);
+#endif // APPLE_SUPPORT_INSTRUMENTED_ALLOCS
     //~ End FGenericPlatformMemory Interface
 	
 	/** Setup the current default CFAllocator to use our malloc functions. */
