[StaticBitArray]
 - Small cast fix

#rnx

diff --git a/Engine/Source/Runtime/Core/Public/Containers/StaticBitArray.h b/Engine/Source/Runtime/Core/Public/Containers/StaticBitArray.h
--- a/Engine/Source/Runtime/Core/Public/Containers/StaticBitArray.h
+++ b/Engine/Source/Runtime/Core/Public/Containers/StaticBitArray.h
@@ -303,7 +303,7 @@ public:
 			const WordType Bits = ~(Words[WordIndex]);
 			UE_ASSUME(Bits != 0);
 
-			const int32 LowestBitIndex = FMath::CountTrailingZeros64(Bits) + (WordIndex << NumBitsPerWordLog2);
+			const int32 LowestBitIndex = (int32)FMath::CountTrailingZeros64(Bits) + (WordIndex << NumBitsPerWordLog2);
 			if (LowestBitIndex < LocalNumBits)
 			{
 				return LowestBitIndex;
