Adding FArchive::ArIsMigratingRemoteObjects

diff --git a/Engine/Source/Runtime/Core/Private/Serialization/Archive.cpp b/Engine/Source/Runtime/Core/Private/Serialization/Archive.cpp
--- a/Engine/Source/Runtime/Core/Private/Serialization/Archive.cpp
+++ b/Engine/Source/Runtime/Core/Private/Serialization/Archive.cpp
@@ -193,6 +193,7 @@ void FArchiveState::Reset()
 	ArUseCustomPropertyList				= false;
 	ArMergeOverrides					= false;
 	ArShouldSkipUpdateCustomVersion		= false;
+	ArIsMigratingRemoteObjects			= false;
 	SavePackageData						= nullptr;
 	SerializedProperty					= nullptr;
 
@@ -253,6 +254,7 @@ void FArchiveState::CopyTrivialFArchiveStatusMembers(const FArchiveState& Archiv
 	ArUseCustomPropertyList              = ArchiveToCopy.ArUseCustomPropertyList;
 	ArMergeOverrides					 = ArchiveToCopy.ArMergeOverrides;
 	ArShouldSkipUpdateCustomVersion		 = ArchiveToCopy.ArShouldSkipUpdateCustomVersion;
+	ArIsMigratingRemoteObjects			 = ArchiveToCopy.ArIsMigratingRemoteObjects;
 	SavePackageData						 = ArchiveToCopy.SavePackageData;
 	SerializedProperty					 = ArchiveToCopy.SerializedProperty;
 #if USE_STABLE_LOCALIZATION_KEYS
@@ -601,6 +603,11 @@ void FArchiveState::SetShouldSkipUpdateCustomVersion(bool bShouldSkip)
 	ForEachState([bShouldSkip](FArchiveState& State) { State.ArShouldSkipUpdateCustomVersion = bShouldSkip; });
 }
 
+void FArchiveState::SetMigratingRemoteObjects(bool bMigrating)
+{
+	ForEachState([bMigrating](FArchiveState& State) { State.ArIsMigratingRemoteObjects = bMigrating; });
+}
+
 void FArchiveState::SetCustomVersion(const FGuid& Key, int32 Version, FName FriendlyName)
 {
 	const_cast<FCustomVersionContainer&>(GetCustomVersions()).SetVersion(Key, Version, FriendlyName);
diff --git a/Engine/Source/Runtime/Core/Public/Serialization/Archive.h b/Engine/Source/Runtime/Core/Public/Serialization/Archive.h
--- a/Engine/Source/Runtime/Core/Public/Serialization/Archive.h
+++ b/Engine/Source/Runtime/Core/Public/Serialization/Archive.h
@@ -326,12 +326,24 @@ public:
 	 */
 	CORE_API void SetShouldSkipUpdateCustomVersion(bool bShouldSkip);
 
-	/** Returns true if this Archive is does not update the custom version data for a given key if has already been set previously in the local custom versions container*/
+	/** Returns true if this Archive does not update the custom version data for a given key if has already been set previously in the local custom versions container*/
 	FORCEINLINE bool ShouldSkipUpdateCustomVersion() const
 	{
 		return ArShouldSkipUpdateCustomVersion;
 	}
 
+	/** 
+	 * If set true, this archive is migrating remote objects
+	 * @param bMigrating Whether this archive is migrating remote objects
+	 */
+	CORE_API void SetMigratingRemoteObjects(bool bMigrating);
+
+	/** Returns true if this Archive is migrating remote object */
+	FORCEINLINE bool IsMigratingRemoteObjects() const
+	{
+		return ArIsMigratingRemoteObjects;
+	}
+
 	/** Returns true if this archive contains errors, which means that further serialization is generally not safe. */
 	FORCEINLINE bool IsError() const
 	{
@@ -819,6 +831,9 @@ private:
 	 * Setting this to true should improve serialization performance by skipping the step entirely */
 	uint8 ArShouldSkipUpdateCustomVersion : 1;
 
+	/** If true this archive is migrating remote object data. */
+	uint8 ArIsMigratingRemoteObjects : 1;
+
 public:
 	/** Quickly tell if an archive contains script code. */
 	uint8 ArContainsCode : 1;
@@ -2128,6 +2143,8 @@ public:
 	using FArchiveState::ShouldSkipCompilingAssets;
 	using FArchiveState::SetShouldSkipUpdateCustomVersion;
 	using FArchiveState::ShouldSkipUpdateCustomVersion;
+	using FArchiveState::SetMigratingRemoteObjects;
+	using FArchiveState::IsMigratingRemoteObjects;
 	
 private:
 	using FArchiveState::ArUEVer;
