[AutoRTFM] Don't call autortfm_[pre|post]_open for closed calls to autortfm_lookup_function

This is hot code, and the pre / post calls are unnecessary and expensive.

diff --git a/Engine/Source/Runtime/Core/Private/AutoRTFM/CompilerInterface.cpp b/Engine/Source/Runtime/Core/Private/AutoRTFM/CompilerInterface.cpp
--- a/Engine/Source/Runtime/Core/Private/AutoRTFM/CompilerInterface.cpp
+++ b/Engine/Source/Runtime/Core/Private/AutoRTFM/CompilerInterface.cpp
@@ -2,6 +2,7 @@
 
 #if (defined(__AUTORTFM) && __AUTORTFM)
 
+#include "AutoRTFM/AutoRTFM.h"
 #include "ContextInlines.h"
 #include "FunctionMapInlines.h"
 #include "Memcpy.h"
@@ -116,7 +117,13 @@ extern "C" UE_AUTORTFM_API AUTORTFM_NO_ASAN void autortfm_record_masked_write(vo
 	}
 }
 
-extern "C" UE_AUTORTFM_ALWAYS_OPEN UE_AUTORTFM_API void* autortfm_lookup_function(void* OriginalFunction, const char* Where)
+// Mark autortfm_lookup_function() with UE_AUTORTFM_NOAUTORTFM and use
+// UE_AUTORTFM_REGISTER_OPEN_FUNCTION_EXPLICIT to map the closed function to the
+// open. We do this to prevent calls to autortfm_lookup_function() being wrapped
+// with calls to autortfm_pre_open() and autortfm_post_open(), which is
+// expensive for hot code like this.
+UE_AUTORTFM_REGISTER_OPEN_FUNCTION_EXPLICIT(autortfm_lookup_function, autortfm_lookup_function);
+extern "C" UE_AUTORTFM_NOAUTORTFM UE_AUTORTFM_API void* autortfm_lookup_function(void* OriginalFunction, const char* Where)
 {
 	FContext* Context = FContext::Get();
     return FunctionMapLookup(OriginalFunction, Where);
