Add missing linux functions for xstat

diff --git a/Engine/Source/Runtime/Core/Private/AutoRTFM/StdLibSupport.cpp b/Engine/Source/Runtime/Core/Private/AutoRTFM/StdLibSupport.cpp
--- a/Engine/Source/Runtime/Core/Private/AutoRTFM/StdLibSupport.cpp
+++ b/Engine/Source/Runtime/Core/Private/AutoRTFM/StdLibSupport.cpp
@@ -681,8 +681,25 @@ UE_AUTORTFM_REGISTER_SELF_FUNCTION(gettimeofday);
 UE_AUTORTFM_REGISTER_SELF_FUNCTION(gmtime_r);
 UE_AUTORTFM_REGISTER_SELF_FUNCTION(localtime_r);
 UE_AUTORTFM_REGISTER_SELF_FUNCTION(bcmp);
-
 UE_AUTORTFM_REGISTER_SELF_FUNCTION(pthread_getspecific);
+
+int RTFM_stat(const char* Path, struct stat* StatBuf) throw()
+{
+	FContext* Context = FContext::Get();
+	Context->RecordWrite(StatBuf, sizeof(*StatBuf));
+
+	return stat(Path, StatBuf);
+}
+UE_AUTORTFM_REGISTER_OPEN_FUNCTION(stat);
+
+int RTFM_fstat(int Fd, struct stat* StatBuf) throw()
+{
+	FContext* Context = FContext::Get();
+	Context->RecordWrite(StatBuf, sizeof(*StatBuf));
+
+	return fstat(Fd, StatBuf);
+}
+UE_AUTORTFM_REGISTER_OPEN_FUNCTION(fstat);
 #endif // PLATFORM_LINUX
 
 wchar_t* RTFM_wcscpy(wchar_t* Dst, const wchar_t* Src) throw()
