Friending Global comparison operators

diff --git a/Engine/Source/Runtime/Core/Public/Templates/PointerVariants.h b/Engine/Source/Runtime/Core/Public/Templates/PointerVariants.h
--- a/Engine/Source/Runtime/Core/Public/Templates/PointerVariants.h
+++ b/Engine/Source/Runtime/Core/Public/Templates/PointerVariants.h
@@ -118,49 +118,87 @@ namespace UE::Core::Private
 				}, PtrVariant);
 		}
 
+	private:
+		template <typename LhsType, bool bLhsStrength, typename RhsType, bool bRhsStrength, typename>
+		friend bool operator==(const TPtrVariantBase<LhsType, bLhsStrength>& Lhs, const TPtrVariantBase<RhsType, bRhsStrength>& Rhs);
+
+		template<typename LhsType, typename RhsType, bool bRhsStrength, typename>
+		friend bool operator==(LhsType* Lhs, const TPtrVariantBase<RhsType, bRhsStrength>& Rhs);
+
+		template<typename LhsType, bool bLhsStrength, typename RhsType, typename>
+		friend bool operator==(const TPtrVariantBase<LhsType, bLhsStrength>& Lhs, RhsType* Rhs);
+
+	private:
+		const BaseType* GetRawPtrValue_Internal() const
+		{
+			return ::Visit([this](auto& StoredValue)
+					{
+						using StoredValueType = std::decay_t<decltype(StoredValue)>;
+						if constexpr (std::is_same_v<StoredValueType, SharedType>)
+						{
+							if constexpr (bIsStrong)
+							{
+								return IsValid() ? StoredValue.Get() : (const BaseType*)nullptr;
+							}
+							else
+							{
+								return IsValid() ? StoredValue.Pin().Get() : (const BaseType*)nullptr;
+							}
+						}
+						else
+						{
+							return IsValid() ? StoredValue.CastedPtr : (const BaseType*)nullptr;
+						}
+					}, PtrVariant);
+		}
+
 	private:
 		TVariant<SharedType, FObjectPtrWrapper> PtrVariant;
 	};
-}
 
-/* Global Comparisons */
+	template<typename T1, typename T2
+		UE_REQUIRES_DEFINITION(UE_REQUIRES_EXPR((T1*)nullptr == (T2*)nullptr))
+	>
+	struct TCanCompare
+	{
+	};
 
-template <typename LhsType, bool bLhsStrength, typename RhsType, bool bRhsStrength, 
-	decltype((LhsType*)nullptr == (RhsType*)nullptr) = false /* this will prevent the template instantiating if the pointer types are not comparable*/>
-bool operator==(const UE::Core::Private::TPtrVariantBase<LhsType, bLhsStrength>& Lhs, const UE::Core::Private::TPtrVariantBase<RhsType, bRhsStrength>& Rhs)
-{
-	return Lhs.Get() == Rhs.Get();
-}
+	template <typename LhsType, bool bLhsStrength, typename RhsType, bool bRhsStrength, typename = TCanCompare<LhsType, RhsType>>
+	bool operator==(const TPtrVariantBase<LhsType, bLhsStrength>& Lhs, const TPtrVariantBase<RhsType, bRhsStrength>& Rhs)
+	{
+		return Lhs.GetRawPtrValue_Internal() == Rhs.GetRawPtrValue_Internal();
+	}
 
-template <typename LhsType, bool bLhsStrength, typename RhsType, bool bRhsStrength, decltype((LhsType*)nullptr == (RhsType*)nullptr) = false>
-bool operator!=(const UE::Core::Private::TPtrVariantBase<LhsType, bLhsStrength>& Lhs, const UE::Core::Private::TPtrVariantBase<RhsType, bRhsStrength>& Rhs)
-{
-	return !(Lhs == Rhs);
-}
+	template <typename LhsType, bool bLhsStrength, typename RhsType, bool bRhsStrength, typename = TCanCompare<LhsType, RhsType>>
+	bool operator!=(const TPtrVariantBase<LhsType, bLhsStrength>& Lhs, const TPtrVariantBase<RhsType, bRhsStrength>& Rhs)
+	{
+		return !(Lhs == Rhs);
+	}
 
-template<typename LhsType, typename RhsType, bool bRhsStrength, decltype((LhsType*)nullptr == (RhsType*)nullptr) = false>
-bool operator==(LhsType* Lhs, const UE::Core::Private::TPtrVariantBase<RhsType, bRhsStrength>& Rhs)
-{
-	return Lhs == Rhs.Get();
-}
+	template<typename LhsType, typename RhsType, bool bRhsStrength, typename = TCanCompare<LhsType, RhsType>>
+	bool operator==(LhsType* Lhs, const TPtrVariantBase<RhsType, bRhsStrength>& Rhs)
+	{
+		return Lhs == Rhs.GetRawPtrValue_Internal();
+	}
 
-template<typename LhsType, typename RhsType, bool bRhsStrength, decltype((LhsType*)nullptr == (RhsType*)nullptr) = false>
-bool operator!=(LhsType* Lhs, const UE::Core::Private::TPtrVariantBase<RhsType, bRhsStrength>& Rhs)
-{
-	return !(Lhs == Rhs);
-}
+	template<typename LhsType, typename RhsType, bool bRhsStrength, typename = TCanCompare<LhsType, RhsType>>
+	bool operator!=(LhsType* Lhs, const TPtrVariantBase<RhsType, bRhsStrength>& Rhs)
+	{
+		return !(Lhs == Rhs);
+	}
 
-template<typename LhsType, bool bLhsStrength, typename RhsType, decltype((LhsType*)nullptr == (RhsType*)nullptr) = false>
-bool operator==(const UE::Core::Private::TPtrVariantBase<LhsType, bLhsStrength>& Lhs, RhsType* Rhs)
-{
-	return Lhs.Get() == Rhs;
-}
+	template<typename LhsType, bool bLhsStrength, typename RhsType, typename = TCanCompare<LhsType, RhsType>>
+	bool operator==(const TPtrVariantBase<LhsType, bLhsStrength>& Lhs, RhsType* Rhs)
+	{
+		return Lhs.GetRawPtrValue_Internal() == Rhs;
+	}
 
-template<typename LhsType, bool bLhsStrength, typename RhsType, decltype((LhsType*)nullptr == (RhsType*)nullptr) = false>
-bool operator!=(const UE::Core::Private::TPtrVariantBase<LhsType, bLhsStrength>& Lhs, RhsType* Rhs)
-{
-	return !(Lhs == Rhs);
-}
+	template<typename LhsType, bool bLhsStrength, typename RhsType, typename = TCanCompare<LhsType, RhsType>>
+	bool operator!=(const TPtrVariantBase<LhsType, bLhsStrength>& Lhs, RhsType* Rhs)
+	{
+		return !(Lhs == Rhs);
+	}
+} // namespace UE::Core::Private
 
 template<typename BaseType>
 struct TStrongPtrVariant : public UE::Core::Private::TPtrVariantBase<BaseType, true>
