Optimize FGuid::ImportTextItem to read at most 32 characters from the buffer

#rnx

diff --git a/Engine/Source/Runtime/Core/Private/Misc/Guid.cpp b/Engine/Source/Runtime/Core/Private/Misc/Guid.cpp
--- a/Engine/Source/Runtime/Core/Private/Misc/Guid.cpp
+++ b/Engine/Source/Runtime/Core/Private/Misc/Guid.cpp
@@ -33,19 +33,19 @@ bool FGuid::ExportTextItem(FString& ValueStr, FGuid const& DefaultValue, UObject
 
 bool FGuid::ImportTextItem(const TCHAR*& Buffer, int32 PortFlags, class UObject* Parent, FOutputDevice* ErrorText)
 {
-	FStringView BufferView = Buffer;
+	constexpr int32 GuidStrLen = 32;
 
-	if (BufferView.Len() < 32)
+	if (FPlatformString::Strnlen(Buffer, GuidStrLen) < GuidStrLen)
 	{
 		return false;
 	}
 
-	if (!ParseExact(BufferView.Left(32), EGuidFormats::Digits, *this))
+	if (!ParseExact(FStringView(Buffer, GuidStrLen), EGuidFormats::Digits, *this))
 	{
 		return false;
 	}
 
-	Buffer += 32;
+	Buffer += GuidStrLen;
 
 	return true;
 }
