Mac Editor Compile error fix attempt

diff --git a/Engine/Source/Runtime/Core/Public/Templates/PointerVariants.h b/Engine/Source/Runtime/Core/Public/Templates/PointerVariants.h
--- a/Engine/Source/Runtime/Core/Public/Templates/PointerVariants.h
+++ b/Engine/Source/Runtime/Core/Public/Templates/PointerVariants.h
@@ -163,37 +163,37 @@ namespace UE::Core::Private
 	{
 	};
 
-	template <typename LhsType, bool bLhsStrength, typename RhsType, bool bRhsStrength, typename = TCanCompare<LhsType, RhsType>>
+	template <typename LhsType, bool bLhsStrength, typename RhsType, bool bRhsStrength, typename = typename TCanCompare<LhsType, RhsType>>
 	bool operator==(const TPtrVariantBase<LhsType, bLhsStrength>& Lhs, const TPtrVariantBase<RhsType, bRhsStrength>& Rhs)
 	{
 		return Lhs.GetRawPtrValue_Internal() == Rhs.GetRawPtrValue_Internal();
 	}
 
-	template <typename LhsType, bool bLhsStrength, typename RhsType, bool bRhsStrength, typename = TCanCompare<LhsType, RhsType>>
+	template <typename LhsType, bool bLhsStrength, typename RhsType, bool bRhsStrength, typename = typename TCanCompare<LhsType, RhsType>>
 	bool operator!=(const TPtrVariantBase<LhsType, bLhsStrength>& Lhs, const TPtrVariantBase<RhsType, bRhsStrength>& Rhs)
 	{
 		return !(Lhs == Rhs);
 	}
 
-	template<typename LhsType, typename RhsType, bool bRhsStrength, typename = TCanCompare<LhsType, RhsType>>
+	template<typename LhsType, typename RhsType, bool bRhsStrength, typename = typename TCanCompare<LhsType, RhsType>>
 	bool operator==(LhsType* Lhs, const TPtrVariantBase<RhsType, bRhsStrength>& Rhs)
 	{
 		return Lhs == Rhs.GetRawPtrValue_Internal();
 	}
 
-	template<typename LhsType, typename RhsType, bool bRhsStrength, typename = TCanCompare<LhsType, RhsType>>
+	template<typename LhsType, typename RhsType, bool bRhsStrength, typename = typename TCanCompare<LhsType, RhsType>>
 	bool operator!=(LhsType* Lhs, const TPtrVariantBase<RhsType, bRhsStrength>& Rhs)
 	{
 		return !(Lhs == Rhs);
 	}
 
-	template<typename LhsType, bool bLhsStrength, typename RhsType, typename = TCanCompare<LhsType, RhsType>>
+	template<typename LhsType, bool bLhsStrength, typename RhsType, typename = typename TCanCompare<LhsType, RhsType>>
 	bool operator==(const TPtrVariantBase<LhsType, bLhsStrength>& Lhs, RhsType* Rhs)
 	{
 		return Lhs.GetRawPtrValue_Internal() == Rhs;
 	}
 
-	template<typename LhsType, bool bLhsStrength, typename RhsType, typename = TCanCompare<LhsType, RhsType>>
+	template<typename LhsType, bool bLhsStrength, typename RhsType, typename = typename TCanCompare<LhsType, RhsType>>
 	bool operator!=(const TPtrVariantBase<LhsType, bLhsStrength>& Lhs, RhsType* Rhs)
 	{
 		return !(Lhs == Rhs);
