FTextKey internal storage can now split GUIDs into a separate FGuid pool rather than store them as a string

This is controlled by UE_TEXTKEY_SPLIT_GUID, and is enabled by default

This can save memory (at the cost additional parsing), as a FGuid is smaller than its string form.

diff --git a/Engine/Source/Runtime/Core/Private/Internationalization/TextKey.cpp b/Engine/Source/Runtime/Core/Private/Internationalization/TextKey.cpp
--- a/Engine/Source/Runtime/Core/Private/Internationalization/TextKey.cpp
+++ b/Engine/Source/Runtime/Core/Private/Internationalization/TextKey.cpp
@@ -23,10 +23,18 @@ DEFINE_LOG_CATEGORY_STATIC(LogTextKey, Log, All);
 	#define UE_TEXTKEY_USE_UTF8 (1)
 #endif
 
+#ifndef UE_TEXTKEY_SPLIT_GUID
+	#define UE_TEXTKEY_SPLIT_GUID (1)
+#endif
+
 #ifndef UE_TEXTKEY_SLAB_ALLOCATOR_SLAB_SIZE
 	#define UE_TEXTKEY_SLAB_ALLOCATOR_SLAB_SIZE (32768)
 #endif
 
+#ifndef UE_TEXTKEY_GUID_CHUNK_SIZE_BYTES
+	#define UE_TEXTKEY_GUID_CHUNK_SIZE_BYTES (32768)
+#endif
+
 #ifndef UE_TEXTKEY_ELEMENTS_CHUNK_SIZE_BYTES
 	#define UE_TEXTKEY_ELEMENTS_CHUNK_SIZE_BYTES (32768)
 #endif
@@ -116,7 +124,16 @@ public:
 			const TCHAR* AddedString = nullptr;
 			UE_AUTORTFM_OPEN
 			{
-				AddedString = *LegacyTCHARState.Add(InIndex, FString(KeyData.ToView()));
+				FString KeyString;
+				if (KeyData.IsStringView())
+				{
+					KeyString += KeyData.AsStringView();
+				}
+				else
+				{
+					KeyData.AsGuid().AppendString(KeyString, EGuidFormats::Digits);
+				}
+				AddedString = *LegacyTCHARState.Add(InIndex, MoveTemp(KeyString));
 			};
 			return AddedString;
 		}
@@ -129,7 +146,14 @@ public:
 		FReadScopeLock ScopeLock(DataRW);
 
 		const FKeyData& KeyData = KeyDataAllocations.Get(InIndex);
-		Out += KeyData.ToView();
+		if (KeyData.IsStringView())
+		{
+			Out += KeyData.AsStringView();
+		}
+		else
+		{
+			KeyData.AsGuid().AppendString(Out, EGuidFormats::Digits);
+		}
 	}
 
 	void AppendStringByIndex(const int32 InIndex, FStringBuilderBase& Out) const
@@ -139,7 +163,14 @@ public:
 		FReadScopeLock ScopeLock(DataRW);
 
 		const FKeyData& KeyData = KeyDataAllocations.Get(InIndex);
-		Out += KeyData.ToView();
+		if (KeyData.IsStringView())
+		{
+			Out += KeyData.AsStringView();
+		}
+		else
+		{
+			KeyData.AsGuid().AppendString(Out, EGuidFormats::Digits);
+		}
 	}
 
 	uint32 GetHashByIndex(const int32 InIndex) const
@@ -172,9 +203,9 @@ private:
 	{
 		FKeyData()
 #if UE_TEXTKEY_USE_UTF8
-			: Str(UTF8TEXT(""))
+			: DataPtr(UTF8TEXT(""))
 #else
-			: Str(TEXT(""))
+			: DataPtr(TEXT(""))
 #endif
 			, StrLen(0)
 			, StrHash(0)
@@ -182,34 +213,61 @@ private:
 		}
 
 		FKeyData(FTextKeyStringView InStr, const uint32 InStrHash)
-			: Str(InStr.GetData())
+			: DataPtr(InStr.GetData())
 			, StrLen(InStr.Len())
 			, StrHash(InStrHash)
 		{
 		}
 
-		FKeyData(FTextKeyStringView InStr, const FKeyData& InOther)
-			: Str(InStr.GetData())
-			, StrLen(InStr.Len())
-			, StrHash(InOther.StrHash)
+		FKeyData(const FGuid& InGuid, const uint32 InStrHash)
+			: DataPtr(&InGuid)
+			, StrLen(INDEX_NONE)
+			, StrHash(InStrHash)
+		{
+		}
+
+		bool IsStringView() const
 		{
+			return StrLen >= 0;
 		}
 
-		FTextKeyStringView ToView() const
+		FTextKeyStringView AsStringView() const
 		{
-			return FTextKeyStringView(Str, StrLen);
+			checkf(IsStringView(), TEXT("AsStringView called on a FKeyData that doesn't reference a string!"));
+			return FTextKeyStringView(static_cast<const FTextKeyCharType*>(DataPtr), StrLen);
+		}
+
+		bool IsGuid() const
+		{
+			return StrLen == INDEX_NONE;
+		}
+
+		const FGuid& AsGuid() const
+		{
+			checkf(IsGuid(), TEXT("AsGuid called on a FKeyData that doesn't reference a GUID!"));
+			return *static_cast<const FGuid*>(DataPtr);
 		}
 
 		friend FORCEINLINE bool operator==(const FKeyData& A, const FKeyData& B)
 		{
-			// We can use Memcmp here as we know we're comparing two blocks of the same size and don't care about lexical ordering
-			return A.StrLen == B.StrLen && FMemory::Memcmp(A.Str, B.Str, A.StrLen * sizeof(FTextKeyCharType)) == 0;
+			if (A.StrLen == B.StrLen)
+			{
+				if (A.StrLen >= 0)
+				{
+					// We can use Memcmp here as we know we're comparing two blocks of the same size and don't care about lexical ordering
+					return FMemory::Memcmp(A.DataPtr, B.DataPtr, A.StrLen * sizeof(FTextKeyCharType)) == 0;
+				}
+				else
+				{
+					return *static_cast<const FGuid*>(A.DataPtr) == *static_cast<const FGuid*>(B.DataPtr);
+				}
+			}
+			return false;
 		}
 
 		friend FORCEINLINE bool operator!=(const FKeyData& A, const FKeyData& B)
 		{
-			// We can use Memcmp here as we know we're comparing two blocks of the same size and don't care about lexical ordering
-			return A.StrLen != B.StrLen || FMemory::Memcmp(A.Str, B.Str, A.StrLen * sizeof(FTextKeyCharType)) != 0;
+			return !(A == B);
 		}
 
 		friend FORCEINLINE uint32 GetTypeHash(const FKeyData& A)
@@ -217,26 +275,42 @@ private:
 			return A.StrHash;
 		}
 
-		const FTextKeyCharType* Str; // Note: May not be null-terminated; use ToView() to get a valid TStringView
+		/** Internal data pointer; if StrLen >= 0 this points to a string buffer (which may not be null-terminated), or if INDEX_NONE points to a FGuid */
+		const void* DataPtr;
+		/** Length of the string (in characters) if DataPtr points to a string buffer, or INDEX_NONE if it points to a FGuid */
 		int32 StrLen;
+		/** Hash of the source string this data was created from */
 		uint32 StrHash;
 	};
 
 	void FindOrAddImpl(FStringView InStr, const uint32 InStrHash, FTextKey& OutTextKey)
 	{
-		auto ConvertedString = StrCast<FTextKeyCharType>(InStr.GetData(), InStr.Len());
-		const FKeyData KeyData = FKeyData(FTextKeyStringView(ConvertedString.Get(), ConvertedString.Length()), InStrHash);
-
-		int32 Index = FindOrAddString(KeyData);
+		int32 Index = INDEX_NONE;
+		{
+#if UE_TEXTKEY_SPLIT_GUID
+			FGuid KeyGuid;
+			if (FGuid::ParseExact(InStr, EGuidFormats::Digits, KeyGuid)) // Only checking for EGuidFormats::Digits as that's the default of FGuid::ToString() as used by FText keys
+			{
+				const FKeyData KeyData = FKeyData(KeyGuid, InStrHash);
+				Index = FindOrAddData(KeyData);
+			}
+			else
+#endif
+			{
+				auto ConvertedString = StrCast<FTextKeyCharType>(InStr.GetData(), InStr.Len());
+				const FKeyData KeyData = FKeyData(FTextKeyStringView(ConvertedString.Get(), ConvertedString.Length()), InStrHash);
+				Index = FindOrAddData(KeyData);
+			}
+		}
 		check(Index != INDEX_NONE);
 
 		OutTextKey.Index = Index;
 #if UE_TEXTKEY_STORE_EMBEDDED_HASH
-		OutTextKey.StrHash = KeyData.StrHash;
+		OutTextKey.StrHash = InStrHash;
 #endif
 	}
 
-	int32 FindOrAddString(const FKeyData& KeyData)
+	int32 FindOrAddData(const FKeyData& KeyData)
 	{
 		// Read-only
 		int32 NumElementsOnRead = 0;
@@ -267,8 +341,16 @@ private:
 			}
 
 			LLM_SCOPE_BYNAME(TEXT("Localization/TextKeys"));
-			FTextKeyStringView NewStr = StringAllocations.Add(KeyData.ToView());
-			return KeyDataAllocations.Add(FKeyData(NewStr, KeyData));
+			if (KeyData.IsStringView())
+			{
+				FTextKeyStringView NewStr = StringAllocations.Add(KeyData.AsStringView());
+				return KeyDataAllocations.Add(FKeyData(NewStr, KeyData.StrHash));
+			}
+			else
+			{
+				const FGuid& NewGuid = GuidAllocations[GuidAllocations.AddElement(KeyData.AsGuid())];
+				return KeyDataAllocations.Add(FKeyData(NewGuid, KeyData.StrHash));
+			}
 		}
 	}
 
@@ -462,6 +544,7 @@ private:
 
 	mutable FRWLock DataRW;
 	FStringAllocator StringAllocations;
+	TChunkedArray<FGuid, UE_TEXTKEY_GUID_CHUNK_SIZE_BYTES> GuidAllocations;
 	FKeyDataAllocator KeyDataAllocations;
 
 	// Sparse TCHAR state; built on-demand by anything still using the deprecated FTextKey::GetChars function
