print address when catching a SIMD SIGSEGV on Linux to help identify stomps

diff --git a/Engine/Source/Runtime/Core/Private/Unix/UnixPlatformCrashContext.cpp b/Engine/Source/Runtime/Core/Private/Unix/UnixPlatformCrashContext.cpp
--- a/Engine/Source/Runtime/Core/Private/Unix/UnixPlatformCrashContext.cpp
+++ b/Engine/Source/Runtime/Core/Private/Unix/UnixPlatformCrashContext.cpp
@@ -53,7 +53,7 @@ FString DescribeSignal(int32 Signal, siginfo_t* Info, ucontext_t *Context)
 #ifdef __x86_64__	// architecture-specific
 		if (Context && (Context->uc_mcontext.gregs[REG_TRAPNO] == 13))
 		{
-			ErrorString += FString::Printf(TEXT("SIGSEGV: unaligned memory access (SIMD vectors?)"));
+			ErrorString += FString::Printf(TEXT("SIGSEGV: unaligned memory access at address 0x%016llx (SIMD vectors?)"), (uint64)Info->si_addr);
 		}
 		else
 		{
