Add support for an unlimited number of cores in FWindowsPlatformProcess::GetPerFrameProcessorUsage.

The previous code was hardcoded to support a maximum of 128 cores and would assert if more cores were detected.

The problem was reported on UDN.

diff --git a/Engine/Source/Runtime/Core/Private/Windows/WindowsPlatformProcess.cpp b/Engine/Source/Runtime/Core/Private/Windows/WindowsPlatformProcess.cpp
--- a/Engine/Source/Runtime/Core/Private/Windows/WindowsPlatformProcess.cpp
+++ b/Engine/Source/Runtime/Core/Private/Windows/WindowsPlatformProcess.cpp
@@ -779,10 +779,16 @@ bool FWindowsPlatformProcess::GetPerFrameProcessorUsage(uint32 ProcessId, float&
 			LastProcessTime = (double)DeltaProcessCycleTime / DeltaCyclesPerFrame;
 
 			// Idle cycles are stored per core and flipped to allow per-frame calculation
-			const uint32 BufferLength = 1024;
-			check(BufferLength >= NumCores * 8);
+			const uint32 BufferLength = NumCores * 8;
+			static uint64* IdleCycleTimeBuffers[2] = { nullptr };
+
+			if(IdleCycleTimeBuffers[0] == nullptr)
+			{
+				//Alloc buffers on first frame 
+				IdleCycleTimeBuffers[0] = new uint64[BufferLength];
+				IdleCycleTimeBuffers[1] = new uint64[BufferLength];
+			}
 
-			static uint64 IdleCycleTimeBuffers[2][BufferLength] = {{0}};
 			uint64* IdleCycleTime = IdleCycleTimeBuffers[CurrFrameIndex];
 			uint64* PrevIdleCycleTime = IdleCycleTimeBuffers[PrevFrameIndex];
 
