#UE Improve messaging in assert

diff --git a/Engine/Source/Runtime/Core/Private/Compression/OodleDataCompressionFormat.cpp b/Engine/Source/Runtime/Core/Private/Compression/OodleDataCompressionFormat.cpp
--- a/Engine/Source/Runtime/Core/Private/Compression/OodleDataCompressionFormat.cpp
+++ b/Engine/Source/Runtime/Core/Private/Compression/OodleDataCompressionFormat.cpp
@@ -213,7 +213,7 @@ struct FOodleDataCompressionFormat : ICompressionFormat
 		{
 			return false;
 		}
-		check(CompressedBufferSize >= CheckCompressedBufferSize);
+		checkf(CompressedBufferSize >= CheckCompressedBufferSize, TEXT("The compressed buffer size (%ld) must be larger than %ld"), CompressedBufferSize, CheckCompressedBufferSize);
 
 		if ( Flags & COMPRESS_ForPackaging )
 		{
