Add __xstat and __fxstat as they can be inlined at times from stat/fstat

diff --git a/Engine/Source/Runtime/Core/Private/AutoRTFM/StdLibSupport.cpp b/Engine/Source/Runtime/Core/Private/AutoRTFM/StdLibSupport.cpp
--- a/Engine/Source/Runtime/Core/Private/AutoRTFM/StdLibSupport.cpp
+++ b/Engine/Source/Runtime/Core/Private/AutoRTFM/StdLibSupport.cpp
@@ -700,6 +700,24 @@ int RTFM_fstat(int Fd, struct stat* StatBuf) throw()
 	return fstat(Fd, StatBuf);
 }
 UE_AUTORTFM_REGISTER_OPEN_FUNCTION(fstat);
+
+int RTFM___xstat(int Ver, const char* Path, struct stat* StatBuf) throw()
+{
+	FContext* Context = FContext::Get();
+	Context->RecordWrite(StatBuf, sizeof(*StatBuf));
+
+	return __xstat(Ver, Path, StatBuf);
+}
+UE_AUTORTFM_REGISTER_OPEN_FUNCTION(__xstat);
+
+int RTFM___fxstat(int Ver, int Fd, struct stat* StatBuf) throw()
+{
+	FContext* Context = FContext::Get();
+	Context->RecordWrite(StatBuf, sizeof(*StatBuf));
+
+	return __fxstat(Ver, Fd, StatBuf);
+}
+UE_AUTORTFM_REGISTER_OPEN_FUNCTION(__fxstat);
 #endif // PLATFORM_LINUX
 
 wchar_t* RTFM_wcscpy(wchar_t* Dst, const wchar_t* Src) throw()
