Add some more CSV metadata for engine tests.

- Exposes RHIName that is currently game-specific, and also some output res metadata (partially addressing UE-221234).

diff --git a/Engine/Source/Runtime/Core/Private/ProfilingDebugging/CsvProfiler.cpp b/Engine/Source/Runtime/Core/Private/ProfilingDebugging/CsvProfiler.cpp
--- a/Engine/Source/Runtime/Core/Private/ProfilingDebugging/CsvProfiler.cpp
+++ b/Engine/Source/Runtime/Core/Private/ProfilingDebugging/CsvProfiler.cpp
@@ -3134,6 +3134,13 @@ FCsvProfiler::FCsvProfiler()
 	SetMetadataInternal(TEXT("LTOEnabled"), PLATFORM_COMPILER_OPTIMIZATION_LTCG ? TEXT("1") : TEXT("0"));//True if Link Time Optimisation is enabled 
 	SetMetadataInternal(TEXT("ASan"), USING_ADDRESS_SANITISER ? TEXT("1") : TEXT("0"));
 
+	FCoreDelegates::OnSystemResolutionChanged.AddLambda([](uint32 ResX, uint32 ResY)
+		{
+			SetMetadata(TEXT("SystemResolution.ResX"), *LexToString(ResX));
+			SetMetadata(TEXT("SystemResolution.ResY"), *LexToString(ResY));
+		}
+	);
+
 #if !UE_BUILD_SHIPPING
 	// for privacy, personal and free text fields are not allowed in shipping
 	SetMetadataInternal(TEXT("Commandline"), *CommandlineStr, false);
diff --git a/Engine/Source/Runtime/RHI/Private/DynamicRHI.cpp b/Engine/Source/Runtime/RHI/Private/DynamicRHI.cpp
--- a/Engine/Source/Runtime/RHI/Private/DynamicRHI.cpp
+++ b/Engine/Source/Runtime/RHI/Private/DynamicRHI.cpp
@@ -408,6 +408,14 @@ void RHIInit(bool bHasEditorToken)
 		}
 	}
 #endif
+
+	// Set CSV profiler metadata so every capture can see this
+	// GDynamicRHI is checked above
+	CSV_METADATA(TEXT("VerbatimRHIName"), GDynamicRHI->GetName());
+	// there's already an unfortunate precedent to concoct RHIName for PCD3D_ES31 and such platforms
+	CSV_METADATA(TEXT("RHIName"), GMaxRHIFeatureLevel == ERHIFeatureLevel::ES3_1 ? *(FString(GDynamicRHI->GetName()) + TEXT("_ES31")) : GDynamicRHI->GetName());
+	CSV_METADATA(TEXT("RHIFeatureLevel"), *LexToString(GMaxRHIFeatureLevel));
+	CSV_METADATA(TEXT("ShaderPlatform"), *LexToString(GetFeatureLevelShaderPlatform(GMaxRHIFeatureLevel)));
 }
 
 void RHIPostInit(const TArray<uint32>& InPixelFormatByteWidth)
