[TSAN]
  - Avoid overwriting data if they are the same values so that we can see real thread-safety issue if they ever occur

#rnx

diff --git a/Engine/Source/Runtime/Core/Public/Containers/LinkedListBuilder.h b/Engine/Source/Runtime/Core/Public/Containers/LinkedListBuilder.h
--- a/Engine/Source/Runtime/Core/Public/Containers/LinkedListBuilder.h
+++ b/Engine/Source/Runtime/Core/Public/Containers/LinkedListBuilder.h
@@ -39,6 +39,18 @@ public:
 
 	UE_NONCOPYABLE(TLinkedListBuilder);
 
+private:
+	FORCEINLINE void WriteEndPtr(InElementType* NewValue)
+	{
+		// Do not overwrite the same value to avoid dirtying the cache and
+		// also prevent TSAN from thinking we are messing around with existing data.
+		if (*EndPtr != NewValue)
+		{
+			*EndPtr = NewValue;
+		}
+	}
+public:
+
 	explicit TLinkedListBuilder(ElementType** ListStartPtr) :
 		StartPtr(ListStartPtr),
 		EndPtr(ListStartPtr)
@@ -62,7 +74,7 @@ public:
 	// Append element, don't touch next link
 	FORCEINLINE void AppendNoTerminate(ElementType& Element)
 	{
-		*EndPtr = &Element;
+		WriteEndPtr(&Element);
 		EndPtr = LinkAccessor::GetNextPtr(Element);
 	}
 
@@ -93,7 +105,7 @@ public:
 	// Mark end of the list
 	FORCEINLINE void NullTerminate()
 	{
-		*EndPtr = nullptr;
+		WriteEndPtr(nullptr);
 	}
 
 	FORCEINLINE void MoveToEnd()
