Fix TVOS build issues.

diff --git a/Engine/Source/Runtime/Core/Private/IOS/IOSPlatformMisc.cpp b/Engine/Source/Runtime/Core/Private/IOS/IOSPlatformMisc.cpp
--- a/Engine/Source/Runtime/Core/Private/IOS/IOSPlatformMisc.cpp
+++ b/Engine/Source/Runtime/Core/Private/IOS/IOSPlatformMisc.cpp
@@ -1846,7 +1846,7 @@ static TMap<FString, FString> GetPredefinedConfigRuleVars()
 	PredefinedConfigRuleVars.Add(FString(TEXT("versionName")), ProjectVersion);
 	
 	PredefinedConfigRuleVars.Add(FString(TEXT("processorCount")), FString::FromInt(FIOSPlatformMisc::NumberOfCores()));
-	
+#if !PLATFORM_TVOS && !PLATFORM_VISIONOS
 	id<MTLDevice> DefaultMetalDevice = MTLCreateSystemDefaultDevice();
 	if(DefaultMetalDevice)
 	{
@@ -1871,6 +1871,7 @@ static TMap<FString, FString> GetPredefinedConfigRuleVars()
 		}
 	}
 	[DefaultMetalDevice release];
+#endif
 	return PredefinedConfigRuleVars;
 }
 
