Enhance file mapping API with an option to create writeable memory mapped files

diff --git a/Engine/Source/Runtime/Core/Private/Android/AndroidPlatformFile.cpp b/Engine/Source/Runtime/Core/Private/Android/AndroidPlatformFile.cpp
--- a/Engine/Source/Runtime/Core/Private/Android/AndroidPlatformFile.cpp
+++ b/Engine/Source/Runtime/Core/Private/Android/AndroidPlatformFile.cpp
@@ -1111,7 +1111,7 @@ public:
 		close(FileHandle);
 	}
 
-	virtual IMappedFileRegion* MapRegion(int64 Offset = 0, int64 BytesToMap = MAX_int64, bool bPreloadHint = false) override
+	virtual IMappedFileRegion* MapRegion(int64 Offset = 0, int64 BytesToMap = MAX_int64, FFileMappingFlags InFlags = EMappedFileFlags::ENone) override
 	{
 		LLM_PLATFORM_SCOPE(ELLMTag::PlatformMMIO);
 		const int64 CurrentFileSize = GetCurrentFileSize();
@@ -1124,12 +1124,18 @@ public:
 		const int64 AlignedSize = Align(BytesToMap + Offset - AlignedOffset, FileMappingAlignment);
 
 		int Flags = MAP_PRIVATE;
-		if (bPreloadHint)
+		if (EnumHasAnyFlags(InFlags.Flags, EMappedFileFlags::EPreloadHint))
 		{
 			Flags |= MAP_POPULATE;
 		}
 
-		const uint8* AlignedMapPtr = static_cast<const uint8*>(mmap(nullptr, AlignedSize, PROT_READ, Flags, FileHandle, AlignedOffset));
+		int Protection = PROT_READ;
+		if (EnumHasAnyFlags(InFlags.Flags, EMappedFileFlags::EFileWritable))
+		{
+			Protection |= PROT_WRITE;
+		}
+
+		const uint8* AlignedMapPtr = static_cast<const uint8*>(mmap(nullptr, AlignedSize, Protection, Flags, FileHandle, AlignedOffset));
 		if (AlignedMapPtr == MAP_FAILED || AlignedMapPtr == nullptr)
 		{
 #if LOG_ANDROID_FILE
@@ -1441,7 +1447,7 @@ public:
 		return result;
 	}
 
-	virtual IMappedFileHandle* OpenMapped(const TCHAR* Filename) override
+	virtual IMappedFileHandle* OpenMapped(const TCHAR* Filename, EOpenReadFlags OpenOptions = EOpenReadFlags::None) override
 	{
 		FString LocalPath;
 		FString AssetPath;
@@ -1449,7 +1455,7 @@ public:
 
 		const FString NormalizedFilename = LocalPath;
 
-		constexpr int Flags = O_RDONLY;
+		const int Flags = EnumHasAnyFlags(OpenOptions, EOpenReadFlags::AllowWrite) ? O_RDWR : O_RDONLY;
 		const int32 Handle = open(TCHAR_TO_UTF8(*NormalizedFilename), Flags);
 		if (Handle == -1)
 		{
diff --git a/Engine/Source/Runtime/Core/Private/IOS/IOSPlatformFile.cpp b/Engine/Source/Runtime/Core/Private/IOS/IOSPlatformFile.cpp
--- a/Engine/Source/Runtime/Core/Private/IOS/IOSPlatformFile.cpp
+++ b/Engine/Source/Runtime/Core/Private/IOS/IOSPlatformFile.cpp
@@ -511,7 +511,7 @@ public:
 		close(FileHandle);
 	}
 	
-	virtual IMappedFileRegion* MapRegion(int64 Offset = 0, int64 BytesToMap = MAX_int64, bool bPreloadHint = false) override
+	virtual IMappedFileRegion* MapRegion(int64 Offset = 0, int64 BytesToMap = MAX_int64, FFileMappingFlags Flags = EMappedFileFlags::ENone) override
 	{
 		LLM_PLATFORM_SCOPE(ELLMTag::PlatformMMIO);
 		check(Offset < GetFileSize()); // don't map zero bytes and don't map off the end of the file
@@ -519,7 +519,7 @@ public:
 		check(BytesToMap > 0); // don't map zero bytes
 
 		
-		// const uint8* MapPtr = (const uint8 *)mmap(NULL, BytesToMap, PROT_READ, MAP_PRIVATE | (bPreloadHint ? MAP_POPULATE : 0), FileHandle, Offset);
+		// const uint8* MapPtr = (const uint8 *)mmap(NULL, BytesToMap, PROT_READ, MAP_PRIVATE | (EnumHasAnyFlags(Flags.Flags, EMappedFileFlags::EPreloadHint) ? MAP_POPULATE : 0), FileHandle, Offset);
 		// const uint8* MapPtr = (const uint8 *)mmap(NULL, BytesToMap, PROT_READ, MAP_PRIVATE, FileHandle, Offset);
 		//		const uint8* MapPtr = (const uint8 *)mmap(NULL, BytesToMap, PROT_READ, MAP_SHARED, FileHandle, Offset);
 		
@@ -527,7 +527,13 @@ public:
 		//File mapping can extend beyond file size. It's OK, kernel will just fill any leftover page data with zeros
 		const int64 AlignedSize = Align(BytesToMap + Offset - AlignedOffset, FileMappingAlignment);
 		
-		const uint8* AlignedMapPtr = (const uint8 *)mmap(NULL, AlignedSize, PROT_READ, MAP_PRIVATE, FileHandle, AlignedOffset);
+		int Protection = PROT_READ;
+		if (EnumHasAnyFlags(Flags.Flags, EMappedFileFlags::EFileWritable))
+		{
+			Protection |= PROT_WRITE;
+		}
+
+		const uint8* AlignedMapPtr = (const uint8 *)mmap(NULL, AlignedSize, Protection, MAP_PRIVATE, FileHandle, AlignedOffset);
 		if (AlignedMapPtr == (const uint8*)-1 || AlignedMapPtr == nullptr)
 		{
 			UE_LOG(LogIOS, Warning, TEXT("Failed to map memory %s, error is %d"), *Filename, errno);
@@ -904,25 +910,26 @@ IFileHandle* FIOSPlatformFile::OpenWrite(const TCHAR* Filename, bool bAppend, bo
 	return NULL;
 }
 
-IMappedFileHandle* FIOSPlatformFile::OpenMapped(const TCHAR* Filename)
+IMappedFileHandle* FIOSPlatformFile::OpenMapped(const TCHAR* Filename, EOpenReadFlags OpenOptions = EOpenReadFlags::None)
 {
 	FString NormalizedFilename = NormalizeFilename(Filename);
 	
 	// check the read path
 	FString FinalPath = ConvertToIOSPath(NormalizedFilename, false, false);
 	FILE* FP;
-	FP = fopen(TCHAR_TO_UTF8(*FinalPath), "r");
+	const char* OpenMode = EnumHasAnyFlags(OpenOptions, EOpenReadFlags::AllowWrite) ? "r+" : "r";
+	FP = fopen(TCHAR_TO_UTF8(*FinalPath), OpenMode);
 	if(FP == nullptr)
 	{
 		// if not in the read path, check the private write path
 		FinalPath = ConvertToIOSPath(NormalizedFilename, true, false);
-		FP = fopen(TCHAR_TO_UTF8(*FinalPath), "r");
+		FP = fopen(TCHAR_TO_UTF8(*FinalPath), OpenMode);
 
 		if(FP == nullptr)
 		{
 			// if not in the private write path, check the public write path
 			FinalPath = ConvertToIOSPath(NormalizedFilename, true, true);
-			FP = fopen(TCHAR_TO_UTF8(*FinalPath), "r");
+			FP = fopen(TCHAR_TO_UTF8(*FinalPath), OpenMode);
 		}
 	}
 	
diff --git a/Engine/Source/Runtime/Core/Private/Unix/UnixPlatformFile.cpp b/Engine/Source/Runtime/Core/Private/Unix/UnixPlatformFile.cpp
--- a/Engine/Source/Runtime/Core/Private/Unix/UnixPlatformFile.cpp
+++ b/Engine/Source/Runtime/Core/Private/Unix/UnixPlatformFile.cpp
@@ -808,7 +808,7 @@ public:
 		close(FileHandle);
 	}
 
-	virtual IMappedFileRegion* MapRegion(int64 Offset = 0, int64 BytesToMap = MAX_int64, bool bPreloadHint = false) override
+	virtual IMappedFileRegion* MapRegion(int64 Offset = 0, int64 BytesToMap = MAX_int64, FFileMappingFlags Flags = EMappedFileFlags::ENone) override
 	{
 		LLM_PLATFORM_SCOPE(ELLMTag::PlatformMMIO);
 		check(Offset < GetFileSize()); // don't map zero bytes and don't map off the end of the file
@@ -819,13 +819,19 @@ public:
 		//File mapping can extend beyond file size. It's OK, kernel will just fill any leftover page data with zeros
 		const int64 AlignedSize = Align(BytesToMap + Offset - AlignedOffset, FileMappingAlignment);
 
-		int Flags = MAP_PRIVATE;
-		if (bPreloadHint)
+		int InternalFlags = (EnumHasAnyFlags(Flags.Flags, EMappedFileFlags::EMapShared)) ? MAP_SHARED : MAP_PRIVATE;
+		if (EnumHasAnyFlags(Flags.Flags, EMappedFileFlags::EPreloadHint))
 		{
-			Flags |= MAP_POPULATE;
+			InternalFlags |= MAP_POPULATE;
 		}
 
-		const uint8* AlignedMapPtr = static_cast<const uint8*>(mmap(nullptr, AlignedSize, PROT_READ, Flags, FileHandle, AlignedOffset));
+		int Protection = PROT_READ;
+		if (EnumHasAnyFlags(Flags.Flags, EMappedFileFlags::EFileWritable))
+		{
+			Protection |= PROT_WRITE;
+		}
+
+		const uint8* AlignedMapPtr = static_cast<const uint8*>(mmap(nullptr, AlignedSize, Protection, InternalFlags, FileHandle, AlignedOffset));
 		if (AlignedMapPtr == MAP_FAILED || AlignedMapPtr == nullptr)
 		{
 			UE_LOG(LogUnixPlatformFile, Warning, TEXT("Failed to map memory %s, error is %d"), *Filename, errno);
@@ -1220,11 +1226,11 @@ IFileHandle* FUnixPlatformFile::OpenWrite(const TCHAR* Filename, bool bAppend, b
 	return nullptr;
 }
 
-IMappedFileHandle* FUnixPlatformFile::OpenMapped(const TCHAR* Filename)
+IMappedFileHandle* FUnixPlatformFile::OpenMapped(const TCHAR* Filename, EOpenReadFlags OpenOptions)
 {
 	const FString NormalizedFilename = NormalizeFilename(Filename, false);
 
-	constexpr int Flags = O_RDONLY;
+	const int Flags = EnumHasAnyFlags(OpenOptions, EOpenReadFlags::AllowWrite) ? O_RDWR : O_RDONLY;
 	const int32 Handle = open(TCHAR_TO_UTF8(*NormalizedFilename), Flags);
 	if (Handle == -1)
 	{
diff --git a/Engine/Source/Runtime/Core/Private/Windows/WindowsPlatformFile.cpp b/Engine/Source/Runtime/Core/Private/Windows/WindowsPlatformFile.cpp
--- a/Engine/Source/Runtime/Core/Private/Windows/WindowsPlatformFile.cpp
+++ b/Engine/Source/Runtime/Core/Private/Windows/WindowsPlatformFile.cpp
@@ -998,13 +998,17 @@ public:
 		(void)CloseResult;
 #endif
 	}
-	virtual IMappedFileRegion* MapRegion(int64 Offset = 0, int64 BytesToMap = MAX_int64, bool bPreloadHint = false) override
+	virtual IMappedFileRegion* MapRegion(int64 Offset = 0, int64 BytesToMap = MAX_int64, FFileMappingFlags Flags = EMappedFileFlags::ENone) override
 	{
 		check(Offset < GetFileSize()); // don't map zero bytes and don't map off the end of the file
 		BytesToMap = FMath::Min<int64>(BytesToMap, GetFileSize() - Offset);
 		check(BytesToMap > 0); // don't map zero bytes
 
 		DWORD OpenMappingAccess = FILE_MAP_READ;
+		if (EnumHasAnyFlags(Flags.Flags, EMappedFileFlags::EFileWritable))
+		{
+			OpenMappingAccess |= FILE_MAP_WRITE;
+		}
 
 		SIZE_T VirtualMemoryGranularity = FWindowsPlatformMemory::GetConstants().OsAllocationGranularity;
 		int64 AlignedOffset = AlignDown(Offset, VirtualMemoryGranularity);
@@ -1013,7 +1017,7 @@ public:
 		ULARGE_INTEGER LI;
 		LI.QuadPart = AlignedOffset;
 
-		const uint8* AlignedMapPtr = (const uint8*)MapViewOfFile(MappingHandle, FILE_MAP_READ, LI.HighPart, LI.LowPart, AlignedSize + AlignedOffset > GetFileSize() ? 0 : AlignedSize);
+		const uint8* AlignedMapPtr = (const uint8*)MapViewOfFile(MappingHandle, OpenMappingAccess, LI.HighPart, LI.LowPart, AlignedSize + AlignedOffset > GetFileSize() ? 0 : AlignedSize);
 		if (!AlignedMapPtr)
 		{
 			return nullptr;
@@ -1504,16 +1508,20 @@ IFileHandle* FWindowsPlatformFile::OpenWrite(const TCHAR* Filename, bool bAppend
 	}
 }
 
-IMappedFileHandle* FWindowsPlatformFile::OpenMapped(const TCHAR* Filename)
+IMappedFileHandle* FWindowsPlatformFile::OpenMapped(const TCHAR* Filename, EOpenReadFlags OpenOptions)
 {
 	int64 Size = FileSize(Filename);
 	if (Size < 1)
 	{
 		return nullptr;
 	}
-	uint32  Access = GENERIC_READ;
-	uint32  WinFlags = FILE_SHARE_READ;
-	uint32  Create = OPEN_EXISTING;
+	uint32 Access = GENERIC_READ;
+	if (EnumHasAnyFlags(OpenOptions, EOpenReadFlags::AllowWrite))
+	{
+		Access |= GENERIC_WRITE;
+	}
+	uint32 WinFlags = FILE_SHARE_READ;
+	uint32 Create = OPEN_EXISTING;
 	TRACE_PLATFORMFILE_BEGIN_OPEN(Filename);
 	HANDLE Handle = CreateFileW(*FNormalizedFilename(Filename), Access, WinFlags, NULL, Create, FILE_ATTRIBUTE_NORMAL, NULL);
 	if (Handle != INVALID_HANDLE_VALUE)
@@ -1525,7 +1533,7 @@ IMappedFileHandle* FWindowsPlatformFile::OpenMapped(const TCHAR* Filename)
 		TRACE_PLATFORMFILE_FAIL_OPEN(Filename);
 		return nullptr;
 	}
-	HANDLE MappingHandle = CreateFileMapping(Handle, NULL, PAGE_READONLY, 0, 0, NULL);
+	HANDLE MappingHandle = CreateFileMapping(Handle, NULL, EnumHasAnyFlags(OpenOptions, EOpenReadFlags::AllowWrite) ? PAGE_READWRITE : PAGE_READONLY, 0, 0, NULL);
 	if (MappingHandle == NULL)
 	{
 		TRACE_PLATFORMFILE_BEGIN_CLOSE(Handle);
diff --git a/Engine/Source/Runtime/Core/Private/Windows/WindowsPlatformFilePrivate.h b/Engine/Source/Runtime/Core/Private/Windows/WindowsPlatformFilePrivate.h
--- a/Engine/Source/Runtime/Core/Private/Windows/WindowsPlatformFilePrivate.h
+++ b/Engine/Source/Runtime/Core/Private/Windows/WindowsPlatformFilePrivate.h
@@ -68,7 +68,7 @@ public:
 	virtual IFileHandle* OpenRead(const TCHAR* Filename, bool bAllowWrite = false) override;
 	virtual IFileHandle* OpenReadNoBuffering(const TCHAR* Filename, bool bAllowWrite = false) override;
 	virtual IFileHandle* OpenWrite(const TCHAR* Filename, bool bAppend = false, bool bAllowRead = false) override;
-	virtual IMappedFileHandle* OpenMapped(const TCHAR* Filename) override;
+	virtual IMappedFileHandle* OpenMapped(const TCHAR* Filename, EOpenReadFlags OpenOptions = EOpenReadFlags::None) override;
 	virtual bool DirectoryExists(const TCHAR* Directory) override;
 	virtual bool CreateDirectory(const TCHAR* Directory) override;
 	virtual bool DeleteDirectory(const TCHAR* Directory) override;
diff --git a/Engine/Source/Runtime/Core/Public/Async/MappedFileHandle.h b/Engine/Source/Runtime/Core/Public/Async/MappedFileHandle.h
--- a/Engine/Source/Runtime/Core/Public/Async/MappedFileHandle.h
+++ b/Engine/Source/Runtime/Core/Public/Async/MappedFileHandle.h
@@ -81,6 +81,25 @@ public:
 	IMappedFileRegion& operator=(const IMappedFileRegion&) = delete;
 };
 
+enum class EMappedFileFlags
+{
+	ENone = 0,				// Do nothing
+	EPreloadHint = 1,		// Preload the data. This is only a hint and might be ignored, see IMappedFileRegion::PreloadHint
+	EFileWritable = 2,		// Make the mapped file writable. Requires OpenMapped to be called with OpenOptions set to EPlatformFileOpenOptions::AllowWrite
+	EMapShared = 4			// If content of mapping is changed, it will be shared with other processes that might have this mapping opened and\or written to disc
+};
+
+ENUM_CLASS_FLAGS(EMappedFileFlags);
+
+
+struct FFileMappingFlags
+{
+	FFileMappingFlags() : Flags(EMappedFileFlags::ENone) {}
+	FFileMappingFlags(bool bPreloadHint) : Flags(bPreloadHint ? EMappedFileFlags::EPreloadHint : EMappedFileFlags::ENone) {}
+	FFileMappingFlags(EMappedFileFlags InFlags) : Flags(InFlags) {}
+	EMappedFileFlags Flags;
+};
+
 class IMappedFileHandle
 {
 	size_t MappedFileSize;
@@ -109,10 +128,10 @@ public:
 	* Map a region of the file.
 	* @param Offset				Offset into the file to start mapping.
 	* @param BytesToMap			Number of bytes to map. Clamped to the size of the file.
-	* @param bPreloadHint		If true, preload the data. This is only a hint and might be ignored, see IMappedFileRegion::PreloadHint
+	* @param Flags				A combination of EMappingFlags values
 	* @return the mapped region interface. Returns nullptr on failure.
 	**/
-	virtual IMappedFileRegion* MapRegion(int64 Offset = 0, int64 BytesToMap = MAX_int64, bool bPreloadHint = false) = 0;
+	virtual IMappedFileRegion* MapRegion(int64 Offset = 0, int64 BytesToMap = MAX_int64, FFileMappingFlags Flags = EMappedFileFlags::ENone) = 0;
 
 	// Non-copyable
 	IMappedFileHandle(const IMappedFileHandle&) = delete;
diff --git a/Engine/Source/Runtime/Core/Public/GenericPlatform/GenericPlatformFile.h b/Engine/Source/Runtime/Core/Public/GenericPlatform/GenericPlatformFile.h
--- a/Engine/Source/Runtime/Core/Public/GenericPlatform/GenericPlatformFile.h
+++ b/Engine/Source/Runtime/Core/Public/GenericPlatform/GenericPlatformFile.h
@@ -741,12 +741,13 @@ public:
 
 	/** Open a file for async reading. This call does hit the disk; it is synchronous open. 
 	*
-	* @param Filename file to be mapped. This doesn't actually map anything, just opens the file.
+	* @param Filename		File to be mapped. This doesn't actually map anything, just opens the file.
+	* @param OpenOptions	Allows specifying map to be writable via AllowWrite or be ReadOnly by deafault.
 	* @return Close the file by delete'ing the handle. A non-null return value does mean the file exists. 
 	* Null can be returned for many reasons even if the file exists. Perhaps this platform does not support mapped files, or this file is compressed in a pak file.
 	* Generally you attempt to open mapped, and if that fails, then use other file operations instead.
 	*/
-	virtual IMappedFileHandle* OpenMapped(const TCHAR* Filename)
+	virtual IMappedFileHandle* OpenMapped(const TCHAR* Filename, EOpenReadFlags OpenOptions = EOpenReadFlags::None)
 	{
 		return nullptr;
 	}
@@ -839,7 +840,7 @@ public:
 	 * @param From				File to copy from.
 	 * @param ReadFlags			Source file read options.
 	 * @param WriteFlags		Destination file write options.
-	 * @return			true if the file was copied sucessfully.
+	 * @return					true if the file was copied successfully.
 	**/
 	CORE_API virtual bool CopyFile(const TCHAR* To, const TCHAR* From, EPlatformFileRead ReadFlags = EPlatformFileRead::None, EPlatformFileWrite WriteFlags = EPlatformFileWrite::None);
 
diff --git a/Engine/Source/Runtime/Core/Public/HAL/IPlatformFileCachedWrapper.h b/Engine/Source/Runtime/Core/Public/HAL/IPlatformFileCachedWrapper.h
--- a/Engine/Source/Runtime/Core/Public/HAL/IPlatformFileCachedWrapper.h
+++ b/Engine/Source/Runtime/Core/Public/HAL/IPlatformFileCachedWrapper.h
@@ -455,9 +455,9 @@ public:
 	{
 		return LowerLevel->OpenAsyncRead(Filename, bAllowWrite);
 	}
-	virtual IMappedFileHandle* OpenMapped(const TCHAR* Filename) override
+	virtual IMappedFileHandle* OpenMapped(const TCHAR* Filename, EOpenReadFlags OpenOptions = EOpenReadFlags::None) override
 	{
-		return LowerLevel->OpenMapped(Filename);
+		return LowerLevel->OpenMapped(Filename, OpenOptions);
 	}
 	virtual void SetAsyncMinimumPriority(EAsyncIOPriorityAndFlags MinPriority) override
 	{
diff --git a/Engine/Source/Runtime/Core/Public/HAL/IPlatformFileLogWrapper.h b/Engine/Source/Runtime/Core/Public/HAL/IPlatformFileLogWrapper.h
--- a/Engine/Source/Runtime/Core/Public/HAL/IPlatformFileLogWrapper.h
+++ b/Engine/Source/Runtime/Core/Public/HAL/IPlatformFileLogWrapper.h
@@ -503,9 +503,9 @@ public:
 		//@todo no wrapped logging for async file handles (yet)
 		return Result;
 	}
-	virtual IMappedFileHandle* OpenMapped(const TCHAR* Filename) override
+	virtual IMappedFileHandle* OpenMapped(const TCHAR* Filename, EOpenReadFlags OpenOptions = EOpenReadFlags::None) override
 	{
-		return LowerLevel->OpenMapped(Filename);
+		return LowerLevel->OpenMapped(Filename, OpenOptions);
 	}
 	virtual void SetAsyncMinimumPriority(EAsyncIOPriorityAndFlags Priority) override
 	{
diff --git a/Engine/Source/Runtime/Core/Public/HAL/IPlatformFileOpenLogWrapper.h b/Engine/Source/Runtime/Core/Public/HAL/IPlatformFileOpenLogWrapper.h
--- a/Engine/Source/Runtime/Core/Public/HAL/IPlatformFileOpenLogWrapper.h
+++ b/Engine/Source/Runtime/Core/Public/HAL/IPlatformFileOpenLogWrapper.h
@@ -223,9 +223,9 @@ public:
 		// we must not record the "open" here...what matters is when we start reading the file!
 		return new FLoggingAsyncReadFileHandle(this, Filename, LowerLevel->OpenAsyncRead(Filename, bAllowWrite));
 	}
-	virtual IMappedFileHandle* OpenMapped(const TCHAR* Filename) override
+	virtual IMappedFileHandle* OpenMapped(const TCHAR* Filename, EOpenReadFlags OpenOptions = EOpenReadFlags::None) override
 	{
-		return LowerLevel->OpenMapped(Filename);
+		return LowerLevel->OpenMapped(Filename, OpenOptions);
 	}
 
 	void AddToOpenLog(const TCHAR* Filename)
diff --git a/Engine/Source/Runtime/Core/Public/IOS/IOSPlatformFile.h b/Engine/Source/Runtime/Core/Public/IOS/IOSPlatformFile.h
--- a/Engine/Source/Runtime/Core/Public/IOS/IOSPlatformFile.h
+++ b/Engine/Source/Runtime/Core/Public/IOS/IOSPlatformFile.h
@@ -42,7 +42,7 @@ public:
 
 	virtual IFileHandle* OpenRead(const TCHAR* Filename, bool bAllowWrite = false) override;
 	virtual IFileHandle* OpenWrite(const TCHAR* Filename, bool bAppend = false, bool bAllowRead = false) override;
-	virtual IMappedFileHandle* OpenMapped(const TCHAR* Filename) override;
+	virtual IMappedFileHandle* OpenMapped(const TCHAR* Filename, EOpenReadFlags OpenOptions = EOpenReadFlags::None) override;
 
 	virtual bool DirectoryExists(const TCHAR* Directory) override;
 	virtual bool CreateDirectory(const TCHAR* Directory) override;
diff --git a/Engine/Source/Runtime/Core/Public/Unix/UnixPlatformFile.h b/Engine/Source/Runtime/Core/Public/Unix/UnixPlatformFile.h
--- a/Engine/Source/Runtime/Core/Public/Unix/UnixPlatformFile.h
+++ b/Engine/Source/Runtime/Core/Public/Unix/UnixPlatformFile.h
@@ -46,7 +46,7 @@ public:
 
 	CORE_API virtual IFileHandle* OpenRead(const TCHAR* Filename, bool bAllowWrite = false) override;
 	CORE_API virtual IFileHandle* OpenWrite(const TCHAR* Filename, bool bAppend = false, bool bAllowRead = false) override;
-	CORE_API virtual IMappedFileHandle* OpenMapped(const TCHAR* Filename) override;
+	CORE_API virtual IMappedFileHandle* OpenMapped(const TCHAR* Filename, EOpenReadFlags OpenOptions = EOpenReadFlags::None) override;
 
 	CORE_API virtual bool DirectoryExists(const TCHAR* Directory) override;
 	CORE_API virtual bool CreateDirectory(const TCHAR* Directory) override;
diff --git a/Engine/Source/Runtime/PakFile/Private/IPlatformFilePak.cpp b/Engine/Source/Runtime/PakFile/Private/IPlatformFilePak.cpp
--- a/Engine/Source/Runtime/PakFile/Private/IPlatformFilePak.cpp
+++ b/Engine/Source/Runtime/PakFile/Private/IPlatformFilePak.cpp
@@ -3923,7 +3923,7 @@ public:
 	{
 		// we don't own lower level, it is shared
 	}
-	virtual IMappedFileRegion* MapRegion(int64 Offset = 0, int64 BytesToMap = MAX_int64, bool bPreloadHint = false) override
+	virtual IMappedFileRegion* MapRegion(int64 Offset = 0, int64 BytesToMap = MAX_int64, FFileMappingFlags Flags = EMappedFileFlags::ENone) override
 	{
 		//check(Offset + OffsetInPak < PakSize); // don't map zero bytes and don't map off the end of the (real) file
 		check(Offset < GetFileSize()); // don't map zero bytes and don't map off the end of the (virtual) file
@@ -3931,7 +3931,7 @@ public:
 		check(BytesToMap > 0); // don't map zero bytes
 		//check(Offset + BytesToMap <= GetFileSize()); // don't map zero bytes and don't map off the end of the (virtual) file
 		//check(Offset + OffsetInPak + BytesToMap <= PakSize); // don't map zero bytes and don't map off the end of the (real) file
-		return LowerLevel->MapRegion(Offset + OffsetInPak, BytesToMap, bPreloadHint);
+		return LowerLevel->MapRegion(Offset + OffsetInPak, BytesToMap, Flags);
 	}
 };
 
@@ -4590,7 +4590,7 @@ bool FPakPlatformFile::IsPakFileInstalled(const FString& InFilename)
 }
 #endif //HAS_PLATFORM_PAK_INSTALL_CHECK
 
-IMappedFileHandle* FPakPlatformFile::OpenMapped(const TCHAR* Filename)
+IMappedFileHandle* FPakPlatformFile::OpenMapped(const TCHAR* Filename, EOpenReadFlags OpenOptions)
 {
 	if (!GMMIO_Enable)
 	{
@@ -4600,7 +4600,7 @@ IMappedFileHandle* FPakPlatformFile::OpenMapped(const TCHAR* Filename)
 #if !UE_BUILD_SHIPPING
 	if (bLookLooseFirst && IsNonPakFilenameAllowed(Filename))
 	{
-		IMappedFileHandle* Handle = LowerLevel->OpenMapped(Filename);
+		IMappedFileHandle* Handle = LowerLevel->OpenMapped(Filename, OpenOptions);
 		if (Handle != nullptr)
 		{
 			return Handle;
@@ -4621,7 +4621,7 @@ IMappedFileHandle* FPakPlatformFile::OpenMapped(const TCHAR* Filename)
 		FScopeLock Lock(&PakEntry->MappedFileHandleCriticalSection);
 		if (!PakEntry->MappedFileHandle)
 		{
-			PakEntry->MappedFileHandle = LowerLevel->OpenMapped(*PakEntry->GetFilename());
+			PakEntry->MappedFileHandle = LowerLevel->OpenMapped(*PakEntry->GetFilename(), OpenOptions);
 		}
 		if (!PakEntry->MappedFileHandle)
 		{
@@ -4631,7 +4631,7 @@ IMappedFileHandle* FPakPlatformFile::OpenMapped(const TCHAR* Filename)
 	}
 	if (IsNonPakFilenameAllowed(Filename))
 	{
-		return LowerLevel->OpenMapped(Filename);
+		return LowerLevel->OpenMapped(Filename, OpenOptions);
 	}
 	return nullptr;
 }
diff --git a/Engine/Source/Runtime/PakFile/Private/IoDispatcherFileBackend.cpp b/Engine/Source/Runtime/PakFile/Private/IoDispatcherFileBackend.cpp
--- a/Engine/Source/Runtime/PakFile/Private/IoDispatcherFileBackend.cpp
+++ b/Engine/Source/Runtime/PakFile/Private/IoDispatcherFileBackend.cpp
@@ -133,9 +133,9 @@ public:
 
 	virtual ~FMappedFileProxy() { }
 
-	virtual IMappedFileRegion* MapRegion(int64 Offset = 0, int64 BytesToMap = MAX_int64, bool bPreloadHint = false) override
+	virtual IMappedFileRegion* MapRegion(int64 Offset = 0, int64 BytesToMap = MAX_int64, FFileMappingFlags Flags = EMappedFileFlags::ENone) override
 	{
-		return SharedMappedFileHandle != nullptr ? SharedMappedFileHandle->MapRegion(Offset, BytesToMap, bPreloadHint) : nullptr;
+		return SharedMappedFileHandle != nullptr ? SharedMappedFileHandle->MapRegion(Offset, BytesToMap, Flags) : nullptr;
 	}
 private:
 	IMappedFileHandle* SharedMappedFileHandle;
diff --git a/Engine/Source/Runtime/PakFile/Public/IPlatformFilePak.h b/Engine/Source/Runtime/PakFile/Public/IPlatformFilePak.h
--- a/Engine/Source/Runtime/PakFile/Public/IPlatformFilePak.h
+++ b/Engine/Source/Runtime/PakFile/Public/IPlatformFilePak.h
@@ -2504,7 +2504,7 @@ public:
 	PAKFILE_API virtual IAsyncReadFileHandle* OpenAsyncRead(const TCHAR* Filename, bool bAllowWrite = false) override;
 	PAKFILE_API virtual void SetAsyncMinimumPriority(EAsyncIOPriorityAndFlags Priority) override;
 
-	virtual IMappedFileHandle* OpenMapped(const TCHAR* Filename) override;
+	virtual IMappedFileHandle* OpenMapped(const TCHAR* Filename, EOpenReadFlags OpenOptions = EOpenReadFlags::None) override;
 	/**
 	 * Converts a filename to a path inside pak file.
 	 *
diff --git a/Engine/Source/Runtime/SandboxFile/Private/IPlatformFileSandboxWrapper.cpp b/Engine/Source/Runtime/SandboxFile/Private/IPlatformFileSandboxWrapper.cpp
--- a/Engine/Source/Runtime/SandboxFile/Private/IPlatformFileSandboxWrapper.cpp
+++ b/Engine/Source/Runtime/SandboxFile/Private/IPlatformFileSandboxWrapper.cpp
@@ -983,14 +983,14 @@ void FSandboxPlatformFile::SetAsyncMinimumPriority(EAsyncIOPriorityAndFlags Prio
 {
 	LowerLevel->SetAsyncMinimumPriority(Priority);
 }
-IMappedFileHandle* FSandboxPlatformFile::OpenMapped(const TCHAR* Filename)
+IMappedFileHandle* FSandboxPlatformFile::OpenMapped(const TCHAR* Filename, EOpenReadFlags OpenOptions)
 {
 	FString UserFilename(*ConvertToSandboxPath(Filename));
 	if (!OkForInnerAccess(Filename) || LowerLevel->FileExists(*UserFilename))
 	{
-		return LowerLevel->OpenMapped(*UserFilename);
+		return LowerLevel->OpenMapped(*UserFilename, OpenOptions);
 	}
-	return LowerLevel->OpenMapped(Filename);
+	return LowerLevel->OpenMapped(Filename, OpenOptions);
 }
 
 /**
diff --git a/Engine/Source/Runtime/SandboxFile/Public/IPlatformFileSandboxWrapper.h b/Engine/Source/Runtime/SandboxFile/Public/IPlatformFileSandboxWrapper.h
--- a/Engine/Source/Runtime/SandboxFile/Public/IPlatformFileSandboxWrapper.h
+++ b/Engine/Source/Runtime/SandboxFile/Public/IPlatformFileSandboxWrapper.h
@@ -223,7 +223,7 @@ public:
 
 	SANDBOXFILE_API virtual IAsyncReadFileHandle*	OpenAsyncRead(const TCHAR* Filename, bool bAllowWrite = false) override;
 	SANDBOXFILE_API virtual void					SetAsyncMinimumPriority(EAsyncIOPriorityAndFlags Priority) override;
-	SANDBOXFILE_API virtual IMappedFileHandle*		OpenMapped(const TCHAR* Filename) override;
+	SANDBOXFILE_API virtual IMappedFileHandle*		OpenMapped(const TCHAR* Filename, EOpenReadFlags OpenOptions = EOpenReadFlags::None) override;
 
 	friend class FSandboxVisitor;
 	friend class FSandboxStatVisitor;
diff --git a/Engine/Source/Runtime/StorageServerClient/Private/StorageServerPlatformFile.cpp b/Engine/Source/Runtime/StorageServerClient/Private/StorageServerPlatformFile.cpp
--- a/Engine/Source/Runtime/StorageServerClient/Private/StorageServerPlatformFile.cpp
+++ b/Engine/Source/Runtime/StorageServerClient/Private/StorageServerPlatformFile.cpp
@@ -1096,9 +1096,9 @@ bool FStorageServerPlatformFile::IterateDirectoryStat(const TCHAR* Directory, FD
 	return bResult;
 }
 
-IMappedFileHandle* FStorageServerPlatformFile::OpenMapped(const TCHAR* Filename)
+IMappedFileHandle* FStorageServerPlatformFile::OpenMapped(const TCHAR* Filename, EOpenReadFlags OpenOptions)
 {
-	return (LowerLevel && IsNonServerFilenameAllowed(Filename)) ? LowerLevel->OpenMapped(Filename) : nullptr;
+	return (LowerLevel && IsNonServerFilenameAllowed(Filename)) ? LowerLevel->OpenMapped(Filename, OpenOptions) : nullptr;
 }
 
 bool FStorageServerPlatformFile::DirectoryExists(const TCHAR* Directory)
diff --git a/Engine/Source/Runtime/StorageServerClient/Private/StorageServerPlatformFile.h b/Engine/Source/Runtime/StorageServerClient/Private/StorageServerPlatformFile.h
--- a/Engine/Source/Runtime/StorageServerClient/Private/StorageServerPlatformFile.h
+++ b/Engine/Source/Runtime/StorageServerClient/Private/StorageServerPlatformFile.h
@@ -97,7 +97,7 @@ public:
 	virtual bool IterateDirectory(const TCHAR* Directory, FDirectoryVisitor& Visitor) override;
 	virtual bool IterateDirectoryRecursively(const TCHAR* Directory, FDirectoryVisitor& Visitor) override;
 	virtual bool IterateDirectoryStat(const TCHAR* Directory, FDirectoryStatVisitor& Visitor) override;
-	virtual IMappedFileHandle* OpenMapped(const TCHAR* Filename) override;
+	virtual IMappedFileHandle* OpenMapped(const TCHAR* Filename, EOpenReadFlags OpenOptions = EOpenReadFlags::None) override;
 	virtual FString GetFilenameOnDisk(const TCHAR* Filename) override;
 	virtual bool DeleteFile(const TCHAR* Filename) override;
 	virtual bool MoveFile(const TCHAR* To, const TCHAR* From) override;
