Fixed compilation on iOS

diff --git a/Engine/Source/Runtime/Core/Private/IOS/IOSPlatformFile.cpp b/Engine/Source/Runtime/Core/Private/IOS/IOSPlatformFile.cpp
--- a/Engine/Source/Runtime/Core/Private/IOS/IOSPlatformFile.cpp
+++ b/Engine/Source/Runtime/Core/Private/IOS/IOSPlatformFile.cpp
@@ -910,7 +910,7 @@ IFileHandle* FIOSPlatformFile::OpenWrite(const TCHAR* Filename, bool bAppend, bo
 	return NULL;
 }
 
-IMappedFileHandle* FIOSPlatformFile::OpenMapped(const TCHAR* Filename, EOpenReadFlags OpenOptions = EOpenReadFlags::None)
+IMappedFileHandle* FIOSPlatformFile::OpenMapped(const TCHAR* Filename, EOpenReadFlags OpenOptions)
 {
 	FString NormalizedFilename = NormalizeFilename(Filename);
 	
