Fixed StringTables not setting the PKG_RequiresLocalizationGather flag when serializing

diff --git a/Engine/Source/Runtime/Core/Private/Internationalization/StringTableCore.cpp b/Engine/Source/Runtime/Core/Private/Internationalization/StringTableCore.cpp
--- a/Engine/Source/Runtime/Core/Private/Internationalization/StringTableCore.cpp
+++ b/Engine/Source/Runtime/Core/Private/Internationalization/StringTableCore.cpp
@@ -401,6 +401,11 @@ void FStringTable::Serialize(FArchive& Ar)
 			}
 		}
 	}
+
+	if (KeysToEntries.Num() > 0)
+	{
+		Ar.ThisRequiresLocalizationGather();
+	}
 }
 
 bool FStringTable::ExportStrings(const FString& InFilename) const
