Fix use-after-free in FMallocLeakDetection destructor

OpenPointers.Empty() can call back into FMallocLeakDetection, which uses the map while its storage is being reallocated - this can cause a use-after-free.

diff --git a/Engine/Source/Runtime/Core/Private/HAL/MallocLeakDetection.cpp b/Engine/Source/Runtime/Core/Private/HAL/MallocLeakDetection.cpp
--- a/Engine/Source/Runtime/Core/Private/HAL/MallocLeakDetection.cpp
+++ b/Engine/Source/Runtime/Core/Private/HAL/MallocLeakDetection.cpp
@@ -61,7 +61,9 @@ FMallocLeakDetection& FMallocLeakDetection::Get()
 
 FMallocLeakDetection::~FMallocLeakDetection()
 {
+	bRecursive = true; // OpenPointers.Empty() may reallocate / free.
 	OpenPointers.Empty(); // clean up the state
+	bRecursive = false;
 	SetAllocationCollection(false); // disable collection to avoid a call back to this instance when its members are destroyed
 }
 
