Fixed forwarding of FArchiveState::ArThisRequiresLocalizationGather to proxies.

diff --git a/Engine/Source/Runtime/Core/Private/Serialization/Archive.cpp b/Engine/Source/Runtime/Core/Private/Serialization/Archive.cpp
--- a/Engine/Source/Runtime/Core/Private/Serialization/Archive.cpp
+++ b/Engine/Source/Runtime/Core/Private/Serialization/Archive.cpp
@@ -169,7 +169,9 @@ void FArchiveState::Reset()
 	ArIsCriticalError					= false;
 	ArContainsCode						= false;
 	ArContainsMap						= false;
+PRAGMA_DISABLE_DEPRECATION_WARNINGS
 	ArRequiresLocalizationGather		= false;
+PRAGMA_ENABLE_DEPRECATION_WARNINGS
 	ArForceByteSwapping					= false;
 	ArSerializingDefaults				= false;
 	ArIgnoreArchetypeRef				= false;
@@ -230,7 +232,9 @@ void FArchiveState::CopyTrivialFArchiveStatusMembers(const FArchiveState& Archiv
 	ArIsCriticalError                    = ArchiveToCopy.ArIsCriticalError;
 	ArContainsCode                       = ArchiveToCopy.ArContainsCode;
 	ArContainsMap                        = ArchiveToCopy.ArContainsMap;
+PRAGMA_DISABLE_DEPRECATION_WARNINGS
 	ArRequiresLocalizationGather         = ArchiveToCopy.ArRequiresLocalizationGather;
+PRAGMA_ENABLE_DEPRECATION_WARNINGS
 	ArForceByteSwapping                  = ArchiveToCopy.ArForceByteSwapping;
 	ArSerializingDefaults                = ArchiveToCopy.ArSerializingDefaults;
 	ArIgnoreArchetypeRef                 = ArchiveToCopy.ArIgnoreArchetypeRef;
@@ -1425,6 +1429,13 @@ void FArchive::LogfImpl(const TCHAR* Fmt, ...)
 	FMemory::SystemFree( Buffer );
 }
 
+void FArchiveState::ThisRequiresLocalizationGather()
+{
+PRAGMA_DISABLE_DEPRECATION_WARNINGS
+	ForEachState([](FArchiveState& State) { State.ArRequiresLocalizationGather = true; });
+PRAGMA_ENABLE_DEPRECATION_WARNINGS
+}
+
 void FArchiveState::SetUEVer(FPackageFileVersion InVer)
 {
 	ArUEVer = InVer;
diff --git a/Engine/Source/Runtime/Core/Public/Serialization/Archive.h b/Engine/Source/Runtime/Core/Public/Serialization/Archive.h
--- a/Engine/Source/Runtime/Core/Public/Serialization/Archive.h
+++ b/Engine/Source/Runtime/Core/Public/Serialization/Archive.h
@@ -185,10 +185,7 @@ public:
 	}
 
 	/** Sets a flag indicating that this archive contains data required to be gathered for localization. */
-	void ThisRequiresLocalizationGather()
-	{
-		ArRequiresLocalizationGather = true;
-	}
+	CORE_API void ThisRequiresLocalizationGather();
 
 	/**
 	 * Called to retrieve the archetype from the event driven loader. 
@@ -376,7 +373,9 @@ public:
 	/** Returns true if this archive contains data required to be gathered for localization. */
 	FORCEINLINE bool RequiresLocalizationGather() const
 	{
+PRAGMA_DISABLE_DEPRECATION_WARNINGS
 		return ArRequiresLocalizationGather;
+PRAGMA_ENABLE_DEPRECATION_WARNINGS
 	}
 
 	/** Returns true if this archive should always swap bytes, ignoring endian rules. */
@@ -842,6 +841,7 @@ public:
 	uint8 ArContainsMap : 1;
 
 	/** Used to determine whether FArchive contains data required to be gathered for localization. */
+	UE_DEPRECATED(5.6, "Public access to ArRequiresLocalizationGather has been deprecated - please use Ar.RequiresLocalizationGather() and Ar.ThisRequiresLocalizationGather() for getting and setting respectively.")
 	uint8 ArRequiresLocalizationGather : 1;
 
 	/** Whether we should forcefully swap bytes. */
@@ -2103,7 +2103,9 @@ public:
 	using FArchiveState::ArIsCriticalError;
 	using FArchiveState::ArContainsCode;
 	using FArchiveState::ArContainsMap;
+PRAGMA_DISABLE_DEPRECATION_WARNINGS
 	using FArchiveState::ArRequiresLocalizationGather;
+PRAGMA_ENABLE_DEPRECATION_WARNINGS
 	using FArchiveState::ArForceByteSwapping;
 	using FArchiveState::ArIgnoreArchetypeRef;
 	using FArchiveState::ArNoDelta;
