Made FWeakObjectPtr a member instead of a base to allow conversions without weird error messages about inaccessible bases.

diff --git a/Engine/Source/Runtime/Core/Public/UObject/WeakObjectPtrTemplates.h b/Engine/Source/Runtime/Core/Public/UObject/WeakObjectPtrTemplates.h
--- a/Engine/Source/Runtime/Core/Public/UObject/WeakObjectPtrTemplates.h
+++ b/Engine/Source/Runtime/Core/Public/UObject/WeakObjectPtrTemplates.h
@@ -11,17 +11,23 @@
 
 #include <type_traits>
 
-
+struct FFieldPath;
 struct FWeakObjectPtr;
+struct FUObjectItem;
 
 /**
  * TWeakObjectPtr is the templated version of the generic FWeakObjectPtr
  */
 template<class T, class TWeakObjectPtrBase>
-struct TWeakObjectPtr : private TWeakObjectPtrBase
+struct TWeakObjectPtr
 {
 	friend struct FFieldPath;
 
+	template <class, class>
+	friend struct TWeakObjectPtr;
+
+	friend struct FFieldPath;
+
 	// Although templated, these parameters are not intended to be anything other than the default,
 	// and are only templates for module organization reasons.
 	static_assert(std::is_same_v<TWeakObjectPtrBase*, FWeakObjectPtr*>, "TWeakObjectPtrBase should not be overridden");
@@ -38,7 +44,7 @@ public:
 	 * Construct from a null pointer
 	 */
 	FORCEINLINE TWeakObjectPtr(TYPE_OF_NULLPTR) :
-		TWeakObjectPtrBase((UObject*)nullptr)
+		WeakPtr((UObject*)nullptr)
 	{
 	}
 
@@ -51,7 +57,7 @@ public:
 		UE_REQUIRES(std::is_convertible_v<U, T*>)
 	>
 	FORCEINLINE TWeakObjectPtr(U Object) :
-		TWeakObjectPtrBase((const UObject*)Object)
+		WeakPtr((const UObject*)Object)
 	{
 		// This static assert is in here rather than in the body of the class because we want
 		// to be able to define TWeakObjectPtr<UUndefinedClass>.
@@ -67,7 +73,7 @@ public:
 		UE_REQUIRES(std::is_convertible_v<OtherT*, T*>)
 	>
 	FORCEINLINE TWeakObjectPtr(const TWeakObjectPtr<OtherT, TWeakObjectPtrBase>& Other) :
-		TWeakObjectPtrBase(*(TWeakObjectPtrBase*)&Other) // we do a C-style cast to private base here to avoid clang 3.6.0 compilation problems with friend declarations
+		WeakPtr(Other.WeakPtr) // we do a C-style cast to private base here to avoid clang 3.6.0 compilation problems with friend declarations
 	{
 	}
 
@@ -76,7 +82,7 @@ public:
 	 */
 	FORCEINLINE void Reset()
 	{
-		TWeakObjectPtrBase::Reset();
+		WeakPtr.Reset();
 	}
 
 	/**  
@@ -90,7 +96,7 @@ public:
 	FORCEINLINE TWeakObjectPtr& operator=(U* Object)
 	{
 		T* TempObject = Object;
-		TWeakObjectPtrBase::operator=(TempObject);
+		WeakPtr = TempObject;
 		return *this;
 	}
 
@@ -104,7 +110,7 @@ public:
 	>
 	FORCEINLINE TWeakObjectPtr& operator=(const TWeakObjectPtr<OtherT, TWeakObjectPtrBase>& Other)
 	{
-		*(TWeakObjectPtrBase*)this = *(TWeakObjectPtrBase*)&Other; // we do a C-style cast to private base here to avoid clang 3.6.0 compilation problems with friend declarations
+		WeakPtr = Other.WeakPtr;
 
 		return *this;
 	}
@@ -116,7 +122,7 @@ public:
 	 */
 	FORCEINLINE T* Get(bool bEvenIfPendingKill) const
 	{
-		return (T*)TWeakObjectPtrBase::Get(bEvenIfPendingKill);
+		return (T*)WeakPtr.Get(bEvenIfPendingKill);
 	}
 
 	/**  
@@ -124,7 +130,7 @@ public:
 	 */
 	FORCEINLINE T* Get(/*bool bEvenIfPendingKill = false*/) const
 	{
-		return (T*)TWeakObjectPtrBase::Get();
+		return (T*)WeakPtr.Get();
 	}
 
 	/**
@@ -135,7 +141,7 @@ public:
 	FORCEINLINE TStrongObjectPtr<T> Pin(bool bEvenIfPendingKill) const
 	{
 		TStrongObjectPtr<T> StrongPtr;
-		StrongPtr.Attach((T*)TWeakObjectPtrBase::Pin(bEvenIfPendingKill).Detach());
+		StrongPtr.Attach((T*)WeakPtr.Pin(bEvenIfPendingKill).Detach());
 		return StrongPtr;
 	}
 
@@ -145,14 +151,14 @@ public:
 	FORCEINLINE TStrongObjectPtr<T> Pin(/*bool bEvenIfPendingKill = false*/) const
 	{
 		TStrongObjectPtr<T> StrongPtr;
-		StrongPtr.Attach((T*)TWeakObjectPtrBase::Pin().Detach());
+		StrongPtr.Attach((T*)WeakPtr.Pin().Detach());
 		return StrongPtr;
 	}
 
 	/** Deferences the weak pointer even if its marked RF_Unreachable. This is needed to resolve weak pointers during GC (such as ::AddReferenceObjects) */
 	FORCEINLINE T* GetEvenIfUnreachable() const
 	{
-		return (T*)TWeakObjectPtrBase::GetEvenIfUnreachable();
+		return (T*)WeakPtr.GetEvenIfUnreachable();
 	}
 
 	/**  
@@ -187,7 +193,7 @@ public:
 	 */
 	FORCEINLINE bool IsValid(bool bEvenIfPendingKill, bool bThreadsafeTest = false) const
 	{
-		return TWeakObjectPtrBase::IsValid(bEvenIfPendingKill, bThreadsafeTest);
+		return WeakPtr.IsValid(bEvenIfPendingKill, bThreadsafeTest);
 	}
 
 	/**
@@ -199,7 +205,7 @@ public:
 	 */
 	FORCEINLINE bool IsValid(/*bool bEvenIfPendingKill = false, bool bThreadsafeTest = false*/) const
 	{
-		return TWeakObjectPtrBase::IsValid();
+		return WeakPtr.IsValid();
 	}
 
 	/**  
@@ -210,7 +216,7 @@ public:
 	 */
 	FORCEINLINE bool IsStale(bool bIncludingIfPendingKill = true, bool bThreadsafeTest = false) const
 	{
-		return TWeakObjectPtrBase::IsStale(bIncludingIfPendingKill, bThreadsafeTest);
+		return WeakPtr.IsStale(bIncludingIfPendingKill, bThreadsafeTest);
 	}
 	
 	/**
@@ -219,7 +225,7 @@ public:
 	 */
 	FORCEINLINE bool IsExplicitlyNull() const
 	{
-		return TWeakObjectPtrBase::IsExplicitlyNull();
+		return WeakPtr.IsExplicitlyNull();
 	}
 
 	/**
@@ -228,7 +234,7 @@ public:
 	 */
 	FORCEINLINE bool HasSameIndexAndSerialNumber(const TWeakObjectPtr& Other) const
 	{
-		return static_cast<const TWeakObjectPtrBase&>(*this).HasSameIndexAndSerialNumber(static_cast<const TWeakObjectPtrBase&>(Other));
+		return WeakPtr.HasSameIndexAndSerialNumber(Other.WeakPtr);
 	}
 
 	/**
@@ -241,7 +247,7 @@ public:
 	>
 	FORCEINLINE bool HasSameIndexAndSerialNumber(const TWeakObjectPtr<OtherT, TWeakObjectPtrBase>& Other) const
 	{
-		return ((const TWeakObjectPtrBase&)*this).HasSameIndexAndSerialNumber((const TWeakObjectPtrBase&)Other);
+		return WeakPtr.HasSameIndexAndSerialNumber(Other.WeakPtr);
 	}
 
 	/**
@@ -249,13 +255,13 @@ public:
 	 */
 	FORCEINLINE	void Serialize(FArchive& Ar)
 	{
-		Ar << static_cast<TWeakObjectPtrBase&>(*this);
+		Ar << WeakPtr;
 	}
 
 	/** Hash function. */
 	FORCEINLINE uint32 GetWeakPtrTypeHash() const
 	{
-		return static_cast<const TWeakObjectPtrBase&>(*this).GetTypeHash();
+		return WeakPtr.GetTypeHash();
 	}
 
 	/**
@@ -266,14 +272,14 @@ public:
 	template <typename RhsT, typename = decltype((T*)nullptr == (RhsT*)nullptr)>
 	FORCENOINLINE bool operator==(const TWeakObjectPtr<RhsT, TWeakObjectPtrBase>& Rhs) const
 	{
-		return (const TWeakObjectPtrBase&)*this == (const TWeakObjectPtrBase&)Rhs;
+		return WeakPtr == Rhs.WeakPtr;
 	}
 
 	template <typename RhsT, typename = decltype((T*)nullptr == (RhsT*)nullptr)>
 	FORCENOINLINE bool operator==(const RhsT* Rhs) const
 	{
 		// NOTE: this constructs a TWeakObjectPtrBase, which has some amount of overhead, so this may not be an efficient operation
-		return (const TWeakObjectPtrBase&)*this == TWeakObjectPtrBase(Rhs);
+		return WeakPtr == TWeakObjectPtrBase(Rhs);
 	}
 
 	FORCENOINLINE bool operator==(TYPE_OF_NULLPTR) const
@@ -287,23 +293,30 @@ public:
 	 * @param Other weak pointer to compare to
 	 */
 	template <typename RhsT, typename = decltype((T*)nullptr != (RhsT*)nullptr)>
-	FORCENOINLINE bool operator!=(const TWeakObjectPtr<RhsT, TWeakObjectPtrBase>& Rhs) const
+	FORCEINLINE bool operator!=(const TWeakObjectPtr<RhsT, TWeakObjectPtrBase>& Rhs) const
 	{
-		return (const TWeakObjectPtrBase&)*this != (const TWeakObjectPtrBase&)Rhs;
+		return !(*this == Rhs);
 	}
 
 	template <typename RhsT, typename = decltype((T*)nullptr != (RhsT*)nullptr)>
-	FORCENOINLINE bool operator!=(const RhsT* Rhs) const
+	FORCEINLINE bool operator!=(const RhsT* Rhs) const
 	{
-		// NOTE: this constructs a TWeakObjectPtrBase, which has some amount of overhead, so this may not be an efficient operation
-		return (const TWeakObjectPtrBase&)*this != TWeakObjectPtrBase(Rhs);
+		return !(*this == Rhs);
 	}
 
-	FORCENOINLINE bool operator!=(TYPE_OF_NULLPTR) const
+	FORCEINLINE bool operator!=(TYPE_OF_NULLPTR) const
 	{
-		return IsValid();
+		return !(*this == nullptr);
 	}
 #endif
+
+private:
+	FORCEINLINE FUObjectItem* Internal_GetObjectItem() const
+	{
+		return WeakPtr.Internal_GetObjectItem();
+	}
+
+	TWeakObjectPtrBase WeakPtr;
 };
 
 template <typename T>
@@ -319,15 +332,14 @@ TWeakObjectPtr(const TWeakObjectPtr<T>&) -> TWeakObjectPtr<T>;
  * @param Other weak pointer to compare to
  */
 template <typename LhsT, typename RhsT, typename OtherTWeakObjectPtrBase, typename = decltype((LhsT*)nullptr == (RhsT*)nullptr)>
-FORCENOINLINE bool operator==(const LhsT* Lhs, const TWeakObjectPtr<RhsT, OtherTWeakObjectPtrBase>& Rhs)
+FORCEINLINE bool operator==(const LhsT* Lhs, const TWeakObjectPtr<RhsT, OtherTWeakObjectPtrBase>& Rhs)
 {
-	// NOTE: this constructs a TWeakObjectPtrBase, which has some amount of overhead, so this may not be an efficient operation
-	return OtherTWeakObjectPtrBase(Lhs) == (const OtherTWeakObjectPtrBase&)Rhs;
+	return Rhs == Lhs;
 }
 template <typename RhsT, typename OtherTWeakObjectPtrBase>
-FORCENOINLINE bool operator==(TYPE_OF_NULLPTR, const TWeakObjectPtr<RhsT, OtherTWeakObjectPtrBase>& Rhs)
+FORCEINLINE bool operator==(TYPE_OF_NULLPTR, const TWeakObjectPtr<RhsT, OtherTWeakObjectPtrBase>& Rhs)
 {
-	return !Rhs.IsValid();
+	return Rhs == nullptr;
 }
 
 /**
@@ -335,15 +347,14 @@ FORCENOINLINE bool operator==(TYPE_OF_NULLPTR, const TWeakObjectPtr<RhsT, OtherT
  * @param Other weak pointer to compare to
  */
 template <typename LhsT, typename RhsT, typename OtherTWeakObjectPtrBase, typename = decltype((LhsT*)nullptr != (RhsT*)nullptr)>
-FORCENOINLINE bool operator!=(const LhsT* Lhs, const TWeakObjectPtr<RhsT, OtherTWeakObjectPtrBase>& Rhs)
+FORCEINLINE bool operator!=(const LhsT* Lhs, const TWeakObjectPtr<RhsT, OtherTWeakObjectPtrBase>& Rhs)
 {
-	// NOTE: this constructs a TWeakObjectPtrBase, which has some amount of overhead, so this may not be an efficient operation
-	return OtherTWeakObjectPtrBase(Lhs) != (const OtherTWeakObjectPtrBase&)Rhs;
+	return !(Rhs == Lhs);
 }
 template <typename RhsT, typename OtherTWeakObjectPtrBase>
-FORCENOINLINE bool operator!=(TYPE_OF_NULLPTR, const TWeakObjectPtr<RhsT, OtherTWeakObjectPtrBase>& Rhs)
+FORCEINLINE bool operator!=(TYPE_OF_NULLPTR, const TWeakObjectPtr<RhsT, OtherTWeakObjectPtrBase>& Rhs)
 {
-	return Rhs.IsValid();
+	return !(Rhs == nullptr);
 }
 #endif
 
diff --git a/Engine/Source/Runtime/CoreUObject/Public/UObject/WeakObjectPtr.h b/Engine/Source/Runtime/CoreUObject/Public/UObject/WeakObjectPtr.h
--- a/Engine/Source/Runtime/CoreUObject/Public/UObject/WeakObjectPtr.h
+++ b/Engine/Source/Runtime/CoreUObject/Public/UObject/WeakObjectPtr.h
@@ -19,6 +19,8 @@
 
 class FArchive;
 class UObject;
+template <class, class>
+struct TWeakObjectPtr;
 
 /** Invalid FWeakObjectPtr ObjectIndex values must be 0 to support zeroed initialization (this used to be INDEX_NONE, leading to subtle bugs). */
 #ifndef UE_WEAKOBJECTPTR_ZEROINIT_FIX
@@ -47,7 +49,8 @@ struct FWeakObjectPtr
 {
 public:
 
-	friend struct FFieldPath;
+	template <class, class>
+	friend struct TWeakObjectPtr;
 
 #if UE_WEAKOBJECTPTR_ZEROINIT_FIX
 	FWeakObjectPtr() = default;
