PR #12428: Adding "!" to the list of separators for tokenization

#rnx

diff --git a/Engine/Plugins/FX/Niagara/Source/NiagaraEditor/Private/NiagaraNodeCustomHlsl.cpp b/Engine/Plugins/FX/Niagara/Source/NiagaraEditor/Private/NiagaraNodeCustomHlsl.cpp
--- a/Engine/Plugins/FX/Niagara/Source/NiagaraEditor/Private/NiagaraNodeCustomHlsl.cpp
+++ b/Engine/Plugins/FX/Niagara/Source/NiagaraEditor/Private/NiagaraNodeCustomHlsl.cpp
@@ -120,7 +120,7 @@ bool UNiagaraNodeCustomHlsl::GetTokensFromString(const FString& InHlsl, TArray<F
 		return false;
 	}
 
-	FString Separators = TEXT(";/*+-=)(?:, []<>\"\t\r\n{}");
+	FString Separators = TEXT("!;/*+-=)(?:, []<>\"\t\r\n{}");
 	const int32 TargetLength = InHlsl.Len();
 
 	int32 TokenStart = 0;
diff --git a/Engine/Source/Runtime/Core/Private/UObject/FortniteMainBranchObjectVersion.cpp b/Engine/Source/Runtime/Core/Private/UObject/FortniteMainBranchObjectVersion.cpp
--- a/Engine/Source/Runtime/Core/Private/UObject/FortniteMainBranchObjectVersion.cpp
+++ b/Engine/Source/Runtime/Core/Private/UObject/FortniteMainBranchObjectVersion.cpp
@@ -11,7 +11,7 @@ TMap<FGuid, FGuid> FFortniteMainBranchObjectVersion::GetSystemGuids()
 	SystemGuids.Add(DevGuids.MATERIALSHADERMAP_DERIVEDDATA_VER, FGuid("0E6202C151DE402CA451802671AE9347"));
 	SystemGuids.Add(DevGuids.NANITE_DERIVEDDATA_VER, FGuid("DB3ECDFAF27A4344ACB6B1967C21148D"));
 	SystemGuids.Add(DevGuids.NIAGARASHADERMAP_DERIVEDDATA_VER, FGuid("C417A50271B3427D98785388266E3B66"));
-	SystemGuids.Add(DevGuids.Niagara_LatestScriptCompileVersion, FGuid("F8D995224324470E97ED718D33939A0E"));
+	SystemGuids.Add(DevGuids.Niagara_LatestScriptCompileVersion, FGuid("393081597FA7234C92F00869DE262D73"));
 	SystemGuids.Add(DevGuids.SkeletalMeshDerivedDataVersion, FGuid("B9AEAA3EE7AC4FCFB3860E8F3E06DD75"));
 	SystemGuids.Add(DevGuids.STATICMESH_DERIVEDDATA_VER, FGuid("CDFDB82FE0A83DB602AE2EFFD8BB7F42"));
 	SystemGuids.Add(DevGuids.MaterialTranslationDDCVersion, FGuid("E779012A63684F46B85596E4E89D489A"));
