Adding explicit cast operator from TWeakObjectPtr to FWeakObjectPtr

diff --git a/Engine/Source/Runtime/Core/Public/UObject/WeakObjectPtrTemplates.h b/Engine/Source/Runtime/Core/Public/UObject/WeakObjectPtrTemplates.h
--- a/Engine/Source/Runtime/Core/Public/UObject/WeakObjectPtrTemplates.h
+++ b/Engine/Source/Runtime/Core/Public/UObject/WeakObjectPtrTemplates.h
@@ -180,6 +180,14 @@ public:
 	// This is explicitly not added to avoid resolving weak pointers too often - use Get() once in a function.
 	explicit operator bool() const = delete;
 
+	/**
+	 * Cast to the underlying generic FWeakObjectPtr type
+	 */
+	explicit operator TWeakObjectPtrBase() const
+	{
+		return WeakPtr;
+	}
+
 	/**  
 	 * Test if this points to a live UObject.
 	 * This should be done only when needed as excess resolution of the underlying pointer can cause performance issues.
