[AutoRTFM] Check ASAN fake stacks in FTransaction::IsOnStack()

Fixes stack corruption when ASAN detect_stack_use_after_return is enabled.

diff --git a/Engine/Source/Runtime/Core/Private/AutoRTFM/TransactionInlines.h b/Engine/Source/Runtime/Core/Private/AutoRTFM/TransactionInlines.h
--- a/Engine/Source/Runtime/Core/Private/AutoRTFM/TransactionInlines.h
+++ b/Engine/Source/Runtime/Core/Private/AutoRTFM/TransactionInlines.h
@@ -9,12 +9,45 @@
 #include "HAL/Platform.h"
 #include "Utils.h"
 
+#if PLATFORM_HAS_ASAN_INCLUDE
+#	if defined(__SANITIZE_ADDRESS__)
+#		define AUTORTFM_CHECK_ASAN_FAKE_STACKS 1
+#	elif defined(__has_feature)
+#		if __has_feature(address_sanitizer)
+#			define AUTORTFM_CHECK_ASAN_FAKE_STACKS 1
+#		endif
+#	endif
+#endif
+
+#ifndef AUTORTFM_CHECK_ASAN_FAKE_STACKS
+#define AUTORTFM_CHECK_ASAN_FAKE_STACKS 0
+#endif
+
+#if AUTORTFM_CHECK_ASAN_FAKE_STACKS
+#include <sanitizer/asan_interface.h>
+#endif
+
 namespace AutoRTFM
 {
 
 UE_AUTORTFM_FORCEINLINE bool FTransaction::IsOnStack(const void* LogicalAddress) const
 { 
-    return StackRange.Contains(LogicalAddress);
+    if (StackRange.Contains(LogicalAddress))
+    {
+        return true;
+    }
+
+#if AUTORTFM_CHECK_ASAN_FAKE_STACKS
+    if (void* FakeStack = __asan_get_current_fake_stack())
+    {
+        void* Beg = nullptr;
+        void* End = nullptr;
+        void* RealAddress = __asan_addr_is_in_fake_stack(FakeStack, const_cast<void*>(LogicalAddress), &Beg, &End);
+        return RealAddress && StackRange.Contains(RealAddress);
+    }
+#endif  // AUTORTFM_CHECK_ASAN_FAKE_STACKS
+
+    return false;
 }
 
 UE_AUTORTFM_FORCEINLINE bool FTransaction::ShouldRecordWrite(void* LogicalAddress) const
@@ -254,3 +287,5 @@ UE_AUTORTFM_FORCEINLINE void FTransaction::CollectStats() const
 }
 
 } // namespace AutoRTFM
+
+#undef AUTORTFM_CHECK_ASAN_FAKE_STACKS
