Replaced some deprecations with static asserts in ContainerAllocationPolicies.

diff --git a/Engine/Source/Runtime/Core/Public/Containers/ContainerAllocationPolicies.h b/Engine/Source/Runtime/Core/Public/Containers/ContainerAllocationPolicies.h
--- a/Engine/Source/Runtime/Core/Public/Containers/ContainerAllocationPolicies.h
+++ b/Engine/Source/Runtime/Core/Public/Containers/ContainerAllocationPolicies.h
@@ -589,7 +589,7 @@ public:
 		/** Default constructor. */
 		ForElementType()
 		{
-			UE_STATIC_DEPRECATE(5.0, alignof(ElementType) > MinimumAlignment, "Using TAlignedHeapAllocator with an alignment lower than the element type's alignment - please update the alignment parameter");
+			static_assert(alignof(ElementType) <= MinimumAlignment, "Using TAlignedHeapAllocator with an alignment lower than the element type's alignment - please update the alignment parameter");
 		}
 
 		FORCEINLINE ElementType* GetAllocation() const
@@ -814,9 +814,7 @@ public:
 	{
 	public:
 		/** Default constructor. */
-		ForElementType()
-		{
-		}
+		ForElementType() = default;
 
 		FORCEINLINE ElementType* GetAllocation() const
 		{
@@ -1048,7 +1046,7 @@ public:
 		ForElementType()
 			: Data(GetInlineElements())
 		{
-			UE_STATIC_DEPRECATE(5.0, alignof(ElementType) > __STDCPP_DEFAULT_NEW_ALIGNMENT__, "TNonRelocatableInlineAllocator uses GMalloc's default alignment, which is lower than the element type's alignment - please consider a different approach");
+			static_assert(alignof(ElementType) <= __STDCPP_DEFAULT_NEW_ALIGNMENT__, "TNonRelocatableInlineAllocator uses GMalloc's default alignment, which is lower than the element type's alignment - please consider a different approach");
 		}
 
 		~ForElementType()
