Fixed compilation errors on Android caused by the file mapping API change

diff --git a/Engine/Source/Runtime/Core/Private/Android/AndroidPlatformFile.cpp b/Engine/Source/Runtime/Core/Private/Android/AndroidPlatformFile.cpp
--- a/Engine/Source/Runtime/Core/Private/Android/AndroidPlatformFile.cpp
+++ b/Engine/Source/Runtime/Core/Private/Android/AndroidPlatformFile.cpp
@@ -1460,7 +1460,7 @@ public:
 		if (Handle == -1)
 		{
 			const int ErrNo = errno;
-			FString ErrorStr = FString::Printf(TEXT("FAndroidPlatformFile::OpenMappedEx('%s', Flags=0x%08X) failed: errno=%d (%s)"), *NormalizedFilename, Flags, ErrNo, UTF8_TO_TCHAR(strerror(ErrNo));
+			FString ErrorStr = FString::Printf(TEXT("FAndroidPlatformFile::OpenMappedEx('%s', Flags=0x%08X) failed: errno=%d (%s)"), *NormalizedFilename, Flags, ErrNo, UTF8_TO_TCHAR(strerror(ErrNo)));
 #if LOG_ANDROID_FILE
 			FPlatformMisc::LowLevelOutputDebugString(*ErrorStr);
 #endif
@@ -1473,7 +1473,7 @@ public:
 		if (StatResult == -1)
 		{
 			const int ErrNo = errno;
-			FString ErrorStr = FString::Printf(TEXT("FAndroidPlatformFile::OpenMappedEx fstat failed: ('%s', Flags=0x%08X) failed: errno=%d (%s)"), *NormalizedFilename, Flags, ErrNo, UTF8_TO_TCHAR(strerror(ErrNo));
+			FString ErrorStr = FString::Printf(TEXT("FAndroidPlatformFile::OpenMappedEx fstat failed: ('%s', Flags=0x%08X) failed: errno=%d (%s)"), *NormalizedFilename, Flags, ErrNo, UTF8_TO_TCHAR(strerror(ErrNo)));
 
 #if LOG_ANDROID_FILE
 			FPlatformMisc::LowLevelOutputDebugString(*ErrorStr);
