Fixed compilation errors on Unix caused by the file mapping API change

diff --git a/Engine/Source/Runtime/Core/Private/Unix/UnixPlatformFile.cpp b/Engine/Source/Runtime/Core/Private/Unix/UnixPlatformFile.cpp
--- a/Engine/Source/Runtime/Core/Private/Unix/UnixPlatformFile.cpp
+++ b/Engine/Source/Runtime/Core/Private/Unix/UnixPlatformFile.cpp
@@ -1236,7 +1236,7 @@ FOpenMappedResult FUnixPlatformFile::OpenMappedEx(const TCHAR* Filename, EOpenRe
 	{
 		const int ErrNo = errno;
 		FString ErrorStr = FString::Printf(TEXT("open('%s', Flags=0x%08X) failed: errno=%d (%s)"), *NormalizedFilename, Flags, ErrNo, UTF8_TO_TCHAR(strerror(ErrNo)));
-		UE_LOG_UNIX_FILE(Warning, *ErrorStr);
+		UE_LOG_UNIX_FILE(Warning, TEXT("%s"), *ErrorStr);
 
 		return MakeError(ErrorStr);
 	}
@@ -1248,7 +1248,7 @@ FOpenMappedResult FUnixPlatformFile::OpenMappedEx(const TCHAR* Filename, EOpenRe
 	{
 		const int ErrNo = errno;
 		FString ErrorStr = FString::Printf(TEXT("stat('%s', Flags=0x%08X) failed: errno=%d (%s)"), *NormalizedFilename, Flags, ErrNo, UTF8_TO_TCHAR(strerror(ErrNo)));
-		UE_LOG_UNIX_FILE(Warning, ErrorStr);
+		UE_LOG_UNIX_FILE(Warning, TEXT("%s"), *ErrorStr);
 
 		return MakeError(ErrorStr);
 	}
