Fixed compilation errors on iOS caused by the file mapping API change

diff --git a/Engine/Source/Runtime/Core/Private/IOS/IOSPlatformFile.cpp b/Engine/Source/Runtime/Core/Private/IOS/IOSPlatformFile.cpp
--- a/Engine/Source/Runtime/Core/Private/IOS/IOSPlatformFile.cpp
+++ b/Engine/Source/Runtime/Core/Private/IOS/IOSPlatformFile.cpp
@@ -952,7 +952,7 @@ FOpenMappedResult FIOSPlatformFile::OpenMappedEx(const TCHAR* Filename, EOpenRea
 		}
 	}
 	
-	return MakeError(FString::Printf(TEXT("FIOSPlatformFile::OpenMappedEx failed to open file '%s' in '%s' mode"), *NormalizedFilename, EnumHasAnyFlags(OpenOptions, EOpenReadFlags::AllowWrite) ? TEXTVIEW("write") : TEXTVIEW("read")));
+	return MakeError(FString::Printf(TEXT("FIOSPlatformFile::OpenMappedEx failed to open file '%s' in '%s' mode"), *NormalizedFilename, EnumHasAnyFlags(OpenOptions, EOpenReadFlags::AllowWrite) ? TEXT("write") : TEXT("read")));
 }
 
 bool FIOSPlatformFile::DirectoryExists(const TCHAR* Directory)
