Add cross platform SetThreadPriority in LaunchEngineLoop, removing the static initiliazer in FRunnableThreadWin::CreateInternal forcing the current thread (game) to be TPri_Normal

diff --git a/Engine/Source/Developer/LowLevelTestsRunner/Private/TestRunner.cpp b/Engine/Source/Developer/LowLevelTestsRunner/Private/TestRunner.cpp
--- a/Engine/Source/Developer/LowLevelTestsRunner/Private/TestRunner.cpp
+++ b/Engine/Source/Developer/LowLevelTestsRunner/Private/TestRunner.cpp
@@ -211,6 +211,7 @@ void FTestRunner::GlobalSetup()
 
 	// Finish setting up the Game Thread, which requires the command line.
 	FPlatformProcess::SetThreadAffinityMask(FPlatformAffinity::GetMainGameMask());
+	FPlatformProcess::SetThreadPriority(FPlatformAffinity::GetGameThreadPriority());
 	FPlatformProcess::SetupGameThread();
 
 	// Always set up GError to handle FatalError, failed assertions, and crashes and other fatal errors.
diff --git a/Engine/Source/Runtime/Core/Private/Microsoft/MicrosoftRunnableThread.h b/Engine/Source/Runtime/Core/Private/Microsoft/MicrosoftRunnableThread.h
--- a/Engine/Source/Runtime/Core/Private/Microsoft/MicrosoftRunnableThread.h
+++ b/Engine/Source/Runtime/Core/Private/Microsoft/MicrosoftRunnableThread.h
@@ -128,13 +128,6 @@ protected:
 		EThreadPriority InThreadPri = TPri_Normal, uint64 InThreadAffinityMask = 0,
 		EThreadCreateFlags InCreateFlags = EThreadCreateFlags::None) override
 	{
-		static bool bOnce = false;
-		if (!bOnce)
-		{
-			bOnce = true;
-			::SetThreadPriority(::GetCurrentThread(), TranslateThreadPriority(TPri_Normal));
-		}
-
 		check(InRunnable);
 		Runnable = InRunnable;
 		ThreadAffinityMask = InThreadAffinityMask;
diff --git a/Engine/Source/Runtime/Core/Private/Windows/WindowsPlatformProcess.cpp b/Engine/Source/Runtime/Core/Private/Windows/WindowsPlatformProcess.cpp
--- a/Engine/Source/Runtime/Core/Private/Windows/WindowsPlatformProcess.cpp
+++ b/Engine/Source/Runtime/Core/Private/Windows/WindowsPlatformProcess.cpp
@@ -2542,6 +2542,11 @@ namespace WindowsPlatformProcessImpl
 			RealSetThreadDescription(::GetCurrentThread(), lpThreadDescription);
 		}
 	}
+
+	static void SetThreadPriority(int NewPriority)
+	{
+		::SetThreadPriority(::GetCurrentThread(), NewPriority);
+	}
 }
 
 void FWindowsPlatformProcess::SetThreadName( const TCHAR* ThreadName )
@@ -2553,6 +2558,11 @@ void FWindowsPlatformProcess::SetThreadName( const TCHAR* ThreadName )
 	WindowsPlatformProcessImpl::SetThreadName(TCHAR_TO_ANSI(ThreadName));
 }
 
+void FWindowsPlatformProcess::SetThreadPriority(EThreadPriority NewPriority)
+{
+	::SetThreadPriority(::GetCurrentThread(), FRunnableThreadWin::TranslateThreadPriority(NewPriority));
+}
+
 PRAGMA_RESTORE_UNSAFE_TYPECAST_WARNINGS
 
 #include "Windows/HideWindowsPlatformTypes.h"
diff --git a/Engine/Source/Runtime/Core/Public/Windows/WindowsPlatformProcess.h b/Engine/Source/Runtime/Core/Public/Windows/WindowsPlatformProcess.h
--- a/Engine/Source/Runtime/Core/Public/Windows/WindowsPlatformProcess.h
+++ b/Engine/Source/Runtime/Core/Public/Windows/WindowsPlatformProcess.h
@@ -142,6 +142,7 @@ public:
 	static CORE_API uint32 GetCurrentCoreNumber();
 	static CORE_API void SetThreadAffinityMask( uint64 AffinityMask );
 	static CORE_API void SetThreadName( const TCHAR* ThreadName );
+	static CORE_API void SetThreadPriority(EThreadPriority NewPriority);
 	static CORE_API const TCHAR* BaseDir();
 	static CORE_API const TCHAR* UserDir();
 	static CORE_API const TCHAR* UserTempDir();
diff --git a/Engine/Source/Runtime/Launch/Private/LaunchEngineLoop.cpp b/Engine/Source/Runtime/Launch/Private/LaunchEngineLoop.cpp
--- a/Engine/Source/Runtime/Launch/Private/LaunchEngineLoop.cpp
+++ b/Engine/Source/Runtime/Launch/Private/LaunchEngineLoop.cpp
@@ -2334,6 +2334,7 @@ int32 FEngineLoop::PreInitPreStartupScreen(const TCHAR* CmdLine)
 	GGameThreadId = FPlatformTLS::GetCurrentThreadId();
 	GIsGameThreadIdInitialized = true;
 
+	FPlatformProcess::SetThreadPriority(FPlatformAffinity::GetGameThreadPriority());
 	FPlatformProcess::SetThreadAffinityMask(FPlatformAffinity::GetMainGameMask());
 	FPlatformProcess::SetupGameThread();
 
