Fixed file mapping internal flags for Android and iOS

diff --git a/Engine/Source/Runtime/Core/Private/Android/AndroidPlatformFile.cpp b/Engine/Source/Runtime/Core/Private/Android/AndroidPlatformFile.cpp
--- a/Engine/Source/Runtime/Core/Private/Android/AndroidPlatformFile.cpp
+++ b/Engine/Source/Runtime/Core/Private/Android/AndroidPlatformFile.cpp
@@ -1123,16 +1123,16 @@ public:
 		//File mapping can extend beyond file size. It's OK, kernel will just fill any leftover page data with zeros
 		const int64 AlignedSize = Align(BytesToMap + Offset - AlignedOffset, FileMappingAlignment);
 
-		int Flags = MAP_PRIVATE;
-		if (EnumHasAnyFlags(InFlags.Flags, EMappedFileFlags::EPreloadHint))
-		{
-			Flags |= MAP_POPULATE;
-		}
-
 		int Protection = PROT_READ;
+		int Flags = (EnumHasAnyFlags(InFlags.Flags, EMappedFileFlags::EPreloadHint)) ? MAP_POPULATE : 0;
 		if (EnumHasAnyFlags(InFlags.Flags, EMappedFileFlags::EFileWritable))
 		{
 			Protection |= PROT_WRITE;
+			Flags |= MAP_SHARED;
+		}
+		else
+		{
+			Flags |= MAP_PRIVATE;
 		}
 
 		const uint8* AlignedMapPtr = static_cast<const uint8*>(mmap(nullptr, AlignedSize, Protection, Flags, FileHandle, AlignedOffset));
diff --git a/Engine/Source/Runtime/Core/Private/IOS/IOSPlatformFile.cpp b/Engine/Source/Runtime/Core/Private/IOS/IOSPlatformFile.cpp
--- a/Engine/Source/Runtime/Core/Private/IOS/IOSPlatformFile.cpp
+++ b/Engine/Source/Runtime/Core/Private/IOS/IOSPlatformFile.cpp
@@ -517,23 +517,24 @@ public:
 		check(Offset < GetFileSize()); // don't map zero bytes and don't map off the end of the file
 		BytesToMap = FMath::Min<int64>(BytesToMap, GetFileSize() - Offset);
 		check(BytesToMap > 0); // don't map zero bytes
-
-		
-		// const uint8* MapPtr = (const uint8 *)mmap(NULL, BytesToMap, PROT_READ, MAP_PRIVATE | (EnumHasAnyFlags(Flags.Flags, EMappedFileFlags::EPreloadHint) ? MAP_POPULATE : 0), FileHandle, Offset);
-		// const uint8* MapPtr = (const uint8 *)mmap(NULL, BytesToMap, PROT_READ, MAP_PRIVATE, FileHandle, Offset);
-		//		const uint8* MapPtr = (const uint8 *)mmap(NULL, BytesToMap, PROT_READ, MAP_SHARED, FileHandle, Offset);
 		
 		const int64 AlignedOffset = AlignDown(Offset, FileMappingAlignment);
 		//File mapping can extend beyond file size. It's OK, kernel will just fill any leftover page data with zeros
 		const int64 AlignedSize = Align(BytesToMap + Offset - AlignedOffset, FileMappingAlignment);
 		
 		int Protection = PROT_READ;
+		int InternalFlags = (EnumHasAnyFlags(Flags.Flags, EMappedFileFlags::EPreloadHint)) ? MAP_POPULATE : 0;
 		if (EnumHasAnyFlags(Flags.Flags, EMappedFileFlags::EFileWritable))
 		{
 			Protection |= PROT_WRITE;
+			InternalFlags |= MAP_SHARED;
+		}
+		else
+		{
+			InternalFlags |= MAP_PRIVATE;
 		}
 
-		const uint8* AlignedMapPtr = (const uint8 *)mmap(NULL, AlignedSize, Protection, MAP_PRIVATE, FileHandle, AlignedOffset);
+		const uint8* AlignedMapPtr = (const uint8 *)mmap(NULL, AlignedSize, Protection, InternalFlags, FileHandle, AlignedOffset);
 		if (AlignedMapPtr == (const uint8*)-1 || AlignedMapPtr == nullptr)
 		{
 			UE_LOG(LogIOS, Warning, TEXT("Failed to map memory %s, error is %d"), *Filename, errno);
diff --git a/Engine/Source/Runtime/Core/Private/Unix/UnixPlatformFile.cpp b/Engine/Source/Runtime/Core/Private/Unix/UnixPlatformFile.cpp
--- a/Engine/Source/Runtime/Core/Private/Unix/UnixPlatformFile.cpp
+++ b/Engine/Source/Runtime/Core/Private/Unix/UnixPlatformFile.cpp
@@ -819,16 +819,16 @@ public:
 		//File mapping can extend beyond file size. It's OK, kernel will just fill any leftover page data with zeros
 		const int64 AlignedSize = Align(BytesToMap + Offset - AlignedOffset, FileMappingAlignment);
 
-		int InternalFlags = (EnumHasAnyFlags(Flags.Flags, EMappedFileFlags::EMapShared)) ? MAP_SHARED : MAP_PRIVATE;
-		if (EnumHasAnyFlags(Flags.Flags, EMappedFileFlags::EPreloadHint))
-		{
-			InternalFlags |= MAP_POPULATE;
-		}
-
 		int Protection = PROT_READ;
+		int InternalFlags = (EnumHasAnyFlags(Flags.Flags, EMappedFileFlags::EPreloadHint)) ? MAP_POPULATE : 0;		
 		if (EnumHasAnyFlags(Flags.Flags, EMappedFileFlags::EFileWritable))
 		{
 			Protection |= PROT_WRITE;
+			InternalFlags |= MAP_SHARED;
+		}
+		else
+		{
+			InternalFlags |= MAP_PRIVATE;
 		}
 
 		const uint8* AlignedMapPtr = static_cast<const uint8*>(mmap(nullptr, AlignedSize, Protection, InternalFlags, FileHandle, AlignedOffset));
diff --git a/Engine/Source/Runtime/Core/Public/Async/MappedFileHandle.h b/Engine/Source/Runtime/Core/Public/Async/MappedFileHandle.h
--- a/Engine/Source/Runtime/Core/Public/Async/MappedFileHandle.h
+++ b/Engine/Source/Runtime/Core/Public/Async/MappedFileHandle.h
@@ -85,8 +85,8 @@ enum class EMappedFileFlags
 {
 	ENone = 0,				// Do nothing
 	EPreloadHint = 1,		// Preload the data. This is only a hint and might be ignored, see IMappedFileRegion::PreloadHint
-	EFileWritable = 2,		// Make the mapped file writable. Requires OpenMapped to be called with OpenOptions set to EPlatformFileOpenOptions::AllowWrite
-	EMapShared = 4			// If content of mapping is changed, it will be shared with other processes that might have this mapping opened and\or written to disc
+	EFileWritable = 2,		// Make the mapped file writable. Requires OpenMappedEx to be called with OpenOptions set to EOpenReadFlags::AllowWrite.
+							// This will create a shared mapping on Unix platform to allow any writes to be flushed to disk
 };
 
 ENUM_CLASS_FLAGS(EMappedFileFlags);
