Minor renaming in MapRegion on Android file mapping implamentation to match Unix and iOS naming

diff --git a/Engine/Source/Runtime/Core/Private/Android/AndroidPlatformFile.cpp b/Engine/Source/Runtime/Core/Private/Android/AndroidPlatformFile.cpp
--- a/Engine/Source/Runtime/Core/Private/Android/AndroidPlatformFile.cpp
+++ b/Engine/Source/Runtime/Core/Private/Android/AndroidPlatformFile.cpp
@@ -1111,7 +1111,7 @@ public:
 		close(FileHandle);
 	}
 
-	virtual IMappedFileRegion* MapRegion(int64 Offset = 0, int64 BytesToMap = MAX_int64, FFileMappingFlags InFlags = EMappedFileFlags::ENone) override
+	virtual IMappedFileRegion* MapRegion(int64 Offset = 0, int64 BytesToMap = MAX_int64, FFileMappingFlags Flags = EMappedFileFlags::ENone) override
 	{
 		LLM_PLATFORM_SCOPE(ELLMTag::PlatformMMIO);
 		const int64 CurrentFileSize = GetCurrentFileSize();
@@ -1124,18 +1124,18 @@ public:
 		const int64 AlignedSize = Align(BytesToMap + Offset - AlignedOffset, FileMappingAlignment);
 
 		int Protection = PROT_READ;
-		int Flags = (EnumHasAnyFlags(InFlags.Flags, EMappedFileFlags::EPreloadHint)) ? MAP_POPULATE : 0;
-		if (EnumHasAnyFlags(InFlags.Flags, EMappedFileFlags::EFileWritable))
+		int InternalFlags = (EnumHasAnyFlags(Flags.Flags, EMappedFileFlags::EPreloadHint)) ? MAP_POPULATE : 0;
+		if (EnumHasAnyFlags(Flags.Flags, EMappedFileFlags::EFileWritable))
 		{
 			Protection |= PROT_WRITE;
-			Flags |= MAP_SHARED;
+			InternalFlags |= MAP_SHARED;
 		}
 		else
 		{
-			Flags |= MAP_PRIVATE;
+			InternalFlags |= MAP_PRIVATE;
 		}
 
-		const uint8* AlignedMapPtr = static_cast<const uint8*>(mmap(nullptr, AlignedSize, Protection, Flags, FileHandle, AlignedOffset));
+		const uint8* AlignedMapPtr = static_cast<const uint8*>(mmap(nullptr, AlignedSize, Protection, InternalFlags, FileHandle, AlignedOffset));
 		if (AlignedMapPtr == MAP_FAILED || AlignedMapPtr == nullptr)
 		{
 #if LOG_ANDROID_FILE
