Removed memory mapping preloading code on iOS as it's not supported by the system

diff --git a/Engine/Source/Runtime/Core/Private/IOS/IOSPlatformFile.cpp b/Engine/Source/Runtime/Core/Private/IOS/IOSPlatformFile.cpp
--- a/Engine/Source/Runtime/Core/Private/IOS/IOSPlatformFile.cpp
+++ b/Engine/Source/Runtime/Core/Private/IOS/IOSPlatformFile.cpp
@@ -523,7 +523,7 @@ public:
 		const int64 AlignedSize = Align(BytesToMap + Offset - AlignedOffset, FileMappingAlignment);
 		
 		int Protection = PROT_READ;
-		int InternalFlags = (EnumHasAnyFlags(Flags.Flags, EMappedFileFlags::EPreloadHint)) ? MAP_POPULATE : 0;
+		int InternalFlags = 0;
 		if (EnumHasAnyFlags(Flags.Flags, EMappedFileFlags::EFileWritable))
 		{
 			Protection |= PROT_WRITE;
