Fixed compilation in a GenericPlatformFile due to a merge issue

diff --git a/Engine/Source/Runtime/Core/Private/GenericPlatform/GenericPlatformFile.cpp b/Engine/Source/Runtime/Core/Private/GenericPlatform/GenericPlatformFile.cpp
--- a/Engine/Source/Runtime/Core/Private/GenericPlatform/GenericPlatformFile.cpp
+++ b/Engine/Source/Runtime/Core/Private/GenericPlatform/GenericPlatformFile.cpp
@@ -518,7 +518,7 @@ FFileOpenAsyncResult IPlatformFile::OpenAsyncRead(const TCHAR* Filename, EOpenRe
 {
 	if (EnumHasAnyFlags(Flags, EOpenReadFlags::AllowDelete))
 	{
-		return MakeError(TEXTVIEW("The current IPlatformFile implementation does not support the EReadFlags::AllowDelete flag"));
+		return MakeError(TEXT("The current IPlatformFile implementation does not support the EReadFlags::AllowDelete flag"));
 	}
 
 	const bool bWrite = EnumHasAnyFlags(Flags, EOpenReadFlags::AllowWrite);
@@ -529,7 +529,7 @@ FFileOpenAsyncResult IPlatformFile::OpenAsyncRead(const TCHAR* Filename, EOpenRe
 		return MakeValue(Handle);
 	}
 
-	return MakeError(TEXTVIEW("[Unknown Error] Current IPlatformFile implementation does not support returning error messages"));
+	return MakeError(TEXT("[Unknown Error] Current IPlatformFile implementation does not support returning error messages"));
 }
 
 IAsyncReadFileHandle* IPlatformFile::OpenAsyncRead(const TCHAR* Filename, bool bAllowWrite /*= false*/)
