- FParse::LineExtended changed to use an enum to control some parsing options, instead of a single 'bExact' bool that wasn't very descriptive

diff --git a/Engine/Source/Editor/UnrealEd/Private/EditorObject.cpp b/Engine/Source/Editor/UnrealEd/Private/EditorObject.cpp
--- a/Engine/Source/Editor/UnrealEd/Private/EditorObject.cpp
+++ b/Engine/Source/Editor/UnrealEd/Private/EditorObject.cpp
@@ -311,7 +311,7 @@ static const TCHAR* ImportProperties(
 	// Build list of all text properties.
 	bool ImportedBrush = 0;
 	int32 LinesConsumed = 0;
-	while (FParse::LineExtended(&SourceText, StrLine, LinesConsumed, true))
+	while (FParse::LineExtended(&SourceText, StrLine, LinesConsumed, FParse::ELineExtendedFlags::OldExactMode))
 	{
 		// remove extra whitespace and optional semicolon from the end of the line
 		{
@@ -1019,7 +1019,7 @@ static const TCHAR* ImportCreateSubObjectsStep(
 	int32 LinesConsumed = 0;
 	while (CurrentSourceText < EndOfSourceText)
 	{
-		if (!FParse::LineExtended(&CurrentSourceText, StrLine, LinesConsumed, true))
+		if (!FParse::LineExtended(&CurrentSourceText, StrLine, LinesConsumed, FParse::ELineExtendedFlags::OldExactMode))
 		{
 			break;
 		}
@@ -1444,7 +1444,7 @@ static const TCHAR* ImportPropertiesStep(
 	int32 LinesConsumed = 0;
 	while (CurrentSourceText < EndOfSourceText)
 	{
-		if (FParse::LineExtended(&CurrentSourceText, StrLine, LinesConsumed, true))
+		if (FParse::LineExtended(&CurrentSourceText, StrLine, LinesConsumed, FParse::ELineExtendedFlags::OldExactMode))
 		{
 			// remove extra whitespace and optional semicolon from the end of the line
 			{
diff --git a/Engine/Source/Runtime/Core/Private/Misc/ConfigCacheIni.cpp b/Engine/Source/Runtime/Core/Private/Misc/ConfigCacheIni.cpp
--- a/Engine/Source/Runtime/Core/Private/Misc/ConfigCacheIni.cpp
+++ b/Engine/Source/Runtime/Core/Private/Misc/ConfigCacheIni.cpp
@@ -1463,17 +1463,24 @@ void FillFileFromBuffer(FileType* File, FStringView Buffer, bool bHandleSymbolCo
 	FString ProcessedValue;
 	bool Done = false;
 	bool bHasHandledMetadata = false;
-	while( !Done && Ptr - Buffer.GetData() < Buffer.Len() )
+
+	FParse::ELineExtendedFlags Flags = 
+		FParse::ELineExtendedFlags::SwallowDoubleSlashComments |
+		FParse::ELineExtendedFlags::AllowBracketedMultiline |
+		FParse::ELineExtendedFlags::AllowEscapedEOLMultiline |
+		FParse::ELineExtendedFlags::SwallowExtraEOLs;
+
+	while (!Done && Ptr - Buffer.GetData() < Buffer.Len())
 	{
 		// Advance past new line characters
-		while( *Ptr=='\r' || *Ptr=='\n' )
+		while (*Ptr == '\r' || *Ptr == '\n')
 		{
 			Ptr++;
 		}
 
 		// read the next line
 		int32 LinesConsumed = 0;
-		FParse::LineExtended(&Ptr, /* reset */ TheLine, LinesConsumed, !!File->bPythonConfigParserMode);
+		FParse::LineExtended(&Ptr, /* reset */ TheLine, LinesConsumed, Flags);
 		if (Ptr == nullptr || *Ptr == 0)
 		{
 			Done = true;
diff --git a/Engine/Source/Runtime/Core/Private/Misc/Parse.cpp b/Engine/Source/Runtime/Core/Private/Misc/Parse.cpp
--- a/Engine/Source/Runtime/Core/Private/Misc/Parse.cpp
+++ b/Engine/Source/Runtime/Core/Private/Misc/Parse.cpp
@@ -1175,139 +1175,153 @@ bool FParse::Line(const TCHAR** Stream, FStringView& Result, bool bExact)
 }
 
 template<class T>
-bool ParseLineExtended(const TCHAR** Stream, T& Result, int32& LinesConsumed, bool bExact)
+bool ParseLineExtended(const TCHAR** InOutStream, T& Result, int32& LinesConsumed, FParse::ELineExtendedFlags Flags)
 {
+	const TCHAR* Stream = *InOutStream;
 	bool bGotStream = false;
 	bool bIsQuoted = false;
 	bool bIgnore = false;
 	int32 BracketDepth = 0;
 
+	const bool bBreakOnPipe = EnumHasAnyFlags(Flags, FParse::ELineExtendedFlags::BreakOnPipe);
+	const bool bHandleBracketMultiline = EnumHasAnyFlags(Flags, FParse::ELineExtendedFlags::AllowBracketedMultiline);
+	const bool bHandleEscapedMultiline = EnumHasAnyFlags(Flags, FParse::ELineExtendedFlags::AllowEscapedEOLMultiline);
+	const bool bHandleDoubleSlashComments = EnumHasAnyFlags(Flags, FParse::ELineExtendedFlags::SwallowDoubleSlashComments);
+	const bool bHandleSemicolonComments = EnumHasAnyFlags(Flags, FParse::ELineExtendedFlags::SwallowSemicolonComments);
+
 	Result.Reset();
 	LinesConsumed = 0;
 
-	while (**Stream != TEXT('\0') && ((**Stream != TEXT('\n') && **Stream != TEXT('\r')) || BracketDepth > 0))
-	{
-		// Start of comments.
-		if (!bIsQuoted && !bExact && (*Stream)[0] == TEXT('/') && (*Stream)[1] == TEXT('/'))
+	auto IsLineBreak = [](const TCHAR* Str, bool bProcessPipeAsBreak) -> bool
 		{
-			bIgnore = true;
-		}
+			return Str[0] == '\n' || Str[0] == '\r' || (bProcessPipeAsBreak && Str[0] == '|');
+		};
 
-		// Command chaining.
-		if (!bIsQuoted && !bExact && **Stream == TEXT('|'))
+	auto TryConsumeLineBreak = [](const TCHAR* Str, bool bProcessPipeAsBreak, int& NumChars, int& NumLines) -> bool
 		{
-			break;
-		}
+			if (Str[0] == '\n' || Str[0] == '\r' || (bProcessPipeAsBreak && Str[0] == '|'))
+			{
+				NumChars = 1;
+				// pipes are breaks that don't count as multiplle lines
+				if (Str[0] == '|')
+				{
+					NumLines = 0;
+				}
+				else
+				{
+					NumLines = 1;
+					// look for a \r\n (or \n\r?) pair
+					if ((Str[1] == '\n' || Str[1] == '\r') && Str[0] != Str[1])
+					{
+						NumChars = 2;
+					}
+				}
 
-		bGotStream = true;
+				return true;
+			}
+			NumChars = 0;
+			NumLines = 0;
+			return false;
+		};
 
-		// bracketed line break
-		if (**Stream == TEXT('\n') || **Stream == TEXT('\r'))
+	while (Stream[0] != '\0' && (!IsLineBreak(Stream, bBreakOnPipe && !bIsQuoted) || BracketDepth > 0))
+	{
+		// look for unquoted comments
+		if (!bIsQuoted && 
+			((bHandleDoubleSlashComments && Stream[0] == '/' && Stream[1] == '/') || 
+			(bHandleSemicolonComments && Stream[0] == ';')))
 		{
-			checkSlow(BracketDepth > 0);
-
-			Result += TEXT(' ');
-			LinesConsumed++;
-			(*Stream)++;
-			if (**Stream == TEXT('\n') || **Stream == TEXT('\r'))
-			{
-				(*Stream)++;
-			}
+			bIgnore = true;
 		}
-		// allow line break if the end of the line is a backslash
-		else if (!bIsQuoted && (*Stream)[0] == TEXT('\\') && ((*Stream)[1] == TEXT('\n') || (*Stream)[1] == TEXT('\r')))
+
+		bGotStream = true;
+
+		// process "allowed" line breaks (ones inside {} or after a \)
+		int NumChars, NumLines;
+		if ((bHandleBracketMultiline && BracketDepth > 0 && TryConsumeLineBreak(Stream, false, NumChars, NumLines)) ||
+			(bHandleEscapedMultiline && Stream[0] == '\\' && TryConsumeLineBreak(Stream + 1, false, NumChars, NumLines)))
 		{
 			Result += TEXT(' ');
-			LinesConsumed++;
-			(*Stream) += 2;
-			if (**Stream == TEXT('\n') || **Stream == TEXT('\r'))
+			LinesConsumed += NumLines;
+			if (Stream[0] == '\\')
 			{
-				(*Stream)++;
+				Stream++;
 			}
+			Stream += NumChars;
+
 		}
 		// check for starting or ending brace
-		else if (!bIsQuoted && **Stream == TEXT('{'))
+		else if (bHandleBracketMultiline && !bIsQuoted && Stream[0] == '{')
 		{
 			BracketDepth++;
-			(*Stream)++;
+			Stream++;
 		}
-		else if (!bIsQuoted && **Stream == TEXT('}') && BracketDepth > 0)
+		else if (!bIsQuoted && Stream[0] == '}' && BracketDepth > 0)
 		{
 			BracketDepth--;
-			(*Stream)++;
+			Stream++;
 		}
 		// specifically consume escaped backslashes and quotes within quoted strings
-		else if (bIsQuoted && !bIgnore && (*Stream)[0] == TEXT('\\') && ( (*Stream)[1] == TEXT('\"') || (*Stream)[1] == TEXT('\\') ))
+		else if (bIsQuoted && !bIgnore && Stream[0] == '\\' && (Stream[1] == '\"' || Stream[1] == '\\'))
 		{
-			Result += FStringView(*Stream, 2);
-			(*Stream) += 2;
+			Result += FStringView(Stream, 2);
+			Stream += 2;
 		}
 		else
 		{
-			bIsQuoted = bIsQuoted ^ (**Stream == TEXT('\"'));
+			bIsQuoted = bIsQuoted ^ (Stream[0] == '\"');
 
 			// Got stuff.
 			if (!bIgnore)
 			{
-				Result += *((*Stream)++);
+				Result += *(Stream++);
 			}
 			else
 			{
-				(*Stream)++;
+				Stream++;
 			}
 		}
 	}
-	if (**Stream == 0)
+
+	if (Stream[0] == '\0')
 	{
 		if (bGotStream)
 		{
 			LinesConsumed++;
 		}
 	}
-	else if (bExact)
-	{
-		// Eat up exactly one CR/LF.
-		if (**Stream == TEXT('\r') || **Stream == TEXT('\n'))
-		{
-			LinesConsumed++;
-			if (**Stream == TEXT('\r'))
-			{
-				(*Stream)++;
-			}
-			if (**Stream == TEXT('\n'))
-			{
-				(*Stream)++;
-			}
-		}
-	}
-	else
+	else 
 	{
-		// Eat up all CR/LF's.
-		while (**Stream == TEXT('\n') || **Stream == TEXT('\r') || **Stream == TEXT('|'))
+
+		// start eating up line breaks (\r, \n, \r\n, maybe |)
+		int NumChars, NumLines;
+		while (TryConsumeLineBreak(Stream, bBreakOnPipe, NumChars, NumLines))
 		{
-			if (**Stream != TEXT('|'))
-			{
-				LinesConsumed++;
-			}
-			if (((*Stream)[0] == TEXT('\n') && (*Stream)[1] == TEXT('\r')) || ((*Stream)[0] == TEXT('\r') && (*Stream)[1] == TEXT('\n')))
+			// move past the line break
+			Stream += NumChars;
+			// count lines
+			LinesConsumed += NumLines;
+
+			// if we aren't eating up extra lines, then just stop after one
+			if (!EnumHasAnyFlags(Flags, FParse::ELineExtendedFlags::SwallowExtraEOLs))
 			{
-				(*Stream)++;
+				break;
 			}
-			(*Stream)++;
 		}
 	}
 
-	return **Stream != TEXT('\0') || bGotStream;
+	*InOutStream = Stream;
+	return **InOutStream != '\0' || bGotStream;
 }
 
-bool FParse::LineExtended(const TCHAR** Stream, FString& Result, int32& LinesConsumed, bool bExact)
+bool FParse::LineExtended(const TCHAR** Stream, FString& Result, int32& LinesConsumed, ELineExtendedFlags Flags)
 {
-	return ParseLineExtended(Stream, Result, LinesConsumed, bExact);
+	return ParseLineExtended(Stream, Result, LinesConsumed, Flags);
 }
 
-bool FParse::LineExtended(const TCHAR** Stream, FStringBuilderBase& Result, int32& LinesConsumed, bool bExact)
+bool FParse::LineExtended(const TCHAR** Stream, FStringBuilderBase& Result, int32& LinesConsumed, ELineExtendedFlags Flags)
 {
-	return ParseLineExtended(Stream, Result, LinesConsumed, bExact);
+	return ParseLineExtended(Stream, Result, LinesConsumed, Flags);
 }
 
 uint32 FParse::HexNumber(FStringView HexString)
@@ -1499,7 +1513,7 @@ bool FParseLineExtendedTest::RunTest(const FString& Parameters)
 		Result.Reset();
 
 		const TCHAR* Stream = Tests[Index];
-		bool bSuccess = FParse::LineExtended(&Stream, Result, LinesConsumed, false);
+		bool bSuccess = FParse::LineExtended(&Stream, Result, LinesConsumed, FParse::ELineExtendedFlags::OldDefaultMode);
 		TestTrue(*FString::Printf(TEXT("Expecting parsed line [%s] to be [%s]. Result was [%s]."), Tests[Index], Expected[Index], *Result), bSuccess && Result == Expected[Index]);
 	}
 
diff --git a/Engine/Source/Runtime/Core/Public/Misc/Parse.h b/Engine/Source/Runtime/Core/Public/Misc/Parse.h
--- a/Engine/Source/Runtime/Core/Public/Misc/Parse.h
+++ b/Engine/Source/Runtime/Core/Public/Misc/Parse.h
@@ -7,6 +7,7 @@
 #include "CoreTypes.h"
 #include "HAL/PlatformCrt.h"
 #include "Misc/Build.h"
+#include "Misc/EnumClassFlags.h"
 #include "Templates/Function.h"
 
 class FExec;
@@ -85,8 +86,37 @@ struct FParse
 	/** Get a line of Stream, with support for extending beyond that line with certain characters, e.g. {} and \
 	 * the out character array will not include the ignored endlines
 	 */
-	static CORE_API bool LineExtended(const TCHAR** Stream, FString& Result, int32& LinesConsumed, bool Exact = 0);
-	static CORE_API bool LineExtended(const TCHAR** Stream, FStringBuilderBase& Result, int32& LinesConsumed, bool Exact = 0);
+	enum class ELineExtendedFlags : uint8
+	{
+		None = 0,
+		SwallowDoubleSlashComments = 1,
+		SwallowSemicolonComments = 2,
+		SwallowExtraEOLs = 4,
+		BreakOnPipe = 8,
+		AllowBracketedMultiline = 16, // if a EOL happens between {} curly brackets, insert a space, and continue parsing the next line
+		AllowEscapedEOLMultiline = 32, // if a EOL happens right after a \ backslash, insert a space, and continue parsing the next line
+		
+		// Standard maps to an "expected" set of flags (\<eol> is replaced with space, no comments, don't swallow extra EOLs, no { } processing)
+		Standard = AllowEscapedEOLMultiline,
+
+		// OldDefaultMode maps to what the previous 'bExact' mapped to when false
+		OldDefaultMode = SwallowDoubleSlashComments | BreakOnPipe | AllowBracketedMultiline | AllowEscapedEOLMultiline | SwallowExtraEOLs,
+		// OldExactMode maps to what the previous 'bExact' mapped to when true
+		OldExactMode = AllowBracketedMultiline | AllowEscapedEOLMultiline,
+	};
+
+	static CORE_API bool LineExtended(const TCHAR** Stream, FString& Result, int32& LinesConsumed, ELineExtendedFlags Flags = ELineExtendedFlags::Standard);
+	static CORE_API bool LineExtended(const TCHAR** Stream, FStringBuilderBase& Result, int32& LinesConsumed, ELineExtendedFlags Flags = ELineExtendedFlags::Standard);
+	UE_DEPRECATED(5.6, "Use the version that takes ELineExtendedFlags")
+	static CORE_API bool LineExtended(const TCHAR** Stream, FString& Result, int32& LinesConsumed, bool bExact)
+	{
+		return LineExtended(Stream, Result, LinesConsumed, (bExact ? ELineExtendedFlags::OldExactMode : ELineExtendedFlags::OldDefaultMode));
+	}
+	UE_DEPRECATED(5.6, "Use the version that takes ELineExtendedFlags")
+	static CORE_API bool LineExtended(const TCHAR** Stream, FStringBuilderBase& Result, int32& LinesConsumed, bool bExact)
+	{
+		return LineExtended(Stream, Result, LinesConsumed, (bExact ? ELineExtendedFlags::OldExactMode : ELineExtendedFlags::OldDefaultMode));
+	}
 	/** Grabs the next space-delimited (unless SingleCharacterDelimiter is provided) string from the input stream. If quoted, gets entire quoted string. */
 	static CORE_API bool Token( const TCHAR*& Str, TCHAR* Result, int32 MaxLen, bool bUseEscape, const TCHAR SingleCharacterDelimiter = TEXT('\0'));
 	/** Grabs the next space-delimited (unless SingleCharacterDelimiter is provided) string from the input stream. If quoted, gets entire quoted string. */
@@ -198,6 +228,7 @@ struct FParse
 };
 
 ENUM_CLASS_FLAGS(FParse::EGrammarBasedParseFlags);
+ENUM_CLASS_FLAGS(FParse::ELineExtendedFlags)
 
 #if !UE_BUILD_SHIPPING
 /** Needed for the console command "DumpConsoleCommands" */
diff --git a/Engine/Source/Runtime/Core/Tests/Misc/ParseTests.cpp b/Engine/Source/Runtime/Core/Tests/Misc/ParseTests.cpp
--- a/Engine/Source/Runtime/Core/Tests/Misc/ParseTests.cpp
+++ b/Engine/Source/Runtime/Core/Tests/Misc/ParseTests.cpp
@@ -902,5 +902,172 @@ TEST_CASE("Parse::Value::Numbers", "[Smoke]")
 		CHECK(Result == Approx(ExpectedResults).margin(0.0001f));
 	}
 }
+
+TEST_CASE("Parse::Line", "[Parse][Smoke]")
+{
+	SECTION("Line")
+	{
+		const TCHAR* Lines =
+			TEXT("Line\n")
+			TEXT("Line\r\n")
+			TEXT("Line\n\n\n")
+			TEXT("Line//Comment\n")
+			TEXT("//Comment\n")
+			TEXT("\"//Comment\"\n")
+			TEXT("Line1|Line2\n")
+			TEXT("\"Line1|Line2\"\n");
+
+		FString ExpectedResult;
+		FString Result;
+
+		// exact == true
+		const TCHAR* InputLines = Lines;
+		const TCHAR** Input = &InputLines;
+		// Line\n
+		CHECK(FParse::Line(Input, Result, true)); CHECK(Result == TEXT("Line"));
+		// Line\r\n
+		CHECK(FParse::Line(Input, Result, true)); CHECK(Result == TEXT("Line"));
+		// Line\n\n\n
+		CHECK(FParse::Line(Input, Result, true)); CHECK(Result == TEXT("Line"));
+		CHECK(FParse::Line(Input, Result, true)); CHECK(Result == TEXT(""));
+		CHECK(FParse::Line(Input, Result, true)); CHECK(Result == TEXT(""));
+		// Line//Comment\n
+		CHECK(FParse::Line(Input, Result, true)); CHECK(Result == TEXT("Line//Comment"));
+		// //Comment\n
+		CHECK(FParse::Line(Input, Result, true)); CHECK(Result == TEXT("//Comment"));
+		// \"//Comment\"\n
+		CHECK(FParse::Line(Input, Result, true)); CHECK(Result == TEXT("\"//Comment\""));
+		// Line1|Line2\n
+		CHECK(FParse::Line(Input, Result, true)); CHECK(Result == TEXT("Line1|Line2"));
+		// \"Line1|Line2\"\n
+		CHECK(FParse::Line(Input, Result, true)); CHECK(Result == TEXT("\"Line1|Line2\""));
+
+		// exact == false
+		InputLines = Lines;
+		Input = &InputLines;
+		ExpectedResult = TEXT("Line");
+		// Line\n
+		CHECK(FParse::Line(Input, Result, false)); CHECK(Result == TEXT("Line"));
+		// Line\r\n
+		CHECK(FParse::Line(Input, Result, false)); CHECK(Result == TEXT("Line"));
+		// Line\n\n\n
+		CHECK(FParse::Line(Input, Result, false)); CHECK(Result == TEXT("Line"));
+		// Line//Comment\n
+		CHECK(FParse::Line(Input, Result, false)); CHECK(Result == TEXT("Line//Comment"));
+		// //Comment\n
+		CHECK(FParse::Line(Input, Result, false)); CHECK(Result == TEXT(""));
+		// \"//Comment\"\n
+		CHECK(FParse::Line(Input, Result, false)); CHECK(Result == TEXT("\"//Comment\""));
+		// Line1|Line2\n
+		CHECK(FParse::Line(Input, Result, false)); CHECK(Result == TEXT("Line1"));
+		CHECK(FParse::Line(Input, Result, false)); CHECK(Result == TEXT("Line2"));
+		// \"Line1|Line2\"\n
+		CHECK(FParse::Line(Input, Result, false)); CHECK(Result == TEXT("\"Line1|Line2\""));
+	}
+	SECTION("LineExtended")
+	{
+		const TCHAR* Lines =
+			TEXT("Line\r\n\r\n\n")
+			TEXT("Line//Comment\n")
+			TEXT("Line;Comment\n")
+			TEXT("Line1\\\nLine2\n")
+			TEXT("\"Line1\\\nLine2\"\n")
+			TEXT("Line1|Line2\n")
+			TEXT("\"Line1|Line2\"\n")
+			TEXT("{Line1\nLine2}\n")
+			TEXT("\"{Line1\nLine2}\"\n");
+
+		FString ExpectedResult;
+		FString Result;
+		int LinesConsumed = 0;
+
+
+		// OldDefaultMode (SwallowDoubleSlashComments | BreakOnPipe | AllowBracketedMultiline | AllowEscapedEOLMultiline | SwallowExtraEOLs)
+		const TCHAR* InputLines = Lines;
+		const TCHAR** Input = &InputLines;
+		FParse::ELineExtendedFlags Flags = FParse::ELineExtendedFlags::OldDefaultMode;
+		// Line\r\n\r\n\n
+		CHECK(FParse::LineExtended(Input, Result, LinesConsumed, Flags)); CHECK(Result == TEXT("Line"));
+		// Line//Comment\n
+		CHECK(FParse::LineExtended(Input, Result, LinesConsumed, Flags)); CHECK(Result == TEXT("Line"));
+		// Line;Comment\n
+		CHECK(FParse::LineExtended(Input, Result, LinesConsumed, Flags)); CHECK(Result == TEXT("Line;Comment"));
+		// Line1\\\nLine2\n
+		CHECK(FParse::LineExtended(Input, Result, LinesConsumed, Flags)); CHECK(Result == TEXT("Line1 Line2"));
+		// \"Line1\\\nLine2\"\n
+		CHECK(FParse::LineExtended(Input, Result, LinesConsumed, Flags)); CHECK(Result == TEXT("\"Line1 Line2\""));
+		// Line1|Line2\n
+		CHECK(FParse::LineExtended(Input, Result, LinesConsumed, Flags)); CHECK(Result == TEXT("Line1"));
+		CHECK(FParse::LineExtended(Input, Result, LinesConsumed, Flags)); CHECK(Result == TEXT("Line2"));
+		// \"Line1|Line2\"\n
+		CHECK(FParse::LineExtended(Input, Result, LinesConsumed, Flags)); CHECK(Result == TEXT("\"Line1|Line2\""));
+		// {Line1\nLine2}\n
+		CHECK(FParse::LineExtended(Input, Result, LinesConsumed, Flags)); CHECK(Result == TEXT("Line1 Line2"));
+		// \"{Line1\nLine2}\"\n
+		CHECK(FParse::LineExtended(Input, Result, LinesConsumed, Flags)); CHECK(Result == TEXT("\"{Line1"));
+		CHECK(FParse::LineExtended(Input, Result, LinesConsumed, Flags)); CHECK(Result == TEXT("Line2}\""));
+
+		CHECK(FParse::LineExtended(Input, Result, LinesConsumed, Flags) == false);
+
+		// None
+		InputLines = Lines;
+		Input = &InputLines;
+		Flags = FParse::ELineExtendedFlags::None;
+		// Line\r\n\r\n\n
+		CHECK(FParse::LineExtended(Input, Result, LinesConsumed, Flags)); CHECK(Result == TEXT("Line"));
+		CHECK(FParse::LineExtended(Input, Result, LinesConsumed, Flags)); CHECK(Result == TEXT(""));
+		CHECK(FParse::LineExtended(Input, Result, LinesConsumed, Flags)); CHECK(Result == TEXT(""));
+		// Line//Comment\n
+		CHECK(FParse::LineExtended(Input, Result, LinesConsumed, Flags)); CHECK(Result == TEXT("Line//Comment"));
+		// Line;Comment\n
+		CHECK(FParse::LineExtended(Input, Result, LinesConsumed, Flags)); CHECK(Result == TEXT("Line;Comment"));
+		// Line1\\\nLine2\n
+		CHECK(FParse::LineExtended(Input, Result, LinesConsumed, Flags)); CHECK(Result == TEXT("Line1\\"));
+		CHECK(FParse::LineExtended(Input, Result, LinesConsumed, Flags)); CHECK(Result == TEXT("Line2"));
+		// \"Line1\\\nLine2\"\n
+		CHECK(FParse::LineExtended(Input, Result, LinesConsumed, Flags)); CHECK(Result == TEXT("\"Line1\\"));
+		CHECK(FParse::LineExtended(Input, Result, LinesConsumed, Flags)); CHECK(Result == TEXT("Line2\""));
+		// Line1|Line2\n
+		CHECK(FParse::LineExtended(Input, Result, LinesConsumed, Flags)); CHECK(Result == TEXT("Line1|Line2"));
+		// \"Line1|Line2\"\n
+		CHECK(FParse::LineExtended(Input, Result, LinesConsumed, Flags)); CHECK(Result == TEXT("\"Line1|Line2\""));
+		// {Line1\nLine2}\n
+		CHECK(FParse::LineExtended(Input, Result, LinesConsumed, Flags)); CHECK(Result == TEXT("{Line1"));
+		CHECK(FParse::LineExtended(Input, Result, LinesConsumed, Flags)); CHECK(Result == TEXT("Line2}"));
+		// \"{Line1\nLine2}\"\n
+		CHECK(FParse::LineExtended(Input, Result, LinesConsumed, Flags)); CHECK(Result == TEXT("\"{Line1"));
+		CHECK(FParse::LineExtended(Input, Result, LinesConsumed, Flags)); CHECK(Result == TEXT("Line2}\""));
+
+		CHECK(FParse::LineExtended(Input, Result, LinesConsumed, Flags) == false);
+
+		// Config System mode
+		InputLines = Lines;
+		Input = &InputLines;
+		Flags = FParse::ELineExtendedFlags::SwallowDoubleSlashComments | FParse::ELineExtendedFlags::SwallowSemicolonComments | 
+			FParse::ELineExtendedFlags::AllowBracketedMultiline | FParse::ELineExtendedFlags::AllowEscapedEOLMultiline | 
+			FParse::ELineExtendedFlags::SwallowExtraEOLs;
+		// Line\r\n\r\n\n
+		CHECK(FParse::LineExtended(Input, Result, LinesConsumed, Flags)); CHECK(Result == TEXT("Line"));
+		// Line//Comment\n
+		CHECK(FParse::LineExtended(Input, Result, LinesConsumed, Flags)); CHECK(Result == TEXT("Line"));
+		// Line;Comment\n
+		CHECK(FParse::LineExtended(Input, Result, LinesConsumed, Flags)); CHECK(Result == TEXT("Line"));
+		// Line1\\\nLine2\n
+		CHECK(FParse::LineExtended(Input, Result, LinesConsumed, Flags)); CHECK(Result == TEXT("Line1 Line2"));
+		// \"Line1\\\nLine2\"\n
+		CHECK(FParse::LineExtended(Input, Result, LinesConsumed, Flags)); CHECK(Result == TEXT("\"Line1 Line2\""));
+		// Line1|Line2\n
+		CHECK(FParse::LineExtended(Input, Result, LinesConsumed, Flags)); CHECK(Result == TEXT("Line1|Line2"));
+		// \"Line1|Line2\"\n
+		CHECK(FParse::LineExtended(Input, Result, LinesConsumed, Flags)); CHECK(Result == TEXT("\"Line1|Line2\""));
+		// {Line1\nLine2}\n
+		CHECK(FParse::LineExtended(Input, Result, LinesConsumed, Flags)); CHECK(Result == TEXT("Line1 Line2"));
+		// \"{Line1\nLine2}\"\n
+		CHECK(FParse::LineExtended(Input, Result, LinesConsumed, Flags)); CHECK(Result == TEXT("\"{Line1"));
+		CHECK(FParse::LineExtended(Input, Result, LinesConsumed, Flags)); CHECK(Result == TEXT("Line2}\""));
+
+		CHECK(FParse::LineExtended(Input, Result, LinesConsumed, Flags) == false);
+	}
+}
 #endif
 	
