Follow up fixes for the new IPlatformFile::OpenMappedEx overload that returns a FOpenMappedResult with a value or error (see 37356456).
Change StealValue().Get() to StealValue().Release() when converting result to a raw pointer.
Make some consistency tweaks.
Remove the FFileSystemError constructor taking a FStringView to avoid implicit conversion from FString& to FStringView causing unnecessary copies.

#rnx

diff --git a/Engine/Source/Developer/TraceServices/Private/Model/SymslibResolver.cpp b/Engine/Source/Developer/TraceServices/Private/Model/SymslibResolver.cpp
--- a/Engine/Source/Developer/TraceServices/Private/Model/SymslibResolver.cpp
+++ b/Engine/Source/Developer/TraceServices/Private/Model/SymslibResolver.cpp
@@ -59,12 +59,16 @@ namespace
 		bool Load(const TCHAR* FileName)
 		{
 			FOpenMappedResult Result = FPlatformFileManager::Get().GetPlatformFile().OpenMappedEx(FileName);
-			if (Result.HasError())
+			if (Result.HasValue())
 			{
-				return false;
+				Handle = Result.StealValue();
+				Region.Reset(Handle->MapRegion(0, Handle->GetFileSize()));
+			}
+			else
+			{
+				Handle.Reset();
+				Region.Reset();
 			}
-			Handle = Result.StealValue();
-			Region.Reset(Handle->MapRegion(0, Handle->GetFileSize()));
 			return Region.IsValid();
 		}
 
diff --git a/Engine/Source/Programs/SymsLibDump/Private/SymsLibDump.cpp b/Engine/Source/Programs/SymsLibDump/Private/SymsLibDump.cpp
--- a/Engine/Source/Programs/SymsLibDump/Private/SymsLibDump.cpp
+++ b/Engine/Source/Programs/SymsLibDump/Private/SymsLibDump.cpp
@@ -42,18 +42,16 @@ namespace
 		bool Load(const TCHAR* FileName)
 		{
 			FOpenMappedResult Result = FPlatformFileManager::Get().GetPlatformFile().OpenMappedEx(FileName);
-
-			if (!Result.HasError())
+			if (Result.HasValue())
 			{
 				Handle = Result.StealValue();
 				Region.Reset(Handle->MapRegion(0, Handle->GetFileSize()));
 			}
 			else
 			{
-				Handle = nullptr;
+				Handle.Reset();
 				Region.Reset();
 			}
-
 			return Region.IsValid();
 		}
 
diff --git a/Engine/Source/Runtime/Core/Private/GenericPlatform/GenericPlatformFile.cpp b/Engine/Source/Runtime/Core/Private/GenericPlatform/GenericPlatformFile.cpp
--- a/Engine/Source/Runtime/Core/Private/GenericPlatform/GenericPlatformFile.cpp
+++ b/Engine/Source/Runtime/Core/Private/GenericPlatform/GenericPlatformFile.cpp
@@ -498,7 +498,7 @@ FFileOpenResult IPlatformFile::OpenRead(const TCHAR* Filename, IPlatformFile::EO
 {
 	if (EnumHasAnyFlags(Flags, EOpenReadFlags::AllowDelete))
 	{
-		return MakeError(TEXTVIEW("The current IPlatformFile implementation does not support the EReadFlags::AllowDelete flag"));
+		return MakeError(TEXT("The current IPlatformFile implementation does not support the EReadFlags::AllowDelete flag"));
 	}
 
 	const bool bWrite = EnumHasAnyFlags(Flags, EOpenReadFlags::AllowWrite);
@@ -510,7 +510,7 @@ FFileOpenResult IPlatformFile::OpenRead(const TCHAR* Filename, IPlatformFile::EO
 	}
 	else
 	{
-		return MakeError(TEXTVIEW("[Unknown Error] Current IPlatformFile implementation does not support returning error messages"));
+		return MakeError(TEXT("[Unknown Error] Current IPlatformFile implementation does not support returning error messages"));
 	}
 }
 
@@ -550,7 +550,7 @@ IMappedFileHandle* IPlatformFile::OpenMapped(const TCHAR* Filename)
 
 FOpenMappedResult IPlatformFile::OpenMappedEx(const TCHAR* Filename, EOpenReadFlags OpenOptions)
 {
-	return MakeError(TEXTVIEW("OpenMappedEx is not implemented on this platform"));
+	return MakeError(TEXT("OpenMappedEx is not implemented on this platform"));
 }
 
 DEFINE_STAT(STAT_AsyncFileMemory);
diff --git a/Engine/Source/Runtime/Core/Private/Unix/UnixPlatformFile.cpp b/Engine/Source/Runtime/Core/Private/Unix/UnixPlatformFile.cpp
--- a/Engine/Source/Runtime/Core/Private/Unix/UnixPlatformFile.cpp
+++ b/Engine/Source/Runtime/Core/Private/Unix/UnixPlatformFile.cpp
@@ -1238,7 +1238,7 @@ FOpenMappedResult FUnixPlatformFile::OpenMappedEx(const TCHAR* Filename, EOpenRe
 		FString ErrorStr = FString::Printf(TEXT("open('%s', Flags=0x%08X) failed: errno=%d (%s)"), *NormalizedFilename, Flags, ErrNo, UTF8_TO_TCHAR(strerror(ErrNo)));
 		UE_LOG_UNIX_FILE(Warning, TEXT("%s"), *ErrorStr);
 
-		return MakeError(ErrorStr);
+		return MakeError(MoveTemp(ErrorStr));
 	}
 	
 	struct stat FileInfo;
@@ -1250,7 +1250,7 @@ FOpenMappedResult FUnixPlatformFile::OpenMappedEx(const TCHAR* Filename, EOpenRe
 		FString ErrorStr = FString::Printf(TEXT("stat('%s', Flags=0x%08X) failed: errno=%d (%s)"), *NormalizedFilename, Flags, ErrNo, UTF8_TO_TCHAR(strerror(ErrNo)));
 		UE_LOG_UNIX_FILE(Warning, TEXT("%s"), *ErrorStr);
 
-		return MakeError(ErrorStr);
+		return MakeError(MoveTemp(ErrorStr));
 	}
 	
 	return MakeValue(new FUnixMappedFileHandle(Handle, FileInfo.st_size, NormalizedFilename));
diff --git a/Engine/Source/Runtime/Core/Private/Windows/WindowsPlatformFile.cpp b/Engine/Source/Runtime/Core/Private/Windows/WindowsPlatformFile.cpp
--- a/Engine/Source/Runtime/Core/Private/Windows/WindowsPlatformFile.cpp
+++ b/Engine/Source/Runtime/Core/Private/Windows/WindowsPlatformFile.cpp
@@ -1413,30 +1413,26 @@ FFileOpenResult FWindowsPlatformFile::OpenRead(const TCHAR* Filename, EOpenReadF
 		else
 		{
 			delete FileHandle;
-			return MakeError(TEXTVIEW("Failed to create FFileHandleWindows"));
+			return MakeError(TEXT("Failed to create FFileHandleWindows"));
 		}
 	}
 #endif //USE_OVERLAPPED_IO
 	else
 	{
 		TRACE_PLATFORMFILE_FAIL_OPEN(Filename);
-		return MakeError(TEXTVIEW("Failed to open file for reading"), FPlatformMisc::GetLastError());
+		return MakeError(TEXT("Failed to open file for reading"), FPlatformMisc::GetLastError());
 	}
 }
 
 IFileHandle* FWindowsPlatformFile::OpenRead(const TCHAR* Filename, bool bAllowWrite)
 {
 	const EOpenReadFlags Flags = bAllowWrite ? EOpenReadFlags::AllowWrite : EOpenReadFlags::None;
-
 	FFileOpenResult Result = OpenRead(Filename, Flags);
-	if (Result.IsValid())
-	{
-		return Result.StealValue().Release();
-	}
-	else
+	if (Result.HasError())
 	{
 		return nullptr;
 	}
+	return Result.StealValue().Release();
 }
 
 IFileHandle* FWindowsPlatformFile::OpenReadNoBuffering(const TCHAR* Filename, bool bAllowWrite)
@@ -1513,7 +1509,7 @@ FOpenMappedResult FWindowsPlatformFile::OpenMappedEx(const TCHAR* Filename, EOpe
 	int64 Size = FileSize(Filename);
 	if (Size < 1)
 	{
-		return MakeError(TEXTVIEW("File size is 0"));
+		return MakeError(TEXT("File size is 0"));
 	}
 	uint32 Access = GENERIC_READ;
 	if (EnumHasAnyFlags(OpenOptions, EOpenReadFlags::AllowWrite))
diff --git a/Engine/Source/Runtime/Core/Public/GenericPlatform/GenericPlatformFile.h b/Engine/Source/Runtime/Core/Public/GenericPlatform/GenericPlatformFile.h
--- a/Engine/Source/Runtime/Core/Public/GenericPlatform/GenericPlatformFile.h
+++ b/Engine/Source/Runtime/Core/Public/GenericPlatform/GenericPlatformFile.h
@@ -291,12 +291,6 @@ public:
 	{
 	}
 
-	explicit FFileSystemError(FStringView InErrorMessage, int32 InSystemErrorCode = 0)
-		: ErrorMessage(InErrorMessage)
-		, SystemErrorCode(InSystemErrorCode)
-	{
-	}
-
 	/** Return the error message, if a valid system error code was provided then this will be appended to the end of the message */
 	FString GetMessage()
 	{
diff --git a/Engine/Source/Runtime/PakFile/Private/IPlatformFilePak.cpp b/Engine/Source/Runtime/PakFile/Private/IPlatformFilePak.cpp
--- a/Engine/Source/Runtime/PakFile/Private/IPlatformFilePak.cpp
+++ b/Engine/Source/Runtime/PakFile/Private/IPlatformFilePak.cpp
@@ -4593,7 +4593,7 @@ FOpenMappedResult FPakPlatformFile::OpenMappedEx(const TCHAR* Filename, EOpenRea
 {
 	if (!GMMIO_Enable)
 	{
-		return MakeError(TEXTVIEW("Memory mapped IO is not enabled on this platform"));
+		return MakeError(TEXT("Memory mapped IO is not enabled on this platform"));
 	}
 
 #if !UE_BUILD_SHIPPING
@@ -4622,7 +4622,7 @@ FOpenMappedResult FPakPlatformFile::OpenMappedEx(const TCHAR* Filename, EOpenRea
 			FOpenMappedResult Result = LowerLevel->OpenMappedEx(*PakEntry->GetFilename(), OpenOptions);
 			if (Result.HasError())
 			{
-				return MakeError(Result.GetError().GetMessage());
+				return MakeError(Result.StealError());
 			}
 
 			PakEntry->MappedFileHandle = Result.StealValue().Release();
