Add opens around atof, strtof and on Linux strtof32/64
Record the writes to the nptr if not null

diff --git a/Engine/Source/Runtime/Core/Private/AutoRTFM/StdLibSupport.cpp b/Engine/Source/Runtime/Core/Private/AutoRTFM/StdLibSupport.cpp
--- a/Engine/Source/Runtime/Core/Private/AutoRTFM/StdLibSupport.cpp
+++ b/Engine/Source/Runtime/Core/Private/AutoRTFM/StdLibSupport.cpp
@@ -215,6 +215,7 @@ UE_AUTORTFM_REGISTER_SELF_OVERLOADED_FUNCTION(const char* (const char*, int), st
 UE_AUTORTFM_REGISTER_SELF_OVERLOADED_FUNCTION(const char* (const char*, int), strrchr);
 UE_AUTORTFM_REGISTER_SELF_OVERLOADED_FUNCTION(const char* (const char*, const char*), strstr);
 UE_AUTORTFM_REGISTER_SELF_FUNCTION(strlen);
+UE_AUTORTFM_REGISTER_SELF_FUNCTION(atof);
 
 UE_AUTORTFM_REGISTER_SELF_OVERLOADED_FUNCTION(const wchar_t* (const wchar_t*, wchar_t), wcschr);
 UE_AUTORTFM_REGISTER_SELF_OVERLOADED_FUNCTION(wchar_t* (wchar_t*, wchar_t), wcschr);
@@ -230,8 +231,7 @@ RETURN RTFM_StringToFloat(const CHAR* String, CHAR** EndPtr) throw()
 {
     if (nullptr != EndPtr)
     {
-        FContext* Context = FContext::Get();
-        Context->RecordWrite(EndPtr, sizeof(CHAR*));
+        AutoRTFM::RecordOpenWrite(EndPtr);
     }
     return FN(String, EndPtr);
 }
@@ -241,13 +241,14 @@ RETURN RTFM_StringToInt(const CHAR* String, CHAR** EndPtr, int Radix) throw()
 {
     if (nullptr != EndPtr)
     {
-        FContext* Context = FContext::Get();
-        Context->RecordWrite(EndPtr, sizeof(CHAR*));
+        AutoRTFM::RecordOpenWrite(EndPtr);
     }
     return FN(String, EndPtr, Radix);
 }
 
 UE_AUTORTFM_REGISTER_OPEN_FUNCTION_EXPLICIT(strtol, (RTFM_StringToInt<long, char, strtol>));
+UE_AUTORTFM_REGISTER_OPEN_FUNCTION_EXPLICIT(strtof, (RTFM_StringToFloat<float, char, strtof>));
+UE_AUTORTFM_REGISTER_OPEN_FUNCTION_EXPLICIT(strtod, (RTFM_StringToFloat<double, char, strtod>));
 UE_AUTORTFM_REGISTER_OPEN_FUNCTION_EXPLICIT(wcstod, (RTFM_StringToFloat<double, wchar_t, wcstod>));
 UE_AUTORTFM_REGISTER_OPEN_FUNCTION_EXPLICIT(wcstof, (RTFM_StringToFloat<float, wchar_t, wcstof>));
 UE_AUTORTFM_REGISTER_OPEN_FUNCTION_EXPLICIT(wcstold, (RTFM_StringToFloat<long double, wchar_t, wcstold>));
@@ -262,8 +263,7 @@ RETURN RTFM_StringToFloat(const CHAR* String, CHAR** EndPtr, _locale_t Locale)
 {
     if (nullptr != EndPtr)
     {
-        FContext* Context = FContext::Get();
-        Context->RecordWrite(EndPtr, sizeof(CHAR*));
+        AutoRTFM::RecordOpenWrite(EndPtr);
     }
     return FN(String, EndPtr, Locale);
 }
@@ -273,8 +273,7 @@ RETURN RTFM_StringToInt(const CHAR* String, CHAR** EndPtr, int Radix, _locale_t
 {
     if (nullptr != EndPtr)
     {
-        FContext* Context = FContext::Get();
-        Context->RecordWrite(EndPtr, sizeof(CHAR*));
+        AutoRTFM::RecordOpenWrite(EndPtr);
     }
     return FN(String, EndPtr, Radix, Locale);
 }
@@ -682,6 +681,8 @@ UE_AUTORTFM_REGISTER_SELF_FUNCTION(gmtime_r);
 UE_AUTORTFM_REGISTER_SELF_FUNCTION(localtime_r);
 UE_AUTORTFM_REGISTER_SELF_FUNCTION(bcmp);
 UE_AUTORTFM_REGISTER_SELF_FUNCTION(pthread_getspecific);
+UE_AUTORTFM_REGISTER_OPEN_FUNCTION_EXPLICIT(strtof32, (RTFM_StringToFloat<float, char, strtof32>));
+UE_AUTORTFM_REGISTER_OPEN_FUNCTION_EXPLICIT(strtof64, (RTFM_StringToFloat<double, char, strtof64>));
 
 int RTFM_stat(const char* Path, struct stat* StatBuf) throw()
 {
