[Core] Added Reserve to StringBuilders, which adjusts the builders capacity without altering CurrPos

diff --git a/Engine/Source/Runtime/Core/Public/Misc/StringBuilder.h b/Engine/Source/Runtime/Core/Public/Misc/StringBuilder.h
--- a/Engine/Source/Runtime/Core/Public/Misc/StringBuilder.h
+++ b/Engine/Source/Runtime/Core/Public/Misc/StringBuilder.h
@@ -165,6 +165,23 @@ public:
 		return OldCount;
 	}
 
+	/**
+	 * Reserves memory such that the string builder can contain at least the given number of characters.
+	 *
+	 * @note: Unlike AddUninitialized, this will not advance CurPos, and thus the various Append calls
+	 * will start at the end of the string builder (from prior to this Reserve call).
+	 * @param InNewCapacity The total number of characters the string builder can contain.
+	 */
+	inline void Reserve(int32 InNewCapacity)
+	{
+		checkSlow(InNewCapacity >= 0);
+		const int32 CurrentCapacity = (int32)(End - Base);
+		if (CurrentCapacity < InNewCapacity)
+		{
+			Extend(InNewCapacity - CurrentCapacity);
+		}
+	}
+
 	/**
 	 * Modifies the string builder to remove the given number of characters from the end.
 	 */
