[Core] Added ToPascalCase utility

#tests Added corresponding unit tests

diff --git a/Engine/Source/Runtime/Core/Private/String/CaseConversion.cpp b/Engine/Source/Runtime/Core/Private/String/CaseConversion.cpp
--- a/Engine/Source/Runtime/Core/Private/String/CaseConversion.cpp
+++ b/Engine/Source/Runtime/Core/Private/String/CaseConversion.cpp
@@ -1,6 +1,10 @@
 // Copyright Epic Games, Inc. All Rights Reserved.
 
 #include "String/CaseConversion.h"
+
+#include "Containers/UnrealString.h"
+#include "Internationalization/BreakIterator.h"
+#include "Misc/AsciiSet.h"
 #include "Misc/StringBuilder.h"
 
 namespace UE::String::Private
@@ -28,6 +32,70 @@ static inline void LowerCaseTo(TStringView<CharType> Input, TStringBuilderBase<C
 	}
 }
 
+// @note: Currently only supports TCHAR (as does IBreakIterator)
+static inline void PascalCaseTo(FStringView Input, FStringBuilderBase& Output)
+{
+	// Create a sanitized string, with certain chars removed
+	TStringBuilderWithBuffer<TCHAR, 64> CleanStringBuilder;
+	CleanStringBuilder.AddUninitialized(Input.Len());
+
+	// Remove apostrophes before converting case, to avoid "You're" becoming "YouRe"
+	{
+		TCHAR* CleanStringIterator = CleanStringBuilder.GetData();
+
+		int32 NumValidChars = 0;
+		constexpr FAsciiSet CharsToRemove("\'\"");
+		for (const TCHAR& Chr : Input)
+		{
+			if (!CharsToRemove.Contains(Chr))
+			{
+				*CleanStringIterator++ = Chr;
+				++NumValidChars;
+			}
+		}
+
+		const int32 NumCharsToRemove = Input.Len() - NumValidChars;
+		CleanStringBuilder.RemoveSuffix(NumCharsToRemove); // Trim excess chars
+	}
+
+	const TSharedRef<IBreakIterator> BreakIterator = FBreakIterator::CreateCamelCaseBreakIterator();
+
+	const FStringView CleanStringView(CleanStringBuilder);
+	BreakIterator->SetStringRef(CleanStringView);
+
+	Output.Reserve(Output.Len() + CleanStringView.Len());
+
+	{
+		// Remove spaces, snake_case, dashes, dots
+		constexpr FAsciiSet CharsToRemove(" \t_-.");
+		for (int32 PrevBreak = 0, NameBreak = BreakIterator->MoveToNext(); NameBreak != INDEX_NONE; NameBreak = BreakIterator->MoveToNext())
+		{
+			const TCHAR& Char = CleanStringView[PrevBreak];
+			PrevBreak++;
+
+			// Char was a space, etc. Skip over.
+			if (CharsToRemove.Contains(Char))
+			{
+				continue;
+			}
+
+			// Uppercase leading character
+			Output.AppendChar(TChar<TCHAR>::ToUpper(Char));
+
+			if (PrevBreak < CleanStringBuilder.Len())
+			{
+				FStringView Trimmed = MakeStringView(&CleanStringView[PrevBreak], NameBreak - PrevBreak);
+				Trimmed = FAsciiSet::TrimPrefixWith(Trimmed, CharsToRemove);
+				Trimmed = FAsciiSet::TrimSuffixWith(Trimmed, CharsToRemove);
+
+				LowerCaseTo(Trimmed, Output);
+			}
+
+			PrevBreak = NameBreak;
+		}
+	}
+}
+
 } // UE::String::Private
 
 namespace UE::String
@@ -63,4 +131,9 @@ void LowerCaseTo(FWideStringView Input, FWideStringBuilderBase& Output)
 	Private::LowerCaseTo(Input, Output);
 }
 
+void PascalCaseTo(FStringView Input, FStringBuilderBase& Output)
+{
+	Private::PascalCaseTo(Input, Output);
+}
+
 } // UE::String
diff --git a/Engine/Source/Runtime/Core/Public/String/CaseConversion.h b/Engine/Source/Runtime/Core/Public/String/CaseConversion.h
--- a/Engine/Source/Runtime/Core/Public/String/CaseConversion.h
+++ b/Engine/Source/Runtime/Core/Public/String/CaseConversion.h
@@ -2,6 +2,7 @@
 
 #pragma once
 
+#include "Containers/ContainersFwd.h"
 #include "Containers/StringFwd.h"
 #include "Containers/StringView.h"
 #include "CoreTypes.h"
@@ -15,6 +16,9 @@ struct TUpperCase { TStringView<CharType> Input; };
 template <typename CharType>
 struct TLowerCase { TStringView<CharType> Input; };
 
+template <typename CharType>
+struct TPascalCase { TStringView<CharType> Input; };
+
 } // UE::String::Private
 
 namespace UE::String
@@ -50,6 +54,19 @@ CORE_API void LowerCaseTo(FAnsiStringView Input, FAnsiStringBuilderBase& Output)
 CORE_API void LowerCaseTo(FUtf8StringView Input, FUtf8StringBuilderBase& Output);
 CORE_API void LowerCaseTo(FWideStringView Input, FWideStringBuilderBase& Output);
 
+/**
+ * Convert the string to PascalCase and append to the string builder.
+ *
+ * @note Only ASCII characters will be converted.
+ * @param Input   The string to convert to PascalCase.
+ * @param Output  The string builder to append the converted string to.
+ *
+ * @code
+ *	PascalCaseTo(TEXTVIEW("EXAMPLE TEXT"), Output); // Output now contains "ExampleText"
+ * @endcode
+ */
+CORE_API void PascalCaseTo(FStringView Input, FStringBuilderBase& Output);
+
 /**
  * Convert the string to uppercase when appended to a string builder.
  *
@@ -98,6 +115,22 @@ inline Private::TLowerCase<WIDECHAR> LowerCase(FWideStringView Input)
 	return {Input};
 }
 
+/**
+ * Convert the string to PascalCase when appended to a string builder.
+ *
+ * @note Only ASCII characters will be converted.
+ * @param Input   The string to convert to PascalCase.
+ * @return An anonymous object to append to a string builder.
+ *
+ * @code
+ *	Builder << String::PascalCase(TEXTVIEW("EXAMPLE TEXT")); // Builder now contains "ExampleText"
+ * @endcode
+ */
+inline Private::TPascalCase<TCHAR> PascalCase(FStringView Input)
+{
+	return {Input};
+}
+
 } // UE::String
 
 namespace UE::String::Private
@@ -117,4 +150,11 @@ inline TStringBuilderBase<CharType>& operator<<(TStringBuilderBase<CharType>& Bu
 	return Builder;
 }
 
+template <typename CharType>
+inline TStringBuilderBase<CharType>& operator<<(TStringBuilderBase<CharType>& Builder, const TPascalCase<CharType>& Adapter)
+{
+	PascalCaseTo(Adapter.Input, Builder);
+	return Builder;
+}
+
 } // UE::String::Private
diff --git a/Engine/Source/Runtime/Core/Tests/String/CaseConversionTest.cpp b/Engine/Source/Runtime/Core/Tests/String/CaseConversionTest.cpp
--- a/Engine/Source/Runtime/Core/Tests/String/CaseConversionTest.cpp
+++ b/Engine/Source/Runtime/Core/Tests/String/CaseConversionTest.cpp
@@ -88,6 +88,22 @@ TEST_CASE_NAMED(FStringLowerCaseTest, "System::Core::String::LowerCase", "[Core]
 	}
 }
 
+TEST_CASE_NAMED(FStringPascalCaseTest, "System::Core::String::PascalCase", "[Core][String][SmokeFilter]")
+{
+	SECTION("PascalCase with WIDECHAR")
+	{
+		CHECK(WriteToWideString<128>(PascalCase(WIDETEXTVIEW("NO PASCALCASE"))) == WIDETEXTVIEW("NoPascalcase"));
+		CHECK(WriteToWideString<128>(PascalCase(WIDETEXTVIEW("Mixed CASE"))) == WIDETEXTVIEW("MixedCase"));
+		CHECK(WriteToWideString<128>(PascalCase(WIDETEXTVIEW("lowercase words"))) == WIDETEXTVIEW("LowercaseWords"));
+		CHECK(WriteToWideString<128>(PascalCase(WIDETEXTVIEW("  SPACES   AND\tTABS\n"))) == WIDETEXTVIEW("SpacesAndTabs\n"));
+		CHECK(WriteToWideString<128>(PascalCase(WIDETEXTVIEW(""))) == WIDETEXTVIEW(""));
+		CHECK(WriteToWideString<128>(PascalCase(WIDETEXTVIEW("A"))) == WIDETEXTVIEW("A"));
+		CHECK(WriteToWideString<128>(PascalCase(WIDETEXTVIEW("z"))) == WIDETEXTVIEW("Z"));
+		CHECK(WriteToWideString<128>(PascalCase(WIDETEXTVIEW("1234!@#$"))) == WIDETEXTVIEW("1234!@#$"));
+		CHECK(WriteToWideString<128>(PascalCase(WIDETEXTVIEW("You're Text"))) == WIDETEXTVIEW("YoureText"));
+	}
+}
+
 } // UE::String
 
 #endif //WITH_TESTS
