Fixed errors in TRingBuffer when using an allocator with NeedsElementType == false.

diff --git a/Engine/Source/Runtime/Core/Public/Containers/RingBuffer.h b/Engine/Source/Runtime/Core/Public/Containers/RingBuffer.h
--- a/Engine/Source/Runtime/Core/Public/Containers/RingBuffer.h
+++ b/Engine/Source/Runtime/Core/Public/Containers/RingBuffer.h
@@ -813,7 +813,7 @@ private:
 			// Allocate SwapStorage
 			ElementAllocatorType SwapStorageBuffer;
 			AllocatorResizeAllocation(SwapStorageBuffer, 0, NewCapacity);
-			ElementType* SwapStorage = SwapStorageBuffer.GetAllocation();
+			ElementType* SwapStorage = (ElementType*)SwapStorageBuffer.GetAllocation();
 
 			// Move data to swap storage
 			const StorageModuloType MaskedFront = Front & IndexMask;
@@ -976,15 +976,15 @@ private:
 	}
 
 	/** Return a pointer to the underlying storage of the RingBuffer */
-	const ElementType* GetStorage() const
+	FORCEINLINE const ElementType* GetStorage() const
 	{
-		return static_cast<const ElementType*>(AllocatorInstance.GetAllocation());
+		return const_cast<TRingBuffer*>(this)->GetStorage();
 	}
 
 	/** Return a const pointer to the underlying storage of the RingBuffer */
 	ElementType* GetStorage()
 	{
-		return static_cast<ElementType*>(AllocatorInstance.GetAllocation());
+		return (ElementType*)AllocatorInstance.GetAllocation();
 	}
 
 	/* Check and return whether the given Index is within range. */
