ParkingLot: Fixed a race reported by TSan

diff --git a/Engine/Source/Runtime/Core/Private/Async/ParkingLot.cpp b/Engine/Source/Runtime/Core/Private/Async/ParkingLot.cpp
--- a/Engine/Source/Runtime/Core/Private/Async/ParkingLot.cpp
+++ b/Engine/Source/Runtime/Core/Private/Async/ParkingLot.cpp
@@ -59,7 +59,7 @@ private:
 struct alignas(PLATFORM_CACHE_LINE_SIZE) FThread final
 {
 	FThread* Next = nullptr;
-	const void* WaitAddress = nullptr;
+	std::atomic<const void*> WaitAddress = nullptr;
 	uint64 WakeToken = 0;
 	FPlatformManualResetEvent Event;
 	std::atomic<uint32> ReferenceCount = 0;
@@ -483,7 +483,7 @@ void FTable::Reserve(const uint32 ThreadCount)
 		// Add waiting threads to the new table.
 		for (FThread* Thread : Threads)
 		{
-			const uint32 Hash = HashAddress(Thread->WaitAddress);
+			const uint32 Hash = HashAddress(Thread->WaitAddress.load(std::memory_order_relaxed));
 			const uint32 Index = Hash % NewTable.BucketCount;
 			FBucket& Bucket = NewTable.FindOrCreateBucket(Index, AllocateBucket);
 			Bucket.Enqueue(Thread);
@@ -586,7 +586,7 @@ FWaitState WaitUntil(const void* Address, const TFunctionRef<bool()>& CanWait, c
 
 	TRefCountPtr<FThread> Self = FThreadLocalData::Get();
 
-	checkfSlow(!Self->WaitAddress, TEXT("WaitAddress must be null. This can happen if Wait is called by BeforeWait."));
+	checkfSlow(!Self->WaitAddress.load(std::memory_order_relaxed), TEXT("WaitAddress must be null. This can happen if Wait is called by BeforeWait."));
 	checkfSlow(Self->WakeToken == 0, TEXT("WakeToken must be 0. This is an error in ParkingLot."));
 
 	FWaitState State;
@@ -600,7 +600,7 @@ FWaitState WaitUntil(const void* Address, const TFunctionRef<bool()>& CanWait, c
 		{
 			return State;
 		}
-		Self->WaitAddress = Address;
+		Self->WaitAddress.store(Address, std::memory_order_relaxed);
 		Self->Event.Reset();
 		Bucket.Enqueue(&*Self);
 	}
@@ -612,7 +612,7 @@ FWaitState WaitUntil(const void* Address, const TFunctionRef<bool()>& CanWait, c
 	Self->Event.WaitUntil(WaitTime);
 
 	// WaitAddress is reset when the thread is dequeued.
-	if (!Self->WaitAddress)
+	if (!Self->WaitAddress.load(std::memory_order_relaxed))
 	{
 		State.bDidWake = true;
 		State.WakeToken = Self->WakeToken;
@@ -630,7 +630,7 @@ FWaitState WaitUntil(const void* Address, const TFunctionRef<bool()>& CanWait, c
 			if (Thread == Self)
 			{
 				bDequeued = true;
-				Thread->WaitAddress = nullptr;
+				Thread->WaitAddress.store(nullptr, std::memory_order_relaxed);
 				return EQueueAction::RemoveAndStop;
 			}
 			return EQueueAction::Continue;
@@ -662,7 +662,7 @@ void WakeOne(const void* Address, const TFunctionRef<uint64(FWakeState)>& OnWake
 		FBucket& Bucket = FTable::FindOrCreateBucket(Address, BucketLock);
 		Bucket.DequeueIf([Address, &WakeThread](FThread* Thread)
 		{
-			if (Thread->WaitAddress == Address)
+			if (Thread->WaitAddress.load(std::memory_order_relaxed) == Address)
 			{
 				WakeThread = Thread;
 				return EQueueAction::RemoveAndStop;
@@ -677,9 +677,9 @@ void WakeOne(const void* Address, const TFunctionRef<uint64(FWakeState)>& OnWake
 
 	if (WakeThread)
 	{
-		checkSlow(WakeThread->WaitAddress == Address);
+		checkSlow(WakeThread->WaitAddress.load(std::memory_order_relaxed) == Address);
 		WakeThread->WakeToken = WakeToken;
-		WakeThread->WaitAddress = nullptr;
+		WakeThread->WaitAddress.store(nullptr, std::memory_order_relaxed);
 		WakeThread->Event.Notify();
 	}
 }
@@ -710,7 +710,7 @@ uint32 WakeMultiple(const void* const Address, const uint32 WakeCount)
 	{
 		Bucket->DequeueIf([Address, &WakeThreads, WakeCount](FThread* Thread)
 		{
-			if (Thread->WaitAddress == Address)
+			if (Thread->WaitAddress.load(std::memory_order_relaxed) == Address)
 			{
 				const uint32 Count = uint32(WakeThreads.Add(Thread) + 1);
 				return Count == WakeCount ? EQueueAction::RemoveAndStop : EQueueAction::RemoveAndContinue;
@@ -721,8 +721,8 @@ uint32 WakeMultiple(const void* const Address, const uint32 WakeCount)
 
 	for (FThread* WakeThread : WakeThreads)
 	{
-		checkSlow(WakeThread->WaitAddress == Address);
-		WakeThread->WaitAddress = nullptr;
+		checkSlow(WakeThread->WaitAddress.load(std::memory_order_relaxed) == Address);
+		WakeThread->WaitAddress.store(nullptr, std::memory_order_relaxed);
 		WakeThread->Event.Notify();
 	}
 
