Fixed invalid comment about how to disable UE_INTERNAL / UE_EXPERIMENTAL on a given module

#skipci
#rnx

diff --git a/Engine/Source/Runtime/Core/Public/Misc/CoreMiscDefines.h b/Engine/Source/Runtime/Core/Public/Misc/CoreMiscDefines.h
--- a/Engine/Source/Runtime/Core/Public/Misc/CoreMiscDefines.h
+++ b/Engine/Source/Runtime/Core/Public/Misc/CoreMiscDefines.h
@@ -289,7 +289,7 @@ enum EPerElement {PerElement};
  *  future releases and that they should use at their risk. 
  *  Using an API marked as such in a non-engine module will trigger a compilation warning (this relies on the c++ [[deprecated]] attribute internally) that can be suppressed by either : 
  *   - Using PRAGMA_DISABLE_INTERNAL_WARNINGS and PRAGMA_ENABLE_INTERNAL_WARNINGS around the call site (preferred method)
- *   - Setting bValidateInternalApi to true in the calling module's build rules, thus telling the entire module that any internal API call is accepted
+ *   - Setting bValidateInternalApi to false in the calling module's build rules, thus telling the entire module that any internal API call is accepted
  *   - Setting bTreatAsEngineModule to true in the calling module's build rules, thus telling the entire module that any internal (OR experimental) API call is accepted
  * 
  * It is advised to simply avoid using APIs marked as such in non-engine modules. 
@@ -311,7 +311,7 @@ enum EPerElement {PerElement};
  *  future releases and that they should use at their own risk. 
  *  Using an API marked as such in a non-engine module will trigger a compilation warning (this relies on the c++ [[deprecated]] attribute internally) that can be suppressed by either : 
  *   - Using PRAGMA_DISABLE_EXPERIMENTAL_WARNINGS and PRAGMA_ENABLE_EXPERIMENTAL_WARNINGS around the call site (preferred method)
- *   - Setting bValidateExperimentalApi to true in the calling module's build rules, thus telling the entire module that any experimental API call is accepted
+ *   - Setting bValidateExperimentalApi to false in the calling module's build rules, thus telling the entire module that any experimental API call is accepted
  *   - Setting bTreatAsEngineModule to true in the calling module's build rules, thus telling the entire module that any experimental (OR internal) API call is accepted
  *
  * Please refer to UE_DEPRECATED for how to use the macro. 
