Changed FAlignedAllocator to use FMemory::Malloc\Free instead of Ansi allocator and do alignment internally.
Removed now obsolete PLATFORM_USES_ANSI_MALLOC define from the code

diff --git a/Engine/Source/Runtime/Core/Public/Async/Fundamental/TaskDelegate.h b/Engine/Source/Runtime/Core/Public/Async/Fundamental/TaskDelegate.h
--- a/Engine/Source/Runtime/Core/Public/Async/Fundamental/TaskDelegate.h
+++ b/Engine/Source/Runtime/Core/Public/Async/Fundamental/TaskDelegate.h
@@ -209,6 +209,7 @@ namespace LowLevelTasks
 				static_assert(TIsInvocable<TCallableType, ParamTypes...>::Value, "TCallableType is not invocable");
 				static_assert(std::is_same_v<ReturnType, decltype(Callable(UE::Core::Private::IsInvocable::DeclVal<ParamTypes>()...))>, "TCallableType return type does not match");
 				static_assert(sizeof(TTaskDelegateImpl<TCallableType, true>) == sizeof(TTaskDelegateBase), "Size must match the Baseclass");
+				static_assert(alignof(TCallableType) <= FLowLevelTasksBlockAllocationTag::Allocator::MaxAlignment);
 				TCallableType** HeapPtr = reinterpret_cast<TCallableType**>(InlineData);
 				*HeapPtr = reinterpret_cast<TCallableType*>(TConcurrentLinearAllocator<FLowLevelTasksBlockAllocationTag>::template Malloc<alignof(TCallableType)>(sizeof(TCallableType)));
 				new (*HeapPtr) TCallableType(Forward<CallableT>(Callable));
diff --git a/Engine/Source/Runtime/Core/Public/Experimental/ConcurrentLinearAllocator.h b/Engine/Source/Runtime/Core/Public/Experimental/ConcurrentLinearAllocator.h
--- a/Engine/Source/Runtime/Core/Public/Experimental/ConcurrentLinearAllocator.h
+++ b/Engine/Source/Runtime/Core/Public/Experimental/ConcurrentLinearAllocator.h
@@ -3,49 +3,48 @@
 #pragma once
 #include <atomic>
 #if UE_ENABLE_INCLUDE_ORDER_DEPRECATED_IN_5_2
-#include "CoreMinimal.h"
+#	include "CoreMinimal.h"
 #endif
-#include "HAL/MallocAnsi.h"
 #include "HAL/UnrealMemory.h"
 #include "HAL/LowLevelMemTracker.h"
-#include "ProfilingDebugging/MemoryTrace.h"
 #include "Templates/UniquePtr.h"
 #include "Templates/UnrealTypeTraits.h"
 #include "Containers/LockFreeFixedSizeAllocator.h"
 #include "Misc/MemStack.h"
 
 #ifdef USE_MALLOC_BINNED3
-#if USE_MALLOC_BINNED3
-#define SUPPORTS_VERY_LARGE_ALIGNMENTS 0
+#	if USE_MALLOC_BINNED3
+#		define SUPPORTS_VERY_LARGE_ALIGNMENTS 0
+#	else
+#		define SUPPORTS_VERY_LARGE_ALIGNMENTS 1
+#	endif
 #else
-#define SUPPORTS_VERY_LARGE_ALIGNMENTS 1
-#endif
-#else
-#define SUPPORTS_VERY_LARGE_ALIGNMENTS 1
+#	define SUPPORTS_VERY_LARGE_ALIGNMENTS 1
 #endif
 
 #if PLATFORM_HAS_ASAN_INCLUDE
-#include <sanitizer/asan_interface.h>
-#if defined(__SANITIZE_ADDRESS__)
-#define IS_ASAN_ENABLED 1
-#elif defined(__has_feature)
-#if __has_feature(address_sanitizer)
-#define IS_ASAN_ENABLED 1
-#else
-#define IS_ASAN_ENABLED 0
-#endif
+#	include <sanitizer/asan_interface.h>
+#	if defined(__SANITIZE_ADDRESS__)
+#		define IS_ASAN_ENABLED 1
+#	elif defined(__has_feature)
+#		if __has_feature(address_sanitizer)
+#			define IS_ASAN_ENABLED 1
+#		else
+#			define IS_ASAN_ENABLED 0
+#		endif
+#	else
+#		define IS_ASAN_ENABLED 0
+#	endif
 #else
-#define IS_ASAN_ENABLED 0
-#endif
-#else
-#define ASAN_POISON_MEMORY_REGION(addr, size) ((void)(addr), (void)(size))
-#define ASAN_UNPOISON_MEMORY_REGION(addr, size) ((void)(addr), (void)(size))
-#define IS_ASAN_ENABLED 0
+#	define ASAN_POISON_MEMORY_REGION(addr, size) ((void)(addr), (void)(size))
+#	define ASAN_UNPOISON_MEMORY_REGION(addr, size) ((void)(addr), (void)(size))
+#	define IS_ASAN_ENABLED 0
 #endif
 
 struct FOsAllocator
 {
 	static constexpr bool SupportsAlignment = false;
+	static constexpr uint32 MaxAlignment	= 16;
 
 	// Indicates whether memory is acquired from the same source as used by FMalloc i.e. FMemory/FPlatformMemory::BaseAllocator().
 	// It's needed so that ConcurrentLinearAllocator could correctly account for this memory in the memory stats and tracking systems.
@@ -58,34 +57,37 @@ struct FOsAllocator
 	CORE_API static void Free(void* Pointer, SIZE_T Size);
 };
 
-#if PLATFORM_USES_ANSI_MALLOC
-struct FAlignedAllocatorInternal
+struct FAlignedAllocator
 {
 	static constexpr bool SupportsAlignment = true;
 	static constexpr bool UsesFMalloc       = true;
+	static constexpr uint32 MaxAlignment	= 128;
 
 	FORCEINLINE static void* Malloc(SIZE_T Size, uint32 Alignment)
 	{
-		void* Pointer = AnsiMalloc(Size, Alignment);
-#if DO_CHECK
-		check(IsAligned(Pointer, Alignment));
-#endif
-		MemoryTrace_Alloc(uint64(Pointer), Size, Alignment, EMemoryTraceRootHeap::SystemMemory);
-		return Pointer;
+		Alignment = FMath::Max(Alignment, (uint32)sizeof(void*));
+		Size += Alignment;
+		uint8* Result = nullptr;
+		uint8* Ptr = (uint8*)FMemory::Malloc(Size);
+		if (Ptr)
+		{
+			Result = Align(Ptr + sizeof(void*), Alignment);
+			*((void**)(Result - sizeof(void*))) = Ptr;
+		}
+
+		return Result;
 	}
 
 	FORCEINLINE static void Free(void* Pointer, SIZE_T Size)
 	{
-		MemoryTrace_Free(uint64(Pointer), EMemoryTraceRootHeap::SystemMemory);
-		return AnsiFree(Pointer);
+		if (Pointer)
+		{
+			void* OriginalPtr = *((void**)((uint8*)Pointer - sizeof(void*)));
+			FMemory::Free(OriginalPtr);
+		}
 	}
 };
 
-using FAlignedAllocator = FAlignedAllocatorInternal;
-#else // #if PLATFORM_USES_ANSI_MALLOC
-using FAlignedAllocator = FOsAllocator;
-#endif // #else // #if PLATFORM_USES_ANSI_MALLOC
-
 template<uint32 BlockSize, typename Allocator = FOsAllocator>
 class TBlockAllocationCache
 {
@@ -113,6 +115,7 @@ class TBlockAllocationCache
 public:
 	static constexpr bool SupportsAlignment = Allocator::SupportsAlignment;
 	static constexpr bool UsesFMalloc       = Allocator::UsesFMalloc;
+	static constexpr uint32 MaxAlignment	= Allocator::MaxAlignment;
 
 	FORCEINLINE static void* Malloc(SIZE_T Size, uint32 Alignment)
 	{
@@ -148,6 +151,7 @@ public:
 	static_assert(BlockSize == FPageAllocator::PageSize, "Only 64k pages are supported with this cache.");
 	static constexpr bool SupportsAlignment = Allocator::SupportsAlignment;
 	static constexpr bool UsesFMalloc       = Allocator::UsesFMalloc;
+	static constexpr uint32 MaxAlignment	= Allocator::MaxAlignment;
 
 	FORCEINLINE static void* Malloc(SIZE_T Size, uint32 Alignment)
 	{
@@ -534,6 +538,7 @@ public:
 	FORCEINLINE_DEBUGGABLE static void* operator new(size_t Size)
 	{
 		static_assert(TIsDerivedFrom<ObjectType, TConcurrentLinearObject<ObjectType, BlockAllocationTag>>::Value, "TConcurrentLinearObject must be base of it's ObjectType (see CRTP)");
+		static_assert(alignof(ObjectType) <= BlockAllocationTag::Allocator::MaxAlignment);
 		return TConcurrentLinearAllocator<BlockAllocationTag>::template Malloc<alignof(ObjectType)>(Size);
 	}
 
@@ -545,17 +550,20 @@ public:
 
 	FORCEINLINE_DEBUGGABLE static void* operator new[](size_t Size)
 	{
+		static_assert(alignof(ObjectType) <= BlockAllocationTag::Allocator::MaxAlignment);
 		return TConcurrentLinearAllocator<BlockAllocationTag>::template Malloc<alignof(ObjectType)>(Size);
 	}
 
 	FORCEINLINE_DEBUGGABLE static void* operator new(size_t Size, std::align_val_t Align)
 	{
+		static_assert(alignof(ObjectType) <= BlockAllocationTag::Allocator::MaxAlignment);
 		checkSlow(size_t(Align) == alignof(ObjectType));
 		return TConcurrentLinearAllocator<BlockAllocationTag>::template Malloc<alignof(ObjectType)>(Size);
 	}
 
 	FORCEINLINE_DEBUGGABLE static void* operator new[](size_t Size, std::align_val_t Align)
 	{
+		static_assert(alignof(ObjectType) <= BlockAllocationTag::Allocator::MaxAlignment);
 		checkSlow(size_t(Align)  == alignof(ObjectType));
 		return TConcurrentLinearAllocator<BlockAllocationTag>::template Malloc<alignof(ObjectType)>(Size);
 	}
@@ -628,6 +636,7 @@ public:
 				UE::Core::Private::OnInvalidConcurrentLinearArrayAllocatorNum(NewMax, NumBytesPerElement);
 			}
 
+			static_assert(alignof(ElementType) <= BlockAllocationTag::Allocator::MaxAlignment);
 			Data = (ElementType*)TConcurrentLinearAllocator<BlockAllocationTag>::Realloc(Data, NewMax * NumBytesPerElement, alignof(ElementType));
 		}
 		SizeType CalculateSlackReserve(SizeType NewMax, SIZE_T NumBytesPerElement) const
diff --git a/Engine/Source/Runtime/Core/Public/HAL/Platform.h b/Engine/Source/Runtime/Core/Public/HAL/Platform.h
--- a/Engine/Source/Runtime/Core/Public/HAL/Platform.h
+++ b/Engine/Source/Runtime/Core/Public/HAL/Platform.h
@@ -498,10 +498,6 @@
 	#define PLATFORM_USE_FULL_TASK_GRAPH						1
 #endif
 
-#ifndef PLATFORM_USES_ANSI_MALLOC
-	#define PLATFORM_USES_ANSI_MALLOC							1
-#endif
-
 #ifndef PLATFORM_USE_ANSI_POSIX_MALLOC
 	#define PLATFORM_USE_ANSI_POSIX_MALLOC						0
 #endif
