First Person: Implemented support for custom material outputs in MaterialAttributes. This  extends support for the handful of existing supported outputs to now actually compile them without needing to get the custom output attribute at the end of the graph and then pipe it into the corresponding custom output node. This was achieved by extending the FMaterialAttributeDefinitionMap class to store additional meta data about custom outputs and by introducing the concept of attribute bitmasks to it. The list of supported custom outputs can be extended in the future as most custom outputs don't require complex logic in the node.
In addition, the FirstPerson output node (FirstPersonInterpolationAlpha to be precise) is now also supported by MaterialAttributes, enabling workflows where the FirstPerson lerp factor can be set from material functions.

diff --git a/Engine/Source/Runtime/Core/Private/UObject/FortniteReleaseBranchCustomObjectVersion.cpp b/Engine/Source/Runtime/Core/Private/UObject/FortniteReleaseBranchCustomObjectVersion.cpp
--- a/Engine/Source/Runtime/Core/Private/UObject/FortniteReleaseBranchCustomObjectVersion.cpp
+++ b/Engine/Source/Runtime/Core/Private/UObject/FortniteReleaseBranchCustomObjectVersion.cpp
@@ -8,7 +8,7 @@ TMap<FGuid, FGuid> FFortniteReleaseBranchCustomObjectVersion::GetSystemGuids()
 
 	SystemGuids.Add(DevGuids.Niagara_LatestScriptCompileVersion, FGuid("D35A655CD3FDA4408724D77E1316E35B"));
 	SystemGuids.Add(DevGuids.SkeletalMeshDerivedDataVersion, FGuid("5705956EC7134274A5A1BD3FC74DB3A9"));
-	SystemGuids.Add(DevGuids.MaterialTranslationDDCVersion, FGuid("C64A6A89E16442C888B638EB264579DB"));
+	SystemGuids.Add(DevGuids.MaterialTranslationDDCVersion, FGuid("888A502A66F6401580B27451F211C46B"));
 	SystemGuids.Add(DevGuids.GLOBALSHADERMAP_DERIVEDDATA_VER, FGuid("0A0A1BF5D2A94105BB395EF0DA7CAEB9"));
 	SystemGuids.Add(DevGuids.MATERIALSHADERMAP_DERIVEDDATA_VER, FGuid("C0E1F59386894C9D855B85178B2553CB"));
 	return SystemGuids;
