PR #12458: Fix UnitConversion.cpp suspicious concatenation

#rnx

diff --git a/Engine/Source/Runtime/Core/Private/Math/UnitConversion.cpp b/Engine/Source/Runtime/Core/Private/Math/UnitConversion.cpp
--- a/Engine/Source/Runtime/Core/Private/Math/UnitConversion.cpp
+++ b/Engine/Source/Runtime/Core/Private/Math/UnitConversion.cpp
@@ -228,7 +228,7 @@ constexpr const TCHAR* const SupportedUnitsStrings[] = {
 	TEXT("Kilohertz"),			
 	TEXT("Megahertz"),			
 	TEXT("Gigahertz"),			
-	TEXT("RevolutionsPerMinute")
+	TEXT("RevolutionsPerMinute"),
 
 	TEXT("Bytes"),				
 	TEXT("Kilobytes"),		
@@ -257,7 +257,6 @@ constexpr const TCHAR* const SupportedUnitsStrings[] = {
 	TEXT("Percent"),
 
 	TEXT("Times"),
-	TEXT("Multiplier"),	
 
 	TEXT("Pascals"),		
 	TEXT("KiloPascals"),	
