Enable third-party profiler events on Windows even if FRAMEPRO_ENABLED is true.

diff --git a/Engine/Source/Runtime/Core/Private/Windows/WindowsPlatformMisc.cpp b/Engine/Source/Runtime/Core/Private/Windows/WindowsPlatformMisc.cpp
--- a/Engine/Source/Runtime/Core/Private/Windows/WindowsPlatformMisc.cpp
+++ b/Engine/Source/Runtime/Core/Private/Windows/WindowsPlatformMisc.cpp
@@ -1289,7 +1289,8 @@ void FWindowsPlatformMisc::BeginNamedEvent(const struct FColor& Color, const TCH
 
 #if FRAMEPRO_ENABLED
 	FFrameProProfiler::PushEvent(Text);
-#elif UE_EXTERNAL_PROFILING_ENABLED
+#endif
+#if UE_EXTERNAL_PROFILING_ENABLED
 	FExternalProfilerTrace::StartScopedEvent(Color, Text);
 #endif
 }
@@ -1300,7 +1301,8 @@ void FWindowsPlatformMisc::BeginNamedEvent(const struct FColor& Color, const ANS
 
 #if FRAMEPRO_ENABLED
 	FFrameProProfiler::PushEvent(Text);
-#elif UE_EXTERNAL_PROFILING_ENABLED
+#endif
+#if UE_EXTERNAL_PROFILING_ENABLED
 	FExternalProfilerTrace::StartScopedEvent(Color, Text);
 #endif
 }
@@ -1309,7 +1311,8 @@ void FWindowsPlatformMisc::EndNamedEvent()
 {
 #if FRAMEPRO_ENABLED
 	FFrameProProfiler::PopEvent();
-#elif UE_EXTERNAL_PROFILING_ENABLED
+#endif
+#if UE_EXTERNAL_PROFILING_ENABLED
 	FExternalProfilerTrace::EndScopedEvent();
 #endif
 }
