Remove erroneous return statement

diff --git a/Engine/Source/Runtime/Core/Private/ProfilingDebugging/StringTrace.cpp b/Engine/Source/Runtime/Core/Private/ProfilingDebugging/StringTrace.cpp
--- a/Engine/Source/Runtime/Core/Private/ProfilingDebugging/StringTrace.cpp
+++ b/Engine/Source/Runtime/Core/Private/ProfilingDebugging/StringTrace.cpp
@@ -42,7 +42,7 @@ static struct FStaticStringTracer
 		TracedWideStrings.FindOrAdd(Id, &bAlreadyInSet);
 		if (!bAlreadyInSet)
 		{
-			return UE_TRACE_LOG_DEFINITION(Strings, StaticString, Id, true)
+			UE_TRACE_LOG_DEFINITION(Strings, StaticString, Id, true)
 				 << StaticString.DisplayWide(InString, FCString::Strlen(InString));
 		}
 		return UE::Trace::FEventRef64(Id, UE_TRACE_GET_DEFINITION_TYPE_ID(Strings, StaticString));
