Improved move semantics for creating instances of TAttribute with delegates.

diff --git a/Engine/Source/Runtime/Core/Public/Misc/Attribute.h b/Engine/Source/Runtime/Core/Public/Misc/Attribute.h
--- a/Engine/Source/Runtime/Core/Public/Misc/Attribute.h
+++ b/Engine/Source/Runtime/Core/Public/Misc/Attribute.h
@@ -103,7 +103,18 @@ public:
 		const bool bExplicitConstructor = true;
 		return TAttribute( InGetter, bExplicitConstructor );
 	}
-	
+
+	/**
+	 * Static: Creates an attribute that's pre-bound to the specified 'getter' delegate
+	 *
+	 * @param  InGetter		Delegate to bind
+	 */
+	[[nodiscard]] static TAttribute Create(FGetter&& InGetter)
+	{
+		const bool bExplicitConstructor = true;
+		return TAttribute(MoveTemp(InGetter), bExplicitConstructor);
+	}
+
 	/**
 	 * Creates an attribute by binding an arbitrary function that will be called to generate this attribute's value on demand. 
 	 * After binding, the attribute will no longer have a value that can be accessed directly, and instead the bound
@@ -488,6 +499,13 @@ private:
 		, Getter(InGetter)
 	{ }
 
+	/** Special explicit constructor for TAttribute::Create() */
+	TAttribute(FGetter&& InGetter, bool bExplicitConstructor)
+		: Value()
+		, bIsSet(true)
+		, Getter(MoveTemp(InGetter))
+	{ }
+
 	// We declare ourselves as a friend (templated using OtherType) so we can access members as needed
 	template< class OtherType > friend class TAttribute;
 
