[ManualResetEvent]
 - Fix benign race by making bWait atomic relaxed.
          T1 WaitOnAddress -> T1 Read bWait -> [No barrier] -> T2 Write bWait -> T2 Notify

#rnx

diff --git a/Engine/Source/Runtime/Core/Private/Microsoft/MicrosoftPlatformManualResetEvent.cpp b/Engine/Source/Runtime/Core/Private/Microsoft/MicrosoftPlatformManualResetEvent.cpp
--- a/Engine/Source/Runtime/Core/Private/Microsoft/MicrosoftPlatformManualResetEvent.cpp
+++ b/Engine/Source/Runtime/Core/Private/Microsoft/MicrosoftPlatformManualResetEvent.cpp
@@ -25,7 +25,7 @@ bool FMicrosoftPlatformManualResetEvent::WaitUntil(FMonotonicTimePoint WaitTime)
 
 		for (;;)
 		{
-			if (WaitOnAddress(&bWait, &bLocalWait, sizeof(bool), INFINITE) && !bWait)
+			if (WaitOnAddress(&bWait, &bLocalWait, sizeof(bool), INFINITE) && !bWait.load(std::memory_order_relaxed))
 			{
 				return true;
 			}
@@ -44,7 +44,7 @@ bool FMicrosoftPlatformManualResetEvent::WaitUntil(FMonotonicTimePoint WaitTime)
 				return false;
 			}
 			const DWORD WaitMs = DWORD(FPlatformMath::CeilToInt64(WaitSpan.ToMilliseconds()));
-			if (WaitOnAddress(&bWait, &bLocalWait, sizeof(bool), WaitMs) && !bWait)
+			if (WaitOnAddress(&bWait, &bLocalWait, sizeof(bool), WaitMs) && !bWait.load(std::memory_order_relaxed))
 			{
 				return true;
 			}
@@ -54,7 +54,7 @@ bool FMicrosoftPlatformManualResetEvent::WaitUntil(FMonotonicTimePoint WaitTime)
 
 void FMicrosoftPlatformManualResetEvent::Notify()
 {
-	bWait = false;
+	bWait.store(false, std::memory_order_relaxed);
 	WakeByAddressSingle((void*)&bWait);
 }
 
diff --git a/Engine/Source/Runtime/Core/Private/Microsoft/MicrosoftPlatformManualResetEvent.h b/Engine/Source/Runtime/Core/Private/Microsoft/MicrosoftPlatformManualResetEvent.h
--- a/Engine/Source/Runtime/Core/Private/Microsoft/MicrosoftPlatformManualResetEvent.h
+++ b/Engine/Source/Runtime/Core/Private/Microsoft/MicrosoftPlatformManualResetEvent.h
@@ -4,6 +4,8 @@
 
 #include "CoreTypes.h"
 
+#include <atomic>
+
 #define UE_API CORE_API
 
 namespace UE { struct FMonotonicTimePoint; }
@@ -19,14 +21,14 @@ public:
 	FMicrosoftPlatformManualResetEvent(const FMicrosoftPlatformManualResetEvent&) = delete;
 	FMicrosoftPlatformManualResetEvent& operator=(const FMicrosoftPlatformManualResetEvent&) = delete;
 
-	void Reset() { bWait = true; }
+	void Reset() { bWait.store(true, std::memory_order_relaxed); }
 
 	UE_API void Wait();
 	UE_API bool WaitUntil(FMonotonicTimePoint WaitTime);
 	UE_API void Notify();
 
 private:
-	bool bWait = true;
+	std::atomic<bool> bWait = true;
 };
 
 } // UE::HAL::Private
