[AutoRTFM] Add UE_AUTORTFM_ALWAYS_OPEN_DISABLE_MEMORY_VALIDATION

Use this on the UObject constructor, as this is known to upset the validator.

Other:
* Rename the memory validation level enum names to be more accurate.
* Move autortfm_memory_validation_level to AutoRTFMConstants.h so the compiler has visibilty.
* Include AutoRTFMConstants.h from AutoRTFM.h
* Use CHECK-LABEL in more places in the lit tests. This improves check grouping.
* Support #including of standard library includes for windows lit tests (requires additional include directories)

diff --git a/Engine/Source/Programs/AutoRTFMEngineTests/Private/AutoRTFMEngineTests.cpp b/Engine/Source/Programs/AutoRTFMEngineTests/Private/AutoRTFMEngineTests.cpp
--- a/Engine/Source/Programs/AutoRTFMEngineTests/Private/AutoRTFMEngineTests.cpp
+++ b/Engine/Source/Programs/AutoRTFMEngineTests/Private/AutoRTFMEngineTests.cpp
@@ -233,7 +233,7 @@ static void PreInit()
 	FDelayedAutoRegisterHelper::RunAndClearDelayedAutoRegisterDelegates(EDelayedRegisterRunPhase::StatSystemReady);
 
 	// Unexpected memory validation errors should be full assertions.
-	AutoRTFM::ForTheRuntime::SetMemoryValidationLevel(AutoRTFM::EMemoryValidationLevel::Enabled);
+	AutoRTFM::ForTheRuntime::SetMemoryValidationLevel(AutoRTFM::EMemoryValidationLevel::Error);
 
 	// Ensure GCachedScalabilityCVars.bInitialized is set (required by some mesh components)
 	FTaskTagScope::SwapTag(ETaskTag::EGameThread);
diff --git a/Engine/Source/Programs/AutoRTFMTests/Private/AutoRTFMTests.cpp b/Engine/Source/Programs/AutoRTFMTests/Private/AutoRTFMTests.cpp
--- a/Engine/Source/Programs/AutoRTFMTests/Private/AutoRTFMTests.cpp
+++ b/Engine/Source/Programs/AutoRTFMTests/Private/AutoRTFMTests.cpp
@@ -100,7 +100,7 @@ int RunTests(int ArgC, const char* ArgV[])
 	AutoRTFM::ForTheRuntime::SetEnsureOnAbortByLanguage(false);
 
 	// Unexpected memory validation errors should be full assertions.
-	AutoRTFM::ForTheRuntime::SetMemoryValidationLevel(AutoRTFM::EMemoryValidationLevel::Enabled);
+	AutoRTFM::ForTheRuntime::SetMemoryValidationLevel(AutoRTFM::EMemoryValidationLevel::Error);
 
 	int Result = Session.run();
 
diff --git a/Engine/Source/Programs/AutoRTFMTests/Private/Utils.h b/Engine/Source/Programs/AutoRTFMTests/Private/Utils.h
--- a/Engine/Source/Programs/AutoRTFMTests/Private/Utils.h
+++ b/Engine/Source/Programs/AutoRTFMTests/Private/Utils.h
@@ -82,7 +82,7 @@ private:
 
 #define AUTORTFM_SCOPED_ENABLE_MEMORY_VALIDATION_AS_WARNING() \
 	AutoRTFMTestUtils::ScopedMemoryValidationLevel EnableMemoryValidationAsWarning \
-		{AutoRTFM::EMemoryValidationLevel::EnabledAsWarning}; \
+		{AutoRTFM::EMemoryValidationLevel::Warn}; \
 	static_assert(true) /* require semicolon */
 
 static constexpr const TCHAR* kMemoryModifiedWarning =
diff --git a/Engine/Source/Runtime/Core/Private/AutoRTFM/API.cpp b/Engine/Source/Runtime/Core/Private/AutoRTFM/API.cpp
--- a/Engine/Source/Runtime/Core/Private/AutoRTFM/API.cpp
+++ b/Engine/Source/Runtime/Core/Private/AutoRTFM/API.cpp
@@ -61,11 +61,11 @@ namespace
 		case AutoRTFM::EMemoryValidationLevel::Disabled:
 			UE_LOG(LogCore, Log, TEXT("AutoRTFM Memory Validation: disabled"));
 			return FGenericCrashContext::SetGameData(TEXT("AutoRTFMMemoryValidationLevel"), TEXT("Disabled"));
-		case AutoRTFM::EMemoryValidationLevel::EnabledAsWarning:
-			UE_LOG(LogCore, Log, TEXT("AutoRTFM Memory Validation: enabled-as-warning"));
-			return FGenericCrashContext::SetGameData(TEXT("AutoRTFMMemoryValidationLevel"), TEXT("EnabledAsWarning"));
-		case AutoRTFM::EMemoryValidationLevel::Enabled:
-			UE_LOG(LogCore, Log, TEXT("AutoRTFM Memory Validation: enabled"));
+		case AutoRTFM::EMemoryValidationLevel::Warn:
+			UE_LOG(LogCore, Log, TEXT("AutoRTFM Memory Validation: enabled as warning"));
+			return FGenericCrashContext::SetGameData(TEXT("AutoRTFMMemoryValidationLevel"), TEXT("Warn"));
+		case AutoRTFM::EMemoryValidationLevel::Error:
+			UE_LOG(LogCore, Log, TEXT("AutoRTFM Memory Validation: enabled as error"));
 			return FGenericCrashContext::SetGameData(TEXT("AutoRTFMMemoryValidationLevel"), TEXT("Enabled"));
 		}
 	}
@@ -262,8 +262,8 @@ namespace AutoRTFM
 #if UE_AUTORTFM
 			// Check memory validation is enabled, otherwise we won't have a
 			// hash to compare against.
-			ASSERT(GetMemoryValidationLevel() == EMemoryValidationLevel::Enabled ||
-				GetMemoryValidationLevel() == EMemoryValidationLevel::EnabledAsWarning);
+			ASSERT(GetMemoryValidationLevel() == EMemoryValidationLevel::Error ||
+				GetMemoryValidationLevel() == EMemoryValidationLevel::Warn);
 			if (FTransaction* Transaction = AutoRTFM::FContext::Get()->GetMaterializedTransaction()) {
 				Transaction->DebugBreakIfMemoryValidationFails();
 			}
diff --git a/Engine/Source/Runtime/Core/Private/AutoRTFM/Transaction.cpp b/Engine/Source/Runtime/Core/Private/AutoRTFM/Transaction.cpp
--- a/Engine/Source/Runtime/Core/Private/AutoRTFM/Transaction.cpp
+++ b/Engine/Source/Runtime/Core/Private/AutoRTFM/Transaction.cpp
@@ -407,7 +407,7 @@ void FTransaction::ValidateWriteHash() const
 	static constexpr TCHAR Message[] =
 		TEXT("Memory modified in a transaction was also modified in an call to AutoRTFM::Open(). ")
 		TEXT("This may lead to memory corruption if the transaction is aborted.");
-	if (CurrentMemoryValidationLevel == EMemoryValidationLevel::EnabledAsWarning)
+	if (CurrentMemoryValidationLevel == EMemoryValidationLevel::Warn)
 	{
 		if (UNLIKELY(OldHash != NewHash))
 		{
diff --git a/Engine/Source/Runtime/Core/Public/AutoRTFM/AutoRTFM.h b/Engine/Source/Runtime/Core/Public/AutoRTFM/AutoRTFM.h
--- a/Engine/Source/Runtime/Core/Public/AutoRTFM/AutoRTFM.h
+++ b/Engine/Source/Runtime/Core/Public/AutoRTFM/AutoRTFM.h
@@ -37,10 +37,12 @@
 #if UE_AUTORTFM
 #define UE_AUTORTFM_NOAUTORTFM [[clang::noautortfm]]
 #define UE_AUTORTFM_ALWAYS_OPEN [[clang::autortfm_always_open]]
+#define UE_AUTORTFM_ALWAYS_OPEN_NO_MEMORY_VALIDATION [[clang::autortfm_always_open_disable_memory_validation]]
 #define UE_AUTORTFM_CALLSITE_FORCEINLINE [[clang::always_inline]]
 #else
 #define UE_AUTORTFM_NOAUTORTFM
 #define UE_AUTORTFM_ALWAYS_OPEN
+#define UE_AUTORTFM_ALWAYS_OPEN_NO_MEMORY_VALIDATION
 #define UE_AUTORTFM_CALLSITE_FORCEINLINE
 #endif
 
@@ -52,6 +54,8 @@
 #define UE_AUTORTFM_ASSUME_SAFE
 #endif
 
+#include "AutoRTFMConstants.h"
+
 #if defined(UE_AUTORTFM_STANDALONE)
 #include <stdlib.h>
 #include <memory.h>
@@ -130,25 +134,6 @@ typedef enum
 	autortfm_status_in_static_local_initializer
 } autortfm_status;
 
-// An enumerator of transactional memory validation levels.
-// Memory validation is used to detect modification by open-code to memory that was written by a
-// transaction. In this situation, aborting the transaction can corrupt memory as the undo will
-// overwrite the writes made in the open-code.
-typedef enum
-{
-	// Use the default memory validation level.
-	autortfm_memory_validation_level_default,
-
-	// Disable memory validation.
-	autortfm_memory_validation_level_disabled,
-
-	// Enable memory validation. Memory validation failures are treated as warnings.
-	autortfm_memory_validation_level_warning,
-
-	// Enable memory validation. Memory validation are treated as fatal.
-	autortfm_memory_validation_level_enabled,
-} autortfm_memory_validation_level;
-
 #if UE_AUTORTFM_ENABLED
 // Note: There is no implementation of this function.
 // The AutoRTFM compiler will replace all calls to this function with a constant boolean value.
@@ -401,7 +386,6 @@ static UE_AUTORTFM_FORCEINLINE void autortfm_check_abi(void* ptr, size_t size)
 }
 
 #include <algorithm>
-#include <initializer_list>
 #include <tuple>
 
 namespace AutoRTFM
@@ -475,11 +459,11 @@ enum class EMemoryValidationLevel : uint8_t
 	// Disable memory validation.
 	Disabled = autortfm_memory_validation_level_disabled,
 
-	// Enable memory validation. Memory validation failures are treated as warnings.
-	EnabledAsWarning = autortfm_memory_validation_level_warning,
+	// Memory validation enabled as warnings.
+	Warn = autortfm_memory_validation_level_warn,
 
-	// Enable memory validation. Memory validation are treated as fatal.
-	Enabled = autortfm_memory_validation_level_enabled,
+	// Memory validation enabled as errors.
+	Error = autortfm_memory_validation_level_error,
 };
 
 #if UE_AUTORTFM
diff --git a/Engine/Source/Runtime/Core/Public/AutoRTFM/AutoRTFMConstants.h b/Engine/Source/Runtime/Core/Public/AutoRTFM/AutoRTFMConstants.h
--- a/Engine/Source/Runtime/Core/Public/AutoRTFM/AutoRTFMConstants.h
+++ b/Engine/Source/Runtime/Core/Public/AutoRTFM/AutoRTFMConstants.h
@@ -2,16 +2,41 @@
 
 #pragma once
 
-#include <cstddef>
 #include <cstdint>
 
 // WARNING: Any change in these constants will require a re-patch and re-build of LLVM!
 
-namespace AutoRTFM
+#ifdef __cplusplus
+extern "C"
 {
-namespace Constants
+#endif
+
+// An enumerator of transactional memory validation levels.
+// Memory validation is used to detect modification by open-code to memory that was written by a
+// transaction. In this situation, aborting the transaction can corrupt memory as the undo will
+// overwrite the writes made in the open-code.
+typedef enum
 {
+	// Use the default memory validation level.
+	autortfm_memory_validation_level_default,
+
+	// Disable memory validation.
+	autortfm_memory_validation_level_disabled,
+
+	// Enable memory validation. Memory validation failures are treated as warnings.
+	autortfm_memory_validation_level_warn,
 
+	// Enable memory validation. Memory validation are treated as fatal.
+	autortfm_memory_validation_level_error,
+} autortfm_memory_validation_level;
+
+#ifdef __cplusplus
+}
+#endif
+
+#ifdef __cplusplus
+namespace AutoRTFM::Constants
+{
 	inline constexpr uint32_t Major = 0;
 	inline constexpr uint32_t Minor = 2;
 	inline constexpr uint32_t Patch = 0;
@@ -21,5 +46,5 @@ namespace Constants
 	// and were definitely injected by us.
 	inline constexpr uint64_t MagicMike = 0xa273000000000000;
 
-} // namespace Constants
-} // namespace AutoRTFM
+} // namespace AutoRTFM::Constants
+#endif
diff --git a/Engine/Source/Runtime/CoreUObject/Private/UObject/UObjectGlobals.cpp b/Engine/Source/Runtime/CoreUObject/Private/UObject/UObjectGlobals.cpp
--- a/Engine/Source/Runtime/CoreUObject/Private/UObject/UObjectGlobals.cpp
+++ b/Engine/Source/Runtime/CoreUObject/Private/UObject/UObjectGlobals.cpp
@@ -3771,7 +3771,8 @@ void UObject::PostInitProperties()
 	FOverridableManager::Get().ClearOverrides(*this);
 }
 
-UE_AUTORTFM_ALWAYS_OPEN
+// UObjects can be reconstructed after use, which can upset the validator.
+UE_AUTORTFM_ALWAYS_OPEN_NO_MEMORY_VALIDATION
 UObject::UObject()
 {
 #if UE_WITH_REMOTE_OBJECT_HANDLE
