[TaskGraph]
 - Fix possible race on platforms with weak memory ordering when using the result of an AddSubsequent call to decide if it is safe to start using the memory modified by the task

diff --git a/Engine/Source/Runtime/Core/Public/Tasks/TaskPrivate.h b/Engine/Source/Runtime/Core/Public/Tasks/TaskPrivate.h
--- a/Engine/Source/Runtime/Core/Public/Tasks/TaskPrivate.h
+++ b/Engine/Source/Runtime/Core/Public/Tasks/TaskPrivate.h
@@ -761,7 +761,10 @@ public:
 				bool PushIfNotClosed(FTaskBase* NewItem)
 				{
 					TASKGRAPH_VERBOSE_EVENT_SCOPE(FSubsequents::PushIfNotClosed);
-					if (bIsClosed.load(std::memory_order_relaxed))
+					// AddSubsequent expects acquire as some code can use the result here
+					// to decide if a task is finished and if the memory it produced/modified
+					// is safe to access.
+					if (bIsClosed.load(std::memory_order_acquire))
 					{
 						return false;
 					}
