Disable SCNetworkReachability warnings for VisionOS 2 - and for iOS18.1 until reimplementation is ready.  Also, disable SKStoreReviewController deprecation warning.

diff --git a/Engine/Source/Runtime/Core/Private/IOS/IOSPlatformMisc.cpp b/Engine/Source/Runtime/Core/Private/IOS/IOSPlatformMisc.cpp
--- a/Engine/Source/Runtime/Core/Private/IOS/IOSPlatformMisc.cpp
+++ b/Engine/Source/Runtime/Core/Private/IOS/IOSPlatformMisc.cpp
@@ -931,7 +931,10 @@ bool FIOSPlatformMisc::GetDiskTotalAndFreeSpace(const FString& InPath, uint64& T
 void FIOSPlatformMisc::RequestStoreReview()
 {
 #if UE_WITH_STORE_KIT && !PLATFORM_TVOS
+PRAGMA_DISABLE_DEPRECATION_WARNINGS
+	// Deprecated for iOS18.0 & visionOS 2.0.  To be replaced with Swift only StoreKit::RequestReviewAction call (UE-228925)
     [SKStoreReviewController requestReviewInScene:[[[[UIApplication sharedApplication] delegate] window] windowScene]];
+PRAGMA_ENABLE_DEPRECATION_WARNINGS
 #endif
 }
 
@@ -1563,6 +1566,7 @@ static inline ENetworkConnectionType CheckNetworkConnectionType()
 	ZeroAddress.sin_len = sizeof(ZeroAddress);
 	ZeroAddress.sin_family = AF_INET;
 	
+PRAGMA_DISABLE_DEPRECATION_WARNINGS
 	SCNetworkReachabilityRef ReachabilityRef = SCNetworkReachabilityCreateWithAddress(kCFAllocatorDefault, (const struct sockaddr*)&ZeroAddress);
 	SCNetworkReachabilityFlags ReachabilityFlags;
 	bool bFlagsAvailable = SCNetworkReachabilityGetFlags(ReachabilityRef, &ReachabilityFlags);
@@ -1582,6 +1586,7 @@ static inline ENetworkConnectionType CheckNetworkConnectionType()
         bHasActiveCellConnection = bReachable && (ReachabilityFlags & kSCNetworkReachabilityFlagsIsWWAN) != 0;
         bInAirplaneMode = ReachabilityFlags == 0;
 	}
+PRAGMA_ENABLE_DEPRECATION_WARNINGS
 	
     if (bHasActiveWiFiConnection)
     {
