Always assert on an AutoRTFM failure by default

diff --git a/Engine/Source/Runtime/Core/Private/AutoRTFM/API.cpp b/Engine/Source/Runtime/Core/Private/AutoRTFM/API.cpp
--- a/Engine/Source/Runtime/Core/Private/AutoRTFM/API.cpp
+++ b/Engine/Source/Runtime/Core/Private/AutoRTFM/API.cpp
@@ -29,7 +29,7 @@ namespace
 
 	int GMemoryValidationLevel = static_cast<int>(GDefaultMemoryValidationLevel);
 
-	bool GAutoRTFMAssertOnError = false;
+	bool GAutoRTFMAssertOnError = true;
 
 	void OnIsAutoRTFMEnabledChanged()
 	{
