[Backout] - CL37545653
Original CL Desc
-----------------------------------------------------------------
Always assert on an AutoRTFM failure by default

diff --git a/Engine/Source/Runtime/Core/Private/AutoRTFM/API.cpp b/Engine/Source/Runtime/Core/Private/AutoRTFM/API.cpp
--- a/Engine/Source/Runtime/Core/Private/AutoRTFM/API.cpp
+++ b/Engine/Source/Runtime/Core/Private/AutoRTFM/API.cpp
@@ -29,7 +29,7 @@ namespace
 
 	int GMemoryValidationLevel = static_cast<int>(GDefaultMemoryValidationLevel);
 
-	bool GAutoRTFMAssertOnError = true;
+	bool GAutoRTFMAssertOnError = false;
 
 	void OnIsAutoRTFMEnabledChanged()
 	{
