Reduced binary size and debug info by refactoring TArray to limit template instantiations:

- TArray's default allocator is now based on ForAnyElementType (type-erased) instead of ForElementType.
- Functions which reallocate or resize are now based on collapsed single functions which are templated on allocator instance rather than allocator, and allocator flags are passed as integers rather than the instance type.

diff --git a/Engine/Source/Runtime/Core/Public/Containers/Array.h b/Engine/Source/Runtime/Core/Public/Containers/Array.h
--- a/Engine/Source/Runtime/Core/Public/Containers/Array.h
+++ b/Engine/Source/Runtime/Core/Public/Containers/Array.h
@@ -350,6 +350,211 @@ namespace UE::Core::Private
 	constexpr bool TIsTArrayOrDerivedFromTArray_V = sizeof(ResolveIsTArrayPtr((T*)nullptr)) == 2;
 
 	[[noreturn]] CORE_API void OnInvalidArrayNum(unsigned long long NewNum);
+
+	// A hacky way to get the SizeType, since it's defined in the
+	// (outer) allocator type, not the (inner) allocator instance type
+	template <typename AllocatorInstanceType>
+	using TAllocatorSizeType_T = decltype(std::declval<AllocatorInstanceType&>().GetInitialCapacity());
+
+	// Flags are passed as a uint32 to minimize PDB impact of these generated symbols.
+	//
+	// 1 == TAllocatorTraits<>::SupportsElementAlignment
+	// 2 == TAllocatorTraits<>::SupportsSlackTracking
+	//
+	// When C++20 is guaranteed, concept checks can be used instead.
+	template <typename AllocatorType>
+	constexpr uint32 GetAllocatorFlags()
+	{
+		uint32 Result = 0;
+		if constexpr (TAllocatorTraits<AllocatorType>::SupportsElementAlignment)
+		{
+			Result |= 1;
+		}
+		if constexpr (TAllocatorTraits<AllocatorType>::SupportsSlackTracking)
+		{
+			Result |= 2;
+		}
+		return Result;
+	}
+
+	// This should be used for repeated growing operations when reallocations are to be amortized over multiple inserts.
+	template <uint32 Flags, typename AllocatorInstanceType>
+	FORCENOINLINE TAllocatorSizeType_T<AllocatorInstanceType> ReallocGrow(
+		std::size_t                                  ElementSize,
+		std::size_t                                  ElementAlignment,
+		TAllocatorSizeType_T<AllocatorInstanceType>  Count,
+		AllocatorInstanceType&                       AllocatorInstance,
+		TAllocatorSizeType_T<AllocatorInstanceType>& ArrayNum,
+		TAllocatorSizeType_T<AllocatorInstanceType>& ArrayMax
+	)
+	{
+		using SizeType  = TAllocatorSizeType_T<AllocatorInstanceType>;
+		using USizeType = std::make_unsigned_t<SizeType>;
+
+		const USizeType UCount  = (USizeType)Count;
+		const USizeType UOldNum = (USizeType)ArrayNum;
+		const USizeType UOldMax = (USizeType)ArrayMax;
+		const USizeType UNewNum = UOldNum + UCount;
+		const SizeType  OldNum  = (SizeType)UOldNum;
+		const SizeType  OldMax  = (SizeType)UOldMax;
+		const SizeType  NewNum  = (SizeType)UNewNum;
+
+		checkSlow((OldNum >= 0) & (OldMax >= OldNum) & (Count >= 0)); // & for one branch
+
+		ArrayNum = NewNum;
+
+#if DO_GUARD_SLOW
+		if (UNewNum > UOldMax)
+#else
+		// SECURITY - This check will guard against negative counts too, in case the checkSlow above is compiled out.
+		// However, it results in slightly worse code generation.
+		if (UCount > UOldMax - UOldNum)
+#endif
+		{
+			// This should only happen when we've underflowed or overflowed SizeType
+			if (NewNum < OldNum)
+			{
+				OnInvalidArrayNum((unsigned long long)UNewNum);
+			}
+			SizeType NewMax;
+			if constexpr (!!(Flags & 1)) // TAllocatorTraits<AllocatorType>::SupportsElementAlignment
+			{
+				NewMax = AllocatorInstance.CalculateSlackGrow(NewNum, OldMax, ElementSize, ElementAlignment);
+				AllocatorInstance.ResizeAllocation(UOldNum, NewMax, ElementSize, ElementAlignment);
+			}
+			else
+			{
+				NewMax = AllocatorInstance.CalculateSlackGrow(NewNum, OldMax, ElementSize);
+				AllocatorInstance.ResizeAllocation(UOldNum, NewMax, ElementSize);
+			}
+			ArrayMax = NewMax;
+#if UE_ENABLE_ARRAY_SLACK_TRACKING
+			if constexpr (!!(Flags & 2)) // TAllocatorTraits<AllocatorType>::SupportsSlackTracking
+			{
+				AllocatorInstance.SlackTrackerLogNum(NewNum);
+			}
+#endif
+		}
+
+		return OldNum;
+	}
+
+	// This should be used for repeated shrinking operations when reallocations are to be amortized over multiple removals.
+	template <uint32 Flags, typename AllocatorInstanceType>
+	FORCENOINLINE void ReallocShrink(
+		size_t                                       ElementSize,
+		std::size_t                                  ElementAlignment,
+		AllocatorInstanceType&                       AllocatorInstance,
+		TAllocatorSizeType_T<AllocatorInstanceType>  ArrayNum,
+		TAllocatorSizeType_T<AllocatorInstanceType>& ArrayMax
+	)
+	{
+		using SizeType  = TAllocatorSizeType_T<AllocatorInstanceType>;
+
+		SizeType OldArrayMax = ArrayMax;
+
+		if constexpr (!!(Flags & 1)) // TAllocatorTraits<AllocatorType>::SupportsElementAlignment
+		{
+			SizeType NewArrayMax = AllocatorInstance.CalculateSlackShrink(ArrayNum, OldArrayMax, ElementSize, ElementAlignment);
+			if (NewArrayMax != OldArrayMax)
+			{
+				ArrayMax = NewArrayMax;
+				AllocatorInstance.ResizeAllocation(ArrayNum, NewArrayMax, ElementSize, ElementAlignment);
+			}
+		}
+		else
+		{
+			SizeType NewArrayMax = AllocatorInstance.CalculateSlackShrink(ArrayNum, OldArrayMax, ElementSize);
+			if (NewArrayMax != OldArrayMax)
+			{
+				ArrayMax = NewArrayMax;
+				AllocatorInstance.ResizeAllocation(ArrayNum, NewArrayMax, ElementSize);
+			}
+		}
+	}
+
+	// This should be used for setting an allocation to a specific size.
+	// Precondition: NewMax >= ArrayNum.
+	template <uint32 Flags, typename AllocatorInstanceType>
+	FORCENOINLINE void ReallocTo(
+		size_t                                       ElementSize,
+		std::size_t                                  ElementAlignment,
+		TAllocatorSizeType_T<AllocatorInstanceType>  NewMax,
+		AllocatorInstanceType&                       AllocatorInstance,
+		TAllocatorSizeType_T<AllocatorInstanceType>  ArrayNum,
+		TAllocatorSizeType_T<AllocatorInstanceType>& ArrayMax
+	)
+	{
+		if constexpr (!!(Flags & 1)) // TAllocatorTraits<AllocatorType>::SupportsElementAlignment
+		{
+			if (NewMax)
+			{
+				NewMax = AllocatorInstance.CalculateSlackReserve(NewMax, ElementSize, ElementAlignment);
+			}
+			if (NewMax != ArrayMax)
+			{
+				ArrayMax = NewMax;
+				AllocatorInstance.ResizeAllocation(ArrayNum, NewMax, ElementSize, ElementAlignment);
+			}
+		}
+		else
+		{
+			if (NewMax)
+			{
+				NewMax = AllocatorInstance.CalculateSlackReserve(NewMax, ElementSize);
+			}
+			if (NewMax != ArrayMax)
+			{
+				ArrayMax = NewMax;
+				AllocatorInstance.ResizeAllocation(ArrayNum, NewMax, ElementSize);
+			}
+		}
+	}
+
+	template <uint32 Flags, typename AllocatorInstanceType>
+	FORCENOINLINE void ReallocForCopy(
+		size_t                                       ElementSize,
+		std::size_t                                  ElementAlignment,
+		TAllocatorSizeType_T<AllocatorInstanceType>  NewMax,
+		TAllocatorSizeType_T<AllocatorInstanceType>  PrevMax,
+		AllocatorInstanceType&                       AllocatorInstance,
+		TAllocatorSizeType_T<AllocatorInstanceType>  ArrayNum,
+		TAllocatorSizeType_T<AllocatorInstanceType>& ArrayMax
+	)
+	{
+		if constexpr (!!(Flags & 1)) // TAllocatorTraits<AllocatorType>::SupportsElementAlignment
+		{
+			if (NewMax)
+			{
+				NewMax = AllocatorInstance.CalculateSlackReserve(NewMax, ElementSize, ElementAlignment);
+			}
+			if (NewMax > PrevMax)
+			{
+				AllocatorInstance.ResizeAllocation(0, NewMax, ElementSize, ElementAlignment);
+			}
+			else
+			{
+				NewMax = PrevMax;
+			}
+		}
+		else
+		{
+			if (NewMax)
+			{
+				NewMax = AllocatorInstance.CalculateSlackReserve(NewMax, ElementSize);
+			}
+			if (NewMax > PrevMax)
+			{
+				AllocatorInstance.ResizeAllocation(0, NewMax, ElementSize);
+			}
+			else
+			{
+				NewMax = PrevMax;
+			}
+		}
+		ArrayMax = NewMax;
+	}
+
 }
 
 
@@ -888,7 +1093,13 @@ public:
 		RemoveAtImpl(ArrayNum - 1);
 		if (AllowShrinking == EAllowShrinking::Yes)
 		{
-			ResizeShrink();
+			UE::Core::Private::ReallocShrink<UE::Core::Private::GetAllocatorFlags<AllocatorType>()>(
+				sizeof(ElementType),
+				alignof(ElementType),
+				AllocatorInstance,
+				ArrayNum,
+				ArrayMax
+			);
 		}
 		return Result;
 	}
@@ -981,7 +1192,14 @@ public:
 		CheckInvariants();
 		if (ArrayMax != ArrayNum)
 		{
-			ResizeTo(ArrayNum);
+			UE::Core::Private::ReallocTo<UE::Core::Private::GetAllocatorFlags<AllocatorType>()>(
+				sizeof(ElementType),
+				alignof(ElementType),
+				ArrayNum,
+				AllocatorInstance,
+				ArrayNum,
+				ArrayMax
+			);
 		}
 	}
 
@@ -1373,102 +1591,67 @@ public:
 	 * @param Count Number of elements to add.
 	 * @returns Number of elements in array before addition.
 	 */
-	FORCEINLINE SizeType AddUninitialized()
+	SizeType AddUninitialized()
 	{
-		CheckInvariants();
-
-		const USizeType OldNum = (USizeType)ArrayNum;
-		const USizeType NewNum = OldNum + (USizeType)1;
-		ArrayNum = (SizeType)NewNum;
-		if (NewNum > (USizeType)ArrayMax)
-		{
-			ResizeGrow((SizeType)OldNum);
-		}
-		else
-		{
-			SlackTrackerNumChanged();
-		}
-
-		return OldNum;
+		return UE::Core::Private::ReallocGrow<UE::Core::Private::GetAllocatorFlags<AllocatorType>()>(
+			sizeof(ElementType),
+			alignof(ElementType),
+			1,
+			AllocatorInstance,
+			ArrayNum,
+			ArrayMax
+		);
 	}
 	FORCEINLINE SizeType AddUninitialized(SizeType Count)
 	{
-		CheckInvariants();
-		checkSlow(Count >= 0);
-
-		const USizeType OldNum = (USizeType)ArrayNum;
-		const USizeType NewNum = OldNum + (USizeType)Count;
-		ArrayNum = (SizeType)NewNum;
-
-#if DO_GUARD_SLOW
-		if (NewNum > (USizeType)ArrayMax)
-#else
-		// SECURITY - This check will guard against negative counts too, in case the checkSlow(Count >= 0) above is compiled out.
-		// However, it results in slightly worse code generation.
-		if ((USizeType)Count > (USizeType)ArrayMax - OldNum)
-#endif
-		{
-			ResizeGrow((SizeType)OldNum);
-		}
-		else
-		{
-			SlackTrackerNumChanged();
-		}
-
-		return OldNum;
+		// Should be SetNumUninitialized?
+		return UE::Core::Private::ReallocGrow<UE::Core::Private::GetAllocatorFlags<AllocatorType>()>(
+			sizeof(ElementType),
+			alignof(ElementType),
+			Count,
+			AllocatorInstance,
+			ArrayNum,
+			ArrayMax
+		);
 	}
 
 private:
 	void InsertUninitializedImpl(SizeType Index)
 	{
-		CheckInvariants();
-		checkSlow((Index >= 0) & (Index <= ArrayNum));
-
-		const USizeType OldNum = (USizeType)ArrayNum;
-		const USizeType NewNum = OldNum + (USizeType)1;
-		ArrayNum = (SizeType)NewNum;
-
-		if (NewNum > (USizeType)ArrayMax)
-		{
-			ResizeGrow((SizeType)OldNum);
-		}
-		else
-		{
-			SlackTrackerNumChanged();
-		}
+		SizeType OldNum = UE::Core::Private::ReallocGrow<UE::Core::Private::GetAllocatorFlags<AllocatorType>()>(
+			sizeof(ElementType),
+			alignof(ElementType),
+			1,
+			AllocatorInstance,
+			ArrayNum,
+			ArrayMax
+		);
 		ElementType* Data = GetData() + Index;
 		RelocateConstructItems<ElementType>((void*)(Data + 1), Data, OldNum - Index);
 	}
-	template <typename OtherSizeType>
-	void InsertUninitializedImpl(SizeType Index, OtherSizeType Count)
+	void InsertUninitializedImpl(SizeType Index, SizeType Count)
 	{
-		CheckInvariants();
-		checkSlow((Count >= 0) & (Index >= 0) & (Index <= ArrayNum));
-
-		SizeType ConvertedCount = Count;
-		checkf((OtherSizeType)ConvertedCount == Count, TEXT("Invalid number of elements to add to this array type: %lld"), (long long)ConvertedCount);
-
-		const USizeType OldNum = (USizeType)ArrayNum;
-		const USizeType NewNum = OldNum + (USizeType)Count;
-		ArrayNum = (SizeType)NewNum;
-
-#if DO_GUARD_SLOW
-		if (NewNum > (USizeType)ArrayMax)
-#else
-		// SECURITY - This check will guard against negative counts too, in case the checkSlow(Count >= 0) above is compiled out.
-		// However, it results in slightly worse code generation.
-		if ((USizeType)Count > (USizeType)ArrayMax - OldNum)
-#endif
-		{
-			ResizeGrow((SizeType)OldNum);
-		}
-		else
-		{
-			SlackTrackerNumChanged();
-		}
+		// Should be SetNumUninitialized?
+		SizeType OldNum = UE::Core::Private::ReallocGrow<UE::Core::Private::GetAllocatorFlags<AllocatorType>()>(
+			sizeof(ElementType),
+			alignof(ElementType),
+			Count,
+			AllocatorInstance,
+			ArrayNum,
+			ArrayMax
+		);
 		ElementType* Data = GetData() + Index;
 		RelocateConstructItems<ElementType>((void*)(Data + Count), Data, OldNum - Index);
 	}
+	template <
+		typename OtherSizeType
+		UE_REQUIRES(sizeof(OtherSizeType) > sizeof(SizeType))
+	>
+	void InsertUninitializedImpl(SizeType Index, OtherSizeType Count)
+	{
+		checkf((OtherSizeType)(SizeType)Count == Count, TEXT("Invalid number of elements to add to this array type: %lld"), (long long)(SizeType)Count);
+		return InsertUninitializedImpl(Index, (SizeType)Count);
+	}
 
 public:
 	/**
@@ -1790,7 +1973,13 @@ public:
 		RemoveAtImpl(Index);
 		if (AllowShrinking == EAllowShrinking::Yes)
 		{
-			ResizeShrink();
+			UE::Core::Private::ReallocShrink<UE::Core::Private::GetAllocatorFlags<AllocatorType>()>(
+				sizeof(ElementType),
+				alignof(ElementType),
+				AllocatorInstance,
+				ArrayNum,
+				ArrayMax
+			);
 		}
 	}
 
@@ -1815,7 +2004,13 @@ public:
 			RemoveAtImpl(Index, (SizeType)Count);
 			if (AllowShrinking == EAllowShrinking::Yes)
 			{
-				ResizeShrink();
+				UE::Core::Private::ReallocShrink<UE::Core::Private::GetAllocatorFlags<AllocatorType>()>(
+					sizeof(ElementType),
+					alignof(ElementType),
+					AllocatorInstance,
+					ArrayNum,
+					ArrayMax
+				);
 			}
 		}
 	}
@@ -1883,7 +2078,13 @@ public:
 		RemoveAtSwapImpl(Index);
 		if (AllowShrinking == EAllowShrinking::Yes)
 		{
-			ResizeShrink();
+			UE::Core::Private::ReallocShrink<UE::Core::Private::GetAllocatorFlags<AllocatorType>()>(
+				sizeof(ElementType),
+				alignof(ElementType),
+				AllocatorInstance,
+				ArrayNum,
+				ArrayMax
+			);
 		}
 	}
 
@@ -1909,11 +2110,17 @@ public:
 		RangeCheck(Index, Count);
 		if (Count)
 		{
-		    RemoveAtSwapImpl(Index, Count);
-		    if (AllowShrinking == EAllowShrinking::Yes)
-		    {
-			    ResizeShrink();
-		    }
+			RemoveAtSwapImpl(Index, Count);
+			if (AllowShrinking == EAllowShrinking::Yes)
+			{
+				UE::Core::Private::ReallocShrink<UE::Core::Private::GetAllocatorFlags<AllocatorType>()>(
+					sizeof(ElementType),
+					alignof(ElementType),
+					AllocatorInstance,
+					ArrayNum,
+					ArrayMax
+				);
+			}
 		}
 	}
 	template <typename CountType>
@@ -1973,7 +2180,14 @@ public:
 
 		if (ArrayMax != Slack)
 		{
-			ResizeTo(Slack);
+			UE::Core::Private::ReallocTo<UE::Core::Private::GetAllocatorFlags<AllocatorType>()>(
+				sizeof(ElementType),
+				alignof(ElementType),
+				Slack,
+				AllocatorInstance,
+				ArrayNum,
+				ArrayMax
+			);
 		}
 	}
 
@@ -2236,8 +2450,6 @@ public:
 	template <typename... ArgsType>
 	FORCEINLINE SizeType Emplace(ArgsType&&... Args)
 	{
-		const SizeType Index = AddUninitialized();
-
 		// If this fails to compile when trying to call Emplace with a non-public constructor,
 		// do not make TArray a friend.
 		//
@@ -2260,8 +2472,17 @@ public:
 		//     // Won't compile if the caller doesn't have access to FMyType::FPrivateToken
 		//     Arr.Emplace(FMyType::FPrivateToken{}, 5, 3.14f, TEXT("Banana"));
 		//
-		::new((void*)(GetData() + Index)) ElementType(Forward<ArgsType>(Args)...);
-		return Index;
+		SizeType Result = UE::Core::Private::ReallocGrow<UE::Core::Private::GetAllocatorFlags<AllocatorType>()>(
+			sizeof(ElementType),
+			alignof(ElementType),
+			1,
+			AllocatorInstance,
+			ArrayNum,
+			ArrayMax
+		);
+		void* Ptr = (char*)AllocatorInstance.GetAllocation() + sizeof(ElementType) * Result;
+		(void)new (Ptr) ElementType(Forward<ArgsType>(Args)...);
+		return ArrayNum - 1;
 	}
 
 	/**
@@ -2273,10 +2494,38 @@ public:
 	template <typename... ArgsType>
 	FORCEINLINE ElementType& Emplace_GetRef(ArgsType&&... Args) UE_LIFETIMEBOUND
 	{
-		const SizeType Index = AddUninitialized();
-		ElementType* Ptr = GetData() + Index;
-		::new((void*)Ptr) ElementType(Forward<ArgsType>(Args)...);
-		return *Ptr;
+		// If this fails to compile when trying to call Emplace with a non-public constructor,
+		// do not make TArray a friend.
+		//
+		// Instead, prefer this pattern:
+		//
+		//     class FMyType
+		//     {
+		//     private:
+		//         struct FPrivateToken { explicit FPrivateToken() = default; };
+		//
+		//     public:
+		//         // This has an equivalent access level to a private constructor,
+		//         // as only friends of FMyType will have access to FPrivateToken,
+		//         // but Emplace can legally call it since it's public.
+		//         explicit FMyType(FPrivateToken, int32 Int, float Real, const TCHAR* String);
+		//     };
+		//
+		//     TArray<FMyType> Arr:
+		//
+		//     // Won't compile if the caller doesn't have access to FMyType::FPrivateToken
+		//     Arr.Emplace(FMyType::FPrivateToken{}, 5, 3.14f, TEXT("Banana"));
+		//
+		SizeType Result = UE::Core::Private::ReallocGrow<UE::Core::Private::GetAllocatorFlags<AllocatorType>()>(
+			sizeof(ElementType),
+			alignof(ElementType),
+			1,
+			AllocatorInstance,
+			ArrayNum,
+			ArrayMax
+		);
+		void* Ptr = (char*)AllocatorInstance.GetAllocation() + sizeof(ElementType) * Result;
+		return *new (Ptr) ElementType(Forward<ArgsType>(Args)...);
 	}
 
 	/**
@@ -2647,7 +2896,14 @@ public:
 		}
 		else if (Number > ArrayMax)
 		{
-			ResizeTo(Number);
+			UE::Core::Private::ReallocTo<UE::Core::Private::GetAllocatorFlags<AllocatorType>()>(
+				sizeof(ElementType),
+				alignof(ElementType),
+				Number,
+				AllocatorInstance,
+				ArrayNum,
+				ArrayMax
+			);
 		}
 	}
 
@@ -2795,7 +3051,13 @@ public:
 
 		if (bRemoved && AllowShrinking == EAllowShrinking::Yes)
 		{
-			ResizeShrink();
+			UE::Core::Private::ReallocShrink<UE::Core::Private::GetAllocatorFlags<AllocatorType>()>(
+				sizeof(ElementType),
+				alignof(ElementType),
+				AllocatorInstance,
+				ArrayNum,
+				ArrayMax
+			);
 		}
 
 		return OriginalNum - ArrayNum;
@@ -2866,7 +3128,13 @@ public:
 
 		if (bRemoved && AllowShrinking == EAllowShrinking::Yes)
 		{
-			ResizeShrink();
+			UE::Core::Private::ReallocShrink<UE::Core::Private::GetAllocatorFlags<AllocatorType>()>(
+				sizeof(ElementType),
+				alignof(ElementType),
+				AllocatorInstance,
+				ArrayNum,
+				ArrayMax
+			);
 		}
 
 		return OriginalNum - ArrayNum;
@@ -3088,108 +3356,6 @@ private:
 #endif
 
 private:
-	void AllocatorResizeAllocation(SizeType CurrentArrayNum, SizeType NewArrayMax)
-	{
-		if constexpr (TAllocatorTraits<AllocatorType>::SupportsElementAlignment)
-		{
-			AllocatorInstance.ResizeAllocation(CurrentArrayNum, NewArrayMax, sizeof(ElementType), alignof(ElementType));
-		}
-		else
-		{
-			AllocatorInstance.ResizeAllocation(CurrentArrayNum, NewArrayMax, sizeof(ElementType));
-		}
-		SlackTrackerNumChanged();
-	}
-
-	SizeType AllocatorCalculateSlackShrink(SizeType CurrentArrayNum, SizeType NewArrayMax)
-	{
-		if constexpr (TAllocatorTraits<AllocatorType>::SupportsElementAlignment)
-		{
-			return AllocatorInstance.CalculateSlackShrink(CurrentArrayNum, NewArrayMax, sizeof(ElementType), alignof(ElementType));
-		}
-		else
-		{
-			return AllocatorInstance.CalculateSlackShrink(CurrentArrayNum, NewArrayMax, sizeof(ElementType));
-		}
-	}
-
-	SizeType AllocatorCalculateSlackGrow(SizeType CurrentArrayNum, SizeType NewArrayMax)
-	{
-		if constexpr (TAllocatorTraits<AllocatorType>::SupportsElementAlignment)
-		{
-			return AllocatorInstance.CalculateSlackGrow(CurrentArrayNum, NewArrayMax, sizeof(ElementType), alignof(ElementType));
-		}
-		else
-		{
-			return AllocatorInstance.CalculateSlackGrow(CurrentArrayNum, NewArrayMax, sizeof(ElementType));
-		}
-	}
-
-	SizeType AllocatorCalculateSlackReserve(SizeType NewArrayMax)
-	{
-		if constexpr (TAllocatorTraits<AllocatorType>::SupportsElementAlignment)
-		{
-			return AllocatorInstance.CalculateSlackReserve(NewArrayMax, sizeof(ElementType), alignof(ElementType));
-		}
-		else
-		{
-			return AllocatorInstance.CalculateSlackReserve(NewArrayMax, sizeof(ElementType));
-		}
-	}
-
-	FORCENOINLINE void ResizeGrow(SizeType OldNum)
-	{
-		SizeType LocalArrayNum = ArrayNum;
-
-		// This should only happen when we've underflowed or overflowed SizeType in the caller
-		if (LocalArrayNum < OldNum)
-		{
-			// Cast to USizeType first to prevent sign extension on negative sizes, producing unusually large values.
-			UE::Core::Private::OnInvalidArrayNum((unsigned long long)(USizeType)LocalArrayNum);
-		}
-		ArrayMax = AllocatorCalculateSlackGrow(LocalArrayNum, ArrayMax);
-		AllocatorResizeAllocation(OldNum, ArrayMax);
-	}
-	FORCENOINLINE void ResizeShrink()
-	{
-		const SizeType NewArrayMax = AllocatorCalculateSlackShrink(ArrayNum, ArrayMax);
-		if (NewArrayMax != ArrayMax)
-		{
-			ArrayMax = NewArrayMax;
-			check(ArrayMax >= ArrayNum);
-			AllocatorResizeAllocation(ArrayNum, ArrayMax);
-		}
-	}
-	FORCENOINLINE void ResizeTo(SizeType NewMax)
-	{
-		if (NewMax)
-		{
-			NewMax = AllocatorCalculateSlackReserve(NewMax);
-		}
-		if (NewMax != ArrayMax)
-		{
-			ArrayMax = NewMax;
-			AllocatorResizeAllocation(ArrayNum, ArrayMax);
-		}
-	}
-	FORCENOINLINE void ResizeForCopy(SizeType NewMax, SizeType PrevMax)
-	{
-		if (NewMax)
-		{
-			NewMax = AllocatorCalculateSlackReserve(NewMax);
-		}
-		if (NewMax > PrevMax)
-		{
-			AllocatorResizeAllocation(0, NewMax);
-			ArrayMax = NewMax;
-		}
-		else
-		{
-			ArrayMax = PrevMax;
-		}
-	}
-
-
 	/**
 	 * Copies data from one array into this array. Uses the fast path if the
 	 * data in question does not need a constructor.
@@ -3206,7 +3372,15 @@ private:
 		ArrayNum = NewNum;
 		if (OtherNum || PrevMax)
 		{
-			ResizeForCopy(NewNum, PrevMax);
+			UE::Core::Private::ReallocForCopy<UE::Core::Private::GetAllocatorFlags<AllocatorType>()>(
+				sizeof(ElementType),
+				alignof(ElementType),
+				NewNum,
+				PrevMax,
+				AllocatorInstance,
+				ArrayNum,
+				ArrayMax
+			);
 			ConstructItems<ElementType>((void*)GetData(), OtherData, OtherNum);
 		}
 		else
@@ -3243,7 +3417,15 @@ private:
 				UE::Core::Private::OnInvalidArrayNum((unsigned long long)NewMax);
 			}
 
-			ResizeForCopy(NewNum + ExtraSlack, PrevMax);
+			UE::Core::Private::ReallocForCopy<UE::Core::Private::GetAllocatorFlags<AllocatorType>()>(
+				sizeof(ElementType),
+				alignof(ElementType),
+				NewNum + ExtraSlack,
+				PrevMax,
+				AllocatorInstance,
+				ArrayNum,
+				ArrayMax
+			);
 			ConstructItems<ElementType>((void*)GetData(), OtherData, OtherNum);
 		}
 		else
@@ -3759,7 +3941,15 @@ struct TArrayPrivateFriend
 			// Serialize simple bytes which require no construction or destruction.
 			if ((A.ArrayNum || A.ArrayMax) && Ar.IsLoading())
 			{
-				A.ResizeForCopy(A.ArrayNum, A.ArrayMax);
+				UE::Core::Private::ReallocForCopy<UE::Core::Private::GetAllocatorFlags<AllocatorType>()>(
+					sizeof(ElementType),
+					alignof(ElementType),
+					A.ArrayNum,
+					A.ArrayMax,
+					A.AllocatorInstance,
+					A.ArrayNum,
+					A.ArrayMax
+				);
 			}
 
 			if constexpr (TIsUECoreVariant<ElementType, double>::Value)
diff --git a/Engine/Source/Runtime/Core/Public/Containers/ContainerAllocationPolicies.h b/Engine/Source/Runtime/Core/Public/Containers/ContainerAllocationPolicies.h
--- a/Engine/Source/Runtime/Core/Public/Containers/ContainerAllocationPolicies.h
+++ b/Engine/Source/Runtime/Core/Public/Containers/ContainerAllocationPolicies.h
@@ -629,7 +629,7 @@ private:
 	using USizeType = std::make_unsigned_t<SizeType>;
 
 public:
-	enum { NeedsElementType = true };
+	enum { NeedsElementType = false };
 	enum { RequireRangeCheck = true };
 
 	class ForAnyElementType
