Fix CI failure: add an AutoRTFM wrapper for the MSVC stdlib's _strtoi64 function.

diff --git a/Engine/Source/Runtime/Core/Private/AutoRTFM/StdLibSupport.cpp b/Engine/Source/Runtime/Core/Private/AutoRTFM/StdLibSupport.cpp
--- a/Engine/Source/Runtime/Core/Private/AutoRTFM/StdLibSupport.cpp
+++ b/Engine/Source/Runtime/Core/Private/AutoRTFM/StdLibSupport.cpp
@@ -268,6 +268,16 @@ RETURN RTFM_StringToFloat(const CHAR* String, CHAR** EndPtr, _locale_t Locale)
     return FN(String, EndPtr, Locale);
 }
 
+template<typename RETURN, typename CHAR, RETURN FN(const CHAR*, CHAR**, int)>
+RETURN RTFM_StringToIntNoLocale(const CHAR* String, CHAR** EndPtr, int Radix)
+{
+    if (nullptr != EndPtr)
+    {
+        AutoRTFM::RecordOpenWrite(EndPtr);
+    }
+    return FN(String, EndPtr, Radix);
+}
+
 template<typename RETURN, typename CHAR, RETURN FN(const CHAR*, CHAR**, int, _locale_t)>
 RETURN RTFM_StringToInt(const CHAR* String, CHAR** EndPtr, int Radix, _locale_t Locale)
 {
@@ -278,6 +288,8 @@ RETURN RTFM_StringToInt(const CHAR* String, CHAR** EndPtr, int Radix, _locale_t
     return FN(String, EndPtr, Radix, Locale);
 }
 
+UE_AUTORTFM_REGISTER_OPEN_FUNCTION_EXPLICIT(_strtoi64, (RTFM_StringToIntNoLocale<long long, char, _strtoi64>));
+
 UE_AUTORTFM_REGISTER_OPEN_FUNCTION_EXPLICIT(_wcstoi64, (RTFM_StringToInt<long long, wchar_t, _wcstoi64>));
 UE_AUTORTFM_REGISTER_OPEN_FUNCTION_EXPLICIT(_wcstoui64, (RTFM_StringToInt<unsigned long long, wchar_t, _wcstoui64>));
 UE_AUTORTFM_REGISTER_OPEN_FUNCTION_EXPLICIT(_wcstod_l, (RTFM_StringToFloat<double, wchar_t, _wcstod_l>));
