Fix a few UE headers to avoid hitting the SwitchUnhandledEnumerator warning without making any changes in behavior.
This is likely scratching the surface, but was all that was needed to enable that warning on the Verse compiler and runtime modules where we want to enable it.

diff --git a/Engine/Source/Developer/CollectionManager/Public/CollectionManagerTypes.h b/Engine/Source/Developer/CollectionManager/Public/CollectionManagerTypes.h
--- a/Engine/Source/Developer/CollectionManager/Public/CollectionManagerTypes.h
+++ b/Engine/Source/Developer/CollectionManager/Public/CollectionManagerTypes.h
@@ -138,6 +138,8 @@ struct ECollectionShareType
 			return NSLOCTEXT("ECollectionShareType", "CST_Private_Description", "Private. This collection is only visible to you.");
 		case CST_Shared:
 			return NSLOCTEXT("ECollectionShareType", "CST_Shared_Description", "Shared. This collection is visible to everyone.");
+		case CST_System:
+		case CST_All:
 		default:
 			break;
 		}
@@ -154,6 +156,8 @@ struct ECollectionShareType
 			return FName(*FString::Printf(TEXT("ContentBrowser.Private%s"), InSizeSuffix));
 		case CST_Shared:
 			return FName(*FString::Printf(TEXT("ContentBrowser.Shared%s"), InSizeSuffix));
+		case CST_System:
+		case CST_All:
 		default:
 			break;
 		}
diff --git a/Engine/Source/Runtime/ApplicationCore/Public/GenericPlatform/GenericWindow.h b/Engine/Source/Runtime/ApplicationCore/Public/GenericPlatform/GenericWindow.h
--- a/Engine/Source/Runtime/ApplicationCore/Public/GenericPlatform/GenericWindow.h
+++ b/Engine/Source/Runtime/ApplicationCore/Public/GenericPlatform/GenericWindow.h
@@ -53,6 +53,7 @@ FORCEINLINE const TCHAR* LexToString(EWindowMode::Type WindowMode)
 	case EWindowMode::Fullscreen: return TEXT("Fullscreen");
 	case EWindowMode::WindowedFullscreen: return TEXT("WindowedFullscreen");
 	case EWindowMode::Windowed: return TEXT("Windowed");
+	case EWindowMode::NumWindowModes:
 	default: return TEXT("Unknown");
 	}
 }
diff --git a/Engine/Source/Runtime/Core/Public/GenericPlatform/GenericPlatformCrashContext.h b/Engine/Source/Runtime/Core/Public/GenericPlatform/GenericPlatformCrashContext.h
--- a/Engine/Source/Runtime/Core/Public/GenericPlatform/GenericPlatformCrashContext.h
+++ b/Engine/Source/Runtime/Core/Public/GenericPlatform/GenericPlatformCrashContext.h
@@ -736,11 +736,18 @@ public:
 		{
 		// Verse runtime errors only halt the Verse runtime itself; they do not result in a crash.
 		// Certain runtime errors may be recoverable from in the future.
-		case ECrashContextType::VerseRuntimeError: [[fallthrough]];
+		case ECrashContextType::VerseRuntimeError:
 		case ECrashContextType::Ensure:
-			return true;
 		case ECrashContextType::Stall:
 			return true;
+
+		case ECrashContextType::Crash:
+		case ECrashContextType::Assert:
+		case ECrashContextType::GPUCrash:
+		case ECrashContextType::Hang:
+		case ECrashContextType::OutOfMemory:
+		case ECrashContextType::AbnormalShutdown:
+		case ECrashContextType::Max:
 		default:
 			return false;
 		}
diff --git a/Engine/Source/Runtime/Core/Public/Math/Matrix.inl b/Engine/Source/Runtime/Core/Public/Math/Matrix.inl
--- a/Engine/Source/Runtime/Core/Public/Math/Matrix.inl
+++ b/Engine/Source/Runtime/Core/Public/Math/Matrix.inl
@@ -657,6 +657,7 @@ inline TVector<T> TMatrix<T>::GetScaledAxis(EAxis::Type InAxis) const
 	case EAxis::Z:
 		return TVector<T>(M[2][0], M[2][1], M[2][2]);
 
+	case EAxis::None:
 	default:
 		ensure(0);
 		return TVector<T>::ZeroVector;
diff --git a/Engine/Source/Runtime/Core/Public/Math/Rotator.h b/Engine/Source/Runtime/Core/Public/Math/Rotator.h
--- a/Engine/Source/Runtime/Core/Public/Math/Rotator.h
+++ b/Engine/Source/Runtime/Core/Public/Math/Rotator.h
@@ -778,6 +778,7 @@ FORCEINLINE T TRotator<T>::GetComponentForAxis(EAxis::Type Axis) const
 		return Pitch;
 	case EAxis::Z:
 		return Yaw;
+	case EAxis::None:
 	default:
 		return 0.f;
 	}
@@ -797,6 +798,9 @@ FORCEINLINE void TRotator<T>::SetComponentForAxis(EAxis::Type Axis, T Component)
 	case EAxis::Z:
 		Yaw = Component;
 		break;
+	case EAxis::None:
+	default:
+		break;
 	}
 }
 
diff --git a/Engine/Source/Runtime/Core/Public/Math/Vector.h b/Engine/Source/Runtime/Core/Public/Math/Vector.h
--- a/Engine/Source/Runtime/Core/Public/Math/Vector.h
+++ b/Engine/Source/Runtime/Core/Public/Math/Vector.h
@@ -1952,6 +1952,7 @@ FORCEINLINE T TVector<T>::GetComponentForAxis(EAxis::Type Axis) const
         return Y;
     case EAxis::Z:
         return Z;
+	case EAxis::None:
     default:
         return 0.f;
     }
@@ -1971,6 +1972,9 @@ FORCEINLINE void TVector<T>::SetComponentForAxis(EAxis::Type Axis, T Component)
     case EAxis::Z:
         Z = Component;
         break;
+	case EAxis::None:
+	default:
+		break;
     }
 }
 
diff --git a/Engine/Source/Runtime/Core/Public/Misc/QueuedThreadPool.h b/Engine/Source/Runtime/Core/Public/Misc/QueuedThreadPool.h
--- a/Engine/Source/Runtime/Core/Public/Misc/QueuedThreadPool.h
+++ b/Engine/Source/Runtime/Core/Public/Misc/QueuedThreadPool.h
@@ -37,6 +37,7 @@ inline const TCHAR* LexToString(EQueuedWorkPriority Priority)
 		return TEXT("Low");
 	case EQueuedWorkPriority::Lowest:
 		return TEXT("Lowest");
+	case EQueuedWorkPriority::Count:
 	default:
 		check(false);
 		return TEXT("Unknown");
diff --git a/Engine/Source/Runtime/Core/Public/PixelFormat.h b/Engine/Source/Runtime/Core/Public/PixelFormat.h
--- a/Engine/Source/Runtime/Core/Public/PixelFormat.h
+++ b/Engine/Source/Runtime/Core/Public/PixelFormat.h
@@ -269,6 +269,7 @@ ENUM_CLASS_FLAGS(EPixelFormatCapabilities);
 
 static inline bool IsETCBlockCompressedPixelFormat(EPixelFormat PixelFormat)
 {
+	PRAGMA_DISABLE_SWITCH_UNHANDLED_ENUM_CASE_WARNINGS;
 	switch (PixelFormat)
 	{
 	case PF_ETC2_RGB:
@@ -277,11 +278,13 @@ static inline bool IsETCBlockCompressedPixelFormat(EPixelFormat PixelFormat)
 	case PF_ETC2_RG11_EAC:
 		return true;
 	}
+	PRAGMA_RESTORE_SWITCH_UNHANDLED_ENUM_CASE_WARNINGS;
 	return false;
 }
 
 FORCEINLINE bool IsASTCBlockCompressedTextureFormat(EPixelFormat PixelFormat)
 {
+	PRAGMA_DISABLE_SWITCH_UNHANDLED_ENUM_CASE_WARNINGS;
 	switch (PixelFormat)
 	{
 	case PF_ASTC_4x4:
@@ -303,10 +306,12 @@ FORCEINLINE bool IsASTCBlockCompressedTextureFormat(EPixelFormat PixelFormat)
 	default:
 		return false;
 	}
+	PRAGMA_RESTORE_SWITCH_UNHANDLED_ENUM_CASE_WARNINGS;
 }
 
 static inline bool IsBlockCompressedFormat(EPixelFormat Format)
 {
+	PRAGMA_DISABLE_SWITCH_UNHANDLED_ENUM_CASE_WARNINGS;
 	switch (Format)
 	{
 	case PF_DXT1:
@@ -318,6 +323,7 @@ static inline bool IsBlockCompressedFormat(EPixelFormat Format)
 	case PF_BC7:
 		return true;
 	}
+	PRAGMA_RESTORE_SWITCH_UNHANDLED_ENUM_CASE_WARNINGS;
 	return false;
 }
 
@@ -331,6 +337,7 @@ FORCEINLINE bool IsHDR(EPixelFormat PixelFormat)
 
 FORCEINLINE bool IsInteger(EPixelFormat PixelFormat)
 {
+	PRAGMA_DISABLE_SWITCH_UNHANDLED_ENUM_CASE_WARNINGS;
 	switch (PixelFormat)
 	{
 	case PF_R32_UINT:
@@ -352,12 +359,14 @@ FORCEINLINE bool IsInteger(EPixelFormat PixelFormat)
 	case PF_R16G16_SINT:
 		return true;
 	}
+	PRAGMA_RESTORE_SWITCH_UNHANDLED_ENUM_CASE_WARNINGS;
 	return false;
 }
 
 // IsFloatFormat does not include all floating point formats; see also IsHDR
 static bool IsFloatFormat(EPixelFormat Format)
 {
+	PRAGMA_DISABLE_SWITCH_UNHANDLED_ENUM_CASE_WARNINGS;
 	switch (Format)
 	{
 	case PF_A32B32G32R32F:
@@ -372,11 +381,13 @@ static bool IsFloatFormat(EPixelFormat Format)
 	case PF_FloatR11G11B10:
 		return true;
 	}
+	PRAGMA_RESTORE_SWITCH_UNHANDLED_ENUM_CASE_WARNINGS;
 	return false;
 }
 
 FORCEINLINE bool IsDepthOrStencilFormat(EPixelFormat Format)
 {
+	PRAGMA_DISABLE_SWITCH_UNHANDLED_ENUM_CASE_WARNINGS;
 	switch (Format)
 	{
 	case PF_D24:
@@ -386,22 +397,26 @@ FORCEINLINE bool IsDepthOrStencilFormat(EPixelFormat Format)
 	case PF_R32_FLOAT:
 		return true;
 	}
+	PRAGMA_RESTORE_SWITCH_UNHANDLED_ENUM_CASE_WARNINGS;
 	return false;
 }
 
 FORCEINLINE bool IsStencilFormat(EPixelFormat Format)
 {
+	PRAGMA_DISABLE_SWITCH_UNHANDLED_ENUM_CASE_WARNINGS;
 	switch (Format)
 	{
 	case PF_DepthStencil:
 	case PF_X24_G8:
 		return true;
 	}
+	PRAGMA_RESTORE_SWITCH_UNHANDLED_ENUM_CASE_WARNINGS;
 	return false;
 }
 
 FORCEINLINE bool IsDXTCBlockCompressedTextureFormat(EPixelFormat PixelFormat)
 {
+	PRAGMA_DISABLE_SWITCH_UNHANDLED_ENUM_CASE_WARNINGS;
 	switch (PixelFormat)
 	{
 	case PF_DXT1:
@@ -415,6 +430,7 @@ FORCEINLINE bool IsDXTCBlockCompressedTextureFormat(EPixelFormat PixelFormat)
 	default:
 		return false;
 	}
+	PRAGMA_RESTORE_SWITCH_UNHANDLED_ENUM_CASE_WARNINGS;
 }
 
 FORCEINLINE bool RequiresBlock4Alignment(EPixelFormat PixelFormat)
diff --git a/Engine/Source/Runtime/CoreUObject/Public/Cooker/CookDependency.h b/Engine/Source/Runtime/CoreUObject/Public/Cooker/CookDependency.h
--- a/Engine/Source/Runtime/CoreUObject/Public/Cooker/CookDependency.h
+++ b/Engine/Source/Runtime/CoreUObject/Public/Cooker/CookDependency.h
@@ -421,6 +421,7 @@ inline FCbFieldViewIterator FCookDependency::GetFunctionArgs() const
 
 inline FName FCookDependency::GetPackageName() const
 {
+	PRAGMA_DISABLE_SWITCH_UNHANDLED_ENUM_CASE_WARNINGS;
 	switch (GetType())
 	{
 	case ECookDependency::TransitiveBuild:
@@ -430,6 +431,7 @@ inline FName FCookDependency::GetPackageName() const
 	default:
 		return NAME_None;
 	};
+	PRAGMA_RESTORE_SWITCH_UNHANDLED_ENUM_CASE_WARNINGS;
 }
 
 inline bool FCookDependency::IsAlsoAddRuntimeDependency() const
@@ -523,6 +525,7 @@ inline bool FCookDependency::operator<(const FCookDependency& Other) const
 			return false; // equal
 		}
 		return ARFilterLessThan(*ARFilter, *Other.ARFilter);
+	case ECookDependency::Count:
 	default:
 		checkNoEntry();
 		return false;
@@ -598,6 +601,7 @@ inline bool FCookDependency::operator==(const FCookDependency& Other) const
 
 		return ARFilterEqual(*ARFilter, *Other.ARFilter);
 	}
+	case ECookDependency::Count:
 	default:
 		checkNoEntry();
 		return false;
diff --git a/Engine/Source/Runtime/CoreUObject/Public/UObject/FastReferenceCollector.h b/Engine/Source/Runtime/CoreUObject/Public/UObject/FastReferenceCollector.h
--- a/Engine/Source/Runtime/CoreUObject/Public/UObject/FastReferenceCollector.h
+++ b/Engine/Source/Runtime/CoreUObject/Public/UObject/FastReferenceCollector.h
@@ -698,6 +698,7 @@ FORCEINLINE_DEBUGGABLE void VisitMembers(DispatcherType& Dispatcher, FSchemaView
 			case EMemberType::VerseValueArray:			Dispatcher.HandleVerseValueArray(*(TArray<Verse::VValue>*)MemberPtr, FMemberId(DebugIdx), Origin);
 			break;
 #endif
+			case EMemberType::Count:
 			default:									LogIllegalTypeFatal(Member.Type, DebugIdx, Instance);
 			return;
 			}
diff --git a/Engine/Source/Runtime/CoreUObject/Public/UObject/UnrealType.h b/Engine/Source/Runtime/CoreUObject/Public/UObject/UnrealType.h
--- a/Engine/Source/Runtime/CoreUObject/Public/UObject/UnrealType.h
+++ b/Engine/Source/Runtime/CoreUObject/Public/UObject/UnrealType.h
@@ -1922,6 +1922,7 @@ public:
 	{
 		if (const EName* TagType = Tag.Type.ToEName(); TagType && Tag.Type.GetNumber() == NAME_NO_NUMBER_INTERNAL)
 		{
+			PRAGMA_DISABLE_SWITCH_UNHANDLED_ENUM_CASE_WARNINGS;
 			switch (*TagType)
 			{
 			case NAME_Int8Property:
@@ -1985,6 +1986,7 @@ public:
 				// We didn't convert it
 				break;
 			}
+			PRAGMA_RESTORE_SWITCH_UNHANDLED_ENUM_CASE_WARNINGS;
 		}
 
 		return EConvertFromTypeResult::UseSerializeItem;
diff --git a/Engine/Source/Runtime/Json/Public/Serialization/JsonReader.h b/Engine/Source/Runtime/Json/Public/Serialization/JsonReader.h
--- a/Engine/Source/Runtime/Json/Public/Serialization/JsonReader.h
+++ b/Engine/Source/Runtime/Json/Public/Serialization/JsonReader.h
@@ -111,6 +111,11 @@ public:
 					ReadWasSuccess = ReadNextObjectValue( /*OUT*/ CurrentToken );
 					break;
 
+				case EJson::None:
+				case EJson::Null:
+				case EJson::String:
+				case EJson::Number:
+				case EJson::Boolean:
 				default:
 					ReadWasSuccess = ReadStart( /*OUT*/ CurrentToken );
 					break;
diff --git a/Engine/Source/Runtime/Json/Public/Serialization/JsonSerializer.h b/Engine/Source/Runtime/Json/Public/Serialization/JsonSerializer.h
--- a/Engine/Source/Runtime/Json/Public/Serialization/JsonSerializer.h
+++ b/Engine/Source/Runtime/Json/Public/Serialization/JsonSerializer.h
@@ -98,6 +98,11 @@ public:
 		case EJson::Array:
 			OutValue = MakeShared<FJsonValueArray>(State.Array);
 			break;
+		case EJson::None:
+		case EJson::Null:
+		case EJson::String:
+		case EJson::Number:
+		case EJson::Boolean:
 		default:
 			// FIXME: would be nice to handle non-composite root values but StackState Deserialize just drops them on the floor
 			return false;
@@ -548,6 +553,7 @@ private:
 				}
 				break;
 
+			case EJson::None:
 			default: 
 				UE_LOG(LogJson, Fatal,TEXT("Could not print Json Value, unrecognized type."));
 			}
