Trace scoped from TraceLog and stats.

* Hook the optional profiler callbacks in TraceLog to emit CpuProfiler events. Enable by passing `-tracescopes` on the commandline.
* Trace `Block Pool`, `Shared buffer` and `Fixed buffer`  as memory stat events.

diff --git a/Engine/Source/Runtime/Core/Private/ProfilingDebugging/TraceAuxiliary.cpp b/Engine/Source/Runtime/Core/Private/ProfilingDebugging/TraceAuxiliary.cpp
--- a/Engine/Source/Runtime/Core/Private/ProfilingDebugging/TraceAuxiliary.cpp
+++ b/Engine/Source/Runtime/Core/Private/ProfilingDebugging/TraceAuxiliary.cpp
@@ -81,10 +81,15 @@ CSV_DEFINE_CATEGORY(Trace, true);
 ////////////////////////////////////////////////////////////////////////////////
 DECLARE_STATS_GROUP(TEXT("TraceLog"), STATGROUP_Trace, STATCAT_Advanced);
 DECLARE_MEMORY_STAT(TEXT("Memory used"), STAT_TraceMemoryUsed, STATGROUP_Trace);
+DECLARE_MEMORY_STAT(TEXT("Block pool"), STAT_TraceBlockPoolAllocated, STATGROUP_Trace);
+DECLARE_MEMORY_STAT(TEXT("Shared buffer"), STAT_TraceSharedAllocated, STATGROUP_Trace);
+DECLARE_MEMORY_STAT(TEXT("Fixed buffer"), STAT_TraceFixedAllocated, STATGROUP_Trace);
 DECLARE_MEMORY_STAT(TEXT("Important event cache"), STAT_TraceCacheAllocated, STATGROUP_Trace);
 DECLARE_MEMORY_STAT(TEXT("Important event cache used"), STAT_TraceCacheUsed, STATGROUP_Trace);
 DECLARE_MEMORY_STAT(TEXT("Important event cache waste"), STAT_TraceCacheWaste, STATGROUP_Trace);
 DECLARE_MEMORY_STAT(TEXT("Sent"), STAT_TraceSent, STATGROUP_Trace);
+DECLARE_MEMORY_STAT(TEXT("Traced"), STAT_TraceTraced, STATGROUP_Trace);
+DECLARE_MEMORY_STAT(TEXT("Emitted"), STAT_TraceEmitted, STATGROUP_Trace);
 
 ////////////////////////////////////////////////////////////////////////////////
 enum class ETraceConnectType
@@ -840,10 +845,15 @@ void FTraceAuxiliaryImpl::StartEndFramePump()
 			UE::Trace::FStatistics Stats;
 			UE::Trace::GetStatistics(Stats);
 			SET_MEMORY_STAT(STAT_TraceMemoryUsed, Stats.MemoryUsed);
+			SET_MEMORY_STAT(STAT_TraceBlockPoolAllocated, Stats.BlockPoolAllocated);
+			SET_MEMORY_STAT(STAT_TraceSharedAllocated, Stats.SharedBufferAllocated);
+			SET_MEMORY_STAT(STAT_TraceFixedAllocated, Stats.FixedBufferAllocated);
 			SET_MEMORY_STAT(STAT_TraceCacheAllocated, Stats.CacheAllocated);
 			SET_MEMORY_STAT(STAT_TraceCacheUsed, Stats.CacheUsed);
 			SET_MEMORY_STAT(STAT_TraceCacheWaste, Stats.CacheWaste);
 			SET_MEMORY_STAT(STAT_TraceSent, Stats.BytesSent);
+			SET_MEMORY_STAT(STAT_TraceTraced, Stats.BytesTraced);
+			SET_MEMORY_STAT(STAT_TraceEmitted, Stats.BytesEmitted);
 
 #if ENABLE_LOW_LEVEL_MEM_TRACKER
 			const int64 NonCacheMemory = Stats.MemoryUsed - Stats.CacheAllocated;
@@ -907,6 +917,18 @@ void OnMessageCallback(const UE::Trace::FMessageEvent& Message)
 	#undef LOG_WITH_VERBOSITY
 }
 
+////////////////////////////////////////////////////////////////////////////////
+void OnScopeBeginCallback(const ANSICHAR* ScopeName)
+{
+	FCpuProfilerTrace::OutputBeginDynamicEvent(ScopeName);
+}
+
+////////////////////////////////////////////////////////////////////////////////
+void OnScopeEndCallback()
+{
+	FCpuProfilerTrace::OutputEndEvent();
+}
+
 ////////////////////////////////////////////////////////////////////////////////
 static void SetupInitFromConfig(UE::Trace::FInitializeDesc& OutDesc)
 {
@@ -1685,6 +1707,11 @@ void FTraceAuxiliary::Initialize(const TCHAR* CommandLine)
 
 	// Initialize Trace. The settings are stored in a static for posterity.
 	UE::Trace::FInitializeDesc& Desc = GInitializeDesc;
+	if (FParse::Param(CommandLine, TEXT("tracescopes")))
+	{
+		Desc.OnScopeBeginFunc = &OnScopeBeginCallback;
+		Desc.OnScopeEndFunc = &OnScopeEndCallback;
+	}
 #if WITH_EDITOR
 	Desc.TailSizeBytes = 32 << 20;
 #endif
