[core] Expression parser: moved basic math token definitions to static functions so they can be resused

diff --git a/Engine/Source/Runtime/Core/Private/Math/BasicMathExpressionEvaluator.cpp b/Engine/Source/Runtime/Core/Private/Math/BasicMathExpressionEvaluator.cpp
--- a/Engine/Source/Runtime/Core/Private/Math/BasicMathExpressionEvaluator.cpp
+++ b/Engine/Source/Runtime/Core/Private/Math/BasicMathExpressionEvaluator.cpp
@@ -139,13 +139,10 @@ namespace ExpressionParser
 	}
 }
 
-FBasicMathExpressionEvaluator::FBasicMathExpressionEvaluator()
+void FBasicMathExpressionEvaluator::DefineArithmeticTokens(FTokenDefinitions& TokenDefinitions)
 {
 	using namespace ExpressionParser;
 
-	TokenDefinitions.IgnoreWhitespace();
-	TokenDefinitions.DefineToken(&ConsumeSymbol<FSubExpressionStart>);
-	TokenDefinitions.DefineToken(&ConsumeSymbol<FSubExpressionEnd>);
 	TokenDefinitions.DefineToken(&ConsumeSymbol<FPlusEquals>);
 	TokenDefinitions.DefineToken(&ConsumeSymbol<FMinusEquals>);
 	TokenDefinitions.DefineToken(&ConsumeSymbol<FStarEquals>);
@@ -157,9 +154,12 @@ FBasicMathExpressionEvaluator::FBasicMathExpressionEvaluator()
 	TokenDefinitions.DefineToken(&ConsumeSymbol<FPercent>);
 	TokenDefinitions.DefineToken(&ConsumeSymbol<FSquareRoot>);
 	TokenDefinitions.DefineToken(&ConsumeSymbol<FPower>);
-	TokenDefinitions.DefineToken(&ConsumeLocalizedNumberWithAgnosticFallback);
+}
+
+void FBasicMathExpressionEvaluator::DefineArithmeticOperators(FExpressionGrammar& Grammar)
+{
+	using namespace ExpressionParser;
 
-	Grammar.DefineGrouping<FSubExpressionStart, FSubExpressionEnd>();
 	Grammar.DefinePreUnaryOperator<FPlus>();
 	Grammar.DefinePreUnaryOperator<FMinus>();
 	Grammar.DefinePreUnaryOperator<FSquareRoot>();
@@ -172,6 +172,12 @@ FBasicMathExpressionEvaluator::FBasicMathExpressionEvaluator()
 	Grammar.DefineBinaryOperator<FPercent>(4, EAssociativity::LeftToRight);
 	Grammar.DefineBinaryOperator<FPower>(3);
 
+}
+
+void FBasicMathExpressionEvaluator::MapArithmeticOperators(FOperatorJumpTable& JumpTable)
+{
+	using namespace ExpressionParser;
+
 	JumpTable.MapPreUnary<FPlus>([](double N)			{ return N; });
 	JumpTable.MapPreUnary<FMinus>([](double N)			{ return -N; });
 	JumpTable.MapPreUnary<FSquareRoot>([](double A)		{ return double(FMath::Sqrt((float)A)); }); //@TODO: FLOATPRECISION: Needs a double version of Sqrt
@@ -199,6 +205,22 @@ FBasicMathExpressionEvaluator::FBasicMathExpressionEvaluator()
 	});
 }
 
+FBasicMathExpressionEvaluator::FBasicMathExpressionEvaluator()
+{
+	using namespace ExpressionParser;
+
+	TokenDefinitions.IgnoreWhitespace();
+	TokenDefinitions.DefineToken(&ConsumeSymbol<FSubExpressionStart>);
+	TokenDefinitions.DefineToken(&ConsumeSymbol<FSubExpressionEnd>);
+	DefineArithmeticTokens(TokenDefinitions);
+	TokenDefinitions.DefineToken(&ConsumeLocalizedNumberWithAgnosticFallback);
+
+	Grammar.DefineGrouping<FSubExpressionStart, FSubExpressionEnd>();
+	DefineArithmeticOperators(Grammar);
+
+	MapArithmeticOperators(JumpTable);
+}
+
 TValueOrError<double, FExpressionError> FBasicMathExpressionEvaluator::Evaluate(const TCHAR* InExpression, double InExistingValue) const
 {
 	using namespace ExpressionParser;
diff --git a/Engine/Source/Runtime/Core/Public/Math/BasicMathExpressionEvaluator.h b/Engine/Source/Runtime/Core/Public/Math/BasicMathExpressionEvaluator.h
--- a/Engine/Source/Runtime/Core/Public/Math/BasicMathExpressionEvaluator.h
+++ b/Engine/Source/Runtime/Core/Public/Math/BasicMathExpressionEvaluator.h
@@ -87,7 +87,16 @@ public:
 
 	/** Evaluate the given expression, resulting in either a double value, or an error */
 	CORE_API TValueOrError<double, FExpressionError> Evaluate(const TCHAR* InExpression, double InExistingValue = 0) const;
-	
+
+	/**
+	 * Utils that can be re-used as part of more complex evaluators
+	 * 
+	 * Note precedences are 3-5 for power/multiplicative/additive 
+	 */
+	CORE_API static void DefineArithmeticTokens(FTokenDefinitions& TokenDefinitions);
+	CORE_API static void DefineArithmeticOperators(FExpressionGrammar& TokenDefinitions);
+	CORE_API static void MapArithmeticOperators(FOperatorJumpTable& JumpTable);
+
 private:
 	FTokenDefinitions TokenDefinitions;
 	FExpressionGrammar Grammar;
