Always assert by default when hitting an Abort in AutoRTFM
If we have disabled ensures, make sure we disable asserts too

diff --git a/Engine/Source/Runtime/Core/Private/AutoRTFM/API.cpp b/Engine/Source/Runtime/Core/Private/AutoRTFM/API.cpp
--- a/Engine/Source/Runtime/Core/Private/AutoRTFM/API.cpp
+++ b/Engine/Source/Runtime/Core/Private/AutoRTFM/API.cpp
@@ -29,7 +29,7 @@ namespace
 
 	int GMemoryValidationLevel = static_cast<int>(GDefaultMemoryValidationLevel);
 
-	bool GAutoRTFMAssertOnError = false;
+	bool GAutoRTFMAssertOnError = true;
 
 	void OnIsAutoRTFMEnabledChanged()
 	{
@@ -171,6 +171,12 @@ namespace AutoRTFM
 		{
 #if UE_AUTORTFM
 			GAutoRTFMEnsureOnAbortByLanguage = bEnabled;
+
+			// if we have disabled ensures on aborts lets make sure we have disabled asserts too
+			if (bEnabled == false)
+			{
+				GAutoRTFMAssertOnError = false;
+			}
 #endif
 		}
 
