Unify Stats headers

NOTE: This change is primarily moving existing code around. There are very few removals and no new code (aside from headers/a couple macros).

The goal of this change is to merge Stats.h and Stats2.h to simplify understanding and working with the various stat macros.

We have several different systems that utilize these macros based on what's defined:
- When STATS is 1 we're using the UE Stats system which tracks various stats at runtime to allow for real time displays as well as piping to profilers.
- When STATS is 0 but ENABLE_STATNAMEDEVENTS is 1, we use a simpler version of these macros to just emit named events to profilers.
- If STATS is 0 and ENABLE_STATNAMEDEVENTS is 0 but USE_HITCH_DETECTION is 1, we use these macros to do very lightweight stack tracking for sake of logging hitch callstack. This can also be present when ENABLE_STATNAMEDEVENTS is 1.

The main point of confusion is functionality for all 3 cases was mixed across both Stats.h and Stats2.h making it difficult to understand, especially with the complicated preprocessor conditions.

The result of this change is all the macro defintions for each case are defined in Stats.h and Stats2.h is a deprecated stub.

Summary of changes:

Stats2
- All the non-macro types that were defined in Stats2.h when STATS == 1 were moved to StatsSystemTypes.h.
- FThreadIdleStats was moved to ThreadIdleStats.h as it's self contained. Now wrapped in UE::Stats namespace.
- DirectCommand was moved to StatsCommand.h to match its implementation, StatsCommand.cpp. Now wrapped in UE::Stats namespace.
- The FStats struct was moved to StatsSystem.h and its implementation to StatsSystem.cpp. Now wrapped in UE::Stats namespace.
- The global stat groups/stats like STAT_FrameTime were moved to GlobalStats.ini.
- The DEFINE_STAT calls for the global stats was moved to Stats.cpp.
- All the macros were moved to Stats.h in the first #if STATS section.

Stats.h
- EStatFlags and PROFILER_CHAR were moved to StatsCommon.h
- FLightweightStatScope was renamed to FHitchTrackingStatScope and moved to HitchTrackingStatScope.h. Now wrapped in UE::Stats namespace. There is also now just a single version of this class.
- The lighweight versions of TStatId and FScopeCycleCounter were moved to LightweightStats.h
- FDynamicStats was moved to DynamicStats.h.

Preflights:
Presubmit Engine: 67119aab970f1596d94c6e49, 671fce17281de00699c47626
Compilation Tests Not Windows: 67119aa6331ffb1bc53659e6, 671fce1e281de00699c476ff
Compilation Tests Windows: 67119ab06673800f998b111b, 671fce23281de00699c4777d
Editor Tools & Monolithics: 67119ab4fc4aed04d3364cb3, 671fce28604dfd607945a32c
Standard Build: 671fce4818435b89de74f014

#tests Many compilation preflights
[RN] Merged Stat system headers and made the various code paths clearer

diff --git a/Engine/Plugins/Animation/LiveLinkHub/Source/LiveLinkHub/Private/LiveLinkHubClient.cpp b/Engine/Plugins/Animation/LiveLinkHub/Source/LiveLinkHub/Private/LiveLinkHubClient.cpp
--- a/Engine/Plugins/Animation/LiveLinkHub/Source/LiveLinkHub/Private/LiveLinkHubClient.cpp
+++ b/Engine/Plugins/Animation/LiveLinkHub/Source/LiveLinkHub/Private/LiveLinkHubClient.cpp
@@ -13,7 +13,7 @@
 #include "Misc/App.h"
 #include "Modules/ModuleManager.h"
 #include "Recording/LiveLinkHubPlaybackSourceSettings.h"
-#include "Stats/Stats2.h"
+#include "Stats/Stats.h"
 #include "UObject/Package.h"
 #include "UObject/StrongObjectPtr.h"
 #include "UObject/UObjectGlobals.h"
diff --git a/Engine/Plugins/Animation/LiveLinkHub/Source/LiveLinkHub/Private/LiveLinkHubLog.h b/Engine/Plugins/Animation/LiveLinkHub/Source/LiveLinkHub/Private/LiveLinkHubLog.h
--- a/Engine/Plugins/Animation/LiveLinkHub/Source/LiveLinkHub/Private/LiveLinkHubLog.h
+++ b/Engine/Plugins/Animation/LiveLinkHub/Source/LiveLinkHub/Private/LiveLinkHubLog.h
@@ -3,7 +3,7 @@
 #pragma once
 
 #include "Logging/LogMacros.h"
-#include "Stats/Stats2.h"
+#include "Stats/Stats.h"
 
 DECLARE_LOG_CATEGORY_EXTERN(LogLiveLinkHub, Log, All);
 
diff --git a/Engine/Plugins/Developer/Concert/ConcertSync/ConcertSyncServer/Source/ConcertSyncServer/Public/ConcertSyncServerLoop.inl b/Engine/Plugins/Developer/Concert/ConcertSync/ConcertSyncServer/Source/ConcertSyncServer/Public/ConcertSyncServerLoop.inl
--- a/Engine/Plugins/Developer/Concert/ConcertSync/ConcertSyncServer/Source/ConcertSyncServer/Public/ConcertSyncServerLoop.inl
+++ b/Engine/Plugins/Developer/Concert/ConcertSync/ConcertSyncServer/Source/ConcertSyncServer/Public/ConcertSyncServerLoop.inl
@@ -17,6 +17,7 @@
 #include "Interfaces/IPluginManager.h"
 #include "Async/TaskGraphInterfaces.h"
 #include "HAL/PlatformApplicationMisc.h"
+#include "Stats/StatsSystem.h"
 
 #include "Runtime/Launch/Resources/Version.h"
 
@@ -204,7 +205,7 @@ int32 ConcertSyncServerLoop(const TCHAR* CommandLine, const FConcertSyncServerLo
 		FTSTicker::GetCoreTicker().Tick(DeltaTime);
 
 		GFrameCounter++;
-		FStats::AdvanceFrame(false);
+		UE::Stats::FStats::AdvanceFrame(false);
 		GLog->FlushThreadedLogs();
 
 		// Run garbage collection for the UObjects for the rest of the frame or at least to 2 ms
diff --git a/Engine/Plugins/Media/ElectraPlayer/Source/ElectraPlayerRuntime/Private/Runtime/Decoder/AudioDecoder.cpp b/Engine/Plugins/Media/ElectraPlayer/Source/ElectraPlayerRuntime/Private/Runtime/Decoder/AudioDecoder.cpp
--- a/Engine/Plugins/Media/ElectraPlayer/Source/ElectraPlayerRuntime/Private/Runtime/Decoder/AudioDecoder.cpp
+++ b/Engine/Plugins/Media/ElectraPlayer/Source/ElectraPlayerRuntime/Private/Runtime/Decoder/AudioDecoder.cpp
@@ -3,7 +3,6 @@
 #include "CoreMinimal.h"
 #include "HAL/LowLevelMemTracker.h"
 #include "Stats/Stats.h"
-#include "Stats/Stats2.h"
 
 #include "PlayerCore.h"
 #include "PlayerRuntimeGlobal.h"
diff --git a/Engine/Plugins/Media/ElectraPlayer/Source/ElectraPlayerRuntime/Private/Runtime/Decoder/VideoDecoder.cpp b/Engine/Plugins/Media/ElectraPlayer/Source/ElectraPlayerRuntime/Private/Runtime/Decoder/VideoDecoder.cpp
--- a/Engine/Plugins/Media/ElectraPlayer/Source/ElectraPlayerRuntime/Private/Runtime/Decoder/VideoDecoder.cpp
+++ b/Engine/Plugins/Media/ElectraPlayer/Source/ElectraPlayerRuntime/Private/Runtime/Decoder/VideoDecoder.cpp
@@ -3,7 +3,6 @@
 #include "CoreMinimal.h"
 #include "HAL/LowLevelMemTracker.h"
 #include "Stats/Stats.h"
-#include "Stats/Stats2.h"
 
 #include "PlayerCore.h"
 #include "PlayerRuntimeGlobal.h"
diff --git a/Engine/Plugins/VirtualProduction/TextureShare/Source/TextureShareCore/Private/IPC/TextureShareCoreInterprocessEvent.cpp b/Engine/Plugins/VirtualProduction/TextureShare/Source/TextureShareCore/Private/IPC/TextureShareCoreInterprocessEvent.cpp
--- a/Engine/Plugins/VirtualProduction/TextureShare/Source/TextureShareCore/Private/IPC/TextureShareCoreInterprocessEvent.cpp
+++ b/Engine/Plugins/VirtualProduction/TextureShare/Source/TextureShareCore/Private/IPC/TextureShareCoreInterprocessEvent.cpp
@@ -3,7 +3,8 @@
 #include "IPC/TextureShareCoreInterprocessEvent.h"
 
 #include "ProfilingDebugging/CsvProfiler.h"
-#include "Stats/Stats2.h"
+#include "Stats/Stats.h"
+#include "Stats/ThreadIdleStats.h"
 
 #include "Module/TextureShareCoreLog.h"
 
@@ -103,7 +104,7 @@ bool FTextureShareCoreInterprocessEventWin::Wait(uint32 MaxMillisecondsToWait, c
 
 #if !TEXTURESHARECORE_SDK
 	CSV_SCOPED_WAIT(MaxMillisecondsToWait);
-	FThreadIdleStats::FScopeIdle Scope(bIgnoreThreadIdleStats);
+	UE::Stats::FThreadIdleStats::FScopeIdle Scope(bIgnoreThreadIdleStats);
 #endif
 
 	if (!MaxMillisecondsToWait)
diff --git a/Engine/Source/Developer/Datasmith/DatasmithExporter/Private/DatasmithExporterManager.cpp b/Engine/Source/Developer/Datasmith/DatasmithExporter/Private/DatasmithExporterManager.cpp
--- a/Engine/Source/Developer/Datasmith/DatasmithExporter/Private/DatasmithExporterManager.cpp
+++ b/Engine/Source/Developer/Datasmith/DatasmithExporter/Private/DatasmithExporterManager.cpp
@@ -21,6 +21,7 @@
 #include "Misc/AssertionMacros.h"
 #include "RequiredProgramMainCPPInclude.h"
 #include "StandaloneRenderer.h"
+#include "Stats/StatsSystem.h"
 
 #ifndef UE_BUILD_DEVELOPMENT_WITH_DEBUGGAME
 	#define UE_BUILD_DEVELOPMENT_WITH_DEBUGGAME 0
@@ -177,7 +178,7 @@ uint32 FDatasmithGameThread::Run()
 	while (bKeepRunning)
 	{
 		FTaskGraphInterface::Get().ProcessThreadUntilIdle(ENamedThreads::GameThread);
-		FStats::AdvanceFrame(false);
+		UE::Stats::FStats::AdvanceFrame(false);
 		FTSTicker::GetCoreTicker().Tick(FApp::GetDeltaTime());
 		FSlateApplication::Get().PumpMessages();
 		FSlateApplication::Get().Tick();
diff --git a/Engine/Source/Developer/ProfilerService/Private/ProfilerServiceManager.cpp b/Engine/Source/Developer/ProfilerService/Private/ProfilerServiceManager.cpp
--- a/Engine/Source/Developer/ProfilerService/Private/ProfilerServiceManager.cpp
+++ b/Engine/Source/Developer/ProfilerService/Private/ProfilerServiceManager.cpp
@@ -5,6 +5,7 @@
 #include "Misc/App.h"
 #include "ProfilerServiceMessages.h"
 #include "Serialization/MemoryReader.h"
+#include "Stats/StatsCommand.h"
 #include "Stats/StatsData.h"
 #include "Stats/StatsFile.h"
 #include "Tasks/Pipe.h"
@@ -60,7 +61,7 @@ FProfilerServiceManager::FProfilerServiceManager()
 void FProfilerServiceManager::StartCapture()
 {
 #if STATS
-	DirectStatsCommand(TEXT("stat startfile"));
+	UE::Stats::DirectStatsCommand(TEXT("stat startfile"));
 #endif
 }
 
@@ -68,7 +69,7 @@ void FProfilerServiceManager::StartCapture()
 void FProfilerServiceManager::StopCapture()
 {
 #if STATS
-	DirectStatsCommand(TEXT("stat stopfile"),true);
+	UE::Stats::DirectStatsCommand(TEXT("stat stopfile"),true);
 	// Not thread-safe, but in this case it is ok, because we are waiting for completion.
 	LastStatsFilename = FCommandStatsFile::Get().LastFileSaved;
 #endif
diff --git a/Engine/Source/Editor/UnrealEd/Private/CookOnTheFlyServer.cpp b/Engine/Source/Editor/UnrealEd/Private/CookOnTheFlyServer.cpp
--- a/Engine/Source/Editor/UnrealEd/Private/CookOnTheFlyServer.cpp
+++ b/Engine/Source/Editor/UnrealEd/Private/CookOnTheFlyServer.cpp
@@ -142,6 +142,7 @@
 #include "ShaderLibraryChunkDataGenerator.h"
 #include "ShaderStats.h"
 #include "ShaderStatsCollector.h"
+#include "Stats/StatsSystem.h"
 #include "String/Find.h"
 #include "String/ParseLines.h"
 #include "String/ParseTokens.h"
@@ -2267,7 +2268,7 @@ static void ProcessDeferredCommands(UCookOnTheFlyServer& COTFS)
 static void CBTBTickCommandletStats()
 {
 	UE_SCOPED_COOKTIMER_AND_DURATION(CookByTheBook_TickCommandletStats, DetailedCookStats::TickLoopTickCommandletStatsTimeSec);
-	FStats::TickCommandletStats();
+	UE::Stats::FStats::TickCommandletStats();
 }
 
 static void TickShaderCompilingManager(UE::Cook::FTickStackData& StackData)
diff --git a/Engine/Source/Programs/UnrealFrontend/Private/Commands/UserInterfaceCommand.cpp b/Engine/Source/Programs/UnrealFrontend/Private/Commands/UserInterfaceCommand.cpp
--- a/Engine/Source/Programs/UnrealFrontend/Private/Commands/UserInterfaceCommand.cpp
+++ b/Engine/Source/Programs/UnrealFrontend/Private/Commands/UserInterfaceCommand.cpp
@@ -20,6 +20,7 @@
 #include "OutputLogModule.h"
 #include "IDirectoryWatcher.h"
 #include "DirectoryWatcherModule.h"
+#include "Stats/StatsSystem.h"
 
 #define IDEAL_FRAMERATE 60;
 
@@ -119,7 +120,7 @@ void FUserInterfaceCommand::Run(  )
 		DeltaTime =  CurrentTime - LastTime;
 		LastTime = CurrentTime;
 
-		FStats::AdvanceFrame( false );
+		UE::Stats::FStats::AdvanceFrame( false );
 		FCoreDelegates::OnEndFrame.Broadcast();
 
 		GLog->FlushThreadedLogs();
diff --git a/Engine/Source/Runtime/Core/Private/GenericPlatform/GenericPlatformProcess.cpp b/Engine/Source/Runtime/Core/Private/GenericPlatform/GenericPlatformProcess.cpp
--- a/Engine/Source/Runtime/Core/Private/GenericPlatform/GenericPlatformProcess.cpp
+++ b/Engine/Source/Runtime/Core/Private/GenericPlatform/GenericPlatformProcess.cpp
@@ -13,6 +13,7 @@
 #include "Misc/CommandLine.h"
 #include "Misc/Paths.h"
 #include "Stats/Stats.h"
+#include "Stats/ThreadIdleStats.h"
 #include "Misc/CoreStats.h"
 #include "Misc/EventPool.h"
 #include "Misc/EngineVersion.h"
@@ -374,7 +375,7 @@ void FGenericPlatformProcess::ExploreFolder( const TCHAR* FilePath )
 void FGenericPlatformProcess::Sleep( float Seconds )
 {
 	SCOPE_CYCLE_COUNTER(STAT_Sleep);
-	FThreadIdleStats::FScopeIdle Scope;
+	UE::Stats::FThreadIdleStats::FScopeIdle Scope;
 	SleepNoStats(Seconds);
 }
 
@@ -415,7 +416,7 @@ void FGenericPlatformProcess::ConditionalSleep(TFunctionRef<bool()> Condition, f
 	}
 
 	SCOPE_CYCLE_COUNTER(STAT_Sleep);
-	FThreadIdleStats::FScopeIdle Scope;
+	UE::Stats::FThreadIdleStats::FScopeIdle Scope;
 	do
 	{
 		FPlatformProcess::SleepNoStats(SleepTime);
@@ -432,7 +433,7 @@ bool FPThreadEvent::Wait(uint32 WaitTime, const bool bIgnoreThreadIdleStats /*=
 
 	SCOPE_CYCLE_COUNTER(STAT_EventWait);
 	CSV_SCOPED_WAIT(WaitTime);
-	FThreadIdleStats::FScopeIdle Scope(bIgnoreThreadIdleStats);
+	UE::Stats::FThreadIdleStats::FScopeIdle Scope(bIgnoreThreadIdleStats);
 
 	check(bInitialized);
 
diff --git a/Engine/Source/Runtime/Core/Private/ProfilingDebugging/CsvProfiler.cpp b/Engine/Source/Runtime/Core/Private/ProfilingDebugging/CsvProfiler.cpp
--- a/Engine/Source/Runtime/Core/Private/ProfilingDebugging/CsvProfiler.cpp
+++ b/Engine/Source/Runtime/Core/Private/ProfilingDebugging/CsvProfiler.cpp
@@ -29,7 +29,7 @@
 #include "Misc/EngineVersion.h"
 #include "ProfilingDebugging/CpuProfilerTrace.h"
 #include "Stats/Stats.h"
-#include "Stats/Stats2.h"
+#include "Stats/ThreadIdleStats.h"
 #include "HAL/LowLevelMemTracker.h"
 #include "Misc/Compression.h"
 #include "Misc/ConfigCacheIni.h"
@@ -3879,7 +3879,7 @@ void FCsvProfiler::BeginWait()
 #if CSV_PROFILER_SUPPORT_NAMED_EVENTS
 			if (UNLIKELY(GCsvProfilerNamedEventsExclusive))
 			{
-				if ( FThreadIdleStats::Get().IsCriticalPath() )
+				if ( UE::Stats::FThreadIdleStats::Get().IsCriticalPath() )
 				{
 					CsvBeginNamedEvent(FColor(192, 96, 96), "CsvEventWait");
 
diff --git a/Engine/Source/Runtime/Core/Private/Stats/Stats2.cpp b/Engine/Source/Runtime/Core/Private/Stats/Stats2.cpp
--- a/Engine/Source/Runtime/Core/Private/Stats/Stats2.cpp
+++ b/Engine/Source/Runtime/Core/Private/Stats/Stats2.cpp
@@ -1,1444 +1,3 @@
 // Copyright Epic Games, Inc. All Rights Reserved.
 
-#include "Stats/Stats2.h"
-#include "Misc/AssertionMacros.h"
-#include "HAL/UnrealMemory.h"
-#include "Templates/UnrealTemplate.h"
-#include "Containers/Array.h"
-#include "Misc/CString.h"
-#include "Containers/UnrealString.h"
-#include "UObject/NameTypes.h"
-#include "Logging/LogMacros.h"
-#include "Misc/Parse.h"
-#include "Misc/StringBuilder.h"
-#include "HAL/Runnable.h"
-#include "HAL/RunnableThread.h"
-#include "Misc/SingleThreadRunnable.h"
-#include "HAL/ThreadSafeCounter.h"
-#include "Misc/ScopeLock.h"
-#include "Misc/CommandLine.h"
-#include "Containers/Map.h"
-#include "Stats/Stats.h"
-#include "String/Find.h"
-#include "Async/AsyncWork.h"
-#include "Containers/Ticker.h"
-#include "Stats/StatsData.h"
-#include "HAL/IConsoleManager.h"
-#include "Async/TaskGraphInterfaces.h"
-#include "HAL/LowLevelMemTracker.h"
-#include "Tasks/Pipe.h"
-
-/*-----------------------------------------------------------------------------
-	Global
------------------------------------------------------------------------------*/
-
-struct FStats2Globals
-{
-	static void Get()
-	{
-#if	STATS
-		FStartupMessages::Get();
-		IStatGroupEnableManager::Get();
-#endif // STATS
-	}
-};
-
-#if STATS
-CORE_API UE::Tasks::FPipe GStatsPipe{ UE_SOURCE_LOCATION };
-#endif
-
-static struct FForceInitAtBootFStats2 : public TForceInitAtBoot<FStats2Globals>
-{} FForceInitAtBootFStats2;
-
-DECLARE_DWORD_COUNTER_STAT( TEXT("Frame Packets Received"),STAT_StatFramePacketsRecv,STATGROUP_StatSystem);
-
-DECLARE_CYCLE_STAT(TEXT("StatsNew Tick"),STAT_StatsNewTick,STATGROUP_StatSystem);
-DECLARE_CYCLE_STAT(TEXT("Parse Meta"),STAT_StatsNewParseMeta,STATGROUP_StatSystem);
-DECLARE_CYCLE_STAT(TEXT("Scan For Advance"),STAT_ScanForAdvance,STATGROUP_StatSystem);
-DECLARE_CYCLE_STAT(TEXT("Add To History"),STAT_StatsNewAddToHistory,STATGROUP_StatSystem);
-DECLARE_CYCLE_STAT(TEXT("Flush Raw Stats"),STAT_FlushRawStats,STATGROUP_StatSystem);
-
-DECLARE_MEMORY_STAT( TEXT("Stats Descriptions"), STAT_StatDescMemory, STATGROUP_StatSystem );
-
-DEFINE_STAT(STAT_FrameTime);
-DEFINE_STAT(STAT_NamedMarker);
-DEFINE_STAT(STAT_SecondsPerCycle);
-
-#if !(defined(DISABLE_THREAD_IDLE_STATS) && DISABLE_THREAD_IDLE_STATS)
-#if CPUPROFILERTRACE_ENABLED
-	UE_TRACE_CHANNEL_DEFINE(ThreadIdleScopeChannel);
-	TRACE_CPUPROFILER_EVENT_DECLARE(ThreadIdleScopeTraceEventId);
-#endif
-
-CORE_API FThreadIdleStats::FScopeIdle::FScopeIdle(bool bInIgnore /* = false */)
-	: Start(FPlatformTime::Cycles())
-	, bIgnore(bInIgnore || FThreadIdleStats::Get().bInIdleScope)
-#if CPUPROFILERTRACE_ENABLED
-	, TraceEventScope(ThreadIdleScopeTraceEventId, TEXT("FThreadIdleStats::FScopeIdle"), ThreadIdleScopeChannel, !bIgnore, __FILE__, __LINE__)
-#endif
-{
-	if (!bIgnore)
-	{
-		FThreadIdleStats& IdleStats = FThreadIdleStats::Get();
-		IdleStats.bInIdleScope = true;
-	}
-}
-#endif // #if !(defined(DISABLE_THREAD_IDLE_STATS) && DISABLE_THREAD_IDLE_STATS)
-
-/*-----------------------------------------------------------------------------
-	DebugLeakTest, for the stats based memory profiler
------------------------------------------------------------------------------*/
-
-#if	!UE_BUILD_SHIPPING
-
-static TAutoConsoleVariable<int32> CVarEnableLeakTest(
-	TEXT( "EnableLeakTest" ),
-	0,
-	TEXT( "If set to 1, enables leak test, for testing stats based memory profiler" )
-	);
-
-void DebugLeakTest()
-{
-	if (CVarEnableLeakTest.GetValueOnGameThread() == 1)
-	{
-		if (GFrameCounter == 60)
-		{
-			DirectStatsCommand( TEXT( "stat namedmarker Frame-060" ), true );
-		}
-
-		if (GFrameCounter == 120)
-		{
-			DirectStatsCommand( TEXT( "stat namedmarker Frame-120" ), true );
-		}
-
-
-		if (GFrameCounter == 240)
-		{
-			DirectStatsCommand( TEXT( "stat namedmarker Frame-240" ), true );
-		}
-
-		if (GFrameCounter == 300)
-		{
-			RequestEngineExit(TEXT("DebugLeakTest hit frame 300"));
-		}
-
-		// Realloc.
-		static TArray<uint8> Array;
-		static int32 Initial = 1;
-		{
-			DECLARE_SCOPE_CYCLE_COUNTER( TEXT( "LeakTest::Realloc" ), Stat_LeakTest_Realloc, STATGROUP_Quick );
-			Array.AddZeroed( Initial );
-			Initial += 100;
-		}
-
-		if (GFrameCounter == 300)
-		{
-			UE_LOG( LogTemp, Warning, TEXT( "Stat_ReallocTest: %i / %i" ), Array.GetAllocatedSize(), Initial );
-		}
-
-		// General memory leak.
-		{
-			DECLARE_SCOPE_CYCLE_COUNTER( TEXT( "LeakTest::NewInt8" ), Stat_LeakTest_NewInt8, STATGROUP_Quick );
-			int8* Leak = new int8[1000 * 1000];
-		} //-V773
-
-
-		if (GFrameCounter < 250)
-		{
-			// Background threads memory test.
-			struct FAllocTask
-			{
-				static void Alloc()
-				{
-					DECLARE_SCOPE_CYCLE_COUNTER( TEXT( "FAllocTask::Alloc" ), Stat_FAllocTask_Alloc, STATGROUP_Quick );
-
-					int8* IntAlloc = new int8[112233];
-					int8* LeakTask = new int8[100000];
-					delete[] IntAlloc;
-				} //-V773
-			};
-
-			for (int32 Index = 0; Index < 40; ++Index)
-			{
-				FSimpleDelegateGraphTask::CreateAndDispatchWhenReady( FSimpleDelegateGraphTask::FDelegate::CreateStatic( FAllocTask::Alloc ), TStatId() );
-			}
-
-			class FAllocPool : public FNonAbandonableTask
-			{
-			public:
-				void DoWork()
-				{
-					DECLARE_SCOPE_CYCLE_COUNTER( TEXT( "FAllocPool::DoWork" ), Stat_FAllocPool_DoWork, STATGROUP_Quick );
-
-					int8* IntAlloc = new int8[223311];
-					int8* LeakTask = new int8[100000];
-					delete[] IntAlloc;
-				} //-V773
-
-				TStatId GetStatId() const
-				{
-					return TStatId();
-				}
-			};
-
-			for (int32 Index = 0; Index < 40; ++Index)
-			{
-				(new FAutoDeleteAsyncTask<FAllocPool>())->StartBackgroundTask();
-			}
-		}
-
-		for (int32 Index = 0; Index < 40; ++Index)
-		{
-			DECLARE_SCOPE_CYCLE_COUNTER( TEXT( "DebugLeakTest::Alloc" ), Stat_LeakTest_Alloc, STATGROUP_Quick );
-
-			int8* IntAlloc = new int8[331122];
-			int8* LeakTask = new int8[100000];
-			delete[] IntAlloc;
-		} //-V773
-
-		if (IsEngineExitRequested())
-		{
-			// If we are writing stats data, stop it now.
-			DirectStatsCommand( TEXT( "stat stopfile" ), true );
-		}
-	}
-}
-
-#endif // !UE_BUILD_SHIPPING
-
-/*-----------------------------------------------------------------------------
-	FStats2
------------------------------------------------------------------------------*/
-
-TAtomic<int32> FStats::GameThreadStatsFrame(1);
-
-void FStats::AdvanceFrame( bool bDiscardCallstack, const FOnAdvanceRenderingThreadStats& AdvanceRenderingThreadStatsDelegate /*= FOnAdvanceRenderingThreadStats()*/ )
-{
-#if STATS
-	TRACE_CPUPROFILER_EVENT_SCOPE(FStats::AdvanceFrame);
-	LLM_SCOPE(ELLMTag::Stats);
-	check( IsInGameThread() );
-	static int32 PrimaryDisableChangeTagStartFrame = -1;
-	int64 Frame = ++GameThreadStatsFrame;
-
-	if( bDiscardCallstack )
-	{
-		FThreadStats::FrameDataIsIncomplete(); // we won't collect call stack stats this frame
-	}
-	if( PrimaryDisableChangeTagStartFrame == -1 )
-	{
-		PrimaryDisableChangeTagStartFrame = FThreadStats::PrimaryDisableChangeTag();
-	}
-	if( !FThreadStats::IsCollectingData() || PrimaryDisableChangeTagStartFrame != FThreadStats::PrimaryDisableChangeTag() )
-	{
-		Frame = -Frame; // mark this as a bad frame
-	}
-
-	// Update the seconds per cycle.
-	SET_FLOAT_STAT( STAT_SecondsPerCycle, FPlatformTime::GetSecondsPerCycle() );
-
-	FThreadStats::AddMessage( FStatConstants::AdvanceFrame.GetEncodedName(), EStatOperation::AdvanceFrameEventGameThread, Frame ); // we need to flush here if we aren't collecting stats to make sure the meta data is up to date
-
-	if( AdvanceRenderingThreadStatsDelegate.IsBound() )
-	{
-		AdvanceRenderingThreadStatsDelegate.Execute( bDiscardCallstack, Frame, PrimaryDisableChangeTagStartFrame );
-	}
-	else
-	{
-		// There is no rendering thread, so this message is sufficient to make stats happy and don't leak memory.
-		FThreadStats::AddMessage( FStatConstants::AdvanceFrame.GetEncodedName(), EStatOperation::AdvanceFrameEventRenderThread, Frame );
-	}
-
-	FThreadStats::ExplicitFlush( bDiscardCallstack );
-	FThreadStats::WaitForStats();
-	PrimaryDisableChangeTagStartFrame = FThreadStats::PrimaryDisableChangeTag();
-#endif
-}
-
-void FStats::TickCommandletStats()
-{
-	if (EnabledForCommandlet())
-	{
-		//FThreadStats* ThreadStats = FThreadStats::GetThreadStats();
-		//check( ThreadStats->ScopeCount == 0 && TEXT( "FStats::TickCommandletStats must be called outside any scope counters" ) );
-
-		FTaskGraphInterface::Get().ProcessThreadUntilIdle( ENamedThreads::GameThread );
-		FTSTicker::GetCoreTicker().Tick( 1 / 60.0f );
-
-		FStats::AdvanceFrame( false );
-	}
-}
-
-bool FStats::EnabledForCommandlet()
-{
-	static bool bHasStatsForCommandletsToken = HasLoadTimeStatsForCommandletToken() || HasLoadTimeFileForCommandletToken();
-	return bHasStatsForCommandletsToken;
-}
-
-bool FStats::HasLoadTimeStatsForCommandletToken()
-{
-	static bool bHasLoadTimeStatsForCommandletToken = FParse::Param( FCommandLine::Get(), TEXT( "LoadTimeStatsForCommandlet" ) );
-	return bHasLoadTimeStatsForCommandletToken;
-}
-
-bool FStats::HasLoadTimeFileForCommandletToken()
-{
-	static bool bHasLoadTimeFileForCommandletToken = FParse::Param( FCommandLine::Get(), TEXT( "LoadTimeFileForCommandlet" ) );
-	return bHasLoadTimeFileForCommandletToken;
-}
-
-/* Todo
-
-ST_int64 combined with IsPackedCCAndDuration is fairly bogus, ST_uint32_pair should probably be a first class type. IsPackedCCAndDuration can stay, it says what kind of data is in the pair. remove FromPackedCallCountDuration_Duration et al
-
-averages stay as int64, it might be better if they were floats, but then we would probably want a ST_float_pair too.
-
-//@todo merge these two after we have redone the user end
-
-set a DEBUG_STATS define that allows this all to be debugged. Otherwise, inline and get rid of checks to the MAX. Also maybe just turn off stats in debug or debug game.
-
-It should be possible to load stats data without STATS
-
-//@todo this probably goes away after we redo the programmer interface
-
-We are only saving condensed frames. For the "core view", we want a single frame capture of FStatsThreadState::History. That is the whole enchillada including start and end time for every scope.
-
-//@todo Legacy API, phase out after we have changed the programmer-facing APU
-
-stats2.h, stats2.cpp, statsrender2.cpp - get rid of the 2s.
-
-//@todo, delegate under a global bool?
-
-//@todo The rest is all horrid hacks to bridge the game between the old and the new system while we finish the new system
-
-//@todo split header
-
-//@todo metadata probably needs a different queue, otherwise it would be possible to load a module, run code in a thread and have the events arrive before the meta data
-
-delete "wait for render commands", and generally be on the look out for stats that are never used, also stat system stuff
-
-sweep INC type things for dump code that calls INC in a loop instead of just calling INC_BY once
-
-FORCEINLINE_STATS void Start(FName InStatId)
-{
-check(InStatId != NAME_None);
-
-^^^should be a checkstats
-
-
-*/
-
-#if STATS
-
-#include "Stats/StatsMallocProfilerProxy.h"
-
-TStatIdData TStatId::TStatId_NAME_None;
-
-/*-----------------------------------------------------------------------------
-	FStartupMessages
------------------------------------------------------------------------------*/
-
-void FStartupMessages::AddThreadMetadata( const FName InThreadName, uint32 InThreadID )
-{
-	// Make unique name.
-	const FString ThreadName = FStatsUtils::BuildUniqueThreadName( InThreadID );
-
-	FStartupMessages::AddMetadata( InThreadName, *ThreadName, STAT_GROUP_TO_FStatGroup( STATGROUP_Threads )::GetGroupName(), STAT_GROUP_TO_FStatGroup( STATGROUP_Threads )::GetGroupCategory(), STAT_GROUP_TO_FStatGroup( STATGROUP_Threads )::GetDescription(), true, EStatDataType::ST_int64, true, false );
-}
-
-
-void FStartupMessages::AddMetadata( FName InStatName, const TCHAR* InStatDesc, const char* InGroupName, const char* InGroupCategory, const TCHAR* InGroupDesc, bool bShouldClearEveryFrame, EStatDataType::Type InStatType, bool bCycleStat, bool bSortByName, FPlatformMemory::EMemoryCounterRegion InMemoryRegion /*= FPlatformMemory::MCR_Invalid*/ )
-{
-	LLM_SCOPE(ELLMTag::Stats);
-	FScopeLock Lock( &CriticalSection );
-
-	DelayedMessages.Emplace( InGroupName, EStatDataType::ST_None, "Groups", InGroupCategory, InGroupDesc, false, false, bSortByName );
-	DelayedMessages.Emplace( InStatName, InStatType, InGroupName, InGroupCategory, InStatDesc, bShouldClearEveryFrame, bCycleStat, bSortByName, InMemoryRegion );
-}
-
-
-FStartupMessages& FStartupMessages::Get()
-{
-	static FStartupMessages* Messages = NULL;
-	if( !Messages )
-	{
-		LLM_SCOPE(ELLMTag::Stats);
-		check( IsInGameThread() );
-		Messages = new FStartupMessages;
-	}
-	return *Messages;
-}
-
-/*-----------------------------------------------------------------------------
-	FThreadSafeStaticStatBase
------------------------------------------------------------------------------*/
-
-const TStatIdData* FThreadSafeStaticStatBase::DoSetup(const char* InStatName, const TCHAR* InStatDesc, const char* InGroupName, const char* InGroupCategory, const TCHAR* InGroupDesc, bool bDefaultEnable, bool bShouldClearEveryFrame, EStatDataType::Type InStatType, bool bCycleStat, bool bSortByName, FPlatformMemory::EMemoryCounterRegion InMemoryRegion) const
-{
-	FName TempName(InStatName);
-
-	// send meta data, we don't use normal messages because the stats thread might not be running yet
-	FStartupMessages::Get().AddMetadata(TempName, InStatDesc, InGroupName, InGroupCategory, InGroupDesc, bShouldClearEveryFrame, InStatType, bCycleStat, bSortByName, InMemoryRegion);
-
-	TStatIdData const* LocalHighPerformanceEnable(IStatGroupEnableManager::Get().GetHighPerformanceEnableForStat(FName(InStatName), InGroupName, InGroupCategory, bDefaultEnable, bShouldClearEveryFrame, InStatType, InStatDesc, bCycleStat, bSortByName, InMemoryRegion).GetRawPointer());
-	TStatIdData const* OldHighPerformanceEnable = HighPerformanceEnable.Exchange(LocalHighPerformanceEnable);
-	check(!OldHighPerformanceEnable || LocalHighPerformanceEnable == OldHighPerformanceEnable); // we are assigned two different groups?
-
-	return LocalHighPerformanceEnable;
-}
-
-/*-----------------------------------------------------------------------------
-	FStatGroupEnableManager
------------------------------------------------------------------------------*/
-
-DEFINE_LOG_CATEGORY_STATIC(LogStatGroupEnableManager, Log, All);
-
-class FStatGroupEnableManager : public IStatGroupEnableManager
-{
-	struct FGroupEnable
-	{
-		TMap<FName, TStatIdData *> NamesInThisGroup;
-		TMap<FName, TStatIdData *> AlwaysEnabledNamesInThisGroup;
-		bool DefaultEnable;
-		bool CurrentEnable;
-
-		FGroupEnable(bool InDefaultEnable)
-			: DefaultEnable(InDefaultEnable)
-			, CurrentEnable(InDefaultEnable)
-		{
-		}
-	};
-
-	enum
-	{
-		/** Number of stats pointer allocated per block. */
-		NUM_PER_BLOCK = 16384,
-	};
-
-
-	TMap<FName, FGroupEnable> HighPerformanceEnable;
-
-	/** Used to synchronize the access to the high performance stats groups. */
-	FCriticalSection SynchronizationObject;
-
-	/** Pointer to the long name in the names block. */
-	TArray<TArray<TStatIdData>> StatIDs;
-
-	/** Holds the amount of memory allocated for the stats descriptions. */
-	FThreadSafeCounter MemoryCounter;
-
-	// these control what happens to groups that haven't been registered yet
-	TMap<FName, bool> EnableForNewGroup;
-	bool EnableForNewGroups;
-	bool UseEnableForNewGroups;
-
-	void EnableStat(FName StatName, TStatIdData* DisablePtr)
-	{
-		DisablePtr->Name = NameToMinimalName(StatName);
-	}
-
-	void DisableStat(TStatIdData* DisablePtr)
-	{
-		DisablePtr->Name = TStatId::GetStatNone().Name.Load();
-	}
-
-public:
-	FStatGroupEnableManager()
-		: EnableForNewGroups(false)
-		, UseEnableForNewGroups(false)
-	{
-		check(IsInGameThread());
-	}
-
-	virtual void UpdateMemoryUsage() override
-	{
-		// Update the stats descriptions memory usage.
-		const int32 MemoryUsage = MemoryCounter.GetValue();
-		SET_MEMORY_STAT( STAT_StatDescMemory, MemoryUsage );
-	}
-
-	virtual void SetHighPerformanceEnableForGroup(FName Group, bool Enable) override
-	{
-		FScopeLock ScopeLock(&SynchronizationObject);
-		FThreadStats::PrimaryDisableChangeTagLockAdd();
-		FGroupEnable* Found = HighPerformanceEnable.Find(Group);
-		if (Found)
-		{
-			Found->CurrentEnable = Enable;
-			if (Enable)
-			{
-				for (auto It = Found->NamesInThisGroup.CreateIterator(); It; ++It)
-				{
-					EnableStat(It.Key(), It.Value());
-				}
-			}
-			else
-			{
-				for (auto It = Found->NamesInThisGroup.CreateIterator(); It; ++It)
-				{
-					DisableStat(It.Value());
-				}
-			}
-		}
-		FThreadStats::PrimaryDisableChangeTagLockSubtract();
-	}
-
-	virtual void SetHighPerformanceEnableForAllGroups(bool Enable) override
-	{
-		FScopeLock ScopeLock(&SynchronizationObject);
-		FThreadStats::PrimaryDisableChangeTagLockAdd();
-		for (auto It = HighPerformanceEnable.CreateIterator(); It; ++It)
-		{
-			It.Value().CurrentEnable = Enable;
-			if (Enable)
-			{
-				for (auto ItInner = It.Value().NamesInThisGroup.CreateConstIterator(); ItInner; ++ItInner)
-				{
-					EnableStat(ItInner.Key(), ItInner.Value());
-				}
-			}
-			else
-			{
-				for (auto ItInner = It.Value().NamesInThisGroup.CreateConstIterator(); ItInner; ++ItInner)
-				{
-					DisableStat(ItInner.Value());
-				}
-			}
-		}
-		FThreadStats::PrimaryDisableChangeTagLockSubtract();
-	}
-	virtual void ResetHighPerformanceEnableForAllGroups() override
-	{
-		FScopeLock ScopeLock(&SynchronizationObject);
-		FThreadStats::PrimaryDisableChangeTagLockAdd();
-		for (auto It = HighPerformanceEnable.CreateIterator(); It; ++It)
-		{
-			It.Value().CurrentEnable = It.Value().DefaultEnable;
-			if (It.Value().DefaultEnable)
-			{
-				for (auto ItInner = It.Value().NamesInThisGroup.CreateConstIterator(); ItInner; ++ItInner)
-				{
-					EnableStat(ItInner.Key(), ItInner.Value());
-				}
-			}
-			else
-			{
-				for (auto ItInner = It.Value().NamesInThisGroup.CreateConstIterator(); ItInner; ++ItInner)
-				{
-					DisableStat(ItInner.Value());
-				}
-			}
-		}
-		FThreadStats::PrimaryDisableChangeTagLockSubtract();
-	}
-
-	virtual TStatId GetHighPerformanceEnableForStat(FName StatShortName, const char* InGroup, const char* InCategory, bool bDefaultEnable, bool bShouldClearEveryFrame, EStatDataType::Type InStatType, TCHAR const* InDescription, bool bCycleStat, bool bSortByName, FPlatformMemory::EMemoryCounterRegion MemoryRegion = FPlatformMemory::MCR_Invalid) override
-	{
-		LLM_SCOPE(ELLMTag::Stats);
-
-		FName Stat;
-		TStatIdData* Result;
-		FString StatDescription;
-		{
-			FScopeLock ScopeLock(&SynchronizationObject);
-
-			FStatNameAndInfo LongName(StatShortName, InGroup, InCategory, InDescription, InStatType, bShouldClearEveryFrame, bCycleStat, bSortByName, MemoryRegion);
-
-			Stat = LongName.GetEncodedName();
-
-			FName Group(InGroup);
-			FGroupEnable* Found = HighPerformanceEnable.Find(Group);
-			if (Found)
-			{
-				if (Found->DefaultEnable != bDefaultEnable)
-				{
-					UE_LOG(LogStatGroupEnableManager, Fatal, TEXT("Stat group %s was was defined both on and off by default."), *Group.ToString());
-				}
-				TStatIdData** StatFound = Found->NamesInThisGroup.Find( Stat );
-				TStatIdData** StatFoundAlways = Found->AlwaysEnabledNamesInThisGroup.Find( Stat );
-				if( StatFound )
-				{
-					if( StatFoundAlways )
-					{
-						UE_LOG( LogStatGroupEnableManager, Fatal, TEXT( "Stat %s is both always enabled and not always enabled, so it was used for two different things." ), *Stat.ToString() );
-					}
-					return TStatId( *StatFound );
-				}
-				else if( StatFoundAlways )
-				{
-					return TStatId( *StatFoundAlways );
-				}
-			}
-			else
-			{
-				Found = &HighPerformanceEnable.Add( Group, FGroupEnable( bDefaultEnable || !bShouldClearEveryFrame ) );
-
-				// this was set up before we saw the group, so set the enable now
-				if (EnableForNewGroup.Contains(Group))
-				{
-					Found->CurrentEnable = EnableForNewGroup.FindChecked(Group);
-					EnableForNewGroup.Remove(Group); // by definition, we will never need this again
-				}
-				else if (UseEnableForNewGroups)
-				{
-					Found->CurrentEnable = EnableForNewGroups;
-				}
-			}
-			if (StatIDs.Num() == 0 || StatIDs.Last().Num() == NUM_PER_BLOCK)
-			{
-				TArray<TStatIdData>& NewBlock = StatIDs.AddDefaulted_GetRef();
-				NewBlock.Reserve(NUM_PER_BLOCK);
-			}
-			Result = &StatIDs.Last().AddDefaulted_GetRef();
-
-			StatDescription = InDescription ? InDescription : StatShortName.GetPlainNameString();
-
-			const int32 StatDescLen = StatDescription.Len() + 1;
-
-			int32 MemoryAllocated = 0;
-
-			// Get the wide stat description.
-			{
-				auto StatDescriptionWide = StringCast<WIDECHAR>(*StatDescription, StatDescLen);
-				int32 StatDescriptionWideLength = StatDescriptionWide.Length();
-				TUniquePtr<WIDECHAR[]> StatDescWide = MakeUnique<WIDECHAR[]>(StatDescriptionWideLength + 1);
-				TCString<WIDECHAR>::Strcpy(StatDescWide.Get(), StatDescriptionWideLength + 1, StatDescriptionWide.Get());
-				Result->StatDescriptionWide = MoveTemp(StatDescWide);
-
-				MemoryAllocated += StatDescriptionWideLength * sizeof(WIDECHAR);
-			}
-
-			// Get the ansi stat description.
-			{
-				auto StatDescriptionAnsi = StringCast<ANSICHAR>(*StatDescription, StatDescLen);
-				int32 StatDescriptionAnsiLength = StatDescriptionAnsi.Length();
-				TUniquePtr<ANSICHAR[]> StatDescAnsi = MakeUnique<ANSICHAR[]>(StatDescriptionAnsiLength + 1);
-				TCString<ANSICHAR>::Strcpy(StatDescAnsi.Get(), StatDescriptionAnsiLength + 1, StatDescriptionAnsi.Get());
-				Result->StatDescriptionAnsi = MoveTemp(StatDescAnsi);
-
-				MemoryAllocated += StatDescriptionAnsiLength * sizeof(ANSICHAR);
-			}
-
-			MemoryCounter.Add( MemoryAllocated );
-
-			if( Found->CurrentEnable )
-			{
-				EnableStat( Stat, Result );
-			}
-
-			if( bShouldClearEveryFrame )
-			{
-				Found->NamesInThisGroup.Add( Stat, Result );
-			}
-			else
-			{
-				Found->AlwaysEnabledNamesInThisGroup.Add( Stat, Result );
-			}
-		}
-
-#if STATSTRACE_ENABLED
-		if (!bCycleStat && (InStatType == EStatDataType::ST_int64 || InStatType == EStatDataType::ST_double))
-		{
-			if (!StatShortName.GetDisplayNameEntry()->IsWide())
-			{
-				ANSICHAR NameBuffer[1024];
-				StatShortName.GetPlainANSIString(NameBuffer);
-				FStatsTrace::DeclareStat(Stat, NameBuffer, *StatDescription, InGroup, InStatType == EStatDataType::ST_double, MemoryRegion != FPlatformMemory::MCR_Invalid, bShouldClearEveryFrame);
-			}
-			else
-			{
-				FString StatShortNameString = StatShortName.GetPlainNameString();
-				FStatsTrace::DeclareStat(Stat, TCHAR_TO_ANSI(*StatShortNameString), *StatDescription, InGroup, InStatType == EStatDataType::ST_double, MemoryRegion != FPlatformMemory::MCR_Invalid, bShouldClearEveryFrame);
-			}
-		}
-#endif
-		return TStatId(Result);
-	}
-
-	void ListGroup(FName Group)
-	{
-		FGroupEnable* Found = HighPerformanceEnable.Find(Group);
-		if (Found)
-		{
-			UE_LOG(LogStatGroupEnableManager, Display, TEXT("  %d  default %d %s"), !!Found->CurrentEnable, !!Found->DefaultEnable, *Group.ToString());
-		}
-	}
-
-	void ListGroups(bool bDetailed = false)
-	{
-		for (auto It = HighPerformanceEnable.CreateConstIterator(); It; ++It)
-		{
-			UE_LOG(LogStatGroupEnableManager, Display, TEXT("  %d  default %d %s"), !!It.Value().CurrentEnable, !!It.Value().DefaultEnable, *(It.Key().ToString()));
-			if (bDetailed)
-			{
-				for (auto ItInner = It.Value().NamesInThisGroup.CreateConstIterator(); ItInner; ++ItInner)
-				{
-					UE_LOG(LogStatGroupEnableManager, Display, TEXT("      %d %s"), !ItInner.Value()->IsNone(), *ItInner.Key().ToString());
-				}
-				for( auto ItInner = It.Value().AlwaysEnabledNamesInThisGroup.CreateConstIterator(); ItInner; ++ItInner )
-				{
-					UE_LOG( LogStatGroupEnableManager, Display, TEXT( "      (always enabled) %s" ), *ItInner.Key().ToString() );
-				}
-			}
-		}
-	}
-
-	FName CheckGroup(TCHAR const *& Cmd, bool Enable)
-	{
-		FString MaybeGroup;
-		FParse::Token(Cmd, MaybeGroup, false);
-		MaybeGroup = FString(TEXT("STATGROUP_")) + MaybeGroup;
-		FName MaybeGroupFName(*MaybeGroup);
-
-		FGroupEnable* Found = HighPerformanceEnable.Find(MaybeGroupFName);
-		if (!Found)
-		{
-			EnableForNewGroup.Add(MaybeGroupFName, Enable);
-			ListGroups();
-			UE_LOG(LogStatGroupEnableManager, Display, TEXT("Group Not Found %s"), *MaybeGroupFName.ToString());
-			return NAME_None;
-		}
-		SetHighPerformanceEnableForGroup(MaybeGroupFName, Enable);
-		ListGroup(MaybeGroupFName);
-		return MaybeGroupFName;
-	}
-
-	void StatGroupEnableManagerCommand(FString const& InCmd) override
-	{
-		FScopeLock ScopeLock(&SynchronizationObject);
-		const TCHAR* Cmd = *InCmd;
-		if( FParse::Command(&Cmd,TEXT("list")) )
-		{
-			ListGroups();
-		}
-		else if( FParse::Command(&Cmd,TEXT("listall")) )
-		{
-			ListGroups(true);
-		}
-		else if ( FParse::Command(&Cmd,TEXT("enable")) )
-		{
-			CheckGroup(Cmd, true);
-		}
-		else if ( FParse::Command(&Cmd,TEXT("disable")) )
-		{
-			CheckGroup(Cmd, false);
-		}
-		else if ( FParse::Command(&Cmd,TEXT("none")) )
-		{
-			EnableForNewGroups = false;
-			UseEnableForNewGroups = true;
-			SetHighPerformanceEnableForAllGroups(false);
-			ListGroups();
-		}
-		else if ( FParse::Command(&Cmd,TEXT("all")) )
-		{
-			EnableForNewGroups = true;
-			UseEnableForNewGroups = true;
-			SetHighPerformanceEnableForAllGroups(true);
-			ListGroups();
-		}
-		else if ( FParse::Command(&Cmd,TEXT("default")) )
-		{
-			UseEnableForNewGroups = false;
-			EnableForNewGroup.Empty();
-			ResetHighPerformanceEnableForAllGroups();
-			ListGroups();
-		}
-	}
-};
-
-IStatGroupEnableManager& IStatGroupEnableManager::Get()
-{
-	static FStatGroupEnableManager Singleton;
-	return Singleton;
-}
-
-/*-----------------------------------------------------------------------------
-	FStatNameAndInfo
------------------------------------------------------------------------------*/
-
-FName FStatNameAndInfo::ToLongName(FName InStatName, char const* InGroup, char const* InCategory, TCHAR const* InDescription, bool InSortByName)
-{
-	TStringBuilder<256> LongName;
-	if (InGroup)
-	{
-		LongName += TEXT("//");
-		LongName += InGroup;
-		LongName += TEXT("//");
-	}
-	InStatName.AppendString(LongName);
-	if (InDescription)
-	{
-		LongName += TEXT("///");
-		FStatsUtils::ToEscapedString(InDescription, LongName);
-		LongName += TEXT("///");
-	}
-	if (InCategory)
-	{
-		LongName += TEXT("####");
-		LongName += InCategory;
-		LongName += TEXT("####");
-	}
-	if (InSortByName)
-	{
-		LongName += TEXT("/#/#");
-		LongName += TEXT("SortByName");
-		LongName += TEXT("/#/#");
-	}
-	return FName(*LongName);
-}
-
-FName FStatNameAndInfo::GetShortNameFrom(FName InLongName)
-{
-	TStringBuilder<256> Temp;
-	Temp << InLongName;
-	FStringView Input(Temp);
-
-	if (Input.StartsWith(TEXT("//"), ESearchCase::CaseSensitive))
-	{
-		Input.RemovePrefix(2);
-		const int32 IndexEnd = UE::String::FindFirst(Input, TEXT("//"));
-		if (IndexEnd == INDEX_NONE)
-		{
-			checkStats(0);
-			return InLongName;
-		}
-		Input.RightChopInline(IndexEnd + 2);
-	}
-	const int32 DescIndexEnd = UE::String::FindFirst(Input, TEXT("///"), ESearchCase::CaseSensitive);
-	if (DescIndexEnd != INDEX_NONE)
-	{
-		Input.LeftInline(DescIndexEnd);
-	}
-	const int32 CategoryIndexEnd = UE::String::FindFirst(Input, TEXT( "####" ), ESearchCase::CaseSensitive );
-	if( DescIndexEnd == INDEX_NONE && CategoryIndexEnd != INDEX_NONE )
-	{
-		Input.LeftInline(CategoryIndexEnd);
-	}
-	const int32 SortByNameIndexEnd = UE::String::FindFirst(Input, TEXT( "/#/#" ), ESearchCase::CaseSensitive );
-	if( DescIndexEnd == INDEX_NONE && CategoryIndexEnd == INDEX_NONE && SortByNameIndexEnd != INDEX_NONE )
-	{
-		Input.LeftInline(SortByNameIndexEnd);
-	}
-	return FName(Input);
-}
-
-FName FStatNameAndInfo::GetGroupNameFrom(FName InLongName)
-{
-	TStringBuilder<256> Temp;
-	Temp << InLongName;
-	FStringView Input(Temp);
-
-	if (Input.StartsWith(TEXT("//"), ESearchCase::CaseSensitive))
-	{
-		Input.RemovePrefix(2);
-		if (Input.StartsWith(TEXT("Groups//")))
-		{
-			Input.RemovePrefix(8);
-		}
-		const int32 IndexEnd = UE::String::FindFirst(Input, TEXT("//"));
-		if (IndexEnd != INDEX_NONE)
-		{
-			return FName(Input.Left(IndexEnd));
-		}
-		checkStats(0);
-	}
-	return FName();
-}
-
-FString FStatNameAndInfo::GetDescriptionFrom(FName InLongName)
-{
-	FString Input(InLongName.ToString());
-
-	const int32 IndexStart = Input.Find(TEXT("///"), ESearchCase::CaseSensitive);
-	if (IndexStart != INDEX_NONE)
-	{
-		Input.RightChopInline(IndexStart + 3, EAllowShrinking::No);
-		const int32 IndexEnd = Input.Find(TEXT("///"), ESearchCase::CaseSensitive);
-		if (IndexEnd != INDEX_NONE)
-		{
-			return FStatsUtils::FromEscapedString(*Input.Left(IndexEnd));
-		}
-	}
-	return FString();
-}
-
-FName FStatNameAndInfo::GetGroupCategoryFrom(FName InLongName)
-{
-	FString Input(InLongName.ToString());
-
-	const int32 IndexStart = Input.Find(TEXT("####"), ESearchCase::CaseSensitive);
-	if (IndexStart != INDEX_NONE)
-	{
-		Input.RightChopInline(IndexStart + 4, EAllowShrinking::No);
-		const int32 IndexEnd = Input.Find(TEXT("####"), ESearchCase::CaseSensitive);
-		if (IndexEnd != INDEX_NONE)
-		{
-			return FName(*Input.Left(IndexEnd));
-		}
-		checkStats(0);
-	}
-	return NAME_None;
-}
-
-bool FStatNameAndInfo::GetSortByNameFrom(FName InLongName)
-{
-	FString Input(InLongName.ToString());
-
-	const int32 IndexStart = Input.Find(TEXT("/#/#"), ESearchCase::CaseSensitive);
-	if (IndexStart != INDEX_NONE)
-	{
-		Input.RightChopInline(IndexStart + 4, EAllowShrinking::No);
-		const int32 IndexEnd = Input.Find(TEXT("/#/#"), ESearchCase::CaseSensitive);
-		if (IndexEnd != INDEX_NONE)
-		{
-			return Input.Left(IndexEnd) == TEXT("SortByName");
-		}
-		checkStats(0);
-	}
-	return false;
-}
-
-/*-----------------------------------------------------------------------------
-	FStatsThread
------------------------------------------------------------------------------*/
-
-static TAutoConsoleVariable<int32> CVarDumpStatPackets(	TEXT("DumpStatPackets"),0,	TEXT("If true, dump stat packets."));
-
-/** The rendering thread runnable object. */
-class FStatsThread
-{
-	/** Array of stat packets, queued data to be processed on this thread. */
-	FStatPacketArray IncomingData;
-
-	/** Stats state. */
-	FStatsThreadState& State;
-
-	/** Whether we are ready to process the packets, sets by game or render packets. */
-	bool bReadyToProcess;
-
-public:
-	static FStatsThread& Get()
-	{
-		static FStatsThread Singleton;
-		return Singleton;
-	}
-
-	/** Default constructor. */
-	FStatsThread()
-		: State(FStatsThreadState::GetLocalState())
-		, bReadyToProcess(false)
-	{
-	}
-
-	/** Received a stat packet from other thread and add to the processing queue. */
-	void StatMessage(FStatPacket* Packet)
-	{
-		LLM_SCOPE(ELLMTag::Stats);
-
-		if (CVarDumpStatPackets.GetValueOnAnyThread())
-		{
-			UE_LOG(LogStats, Log, TEXT("Packet from %x with %d messages"), Packet->ThreadId, Packet->StatMessages.Num());
-		}
-
-		bReadyToProcess = Packet->ThreadType != EThreadType::Other;
-		IncomingData.Packets.Add(Packet);
-		State.NumStatMessages.Add(Packet->StatMessages.Num());
-
-		Process();
-	}
-
-private:
-	void Process()
-	{
-		static double LastTime = -1.0;
-		bool bShouldProcess = false;
-
-		const int32 MaxIncomingPackets = 16;
-		if (FThreadStats::bIsRawStatsActive)
-		{
-			// For raw stats we process every 24MB of packet data to minimize the stats messages memory usage.
-			//const bool bShouldProcessRawStats = IncomingData.Packets.Num() > 10;
-			const int32 MaxIncomingMessages = 24 * 1024 * 1024 / sizeof(FStatMessage);
-
-			int32 IncomingDataMessages = 0;
-			for (FStatPacket* Packet : IncomingData.Packets)
-			{
-				IncomingDataMessages += Packet->StatMessages.Num();
-			}
-
-			bShouldProcess = IncomingDataMessages > MaxIncomingMessages || IncomingData.Packets.Num() > MaxIncomingPackets;
-		}
-		else
-		{
-			// For regular stats we won't process more than every 5ms or every 16 packets.
-			// Commandlet stats are flushed as soon as.
-			bShouldProcess = bReadyToProcess && (FPlatformTime::Seconds() - LastTime > 0.005f || IncomingData.Packets.Num() > MaxIncomingPackets || FStats::EnabledForCommandlet());
-		}
-
-		if (bShouldProcess)
-		{
-			SCOPE_CYCLE_COUNTER(STAT_StatsNewTick);
-
-			IStatGroupEnableManager::Get().UpdateMemoryUsage();
-			State.UpdateStatMessagesMemoryUsage();
-
-			bReadyToProcess = false;
-			FStatPacketArray NowData;
-			Exchange(NowData.Packets, IncomingData.Packets);
-			INC_DWORD_STAT_BY(STAT_StatFramePacketsRecv, NowData.Packets.Num());
-			{
-				SCOPE_CYCLE_COUNTER(STAT_StatsNewParseMeta);
-				TArray64<FStatMessage> MetaMessages;
-				{
-					FScopeLock Lock(&FStartupMessages::Get().CriticalSection);
-					Exchange(FStartupMessages::Get().DelayedMessages, MetaMessages);
-				}
-				if (MetaMessages.Num())
-				{
-					State.ProcessMetaDataOnly(MetaMessages);
-				}
-			}
-			{
-				SCOPE_CYCLE_COUNTER(STAT_ScanForAdvance);
-				State.ScanForAdvance(NowData);
-			}
-
-			if (FThreadStats::bIsRawStatsActive)
-			{
-				// Process raw stats.
-				State.ProcessRawStats(NowData);
-				State.ResetRegularStats();
-			}
-			else
-			{
-				// Process regular stats.
-				SCOPE_CYCLE_COUNTER(STAT_StatsNewAddToHistory);
-				State.ResetRawStats();
-				State.AddToHistoryAndEmpty(NowData);
-			}
-			check(!NowData.Packets.Num());
-			LastTime = FPlatformTime::Seconds();
-		}
-	}
-};
-
-/*-----------------------------------------------------------------------------
-	FThreadStatsPool
------------------------------------------------------------------------------*/
-
-FThreadStatsPool::FThreadStatsPool()
-{
-	LLM_SCOPE(ELLMTag::Stats);
-	for( int32 Index = 0; Index < NUM_ELEMENTS_IN_POOL; ++Index )
-	{
-		Pool.Push( new FThreadStats(EConstructor::FOR_POOL) );
-	}
-}
-
-FThreadStatsPool& FThreadStatsPool::Get()
-{
-	static FThreadStatsPool Singleton;
-	return Singleton;
-}
-
-FThreadStats* FThreadStatsPool::GetFromPool()
-{
-	LLM_SCOPE(ELLMTag::Stats);
-	FPlatformMisc::MemoryBarrier();
-	FThreadStats* Address = Pool.Pop();
-	while (!Address)
-	{
-		Pool.Push(new FThreadStats(EConstructor::FOR_POOL));
-		Address = Pool.Pop();
-	}
-	return new(Address) FThreadStats();
-}
-
-void FThreadStatsPool::ReturnToPool( FThreadStats* Instance )
-{
-	check(Instance);
-	Instance->~FThreadStats();
-	Pool.Push(Instance);
-}
-
-/*-----------------------------------------------------------------------------
-	FThreadStats
------------------------------------------------------------------------------*/
-
-uint32 FThreadStats::TlsSlot = FPlatformTLS::InvalidTlsSlot;
-FThreadSafeCounter FThreadStats::PrimaryEnableCounter;
-FThreadSafeCounter FThreadStats::PrimaryEnableUpdateNumber;
-FThreadSafeCounter FThreadStats::PrimaryDisableChangeTagLock;
-bool FThreadStats::bPrimaryEnable = false;
-bool FThreadStats::bPrimaryDisableForever = false;
-bool FThreadStats::bIsRawStatsActive = false;
-
-FThreadStats::FThreadStats():
-	CurrentGameFrame(FStats::GameThreadStatsFrame.Load(EMemoryOrder::Relaxed)),
-	ScopeCount(0),
-	bWaitForExplicitFlush(0),
-	MemoryMessageScope(0),
-	bReentranceGuard(false),
-	bSawExplicitFlush(false)
-{
-	Packet.SetThreadProperties();
-
-	check(FPlatformTLS::IsValidTlsSlot(TlsSlot));
-	FPlatformTLS::SetTlsValue(TlsSlot, this);
-}
-
-FThreadStats::FThreadStats( EConstructor ):
-	CurrentGameFrame(-1),
-	ScopeCount(0),
-	bWaitForExplicitFlush(0),
-	MemoryMessageScope(0),
-	bReentranceGuard(false),
-	bSawExplicitFlush(false)
-{}
-
-void FThreadStats::CheckEnable()
-{
-	bool bOldPrimaryEnable(bPrimaryEnable);
-	bool bNewPrimaryEnable( WillEverCollectData() && (!IsRunningCommandlet() || FStats::EnabledForCommandlet()) && IsThreadingReady() && (PrimaryEnableCounter.GetValue()) );
-	if (bPrimaryEnable != bNewPrimaryEnable)
-	{
-		PrimaryDisableChangeTagLockAdd();
-		bPrimaryEnable = bNewPrimaryEnable;
-		PrimaryDisableChangeTagLockSubtract();
-	}
-}
-
-void FThreadStats::Flush( bool bHasBrokenCallstacks /*= false*/, bool bForceFlush /*= false*/ )
-{
-	if (bPrimaryDisableForever)
-	{
-		Packet.StatMessages.Empty();
-		return;
-	}
-
-	if( bIsRawStatsActive )
-	{
-		FlushRawStats(bHasBrokenCallstacks, bForceFlush);
-	}
-	else
-	{
-		FlushRegularStats(bHasBrokenCallstacks, bForceFlush);
-	}
-}
-
-void FThreadStats::SendMessage_Async(FStatPacket* ToSend)
-{
-	if (FPlatformProcess::SupportsMultithreading())
-	{
-		GStatsPipe.Launch(UE_SOURCE_LOCATION, [ToSend] { FStatsThread::Get().StatMessage(ToSend); });
-	}
-	else
-	{
-		FFunctionGraphTask::CreateAndDispatchWhenReady([ToSend] { FStatsThread::Get().StatMessage(ToSend); }, TStatId{}, nullptr, ENamedThreads::GameThread);
-	}
-}
-
-void FThreadStats::FlushRegularStats( bool bHasBrokenCallstacks, bool bForceFlush )
-{
-	LLM_SCOPE(ELLMTag::Stats);
-
-	if (bReentranceGuard)
-	{
-		return;
-	}
-	TGuardValue<bool> Guard( bReentranceGuard, true );
-
-	enum
-	{
-		PRESIZE_MAX_NUM_ENTRIES = 10,
-		PRESIZE_MAX_SIZE = 256*1024,
-	};
-
-
-	// Sends all collected messages when:
-	// The current game frame has changed.
-	// This a force flush when we shutting down the thread stats.
-	// This is an explicit flush from the game thread or the render thread.
-	const bool bFrameHasChanged = DetectAndUpdateCurrentGameFrame();
-	const bool bSendStatPacket = bFrameHasChanged || bForceFlush || bSawExplicitFlush;
-	if( !bSendStatPacket )
-	{
-		return;
-	}
-
-	if ((!ScopeCount || bForceFlush) && Packet.StatMessages.Num())
-	{
-		if( Packet.StatMessagesPresize.Num() >= PRESIZE_MAX_NUM_ENTRIES )
-		{
-			Packet.StatMessagesPresize.RemoveAt(0);
-		}
-		if (Packet.StatMessages.Num() < PRESIZE_MAX_SIZE)
-		{
-			Packet.StatMessagesPresize.Add(Packet.StatMessages.Num());
-		}
-		else
-		{
-			UE_LOG( LogStats, Verbose, TEXT( "StatMessage Packet has more than %i messages.  Ignoring for the presize history." ), (int32)PRESIZE_MAX_SIZE );
-		}
-		FStatPacket* ToSend = new FStatPacket(Packet);
-		Exchange(ToSend->StatMessages, Packet.StatMessages);
-		ToSend->bBrokenCallstacks = bHasBrokenCallstacks;
-
-		check(!Packet.StatMessages.Num());
-		if( Packet.StatMessagesPresize.Num() > 0 )
-		{
-			int32 MaxPresize = Packet.StatMessagesPresize[0];
-			for (int32 Index = 0; Index < Packet.StatMessagesPresize.Num(); ++Index)
-			{
-				if (MaxPresize < Packet.StatMessagesPresize[Index])
-				{
-					MaxPresize = Packet.StatMessagesPresize[Index];
-				}
-			}
-			Packet.StatMessages.Empty(MaxPresize);
-		}
-
-		SendMessage_Async(ToSend);
-		UpdateExplicitFlush();
-	}
-}
-
-void FThreadStats::FlushRawStats( bool bHasBrokenCallstacks /*= false*/, bool bForceFlush /*= false*/ )
-{
-	LLM_SCOPE(ELLMTag::Stats);
-
-	if (bReentranceGuard)
-	{
-		return;
-	}
-	TGuardValue<bool> Guard( bReentranceGuard, true );
-
-	enum
-	{
-		/** Maximum number of messages in the stat packet. */
-		MAX_RAW_MESSAGES_IN_PACKET = 1024*1024 / sizeof(FStatMessage),
-	};
-
-	// Sends all collected messages when:
-	// Number of messages is greater than MAX_RAW_MESSAGES_IN_PACKET.
-	// The current game frame has changed.
-	// This a force flush when we shutting down the thread stats.
-	// This is an explicit flush from the game thread or the render thread.
-	const bool bFrameHasChanged = DetectAndUpdateCurrentGameFrame();
-	const int32 NumMessages = Packet.StatMessages.Num();
-	if( NumMessages > MAX_RAW_MESSAGES_IN_PACKET || bFrameHasChanged || bForceFlush || bSawExplicitFlush )
-	{
-		SCOPE_CYCLE_COUNTER(STAT_FlushRawStats);
-
-		FStatPacket* ToSend = new FStatPacket(Packet);
-		Exchange(ToSend->StatMessages, Packet.StatMessages);
-		ToSend->bBrokenCallstacks = bHasBrokenCallstacks;
-
-		check(!Packet.StatMessages.Num());
-
-		SendMessage_Async(ToSend);
-		UpdateExplicitFlush();
-
-		const float NumMessagesAsMB = float(NumMessages * sizeof(FStatMessage)) / 1024.0f / 1024.0f;
-		if( NumMessages > 524288 )
-		{
-			UE_LOG( LogStats, Warning, TEXT( "FlushRawStats NumMessages: %i (%.2f MB), Thread: %u" ), NumMessages, NumMessagesAsMB, Packet.ThreadId );
-		}
-
-		UE_LOG( LogStats, Verbose, TEXT( "FlushRawStats NumMessages: %i (%.2f MB), Thread: %u" ), NumMessages, NumMessagesAsMB, Packet.ThreadId );
-	}
-}
-
-void FThreadStats::CheckForCollectingStartupStats()
-{
-	FString CmdLine(FCommandLine::Get());
-	FString StatCmds(TEXT("-StatCmds="));
-	while (1)
-	{
-		FString Cmds;
-		if (!FParse::Value(*CmdLine, *StatCmds, Cmds, false))
-		{
-			break;
-		}
-		TArray<FString> CmdsArray;
-		Cmds.ParseIntoArray(CmdsArray, TEXT( "," ), true);
-		for (int32 Index = 0; Index < CmdsArray.Num(); Index++)
-		{
-			CmdsArray[Index].TrimStartInline();
-			FString StatCmd = FString("stat ") + CmdsArray[Index];
-			UE_LOG(LogStatGroupEnableManager, Log, TEXT("Sending Stat Command '%s'"), *StatCmd);
-			DirectStatsCommand(*StatCmd);
-		}
-		int32 Index = CmdLine.Find(*StatCmds);
-		ensure(Index >= 0);
-		if (Index == INDEX_NONE)
-		{
-			break;
-		}
-		CmdLine.MidInline(Index + StatCmds.Len(), MAX_int32, EAllowShrinking::No);
-	}
-
-	if (FParse::Param( FCommandLine::Get(), TEXT( "LoadTimeStats" ) ))
-	{
-		DirectStatsCommand( TEXT( "stat group enable LinkerLoad" ) );
-		DirectStatsCommand( TEXT( "stat group enable AsyncLoad" ) );
-		DirectStatsCommand( TEXT( "stat group enable LoadTimeVerbose" ) );
-		DirectStatsCommand( TEXT( "stat dumpsum -start -ms=250 -num=240" ), true );
-	}
-	else if (FParse::Param( FCommandLine::Get(), TEXT( "LoadTimeFile" ) ) || FStats::HasLoadTimeFileForCommandletToken())
-	{
-		DirectStatsCommand( TEXT( "stat group enable LinkerLoad" ) );
-		DirectStatsCommand( TEXT( "stat group enable AsyncLoad" ) );
-		DirectStatsCommand( TEXT( "stat group enable LoadTimeVerbose" ) );
-		DirectStatsCommand( TEXT( "stat startfile" ), true );
-	}
-	else if (FStats::HasLoadTimeStatsForCommandletToken())
-	{
-		DirectStatsCommand( TEXT( "stat group enable LinkerLoad" ) );
-		DirectStatsCommand( TEXT( "stat group enable AsyncLoad" ) );
-		DirectStatsCommand( TEXT( "stat group enable LoadTimeVerbose" ) );
-		DirectStatsCommand( TEXT( "stat dumpsum -start" ), true );
-	}
-
-#if UE_STATS_MEMORY_PROFILER_ENABLED
-	// Now we can safely enable malloc profiler.
-	if (FStatsMallocProfilerProxy::HasMemoryProfilerToken())
-	{
-		// Enable all available groups and enable malloc profiler.
-		IStatGroupEnableManager::Get().StatGroupEnableManagerCommand( TEXT( "all" ) );
-		FStatsMallocProfilerProxy::Get()->SetState( true );
-		DirectStatsCommand( TEXT( "stat startfileraw" ), true );
-	}
-#endif //UE_STATS_MEMORY_PROFILER_ENABLED
-
-	STAT_ADD_CUSTOMMESSAGE_NAME( STAT_NamedMarker, TEXT("CheckForCollectingStartupStats") );
-}
-
-void FThreadStats::ExplicitFlush(bool DiscardCallstack)
-{
-	FThreadStats* ThreadStats = GetThreadStats();
-	//check(ThreadStats->Packet.ThreadType != EThreadType::Other);
-	if (ThreadStats->bWaitForExplicitFlush)
-	{
-		ThreadStats->ScopeCount--; // the main thread pre-incremented this to prevent stats from being sent. we send them at the next available opportunity
-		ThreadStats->bWaitForExplicitFlush = 0;
-	}
-	bool bHasBrokenCallstacks = false;
-	if (DiscardCallstack && ThreadStats->ScopeCount)
-	{
-		ThreadStats->ScopeCount = 0;
-		bHasBrokenCallstacks = true;
-	}
-	ThreadStats->bSawExplicitFlush = true;
-	ThreadStats->Flush(bHasBrokenCallstacks);
-}
-
-void FThreadStats::StartThread()
-{
-	FThreadStats::FrameDataIsIncomplete(); // make this non-zero
-	check(IsInGameThread());
-	check(!IsThreadingReady());
-	// Preallocate a bunch of FThreadStats to avoid dynamic memory allocation.
-	// (Must do this before we expose ourselves to other threads via tls).
-	FThreadStatsPool::Get();
-	FStatsThreadState::GetLocalState(); // start up the state
-	if (!FPlatformTLS::IsValidTlsSlot(TlsSlot))
-	{
-		TlsSlot = FPlatformTLS::AllocTlsSlot();
-		check(FPlatformTLS::IsValidTlsSlot(TlsSlot));
-	}
-
-	check(IsThreadingReady());
-	CheckEnable();
-
-	if( FThreadStats::WillEverCollectData() )
-	{
-		FThreadStats::ExplicitFlush(); // flush the stats and set update the scope so we don't flush again until a frame update, this helps prevent fragmentation
-	}
-	FStartupMessages::Get().AddThreadMetadata( NAME_GameThread, FPlatformTLS::GetCurrentThreadId() );
-
-	CheckForCollectingStartupStats();
-
-	UE_LOG( LogStats, Log, TEXT( "Stats thread started at %f" ), FPlatformTime::Seconds() - GStartTime );
-}
-
-static int32 CurrentEventIndex = 0;
-
-static UE::Tasks::FTask LastFramesBarriers[MAX_STAT_LAG];
-
-static FGraphEventRef LastFramesEvents[MAX_STAT_LAG];
-
-void FThreadStats::StopThread()
-{
-	// Nothing to stop if it was never started
-	if (IsThreadingReady())
-	{
-		if (FStats::HasLoadTimeStatsForCommandletToken())
-		{
-			// Dump all the collected stats to the log, if any.
-			DirectStatsCommand( TEXT( "stat dumpsum -stop -ms=100" ), true );
-		}
-
-		// If we are writing stats data, stop it now.
-		DirectStatsCommand( TEXT( "stat stopfile" ), true );
-
-		FThreadStats::PrimaryDisableForever();
-
-		WaitForStats();
-	}
-
-	// wait for the pipe to complete all currently piped tasks
-	while (GStatsPipe.HasWork())
-	{
-		FPlatformProcess::Yield();
-	}
-	for (int32 Index = 0; Index < MAX_STAT_LAG; Index++)
-	{
-		LastFramesBarriers[Index] = {}; // release memory before the allocator is destroyed
-	}
-}
-
-void FThreadStats::WaitForStats()
-{
-	if (FPlatformProcess::SkipWaitForStats())
-	{
-		return;
-	}
-
-	check(IsInGameThread());
-	if (IsThreadingReady() && !bPrimaryDisableForever)
-	{
-		DECLARE_CYCLE_STAT(TEXT("WaitForStats"), STAT_WaitForStats, STATGROUP_Engine);
-
-		int32 EventIndex = (CurrentEventIndex + MAX_STAT_LAG - 1) % MAX_STAT_LAG;
-
-		if (FPlatformProcess::SupportsMultithreading())
-		{
-			SCOPE_CYCLE_COUNTER(STAT_WaitForStats);
-			LastFramesBarriers[EventIndex].Wait();
-			LastFramesBarriers[EventIndex] = GStatsPipe.Launch(UE_SOURCE_LOCATION, [] {});
-		}
-		else
-		{
-			{
-				SCOPE_CYCLE_COUNTER(STAT_WaitForStats);
-				if (LastFramesEvents[EventIndex].GetReference())
-				{
-					FTaskGraphInterface::Get().WaitUntilTaskCompletes(LastFramesEvents[EventIndex], ENamedThreads::GameThread_Local);
-				}
-			}
-
-			LastFramesEvents[EventIndex] = TGraphTask<FNullGraphTask>::CreateTask(NULL, ENamedThreads::GameThread).
-				ConstructAndDispatchWhenReady(TStatId{}, ENamedThreads::GameThread);
-		}
-
-		CurrentEventIndex++;
-
-#if	!UE_BUILD_SHIPPING
-		DebugLeakTest();
-#endif // !UE_BUILD_SHIPPING
-	}
-}
-
-#endif
+// TODO: remove
diff --git a/Engine/Source/Runtime/Core/Private/Stats/StatsCommand.cpp b/Engine/Source/Runtime/Core/Private/Stats/StatsCommand.cpp
--- a/Engine/Source/Runtime/Core/Private/Stats/StatsCommand.cpp
+++ b/Engine/Source/Runtime/Core/Private/Stats/StatsCommand.cpp
@@ -1,5 +1,7 @@
 // Copyright Epic Games, Inc. All Rights Reserved.
 
+#include "Stats/StatsCommand.h"
+
 #include "CoreTypes.h"
 #include "Misc/AssertionMacros.h"
 #include "Math/NumericLimits.h"
@@ -2195,196 +2197,204 @@ protected:
 	virtual bool Exec_Runtime( UWorld*, const TCHAR* Cmd, FOutputDevice& Ar ) override
 	{
 		// Block the thread as this affects external stat states now
-		return DirectStatsCommand(Cmd,true,&Ar);
+		return UE::Stats::DirectStatsCommand(Cmd,true,&Ar);
 	}
 }
 StatCmdCoreExec;
 
-
-bool DirectStatsCommand(const TCHAR* Cmd, bool bBlockForCompletion /*= false*/, FOutputDevice* Ar /*= nullptr*/)
+namespace UE::Stats
 {
-	bool bResult = false;
-	const bool bStatCommand = FParse::Command(&Cmd,TEXT("stat"));
-	const bool bBudgetCommand = FParse::Command(&Cmd,TEXT("budget"));
-	
-	if(bStatCommand || bBudgetCommand)
+	bool DirectStatsCommand(const TCHAR* Cmd, bool bBlockForCompletion /*= false*/, FOutputDevice* Ar /*= nullptr*/)
 	{
-		FString AddArgs;
-		const TCHAR* TempCmd = Cmd;
+		bool bResult = false;
+		const bool bStatCommand = FParse::Command(&Cmd, TEXT("stat"));
+		const bool bBudgetCommand = FParse::Command(&Cmd, TEXT("budget"));
+
+		if (bStatCommand || bBudgetCommand)
+		{
+			FString AddArgs;
+			const TCHAR* TempCmd = Cmd;
 
-		FString ArgNoWhitespaces = FDefaultValueHelper::RemoveWhitespaces(TempCmd);
-		const bool bIsEmpty = ArgNoWhitespaces.IsEmpty();
+			FString ArgNoWhitespaces = FDefaultValueHelper::RemoveWhitespaces(TempCmd);
+			const bool bIsEmpty = ArgNoWhitespaces.IsEmpty();
 #if STATS
-		bResult = true;
+			bResult = true;
 
-		if(bStatCommand)
-		{
-			if (bIsEmpty && Ar)
-			{
-				PrintStatsHelpToOutputDevice(*Ar);
-			}
-			else if (FParse::Command(&TempCmd, TEXT("STARTFILE")))
+			if (bStatCommand)
 			{
-				FString Filename;
-				AddArgs += TEXT(" ");
-				if (FParse::Line(&TempCmd, Filename, true))
+				if (bIsEmpty && Ar)
 				{
-					AddArgs += Filename;
+					PrintStatsHelpToOutputDevice(*Ar);
 				}
-				else
+				else if (FParse::Command(&TempCmd, TEXT("STARTFILE")))
 				{
-					AddArgs += CreateProfileFilename(FStatConstants::StatsFileExtension, true);
+					FString Filename;
+					AddArgs += TEXT(" ");
+					if (FParse::Line(&TempCmd, Filename, true))
+					{
+						AddArgs += Filename;
+					}
+					else
+					{
+						AddArgs += CreateProfileFilename(FStatConstants::StatsFileExtension, true);
+					}
 				}
-			}
-			else if (FParse::Command(&TempCmd, TEXT("StartFileRaw")))
-			{
-				AddArgs += TEXT(" ");
-				AddArgs += CreateProfileFilename(FStatConstants::StatsFileRawExtension, true);
-			}
-			else if (FParse::Command(&TempCmd, TEXT("DUMPFRAME")))
-			{
-			}
-			else if (FParse::Command(&TempCmd, TEXT("DUMPNONFRAME")))
-			{
-			}
-			else if (FParse::Command(&TempCmd, TEXT("DUMPCPU")))
-			{
-			}
-			else if (FParse::Command(&TempCmd, TEXT("STOP")))
-			{
-			}
-			else if (FParse::Command(&TempCmd, TEXT("DUMPAVE")))
-			{
-			}
-			else if (FParse::Command(&TempCmd, TEXT("DUMPMAX")))
-			{
-			}
-			else if (FParse::Command(&TempCmd, TEXT("DUMPSUM")))
-			{
-			}
-			else if (FParse::Command(&TempCmd, TEXT("DUMPSPAM")))
-			{
-			}
-			else if (FParse::Command(&TempCmd, TEXT("DUMPHITCHES")))
-			{
-			}
-			else if (FParse::Command(&TempCmd, TEXT("DumpEvents")))
-			{
-			}
-			else if (FParse::Command(&TempCmd, TEXT("STOPFILE")))
-			{
-			}
-			else if (FParse::Command(&TempCmd, TEXT("TESTFILE")))
-			{
-			}
-			else if (FParse::Command(&TempCmd, TEXT("testdisable")))
-			{
-			}
-			else if (FParse::Command(&TempCmd, TEXT("none")))
-			{
-			}
-			else if (FParse::Command(&TempCmd, TEXT("group")))
-			{
-			}
-			else if (FParse::Command(&TempCmd, TEXT("hier")))
-			{
-			}
-			else if (FParse::Command(&TempCmd, TEXT("net")))
-			{
-			}
-			else if (FParse::Command(&TempCmd, TEXT("toggledebug")))
-			{
-			}
-			else if (FParse::Command(&TempCmd, TEXT("memoryprofiler")))
-			{
-			}
-			else if (FParse::Command(&TempCmd, TEXT("slow")))
-			{
-			}
-			else if (FParse::Command(&TempCmd, TEXT("namedmarker")))
-			{
-			}
-			else
-			{
-				bResult = false;
-
-				FString MaybeGroup;
-				if (FParse::Token(TempCmd, MaybeGroup, false) && MaybeGroup.Len() > 0)
+				else if (FParse::Command(&TempCmd, TEXT("StartFileRaw")))
+				{
+					AddArgs += TEXT(" ");
+					AddArgs += CreateProfileFilename(FStatConstants::StatsFileRawExtension, true);
+				}
+				else if (FParse::Command(&TempCmd, TEXT("DUMPFRAME")))
+				{
+				}
+				else if (FParse::Command(&TempCmd, TEXT("DUMPNONFRAME")))
+				{
+				}
+				else if (FParse::Command(&TempCmd, TEXT("DUMPCPU")))
 				{
-					// If there is + at the end of the group name, remove it
-					const int32 PlusPos = MaybeGroup.Len() - 1;
-					const bool bHierarchy = MaybeGroup[MaybeGroup.Len() - 1] == TEXT('+');
-					if (bHierarchy)
+				}
+				else if (FParse::Command(&TempCmd, TEXT("STOP")))
+				{
+				}
+				else if (FParse::Command(&TempCmd, TEXT("DUMPAVE")))
+				{
+				}
+				else if (FParse::Command(&TempCmd, TEXT("DUMPMAX")))
+				{
+				}
+				else if (FParse::Command(&TempCmd, TEXT("DUMPSUM")))
+				{
+				}
+				else if (FParse::Command(&TempCmd, TEXT("DUMPSPAM")))
+				{
+				}
+				else if (FParse::Command(&TempCmd, TEXT("DUMPHITCHES")))
+				{
+				}
+				else if (FParse::Command(&TempCmd, TEXT("DumpEvents")))
+				{
+				}
+				else if (FParse::Command(&TempCmd, TEXT("STOPFILE")))
+				{
+				}
+				else if (FParse::Command(&TempCmd, TEXT("TESTFILE")))
+				{
+				}
+				else if (FParse::Command(&TempCmd, TEXT("testdisable")))
+				{
+				}
+				else if (FParse::Command(&TempCmd, TEXT("none")))
+				{
+				}
+				else if (FParse::Command(&TempCmd, TEXT("group")))
+				{
+				}
+				else if (FParse::Command(&TempCmd, TEXT("hier")))
+				{
+				}
+				else if (FParse::Command(&TempCmd, TEXT("net")))
+				{
+				}
+				else if (FParse::Command(&TempCmd, TEXT("toggledebug")))
+				{
+				}
+				else if (FParse::Command(&TempCmd, TEXT("memoryprofiler")))
+				{
+				}
+				else if (FParse::Command(&TempCmd, TEXT("slow")))
+				{
+				}
+				else if (FParse::Command(&TempCmd, TEXT("namedmarker")))
+				{
+				}
+				else
+				{
+					bResult = false;
+
+					FString MaybeGroup;
+					if (FParse::Token(TempCmd, MaybeGroup, false) && MaybeGroup.Len() > 0)
 					{
-						MaybeGroup.RemoveAt(PlusPos, EAllowShrinking::No);
-					}
+						// If there is + at the end of the group name, remove it
+						const int32 PlusPos = MaybeGroup.Len() - 1;
+						const bool bHierarchy = MaybeGroup[MaybeGroup.Len() - 1] == TEXT('+');
+						if (bHierarchy)
+						{
+							MaybeGroup.RemoveAt(PlusPos, EAllowShrinking::No);
+						}
 
-					const FName MaybeGroupFName = FName(*(FString(TEXT("STATGROUP_")) + MaybeGroup));
-					bResult = FStatGroupGameThreadNotifier::Get().StatGroupNames.Contains(MaybeGroupFName);
+						const FName MaybeGroupFName = FName(*(FString(TEXT("STATGROUP_")) + MaybeGroup));
+						bResult = FStatGroupGameThreadNotifier::Get().StatGroupNames.Contains(MaybeGroupFName);
+					}
 				}
 			}
-		}
-		else
-		{
-			FString BudgetSection;
-			const TCHAR* TmpCmd = Cmd;
-			if (FParse::Token(TmpCmd, BudgetSection, false) && BudgetSection.Len() > 0)
+			else
 			{
-				FScopeLock BudgetINILock(&BudgetStatMapCS);   //Make sure stats thread isn't currently reading from this data
+				FString BudgetSection;
+				const TCHAR* TmpCmd = Cmd;
+				if (FParse::Token(TmpCmd, BudgetSection, false) && BudgetSection.Len() > 0)
 				{
-					FBudgetData& BudgetData = BudgetStatMapping.FindOrAdd(*BudgetSection);
-					BudgetData = FBudgetData();
-					GConfig->GetArray(*BudgetSection, TEXT("Stats"), BudgetData.Stats, GEngineIni);
-					
-					TArray<FString> Lines;
-					GConfig->GetSection(*BudgetSection, Lines, GEngineIni);
-					for(const FString& Line : Lines)
+					FScopeLock BudgetINILock(&BudgetStatMapCS);   //Make sure stats thread isn't currently reading from this data
 					{
-						if(!Line.Contains(TEXT("+Stats=")))	//ignore stats array
+						FBudgetData& BudgetData = BudgetStatMapping.FindOrAdd(*BudgetSection);
+						BudgetData = FBudgetData();
+						GConfig->GetArray(*BudgetSection, TEXT("Stats"), BudgetData.Stats, GEngineIni);
+
+						TArray<FString> Lines;
+						GConfig->GetSection(*BudgetSection, Lines, GEngineIni);
+						for (const FString& Line : Lines)
 						{
-							FString ThreadName;
-							Line.Split(FString(TEXT("=")), &ThreadName, nullptr);
-							float Budget = -1.f;
-							if(GConfig->GetFloat(*BudgetSection, *ThreadName, Budget, GEngineIni))
+							if (!Line.Contains(TEXT("+Stats=")))	//ignore stats array
 							{
-								BudgetData.ThreadBudgetMap.FindOrAdd(FName(*ThreadName)) = Budget;
+								FString ThreadName;
+								Line.Split(FString(TEXT("=")), &ThreadName, nullptr);
+								float Budget = -1.f;
+								if (GConfig->GetFloat(*BudgetSection, *ThreadName, Budget, GEngineIni))
+								{
+									BudgetData.ThreadBudgetMap.FindOrAdd(FName(*ThreadName)) = Budget;
+								}
 							}
 						}
-					}
 
-					BudgetData.Process();
+						BudgetData.Process();
+					}
 				}
 			}
-		}
-		
+
 #endif
 
-		check(IsInGameThread());
-		if( !bIsEmpty )
-		{
-			const FString FullCmd = FString(Cmd) + AddArgs;
+			check(IsInGameThread());
+			if (!bIsEmpty)
+			{
+				const FString FullCmd = FString(Cmd) + AddArgs;
 #if STATS
-			// make sure these are initialized on the game thread
-			FLatestGameThreadStatsData::Get();
-			FStatGroupGameThreadNotifier::Get();
+				// make sure these are initialized on the game thread
+				FLatestGameThreadStatsData::Get();
+				FStatGroupGameThreadNotifier::Get();
 
-			DECLARE_CYCLE_STAT(TEXT("FSimpleDelegateGraphTask.StatCmd"),
-				STAT_FSimpleDelegateGraphTask_StatCmd,
-				STATGROUP_TaskGraphTasks);
+				DECLARE_CYCLE_STAT(TEXT("FSimpleDelegateGraphTask.StatCmd"),
+					STAT_FSimpleDelegateGraphTask_StatCmd,
+					STATGROUP_TaskGraphTasks);
 
-			UE::Tasks::FTask Task = GStatsPipe.Launch(UE_SOURCE_LOCATION, [FullCmd, bStatCommand, Ar] { StatCmd(FullCmd, bStatCommand, Ar); });
-			if (bBlockForCompletion)
-			{
-				Task.Wait();
-				GLog->FlushThreadedLogs();
-			}
+				UE::Tasks::FTask Task = GStatsPipe.Launch(UE_SOURCE_LOCATION, [FullCmd, bStatCommand, Ar] { StatCmd(FullCmd, bStatCommand, Ar); });
+				if (bBlockForCompletion)
+				{
+					Task.Wait();
+					GLog->FlushThreadedLogs();
+				}
 #else // STATS
-			// If stats aren't enabled, broadcast so engine stats can still be triggered
-			StatCmd(FullCmd, bStatCommand, Ar);
+				// If stats aren't enabled, broadcast so engine stats can still be triggered
+				StatCmd(FullCmd, bStatCommand, Ar);
 #endif // STATS
+			}
 		}
+		return bResult;
 	}
-	return bResult;
+} // namespace UE::Stats
+
+// Deprecated
+bool DirectStatsCommand(const TCHAR* Cmd, bool bBlockForCompletion, FOutputDevice* Ar)
+{
+	return UE::Stats::DirectStatsCommand(Cmd, bBlockForCompletion, Ar);
 }
 
 #if STATS
@@ -2413,4 +2423,3 @@ void GetPermanentStats(TArray<FStatMessage>& OutStats)
 }
 
 #endif
-
diff --git a/Engine/Source/Runtime/Core/Private/Stats/StatsMisc.cpp b/Engine/Source/Runtime/Core/Private/Stats/StatsMisc.cpp
--- a/Engine/Source/Runtime/Core/Private/Stats/StatsMisc.cpp
+++ b/Engine/Source/Runtime/Core/Private/Stats/StatsMisc.cpp
@@ -4,36 +4,6 @@
 #include "Stats/Stats.h"
 #include "Logging/LogMacros.h"
 #include "CoreGlobals.h"
-#include "AutoRTFM/AutoRTFM.h"
-
-
-#if !STATS && !UE_BUILD_DEBUG && defined(USE_LIGHTWEIGHT_STATS_FOR_HITCH_DETECTION) && USE_LIGHTWEIGHT_STATS_FOR_HITCH_DETECTION && USE_HITCH_DETECTION
-#include "HAL/ThreadHeartBeat.h"
-#include "HAL/ThreadManager.h"
-
-UE_AUTORTFM_ALWAYS_OPEN
-void FLightweightStatScope::ReportHitch()
-{
-	if (StatString)
-	{
-		float Delta = float(FGameThreadHitchHeartBeat::Get().GetCurrentTime() - FGameThreadHitchHeartBeat::Get().GetFrameStartTime()) * 1000.0f;
-
-		const uint32 CurrentThreadId = FPlatformTLS::GetCurrentThreadId();
-		const bool isGT = CurrentThreadId == GGameThreadId;
-		const FString& ThreadString = FThreadManager::GetThreadName(CurrentThreadId);
-		FString StackString = StatString; // possibly convert from ANSICHAR
-
-		if (!isGT && (StackString == TEXT("STAT_EventWait") || StackString == TEXT("STAT_FQueuedThread_Run_WaitForWork")))
-		{
-			return;
-		}
-
-		UE_LOG(LogCore, Error, TEXT("Leaving stat scope on hitch (+%8.2fms) [%s] %s"), Delta, *ThreadString, *StackString);
-	}
-}
-
-#endif
-
 
 FConditionalScopeLogTime::FConditionalScopeLogTime( bool bCondition, const WIDECHAR* InName, FTotalTimeAndCount* InCumulative /*= nullptr */, EScopeLogTimeUnits InUnits /*= ScopeLog_Milliseconds */ )
 : StartTime( bCondition ? FPlatformTime::Seconds() : 0.0 )
diff --git a/Engine/Source/Runtime/Core/Private/Windows/WindowsPlatformProcess.cpp b/Engine/Source/Runtime/Core/Private/Windows/WindowsPlatformProcess.cpp
--- a/Engine/Source/Runtime/Core/Private/Windows/WindowsPlatformProcess.cpp
+++ b/Engine/Source/Runtime/Core/Private/Windows/WindowsPlatformProcess.cpp
@@ -28,6 +28,7 @@
 #include "ProfilingDebugging/CpuProfilerTrace.h"
 #include "ProfilingDebugging/CsvProfiler.h"
 #include "Stats/Stats.h"
+#include "Stats/ThreadIdleStats.h"
 #include "Templates/UnrealTemplate.h"
 #include "Trace/Trace.h"
 #include "Trace/Trace.inl"
@@ -1665,7 +1666,7 @@ bool FWindowsPlatformProcess::ResolveNetworkPath( FString InUNCPath, FString& Ou
 void FWindowsPlatformProcess::Sleep( float Seconds )
 {
 	SCOPE_CYCLE_COUNTER( STAT_Sleep );
-	FThreadIdleStats::FScopeIdle Scope;
+	UE::Stats::FThreadIdleStats::FScopeIdle Scope;
 	SleepNoStats( Seconds );
 }
 
@@ -1735,7 +1736,7 @@ bool FEventWin::Wait(uint32 WaitTime, const bool bIgnoreThreadIdleStats /*= fals
 	check(Event);
 
 	LowLevelTasks::FOversubscriptionScope _(WaitTime != 0); // Let the scheduler know one of its thread might be waiting.
-	FThreadIdleStats::FScopeIdle Scope( bIgnoreThreadIdleStats );
+	UE::Stats::FThreadIdleStats::FScopeIdle Scope( bIgnoreThreadIdleStats );
 	return (WaitForSingleObject( Event, WaitTime ) == WAIT_OBJECT_0);
 }
 
diff --git a/Engine/Source/Runtime/Core/Public/Stats/HitchTrackingStatScope.h b/Engine/Source/Runtime/Core/Public/Stats/HitchTrackingStatScope.h
new file mode 100644
--- /dev/null
+++ b/Engine/Source/Runtime/Core/Public/Stats/HitchTrackingStatScope.h
@@ -0,0 +1,45 @@
+// Copyright Epic Games, Inc. All Rights Reserved.
+#pragma once
+
+#include "CoreTypes.h"
+#include "Misc/Build.h"
+#include "Stats/StatsCommon.h"
+
+#if UE_ENABLE_INCLUDE_ORDER_DEPRECATED_IN_5_6
+#include "ProfilingDebugging/UMemoryDefines.h"
+#endif
+
+#ifndef USE_LIGHTWEIGHT_STATS_FOR_HITCH_DETECTION
+#define USE_LIGHTWEIGHT_STATS_FOR_HITCH_DETECTION 1
+#endif
+
+#if USE_LIGHTWEIGHT_STATS_FOR_HITCH_DETECTION && USE_HITCH_DETECTION
+extern CORE_API TSAN_ATOMIC(bool) GHitchDetected;
+
+namespace UE::Stats
+{
+	class FHitchTrackingStatScope
+	{
+		const PROFILER_CHAR* StatString;
+	public:
+		FORCEINLINE FHitchTrackingStatScope(const PROFILER_CHAR* InStat)
+		{
+			StatString = GHitchDetected ? nullptr : InStat;
+		}
+
+		FORCEINLINE ~FHitchTrackingStatScope()
+		{
+			if (GHitchDetected && StatString)
+			{
+				ReportHitch();
+			}
+		}
+
+		CORE_API void ReportHitch();
+	};
+} // namespace UE::Stats
+
+UE_DEPRECATED(5.6, "Use UE::Stats::FHitchTrackingStatScope instead.")
+typedef UE::Stats::FHitchTrackingStatScope FLightweightStatScope;
+
+#endif // USE_LIGHTWEIGHT_STATS_FOR_HITCH_DETECTION && USE_HITCH_DETECTION
\ No newline at end of file
diff --git a/Engine/Source/Runtime/Core/Public/Stats/LightweightStats.h b/Engine/Source/Runtime/Core/Public/Stats/LightweightStats.h
new file mode 100644
--- /dev/null
+++ b/Engine/Source/Runtime/Core/Public/Stats/LightweightStats.h
@@ -0,0 +1,109 @@
+// Copyright Epic Games, Inc. All Rights Reserved.
+#pragma once
+
+#include "CoreTypes.h"
+#include "CoreGlobals.h"
+#include "Misc/Build.h"
+#include "StatsCommon.h"
+
+#ifndef UE_USE_LIGHTWEIGHT_STATS
+#define UE_USE_LIGHTWEIGHT_STATS (!STATS && ENABLE_STATNAMEDEVENTS)
+#endif
+
+#if UE_USE_LIGHTWEIGHT_STATS
+
+#include "AutoRTFM/AutoRTFM.h"
+#include "Math/Color.h"
+#include "Stats/HitchTrackingStatScope.h"
+
+struct TStatId
+{
+	const PROFILER_CHAR* StatString;
+
+	FORCEINLINE TStatId()
+		: StatString(nullptr)
+	{
+	}
+
+	FORCEINLINE TStatId(const PROFILER_CHAR* InString)
+		: StatString(InString)
+	{
+	}
+
+	FORCEINLINE bool IsValidStat() const
+	{
+		return StatString != nullptr;
+	}
+
+	FORCEINLINE bool operator==(TStatId Other) const
+	{
+		return StatString == Other.StatString;
+	}
+
+	FORCEINLINE bool operator!=(TStatId Other) const
+	{
+		return StatString != Other.StatString;
+	}
+};
+
+class FScopeCycleCounter
+{
+public:
+	FORCEINLINE FScopeCycleCounter(TStatId InStatId, EStatFlags StatFlags, bool bAlways = false)
+		:
+#if USE_LIGHTWEIGHT_STATS_FOR_HITCH_DETECTION && USE_HITCH_DETECTION
+		StatScope(InStatId.StatString),
+#endif
+
+		bPop(false)
+	{
+		if (GCycleStatsShouldEmitNamedEvents && InStatId.IsValidStat())
+		{
+			AutoRTFM::Open([&]
+				{
+					bPop = true;
+					FPlatformMisc::BeginNamedEvent(FColor(0), InStatId.StatString);
+				});
+			AutoRTFM::PushOnAbortHandler(this, []() { FPlatformMisc::EndNamedEvent(); });
+		}
+	}
+
+	FORCEINLINE FScopeCycleCounter(TStatId InStatId, bool bAlways = false)
+		: FScopeCycleCounter(InStatId, EStatFlags::None, bAlways)
+	{
+	}
+
+	FORCEINLINE ~FScopeCycleCounter()
+	{
+		AutoRTFM::PopOnAbortHandler(this);
+
+		AutoRTFM::Open([&]
+			{
+				if (bPop)
+				{
+					FPlatformMisc::EndNamedEvent();
+				}
+			});
+	}
+private:
+#if USE_LIGHTWEIGHT_STATS_FOR_HITCH_DETECTION && USE_HITCH_DETECTION
+	UE::Stats::FHitchTrackingStatScope StatScope;
+#endif
+	bool bPop;
+};
+
+#elif !STATS // UE_USE_LIGHTWEIGHT_STATS
+
+struct TStatId {};
+
+class FScopeCycleCounter
+{
+public:
+	inline FScopeCycleCounter(TStatId, EStatFlags, bool bAlways = false)
+	{
+	}
+	inline FScopeCycleCounter(TStatId, bool bAlways = false)
+	{
+	}
+};
+#endif // UE_USE_LIGHTWEIGHT_STATS
\ No newline at end of file
