Unify Stats headers

diff --git a/Engine/Source/Runtime/Core/Public/Stats/Stats.h b/Engine/Source/Runtime/Core/Public/Stats/Stats.h
--- a/Engine/Source/Runtime/Core/Public/Stats/Stats.h
+++ b/Engine/Source/Runtime/Core/Public/Stats/Stats.h
@@ -2,29 +2,29 @@
 
 #pragma once
 
+#include "CoreGlobals.h"
 #include "CoreTypes.h"
-#if UE_ENABLE_INCLUDE_ORDER_DEPRECATED_IN_5_4
-#include "Math/Color.h"
-#endif
-#if ENABLE_STATNAMEDEVENTS
-#include "Math/Color.h"
-#endif
-#include "UObject/NameTypes.h"
+#include "StatsCommon.h"
+#include "Stats/DynamicStats.h"
+#include "Stats/HitchTrackingStatScope.h"
+#include "Stats/LightweightStats.h"
+#include "Stats/StatsSystemTypes.h"
+
+#if UE_ENABLE_INCLUDE_ORDER_DEPRECATED_IN_5_6
+// These includes are no longer required for this file.
 #include "Delegates/Delegate.h"
 #include "HAL/PlatformTime.h"
 #include "HAL/ThreadSingleton.h"
-#include "StatsCommon.h"
 #include "ProfilingDebugging/UMemoryDefines.h"
-#include "Stats/Stats2.h"
 #include "AutoRTFM/AutoRTFM.h"
-#include "Misc/SourceLocation.h"
-
-/** 
- *	Learn about the Stats System at docs.unrealengine.com
- */
+#include "Math/Color.h"
+#include "UObject/NameTypes.h"
 
-struct TStatId;
-enum class EStatFlags : uint8;
+// Include these directly where needed instead.
+#include "Stats/StatsCommand.h"
+#include "Stats/StatsSystem.h"
+#include "Stats/ThreadIdleStats.h"
+#endif // UE_ENABLE_INCLUDE_ORDER_DEPRECATED_IN_5_6
 
 // used by the profiler
 enum EStatType
@@ -38,226 +38,467 @@ enum EStatType
 	STATTYPE_Error
 };
 
-/*----------------------------------------------------------------------------
-	Stats helpers
-----------------------------------------------------------------------------*/
-
 #if STATS
 	#define STAT(x) x
 #else
 	#define STAT(x)
 #endif
 
-#ifndef USE_LIGHTWEIGHT_STATS_FOR_HITCH_DETECTION
-#define USE_LIGHTWEIGHT_STATS_FOR_HITCH_DETECTION 1
-#endif
-
 #if STATS
-
-/**
- * This is a utility class for counting the number of cycles during the
- * lifetime of the object. It updates the per thread values for this
- * stat.
- */
-class FScopeCycleCounter : public FCycleCounter
-{
-public:
-	/**
-	 * Pushes the specified stat onto the hierarchy for this thread. Starts
-	 * the timing of the cycles used
-	 */
-	FORCEINLINE_STATS FScopeCycleCounter( TStatId StatId, EStatFlags StatFlags, bool bAlways = false, UE::FSourceLocation SourceLocation = UE::FSourceLocation::Current())
-	{
-		AutoRTFM::Open([&]{ Start(StatId, StatFlags, bAlways, SourceLocation); });
-		AutoRTFM::PushOnAbortHandler(this, [this](){ this->Stop(); });
-	}
-
-	FORCEINLINE_STATS FScopeCycleCounter(TStatId StatId, bool bAlways = false, UE::FSourceLocation SourceLocation = UE::FSourceLocation::Current())
-		: FScopeCycleCounter(StatId, EStatFlags::None, bAlways, SourceLocation)
-	{
-	}
-
-	/**
-	 * Updates the stat with the time spent
-	 */
-	FORCEINLINE_STATS ~FScopeCycleCounter()
-	{
-		AutoRTFM::PopOnAbortHandler(this);
-		AutoRTFM::Open([&]
-		{
-			Stop();
-		});
-	}
-
+// ----------------------------------------------------------------------------
+// Stats
+// ----------------------------------------------------------------------------
+// These macros emit events to the stat system as well as to profilers.
+// See official docs for more information on the Stats System.
+// ----------------------------------------------------------------------------
+
+#define DECLARE_STAT_GROUP(Description, StatName, StatCategory, InDefaultEnable, InCompileTimeEnable, InSortByName) \
+struct FStatGroup_##StatName\
+{ \
+	enum \
+	{ \
+		DefaultEnable = InDefaultEnable, \
+		CompileTimeEnable = InCompileTimeEnable, \
+		SortByName = InSortByName \
+	}; \
+	static FORCEINLINE const char* GetGroupName() \
+	{ \
+		return #StatName; \
+	} \
+	static FORCEINLINE const char* GetGroupCategory() \
+	{ \
+		return #StatCategory; \
+	} \
+	static FORCEINLINE const TCHAR* GetDescription() \
+	{ \
+		return Description; \
+	} \
+	static FORCEINLINE bool IsDefaultEnabled() \
+	{ \
+		return (bool)DefaultEnable; \
+	} \
+	static FORCEINLINE bool IsCompileTimeEnable() \
+	{ \
+		return (bool)CompileTimeEnable; \
+	} \
+	static FORCEINLINE bool GetSortByName() \
+	{ \
+		return (bool)SortByName; \
+	} \
 };
 
-FORCEINLINE void StatsPrimaryEnableAdd(int32 Value = 1)
-{
-	FThreadStats::PrimaryEnableAdd(Value);
-}
-FORCEINLINE void StatsPrimaryEnableSubtract(int32 Value = 1)
-{
-	FThreadStats::PrimaryEnableSubtract(Value);
-}
-
-UE_DEPRECATED(5.1, "Use StatsPrimaryEnableAdd instead")
-FORCEINLINE void StatsMasterEnableAdd(int32 Value = 1)
-{
-	StatsPrimaryEnableAdd(Value);
-}
-
-UE_DEPRECATED(5.1, "Use StatsPrimaryEnableSubtract instead")
-FORCEINLINE void StatsMasterEnableSubtract(int32 Value = 1)
-{
-	StatsPrimaryEnableSubtract(Value);
-}
+#define DECLARE_STAT(Description, StatName, GroupName, StatType, StatFlags, MemoryRegion) \
+struct FStat_##StatName\
+{ \
+	typedef FStatGroup_##GroupName TGroup; \
+	static FORCEINLINE const char* GetStatName() \
+	{ \
+		return #StatName; \
+	} \
+	static FORCEINLINE const TCHAR* GetDescription() \
+	{ \
+		return Description; \
+	} \
+	static FORCEINLINE EStatDataType::Type GetStatType() \
+	{ \
+		return StatType; \
+	} \
+	static FORCEINLINE bool IsClearEveryFrame() \
+	{ \
+		return EnumHasAnyFlags(GetFlags(), EStatFlags::ClearEveryFrame); \
+	} \
+	static FORCEINLINE bool IsCycleStat() \
+	{ \
+		return EnumHasAnyFlags(GetFlags(), EStatFlags::CycleStat); \
+	} \
+	static FORCEINLINE EStatFlags GetFlags() \
+	{ \
+		return StatFlags; \
+	} \
+	static FORCEINLINE FPlatformMemory::EMemoryCounterRegion GetMemoryRegion() \
+	{ \
+		return MemoryRegion; \
+	} \
+};
 
-#else	//STATS
-#if ENABLE_STATNAMEDEVENTS
+#define GET_STATID(Stat) (StatPtr_##Stat.GetStatId())
+#define GET_STATFNAME(Stat) (StatPtr_##Stat.GetStatFName())
+#define GET_STATDESCRIPTION(Stat) (FStat_##Stat::GetDescription())
+#define GET_STATISEVERYFRAME(Stat) (FStat_##Stat::IsClearEveryFrame())
+#define GET_STATFLAGS(Stat) (FStat_##Stat::GetFlags())
+
+#define STAT_GROUP_TO_FStatGroup(Group) FStatGroup_##Group
+
+/*-----------------------------------------------------------------------------
+	Local
+-----------------------------------------------------------------------------*/
+
+#define DEFINE_STAT(Stat) \
+	struct FThreadSafeStaticStat<FStat_##Stat> StatPtr_##Stat;
+
+#define RETURN_QUICK_DECLARE_CYCLE_STAT(StatId,GroupId) \
+	DECLARE_STAT(TEXT(#StatId),StatId,GroupId,EStatDataType::ST_int64, EStatFlags::ClearEveryFrame | EStatFlags::CycleStat, FPlatformMemory::MCR_Invalid); \
+	static DEFINE_STAT(StatId) \
+	return GET_STATID(StatId);
+
+#define QUICK_USE_CYCLE_STAT(StatId,GroupId) [](){ RETURN_QUICK_DECLARE_CYCLE_STAT(StatId, GroupId); }()
+
+#define DECLARE_CYCLE_STAT(CounterName,StatId,GroupId) \
+	DECLARE_STAT(CounterName,StatId,GroupId,EStatDataType::ST_int64, EStatFlags::ClearEveryFrame | EStatFlags::CycleStat, FPlatformMemory::MCR_Invalid); \
+	static DEFINE_STAT(StatId)
+#define DECLARE_CYCLE_STAT_WITH_FLAGS(CounterName,StatId,GroupId,StatFlags) \
+	DECLARE_STAT(CounterName,StatId,GroupId,EStatDataType::ST_int64, (StatFlags) | EStatFlags::ClearEveryFrame | EStatFlags::CycleStat, FPlatformMemory::MCR_Invalid); \
+	static DEFINE_STAT(StatId)
+#define DECLARE_FLOAT_COUNTER_STAT(CounterName,StatId,GroupId) \
+	DECLARE_STAT(CounterName,StatId,GroupId,EStatDataType::ST_double,EStatFlags::ClearEveryFrame, FPlatformMemory::MCR_Invalid); \
+	static DEFINE_STAT(StatId)
+#define DECLARE_DWORD_COUNTER_STAT(CounterName,StatId,GroupId) \
+	DECLARE_STAT(CounterName,StatId,GroupId,EStatDataType::ST_int64,EStatFlags::ClearEveryFrame, FPlatformMemory::MCR_Invalid); \
+	static DEFINE_STAT(StatId)
+#define DECLARE_FLOAT_ACCUMULATOR_STAT(CounterName,StatId,GroupId) \
+	DECLARE_STAT(CounterName,StatId,GroupId,EStatDataType::ST_double, EStatFlags::None, FPlatformMemory::MCR_Invalid); \
+	static DEFINE_STAT(StatId)
+#define DECLARE_DWORD_ACCUMULATOR_STAT(CounterName,StatId,GroupId) \
+	DECLARE_STAT(CounterName,StatId,GroupId,EStatDataType::ST_int64, EStatFlags::None, FPlatformMemory::MCR_Invalid); \
+	static DEFINE_STAT(StatId)
+
+/** FName stat that allows sending a string based data. */
+#define DECLARE_FNAME_STAT(CounterName,StatId,GroupId) \
+	DECLARE_STAT(CounterName,StatId,GroupId,EStatDataType::ST_FName, EStatFlags::None, FPlatformMemory::MCR_Invalid); \
+	static DEFINE_STAT(StatId)
+
+/** This is a fake stat, mostly used to implement memory message or other custom stats that don't easily fit into the system. */
+#define DECLARE_PTR_STAT(CounterName,StatId,GroupId)\
+	DECLARE_STAT(CounterName,StatId,GroupId,EStatDataType::ST_Ptr, EStatFlags::None, FPlatformMemory::MCR_Invalid); \
+	static DEFINE_STAT(StatId)
+
+#define DECLARE_MEMORY_STAT(CounterName,StatId,GroupId) \
+	DECLARE_STAT(CounterName,StatId,GroupId,EStatDataType::ST_int64, EStatFlags::None, FPlatformMemory::MCR_Physical); \
+	static DEFINE_STAT(StatId)
+
+#define DECLARE_MEMORY_STAT_POOL(CounterName,StatId,GroupId,Pool) \
+	DECLARE_STAT(CounterName,StatId,GroupId,EStatDataType::ST_int64, EStatFlags::None, Pool); \
+	static DEFINE_STAT(StatId)
+
+/*-----------------------------------------------------------------------------
+	Extern
+-----------------------------------------------------------------------------*/
+
+#define DECLARE_CYCLE_STAT_EXTERN(CounterName,StatId,GroupId, APIX) \
+	DECLARE_STAT(CounterName,StatId,GroupId,EStatDataType::ST_int64, EStatFlags::ClearEveryFrame | EStatFlags::CycleStat, FPlatformMemory::MCR_Invalid); \
+	extern APIX DEFINE_STAT(StatId);
+#define DECLARE_CYCLE_STAT_WITH_FLAGS_EXTERN(CounterName,StatId,GroupId,StatFlags, APIX) \
+	DECLARE_STAT(CounterName,StatId,GroupId,EStatDataType::ST_int64, (StatFlags) | EStatFlags::ClearEveryFrame | EStatFlags::CycleStat, FPlatformMemory::MCR_Invalid); \
+	extern APIX DEFINE_STAT(StatId);
+#define DECLARE_FLOAT_COUNTER_STAT_EXTERN(CounterName,StatId,GroupId, API) \
+	DECLARE_STAT(CounterName,StatId,GroupId,EStatDataType::ST_double, EStatFlags::ClearEveryFrame, FPlatformMemory::MCR_Invalid); \
+	extern API DEFINE_STAT(StatId);
+#define DECLARE_DWORD_COUNTER_STAT_EXTERN(CounterName,StatId,GroupId, API) \
+	DECLARE_STAT(CounterName,StatId,GroupId,EStatDataType::ST_int64, EStatFlags::ClearEveryFrame, FPlatformMemory::MCR_Invalid); \
+	extern API DEFINE_STAT(StatId);
+#define DECLARE_FLOAT_ACCUMULATOR_STAT_EXTERN(CounterName,StatId,GroupId, API) \
+	DECLARE_STAT(CounterName,StatId,GroupId,EStatDataType::ST_double, EStatFlags::None, FPlatformMemory::MCR_Invalid); \
+	extern API DEFINE_STAT(StatId);
+#define DECLARE_DWORD_ACCUMULATOR_STAT_EXTERN(CounterName,StatId,GroupId, API) \
+	DECLARE_STAT(CounterName,StatId,GroupId,EStatDataType::ST_int64, EStatFlags::None, FPlatformMemory::MCR_Invalid); \
+	extern API DEFINE_STAT(StatId);
+
+/** FName stat that allows sending a string based data. */
+#define DECLARE_FNAME_STAT_EXTERN(CounterName,StatId,GroupId, API) \
+	DECLARE_STAT(CounterName,StatId,GroupId,EStatDataType::ST_FName, EStatFlags::None, FPlatformMemory::MCR_Invalid); \
+	extern API DEFINE_STAT(StatId);
+
+/** This is a fake stat, mostly used to implement memory message or other custom stats that don't easily fit into the system. */
+#define DECLARE_PTR_STAT_EXTERN(CounterName,StatId,GroupId, API) \
+	DECLARE_STAT(CounterName,StatId,GroupId,EStatDataType::ST_Ptr, EStatFlags::None, FPlatformMemory::MCR_Invalid); \
+	extern API DEFINE_STAT(StatId);
+
+#define DECLARE_MEMORY_STAT_EXTERN(CounterName,StatId,GroupId, API) \
+	DECLARE_STAT(CounterName,StatId,GroupId,EStatDataType::ST_int64, EStatFlags::None, FPlatformMemory::MCR_Physical); \
+	extern API DEFINE_STAT(StatId);
+
+#define DECLARE_MEMORY_STAT_POOL_EXTERN(CounterName,StatId,GroupId,Pool, API) \
+	DECLARE_STAT(CounterName,StatId,GroupId,EStatDataType::ST_int64, EStatFlags::None, Pool); \
+	extern API DEFINE_STAT(StatId);
+
+/** Macro for declaring group factory instances */
+#define DECLARE_STATS_GROUP(GroupDesc, GroupId, GroupCat) \
+	DECLARE_STAT_GROUP(GroupDesc, GroupId, GroupCat, true, true, false);
+
+#define DECLARE_STATS_GROUP_SORTBYNAME(GroupDesc, GroupId, GroupCat) \
+	DECLARE_STAT_GROUP(GroupDesc, GroupId, GroupCat, true, true, true);
+
+#define DECLARE_STATS_GROUP_VERBOSE(GroupDesc, GroupId, GroupCat) \
+	DECLARE_STAT_GROUP(GroupDesc, GroupId, GroupCat, false, true, false);
+
+#define DECLARE_STATS_GROUP_MAYBE_COMPILED_OUT(GroupDesc, GroupId, GroupCat, CompileIn) \
+	DECLARE_STAT_GROUP(GroupDesc, GroupId, GroupCat, false, CompileIn, false);
 
-struct FStatStringWrapper
-{
-	
-};
+#define DECLARE_SCOPE_CYCLE_COUNTER(CounterName,Stat,GroupId) \
+	DECLARE_STAT(CounterName,Stat,GroupId,EStatDataType::ST_int64, EStatFlags::ClearEveryFrame | EStatFlags::CycleStat, FPlatformMemory::MCR_Invalid); \
+	static DEFINE_STAT(Stat) \
+	FScopeCycleCounter CycleCount_##Stat(GET_STATID(Stat), GET_STATFLAGS(Stat));
 
-#if	PLATFORM_USES_ANSI_STRING_FOR_EXTERNAL_PROFILING
-typedef ANSICHAR PROFILER_CHAR;
-#else
-typedef WIDECHAR PROFILER_CHAR;
-#endif // PLATFORM_USES_ANSI_STRING_FOR_EXTERNAL_PROFILING
+#define QUICK_SCOPE_CYCLE_COUNTER(Stat) \
+	DECLARE_SCOPE_CYCLE_COUNTER(TEXT(#Stat),Stat,STATGROUP_Quick)
 
-struct TStatId
-{
-	const PROFILER_CHAR* StatString;
-
-	FORCEINLINE TStatId()
-	: StatString(nullptr)
-	{
-	}
-
-	FORCEINLINE TStatId(const PROFILER_CHAR* InString)
-	: StatString(InString)
-	{
-	}
-	
-	FORCEINLINE bool IsValidStat() const
-	{
-		return StatString != nullptr;
-	}
-
-	FORCEINLINE bool operator==(TStatId Other) const
-	{
-		return StatString == Other.StatString;
-	}
-
-	FORCEINLINE bool operator!=(TStatId Other) const
-	{
-		return StatString != Other.StatString;
-	}
-};
+#define SCOPE_CYCLE_COUNTER(Stat) \
+	FScopeCycleCounter CycleCount_##Stat(GET_STATID(Stat), GET_STATFLAGS(Stat));
 
-#if USE_LIGHTWEIGHT_STATS_FOR_HITCH_DETECTION && USE_HITCH_DETECTION
-extern CORE_API TSAN_ATOMIC(bool) GHitchDetected;
+#define SCOPE_CYCLE_COUNTER_STATID(StatId) \
+	FScopeCycleCounter CycleCount_FromStatId(StatId);
 
-class FLightweightStatScope
-{
-	const PROFILER_CHAR* StatString;
-public:
-	FORCEINLINE FLightweightStatScope(const PROFILER_CHAR* InStat)
-	{
-		StatString = GHitchDetected ? nullptr : InStat;
-	}
-
-	FORCEINLINE ~FLightweightStatScope()
-	{
-		if (GHitchDetected && StatString)
-		{
-			ReportHitch();
-		}
-	}
-
-	CORE_API void ReportHitch();
-};
+#define CONDITIONAL_SCOPE_CYCLE_COUNTER(Stat,bCondition) \
+	FScopeCycleCounter CycleCount_##Stat(bCondition ? GET_STATID(Stat) : TStatId(), GET_STATFLAGS(Stat));
 
-#endif
+#define SCOPE_SECONDS_ACCUMULATOR(Stat) \
+	FSimpleScopeSecondsStat SecondsAccum_##Stat(GET_STATID(Stat));
 
-class FScopeCycleCounter
-{
-public:
-	FORCEINLINE FScopeCycleCounter(TStatId InStatId, EStatFlags StatFlags, bool bAlways = false)
-		: 
-#if USE_LIGHTWEIGHT_STATS_FOR_HITCH_DETECTION && USE_HITCH_DETECTION
-		StatScope(InStatId.StatString),
-#endif
-		
-		bPop(false)
-	{
-		if (GCycleStatsShouldEmitNamedEvents && InStatId.IsValidStat())
-		{
-			AutoRTFM::Open([&]
-			{
-				bPop = true;
-				FPlatformMisc::BeginNamedEvent(FColor(0), InStatId.StatString);
-			});
-			AutoRTFM::PushOnAbortHandler(this, [](){ FPlatformMisc::EndNamedEvent(); });
-		}
-	}
-
-	FORCEINLINE FScopeCycleCounter(TStatId InStatId, bool bAlways = false)
-		: FScopeCycleCounter(InStatId, EStatFlags::None, bAlways)
-	{
-	}
-
-	FORCEINLINE ~FScopeCycleCounter()
-	{
-		AutoRTFM::PopOnAbortHandler(this);
-
-		AutoRTFM::Open([&]
-		{
-			if (bPop)
-			{
-				FPlatformMisc::EndNamedEvent();
-			}
-		});
-	}
-private:
-#if USE_LIGHTWEIGHT_STATS_FOR_HITCH_DETECTION && USE_HITCH_DETECTION
-	FLightweightStatScope StatScope;
-#endif
-	bool bPop;
-};
+#define SCOPE_MS_ACCUMULATOR(Stat) \
+	FSimpleScopeSecondsStat SecondsAccum_##Stat(GET_STATID(Stat), 1000.0);
 
-#else	//ENABLE_STATNAMEDEVENTS
-struct TStatId {};
+#define SET_CYCLE_COUNTER(Stat,Cycles) \
+{\
+	if (FThreadStats::IsCollectingData() || !GET_STATISEVERYFRAME(Stat)) \
+		FThreadStats::AddMessage(GET_STATFNAME(Stat), EStatOperation::Set, int64(Cycles), true);\
+}
 
-class FScopeCycleCounter
-{
-public:
-	FORCEINLINE_STATS FScopeCycleCounter(TStatId, EStatFlags, bool bAlways = false)
-	{
-	}
-	FORCEINLINE_STATS FScopeCycleCounter(TStatId, bool bAlways = false)
-	{
-	}
-};
-#endif
+#define INC_DWORD_STAT(Stat) \
+{\
+	if (FThreadStats::IsCollectingData() || !GET_STATISEVERYFRAME(Stat)) \
+	{ \
+		const FName StatName = GET_STATFNAME(Stat); \
+		FThreadStats::AddMessage(StatName, EStatOperation::Add, int64(1));\
+		TRACE_STAT_INCREMENT(StatName); \
+	} \
+}
+#define INC_FLOAT_STAT_BY(Stat, Amount) \
+{\
+	double AddAmount = double(Amount); \
+	if (AddAmount != 0.0) \
+	{ \
+		if (FThreadStats::IsCollectingData() || !GET_STATISEVERYFRAME(Stat)) \
+		{ \
+			const FName StatName = GET_STATFNAME(Stat); \
+			FThreadStats::AddMessage(StatName, EStatOperation::Add, AddAmount); \
+			TRACE_STAT_ADD(StatName, AddAmount); \
+		} \
+	} \
+}
+#define INC_DWORD_STAT_BY(Stat, Amount) \
+{\
+	int64 AddAmount = int64(Amount); \
+	if (AddAmount != 0) \
+	{ \
+		if (FThreadStats::IsCollectingData() || !GET_STATISEVERYFRAME(Stat)) \
+		{ \
+			const FName StatName = GET_STATFNAME(Stat); \
+			FThreadStats::AddMessage(StatName, EStatOperation::Add, AddAmount); \
+			TRACE_STAT_ADD(StatName, AddAmount); \
+		} \
+	} \
+}
+#define INC_MEMORY_STAT_BY(Stat, Amount) \
+{\
+	int64 AddAmount = int64(Amount); \
+	if (AddAmount != 0) \
+	{ \
+		if (FThreadStats::IsCollectingData() || !GET_STATISEVERYFRAME(Stat)) \
+		{ \
+			const FName StatName = GET_STATFNAME(Stat); \
+			FThreadStats::AddMessage(StatName, EStatOperation::Add, AddAmount); \
+			TRACE_STAT_ADD(StatName, AddAmount); \
+		} \
+	} \
+}
+#define DEC_DWORD_STAT(Stat) \
+{\
+	if (FThreadStats::IsCollectingData() || !GET_STATISEVERYFRAME(Stat)) \
+	{ \
+		const FName StatName = GET_STATFNAME(Stat); \
+		FThreadStats::AddMessage(StatName, EStatOperation::Subtract, int64(1));\
+		TRACE_STAT_DECREMENT(StatName); \
+	} \
+}
+#define DEC_FLOAT_STAT_BY(Stat,Amount) \
+{\
+	double SubtractAmount = double(Amount); \
+	if (SubtractAmount != 0.0) \
+	{ \
+		if (FThreadStats::IsCollectingData() || !GET_STATISEVERYFRAME(Stat)) \
+		{ \
+			const FName StatName = GET_STATFNAME(Stat); \
+			FThreadStats::AddMessage(StatName, EStatOperation::Subtract, SubtractAmount); \
+			TRACE_STAT_ADD(StatName, -SubtractAmount); \
+		} \
+	} \
+}
+#define DEC_DWORD_STAT_BY(Stat,Amount) \
+{\
+	int64 SubtractAmount = int64(Amount); \
+	if (SubtractAmount != 0) \
+	{ \
+		if (FThreadStats::IsCollectingData() || !GET_STATISEVERYFRAME(Stat)) \
+		{ \
+			const FName StatName = GET_STATFNAME(Stat); \
+			FThreadStats::AddMessage(StatName, EStatOperation::Subtract, SubtractAmount); \
+			TRACE_STAT_ADD(StatName, -SubtractAmount); \
+		} \
+	} \
+}
+#define DEC_MEMORY_STAT_BY(Stat,Amount) \
+{\
+	int64 SubtractAmount = int64(Amount); \
+	if (SubtractAmount != 0) \
+	{ \
+		if (FThreadStats::IsCollectingData() || !GET_STATISEVERYFRAME(Stat)) \
+		{ \
+			const FName StatName = GET_STATFNAME(Stat); \
+			FThreadStats::AddMessage(StatName, EStatOperation::Subtract, SubtractAmount); \
+			TRACE_STAT_ADD(StatName, -SubtractAmount); \
+		} \
+	} \
+}
+#define SET_MEMORY_STAT(Stat,Value) \
+{\
+	if (FThreadStats::IsCollectingData() || !GET_STATISEVERYFRAME(Stat)) \
+	{ \
+		const FName StatName = GET_STATFNAME(Stat); \
+		int64 SetValue = int64(Value); \
+		FThreadStats::AddMessage(StatName, EStatOperation::Set, SetValue); \
+		TRACE_STAT_SET(StatName, SetValue); \
+	} \
+}
+#define SET_DWORD_STAT(Stat,Value) \
+{\
+	if (FThreadStats::IsCollectingData() || !GET_STATISEVERYFRAME(Stat)) \
+	{ \
+		const FName StatName = GET_STATFNAME(Stat); \
+		int64 SetValue = int64(Value); \
+		FThreadStats::AddMessage(StatName, EStatOperation::Set, SetValue); \
+		TRACE_STAT_SET(StatName, SetValue); \
+	} \
+}
+#define SET_FLOAT_STAT(Stat,Value) \
+{\
+	if (FThreadStats::IsCollectingData() || !GET_STATISEVERYFRAME(Stat)) \
+	{ \
+		const FName StatName = GET_STATFNAME(Stat); \
+		double SetValue = double(Value); \
+		FThreadStats::AddMessage(StatName, EStatOperation::Set, SetValue); \
+		TRACE_STAT_SET(StatName, SetValue); \
+	} \
+}
 
-FORCEINLINE void StatsPrimaryEnableAdd(int32 Value = 1)
-{
+#define STAT_ADD_CUSTOMMESSAGE_NAME(Stat,Value) \
+{\
+	FThreadStats::AddMessage(GET_STATFNAME(Stat), EStatOperation::SpecialMessageMarker, FName(Value));\
 }
-FORCEINLINE void StatsPrimaryEnableSubtract(int32 Value = 1)
-{
+#define STAT_ADD_CUSTOMMESSAGE_PTR(Stat,Value) \
+{\
+	FThreadStats::AddMessage(GET_STATFNAME(Stat), EStatOperation::SpecialMessageMarker, uint64(Value));\
 }
 
+#define SET_CYCLE_COUNTER_FName(Stat,Cycles) \
+{\
+	FThreadStats::AddMessage(Stat, EStatOperation::Set, int64(Cycles), true);\
+}
 
+#define INC_DWORD_STAT_FName(Stat) \
+{\
+	FThreadStats::AddMessage(Stat, EStatOperation::Add, int64(1));\
+	TRACE_STAT_INCREMENT(Stat); \
+}
+#define INC_FLOAT_STAT_BY_FName(Stat, Amount) \
+{\
+	double AddAmount = double(Amount); \
+	if (AddAmount != 0.0) \
+	{ \
+		FThreadStats::AddMessage(Stat, EStatOperation::Add, AddAmount); \
+		TRACE_STAT_ADD(Stat, AddAmount); \
+	} \
+}
+#define INC_DWORD_STAT_BY_FName(Stat, Amount) \
+{\
+	int64 AddAmount = int64(Amount); \
+	if (AddAmount != 0) \
+	{ \
+		FThreadStats::AddMessage(Stat, EStatOperation::Add, AddAmount); \
+		TRACE_STAT_ADD(Stat, AddAmount); \
+	} \
+}
+#define INC_DWORD_STAT_FNAME_BY(Stat, Amount) INC_DWORD_STAT_BY_FName(Stat, Amount)
+#define INC_MEMORY_STAT_BY_FName(Stat, Amount) \
+{\
+	int64 AddAmount = int64(Amount); \
+	if (AddAmount != 0) \
+	{ \
+		FThreadStats::AddMessage(Stat, EStatOperation::Add, AddAmount); \
+		TRACE_STAT_ADD(Stat, AddAmount); \
+	} \
+}
+#define DEC_DWORD_STAT_FName(Stat) \
+{\
+	FThreadStats::AddMessage(Stat, EStatOperation::Subtract, int64(1));\
+	TRACE_STAT_DECREMENT(Stat); \
+}
+#define DEC_FLOAT_STAT_BY_FName(Stat,Amount) \
+{\
+	double SubtractAmount = double(Amount); \
+	if (SubtractAmount != 0.0) \
+	{ \
+		FThreadStats::AddMessage(Stat, EStatOperation::Subtract, SubtractAmount); \
+		TRACE_STAT_ADD(Stat, -SubtractAmount); \
+	} \
+}
+#define DEC_DWORD_STAT_BY_FName(Stat,Amount) \
+{\
+	int64 SubtractAmount = int64(Amount); \
+	if (SubtractAmount != 0) \
+	{ \
+		FThreadStats::AddMessage(Stat, EStatOperation::Subtract, SubtractAmount); \
+		TRACE_STAT_ADD(Stat, -SubtractAmount); \
+	} \
+}
+#define DEC_DWORD_STAT_FNAME_BY(Stat,Amount) DEC_DWORD_STAT_BY_FName(Stat,Amount)
+#define DEC_MEMORY_STAT_BY_FName(Stat,Amount) \
+{\
+	int64 SubtractAmount = int64(Amount); \
+	if (SubtractAmount != 0) \
+	{ \
+		FThreadStats::AddMessage(Stat, EStatOperation::Subtract, SubtractAmount); \
+		TRACE_STAT_ADD(Stat, -SubtractAmount); \
+	} \
+}
+#define SET_MEMORY_STAT_FName(Stat,Value) \
+{\
+	int64 SetValue = int64(Value); \
+	FThreadStats::AddMessage(Stat, EStatOperation::Set, SetValue); \
+	TRACE_STAT_SET(Stat, SetValue); \
+}
+#define SET_DWORD_STAT_FName(Stat,Value) \
+{\
+	int64 SetValue = int64(Value); \
+	FThreadStats::AddMessage(Stat, EStatOperation::Set, SetValue); \
+	TRACE_STAT_SET(Stat, SetValue); \
+}
+#define SET_FLOAT_STAT_FName(Stat,Value) \
+{\
+	double SetValue = double(Value); \
+	FThreadStats::AddMessage(Stat, EStatOperation::Set, SetValue); \
+	TRACE_STAT_SET(Stat, SetValue); \
+}
 
-#if ENABLE_STATNAMEDEVENTS
+#elif UE_USE_LIGHTWEIGHT_STATS // STATS
+// STATS == 0
 
-#if PLATFORM_USES_ANSI_STRING_FOR_EXTERNAL_PROFILING
-#define ANSI_TO_PROFILING(x) x
-#else
-#define ANSI_TO_PROFILING(x) TEXT(x)
-#endif
+// ----------------------------------------------------------------------------
+// Lightweight Stats
+// ----------------------------------------------------------------------------
+// This defines a simpler and lower overhead version of the stat macros.
+// These will not emit any data to the runtime stat system and only emit
+// profiling namedevents if enabled.
+// ----------------------------------------------------------------------------
 
 #define SCOPE_CYCLE_COUNTER_TO_TRACE(StatString, StatName, Condition) \
 	TRACE_CPUPROFILER_EVENT_DECLARE(PREPROCESSOR_JOIN(PREPROCESSOR_JOIN(__Decl_, StatName), __LINE__)); \
@@ -289,47 +530,40 @@ FORCEINLINE void StatsPrimaryEnableSubtract(int32 Value = 1)
 
 
 #elif USE_LIGHTWEIGHT_STATS_FOR_HITCH_DETECTION && USE_HITCH_DETECTION
-extern CORE_API TSAN_ATOMIC(bool) GHitchDetected;
+// STATS == 0
+// UE_USE_LIGHTWEIGHT_STATS == 0
 
-class FLightweightStatScope
-{
-	const TCHAR* StatString;
-public:
-	FORCEINLINE FLightweightStatScope(const TCHAR* InStat)
-	{
-		StatString = GHitchDetected ? nullptr : InStat;
-	}
-
-	FORCEINLINE ~FLightweightStatScope()
-	{
-		if (GHitchDetected && StatString)
-		{
-			ReportHitch();
-		}
-	}
-
-	CORE_API void ReportHitch();
-};
+// ----------------------------------------------------------------------------
+// Hitch Detection
+// ----------------------------------------------------------------------------
+// This replaces the stat macros with a very lightweight scoped struct to create 
+// a rough approximation of the current callstack for hitch tracking purposes.
+// This let's us avoid traversing the actual callstack at runtime during a hitch which is very slow.
+// ----------------------------------------------------------------------------
 
 #define DECLARE_SCOPE_CYCLE_COUNTER(CounterName,Stat,GroupId) \
-	FLightweightStatScope LightweightStatScope_##Stat(TEXT(#Stat));
+	UE::Stats::FHitchTrackingStatScope HitchTrackingStatScope_##Stat(ANSI_TO_PROFILING(#Stat));
 
 #define QUICK_SCOPE_CYCLE_COUNTER(Stat) \
-	FLightweightStatScope LightweightStatScope_##Stat(TEXT(#Stat));
+	UE::Stats::FHitchTrackingStatScope HitchTrackingStatScope_##Stat(ANSI_TO_PROFILING(#Stat));
 
 #define SCOPE_CYCLE_COUNTER(Stat) \
-	FLightweightStatScope LightweightStatScope_##Stat(TEXT(#Stat));
+	UE::Stats::FHitchTrackingStatScope HitchTrackingStatScope_##Stat(ANSI_TO_PROFILING(#Stat));
 
 #define SCOPE_CYCLE_COUNTER_STATID(StatId) \
-	FLightweightStatScope LightweightStatScope_##Stat(TEXT("Lightweight StatId Scope"));
+	UE::Stats::FHitchTrackingStatScope HitchTrackingStatScope_##Stat(ANSI_TO_PROFILING("Lightweight StatId Scope"));
 
 #define CONDITIONAL_SCOPE_CYCLE_COUNTER(Stat,bCondition) \
-	FLightweightStatScope LightweightStatScope_##Stat(bCondition ? TEXT(#Stat) : nullptr);
+	UE::Stats::FHitchTrackingStatScope HitchTrackingStatScope_##Stat(bCondition ? ANSI_TO_PROFILING(#Stat) : nullptr);
 
 #define RETURN_QUICK_DECLARE_CYCLE_STAT(StatId,GroupId) return TStatId();
 #define GET_STATID(Stat) (TStatId())
 
 #else
+// STATS == 0
+// UE_USE_LIGHTWEIGHT_STATS == 0
+// (USE_LIGHTWEIGHT_STATS_FOR_HITCH_DETECTION && USE_HITCH_DETECTION) == 0
+
 #define SCOPE_CYCLE_COUNTER(Stat)
 #define SCOPE_CYCLE_COUNTER_STATID(StatId)
 #define QUICK_SCOPE_CYCLE_COUNTER(Stat)
@@ -337,9 +571,12 @@ public:
 #define CONDITIONAL_SCOPE_CYCLE_COUNTER(Stat,bCondition)
 #define RETURN_QUICK_DECLARE_CYCLE_STAT(StatId,GroupId) return TStatId();
 #define GET_STATID(Stat) (TStatId())
+#endif // STATS
 
-#endif
-
+// ----------------------------------------------------------------------------
+// Default Macro Definitions
+// ----------------------------------------------------------------------------
+#if !STATS
 #define SCOPE_SECONDS_ACCUMULATOR(Stat)
 #define SCOPE_MS_ACCUMULATOR(Stat)
 #define DEFINE_STAT(Stat)
@@ -402,115 +639,56 @@ public:
 
 #define GET_STATFNAME(Stat) (FName())
 #define GET_STATDESCRIPTION(Stat) (nullptr)
-
-#endif
-
-
-/** Helper class used to generate dynamic stat ids. */
-struct FDynamicStats
-{
-	/**
-	* Create a new stat id and registers it with the stats system.
-	* This is the only way to create dynamic stat ids at runtime.
-	* Can be used only with FScopeCycleCounters.
-	*
-	* Store the created stat id.
-	* Expensive method, avoid calling that method every frame.
-	*
-	* Example:
-	*	FDynamicStats::CreateStatId<STAT_GROUP_TO_FStatGroup( STATGROUP_UObjects )>( FString::Printf(TEXT("MyDynamicStat_%i"),Index) )
-	*/
-	template< typename TStatGroup >
-	static TStatId CreateStatId( const FString& StatNameOrDescription )
-	{
-#if	STATS
-		return CreateStatIdInternal<TStatGroup>( FName( *StatNameOrDescription ), EStatDataType::ST_int64, true);
-#else
-		return TStatId();
-#endif // STATS
-	}
-
-	template< typename TStatGroup >
-	static TStatId CreateStatIdInt64(const FString& StatNameOrDescription, bool bIsAccumulator = false)
-	{
-#if	STATS
-		return CreateStatIdInternal<TStatGroup>(FName(*StatNameOrDescription), EStatDataType::ST_int64, false, !bIsAccumulator);
-#else
-		return TStatId();
-#endif // STATS
-	}
-
-	template< typename TStatGroup >
-	static TStatId CreateStatIdDouble(const FString& StatNameOrDescription, bool bIsAccumulator=false)
-	{
-#if	STATS
-		return CreateStatIdInternal<TStatGroup>(FName(*StatNameOrDescription), EStatDataType::ST_double, false, !bIsAccumulator);
-#else
-		return TStatId();
-#endif // STATS
-	}
-
-	template< typename TStatGroup >
-	static TStatId CreateStatId(const FName StatNameOrDescription, bool IsTimer = true)
-	{
-#if	STATS
-		return CreateStatIdInternal<TStatGroup>(StatNameOrDescription, EStatDataType::ST_int64, IsTimer);
-#else
-		return TStatId();
-#endif // STATS
-	}
-
-	template< typename TStatGroup >
-	static TStatId CreateMemoryStatId(const FString& StatNameOrDescription, FPlatformMemory::EMemoryCounterRegion MemRegion=FPlatformMemory::MCR_Physical)
-	{
-#if	STATS
-		return CreateMemoryStatId<TStatGroup>(FName(*StatNameOrDescription), MemRegion);
-#else
-		return TStatId();
-#endif // STATS
-	}
-
-	template< typename TStatGroup >
-	static TStatId CreateMemoryStatId(const FName StatNameOrDescription, FPlatformMemory::EMemoryCounterRegion MemRegion=FPlatformMemory::MCR_Physical)
-	{
-#if	STATS
-		FStartupMessages::Get().AddMetadata(StatNameOrDescription, *StatNameOrDescription.ToString(),
-			TStatGroup::GetGroupName(),
-			TStatGroup::GetGroupCategory(),
-			TStatGroup::GetDescription(),
-			false, EStatDataType::ST_int64, false, false, MemRegion);
-
-		TStatId StatID = IStatGroupEnableManager::Get().GetHighPerformanceEnableForStat(StatNameOrDescription,
-			TStatGroup::GetGroupName(),
-			TStatGroup::GetGroupCategory(),
-			TStatGroup::DefaultEnable,
-			false, EStatDataType::ST_int64, *StatNameOrDescription.ToString(), false, false, MemRegion);
-
-		return StatID;
-#else
-		return TStatId();
-#endif // STATS
-	}
-
-#if	STATS
-private: // private since this can only be declared if STATS is defined, due to EStatDataType in signature
-	template< typename TStatGroup >
-	static TStatId CreateStatIdInternal(const FName StatNameOrDescription, EStatDataType::Type Type, bool IsTimer, bool bClearEveryFrame=true)
-	{
-
-		FStartupMessages::Get().AddMetadata(StatNameOrDescription, nullptr,
-			TStatGroup::GetGroupName(),
-			TStatGroup::GetGroupCategory(),
-			TStatGroup::GetDescription(),
-			bClearEveryFrame, Type, IsTimer, false);
-
-		TStatId StatID = IStatGroupEnableManager::Get().GetHighPerformanceEnableForStat(StatNameOrDescription,
-			TStatGroup::GetGroupName(),
-			TStatGroup::GetGroupCategory(),
-			TStatGroup::DefaultEnable,
-			bClearEveryFrame, Type, nullptr, IsTimer, false);
-
-		return StatID;
-	}
-#endif // STATS
-};
+#endif // !STATS
+
+#include "Stats/GlobalStats.inl"
+
+// Deprecated includes from Stats2.h
+#if UE_ENABLE_INCLUDE_ORDER_DEPRECATED_IN_5_6
+// These should be included directly instead of through this header.
+#include "Stats/HitchTrackingStatScope.h"
+#include "Stats/StatsCommand.h"
+#include "Stats/StatsSystem.h"
+#include "Stats/StatsSystemTypes.h"
+#include "Stats/ThreadIdleStats.h"
+
+// Includes from Stats2.h
+#include "Containers/Array.h"
+#include "Containers/ChunkedArray.h"
+#include "Containers/ContainerAllocationPolicies.h"
+#include "Containers/LockFreeList.h"
+#include "Containers/UnrealString.h"
+#include "CoreGlobals.h"
+#include "CoreTypes.h"
+#include "Delegates/Delegate.h"
+#include "Delegates/DelegateBase.h"
+#include "HAL/CriticalSection.h"
+#include "HAL/LowLevelMemTracker.h"
+#include "HAL/PlatformCrt.h"
+#include "HAL/PlatformMemory.h"
+#include "HAL/PlatformMisc.h"
+#include "HAL/PlatformTLS.h"
+#include "HAL/PlatformTime.h"
+#include "HAL/ThreadSafeCounter.h"
+#include "HAL/ThreadSingleton.h"
+#include "HAL/UnrealMemory.h"
+#include "Math/Color.h"
+#include "Math/NumericLimits.h"
+#include "Misc/AssertionMacros.h"
+#include "Misc/Build.h"
+#include "Misc/CString.h"
+#include "Misc/EnumClassFlags.h"
+#include "ProfilingDebugging/CpuProfilerTrace.h"
+#include "ProfilingDebugging/MiscTrace.h"
+#include "StatsCommon.h"
+#include "StatsTrace.h"
+#include "Templates/Atomic.h"
+#include "Templates/TypeCompatibleBytes.h"
+#include "Templates/UniquePtr.h"
+#include "Templates/UnrealTemplate.h"
+#include "Trace/Detail/Channel.h"
+#include "Trace/Detail/Channel.inl"
+#include "Trace/Trace.h"
+#include "UObject/NameTypes.h"
+#include "UObject/UnrealNames.h"
+#endif // UE_ENABLE_INCLUDE_ORDER_DEPRECATED_IN_5_6
\ No newline at end of file
