Unify Stats headers

diff --git a/Engine/Source/Runtime/Core/Public/Stats/Stats2.h b/Engine/Source/Runtime/Core/Public/Stats/Stats2.h
--- a/Engine/Source/Runtime/Core/Public/Stats/Stats2.h
+++ b/Engine/Source/Runtime/Core/Public/Stats/Stats2.h
@@ -2,2469 +2,9 @@
 
 #pragma once
 
-#include "Containers/Array.h"
-#include "Containers/ChunkedArray.h"
-#include "Containers/ContainerAllocationPolicies.h"
-#include "Containers/LockFreeList.h"
-#include "Containers/UnrealString.h"
-#include "CoreGlobals.h"
-#include "CoreTypes.h"
-#include "Delegates/Delegate.h"
-#include "Delegates/DelegateBase.h"
-#include "HAL/CriticalSection.h"
-#include "HAL/LowLevelMemTracker.h"
-#include "HAL/PlatformCrt.h"
-#include "HAL/PlatformMemory.h"
-#include "HAL/PlatformMisc.h"
-#include "HAL/PlatformTLS.h"
-#include "HAL/PlatformTime.h"
-#include "HAL/ThreadSafeCounter.h"
-#include "HAL/ThreadSingleton.h"
-#include "HAL/UnrealMemory.h"
-#include "Math/Color.h"
-#include "Math/NumericLimits.h"
-#include "Misc/AssertionMacros.h"
-#include "Misc/Build.h"
-#include "Misc/CString.h"
-#include "Misc/EnumClassFlags.h"
-#include "Misc/SourceLocation.h"
-#include "ProfilingDebugging/CpuProfilerTrace.h"
-#include "ProfilingDebugging/MiscTrace.h"
-#include "StatsCommon.h"
-#include "StatsTrace.h"
-#include "Templates/Atomic.h"
-#include "Templates/TypeCompatibleBytes.h"
-#include "Templates/UniquePtr.h"
-#include "Templates/UnrealTemplate.h"
-#include "Trace/Detail/Channel.h"
-#include "Trace/Detail/Channel.inl"
-#include "Trace/Trace.h"
-#include "UObject/NameTypes.h"
-#include "UObject/UnrealNames.h"
+#include "Misc/CoreMiscDefines.h"
 
-class FOutputDevice;
-class FScopeCycleCounter;
-class FThreadStats;
-struct TStatId;
-template <typename T> struct TIsPODType;
+UE_DEPRECATED_HEADER(5.6, "All definitions have been moved to Stats.h, please use that instead.")
 
-/**
-* This is thread-private information about the thread idle stats, which we always collect, even in final builds
-*/
-class CORE_API FThreadIdleStats : public TThreadSingleton<FThreadIdleStats>
-{
-	friend class TThreadSingleton<FThreadIdleStats>;
+#include "Stats/Stats.h"
 
-	FThreadIdleStats()
-		: Waits(0)
-		, WaitsCriticalPath(0)
-		, IsCriticalPathCounter(1)
-		, bInIdleScope(false)
-	{}
-
-public:
-
-	/** Total cycles we waited for sleep or event. **/
-	uint32 Waits;
-
-	/** Total cycles we waited for sleep or event on the critical path. **/
-	uint32 WaitsCriticalPath;
-
-	int IsCriticalPathCounter;
-	bool bInIdleScope;
-
-	static void BeginCriticalPath()
-	{
-		FThreadIdleStats::Get().IsCriticalPathCounter++;
-	}
-
-	static void EndCriticalPath()
-	{
-		FThreadIdleStats::Get().IsCriticalPathCounter--;
-	}
-
-	struct FScopeNonCriticalPath
-	{
-		FScopeNonCriticalPath()
-		{
-			FThreadIdleStats::Get().IsCriticalPathCounter--;
-		}
-		~FScopeNonCriticalPath()
-		{
-			FThreadIdleStats::Get().IsCriticalPathCounter++;
-		}
-	};
-
-	bool IsCriticalPath() const 
-	{
-		return IsCriticalPathCounter > 0;
-	}
-
-	void Reset()
-	{
-		Waits = 0;
-		WaitsCriticalPath = 0;
-		IsCriticalPathCounter = 1;
-	}
-
-
-
-	struct FScopeIdle
-	{
-#if defined(DISABLE_THREAD_IDLE_STATS) && DISABLE_THREAD_IDLE_STATS
-		FScopeIdle( bool bInIgnore = false )
-		{}
-#else
-		/** Starting cycle counter. */
-		const uint32 Start;
-
-		/** If true, we ignore this thread idle stats. */
-		const bool bIgnore;
-
-#if CPUPROFILERTRACE_ENABLED
-		FCpuProfilerTrace::FEventScope TraceEventScope;
-#endif
-
-		CORE_API FScopeIdle(bool bInIgnore = false);
-
-		~FScopeIdle()
-		{
-			if( !bIgnore )
-			{
-				FThreadIdleStats& IdleStats = FThreadIdleStats::Get();
-				uint32 CyclesElapsed = FPlatformTime::Cycles() - Start;
-				IdleStats.Waits += CyclesElapsed;
-
-				if (IdleStats.IsCriticalPath())
-				{
-					IdleStats.WaitsCriticalPath += CyclesElapsed;
-				}
-
-				IdleStats.bInIdleScope = false;
-			}
-		}
-#endif
-	};
-};
-
-// Pass a console command directly to the stats system, return true if it is known command, false means it might be a stats command
-CORE_API bool DirectStatsCommand(const TCHAR* Cmd, bool bBlockForCompletion = false, FOutputDevice* Ar = nullptr);
-
-/** Helper struct that contains method available even when the stats are disabled. */
-struct FStats
-{
-	/** Delegate to fire every time we need to advance the stats for the rendering thread. */
-	DECLARE_DELEGATE_ThreeParams( FOnAdvanceRenderingThreadStats, bool /*bDiscardCallstack*/, int64 /*StatsFrame*/, int32 /*PrimaryDisableChangeTagStartFrame*/ );
-
-	/** Advances stats for the current frame. */
-	static CORE_API void AdvanceFrame( bool bDiscardCallstack, const FOnAdvanceRenderingThreadStats& AdvanceRenderingThreadStatsDelegate = FOnAdvanceRenderingThreadStats() );
-
-	/** Advances stats for commandlets, only valid if the command line has the proper token. @see HasStatsForCommandletsToken */
-	static CORE_API void TickCommandletStats();
-
-	/**
-	* @return true, if the command line has the LoadTimeStatsForCommandlet or LoadTimeFileForCommandlet token which enables stats in the commandlets.
-	* !!!CAUTION!!! You need to manually advance stats frame in order to maintain the data integrity and not to leak the memory.
-	*/
-	static CORE_API bool EnabledForCommandlet();
-
-	/**
-	* @return true, if the command line has the LoadTimeStatsForCommandlet token which enables LoadTimeStats equivalent for commandlets.
-	* All collected stats will be dumped to the log file at the end of running the specified commandlet.
-	*/
-	static CORE_API bool HasLoadTimeStatsForCommandletToken();
-
-	/**
-	* @return true, if the command line has the LoadTimeFileForCommandlet token which enables LoadTimeFile equivalent for commandlets.
-	*/
-	static CORE_API bool HasLoadTimeFileForCommandletToken();
-
-	/** Current game thread stats frame. */
-	static CORE_API TAtomic<int32> GameThreadStatsFrame;
-};
-
-enum class EStatFlags : uint8
-{
-	None            = 0,
-	ClearEveryFrame = 1 << 0,
-	CycleStat       = 1 << 1,
-	Verbose         = 1 << 2, // Profiling scopes for this stat will no generate a trace event by default. See GShouldEmitVerboseNamedEvents.
-};
-ENUM_CLASS_FLAGS(EStatFlags);
-
-#if STATS
-
-struct TStatIdData
-{
-	bool IsNone() const
-	{
-		FMinimalName LocalName = Name.Load(EMemoryOrder::Relaxed);
-		return LocalName.IsNone();
-	}
-
-	/** Name of the active stat; stored as a minimal name to minimize the data size */
-	TAtomic<FMinimalName> Name;
-
-	/** const ANSICHAR* pointer to a string describing the stat */
-	TUniquePtr<WIDECHAR[]> StatDescriptionWide;
-
-	/** const WIDECHAR* pointer to a string describing the stat */
-	TUniquePtr<ANSICHAR[]> StatDescriptionAnsi;
-};
-
-struct TStatId
-{
-	FORCEINLINE TStatId()
-		: StatIdPtr(&TStatId_NAME_None)
-	{
-	}
-	FORCEINLINE TStatId(TStatIdData const* InStatIdPtr)
-		: StatIdPtr(InStatIdPtr)
-	{
-	}
-	FORCEINLINE bool IsValidStat() const
-	{
-		return !IsNone();
-	}
-	FORCEINLINE bool IsNone() const
-	{
-		return StatIdPtr->IsNone();
-	}
-	FORCEINLINE TStatIdData const* GetRawPointer() const
-	{
-		return StatIdPtr;
-	}
-
-	FORCEINLINE FMinimalName GetMinimalName(EMemoryOrder MemoryOrder) const
-	{
-		return StatIdPtr->Name.Load(MemoryOrder);
-	}
-
-	FORCEINLINE FName GetName() const
-	{
-		return MinimalNameToName(StatIdPtr->Name);
-	}
-
-	FORCEINLINE static const TStatIdData& GetStatNone()
-	{
-		return TStatId_NAME_None;
-	}
-
-	/**
-	 * @return a stat description as an ansi string.
-	 * StatIdPtr must point to a valid FName pointer.
-	 * @see FStatGroupEnableManager::GetHighPerformanceEnableForStat
-	 */
-	FORCEINLINE const ANSICHAR* GetStatDescriptionANSI() const
-	{
-		return StatIdPtr->StatDescriptionAnsi.Get();
-	}
-
-	/**
-	 * @return a stat description as a wide string.
-	 * StatIdPtr must point to a valid FName pointer.
-	 * @see FStatGroupEnableManager::GetHighPerformanceEnableForStat
-	 */
-	FORCEINLINE const WIDECHAR* GetStatDescriptionWIDE() const
-	{
-		return StatIdPtr->StatDescriptionWide.Get();
-	}
-
-	FORCEINLINE bool operator==(TStatId Other) const
-	{
-		return StatIdPtr == Other.StatIdPtr;
-	}
-
-	FORCEINLINE bool operator!=(TStatId Other) const
-	{
-		return StatIdPtr != Other.StatIdPtr;
-	}
-
-	friend uint32 GetTypeHash(TStatId StatId)
-	{
-		return GetTypeHash(StatId.StatIdPtr);
-	}
-
-private:
-	/** NAME_None. */
-	CORE_API static TStatIdData TStatId_NAME_None;
-
-	/**
-	 *	Holds a pointer to the stat long name if enabled, or to the NAME_None if disabled.
-	 *	@see FStatGroupEnableManager::EnableStat
-	 *	@see FStatGroupEnableManager::DisableStat
-	 *	
-	 *	Next pointer points to the ansi string with a stat description
-	 *	Next pointer points to the wide string with a stat description
-	 *	@see FStatGroupEnableManager::GetHighPerformanceEnableForStat 
-	 */
-	TStatIdData const* StatIdPtr;
-};
-
-/**
- * For packet messages, this indicated what sort of thread timing we use. 
- * Game and Other use CurrentGameFrame, Renderer is CurrentRenderFrame
- */
-namespace EThreadType
-{
-	enum Type
-	{
-		Invalid,
-		Game,
-		Renderer,
-		Other,
-	};
-}
-
-
-/**
- * What the type of the payload is
- */
-struct EStatDataType
-{
-	enum Type
-	{
-		Invalid,
-		/** Not defined. */
-		ST_None,
-		/** int64. */
-		ST_int64,
-		/** double. */
-		ST_double,
-		/** FName. */
-		ST_FName,
-		/** Memory pointer, stored as uint64. */
-		ST_Ptr,
-
-		Num,
-		Mask = 0x7,
-		Shift = 0,
-		NumBits = 3
-	};
-};
-
-/**
- * The operation being performed by this message
- */
-struct EStatOperation
-{
-	enum Type
-	{
-		Invalid,
-		/** Indicates metadata message. */
-		SetLongName,
-		/** Special message for advancing the stats frame from the game thread. */
-		AdvanceFrameEventGameThread,    
-		/** Special message for advancing the stats frame from the render thread. */
-		AdvanceFrameEventRenderThread,  
-		/** Indicates begin of the cycle scope. */
-		CycleScopeStart,
-		/** Indicates end of the cycle scope. */
-		CycleScopeEnd,
-		/** This is not a regular stat operation, but just a special message marker to determine that we encountered a special data in the stat file. */
-		SpecialMessageMarker,
-		/** Set operation. */
-		Set,
-		/** Clear operation. */
-		Clear,
-		/** Add operation. */
-		Add,
-		/** Subtract operation. */
-		Subtract,
-
-		// these are special ones for processed data
-		ChildrenStart,
-		ChildrenEnd,
-		Leaf,
-		MaxVal,
-
-		/** This is a memory operation. @see EMemoryOperation. */
-		Memory UE_DEPRECATED(5.3, "Use Trace/MemoryInsights and/or LLM for memory profiling."),
-
-		Num,
-		Mask = 0xf,
-		Shift = EStatDataType::Shift + EStatDataType::NumBits,
-		NumBits = 4
-	};
-};
-
-/**
- * Message flags
- */
-struct EStatMetaFlags
-{
-	enum Type
-	{
-		Invalid							= 0x00,
-		/** this bit is always one and is used for error checking. */
-		DummyAlwaysOne					= 0x01,  
-		/** if true, then this message contains all meta data as well and the name is long. NOT USED. */
-		HasLongNameAndMetaInfo			= 0x02,
-		/** if true, then this message contains and int64 cycle or IsPackedCCAndDuration. */
-		IsCycle							= 0x04,	
-		/** if true, then this message contains a memory stat. */
-		IsMemory						= 0x08, 
-		/** if true, then this is actually two uint32s, the cycle count and the call count, see FromPackedCallCountDuration_Duration. */
-		IsPackedCCAndDuration			= 0x10,
-		/** if true, then this stat is cleared every frame. */
-		ShouldClearEveryFrame			= 0x20, 
-		/** used only on disk on on the wire, indicates that we serialized the FName string. */
-		SendingFName					= 0x40,
-
-		Num								= 0x80,
-		Mask							= 0xff,
-		Shift = EStatOperation::Shift + EStatOperation::NumBits,
-		NumBits = 8
-	};
-};
-
-//@todo merge these two after we have redone the user end
-/**
- * Wrapper for memory region
- */
-struct EMemoryRegion
-{
-	enum Type
-	{
-		Invalid = FPlatformMemory::MCR_Invalid,
-
-		Num = FPlatformMemory::MCR_MAX,
-		Mask = 0xf,
-		Shift = EStatMetaFlags::Shift + EStatMetaFlags::NumBits,
-		NumBits = 4
-	};
-	static_assert(FPlatformMemory::MCR_MAX < (1 << NumBits), "Need to expand memory region field.");
-};
-
-/** Memory operation for STAT_Memory_AllocPtr. */
-enum class UE_DEPRECATED(5.3, "Use Trace/MemoryInsights and/or LLM for memory profiling.") EMemoryOperation : uint8
-{
-	/** Invalid. */
-	Invalid,	
-	/** Alloc. */
-	Alloc,
-	/** Free. */
-	Free,
-	/** Realloc. */
-	Realloc,
-
-	Num,
-	Mask = 0x7,
-	/** Every allocation is aligned to 8 or 16 bytes, so we have 3/4 bits free to use. */
-	NumBits = 3,
-};
-
-/**
- * A few misc final bit packing computations
- */
-namespace EStatAllFields
-{
-	enum Type
-	{
-		NumBits = EMemoryRegion::Shift + EMemoryRegion::NumBits,
-		StartShift = 28 - NumBits,
-	};
-}
-
-static_assert(EStatAllFields::StartShift > 0, "Too many stat fields.");
-
-FORCEINLINE int64 ToPackedCallCountDuration(uint32 CallCount, uint32 Duration)
-{
-	return (int64(CallCount) << 32) | Duration;
-}
-
-FORCEINLINE uint32 FromPackedCallCountDuration_CallCount(int64 Both)
-{
-	return uint32(Both >> 32);
-}
-
-FORCEINLINE uint32 FromPackedCallCountDuration_Duration(int64 Both)
-{
-	return uint32(Both & MAX_uint32);
-}
-
-/**
- * Helper class that stores an FName and all meta information in 8 bytes. Kindof icky.
- */
-class FStatNameAndInfo
-{
-	/**
-	 * Store name and number separately in case UE_FNAME_OUTLINE_NUMBER is set, so the high bits of the Number are used for other fields.
-	 */
-	FNameEntryId Index;
-	int32 Number;
-
-public:
-	FORCEINLINE_STATS FStatNameAndInfo()
-	{
-	}
-
-	/**
-	 * Build from a raw FName
-	 */
-	FORCEINLINE_STATS FStatNameAndInfo(FName Other, bool bAlreadyHasMeta)
-		: Index(Other.GetComparisonIndex())
-		, Number(Other.GetNumber())
-	{
-		if (!bAlreadyHasMeta)
-		{
-			// ok, you can't have numbered stat FNames too large
-			checkStats(!(Number >> EStatAllFields::StartShift));
-			Number |= EStatMetaFlags::DummyAlwaysOne << (EStatMetaFlags::Shift + EStatAllFields::StartShift);
-		}
-		CheckInvariants();
-	}
-
-	/**
-	 * Build with stat metadata
-	 */
-	FORCEINLINE_STATS FStatNameAndInfo(FName InStatName, char const* InGroup, char const* InCategory, TCHAR const* InDescription, EStatDataType::Type InStatType, bool bShouldClearEveryFrame, bool bCycleStat, bool bSortByName, FPlatformMemory::EMemoryCounterRegion MemoryRegion = FPlatformMemory::MCR_Invalid)
-	{
-		FName LongName = ToLongName(InStatName, InGroup, InCategory, InDescription, bSortByName);
-		Index = LongName.GetComparisonIndex();
-		Number = LongName.GetNumber();
-
-		// ok, you can't have numbered stat FNames too large
-		checkStats(!(Number >> EStatAllFields::StartShift));
-		Number |= (EStatMetaFlags::DummyAlwaysOne | EStatMetaFlags::HasLongNameAndMetaInfo) << (EStatMetaFlags::Shift + EStatAllFields::StartShift);
-
-		SetField<EStatDataType>(InStatType);
-		SetFlag(EStatMetaFlags::ShouldClearEveryFrame, bShouldClearEveryFrame);
-		SetFlag(EStatMetaFlags::IsCycle, bCycleStat);
-		if (MemoryRegion != FPlatformMemory::MCR_Invalid)
-		{
-			SetFlag(EStatMetaFlags::IsMemory, true);
-			SetField<EMemoryRegion>(EMemoryRegion::Type(MemoryRegion));
-		}
-
-		CheckInvariants();
-	}
-
-	/**
-	 * Internal use, used by the deserializer
-	 */
-	FORCEINLINE_STATS void SetNumberDirect(int32 InNumber)
-	{
-		Number = InNumber;
-	}
-
-	/**
-	* Internal use, used by the serializer
-	 */
-	FORCEINLINE_STATS int32 GetRawNumber() const
-	{
-		CheckInvariants();
-		return Number;
-	}
-
-	/**
-	 * Internal use by FStatsThreadState to force an update to the long name
-	 */
-	FORCEINLINE_STATS void SetRawName(FName RawName)
-	{
-		// ok, you can't have numbered stat FNames too large
-		checkStats(!(RawName.GetNumber() >> EStatAllFields::StartShift));
-		CheckInvariants();
-		int32 LocalNumber = Number;
-		LocalNumber &= ~((1 << EStatAllFields::StartShift) - 1);
-		Index = RawName.GetComparisonIndex();
-		Number = (LocalNumber | RawName.GetNumber());
-	}
-
-	/**
-	 * The encoded FName with the correct, original Number
-	 * The original number usually is 0
-	 */
-	FORCEINLINE_STATS FName GetRawName() const
-	{
-		CheckInvariants();
-		return FName(Index, Index, Number & ((1 << EStatAllFields::StartShift) - 1));
-	}
-
-	/**
-	 * The encoded FName with the encoded, new Number
-	 * The number contains all encoded metadata
-	 */
-	FORCEINLINE_STATS FName GetEncodedName() const
-	{
-		CheckInvariants();
-		return FName(Index, Index, Number);
-	}
-
-	/**
-	 * Expensive! Extracts the shortname if this is a long name or just returns the name
-	 */
-	FORCEINLINE_STATS FName GetShortName() const
-	{
-		CheckInvariants();
-		return GetShortNameFrom(GetRawName());
-	}
-
-	/**
-	 * Expensive! Extracts the group name if this is a long name or just returns none
-	 */
-	FORCEINLINE_STATS FName GetGroupName() const
-	{
-		CheckInvariants();
-		return GetGroupNameFrom(GetRawName());
-	}
-
-	/**
-	 * Expensive! Extracts the group category if this is a long name or just returns none
-	 */
-	FORCEINLINE_STATS FName GetGroupCategory() const
-	{
-		CheckInvariants();
-		return GetGroupCategoryFrom(GetRawName());
-	}
-
-	/**
-	 * Expensive! Extracts the description if this is a long name or just returns the empty string
-	 */
-	FORCEINLINE_STATS FString GetDescription() const
-	{
-		CheckInvariants();
-		return GetDescriptionFrom(GetRawName());
-	}
-	/**
-	 * Expensive! Extracts the sort by name flag
-	 */
-	FORCEINLINE_STATS bool GetSortByName() const
-	{
-		CheckInvariants();
-		return GetSortByNameFrom(GetRawName());
-	}
-
-
-	/**
-	 * Makes sure this object is in good shape
-	 */
-	FORCEINLINE_STATS void CheckInvariants() const
-	{
-		checkStats((Number & (EStatMetaFlags::DummyAlwaysOne << (EStatAllFields::StartShift + EStatMetaFlags::Shift)))
-			&& Index);
-	}
-
-	/**
-	 * returns an encoded field
-	 * @return the field
-	 */
-	template<typename TField>
-	typename TField::Type GetField() const
-	{
-		CheckInvariants();
-		int32 LocalNumber = Number;
-		LocalNumber = (LocalNumber >> (EStatAllFields::StartShift + TField::Shift)) & TField::Mask;
-		checkStats(LocalNumber != TField::Invalid && LocalNumber < TField::Num);
-		return typename TField::Type(LocalNumber);
-	}
-
-	/**
-	 * sets an encoded field
-	 * @param Value, value to set
-	 */
-	template<typename TField>
-	void SetField(typename TField::Type Value)
-	{
-		int32 LocalNumber = Number;
-		CheckInvariants();
-		checkStats(Value < TField::Num && Value != TField::Invalid);
-		LocalNumber &= ~(TField::Mask << (EStatAllFields::StartShift + TField::Shift));
-		LocalNumber |= Value << (EStatAllFields::StartShift + TField::Shift);
-		Number = LocalNumber;
-		CheckInvariants();
-	}
-
-	/**
-	 * returns an encoded flag
-	 * @param Bit, flag to read
-	 */
-	bool GetFlag(EStatMetaFlags::Type Bit) const
-	{
-		int32 LocalNumber = Number;
-		CheckInvariants();
-		checkStats(Bit < EStatMetaFlags::Num && Bit != EStatMetaFlags::Invalid);
-		return !!((LocalNumber >> (EStatAllFields::StartShift + EStatMetaFlags::Shift)) & Bit);
-	}
-
-	/**
-	 * sets an encoded flag
-	 * @param Bit, flag to set
-	 * @param Value, value to set
-	 */
-	void SetFlag(EStatMetaFlags::Type Bit, bool Value)
-	{
-		int32 LocalNumber = Number;
-		CheckInvariants();
-		checkStats(Bit < EStatMetaFlags::Num && Bit != EStatMetaFlags::Invalid);
-		if (Value)
-		{
-			LocalNumber |= (Bit << (EStatAllFields::StartShift + EStatMetaFlags::Shift));
-		}
-		else
-		{
-			LocalNumber &= ~(Bit << (EStatAllFields::StartShift + EStatMetaFlags::Shift));
-		}
-		Number = LocalNumber;
-		CheckInvariants();
-	}
-
-
-	/**
-	 * Builds a long name from the three parts
-	 * @param InStatName, Short name
-	 * @param InGroup, Group name
-	 * @param InCategory, Category name
-	 * @param InDescription, Description
-	 * @param InSortByName, Whether this stats need to be sorted by name
-	 * @return the packed FName
-	 */
-	CORE_API static FName ToLongName(FName InStatName, char const* InGroup, char const* InCategory, TCHAR const* InDescription, bool InSortByName);
-	CORE_API static FName GetShortNameFrom(FName InLongName);
-	CORE_API static FName GetGroupNameFrom(FName InLongName);
-	CORE_API static FName GetGroupCategoryFrom(FName InLongName);
-	CORE_API static FString GetDescriptionFrom(FName InLongName);
-	CORE_API static bool GetSortByNameFrom(FName InLongName);
-	
-};
-
-
-/** Union for easier debugging. */
-union UStatData
-{
-private:
-	/** For ST_double. */
-	double	Float;
-	/** For ST_int64 and IsCycle or IsMemory. */
-	int64	Cycles;
-	/** For ST_Ptr. */
-	uint64	Ptr;
-	/** ST_int64 and IsPackedCCAndDuration. */
-	uint32	CCAndDuration[2];
-	/** For FName. */
-	CORE_API const FString GetName() const
-	{
-		return FName::SafeString(FNameEntryId::FromUnstableInt(static_cast<uint32>(Cycles)));
-	}
-};
-
-/**
-* 16 byte stat message. Everything is a message
-*/
-struct FStatMessage
-{
-	/**
-	* Generic payload
-	*/
-	enum
-	{
-		DATA_SIZE=8,
-		DATA_ALIGN=8,
-	};
-	union
-	{
-#if	UE_BUILD_DEBUG
-		UStatData								DebugStatData;
-#endif // UE_BUILD_DEBUG
-		TAlignedBytes<DATA_SIZE, DATA_ALIGN>	StatData;
-	};
-
-	/**
-	* Name and the meta info.
-	*/
-	FStatNameAndInfo						NameAndInfo;
-
-	FStatMessage()
-	{
-	}
-
-	/**
-	* Build a meta data message
-	*/
-	FStatMessage(FName InStatName, EStatDataType::Type InStatType, char const* InGroup, char const* InCategory, TCHAR const* InDescription, bool bShouldClearEveryFrame, bool bCycleStat, bool bSortByName, FPlatformMemory::EMemoryCounterRegion MemoryRegion = FPlatformMemory::MCR_Invalid)
-		: NameAndInfo(InStatName, InGroup, InCategory, InDescription, InStatType, bShouldClearEveryFrame, bCycleStat, bSortByName, MemoryRegion)
-	{
-		NameAndInfo.SetField<EStatOperation>(EStatOperation::SetLongName);
-	}
-
-	explicit FORCEINLINE_STATS FStatMessage(FStatNameAndInfo InStatName)
-		: NameAndInfo(InStatName)
-	{
-	}
-
-	/**
-	* Clock operation
-	*/
-	FORCEINLINE_STATS FStatMessage(FName InStatName, EStatOperation::Type InStatOperation)
-		: NameAndInfo(InStatName, true)
-	{
-		NameAndInfo.SetField<EStatOperation>(InStatOperation);
-		checkStats(NameAndInfo.GetField<EStatDataType>() == EStatDataType::ST_int64);
-		checkStats(NameAndInfo.GetFlag(EStatMetaFlags::IsCycle) == true);
-
-		// these branches are FORCEINLINE_STATS of constants in almost all cases, so they disappear
-		if (InStatOperation == EStatOperation::CycleScopeStart || InStatOperation == EStatOperation::CycleScopeEnd)
-		{
-			GetValue_int64()= int64(FPlatformTime::Cycles());
-		}
-		else
-		{
-			checkStats(0);
-		}
-	}
-
-	/**
-	* int64 operation
-	*/
-	FORCEINLINE_STATS FStatMessage(FName InStatName, EStatOperation::Type InStatOperation, int64 Value, bool bIsCycle)
-		: NameAndInfo(InStatName, true)
-	{
-		NameAndInfo.SetField<EStatOperation>(InStatOperation);
-		checkStats(NameAndInfo.GetField<EStatDataType>() == EStatDataType::ST_int64);
-		checkStats(NameAndInfo.GetFlag(EStatMetaFlags::IsCycle) == bIsCycle);
-		GetValue_int64() = Value;
-	}
-
-	/**
-	* double operation
-	*/
-	FORCEINLINE_STATS FStatMessage(FName InStatName, EStatOperation::Type InStatOperation, double Value, bool)
-		: NameAndInfo(InStatName, true)
-	{
-		NameAndInfo.SetField<EStatOperation>(InStatOperation);
-		checkStats(NameAndInfo.GetField<EStatDataType>() == EStatDataType::ST_double);
-		checkStats(NameAndInfo.GetFlag(EStatMetaFlags::IsCycle) == false);
-		GetValue_double() = Value;
-	}
-
-	/**
-	* name operation
-	*/
-	FORCEINLINE_STATS FStatMessage(FName InStatName, EStatOperation::Type InStatOperation, FName Value, bool)
-		: NameAndInfo(InStatName, true)
-	{
-		NameAndInfo.SetField<EStatOperation>(InStatOperation);
-		checkStats(NameAndInfo.GetField<EStatDataType>() == EStatDataType::ST_FName);
-		checkStats(NameAndInfo.GetFlag(EStatMetaFlags::IsCycle) == false);
-		GetValue_FMinimalName() = NameToMinimalName(Value);
-	}
-
-	/**
-	 * Ptr operation
-	 */
-	FORCEINLINE_STATS FStatMessage(FName InStatName, EStatOperation::Type InStatOperation, uint64 Value, bool)
-		: NameAndInfo(InStatName, true)
-	{
-		NameAndInfo.SetField<EStatOperation>(InStatOperation);
-		checkStats(NameAndInfo.GetField<EStatDataType>() == EStatDataType::ST_Ptr);
-		checkStats(NameAndInfo.GetFlag(EStatMetaFlags::IsCycle) == false);
-		GetValue_Ptr() = Value;
-	}
-
-	/**
-	* Clear any data type
-	*/
-	FORCEINLINE_STATS void Clear()
-	{
-		static_assert(sizeof(uint64) == DATA_SIZE, "Bad clear.");
-		*(int64*)&StatData = 0;
-	}
-
-	/**
-	* Payload retrieval and setting methods
-	*/
-	FORCEINLINE_STATS int64& GetValue_int64()
-	{
-		static_assert(sizeof(int64) <= DATA_SIZE && alignof(int64) <= DATA_ALIGN, "Bad data for stat message.");
-		checkStats(NameAndInfo.GetField<EStatDataType>() == EStatDataType::ST_int64);
-		return *(int64*)&StatData;
-	}
-	FORCEINLINE_STATS int64 GetValue_int64() const
-	{
-		checkStats(NameAndInfo.GetField<EStatDataType>() == EStatDataType::ST_int64);
-		return *(int64 const*)&StatData;
-	}
-
-	FORCEINLINE_STATS uint64& GetValue_Ptr()
-	{
-		static_assert(sizeof(uint64) <= DATA_SIZE && alignof(uint64) <= DATA_ALIGN, "Bad data for stat message.");
-		checkStats(NameAndInfo.GetField<EStatDataType>() == EStatDataType::ST_Ptr);
-		return *(uint64*)&StatData;
-	}
-	FORCEINLINE_STATS uint64 GetValue_Ptr() const
-	{
-		checkStats(NameAndInfo.GetField<EStatDataType>() == EStatDataType::ST_Ptr);
-		return *(uint64 const*)&StatData;
-	}
-
-	FORCEINLINE_STATS int64 GetValue_Duration() const
-	{
-		checkStats(NameAndInfo.GetFlag(EStatMetaFlags::IsCycle) && NameAndInfo.GetField<EStatDataType>() == EStatDataType::ST_int64);
-		if (NameAndInfo.GetFlag(EStatMetaFlags::IsPackedCCAndDuration))
-		{
-			return FromPackedCallCountDuration_Duration(*(int64 const*)&StatData);
-		}
-		return *(int64 const*)&StatData;
-	}
-
-	FORCEINLINE_STATS uint32 GetValue_CallCount() const
-	{
-		checkStats(NameAndInfo.GetFlag(EStatMetaFlags::IsPackedCCAndDuration) && NameAndInfo.GetFlag(EStatMetaFlags::IsCycle) && NameAndInfo.GetField<EStatDataType>() == EStatDataType::ST_int64);
-		return FromPackedCallCountDuration_CallCount(*(int64 const*)&StatData);
-	}
-
-	FORCEINLINE_STATS double& GetValue_double()
-	{
-		static_assert(sizeof(double) <= DATA_SIZE && alignof(double) <= DATA_ALIGN, "Bad data for stat message.");
-		checkStats(NameAndInfo.GetField<EStatDataType>() == EStatDataType::ST_double);
-		return *(double*)&StatData;
-	}
-
-	FORCEINLINE_STATS double GetValue_double() const
-	{
-		checkStats(NameAndInfo.GetField<EStatDataType>() == EStatDataType::ST_double);
-		return *(double const*)&StatData;
-	}
-
-	FORCEINLINE_STATS FMinimalName& GetValue_FMinimalName()
-	{
-		static_assert(sizeof(FMinimalName) <= DATA_SIZE && alignof(FMinimalName) <= DATA_ALIGN, "Bad data for stat message.");
-		checkStats(NameAndInfo.GetField<EStatDataType>() == EStatDataType::ST_FName);
-		return *(FMinimalName*)&StatData;
-	}
-
-	FORCEINLINE_STATS FMinimalName GetValue_FMinimalName() const
-	{
-		checkStats(NameAndInfo.GetField<EStatDataType>() == EStatDataType::ST_FName);
-		return *(FMinimalName const*)&StatData;
-	}
-
-	FORCEINLINE_STATS FName GetValue_FName() const
-	{
-		checkStats(NameAndInfo.GetField<EStatDataType>() == EStatDataType::ST_FName);
-		return MinimalNameToName(*(FMinimalName const*)&StatData);
-	}
-};
-template<> struct TIsPODType<FStatMessage> { enum { Value = true }; };
-
-CORE_API void GetPermanentStats(TArray<FStatMessage>& OutStats);
-
-/**
- *	Based on FStatMessage, but supports more than 8 bytes of stat data storage.
- */
-template< typename TEnum >
-struct TStatMessage
-{
-	typedef TEnum TStructEnum;
-	static constexpr int32 EnumCount = TEnum::Num;
-
-	/**
-	* Generic payload
-	*/
-	enum
-	{
-		DATA_SIZE=8*EnumCount,
-		DATA_ALIGN=8,
-	};
-
-	union
-	{
-#if	UE_BUILD_DEBUG
-		UStatData								DebugStatData[EnumCount];
-#endif // UE_BUILD_DEBUG
-		TAlignedBytes<DATA_SIZE, DATA_ALIGN>	StatData;
-	};
-
-	/**
-	* Name and the meta info.
-	*/
-	FStatNameAndInfo							NameAndInfo;
-
-	TStatMessage()
-	{}
-
-	/**
-	* Copy constructor from FStatMessage
-	*/
-	explicit FORCEINLINE_STATS TStatMessage(const FStatMessage& Other)
-		: NameAndInfo(Other.NameAndInfo)
-	{
-		// Reset data type and clear all fields.
-		NameAndInfo.SetField<EStatDataType>( EStatDataType::ST_None );
-		Clear();
-	}
-
-	/** Assignment operator for FStatMessage. */
-	TStatMessage& operator=(const FStatMessage& Other)
-	{
-		NameAndInfo = Other.NameAndInfo;
-		// Reset data type and clear all fields.
-		NameAndInfo.SetField<EStatDataType>( EStatDataType::ST_None );
-		Clear();
-		return *this;
-	}
-
-	/** Fixes stat data type for all fields. */
-	void FixStatData( const EStatDataType::Type NewType )
-	{
-		const EStatDataType::Type OldType = NameAndInfo.GetField<EStatDataType>();
-		if( OldType != NewType )
-		{
-			// Convert from the old type to the new type.
-			if( OldType == EStatDataType::ST_int64 && NewType == EStatDataType::ST_double )
-			{
-				// Get old values.
-				int64 OldValues[TEnum::Num];
-				for( int32 FieldIndex = 0; FieldIndex < EnumCount; ++FieldIndex )
-				{
-					OldValues[FieldIndex] = GetValue_int64((typename TEnum::Type)FieldIndex);
-				}
-				
-				// Set new stat data type
-				NameAndInfo.SetField<EStatDataType>(NewType);
-				for( int32 FieldIndex = 0; FieldIndex < EnumCount; ++FieldIndex )
-				{
-					GetValue_double((typename TEnum::Type)FieldIndex) = (double)OldValues[FieldIndex];
-				}
-			}
-			else if( OldType == EStatDataType::ST_double && NewType == EStatDataType::ST_int64 )
-			{
-				// Get old values.
-				double OldValues[TEnum::Num];
-				for( int32 FieldIndex = 0; FieldIndex < EnumCount; ++FieldIndex )
-				{
-					OldValues[FieldIndex] = GetValue_double((typename TEnum::Type)FieldIndex);
-				}
-
-				// Set new stat data type
-				NameAndInfo.SetField<EStatDataType>(NewType);
-				for( int32 FieldIndex = 0; FieldIndex < EnumCount; ++FieldIndex )
-				{
-					GetValue_int64((typename TEnum::Type)FieldIndex) = (int64)OldValues[FieldIndex];
-				}
-			}
-		}
-	}
-
-	/**
-	* Clear any data type
-	*/
-	FORCEINLINE_STATS void Clear()
-	{
-		static_assert(sizeof(uint64) == DATA_SIZE / EnumCount, "Bad clear.");
-
-		for( int32 FieldIndex = 0; FieldIndex < EnumCount; ++FieldIndex )
-		{
-			*((int64*)&StatData+FieldIndex) = 0;
-		}
-	}
-
-	/**
-	* Payload retrieval and setting methods
-	*/
-	FORCEINLINE_STATS int64& GetValue_int64( typename TEnum::Type Index )
-	{
-		static_assert(sizeof(int64) <= DATA_SIZE && alignof(int64) <= DATA_ALIGN, "Bad data for stat message.");
-		checkStats(NameAndInfo.GetField<EStatDataType>() == EStatDataType::ST_int64);
-		checkStats(Index<EnumCount);
-		int64& Value = *((int64*)&StatData+(uint32)Index);
-		return Value;
-	}
-	FORCEINLINE_STATS int64 GetValue_int64( typename TEnum::Type Index ) const
-	{
-		checkStats(NameAndInfo.GetField<EStatDataType>() == EStatDataType::ST_int64);
-		checkStats(Index<EnumCount);
-		const int64 Value = *((int64*)&StatData+(uint32)Index);
-		return Value;
-	}
-
-	FORCEINLINE_STATS int64 GetValue_Duration( typename TEnum::Type Index ) const
-	{
-		checkStats(NameAndInfo.GetFlag(EStatMetaFlags::IsCycle) && NameAndInfo.GetField<EStatDataType>() == EStatDataType::ST_int64);
-		checkStats(Index<EnumCount);
-		if (NameAndInfo.GetFlag(EStatMetaFlags::IsPackedCCAndDuration))
-		{
-			const uint32 Value = FromPackedCallCountDuration_Duration(*((int64 const*)&StatData+(uint32)Index));
-			return Value;
-		}
-		const int64 Value = *((int64 const*)&StatData+(uint32)Index);
-		return Value;
-	}
-
-	FORCEINLINE_STATS uint32 GetValue_CallCount( typename TEnum::Type Index ) const
-	{
-		checkStats(NameAndInfo.GetFlag(EStatMetaFlags::IsPackedCCAndDuration) && NameAndInfo.GetFlag(EStatMetaFlags::IsCycle) && NameAndInfo.GetField<EStatDataType>() == EStatDataType::ST_int64);
-		checkStats(Index<EnumCount);
-		const uint32 Value = FromPackedCallCountDuration_CallCount(*((int64 const*)&StatData+(uint32)Index));
-		return Value;
-	}
-
-	FORCEINLINE_STATS double& GetValue_double( typename TEnum::Type Index )
-	{
-		static_assert(sizeof(double) <= DATA_SIZE && alignof(double) <= DATA_ALIGN, "Bad data for stat message.");
-		checkStats(Index<EnumCount);
-		checkStats(NameAndInfo.GetField<EStatDataType>() == EStatDataType::ST_double);
-		double& Value = *((double*)&StatData+(uint32)Index);
-		return Value;
-	}
-
-	FORCEINLINE_STATS double GetValue_double( typename TEnum::Type Index ) const
-	{
-		checkStats(NameAndInfo.GetField<EStatDataType>() == EStatDataType::ST_double);
-		checkStats(Index<EnumCount);
-		const double Value = *((double const*)&StatData+(uint32)Index);
-		return Value;
-	}
-
-	FORCEINLINE_STATS FName GetShortName() const
-	{
-		return NameAndInfo.GetShortName();
-	}
-
-	FORCEINLINE_STATS FString GetDescription() const
-	{
-		return NameAndInfo.GetDescription();
-	}
-};
-
-/** Enumerates fields of the FComplexStatMessage. */
-struct EComplexStatField
-{
-	enum Type
-	{
-		/** Summed inclusive time. */
-		IncSum,
-		/** Average inclusive time. */
-		IncAve,
-		/** Maximum inclusive time. */
-		IncMax,
-		/** Minimum inclusive time. */
-		IncMin,
-		/** Summed exclusive time. */
-		ExcSum,
-		/** Average exclusive time. */
-		ExcAve,
-		/** Maximum exclusive time. */
-		ExcMax,
-		/** Minimum exclusive time. */
-		ExcMin,
-
-		/** Number of enumerates. */
-		Num,
-	};
-};
-
-/**
- *	This type of stat message holds data defined by associated enumeration @see EComplexStatField
- *	By default any of these messages contain a valid data, so check before accessing the data.
- */
-typedef TStatMessage<EComplexStatField> FComplexStatMessage;
-
-template<> struct TIsPODType<FComplexStatMessage> { enum { Value = true }; };
-
-
-struct EStatMessagesArrayConstants
-{
-	enum
-	{
-#if WITH_EDITOR
-		MESSAGES_CHUNK_SIZE = 4*1024, // Smaller chunks prevent malloc deadlocks when running in the editor with TBB allocator
-#else
-		MESSAGES_CHUNK_SIZE = 64*1024,
-#endif
-	};
-};
-typedef TChunkedArray<FStatMessage,(uint32)EStatMessagesArrayConstants::MESSAGES_CHUNK_SIZE> FStatMessagesArray;
-
-/**
-* A stats packet. Sent between threads. Includes and array of messages and some information about the thread. 
-*/
-struct FStatPacket
-{
-	/** Assigned later, this is the frame number this packet is for. @see FStatsThreadState::ScanForAdvance or FThreadStats::DetectAndUpdateCurrentGameFrame */
-	int64 Frame;
-	/** ThreadId this packet came from **/
-	uint32 ThreadId;
-	/** type of thread this packet came from **/
-	EThreadType::Type ThreadType;
-	/** true if this packet has broken callstacks **/
-	bool bBrokenCallstacks;
-	/** messages in this packet **/
-	FStatMessagesArray StatMessages;
-	/** Size we presize the message buffer to, currently the max of what we have seen for the last PRESIZE_MAX_NUM_ENTRIES. **/
-	TArray<int32> StatMessagesPresize;
-
-	/** constructor **/
-	FStatPacket()
-		: Frame(1)
-		, ThreadId(0)
-		, ThreadType(EThreadType::Invalid)
-		, bBrokenCallstacks(false)
-	{
-	}
-
-	/** Copy constructor. !!!CAUTION!!! does not copy the data **/
-	FStatPacket(FStatPacket const& Other)
-		: Frame(Other.Frame)
-		, ThreadId(Other.ThreadId)
-		, ThreadType(Other.ThreadType)
-		, bBrokenCallstacks(false)
-		, StatMessagesPresize(Other.StatMessagesPresize)
-	{
-	}
-
-	/** Initializes thread related properties for the stats packet. */
-	void SetThreadProperties()
-	{
-		ThreadId = FPlatformTLS::GetCurrentThreadId();
-		PRAGMA_DISABLE_DEPRECATION_WARNINGS
-		if (ThreadId == GGameThreadId)
-		{
-			ThreadType = EThreadType::Game;
-		}
-		else if (ThreadId == GRenderThreadId)
-		{
-			ThreadType = EThreadType::Renderer;
-		}
-		else
-		{
-			ThreadType = EThreadType::Other;
-		}
-		PRAGMA_ENABLE_DEPRECATION_WARNINGS
-	}
-
-
-	void AssignFrame( int64 InFrame )
-	{
-		Frame = InFrame;
-	}
-};
-
-/** Helper struct used to monitor the scope of the message. */
-struct FStatMessageLock
-{
-	FStatMessageLock(int32& InMessageScope)
-		: MessageScope(InMessageScope)
-	{
-		MessageScope++;
-	}
-
-	~FStatMessageLock()
-	{
-		MessageScope--;
-	}
-
-protected:
-	int32& MessageScope;
-};
-
-/** Preallocates FThreadStats to avoid dynamic memory allocation. */
-struct FThreadStatsPool
-{
-private:
-	enum
-	{
-		/**
-		 *	Number of elements preallocated in the pool.
-		 *  Should be enough to handle all threads.
-		 *  Increase if you encounter a crash.
-		 */
-		NUM_ELEMENTS_IN_POOL = 512
-	};
-
-	/** Lock free pool of FThreadStats instances. */
-	TLockFreePointerListUnordered<FThreadStats, 0> Pool;
-
-public:
-	/** Default constructor. */
-	CORE_API FThreadStatsPool();
-
-	/** Singleton accessor. */
-	CORE_API static FThreadStatsPool& Get();
-
-	/** Gets an instance from the pool and call the default constructor on it. */
-	CORE_API FThreadStats* GetFromPool();
-
-	/** Return an instance to the pool. */
-	CORE_API void ReturnToPool(FThreadStats* Instance);
-};
-
-/** Fake type to distinguish constructors. */
-enum EConstructor
-{
-	FOR_POOL
-};
-
-/**
-* This is thread-private information about the stats we are acquiring. Pointers to these objects are held in TLS.
-*/
-class FThreadStats : FNoncopyable
-{
-	friend class FStatsMallocProfilerProxy;
-	friend class FStatsThreadState;
-	friend class FStatsThread;
-	friend struct FThreadStatsPool;
-
-	/** Used to control when we are collecting stats. User of the stats system increment and decrement this counter as they need data. **/
-	CORE_API static FThreadSafeCounter PrimaryEnableCounter;
-	/** Every time bPrimaryEnable changes, we update this. This is used to determine frames that have complete data. **/
-	CORE_API static FThreadSafeCounter PrimaryEnableUpdateNumber;
-	/** while bPrimaryEnable (or other things affecting stat collection) is chaning, we lock this. This is used to determine frames that have complete data. **/
-	CORE_API static FThreadSafeCounter PrimaryDisableChangeTagLock;
-	/** TLS slot that holds a FThreadStats. **/
-	CORE_API static uint32 TlsSlot;
-	/** Computed by CheckEnable, the current "primary control" for stats collection, based on PrimaryEnableCounter and a few other things. **/
-	CORE_API static bool bPrimaryEnable;
-	/** Set to permanently disable the stats system. **/
-	CORE_API static bool bPrimaryDisableForever;
-	/** True if we running in the raw stats mode, all stats processing is disabled, captured stats messages are written in timely manner, memory overhead is minimal. */
-	CORE_API static bool bIsRawStatsActive;
-
-	/** The data we are eventually going to send to the stats thread. **/
-	FStatPacket Packet;
-
-	/** Current game frame for this thread stats. */
-	int32 CurrentGameFrame;
-
-	/** Tracks current stack depth for cycle counters. **/
-	int32 ScopeCount;
-
-	/** Tracks current stack depth for cycle counters. **/
-	int32 bWaitForExplicitFlush;
-
-	/** 
-	 * Tracks the progress of adding a new memory stat message and prevents from using the memory profiler in the scope.
-	 * Mostly related to ignoring all memory allocations in AddStatMessage().
-	 * Memory usage of the stats messages is handled by the STAT_StatMessagesMemory.
-	 */
-	int32 MemoryMessageScope;
-
-	/** We have a scoped cycle counter in the FlushRawStats which causes an infinite recursion, ReentranceGuard will solve that issue. */
-	bool bReentranceGuard;
-
-	/** Tracks current stack depth for cycle counters. **/
-	bool bSawExplicitFlush;
-
-	/** Gathers information about the current thread and sets up the TLS value. **/
-	CORE_API FThreadStats();
-
-	/** Constructor used for the pool. */
-	CORE_API FThreadStats(EConstructor);
-
-	void SendMessage_Async(FStatPacket* ToSend);
-
-public:
-	/** Checks the TLS for a thread packet and if it isn't found, it makes a new one. **/
-	static FORCEINLINE_STATS FThreadStats* GetThreadStats()
-	{
-		FThreadStats* Stats = (FThreadStats*)FPlatformTLS::GetTlsValue(TlsSlot);
-		if (!Stats)
-		{
-			Stats = FThreadStatsPool::Get().GetFromPool();
-		}
-		return Stats;
-	}
-
-	/** This should be called when conditions have changed such that stat collection may now be enabled or not **/
-	static CORE_API void CheckEnable();
-
-	/**
-	 *	Checks if the game frame has changed and updates the current game frame.
-	 *	Used by the flushing mechanism to optimize the memory usage, only valid for packets from other threads.
-	 *	@return true, if the frame has changed
-	 */
-	FORCEINLINE_STATS bool DetectAndUpdateCurrentGameFrame()
-	{
-		if( Packet.ThreadType == EThreadType::Other )
-		{
-			FPlatformMisc::MemoryBarrier();
-			int32 Frame = FStats::GameThreadStatsFrame;
-			const bool bFrameHasChanged = Frame > CurrentGameFrame;
-			if( bFrameHasChanged )
-			{
-				CurrentGameFrame = Frame;
-				Packet.AssignFrame( CurrentGameFrame );
-				return true;
-			}
-		}
-				
-		return false;
-	}
-
-	/** Maintains the explicit flush. */
-	FORCEINLINE_STATS void UpdateExplicitFlush()
-	{
-		if (Packet.ThreadType != EThreadType::Other && bSawExplicitFlush)
-		{
-			bSawExplicitFlush = false;
-			bWaitForExplicitFlush = 1;
-			ScopeCount++; // prevent sends until the next explicit flush
-		}
-	}
-
-	/** Send any outstanding packets to the stats thread. **/
-	CORE_API void Flush(bool bHasBrokenCallstacks = false, bool bForceFlush = false);
-
-	/** Flushes the regular stats, the realtime stats. */
-	CORE_API void FlushRegularStats(bool bHasBrokenCallstacks, bool bForceFlush);
-
-	/** Flushes the raw stats, low memory and performance overhead, but not realtime. */
-	CORE_API void FlushRawStats(bool bHasBrokenCallstacks = false, bool bForceFlush = false);
-
-	/** Checks the command line whether we want to enable collecting startup stats. */
-	static CORE_API void CheckForCollectingStartupStats();
-
-	FORCEINLINE_STATS void AddStatMessage( const FStatMessage& StatMessage )
-	{
-		LLM_SCOPE(ELLMTag::Stats);
-		FStatMessageLock MessageLock(MemoryMessageScope);
-		Packet.StatMessages.AddElement(StatMessage);
-	}
-
-public:
-	/** This should be called when a thread exits, this deletes FThreadStats from the heap and TLS. **/
-	static void Shutdown()
-	{
-		FThreadStats* Stats = IsThreadingReady() ? (FThreadStats*)FPlatformTLS::GetTlsValue(TlsSlot) : nullptr;
-		if (Stats)
-		{
-			// Send all remaining messages.
-			Stats->Flush(false, true);
-			FPlatformTLS::SetTlsValue(TlsSlot, nullptr);
-			FThreadStatsPool::Get().ReturnToPool(Stats);
-		}
-	}
-
-	/** Clock operation. **/
-	static FORCEINLINE_STATS void AddMessage(FName InStatName, EStatOperation::Type InStatOperation)
-	{
-		checkStats((InStatOperation == EStatOperation::CycleScopeStart || InStatOperation == EStatOperation::CycleScopeEnd));
-		FThreadStats* ThreadStats = GetThreadStats();
-		// these branches are handled by the optimizer
-		if (InStatOperation == EStatOperation::CycleScopeStart)
-		{
-			ThreadStats->ScopeCount++;
-			ThreadStats->AddStatMessage(FStatMessage(InStatName, InStatOperation));
-
-			if( bIsRawStatsActive )
-			{
-				ThreadStats->FlushRawStats();
-			}
-		}
-		else if (InStatOperation == EStatOperation::CycleScopeEnd)
-		{
-			if (ThreadStats->ScopeCount > ThreadStats->bWaitForExplicitFlush)
-			{
-				ThreadStats->AddStatMessage(FStatMessage(InStatName, InStatOperation));
-				ThreadStats->ScopeCount--;
-				if (!ThreadStats->ScopeCount)
-				{
-					ThreadStats->Flush();
-				}
-				else if( bIsRawStatsActive )
-				{
-					ThreadStats->FlushRawStats();
-				}
-			}
-			// else we dumped this frame without closing scope, so we just drop the closes on the floor
-		}
-	}
-
-	/** Any non-clock operation with an ordinary payload. **/
-	template<typename TValue>
-	static FORCEINLINE_STATS void AddMessage(FName InStatName, EStatOperation::Type InStatOperation, TValue Value, bool bIsCycle = false)
-	{
-		if (!InStatName.IsNone() && WillEverCollectData() && IsThreadingReady())
-		{
-			FThreadStats* ThreadStats = GetThreadStats();
-			ThreadStats->AddStatMessage(FStatMessage(InStatName, InStatOperation, Value, bIsCycle));
-			if(!ThreadStats->ScopeCount)
-			{
-				ThreadStats->Flush();
-			}
-			else if( bIsRawStatsActive )
-			{
-				ThreadStats->FlushRawStats();
-			}
-		}
-	}
-
-	/** Pseudo-Memory operation. */
-	template<typename TValue>
-	UE_DEPRECATED(5.3, "Use Trace/MemoryInsights and/or LLM for memory profiling.")
-	FORCEINLINE_STATS void AddMemoryMessage( FName InStatName, TValue Value )
-	{
-#if UE_STATS_MEMORY_PROFILER_ENABLED
-		AddStatMessage(FStatMessage(InStatName, EStatOperation::Memory, Value, false));
-#endif //UE_STATS_MEMORY_PROFILER_ENABLED
-	}
-
-	/** 
-	 * Used to force a flush at the next available opportunity. This is not useful for threads other than the main and render thread. 
-	 * if DiscardCallstack is true, we also dump call stacks, making the next available opportunity at the next stat or stat close.
-	**/
-	static CORE_API void ExplicitFlush(bool DiscardCallstack = false);
-
-	/** Return true if we are currently collecting data **/
-	static FORCEINLINE_STATS bool IsCollectingData()
-	{
-		return bPrimaryEnable;
-	}
-	static FORCEINLINE_STATS bool IsCollectingData(TStatId StatId)
-	{
-		// we don't test StatId for nullptr here because we assume it is non-null. If it is nullptr, that indicates a problem with higher level code.
-		return !StatId.IsNone() && IsCollectingData();
-	}
-
-	/** Return true if we are currently collecting data **/
-	static FORCEINLINE_STATS bool WillEverCollectData()
-	{
-		return !bPrimaryDisableForever;
-	}
-
-	/** Return true if the threading is ready **/
-	static FORCEINLINE_STATS bool IsThreadingReady()
-	{
-		return FPlatformTLS::IsValidTlsSlot(TlsSlot);
-	}
-
-	/** Indicate that you would like the system to begin collecting data, if it isn't already collecting data. Think reference count. **/
-	static FORCEINLINE_STATS void PrimaryEnableAdd(int32 Value = 1)
-	{
-		PrimaryEnableCounter.Add(Value);
-		CheckEnable();
-	}
-
-	/** Indicate that you no longer need stat data, if nobody else needs stat data, then no stat data will be collected. Think reference count. **/
-	static FORCEINLINE_STATS void PrimaryEnableSubtract(int32 Value = 1)
-	{
-		PrimaryEnableCounter.Subtract(Value);
-		CheckEnable();
-	}
-
-	/** Indicate that you no longer need stat data, forever. **/
-	static FORCEINLINE_STATS void PrimaryDisableForever()
-	{
-		bPrimaryDisableForever = true;
-		CheckEnable();
-	}
-
-	/** This is called before we start to change something that will invalidate. **/
-	static FORCEINLINE_STATS void PrimaryDisableChangeTagLockAdd(int32 Value = 1)
-	{
-		PrimaryDisableChangeTagLock.Add(Value);
-		FPlatformMisc::MemoryBarrier();
-		PrimaryEnableUpdateNumber.Increment();
-	}
-
-	/** Indicate that you no longer need stat data, if nobody else needs stat data, then no stat data will be collected. Think reference count. **/
-	static FORCEINLINE_STATS void PrimaryDisableChangeTagLockSubtract(int32 Value = 1)
-	{
-		FPlatformMisc::MemoryBarrier();
-		PrimaryEnableUpdateNumber.Increment();
-		FPlatformMisc::MemoryBarrier();
-		PrimaryDisableChangeTagLock.Subtract(Value);
-	}
-
-	/** Everytime primary enable changes, this number increases. This is used to determine full frames. **/
-	static FORCEINLINE_STATS int32 PrimaryDisableChangeTag()
-	{
-		if (PrimaryDisableChangeTagLock.GetValue())
-		{
-			// while locked we are continually invalid, so we will just keep giving unique numbers
-			return PrimaryEnableUpdateNumber.Increment();
-		}
-		return PrimaryEnableUpdateNumber.GetValue();
-	}
-
-	/** Indicate that you would like the system to begin collecting data, if it isn't already collecting data. Think reference count. **/
-	UE_DEPRECATED(5.1, "Use PrimaryEnableAdd instead")
-	static FORCEINLINE_STATS void MasterEnableAdd(int32 Value = 1)
-	{
-		PrimaryEnableAdd(Value);
-	}
-
-	/** Indicate that you no longer need stat data, if nobody else needs stat data, then no stat data will be collected. Think reference count. **/
-	UE_DEPRECATED(5.1, "Use PrimaryEnableSubtract instead")
-	static FORCEINLINE_STATS void MasterEnableSubtract(int32 Value = 1)
-	{
-		PrimaryEnableSubtract(Value);
-	}
-
-	/** Indicate that you no longer need stat data, forever. **/
-	UE_DEPRECATED(5.1, "Use PrimaryDisableForever instead")
-	static FORCEINLINE_STATS void MasterDisableForever()
-	{
-		PrimaryDisableForever();
-	}
-
-	/** This is called before we start to change something that will invalidate. **/
-	UE_DEPRECATED(5.1, "Use PrimaryDisableChangeTagLockAdd instead")
-	static FORCEINLINE_STATS void MasterDisableChangeTagLockAdd(int32 Value = 1)
-	{
-		PrimaryDisableChangeTagLockAdd(Value);
-	}
-
-	/** Indicate that you no longer need stat data, if nobody else needs stat data, then no stat data will be collected. Think reference count. **/
-	UE_DEPRECATED(5.1, "Use PrimaryDisableChangeTagLockSubtract instead")
-	static FORCEINLINE_STATS void MasterDisableChangeTagLockSubtract(int32 Value = 1)
-	{
-		PrimaryDisableChangeTagLockSubtract(Value);
-	}
-
-	/** Call this if something disrupts data gathering. For example when the render thread is killed, data is abandoned.**/
-	static FORCEINLINE_STATS void FrameDataIsIncomplete()
-	{
-		FPlatformMisc::MemoryBarrier();
-		PrimaryEnableUpdateNumber.Increment();
-		FPlatformMisc::MemoryBarrier();
-	}
-
-	/** Enables the raw stats mode. */
-	static FORCEINLINE_STATS void EnableRawStats() TSAN_SAFE
-	{
-		bIsRawStatsActive = true;
-		FPlatformMisc::MemoryBarrier();
-	}
-
-	/** Disables the raw stats mode. */
-	static FORCEINLINE_STATS void DisableRawStats() TSAN_SAFE
-	{
-		bIsRawStatsActive = false;
-		FPlatformMisc::MemoryBarrier();
-	}
-
-	/** Called by launch engine loop to start the stats thread **/
-	static CORE_API void StartThread();
-	/** Called by launch engine loop to stop the stats thread **/
-	static CORE_API void StopThread();
-	/** Called by the engine loop to make sure the stats thread isn't getting too far behind. **/
-	static CORE_API void WaitForStats();
-};
-
-
-/**
- * This is a utility class for counting the number of cycles during the
- * lifetime of the object. It creates messages for the stats thread.
- */
-class FCycleCounter
-{
-	/** Event types emitted */
-	enum 
-	{
-		NamedEvent			= 1 << 0,
-		TraceEvent			= 1 << 1,
-		ThreadStatsEvent	= 1 << 2,
-	};
-
-	/** Name of the stat, usually a short name **/
-	FName StatId;
-	uint8 EmittedEvent = 0;
-
-public:
-
-	/**
-	 * Pushes the specified stat onto the hierarchy for this thread. Starts
-	 * the timing of the cycles used
-	 */
-	FORCEINLINE_STATS void Start(TStatId InStatId, EStatFlags InStatFlags, bool bAlways = false, UE::FSourceLocation SourceLocation = UE::FSourceLocation::Current())
-	{
-		FMinimalName StatMinimalName = InStatId.GetMinimalName(EMemoryOrder::Relaxed);
-		if (StatMinimalName.IsNone())
-		{
-			return;
-		}
-
-		// Emit named event for active cycle stat.
-		if (GCycleStatsShouldEmitNamedEvents
-			&& (GShouldEmitVerboseNamedEvents || !EnumHasAnyFlags(InStatFlags, EStatFlags::Verbose)))
-		{
-#if PLATFORM_USES_ANSI_STRING_FOR_EXTERNAL_PROFILING
-			FPlatformMisc::BeginNamedEvent(FColor(0), InStatId.GetStatDescriptionANSI());
-#else
-			FPlatformMisc::BeginNamedEvent(FColor(0), InStatId.GetStatDescriptionWIDE());
-#endif
-			EmittedEvent |= NamedEvent;
-
-#if CPUPROFILERTRACE_ENABLED
-			if (UE_TRACE_CHANNELEXPR_IS_ENABLED(CpuChannel))
-			{
-				FName StatName = MinimalNameToName(StatMinimalName);
-				FCpuProfilerTrace::OutputBeginDynamicEventWithId(StatName, InStatId.GetStatDescriptionWIDE(), SourceLocation.GetFileName(), SourceLocation.GetLine());
-				EmittedEvent |= TraceEvent;
-			}
-#endif
-		}
-
-		if ((bAlways && FThreadStats::WillEverCollectData()) || FThreadStats::IsCollectingData())
-		{
-			FName StatName = MinimalNameToName(StatMinimalName);
-			StatId = StatName;
-			FThreadStats::AddMessage(StatName, EStatOperation::CycleScopeStart);
-			EmittedEvent |= ThreadStatsEvent;
-		}
-	}
-
-	FORCEINLINE_STATS void Start(TStatId InStatId, bool bAlways = false, UE::FSourceLocation SourceLocation = UE::FSourceLocation::Current())
-	{
-		Start(InStatId, EStatFlags::None, bAlways, SourceLocation);
-	}
-
-	FORCEINLINE_STATS void StartTrace(const FName Name)
-	{
-#if CPUPROFILERTRACE_ENABLED
-		if (UE_TRACE_CHANNELEXPR_IS_ENABLED(CpuChannel))
-		{
-			FCpuProfilerTrace::OutputBeginDynamicEvent(Name);
-			EmittedEvent |= TraceEvent;
-		}
-#endif
-	}
-
-	FORCEINLINE_STATS void StartTrace(const FName Name, const TCHAR* Desc)
-	{
-#if CPUPROFILERTRACE_ENABLED
-		if (UE_TRACE_CHANNELEXPR_IS_ENABLED(CpuChannel))
-		{
-			FCpuProfilerTrace::OutputBeginDynamicEventWithId(Name, Desc);
-			EmittedEvent |= TraceEvent;
-		}
-#endif
-	}
-
-	/**
-	 * Stops the capturing and stores the result
-	 */
-	FORCEINLINE_STATS void Stop()
-	{
-		if (EmittedEvent & NamedEvent)
-		{
-			FPlatformMisc::EndNamedEvent();
-		}
-
-#if CPUPROFILERTRACE_ENABLED
-		if (EmittedEvent & TraceEvent)
-		{
-			FCpuProfilerTrace::OutputEndEvent();
-		}
-#endif
-
-		if (EmittedEvent & ThreadStatsEvent)
-		{
-			FThreadStats::AddMessage(StatId, EStatOperation::CycleScopeEnd);
-		}
-
-		EmittedEvent = 0;
-	}
-
-	/**
-	 * Stops the capturing and stores the result and resets the stat id.
-	 */
-	FORCEINLINE_STATS void StopAndResetStatId()
-	{
-		Stop();
-		StatId = NAME_None;
-	}
-};
-
-class FSimpleScopeSecondsStat
-{
-public:
-
-	FSimpleScopeSecondsStat(TStatId InStatId, double InScale=1.0)
-		: StartTime(FPlatformTime::Seconds())
-		, StatId(InStatId)
-		, Scale(InScale)
-	{
-
-	}
-
-	virtual ~FSimpleScopeSecondsStat()
-	{
-		double TotalTime = (FPlatformTime::Seconds() - StartTime) * Scale;
-		FThreadStats::AddMessage(StatId.GetName(), EStatOperation::Add, TotalTime);
-	}
-
-private:
-
-	double StartTime;
-	TStatId StatId;
-	double Scale;
-};
-
-/** Manages startup messages, usually to update the metadata. */
-class FStartupMessages
-{
-	friend class FStatsThread;
-
-	TArray64<FStatMessage> DelayedMessages;
-	FCriticalSection CriticalSection;
-
-public:
-	/** Adds a thread metadata. */
-	CORE_API void AddThreadMetadata( const FName InThreadFName, uint32 InThreadID );
-
-	/** Adds a regular metadata. */
-	CORE_API void AddMetadata( FName InStatName, const TCHAR* InStatDesc, const char* InGroupName, const char* InGroupCategory, const TCHAR* InGroupDesc, bool bShouldClearEveryFrame, EStatDataType::Type InStatType, bool bCycleStat, bool bSortByName, FPlatformMemory::EMemoryCounterRegion InMemoryRegion = FPlatformMemory::MCR_Invalid );
-
-	/** Access the singleton. */
-	CORE_API static FStartupMessages& Get();
-};
-
-/**
-* Single interface to control high performance stat disable
-*/
-class IStatGroupEnableManager
-{
-public:
-	/** Return the singleton, must be called from the main thread. **/
-	CORE_API static IStatGroupEnableManager& Get();
-
-	/** virtual destructor. **/
-	virtual ~IStatGroupEnableManager() 
-	{
-	}
-
-	/**
-	 * Returns a pointer to a bool (valid forever) that determines if this group is active
-	 * This should be CACHED. We will get a few calls from different stats and different threads and stuff, but once things are "warmed up", this should NEVER be called.
-	 * This function will also register any stats with the StatsTrace system 
-	 * @param InGroup, group to look up
-	 * @param InCategory, the category the group belongs to
-	 * @param bDefaultEnable, If this is the first time this group has been set up, this sets the default enable value for this group.
-	 * @param bShouldClearEveryFrame, If this is true, this is a memory counter or an accumulator
-	 * @return a pointer to a FName (valid forever) that determines if this group is active
-	 */
-	virtual TStatId GetHighPerformanceEnableForStat(FName StatShortName, const char* InGroup, const char* InCategory, bool bDefaultEnable, bool bShouldClearEveryFrame, EStatDataType::Type InStatType, TCHAR const* InDescription, bool bCycleStat, bool bSortByName, FPlatformMemory::EMemoryCounterRegion MemoryRegion = FPlatformMemory::MCR_Invalid) = 0;
-
-	/**
-	 * Enables or disabled a particular group of stats
-	 * Disabling a memory group, ever, is usually a bad idea
-	 * @param Group, group to look up
-	 * @param Enable, this should be true if we want to collect stats for this group
-	 */
-	virtual void SetHighPerformanceEnableForGroup(FName Group, bool Enable)=0;
-
-	/**
-	 * Enables or disabled all groups of stats
-	 * Disabling a memory group, ever, is usually a bad idea. SO if you disable all groups, you will wreck memory stats usually.
-	 * @param Enable, this should be true if we want to collect stats for all groups
-	 */
-	virtual void SetHighPerformanceEnableForAllGroups(bool Enable)=0;
-
-	/**
-	 * Resets all stats to their default collection state, which was set when they were looked up initially
-	 */
-	virtual void ResetHighPerformanceEnableForAllGroups()=0;
-
-	/**
-	 * Runs a group command
-	 * @param Cmd, Command to run
-	 */
-	virtual void StatGroupEnableManagerCommand(FString const& Cmd)=0;
-
-	/** Updates memory usage. */
-	virtual void UpdateMemoryUsage() = 0;
-};
-
-
-/**
-**********************************************************************************
-*/
-struct FThreadSafeStaticStatBase
-{
-protected:
-	mutable TAtomic<const TStatIdData*> HighPerformanceEnable; // must be uninitialized, because we need atomic initialization
-	CORE_API const TStatIdData* DoSetup(const char* InStatName, const TCHAR* InStatDesc, const char* InGroupName, const char* InGroupCategory, const TCHAR* InGroupDesc, bool bDefaultEnable, bool bShouldClearEveryFrame, EStatDataType::Type InStatType, bool bCycleStat, bool bSortByName, FPlatformMemory::EMemoryCounterRegion InMemoryRegion) const;
-};
-
-template<class TStatData, bool TCompiledIn>
-struct FThreadSafeStaticStatInner : public FThreadSafeStaticStatBase
-{
-	FORCEINLINE_STATS TStatId GetStatId() const
-	{
-		const TStatIdData* LocalHighPerformanceEnable = HighPerformanceEnable.Load(EMemoryOrder::Relaxed);
-		if (!LocalHighPerformanceEnable)
-		{
-			LocalHighPerformanceEnable = DoSetup(TStatData::GetStatName(), TStatData::GetDescription(), TStatData::TGroup::GetGroupName(), TStatData::TGroup::GetGroupCategory(), TStatData::TGroup::GetDescription(), TStatData::TGroup::IsDefaultEnabled(), TStatData::IsClearEveryFrame(), TStatData::GetStatType(), TStatData::IsCycleStat(), TStatData::TGroup::GetSortByName(), TStatData::GetMemoryRegion() );
-		}
-		return TStatId(LocalHighPerformanceEnable);
-	}
-	FORCEINLINE FName GetStatFName() const
-	{
-		return GetStatId().GetName();
-	}
-};
-
-template<class TStatData>
-struct FThreadSafeStaticStatInner<TStatData, false>
-{
-	FORCEINLINE TStatId GetStatId()
-	{
-		return TStatId();
-	}
-	FORCEINLINE FName GetStatFName() const
-	{
-		return FName();
-	}
-};
-
-template<class TStatData>
-struct FThreadSafeStaticStat : public FThreadSafeStaticStatInner<TStatData, TStatData::TGroup::CompileTimeEnable>
-{
-	FThreadSafeStaticStat()
-	{
-		//This call will result in registering the Group if it's compile time enabled. 
-		//It fixes a bug when a StatGroup only has counters that are using the INC_\DEC_ macros. 
-		//Those macros are guarded for the stats collection to be active which prevented the registration of the stat group.
-		//It was not possible to activate the stat group unless another was already active.
-		//Most groups are registered when a FScopeCycleCounter is declared as GetStatId is called as the constructor parameter.
-		FThreadSafeStaticStatInner<TStatData, TStatData::TGroup::CompileTimeEnable>::GetStatId();
-	}
-};
-
-#define DECLARE_STAT_GROUP(Description, StatName, StatCategory, InDefaultEnable, InCompileTimeEnable, InSortByName) \
-struct FStatGroup_##StatName\
-{ \
-	enum \
-	{ \
-		DefaultEnable = InDefaultEnable, \
-		CompileTimeEnable = InCompileTimeEnable, \
-		SortByName = InSortByName \
-	}; \
-	static FORCEINLINE const char* GetGroupName() \
-	{ \
-		return #StatName; \
-	} \
-	static FORCEINLINE const char* GetGroupCategory() \
-	{ \
-		return #StatCategory; \
-	} \
-	static FORCEINLINE const TCHAR* GetDescription() \
-	{ \
-		return Description; \
-	} \
-	static FORCEINLINE bool IsDefaultEnabled() \
-	{ \
-		return (bool)DefaultEnable; \
-	} \
-	static FORCEINLINE bool IsCompileTimeEnable() \
-	{ \
-		return (bool)CompileTimeEnable; \
-	} \
-	static FORCEINLINE bool GetSortByName() \
-	{ \
-		return (bool)SortByName; \
-	} \
-};
-
-#define DECLARE_STAT(Description, StatName, GroupName, StatType, StatFlags, MemoryRegion) \
-struct FStat_##StatName\
-{ \
-	typedef FStatGroup_##GroupName TGroup; \
-	static FORCEINLINE const char* GetStatName() \
-	{ \
-		return #StatName; \
-	} \
-	static FORCEINLINE const TCHAR* GetDescription() \
-	{ \
-		return Description; \
-	} \
-	static FORCEINLINE EStatDataType::Type GetStatType() \
-	{ \
-		return StatType; \
-	} \
-	static FORCEINLINE bool IsClearEveryFrame() \
-	{ \
-		return EnumHasAnyFlags(GetFlags(), EStatFlags::ClearEveryFrame); \
-	} \
-	static FORCEINLINE bool IsCycleStat() \
-	{ \
-		return EnumHasAnyFlags(GetFlags(), EStatFlags::CycleStat); \
-	} \
-	static FORCEINLINE EStatFlags GetFlags() \
-	{ \
-		return StatFlags; \
-	} \
-	static FORCEINLINE FPlatformMemory::EMemoryCounterRegion GetMemoryRegion() \
-	{ \
-		return MemoryRegion; \
-	} \
-};
-
-#define GET_STATID(Stat) (StatPtr_##Stat.GetStatId())
-#define GET_STATFNAME(Stat) (StatPtr_##Stat.GetStatFName())
-#define GET_STATDESCRIPTION(Stat) (FStat_##Stat::GetDescription())
-#define GET_STATISEVERYFRAME(Stat) (FStat_##Stat::IsClearEveryFrame())
-#define GET_STATFLAGS(Stat) (FStat_##Stat::GetFlags())
-
-#define STAT_GROUP_TO_FStatGroup(Group) FStatGroup_##Group
-
-/*-----------------------------------------------------------------------------
-	Local
------------------------------------------------------------------------------*/
-
-#define DEFINE_STAT(Stat) \
-	struct FThreadSafeStaticStat<FStat_##Stat> StatPtr_##Stat;
-
-#define RETURN_QUICK_DECLARE_CYCLE_STAT(StatId,GroupId) \
-	DECLARE_STAT(TEXT(#StatId),StatId,GroupId,EStatDataType::ST_int64, EStatFlags::ClearEveryFrame | EStatFlags::CycleStat, FPlatformMemory::MCR_Invalid); \
-	static DEFINE_STAT(StatId) \
-	return GET_STATID(StatId);
-
-#define QUICK_USE_CYCLE_STAT(StatId,GroupId) [](){ RETURN_QUICK_DECLARE_CYCLE_STAT(StatId, GroupId); }()
-
-#define DECLARE_CYCLE_STAT(CounterName,StatId,GroupId) \
-	DECLARE_STAT(CounterName,StatId,GroupId,EStatDataType::ST_int64, EStatFlags::ClearEveryFrame | EStatFlags::CycleStat, FPlatformMemory::MCR_Invalid); \
-	static DEFINE_STAT(StatId)
-#define DECLARE_CYCLE_STAT_WITH_FLAGS(CounterName,StatId,GroupId,StatFlags) \
-	DECLARE_STAT(CounterName,StatId,GroupId,EStatDataType::ST_int64, (StatFlags) | EStatFlags::ClearEveryFrame | EStatFlags::CycleStat, FPlatformMemory::MCR_Invalid); \
-	static DEFINE_STAT(StatId)
-#define DECLARE_FLOAT_COUNTER_STAT(CounterName,StatId,GroupId) \
-	DECLARE_STAT(CounterName,StatId,GroupId,EStatDataType::ST_double,EStatFlags::ClearEveryFrame, FPlatformMemory::MCR_Invalid); \
-	static DEFINE_STAT(StatId)
-#define DECLARE_DWORD_COUNTER_STAT(CounterName,StatId,GroupId) \
-	DECLARE_STAT(CounterName,StatId,GroupId,EStatDataType::ST_int64,EStatFlags::ClearEveryFrame, FPlatformMemory::MCR_Invalid); \
-	static DEFINE_STAT(StatId)
-#define DECLARE_FLOAT_ACCUMULATOR_STAT(CounterName,StatId,GroupId) \
-	DECLARE_STAT(CounterName,StatId,GroupId,EStatDataType::ST_double, EStatFlags::None, FPlatformMemory::MCR_Invalid); \
-	static DEFINE_STAT(StatId)
-#define DECLARE_DWORD_ACCUMULATOR_STAT(CounterName,StatId,GroupId) \
-	DECLARE_STAT(CounterName,StatId,GroupId,EStatDataType::ST_int64, EStatFlags::None, FPlatformMemory::MCR_Invalid); \
-	static DEFINE_STAT(StatId)
-
-/** FName stat that allows sending a string based data. */
-#define DECLARE_FNAME_STAT(CounterName,StatId,GroupId) \
-	DECLARE_STAT(CounterName,StatId,GroupId,EStatDataType::ST_FName, EStatFlags::None, FPlatformMemory::MCR_Invalid); \
-	static DEFINE_STAT(StatId)
-
-/** This is a fake stat, mostly used to implement memory message or other custom stats that don't easily fit into the system. */
-#define DECLARE_PTR_STAT(CounterName,StatId,GroupId)\
-	DECLARE_STAT(CounterName,StatId,GroupId,EStatDataType::ST_Ptr, EStatFlags::None, FPlatformMemory::MCR_Invalid); \
-	static DEFINE_STAT(StatId)
-
-#define DECLARE_MEMORY_STAT(CounterName,StatId,GroupId) \
-	DECLARE_STAT(CounterName,StatId,GroupId,EStatDataType::ST_int64, EStatFlags::None, FPlatformMemory::MCR_Physical); \
-	static DEFINE_STAT(StatId)
-
-#define DECLARE_MEMORY_STAT_POOL(CounterName,StatId,GroupId,Pool) \
-	DECLARE_STAT(CounterName,StatId,GroupId,EStatDataType::ST_int64, EStatFlags::None, Pool); \
-	static DEFINE_STAT(StatId)
-
-/*-----------------------------------------------------------------------------
-	Extern
------------------------------------------------------------------------------*/
-
-#define DECLARE_CYCLE_STAT_EXTERN(CounterName,StatId,GroupId, APIX) \
-	DECLARE_STAT(CounterName,StatId,GroupId,EStatDataType::ST_int64, EStatFlags::ClearEveryFrame | EStatFlags::CycleStat, FPlatformMemory::MCR_Invalid); \
-	extern APIX DEFINE_STAT(StatId);
-#define DECLARE_CYCLE_STAT_WITH_FLAGS_EXTERN(CounterName,StatId,GroupId,StatFlags, APIX) \
-	DECLARE_STAT(CounterName,StatId,GroupId,EStatDataType::ST_int64, (StatFlags) | EStatFlags::ClearEveryFrame | EStatFlags::CycleStat, FPlatformMemory::MCR_Invalid); \
-	extern APIX DEFINE_STAT(StatId);
-#define DECLARE_FLOAT_COUNTER_STAT_EXTERN(CounterName,StatId,GroupId, API) \
-	DECLARE_STAT(CounterName,StatId,GroupId,EStatDataType::ST_double, EStatFlags::ClearEveryFrame, FPlatformMemory::MCR_Invalid); \
-	extern API DEFINE_STAT(StatId);
-#define DECLARE_DWORD_COUNTER_STAT_EXTERN(CounterName,StatId,GroupId, API) \
-	DECLARE_STAT(CounterName,StatId,GroupId,EStatDataType::ST_int64, EStatFlags::ClearEveryFrame, FPlatformMemory::MCR_Invalid); \
-	extern API DEFINE_STAT(StatId);
-#define DECLARE_FLOAT_ACCUMULATOR_STAT_EXTERN(CounterName,StatId,GroupId, API) \
-	DECLARE_STAT(CounterName,StatId,GroupId,EStatDataType::ST_double, EStatFlags::None, FPlatformMemory::MCR_Invalid); \
-	extern API DEFINE_STAT(StatId);
-#define DECLARE_DWORD_ACCUMULATOR_STAT_EXTERN(CounterName,StatId,GroupId, API) \
-	DECLARE_STAT(CounterName,StatId,GroupId,EStatDataType::ST_int64, EStatFlags::None, FPlatformMemory::MCR_Invalid); \
-	extern API DEFINE_STAT(StatId);
-
-/** FName stat that allows sending a string based data. */
-#define DECLARE_FNAME_STAT_EXTERN(CounterName,StatId,GroupId, API) \
-	DECLARE_STAT(CounterName,StatId,GroupId,EStatDataType::ST_FName, EStatFlags::None, FPlatformMemory::MCR_Invalid); \
-	extern API DEFINE_STAT(StatId);
-
-/** This is a fake stat, mostly used to implement memory message or other custom stats that don't easily fit into the system. */
-#define DECLARE_PTR_STAT_EXTERN(CounterName,StatId,GroupId, API) \
-	DECLARE_STAT(CounterName,StatId,GroupId,EStatDataType::ST_Ptr, EStatFlags::None, FPlatformMemory::MCR_Invalid); \
-	extern API DEFINE_STAT(StatId);
-
-#define DECLARE_MEMORY_STAT_EXTERN(CounterName,StatId,GroupId, API) \
-	DECLARE_STAT(CounterName,StatId,GroupId,EStatDataType::ST_int64, EStatFlags::None, FPlatformMemory::MCR_Physical); \
-	extern API DEFINE_STAT(StatId);
-
-#define DECLARE_MEMORY_STAT_POOL_EXTERN(CounterName,StatId,GroupId,Pool, API) \
-	DECLARE_STAT(CounterName,StatId,GroupId,EStatDataType::ST_int64, EStatFlags::None, Pool); \
-	extern API DEFINE_STAT(StatId);
-
-/** Macro for declaring group factory instances */
-#define DECLARE_STATS_GROUP(GroupDesc, GroupId, GroupCat) \
-	DECLARE_STAT_GROUP(GroupDesc, GroupId, GroupCat, true, true, false);
-
-#define DECLARE_STATS_GROUP_SORTBYNAME(GroupDesc, GroupId, GroupCat) \
-	DECLARE_STAT_GROUP(GroupDesc, GroupId, GroupCat, true, true, true);
-
-#define DECLARE_STATS_GROUP_VERBOSE(GroupDesc, GroupId, GroupCat) \
-	DECLARE_STAT_GROUP(GroupDesc, GroupId, GroupCat, false, true, false);
-
-#define DECLARE_STATS_GROUP_MAYBE_COMPILED_OUT(GroupDesc, GroupId, GroupCat, CompileIn) \
-	DECLARE_STAT_GROUP(GroupDesc, GroupId, GroupCat, false, CompileIn, false);
-
-#define DECLARE_SCOPE_CYCLE_COUNTER(CounterName,Stat,GroupId) \
-	DECLARE_STAT(CounterName,Stat,GroupId,EStatDataType::ST_int64, EStatFlags::ClearEveryFrame | EStatFlags::CycleStat, FPlatformMemory::MCR_Invalid); \
-	static DEFINE_STAT(Stat) \
-	FScopeCycleCounter CycleCount_##Stat(GET_STATID(Stat), GET_STATFLAGS(Stat));
-
-#define QUICK_SCOPE_CYCLE_COUNTER(Stat) \
-	DECLARE_SCOPE_CYCLE_COUNTER(TEXT(#Stat),Stat,STATGROUP_Quick)
-
-#define SCOPE_CYCLE_COUNTER(Stat) \
-	FScopeCycleCounter CycleCount_##Stat(GET_STATID(Stat), GET_STATFLAGS(Stat));
-
-#define SCOPE_CYCLE_COUNTER_STATID(StatId) \
-	FScopeCycleCounter CycleCount_FromStatId(StatId);
-
-#define CONDITIONAL_SCOPE_CYCLE_COUNTER(Stat,bCondition) \
-	FScopeCycleCounter CycleCount_##Stat(bCondition ? GET_STATID(Stat) : TStatId(), GET_STATFLAGS(Stat));
-
-#define SCOPE_SECONDS_ACCUMULATOR(Stat) \
-	FSimpleScopeSecondsStat SecondsAccum_##Stat(GET_STATID(Stat));
-
-#define SCOPE_MS_ACCUMULATOR(Stat) \
-	FSimpleScopeSecondsStat SecondsAccum_##Stat(GET_STATID(Stat), 1000.0);
-
-#define SET_CYCLE_COUNTER(Stat,Cycles) \
-{\
-	if (FThreadStats::IsCollectingData() || !GET_STATISEVERYFRAME(Stat)) \
-		FThreadStats::AddMessage(GET_STATFNAME(Stat), EStatOperation::Set, int64(Cycles), true);\
-}
-
-#define INC_DWORD_STAT(Stat) \
-{\
-	if (FThreadStats::IsCollectingData() || !GET_STATISEVERYFRAME(Stat)) \
-	{ \
-		const FName StatName = GET_STATFNAME(Stat); \
-		FThreadStats::AddMessage(StatName, EStatOperation::Add, int64(1));\
-		TRACE_STAT_INCREMENT(StatName); \
-	} \
-}
-#define INC_FLOAT_STAT_BY(Stat, Amount) \
-{\
-	double AddAmount = double(Amount); \
-	if (AddAmount != 0.0) \
-	{ \
-		if (FThreadStats::IsCollectingData() || !GET_STATISEVERYFRAME(Stat)) \
-		{ \
-			const FName StatName = GET_STATFNAME(Stat); \
-			FThreadStats::AddMessage(StatName, EStatOperation::Add, AddAmount); \
-			TRACE_STAT_ADD(StatName, AddAmount); \
-		} \
-	} \
-}
-#define INC_DWORD_STAT_BY(Stat, Amount) \
-{\
-	int64 AddAmount = int64(Amount); \
-	if (AddAmount != 0) \
-	{ \
-		if (FThreadStats::IsCollectingData() || !GET_STATISEVERYFRAME(Stat)) \
-		{ \
-			const FName StatName = GET_STATFNAME(Stat); \
-			FThreadStats::AddMessage(StatName, EStatOperation::Add, AddAmount); \
-			TRACE_STAT_ADD(StatName, AddAmount); \
-		} \
-	} \
-}
-#define INC_MEMORY_STAT_BY(Stat, Amount) \
-{\
-	int64 AddAmount = int64(Amount); \
-	if (AddAmount != 0) \
-	{ \
-		if (FThreadStats::IsCollectingData() || !GET_STATISEVERYFRAME(Stat)) \
-		{ \
-			const FName StatName = GET_STATFNAME(Stat); \
-			FThreadStats::AddMessage(StatName, EStatOperation::Add, AddAmount); \
-			TRACE_STAT_ADD(StatName, AddAmount); \
-		} \
-	} \
-}
-#define DEC_DWORD_STAT(Stat) \
-{\
-	if (FThreadStats::IsCollectingData() || !GET_STATISEVERYFRAME(Stat)) \
-	{ \
-		const FName StatName = GET_STATFNAME(Stat); \
-		FThreadStats::AddMessage(StatName, EStatOperation::Subtract, int64(1));\
-		TRACE_STAT_DECREMENT(StatName); \
-	} \
-}
-#define DEC_FLOAT_STAT_BY(Stat,Amount) \
-{\
-	double SubtractAmount = double(Amount); \
-	if (SubtractAmount != 0.0) \
-	{ \
-		if (FThreadStats::IsCollectingData() || !GET_STATISEVERYFRAME(Stat)) \
-		{ \
-			const FName StatName = GET_STATFNAME(Stat); \
-			FThreadStats::AddMessage(StatName, EStatOperation::Subtract, SubtractAmount); \
-			TRACE_STAT_ADD(StatName, -SubtractAmount); \
-		} \
-	} \
-}
-#define DEC_DWORD_STAT_BY(Stat,Amount) \
-{\
-	int64 SubtractAmount = int64(Amount); \
-	if (SubtractAmount != 0) \
-	{ \
-		if (FThreadStats::IsCollectingData() || !GET_STATISEVERYFRAME(Stat)) \
-		{ \
-			const FName StatName = GET_STATFNAME(Stat); \
-			FThreadStats::AddMessage(StatName, EStatOperation::Subtract, SubtractAmount); \
-			TRACE_STAT_ADD(StatName, -SubtractAmount); \
-		} \
-	} \
-}
-#define DEC_MEMORY_STAT_BY(Stat,Amount) \
-{\
-	int64 SubtractAmount = int64(Amount); \
-	if (SubtractAmount != 0) \
-	{ \
-		if (FThreadStats::IsCollectingData() || !GET_STATISEVERYFRAME(Stat)) \
-		{ \
-			const FName StatName = GET_STATFNAME(Stat); \
-			FThreadStats::AddMessage(StatName, EStatOperation::Subtract, SubtractAmount); \
-			TRACE_STAT_ADD(StatName, -SubtractAmount); \
-		} \
-	} \
-}
-#define SET_MEMORY_STAT(Stat,Value) \
-{\
-	if (FThreadStats::IsCollectingData() || !GET_STATISEVERYFRAME(Stat)) \
-	{ \
-		const FName StatName = GET_STATFNAME(Stat); \
-		int64 SetValue = int64(Value); \
-		FThreadStats::AddMessage(StatName, EStatOperation::Set, SetValue); \
-		TRACE_STAT_SET(StatName, SetValue); \
-	} \
-}
-#define SET_DWORD_STAT(Stat,Value) \
-{\
-	if (FThreadStats::IsCollectingData() || !GET_STATISEVERYFRAME(Stat)) \
-	{ \
-		const FName StatName = GET_STATFNAME(Stat); \
-		int64 SetValue = int64(Value); \
-		FThreadStats::AddMessage(StatName, EStatOperation::Set, SetValue); \
-		TRACE_STAT_SET(StatName, SetValue); \
-	} \
-}
-#define SET_FLOAT_STAT(Stat,Value) \
-{\
-	if (FThreadStats::IsCollectingData() || !GET_STATISEVERYFRAME(Stat)) \
-	{ \
-		const FName StatName = GET_STATFNAME(Stat); \
-		double SetValue = double(Value); \
-		FThreadStats::AddMessage(StatName, EStatOperation::Set, SetValue); \
-		TRACE_STAT_SET(StatName, SetValue); \
-	} \
-}
-
-#define STAT_ADD_CUSTOMMESSAGE_NAME(Stat,Value) \
-{\
-	FThreadStats::AddMessage(GET_STATFNAME(Stat), EStatOperation::SpecialMessageMarker, FName(Value));\
-}
-#define STAT_ADD_CUSTOMMESSAGE_PTR(Stat,Value) \
-{\
-	FThreadStats::AddMessage(GET_STATFNAME(Stat), EStatOperation::SpecialMessageMarker, uint64(Value));\
-}
-
-#define SET_CYCLE_COUNTER_FName(Stat,Cycles) \
-{\
-	FThreadStats::AddMessage(Stat, EStatOperation::Set, int64(Cycles), true);\
-}
-
-#define INC_DWORD_STAT_FName(Stat) \
-{\
-	FThreadStats::AddMessage(Stat, EStatOperation::Add, int64(1));\
-	TRACE_STAT_INCREMENT(Stat); \
-}
-#define INC_FLOAT_STAT_BY_FName(Stat, Amount) \
-{\
-	double AddAmount = double(Amount); \
-	if (AddAmount != 0.0) \
-	{ \
-		FThreadStats::AddMessage(Stat, EStatOperation::Add, AddAmount); \
-		TRACE_STAT_ADD(Stat, AddAmount); \
-	} \
-}
-#define INC_DWORD_STAT_BY_FName(Stat, Amount) \
-{\
-	int64 AddAmount = int64(Amount); \
-	if (AddAmount != 0) \
-	{ \
-		FThreadStats::AddMessage(Stat, EStatOperation::Add, AddAmount); \
-		TRACE_STAT_ADD(Stat, AddAmount); \
-	} \
-}
-#define INC_DWORD_STAT_FNAME_BY(Stat, Amount) INC_DWORD_STAT_BY_FName(Stat, Amount)
-#define INC_MEMORY_STAT_BY_FName(Stat, Amount) \
-{\
-	int64 AddAmount = int64(Amount); \
-	if (AddAmount != 0) \
-	{ \
-		FThreadStats::AddMessage(Stat, EStatOperation::Add, AddAmount); \
-		TRACE_STAT_ADD(Stat, AddAmount); \
-	} \
-}
-#define DEC_DWORD_STAT_FName(Stat) \
-{\
-	FThreadStats::AddMessage(Stat, EStatOperation::Subtract, int64(1));\
-	TRACE_STAT_DECREMENT(Stat); \
-}
-#define DEC_FLOAT_STAT_BY_FName(Stat,Amount) \
-{\
-	double SubtractAmount = double(Amount); \
-	if (SubtractAmount != 0.0) \
-	{ \
-		FThreadStats::AddMessage(Stat, EStatOperation::Subtract, SubtractAmount); \
-		TRACE_STAT_ADD(Stat, -SubtractAmount); \
-	} \
-}
-#define DEC_DWORD_STAT_BY_FName(Stat,Amount) \
-{\
-	int64 SubtractAmount = int64(Amount); \
-	if (SubtractAmount != 0) \
-	{ \
-		FThreadStats::AddMessage(Stat, EStatOperation::Subtract, SubtractAmount); \
-		TRACE_STAT_ADD(Stat, -SubtractAmount); \
-	} \
-}
-#define DEC_DWORD_STAT_FNAME_BY(Stat,Amount) DEC_DWORD_STAT_BY_FName(Stat,Amount)
-#define DEC_MEMORY_STAT_BY_FName(Stat,Amount) \
-{\
-	int64 SubtractAmount = int64(Amount); \
-	if (SubtractAmount != 0) \
-	{ \
-		FThreadStats::AddMessage(Stat, EStatOperation::Subtract, SubtractAmount); \
-		TRACE_STAT_ADD(Stat, -SubtractAmount); \
-	} \
-}
-#define SET_MEMORY_STAT_FName(Stat,Value) \
-{\
-	int64 SetValue = int64(Value); \
-	FThreadStats::AddMessage(Stat, EStatOperation::Set, SetValue); \
-	TRACE_STAT_SET(Stat, SetValue); \
-}
-#define SET_DWORD_STAT_FName(Stat,Value) \
-{\
-	int64 SetValue = int64(Value); \
-	FThreadStats::AddMessage(Stat, EStatOperation::Set, SetValue); \
-	TRACE_STAT_SET(Stat, SetValue); \
-}
-#define SET_FLOAT_STAT_FName(Stat,Value) \
-{\
-	double SetValue = double(Value); \
-	FThreadStats::AddMessage(Stat, EStatOperation::Set, SetValue); \
-	TRACE_STAT_SET(Stat, SetValue); \
-}
-
-
-/**
- * Unique group identifiers. Note these don't have to defined in this header
- * but they do have to be unique. You're better off defining these in your
- * own headers/cpp files
- */
-DECLARE_STATS_GROUP(TEXT("AI"),STATGROUP_AI, STATCAT_Advanced);
-DECLARE_STATS_GROUP(TEXT("Anim"),STATGROUP_Anim, STATCAT_Advanced);
-DECLARE_STATS_GROUP(TEXT("Async I/O"),STATGROUP_AsyncIO, STATCAT_Advanced);
-DECLARE_STATS_GROUP(TEXT("Audio"), STATGROUP_Audio, STATCAT_Advanced);
-DECLARE_STATS_GROUP(TEXT("Beam Particles"),STATGROUP_BeamParticles, STATCAT_Advanced);
-DECLARE_STATS_GROUP(TEXT("CPU Stalls"), STATGROUP_CPUStalls, STATCAT_Advanced);
-DECLARE_STATS_GROUP(TEXT("Canvas"),STATGROUP_Canvas, STATCAT_Advanced);
-DECLARE_STATS_GROUP(TEXT("Character"), STATGROUP_Character, STATCAT_Advanced);
-DECLARE_STATS_GROUP(TEXT("Collision"),STATGROUP_Collision, STATCAT_Advanced);
-DECLARE_STATS_GROUP(TEXT("CollisionTags"), STATGROUP_CollisionTags, STATCAT_Advanced);
-DECLARE_STATS_GROUP_VERBOSE(TEXT("CollisionVerbose"),STATGROUP_CollisionVerbose, STATCAT_Advanced);
-DECLARE_STATS_GROUP(TEXT("D3D11RHI"),STATGROUP_D3D11RHI, STATCAT_Advanced);
-DECLARE_STATS_GROUP(TEXT("DDC"),STATGROUP_DDC, STATCAT_Advanced);
-DECLARE_STATS_GROUP(TEXT("Default Stat Group"),STATGROUP_Default, STATCAT_Advanced);
-DECLARE_STATS_GROUP(TEXT("Engine"),STATGROUP_Engine, STATCAT_Advanced);
-DECLARE_STATS_GROUP(TEXT("FPS Chart"),STATGROUP_FPSChart, STATCAT_Advanced);
-DECLARE_STATS_GROUP(TEXT("GPU"), STATGROUP_GPU, STATCAT_Advanced);
-DECLARE_STATS_GROUP(TEXT("GPU Particles"),STATGROUP_GPUParticles, STATCAT_Advanced);
-DECLARE_STATS_GROUP(TEXT("Game"),STATGROUP_Game, STATCAT_Advanced);
-DECLARE_STATS_GROUP(TEXT("GPU Defrag"), STATGROUP_GPUDEFRAG, STATCAT_Advanced);
-DECLARE_STATS_GROUP(TEXT("Gnm"),STATGROUP_PS4RHI, STATCAT_Advanced);
-DECLARE_STATS_GROUP(TEXT("Init Views"),STATGROUP_InitViews, STATCAT_Advanced);
-DECLARE_STATS_GROUP(TEXT("Landscape"),STATGROUP_Landscape, STATCAT_Advanced);
-DECLARE_STATS_GROUP(TEXT("Light Rendering"),STATGROUP_LightRendering, STATCAT_Advanced);
-DECLARE_STATS_GROUP(TEXT("LoadTime"), STATGROUP_LoadTime, STATCAT_Advanced);
-DECLARE_STATS_GROUP(TEXT("LoadTimeClass"), STATGROUP_LoadTimeClass, STATCAT_Advanced);
-DECLARE_STATS_GROUP(TEXT("LoadTimeClassCount"), STATGROUP_LoadTimeClassCount, STATCAT_Advanced);
-DECLARE_STATS_GROUP_VERBOSE(TEXT("LoadTimeVerbose"), STATGROUP_LoadTimeVerbose, STATCAT_Advanced);
-DECLARE_STATS_GROUP(TEXT("Media"),STATGROUP_Media, STATCAT_Advanced);
-DECLARE_STATS_GROUP(TEXT("Memory Allocator"),STATGROUP_MemoryAllocator, STATCAT_Advanced);
-DECLARE_STATS_GROUP(TEXT("Memory Platform"),STATGROUP_MemoryPlatform, STATCAT_Advanced);
-DECLARE_STATS_GROUP(TEXT("Memory StaticMesh"),STATGROUP_MemoryStaticMesh, STATCAT_Advanced);
-DECLARE_STATS_GROUP(TEXT("Memory"),STATGROUP_Memory, STATCAT_Advanced);
-DECLARE_STATS_GROUP(TEXT("Mesh Particles"),STATGROUP_MeshParticles, STATCAT_Advanced);
-DECLARE_STATS_GROUP(TEXT("Metal"),STATGROUP_MetalRHI, STATCAT_Advanced);
-DECLARE_STATS_GROUP(TEXT("AGX"),STATGROUP_AGXRHI, STATCAT_Advanced);
-DECLARE_STATS_GROUP(TEXT("Morph"),STATGROUP_MorphTarget, STATCAT_Advanced);
-DECLARE_STATS_GROUP(TEXT("Navigation"),STATGROUP_Navigation, STATCAT_Advanced);
-DECLARE_STATS_GROUP(TEXT("Net"),STATGROUP_Net, STATCAT_Advanced);
-
-#if !UE_BUILD_SHIPPING
-DECLARE_STATS_GROUP(TEXT("Packet"),STATGROUP_Packet, STATCAT_Advanced);
-#endif
-
-DECLARE_STATS_GROUP(TEXT("Object"),STATGROUP_Object, STATCAT_Advanced);
-DECLARE_STATS_GROUP_VERBOSE(TEXT("ObjectVerbose"),STATGROUP_ObjectVerbose, STATCAT_Advanced);
-DECLARE_STATS_GROUP(TEXT("OpenGL RHI"),STATGROUP_OpenGLRHI, STATCAT_Advanced);
-DECLARE_STATS_GROUP(TEXT("Pak File"),STATGROUP_PakFile, STATCAT_Advanced);
-DECLARE_STATS_GROUP(TEXT("Particle Mem"),STATGROUP_ParticleMem, STATCAT_Advanced);
-DECLARE_STATS_GROUP(TEXT("Particles"),STATGROUP_Particles, STATCAT_Advanced);
-DECLARE_STATS_GROUP(TEXT("Physics"),STATGROUP_Physics, STATCAT_Advanced);
-DECLARE_STATS_GROUP(TEXT("Platform"), STATGROUP_Platform, STATCAT_Advanced);
-DECLARE_STATS_GROUP(TEXT("Profiler"), STATGROUP_Profiler, STATCAT_Advanced);
-DECLARE_STATS_GROUP(TEXT("Quick"), STATGROUP_Quick, STATCAT_Advanced);
-DECLARE_STATS_GROUP(TEXT("RHI"), STATGROUP_RHI, STATCAT_Advanced);
-DECLARE_STATS_GROUP(TEXT("RDG"), STATGROUP_RDG, STATCAT_Advanced);
-DECLARE_STATS_GROUP(TEXT("Render Thread"),STATGROUP_RenderThreadProcessing, STATCAT_Advanced);
-DECLARE_STATS_GROUP(TEXT("Render Target Pool"), STATGROUP_RenderTargetPool, STATCAT_Advanced);
-DECLARE_STATS_GROUP(TEXT("Render Scaling"), STATGROUP_RenderScaling, STATCAT_Advanced);
-DECLARE_STATS_GROUP(TEXT("Scene Memory"),STATGROUP_SceneMemory, STATCAT_Advanced);
-DECLARE_STATS_GROUP(TEXT("Scene Rendering"),STATGROUP_SceneRendering, STATCAT_Advanced);
-DECLARE_STATS_GROUP(TEXT("Scene Update"),STATGROUP_SceneUpdate, STATCAT_Advanced);
-DECLARE_STATS_GROUP(TEXT("Server CPU"),STATGROUP_ServerCPU, STATCAT_Advanced);
-DECLARE_STATS_GROUP(TEXT("MapBuildData"),STATGROUP_MapBuildData, STATCAT_Advanced);
-DECLARE_STATS_GROUP(TEXT("Shader Compiling"),STATGROUP_ShaderCompiling, STATCAT_Advanced);
-DECLARE_STATS_GROUP(TEXT("Shader Compression"),STATGROUP_Shaders, STATCAT_Advanced);
-DECLARE_STATS_GROUP(TEXT("Shadow Rendering"),STATGROUP_ShadowRendering, STATCAT_Advanced);
-DECLARE_STATS_GROUP(TEXT("Stat System"),STATGROUP_StatSystem, STATCAT_Advanced);
-DECLARE_STATS_GROUP_SORTBYNAME(TEXT("Streaming Overview"),STATGROUP_StreamingOverview, STATCAT_Advanced);
-DECLARE_STATS_GROUP_SORTBYNAME(TEXT("Streaming Details"),STATGROUP_StreamingDetails, STATCAT_Advanced);
-DECLARE_STATS_GROUP_SORTBYNAME(TEXT("Streaming"),STATGROUP_Streaming, STATCAT_Advanced);
-DECLARE_STATS_GROUP(TEXT("Target Platform"),STATGROUP_TargetPlatform, STATCAT_Advanced);
-DECLARE_STATS_GROUP(TEXT("Text"),STATGROUP_Text, STATCAT_Advanced);
-DECLARE_STATS_GROUP(TEXT("ThreadPool Async Tasks"),STATGROUP_ThreadPoolAsyncTasks, STATCAT_Advanced);
-DECLARE_STATS_GROUP(TEXT("Threading"),STATGROUP_Threading, STATCAT_Advanced);
-DECLARE_STATS_GROUP(TEXT("Threads"),STATGROUP_Threads, STATCAT_Advanced);
-DECLARE_STATS_GROUP(TEXT("Tickables"),STATGROUP_Tickables, STATCAT_Advanced);
-DECLARE_STATS_GROUP(TEXT("Trail Particles"),STATGROUP_TrailParticles, STATCAT_Advanced);
-DECLARE_STATS_GROUP(TEXT("UI"),STATGROUP_UI, STATCAT_Advanced);
-DECLARE_STATS_GROUP(TEXT("UObjects"),STATGROUP_UObjects, STATCAT_Advanced);
-DECLARE_STATS_GROUP(TEXT("User"),STATGROUP_User, STATCAT_Advanced);
-
-DECLARE_CYCLE_STAT_EXTERN(TEXT("FrameTime"),STAT_FrameTime,STATGROUP_Engine, CORE_API);
-DECLARE_FNAME_STAT_EXTERN(TEXT("NamedMarker"),STAT_NamedMarker,STATGROUP_StatSystem, CORE_API);
-DECLARE_FLOAT_COUNTER_STAT_EXTERN(TEXT("Seconds Per Cycle"),STAT_SecondsPerCycle,STATGROUP_Engine, CORE_API);
-
-#endif
-
-#if STATS || ENABLE_STATNAMEDEVENTS
-namespace Stats
-{
-	UE_DEPRECATED(5.1, "No replacement. Engine usage was redundant")
-	FORCEINLINE bool IsThreadCollectingData()
-	{
-#if STATS && CPUPROFILERTRACE_ENABLED
-		return UE_TRACE_CHANNELEXPR_IS_ENABLED(CpuChannel) || FThreadStats::IsCollectingData();
-#else
-		return GCycleStatsShouldEmitNamedEvents != 0;
-#endif
-	}
-}
-#endif
diff --git a/Engine/Source/Runtime/Core/Public/Stats/StatsCommon.h b/Engine/Source/Runtime/Core/Public/Stats/StatsCommon.h
--- a/Engine/Source/Runtime/Core/Public/Stats/StatsCommon.h
+++ b/Engine/Source/Runtime/Core/Public/Stats/StatsCommon.h
@@ -3,6 +3,7 @@
 #pragma once
 
 #include "CoreTypes.h"
+#include "Misc/EnumClassFlags.h"
 
 #define FORCEINLINE_STATS FORCEINLINE
 //#define FORCEINLINE_STATS FORCEINLINE_DEBUGGABLE
@@ -12,3 +13,23 @@
 #error "STATS must be defined as either zero or one."
 #endif
 
+#if	PLATFORM_USES_ANSI_STRING_FOR_EXTERNAL_PROFILING
+typedef ANSICHAR PROFILER_CHAR;
+#else
+typedef WIDECHAR PROFILER_CHAR;
+#endif // PLATFORM_USES_ANSI_STRING_FOR_EXTERNAL_PROFILING
+
+#if PLATFORM_USES_ANSI_STRING_FOR_EXTERNAL_PROFILING
+#define ANSI_TO_PROFILING(x) x
+#else
+#define ANSI_TO_PROFILING(x) TEXT(x)
+#endif
+
+enum class EStatFlags : uint8
+{
+	None = 0,
+	ClearEveryFrame = 1 << 0,
+	CycleStat = 1 << 1,
+	Verbose = 1 << 2, // Profiling scopes for this stat will no generate a trace event by default. See GShouldEmitVerboseNamedEvents.
+};
+ENUM_CLASS_FLAGS(EStatFlags);
\ No newline at end of file
diff --git a/Engine/Source/Runtime/CoreUObject/Public/UObject/Script.h b/Engine/Source/Runtime/CoreUObject/Public/UObject/Script.h
--- a/Engine/Source/Runtime/CoreUObject/Public/UObject/Script.h
+++ b/Engine/Source/Runtime/CoreUObject/Public/UObject/Script.h
@@ -6,6 +6,7 @@
 
 #pragma once
 
+#include "Delegates/Delegate.h"
 #include "HAL/ThreadSingleton.h"
 #if UE_ENABLE_INCLUDE_ORDER_DEPRECATED_IN_5_4
 #include "Internationalization/Text.h"
diff --git a/Engine/Source/Runtime/D3D12RHI/Private/D3D12Submission.cpp b/Engine/Source/Runtime/D3D12RHI/Private/D3D12Submission.cpp
--- a/Engine/Source/Runtime/D3D12RHI/Private/D3D12Submission.cpp
+++ b/Engine/Source/Runtime/D3D12RHI/Private/D3D12Submission.cpp
@@ -5,6 +5,7 @@
 #include "HAL/Runnable.h"
 #include "HAL/RunnableThread.h"
 #include "IRenderCaptureProvider.h"
+#include "Stats/ThreadIdleStats.h"
 
 #ifndef D3D12_PLATFORM_SUPPORTS_BLOCKING_FENCES
 #define D3D12_PLATFORM_SUPPORTS_BLOCKING_FENCES 1
@@ -1146,7 +1147,7 @@ void FD3D12DynamicRHI::ProcessInterruptQueueUntil(FGraphEvent* GraphEvent)
 		if (!GraphEvent || !GraphEvent->IsComplete())
 		{
 			// If we're waiting for a sync point, accumulate the idle time
-			FThreadIdleStats::FScopeIdle IdleScope(/* bIgnore = */GraphEvent == nullptr);
+			UE::Stats::FThreadIdleStats::FScopeIdle IdleScope(/* bIgnore = */GraphEvent == nullptr);
 
 		Retry:
 			if (InterruptCS.TryLock())
diff --git a/Engine/Source/Runtime/Engine/Private/LevelTick.cpp b/Engine/Source/Runtime/Engine/Private/LevelTick.cpp
--- a/Engine/Source/Runtime/Engine/Private/LevelTick.cpp
+++ b/Engine/Source/Runtime/Engine/Private/LevelTick.cpp
@@ -28,6 +28,7 @@
 #include "UObject/Stack.h"
 #include "PhysicsEngine/CollisionAnalyzerCapture.h"
 #include "Rendering/RenderCommandPipes.h"
+#include "Stats/ThreadIdleStats.h"
 
 #if !UE_SERVER
 #include "IMediaModule.h"
@@ -1269,7 +1270,7 @@ void UWorld::Tick( ELevelTick TickType, float DeltaSeconds )
 		return;
 	}
 
-	FThreadIdleStats::BeginCriticalPath();
+	UE::Stats::FThreadIdleStats::BeginCriticalPath();
 
 	RHI_BREADCRUMB_EVENT_GAMETHREAD("WorldTick");
 
@@ -1810,7 +1811,7 @@ void UWorld::Tick( ELevelTick TickType, float DeltaSeconds )
 	StateStreamManager->Game_BeginTick();
 #endif
 
-	FThreadIdleStats::EndCriticalPath();
+	UE::Stats::FThreadIdleStats::EndCriticalPath();
 }
 
 void UWorld::CleanupActors()
diff --git a/Engine/Source/Runtime/Engine/Private/UnrealClient.cpp b/Engine/Source/Runtime/Engine/Private/UnrealClient.cpp
--- a/Engine/Source/Runtime/Engine/Private/UnrealClient.cpp
+++ b/Engine/Source/Runtime/Engine/Private/UnrealClient.cpp
@@ -32,6 +32,7 @@
 #include "RenderGraphUtils.h"
 #include "DynamicResolutionState.h"
 #include "Stats/StatsTrace.h"
+#include "Stats/ThreadIdleStats.h"
 #include "RHIUtilities.h"
 
 DEFINE_LOG_CATEGORY_STATIC(LogClient, Log, All);
@@ -1773,7 +1774,7 @@ void FViewport::Draw( bool bShouldPresent /*= true */)
 						static uint32 Lastimestamp = 0;
 						static bool bStarted = false;
 						uint32 CurrentTime	= FPlatformTime::Cycles();
-						FThreadIdleStats& GameThread = FThreadIdleStats::Get();
+						UE::Stats::FThreadIdleStats& GameThread = UE::Stats::FThreadIdleStats::Get();
 						if (bStarted)
 						{
 							uint32 ThreadTime	= CurrentTime - Lastimestamp;
diff --git a/Engine/Source/Runtime/Engine/Private/UnrealEngine.cpp b/Engine/Source/Runtime/Engine/Private/UnrealEngine.cpp
--- a/Engine/Source/Runtime/Engine/Private/UnrealEngine.cpp
+++ b/Engine/Source/Runtime/Engine/Private/UnrealEngine.cpp
@@ -28,6 +28,7 @@ UnrealEngine.cpp: Implements the UEngine class and helpers.
 #include "Misc/OutputDeviceFile.h"
 #include "Math/ColorList.h"
 #include "Stats/StatsMisc.h"
+#include "Stats/ThreadIdleStats.h"
 #include "Misc/AutomationTest.h"
 #include "Misc/CoreDelegates.h"
 #include "Misc/StringFormatArg.h"
@@ -2731,7 +2732,7 @@ void UEngine::UpdateTimeAndHandleMaxTickRate()
 		if( WaitTime > 0 )
 		{
 			// track all this waiting so that Game Thread is correct
-			FThreadIdleStats::FScopeIdle Scope;
+			UE::Stats::FThreadIdleStats::FScopeIdle Scope;
 
 			FSimpleScopeSecondsCounter ActualWaitTimeCounter(ActualWaitTime);
 			double WaitEndTime = CurrentRealTime + WaitTime;
diff --git a/Engine/Source/Runtime/Launch/Private/LaunchEngineLoop.cpp b/Engine/Source/Runtime/Launch/Private/LaunchEngineLoop.cpp
--- a/Engine/Source/Runtime/Launch/Private/LaunchEngineLoop.cpp
+++ b/Engine/Source/Runtime/Launch/Private/LaunchEngineLoop.cpp
@@ -45,6 +45,8 @@
 #include "Serialization/CompactBinarySerialization.h"
 #include "Serialization/CompactBinaryWriter.h"
 #include "Stats/StatsMallocProfilerProxy.h"
+#include "Stats/StatsSystem.h"
+#include "Stats/ThreadIdleStats.h"
 #include "Trace/Trace.inl"
 #include "ProfilingDebugging/TraceAuxiliary.h"
 #include "ProfilingDebugging/CsvProfiler.h"
@@ -2378,7 +2380,7 @@ int32 FEngineLoop::PreInitPreStartupScreen(const TCHAR* CmdLine)
 #endif
 #if STATS
 		// Leave the stats enabled.
-		if (!FStats::EnabledForCommandlet())
+		if (!UE::Stats::FStats::EnabledForCommandlet())
 		{
 			FThreadStats::PrimaryDisableForever();
 		}
@@ -4134,7 +4136,7 @@ int32 FEngineLoop::PreInitPostStartupScreen(const TCHAR* CmdLine)
 			// We have to close the scope, otherwise we will end with broken stats.
 			CycleCount_AfterStats.StopAndResetStatId();
 #endif // STATS
-			FStats::TickCommandletStats();
+			UE::Stats::FStats::TickCommandletStats();
 
 #if WITH_ENGINE
 			PRIVATE_GRunningCommandletClass = CommandletClass;
@@ -4151,7 +4153,7 @@ int32 FEngineLoop::PreInitPostStartupScreen(const TCHAR* CmdLine)
 			PRIVATE_GRunningCommandletClass = nullptr;
 #endif
 
-			FStats::TickCommandletStats();
+			UE::Stats::FStats::TickCommandletStats();
 
 			RequestEngineExit(FString::Printf(TEXT("Commandlet %s finished execution (result %d)"), *Commandlet->GetName(), ErrorLevel));
 
@@ -4756,7 +4758,7 @@ void FEngineLoop::InitTime()
 void GameLoopIsStarved()
 {
 	FlushPendingDeleteRHIResources_GameThread();
-	FStats::AdvanceFrame( true, FStats::FOnAdvanceRenderingThreadStats::CreateStatic( &AdvanceRenderingThreadStatsGT ) );
+	UE::Stats::FStats::AdvanceFrame( true, UE::Stats::FStats::FOnAdvanceRenderingThreadStats::CreateStatic( &AdvanceRenderingThreadStatsGT ) );
 }
 
 
@@ -5540,7 +5542,7 @@ static inline void BeginFrameRenderThread(FRHICommandListImmediate& RHICmdList,
 	GVisualizeTexture.BeginFrameRenderThread();
 
 	// Waits after this point but before BeginRenderingViewFamilies are not on the RT critical path (since we're waiting for the GT)
-	FThreadIdleStats::EndCriticalPath();
+	UE::Stats::FThreadIdleStats::EndCriticalPath();
 }
 
 
@@ -5769,7 +5771,7 @@ void FEngineLoop::Tick()
 #endif
 	}
 
-	FStats::AdvanceFrame( false, FStats::FOnAdvanceRenderingThreadStats::CreateStatic( &AdvanceRenderingThreadStatsGT ) );
+	UE::Stats::FStats::AdvanceFrame( false, UE::Stats::FStats::FOnAdvanceRenderingThreadStats::CreateStatic( &AdvanceRenderingThreadStatsGT ) );
 
 	{
 		SCOPE_CYCLE_COUNTER( STAT_FrameTime );
diff --git a/Engine/Source/Runtime/RHI/Private/RHICommandList.cpp b/Engine/Source/Runtime/RHI/Private/RHICommandList.cpp
--- a/Engine/Source/Runtime/RHI/Private/RHICommandList.cpp
+++ b/Engine/Source/Runtime/RHI/Private/RHICommandList.cpp
@@ -13,6 +13,7 @@
 #include "RHIFwd.h"
 #include "RHITransition.h"
 #include "Stats/StatsTrace.h"
+#include "Stats/ThreadIdleStats.h"
 #include "HAL/PlatformMisc.h"
 
 CSV_DEFINE_CATEGORY_MODULE(RHI_API, RHITStalls, false);
@@ -1524,7 +1525,7 @@ void FRHICommandListExecutor::WaitOnRHIThreadFence(FGraphEventRef& Fence)
 	check(IsInRenderingThread());
 
 	// Exclude RHIT waits from the RT critical path stat (these waits simply get longer if the RT is running faster, so we don't get useful results)
-	FThreadIdleStats::FScopeNonCriticalPath NonCriticalPathScope;
+	UE::Stats::FThreadIdleStats::FScopeNonCriticalPath NonCriticalPathScope;
 
 	if (Fence.GetReference() && !Fence->IsComplete())
 	{
diff --git a/Engine/Source/Runtime/RHI/Public/RHICommandList.h b/Engine/Source/Runtime/RHI/Public/RHICommandList.h
--- a/Engine/Source/Runtime/RHI/Public/RHICommandList.h
+++ b/Engine/Source/Runtime/RHI/Public/RHICommandList.h
@@ -34,6 +34,7 @@
 #include "RHIShaderParameters.h"
 #include "RHITextureReference.h"
 #include "RHIResourceReplace.h"
+#include "Stats/ThreadIdleStats.h"
 #include "Trace/Trace.h"
 
 #include "DynamicRHI.h"
@@ -122,7 +123,7 @@ extern RHI_API uint32 GRenderThreadIdle[ERenderThreadIdleTypes::Num];
 /** Helper to mark scopes as idle time on the render or RHI threads. */
 struct FRenderThreadIdleScope
 {
-	FThreadIdleStats::FScopeIdle RHIThreadIdleScope;
+	UE::Stats::FThreadIdleStats::FScopeIdle RHIThreadIdleScope;
 
 	const ERenderThreadIdleTypes::Type Type;
 	const bool bCondition;
diff --git a/Engine/Source/Runtime/RenderCore/Public/VirtualTexturing.h b/Engine/Source/Runtime/RenderCore/Public/VirtualTexturing.h
--- a/Engine/Source/Runtime/RenderCore/Public/VirtualTexturing.h
+++ b/Engine/Source/Runtime/RenderCore/Public/VirtualTexturing.h
@@ -12,11 +12,7 @@
 #include "Misc/AssertionMacros.h"
 #include "Misc/EnumClassFlags.h"
 #include "PixelFormat.h"
-#if UE_ENABLE_INCLUDE_ORDER_DEPRECATED_IN_5_2
-#include "RHIDefinitions.h"
-#endif
 #include "Stats/Stats.h"
-#include "Stats/Stats2.h"
 #include "Templates/RefCounting.h"
 #include "Templates/TypeHash.h"
 #include "UObject/NameTypes.h"
diff --git a/Engine/Source/Runtime/Renderer/Private/SceneRendering.cpp b/Engine/Source/Runtime/Renderer/Private/SceneRendering.cpp
--- a/Engine/Source/Runtime/Renderer/Private/SceneRendering.cpp
+++ b/Engine/Source/Runtime/Renderer/Private/SceneRendering.cpp
@@ -100,6 +100,7 @@
 #include "LocalFogVolumeRendering.h"
 #include "OIT/OIT.h"
 #include "Rendering/CustomRenderPass.h"
+#include "Stats/ThreadIdleStats.h"
 #include "CustomRenderPassSceneCapture.h"
 #include "LightFunctionAtlas.h"
 #include "EnvironmentComponentsFlags.h"
@@ -5550,12 +5551,12 @@ void FRendererModule::BeginRenderingViewFamilies(FCanvas* Canvas, TConstArrayVie
 		World = Scene->GetWorld();
 		if (World)
 		{
-			FThreadIdleStats::BeginCriticalPath();
+			UE::Stats::FThreadIdleStats::BeginCriticalPath();
 			{
 				// Guarantee that all render proxies are up to date before kicking off a BeginRenderViewFamily.
 				World->SendAllEndOfFrameUpdates();
 			}
-			FThreadIdleStats::EndCriticalPath();
+			UE::Stats::FThreadIdleStats::EndCriticalPath();
 
 			GetNaniteVisualizationData().Pick(World);
 		}
@@ -5565,7 +5566,7 @@ void FRendererModule::BeginRenderingViewFamilies(FCanvas* Canvas, TConstArrayVie
 		[](FRHICommandList& RHICmdList)
 		{
 			// Rendering is up and running now, so waits are considered part of the RT critical path
-			FThreadIdleStats::BeginCriticalPath();
+			UE::Stats::FThreadIdleStats::BeginCriticalPath();
 		});
 
 	FUniformExpressionCacheAsyncUpdateScope AsyncUpdateScope;
diff --git a/Engine/Source/Runtime/SlateRHIRenderer/Private/SlateRHIRenderer.cpp b/Engine/Source/Runtime/SlateRHIRenderer/Private/SlateRHIRenderer.cpp
--- a/Engine/Source/Runtime/SlateRHIRenderer/Private/SlateRHIRenderer.cpp
+++ b/Engine/Source/Runtime/SlateRHIRenderer/Private/SlateRHIRenderer.cpp
@@ -54,6 +54,7 @@
 #include "RenderCore.h"
 #include "DataDrivenShaderPlatformInfo.h"
 #include "SlatePostProcessor.h"
+#include "Stats/ThreadIdleStats.h"
 
 #if WITH_EDITORONLY_DATA
 #include "ShaderCompiler.h"
@@ -981,7 +982,7 @@ void FSlateRHIRenderer::PresentWindow_RenderThread(FRHICommandListImmediate& RHI
 
 	uint32 RenderThreadIdle = 0;
 
-	FThreadIdleStats& RenderThread = FThreadIdleStats::Get();
+	UE::Stats::FThreadIdleStats& RenderThread = UE::Stats::FThreadIdleStats::Get();
 	GRenderThreadIdle[ERenderThreadIdleTypes::WaitingForAllOtherSleep] = RenderThread.Waits;
 	GRenderThreadIdle[ERenderThreadIdleTypes::WaitingForGPUPresent] += GSwapBufferTime;
 
@@ -1029,7 +1030,7 @@ void FSlateRHIRenderer::PresentWindow_RenderThread(FRHICommandListImmediate& RHI
 		RHICmdList.EnqueueLambda([](FRHICommandListImmediate&)
 		{
 			// Update RHI thread time
-			FThreadIdleStats& RHIThreadStats = FThreadIdleStats::Get();
+			UE::Stats::FThreadIdleStats& RHIThreadStats = UE::Stats::FThreadIdleStats::Get();
 
 			if (!RHITCycles.IsSet())
 			{
diff --git a/Engine/Source/Runtime/VulkanRHI/Private/VulkanRHIPrivate.h b/Engine/Source/Runtime/VulkanRHI/Private/VulkanRHIPrivate.h
--- a/Engine/Source/Runtime/VulkanRHI/Private/VulkanRHIPrivate.h
+++ b/Engine/Source/Runtime/VulkanRHI/Private/VulkanRHIPrivate.h
@@ -8,6 +8,7 @@
 
 // Dependencies
 #include "CoreMinimal.h"
+#include "Containers/IndirectArray.h"
 #include "Misc/ScopeLock.h"
 #include "RHI.h"
 #include "RenderUtils.h"
@@ -31,7 +32,7 @@
 #include "GPUProfiler.h"
 #include "VulkanQueue.h"
 #include "VulkanCommandBuffer.h"
-#include "Stats/Stats2.h"
+#include "Stats/Stats.h"
 
 class FVulkanQueue;
 class FVulkanCommandBuffer;
