Making GIsCriticalError an atomic to fix a TSAN issue.

#rnx

diff --git a/Engine/Source/Runtime/Core/Private/Misc/CoreGlobals.cpp b/Engine/Source/Runtime/Core/Private/Misc/CoreGlobals.cpp
--- a/Engine/Source/Runtime/Core/Private/Misc/CoreGlobals.cpp
+++ b/Engine/Source/Runtime/Core/Private/Misc/CoreGlobals.cpp
@@ -205,7 +205,7 @@ bool					GFirstFrameIntraFrameDebugging	= false;
 bool					GEdSelectionLock				= false;					/* Are selections locked? (you can't select/deselect additional actors) */
 bool					GIsClient						= false;					/* Whether engine was launched as a client */
 bool					GIsServer						= false;					/* Whether engine was launched as a server, true if GIsClient */
-bool					GIsCriticalError				= false;					/* An appError() has occured */
+TSAN_ATOMIC(bool)		GIsCriticalError				= false;					/* An appError() has occured */
 bool					GIsGuarded						= false;					/* Whether execution is happening within main()/WinMain()'s try/catch handler */
 TSAN_ATOMIC(bool)		GIsRunning(false);											/* Whether execution is happening within MainLoop() */
 FIsDuplicatingClassForReinstancing	GIsDuplicatingClassForReinstancing;					        /* Whether we are currently using SDO on a UClass or CDO for live reinstancing */
diff --git a/Engine/Source/Runtime/Core/Public/CoreGlobals.h b/Engine/Source/Runtime/Core/Public/CoreGlobals.h
--- a/Engine/Source/Runtime/Core/Public/CoreGlobals.h
+++ b/Engine/Source/Runtime/Core/Public/CoreGlobals.h
@@ -337,7 +337,7 @@ public:
 extern CORE_API bool GEdSelectionLock;
 extern CORE_API bool GIsClient;
 extern CORE_API bool GIsServer;
-extern CORE_API bool GIsCriticalError;
+extern CORE_API TSAN_ATOMIC(bool) GIsCriticalError;
 extern CORE_API TSAN_ATOMIC(bool) GIsRunning;
 extern CORE_API FIsDuplicatingClassForReinstancing GIsDuplicatingClassForReinstancing;
 
