- Add missing UE_CLIENT define to Build.h
-- This is added to GlobalDefinitions in TargetRules.cs for client builds.

#rnx

diff --git a/Engine/Source/Runtime/Core/Public/Misc/Build.h b/Engine/Source/Runtime/Core/Public/Misc/Build.h
--- a/Engine/Source/Runtime/Core/Public/Misc/Build.h
+++ b/Engine/Source/Runtime/Core/Public/Misc/Build.h
@@ -22,6 +22,9 @@
 #ifndef UE_GAME
 	#define UE_GAME						0
 #endif
+#ifndef UE_CLIENT
+	#define UE_CLIENT					0
+#endif
 #ifndef UE_EDITOR
 	#define UE_EDITOR					0
 #endif
