[AutoRTFM] Add autortfm_memory_validation_level parameter to autortfm_pre_open()

This should have been part of 37474030, but got missed when separating out changes.

diff --git a/Engine/Source/Programs/AutoRTFMTests/Private/Open.cpp b/Engine/Source/Programs/AutoRTFMTests/Private/Open.cpp
--- a/Engine/Source/Programs/AutoRTFMTests/Private/Open.cpp
+++ b/Engine/Source/Programs/AutoRTFMTests/Private/Open.cpp
@@ -437,6 +437,40 @@ TEST_CASE("Open.Collision")
             Catch::Matchers::VectorContains(FString(AutoRTFMTestUtils::kMemoryModifiedWarning)));
     }
 
+    SECTION("Transact(Assign, OpenNoValidation(Assign))")
+    {
+        AutoRTFM::Transact([&] ()
+        {
+            I = 42;
+            AutoRTFM::Open<AutoRTFM::EMemoryValidationLevel::Disabled>([&]
+            {
+                I = 10;
+            });
+        });
+
+        REQUIRE(WarningContext.GetWarnings().empty());
+    }
+
+    SECTION("Transact(Assign, AlwaysOpenNoValidation(Assign))")
+    {
+        struct S
+        {
+            UE_AUTORTFM_ALWAYS_OPEN_NO_MEMORY_VALIDATION
+            static void AssignToInt(int& I)
+            {
+                I = 10;
+            }
+        };
+
+        AutoRTFM::Transact([&] ()
+        {
+            I = 42;
+            S::AssignToInt(I);
+        });
+
+        REQUIRE(WarningContext.GetWarnings().empty());
+    }
+
     SECTION("Transact(Assign, Open(Assign, OpenNoValidation()))")
     {
         AutoRTFM::Transact([&] ()
diff --git a/Engine/Source/Runtime/Core/Private/AutoRTFM/API.cpp b/Engine/Source/Runtime/Core/Private/AutoRTFM/API.cpp
--- a/Engine/Source/Runtime/Core/Private/AutoRTFM/API.cpp
+++ b/Engine/Source/Runtime/Core/Private/AutoRTFM/API.cpp
@@ -393,11 +393,15 @@ AUTORTFM_CORE_API autortfm_status RTFM_autortfm_close(void (*UninstrumentedWork)
 	return static_cast<autortfm_status>(Context->GetStatus());
 }
 
-extern "C" UE_AUTORTFM_NOAUTORTFM UE_AUTORTFM_API void autortfm_pre_open()
+extern "C" UE_AUTORTFM_NOAUTORTFM UE_AUTORTFM_API void autortfm_pre_open(autortfm_memory_validation_level MemoryValidationLevel)
 {
 	if (FTransaction* Transaction = FContext::Get()->GetCurrentTransaction())
 	{
-		Transaction->SetOpenActive(ForTheRuntime::GetMemoryValidationLevel());
+		AutoRTFM::EMemoryValidationLevel Level
+			= MemoryValidationLevel == autortfm_memory_validation_level_default
+			? ForTheRuntime::GetMemoryValidationLevel()
+			: static_cast<AutoRTFM::EMemoryValidationLevel>(MemoryValidationLevel);
+		Transaction->SetOpenActive(Level);
 	}
 }
 
@@ -422,15 +426,15 @@ extern "C" UE_AUTORTFM_NOAUTORTFM UE_AUTORTFM_API void autortfm_post_static_loca
 	FContext::Get()->LeavingStaticLocalInitializer();
 }
 
-AUTORTFM_CORE_API void RTFM_autortfm_open(autortfm_memory_validation_level ValidationLevel, void (*Work)(void*), void* Arg)
+AUTORTFM_CORE_API void RTFM_autortfm_open(autortfm_memory_validation_level MemoryValidationLevel, void (*Work)(void*), void* Arg)
 {
 	FTransaction* Transaction = FContext::Get()->GetCurrentTransaction();
 	if (Transaction)
 	{
 		AutoRTFM::EMemoryValidationLevel Level
-			= ValidationLevel == autortfm_memory_validation_level_default
+			= MemoryValidationLevel == autortfm_memory_validation_level_default
 			? ForTheRuntime::GetMemoryValidationLevel()
-			: static_cast<AutoRTFM::EMemoryValidationLevel>(ValidationLevel);
+			: static_cast<AutoRTFM::EMemoryValidationLevel>(MemoryValidationLevel);
 		Transaction->SetOpenActive(Level);
 	}
 
