[HWRT] Support streaming of FStaticMeshRayTracingProxyLOD with bOwnsBuffers == true.

Context:
- Streaming logic in FRayTracingGeometryManager only handled streaming of offline built BVH not other associated data such as VB/IB.
- Different asset types might serialize different data/layouts so responsibility of managing that shouldn't be on FRayTracingGeometryManager.

Changes:
- new FRayTracingStreamableAsset interface that asset types can implement to support streaming related to ray tracing geometry.
- modified FRayTracingGeometryManager to issue stream in/out requets via FRayTracingStreamableAsset
- FStaticMeshRayTracingProxyLOD implements FRayTracingStreamableAsset interface.
- Modified serialization of FStaticMeshRayTracingProxyLOD to track offset and size of offline BVH and buffers in serialized data blob.

diff --git a/Engine/Source/Runtime/Core/Private/UObject/FortniteReleaseBranchCustomObjectVersion.cpp b/Engine/Source/Runtime/Core/Private/UObject/FortniteReleaseBranchCustomObjectVersion.cpp
--- a/Engine/Source/Runtime/Core/Private/UObject/FortniteReleaseBranchCustomObjectVersion.cpp
+++ b/Engine/Source/Runtime/Core/Private/UObject/FortniteReleaseBranchCustomObjectVersion.cpp
@@ -11,5 +11,6 @@ TMap<FGuid, FGuid> FFortniteReleaseBranchCustomObjectVersion::GetSystemGuids()
 	SystemGuids.Add(DevGuids.MaterialTranslationDDCVersion, FGuid("888A502A66F6401580B27451F211C46B"));
 	SystemGuids.Add(DevGuids.GLOBALSHADERMAP_DERIVEDDATA_VER, FGuid("0A0A1BF5D2A94105BB395EF0DA7CAEB9"));
 	SystemGuids.Add(DevGuids.MATERIALSHADERMAP_DERIVEDDATA_VER, FGuid("C0E1F59386894C9D855B85178B2553CB"));
+	SystemGuids.Add(DevGuids.STATICMESH_DERIVEDDATA_VER, FGuid("9463A91F546C41FCBFF114787E9E7E32"));
 	return SystemGuids;
 }
