- Removed the reference to ApplicationCore from Core.Build.cs when bBuildWithApplicationCore is false (for IOS)

diff --git a/Engine/Source/Runtime/Core/Core.Build.cs b/Engine/Source/Runtime/Core/Core.Build.cs
--- a/Engine/Source/Runtime/Core/Core.Build.cs
+++ b/Engine/Source/Runtime/Core/Core.Build.cs
@@ -135,7 +135,10 @@ public class Core : ModuleRules
 				PublicFrameworks.Add("CoreMotion");
 			}
 
-			PrivateIncludePathModuleNames.Add("ApplicationCore");
+			if (Target.bCompileAgainstApplicationCore)
+			{
+				PrivateIncludePathModuleNames.Add("ApplicationCore");
+			}
 
 			bool bSupportAdvertising = Target.Platform == UnrealTargetPlatform.IOS;
 			if (bSupportAdvertising)
