[Core]
* Fixed unrealpak shipping compile error in TraceAuxiliary.cpp

diff --git a/Engine/Source/Runtime/Core/Private/ProfilingDebugging/TraceAuxiliary.cpp b/Engine/Source/Runtime/Core/Private/ProfilingDebugging/TraceAuxiliary.cpp
--- a/Engine/Source/Runtime/Core/Private/ProfilingDebugging/TraceAuxiliary.cpp
+++ b/Engine/Source/Runtime/Core/Private/ProfilingDebugging/TraceAuxiliary.cpp
@@ -920,13 +920,17 @@ void OnMessageCallback(const UE::Trace::FMessageEvent& Message)
 ////////////////////////////////////////////////////////////////////////////////
 void OnScopeBeginCallback(const ANSICHAR* ScopeName)
 {
+#if CPUPROFILERTRACE_ENABLED
 	FCpuProfilerTrace::OutputBeginDynamicEvent(ScopeName);
+#endif
 }
 
 ////////////////////////////////////////////////////////////////////////////////
 void OnScopeEndCallback()
 {
+#if CPUPROFILERTRACE_ENABLED
 	FCpuProfilerTrace::OutputEndEvent();
+#endif
 }
 
 ////////////////////////////////////////////////////////////////////////////////
