[TaskGraph/Pipe]
 - Remove additional barrier since the root cause for the race was a missing acquire barrier on bIsClosed in PushIfNotClosed fixed a few days ago.

#rnx

diff --git a/Engine/Source/Runtime/Core/Private/Tasks/Pipe.cpp b/Engine/Source/Runtime/Core/Private/Tasks/Pipe.cpp
--- a/Engine/Source/Runtime/Core/Private/Tasks/Pipe.cpp
+++ b/Engine/Source/Runtime/Core/Private/Tasks/Pipe.cpp
@@ -5,8 +5,6 @@
 #include "Containers/Array.h"
 #include "Misc/ScopeExit.h"
 
-DEFINE_LOG_CATEGORY_STATIC(LogTaskGraphPipe, Log, All);
-
 namespace UE::Tasks
 {
 	Private::FTaskBase* FPipe::PushIntoPipe(Private::FTaskBase& Task)
@@ -121,25 +119,10 @@ namespace UE::Tasks
 	void FPipe::ExecutionStarted()
 	{
 		FPipeCallStack::Push(*this);
-
-		// Since different threads can execute the same pipe, we need a single atomic
-		// that each thread are going to synchronize memory ordering with.
-		// Acquire any unsynchronized memory modifications that happened
-		// in the last pipe task execution.
-		if (Executing.load(std::memory_order_acquire))
-		{
-			UE_LOG(LogTaskGraphPipe, Fatal, TEXT("Pipe execution must be sequential"));
-		}
-
-		Executing.store(true, std::memory_order_relaxed);
 	}
 
 	void FPipe::ExecutionFinished()
 	{
-		// Publish any unsynchronized memory modification that happened during 
-		// this pipe task execution before another task is executed on the same pipe.
-		Executing.store(false, std::memory_order_release);
-
 		FPipeCallStack::Pop(*this);
 	}
 
diff --git a/Engine/Source/Runtime/Core/Public/Tasks/Pipe.h b/Engine/Source/Runtime/Core/Public/Tasks/Pipe.h
--- a/Engine/Source/Runtime/Core/Public/Tasks/Pipe.h
+++ b/Engine/Source/Runtime/Core/Public/Tasks/Pipe.h
@@ -134,7 +134,6 @@ namespace UE::Tasks
 		// pipe builds a chain (a linked list) of tasks and so needs to store only the last one. the last task is null if the pipe is not blocked
 		std::atomic<Private::FTaskBase*> LastTask{ nullptr };
 		std::atomic<uint64> TaskCount { 0 };
-		std::atomic<bool> Executing{ false }; // Used for memory ordering between pipe tasks.
 		UE::FEventCount EmptyEvent;
 	public:
 		FORCENOINLINE const TCHAR* GetDebugName() const
