Core: Fix error in TFrameValue(TFrameValue&& InValue) trying to move TFrameValue into the TOptional value

diff --git a/Engine/Source/Runtime/Core/Public/Misc/FrameValue.h b/Engine/Source/Runtime/Core/Public/Misc/FrameValue.h
--- a/Engine/Source/Runtime/Core/Public/Misc/FrameValue.h
+++ b/Engine/Source/Runtime/Core/Public/Misc/FrameValue.h
@@ -47,7 +47,7 @@ public:
 	TFrameValue(TFrameValue&& InValue)
 		: FrameSet(GFrameCounter)
 	{
-		Value = MoveTemp(InValue);
+		Value = MoveTemp(InValue.Value);
 	}
 
 	TFrameValue& operator=(const TFrameValue& InValue)
