[Core] Added FText, FName overloads to TestEqual and related, and updated a related unit test to disambiguate the overload

diff --git a/Engine/Source/Runtime/Core/Private/Misc/AutomationTest.cpp b/Engine/Source/Runtime/Core/Private/Misc/AutomationTest.cpp
--- a/Engine/Source/Runtime/Core/Private/Misc/AutomationTest.cpp
+++ b/Engine/Source/Runtime/Core/Private/Misc/AutomationTest.cpp
@@ -1915,6 +1915,28 @@ bool FAutomationTestBase::TestEqual(const TCHAR* What, FStringView Actual, FStri
 	return true;
 }
 
+bool FAutomationTestBase::TestEqual(const TCHAR* What, const FText Actual, const FText Expected)
+{
+	if (!Actual.EqualToCaseIgnored(Actual))
+	{
+		AddError(FString::Printf(TEXT("Expected '%s' to be %s, but it was %s."), What, *Expected.ToString(), *Actual.ToString()), 1);
+		return false;
+	}
+
+	return true;
+}
+
+bool FAutomationTestBase::TestEqual(const TCHAR* What, const FName Actual, const FName Expected)
+{
+	if (!Actual.IsEqual(Expected))
+	{
+		AddError(FString::Printf(TEXT("Expected '%s' to be %s, but it was %s."), What, *Expected.ToString(), *Actual.ToString()), 1);
+		return false;
+	}
+
+	return true;
+}
+
 bool FAutomationTestBase::TestNotEqual(const TCHAR* What, const TCHAR* Actual, const TCHAR* Expected)
 {
 	bool bAreDifferent = (Actual && Expected) ? (FCString::Stricmp(Actual, Expected) != 0) : (Actual != Expected);
@@ -1969,6 +1991,28 @@ bool FAutomationTestBase::TestNotEqual(const TCHAR* What, const double Actual, c
 	return true;
 }
 
+bool FAutomationTestBase::TestNotEqual(const TCHAR* What, const FText Actual, const FText Expected)
+{
+	if (Actual.EqualToCaseIgnored(Actual))
+	{
+		AddError(FString::Printf(TEXT("Expected '%s' to differ from %s, but it was %s."), What, *Expected.ToString(), *Actual.ToString()), 1);
+		return false;
+	}
+
+	return true;
+}
+
+bool FAutomationTestBase::TestNotEqual(const TCHAR* What, const FName Actual, const FName Expected)
+{
+	if (Actual.IsEqual(Expected))
+	{
+		AddError(FString::Printf(TEXT("Expected '%s' to differ from %s, but it was %s."), What, *Expected.ToString(), *Actual.ToString()), 1);
+		return false;
+	}
+
+	return true;
+}
+
 bool FAutomationTestBase::TestEqualInsensitive(const TCHAR* What, const TCHAR* Actual, const TCHAR* Expected)
 {
 	return TestEqual(What, Actual, Expected);
diff --git a/Engine/Source/Runtime/Core/Public/Misc/AutomationTest.h b/Engine/Source/Runtime/Core/Public/Misc/AutomationTest.h
--- a/Engine/Source/Runtime/Core/Public/Misc/AutomationTest.h
+++ b/Engine/Source/Runtime/Core/Public/Misc/AutomationTest.h
@@ -1941,9 +1941,13 @@ public:
 	CORE_API bool TestEqual(const TCHAR* What, const TCHAR* Actual, const TCHAR* Expected);
 	CORE_API bool TestEqual(const TCHAR* What, FStringView Actual, FStringView Expected);
 	CORE_API bool TestEqual(const TCHAR* What, FUtf8StringView Actual, FUtf8StringView Expected);
+	CORE_API bool TestEqual(const TCHAR* What, const FText Actual, const FText Expected);
+	CORE_API bool TestEqual(const TCHAR* What, const FName Actual, const FName Expected);
 	CORE_API bool TestNotEqual(const TCHAR* What, const TCHAR* Actual, const TCHAR* Expected);
 	CORE_API bool TestNotEqual(const TCHAR* What, FStringView Actual, FStringView Expected);
 	CORE_API bool TestNotEqual(const TCHAR* What, FUtf8StringView Actual, FUtf8StringView Expected);
+	CORE_API bool TestNotEqual(const TCHAR* What, const FText Actual, const FText Expected);
+	CORE_API bool TestNotEqual(const TCHAR* What, const FName Actual, const FName Expected);
 
 	UE_DEPRECATED(5.5, "Use TestEqual instead (string tests are case insensitive by default)")
 	CORE_API bool TestEqualInsensitive(const TCHAR* What, const TCHAR* Actual, const TCHAR* Expected);
diff --git a/Engine/Source/Runtime/Core/Tests/Misc/AsciiSetTest.cpp b/Engine/Source/Runtime/Core/Tests/Misc/AsciiSetTest.cpp
--- a/Engine/Source/Runtime/Core/Tests/Misc/AsciiSetTest.cpp
+++ b/Engine/Source/Runtime/Core/Tests/Misc/AsciiSetTest.cpp
@@ -37,7 +37,7 @@ TEST_CASE_NAMED(TAsciiSetTest, "System::Core::Misc::AsciiSet", "[ApplicationCont
 
 	CHECK_EQUALS(TEXT("Skip"), FAsciiSet::Skip(TEXT("  \t\tHello world!"), Whitespaces), TEXT("Hello world!"));
 	CHECK_EQUALS(TEXT("Skip"), FAsciiSet::Skip(TEXT("Hello world!"), Whitespaces), TEXT("Hello world!"));
-	CHECK_EQUALS(TEXT("Skip to extended ASCII"), FAsciiSet::Skip(" " "\xA0" " abc", Whitespaces), "\xA0" " abc");
+	CHECK_EQUALS(TEXT("Skip to extended ASCII"), FAsciiSet::Skip(" " "\xA0" " abc", Whitespaces), UTF8TEXTVIEW("\xA0" " abc"));
 	CHECK_EQUALS(TEXT("Skip to wide"), FAsciiSet::Skip(TEXT(" 变 abc"), Whitespaces), TEXT("变 abc"));
 	CHECK_EQUALS(TEXT("AdvanceToFirst"),	*FAsciiSet::FindFirstOrEnd("NonWhitespace\t \nNonWhitespace", Whitespaces), '\t');
 	CHECK_EQUALS(TEXT("AdvanceToLast"),	*FAsciiSet::FindLastOrEnd("NonWhitespace\t \nNonWhitespace", Whitespaces), '\n');
