Fix AutoRTFM::InitializeGlobalDataIfNecessary passing a pointer to a local to putenv, which retains the pointer after it returns rather than making a copy of the referenced string.

diff --git a/Engine/Source/Runtime/Core/Private/AutoRTFM/GlobalData.cpp b/Engine/Source/Runtime/Core/Private/AutoRTFM/GlobalData.cpp
--- a/Engine/Source/Runtime/Core/Private/AutoRTFM/GlobalData.cpp
+++ b/Engine/Source/Runtime/Core/Private/AutoRTFM/GlobalData.cpp
@@ -51,7 +51,7 @@ void InitializeGlobalDataIfNecessary()
         {
             GlobalData = new FGlobalData;
             FContext::InitializeGlobalData();
-            char Buffer[128];
+            static char Buffer[128];
             snprintf(Buffer, sizeof(Buffer), "%s=%p", EnvName, GlobalData);
             putenv(Buffer);
         }
