Work around UBSAN inserting its own prefix on functions with type hash information.

diff --git a/Engine/Source/Runtime/Core/Private/AutoRTFM/FunctionMapInlines.h b/Engine/Source/Runtime/Core/Private/AutoRTFM/FunctionMapInlines.h
--- a/Engine/Source/Runtime/Core/Private/AutoRTFM/FunctionMapInlines.h
+++ b/Engine/Source/Runtime/Core/Private/AutoRTFM/FunctionMapInlines.h
@@ -20,12 +20,24 @@ inline void* FunctionMapLookup(void* OldFunction, const char* Where)
 	// 16-bits of the function pointer address as a magic constant check
 	// to give us a much higher confidence that there is actually a closed
 	// variant pointer residing 8-bytes before our function address.
-	const uint64 PrefixData = *(reinterpret_cast<uint64*>(OldFunction) - 1);
-
+	uint64 PrefixData;
+	memcpy(&PrefixData, reinterpret_cast<char*>(OldFunction) - sizeof(uint64), sizeof(uint64));
 	if (LIKELY(Constants::MagicMike == (PrefixData & 0xffff000000000000)))
 	{
 		return reinterpret_cast<void*>(PrefixData & 0x0000ffffffffffff);
 	}
+	// UBSAN adds a type hash prefix to the function as a "prologue" that
+	// ends preceding our Magic Mike prefix. They use 0xc105cafe in the
+	// lower 32-bits to distinguish the 64-bit word containing their type
+	// hash. If we see it, check the preceding 64-bit word for our prefix.
+	else if (UNLIKELY(0xc105cafe == (PrefixData & 0x00000000ffffffff)))
+	{
+		memcpy(&PrefixData, reinterpret_cast<char*>(OldFunction) - sizeof(uint64) * 2, sizeof(uint64));
+		if (LIKELY(Constants::MagicMike == (PrefixData & 0xffff000000000000)))
+		{
+			return reinterpret_cast<void*>(PrefixData & 0x0000ffffffffffff);
+		}
+	}
 
 	// Instead fall back to the slower function map lookup.
     void* const Result = FunctionMapTryLookup(OldFunction);
