[AutoRTFM] Self-register closed variants of ASAN functions used by FTransaction::IsOnStack()

Fixes runtime failures when running tests with ASAN.

diff --git a/Engine/Source/Runtime/Core/Private/AutoRTFM/StdLibSupport.cpp b/Engine/Source/Runtime/Core/Private/AutoRTFM/StdLibSupport.cpp
--- a/Engine/Source/Runtime/Core/Private/AutoRTFM/StdLibSupport.cpp
+++ b/Engine/Source/Runtime/Core/Private/AutoRTFM/StdLibSupport.cpp
@@ -20,6 +20,24 @@
 #include "Windows/WindowsHWrapper.h"
 #endif
 
+#if PLATFORM_HAS_ASAN_INCLUDE
+#	if defined(__SANITIZE_ADDRESS__)
+#		define AUTORTFM_ASAN_ENABLED 1
+#	elif defined(__has_feature)
+#		if __has_feature(address_sanitizer)
+#			define AUTORTFM_ASAN_ENABLED 1
+#		endif
+#	endif
+#endif
+
+#ifndef AUTORTFM_ASAN_ENABLED
+#define AUTORTFM_ASAN_ENABLED 0
+#endif
+
+#if AUTORTFM_ASAN_ENABLED
+#include <sanitizer/asan_interface.h>
+#endif
+
 namespace AutoRTFM
 {
 
@@ -90,6 +108,11 @@ static void ThrowIfFormatContainsPercentN(const wchar_t* Format)
 
 }
 
+#if AUTORTFM_ASAN_ENABLED
+UE_AUTORTFM_REGISTER_SELF_FUNCTION(__asan_addr_is_in_fake_stack);
+UE_AUTORTFM_REGISTER_SELF_FUNCTION(__asan_get_current_fake_stack);
+#endif
+
 void* RTFM_memcpy(void* Dst, const void* Src, size_t Size) throw()
 {
     return Memcpy(Dst, Src, Size, FContext::Get());
