Fixed assert in TFixedAllocator::CalculateSlackGrow to guard against overflow, which was broken in CL# 35869509.

diff --git a/Engine/Source/Runtime/Core/Public/Containers/ContainerAllocationPolicies.h b/Engine/Source/Runtime/Core/Public/Containers/ContainerAllocationPolicies.h
--- a/Engine/Source/Runtime/Core/Public/Containers/ContainerAllocationPolicies.h
+++ b/Engine/Source/Runtime/Core/Public/Containers/ContainerAllocationPolicies.h
@@ -1274,7 +1274,7 @@ public:
 		FORCEINLINE SizeType CalculateSlackGrow(SizeType NewMax, SizeType CurrentMax, SIZE_T NumBytesPerElement) const
 		{
 			// Ensure the requested allocation will fit in the inline data area.
-			check(CurrentMax <= NumInlineElements);
+			check(NewMax <= NumInlineElements);
 			return NumInlineElements;
 		}
 
