Fix UBSAN error from coercing uninitialized memory to a bool in FArchive::operator<<(FArchive&, bool&).

diff --git a/Engine/Source/Runtime/Core/Public/Serialization/Archive.h b/Engine/Source/Runtime/Core/Public/Serialization/Archive.h
--- a/Engine/Source/Runtime/Core/Public/Serialization/Archive.h
+++ b/Engine/Source/Runtime/Core/Public/Serialization/Archive.h
@@ -1463,7 +1463,15 @@ public:
 		else
 #endif
 		{
-			uint32 OldUBoolValue = D ? 1 : 0;
+			uint32 OldUBoolValue = 0;
+
+			if (!Ar.IsLoading())
+			{
+				// Only read from the input if we aren't loading its value.
+				// This avoids triggering UBSAN by unnecessarily coercing uninitialized memory to a bool.
+				OldUBoolValue = D ? 1 : 0;
+			}
+
 			Ar.Serialize(&OldUBoolValue, sizeof(OldUBoolValue));
 
 			if (OldUBoolValue > 1)
