[BitArray]
 - Add custom bitwise combine function

#rnx

diff --git a/Engine/Source/Runtime/Core/Public/Containers/BitArray.h b/Engine/Source/Runtime/Core/Public/Containers/BitArray.h
--- a/Engine/Source/Runtime/Core/Public/Containers/BitArray.h
+++ b/Engine/Source/Runtime/Core/Public/Containers/BitArray.h
@@ -1266,6 +1266,16 @@ public:
 		return Result;
 	}
 
+	/**
+	 * Perform a custom bitwise operation on this bit array with another. This array receives the result and will be sized max(A.Num(), B.Num()).
+	 */
+	template <typename OtherAllocator, typename BinaryOpType>
+	TBitArray& CombineWithBinaryOp(const TBitArray<OtherAllocator>& InOther, EBitwiseOperatorFlags InFlags, BinaryOpType&& InBinaryOp)
+	{
+		UE::Core::Private::BitwiseOperatorImpl(InOther, *this, InFlags, Forward<BinaryOpType>(InBinaryOp));
+		return *this;
+	}
+
 	/**
 	 * Perform a bitwise AND on this bit array with another. This array receives the result and will be sized max(A.Num(), B.Num()).
 	 */
