Fix some instances of SwitchUnhandledEnumerator warnings in UnitConversion.inl. There is no change in behavior.
The warning is not enabled for Core, but the header might be included by modules where it is enabled.

diff --git a/Engine/Source/Runtime/Core/Public/Math/UnitConversion.inl b/Engine/Source/Runtime/Core/Public/Math/UnitConversion.inl
--- a/Engine/Source/Runtime/Core/Public/Math/UnitConversion.inl
+++ b/Engine/Source/Runtime/Core/Public/Math/UnitConversion.inl
@@ -109,21 +109,36 @@ T FUnitConversion::Convert(T InValue, EUnit From, EUnit To)
 		{
 			double NewValue = InValue;
 			// Put it into kelvin
-			switch (From)
+			if (From == EUnit::Celsius)
 			{
-				case EUnit::Celsius:			NewValue = NewValue + 273.15f;					break;
-				case EUnit::Farenheit:			NewValue = (NewValue + 459.67f) * 5.f/9.f;		break;
-				default: 																		break;
+				NewValue = NewValue + 273.15f;
+			}
+			else if (From == EUnit::Farenheit)
+			{
+				NewValue = (NewValue + 459.67f) * 5.f/9.f;
 			}
 			// And out again
-			switch (To)
+			if (To == EUnit::Celsius)
+			{
+				return NewValue - 273.15f;
+			}
+			else if (To == EUnit::Farenheit)
+			{
+				return NewValue * 9.f/5.f - 459.67f;
+			}
+			else
 			{
-				case EUnit::Celsius:			return NewValue - 273.15f;
-				case EUnit::Farenheit:			return NewValue * 9.f/5.f - 459.67f;
-				default: 						return NewValue;
+				return NewValue;
 			}
 		}
 
+		case EUnitType::Impulse:
+		case EUnitType::LuminousIntensity:
+		case EUnitType::Illuminance:
+		case EUnitType::Luminance:
+		case EUnitType::PixelDensity:
+		case EUnitType::ExposureValue:
+		case EUnitType::NumberOf:
 		default:							return InValue;
 	}
 }
