Fix static mesh DDC key generation to prevent an issue where non-Nanite platforms could cache static meshes with Nanite data built from LOD0 instead of the high resolution source, causing Nanite platforms to have incorrect geometry when loading from the cache.

#allowintegration

diff --git a/Engine/Source/Runtime/Core/Private/UObject/FortniteReleaseBranchCustomObjectVersion.cpp b/Engine/Source/Runtime/Core/Private/UObject/FortniteReleaseBranchCustomObjectVersion.cpp
--- a/Engine/Source/Runtime/Core/Private/UObject/FortniteReleaseBranchCustomObjectVersion.cpp
+++ b/Engine/Source/Runtime/Core/Private/UObject/FortniteReleaseBranchCustomObjectVersion.cpp
@@ -11,6 +11,6 @@ TMap<FGuid, FGuid> FFortniteReleaseBranchCustomObjectVersion::GetSystemGuids()
 	SystemGuids.Add(DevGuids.MaterialTranslationDDCVersion, FGuid("888A502A66F6401580B27451F211C46B"));
 	SystemGuids.Add(DevGuids.GLOBALSHADERMAP_DERIVEDDATA_VER, FGuid("0A0A1BF5D2A94105BB395EF0DA7CAEB9"));
 	SystemGuids.Add(DevGuids.MATERIALSHADERMAP_DERIVEDDATA_VER, FGuid("C0E1F59386894C9D855B85178B2553CB"));
-	SystemGuids.Add(DevGuids.STATICMESH_DERIVEDDATA_VER, FGuid("9463A91F546C41FCBFF114787E9E7E32"));
+	SystemGuids.Add(DevGuids.STATICMESH_DERIVEDDATA_VER, FGuid("0424034ECF2E4BADAA42E8C9AC4BF603"));
 	return SystemGuids;
 }
diff --git a/Engine/Source/Runtime/Engine/Private/StaticMesh.cpp b/Engine/Source/Runtime/Engine/Private/StaticMesh.cpp
--- a/Engine/Source/Runtime/Engine/Private/StaticMesh.cpp
+++ b/Engine/Source/Runtime/Engine/Private/StaticMesh.cpp
@@ -3456,8 +3456,8 @@ static FString BuildStaticMeshDerivedDataKeySuffix(const ITargetPlatform* Target
 	}
 #endif
 
-	// Add hi-res mesh description into DDC key
-	if (!Mesh->GetHiResSourceModel().GetMeshDescriptionBulkData()->IsEmpty())
+	// Add hi-res mesh description into DDC key (only on platforms that build from it)
+	if (DoesTargetPlatformSupportNanite(TargetPlatform) && !Mesh->GetHiResSourceModel().GetMeshDescriptionBulkData()->IsEmpty())
 	{
 		KeySuffix += Mesh->GetHiResSourceModel().GetMeshDescriptionBulkData()->GetIdString();
 	}
