[Core]
* Removed Concurrency visualizer external profiler extension because it haven't worked for a long time and there are better tools available

diff --git a/Engine/Source/Runtime/Core/Core.Build.cs b/Engine/Source/Runtime/Core/Core.Build.cs
--- a/Engine/Source/Runtime/Core/Core.Build.cs
+++ b/Engine/Source/Runtime/Core/Core.Build.cs
@@ -229,31 +229,6 @@ public class Core : ModuleRules
 			PrivateDependencyModuleNames.Add("SuperLuminal");
 		}
 
-		// Detect if the Concurrency Viewer Extension is installed
-		if (Target.Platform.IsInGroup(UnrealPlatformGroup.Windows))
-		{
-			bool VSCVDefined = false;
-			string VisualStudioInstallation = Target.WindowsPlatform.IDEDir;
-			if (VisualStudioInstallation != null && VisualStudioInstallation != string.Empty && Directory.Exists(VisualStudioInstallation))
-			{
-				string SubFolderName = @"Common7\IDE\Extensions\rf2nfg00.o0t\SDK\Native\Inc";
-				string IncludeDirectory = Path.Combine(VisualStudioInstallation, SubFolderName);
-
-				if (File.Exists(Path.Combine(IncludeDirectory, "cvmarkers.h"))
-					&& Target.Configuration != UnrealTargetConfiguration.Shipping)
-				{
-					PrivateIncludePaths.Add(IncludeDirectory);
-					PublicDefinitions.Add("WITH_CONCURRENCYVIEWER_PROFILER=1");
-					VSCVDefined = true;
-				}
-			}
-
-			if (!VSCVDefined)
-			{
-				PublicDefinitions.Add("WITH_CONCURRENCYVIEWER_PROFILER=0");
-			}
-		}
-
 		if (Target.bWithDirectXMath && Target.Platform == UnrealTargetPlatform.Win64)
 		{
 			PublicDefinitions.Add("WITH_DIRECTXMATH=1");
diff --git a/Engine/Source/Runtime/Core/Private/ProfilingDebugging/Windows/ConcurrencyViewerExternalProfiler.cpp b/Engine/Source/Runtime/Core/Private/ProfilingDebugging/Windows/ConcurrencyViewerExternalProfiler.cpp
deleted file mode 100644
