Add AutoRTFM::CascadingRetryTransaction to let you force a transaction nest to retry and run a callback after the initial transaction aborts, but before the new one begins.

#okforversepublic

diff --git a/Engine/Source/Programs/AutoRTFMTests/Private/API.cpp b/Engine/Source/Programs/AutoRTFMTests/Private/API.cpp
--- a/Engine/Source/Programs/AutoRTFMTests/Private/API.cpp
+++ b/Engine/Source/Programs/AutoRTFMTests/Private/API.cpp
@@ -3,7 +3,7 @@
 #include "API.h"
 #include "Catch2Includes.h"
 
-#include <AutoRTFM/AutoRTFM.h>
+#include <AutoRTFM/AutoRTFMTesting.h>
 #include <memory>
 #include <thread>
 
@@ -1139,3 +1139,97 @@ TEST_CASE("API.IsOnCurrentTransactionStack")
 		});
 	});
 }
+
+TEST_CASE("API.CascadingRetryTransaction")
+{
+	SECTION("Callback not called outside transaction")
+	{
+		AutoRTFM::CascadingRetryTransaction([&] { FAIL("Unreachable"); });
+	}
+
+	SECTION("Non-nested committed transaction")
+	{
+		bool bFirst = true;
+		AutoRTFM::Testing::Commit([&]
+		{
+			if (bFirst)
+			{
+				bFirst = false;
+				AutoRTFM::CascadingRetryTransaction([&]
+				{
+					// The write above would have been undone.
+					REQUIRE(bFirst);
+					bFirst = false;
+				});
+			}
+		});
+		REQUIRE(!bFirst);
+	}
+
+	SECTION("Non-nested aborted transaction")
+	{
+		bool bFirst = true;
+		AutoRTFM::Testing::Abort([&]
+		{
+			if (bFirst)
+			{
+				bFirst = false;
+				AutoRTFM::CascadingRetryTransaction([&]
+				{
+					// The write above would have been undone.
+					REQUIRE(bFirst);
+					bFirst = false;
+				});
+			}
+
+			AutoRTFM::AbortTransaction();
+		});
+		REQUIRE(!bFirst);
+	}
+
+	SECTION("Nested committed transaction")
+	{
+		bool bFirst = true;
+		AutoRTFM::Testing::Commit([&]
+		{
+			AutoRTFM::Testing::Commit([&]
+			{
+				if (bFirst)
+				{
+					bFirst = false;
+					AutoRTFM::CascadingRetryTransaction([&]
+					{
+						// The write above would have been undone.
+						REQUIRE(bFirst);
+						bFirst = false;
+					});
+				}
+			});
+		});
+		REQUIRE(!bFirst);
+	}
+
+	SECTION("Nested aborted transaction")
+	{
+		bool bFirst = true;
+		AutoRTFM::Testing::Commit([&]
+		{
+			AutoRTFM::Testing::Abort([&]
+			{
+				if (bFirst)
+				{
+					bFirst = false;
+					AutoRTFM::CascadingRetryTransaction([&]
+					{
+						// The write above would have been undone.
+						REQUIRE(bFirst);
+						bFirst = false;
+					});
+				}
+
+				AutoRTFM::AbortTransaction();
+			});
+		});
+		REQUIRE(!bFirst);
+	}
+}
diff --git a/Engine/Source/Runtime/Core/Private/AutoRTFM/API.cpp b/Engine/Source/Runtime/Core/Private/AutoRTFM/API.cpp
--- a/Engine/Source/Runtime/Core/Private/AutoRTFM/API.cpp
+++ b/Engine/Source/Runtime/Core/Private/AutoRTFM/API.cpp
@@ -457,6 +457,13 @@ AUTORTFM_CORE_API void RTFM_autortfm_record_open_write(void*, size_t)
 	UE_LOG(LogAutoRTFM, Fatal, TEXT("The function `autortfm_record_open_write` was called from closed code."));
 }
 
+AUTORTFM_CORE_API void RTFM_CascadingRetryTransactionInternal(TFunction<void()> && RunAfterAbortBeforeRetryWork)
+{
+	FContext* Context = FContext::Get();
+	ASSERT(Context->GetStatus() == EContextStatus::OnTrack);
+	Context->CascadingRetryTransaction(true, MoveTemp(RunAfterAbortBeforeRetryWork));
+}
+
 AUTORTFM_CORE_API void RTFM_OnCommitInternal(TFunction<void()> && Work)
 {
 	FContext* Context = FContext::Get();
@@ -782,6 +789,14 @@ extern "C" bool autortfm_is_on_current_transaction_stack(void* Ptr)
 	}
 }
 
+void ForTheRuntime::CascadingRetryTransactionInternal(TFunction<void()>&& RunAfterAbortBeforeRetryWork)
+{
+	if (autortfm_is_closed())
+	{
+		return RTFM_CascadingRetryTransactionInternal(MoveTemp(RunAfterAbortBeforeRetryWork));
+	}
+}
+
 void ForTheRuntime::OnCommitInternal(TFunction<void()> && Work)
 {
 	if (autortfm_is_closed())
diff --git a/Engine/Source/Runtime/Core/Private/AutoRTFM/Context.cpp b/Engine/Source/Runtime/Core/Private/AutoRTFM/Context.cpp
--- a/Engine/Source/Runtime/Core/Private/AutoRTFM/Context.cpp
+++ b/Engine/Source/Runtime/Core/Private/AutoRTFM/Context.cpp
@@ -189,7 +189,6 @@ ETransactionResult FContext::AbortTransaction(bool bIsClosed, bool bIsCascading)
 		}
 	}
 
-	
 	if (bIsClosed)
 	{
 		Throw();
@@ -198,6 +197,13 @@ ETransactionResult FContext::AbortTransaction(bool bIsClosed, bool bIsCascading)
 	return Result;
 }
 
+void FContext::CascadingRetryTransaction(bool bIsClosed, TFunction<void()>&& RunAfterAbortBeforeRetryWork)
+{
+	ASSERT(!CascadingRetryTransactionWork);
+	CascadingRetryTransactionWork = MoveTemp(RunAfterAbortBeforeRetryWork);
+	AbortTransaction(bIsClosed, true);
+}
+
 bool FContext::IsAborting() const
 {
 	switch (Status)
@@ -474,6 +480,24 @@ ETransactionResult FContext::Transact(void (*InstrumentedFunction)(void*), void*
 				break;
 
 			case EContextStatus::AbortedByCascade:
+				if (CascadingRetryTransactionWork)
+				{
+					// If we get here, it means we are retrying our transaction nest
+					// after a cascade. So clean up the transaction to get it ready
+					// for re-execution.
+					ClearTransactionStatus();
+					CurrentTransaction->Reset();
+
+					// Then call the retry callback to do whatever work the user
+					// required be done before we attempt to re-execute.
+					CascadingRetryTransactionWork();
+					CascadingRetryTransactionWork.Reset();
+
+					// Then get rolling!
+					CurrentTransaction->SetClosedActive();
+					continue;
+				}
+
 				Result = ETransactionResult::AbortedByCascade;
 				break;
 
diff --git a/Engine/Source/Runtime/Core/Private/AutoRTFM/Context.h b/Engine/Source/Runtime/Core/Private/AutoRTFM/Context.h
--- a/Engine/Source/Runtime/Core/Private/AutoRTFM/Context.h
+++ b/Engine/Source/Runtime/Core/Private/AutoRTFM/Context.h
@@ -36,6 +36,7 @@ public:
 	bool IsAborting() const;
 	bool IsTransactional() const;
 	bool IsCommittingOrAborting() const;
+	void CascadingRetryTransaction(bool bIsClosed, TFunction<void()> && RunAfterAbortBeforeRetryWork);
 
     // Record that a write is about to occur at the given LogicalAddress of Size bytes.
     void RecordWrite(void* LogicalAddress, size_t Size);
@@ -159,6 +160,7 @@ private:
     EContextStatus Status{EContextStatus::Idle};
 	uint32 CurrentThreadId{ FPlatformTLS::InvalidTlsSlot };
 	uint32 StackLocalInitializerDepth = 0;
+	TFunction<void()> CascadingRetryTransactionWork;
 };
 
 } // namespace AutoRTFM
diff --git a/Engine/Source/Runtime/Core/Public/AutoRTFM/AutoRTFM.h b/Engine/Source/Runtime/Core/Public/AutoRTFM/AutoRTFM.h
--- a/Engine/Source/Runtime/Core/Public/AutoRTFM/AutoRTFM.h
+++ b/Engine/Source/Runtime/Core/Public/AutoRTFM/AutoRTFM.h
@@ -469,6 +469,7 @@ enum class EMemoryValidationLevel : uint8_t
 #if UE_AUTORTFM
 namespace ForTheRuntime
 {
+	UE_AUTORTFM_API void CascadingRetryTransactionInternal(TFunction<void()>&& RunAfterAbortBeforeRetryWork);
 	UE_AUTORTFM_API void OnCommitInternal(TFunction<void()>&& Work);
 	UE_AUTORTFM_API void OnAbortInternal(TFunction<void()>&& Work);
 	UE_AUTORTFM_API void PushOnAbortHandlerInternal(const void* Key, TFunction<void()>&& Work);
@@ -789,6 +790,31 @@ template<typename FunctorType> [[nodiscard]] static UE_AUTORTFM_FORCEINLINE ECon
 			const_cast<void*>(static_cast<const void*>(&Functor))));
 }
 
+#if UE_AUTORTFM
+// Force a transaction nest to be retried. Once the transaction has aborted,
+// it will call `RunAfterAbortBeforeRetryWork`, before retrying the
+// transaction. This is an expensive operation and should thus be used with
+// extreme caution. If this is called outside of a transaction then the
+// `RunAfterAbortBeforeRetryWork` callback is never executed.
+template<typename FunctorType> static UE_AUTORTFM_FORCEINLINE void CascadingRetryTransaction(FunctorType&& RunAfterAbortBeforeRetryWork)
+{
+	if (autortfm_is_closed())
+	{
+		ForTheRuntime::CascadingRetryTransactionInternal(std::forward<FunctorType>(RunAfterAbortBeforeRetryWork));
+	}
+}
+#else
+// Force a transaction nest to be retried. Once the transaction has aborted,
+// it will call `RunAfterAbortBeforeRetryWork`, before retrying the
+// transaction. This is an expensive operation and should thus be used with
+// extreme caution. If this is called outside of a transaction then the
+// `RunAfterAbortBeforeRetryWork` callback is never executed.
+template<typename FunctorType> static UE_AUTORTFM_FORCEINLINE void CascadingRetryTransaction(FunctorType&& RunAfterAbortBeforeRetryWork)
+{
+	UE_AUTORTFM_UNUSED(RunAfterAbortBeforeRetryWork);
+}
+#endif
+
 #if UE_AUTORTFM
 // Have some work happen when this transaction commits. For nested transactions,
 // this just adds the work to the work deferred until the outer nest's commit.
