Do not call MemoryTrimDelegate from LowLevelAllocatorMemoryTrimDelegate in MB2 as that might lead to deadlocks

diff --git a/Engine/Source/Runtime/Core/Private/HAL/MallocBinned2.cpp b/Engine/Source/Runtime/Core/Private/HAL/MallocBinned2.cpp
--- a/Engine/Source/Runtime/Core/Private/HAL/MallocBinned2.cpp
+++ b/Engine/Source/Runtime/Core/Private/HAL/MallocBinned2.cpp
@@ -469,8 +469,6 @@ void FMallocBinned2::OnMallocInitialized()
 #if UE_USE_VERYLARGEPAGEALLOCATOR
 	FCoreDelegates::GetLowLevelAllocatorMemoryTrimDelegate().AddLambda([this]()
 		{
-			FCoreDelegates::GetMemoryTrimDelegate().Broadcast();
-
 			FScopeLock Lock(&ExternalAllocMutex);
 			CachedOSPageAllocator.FreeAll(&ExternalAllocMutex);
 		}
