- Fixed binary config, and added a cache of the ini files that exist at stage time per plugin (and global), so make faster scanning for .ini files at runtime (no need to touch disk or .pak file indexes)
- Added support for batched plugin mounting, which can greatly reduce the time spent in CalculateDiff while modifying existing branches

diff --git a/Engine/Source/Developer/PakFileUtilities/Private/PakFileUtilities.cpp b/Engine/Source/Developer/PakFileUtilities/Private/PakFileUtilities.cpp
--- a/Engine/Source/Developer/PakFileUtilities/Private/PakFileUtilities.cpp
+++ b/Engine/Source/Developer/PakFileUtilities/Private/PakFileUtilities.cpp
@@ -5396,21 +5396,90 @@ bool MakeBinaryConfig(const TCHAR* CmdLine)
 	}
 
 	FString ProjectDir = FPaths::GetPath(ProjectFile);
+	ProjectDir = ProjectDir.Replace(TEXT("\\"), TEXT("/"));
+	ProjectDir = ProjectDir.Replace(*FPaths::RootDir(), TEXT("../../../"));
 
-	FConfigCacheIni Config(EConfigCacheType::Temporary);
+	FConfigCacheIni Config(EConfigCacheType::Temporary, *PlatformName);
 	FConfigContext Context = FConfigContext::ReadIntoConfigSystem(&Config, PlatformName);
 	Context.ProjectConfigDir = FPaths::Combine(ProjectDir, TEXT("Config/"));
+	Context.bIsMakingBinaryConfig = true;
+
+	auto LoadIniCache = [] (FConfigCacheIni& Config, const TArray<FString>&PluginAndConfigsList, bool bUseLocalPaths)
+	{
+		// reset any current cache
+		delete Config.StagedGlobalConfigCache;
+		Config.StagedGlobalConfigCache = nullptr;
+		Config.StagedPluginConfigCache.Empty();
+
+		FName CurrentPlugin = NAME_None;
+		for (const FString& Line : PluginAndConfigsList)
+		{
+			// indented lines are staged paths to config files
+			if (Line.StartsWith(TEXT("  ")))
+			{
+				// a line looks like: "  D:\UE\Engine\Foo.png|Engine/Foo.png"
+				int PipeLoc;
+				Line.FindChar('|', PipeLoc);
+
+				FString CachedPath;
+				if (bUseLocalPaths)
+				{
+					CachedPath = Line.Mid(2, PipeLoc - 2);
+					// must be relative so that the config discovery code matches paths correctly
+					CachedPath = CachedPath.Replace(TEXT("\\"), TEXT("/"));
+					CachedPath = CachedPath.Replace(*FPaths::RootDir(), TEXT("../../../"));
+				}
+				else
+				{
+					CachedPath = FString(TEXT("../../../")) + Line.Mid(PipeLoc + 1);
+				}
+
+				// the first set of lines are configs without a plugin, so those are global
+				if (CurrentPlugin == NAME_None)
+				{
+					if (Config.StagedGlobalConfigCache == nullptr)
+					{
+						Config.StagedGlobalConfigCache = new TSet<FString>();
+					}
+					Config.StagedGlobalConfigCache->Add(CachedPath);
+				}
+				else
+				{
+					Config.StagedPluginConfigCache[CurrentPlugin].Add(CachedPath);
+				}
+			}
+			else
+			{
+				CurrentPlugin = *Line;
+				Config.StagedPluginConfigCache.FindOrAdd(CurrentPlugin);
+			}
+		}
+	};
+
+	// read in the cache of plugin ini files - 
+	TArray<FString> PluginAndConfigsList;
+	if (!FFileHelper::LoadFileToStringArray(PluginAndConfigsList, *StagedPluginsFile))
+	{
+		UE_LOG(LogMakeBinaryConfig, Fatal, TEXT("Failed to load required file %s"), *StagedPluginsFile);
+	}
+
+	// IMPORTANT - we use local paths to config files when reading in the known configs to limit
+	// the binary config to only contain the staged ini files (ie Restricted files that exist when 
+	// building the binary config, but aren't staged, and shouldn't be in the binary)
+	LoadIniCache(Config, PluginAndConfigsList, true);
+
 	Config.InitializeKnownConfigFiles(Context);
 
-	// removing for now, because this causes issues with some plugins not getting ini files merged in
-	// IPluginManager::Get().IntegratePluginsIntoConfig(Config, *GEngineIni, *PlatformName, *StagedPluginsFile);
+	// then we need to save out with staged paths so that runtime can know what files exist in the cache
+	LoadIniCache(Config, PluginAndConfigsList, false);
+	Config.bIsReadyForUse = true;
 
 	// pull out deny list entries
 
 	TArray<FString> KeyDenyListStrings;
 	TArray<FString> SectionsDenyList;
-	GConfig->GetArray(TEXT("/Script/UnrealEd.ProjectPackagingSettings"), TEXT("IniKeyDenylist"), KeyDenyListStrings, GGameIni);
-	GConfig->GetArray(TEXT("/Script/UnrealEd.ProjectPackagingSettings"), TEXT("IniSectionDenylist"), SectionsDenyList, GGameIni);
+	Config.GetArray(TEXT("/Script/UnrealEd.ProjectPackagingSettings"), TEXT("IniKeyDenylist"), KeyDenyListStrings, GGameIni);
+	Config.GetArray(TEXT("/Script/UnrealEd.ProjectPackagingSettings"), TEXT("IniSectionDenylist"), SectionsDenyList, GGameIni);
 	TArray<FName> KeysDenyList;
 	for (const FString& Key : KeyDenyListStrings)
 	{
diff --git a/Engine/Source/Programs/AutomationTool/Scripts/CopyBuildToStagingDirectory.Automation.cs b/Engine/Source/Programs/AutomationTool/Scripts/CopyBuildToStagingDirectory.Automation.cs
--- a/Engine/Source/Programs/AutomationTool/Scripts/CopyBuildToStagingDirectory.Automation.cs
+++ b/Engine/Source/Programs/AutomationTool/Scripts/CopyBuildToStagingDirectory.Automation.cs
@@ -1546,6 +1546,9 @@ namespace AutomationScripts
 
 					// Stage plugin config files
 					List<KeyValuePair<StagedFileReference, FileReference>> StagedPlugins = SC.FilesToStage.UFSFiles.Where(x => x.Value.HasExtension(".uplugin")).ToList();
+					List<KeyValuePair<StagedFileReference, FileReference>> GlobalConfigCache = SC.FilesToStage.UFSFiles.Where(x => x.Value.HasExtension(".ini")).ToList();
+					Dictionary<StagedFileReference, List<KeyValuePair<StagedFileReference, FileReference>>> PluginConfigCache = new();
+
 					List<PluginInfo> AvailablePlugins = new List<PluginInfo>();
 					if (StagedPlugins.Count > 0 && PlatformExtensionsToStage.Count > 0)
 					{
@@ -1554,11 +1557,29 @@ namespace AutomationScripts
 					}
 					foreach (KeyValuePair<StagedFileReference, FileReference> StagedPlugin in StagedPlugins)
 					{
-						//PluginDescriptor Descriptor = PluginDescriptor.FromFile(StagedPlugin.Value);
+						PluginDescriptor Descriptor = PluginDescriptor.FromFile(StagedPlugin.Value);
+						bool bTrackConfigs = ShouldCreateBinaryConfig(Params, SC) && !Descriptor.bIsPluginExtension;
+
+						// always insert an entry, even if no configs are found, this tells runtime that there are no 
+						// configs for this plugin, and not to look (we still may look on disk for downloaded plugins, etc)
+						if (bTrackConfigs)
+						{
+							PluginConfigCache.Add(StagedPlugin.Key, new List<KeyValuePair<StagedFileReference, FileReference>> ());
+						}
+
 						DirectoryReference PluginConfigDirectory = DirectoryReference.Combine(StagedPlugin.Value.Directory, "Config");
 						if (DirectoryReference.Exists(PluginConfigDirectory))
 						{
 							SC.StageFiles(StagedFileType.UFS, PluginConfigDirectory, "*.ini", StageFilesSearch.AllDirectories);
+
+							// extensions will be added in the child loop below
+							if (bTrackConfigs)
+							{
+								// add any staged files we just added, to the ini cache (not easy to get the staged files list that we just added)
+								IEnumerable<KeyValuePair<StagedFileReference, FileReference>> PluginConfigs = SC.FilesToStage.UFSFiles
+									.Where(x => x.Value.HasExtension(".ini") && x.Key.IsUnderDirectory(StagedPlugin.Key.Directory));
+								PluginConfigCache[StagedPlugin.Key].AddRange(PluginConfigs);
+							}
 						}
 
 						// look in platform extension locations of this plugin for config files
@@ -1577,7 +1598,16 @@ namespace AutomationScripts
 										if (DirectoryReference.Exists(PlatformConfigDirectory))
 										{
 											// stage info Plugin/Config/PlatformName/
-											SC.StageFiles(StagedFileType.UFS, PlatformConfigDirectory, "*.ini", StageFilesSearch.AllDirectories, StagedDirectoryReference.Combine(StagedPlugin.Key.Directory, "Config", PlatformExtensionToStage));
+											StagedDirectoryReference StagedExtensionDir = StagedDirectoryReference.Combine(StagedPlugin.Key.Directory, "Config", PlatformExtensionToStage);
+											SC.StageFiles(StagedFileType.UFS, PlatformConfigDirectory, "*.ini", StageFilesSearch.AllDirectories, StagedExtensionDir);
+
+											if (bTrackConfigs)
+											{
+												// add any staged files we just added, to the ini cache (not easy to get the staged files list that we just added)
+												IEnumerable<KeyValuePair<StagedFileReference, FileReference>> PluginConfigs = SC.FilesToStage.UFSFiles
+													.Where(x => x.Value.HasExtension(".ini") && x.Key.IsUnderDirectory(StagedExtensionDir));
+												PluginConfigCache[StagedPlugin.Key].AddRange(PluginConfigs);
+											}
 										}
 									}
 								}
@@ -1838,22 +1868,30 @@ namespace AutomationScripts
 					// data from a previous run, but still want to stage executables, etc. SkipPak indicates we _would_ have made a 
 					// pak file, but we aren't this time, so we can assume the binary config is already inside the .pak file (and
 					// the UnrealEditor.exe may not even exist if we are making a target-code-only build on top of existing build)
-					if (Params.GenerateOptimizationData && !Params.SkipPak)
+					if (ShouldCreateBinaryConfig(Params, SC) && !Params.SkipPak)
 					{
+						List<string> ConfigCacheInputFileContents = new List<string>();
+
+						// get the configs that aren't in a plugin
+						var GlobalConfigs = GlobalConfigCache.Where(x => !PluginConfigCache.Any(y => y.Value.Any(z => z.Key == x.Key)));
+						ConfigCacheInputFileContents.AddRange(GlobalConfigs.Select(x => $"  {x.Value}|{x.Key}"));
+
 						// get the list of plugins that need to be processed
-						List<string> PluginFilesForTarget = new List<string>();
-						foreach (KeyValuePair<StagedFileReference, FileReference> StagedPlugin in StagedPlugins)
+						foreach (KeyValuePair<StagedFileReference, List<KeyValuePair<StagedFileReference, FileReference>>> ConfigInfo in PluginConfigCache)
 						{
-							PluginFilesForTarget.Add(StagedPlugin.Value.FullName);
+							// store the plugin name (not path) and any configs it may have had staged
+							ConfigCacheInputFileContents.Add(Path.GetFileNameWithoutExtension(ConfigInfo.Key.Name));
+							ConfigCacheInputFileContents.AddRange(ConfigInfo.Value.Select(x => $"  {x.Value}|{x.Key}"));
 						}
-						FileReference PluginListFile = FileReference.Combine(SC.ProjectRoot, "Intermediate", "Config", "PluginList.txt");
-						DirectoryReference.CreateDirectory(PluginListFile.Directory);
-						File.WriteAllLines(PluginListFile.FullName, PluginFilesForTarget.ToArray());
+
+						FileReference ConfigCacheInputFile = FileReference.Combine(SC.ProjectRoot, "Intermediate", "Config", "ConfigCacheInput.txt");
+						DirectoryReference.CreateDirectory(ConfigCacheInputFile.Directory);
+						File.WriteAllLines(ConfigCacheInputFile.FullName, ConfigCacheInputFileContents.ToArray());
 
 						// run the commandlet to generate a binary file
 						String TargetPlatformName = ThisPlatform.GetCookPlatform(Params.DedicatedServer, Params.Client);
 						FileReference OutputFile = FileReference.Combine(SC.ProjectRoot, "Intermediate", "Config", TargetPlatformName, "BinaryConfig.ini");
-						String UnrealPakParams = String.Format("MakeBinaryConfig -Project=\"{0}\" -Platform={1} -OutputFile=\"{2}\" -StagedPluginsFile=\"{3}\"", SC.RawProjectPath, ConfigHierarchy.GetIniPlatformName(ThisPlatform.IniPlatformType), OutputFile.FullName, PluginListFile.FullName);
+						String UnrealPakParams = String.Format("MakeBinaryConfig -Project=\"{0}\" -Platform={1} -OutputFile=\"{2}\" -StagedPluginsFile=\"{3}\"", SC.RawProjectPath, ConfigHierarchy.GetIniPlatformName(ThisPlatform.IniPlatformType), OutputFile.FullName, ConfigCacheInputFile.FullName);
 
 						RunUnrealPak(Params.RawProjectPath, UnrealPakParams, "MakeBinaryConfig");
 
@@ -5215,6 +5253,20 @@ namespace AutomationScripts
 			return ChunkManifestPath;
 		}
 
+		private static bool ShouldCreateBinaryConfig(ProjectParams Params, DeploymentContext SC)
+		{
+			if (Params.GenerateOptimizationData)
+			{
+				return true;
+			}
+
+			// check the project settings if not set
+			ConfigHierarchy PlatformGameConfig = ConfigCache.ReadHierarchy(ConfigHierarchyType.Game, DirectoryReference.FromFile(Params.RawProjectPath), SC.StageTargetPlatform.IniPlatformType, SC.CustomConfig);
+			bool bMakeBinaryConfig;
+			PlatformGameConfig.GetBool("/Script/UnrealEd.ProjectPackagingSettings", "bMakeBinaryConfig", out bMakeBinaryConfig);
+			return bMakeBinaryConfig;
+		}
+
 		private static bool ShouldCreateIoStoreContainerFiles(ProjectParams Params, DeploymentContext SC)
 		{
 			if (!ShouldCreatePak(Params, SC))
diff --git a/Engine/Source/Runtime/Core/Private/Misc/ConfigCacheIni.cpp b/Engine/Source/Runtime/Core/Private/Misc/ConfigCacheIni.cpp
--- a/Engine/Source/Runtime/Core/Private/Misc/ConfigCacheIni.cpp
+++ b/Engine/Source/Runtime/Core/Private/Misc/ConfigCacheIni.cpp
@@ -2,6 +2,7 @@
 
 #include "Misc/ConfigCacheIni.h"
 #include "Misc/ConfigContext.h"
+#include "Misc/ConfigHierarchy.h"
 #include "Misc/ConfigUtilities.h"
 #include "Containers/StringView.h"
 #include "Misc/DateTime.h"
@@ -508,7 +509,7 @@ void FixupArrayOfStructKeysForSection(SectionType* Section, const FString& Secti
 /**
  * Check if an ini file exists, allowing a delegate to determine if it will handle loading it
  */
-/*static*/ bool DoesConfigFileExistWrapper(const TCHAR* IniFile, const TSet<FString>* IniCacheSet)
+/*static*/ bool DoesConfigFileExistWrapper(const TCHAR* IniFile, const TSet<FString>* IniCacheSet, const TSet<FString>* PrimaryConfigFileCache, const TSet<FString>* SecondaryConfigFileCache)
 {
 	// will any delegates return contents via TSPreLoadConfigFileDelegate()?
 	int32 ResponderCount = 0;
@@ -519,12 +520,19 @@ void FixupArrayOfStructKeysForSection(SectionType* Section, const FString& Secti
 		return true;
 	}
 
-	FString OverriddenIniFile = IniFile;
-	if (OverrideFileFromCommandline(OverriddenIniFile))
+	FString IniFileString = IniFile;
+	if (OverrideFileFromCommandline(IniFileString))
 	{
 		return true;
 	}
 
+	// check staged cache (likely for plugin configs)
+	if (PrimaryConfigFileCache != nullptr || SecondaryConfigFileCache != nullptr)
+	{
+		return (PrimaryConfigFileCache && PrimaryConfigFileCache->Contains(IniFileString)) ||
+			(SecondaryConfigFileCache && SecondaryConfigFileCache->Contains(IniFileString));
+	}
+
 	// So far, testing on cooked consoles, cooked desktop, and the Editor
 	// works fine.
 	// However, there was an issue where INIs wouldn't be found during cooking,
@@ -533,7 +541,6 @@ void FixupArrayOfStructKeysForSection(SectionType* Section, const FString& Secti
 	// to prevent using the cache in that case ensures full consistency.
 	if (IniCacheSet && FPlatformProperties::RequiresCookedData())
 	{ 
-		const FString IniFileString(IniFile);
 		const bool bFileExistsCached = IniCacheSet->Contains(IniFileString);
 
 		// This code can be uncommented if we expect there are INIs that are not being
@@ -554,6 +561,7 @@ void FixupArrayOfStructKeysForSection(SectionType* Section, const FString& Secti
 	}
 
 	// otherwise just look for the normal file to exist
+	UE_LOG(LogConfig, VeryVerbose, TEXT("Looking for a config file without a staged cache for %s"), IniFile);
 	const bool bFileExistsCached = IFileManager::Get().FileExists(IniFile);
 	return bFileExistsCached;
 }
@@ -1403,6 +1411,7 @@ void FConfigFile::ProcessCommand(FConfigSection* Section, FStringView SectionNam
 		case FConfigValue::EValueType::Clear:
 			// Remove if present.
 			Section->Remove(Key);
+			Section->CompactStable();
 			break;
 		case FConfigValue::EValueType::ArrayOfStructKey:
 			// track a key to show uniqueness for arrays of structs
@@ -1489,15 +1498,16 @@ void FillFileFromBuffer(FileType* File, FStringView Buffer, bool bHandleSymbolCo
 		{
 			Done = true;
 		}
+		int32 LineLen = TheLine.Len();
 		TCHAR* Start = const_cast<TCHAR*>(*TheLine);
 
 		// Strip trailing spaces from the current line
-		while( *Start && FChar::IsWhitespace(Start[FCString::Strlen(Start)-1]) )
+		while( *Start && FChar::IsWhitespace(Start[LineLen-1]) )
 		{
-			Start[FCString::Strlen(Start)-1] = TEXT('\0');
+			Start[LineLen-1] = TEXT('\0');
+			LineLen--;
 		}
 
-		// look for the ;METADATA line as the first line in the file, which can override passed in settings
 		//@todo (UE-214768) Comment this back in
 		// if (!bHasHandledMetadata)
 		// {
@@ -1519,11 +1529,13 @@ void FillFileFromBuffer(FileType* File, FStringView Buffer, bool bHandleSymbolCo
 			
 		// }
 		// If the first character in the line is [ and last char is ], this line indicates a section name
-		if( *Start=='[' && Start[FCString::Strlen(Start)-1]==']' )
+		if( *Start=='[' && Start[LineLen-1]==']' )
 		{
 			// Remove the brackets
 			Start++;
-			Start[FCString::Strlen(Start)-1] = TEXT('\0');
+			LineLen--;
+			Start[LineLen-1] = TEXT('\0');
+			LineLen--;
 
 			// If we don't have an existing section by this name, add one
 			CurrentSectionName = Start;
@@ -1590,17 +1602,21 @@ void FillFileFromBuffer(FileType* File, FStringView Buffer, bool bHandleSymbolCo
 				{
 					// Terminate the property name, advancing past the =
 					*Value++ = TEXT('\0');
+					// update LineLen since we just chopped it
+					LineLen = (int32)(Value - Start - 1);
 
 					// strip leading whitespace from the property name
 					while (*Start && FChar::IsWhitespace(*Start))
 					{
 						Start++;
+						LineLen--;
 					}
 
 					// ~ is a packaging and should be skipped at runtime
 					if (Start[0] == '~')
 					{
 						Start++;
+						LineLen--;
 					}
 
 					if (bHandleSymbolCommands)
@@ -1609,6 +1625,7 @@ void FillFileFromBuffer(FileType* File, FStringView Buffer, bool bHandleSymbolCo
 						if (Cmd == '+' || Cmd == '-' || Cmd == '.' || Cmd == '!' || Cmd == '@' || Cmd == '*')
 						{
 							Start++;
+							LineLen--;
 						}
 						else
 						{
@@ -1626,9 +1643,10 @@ void FillFileFromBuffer(FileType* File, FStringView Buffer, bool bHandleSymbolCo
 					}
 
 					// Strip trailing spaces from the property name.
-					while (*Start && FChar::IsWhitespace(Start[FCString::Strlen(Start)-1]))
+					while (*Start && FChar::IsWhitespace(Start[LineLen-1]))
 					{
-						Start[FCString::Strlen(Start)-1] = TEXT('\0');
+						Start[LineLen-1] = TEXT('\0');
+						LineLen--;
 					}
 
 					const TCHAR* KeyName = Start;
@@ -3088,17 +3106,28 @@ void FConfigBranch::RunOnEachCommandStream(TFunction<void(FConfigCommandStream&
 }
 
 
-bool FConfigBranch::AddDynamicLayerToHierarchy(const FString& Filename, FConfigModificationTracker* ModificationTracker)
+bool FConfigBranch::AddDynamicLayerToHierarchy(const FString& Filename, FConfigModificationTracker* ModificationTracker, TSet<FString>* GlobalConfigFileCache, TSet<FString>* PluginConfigFileCache)
 {
-	if (!DoesConfigFileExistWrapper(*Filename))
+	if (!DoesConfigFileExistWrapper(*Filename, nullptr, GlobalConfigFileCache, PluginConfigFileCache))
 	{
 		return false;
 	}
 
-	return AddDynamicLayersToHierarchy({ Filename }, NAME_None, DynamicLayerPriority::Unknown, ModificationTracker);
+	return AddDynamicLayersToHierarchy({ { Filename, NAME_None, (uint16)DynamicLayerPriority::Unknown } }, ModificationTracker, GlobalConfigFileCache, PluginConfigFileCache);
 }
 
 bool FConfigBranch::AddDynamicLayersToHierarchy(const TArray<FString>& Filenames, FName Tag, DynamicLayerPriority Priority, FConfigModificationTracker* ModificationTracker)
+{
+	TArray<FDynamicLayerInfo> Layers;
+	Layers.Reserve(Filenames.Num());
+	for (const FString& Filename : Filenames)
+	{
+		Layers.Add( { Filename, Tag, (uint16)Priority } );
+	}
+	return AddDynamicLayersToHierarchy(Layers, ModificationTracker);
+}
+
+bool FConfigBranch::AddDynamicLayersToHierarchy(const TArray<FDynamicLayerInfo>& Layers, FConfigModificationTracker* ModificationTracker, TSet<FString>* GlobalConfigFileCache, TSet<FString>* PluginConfigFileCache)
 {
 	static bool bDumpIniLoadInfo = FParse::Param(FCommandLine::Get(), TEXT("dumpiniloads"));
 
@@ -3109,19 +3138,19 @@ bool FConfigBranch::AddDynamicLayersToHierarchy(const TArray<FString>& Filenames
 	FConfigCommandStream Patch;
 
 	TArray<FConfigCommandStream*> AddedLayers;
-	for (const FString& Filename : Filenames)
+	for (const FDynamicLayerInfo& Layer : Layers)
 	{
-		UE_CLOG(bDumpIniLoadInfo, LogConfig, Display, TEXT("Looking for file: %s"), *Filename);
+		UE_CLOG(bDumpIniLoadInfo, LogConfig, Display, TEXT("Looking for file: %s"), *Layer.Filename);
 
-		UE_LOG(LogConfig, Verbose, TEXT("Adding Dynamic layer %s to Branch %s"), *Filename, *IniName.ToString());
+		UE_LOG(LogConfig, Verbose, TEXT("Adding Dynamic layer %s to Branch %s"), *Layer.Filename, *IniName.ToString());
 
-		if (!DoesConfigFileExistWrapper(*Filename))
+		if (!DoesConfigFileExistWrapper(*Layer.Filename, nullptr, GlobalConfigFileCache, PluginConfigFileCache))
 		{
 			UE_LOG(LogConfig, Verbose, TEXT("  .. doesn't exist!"));
 			continue;
 		}
 
-		UE_CLOG(bDumpIniLoadInfo, LogConfig, Display, TEXT("   Found %s!"), *Filename);
+		UE_CLOG(bDumpIniLoadInfo, LogConfig, Display, TEXT("   Found %s!"), *Layer.Filename);
 
 		if (AddedLayers.Num() == 0)
 		{
@@ -3131,10 +3160,10 @@ bool FConfigBranch::AddDynamicLayersToHierarchy(const TArray<FString>& Filenames
 		
 		// make and read in the layer
 		FConfigCommandStream* DynamicLayer = new FConfigCommandStream;
-		FillFileFromDisk(DynamicLayer, Filename, true);
-		DynamicLayer->Priority = (uint16)Priority;
-		DynamicLayer->Filename = Filename;
-		DynamicLayer->Tag = Tag;
+		FillFileFromDisk(DynamicLayer, Layer.Filename, true);
+		DynamicLayer->Priority = Layer.Priority;
+		DynamicLayer->Filename = Layer.Filename;
+		DynamicLayer->Tag = Layer.Tag;
 
 		// remember in local array, then figure out how to remember it permanently
 		AddedLayers.Add(DynamicLayer);
@@ -3182,8 +3211,10 @@ bool FConfigBranch::AddDynamicLayersToHierarchy(const TArray<FString>& Filenames
 					{
 						UE_LOG(LogConfig, Verbose, TEXT("  .. .. .. tracking cvars"), *Pair.Key);
 						
+						TMap<FName, FConfigSection>& PerIniCVarSections = CVarTracker->CVarEntriesPerBranchPerTag.FindOrAdd(Layer.Tag);
+						FConfigSection& TrackedCVarSection = PerIniCVarSections.FindOrAdd(IniName);
+						
 						const FConfigSectionMap& ModifiedCVars = (const FConfigSectionMap&)Pair.Value;  
-						FConfigSection& TrackedCVarSection = CVarTracker->CVarEntriesPerBranch.FindOrAdd(IniName);
 						for (const TPair<FName, FConfigValue>& CVarPair : ModifiedCVars)
 						{
 							UE_LOG(LogConfig, Verbose, TEXT("  .. .. .. .. %s = %s"), *CVarPair.Key.ToString(), *CVarPair.Value.GetSavedValue());
@@ -3195,7 +3226,7 @@ bool FConfigBranch::AddDynamicLayersToHierarchy(const TArray<FString>& Filenames
 			}
 			if (ModificationTracker->bTrackLoadedFiles)
 			{
-				ModificationTracker->LoadedFiles.Add(Filename);
+				ModificationTracker->LoadedFiles.Add(Layer.Filename);
 			}
 		}
 	}
@@ -3291,7 +3322,7 @@ bool FConfigBranch::AddDynamicLayerStringToHierarchy(const FString& Filename, co
 					FConfigSection NewSection;
 					// copy just the SectionMap parts
 					(FConfigSectionMap&)NewSection = (const FConfigSectionMap&)Pair.Value;
-					CVarTracker->CVarEntriesPerBranch.Add(IniName, NewSection);
+					CVarTracker->CVarEntriesPerBranchPerTag.FindOrAdd(Tag).Add(IniName, NewSection);
 				}
 			}
 		}
@@ -3387,15 +3418,15 @@ bool FConfigBranch::RemoveDynamicLayersFromHierarchy(const TArray<FString>& File
 
 }
 
-void FConfigBranch::RemoveTagFromHierarchy(FName Tag, FConfigModificationTracker* ModificationTracker)
+void FConfigBranch::RemoveTagsFromHierarchy(const TArray<FName>& Tags, FConfigModificationTracker* ModificationTracker)
 {
 	// gather tagged layers
 	TArray<FString> LayersToRemove;
 	for (DynamicLayerList::TIterator Node(DynamicLayers.GetHead()); Node; ++Node)
 	{
-		if (Node->Tag == Tag)
+		if (Tags.Contains(Node->Tag))
 		{
-			UE_LOG(LogConfig, Verbose, TEXT("Removing dynamic layer %s from branch %s with tag %s"), *Node->Filename, *IniName.ToString(), *Tag.ToString());
+			UE_LOG(LogConfig, Verbose, TEXT("Removing dynamic layer %s from branch %s with tag %s"), *Node->Filename, *IniName.ToString(), *Node->Tag.ToString());
 			// @todo make a version that takes CommandStreams, not Filenames, for speed
 			LayersToRemove.Add(Node->Filename);
 		}
@@ -3445,11 +3476,13 @@ void FConfigBranch::SafeReload()
 		DynamicLayers.RemoveNode(HeadNode, false);
 		EmptiedDynamicLayers.AddTail(HeadNode);
 	}
+	TArray<FDynamicLayerInfo> ReloadInfos;
 	for (FConfigBranch::DynamicLayerList::TIterator Node(EmptiedDynamicLayers.GetHead()); Node; ++Node)
 	{
 		FConfigCommandStream& S = *Node.GetNode()->GetValue();
-		AddDynamicLayersToHierarchy({ S.Filename }, S.Tag, (DynamicLayerPriority)S.Priority);
+		ReloadInfos.Add( { S.Filename, S.Tag, S.Priority } );
 	}
+	AddDynamicLayersToHierarchy(ReloadInfos);
 
 	UE_LOG(LogConfig, Log, TEXT("Branch '%s' had been unloaded. Reloading on-demand took %.2fms"), *IniName.ToString(), (FPlatformTime::Seconds() - StartTime) * 1000.0f);
 }
@@ -4537,29 +4570,44 @@ void FConfigCacheIni::Exit()
 #endif
 }
 
-void FConfigCacheIni::DumpFile(FOutputDevice& Ar, const FString& Filename, const FConfigFile& File)
+static void DumpBranch(FOutputDevice& Ar, const FConfigBranch& Branch)
 {
-	Ar.Logf(TEXT("FileName: %s"), *Filename);
+	Ar.Logf(TEXT("Branch Name: %s"), *Branch.IniName.ToString());
+	Ar.Logf(TEXT("Branch Filename: %s"), *Branch.IniPath);
+	Ar.Logf(TEXT("Branch Static Hierarchy:"));
+	for (const TPair<int32, FString>& Pair: Branch.Hierarchy)
+	{
+		Ar.Logf(TEXT("  %s"), *Pair.Value);
+	}
+	if (Branch.DynamicLayers.GetHead())
+	{
+		Ar.Logf(TEXT("Branch Dynamic Layers:"));
+		for (FConfigBranch::DynamicLayerList::TIterator Node(Branch.DynamicLayers.GetHead()); Node; ++Node)
+		{
+			Ar.Logf(TEXT("  %s"), *Node.GetNode()->GetValue()->Filename);
+		}
+	}
 
+	Ar.Logf(TEXT("Branch Values:"));
 	// sort the sections (and keys below) for easier diffing
 	TArray<FString> SectionKeys;
-	File.GetKeys(SectionKeys);
+	Branch.InMemoryFile.GetKeys(SectionKeys);
 	SectionKeys.Sort();
 	for (const FString& SectionKey : SectionKeys)
 	{
-		const FConfigSection& Sec = File[SectionKey];
+		const FConfigSection* Sec = Branch.InMemoryFile.FindSection(SectionKey);
 		Ar.Logf(TEXT("   [%s]"), *SectionKey);
 
 		TArray<FName> Keys;
-		Sec.GetKeys(Keys);
+		Sec->GetKeys(Keys);
 		Keys.Sort(FNameLexicalLess());
 		for (FName Key : Keys)
 		{
 			TArray<FConfigValue> Values;
-			Sec.MultiFind(Key, Values, true);
+			Sec->MultiFind(Key, Values, true);
 			for (const FConfigValue& Value : Values)
 			{
-				Ar.Logf(TEXT("   %s=%s"), *Key.ToString(), *Value.GetValue());
+				Ar.Logf(TEXT("   %s=%s"), *Key.ToString(), *Value.GetSavedValueForWriting());
 			}
 		}
 
@@ -4573,7 +4621,7 @@ void FConfigCacheIni::Dump(FOutputDevice& Ar, const TCHAR* BaseIniName)
 	{
 		if (BaseIniName == nullptr || Branch.IniName == BaseIniName)
 		{
-			DumpFile(Ar, Branch.IniName.ToString(), Branch.InMemoryFile);
+			DumpBranch(Ar, Branch);
 		}
 	}
 
@@ -4585,7 +4633,7 @@ void FConfigCacheIni::Dump(FOutputDevice& Ar, const TCHAR* BaseIniName)
 	{
 		if (BaseIniName == nullptr || FPaths::GetBaseFilename(Key) == BaseIniName)
 		{
-			DumpFile(Ar, Key, OtherFiles[Key]->InMemoryFile);
+			DumpBranch(Ar, OtherFiles[Key]->InMemoryFile);
 		}
 	}
 }
@@ -5361,6 +5409,19 @@ void FConfigCacheIni::Serialize(FArchive& Ar)
 	Ar << bIsReadyForUse;
 	Ar << Type;
 	Ar << PlatformName;
+	Ar << StagedPluginConfigCache;
+
+	bool bHasGlobalCache = StagedGlobalConfigCache != nullptr;
+	Ar << bHasGlobalCache;
+	if (bHasGlobalCache)
+	{
+		if (Ar.IsLoading()) StagedGlobalConfigCache = new TSet<FString>();
+		// checking for null for SA validation
+		if (StagedGlobalConfigCache != nullptr)
+		{
+			Ar << *StagedGlobalConfigCache;
+		}
+	}
 }
 
 void FConfigCacheIni::SerializeStateForBootstrap_Impl(FArchive& Ar)
@@ -5505,20 +5566,20 @@ FArchive& operator<<(FArchive& Ar, FConfigBranch& ConfigBranch)
 {
 	Ar << ConfigBranch.bIsHierarchical;
 	Ar << ConfigBranch.InMemoryFile;
+	Ar << ConfigBranch.Hierarchy;
+	Ar << ConfigBranch.CombinedStaticLayers;
+	Ar << ConfigBranch.FinalCombinedLayers;
+	Ar << ConfigBranch.IniName;
+	Ar << ConfigBranch.IniPath;
 
 	// needed to count full memory usage
 	if (!Ar.IsPersistent())
 	{
-		Ar << ConfigBranch.IniName;
-		Ar << ConfigBranch.IniPath;
 		Ar << ConfigBranch.Platform;
 		Ar << ConfigBranch.SourceEngineConfigDir;
 		Ar << ConfigBranch.SourceProjectConfigDir;
-		Ar << ConfigBranch.Hierarchy;
 		Ar << ConfigBranch.StaticLayers;
 		Ar << ConfigBranch.SavedLayer;
-		Ar << ConfigBranch.CombinedStaticLayers;
-		Ar << ConfigBranch.FinalCombinedLayers;
 		Ar << ConfigBranch.CommandLineOverrides;
 		Ar << ConfigBranch.RuntimeChanges;
 		for (FConfigBranch::DynamicLayerList::TIterator Node(ConfigBranch.DynamicLayers.GetHead()); Node; ++Node)
@@ -5547,24 +5608,29 @@ bool FConfigCacheIni::CreateGConfigFromSaved(const TCHAR* Filename)
 		return false;
 	}
 
-	UE_LOG(LogInit, Display, TEXT("Loading binary GConfig..."));
+	UE_LOG(LogInit, Display, TEXT("Loading binary GConfig from %lld bytes of data..."), Size);
 
 	// serialize right out of the preloaded data
 	FLargeMemoryReader MemoryReader((uint8*)PreloadedData, Size);
-	FKnownConfigFiles Names;
+	MemoryReader.SetIsPersistent(true);
 	GConfig = new FConfigCacheIni(EConfigCacheType::Temporary, NAME_None /* Platform*/, true /* bInGloballyRegistered */);
 
 	// make an object that we can use to pass to delegates for any extra binary data they want to write
-//	FCoreDelegates::FExtraBinaryConfigData ExtraData(*GConfig, false);
+	FCoreDelegates::FExtraBinaryConfigData ExtraData(*GConfig, false);
 
 	GConfig->Serialize(MemoryReader);
 
-	// Forced to be disk backed so that GameUserSettings does get written out.
+	// fix up some things that weren't saved out
 	GConfig->Type = EConfigCacheType::DiskBacked;
-	MemoryReader << Names;// << ExtraData.Data;
+	for (FConfigBranch& Branch : GConfig->KnownFiles.Branches)
+	{
+		Branch.IniPath = FConfigCacheIni::GetDestIniFilename(*Branch.IniName.ToString(), nullptr, *FPaths::GeneratedConfigDir());
+	}
+	
+	MemoryReader << ExtraData.Data;
 
 	// now let the delegates pull their data out, after GConfig is set up
-//	FCoreDelegates::TSAccessExtraBinaryConfigData().Broadcast(ExtraData);
+	FCoreDelegates::TSAccessExtraBinaryConfigData().Broadcast(ExtraData);
 
 	{
 		TRACE_CPUPROFILER_EVENT_SCOPE(ConfigReadyForUseBroadcast);
@@ -5629,7 +5695,7 @@ static void InitializeConfigRemap()
 		// if there isn't an active project, then skip the project pass
 		if (Pass == 1 && FPaths::ProjectDir() == FPaths::EngineDir())
 		{
-			continue;
+			continue;\
 		}
 		
 		Context.Load(*FPaths::Combine(Pass == 0 ? FPaths::EngineDir() : FPaths::ProjectDir(), TEXT("Config/ConfigRedirects.ini")));
@@ -5659,6 +5725,13 @@ static void InitializeConfigRemap()
 
 void FConfigCacheIni::InitializeConfigSystem()
 {
+	// cache existence of a few key files that may be checked over and over
+	// this could be done with the Staged caches, but this will at least speed up a repeated FilExists check when not using binaryconfig
+	GConfigLayers[0].bHasCheckedExist = true;
+	GConfigLayers[0].bExists = DoesConfigFileExistWrapper(*FString(GConfigLayers[0].Path).Replace(TEXT("{ENGINE}/"), *FPaths::EngineDir()));
+	GPluginLayers[0].bHasCheckedExist = true;
+	GPluginLayers[0].bExists = DoesConfigFileExistWrapper(*FString(GPluginLayers[0].Path).Replace(TEXT("{ENGINE}/"), *FPaths::EngineDir()));
+
 	// assign the G***Ini strings for the known ini's
 	#define ASSIGN_GLOBAL_INI_STRING(IniName) G##IniName##Ini = FString(#IniName);
 		// GEngineIni = FString("Engine")
@@ -5678,11 +5751,7 @@ void FConfigCacheIni::InitializeConfigSystem()
 		// So the safest thing to do is to re-read the file after binary configs load.
 		Context.Load(TEXT("GameUserSettings"), GGameUserSettingsIni);
 
-#if WITH_EDITOR
-		// a cooked editor (cooked cooker more likely) can be initialized from binary config for speed, 
-		// but we still need the other files as well
 		LoadRemainingConfigFiles(Context);
-#endif
 		return;
 	}
 #endif
@@ -6301,130 +6370,169 @@ FTransactionallySafeRWLock FConfigFile::ConfigFileMapLock;
 
 void FConfigCacheIni::AddPluginToAllBranches(FName PluginName, FConfigModificationTracker* ModificationTracker)
 {
-	GConfig->AddPluginToBranches(PluginName, ModificationTracker);
+	AddMultiplePluginsToAllBranches({PluginName}, ModificationTracker);
+}
+
+void FConfigCacheIni::RemoveTagFromAllBranches(FName Tag, FConfigModificationTracker* ModificationTracker)
+{
+	RemoveMultipleTagsFromAllBranches({Tag}, ModificationTracker);
+}
+
+void FConfigCacheIni::AddMultiplePluginsToAllBranches(const TArray<FName>& PluginNames, FConfigModificationTracker* ModificationTracker)
+{
+	GConfig->AddPluginsToBranches(PluginNames, ModificationTracker);
 	
 #if ALLOW_OTHER_PLATFORM_CONFIG
 	FScopeLock Lock(&ConfigForPlatformLock);
 	// need to walk over the other platforms without calling ForPlatform because that could end up loading pending plugins
 	for (auto Pair : ConfigForPlatform)
 	{
-		Pair.Value->PendingModificationPlugins.Add(PluginName);
+		Pair.Value->PendingModificationPlugins.Append(PluginNames);
 	}
 #endif
 }
 
-void FConfigCacheIni::RemoveTagFromAllBranches(FName Tag, FConfigModificationTracker* ModificationTracker)
+void FConfigCacheIni::RemoveMultipleTagsFromAllBranches(const TArray<FName>& Tags, FConfigModificationTracker* ModificationTracker)
 {
-	GConfig->RemoveTagFromBranches(Tag, ModificationTracker);
+	GConfig->RemoveTagsFromBranches(Tags, ModificationTracker);
 	
 #if ALLOW_OTHER_PLATFORM_CONFIG
 	// need to walk over the other platforms without calling ForPlatform because that could end up loading pending plugins
 	for (auto Pair : ConfigForPlatform)
 	{
-		Pair.Value->RemoveTagFromBranches(Tag, ModificationTracker);
+		Pair.Value->RemoveTagsFromBranches(Tags, ModificationTracker);
 	}
 #endif
 }
 
-void FConfigCacheIni::AddPluginToBranches(FName PluginName, FConfigModificationTracker* ModificationTracker)
+void FConfigCacheIni::AddPluginsToBranches(const TArray<FName>& PluginNames, FConfigModificationTracker* ModificationTracker)
 {
 	// @todo make sure we are still pending
 	
-	FPluginInfo* PluginInfo = nullptr;
+	TMap<FConfigBranch*, TArray<FDynamicLayerInfo>> AllDynamicLayers;
+	
+	for (FName PluginName : PluginNames)
 	{
-		UE::TScopeLock Lock(RegisteredPluginsLock);
 		
-		PluginInfo = RegisteredPlugins.FindRef(PluginName);
-		if (PluginInfo == nullptr)
+		FPluginInfo* PluginInfo = nullptr;
 		{
-			UE_LOG(LogConfig, Warning, TEXT("Attempting to load a dynamic plugin (%s) that was not registered ahead of time!"), *PluginName.ToString());
-			return;
+			UE::TScopeLock Lock(RegisteredPluginsLock);
+			
+			PluginInfo = RegisteredPlugins.FindRef(PluginName);
+			if (PluginInfo == nullptr)
+			{
+				UE_LOG(LogConfig, Warning, TEXT("Attempting to load a dynamic plugin (%s) that was not registered ahead of time!"), *PluginName.ToString());
+				return;
+			}
 		}
-	}
+		
+		FString PluginConfigDir = FPaths::Combine(PluginInfo->PluginDir, TEXT("Config"));
+		FString PlatformNameStr(PlatformName.ToString());
+		TSet<FString> SlowPluginConfigs;
+		// if we already cached this plugin's configs offline, use it
+		TSet<FString>* PluginConfigs = StagedPluginConfigCache.Find(PluginName);
+		bool bNamesAreFullPaths = true;
+		if (PluginConfigs == nullptr)
+		{
+			bNamesAreFullPaths = false;
+			PluginConfigs = &SlowPluginConfigs;
 	
-	FString PlatformNameStr(PlatformName.ToString());
-	TArray<FString> PluginConfigs;
-	FString PluginConfigDir = FPaths::Combine(PluginInfo->PluginDir, TEXT("Config"));
-	FString PlatformConfigDir = FPaths::Combine(PluginConfigDir, PlatformNameStr);
-	IFileManager::Get().FindFiles(PluginConfigs, *PluginConfigDir, TEXT("ini"));
-	IFileManager::Get().FindFiles(PluginConfigs, *PlatformConfigDir, TEXT("ini"));
+			FString PlatformConfigDir = FPaths::Combine(PluginConfigDir, PlatformNameStr);
+			TArray<FString> LocalConfigs;
+			IFileManager::Get().FindFiles(LocalConfigs, *PluginConfigDir, TEXT("ini"));
+			IFileManager::Get().FindFiles(LocalConfigs, *PlatformConfigDir, TEXT("ini"));
 	
-#if !UE_BUILD_SHIPPING
-	for (const FString& F : PluginConfigs)
-	{
-		UE_LOG(LogConfig, Verbose, TEXT("Found config file %s in plugin dir %s"), *F, *PluginInfo->PluginDir);
-	}
-#endif // !UE_BUILD_SHIPPING
+			// if this plugin has any platform extensions, then we need to look in them for files, in so that we can load them
+			// even if there is no platform-less config file in the plugin itself
+			for (FString& ChildPluginDir : PluginInfo->ChildPluginDirs)
+			{
+				if (ChildPluginDir.Contains(*FString::Printf(TEXT("/%s/"), *PlatformNameStr)))
+				{
+					FString PlatformExtConfigDir = FPaths::Combine(ChildPluginDir, TEXT("Config"));
+					IFileManager::Get().FindFiles(LocalConfigs, *PlatformExtConfigDir, TEXT("ini"));
+				}
+			}
 	
-	// if this plugin has any platform extensions, then we need to look in them for files, in so that we can load them
-	// even if there is no platform-less config file in the plugin itself
-	for (FString& ChildPluginDir : PluginInfo->ChildPluginDirs)
-	{
-		if (ChildPluginDir.Contains(*FString::Printf(TEXT("/%s/"), *PlatformNameStr)))
-		{
-			FString PlatformExtConfigDir = FPaths::Combine(ChildPluginDir, TEXT("Config"));
-			IFileManager::Get().FindFiles(PluginConfigs, *PlatformExtConfigDir, TEXT("ini"));
+			SlowPluginConfigs = TSet<FString>(LocalConfigs);
 		}
-	}
-	
-	// make a single context that can be used for all the branches modified by this plugin
-	FConfigContext Context = FConfigContext::ReadIntoConfigSystem(this, PlatformNameStr);
-	Context.bIsForPluginModification = true;
-	Context.PluginModificationPriority = PluginInfo->Priority;
-	Context.bIncludeTagNameInBranchName = PluginInfo->bIncludePluginNameInBranchName;
-	Context.ChangeTracker = ModificationTracker;
-	
-	// find branches that are found in the plugin dir or it's platform dirs
-	FString StrippedPart = PluginInfo->bIncludePluginNameInBranchName ? PluginName.ToString() : FString();
-	FName CurrentPlatform(FPlatformProperties::IniPlatformName());
-	
-	TSet<FName> LoadedBranches;
-	for (const FString& ConfigFilename : PluginConfigs)
-	{
-		FName BranchName = *FPaths::GetBaseFilename(ConfigFilename).Replace(*StrippedPart, TEXT("")).Replace(*PlatformNameStr, TEXT(""));
-		if (LoadedBranches.Contains(BranchName))
+		
+#if !UE_BUILD_SHIPPING
+		for (const FString& F : *PluginConfigs)
 		{
-			continue;
+			UE_LOG(LogConfig, Verbose, TEXT("Found config file %s in plugin dir %s"), *F, *PluginInfo->PluginDir);
 		}
-		LoadedBranches.Add(BranchName);
+#endif // !UE_BUILD_SHIPPING
+
+		// make a single context that can be used for all the branches modified by this plugin
+		FConfigContext Context = FConfigContext::ReadIntoConfigSystem(this, PlatformNameStr);
+		Context.bIsForPluginModification = true;
+		Context.PluginModificationPriority = PluginInfo->Priority;
+		Context.bIncludeTagNameInBranchName = PluginInfo->bIncludePluginNameInBranchName;
+		Context.ChangeTracker = ModificationTracker;
+		
+		// find branches that are found in the plugin dir or it's platform dirs
+		FString StrippedPart = PluginInfo->bIncludePluginNameInBranchName ? PluginName.ToString() : FString();
+		FName CurrentPlatform(FPlatformProperties::IniPlatformName());
 		
-		// if we have been tracking loaded files, and we've already loaded this file, we can skip it (it would be the DefaultMyPlugin.ini type of file)
-		if (ModificationTracker && ModificationTracker->bTrackLoadedFiles)
+		TSet<FName> LoadedBranches;
+		for (const FString& ConfigFilename : *PluginConfigs)
 		{
-			FString FullPath = FPaths::Combine(PluginConfigDir, ConfigFilename);
-			if (ModificationTracker->LoadedFiles.Contains(FullPath))
+			FName BranchName = *FPaths::GetBaseFilename(ConfigFilename).Replace(*StrippedPart, TEXT("")).Replace(*PlatformNameStr, TEXT(""));
+			if (LoadedBranches.Contains(BranchName))
 			{
-				UE_LOG(LogConfig, Verbose, TEXT("Skipping already loaded file %s"), *FullPath);
 				continue;
 			}
+			LoadedBranches.Add(BranchName);
+			
+			// if we have been tracking loaded files, and we've already loaded this file, we can skip it (it would be the DefaultMyPlugin.ini type of file)
+			if (ModificationTracker && ModificationTracker->bTrackLoadedFiles)
+			{
+				FString FullPath = bNamesAreFullPaths ? ConfigFilename : FPaths::Combine(PluginConfigDir, ConfigFilename);
+				if (ModificationTracker->LoadedFiles.Contains(FullPath))
+				{
+					UE_LOG(LogConfig, Verbose, TEXT("Skipping already loaded file %s"), *FullPath);
+					continue;
+				}
+			}
+			
+			// look up the branch to see if we can modify it
+			FConfigBranch* Branch = FindBranch(BranchName, FString());
+			if (Branch == nullptr)
+			{
+				// don't log out for other platforms, because they are being loadd later and the ModificationTracker doesn't have full context
+				// @todo: removed this because it was causing the FilterPlugin.ini files to be logged _a lot_.
+				UE_CLOG(PlatformName == CurrentPlatform, LogConfig, Verbose, TEXT("Found unknown .ini file %s in plugindir %s"), *ConfigFilename, *PluginInfo->PluginDir);
+				continue;
+			}
+			
+			UE_LOG(LogConfig, Verbose, TEXT("Modifying branch %s with plugin ini %s"), *BranchName.ToString(), *ConfigFilename);
+			
+			Context.ConfigFileTag = PluginName;
+			
+			// by setting this, Load will not load the layers, it will just call this function 
+			Context.HandleLayersFunction = [&AllDynamicLayers, Branch](const TArray<FDynamicLayerInfo>& Layers)
+			{
+				AllDynamicLayers.FindOrAdd(Branch).Append(Layers);
+			};
+			Context.Load(*BranchName.ToString());
 		}
-		
-		// look up the branch to see if we can modify it
-		FConfigBranch* Branch = FindBranch(BranchName, FString());
-		if (Branch == nullptr)
-		{
-			// don't log out for other platforms, because they are being loadd later and the ModificationTracker doesn't have full context
-			// @todo: removed this because it was causing the FilterPlugin.ini files to be logged _a lot_.
-			UE_CLOG(PlatformName == CurrentPlatform, LogConfig, Verbose, TEXT("Found unknown .ini file %s in plugindir %s"), *ConfigFilename, *PluginInfo->PluginDir);
-			continue;
-		}
-		
-		UE_LOG(LogConfig, Verbose, TEXT("Modifying branch %s with plugin ini %s"), *BranchName.ToString(), *ConfigFilename);
+	}
 
-		Context.ConfigFileTag = PluginName;
-		Context.Load(*BranchName.ToString());
+	for (TPair<FConfigBranch*, TArray<FDynamicLayerInfo>> Pair : AllDynamicLayers)
+	{
+		Pair.Key->AddDynamicLayersToHierarchy(Pair.Value, ModificationTracker);
 	}
 }
 
-void FConfigCacheIni::RemoveTagFromBranches(FName Tag, FConfigModificationTracker* ModificationTracker)
+void FConfigCacheIni::RemoveTagsFromBranches(const TArray<FName>& Tags, FConfigModificationTracker* ModificationTracker)
 {
 	for (uint8 KnownIndex = 0; KnownIndex < (uint8)EKnownIniFile::NumKnownFiles; KnownIndex++)
 	{
-		KnownFiles.Branches[KnownIndex].RemoveTagFromHierarchy(Tag, ModificationTracker);
+		KnownFiles.Branches[KnownIndex].RemoveTagsFromHierarchy(Tags, ModificationTracker);
 	}
 	for (auto& Pair : OtherFiles)
 	{
-		Pair.Value->RemoveTagFromHierarchy(Tag, ModificationTracker);
+		Pair.Value->RemoveTagsFromHierarchy(Tags, ModificationTracker);
 	}
 }
 
@@ -6558,11 +6666,8 @@ FConfigCacheIni* FConfigCacheIni::ForPlatform(FName PlatformName)
 	}
 
 	{
-		for (FName PluginName : PendingModificationPlugins)
-		{
-			// delayed plugin injection
-			PlatformConfig->AddPluginToBranches(PluginName, nullptr);
-		}
+		// delayed plugin injection
+		PlatformConfig->AddPluginsToBranches(PendingModificationPlugins, nullptr);
 	}
 
 	return PlatformConfig;
@@ -6729,7 +6834,7 @@ class FIniExec : public FSelfRegisteringExec
 				{
 					UE::DynamicConfig::PerformDynamicConfig("HotFixText", [Branch, FileName](FConfigModificationTracker* ChangeTracker)
 					{
-						Branch->AddDynamicLayersToHierarchy({ FileName }, "HotfixTest", DynamicLayerPriority::Hotfix, ChangeTracker);
+						Branch->AddDynamicLayersToHierarchy( { { FileName , "HotfixTest", (uint16)DynamicLayerPriority::Hotfix } }, ChangeTracker);
 					});
 				}
 			}
diff --git a/Engine/Source/Runtime/Core/Private/Misc/ConfigContext.cpp b/Engine/Source/Runtime/Core/Private/Misc/ConfigContext.cpp
--- a/Engine/Source/Runtime/Core/Private/Misc/ConfigContext.cpp
+++ b/Engine/Source/Runtime/Core/Private/Misc/ConfigContext.cpp
@@ -27,8 +27,19 @@ namespace
 	FString CurrentIniVersionString = TEXT("CurrentIniVersion");
 	const TCHAR* SectionsToSaveString = TEXT("SectionsToSave");
 	const TCHAR* SaveAllSectionsKey = TEXT("bCanSaveAllSections");
+
+	// some settings for hierarchy keys
+	constexpr int32 KeyFlag_UseGlobalCache = 1;
+	constexpr int32 KeyFlag_UsePluginCache = 2;
+	constexpr int32 KeyFlag_AssumeExists = 4;
+
+	constexpr int32 NumLayerBits = 6;
+	constexpr int32 NumExpansionBits = 6;
+	constexpr int32 NumPlatformBits = 6;
+	constexpr int32 NumFlagsBits = 3;
 }
 
+
 FConfigContext::FConfigContext(FConfigCacheIni* InConfigSystem, bool InIsHierarchicalConfig, const FString& InPlatform, FConfigFile* DestConfigFile)
 	: ConfigSystem(InConfigSystem)
 	, Platform(InPlatform)
@@ -86,13 +97,16 @@ void FConfigContext::CachePaths()
 		{
 			UE::TScopeLock Lock(FConfigCacheIni::RegisteredPluginsLock);
 			
-			FConfigCacheIni::FPluginInfo* PluginInfo = ConfigSystem->RegisteredPlugins.FindRef(ConfigFileTag == NAME_None ? FName(*BaseIniName) : ConfigFileTag);
+			FName PluginName = ConfigFileTag == NAME_None ? FName(*BaseIniName) : ConfigFileTag;
+			FConfigCacheIni::FPluginInfo* PluginInfo = ConfigSystem->RegisteredPlugins.FindRef(PluginName);
 			if (PluginInfo != nullptr)
 			{
 				bIsForPlugin = true;
 				PluginRootDir = PluginInfo->PluginDir;
 				ChildPluginBaseDirs = PluginInfo->ChildPluginDirs;
+				StagedPluginConfigCache = ConfigSystem->StagedPluginConfigCache.Find(PluginName);
 			}
+			StagedGlobalConfigCache = ConfigSystem->StagedGlobalConfigCache;
 		}
 		
 		// for the hierarchy replacements, we need to have a directory called Config - or we will have to do extra processing for these non-standard cases
@@ -477,9 +491,9 @@ bool FConfigContext::PrepareForLoad(bool& bPerformLoad)
 			check(!DestIniFilename.IsEmpty());
 			
 			Branch = &ConfigSystem->AddNewBranch(DestIniFilename);
-			Branch->bIsHierarchical = bIsHierarchicalConfig;
 		}
 	}
+	Branch->bIsHierarchical = bIsHierarchicalConfig;
 
 	if (IsInGameThread()) GPrepareForLoadTime += FPlatformTime::Seconds();
 
@@ -526,8 +540,23 @@ bool FConfigContext::PerformLoad()
 		TArray<FString> GatheredFiles;
 		AddStaticLayersToHierarchy(&GatheredFiles);
 	
-		// now add them all as one operation (optimization to not perform unnecessary duplicated work for each file)
-		Branch->AddDynamicLayersToHierarchy(GatheredFiles, ConfigFileTag, PluginModificationPriority, ChangeTracker);
+		TArray<FDynamicLayerInfo> Layers;
+		for (FString& File : GatheredFiles)
+		{
+			Layers.Add({ File, ConfigFileTag, (uint16)PluginModificationPriority });
+			UE_LOG(LogConfig, Verbose, TEXT("Loading plugin modification file %s"), *File);
+		}
+	
+		// call a function to handle the layers if desired
+		if (HandleLayersFunction != nullptr)
+		{
+			HandleLayersFunction(Layers);
+		}
+		else
+		{
+			// now add them all as one operation (optimization to not perform unnecessary duplicated work for each file)
+			Branch->AddDynamicLayersToHierarchy(Layers, ChangeTracker, ConfigSystem->StagedGlobalConfigCache, ConfigSystem->StagedPluginConfigCache.Find(ConfigFileTag));
+		}
 		if (IsInGameThread()) GPerformLoadTime += FPlatformTime::Seconds();
 		return true;
 	}
@@ -747,7 +776,7 @@ void FConfigContext::AddStaticLayersToHierarchy(TArray<FString>* GatheredLayerFi
 	EConfigExpansionFlags ExpansionMode = EConfigExpansionFlags::ForUncooked;
 	FConfigLayer* Layers = GConfigLayers;
 	int32 NumLayers = UE_ARRAY_COUNT(GConfigLayers);
-	if (FPlatformProperties::RequiresCookedData())
+	if (FPlatformProperties::RequiresCookedData() || bIsMakingBinaryConfig)
 	{
 		ExpansionMode = EConfigExpansionFlags::ForCooked;
 	}
@@ -785,12 +814,31 @@ void FConfigContext::AddStaticLayersToHierarchy(TArray<FString>* GatheredLayerFi
 			continue;
 		}
 
+		// put some info into the key for later use
+		int FlagPartOfKey = 0;
+		if (EnumHasAnyFlags(Layer.Flag, EConfigLayerFlags::UseGlobalConfigCache))
+		{
+			FlagPartOfKey |= KeyFlag_UseGlobalCache;
+		}
+		else if (EnumHasAnyFlags(Layer.Flag, EConfigLayerFlags::UsePluginConfigCache))
+		{
+			FlagPartOfKey |= KeyFlag_UsePluginCache;
+		}
+		if (Layer.bHasCheckedExist)
+		{
+			if (!Layer.bExists)
+			{
+				continue;
+			}
+			FlagPartOfKey |= KeyFlag_AssumeExists;
+		}
+
 		// start replacing basic variables
 		FString LayerPath = PerformBasicReplacements(Layer.Path, *BaseIniName);
 		bool bHasPlatformTag = LayerPath.Contains(TEXT("{PLATFORM}"));
 		bool bHasEngineTag = LayerPath.StartsWith(TEXT("{ENGINE}"));
 		bool bHasProjectTag = !bHasEngineTag && LayerPath.StartsWith(TEXT("{PROJECT}"));
-		
+		bool bHasPluginTag = bIsForPlugin && !bHasEngineTag && !bHasProjectTag && LayerPath.StartsWith(TEXT("{PLUGIN}"));
 
 		// expand if it it has {ED} or {EF} expansion tags
 		if (!EnumHasAnyFlags(Layer.Flag, EConfigLayerFlags::NoExpand))
@@ -878,7 +926,7 @@ void FConfigContext::AddStaticLayersToHierarchy(TArray<FString>* GatheredLayerFi
 					}
 					else
 					{
-						Branch->Hierarchy.AddStaticLayer(PlatformPath, LayerIndex, ExpansionIndex, PlatformIndex);
+						Branch->Hierarchy.AddStaticLayer(PlatformPath, LayerIndex, ExpansionIndex, PlatformIndex, FlagPartOfKey);
 					}
 				}
 			}
@@ -911,7 +959,7 @@ void FConfigContext::AddStaticLayersToHierarchy(TArray<FString>* GatheredLayerFi
 			}
 			else
 			{
-				Branch->Hierarchy.AddStaticLayer(FinalPath, LayerIndex);
+				Branch->Hierarchy.AddStaticLayer(FinalPath, LayerIndex, 0, 0, FlagPartOfKey);
 			}
 		}
 	}
@@ -939,18 +987,23 @@ bool FConfigContext::LoadIniFileHierarchy()
 	// Traverse ini list back to front, merging along the way.
 	for (const TPair<int32, FString>& HierarchyIt : Branch->Hierarchy)
 	{
-		bool bDoCombine = (HierarchyIt.Key != 0);
+		const int KeyWithoutFlags = HierarchyIt.Key & ~((1 << NumFlagsBits) - 1);
 		const FString& IniFileName = HierarchyIt.Value;
 
+		const bool bDoCombine = (KeyWithoutFlags != 0);
+		const bool bUseStagedGlobalCache = ConfigSystem != nullptr && (HierarchyIt.Key & KeyFlag_UseGlobalCache) != 0;
+		const bool bUseStagedPluginCache = (HierarchyIt.Key & KeyFlag_UsePluginCache) != 0;
+		const bool bAssumeExists = (HierarchyIt.Key & KeyFlag_AssumeExists) != 0;
+
 		UE_CLOG(bDumpIniLoadInfo, LogConfig, Display, TEXT("Looking for file: %s"), *IniFileName);
 		
 		// skip non-existant files
-		if (IsUsingLocalIniFile(*IniFileName, nullptr) && !DoesConfigFileExistWrapper(*IniFileName, IniCacheSet))
+		if (!bAssumeExists && IsUsingLocalIniFile(*IniFileName, nullptr) && !DoesConfigFileExistWrapper(*IniFileName, IniCacheSet, bUseStagedGlobalCache ? ConfigSystem->StagedGlobalConfigCache : nullptr, bUseStagedPluginCache ? StagedPluginConfigCache : nullptr))
 		{
 			continue;
 		}
 		
-		if (HierarchyIt.Key != 0)
+		if (KeyWithoutFlags != 0)
 		{
 			bReadAnyFile = true;
 		}
@@ -1149,14 +1202,23 @@ bool FConfigContext::GenerateDestIniFile()
 	FConfigFileHierarchy
 -----------------------------------------------------------------------------*/
 
-constexpr int32 MaxPlatformIndex = 99;
-constexpr int32 GetStaticKey(int32 LayerIndex, int32 ReplacementIndex, int32 PlatformIndex)
+constexpr int32 DynamicKeyOffset = NumLayerBits + NumExpansionBits + NumPlatformBits + NumFlagsBits;
+constexpr int32 LayerOffset = NumExpansionBits + NumPlatformBits + NumFlagsBits;
+constexpr int32 ExpansionOffset = NumPlatformBits + NumFlagsBits;
+constexpr int32 PlatformOffset = NumFlagsBits;
+constexpr int32 FlagsOffset = 0;
+
+
+constexpr int32 GetStaticKey(int32 LayerIndex, int32 ExpansionIndex, int32 PlatformIndex, int32 Flags)
 {
-	return LayerIndex * 10000 + ReplacementIndex * 100 + PlatformIndex;
+	return (LayerIndex << LayerOffset) + (ExpansionIndex << ExpansionOffset) + (PlatformIndex << PlatformOffset) + (Flags << FlagsOffset);
 }
 
+static_assert(UE_ARRAY_COUNT(GConfigLayers) < (1 << NumLayerBits), "Need more NumLayerBits");
+static_assert(UE_ARRAY_COUNT(GConfigExpansions) < (1 << NumExpansionBits), "Need more NumExpansionBits");
+
 FConfigFileHierarchy::FConfigFileHierarchy()
-	: KeyGen(GetStaticKey(UE_ARRAY_COUNT(GConfigLayers) - 1, UE_ARRAY_COUNT(GConfigExpansions) - 1, MaxPlatformIndex))
+	: KeyGen(1 << DynamicKeyOffset)
 {
 }
 
@@ -1166,9 +1228,9 @@ int32 FConfigFileHierarchy::GenerateDynamicKey()
 	return ++KeyGen;
 }
 
-int32 FConfigFileHierarchy::AddStaticLayer(const FString& Filename, int32 LayerIndex, int32 ExpansionIndex /*= 0*/, int32 PlatformIndex /*= 0*/)
+int32 FConfigFileHierarchy::AddStaticLayer(const FString& Filename, int32 LayerIndex, int32 ExpansionIndex /*= 0*/, int32 PlatformIndex /*= 0*/, int32 Flags)
 {
-	int32 Key = GetStaticKey(LayerIndex, ExpansionIndex, PlatformIndex);
+	int32 Key = GetStaticKey(LayerIndex, ExpansionIndex, PlatformIndex, Flags);
 	Emplace(Key, Filename);
 	return Key;
 }
diff --git a/Engine/Source/Runtime/Core/Private/Misc/ConfigUtilities.cpp b/Engine/Source/Runtime/Core/Private/Misc/ConfigUtilities.cpp
--- a/Engine/Source/Runtime/Core/Private/Misc/ConfigUtilities.cpp
+++ b/Engine/Source/Runtime/Core/Private/Misc/ConfigUtilities.cpp
@@ -42,13 +42,16 @@ void PerformDynamicConfig(FName Tag, TFunction<void(class FConfigModificationTra
 			EConsoleVariableFlags Priority = (EConsoleVariableFlags)CVars.CVarPriority;
 
 			// now walk over the updated cvars and set them based on the priority
-			for (const auto& BranchPair : CVars.CVarEntriesPerBranch)
+			for (const auto& TagPair : CVars.CVarEntriesPerBranchPerTag)
 			{
-				const FConfigSection& Section = BranchPair.Value;
-
-				for (const auto& CVarPair : Section)
+				for (const auto& BranchPair : TagPair.Value)
 				{
-					UE::ConfigUtilities::OnSetCVarFromIniEntry(TEXT("DynamicLayer"), *CVarPair.Key.ToString(), *CVarPair.Value.GetValue(), Priority, false, false, Tag);
+					const FConfigSection& Section = BranchPair.Value;
+					
+					for (const auto& CVarPair : Section)
+					{
+						UE::ConfigUtilities::OnSetCVarFromIniEntry(TEXT("DynamicLayer"), *CVarPair.Key.ToString(), *CVarPair.Value.GetValue(), Priority, false, false, TagPair.Key);
+					}
 				}
 			}
 		}
@@ -102,14 +105,18 @@ void PerformDynamicConfig(FName Tag, TFunction<void(class FConfigModificationTra
 				UE_LOG(LogConfig, Verbose, TEXT("  Branch: %s, SetBy Priority: %d"), *SectionPair.Key, CVars.CVarPriority);
 
 				// now walk over the updated cvars and set them based on the priority
-				for (const auto& BranchPair : CVars.CVarEntriesPerBranch)
+				for (const auto& TagPair : CVars.CVarEntriesPerBranchPerTag)
 				{
-					UE_LOG(LogConfig, Verbose, TEXT("    Section %s:"), *BranchPair.Key.ToString());
-					const FConfigSection& Section = BranchPair.Value;
-					
-					for (const auto& CVarPair : Section)
+					UE_LOG(LogConfig, Verbose, TEXT("    Tag %s:"), *TagPair.Key.ToString());
+					for (const auto& BranchPair : TagPair.Value)
 					{
-						UE_LOG(LogConfig, Verbose, TEXT("      %s = %s"), *CVarPair.Key.ToString(), *CVarPair.Value.GetValue());
+						UE_LOG(LogConfig, Verbose, TEXT("      Section %s:"), *BranchPair.Key.ToString());
+						const FConfigSection& Section = BranchPair.Value;
+						
+						for (const auto& CVarPair : Section)
+						{
+							UE_LOG(LogConfig, Verbose, TEXT("        %s = %s"), *CVarPair.Key.ToString(), *CVarPair.Value.GetValue());
+						}
 					}
 				}
 			}
diff --git a/Engine/Source/Runtime/Core/Private/Misc/DataDrivenPlatformInfoRegistry.cpp b/Engine/Source/Runtime/Core/Private/Misc/DataDrivenPlatformInfoRegistry.cpp
--- a/Engine/Source/Runtime/Core/Private/Misc/DataDrivenPlatformInfoRegistry.cpp
+++ b/Engine/Source/Runtime/Core/Private/Misc/DataDrivenPlatformInfoRegistry.cpp
@@ -451,8 +451,15 @@ const TMap<FName, FDataDrivenPlatformInfo>& FDataDrivenPlatformInfoRegistry::Get
 
 				// get the parent to build list later
 				FString IniParent;
-				IniFile.GetString(TEXT("DataDrivenPlatformInfo"), TEXT("IniParent"), IniParent);
-				IniParents.Add(PlatformString, IniParent);
+				bool bSettingPresent = IniFile.GetString(TEXT("DataDrivenPlatformInfo"), TEXT("IniParent"), IniParent);
+				// we want an empty string in here if there is nothing already in it for the platform, however if we have a project
+				// DDPI.ini file, and it doesn't have the setting, it should not overwrite the good value with empty string
+				// so, add empty string if not in map, but only overwrite the engine value if the project value was present (even
+				// if that value is empty string: "IniParent=")
+				if (!IniParents.Contains(PlatformString) || bSettingPresent)
+				{
+					IniParents.Add(PlatformString, IniParent);
+				}
 
 				// get platform name aliases
 				FString PlatformNameAliasesStr;
diff --git a/Engine/Source/Runtime/Core/Public/Misc/ConfigCacheIni.h b/Engine/Source/Runtime/Core/Public/Misc/ConfigCacheIni.h
--- a/Engine/Source/Runtime/Core/Public/Misc/ConfigCacheIni.h
+++ b/Engine/Source/Runtime/Core/Public/Misc/ConfigCacheIni.h
@@ -494,7 +494,8 @@ public:
 	{
 		// this must be an ECVF priority
 		int CVarPriority;
-		TMap<FName, FConfigSection> CVarEntriesPerBranch;
+		// tag -> branch -> cvars
+		TMap<FName, TMap<FName, FConfigSection>> CVarEntriesPerBranchPerTag;
 	};
 	
 	// input
@@ -505,7 +506,7 @@ public:
 	TMap<FName, TSet<FString>> ModifiedSectionsPerBranch;
 	TArray<FString> LoadedFiles;
 	
-	// cvars, input/output
+	// section -> tracker
 	TMap<FString, FCVarTracker> CVars;
 };
 
@@ -1103,17 +1104,20 @@ public:
 	 * Additonally, returns the FConfigFIle object that contains just the loaded sections
 	 * Can return the modified sections if the callers wants to reloadconfig on classes
 	 */
-	CORE_API bool AddDynamicLayerToHierarchy(const FString& Filename, FConfigModificationTracker* ModificationTracker =nullptr);
+	CORE_API bool AddDynamicLayerToHierarchy(const FString& Filename, FConfigModificationTracker* ModificationTracker=nullptr, TSet<FString>* GlobalConfigFileCache=nullptr, TSet<FString>* PluginConfigFileCache=nullptr);
+	CORE_API bool AddDynamicLayersToHierarchy(const TArray<FDynamicLayerInfo>& Layers, FConfigModificationTracker* ModificationTracker=nullptr, TSet<FString>* GlobalConfigFileCache=nullptr, TSet<FString>* PluginConfigFileCache=nullptr);
+
+		UE_DEPRECATED(5.6, "Use AddDynamicLayersToHierarchy that takes a FDynamicLayerInfo list")
 	CORE_API bool AddDynamicLayersToHierarchy(const TArray<FString>& Filenames, FName Tag=NAME_None, DynamicLayerPriority Priority=DynamicLayerPriority::Unknown, FConfigModificationTracker* ModificationTracker =nullptr);
 	
 	// Add a preloaded string as a dynamic layer (useful for hotfixing)
-	CORE_API bool AddDynamicLayerStringToHierarchy(const FString& Filename, const FString& Contents, FName Tag=NAME_None, DynamicLayerPriority Priority=DynamicLayerPriority::Unknown, FConfigModificationTracker* ModificationTracker =nullptr);
+	CORE_API bool AddDynamicLayerStringToHierarchy(const FString& Filename, const FString& Contents, FName Tag=NAME_None, DynamicLayerPriority Priority=DynamicLayerPriority::Unknown, FConfigModificationTracker* ModificationTracker=nullptr);
 	
 	/** Removes a dyanmic file from the hierarchy and recalculates the branch's IniMemoryFile
 	 * Can return the modified sections if the callers wants to reloadconfig on classes
 	 */
-	CORE_API bool RemoveDynamicLayerFromHierarchy(const FString& Filename, FConfigModificationTracker* ModificationTracker =nullptr);
-	CORE_API bool RemoveDynamicLayersFromHierarchy(const TArray<FString>& Filenames, FConfigModificationTracker* ModificationTracker =nullptr);
+	CORE_API bool RemoveDynamicLayerFromHierarchy(const FString& Filename, FConfigModificationTracker* ModificationTracker=nullptr);
+	CORE_API bool RemoveDynamicLayersFromHierarchy(const TArray<FString>& Filenames, FConfigModificationTracker* ModificationTracker=nullptr);
 	
 	/** 
 	 * Frees up the static layer memory, which can be useful if a branch is loaded from, cached, and never used again. However,
@@ -1146,7 +1150,7 @@ public:
 private:
 	void InitFiles();
 	
-	void RemoveTagFromHierarchy(FName Tag, FConfigModificationTracker* ModificationTracker);
+	void RemoveTagsFromHierarchy(const TArray<FName>& Tags, FConfigModificationTracker* ModificationTracker);
 	
 	friend class FConfigCacheIni;
 	
@@ -1886,7 +1890,9 @@ public:
 	static CORE_API void RegisterPlugin(FName PluginName, const FString& PluginDir, const TArray<FString>& ChildPluginDirs, DynamicLayerPriority Priority, bool bIncludePluginNameInBranchName);
 
 	static CORE_API void AddPluginToAllBranches(FName PluginName, FConfigModificationTracker* ModificationTracker=nullptr);
+	static CORE_API void AddMultiplePluginsToAllBranches(const TArray<FName>& PluginName, FConfigModificationTracker* ModificationTracker=nullptr);
 	static CORE_API void RemoveTagFromAllBranches(FName Tag, FConfigModificationTracker* ModificationTracker=nullptr);
+	static CORE_API void RemoveMultipleTagsFromAllBranches(const TArray<FName>& Tags, FConfigModificationTracker* ModificationTracker=nullptr);
 
 	/**
 	* Initializes configs and plugins for every supported platform on worker threads. This function can be
@@ -1902,11 +1908,10 @@ private:
 
 	/** Serialize a bootstrapping state into or from an archive */
 	CORE_API void SerializeStateForBootstrap_Impl(FArchive& Ar);
-	
-	void DumpFile(FOutputDevice& Ar, const FString& Filename, const FConfigFile& File);
 
-	void AddPluginToBranches(FName PluginName, FConfigModificationTracker* ModificationTracker);
-	void RemoveTagFromBranches(FName Tag, FConfigModificationTracker* ModificationTracker);
+
+	void AddPluginsToBranches(const TArray<FName>& PluginNames, FConfigModificationTracker* ModificationTracker);
+	void RemoveTagsFromBranches(const TArray<FName>& Tags, FConfigModificationTracker* ModificationTracker);
 
 	/** Does this ConfigBranch IniName appear in our disallowed removal list **/
 	bool IsConfigBranchNameInNeverUnloadList(const FName& ConfigBranchName);
@@ -1944,7 +1949,13 @@ private:
 	static TMap<FName, FPluginInfo*> RegisteredPlugins;
 	TArray<FName> PendingModificationPlugins;
 	static FTransactionallySafeCriticalSection RegisteredPluginsLock;
-	
+
+	// a cache discovered at staging time and loaded from BinaryConfig.ini
+	TSet<FString>* StagedGlobalConfigCache = nullptr;
+	TMap<FName, TSet<FString>> StagedPluginConfigCache;
+	friend bool MakeBinaryConfig(const TCHAR* CmdLine);// @todo remove this and move the code that fills this out, maybe
+
+
 #if ALLOW_OTHER_PLATFORM_CONFIG
 	static TMap<FName, FConfigCacheIni*> ConfigForPlatform;
 	static FCriticalSection ConfigForPlatformLock;
diff --git a/Engine/Source/Runtime/Core/Public/Misc/ConfigContext.h b/Engine/Source/Runtime/Core/Public/Misc/ConfigContext.h
--- a/Engine/Source/Runtime/Core/Public/Misc/ConfigContext.h
+++ b/Engine/Source/Runtime/Core/Public/Misc/ConfigContext.h
@@ -234,11 +234,16 @@ public:
 	bool bIsForPluginModification = false;
 	// GameFeaturePlugins have WindowsFooGame.ini, not WindowsGame.ini
 	bool bIncludeTagNameInBranchName = false;
-	
+	bool bIsMakingBinaryConfig = false;
+
 	DynamicLayerPriority PluginModificationPriority;
 
 	// if this is non-null, it contains a set of pre-scanned ini files to use to find files, instead of looking on disk
 	const TSet<FString>* IniCacheSet = nullptr;
+	const TSet<FString>* StagedGlobalConfigCache = nullptr;
+	const TSet<FString>* StagedPluginConfigCache = nullptr;
+	
+	TFunction<void(const TArray<FDynamicLayerInfo>&)> HandleLayersFunction = nullptr;
 
 protected:
 
@@ -265,4 +270,4 @@ private:
 	void LogVariables(const TCHAR* InBaseIniName, const FString& Platform);
 };
 
-bool DoesConfigFileExistWrapper(const TCHAR* IniFile, const TSet<FString>* IniCacheSet = nullptr);
+bool DoesConfigFileExistWrapper(const TCHAR* IniFile, const TSet<FString>* IniCacheSet=nullptr, const TSet<FString>* PrimaryConfigFileCache = nullptr, const TSet<FString>* SecondaryConfigFileCache = nullptr);
diff --git a/Engine/Source/Runtime/Core/Public/Misc/ConfigHierarchy.h b/Engine/Source/Runtime/Core/Public/Misc/ConfigHierarchy.h
--- a/Engine/Source/Runtime/Core/Public/Misc/ConfigHierarchy.h
+++ b/Engine/Source/Runtime/Core/Public/Misc/ConfigHierarchy.h
@@ -18,23 +18,23 @@ inline FConfigLayer GConfigLayers[] =
 	{ TEXT("AbsoluteBase"),				TEXT("{ENGINE}/Config/Base.ini"), EConfigLayerFlags::NoExpand},
 
 	// Engine/Base*.ini
-	{ TEXT("Base"),						TEXT("{ENGINE}/Config/Base{TYPE}.ini") },
+	{ TEXT("Base"),						TEXT("{ENGINE}/Config/Base{TYPE}.ini"), EConfigLayerFlags::UseGlobalConfigCache },
 	// Engine/Platform/BasePlatform*.ini
-	{ TEXT("BasePlatform"),				TEXT("{ENGINE}/Config/{PLATFORM}/Base{PLATFORM}{TYPE}.ini")  },
+	{ TEXT("BasePlatform"),				TEXT("{ENGINE}/Config/{PLATFORM}/Base{PLATFORM}{TYPE}.ini"), EConfigLayerFlags::UseGlobalConfigCache  },
 	// Project/Default*.ini
-	{ TEXT("ProjectDefault"),			TEXT("{PROJECT}/Config/Default{TYPE}.ini"), EConfigLayerFlags::AllowCommandLineOverride },
+	{ TEXT("ProjectDefault"),			TEXT("{PROJECT}/Config/Default{TYPE}.ini"), EConfigLayerFlags::AllowCommandLineOverride | EConfigLayerFlags::UseGlobalConfigCache },
 	// Project/Generated*.ini Reserved for files generated by build process and should never be checked in 
-	{ TEXT("ProjectGenerated"),			TEXT("{PROJECT}/Config/Generated{TYPE}.ini") },
+	{ TEXT("ProjectGenerated"),			TEXT("{PROJECT}/Config/Generated{TYPE}.ini"), EConfigLayerFlags::UseGlobalConfigCache },
 	// Project/Custom/CustomConfig/Default*.ini only if CustomConfig is defined
-	{ TEXT("CustomConfig"),				TEXT("{PROJECT}/Config/Custom/{CUSTOMCONFIG}/Default{TYPE}.ini"), EConfigLayerFlags::RequiresCustomConfig },
+	{ TEXT("CustomConfig"),				TEXT("{PROJECT}/Config/Custom/{CUSTOMCONFIG}/Default{TYPE}.ini"), EConfigLayerFlags::RequiresCustomConfig | EConfigLayerFlags::UseGlobalConfigCache },
 	// Engine/Platform/Platform*.ini
-	{ TEXT("EnginePlatform"),			TEXT("{ENGINE}/Config/{PLATFORM}/{PLATFORM}{TYPE}.ini") },
+	{ TEXT("EnginePlatform"),			TEXT("{ENGINE}/Config/{PLATFORM}/{PLATFORM}{TYPE}.ini"), EConfigLayerFlags::UseGlobalConfigCache },
 	// Project/Platform/Platform*.ini
-	{ TEXT("ProjectPlatform"),			TEXT("{PROJECT}/Config/{PLATFORM}/{PLATFORM}{TYPE}.ini") },
+	{ TEXT("ProjectPlatform"),			TEXT("{PROJECT}/Config/{PLATFORM}/{PLATFORM}{TYPE}.ini"), EConfigLayerFlags::UseGlobalConfigCache },
 	// Project/Platform/GeneratedPlatform*.ini Reserved for files generated by build process and should never be checked in 
-	{ TEXT("ProjectPlatformGenerated"),	TEXT("{PROJECT}/Config/{PLATFORM}/Generated{PLATFORM}{TYPE}.ini") },
+	{ TEXT("ProjectPlatformGenerated"),	TEXT("{PROJECT}/Config/{PLATFORM}/Generated{PLATFORM}{TYPE}.ini"), EConfigLayerFlags::UseGlobalConfigCache },
 	// Project/Platform/Custom/CustomConfig/Platform*.ini only if CustomConfig is defined
-	{ TEXT("CustomConfigPlatform"),		TEXT("{PROJECT}/Config/{PLATFORM}/Custom/{CUSTOMCONFIG}/{PLATFORM}{TYPE}.ini"), EConfigLayerFlags::RequiresCustomConfig },
+	{ TEXT("CustomConfigPlatform"),		TEXT("{PROJECT}/Config/{PLATFORM}/Custom/{CUSTOMCONFIG}/{PLATFORM}{TYPE}.ini"), EConfigLayerFlags::RequiresCustomConfig| EConfigLayerFlags::UseGlobalConfigCache },
 	// UserSettings/.../User*.ini
 	{ TEXT("UserSettingsDir"),			TEXT("{USERSETTINGS}Unreal Engine/Engine/Config/User{TYPE}.ini"), EConfigLayerFlags::NoExpand },
 	// UserDir/.../User*.ini
@@ -53,15 +53,15 @@ inline FConfigLayer GPluginLayers[] =
 	{ TEXT("AbsoluteBase"),				TEXT("{ENGINE}/Config/PluginBase.ini"), EConfigLayerFlags::NoExpand},
 
 	// Plugin/Base*.ini
-	{ TEXT("PluginBase"),				TEXT("{PLUGIN}/Config/Base{TYPE}.ini") },
+	{ TEXT("PluginBase"),				TEXT("{PLUGIN}/Config/Base{TYPE}.ini"), EConfigLayerFlags::UsePluginConfigCache },
 	// Plugin/Default*.ini (we use Base and Default as we can have both depending on Engine or Project plugin, but going forward we should stick with Default)
-	{ TEXT("PluginDefault"),			TEXT("{PLUGIN}/Config/Default{TYPE}.ini") },
+	{ TEXT("PluginDefault"),			TEXT("{PLUGIN}/Config/Default{TYPE}.ini"), EConfigLayerFlags::UsePluginConfigCache },
 	// Plugin/Platform/Platform*.ini
-	{ TEXT("PluginPlatform"),			TEXT("{PLUGIN}/Config/{PLATFORM}/{PLATFORM}{TYPE}.ini") },
+	{ TEXT("PluginPlatform"),			TEXT("{PLUGIN}/Config/{PLATFORM}/{PLATFORM}{TYPE}.ini"), EConfigLayerFlags::UsePluginConfigCache },
 	// Project/Default.ini
-	{ TEXT("ProjectDefault"),			TEXT("{PROJECT}/Config/Default{TYPE}.ini") },
+	{ TEXT("ProjectDefault"),			TEXT("{PROJECT}/Config/Default{TYPE}.ini"), EConfigLayerFlags::UseGlobalConfigCache },
 	// Project/Platform/Platform*.ini
-	{ TEXT("ProjectDefault"),			TEXT("{PROJECT}/Config/{PLATFORM}/{PLATFORM}{TYPE}.ini") },
+	{ TEXT("ProjectDefault"),			TEXT("{PROJECT}/Config/{PLATFORM}/{PLATFORM}{TYPE}.ini"), EConfigLayerFlags::UseGlobalConfigCache },
 };
 
 /// <summary>
diff --git a/Engine/Source/Runtime/Core/Public/Misc/ConfigTypes.h b/Engine/Source/Runtime/Core/Public/Misc/ConfigTypes.h
--- a/Engine/Source/Runtime/Core/Public/Misc/ConfigTypes.h
+++ b/Engine/Source/Runtime/Core/Public/Misc/ConfigTypes.h
@@ -27,14 +27,19 @@ public:
 private:
 	int32 GenerateDynamicKey();
 
-	int32 AddStaticLayer(const FString& Filename, int32 LayerIndex, int32 ExpansionIndex = 0, int32 PlatformIndex = 0);
+	int32 AddStaticLayer(const FString& Filename, int32 LayerIndex, int32 ExpansionIndex, int32 PlatformIndex, int32 Flags);
 	int32 AddDynamicLayer(const FString& Filename);
 
 	friend class FConfigFile;
 	friend class FConfigContext;
 };
 
-
+struct FDynamicLayerInfo
+{
+	FString Filename;
+	FName Tag;
+	uint16 Priority;
+};
 
 enum class EConfigLayerFlags : int32
 {
@@ -43,6 +48,8 @@ enum class EConfigLayerFlags : int32
 	DedicatedServerOnly = (1 << 2), // replaces Default, Base, and (NOT {PLATFORM} yet) with an empty string
 	NoExpand = (1 << 4),
 	RequiresCustomConfig = (1 << 5), // disabled if no custom config specified
+	UseGlobalConfigCache = (1 << 6),
+	UsePluginConfigCache = (1 << 7),
 };
 ENUM_CLASS_FLAGS(EConfigLayerFlags);
 
@@ -58,6 +65,10 @@ struct FConfigLayer
 	// Special flag
 	EConfigLayerFlags Flag;
 
+	// Will be cached early on to avoid threading issues
+	uint8 bHasCheckedExist = false;
+	uint8 bExists = false;
+
 };
 
 enum class EConfigExpansionFlags : int32
diff --git a/Engine/Source/Runtime/Projects/Private/PluginManager.cpp b/Engine/Source/Runtime/Projects/Private/PluginManager.cpp
--- a/Engine/Source/Runtime/Projects/Private/PluginManager.cpp
+++ b/Engine/Source/Runtime/Projects/Private/PluginManager.cpp
@@ -1301,72 +1301,6 @@ public:
 	}
 };
 
-bool FPluginManager::IntegratePluginsIntoConfig(FConfigCacheIni& ConfigSystem, const TCHAR* EngineIniName, const TCHAR* PlatformName, const TCHAR* StagedPluginsFile)
-{
-	TArray<FString> PluginList;
-	if (!FFileHelper::LoadFileToStringArray(PluginList, StagedPluginsFile))
-	{
-		return false;
-	}
-
-	// track which plugins were staged and are in the binary config - so at runtime, we will still look at other 
-	TArray<FString> IntegratedPlugins;
-
-	// loop over each one
-	for (FString PluginFile : PluginList)
-	{
-		FPaths::MakeStandardFilename(PluginFile);
-
-		FPluginDescriptor Descriptor;
-		FText FailureReason;
-		if (Descriptor.Load(PluginFile, FailureReason))
-		{
-			// @todo: The type isn't quite right here
-			FPlugin Plugin(PluginFile, Descriptor, FPaths::IsUnderDirectory(PluginFile, FPaths::EngineDir()) ? EPluginType::Engine : EPluginType::Project);
-
-			// we perform Mod plugin processing at runtime
-			if (Plugin.GetType() == EPluginType::Mod)
-			{
-				continue;
-			}
-
-			// mark that we have processed this plugin, so runtime will not scan it again
-			IntegratedPlugins.Add(Plugin.Name);
-
-			FString PluginConfigDir = FPaths::GetPath(Plugin.FileName) / TEXT("Config/");
-
-			// override config cache entries with plugin configs (Engine.ini, Game.ini, etc in <PluginDir>\Config\)
-			TArray<FString> PluginConfigs;
-			IFileManager::Get().FindFiles(PluginConfigs, *PluginConfigDir, TEXT("ini"));
-			for (const FString& ConfigFile : PluginConfigs)
-			{
-				FString BaseConfigFile = *FPaths::GetBaseFilename(ConfigFile);
-
-				FConfigBranch* FoundBranch = ConfigSystem.FindBranch(*BaseConfigFile, TEXT(""));
-				if (FoundBranch != nullptr)
-				{
-					UE_LOG(LogPluginManager, Log, TEXT("Found config from plugin[%s] %s"), *Plugin.GetName(), *BaseConfigFile);
-
-					FoundBranch->AddDynamicLayerToHierarchy(FPaths::Combine(PluginConfigDir, ConfigFile));
-				}
-			}
-
-			if (Descriptor.bCanContainContent)
-			{
-				// we need to look up the section each time because other loops could add entries
-				FConfigFile* EngineConfigFile = ConfigSystem.FindConfigFile(EngineIniName);
-				EngineConfigFile->AddUniqueToSection(TEXT("Core.System"), "Paths", Plugin.GetContentDir());
-			}
-		}
-	}
-
-	// record in the config that the plugin inis have been inserted (so we can know at runtime if we have to load plugins or not)
-	FConfigFile* EngineConfigFile = ConfigSystem.FindConfigFile(EngineIniName);
-	EngineConfigFile->SetArray(TEXT("BinaryConfig"), TEXT("BinaryConfigPlugins"), IntegratedPlugins);
-
-	return true;
-}
-
 void FPluginManager::SetBinariesRootDirectories(const FString& InEngineBinariesRootDir, const FString& InProjectBinariesRootDir)
 {
 	EngineBinariesRootDir = InEngineBinariesRootDir;
@@ -1951,25 +1885,21 @@ bool FPluginManager::ConfigureEnabledPlugins()
 			FCriticalSection ConfigCS;
 			FCriticalSection PluginPakCS;
 			
-			struct FPendingConfigFile
-			{
-				FString PluginName;
-				FString PluginConfigDir;
-				FString PluginConfigFile;
-			};
-			FCriticalSection PendingConfigsCS;
-			TMap<FString, TArray<FPendingConfigFile>> PendingConfigs;
-
 			// Mount all the enabled plugins
 			FConfigModificationTracker ChangeTracker;
 			ChangeTracker.bTrackLoadedFiles = true;
 
+			TArray<FName> PluginNames;
+			PluginNames.Reserve(PluginsArray.Num());
+			
 			// walk over each plugin, and add the config files that are named for the plugin
 			// this is a separate loop so that a plugin can modify another plugin's configs in the second loop below
 			for (TSharedRef<FPlugin> PluginPtr : PluginsArray)
 			{
 				FPlugin& Plugin = *PluginPtr;
 				FName PluginName(*Plugin.Name);
+				PluginNames.Add(PluginName);
+				
 				UE_LOG(LogPluginManager, Log, TEXT("Mounting %s plugin %s"), *EnumToString(Plugin.Type), *Plugin.GetName());
 				UE_LOG(LogPluginManager, Verbose, TEXT("Plugin path: %s"), *Plugin.FileName);
 				
@@ -1982,17 +1912,16 @@ bool FPluginManager::ConfigureEnabledPlugins()
 				Context.ConfigFileTag = *Plugin.Name;
 				Context.Load(*Plugin.Name);
 			}
+			
+			// gang up dynamic layer adds as much as possible to reduce the amount of "CalcDiff" calls etc 
+			FConfigCacheIni::AddMultiplePluginsToAllBranches(PluginNames);
 
 			// walk over each plugin, modify existing ini's, mount the content, etc
-			ParallelFor(PluginsArray.Num(), [&PluginsArray, &ConfigCS, &PluginPakCS, &PendingConfigsCS, &PendingConfigs, &ConfigFilesPluginsCannotOverride, &ChangeTracker, this](int32 Index)
+			ParallelFor(PluginsArray.Num(), [&PluginsArray, &ConfigCS, &PluginPakCS, &ConfigFilesPluginsCannotOverride, &ChangeTracker, this](int32 Index)
 			{
 				FPlugin& Plugin = *PluginsArray[Index];
 				FName PluginName(*Plugin.Name);
 
-				// handle any overrides for Engine.ini, etc in the plugin
-				FConfigCacheIni::AddPluginToAllBranches(PluginName, &ChangeTracker);
-
-
 				// Build the list of content folders
 				if (Plugin.Descriptor.bCanContainContent)
 				{
diff --git a/Engine/Source/Runtime/Projects/Private/PluginManager.h b/Engine/Source/Runtime/Projects/Private/PluginManager.h
--- a/Engine/Source/Runtime/Projects/Private/PluginManager.h
+++ b/Engine/Source/Runtime/Projects/Private/PluginManager.h
@@ -231,8 +231,6 @@ public:
 	virtual bool TrySplitVersePath(const UE::Core::FVersePath& VersePath, FName& OutPackageName, FString& OutLeafPath) override;
 	virtual bool RequiresTempTargetForCodePlugin(const FProjectDescriptor* ProjectDescriptor, const FString& Platform, EBuildConfiguration Configuration, EBuildTargetType TargetType, FText& OutReason) override;
 
-	virtual bool IntegratePluginsIntoConfig(FConfigCacheIni& ConfigSystem, const TCHAR* EngineIniName, const TCHAR* PlatformName, const TCHAR* StagedPluginsFile) override;
-
 	virtual void SetBinariesRootDirectories(const FString& EngineBinariesRootDir, const FString& ProjectBinariesRootDir) override;
 	virtual void SetPreloadBinaries() override;
 	virtual bool GetPreloadBinaries() override;
diff --git a/Engine/Source/Runtime/Projects/Public/Interfaces/IPluginManager.h b/Engine/Source/Runtime/Projects/Public/Interfaces/IPluginManager.h
--- a/Engine/Source/Runtime/Projects/Public/Interfaces/IPluginManager.h
+++ b/Engine/Source/Runtime/Projects/Public/Interfaces/IPluginManager.h
@@ -588,7 +588,8 @@ public:
 	 * @param PlatformName The name of the platform this config is made for
 	 * @param StagedPluginsFile A path to a file that contains all plugins that have been staged, and should be evaluated
 	 */
-	virtual bool IntegratePluginsIntoConfig(FConfigCacheIni& ConfigSystem, const TCHAR* EngineIniName, const TCHAR* PlatformName, const TCHAR* StagedPluginsFile) = 0;
+	UE_DEPRECATED(5.6, "This has never worked and is unsafe. Do not call this function")
+	virtual bool IntegratePluginsIntoConfig(FConfigCacheIni& ConfigSystem, const TCHAR* EngineIniName, const TCHAR* PlatformName, const TCHAR* StagedPluginsFile) { return false; }
 
 	/**
 	* Set root directories for where to find binaries for plugins.
