Removed Deprecated Code Paths from LLM

diff --git a/Engine/Source/Runtime/Core/Private/HAL/LowLevelMemTracker.cpp b/Engine/Source/Runtime/Core/Private/HAL/LowLevelMemTracker.cpp
--- a/Engine/Source/Runtime/Core/Private/HAL/LowLevelMemTracker.cpp
+++ b/Engine/Source/Runtime/Core/Private/HAL/LowLevelMemTracker.cpp
@@ -1749,40 +1749,6 @@ extern const ANSICHAR* LLMGetTagNameANSI(ELLMTag Tag)
 	}
 }
 
-extern FName LLMGetTagStat(ELLMTag Tag)
-{
-#define LLM_TAG_STAT_ARRAY(Enum,Str,Stat,Group,ParentTag) Stat,
-	static FName Names[] = { LLM_ENUM_GENERIC_TAGS(LLM_TAG_STAT_ARRAY) };
-#undef LLM_TAG_STAT_ARRAY
-
-	int32 Index = static_cast<int32>(Tag);
-	if (Index >= 0 && Index < UE_ARRAY_COUNT(Names))
-	{
-		return Names[Index];
-	}
-	else
-	{
-		return NAME_None;
-	}
-}
-
-extern FName LLMGetTagStatGroup(ELLMTag Tag)
-{
-#define LLM_TAG_STATGROUP_ARRAY(Enum,Str,Stat,Group,ParentTag) Group,
-	static FName Names[] = { LLM_ENUM_GENERIC_TAGS(LLM_TAG_STATGROUP_ARRAY) };
-#undef LLM_TAG_STAT_ARRAY
-
-	int32 Index = static_cast<int32>(Tag);
-	if (Index >= 0 && Index < UE_ARRAY_COUNT(Names))
-	{
-		return Names[Index];
-	}
-	else
-	{
-		return NAME_None;
-	}
-}
-
 FLowLevelMemTracker& FLowLevelMemTracker::Construct()
 {
 	static FLowLevelMemTracker Tracker;
@@ -3315,61 +3281,6 @@ ELLMTag FLowLevelMemTracker::GetTagClosestEnumTag(const UE::LLMPrivate::FTagData
 	return TagData->GetEnumTag();
 }
 
-// Deprecated in 5.3
-int64 FLowLevelMemTracker::GetTagAmountForTracker(ELLMTracker Tracker, ELLMTag Tag, bool bPeakAmount)
-{
-	if (!IsEnabled())
-	{
-		return 0;
-	}
-
-	BootstrapInitialise();
-
-	return GetTagAmountForTracker(Tracker, FindTagData(Tag), bPeakAmount ? UE::LLM::ESizeParams::ReportPeak : UE::LLM::ESizeParams::Default);
-}
-
-// Deprecated in 5.3
-int64 FLowLevelMemTracker::GetTagAmountForTracker(ELLMTracker Tracker, const UE::LLMPrivate::FTagData* TagData,
-	bool bPeakAmount)
-{
-	using namespace UE::LLMPrivate;
-
-	if (!IsEnabled())
-	{
-		return 0;
-	}
-
-	BootstrapInitialise();
-
-	if (TagData == nullptr)
-	{
-		return 0;
-	}
-
-	FScopeLock UpdateScopeLock(&UpdateLock); // uses of TagSizes are guarded by the UpdateLock
-	return GetTracker(Tracker)->GetTagAmount(TagData, bPeakAmount ? UE::LLM::ESizeParams::ReportPeak : UE::LLM::ESizeParams::Default);
-}
-
-// Deprecated in 5.3
-int64 FLowLevelMemTracker::GetTagAmountForTracker(ELLMTracker Tracker, FName Tag, ELLMTagSet TagSet,
-	bool bPeakAmount)
-{
-	if (!IsEnabled())
-	{
-		return 0;
-	}
-
-	BootstrapInitialise();
-
-	UE::LLM::ESizeParams SizeParams(UE::LLM::ESizeParams::Default);
-	if (bPeakAmount)
-	{
-		EnumAddFlags(SizeParams, UE::LLM::ESizeParams::ReportPeak);
-	}
-
-	return GetTagAmountForTracker(Tracker, FindTagData(Tag, TagSet, UE::LLMPrivate::ETagReferenceSource::FunctionAPI), bPeakAmount ? UE::LLM::ESizeParams::ReportPeak : UE::LLM::ESizeParams::Default);
-}
-
 int64 FLowLevelMemTracker::GetTagAmountForTracker(ELLMTracker Tracker, const UE::LLMPrivate::FTagData* TagData,
 	UE::LLM::ESizeParams SizeParams)
 {
@@ -3446,27 +3357,6 @@ void FLowLevelMemTracker::SetTagAmountForTracker(ELLMTracker Tracker, FName Tag,
 	GetTracker(Tracker)->SetTagAmountExternal(TagData, Amount, bAddToTotal);
 }
 
-int64 FLowLevelMemTracker::GetActiveTag(ELLMTracker Tracker)
-{
-	using namespace UE::LLMPrivate;
-
-	if (!IsEnabled())
-	{
-		return static_cast<int64>(ELLMTag::Untagged);
-	}
-	BootstrapInitialise();
-
-	const FTagData* TagData = GetActiveTagData(Tracker);
-	if (TagData)
-	{
-		return static_cast<int64>(TagData->GetContainingEnum());
-	}
-	else
-	{
-		return static_cast<int64>(ELLMTag::Untagged);
-	}
-}
-
 const UE::LLMPrivate::FTagData* FLowLevelMemTracker::GetActiveTagData(ELLMTracker Tracker, ELLMTagSet TagSet /*= ELLMTagSet::None*/)
 {
 	if (!IsEnabled())
diff --git a/Engine/Source/Runtime/Core/Public/HAL/LowLevelMemTracker.h b/Engine/Source/Runtime/Core/Public/HAL/LowLevelMemTracker.h
--- a/Engine/Source/Runtime/Core/Public/HAL/LowLevelMemTracker.h
+++ b/Engine/Source/Runtime/Core/Public/HAL/LowLevelMemTracker.h
@@ -316,11 +316,6 @@ enum class ELLMAllocType
 extern const ANSICHAR* LLMGetTagNameANSI(ELLMTag Tag);
 extern const TCHAR* LLMGetTagName(ELLMTag Tag);
 extern CORE_API const FName LLMGetUntaggedTagName(ELLMTagSet TagSet);
-UE_DEPRECATED(4.27, "This function was an unused implementation detail; contact Epic if you need to keep its functionality.")
-extern FName LLMGetTagStatGroup(ELLMTag Tag);
-UE_DEPRECATED(4.27, "This function was an unused implementation detail; contact Epic if you need to keep its functionality.")
-extern FName LLMGetTagStat(ELLMTag Tag);
-
 /*
  * LLM utility macros
  */
@@ -673,10 +668,6 @@ public:
 	/** For some tag sets, it's really useful to reduce threads, to attribute allocations to assets, for instance. */
 	CORE_API bool ShouldReduceThreads();
 
-	/** Get the top active tag for the given tracker */
-	UE_DEPRECATED(4.27, "Tags have been changed to FNames and the old ELLMTag is now only the top-level coarse tag. Use GetActivateTagData instead to get the current Tag instead of its toplevel parent.")
-	CORE_API int64 GetActiveTag(ELLMTracker Tracker);
-
 	/** Get an opaque identifier for the top active tag for the given tracker. */
 	CORE_API const UE::LLMPrivate::FTagData* GetActiveTagData(ELLMTracker Tracker, ELLMTagSet TagSet = ELLMTagSet::None);
 
@@ -734,17 +725,6 @@ public:
 
     CORE_API int64 GetTagAmountForTracker(ELLMTracker Tracker, FName Tag, ELLMTagSet TagSet, UE::LLM::ESizeParams SizeParams = UE::LLM::ESizeParams::Default);
 
-	/** Get the amount of memory for an ELLMTag from the given tracker. */
-	UE_DEPRECATED(5.3, "Use version that takes ESizeParams instead")
-	CORE_API int64 GetTagAmountForTracker(ELLMTracker Tracker, ELLMTag Tag, bool bPeakAmount);
-
-	/** Get the amount of memory for a FTagData from the given tracker. */
-	UE_DEPRECATED(5.3, "Use version that takes ESizeParams instead")
-	CORE_API int64 GetTagAmountForTracker(ELLMTracker Tracker, const UE::LLMPrivate::FTagData* TagData, bool bPeakAmount);
-
-	UE_DEPRECATED(5.3, "Use version that takes ESizeParams instead")
-	CORE_API int64 GetTagAmountForTracker(ELLMTracker Tracker, FName Tag, ELLMTagSet TagSet, bool bPeakAmount);
-
 	/** Set the amount of memory for an ELLMTag for a given tracker and optionally update the total tracked memory. */
 	CORE_API void SetTagAmountForTracker(ELLMTracker Tracker, ELLMTag Tag, int64 Amount, bool bAddToTotal);
 
