[AutoRTFM] Disable memory validation for opens in core ref counting classes

These opens are safe, and called very frequently.

diff --git a/Engine/Source/Runtime/Core/Public/Templates/RefCounting.h b/Engine/Source/Runtime/Core/Public/Templates/RefCounting.h
--- a/Engine/Source/Runtime/Core/Public/Templates/RefCounting.h
+++ b/Engine/Source/Runtime/Core/Public/Templates/RefCounting.h
@@ -93,7 +93,8 @@ public:
 			// though we are deferring reference count changes until commit time.
 			// Be advised that GetRefCount() would reveal our trickery, since it
 			// always returns the true refcount.
-			UE_AUTORTFM_OPEN
+			// Memory validation disabled as this is hot code.
+			UE_AUTORTFM_OPEN_NO_VALIDATION
 			{
 				RefsToReturn = uint32(NumRefs.load(std::memory_order_relaxed)) - 1;
 			};
@@ -117,7 +118,8 @@ public:
 		// Inside of an AutoRTFM transaction, the returned refcount value may be higher than you'd
 		// expect, because all Release calls are deferred until the transaction commit time.
 		uint32 Refs;
-		UE_AUTORTFM_OPEN
+		// Memory validation disabled as this is hot code.
+		UE_AUTORTFM_OPEN_NO_VALIDATION
 		{
 			Refs = uint32(NumRefs.load(std::memory_order_relaxed));
 		};
@@ -267,7 +269,8 @@ public:
 		if constexpr (Mode == ERefCountingMode::ThreadSafe)
 		{
 			uint32 Refs;
-			UE_AUTORTFM_OPEN
+			// Memory validation disabled as this is hot code.
+			UE_AUTORTFM_OPEN_NO_VALIDATION
 			{
 				Refs = RefCount++;
 			};
@@ -348,7 +351,8 @@ public:
 
 			uint32 Count = 0;
 
-			UE_AUTORTFM_OPEN
+			// Memory validation disabled as this is hot code.
+			UE_AUTORTFM_OPEN_NO_VALIDATION
 			{
 				// This reference count may be accessed by multiple threads
 				Count = RefCount.load(std::memory_order_relaxed);
