Changed PlatformProcessMemory LLM reporting to use PlatformStats.UsedPhysical instead of PlatformStats.TotalPhysical - PlatformStats.AvailablePhysical
Some platforms cache free memory in the allocator so TotalPhysical contains it, while UsedPhysical reports only what's actually used by the allocator

diff --git a/Engine/Source/Runtime/Core/Private/HAL/LowLevelMemTracker.cpp b/Engine/Source/Runtime/Core/Private/HAL/LowLevelMemTracker.cpp
--- a/Engine/Source/Runtime/Core/Private/HAL/LowLevelMemTracker.cpp
+++ b/Engine/Source/Runtime/Core/Private/HAL/LowLevelMemTracker.cpp
@@ -1985,11 +1985,8 @@ void FLowLevelMemTracker::TickInternal()
 	// On some Android devices total used mem (VmRss + VmSwap) does not include GPU memory allocations
 	// and so it's going to be lower than TrackedTotal
 	const int64 PlatformProcessMemory = FMath::Max((int64)PlatformStats.UsedPhysical, TrackedTotal);
-#elif  PLATFORM_IOS || UE_SERVER
-	const int64 PlatformProcessMemory = static_cast<int64>(PlatformStats.UsedPhysical);
 #else
-	const int64 PlatformProcessMemory = static_cast<int64>(PlatformStats.TotalPhysical) -
-		static_cast<int64>(PlatformStats.AvailablePhysical);
+	const int64 PlatformProcessMemory = static_cast<int64>(PlatformStats.UsedPhysical);
 #endif
 	// Update how much overhead LLM is using. Note we used to also added sizeof(FLowLevelMemTracker), but this
 	// isn't allocated and is in the data segment. So adding it throws the numbers off.
