Patch for PackageWriter enum not being explicitly associated for LexToString implementation.

diff --git a/Engine/Source/Runtime/Core/Public/Serialization/PackageWriter.h b/Engine/Source/Runtime/Core/Public/Serialization/PackageWriter.h
--- a/Engine/Source/Runtime/Core/Public/Serialization/PackageWriter.h
+++ b/Engine/Source/Runtime/Core/Public/Serialization/PackageWriter.h
@@ -454,6 +454,8 @@ static inline const ANSICHAR* LexToString(IPackageWriter::FBulkDataInfo::EType V
 		return "Mmap";
 	case IPackageWriter::FBulkDataInfo::Optional:
 		return "Optional";
+	case IPackageWriter::FBulkDataInfo::NumTypes:
+		return "NumTypes";
 	default:
 		checkNoEntry();
 		return "Unknown";
