Add AutoRTFM::Unreachable, and fix a bug with AutoRTFM::CascadingRetryTransaction with respect to AutoRTFM::IsTransactional.

diff --git a/Engine/Source/Programs/AutoRTFMTests/Private/API.cpp b/Engine/Source/Programs/AutoRTFMTests/Private/API.cpp
--- a/Engine/Source/Programs/AutoRTFMTests/Private/API.cpp
+++ b/Engine/Source/Programs/AutoRTFMTests/Private/API.cpp
@@ -1232,4 +1232,38 @@ TEST_CASE("API.CascadingRetryTransaction")
 		});
 		REQUIRE(!bFirst);
 	}
+
+	SECTION("IsTransactional is false during the retry")
+	{
+		bool bFirst = true;
+		AutoRTFM::Testing::Commit([&]
+		{
+			if (bFirst)
+			{
+				AutoRTFM::CascadingRetryTransaction([&]
+				{
+					REQUIRE(!AutoRTFM::IsTransactional());
+					bFirst = false;
+				});
+			}
+		});
+		REQUIRE(!bFirst);
+	}
+
+	SECTION("IsClosed is false during the retry")
+	{
+		bool bFirst = true;
+		AutoRTFM::Testing::Commit([&]
+		{
+			if (bFirst)
+			{
+				AutoRTFM::CascadingRetryTransaction([&]
+				{
+					REQUIRE(!AutoRTFM::IsClosed());
+					bFirst = false;
+				});
+			}
+		});
+		REQUIRE(!bFirst);
+	}
 }
diff --git a/Engine/Source/Programs/AutoRTFMTests/Private/AutoRTFMTests.cpp b/Engine/Source/Programs/AutoRTFMTests/Private/AutoRTFMTests.cpp
--- a/Engine/Source/Programs/AutoRTFMTests/Private/AutoRTFMTests.cpp
+++ b/Engine/Source/Programs/AutoRTFMTests/Private/AutoRTFMTests.cpp
@@ -89,8 +89,9 @@ int RunTests(int ArgC, const char* ArgV[])
 		AutoRTFM::ForTheRuntime::SetRetryTransaction(AutoRTFM::ForTheRuntime::EAutoRTFMRetryTransactionState::RetryNonNested);
 	}
 
-	const TCHAR* CommandLine = TEXT("-Multiprocess");
+	const TCHAR* CommandLine = TEXT("-Multiprocess -LogCmds=\"LogStreaming off, LogUObjectGlobals off, LogPackageName off\" -AsyncLoadingThread");
 	GEngineLoop.PreInit(CommandLine);
+
 	FModuleManager::Get().StartProcessingNewlyLoadedObjects();
 
 	// Enable AutoRTFM.
diff --git a/Engine/Source/Runtime/Core/Private/AutoRTFM/CompilerInterface.cpp b/Engine/Source/Runtime/Core/Private/AutoRTFM/CompilerInterface.cpp
--- a/Engine/Source/Runtime/Core/Private/AutoRTFM/CompilerInterface.cpp
+++ b/Engine/Source/Runtime/Core/Private/AutoRTFM/CompilerInterface.cpp
@@ -147,17 +147,21 @@ extern "C" UE_AUTORTFM_API void autortfm_memset(void* Dst, int Value, size_t Siz
     Memset(Dst, Value, Size, Context);
 }
 
-extern "C" UE_AUTORTFM_API void autortfm_llvm_fail(const char* Message)
+extern "C" UE_AUTORTFM_API FORCENOINLINE void autortfm_report_error(const UE_AUTORTFM_TCHAR* const Message)
 {
+	FContext* const Context = FContext::Get();
+
+	UE_CLOG(!Context->IsTransactional(), LogAutoRTFM, Fatal, TEXT("The function `autortfm_report_error` was called from outside a transact."));
+
     if (Message)
     {
 		if (ForTheRuntime::IsAutoRTFMAssertOnError())
 		{
-			UE_LOG(LogAutoRTFM, Fatal, TEXT("Transaction failing because of language issue '%s'."), ANSI_TO_TCHAR(Message));
+			UE_LOG(LogAutoRTFM, Fatal, TEXT("Transaction failing because of language issue '%s'."), Message);
 		}
 		else
 		{
-			ensureMsgf(!ForTheRuntime::IsEnsureOnAbortByLanguageEnabled(), TEXT("Transaction failing because of language issue '%s'."), ANSI_TO_TCHAR(Message));
+			ensureMsgf(!ForTheRuntime::IsEnsureOnAbortByLanguageEnabled(), TEXT("Transaction failing because of language issue '%s'."), Message);
 		}
     }
     else
@@ -172,10 +176,26 @@ extern "C" UE_AUTORTFM_API void autortfm_llvm_fail(const char* Message)
 		}
 	}
 
-	FContext* Context = FContext::Get();
     Context->AbortByLanguageAndThrow();
 }
 
+extern "C" UE_AUTORTFM_API void autortfm_unreachable(const UE_AUTORTFM_TCHAR* const Message)
+{
+	autortfm_report_error(Message);
+}
+
+extern "C" UE_AUTORTFM_API void autortfm_llvm_fail(const char* const Message)
+{
+	if (Message)
+	{
+		autortfm_report_error(ANSI_TO_TCHAR(Message));
+	}
+	else
+	{
+		autortfm_report_error(nullptr);
+	}
+}
+
 extern "C" UE_AUTORTFM_API void autortfm_llvm_missing_function()
 {
 	if (ForTheRuntime::IsAutoRTFMAssertOnError())
diff --git a/Engine/Source/Runtime/Core/Private/AutoRTFM/Context.cpp b/Engine/Source/Runtime/Core/Private/AutoRTFM/Context.cpp
--- a/Engine/Source/Runtime/Core/Private/AutoRTFM/Context.cpp
+++ b/Engine/Source/Runtime/Core/Private/AutoRTFM/Context.cpp
@@ -200,7 +200,13 @@ ETransactionResult FContext::AbortTransaction(bool bIsClosed, bool bIsCascading)
 void FContext::CascadingRetryTransaction(bool bIsClosed, TFunction<void()>&& RunAfterAbortBeforeRetryWork)
 {
 	ASSERT(!CascadingRetryTransactionWork);
-	CascadingRetryTransactionWork = MoveTemp(RunAfterAbortBeforeRetryWork);
+	
+	// We explicitly must copy the function here because the original was allocated
+	// within a transactional context, and thus the memory is allocating under
+	// transactionalized conditions. By copying, we create an open copy of the callback.
+	TFunction<void()> Copy(RunAfterAbortBeforeRetryWork);
+
+	CascadingRetryTransactionWork = MoveTemp(Copy);
 	AbortTransaction(bIsClosed, true);
 }
 
@@ -488,11 +494,17 @@ ETransactionResult FContext::Transact(void (*InstrumentedFunction)(void*), void*
 					ClearTransactionStatus();
 					CurrentTransaction->Reset();
 
+					ASSERT(Status == EContextStatus::OnTrack);
+					Status = EContextStatus::InRetryTransaction;
+
 					// Then call the retry callback to do whatever work the user
 					// required be done before we attempt to re-execute.
 					CascadingRetryTransactionWork();
 					CascadingRetryTransactionWork.Reset();
 
+					ASSERT(Status == EContextStatus::InRetryTransaction);
+					Status = EContextStatus::OnTrack;
+
 					// Then get rolling!
 					CurrentTransaction->SetClosedActive();
 					continue;
diff --git a/Engine/Source/Runtime/Core/Private/AutoRTFM/ContextStatus.h b/Engine/Source/Runtime/Core/Private/AutoRTFM/ContextStatus.h
--- a/Engine/Source/Runtime/Core/Private/AutoRTFM/ContextStatus.h
+++ b/Engine/Source/Runtime/Core/Private/AutoRTFM/ContextStatus.h
@@ -12,7 +12,7 @@ inline const char* GetContextStatusName(EContextStatus Status)
     switch (Status)
     {
     default:
-        Unreachable();
+        InternalUnreachable();
         return nullptr;
     case EContextStatus::Idle:
         return "Idle";
diff --git a/Engine/Source/Runtime/Core/Private/AutoRTFM/HitSet.h b/Engine/Source/Runtime/Core/Private/AutoRTFM/HitSet.h
--- a/Engine/Source/Runtime/Core/Private/AutoRTFM/HitSet.h
+++ b/Engine/Source/Runtime/Core/Private/AutoRTFM/HitSet.h
@@ -54,7 +54,7 @@ namespace AutoRTFM
 				switch (InsertNoResize(K.GetPayload()))
 				{
 				default:
-					AutoRTFM::Unreachable();
+					AutoRTFM::InternalUnreachable();
 				case EInsertResult::Inserted:
 					return false;
 				case EInsertResult::Exists:
diff --git a/Engine/Source/Runtime/Core/Private/AutoRTFM/Stats.cpp b/Engine/Source/Runtime/Core/Private/AutoRTFM/Stats.cpp
--- a/Engine/Source/Runtime/Core/Private/AutoRTFM/Stats.cpp
+++ b/Engine/Source/Runtime/Core/Private/AutoRTFM/Stats.cpp
@@ -123,7 +123,7 @@ void AutoRTFM::FStats::Report() const
 		{
 			switch (static_cast<EStatsKind>(I))
 			{
-			default: AutoRTFM::Unreachable(); break;
+			default: AutoRTFM::InternalUnreachable(); break;
 #define REPORT_CASE(x) case (x): Report<x>(Datas[static_cast<size_t>(x)]); break
 			REPORT_CASE(AutoRTFM::EStatsKind::Transaction);
 			REPORT_CASE(AutoRTFM::EStatsKind::Commit);
diff --git a/Engine/Source/Runtime/Core/Private/AutoRTFM/Utils.h b/Engine/Source/Runtime/Core/Private/AutoRTFM/Utils.h
--- a/Engine/Source/Runtime/Core/Private/AutoRTFM/Utils.h
+++ b/Engine/Source/Runtime/Core/Private/AutoRTFM/Utils.h
@@ -13,7 +13,7 @@ DECLARE_LOG_CATEGORY_EXTERN(LogAutoRTFM, Display, All)
 namespace AutoRTFM
 {
 
-[[noreturn]] inline void Unreachable() //-V1082 Silence PVS [[noreturn]] false positive warning
+[[noreturn]] inline void InternalUnreachable() //-V1082 Silence PVS [[noreturn]] false positive warning
 {
 	LowLevelFatalError(TEXT("Unreachable encountered!"));
 
diff --git a/Engine/Source/Runtime/Core/Public/AutoRTFM/AutoRTFM.h b/Engine/Source/Runtime/Core/Public/AutoRTFM/AutoRTFM.h
--- a/Engine/Source/Runtime/Core/Public/AutoRTFM/AutoRTFM.h
+++ b/Engine/Source/Runtime/Core/Public/AutoRTFM/AutoRTFM.h
@@ -68,6 +68,7 @@
 #define PRAGMA_DISABLE_UNREACHABLE_CODE_WARNINGS
 #define PRAGMA_RESTORE_UNREACHABLE_CODE_WARNINGS
 #define LIKELY(EXPR) __builtin_expect(!!(EXPR), 1)
+#define UE_AUTORTFM_TCHAR char
 template <typename FuncType> class TFunction;  // Should be replaced with `std::function`
 #else
 #include <HAL/Platform.h>
@@ -78,6 +79,7 @@ template <typename FuncType> class TFunction;  // Should be replaced with `std::
 #define UE_AUTORTFM_FORCEINLINE FORCEINLINE
 #define UE_AUTORTFM_MEMCPY FPlatformMemory::Memcpy
 #define UE_AUTORTFM_MOVE MoveTemp
+#define UE_AUTORTFM_TCHAR TCHAR
 #endif
 
 #if defined(UE_BUILD_DEBUG) && UE_BUILD_DEBUG
@@ -131,7 +133,8 @@ typedef enum
 	autortfm_status_aborted_by_request,
 	autortfm_status_committing,
 	autortfm_status_aborted_by_cascade,
-	autortfm_status_in_static_local_initializer
+	autortfm_status_in_static_local_initializer,
+	autortfm_status_in_retry_transaction
 } autortfm_status;
 
 #if UE_AUTORTFM_ENABLED
@@ -382,6 +385,15 @@ static UE_AUTORTFM_FORCEINLINE void autortfm_check_abi(void* ptr, size_t size)
 }
 #endif
 
+#if UE_AUTORTFM_ENABLED
+UE_AUTORTFM_API void autortfm_unreachable(const UE_AUTORTFM_TCHAR* const Message);
+#else
+static UE_AUTORTFM_FORCEINLINE void autortfm_unreachable(const UE_AUTORTFM_TCHAR* const Message)
+{
+	UE_AUTORTFM_UNUSED(Message);
+}
+#endif
+
 #ifdef __cplusplus
 }
 
@@ -445,6 +457,10 @@ enum class EContextStatus : uint8_t
 	// Means we are in a static local initializer which always run in the open.
 	// `IsTransactional()` will return `false` when in this state.
 	InStaticLocalInitializer = autortfm_status_in_static_local_initializer,
+
+	// Means we are in a transaction retry currently, inbetween the aborted previous transaction
+	// and the re-runned retried transaction.
+	InRetryTransaction = autortfm_status_in_retry_transaction,
 };
 
 // An enumerator of transactional memory validation levels.
@@ -936,6 +952,22 @@ template<typename Type> static UE_AUTORTFM_FORCEINLINE void RecordOpenWrite(Type
 	autortfm_record_open_write(Ptr, sizeof(Type));
 }
 
+// Report that a unreachable codepath is being hit. Used to manually ban certain codepaths
+// from being transactionally safe.
+static UE_AUTORTFM_FORCEINLINE void Unreachable(const TCHAR* const Message = nullptr)
+{
+	autortfm_unreachable(Message);
+}
+
+// If we are running within a transaction, call `AutoRTFM::Unreachable`.
+static UE_AUTORTFM_FORCEINLINE void UnreachableIfTransactional(const TCHAR* const Message = nullptr)
+{
+	if (AutoRTFM::IsTransactional())
+	{
+		AutoRTFM::Unreachable(Message);
+	}
+}
+
 // A collection of power-user functions that are reserved for use by the AutoRTFM runtime only.
 namespace ForTheRuntime
 {
