Fix ForkProcess logs and log the timestamp at which the fork signal event was received

diff --git a/Engine/Source/Runtime/Core/Private/Unix/UnixPlatformProcess.cpp b/Engine/Source/Runtime/Core/Private/Unix/UnixPlatformProcess.cpp
--- a/Engine/Source/Runtime/Core/Private/Unix/UnixPlatformProcess.cpp
+++ b/Engine/Source/Runtime/Core/Private/Unix/UnixPlatformProcess.cpp
@@ -59,6 +59,15 @@ namespace UnixPlatformProcess
 		GLogMemoryStatsWhenForking,
 		TEXT("When true the parent process will log memory stats before every fork."),
 		ECVF_Default);
+
+
+	FString GetReadableTime(double InDeltaSeconds)
+	{
+		FDateTime Now = FDateTime::UtcNow();
+		const FTimespan Delta = FTimespan::FromSeconds(InDeltaSeconds);
+		Now -= Delta;
+		return Now.ToString(TEXT("%H.%M.%S:%s"));
+	}
 }
 
 #if IS_MONOLITHIC
@@ -1445,10 +1454,10 @@ FGenericPlatformProcess::EWaitAndForkResult FUnixPlatformProcess::WaitAndFork()
 	struct FForkSignalData
 	{
 		FForkSignalData() = default;
-		FForkSignalData(int32 InSignal, FDateTime&& InTime) : SignalValue(InSignal), Time(MoveTemp(InTime)) {}
+		FForkSignalData(int32 InSignal, double InSeconds) : SignalValue(InSignal), TimeSeconds(InSeconds) {}
 
 		int32 SignalValue = 0;
-		FDateTime Time;
+		double TimeSeconds = 0.0;
 	};
 
 	static TCircularQueue<FForkSignalData> WaitAndForkSignalQueue(WAIT_AND_FORK_QUEUE_LENGTH);
@@ -1461,7 +1470,7 @@ FGenericPlatformProcess::EWaitAndForkResult FUnixPlatformProcess::WaitAndFork()
 	{
 		for (int32 ForkIdx = 0; ForkIdx < NumForks; ++ForkIdx)
 		{
-			WaitAndForkSignalQueue.Enqueue(FForkSignalData(ForkIdx + 1, FDateTime()));
+			WaitAndForkSignalQueue.Enqueue(FForkSignalData(ForkIdx + 1, FPlatformTime::Seconds()));
 		}
 	}
 
@@ -1498,7 +1507,7 @@ FGenericPlatformProcess::EWaitAndForkResult FUnixPlatformProcess::WaitAndFork()
 		Action.sa_sigaction = [](int32 Signal, siginfo_t* Info, void* Context) {
 			if (Signal == WAIT_AND_FORK_QUEUE_SIGNAL && Info)
 			{
-				WaitAndForkSignalQueue.Enqueue(FForkSignalData(Info->si_value.sival_int, FDateTime()));
+				WaitAndForkSignalQueue.Enqueue(FForkSignalData(Info->si_value.sival_int, FPlatformTime::Seconds()));
 			}
 		};
 		sigaction(WAIT_AND_FORK_QUEUE_SIGNAL, &Action, nullptr);
@@ -1555,11 +1564,9 @@ FGenericPlatformProcess::EWaitAndForkResult FUnixPlatformProcess::WaitAndFork()
 			// Sleep for a short while to avoid spamming new processes to the OS all at once
 			FPlatformProcess::Sleep(WAIT_AND_FORK_CHILD_SPAWN_DELAY);
 
-			FDateTime SignalReceived = SignalData.Time;
-
 			FCoreDelegates::OnParentPreFork.Broadcast();
 
-			UE_LOG(LogFork, Log, TEXT("[Parent] WaitAndFork processing child request %04hx-%04hx received at: %s"), Cookie, ChildIdx, *SignalReceived.ToString());
+			UE_LOG(LogFork, Log, TEXT("[Parent] WaitAndFork processing child request %04hx-%04hx received at: %s"), Cookie, ChildIdx, *UnixPlatformProcess::GetReadableTime(FPlatformTime::Seconds() - SignalData.TimeSeconds));
 
 			if (UnixPlatformProcess::GLogMemoryStatsWhenForking)
 			{
@@ -1741,7 +1748,7 @@ FGenericPlatformProcess::EWaitAndForkResult FUnixPlatformProcess::WaitAndFork()
 					else if (NumForks > 0 && ChildPidAndSignal.SignalValue > 0 && ChildPidAndSignal.SignalValue <= NumForks)
 					{
 						UE_LOG(LogFork, Log, TEXT("[Parent] WaitAndFork child %d missing. This was NumForks child %d. Relaunching..."), ChildPidAndSignal.Pid, ChildPidAndSignal.SignalValue);
-						WaitAndForkSignalQueue.Enqueue(FForkSignalData(ChildPidAndSignal.SignalValue, FDateTime()));
+						WaitAndForkSignalQueue.Enqueue(FForkSignalData(ChildPidAndSignal.SignalValue, FPlatformTime::Seconds()));
 					}
 					else
 					{
