[AutoRTFM] Annotate structure log APIs with UE_AUTORTFM_ALWAYS_OPEN

Matches the always-open boundary of the non-structured logging API.

diff --git a/Engine/Source/Programs/AutoRTFMTests/Private/UECoreTests.cpp b/Engine/Source/Programs/AutoRTFMTests/Private/UECoreTests.cpp
--- a/Engine/Source/Programs/AutoRTFMTests/Private/UECoreTests.cpp
+++ b/Engine/Source/Programs/AutoRTFMTests/Private/UECoreTests.cpp
@@ -9,6 +9,8 @@
 #include "Internationalization/TextCache.h"
 #include "Internationalization/TextFormatter.h"
 #include "Internationalization/TextHistory.h"
+#include "Logging/LogMacros.h"
+#include "Logging/StructuredLog.h"
 #include "Serialization/CustomVersion.h"
 #include "UObject/NameTypes.h"
 #include "UObject/UObjectArray.h"
@@ -19,6 +21,8 @@
 #include "Misc/ConfigCacheIni.h"
 #include "StructUtils/PropertyBag.h"
 
+DEFINE_LOG_CATEGORY_STATIC(LogAutoRTFMTests, Log, All)
+
 TEST_CASE("UECore.FDelegateHandle")
 {
 	FDelegateHandle Handle;
@@ -2342,3 +2346,22 @@ TEST_CASE("UECore.FAssetDataTagMapSharedView")
 		}
 	}
 }
+
+TEST_CASE("UECore.UE_LOGFMT")
+{
+	SECTION("Commit")
+	{
+		AutoRTFM::Testing::Commit([&]
+		{
+			UE_LOGFMT(LogAutoRTFMTests, Log, "{Animal} says {Sound}", TEXT("Cat"), TEXT("meow!"));
+		});
+	}
+	SECTION("Abort")
+	{
+		AutoRTFM::Testing::Abort([&]
+		{
+			UE_LOGFMT(LogAutoRTFMTests, Log, "{Animal} says {Sound}", TEXT("Cat"), TEXT("meow!"));
+			AutoRTFM::AbortTransaction();
+		});
+	}
+}
diff --git a/Engine/Source/Runtime/Core/Private/Logging/StructuredLog.cpp b/Engine/Source/Runtime/Core/Private/Logging/StructuredLog.cpp
--- a/Engine/Source/Runtime/Core/Private/Logging/StructuredLog.cpp
+++ b/Engine/Source/Runtime/Core/Private/Logging/StructuredLog.cpp
@@ -993,6 +993,7 @@ inline static void DispatchStaticLogRecord(const StaticLogRecordType& Log, const
 	DispatchLogRecord(Record);
 }
 
+UE_AUTORTFM_ALWAYS_OPEN
 void LogWithFieldArray(const FLogCategoryBase& Category, const FStaticLogRecord& Log, const FLogField* Fields, const int32 FieldCount)
 {
 #if !NO_LOGGING
@@ -1000,6 +1001,7 @@ void LogWithFieldArray(const FLogCategoryBase& Category, const FStaticLogRecord&
 #endif
 }
 
+UE_AUTORTFM_ALWAYS_OPEN
 void LogWithNoFields(const FLogCategoryBase& Category, const FStaticLogRecord& Log)
 {
 #if !NO_LOGGING
@@ -1009,6 +1011,7 @@ void LogWithNoFields(const FLogCategoryBase& Category, const FStaticLogRecord& L
 #endif
 }
 
+UE_AUTORTFM_ALWAYS_OPEN
 void LogWithFieldArray(const FLogCategoryBase& Category, const FStaticLocalizedLogRecord& Log, const FLogField* Fields, const int32 FieldCount)
 {
 #if !NO_LOGGING
@@ -1019,6 +1022,7 @@ void LogWithFieldArray(const FLogCategoryBase& Category, const FStaticLocalizedL
 #endif
 }
 
+UE_AUTORTFM_ALWAYS_OPEN
 void LogWithNoFields(const FLogCategoryBase& Category, const FStaticLocalizedLogRecord& Log)
 {
 #if !NO_LOGGING
@@ -1026,6 +1030,7 @@ void LogWithNoFields(const FLogCategoryBase& Category, const FStaticLocalizedLog
 #endif
 }
 
+UE_AUTORTFM_ALWAYS_OPEN
 void FatalLogWithFieldArray(const FLogCategoryBase& Category, const FStaticLogRecord& Log, const FLogField* Fields, const int32 FieldCount)
 {
 #if !NO_LOGGING
@@ -1040,6 +1045,7 @@ void FatalLogWithFieldArray(const FLogCategoryBase& Category, const FStaticLogRe
 #endif
 }
 
+UE_AUTORTFM_ALWAYS_OPEN
 void FatalLogWithNoFields(const FLogCategoryBase& Category, const FStaticLogRecord& Log)
 {
 #if !NO_LOGGING
@@ -1049,6 +1055,7 @@ void FatalLogWithNoFields(const FLogCategoryBase& Category, const FStaticLogReco
 #endif
 }
 
+UE_AUTORTFM_ALWAYS_OPEN
 void FatalLogWithFieldArray(const FLogCategoryBase& Category, const FStaticLocalizedLogRecord& Log, const FLogField* Fields, const int32 FieldCount)
 {
 #if !NO_LOGGING
@@ -1063,6 +1070,7 @@ void FatalLogWithFieldArray(const FLogCategoryBase& Category, const FStaticLocal
 #endif
 }
 
+UE_AUTORTFM_ALWAYS_OPEN
 void FatalLogWithNoFields(const FLogCategoryBase& Category, const FStaticLocalizedLogRecord& Log)
 {
 #if !NO_LOGGING
diff --git a/Engine/Source/Runtime/Core/Public/Logging/StructuredLog.h b/Engine/Source/Runtime/Core/Public/Logging/StructuredLog.h
--- a/Engine/Source/Runtime/Core/Public/Logging/StructuredLog.h
+++ b/Engine/Source/Runtime/Core/Public/Logging/StructuredLog.h
@@ -449,7 +449,7 @@ struct TFieldArgType<TLogFieldName<NameType>>
 
 /** Log with fields created from the arguments, which may be values or pairs of name/value. */
 template <typename... FieldArgTypes, typename LogType>
-UE_COLD UE_DEBUG_SECTION void LogWithFields(const FLogCategoryBase& Category, const LogType& Log, typename TFieldArgType<FieldArgTypes>::Type... FieldArgs)
+UE_AUTORTFM_ALWAYS_OPEN UE_COLD UE_DEBUG_SECTION void LogWithFields(const FLogCategoryBase& Category, const LogType& Log, typename TFieldArgType<FieldArgTypes>::Type... FieldArgs)
 {
 	constexpr int32 FieldCount = FLogFieldCreator::template GetCount<FieldArgTypes...>();
 	static_assert(FieldCount > 0);
@@ -460,7 +460,7 @@ UE_COLD UE_DEBUG_SECTION void LogWithFields(const FLogCategoryBase& Category, co
 
 /** Fatal log with fields created from the arguments, which may be values or pairs of name/value. */
 template <typename... FieldArgTypes, typename LogType>
-inline void UE_COLD FatalLogWithFields(const FLogCategoryBase& Category, const LogType& Log, typename TFieldArgType<FieldArgTypes>::Type... FieldArgs)
+UE_AUTORTFM_ALWAYS_OPEN inline void UE_COLD FatalLogWithFields(const FLogCategoryBase& Category, const LogType& Log, typename TFieldArgType<FieldArgTypes>::Type... FieldArgs)
 {
 	constexpr int32 FieldCount = FLogFieldCreator::template GetCount<FieldArgTypes...>();
 	static_assert(FieldCount > 0);
@@ -471,7 +471,7 @@ inline void UE_COLD FatalLogWithFields(const FLogCategoryBase& Category, const L
 
 /** Log if the category is active at this level of verbosity. */
 template <typename LogCategoryType, ELogVerbosity::Type Verbosity, typename LogRecordType, typename... FieldArgTypes>
-inline void LogIfActive(const LogCategoryType& Category, const LogRecordType& Log, FieldArgTypes&&... FieldArgs)
+UE_AUTORTFM_ALWAYS_OPEN inline void LogIfActive(const LogCategoryType& Category, const LogRecordType& Log, FieldArgTypes&&... FieldArgs)
 {
 	static_assert((Verbosity & ELogVerbosity::VerbosityMask) < ELogVerbosity::NumVerbosity && Verbosity > 0, "Verbosity must be constant and in range.");
 	if constexpr ((Verbosity & ELogVerbosity::VerbosityMask) == ELogVerbosity::Fatal)
