Fix config unloading logs to reduce spam after heavy GFP activity

#rnx

diff --git a/Engine/Source/Runtime/Core/Private/Misc/ConfigCacheIni.cpp b/Engine/Source/Runtime/Core/Private/Misc/ConfigCacheIni.cpp
--- a/Engine/Source/Runtime/Core/Private/Misc/ConfigCacheIni.cpp
+++ b/Engine/Source/Runtime/Core/Private/Misc/ConfigCacheIni.cpp
@@ -3642,7 +3642,7 @@ void FConfigCacheIni::Tick(float DeltaSeconds)
 	}
 	BranchesToCheck[1] = OtherFiles.FindRef(OtherFileNames[OtherFileToCheckForUnload++]);
 
-	checkf(OtherFileNames.Num() == OtherFiles.Num(), TEXT("OtherFIles and OtherFileNames are out of sync! %d other files, %d other file names!"), OtherFileNames.Num(), OtherFiles.Num());
+	checkf(OtherFileNames.Num() == OtherFiles.Num(), TEXT("OtherFiles and OtherFileNames are out of sync! %d other files, %d other file names!"), OtherFileNames.Num(), OtherFiles.Num());
 	
 	// now check for unused files
 	double Now = FPlatformTime::Seconds();
@@ -3726,7 +3726,7 @@ FConfigBranch* FConfigCacheIni::FindBranch(FName BaseIniName, const FString& Fil
 	if (Branch && Branch->InactiveTimer >= 0 && GTimeToUnloadConfig > 0)
 	{
 		Branch->InactiveTimer = FPlatformTime::Seconds();
-		UE_LOG(LogConfig, Log, TEXT("REsetting InactiveTimer for %s"), *Branch->IniName.ToString());
+		UE_LOG(LogConfig, Verbose, TEXT("Resetting InactiveTimer for %s"), *Branch->IniName.ToString());
 	}
 
 	return Branch;
@@ -3862,7 +3862,7 @@ FConfigFile* FConfigCacheIni::FindConfigFileWithBaseName(FName BaseName)
 		if (Result && Result->InactiveTimer >= 0 && GTimeToUnloadConfig > 0)
 		{
 			Result->InactiveTimer = FPlatformTime::Seconds();
-			UE_LOG(LogConfig, Log, TEXT("REsetting InactiveTimer for %s"), *Result->IniName.ToString());
+			UE_LOG(LogConfig, Verbose, TEXT("Resetting InactiveTimer for %s"), *Result->IniName.ToString());
 		}
 		return &Result->InMemoryFile;
 	}
