Fix build failure in the AutoRTFM runtime.

#rnx

diff --git a/Engine/Source/Runtime/Core/Private/AutoRTFM/API.cpp b/Engine/Source/Runtime/Core/Private/AutoRTFM/API.cpp
--- a/Engine/Source/Runtime/Core/Private/AutoRTFM/API.cpp
+++ b/Engine/Source/Runtime/Core/Private/AutoRTFM/API.cpp
@@ -366,7 +366,7 @@ AUTORTFM_CORE_API autortfm_result RTFM_autortfm_cascading_abort_transaction()
 AUTORTFM_CORE_API void RTFM_autortfm_clear_transaction_status()
 {
 	UE_LOG(LogAutoRTFM, Fatal, TEXT("The function `autortfm_clear_transaction_status` was called from closed code."));
-	AutoRTFM::Unreachable();
+	AutoRTFM::InternalUnreachable();
 }
 
 AUTORTFM_CORE_API autortfm_status RTFM_autortfm_close(void (*UninstrumentedWork)(void*), void (*InstrumentedWork)(void*), void* Arg)
@@ -530,7 +530,7 @@ AUTORTFM_CORE_API void RTFM_autortfm_did_free(void* Ptr)
 {
 	// We should never-ever-ever actually free memory from within closed code of
 	// a transaction.
-	AutoRTFM::Unreachable();
+	AutoRTFM::InternalUnreachable();
 }
 
 AUTORTFM_CORE_API void RTFM_autortfm_check_consistency_assuming_no_races()
@@ -924,7 +924,7 @@ extern "C" void autortfm_check_consistency_assuming_no_races()
 
     if (FContext::Get()->IsTransactional())
     {
-        AutoRTFM::Unreachable();
+        AutoRTFM::InternalUnreachable();
     }
 }
 
diff --git a/Engine/Source/Runtime/Core/Private/AutoRTFM/Context.cpp b/Engine/Source/Runtime/Core/Private/AutoRTFM/Context.cpp
--- a/Engine/Source/Runtime/Core/Private/AutoRTFM/Context.cpp
+++ b/Engine/Source/Runtime/Core/Private/AutoRTFM/Context.cpp
@@ -330,7 +330,7 @@ void FContext::ClearTransactionStatus()
 		Status = EContextStatus::OnTrack;
 		break;
 	default:
-		AutoRTFM::Unreachable();
+		AutoRTFM::InternalUnreachable();
 	}
 }
 
@@ -355,7 +355,7 @@ ETransactionResult FContext::ResolveNestedTransaction(FTransaction* NewTransacti
 	case EContextStatus::AbortedByCascade:
 		return ETransactionResult::AbortedByCascade;
 	default:
-		AutoRTFM::Unreachable();
+		AutoRTFM::InternalUnreachable();
 	}
 }
 
