- Disabled binary config in editor, enabled on client (on behalf of Josh.Adams)

diff --git a/Engine/Source/Runtime/Core/Private/Misc/ConfigCacheIni.cpp b/Engine/Source/Runtime/Core/Private/Misc/ConfigCacheIni.cpp
--- a/Engine/Source/Runtime/Core/Private/Misc/ConfigCacheIni.cpp
+++ b/Engine/Source/Runtime/Core/Private/Misc/ConfigCacheIni.cpp
@@ -5741,7 +5741,7 @@ void FConfigCacheIni::InitializeConfigSystem()
 	
 	InitializeConfigRemap();
 	
-#if PLATFORM_SUPPORTS_BINARYCONFIG && PRELOAD_BINARY_CONFIG
+#if PLATFORM_SUPPORTS_BINARYCONFIG && PRELOAD_BINARY_CONFIG && !WITH_EDITOR
 	// attempt to load from staged binary config data
 	if (!FParse::Param(FCommandLine::Get(), TEXT("textconfig")) &&
 		FConfigCacheIni::CreateGConfigFromSaved(nullptr))
