Added separate stats for total cached memory and immediately freeable cached memory in MB2.
Subtract total cached memory in GMalloc from FMallocUnused as it's technically is not used by any FMemory allocations atm

diff --git a/Engine/Source/Runtime/Core/Private/HAL/LowLevelMemTracker.cpp b/Engine/Source/Runtime/Core/Private/HAL/LowLevelMemTracker.cpp
--- a/Engine/Source/Runtime/Core/Private/HAL/LowLevelMemTracker.cpp
+++ b/Engine/Source/Runtime/Core/Private/HAL/LowLevelMemTracker.cpp
@@ -1996,7 +1996,8 @@ void FLowLevelMemTracker::TickInternal()
 	// Calculate FMalloc unused stat and set it in the Default tracker.
 	const int64 FMallocAmount = DefaultTracker.GetAllocTypeAmount(ELLMAllocType::FMalloc);
 	const int64 FMallocPlatformAmount = PlatformTracker.GetTagAmount(FindOrAddTagData(ELLMTag::FMalloc));
-	int64 FMallocUnused = FMallocPlatformAmount - FMallocAmount;
+	const int64 CachedFreeMemory = GMalloc ? GMalloc->GetTotalFreeCachedMemorySize() : 0;
+	int64 FMallocUnused = FMallocPlatformAmount - FMallocAmount - CachedFreeMemory;
 	if (FMallocPlatformAmount == 0)
 	{
 		// We do not have instrumentation for this allocator, and so can not calculate how much memory it is using
diff --git a/Engine/Source/Runtime/Core/Private/HAL/MallocBinned2.cpp b/Engine/Source/Runtime/Core/Private/HAL/MallocBinned2.cpp
--- a/Engine/Source/Runtime/Core/Private/HAL/MallocBinned2.cpp
+++ b/Engine/Source/Runtime/Core/Private/HAL/MallocBinned2.cpp
@@ -31,13 +31,6 @@ PRAGMA_DISABLE_UNSAFE_TYPECAST_WARNINGS
 #	define GBinned2MoveOSFreesOffTimeCriticalThreads	UE_DEFAULT_GBinned2MoveOSFreesOffTimeCriticalThreads
 #endif
 
-int32 EnableLegacyCachedOSPageAllocatorFreeMemReporting = 0;
-static FAutoConsoleVariableRef GEnableLegacyCachedOSPageAllocatorFreeMemReportingCVar(
-	TEXT("MallocBinned2.EnableLegacyCachedOSPageAllocatorFreeMemReporting"),
-	EnableLegacyCachedOSPageAllocatorFreeMemReporting,
-	TEXT("Temporary. Whether or not report full amount of cached memory in CachedOSPageAllocator (legacy) or report only the amount of memory immediately available to return to the OS")
-);
-
 #if UE_MB2_ALLOCATOR_STATS
 	std::atomic<int64> AllocatedOSSmallPoolMemory(0);
 
@@ -987,8 +980,9 @@ void FMallocBinned2::DumpAllocatorStats(class FOutputDevice& Ar)
 
 void FMallocBinned2::UpdateStats()
 {
-	//TODO: this will report total cached free memory in the COSPA, however a big chunk of that memory can be immediately freed to the OS
-	CSV_CUSTOM_STAT(FMemory, AllocatorCachedSlackMB, (int32)(CachedOSPageAllocator.GetCachedFreeTotal()/(1024*1024)), ECsvCustomStatOp::Set);
+	//Report total cached free memory in the COSPA and separately report memory that can be immediately freed back the kernel at any time
+	CSV_CUSTOM_STAT(FMemory, AllocatorCachedSlackMB, (int32)(CachedOSPageAllocator.GetCachedFreeTotal() / (1024 * 1024)), ECsvCustomStatOp::Set);
+	CSV_CUSTOM_STAT(FMemory, AllocatorImmediatelyFreeableCachedSlackMB, (int32)(CachedOSPageAllocator.GetCachedImmediatelyFreeable() / (1024 * 1024)), ECsvCustomStatOp::Set);
 
 	CachedOSPageAllocator.UpdateStats();
 	FScopedVirtualMallocTimer::UpdateStats();
diff --git a/Engine/Source/Runtime/Core/Public/HAL/MallocBinned2.h b/Engine/Source/Runtime/Core/Public/HAL/MallocBinned2.h
--- a/Engine/Source/Runtime/Core/Public/HAL/MallocBinned2.h
+++ b/Engine/Source/Runtime/Core/Public/HAL/MallocBinned2.h
@@ -58,8 +58,6 @@ struct FGenericMemoryStats;
 	extern int32 RecursionCounter;
 #endif
 
-extern int32 EnableLegacyCachedOSPageAllocatorFreeMemReporting;
-
 // Canary value used in FFreeBlock
 // A constant value unless we're compiled with fork support in which case there are two values identifying whether the page
 // was allocated pre- or post-fork
@@ -514,15 +512,14 @@ public:
 	CORE_API virtual void OnMallocInitialized() override;
 	CORE_API virtual void OnPreFork() override;
 	CORE_API virtual void OnPostFork() override;
-	virtual uint64 GetFreeCachedMemorySize() const
+	virtual uint64 GetImmediatelyFreeableCachedMemorySize() const override
 	{
-		if (EnableLegacyCachedOSPageAllocatorFreeMemReporting)
-		{
-			return CachedOSPageAllocator.GetCachedFreeTotal();
-		}
-
 		return CachedOSPageAllocator.GetCachedImmediatelyFreeable();
 	}
+	virtual uint64 GetTotalFreeCachedMemorySize() const override
+	{
+		return CachedOSPageAllocator.GetCachedFreeTotal();
+	}
 	// End FMalloc interface.
 
 	CORE_API void* MallocExternalSmall(SIZE_T Size, uint32 Alignment);
diff --git a/Engine/Source/Runtime/Core/Public/HAL/MemoryBase.h b/Engine/Source/Runtime/Core/Public/HAL/MemoryBase.h
--- a/Engine/Source/Runtime/Core/Public/HAL/MemoryBase.h
+++ b/Engine/Source/Runtime/Core/Public/HAL/MemoryBase.h
@@ -251,7 +251,17 @@ public:
 	/**
 	 * Returns the amount of free memory cached by the allocator that can be returned to the system in case of a memory shortage
 	 */
-	virtual uint64 GetFreeCachedMemorySize() const
+	virtual uint64 GetImmediatelyFreeableCachedMemorySize() const
+	{
+		return 0;
+	}
+
+	/**
+	 * Returns the amount of total free memory cached by the allocator.
+	 * This includes memory that can be returned to the system in case of a memory shortage, see GetImmediatelyFreeableCachedMemorySize()
+	 * as well as any memory that can't be returned back to the kernel, but can be used to satisfy some of the allocation requirements.
+	 */
+	virtual uint64 GetTotalFreeCachedMemorySize() const
 	{
 		return 0;
 	}
