Csv Profiler: Add a compile define CSV_PROFILER_ALLOW_SENSITIVE_BUILTIN_METADATA so sensitive fields like commandline and loginID can be stripped (or included if required) on a per-project basis. If not specified expliticly, this defaults to 0 in shipping builds (matching the existing behavior) and 1 otherwise.
#tests tested locally
[RN] Csv Profiler: Add a compile define CSV_PROFILER_ALLOW_SENSITIVE_BUILTIN_METADATA so sensitive fields like commandline and loginID can be stripped (or included if required) on a per-project basis. If not specified expliticly, this defaults to 0 in shipping builds and 1 otherwise.

diff --git a/Engine/Source/Runtime/Core/Private/ProfilingDebugging/CsvProfiler.cpp b/Engine/Source/Runtime/Core/Private/ProfilingDebugging/CsvProfiler.cpp
--- a/Engine/Source/Runtime/Core/Private/ProfilingDebugging/CsvProfiler.cpp
+++ b/Engine/Source/Runtime/Core/Private/ProfilingDebugging/CsvProfiler.cpp
@@ -49,6 +49,11 @@
 #define CSV_PROFILER_SUPPORT_NAMED_EVENTS ENABLE_NAMED_EVENTS
 #endif
 
+#ifndef CSV_PROFILER_ALLOW_SENSITIVE_BUILTIN_METADATA
+// Disable sensitive metadata like LoginId and Commandline outside of shipping builds
+#define CSV_PROFILER_ALLOW_SENSITIVE_BUILTIN_METADATA !UE_BUILD_SHIPPING
+#endif
+
 #define REPAIR_MARKER_STACKS 1
 
 // Note: Enabling this will break a lot of reporting but it may be useful for debugging. If disabled (default), the thread name is omitted from custom stat columns
@@ -3141,7 +3146,7 @@ FCsvProfiler::FCsvProfiler()
 		}
 	);
 
-#if !UE_BUILD_SHIPPING
+#if CSV_PROFILER_ALLOW_SENSITIVE_BUILTIN_METADATA
 	// for privacy, personal and free text fields are not allowed in shipping
 	SetMetadataInternal(TEXT("Commandline"), *CommandlineStr, false);
 	SetMetadataInternal(TEXT("LoginID"), *FPlatformMisc::GetLoginId());
@@ -3157,7 +3162,7 @@ FCsvProfiler::FCsvProfiler()
 	{
 		SetMetadataInternal(TEXT("DeviceID"), *DeviceID);
 	}
-#endif
+#endif //CSV_PROFILER_ALLOW_SENSITIVE_BUILTIN_METADATA
 }
 
 FCsvProfiler::~FCsvProfiler()
