[Core/Logging] Fix program counter passed to StaticFailDebugV() and FDebug::ProcessFatalError()

Call PLATFORM_RETURN_ADDRESS() in BasicFatalLog() and pass that as an additional parameter to BasicFatalLogV()

diff --git a/Engine/Source/Runtime/Core/Private/Logging/StructuredLog.cpp b/Engine/Source/Runtime/Core/Private/Logging/StructuredLog.cpp
--- a/Engine/Source/Runtime/Core/Private/Logging/StructuredLog.cpp
+++ b/Engine/Source/Runtime/Core/Private/Logging/StructuredLog.cpp
@@ -1254,13 +1254,13 @@ void BasicLog(const FLogCategoryBase& Category, const FStaticBasicLogRecord* Log
 }
 
 UE_AUTORTFM_ALWAYS_OPEN
-static void BasicFatalLogV(const FLogCategoryBase& Category, const FStaticBasicLogRecord* Log, va_list Args)
+static void BasicFatalLogV(const FLogCategoryBase& Category, const FStaticBasicLogRecord* Log, void* ProgramCounter, va_list Args)
 {
 #if !NO_LOGGING
-	StaticFailDebugV(TEXT("Fatal error:"), "", Log->File, Log->Line, /*bIsEnsure*/ false, PLATFORM_RETURN_ADDRESS(), Log->Format, Args);
+	StaticFailDebugV(TEXT("Fatal error:"), "", Log->File, Log->Line, /*bIsEnsure*/ false, ProgramCounter, Log->Format, Args);
 
 	UE_DEBUG_BREAK_AND_PROMPT_FOR_REMOTE();
-	FDebug::ProcessFatalError(PLATFORM_RETURN_ADDRESS());
+	FDebug::ProcessFatalError(ProgramCounter);
 #endif
 }
 
@@ -1269,7 +1269,7 @@ void BasicFatalLog(const FLogCategoryBase& Category, const FStaticBasicLogRecord
 #if !NO_LOGGING
 	va_list Args;
 	va_start(Args, Log);
-	BasicFatalLogV(Category, Log, Args);
+	BasicFatalLogV(Category, Log, PLATFORM_RETURN_ADDRESS(), Args);
 	va_end(Args);
 #endif
 }
