[AutoRTFM] Fix improperly-typed TCHAR variables.

This prevented the AutoRTFM compiler's lit tests from building, as they have no concept of TCHAR.

diff --git a/Engine/Source/Runtime/Core/Public/AutoRTFM/AutoRTFM.h b/Engine/Source/Runtime/Core/Public/AutoRTFM/AutoRTFM.h
--- a/Engine/Source/Runtime/Core/Public/AutoRTFM/AutoRTFM.h
+++ b/Engine/Source/Runtime/Core/Public/AutoRTFM/AutoRTFM.h
@@ -1005,13 +1005,13 @@ template<typename Type> static UE_AUTORTFM_FORCEINLINE void RecordOpenWrite(Type
 
 // Report that a unreachable codepath is being hit. Used to manually ban certain codepaths
 // from being transactionally safe.
-static UE_AUTORTFM_FORCEINLINE void Unreachable(const TCHAR* const Message = nullptr)
+static UE_AUTORTFM_FORCEINLINE void Unreachable(const UE_AUTORTFM_TCHAR* const Message = nullptr)
 {
 	autortfm_unreachable(Message);
 }
 
 // If we are running within a transaction, call `AutoRTFM::Unreachable`.
-static UE_AUTORTFM_FORCEINLINE void UnreachableIfTransactional(const TCHAR* const Message = nullptr)
+static UE_AUTORTFM_FORCEINLINE void UnreachableIfTransactional(const UE_AUTORTFM_TCHAR* const Message = nullptr)
 {
 	if (AutoRTFM::IsTransactional())
 	{
