[Insights] Improved the debug stats/counters for TraceLog:
    - Added defienes in TraceAuxiliary.cpp to easily enable/disable code for STAT, TRACE, LLM and/or CSV stats/counters.
    - Added additional debug counters/stats: "Cache Unused", "Emitted not Tarced", "Memory Error" (total - block pool - fixed - shared - cache).
    - Added also trace counters API stats and made it default (instead of STAT counters).
    - Added block pool, fixed buffers, shared buffers, emitted and traced stats also to the output of "trace.status" console command.
    - Added initial zero values to all stats/counters (improves graph display in Insights for all the stats).
    - Simplified code re registering EndFrame callbacks.
    - Added support for registering a callback to be called each time TraceLog updates: added OnUpdateFunc in UE::Trace::FInitializeDesc + added  UE::Trace::SetUpdateCallback(func).
    - The new trace update callback is now used to emit stats/counters also during engine initialization. Once the engine finishes the initialization, the stats counters will be emitted only once per frame (by resetting the update callback and further emitting stats from end frame updates).

diff --git a/Engine/Source/Runtime/Core/Private/ProfilingDebugging/TraceAuxiliary.cpp b/Engine/Source/Runtime/Core/Private/ProfilingDebugging/TraceAuxiliary.cpp
--- a/Engine/Source/Runtime/Core/Private/ProfilingDebugging/TraceAuxiliary.cpp
+++ b/Engine/Source/Runtime/Core/Private/ProfilingDebugging/TraceAuxiliary.cpp
@@ -1,11 +1,14 @@
 // Copyright Epic Games, Inc. All Rights Reserved.
 
 #include "ProfilingDebugging/TraceAuxiliary.h"
-#include "ProfilingDebugging/TagTrace.h"
-#include "ProfilingDebugging/StringsTrace.h"
-#include "Trace/Trace.h"
+
 #include "CoreGlobals.h"
 #include "Misc/ScopeRWLock.h"
+#include "ProfilingDebugging/CountersTrace.h"
+#include "ProfilingDebugging/CsvProfilerConfig.h"
+#include "ProfilingDebugging/StringsTrace.h"
+#include "ProfilingDebugging/TagTrace.h"
+#include "Trace/Trace.h"
 
 #if PLATFORM_WINDOWS
 #include "Windows/AllowWindowsPlatformTypes.h"
@@ -29,6 +32,22 @@
 	#define UE_TRACE_AUTOSTART 1
 #endif
 
+#if !defined(UE_TRACEAUX_EMIT_STAT_COUNTERS)
+	#define UE_TRACEAUX_EMIT_STAT_COUNTERS (0 && STATS)
+#endif
+
+#if !defined(UE_TRACEAUX_EMIT_TRACE_COUNTERS)
+	#define UE_TRACEAUX_EMIT_TRACE_COUNTERS (1 && COUNTERSTRACE_ENABLED)
+#endif
+
+#if !defined(UE_TRACEAUX_EMIT_LLM_STATS)
+	#define UE_TRACEAUX_EMIT_LLM_STATS (1 && ENABLE_LOW_LEVEL_MEM_TRACKER)
+#endif
+
+#if !defined(UE_TRACEAUX_EMIT_CSV_STATS)
+	#define UE_TRACEAUX_EMIT_CSV_STATS (1 && CSV_PROFILER_STATS)
+#endif
+
 #if UE_TRACE_SERVER_LAUNCH_ENABLED || UE_TRACE_SERVER_CONTROLS_ENABLED
 #include "HAL/PlatformProcess.h"
 #include "Misc/CommandLine.h"
@@ -75,21 +94,48 @@ const FTraceAuxiliary::FChannelPreset GMemoryLightChannels(TEXT("Memory_Light"),
 const TCHAR* GTraceConfigSection = TEXT("Trace.Config");
 static UE::Trace::FInitializeDesc GInitializeDesc;
 
-////////////////////////////////////////////////////////////////////////////////
 CSV_DEFINE_CATEGORY(Trace, true);
 
-////////////////////////////////////////////////////////////////////////////////
+#if UE_TRACEAUX_EMIT_STAT_COUNTERS
 DECLARE_STATS_GROUP(TEXT("TraceLog"), STATGROUP_Trace, STATCAT_Advanced);
-DECLARE_MEMORY_STAT(TEXT("Memory used"), STAT_TraceMemoryUsed, STATGROUP_Trace);
-DECLARE_MEMORY_STAT(TEXT("Block pool"), STAT_TraceBlockPoolAllocated, STATGROUP_Trace);
-DECLARE_MEMORY_STAT(TEXT("Shared buffer"), STAT_TraceSharedAllocated, STATGROUP_Trace);
-DECLARE_MEMORY_STAT(TEXT("Fixed buffer"), STAT_TraceFixedAllocated, STATGROUP_Trace);
-DECLARE_MEMORY_STAT(TEXT("Important event cache"), STAT_TraceCacheAllocated, STATGROUP_Trace);
-DECLARE_MEMORY_STAT(TEXT("Important event cache used"), STAT_TraceCacheUsed, STATGROUP_Trace);
-DECLARE_MEMORY_STAT(TEXT("Important event cache waste"), STAT_TraceCacheWaste, STATGROUP_Trace);
-DECLARE_MEMORY_STAT(TEXT("Sent"), STAT_TraceSent, STATGROUP_Trace);
-DECLARE_MEMORY_STAT(TEXT("Traced"), STAT_TraceTraced, STATGROUP_Trace);
-DECLARE_MEMORY_STAT(TEXT("Emitted"), STAT_TraceEmitted, STATGROUP_Trace);
+DECLARE_MEMORY_STAT(TEXT("Total Memory Used"),             STAT_Trace_MemoryUsed,                 STATGROUP_Trace);
+DECLARE_MEMORY_STAT(TEXT("Block Pool Memory"),             STAT_Trace_Memory_BlockPool,           STATGROUP_Trace);
+DECLARE_MEMORY_STAT(TEXT("Fixed Buffers Memory"),          STAT_Trace_Memory_FixedBuffers,        STATGROUP_Trace);
+DECLARE_MEMORY_STAT(TEXT("Shared Buffers Memory"),         STAT_Trace_Memory_SharedBuffers,       STATGROUP_Trace);
+DECLARE_MEMORY_STAT(TEXT("Important Events Cache"),        STAT_Trace_Memory_Cache,               STATGROUP_Trace);
+DECLARE_MEMORY_STAT(TEXT("Important Events Cache Used"),   STAT_Trace_Memory_Cache_Used,          STATGROUP_Trace);
+DECLARE_MEMORY_STAT(TEXT("Important Events Cache Unused"), STAT_Trace_Memory_Cache_Unused,        STATGROUP_Trace);
+DECLARE_MEMORY_STAT(TEXT("Important Events Cache Wasted"), STAT_Trace_Memory_Cache_Wasted,        STATGROUP_Trace);
+DECLARE_MEMORY_STAT(TEXT("Bytes Emitted"),                 STAT_Trace_Bandwidth_Emitted,          STATGROUP_Trace);
+DECLARE_MEMORY_STAT(TEXT("Bytes Emitted Not Traced"),      STAT_Trace_Bandwidth_EmittedNotTraced, STATGROUP_Trace);
+DECLARE_MEMORY_STAT(TEXT("Bytes Traced"),                  STAT_Trace_Bandwidth_Traced,           STATGROUP_Trace);
+DECLARE_MEMORY_STAT(TEXT("Bytes Sent"),                    STAT_Trace_Bandwidth_Sent,             STATGROUP_Trace);
+DECLARE_MEMORY_STAT(TEXT("Memory Error"),                  STAT_Trace_MemoryError,                STATGROUP_Trace);
+#endif
+
+#if UE_TRACEAUX_EMIT_TRACE_COUNTERS
+TRACE_DECLARE_UNCHECKED_MEMORY_COUNTER(Trace_MemoryUsed,                 TEXT("Trace/MemoryUsed"));
+TRACE_DECLARE_UNCHECKED_MEMORY_COUNTER(Trace_Memory_BlockPool,           TEXT("Trace/Memory/BlockPool"));
+TRACE_DECLARE_UNCHECKED_MEMORY_COUNTER(Trace_Memory_FixedBuffers,        TEXT("Trace/Memory/FixedBuffers"));
+TRACE_DECLARE_UNCHECKED_MEMORY_COUNTER(Trace_Memory_SharedBuffers,       TEXT("Trace/Memory/SharedBuffers"));
+TRACE_DECLARE_UNCHECKED_MEMORY_COUNTER(Trace_Memory_Cache,               TEXT("Trace/Memory/Cache"));
+TRACE_DECLARE_UNCHECKED_MEMORY_COUNTER(Trace_Memory_Cache_Used,          TEXT("Trace/Memory/Cache/Used"));
+TRACE_DECLARE_UNCHECKED_MEMORY_COUNTER(Trace_Memory_Cache_Unused,        TEXT("Trace/Memory/Cache/Unused"));
+TRACE_DECLARE_UNCHECKED_MEMORY_COUNTER(Trace_Memory_Cache_Wasted,        TEXT("Trace/Memory/Cache/Wasted"));
+TRACE_DECLARE_UNCHECKED_MEMORY_COUNTER(Trace_Bandwidth_Emitted,          TEXT("Trace/Bandwidth/Emitted"));
+TRACE_DECLARE_UNCHECKED_MEMORY_COUNTER(Trace_Bandwidth_EmittedNotTraced, TEXT("Trace/Bandwidth/EmittedNotTraced"));
+TRACE_DECLARE_UNCHECKED_MEMORY_COUNTER(Trace_Bandwidth_Traced,           TEXT("Trace/Bandwidth/Traced"));
+TRACE_DECLARE_UNCHECKED_MEMORY_COUNTER(Trace_Bandwidth_Sent,             TEXT("Trace/Bandwidth/Sent"));
+TRACE_DECLARE_UNCHECKED_MEMORY_COUNTER(Trace_MemoryError,                TEXT("Trace/MemoryError"));
+#endif
+
+#if UE_TRACEAUX_EMIT_LLM_STATS
+LLM_DEFINE_TAG(TraceLog);
+LLM_DEFINE_TAG(TraceLog_BlockPool);
+LLM_DEFINE_TAG(TraceLog_FixedBuffers);
+LLM_DEFINE_TAG(TraceLog_SharedBuffers);
+LLM_DEFINE_TAG(TraceLog_Cache);
+#endif
 
 ////////////////////////////////////////////////////////////////////////////////
 enum class ETraceConnectType
@@ -133,7 +179,7 @@ public:
 	void SetTruncateFile(bool bTruncateFile);
 	void UpdateCsvStats() const;
 	void StartWorkerThread();
-	void StartEndFramePump();
+	void RegisterEndFrameCallbacks();
 	bool WriteSnapshot(const TCHAR* InFilePath, const FTraceAuxiliary::FLogCategoryAlias& LogCategory);
 	bool SendSnapshot(const TCHAR* InHost, uint32 InPort, const FTraceAuxiliary::FLogCategoryAlias& LogCategory);
 
@@ -470,21 +516,21 @@ bool FTraceAuxiliaryImpl::ConnectRelay(UPTRINT Handle, UE::Trace::IoWriteFunc Wr
 		UE_LOG_REF(LogTrace, Error, TEXT("Already tracing from unknown source."));
 		return true;
 	}
-	
+
 	const bool bConnected = UE::Trace::RelayTo(Handle, WriteFunc, CloseFunc);
 	if (!bConnected)
 	{
 		UE_LOG(LogTrace, Display, TEXT("Trace started (connected to relay endpoint)"));
-		
+
 		FString StartedDest = FString::Printf(TEXT("Relay endpoint (Handle: 0x%llx)"), Handle);
 		FTraceAuxiliary::EConnectionType StartedType = FTraceAuxiliary::EConnectionType::Relay;
-		
+
 		{
 			FWriteScopeLock _(CurrentTargetLock);
 			CurrentTraceTarget.TraceType = StartedType;
 			CurrentTraceTarget.TraceDest = StartedDest;
 		}
-		
+
 		TRACE_CPUPROFILER_EVENT_SCOPE(FTraceAux_OnTraceStarted);
 		FTraceAuxiliary::OnTraceStarted.Broadcast(StartedType, StartedDest);
 	}
@@ -829,26 +875,6 @@ void FTraceAuxiliaryImpl::GetActiveChannelsString(FStringBuilderBase& String) co
 	}, &String);
 }
 
-////////////////////////////////////////////////////////////////////////////////
-void FTraceAuxiliaryImpl::UpdateCsvStats() const
-{
-#if TRACE_PRIVATE_STATISTICS
-	// Only publish CSV stats if we have ever run tracing in order to reduce overhead in most runs.
-	static bool bDoStats = false;
-	if (UE::Trace::IsTracing() || bDoStats)
-	{
-		bDoStats = true;
-
-		UE::Trace::FStatistics Stats;
-		UE::Trace::GetStatistics(Stats);
-
-		CSV_CUSTOM_STAT(Trace, MemoryUsedMb,	double(Stats.MemoryUsed) / 1024.0 / 1024.0, ECsvCustomStatOp::Set);
-		CSV_CUSTOM_STAT(Trace, CacheUsedMb,		double(Stats.CacheUsed) / 1024.0 / 1024.0, ECsvCustomStatOp::Set);
-		CSV_CUSTOM_STAT(Trace, CacheWasteMb,	double(Stats.CacheWaste) / 1024.0 / 1024.0, ECsvCustomStatOp::Set);
-	}
-#endif
-}
-
 ////////////////////////////////////////////////////////////////////////////////
 void FTraceAuxiliaryImpl::StartWorkerThread()
 {
@@ -860,53 +886,122 @@ void FTraceAuxiliaryImpl::StartWorkerThread()
 	}
 }
 
-LLM_DEFINE_TAG(Trace_TraceLog);
-LLM_DEFINE_TAG(Trace_TraceLog_Cache);
+////////////////////////////////////////////////////////////////////////////////
+void TraceAuxiliaryEmitStats(const UE::Trace::FStatistics& Stats)
+{
+	uint64 TotalMemoryUsed = Stats.BlockPoolAllocated + Stats.FixedBufferAllocated + Stats.SharedBufferAllocated + Stats.CacheAllocated;
+	//ensure(TotalMemoryUsed == Stats.MemoryUsed);
+	//ensure(Stats.BytesEmitted >= Stats.BytesTraced);
+
+#if UE_TRACEAUX_EMIT_STAT_COUNTERS
+	SET_MEMORY_STAT(STAT_Trace_MemoryUsed, Stats.MemoryUsed);
+	SET_MEMORY_STAT(STAT_Trace_Memory_BlockPool, Stats.BlockPoolAllocated);
+	SET_MEMORY_STAT(STAT_Trace_Memory_FixedBuffers, Stats.FixedBufferAllocated);
+	SET_MEMORY_STAT(STAT_Trace_Memory_SharedBuffers, Stats.SharedBufferAllocated);
+	SET_MEMORY_STAT(STAT_Trace_Memory_Cache, Stats.CacheAllocated);
+	SET_MEMORY_STAT(STAT_Trace_Memory_Cache_Used, Stats.CacheUsed);
+	SET_MEMORY_STAT(STAT_Trace_Memory_Cache_Unused, (int64)Stats.CacheAllocated - (int64)Stats.CacheUsed);
+	SET_MEMORY_STAT(STAT_Trace_Memory_Cache_Wasted, Stats.CacheWaste);
+	SET_MEMORY_STAT(STAT_Trace_Bandwidth_Emitted, Stats.BytesEmitted);
+	SET_MEMORY_STAT(STAT_Trace_Bandwidth_EmittedNotTraced, (int64)Stats.BytesEmitted - (int64)Stats.BytesTraced);
+	SET_MEMORY_STAT(STAT_Trace_Bandwidth_Traced, Stats.BytesTraced);
+	SET_MEMORY_STAT(STAT_Trace_Bandwidth_Sent, Stats.BytesSent);
+	SET_MEMORY_STAT(STAT_Trace_MemoryError, (int64)Stats.MemoryUsed - (int64)TotalMemoryUsed);
+#endif
+
+#if UE_TRACEAUX_EMIT_TRACE_COUNTERS
+	TRACE_COUNTER_SET(Trace_MemoryUsed, Stats.MemoryUsed);
+	TRACE_COUNTER_SET(Trace_Memory_BlockPool, Stats.BlockPoolAllocated);
+	TRACE_COUNTER_SET(Trace_Memory_FixedBuffers, Stats.FixedBufferAllocated);
+	TRACE_COUNTER_SET(Trace_Memory_SharedBuffers, Stats.SharedBufferAllocated);
+	TRACE_COUNTER_SET(Trace_Memory_Cache, Stats.CacheAllocated);
+	TRACE_COUNTER_SET(Trace_Memory_Cache_Used, Stats.CacheUsed);
+	TRACE_COUNTER_SET(Trace_Memory_Cache_Unused, (int64)Stats.CacheAllocated - (int64)Stats.CacheUsed);
+	TRACE_COUNTER_SET(Trace_Memory_Cache_Wasted, Stats.CacheWaste);
+	TRACE_COUNTER_SET(Trace_Bandwidth_Emitted, Stats.BytesEmitted);
+	TRACE_COUNTER_SET(Trace_Bandwidth_EmittedNotTraced, (int64)Stats.BytesEmitted - (int64)Stats.BytesTraced);
+	TRACE_COUNTER_SET(Trace_Bandwidth_Traced, Stats.BytesTraced);
+	TRACE_COUNTER_SET(Trace_Bandwidth_Sent, Stats.BytesSent);
+	TRACE_COUNTER_SET(Trace_MemoryError, (int64)Stats.MemoryUsed - (int64)TotalMemoryUsed);
+#endif
+
+#if UE_TRACEAUX_EMIT_LLM_STATS && ENABLE_LOW_LEVEL_MEM_TRACKER
+	FLowLevelMemTracker::Get().SetTagAmountForTracker(ELLMTracker::Default, LLM_TAG_NAME(TraceLog), ELLMTagSet::None, Stats.MemoryUsed, true);
+	FLowLevelMemTracker::Get().SetTagAmountForTracker(ELLMTracker::Default, LLM_TAG_NAME(TraceLog_BlockPool), ELLMTagSet::None, Stats.BlockPoolAllocated, true);
+	FLowLevelMemTracker::Get().SetTagAmountForTracker(ELLMTracker::Default, LLM_TAG_NAME(TraceLog_FixedBuffers), ELLMTagSet::None, Stats.FixedBufferAllocated, true);
+	FLowLevelMemTracker::Get().SetTagAmountForTracker(ELLMTracker::Default, LLM_TAG_NAME(TraceLog_SharedBuffers), ELLMTagSet::None, Stats.SharedBufferAllocated, true);
+	FLowLevelMemTracker::Get().SetTagAmountForTracker(ELLMTracker::Default, LLM_TAG_NAME(TraceLog_Cache), ELLMTagSet::None, Stats.CacheAllocated, true);
+#endif
+}
 
 ////////////////////////////////////////////////////////////////////////////////
-void FTraceAuxiliaryImpl::StartEndFramePump()
+void TraceAuxiliaryOnTraceUpdateCallback()
 {
-	if (!GEndFrameDelegateHandle.IsValid())
-	{
-		// If the worker thread is disabled, pump the update from end frame
-		GEndFrameDelegateHandle = FCoreDelegates::OnEndFrame.AddStatic(UE::Trace::Update);
-	}
-	if (!GEndFrameStatDelegateHandle.IsValid())
+	// This is called each time the Trace system updates.
+	// If the trace system uses a worker thread, this is called from the worker thread.
+
+#if UE_TRACEAUX_EMIT_STAT_COUNTERS || UE_TRACEAUX_EMIT_TRACE_COUNTERS || UE_TRACEAUX_EMIT_LLM_STATS
+	UE::Trace::FStatistics Stats;
+	UE::Trace::GetStatistics(Stats);
+	TraceAuxiliaryEmitStats(Stats);
+#endif
+}
+
+////////////////////////////////////////////////////////////////////////////////
+void TraceAuxiliaryOnEndFrameCallback()
+{
+	// This is called at end of each frame.
+
+	// Reset the trace update callback. From now on, we only trace stats once per frame.
+	UE::Trace::SetUpdateCallback(nullptr);
+
+#if UE_TRACEAUX_EMIT_STAT_COUNTERS || UE_TRACEAUX_EMIT_TRACE_COUNTERS || UE_TRACEAUX_EMIT_LLM_STATS
+	UE::Trace::FStatistics Stats;
+	UE::Trace::GetStatistics(Stats);
+	TraceAuxiliaryEmitStats(Stats);
+#endif
+
+#if UE_TRACEAUX_EMIT_CSV_STATS && TRACE_PRIVATE_STATISTICS
+	// Only publish CSV stats if we have ever run tracing in order to reduce overhead in most runs.
+	static bool bDoCsvStats = false;
+	if (UE::Trace::IsTracing() || bDoCsvStats)
 	{
-		// Update stats every frame
-		GEndFrameStatDelegateHandle = FCoreDelegates::OnEndFrame.AddLambda([]()
-		{
-			UE::Trace::FStatistics Stats;
-			UE::Trace::GetStatistics(Stats);
-			SET_MEMORY_STAT(STAT_TraceMemoryUsed, Stats.MemoryUsed);
-			SET_MEMORY_STAT(STAT_TraceBlockPoolAllocated, Stats.BlockPoolAllocated);
-			SET_MEMORY_STAT(STAT_TraceSharedAllocated, Stats.SharedBufferAllocated);
-			SET_MEMORY_STAT(STAT_TraceFixedAllocated, Stats.FixedBufferAllocated);
-			SET_MEMORY_STAT(STAT_TraceCacheAllocated, Stats.CacheAllocated);
-			SET_MEMORY_STAT(STAT_TraceCacheUsed, Stats.CacheUsed);
-			SET_MEMORY_STAT(STAT_TraceCacheWaste, Stats.CacheWaste);
-			SET_MEMORY_STAT(STAT_TraceSent, Stats.BytesSent);
-			SET_MEMORY_STAT(STAT_TraceTraced, Stats.BytesTraced);
-			SET_MEMORY_STAT(STAT_TraceEmitted, Stats.BytesEmitted);
-
-#if ENABLE_LOW_LEVEL_MEM_TRACKER
-			const int64 NonCacheMemory = Stats.MemoryUsed - Stats.CacheAllocated;
-			FLowLevelMemTracker::Get().SetTagAmountForTracker(ELLMTracker::Default, LLM_TAG_NAME(Trace_TraceLog), ELLMTagSet::None, NonCacheMemory, true);
-			FLowLevelMemTracker::Get().SetTagAmountForTracker(ELLMTracker::Default, LLM_TAG_NAME(Trace_TraceLog_Cache), ELLMTagSet::None, Stats.CacheAllocated, true);
+		bDoCsvStats = true;
+
+#if !(UE_TRACEAUX_EMIT_STAT_COUNTERS || UE_TRACEAUX_EMIT_TRACE_COUNTERS || UE_TRACEAUX_EMIT_LLM_STATS)
+		UE::Trace::FStatistics Stats;
+		UE::Trace::GetStatistics(Stats);
 #endif
-		});
+
+		CSV_CUSTOM_STAT(Trace, MemoryUsedMb,	double(Stats.MemoryUsed) / 1024.0 / 1024.0, ECsvCustomStatOp::Set);
+		CSV_CUSTOM_STAT(Trace, BlockPoolMb,		double(Stats.BlockPoolAllocated) / 1024.0 / 1024.0, ECsvCustomStatOp::Set);
+		CSV_CUSTOM_STAT(Trace, FixedBuffersMb,	double(Stats.FixedBufferAllocated) / 1024.0 / 1024.0, ECsvCustomStatOp::Set);
+		CSV_CUSTOM_STAT(Trace, SharedBuffersMb,	double(Stats.SharedBufferAllocated) / 1024.0 / 1024.0, ECsvCustomStatOp::Set);
+		CSV_CUSTOM_STAT(Trace, CacheMb,			double(Stats.CacheAllocated) / 1024.0 / 1024.0, ECsvCustomStatOp::Set);
 	}
+#endif
+}
+
+////////////////////////////////////////////////////////////////////////////////
+void FTraceAuxiliaryImpl::RegisterEndFrameCallbacks()
+{
+	// Always register end frame updates. This path is short circuited if a worker thread exists.
+	// If the worker thread is disabled, pump the update from end frame.
+	FCoreDelegates::OnEndFrame.AddStatic(UE::Trace::Update);
+
+	// Update stats every frame.
+	FCoreDelegates::OnEndFrame.AddStatic(TraceAuxiliaryOnEndFrameCallback);
 }
 
 ////////////////////////////////////////////////////////////////////////////////
-void OnConnectionCallback()
+void TraceAuxiliaryOnConnectionCallback()
 {
 	TRACE_CPUPROFILER_EVENT_SCOPE(FTraceAux_OnConnection);
 	FTraceAuxiliary::OnConnection.Broadcast();
 }
 
 ////////////////////////////////////////////////////////////////////////////////
-void OnMessageCallback(const UE::Trace::FMessageEvent& Message)
+void TraceAuxiliaryOnMessageCallback(const UE::Trace::FMessageEvent& Message)
 {
 	#define LOG_WITH_VERBOSITY(Verbosity, Text) \
 		UE_LOG(LogTrace, Verbosity, TEXT("%s"), Text)
@@ -952,7 +1047,7 @@ void OnMessageCallback(const UE::Trace::FMessageEvent& Message)
 }
 
 ////////////////////////////////////////////////////////////////////////////////
-void OnScopeBeginCallback(const ANSICHAR* ScopeName)
+void TraceAuxiliaryOnScopeBeginCallback(const ANSICHAR* ScopeName)
 {
 #if CPUPROFILERTRACE_ENABLED
 	FCpuProfilerTrace::OutputBeginDynamicEvent(ScopeName);
@@ -960,7 +1055,7 @@ void OnScopeBeginCallback(const ANSICHAR* ScopeName)
 }
 
 ////////////////////////////////////////////////////////////////////////////////
-void OnScopeEndCallback()
+void TraceAuxiliaryOnScopeEndCallback()
 {
 #if CPUPROFILERTRACE_ENABLED
 	FCpuProfilerTrace::OutputEndEvent();
@@ -968,7 +1063,7 @@ void OnScopeEndCallback()
 }
 
 ////////////////////////////////////////////////////////////////////////////////
-static void SetupInitFromConfig(UE::Trace::FInitializeDesc& OutDesc)
+static void TraceAuxiliarySetupInitFromConfig(UE::Trace::FInitializeDesc& OutDesc)
 {
 	if (!GConfig)
 	{
@@ -1120,13 +1215,24 @@ static void TraceAuxiliaryStatus()
 	UE::Trace::FStatistics Stats;
 	UE::Trace::GetStatistics(Stats);
 	constexpr double MiB = 1.0 / (1024.0 * 1024.0);
-	UE_LOG(LogConsoleResponse, Display, TEXT("- Memory used: %.02f MiB"),
+	UE_LOG(LogConsoleResponse, Display, TEXT("- Memory Used: %.02f MiB"),
 		double(Stats.MemoryUsed) * MiB);
-	UE_LOG(LogConsoleResponse, Display, TEXT("- Important Events cache: %.02f MiB (%.02f MiB used + %0.02f MiB unused | %0.02f MiB waste)"),
+	UE_LOG(LogConsoleResponse, Display, TEXT("- Block Pool: %.02f MiB"),
+		double(Stats.BlockPoolAllocated) * MiB);
+	UE_LOG(LogConsoleResponse, Display, TEXT("- Fixed Buffers: %.02f MiB"),
+		double(Stats.FixedBufferAllocated) * MiB);
+	UE_LOG(LogConsoleResponse, Display, TEXT("- Shared Buffers: %.02f MiB"),
+		double(Stats.SharedBufferAllocated) * MiB);
+	UE_LOG(LogConsoleResponse, Display, TEXT("- Important Events Cache: %.02f MiB (%.02f MiB used + %0.02f MiB unused | %0.02f MiB waste)"),
 		double(Stats.CacheAllocated) * MiB,
 		double(Stats.CacheUsed) * MiB,
 		double(Stats.CacheAllocated - Stats.CacheUsed) * MiB,
 		double(Stats.CacheWaste) * MiB);
+	UE_LOG(LogConsoleResponse, Display, TEXT("- Emitted: %.02f MiB"),
+		double(Stats.BytesEmitted) * MiB);
+	UE_LOG(LogConsoleResponse, Display, TEXT("- Traced: %.02f MiB (-%llu bytes)"),
+		double(Stats.BytesTraced) * MiB,
+		Stats.BytesEmitted - Stats.BytesTraced);
 	UE_LOG(LogConsoleResponse, Display, TEXT("- Sent: %.02f MiB"),
 		double(Stats.BytesSent) * MiB);
 
@@ -1547,7 +1653,7 @@ bool FTraceAuxiliary::Start(EConnectionType Type, const TCHAR* Target, const TCH
 #if UE_TRACE_ENABLED
 	// Use StartCustom function to start a custom trace
 	check(Type != EConnectionType::Relay);
-	
+
 	TRACE_CPUPROFILER_EVENT_SCOPE(FTraceAux_Start);
 
 	if (!TraceAuxiliaryStartShared(Channels, Options, LogCategory))
@@ -1577,12 +1683,12 @@ bool FTraceAuxiliary::Relay(UPTRINT Handle, UE::Trace::IoWriteFunc WriteFunc, UE
 {
 #if UE_TRACE_ENABLED
 	TRACE_CPUPROFILER_EVENT_SCOPE(FTraceAux_Start);
-	
+
 	if (!TraceAuxiliaryStartShared(Channels, Options, LogTrace))
 	{
 		return false;
 	}
-	
+
 	uint16 SendFlags = (Options && Options->bExcludeTail) ? UE::Trace::FSendFlags::ExcludeTail : 0;
 	return GTraceAuxiliary.ConnectRelay(Handle, WriteFunc, CloseFunc, SendFlags);
 #else
@@ -1688,7 +1794,7 @@ void FTraceAuxiliary::Initialize(const TCHAR* CommandLine)
 	{
 		return;
 	}
-	
+
 	TRACE_CPUPROFILER_EVENT_SCOPE(FTraceAux_Init);
 	UE_MEMSCOPE(TRACE_TAG);
 
@@ -1704,7 +1810,7 @@ void FTraceAuxiliary::Initialize(const TCHAR* CommandLine)
 #if UE_TRACE_ENABLED
 
 	// Setup message callback so we get feedback from TraceLog
-	UE::Trace::SetMessageCallback(&OnMessageCallback);
+	UE::Trace::SetMessageCallback(&TraceAuxiliaryOnMessageCallback);
 
 	FParse::Bool(CommandLine, TEXT("-traceautostart="), GTraceAutoStart);
 	UE_LOG(LogTrace, Verbose, TEXT("Trace auto start = %d."), GTraceAutoStart);
@@ -1724,7 +1830,7 @@ void FTraceAuxiliary::Initialize(const TCHAR* CommandLine)
 	bInitialized = true;
 
 	UE_LOG(LogTrace, Log, TEXT("Initializing trace..."));
-	
+
 	const TCHAR* AppName = TEXT(UE_APP_NAME);
 	const TCHAR* ProjectName = TEXT("");
 #if IS_MONOLITHIC && !IS_PROGRAM
@@ -1780,21 +1886,28 @@ void FTraceAuxiliary::Initialize(const TCHAR* CommandLine)
 	bool bShouldStartWorkerThread = false;
 	StartFromCommandlineArguments(CommandLine, bShouldStartWorkerThread);
 
+	// Emit empty stats once (to ensure all stats/counters start from zero).
+	{
+		UE::Trace::FStatistics Stats;
+		TraceAuxiliaryEmitStats(Stats);
+	}
+
 	// Initialize Trace. The settings are stored in a static for posterity.
 	UE::Trace::FInitializeDesc& Desc = GInitializeDesc;
 	if (FParse::Param(CommandLine, TEXT("tracescopes")))
 	{
-		Desc.OnScopeBeginFunc = &OnScopeBeginCallback;
-		Desc.OnScopeEndFunc = &OnScopeEndCallback;
+		Desc.OnScopeBeginFunc = &TraceAuxiliaryOnScopeBeginCallback;
+		Desc.OnScopeEndFunc = &TraceAuxiliaryOnScopeEndCallback;
 	}
 #if WITH_EDITOR
 	Desc.TailSizeBytes = 32 << 20;
 #endif
-	SetupInitFromConfig(Desc);
+	TraceAuxiliarySetupInitFromConfig(Desc);
 
 	Desc.bUseWorkerThread = bShouldStartWorkerThread;
 	Desc.bUseImportantCache = (FParse::Param(CommandLine, TEXT("tracenocache")) == false);
-	Desc.OnConnectionFunc = &OnConnectionCallback;
+	Desc.OnConnectionFunc = &TraceAuxiliaryOnConnectionCallback;
+	Desc.OnUpdateFunc = &TraceAuxiliaryOnTraceUpdateCallback;
 
 	FGuid SessionGuid;
 	if (!FParse::Value(CommandLine, TEXT("-tracesessionguid="), SessionGuid))
@@ -1836,8 +1949,8 @@ void FTraceAuxiliary::Initialize(const TCHAR* CommandLine)
 	// Setup known on connection callbacks
 	OnConnection.AddStatic(FStringTrace::OnConnection);
 
-	// Always register end frame updates. This path is short circuited if a worker thread exists.
-	GTraceAuxiliary.StartEndFramePump();
+	// Register end frame callbacks
+	GTraceAuxiliary.RegisterEndFrameCallbacks();
 
 	// Initialize callstack tracing with the regular malloc (it might have already been initialized by memory tracing).
 	CallstackTrace_Create(GMalloc);
@@ -1849,10 +1962,6 @@ void FTraceAuxiliary::Initialize(const TCHAR* CommandLine)
 	FPlatformEventsTrace::Init(Microseconds);
 	FPlatformEventsTrace::PostInit();
 
-#if CSV_PROFILER_STATS
-	FCoreDelegates::OnEndFrame.AddRaw(&GTraceAuxiliary, &FTraceAuxiliaryImpl::UpdateCsvStats);
-#endif
-
 	if (GTraceAutoStart)
 	{
 		FModuleManager::Get().OnModulesChanged().AddLambda([](FName Name, EModuleChangeReason Reason)
@@ -2290,3 +2399,5 @@ bool FTraceServerControls::Stop()
 }
 
 #endif // UE_TRACE_SERVER_CONTROLS_ENABLED
+
+////////////////////////////////////////////////////////////////////////////////
diff --git a/Engine/Source/Runtime/TraceLog/Private/Trace/Trace.cpp b/Engine/Source/Runtime/TraceLog/Private/Trace/Trace.cpp
--- a/Engine/Source/Runtime/TraceLog/Private/Trace/Trace.cpp
+++ b/Engine/Source/Runtime/TraceLog/Private/Trace/Trace.cpp
@@ -16,6 +16,7 @@ namespace Private
 
 ////////////////////////////////////////////////////////////////////////////////
 void	Message_SetCallback(OnMessageFunc Callback);
+void	Writer_SetUpdateCallback(OnUpdateFunc Callback);
 void	Writer_MemorySetHooks(AllocFunc, FreeFunc);
 void	Writer_Initialize(const FInitializeDesc&);
 void	Writer_WorkerCreate();
@@ -25,7 +26,7 @@ bool	Writer_SendTo(const ANSICHAR*, uint32, uint32);
 bool	Writer_WriteTo(const ANSICHAR*, uint32);
 bool	Writer_RelayTo(UPTRINT, IoWriteFunc, IoCloseFunc, uint16);
 bool	Writer_WriteSnapshotTo(const ANSICHAR*);
-bool	Writer_SendSnapshotTo(const ANSICHAR*, uint32);	
+bool	Writer_SendSnapshotTo(const ANSICHAR*, uint32);
 bool	Writer_IsTracing();
 bool	Writer_IsTracingTo(uint32 (&OutSessionGuid)[4], uint32 (&OutTraceGuid)[4]);
 bool	Writer_Stop();
@@ -66,6 +67,12 @@ void SetMessageCallback(OnMessageFunc MessageFunc)
 	Private::Message_SetCallback(MessageFunc);
 }
 
+////////////////////////////////////////////////////////////////////////////////
+void SetUpdateCallback(OnUpdateFunc UpdateFunc)
+{
+	Private::Writer_SetUpdateCallback(UpdateFunc);
+}
+
 ////////////////////////////////////////////////////////////////////////////////
 void Initialize(const FInitializeDesc& Desc)
 {
@@ -146,7 +153,7 @@ bool IsTracingTo(uint32 (&OutSessionGuid)[4], uint32 (&OutTraceGuid)[4])
 {
 	return Private::Writer_IsTracingTo(OutSessionGuid, OutTraceGuid);
 }
-	
+
 ////////////////////////////////////////////////////////////////////////////////
 bool Stop()
 {
@@ -201,7 +208,7 @@ void StartWorkerThread()
 {
 	Private::Writer_WorkerCreate();
 }
-	
+
 ////////////////////////////////////////////////////////////////////////////////
 FChannel* FindChannel(const TCHAR* ChannelName)
 {
@@ -209,7 +216,7 @@ FChannel* FindChannel(const TCHAR* ChannelName)
 	ToAnsiCheap(ChannelNameA, ChannelName);
 	return FChannel::FindChannel(ChannelNameA);
 }
-	
+
 ////////////////////////////////////////////////////////////////////////////////
 FChannel* FindChannel(FChannelId ChannelId)
 {
@@ -234,7 +241,7 @@ UE_TRACE_MINIMAL_EVENT_END()
 UE_TRACE_MINIMAL_EVENT_BEGIN($Trace, ThreadGroupEnd, NoSync|Important)
 UE_TRACE_MINIMAL_EVENT_END()
 
-	
+
 ////////////////////////////////////////////////////////////////////////////////
 void ThreadRegister(const TCHAR* Name, uint32 SystemId, int32 SortHint)
 {
diff --git a/Engine/Source/Runtime/TraceLog/Private/Trace/Writer.cpp b/Engine/Source/Runtime/TraceLog/Private/Trace/Writer.cpp
--- a/Engine/Source/Runtime/TraceLog/Private/Trace/Writer.cpp
+++ b/Engine/Source/Runtime/TraceLog/Private/Trace/Writer.cpp
@@ -37,6 +37,7 @@ namespace Private {
 
 ////////////////////////////////////////////////////////////////////////////////
 int32			Encode(const void*, int32, void*, int32);
+void			Writer_SetUpdateCallback(OnUpdateFunc Callback);
 void			Writer_SendData(uint32, uint8* __restrict, uint32);
 void			Writer_InitializeTail(int32);
 void			Writer_ShutdownTail();
@@ -65,7 +66,7 @@ void			Writer_FreeBlockListToPool(FWriteBuffer*, FWriteBuffer*);
 ////////////////////////////////////////////////////////////////////////////////
 struct FTraceGuid
 {
-	uint32 Bits[4];	
+	uint32 Bits[4];
 };
 
 ////////////////////////////////////////////////////////////////////////////////
@@ -186,7 +187,7 @@ void Writer_CreateGuid(FTraceGuid* OutGuid)
 	// L'Ecuyer, Pierre (1999). "Tables of Linear Congruential Generators of Different Sizes and Good Lattice Structure"
 	// corrected with errata
 	// Assuming m = 2e64
-	constexpr uint64 C = 0x369DEA0F31A53F85;	
+	constexpr uint64 C = 0x369DEA0F31A53F85;
 	constexpr uint64 I = 1ull;
 
 	const uint64 TopBits = State * C + I;
@@ -205,7 +206,7 @@ void Writer_CreateGuid(FTraceGuid* OutGuid)
 	Octets[6] = Version | (~VersionMask & Octets[6]); // Octet 9
 	Octets[8] = Variant | (~VariantMask & Octets[8]); // Octet 7
 }
-	
+
 ////////////////////////////////////////////////////////////////////////////////
 
 OnScopeBeginFunc* FProfilerScope::OnScopeBegin = nullptr;
@@ -404,7 +405,7 @@ void Writer_SendDataNoInstr(uint32 ThreadId, uint8* __restrict Data, uint32 Size
 	{
 		return;
 	}
-	
+
 
 	// Smaller buffers usually aren't redundant enough to benefit from being
 	// compressed. They often end up being larger.
@@ -455,7 +456,7 @@ void Writer_SendData(uint32 ThreadId, uint8* __restrict Data, uint32 Size)
 	FProfilerScope _(__func__);
 	Writer_SendDataNoInstr(ThreadId, Data, Size);
 }
-	
+
 ////////////////////////////////////////////////////////////////////////////////
 static void Writer_DescribeEvents(FEventNode::FIter Iter)
 {
@@ -507,11 +508,13 @@ static void Writer_DescribeAnnounce()
 
 
 ////////////////////////////////////////////////////////////////////////////////
-static int8			GSyncPacketCountdown;	// = 0
-static const int8	GNumSyncPackets			= 3;
+static int8				GSyncPacketCountdown; // = 0
+static const int8		GNumSyncPackets = 3;
 static OnConnectFunc*	GOnConnection = nullptr;
-static FTraceGuid GSessionGuid; // = {0, 0, 0, 0};
-static FTraceGuid GTraceGuid; // = {0, 0, 0, 0};
+static OnUpdateFunc*	GOnUpdate = nullptr;
+static FTraceGuid		GSessionGuid; // = {0, 0, 0, 0};
+static FTraceGuid		GTraceGuid; // = {0, 0, 0, 0};
+
 
 ////////////////////////////////////////////////////////////////////////////////
 static void Writer_SendSync()
@@ -609,7 +612,7 @@ static bool Writer_UpdateConnection()
 	FEventNode::OnConnect();
 	Writer_DescribeEvents(FEventNode::ReadNew());
 
-	// Send cached events (i.e. importants) 
+	// Send cached events (i.e. importants)
 	Writer_CacheOnConnect();
 
 	// Issue on connection callback. This allows writing events that are
@@ -718,10 +721,17 @@ static volatile bool	GWorkerThreadQuit;	// = false;
 static uint32			GSleepTimeInMS = 17;
 static volatile uint32	GUpdateInProgress = 1;	// Don't allow updates until initialized
 
+////////////////////////////////////////////////////////////////////////////////
+void Writer_SetUpdateCallback(OnUpdateFunc Callback)
+{
+	GOnUpdate = Callback;
+}
+
 ////////////////////////////////////////////////////////////////////////////////
 static void Writer_WorkerUpdateInternal()
 {
 	FProfilerScope _(__func__);
+
 	Writer_UpdateControl();
 	Writer_UpdateConnection();
 	Writer_DescribeAnnounce();
@@ -736,6 +746,11 @@ static void Writer_WorkerUpdateInternal()
 		Writer_FlushSendBuffer();
 	}
 #endif
+
+	if (OnUpdateFunc* OnUpdate = GOnUpdate)
+	{
+		OnUpdate();
+	}
 }
 
 ////////////////////////////////////////////////////////////////////////////////
@@ -745,7 +760,7 @@ static void Writer_WorkerUpdate()
 	{
 		return;
 	}
-	
+
 	Writer_WorkerUpdateInternal();
 
 	AtomicExchangeRelease(&GUpdateInProgress, 0u);
@@ -892,7 +907,7 @@ void Writer_Initialize(const FInitializeDesc& Desc)
 	FProfilerScope::OnScopeEnd = Desc.OnScopeEndFunc;
 
 	FProfilerScope _(__func__);
-	
+
 	Writer_InitializeTail(Desc.TailSizeBytes);
 
 	if (Desc.bUseImportantCache)
@@ -922,8 +937,9 @@ void Writer_Initialize(const FInitializeDesc& Desc)
 
 	// Store callback on connection
 	GOnConnection = Desc.OnConnectionFunc;
+	GOnUpdate = Desc.OnUpdateFunc;
 
-	// Allow the worker thread to start updating 
+	// Allow the worker thread to start updating
 	AtomicStoreRelease(&GUpdateInProgress, uint32(0));
 }
 
@@ -988,7 +1004,7 @@ bool Writer_SendTo(const ANSICHAR* Host, uint32 Flags, uint32 Port)
 	}
 
 	GPendingWriterState = { &IoWrite, &IoClose };
-	
+
 	AtomicStoreRelease(&GPendingDataHandle, DataHandle);
 	return true;
 #else
@@ -1021,7 +1037,7 @@ bool Writer_WriteTo(const ANSICHAR* Path, uint32 Flags)
 	}
 
 	GPendingWriterState = { &IoWrite, &IoClose };
-	
+
 	AtomicStoreRelease(&GPendingDataHandle, DataHandle);
 	return true;
 #else
@@ -1036,9 +1052,9 @@ bool Writer_RelayTo(UPTRINT InHandle, IoWriteFunc WriteFunc, IoCloseFunc CloseFu
 	{
 		return false;
 	}
-	
+
 	Writer_InternalInitialize();
-	
+
 	UPTRINT DataHandle = InHandle;
 	DataHandle = Writer_PackSendFlags(DataHandle, Flags);
 	if (!DataHandle)
@@ -1047,7 +1063,7 @@ bool Writer_RelayTo(UPTRINT InHandle, IoWriteFunc WriteFunc, IoCloseFunc CloseFu
 	}
 
 	GPendingWriterState = { WriteFunc, CloseFunc };
-	
+
 	AtomicStoreRelease(&GPendingDataHandle, DataHandle);
 	return true;
 }
@@ -1155,7 +1171,7 @@ bool Writer_WriteSnapshot(const FSnapshotTarget& Target)
 	{
 		return false;
 	}
-	
+
 	// Bring everything up to date with the active tracing connection.
 	//  Any connection writes after we call the worker update
 	//  will need to treat source data structures as read-only.
@@ -1177,7 +1193,7 @@ bool Writer_WriteSnapshot(const FSnapshotTarget& Target)
 
 	{
 		const bool bExistingDataHandle = GDataHandle != 0;
-		
+
 		TStashGlobal DataHandle(GDataHandle);
 		TStashGlobal StateWrite(GWriterState.Write);
 		TStashGlobal StateClose(GWriterState.Close);
@@ -1188,14 +1204,14 @@ bool Writer_WriteSnapshot(const FSnapshotTarget& Target)
 		if (Target.Type == FSnapshotTarget::EType::FileTarget)
 		{
 #if TRACE_PRIVATE_ALLOW_FILE
-			// Open the snapshot file 
+			// Open the snapshot file
 			GDataHandle = FileOpen(Target.File.Path);
 #endif
 		}
 		else
 		{
 #if TRACE_PRIVATE_ALLOW_TCP
-			// Open the snapshot connection and write 
+			// Open the snapshot connection and write
 			const uint32 Port = Target.Host.Port ? Target.Host.Port : 1981;
 			GDataHandle = TcpSocketConnect(Target.Host.Host, uint16(Port));
 			if (!GDataHandle)
@@ -1243,7 +1259,7 @@ bool Writer_WriteSnapshot(const FSnapshotTarget& Target)
 
 	return true;
 }
-	
+
 ////////////////////////////////////////////////////////////////////////////////
 bool Writer_WriteSnapshotTo(const ANSICHAR* Path)
 {
@@ -1262,13 +1278,13 @@ bool Writer_SendSnapshotTo(const ANSICHAR* Host, uint32 Port)
 	Target.Host.Port = Port;
 	return Writer_WriteSnapshot(Target);
 }
-	
+
 ////////////////////////////////////////////////////////////////////////////////
 bool Writer_IsTracing()
 {
 	return AtomicLoadRelaxed(&GDataHandle) != 0 || AtomicLoadRelaxed(&GPendingDataHandle) != 0;
 }
-	
+
 ////////////////////////////////////////////////////////////////////////////////
 bool Writer_IsTracingTo(uint32 (&OutSessionGuid)[4], uint32 (&OutTraceGuid)[4])
 {
diff --git a/Engine/Source/Runtime/TraceLog/Public/Trace/Trace.h b/Engine/Source/Runtime/TraceLog/Public/Trace/Trace.h
--- a/Engine/Source/Runtime/TraceLog/Public/Trace/Trace.h
+++ b/Engine/Source/Runtime/TraceLog/Public/Trace/Trace.h
@@ -53,8 +53,8 @@ enum AnsiString {};
 enum WideString {};
 
 // Reference to a definition event.
-template<typename IdType>	
-struct TEventRef 
+template<typename IdType>
+struct TEventRef
 {
 	using ReferenceType = IdType;
 
@@ -98,7 +98,7 @@ inline uint64 TEventRef<uint64>::GetHash() const
 }
 
 typedef TEventRef<uint8> FEventRef8;
-typedef TEventRef<uint16> FEventRef16;	
+typedef TEventRef<uint16> FEventRef16;
 typedef TEventRef<uint32> FEventRef32;
 typedef TEventRef<uint64> FEventRef64;
 
@@ -108,8 +108,6 @@ TEventRef<IdType> MakeEventRef(IdType InId, uint32 InTypeId)
 	return TEventRef<IdType>(InId, InTypeId);
 }
 
-using OnConnectFunc = void(void);
-
 enum class EMessageType : uint8
 {
 	Reserved			= 0,
@@ -135,7 +133,7 @@ enum class EMessageType : uint8
 	FatalStart			= 0x40,
 	OOMFatal,
 };
-	
+
 struct FMessageEvent
 {
 	/** Type of message */
@@ -147,6 +145,8 @@ struct FMessageEvent
 };
 
 using OnMessageFunc = void(const FMessageEvent&);
+using OnConnectFunc = void(void);
+using OnUpdateFunc = void(void);
 using OnScopeBeginFunc = void(const ANSICHAR*);
 using OnScopeEndFunc = void(void);
 
@@ -158,6 +158,7 @@ struct FInitializeDesc
 	bool				bUseImportantCache	= true;
 	uint32				SessionGuid[4]		= {0,0,0,0}; // leave as zero to generate random
 	OnConnectFunc*		OnConnectionFunc	= nullptr;
+	OnUpdateFunc*		OnUpdateFunc		= nullptr;
 	OnScopeBeginFunc*	OnScopeBeginFunc	= nullptr;
 	OnScopeEndFunc*		OnScopeEndFunc		= nullptr;
 };
@@ -174,11 +175,11 @@ struct FChannelInfo
 };
 
 class FChannel;
-	
+
 typedef void*		AllocFunc(SIZE_T, uint32);
 typedef void		FreeFunc(void*, SIZE_T);
 typedef void		ChannelIterFunc(const ANSICHAR*, bool, void*);
-/*  The callback provides information about a channel and a user provided pointer. 
+/*  The callback provides information about a channel and a user provided pointer.
 	Returning false from the callback will stop the enumeration */
 typedef bool		ChannelIterCallback(const FChannelInfo& Info, void*/*User*/);
 typedef bool		IoWriteFunc(UPTRINT Handle, const void* Data, uint32 Size);
@@ -207,8 +208,9 @@ struct FSendFlags
 
 UE_TRACE_API void	SetMemoryHooks(AllocFunc Alloc, FreeFunc Free) UE_TRACE_IMPL();
 UE_TRACE_API void	SetMessageCallback(OnMessageFunc MessageFunc) UE_TRACE_IMPL();
+UE_TRACE_API void	SetUpdateCallback(OnUpdateFunc UpdateFunc) UE_TRACE_IMPL();
 UE_TRACE_API void	Initialize(const FInitializeDesc& Desc) UE_TRACE_IMPL();
-UE_TRACE_API void	StartWorkerThread() UE_TRACE_IMPL();	
+UE_TRACE_API void	StartWorkerThread() UE_TRACE_IMPL();
 UE_TRACE_API void	Shutdown() UE_TRACE_IMPL();
 UE_TRACE_API void	Panic() UE_TRACE_IMPL();
 UE_TRACE_API void	Update() UE_TRACE_IMPL();
@@ -217,7 +219,7 @@ UE_TRACE_API bool	SendTo(const TCHAR* Host, uint32 Port=0, uint16 Flags=FSendFla
 UE_TRACE_API bool	WriteTo(const TCHAR* Path, uint16 Flags=FSendFlags::None) UE_TRACE_IMPL(false);
 UE_TRACE_API bool	RelayTo(UPTRINT InHandle, IoWriteFunc WriteFunc, IoCloseFunc CloseFunc, uint16 Flags = FSendFlags::None) UE_TRACE_IMPL(false);
 UE_TRACE_API bool	WriteSnapshotTo(const TCHAR* Path) UE_TRACE_IMPL(false);
-UE_TRACE_API bool	SendSnapshotTo(const TCHAR* Host, uint32 Port) UE_TRACE_IMPL(false);	
+UE_TRACE_API bool	SendSnapshotTo(const TCHAR* Host, uint32 Port) UE_TRACE_IMPL(false);
 UE_TRACE_API bool	IsTracing() UE_TRACE_IMPL(false);
 UE_TRACE_API bool	IsTracingTo(uint32 (&OutSessionGuid)[4], uint32 (&OutTraceGuid)[4]) UE_TRACE_IMPL(false);
 UE_TRACE_API bool	Stop() UE_TRACE_IMPL(false);
@@ -228,7 +230,7 @@ UE_TRACE_API void	EnumerateChannels(ChannelIterCallback IterFunc, void* User) UE
 UE_TRACE_API void	ThreadRegister(const TCHAR* Name, uint32 SystemId, int32 SortHint) UE_TRACE_IMPL();
 UE_TRACE_API void	ThreadGroupBegin(const TCHAR* Name) UE_TRACE_IMPL();
 UE_TRACE_API void	ThreadGroupEnd() UE_TRACE_IMPL();
-	
+
 UE_TRACE_API FChannel* FindChannel(const TCHAR* ChannelName) UE_TRACE_IMPL(nullptr);
 UE_TRACE_API FChannel* FindChannel(FChannelId ChannelId) UE_TRACE_IMPL(nullptr);
 
