Unify Mac and iOS platform defines under ApplePlatform.h

diff --git a/Engine/Source/Runtime/Core/Public/Apple/ApplePlatform.h b/Engine/Source/Runtime/Core/Public/Apple/ApplePlatform.h
new file mode 100644
--- /dev/null
+++ b/Engine/Source/Runtime/Core/Public/Apple/ApplePlatform.h
@@ -0,0 +1,82 @@
+// Copyright Epic Games, Inc. All Rights Reserved.
+
+
+/*=============================================================================================
+	ApplePlatform.h: Common setup for Apple platforms
+==============================================================================================*/
+
+#pragma once
+
+#include "Clang/ClangPlatform.h"
+
+// Base defines, must define these for the platform, there are no defaults
+#define PLATFORM_64BITS					1
+// Technically the underlying platform has 128bit atomics, but clang might not issue optimal code
+#define PLATFORM_HAS_128BIT_ATOMICS		0
+
+// Base defines, defaults are commented out
+#define PLATFORM_LITTLE_ENDIAN							1
+#define PLATFORM_SEH_EXCEPTIONS_DISABLED				1
+#define PLATFORM_SUPPORTS_PRAGMA_PACK					1
+#define PLATFORM_ENABLE_VECTORINTRINSICS				1
+#define PLATFORM_USE_SYSTEM_VSWPRINTF					0
+#define PLATFORM_COMPILER_DISTINGUISHES_INT_AND_LONG	1
+#define PLATFORM_WCHAR_IS_4_BYTES						1
+#define PLATFORM_HAS_BSD_TIME							1
+#define PLATFORM_HAS_BSD_IPV6_SOCKETS					1
+#define PLATFORM_HAS_BSD_SOCKET_FEATURE_MSG_DONTWAIT	1
+#define PLATFORM_HAS_MULTITHREADED_PREMAIN				1
+#define PLATFORM_SUPPORTS_TEXTURE_STREAMING				1
+#define PLATFORM_SUPPORTS_STACK_SYMBOLS					1
+#define PLATFORM_IS_ANSI_MALLOC_THREADSAFE				1
+
+#define PLATFORM_BREAK()                                __builtin_debugtrap()
+
+#define PLATFORM_CODE_SECTION(Name)						__attribute__((section("__TEXT,__" Name ",regular,pure_instructions"))) \
+														__attribute__((aligned(4)))
+
+// Ensure we can use this builtin - seems to be present on Clang 9, GCC 11 and MSVC 19.26,
+// but gives spurious "non-void function 'BitCast' should return a value" errors on some
+// Mac and Android toolchains when building PCHs, so avoid those.
+#undef PLATFORM_COMPILER_SUPPORTS_BUILTIN_BITCAST
+#define PLATFORM_COMPILER_SUPPORTS_BUILTIN_BITCAST (__clang_major__ >= 13)
+
+// Function type macros.
+#define VARARGS																/* Functions with variable arguments */
+#define CDECL																/* Standard C function */
+#define STDCALL																/* Standard calling convention */
+#if UE_BUILD_DEBUG || UE_DISABLE_FORCE_INLINE
+#	define FORCEINLINE				inline 									/* Don't force code to be inline */
+#else
+#	define FORCEINLINE				inline __attribute__ ((always_inline))	/* Force code to be inline */
+#endif
+#define FORCENOINLINE				__attribute__((noinline))				/* Force code to NOT be inline */
+#define FUNCTION_CHECK_RETURN_END	__attribute__ ((warn_unused_result))	/* Warn that callers should not ignore the return value. */
+#define FUNCTION_NO_RETURN_END		__attribute__ ((noreturn))				/* Indicate that the function never returns. */
+
+#define ABSTRACT abstract
+
+// We can use pragma optimisations on and off as of Apple LLVM 7.3.0 but not before.
+#if (__clang_major__ > 7) || (__clang_major__ == 7 && __clang_minor__ >= 3)
+#	define PRAGMA_DISABLE_OPTIMIZATION_ACTUAL _Pragma("clang optimize off")
+#	define PRAGMA_ENABLE_OPTIMIZATION_ACTUAL  _Pragma("clang optimize on")
+#endif
+
+// Alignment.
+#define GCC_PACK(n)		__attribute__((packed,aligned(n)))
+#define GCC_ALIGN(n)	__attribute__((aligned(n)))
+
+// operator new/delete operators
+// As of 10.9 we need to use _NOEXCEPT & cxx_noexcept compatible definitions
+#if __has_feature(cxx_noexcept)
+#	define OPERATOR_NEW_THROW_SPEC
+#else
+#	define OPERATOR_NEW_THROW_SPEC		throw (std::bad_alloc)
+#endif
+#define OPERATOR_DELETE_THROW_SPEC		noexcept
+#define OPERATOR_NEW_NOTHROW_SPEC		noexcept
+#define OPERATOR_DELETE_NOTHROW_SPEC	noexcept
+
+// DLL export and import definitions
+#define DLLEXPORT			__attribute__((visibility("default")))
+#define DLLIMPORT			__attribute__((visibility("default")))
\ No newline at end of file
diff --git a/Engine/Source/Runtime/Core/Public/IOS/IOSPlatform.h b/Engine/Source/Runtime/Core/Public/IOS/IOSPlatform.h
--- a/Engine/Source/Runtime/Core/Public/IOS/IOSPlatform.h
+++ b/Engine/Source/Runtime/Core/Public/IOS/IOSPlatform.h
@@ -6,7 +6,7 @@
 
 #pragma once
 
-#include "Clang/ClangPlatform.h"
+#include "Apple/ApplePlatform.h"
 #include "Availability.h"
 
 /**
@@ -25,38 +25,15 @@ typedef FIOSPlatformTypes FPlatformTypes;
 // Base defines, must define these for the platform, there are no defaults
 #define PLATFORM_DESKTOP				0
 
-#ifdef __LP64__
-#define PLATFORM_64BITS					1
-// Technically the underlying platform has 128bit atomics, but clang might not issue optimal code
-#define PLATFORM_HAS_128BIT_ATOMICS		0
-#else
-#define PLATFORM_64BITS					0
-#define PLATFORM_HAS_128BIT_ATOMICS		0
-#endif
-
 // Base defines, defaults are commented out
-#define PLATFORM_LITTLE_ENDIAN							1
-#define PLATFORM_SEH_EXCEPTIONS_DISABLED				1
-#define PLATFORM_SUPPORTS_PRAGMA_PACK					1
-#define PLATFORM_COMPILER_DISTINGUISHES_INT_AND_LONG	1
-#define PLATFORM_WCHAR_IS_4_BYTES						1
 #define PLATFORM_TCHAR_IS_CHAR16						1
-#define PLATFORM_USE_SYSTEM_VSWPRINTF					0
-#define PLATFORM_HAS_BSD_TIME							1
-#define PLATFORM_HAS_BSD_IPV6_SOCKETS					1
-#define PLATFORM_HAS_BSD_SOCKET_FEATURE_MSG_DONTWAIT	1
-#define PLATFORM_HAS_MULTITHREADED_PREMAIN				1
 #define PLATFORM_MAX_FILEPATH_LENGTH_DEPRECATED			IOS_MAX_PATH
-#define PLATFORM_SUPPORTS_TEXTURE_STREAMING				1
 #define PLATFORM_BUILTIN_VERTEX_HALF_FLOAT				0
 #define PLATFORM_SUPPORTS_MULTIPLE_NATIVE_WINDOWS		0
 #define PLATFORM_ALLOW_NULL_RHI							1
-#define PLATFORM_ENABLE_VECTORINTRINSICS				1
 #define PLATFORM_ENABLE_VECTORINTRINSICS_NEON			1
-#define PLATFORM_SUPPORTS_STACK_SYMBOLS					1
 #define PLATFORM_SUPPORTS_EARLY_MOVIE_PLAYBACK			1 // movies will start before engine is initalized
 #define PLATFORM_USE_FULL_TASK_GRAPH					0 // @todo platplug: not platplug, but should investigate soon anyway
-#define PLATFORM_IS_ANSI_MALLOC_THREADSAFE				1
 
 #define PLATFORM_NUM_AUDIODECOMPRESSION_PRECACHE_BUFFERS		0
 #if PLATFORM_TVOS
@@ -77,69 +54,18 @@ typedef FIOSPlatformTypes FPlatformTypes;
 #define PLATFORM_SUPPORTS_VIRTUAL_TEXTURE_STREAMING		1
 #define PLATFORM_SUPPORTS_BINDLESS_RENDERING			0
 
-// Ensure we can use this builtin - seems to be present on Clang 9, GCC 11 and MSVC 19.26,
-// but gives spurious "non-void function 'BitCast' should return a value" errors on some
-// Mac and Android toolchains when building PCHs, so avoid those.
-#undef PLATFORM_COMPILER_SUPPORTS_BUILTIN_BITCAST
-#define PLATFORM_COMPILER_SUPPORTS_BUILTIN_BITCAST (__clang_major__ >= 13)
-
 #define PLATFORM_RETURN_ADDRESS_FOR_CALLSTACKTRACING    PLATFORM_RETURN_ADDRESS
 
 #define PLATFORM_GLOBAL_LOG_CATEGORY					LogIOS
 
-#define PLATFORM_BREAK()                                __builtin_debugtrap()
-
-#define PLATFORM_CODE_SECTION(Name)						__attribute__((section("__TEXT,__" Name ",regular,pure_instructions"))) \
-														__attribute__((aligned(4)))
-
 //mallocpoison not safe with aligned ansi allocator.  returns the larger unaligned size during Free() which causes writes off the end of the allocation.
 #define UE_USE_MALLOC_FILL_BYTES 0 
 
-// Function type macros.
-#define VARARGS															/* Functions with variable arguments */
-#define CDECL															/* Standard C function */
-#define STDCALL															/* Standard calling convention */
-#if UE_BUILD_DEBUG || UE_DISABLE_FORCE_INLINE
-#define FORCEINLINE inline 												/* Don't force code to be inline */
-#else
-#define FORCEINLINE inline __attribute__ ((always_inline))				/* Force code to be inline */
-#endif
-#define FORCENOINLINE __attribute__((noinline))							/* Force code to NOT be inline */
-#define FUNCTION_CHECK_RETURN_END __attribute__ ((warn_unused_result))	/* Warn that callers should not ignore the return value. */
-#define FUNCTION_NO_RETURN_END __attribute__ ((noreturn))				/* Indicate that the function never returns. */
-
-#define ABSTRACT abstract
-
-// We can use pragma optimisation's on and off as of Apple LLVM 7.3.0 but not before.
-#if (__clang_major__ > 7) || (__clang_major__ == 7 && __clang_minor__ >= 3)
-#define PRAGMA_DISABLE_OPTIMIZATION_ACTUAL _Pragma("clang optimize off")
-#define PRAGMA_ENABLE_OPTIMIZATION_ACTUAL  _Pragma("clang optimize on")
-#endif
-
 // Strings.
 // Strings.
 #define LINE_TERMINATOR TEXT("\n")
 #define LINE_TERMINATOR_ANSI "\n"
 
-// Alignment.
-#define GCC_PACK(n) __attribute__((packed,aligned(n)))
-#define GCC_ALIGN(n) __attribute__((aligned(n)))
-
-// operator new/delete operators
-// As of 10.9 we need to use _NOEXCEPT & cxx_noexcept compatible definitions
-#if __has_feature(cxx_noexcept)
-#define OPERATOR_NEW_THROW_SPEC
-#else
-#define OPERATOR_NEW_THROW_SPEC throw (std::bad_alloc)
-#endif
-#define OPERATOR_DELETE_THROW_SPEC noexcept
-#define OPERATOR_NEW_NOTHROW_SPEC  noexcept
-#define OPERATOR_DELETE_NOTHROW_SPEC  noexcept
-
-// DLL export and import definitions
-#define DLLEXPORT __attribute__((visibility("default")))
-#define DLLIMPORT
-
 #define IOS_MAX_PATH 1024
 
-static_assert(__IPHONE_OS_VERSION_MAX_ALLOWED >= 13000, "Unreal requires Xcode 11 or later to build"); 
+static_assert(__IPHONE_OS_VERSION_MAX_ALLOWED >= 13000, "Unreal requires Xcode 11 or later to build"); 
\ No newline at end of file
diff --git a/Engine/Source/Runtime/Core/Public/Mac/MacPlatform.h b/Engine/Source/Runtime/Core/Public/Mac/MacPlatform.h
--- a/Engine/Source/Runtime/Core/Public/Mac/MacPlatform.h
+++ b/Engine/Source/Runtime/Core/Public/Mac/MacPlatform.h
@@ -6,7 +6,7 @@
 
 #pragma once
 
-#include "Clang/ClangPlatform.h"
+#include "Apple/ApplePlatform.h"
 
 #define PLATFORM_MAC_USE_CHAR16 1
 
@@ -38,45 +38,30 @@ typedef FMacPlatformTypes FPlatformTypes;
 #endif
 
 // Base defines, must define these for the platform, there are no defaults
-#define PLATFORM_DESKTOP				1
-#define PLATFORM_64BITS					1
-// Technically the underlying platform has 128bit atomics, but clang might not issue optimal code
-#define PLATFORM_HAS_128BIT_ATOMICS		0
-#define PLATFORM_CAN_SUPPORT_EDITORONLY_DATA	1
+#define PLATFORM_DESKTOP								1
+#define PLATFORM_CAN_SUPPORT_EDITORONLY_DATA			1
 
 // Base defines, defaults are commented out
-
-#define PLATFORM_LITTLE_ENDIAN						1
-//#define PLATFORM_EXCEPTIONS_DISABLED				!PLATFORM_DESKTOP
-#define PLATFORM_SEH_EXCEPTIONS_DISABLED				1
-#define PLATFORM_SUPPORTS_PRAGMA_PACK				1
-#define PLATFORM_ENABLE_VECTORINTRINSICS			1
-#define PLATFORM_ENABLE_VECTORINTRINSICS_NEON       PLATFORM_MAC_ARM64
+//#define PLATFORM_EXCEPTIONS_DISABLED					!PLATFORM_DESKTOP
+#define PLATFORM_ENABLE_VECTORINTRINSICS_NEON			PLATFORM_MAC_ARM64
 // FMA3 support was added starting from Intel Haswell
 #ifndef PLATFORM_ALWAYS_HAS_FMA3
-	#define PLATFORM_ALWAYS_HAS_FMA3				0
+	#define PLATFORM_ALWAYS_HAS_FMA3					0
 #endif
-//#define PLATFORM_USE_LS_SPEC_FOR_WIDECHAR			1
-#define PLATFORM_USE_SYSTEM_VSWPRINTF				0
-#define PLATFORM_COMPILER_DISTINGUISHES_INT_AND_LONG			1
-#define PLATFORM_WCHAR_IS_4_BYTES					1
+//#define PLATFORM_USE_LS_SPEC_FOR_WIDECHAR				1
+#define PLATFORM_COMPILER_DISTINGUISHES_INT_AND_LONG	1
+#define PLATFORM_WCHAR_IS_4_BYTES						1
 #if PLATFORM_MAC_USE_CHAR16
 	#define PLATFORM_TCHAR_IS_CHAR16					1
 #else
 	#define PLATFORM_TCHAR_IS_4_BYTES					1
 #endif
-#define PLATFORM_HAS_BSD_TIME							1
-#define PLATFORM_HAS_BSD_IPV6_SOCKETS					1
 #define PLATFORM_HAS_BSD_SOCKET_FEATURE_IOCTL			1
 #define PLATFORM_HAS_BSD_SOCKET_FEATURE_POLL			1
 //#define PLATFORM_USE_PTHREADS							1
-#define PLATFORM_HAS_MULTITHREADED_PREMAIN				1
 #define PLATFORM_MAX_FILEPATH_LENGTH_DEPRECATED			MAC_MAX_PATH
 #define PLATFORM_SUPPORTS_TBB							1
 #define PLATFORM_SUPPORTS_MIMALLOC						PLATFORM_64BITS
-#define PLATFORM_SUPPORTS_STACK_SYMBOLS					1
-#define PLATFORM_HAS_BSD_SOCKET_FEATURE_MSG_DONTWAIT	1
-#define PLATFORM_IS_ANSI_MALLOC_THREADSAFE				1
 #define PLATFORM_SUPPORTS_VIRTUAL_TEXTURE_STREAMING		1
 #define PLATFORM_SUPPORTS_MESH_SHADERS                  1
 #define PLATFORM_SUPPORTS_BINDLESS_RENDERING            1
@@ -84,72 +69,16 @@ typedef FMacPlatformTypes FPlatformTypes;
  
 #define PLATFORM_ENABLE_POPCNT_INTRINSIC				1
 
-// Ensure we can use this builtin - seems to be present on Clang 9, GCC 11 and MSVC 19.26,
-// but gives spurious "non-void function 'BitCast' should return a value" errors on some
-// Mac and Android toolchains when building PCHs, so avoid those.
-#undef PLATFORM_COMPILER_SUPPORTS_BUILTIN_BITCAST
-#define PLATFORM_COMPILER_SUPPORTS_BUILTIN_BITCAST (__clang_major__ >= 13)
-
 #define PLATFORM_GLOBAL_LOG_CATEGORY					LogMac
 
 #if WITH_EDITOR
 #define PLATFORM_FILE_READER_BUFFER_SIZE				(256*1024)
 #endif
 
-#if PLATFORM_MAC_X86
-	#define PLATFORM_BREAK()							__asm__("int $3")
-#else
-    #define PLATFORM_BREAK()                            __builtin_debugtrap()
-#endif
-
-#define PLATFORM_CODE_SECTION(Name)						__attribute__((section("__TEXT,__" Name ",regular,pure_instructions")))
-
-// Function type macros.
-#define VARARGS															/* Functions with variable arguments */
-#define CDECL															/* Standard C function */
-#define STDCALL															/* Standard calling convention */
-#if UE_BUILD_DEBUG
-#define FORCEINLINE inline 												/* Don't force code to be inline */
-#else
-#define FORCEINLINE inline __attribute__ ((always_inline))				/* Force code to be inline */
-#endif
-
-#define FORCENOINLINE __attribute__((noinline))							/* Force code to NOT be inline */
-#define FUNCTION_CHECK_RETURN_END __attribute__ ((warn_unused_result))	/* Warn that callers should not ignore the return value. */
-#define FUNCTION_NO_RETURN_END __attribute__ ((noreturn))				/* Indicate that the function never returns. */
-
-#define ABSTRACT abstract
-
-// We can use pragma optimisation's on and off as of Apple LLVM 7.3.0 but not before.
-#if (__clang_major__ > 7) || (__clang_major__ == 7 && __clang_minor__ >= 3)
-#define PRAGMA_DISABLE_OPTIMIZATION_ACTUAL _Pragma("clang optimize off")
-#define PRAGMA_ENABLE_OPTIMIZATION_ACTUAL  _Pragma("clang optimize on")
-#endif
-
-// Strings.
-// Alignment.
-#define GCC_PACK(n) __attribute__((packed,aligned(n)))
-#define GCC_ALIGN(n) __attribute__((aligned(n)))
-
-// operator new/delete operators
-// As of 10.9 we need to use _NOEXCEPT & cxx_noexcept compatible definitions
-#if __has_feature(cxx_noexcept)
-#define OPERATOR_NEW_THROW_SPEC
-#else
-#define OPERATOR_NEW_THROW_SPEC throw (std::bad_alloc)
-#endif
-#define OPERATOR_DELETE_THROW_SPEC _NOEXCEPT
-#define OPERATOR_NEW_NOTHROW_SPEC  _NOEXCEPT
-#define OPERATOR_DELETE_NOTHROW_SPEC  _NOEXCEPT
-
 #ifdef PLATFORM_MAC_ARM64
 #define PLATFORM_CACHE_LINE_SIZE	128
 #else
 #define PLATFORM_CACHE_LINE_SIZE	64
 #endif
 
-// DLL export and import definitions
-#define DLLEXPORT			__attribute__((visibility("default")))
-#define DLLIMPORT			__attribute__((visibility("default")))
-
-#define MAC_MAX_PATH 1024
+#define MAC_MAX_PATH 1024
\ No newline at end of file
