Implement FGenericPlatformMemory::MapNamedSharedMemoryRegion & FGenericPlatformMemory::UnmapNamedSharedMemoryRegion on Mac
This is a duplication of the Unix implementation with a check on the name length

#tests Local testing with a project that use SharedMemory

diff --git a/Engine/Source/Runtime/Core/Private/Mac/MacPlatformMemory.cpp b/Engine/Source/Runtime/Core/Private/Mac/MacPlatformMemory.cpp
--- a/Engine/Source/Runtime/Core/Private/Mac/MacPlatformMemory.cpp
+++ b/Engine/Source/Runtime/Core/Private/Mac/MacPlatformMemory.cpp
@@ -19,6 +19,8 @@
 
 #include <sys/param.h>
 #include <sys/mount.h>
+#include <sys/mman.h>
+#include <sys/posix_shm.h>
 
 #include <mach/vm_page_size.h>
 
@@ -236,3 +238,140 @@ const FPlatformMemoryConstants& FMacPlatformMemory::GetConstants()
 
 	return MemoryConstants;	
 }
+
+FPlatformMemory::FSharedMemoryRegion* FMacPlatformMemory::MapNamedSharedMemoryRegion(const FString& InName, bool bCreate, uint32 AccessMode, SIZE_T Size)
+{
+	// The maximum is PSHMNAMLEN but a '/' is added below
+	constexpr uint32 SHM_MAX_FILENAME = PSHMNAMLEN - 1;
+	ensureMsgf(InName.Len() <= SHM_MAX_FILENAME, TEXT("Trying to create a SharedMemoryRegion with a name size > %d will likely failed on Mac (Name=\"%s\", Size=%d)"), SHM_MAX_FILENAME, *InName, InName.Len());
+	
+	// expecting platform-independent name, so convert it to match platform requirements
+	FString Name("/");
+	Name += InName;
+	FTCHARToUTF8 NameUTF8(*Name);
+
+	// correct size to match platform constraints
+	FPlatformMemoryConstants MemConstants = FPlatformMemory::GetConstants();
+	check(MemConstants.PageSize > 0);	// also relying on it being power of two, which should be true in foreseeable future
+	if (Size & (MemConstants.PageSize - 1))
+	{
+		Size = Size & ~(MemConstants.PageSize - 1);
+		Size += MemConstants.PageSize;
+	}
+
+	int ShmOpenFlags = bCreate ? O_CREAT : 0;
+	// note that you cannot combine O_RDONLY and O_WRONLY to get O_RDWR
+	check(AccessMode != 0);
+	if (AccessMode == FPlatformMemory::ESharedMemoryAccess::Read)
+	{
+		ShmOpenFlags |= O_RDONLY;
+	}
+	else if (AccessMode == FPlatformMemory::ESharedMemoryAccess::Write)
+	{
+		ShmOpenFlags |= O_WRONLY;
+	}
+	else if (AccessMode == (FPlatformMemory::ESharedMemoryAccess::Write | FPlatformMemory::ESharedMemoryAccess::Read))
+	{
+		ShmOpenFlags |= O_RDWR;
+	}
+
+	int ShmOpenMode = (S_IRUSR | S_IWUSR) | (S_IRGRP | S_IWGRP) | (S_IROTH | S_IWOTH );	// 0666
+
+	// open the object
+	int SharedMemoryFd = shm_open(NameUTF8.Get(), ShmOpenFlags, ShmOpenMode);
+	if (SharedMemoryFd == -1)
+	{
+		int ErrNo = errno;
+		UE_LOG(LogHAL, Warning, TEXT("shm_open(name='%s', flags=0x%x, mode=0x%x) failed with errno = %d (%s)"), *Name, ShmOpenFlags, ShmOpenMode, ErrNo, 
+			StringCast< TCHAR >(strerror(ErrNo)).Get());
+		return NULL;
+	}
+
+	// truncate if creating (note that we may still don't have rights to do so)
+	if (bCreate)
+	{
+		int Res = ftruncate(SharedMemoryFd, Size);
+		if (Res != 0)
+		{
+			int ErrNo = errno;
+			UE_LOG(LogHAL, Warning, TEXT("ftruncate(fd=%d, size=%d) failed with errno = %d (%s)"), SharedMemoryFd, Size, ErrNo, 
+				StringCast< TCHAR >(strerror(ErrNo)).Get());
+			shm_unlink(NameUTF8.Get());
+			return NULL;
+		}
+	}
+
+	// map
+	int MmapProtFlags = 0;
+	if (AccessMode & FPlatformMemory::ESharedMemoryAccess::Read)
+	{
+		MmapProtFlags |= PROT_READ;
+	}
+
+	if (AccessMode & FPlatformMemory::ESharedMemoryAccess::Write)
+	{
+		MmapProtFlags |= PROT_WRITE;
+	}
+
+	void *Ptr = mmap(NULL, Size, MmapProtFlags, MAP_SHARED, SharedMemoryFd, 0);
+	if (Ptr == MAP_FAILED)
+	{
+		int ErrNo = errno;
+		UE_LOG(LogHAL, Warning, TEXT("mmap(addr=NULL, length=%d, prot=0x%x, flags=MAP_SHARED, fd=%d, 0) failed with errno = %d (%s)"), Size, MmapProtFlags, SharedMemoryFd, ErrNo, 
+			StringCast< TCHAR >(strerror(ErrNo)).Get());
+
+		if (bCreate)
+		{
+			shm_unlink(NameUTF8.Get());
+		}
+		return NULL;
+	}
+
+	return new FMacSharedMemoryRegion(Name, AccessMode, Ptr, Size, SharedMemoryFd, bCreate);
+}
+
+bool FMacPlatformMemory::UnmapNamedSharedMemoryRegion(FSharedMemoryRegion * MemoryRegion)
+{
+	bool bAllSucceeded = true;
+
+	if (MemoryRegion)
+	{
+		FMacSharedMemoryRegion * MacRegion = static_cast< FMacSharedMemoryRegion* >( MemoryRegion );
+
+		if (munmap(MacRegion->GetAddress(), MacRegion->GetSize()) == -1) 
+		{
+			bAllSucceeded = false;
+
+			int ErrNo = errno;
+			UE_LOG(LogHAL, Warning, TEXT("munmap(addr=%p, len=%d) failed with errno = %d (%s)"), MacRegion->GetAddress(), MacRegion->GetSize(), ErrNo, 
+				StringCast< TCHAR >(strerror(ErrNo)).Get());
+		}
+
+		if (close(MacRegion->GetFileDescriptor()) == -1)
+		{
+			bAllSucceeded = false;
+
+			int ErrNo = errno;
+			UE_LOG(LogHAL, Warning, TEXT("close(fd=%d) failed with errno = %d (%s)"), MacRegion->GetFileDescriptor(), ErrNo, 
+				StringCast< TCHAR >(strerror(ErrNo)).Get());
+		}
+
+		if (MacRegion->NeedsToUnlinkRegion())
+		{
+			FTCHARToUTF8 NameUTF8(MacRegion->GetName());
+			if (shm_unlink(NameUTF8.Get()) == -1)
+			{
+				bAllSucceeded = false;
+
+				int ErrNo = errno;
+				UE_LOG(LogHAL, Warning, TEXT("shm_unlink(name='%s') failed with errno = %d (%s)"), MacRegion->GetName(), ErrNo, 
+					StringCast< TCHAR >(strerror(ErrNo)).Get());
+			}
+		}
+
+		// delete the region
+		delete MacRegion;
+	}
+
+	return bAllSucceeded;
+}
diff --git a/Engine/Source/Runtime/Core/Public/Mac/MacPlatformMemory.h b/Engine/Source/Runtime/Core/Public/Mac/MacPlatformMemory.h
--- a/Engine/Source/Runtime/Core/Public/Mac/MacPlatformMemory.h
+++ b/Engine/Source/Runtime/Core/Public/Mac/MacPlatformMemory.h
@@ -14,12 +14,43 @@
 **/
 struct CORE_API FMacPlatformMemory : public FApplePlatformMemory
 {
+	/**
+	 * Mac representation of a shared memory region
+	 */
+	struct FMacSharedMemoryRegion : public FSharedMemoryRegion
+	{
+		/** Returns file descriptor of a shared memory object */
+		int GetFileDescriptor() const { return Fd; }
+
+		/** Returns true if we need to unlink this region on destruction (no other process will be able to access it) */
+		bool NeedsToUnlinkRegion() const { return bCreatedThisRegion; }
+
+		FMacSharedMemoryRegion(const FString& InName, uint32 InAccessMode, void* InAddress, SIZE_T InSize, int InFd, bool bInCreatedThisRegion)
+			:	FSharedMemoryRegion(InName, InAccessMode, InAddress, InSize)
+			,	Fd(InFd)
+			,	bCreatedThisRegion(bInCreatedThisRegion)
+		{}
+
+	protected:
+
+		/** File descriptor of a shared region */
+		int				Fd;
+
+		/** Whether we created this region */
+		bool			bCreatedThisRegion;
+	};
+
+	
+	
 	//~ Begin FGenericPlatformMemory Interface
 	static FPlatformMemoryStats GetStats();
 	static const FPlatformMemoryConstants& GetConstants();
 	static FMalloc* BaseAllocator();
 	
 	static FGenericPlatformMemoryStats::EMemoryPressureStatus MemoryPressureStatus;
+	
+	static CORE_API FSharedMemoryRegion * MapNamedSharedMemoryRegion(const FString& InName, bool bCreate, uint32 AccessMode, SIZE_T Size);
+	static CORE_API bool UnmapNamedSharedMemoryRegion(FSharedMemoryRegion * MemoryRegion);
 	//~ End FGenericPlatformMemory Interface
 };
 
