Fixed compilation on Mac platform due to a missing UE_DISABLE_FORCE_INLINE define

diff --git a/Engine/Source/Runtime/Core/Public/Apple/ApplePlatform.h b/Engine/Source/Runtime/Core/Public/Apple/ApplePlatform.h
--- a/Engine/Source/Runtime/Core/Public/Apple/ApplePlatform.h
+++ b/Engine/Source/Runtime/Core/Public/Apple/ApplePlatform.h
@@ -10,9 +10,9 @@
 #include "Clang/ClangPlatform.h"
 
 // Base defines, must define these for the platform, there are no defaults
-#define PLATFORM_64BITS					1
+#define PLATFORM_64BITS									1
 // Technically the underlying platform has 128bit atomics, but clang might not issue optimal code
-#define PLATFORM_HAS_128BIT_ATOMICS		0
+#define PLATFORM_HAS_128BIT_ATOMICS						0
 
 // Base defines, defaults are commented out
 #define PLATFORM_LITTLE_ENDIAN							1
@@ -45,11 +45,6 @@
 #define VARARGS																/* Functions with variable arguments */
 #define CDECL																/* Standard C function */
 #define STDCALL																/* Standard calling convention */
-#if UE_BUILD_DEBUG || UE_DISABLE_FORCE_INLINE
-#	define FORCEINLINE				inline 									/* Don't force code to be inline */
-#else
-#	define FORCEINLINE				inline __attribute__ ((always_inline))	/* Force code to be inline */
-#endif
 #define FORCENOINLINE				__attribute__((noinline))				/* Force code to NOT be inline */
 #define FUNCTION_CHECK_RETURN_END	__attribute__ ((warn_unused_result))	/* Warn that callers should not ignore the return value. */
 #define FUNCTION_NO_RETURN_END		__attribute__ ((noreturn))				/* Indicate that the function never returns. */
diff --git a/Engine/Source/Runtime/Core/Public/IOS/IOSPlatform.h b/Engine/Source/Runtime/Core/Public/IOS/IOSPlatform.h
--- a/Engine/Source/Runtime/Core/Public/IOS/IOSPlatform.h
+++ b/Engine/Source/Runtime/Core/Public/IOS/IOSPlatform.h
@@ -23,7 +23,7 @@ struct FIOSPlatformTypes : public FGenericPlatformTypes
 typedef FIOSPlatformTypes FPlatformTypes;
 
 // Base defines, must define these for the platform, there are no defaults
-#define PLATFORM_DESKTOP				0
+#define PLATFORM_DESKTOP								0
 
 // Base defines, defaults are commented out
 #define PLATFORM_TCHAR_IS_CHAR16						1
@@ -35,15 +35,15 @@ typedef FIOSPlatformTypes FPlatformTypes;
 #define PLATFORM_SUPPORTS_EARLY_MOVIE_PLAYBACK			1 // movies will start before engine is initalized
 #define PLATFORM_USE_FULL_TASK_GRAPH					0 // @todo platplug: not platplug, but should investigate soon anyway
 
-#define PLATFORM_NUM_AUDIODECOMPRESSION_PRECACHE_BUFFERS		0
+#define PLATFORM_NUM_AUDIODECOMPRESSION_PRECACHE_BUFFERS 0
 #if PLATFORM_TVOS
-#define PLATFORM_USES_GLES								0
-#define PLATFORM_HAS_TOUCH_MAIN_SCREEN					0
-#define	PLATFORM_SUPPORTS_OPUS_CODEC					0
-#define PLATFORM_SUPPORTS_VORBIS_CODEC					0
+#	define PLATFORM_USES_GLES							0
+#	define PLATFORM_HAS_TOUCH_MAIN_SCREEN				0
+#	define	PLATFORM_SUPPORTS_OPUS_CODEC				0
+#	define PLATFORM_SUPPORTS_VORBIS_CODEC				0
 #else
-#define PLATFORM_USES_GLES								1
-#define PLATFORM_HAS_TOUCH_MAIN_SCREEN					1
+#	define PLATFORM_USES_GLES							1
+#	define PLATFORM_HAS_TOUCH_MAIN_SCREEN				1
 #endif
 #define PLATFORM_UI_HAS_MOBILE_SCROLLBARS				1
 #define PLATFORM_UI_NEEDS_TOOLTIPS						0
@@ -59,13 +59,18 @@ typedef FIOSPlatformTypes FPlatformTypes;
 #define PLATFORM_GLOBAL_LOG_CATEGORY					LogIOS
 
 //mallocpoison not safe with aligned ansi allocator.  returns the larger unaligned size during Free() which causes writes off the end of the allocation.
-#define UE_USE_MALLOC_FILL_BYTES 0 
+#define UE_USE_MALLOC_FILL_BYTES						0 
+
+#define IOS_MAX_PATH									1024
+
+#if UE_BUILD_DEBUG || UE_DISABLE_FORCE_INLINE
+#	define FORCEINLINE									inline 									/* Don't force code to be inline */
+#else
+#	define FORCEINLINE									inline __attribute__ ((always_inline))	/* Force code to be inline */
+#endif
 
-// Strings.
 // Strings.
 #define LINE_TERMINATOR TEXT("\n")
 #define LINE_TERMINATOR_ANSI "\n"
 
-#define IOS_MAX_PATH 1024
-
 static_assert(__IPHONE_OS_VERSION_MAX_ALLOWED >= 13000, "Unreal requires Xcode 11 or later to build"); 
\ No newline at end of file
diff --git a/Engine/Source/Runtime/Core/Public/Mac/MacPlatform.h b/Engine/Source/Runtime/Core/Public/Mac/MacPlatform.h
--- a/Engine/Source/Runtime/Core/Public/Mac/MacPlatform.h
+++ b/Engine/Source/Runtime/Core/Public/Mac/MacPlatform.h
@@ -30,11 +30,11 @@ typedef FMacPlatformTypes FPlatformTypes;
 
 // Define ARM64 / X86 here so we can run UBT once for both platforms
 #if __is_target_arch(arm64) || __is_target_arch(arm64e)
-    #define PLATFORM_MAC_ARM64 1
-    #define PLATFORM_MAC_X86 0
+#	define PLATFORM_MAC_ARM64							1
+#	define PLATFORM_MAC_X86								0
 #else
-    #define PLATFORM_MAC_ARM64 0
-    #define PLATFORM_MAC_X86 1
+#	define PLATFORM_MAC_ARM64							0
+#	define PLATFORM_MAC_X86								1
 #endif
 
 // Base defines, must define these for the platform, there are no defaults
@@ -46,15 +46,15 @@ typedef FMacPlatformTypes FPlatformTypes;
 #define PLATFORM_ENABLE_VECTORINTRINSICS_NEON			PLATFORM_MAC_ARM64
 // FMA3 support was added starting from Intel Haswell
 #ifndef PLATFORM_ALWAYS_HAS_FMA3
-	#define PLATFORM_ALWAYS_HAS_FMA3					0
+#	define PLATFORM_ALWAYS_HAS_FMA3						0
 #endif
 //#define PLATFORM_USE_LS_SPEC_FOR_WIDECHAR				1
 #define PLATFORM_COMPILER_DISTINGUISHES_INT_AND_LONG	1
 #define PLATFORM_WCHAR_IS_4_BYTES						1
 #if PLATFORM_MAC_USE_CHAR16
-	#define PLATFORM_TCHAR_IS_CHAR16					1
+#	define PLATFORM_TCHAR_IS_CHAR16						1
 #else
-	#define PLATFORM_TCHAR_IS_4_BYTES					1
+#	define PLATFORM_TCHAR_IS_4_BYTES					1
 #endif
 #define PLATFORM_HAS_BSD_SOCKET_FEATURE_IOCTL			1
 #define PLATFORM_HAS_BSD_SOCKET_FEATURE_POLL			1
@@ -72,13 +72,19 @@ typedef FMacPlatformTypes FPlatformTypes;
 #define PLATFORM_GLOBAL_LOG_CATEGORY					LogMac
 
 #if WITH_EDITOR
-#define PLATFORM_FILE_READER_BUFFER_SIZE				(256*1024)
+#	define PLATFORM_FILE_READER_BUFFER_SIZE				(256*1024)
 #endif
 
 #ifdef PLATFORM_MAC_ARM64
-#define PLATFORM_CACHE_LINE_SIZE	128
+#	define PLATFORM_CACHE_LINE_SIZE						128
 #else
-#define PLATFORM_CACHE_LINE_SIZE	64
+#	define PLATFORM_CACHE_LINE_SIZE						64
 #endif
 
-#define MAC_MAX_PATH 1024
\ No newline at end of file
+#define MAC_MAX_PATH									1024
+
+#if UE_BUILD_DEBUG
+#	define FORCEINLINE									inline 									/* Don't force code to be inline */
+#else
+#	define FORCEINLINE									inline __attribute__ ((always_inline))	/* Force code to be inline */
+#endif
\ No newline at end of file
