Fixed LLM Untracked stat on a platforms that cache memory in the allocator

diff --git a/Engine/Source/Runtime/Core/Private/HAL/LowLevelMemTracker.cpp b/Engine/Source/Runtime/Core/Private/HAL/LowLevelMemTracker.cpp
--- a/Engine/Source/Runtime/Core/Private/HAL/LowLevelMemTracker.cpp
+++ b/Engine/Source/Runtime/Core/Private/HAL/LowLevelMemTracker.cpp
@@ -1977,7 +1977,7 @@ void FLowLevelMemTracker::TickInternal()
 
 	// Cache the amount of memory used early, since some of these functions (FindOrAddTagData) can
 	// cause allocations, which will throw the numbers off slightly.
-	FPlatformMemoryStats PlatformStats = FPlatformMemory::GetStatsRaw();
+	FPlatformMemoryStats PlatformStats = FPlatformMemory::GetStats();
 #if PLATFORM_DESKTOP
 	// virtual is working set + paged out memory.
 	const int64 PlatformProcessMemory = static_cast<int64>(PlatformStats.UsedVirtual);
