[AutoRTFM] Fix Chaos halloween map.

diff --git a/Engine/Source/Runtime/Core/Private/Internationalization/TextHistory.cpp b/Engine/Source/Runtime/Core/Private/Internationalization/TextHistory.cpp
--- a/Engine/Source/Runtime/Core/Private/Internationalization/TextHistory.cpp
+++ b/Engine/Source/Runtime/Core/Private/Internationalization/TextHistory.cpp
@@ -2537,8 +2537,8 @@ void FTextHistory_StringTableEntry::FStringTableReferenceData::Initialize(FName
 
 bool FTextHistory_StringTableEntry::FStringTableReferenceData::IsIdentical(const FStringTableReferenceData& Other) const
 {
-	FScopeLock ScopeLock(&DataCS);
-	FScopeLock OtherScopeLock(&Other.DataCS);
+	UE::TScopeLock ScopeLock(DataCS);
+	UE::TScopeLock OtherScopeLock(Other.DataCS);
 
 	return TableId == Other.TableId
 		&& Key == Other.Key;
@@ -2546,19 +2546,19 @@ bool FTextHistory_StringTableEntry::FStringTableReferenceData::IsIdentical(const
 
 FName FTextHistory_StringTableEntry::FStringTableReferenceData::GetTableId() const
 {
-	FScopeLock ScopeLock(&DataCS);
+	UE::TScopeLock ScopeLock(DataCS);
 	return TableId;
 }
 
 FTextKey FTextHistory_StringTableEntry::FStringTableReferenceData::GetKey() const
 {
-	FScopeLock ScopeLock(&DataCS);
+	UE::TScopeLock ScopeLock(DataCS);
 	return Key;
 }
 
 void FTextHistory_StringTableEntry::FStringTableReferenceData::GetTableIdAndKey(FName& OutTableId, FTextKey& OutKey) const
 {
-	FScopeLock ScopeLock(&DataCS);
+	UE::TScopeLock ScopeLock(DataCS);
 	OutTableId = TableId;
 	OutKey = Key;
 }
@@ -2576,7 +2576,7 @@ void FTextHistory_StringTableEntry::FStringTableReferenceData::CollectStringTabl
 {
 	if (Record.GetUnderlyingArchive().IsObjectReferenceCollector())
 	{
-		FScopeLock ScopeLock(&DataCS);
+		UE::TScopeLock ScopeLock(DataCS);
 
 		const FName OldTableId = TableId;
 		IStringTableEngineBridge::CollectStringTableAssetReferences(TableId, Record.EnterField(TEXT("AssetReferences")));
@@ -2601,7 +2601,7 @@ FStringTableEntryConstPtr FTextHistory_StringTableEntry::FStringTableReferenceDa
 
 	if (!StringTableEntryPin.IsValid() || !StringTableEntryPin->IsOwned())
 	{
-		FScopeLock ScopeLock(&DataCS);
+		UE::TScopeLock ScopeLock(DataCS);
 
 		// Reset for the case it was disowned rather than became null
 		StringTableEntry.Reset();
@@ -2657,7 +2657,7 @@ void FTextHistory_StringTableEntry::FStringTableReferenceData::ConditionalBeginA
 
 	FName TableIdToLoad;
 	{
-		FScopeLock ScopeLock(&DataCS);
+		UE::TScopeLock ScopeLock(DataCS);
 
 		if (LoadingPhase != EStringTableLoadingPhase::PendingLoad)
 		{
@@ -2677,7 +2677,7 @@ void FTextHistory_StringTableEntry::FStringTableReferenceData::ConditionalBeginA
 			return;
 		}
 
-		FScopeLock ScopeLock(&This->DataCS);
+		UE::TScopeLock ScopeLock(This->DataCS);
 		check(This->TableId == InRequestedTableId);
 
 		// If this string table loaded, then update the table ID using the potentially redirected value
diff --git a/Engine/Source/Runtime/Core/Private/Internationalization/TextHistory.h b/Engine/Source/Runtime/Core/Private/Internationalization/TextHistory.h
--- a/Engine/Source/Runtime/Core/Private/Internationalization/TextHistory.h
+++ b/Engine/Source/Runtime/Core/Private/Internationalization/TextHistory.h
@@ -11,6 +11,7 @@
 #include "Internationalization/Text.h"
 #include "Internationalization/TextKey.h"
 #include "Misc/DateTime.h"
+#include "Misc/TransactionallySafeCriticalSection.h"
 #include "Serialization/StructuredArchive.h"
 #include "Templates/SharedPointerInternals.h"
 #include "Templates/UnrealTemplate.h"
@@ -649,7 +650,7 @@ private:
 		FTextConstDisplayStringPtr DisplayString;
 
 		/** Critical section preventing concurrent access to the resolved data */
-		mutable FCriticalSection DataCS;
+		mutable FTransactionallySafeCriticalSection DataCS;
 	};
 	typedef TSharedPtr<FStringTableReferenceData, ESPMode::ThreadSafe> FStringTableReferenceDataPtr;
 	typedef TWeakPtr<FStringTableReferenceData, ESPMode::ThreadSafe> FStringTableReferenceDataWeakPtr;
