Implement mapped file support for all Apple platforms

* Lift the mapped file access implementation to Apple layer.
* Added functionality to redirect file names per platform.
* Changed log category from LogIOS to LogHAL.

diff --git a/Engine/Source/Runtime/Core/Private/Apple/ApplePlatformFile.cpp b/Engine/Source/Runtime/Core/Private/Apple/ApplePlatformFile.cpp
--- a/Engine/Source/Runtime/Core/Private/Apple/ApplePlatformFile.cpp
+++ b/Engine/Source/Runtime/Core/Private/Apple/ApplePlatformFile.cpp
@@ -17,6 +17,9 @@
 #include "Templates/Function.h"
 #include <sys/stat.h>
 
+#include "Async/MappedFileHandle.h"
+#include <sys/mman.h>
+
 // make an FTimeSpan object that represents the "epoch" for time_t (from a stat struct)
 const FDateTime MacEpoch(1970, 1, 1);
 
@@ -477,6 +480,126 @@ __thread FFileHandleApple* FFileHandleApple::ActiveHandles[ FFileHandleApple::AC
 __thread double FFileHandleApple::AccessTimes[ FFileHandleApple::ACTIVE_HANDLE_COUNT ];
 #endif
 
+static SIZE_T FileMappingAlignment = FPlatformMemory::GetConstants().PageSize;
+
+class FMappedFileRegion final : public IMappedFileRegion
+{
+public:
+	class FMappedFileHandle* Parent;
+	const uint8* AlignedPtr;
+	uint64 AlignedSize;
+	
+	FMappedFileRegion(const uint8* InMappedPtr, const uint8* InAlignedPtr, size_t InMappedSize, uint64 InAlignedSize, const FString& InDebugFilename, size_t InDebugOffsetIntoFile, class FMappedFileHandle* InParent)
+		: IMappedFileRegion(InMappedPtr, InMappedSize, InDebugFilename, InDebugOffsetIntoFile)
+		, Parent(InParent)
+		, AlignedPtr(InAlignedPtr)
+		, AlignedSize(InAlignedSize)
+	{
+	}
+	
+	~FMappedFileRegion();
+	
+	virtual void PreloadHint(int64 PreloadOffset = 0, int64 BytesToPreload = MAX_int64) override
+	{
+		int64 Size = GetMappedSize();
+		const uint8* Ptr = GetMappedPtr();
+		int32 FoolTheOptimizer = 0;
+		while (Size > 0)
+		{
+			FoolTheOptimizer += Ptr[0];
+			Size -= 4096;
+			Ptr += 4096;
+		}
+		if (FoolTheOptimizer == 0xbadf00d)
+		{
+			FPlatformProcess::Sleep(0.0f); // this will more or less never happen, but we can't let the optimizer strip these reads
+		}
+	}
+	
+};
+
+class FMappedFileHandle final : public IMappedFileHandle
+{
+	const uint8* MappedPtr;
+	FString Filename;
+	int32 NumOutstandingRegions;
+	int FileHandle;
+	
+public:
+	
+	FMappedFileHandle(int InFileHandle, int64 FileSize, const FString& InFilename)
+		: IMappedFileHandle(FileSize)
+		, MappedPtr(nullptr)
+#if !UE_BUILD_SHIPPING
+		, Filename(InFilename)
+#endif
+		, NumOutstandingRegions(0)
+		, FileHandle(InFileHandle)
+	{
+	}
+
+	~FMappedFileHandle()
+	{
+		// can't delete the file before you delete all outstanding regions
+		UE_CLOG(NumOutstandingRegions, LogHAL, Fatal, TEXT("Cleaning mapped file with alive mapped regions."));
+		close(FileHandle);
+	}
+	
+	virtual IMappedFileRegion* MapRegion(int64 Offset = 0, int64 BytesToMap = MAX_int64, FFileMappingFlags Flags = EMappedFileFlags::ENone) override
+	{
+		LLM_PLATFORM_SCOPE(ELLMTag::PlatformMMIO);
+		check(Offset < GetFileSize()); // don't map zero bytes and don't map off the end of the file
+		BytesToMap = FMath::Min<int64>(BytesToMap, GetFileSize() - Offset);
+		check(BytesToMap > 0); // don't map zero bytes
+		
+		const int64 AlignedOffset = AlignDown(Offset, FileMappingAlignment);
+		//File mapping can extend beyond file size. It's OK, kernel will just fill any leftover page data with zeros
+		const int64 AlignedSize = Align(BytesToMap + Offset - AlignedOffset, FileMappingAlignment);
+		
+		int Protection = PROT_READ;
+		int InternalFlags = 0;
+		if (EnumHasAnyFlags(Flags.Flags, EMappedFileFlags::EFileWritable))
+		{
+			Protection |= PROT_WRITE;
+			InternalFlags |= MAP_SHARED;
+		}
+		else
+		{
+			InternalFlags |= MAP_PRIVATE;
+		}
+
+		const uint8* AlignedMapPtr = (const uint8 *)mmap(NULL, AlignedSize, Protection, InternalFlags, FileHandle, AlignedOffset);
+		if (AlignedMapPtr == (const uint8*)-1 || AlignedMapPtr == nullptr)
+		{
+			UE_LOG(LogHAL, Warning, TEXT("Failed to map memory %s, error is %d"), *Filename, errno);
+			return nullptr;
+		}
+		LLM_IF_ENABLED(FLowLevelMemTracker::Get().OnLowLevelAlloc(ELLMTracker::Platform, AlignedMapPtr, AlignedSize));
+
+		// create a mapping for this range
+		const uint8* MapPtr = AlignedMapPtr + Offset - AlignedOffset;
+		FMappedFileRegion* Result = new FMappedFileRegion(MapPtr, AlignedMapPtr, BytesToMap, AlignedSize, Filename, Offset, this);
+		NumOutstandingRegions++;
+		return Result;
+	}
+	
+	void UnMap(FMappedFileRegion* Region)
+	{
+		LLM_PLATFORM_SCOPE(ELLMTag::PlatformMMIO);
+		check(NumOutstandingRegions > 0);
+		NumOutstandingRegions--;
+		
+		LLM_IF_ENABLED(FLowLevelMemTracker::Get().OnLowLevelFree(ELLMTracker::Platform, (void*)Region->AlignedPtr));
+		int Res = munmap((void*)Region->AlignedPtr, Region->AlignedSize);
+		checkf(Res == 0, TEXT("Failed to unmap, error is %d, errno is %d [params: %x, %d]"), Res, errno, MappedPtr, GetFileSize());
+	}
+};
+
+FMappedFileRegion::~FMappedFileRegion()
+{
+	Parent->UnMap(this);
+}
+
 /**
  * Mac File I/O implementation
 **/
@@ -740,6 +863,51 @@ IFileHandle* FApplePlatformFile::OpenWrite(const TCHAR* Filename, bool bAppend,
 	}
 }
 
+FOpenMappedResult FApplePlatformFile::OpenMappedEx(const TCHAR* Filename, EOpenReadFlags OpenOptions)
+{
+	FString NormalizedFilename = NormalizeFilename(Filename);
+	
+	// check the read path
+	FString FinalPath = ConvertToPlatformPath(NormalizedFilename, false, false);
+	FILE* FP;
+	const char* OpenMode = EnumHasAnyFlags(OpenOptions, EOpenReadFlags::AllowWrite) ? "r+" : "r";
+	FP = fopen(TCHAR_TO_UTF8(*FinalPath), OpenMode);
+	if(FP == nullptr)
+	{
+		// if not in the read path, check the private write path
+		FinalPath = ConvertToPlatformPath(NormalizedFilename, true, false);
+		FP = fopen(TCHAR_TO_UTF8(*FinalPath), OpenMode);
+
+		if(FP == nullptr)
+		{
+			// if not in the private write path, check the public write path
+			FinalPath = ConvertToPlatformPath(NormalizedFilename, true, true);
+			FP = fopen(TCHAR_TO_UTF8(*FinalPath), OpenMode);
+		}
+	}
+	
+	if (FP != nullptr)
+	{
+		int32 Handle = fileno(FP);
+	
+		if (Handle != -1)
+		{
+			struct stat FileInfo;
+			FileInfo.st_size = -1;
+			// check the read path
+			if(fstat(Handle, &FileInfo) == -1)
+			{
+				return MakeError(FString::Printf(TEXT("FApplePlatformFile::OpenMappedEx failed to get file info for file '%s'"), *NormalizedFilename));
+			}
+			uint64 FileSize = FileInfo.st_size;
+
+			return MakeValue(new FMappedFileHandle(Handle, FileSize, FinalPath));
+		}
+	}
+	
+	return MakeError(FString::Printf(TEXT("FApplePlatformFile::OpenMappedEx failed to open file '%s' in '%s' mode"), *NormalizedFilename, EnumHasAnyFlags(OpenOptions, EOpenReadFlags::AllowWrite) ? TEXT("write") : TEXT("read")));
+}
+
 bool FApplePlatformFile::DirectoryExists(const TCHAR* Directory)
 {
 	struct stat FileInfo;
diff --git a/Engine/Source/Runtime/Core/Private/IOS/IOSPlatformFile.cpp b/Engine/Source/Runtime/Core/Private/IOS/IOSPlatformFile.cpp
--- a/Engine/Source/Runtime/Core/Private/IOS/IOSPlatformFile.cpp
+++ b/Engine/Source/Runtime/Core/Private/IOS/IOSPlatformFile.cpp
@@ -15,7 +15,6 @@
 #include "Misc/ScopeLock.h"
 
 #include "Async/MappedFileHandle.h"
-#include <sys/mman.h>
 
 #if PLATFORM_USE_PLATFORM_FILE_MANAGED_STORAGE_WRAPPER
 #include "HAL/IPlatformFileManagedStorageWrapper.h"
@@ -447,126 +446,6 @@ int32 FIOSFileHandle::ManagedFilesTlsSlot = FPlatformTLS::AllocTlsSlot();
 #endif
 
 
-class FIOSMappedFileRegion final : public IMappedFileRegion
-{
-public:
-	class FIOSMappedFileHandle* Parent;
-	const uint8* AlignedPtr;
-	uint64 AlignedSize;
-	
-	FIOSMappedFileRegion(const uint8* InMappedPtr, const uint8* InAlignedPtr, size_t InMappedSize, uint64 InAlignedSize, const FString& InDebugFilename, size_t InDebugOffsetIntoFile, class FIOSMappedFileHandle* InParent)
-		: IMappedFileRegion(InMappedPtr, InMappedSize, InDebugFilename, InDebugOffsetIntoFile)
-		, Parent(InParent)
-		, AlignedPtr(InAlignedPtr)
-		, AlignedSize(InAlignedSize)
-	{
-	}
-	
-	~FIOSMappedFileRegion();
-	
-	virtual void PreloadHint(int64 PreloadOffset = 0, int64 BytesToPreload = MAX_int64) override
-	{
-		int64 Size = GetMappedSize();
-		const uint8* Ptr = GetMappedPtr();
-		int32 FoolTheOptimizer = 0;
-		while (Size > 0)
-		{
-			FoolTheOptimizer += Ptr[0];
-			Size -= 4096;
-			Ptr += 4096;
-		}
-		if (FoolTheOptimizer == 0xbadf00d)
-		{
-			FPlatformProcess::Sleep(0.0f); // this will more or less never happen, but we can't let the optimizer strip these reads
-		}
-	}
-	
-};
-
-static SIZE_T FileMappingAlignment = FPlatformMemory::GetConstants().PageSize;
-
-class FIOSMappedFileHandle final : public IMappedFileHandle
-{
-	const uint8* MappedPtr;
-	FString Filename;
-	int32 NumOutstandingRegions;
-	int FileHandle;
-	
-public:
-	
-	FIOSMappedFileHandle(int InFileHandle, int64 FileSize, const FString& InFilename)
-		: IMappedFileHandle(FileSize)
-		, MappedPtr(nullptr)
-#if !UE_BUILD_SHIPPING
-		, Filename(InFilename)
-#endif
-		, NumOutstandingRegions(0)
-		, FileHandle(InFileHandle)
-	{
-	}
-
-	~FIOSMappedFileHandle()
-	{
-		// can't delete the file before you delete all outstanding regions
-		UE_CLOG(NumOutstandingRegions, LogIOS, Fatal, TEXT("Cleaning mapped file with alive mapped regions."));
-		close(FileHandle);
-	}
-	
-	virtual IMappedFileRegion* MapRegion(int64 Offset = 0, int64 BytesToMap = MAX_int64, FFileMappingFlags Flags = EMappedFileFlags::ENone) override
-	{
-		LLM_PLATFORM_SCOPE(ELLMTag::PlatformMMIO);
-		check(Offset < GetFileSize()); // don't map zero bytes and don't map off the end of the file
-		BytesToMap = FMath::Min<int64>(BytesToMap, GetFileSize() - Offset);
-		check(BytesToMap > 0); // don't map zero bytes
-		
-		const int64 AlignedOffset = AlignDown(Offset, FileMappingAlignment);
-		//File mapping can extend beyond file size. It's OK, kernel will just fill any leftover page data with zeros
-		const int64 AlignedSize = Align(BytesToMap + Offset - AlignedOffset, FileMappingAlignment);
-		
-		int Protection = PROT_READ;
-		int InternalFlags = 0;
-		if (EnumHasAnyFlags(Flags.Flags, EMappedFileFlags::EFileWritable))
-		{
-			Protection |= PROT_WRITE;
-			InternalFlags |= MAP_SHARED;
-		}
-		else
-		{
-			InternalFlags |= MAP_PRIVATE;
-		}
-
-		const uint8* AlignedMapPtr = (const uint8 *)mmap(NULL, AlignedSize, Protection, InternalFlags, FileHandle, AlignedOffset);
-		if (AlignedMapPtr == (const uint8*)-1 || AlignedMapPtr == nullptr)
-		{
-			UE_LOG(LogIOS, Warning, TEXT("Failed to map memory %s, error is %d"), *Filename, errno);
-			return nullptr;
-		}
-		LLM_IF_ENABLED(FLowLevelMemTracker::Get().OnLowLevelAlloc(ELLMTracker::Platform, AlignedMapPtr, AlignedSize));
-
-		// create a mapping for this range
-		const uint8* MapPtr = AlignedMapPtr + Offset - AlignedOffset;
-		FIOSMappedFileRegion* Result = new FIOSMappedFileRegion(MapPtr, AlignedMapPtr, BytesToMap, AlignedSize, Filename, Offset, this);
-		NumOutstandingRegions++;
-		return Result;
-	}
-	
-	void UnMap(FIOSMappedFileRegion* Region)
-	{
-		LLM_PLATFORM_SCOPE(ELLMTag::PlatformMMIO);
-		check(NumOutstandingRegions > 0);
-		NumOutstandingRegions--;
-		
-		LLM_IF_ENABLED(FLowLevelMemTracker::Get().OnLowLevelFree(ELLMTracker::Platform, (void*)Region->AlignedPtr));
-		int Res = munmap((void*)Region->AlignedPtr, Region->AlignedSize);
-		checkf(Res == 0, TEXT("Failed to unmap, error is %d, errno is %d [params: %x, %d]"), Res, errno, MappedPtr, GetFileSize());
-	}
-};
-
-FIOSMappedFileRegion::~FIOSMappedFileRegion()
-{
-	Parent->UnMap(this);
-}
-
 
 
 /**
@@ -596,13 +475,13 @@ FString FIOSPlatformFile::ConvertToAbsolutePathForExternalAppForRead( const TCHA
 {
     struct stat FileInfo;
     FString NormalizedFilename = NormalizeFilename(Filename);
-    if (stat(TCHAR_TO_UTF8(*ConvertToIOSPath(NormalizedFilename, false, false)), &FileInfo) == -1)
+    if (stat(TCHAR_TO_UTF8(*ConvertToPlatformPath(NormalizedFilename, false, false)), &FileInfo) == -1)
     {
         return ConvertToAbsolutePathForExternalAppForWrite(Filename);
     }
     else
     {
-        return ConvertToIOSPath(NormalizedFilename, false, false);
+        return ConvertToPlatformPath(NormalizedFilename, false, false);
     }
 }
 
@@ -612,11 +491,11 @@ FString FIOSPlatformFile::ConvertToAbsolutePathForExternalAppForWrite( const TCH
     FString NormalizedFilename = NormalizeFilename(Filename);
 	if (bCreatePublicFiles)
 	{
-		return ConvertToIOSPath(NormalizedFilename, true, true);
+		return ConvertToPlatformPath(NormalizedFilename, true, true);
 	}
 	else
 	{
-		return ConvertToIOSPath(NormalizedFilename, true, false);
+		return ConvertToPlatformPath(NormalizedFilename, true, false);
 	}
 }
 
@@ -625,13 +504,13 @@ bool FIOSPlatformFile::FileExists(const TCHAR* Filename)
 	struct stat FileInfo;
 	FString NormalizedFilename = NormalizeFilename(Filename);
 	// check the read path
-	if (stat(TCHAR_TO_UTF8(*ConvertToIOSPath(NormalizedFilename, false, false)), &FileInfo) == -1)
+	if (stat(TCHAR_TO_UTF8(*ConvertToPlatformPath(NormalizedFilename, false, false)), &FileInfo) == -1)
 	{
 		// if not in read path, check the private write path
-		if (stat(TCHAR_TO_UTF8(*ConvertToIOSPath(NormalizedFilename, true, false)), &FileInfo) == -1)
+		if (stat(TCHAR_TO_UTF8(*ConvertToPlatformPath(NormalizedFilename, true, false)), &FileInfo) == -1)
 		{
 			// if not in the private write path, check the public write path
-			if (stat(TCHAR_TO_UTF8(*ConvertToIOSPath(NormalizedFilename, true, true)), &FileInfo) == -1)
+			if (stat(TCHAR_TO_UTF8(*ConvertToPlatformPath(NormalizedFilename, true, true)), &FileInfo) == -1)
 			{
 				return false;
 			}
@@ -647,13 +526,13 @@ int64 FIOSPlatformFile::FileSize(const TCHAR* Filename)
 	FileInfo.st_size = -1;
 	FString NormalizedFilename = NormalizeFilename(Filename);
 	// check the read path
-	if(stat(TCHAR_TO_UTF8(*ConvertToIOSPath(NormalizedFilename, false, false)), &FileInfo) == -1)
+	if(stat(TCHAR_TO_UTF8(*ConvertToPlatformPath(NormalizedFilename, false, false)), &FileInfo) == -1)
 	{
 		// if not in read path, check the private write path
-		if(stat(TCHAR_TO_UTF8(*ConvertToIOSPath(NormalizedFilename, true, false)), &FileInfo) == -1)
+		if(stat(TCHAR_TO_UTF8(*ConvertToPlatformPath(NormalizedFilename, true, false)), &FileInfo) == -1)
 		{
 			// if not in the private write path, check the public write path
-			if(stat(TCHAR_TO_UTF8(*ConvertToIOSPath(NormalizedFilename, true, true)), &FileInfo) == -1)
+			if(stat(TCHAR_TO_UTF8(*ConvertToPlatformPath(NormalizedFilename, true, true)), &FileInfo) == -1)
 			{
 				return -1;
 			}
@@ -675,8 +554,8 @@ bool FIOSPlatformFile::DeleteFile(const TCHAR* Filename)
 	struct stat FileInfo;
 	FileInfo.st_size = -1;
 	FString NormalizedFilename = NormalizeFilename(Filename);
-	FString IOSPrivateWriteFilename = ConvertToIOSPath(NormalizedFilename, true, false);
-	FString IOSPublicWriteFilename = ConvertToIOSPath(NormalizedFilename, true, true);
+	FString IOSPrivateWriteFilename = ConvertToPlatformPath(NormalizedFilename, true, false);
+	FString IOSPublicWriteFilename = ConvertToPlatformPath(NormalizedFilename, true, true);
 
 	// Try to delete the file from both the public and private write paths
 	bool bDeletedPrivate = unlink(TCHAR_TO_UTF8(*IOSPrivateWriteFilename)) == 0;
@@ -688,16 +567,16 @@ bool FIOSPlatformFile::DeleteFile(const TCHAR* Filename)
 bool FIOSPlatformFile::IsReadOnly(const TCHAR* Filename)
 {
 	FString NormalizedFilename = NormalizeFilename(Filename);
-	FString Filepath = ConvertToIOSPath(NormalizedFilename, false, false);
+	FString Filepath = ConvertToPlatformPath(NormalizedFilename, false, false);
 	// check read path
 	if (access(TCHAR_TO_UTF8(*Filepath), F_OK) == -1)
 	{
 		// if not in read path, check private write path
-		Filepath = ConvertToIOSPath(NormalizedFilename, true, false);
+		Filepath = ConvertToPlatformPath(NormalizedFilename, true, false);
 		if (access(TCHAR_TO_UTF8(*Filepath), F_OK) == -1)
 		{
 			// if not in private write path, check public write path
-			Filepath = ConvertToIOSPath(NormalizedFilename, true, true);
+			Filepath = ConvertToPlatformPath(NormalizedFilename, true, true);
 
 			if (access(TCHAR_TO_UTF8(*Filepath), F_OK) == -1)
 			{
@@ -716,18 +595,18 @@ bool FIOSPlatformFile::IsReadOnly(const TCHAR* Filename)
 bool FIOSPlatformFile::MoveFile(const TCHAR* To, const TCHAR* From)
 {
 	// move to the write path
-	FString ToIOSFilename = ConvertToIOSPath(NormalizeFilename(To), true, bCreatePublicFiles);
+	FString ToIOSFilename = ConvertToPlatformPath(NormalizeFilename(To), true, bCreatePublicFiles);
 	// move from the read path if the file exists there
-	FString FromIOSFilename = ConvertToIOSPath(NormalizeFilename(From), false, false);
+	FString FromIOSFilename = ConvertToPlatformPath(NormalizeFilename(From), false, false);
 	if (!FileExists(*FromIOSFilename))
 	{
 		// otherwise try the private write path
-		FromIOSFilename = ConvertToIOSPath(NormalizeFilename(From), true, false);
+		FromIOSFilename = ConvertToPlatformPath(NormalizeFilename(From), true, false);
 
 		if (!FileExists(*FromIOSFilename))
 		{
 			// and finally try the public write path
-			FromIOSFilename = ConvertToIOSPath(NormalizeFilename(From), true, true);
+			FromIOSFilename = ConvertToPlatformPath(NormalizeFilename(From), true, true);
 		}
 	}
 	return rename(TCHAR_TO_UTF8(*FromIOSFilename), TCHAR_TO_UTF8(*ToIOSFilename)) != -1;
@@ -736,7 +615,7 @@ bool FIOSPlatformFile::MoveFile(const TCHAR* To, const TCHAR* From)
 bool FIOSPlatformFile::SetReadOnly(const TCHAR* Filename, bool bNewReadOnlyValue)
 {
 	struct stat FileInfo;
-	FString IOSFilename = ConvertToIOSPath(NormalizeFilename(Filename), false, false);
+	FString IOSFilename = ConvertToPlatformPath(NormalizeFilename(Filename), false, false);
 	if (stat(TCHAR_TO_UTF8(*IOSFilename), &FileInfo) != -1)
 	{
 		if (bNewReadOnlyValue)
@@ -759,13 +638,13 @@ FDateTime FIOSPlatformFile::GetTimeStamp(const TCHAR* Filename)
 	struct stat FileInfo;
 	FString NormalizedFilename = NormalizeFilename(Filename);
 	// check the read path
-	if(stat(TCHAR_TO_UTF8(*ConvertToIOSPath(NormalizedFilename, false, false)), &FileInfo) == -1)
+	if(stat(TCHAR_TO_UTF8(*ConvertToPlatformPath(NormalizedFilename, false, false)), &FileInfo) == -1)
 	{
 		// if not in the read path, check the private write path
-		if(stat(TCHAR_TO_UTF8(*ConvertToIOSPath(NormalizedFilename, true, false)), &FileInfo) == -1)
+		if(stat(TCHAR_TO_UTF8(*ConvertToPlatformPath(NormalizedFilename, true, false)), &FileInfo) == -1)
 		{
 			// if not in the private write path, check the public write path
-			if(stat(TCHAR_TO_UTF8(*ConvertToIOSPath(NormalizedFilename, true, true)), &FileInfo) == -1)
+			if(stat(TCHAR_TO_UTF8(*ConvertToPlatformPath(NormalizedFilename, true, true)), &FileInfo) == -1)
 			{
 				return FDateTime::MinValue();
 			}
@@ -781,10 +660,10 @@ void FIOSPlatformFile::SetTimeStamp(const TCHAR* Filename, const FDateTime DateT
 {
 	// get file times
 	struct stat FileInfo;
-	FString IOSFilename = ConvertToIOSPath(NormalizeFilename(Filename), true, false);
+	FString IOSFilename = ConvertToPlatformPath(NormalizeFilename(Filename), true, false);
 	if(stat(TCHAR_TO_UTF8(*IOSFilename), &FileInfo) == -1)
 	{
-		IOSFilename = ConvertToIOSPath(NormalizeFilename(Filename), true, true);
+		IOSFilename = ConvertToPlatformPath(NormalizeFilename(Filename), true, true);
 		if(stat(TCHAR_TO_UTF8(*IOSFilename), &FileInfo) == -1)
 		{
 			return;
@@ -804,13 +683,13 @@ FDateTime FIOSPlatformFile::GetAccessTimeStamp(const TCHAR* Filename)
 	struct stat FileInfo;
 	FString NormalizedFilename = NormalizeFilename(Filename);
 	// check the read path
-	if(stat(TCHAR_TO_UTF8(*ConvertToIOSPath(NormalizedFilename, false, false)), &FileInfo) == -1)
+	if(stat(TCHAR_TO_UTF8(*ConvertToPlatformPath(NormalizedFilename, false, false)), &FileInfo) == -1)
 	{
 		// if not in the read path, check the private write path
-		if(stat(TCHAR_TO_UTF8(*ConvertToIOSPath(NormalizedFilename, true, false)), &FileInfo) == -1)
+		if(stat(TCHAR_TO_UTF8(*ConvertToPlatformPath(NormalizedFilename, true, false)), &FileInfo) == -1)
 		{
 			// if not in the private write path, check the public write path
-			if(stat(TCHAR_TO_UTF8(*ConvertToIOSPath(NormalizedFilename, true, true)), &FileInfo) == -1)
+			if(stat(TCHAR_TO_UTF8(*ConvertToPlatformPath(NormalizedFilename, true, true)), &FileInfo) == -1)
 			{
 				return FDateTime::MinValue();
 			}
@@ -833,13 +712,13 @@ FFileStatData FIOSPlatformFile::GetStatData(const TCHAR* FilenameOrDirectory)
 	FString NormalizedFilename = NormalizeFilename(FilenameOrDirectory);
 
 	// check the read path
-	if(stat(TCHAR_TO_UTF8(*ConvertToIOSPath(NormalizedFilename, false, false)), &FileInfo) == -1)
+	if(stat(TCHAR_TO_UTF8(*ConvertToPlatformPath(NormalizedFilename, false, false)), &FileInfo) == -1)
 	{
 		// if not in the read path, check the private write path
-		if(stat(TCHAR_TO_UTF8(*ConvertToIOSPath(NormalizedFilename, true, false)), &FileInfo) == -1)
+		if(stat(TCHAR_TO_UTF8(*ConvertToPlatformPath(NormalizedFilename, true, false)), &FileInfo) == -1)
 		{
 			// if not in the private write path, check the public write path
-			if(stat(TCHAR_TO_UTF8(*ConvertToIOSPath(NormalizedFilename, true, true)), &FileInfo) == -1)
+			if(stat(TCHAR_TO_UTF8(*ConvertToPlatformPath(NormalizedFilename, true, true)), &FileInfo) == -1)
 			{
 				return FFileStatData();
 			}
@@ -854,18 +733,18 @@ IFileHandle* FIOSPlatformFile::OpenRead(const TCHAR* Filename, bool bAllowWrite)
 	FString NormalizedFilename = NormalizeFilename(Filename);
 
 	// check the read path
-	FString FinalPath = ConvertToIOSPath(NormalizedFilename, false, false);
+	FString FinalPath = ConvertToPlatformPath(NormalizedFilename, false, false);
 	int32 Handle = open(TCHAR_TO_UTF8(*FinalPath), O_RDONLY);
 	if(Handle == -1)
 	{
 		// if not in the read path, check the private write path
-		FinalPath = ConvertToIOSPath(NormalizedFilename, true, false);
+		FinalPath = ConvertToPlatformPath(NormalizedFilename, true, false);
 		Handle = open(TCHAR_TO_UTF8(*FinalPath), O_RDONLY);
 		
 		if(Handle == -1)
 		{
 			// if not in the private write path, check the public write path
-			FinalPath = ConvertToIOSPath(NormalizedFilename, true, true);
+			FinalPath = ConvertToPlatformPath(NormalizedFilename, true, true);
 			Handle = open(TCHAR_TO_UTF8(*FinalPath), O_RDONLY);
 		}
 	}
@@ -892,7 +771,7 @@ IFileHandle* FIOSPlatformFile::OpenWrite(const TCHAR* Filename, bool bAppend, bo
 	{
 		Flags |= O_WRONLY;
 	}
-	FString IOSFilename = ConvertToIOSPath(NormalizeFilename(Filename), true, bCreatePublicFiles);
+	FString IOSFilename = ConvertToPlatformPath(NormalizeFilename(Filename), true, bCreatePublicFiles);
 	int32 Handle = open(TCHAR_TO_UTF8(*IOSFilename), Flags, S_IRUSR | S_IWUSR | S_IRGRP | S_IWGRP | S_IROTH | S_IWOTH);
 
 	if (Handle != -1)
@@ -912,60 +791,16 @@ IFileHandle* FIOSPlatformFile::OpenWrite(const TCHAR* Filename, bool bAppend, bo
 	return NULL;
 }
 
-FOpenMappedResult FIOSPlatformFile::OpenMappedEx(const TCHAR* Filename, EOpenReadFlags OpenOptions)
-{
-	FString NormalizedFilename = NormalizeFilename(Filename);
-	
-	// check the read path
-	FString FinalPath = ConvertToIOSPath(NormalizedFilename, false, false);
-	FILE* FP;
-	const char* OpenMode = EnumHasAnyFlags(OpenOptions, EOpenReadFlags::AllowWrite) ? "r+" : "r";
-	FP = fopen(TCHAR_TO_UTF8(*FinalPath), OpenMode);
-	if(FP == nullptr)
-	{
-		// if not in the read path, check the private write path
-		FinalPath = ConvertToIOSPath(NormalizedFilename, true, false);
-		FP = fopen(TCHAR_TO_UTF8(*FinalPath), OpenMode);
-
-		if(FP == nullptr)
-		{
-			// if not in the private write path, check the public write path
-			FinalPath = ConvertToIOSPath(NormalizedFilename, true, true);
-			FP = fopen(TCHAR_TO_UTF8(*FinalPath), OpenMode);
-		}
-	}
-	
-	if (FP != nullptr)
-	{
-		int32 Handle = fileno(FP);
-	
-		if (Handle != -1)
-		{
-			struct stat FileInfo;
-			FileInfo.st_size = -1;
-			// check the read path
-			if(fstat(Handle, &FileInfo) == -1)
-			{
-				return MakeError(FString::Printf(TEXT("FIOSPlatformFile::OpenMappedEx failed to get file info for file '%s'"), *NormalizedFilename));
-			}
-			uint64 FileSize = FileInfo.st_size;
-
-			return MakeValue(new FIOSMappedFileHandle(Handle, FileSize, FinalPath));
-		}
-	}
-	
-	return MakeError(FString::Printf(TEXT("FIOSPlatformFile::OpenMappedEx failed to open file '%s' in '%s' mode"), *NormalizedFilename, EnumHasAnyFlags(OpenOptions, EOpenReadFlags::AllowWrite) ? TEXT("write") : TEXT("read")));
-}
 
 bool FIOSPlatformFile::DirectoryExists(const TCHAR* Directory)
 {
 	struct stat FileInfo;
 	FString NormalizedDirectory = NormalizeFilename(Directory);
-	if (stat(TCHAR_TO_UTF8(*ConvertToIOSPath(NormalizedDirectory, false, false)), &FileInfo)== -1)
+	if (stat(TCHAR_TO_UTF8(*ConvertToPlatformPath(NormalizedDirectory, false, false)), &FileInfo)== -1)
 	{
-		if (stat(TCHAR_TO_UTF8(*ConvertToIOSPath(NormalizedDirectory, true, false)), &FileInfo)== -1)
+		if (stat(TCHAR_TO_UTF8(*ConvertToPlatformPath(NormalizedDirectory, true, false)), &FileInfo)== -1)
 		{
-			if (stat(TCHAR_TO_UTF8(*ConvertToIOSPath(NormalizedDirectory, true, true)), &FileInfo)== -1)
+			if (stat(TCHAR_TO_UTF8(*ConvertToPlatformPath(NormalizedDirectory, true, true)), &FileInfo)== -1)
 			{
 				return false;
 			}
@@ -976,7 +811,7 @@ bool FIOSPlatformFile::DirectoryExists(const TCHAR* Directory)
 
 bool FIOSPlatformFile::CreateDirectory(const TCHAR* Directory)
 {
-	FString IOSDirectory = ConvertToIOSPath(NormalizeFilename(Directory), true, bCreatePublicFiles);
+	FString IOSDirectory = ConvertToPlatformPath(NormalizeFilename(Directory), true, bCreatePublicFiles);
 	CFStringRef CFDirectory = FPlatformString::TCHARToCFString(*IOSDirectory);
 	bool Result = [[NSFileManager defaultManager] createDirectoryAtPath:(NSString*)CFDirectory withIntermediateDirectories:true attributes:nil error:nil];
 	CFRelease(CFDirectory);
@@ -985,8 +820,8 @@ bool FIOSPlatformFile::CreateDirectory(const TCHAR* Directory)
 
 bool FIOSPlatformFile::DeleteDirectory(const TCHAR* Directory)
 {
-	FString IOSPrivateWriteDirectory = ConvertToIOSPath(NormalizeFilename(Directory), true, false);
-	FString IOSPublicWriteDirectory = ConvertToIOSPath(NormalizeFilename(Directory), true, true);
+	FString IOSPrivateWriteDirectory = ConvertToPlatformPath(NormalizeFilename(Directory), true, false);
+	FString IOSPublicWriteDirectory = ConvertToPlatformPath(NormalizeFilename(Directory), true, true);
 	
 	// Try to delete the directory in both the private and public write paths
 	bool bDeletedPrivate = rmdir(TCHAR_TO_UTF8(*IOSPrivateWriteDirectory));
@@ -1023,13 +858,13 @@ bool FIOSPlatformFile::IterateDirectoryStat(const TCHAR* Directory, FDirectorySt
 		struct stat FileInfo;
 
 		// check the read path
-		if(stat(TCHAR_TO_UTF8(*ConvertToIOSPath(FullNormalizedPath, false, false)), &FileInfo) == -1)
+		if(stat(TCHAR_TO_UTF8(*ConvertToPlatformPath(FullNormalizedPath, false, false)), &FileInfo) == -1)
 		{
 			// if not in the read path, check the private write path
-			if(stat(TCHAR_TO_UTF8(*ConvertToIOSPath(FullNormalizedPath, true, false)), &FileInfo) == -1)
+			if(stat(TCHAR_TO_UTF8(*ConvertToPlatformPath(FullNormalizedPath, true, false)), &FileInfo) == -1)
 			{
 				// if not in the private write path, check the public write path
-				if(stat(TCHAR_TO_UTF8(*ConvertToIOSPath(FullNormalizedPath, true, true)), &FileInfo) == -1)
+				if(stat(TCHAR_TO_UTF8(*ConvertToPlatformPath(FullNormalizedPath, true, true)), &FileInfo) == -1)
 				{
 					return true;
 				}
@@ -1077,16 +912,16 @@ bool FIOSPlatformFile::IterateDirectoryCommon(const TCHAR* Directory, const TFun
 
 	FString NormalizedDirectory = NormalizeFilename(Directory);
 	// If Directory is an empty string, assume that we want to iterate Binaries/Mac (current dir), but because we're an app bundle, iterate bundle's Contents/Frameworks instead
-	DIR* Handle = opendir(Directory[0] ? TCHAR_TO_UTF8(*ConvertToIOSPath(NormalizedDirectory, false, false)) : FrameworksPath);
+	DIR* Handle = opendir(Directory[0] ? TCHAR_TO_UTF8(*ConvertToPlatformPath(NormalizedDirectory, false, false)) : FrameworksPath);
 	if(!Handle)
 	{
 		// look in the private write file path if it's not in the read file path
-		Handle = opendir(Directory[0] ? TCHAR_TO_UTF8(*ConvertToIOSPath(NormalizedDirectory, true, false)) : FrameworksPath);
+		Handle = opendir(Directory[0] ? TCHAR_TO_UTF8(*ConvertToPlatformPath(NormalizedDirectory, true, false)) : FrameworksPath);
 		
 		if(!Handle)
 		{
 			// look in the public write file path if it's not in the private write file path
-			Handle = opendir(Directory[0] ? TCHAR_TO_UTF8(*ConvertToIOSPath(NormalizedDirectory, true, true)) : FrameworksPath);
+			Handle = opendir(Directory[0] ? TCHAR_TO_UTF8(*ConvertToPlatformPath(NormalizedDirectory, true, true)) : FrameworksPath);
 		}
 	}
 	if (Handle)
@@ -1105,7 +940,7 @@ bool FIOSPlatformFile::IterateDirectoryCommon(const TCHAR* Directory, const TFun
 	return Result;
 }
 
-FString FIOSPlatformFile::ConvertToIOSPath(const FString& Filename, bool bForWrite, bool bIsPublicWrite)
+FString FIOSPlatformFile::ConvertToPlatformPath(const FString& Filename, bool bForWrite, bool bIsPublicWrite)
 {
 	FString Result = Filename;
 	if (Result.Contains(TEXT("/OnDemandResources/")) || Result.StartsWith(TEXT("/var/")))
diff --git a/Engine/Source/Runtime/Core/Public/Apple/ApplePlatformFile.h b/Engine/Source/Runtime/Core/Public/Apple/ApplePlatformFile.h
--- a/Engine/Source/Runtime/Core/Public/Apple/ApplePlatformFile.h
+++ b/Engine/Source/Runtime/Core/Public/Apple/ApplePlatformFile.h
@@ -41,6 +41,7 @@ public:
 
 	virtual IFileHandle* OpenRead(const TCHAR* Filename, bool bAllowWrite = false) override;
 	virtual IFileHandle* OpenWrite(const TCHAR* Filename, bool bAppend = false, bool bAllowRead = false) override;
+	virtual FOpenMappedResult OpenMappedEx(const TCHAR* Filename, EOpenReadFlags OpenOptions = EOpenReadFlags::None) override;
 	virtual bool DirectoryExists(const TCHAR* Directory) override;
 	virtual bool CreateDirectory(const TCHAR* Directory) override;
 	virtual bool DeleteDirectory(const TCHAR* Directory) override;
@@ -52,6 +53,7 @@ public:
 
 	virtual bool CopyFile(const TCHAR* To, const TCHAR* From, EPlatformFileRead ReadFlags = EPlatformFileRead::None, EPlatformFileWrite WriteFlags = EPlatformFileWrite::None) override;
 private:
+	FString ConvertToPlatformPath(const FString& Filename, bool bForWrite, bool bIsPublicWrite) { return Filename; }
 	bool IterateDirectoryCommon(const TCHAR* Directory, const TFunctionRef<bool(struct dirent*)>& Visitor);
 
 	int32 Stat(const TCHAR* Filename, struct stat* OutFileInfo);
diff --git a/Engine/Source/Runtime/Core/Public/IOS/IOSPlatformFile.h b/Engine/Source/Runtime/Core/Public/IOS/IOSPlatformFile.h
--- a/Engine/Source/Runtime/Core/Public/IOS/IOSPlatformFile.h
+++ b/Engine/Source/Runtime/Core/Public/IOS/IOSPlatformFile.h
@@ -42,7 +42,6 @@ public:
 
 	virtual IFileHandle* OpenRead(const TCHAR* Filename, bool bAllowWrite = false) override;
 	virtual IFileHandle* OpenWrite(const TCHAR* Filename, bool bAppend = false, bool bAllowRead = false) override;
-	virtual FOpenMappedResult OpenMappedEx(const TCHAR* Filename, EOpenReadFlags OpenOptions = EOpenReadFlags::None) override;
 
 	virtual bool DirectoryExists(const TCHAR* Directory) override;
 	virtual bool CreateDirectory(const TCHAR* Directory) override;
@@ -56,8 +55,7 @@ public:
 	
 private:
 	bool IterateDirectoryCommon(const TCHAR* Directory, const TFunctionRef<bool(struct dirent*)>& Visitor);
-
-	FString ConvertToIOSPath(const FString& Filename, bool bForWrite, bool bIsPublicWrite);
+	FString ConvertToPlatformPath(const FString& Filename, bool bForWrite, bool bIsPublicWrite);
 
 	FString ReadDirectory;
 	FString WriteDirectory;
