[Insights] Fixed TSAN warning re trace counters.

#rnx

diff --git a/Engine/Source/Runtime/Core/Private/ProfilingDebugging/TraceAuxiliary.cpp b/Engine/Source/Runtime/Core/Private/ProfilingDebugging/TraceAuxiliary.cpp
--- a/Engine/Source/Runtime/Core/Private/ProfilingDebugging/TraceAuxiliary.cpp
+++ b/Engine/Source/Runtime/Core/Private/ProfilingDebugging/TraceAuxiliary.cpp
@@ -114,19 +114,19 @@ DECLARE_MEMORY_STAT(TEXT("Memory Error"),                  STAT_Trace_MemoryErro
 #endif
 
 #if UE_TRACEAUX_EMIT_TRACE_COUNTERS
-TRACE_DECLARE_UNCHECKED_MEMORY_COUNTER(Trace_MemoryUsed,                 TEXT("Trace/MemoryUsed"));
-TRACE_DECLARE_UNCHECKED_MEMORY_COUNTER(Trace_Memory_BlockPool,           TEXT("Trace/Memory/BlockPool"));
-TRACE_DECLARE_UNCHECKED_MEMORY_COUNTER(Trace_Memory_FixedBuffers,        TEXT("Trace/Memory/FixedBuffers"));
-TRACE_DECLARE_UNCHECKED_MEMORY_COUNTER(Trace_Memory_SharedBuffers,       TEXT("Trace/Memory/SharedBuffers"));
-TRACE_DECLARE_UNCHECKED_MEMORY_COUNTER(Trace_Memory_Cache,               TEXT("Trace/Memory/Cache"));
-TRACE_DECLARE_UNCHECKED_MEMORY_COUNTER(Trace_Memory_Cache_Used,          TEXT("Trace/Memory/Cache/Used"));
-TRACE_DECLARE_UNCHECKED_MEMORY_COUNTER(Trace_Memory_Cache_Unused,        TEXT("Trace/Memory/Cache/Unused"));
-TRACE_DECLARE_UNCHECKED_MEMORY_COUNTER(Trace_Memory_Cache_Wasted,        TEXT("Trace/Memory/Cache/Wasted"));
-TRACE_DECLARE_UNCHECKED_MEMORY_COUNTER(Trace_Bandwidth_Emitted,          TEXT("Trace/Bandwidth/Emitted"));
-TRACE_DECLARE_UNCHECKED_MEMORY_COUNTER(Trace_Bandwidth_EmittedNotTraced, TEXT("Trace/Bandwidth/EmittedNotTraced"));
-TRACE_DECLARE_UNCHECKED_MEMORY_COUNTER(Trace_Bandwidth_Traced,           TEXT("Trace/Bandwidth/Traced"));
-TRACE_DECLARE_UNCHECKED_MEMORY_COUNTER(Trace_Bandwidth_Sent,             TEXT("Trace/Bandwidth/Sent"));
-TRACE_DECLARE_UNCHECKED_MEMORY_COUNTER(Trace_MemoryError,                TEXT("Trace/MemoryError"));
+TRACE_DECLARE_UNCHECKED_ATOMIC_MEMORY_COUNTER(Trace_MemoryUsed,                 TEXT("Trace/MemoryUsed"));
+TRACE_DECLARE_UNCHECKED_ATOMIC_MEMORY_COUNTER(Trace_Memory_BlockPool,           TEXT("Trace/Memory/BlockPool"));
+TRACE_DECLARE_UNCHECKED_ATOMIC_MEMORY_COUNTER(Trace_Memory_FixedBuffers,        TEXT("Trace/Memory/FixedBuffers"));
+TRACE_DECLARE_UNCHECKED_ATOMIC_MEMORY_COUNTER(Trace_Memory_SharedBuffers,       TEXT("Trace/Memory/SharedBuffers"));
+TRACE_DECLARE_UNCHECKED_ATOMIC_MEMORY_COUNTER(Trace_Memory_Cache,               TEXT("Trace/Memory/Cache"));
+TRACE_DECLARE_UNCHECKED_ATOMIC_MEMORY_COUNTER(Trace_Memory_Cache_Used,          TEXT("Trace/Memory/Cache/Used"));
+TRACE_DECLARE_UNCHECKED_ATOMIC_MEMORY_COUNTER(Trace_Memory_Cache_Unused,        TEXT("Trace/Memory/Cache/Unused"));
+TRACE_DECLARE_UNCHECKED_ATOMIC_MEMORY_COUNTER(Trace_Memory_Cache_Wasted,        TEXT("Trace/Memory/Cache/Wasted"));
+TRACE_DECLARE_UNCHECKED_ATOMIC_MEMORY_COUNTER(Trace_Bandwidth_Emitted,          TEXT("Trace/Bandwidth/Emitted"));
+TRACE_DECLARE_UNCHECKED_ATOMIC_MEMORY_COUNTER(Trace_Bandwidth_EmittedNotTraced, TEXT("Trace/Bandwidth/EmittedNotTraced"));
+TRACE_DECLARE_UNCHECKED_ATOMIC_MEMORY_COUNTER(Trace_Bandwidth_Traced,           TEXT("Trace/Bandwidth/Traced"));
+TRACE_DECLARE_UNCHECKED_ATOMIC_MEMORY_COUNTER(Trace_Bandwidth_Sent,             TEXT("Trace/Bandwidth/Sent"));
+TRACE_DECLARE_UNCHECKED_ATOMIC_MEMORY_COUNTER(Trace_MemoryError,                TEXT("Trace/MemoryError"));
 #endif
 
 #if UE_TRACEAUX_EMIT_LLM_STATS
