Allow DefaultDeviceProfiles.ini to define future iOS devices. There was an issue whereby new IOSDeviceMappings would be ignored when added application side to DefaultDeviceProfiles.ini. This was due to BaseDeviceProfiles.ini defining fallback devices for future proofing (iPhone*, iPad*, etc.) that would supersede any mapping defined app side. To avoid this, move fallback device mappings to its own IOSFallbackDeviceMappings section and handle separately.

[RN] Support addition of future iOS device profiles via game side DefaultDeviceProfiles.ini

diff --git a/Engine/Config/BaseDeviceProfiles.ini b/Engine/Config/BaseDeviceProfiles.ini
--- a/Engine/Config/BaseDeviceProfiles.ini
+++ b/Engine/Config/BaseDeviceProfiles.ini
@@ -296,6 +296,8 @@ AppleTV5,=AppleTV
 AppleTV6,=AppleTV4K
 AppleTV11,=AppleTV2_4K
 RealityDevice14,1=VisionPro
+
+[IOSFallbackDeviceMappings]
 ; Below are generic fallbacks for unrecognized devices. Add any specific devices above this section.
 iPhone=iPhone15
 iPod=iPodTouch7
@@ -303,6 +305,7 @@ iPad=iPad10
 AppleTV=AppleTV2_4K
 VisionPro0.*=VisionProSimulator
 VisionPro1.*=VisionPro
+..*=IOS
 
 [IOS DeviceProfile]
 DeviceType=IOS
diff --git a/Engine/Source/Runtime/Core/Private/IOS/IOSPlatformMisc.cpp b/Engine/Source/Runtime/Core/Private/IOS/IOSPlatformMisc.cpp
--- a/Engine/Source/Runtime/Core/Private/IOS/IOSPlatformMisc.cpp
+++ b/Engine/Source/Runtime/Core/Private/IOS/IOSPlatformMisc.cpp
@@ -427,40 +427,47 @@ FString GetIOSDeviceIDString()
 	return CachedResult;
 }
 
+static void MatchIOSDeviceMapping(const FString& DeviceIDString, const TCHAR* DeviceMappingSection, FString& IOSDeviceProfileName)
+{
+	TArray<FString> Mappings;
+	if (ensure(GConfig->GetSection(DeviceMappingSection, Mappings, GDeviceProfilesIni)))
+	{
+		for (const FString& MappingString : Mappings)
+		{
+			FString MappingRegex, ProfileName;
+			if (MappingString.Split(TEXT("="), &MappingRegex, &ProfileName))
+			{
+				const FRegexPattern RegexPattern(MappingRegex);
+				FRegexMatcher RegexMatcher(RegexPattern, *DeviceIDString);
+				if (RegexMatcher.FindNext())
+				{
+					FPlatformMisc::LowLevelOutputDebugStringf(TEXT("Matched %s as %s") LINE_TERMINATOR, *MappingRegex, *ProfileName);
+					IOSDeviceProfileName = ProfileName;
+					break;
+				}
+			}
+			else
+			{
+				FPlatformMisc::LowLevelOutputDebugStringf(TEXT("Invalid %s: %s") LINE_TERMINATOR, DeviceMappingSection, *MappingString);
+			}
+		}
+	}
+}
 
 const TCHAR* FIOSPlatformMisc::GetDefaultDeviceProfileName()
 {
 	static FString IOSDeviceProfileName;
 	if (IOSDeviceProfileName.Len() == 0)
-	{
-		IOSDeviceProfileName = TEXT("IOS");
-
+	{		
 		FString DeviceIDString = GetIOSDeviceIDString();
 		FPlatformMisc::LowLevelOutputDebugStringf(TEXT("Device Type: %s") LINE_TERMINATOR, *DeviceIDString);
-
-		TArray<FString> Mappings;
-		if (ensure(GConfig->GetSection(TEXT("IOSDeviceMappings"), Mappings, GDeviceProfilesIni)))
+		
+		MatchIOSDeviceMapping( DeviceIDString, TEXT("IOSDeviceMappings"), IOSDeviceProfileName );
+		if (IOSDeviceProfileName.Len() == 0)
 		{
-			for (const FString& MappingString : Mappings)
-			{
-				FString MappingRegex, ProfileName;
-				if (MappingString.Split(TEXT("="), &MappingRegex, &ProfileName))
-				{
-					const FRegexPattern RegexPattern(MappingRegex);
-					FRegexMatcher RegexMatcher(RegexPattern, *DeviceIDString);
-					if (RegexMatcher.FindNext())
-					{
-						FPlatformMisc::LowLevelOutputDebugStringf(TEXT("Matched %s as %s") LINE_TERMINATOR, *MappingRegex, *ProfileName);
-						IOSDeviceProfileName = ProfileName;
-						break;
-					}
-				}
-				else
-				{
-					FPlatformMisc::LowLevelOutputDebugStringf(TEXT("Invalid IOSDeviceMappings: %s") LINE_TERMINATOR, *MappingString);
-				}
-			}
+			MatchIOSDeviceMapping( DeviceIDString, TEXT("IOSFallbackDeviceMappings"), IOSDeviceProfileName );	
 		}
+		check(IOSDeviceProfileName.Len() > 0)
 	}
 
 	return *IOSDeviceProfileName;
