[AutoRTFM] Migrate FStartupMessages to transactionally-safe critical sections.

This code can be reached from Verse, although I don't have exact steps to do so.

diff --git a/Engine/Source/Runtime/Core/Private/Stats/StatsSystemTypes.cpp b/Engine/Source/Runtime/Core/Private/Stats/StatsSystemTypes.cpp
--- a/Engine/Source/Runtime/Core/Private/Stats/StatsSystemTypes.cpp
+++ b/Engine/Source/Runtime/Core/Private/Stats/StatsSystemTypes.cpp
@@ -247,7 +247,7 @@ void FStartupMessages::AddThreadMetadata(const FName InThreadName, uint32 InThre
 void FStartupMessages::AddMetadata(FName InStatName, const TCHAR* InStatDesc, const char* InGroupName, const char* InGroupCategory, const TCHAR* InGroupDesc, bool bShouldClearEveryFrame, EStatDataType::Type InStatType, bool bCycleStat, bool bSortByName, FPlatformMemory::EMemoryCounterRegion InMemoryRegion /*= FPlatformMemory::MCR_Invalid*/)
 {
 	LLM_SCOPE(ELLMTag::Stats);
-	FScopeLock Lock(&CriticalSection);
+	UE::TScopeLock Lock(CriticalSection);
 
 	DelayedMessages.Emplace(InGroupName, EStatDataType::ST_None, "Groups", InGroupCategory, InGroupDesc, false, false, bSortByName);
 	DelayedMessages.Emplace(InStatName, InStatType, InGroupName, InGroupCategory, InStatDesc, bShouldClearEveryFrame, bCycleStat, bSortByName, InMemoryRegion);
@@ -888,7 +888,7 @@ private:
 				SCOPE_CYCLE_COUNTER(STAT_StatsNewParseMeta);
 				TArray64<FStatMessage> MetaMessages;
 				{
-					FScopeLock Lock(&FStartupMessages::Get().CriticalSection);
+					UE::TScopeLock Lock(FStartupMessages::Get().CriticalSection);
 					Exchange(FStartupMessages::Get().DelayedMessages, MetaMessages);
 				}
 				if (MetaMessages.Num())
diff --git a/Engine/Source/Runtime/Core/Public/Stats/StatsSystemTypes.h b/Engine/Source/Runtime/Core/Public/Stats/StatsSystemTypes.h
--- a/Engine/Source/Runtime/Core/Public/Stats/StatsSystemTypes.h
+++ b/Engine/Source/Runtime/Core/Public/Stats/StatsSystemTypes.h
@@ -14,7 +14,6 @@
 #include "CoreGlobals.h"
 #include "CoreTypes.h"
 #include "Delegates/Delegate.h"
-#include "HAL/CriticalSection.h"
 #include "HAL/LowLevelMemTracker.h"
 #include "HAL/PlatformCrt.h"
 #include "HAL/PlatformMemory.h"
@@ -29,6 +28,7 @@
 #include "Misc/CString.h"
 #include "Misc/EnumClassFlags.h"
 #include "Misc/SourceLocation.h"
+#include "Misc/TransactionallySafeCriticalSection.h"
 #include "ProfilingDebugging/CpuProfilerTrace.h"
 #include "ProfilingDebugging/MiscTrace.h"
 #include "Stats/StatsCommon.h"
@@ -1688,7 +1688,7 @@ class FStartupMessages
 	friend class FStatsThread;
 
 	TArray64<FStatMessage> DelayedMessages;
-	FCriticalSection CriticalSection;
+	FTransactionallySafeCriticalSection CriticalSection;
 
 public:
 	/** Adds a thread metadata. */
