Fixing up the issues found with binary config:
- Added support for -ini: param
- Disabled on Editor/Server
- Re-enabled for Androir

diff --git a/Engine/Source/Runtime/Core/Core.Build.cs b/Engine/Source/Runtime/Core/Core.Build.cs
--- a/Engine/Source/Runtime/Core/Core.Build.cs
+++ b/Engine/Source/Runtime/Core/Core.Build.cs
@@ -298,7 +298,7 @@ public class Core : ModuleRules
 
 	protected virtual bool SupportsBinaryConfig(ReadOnlyTargetRules Target)
 	{
-		return Target.Platform != UnrealTargetPlatform.Android;
+		return true;// Target.Platform != UnrealTargetPlatform.Android;
 	}
 
 	// Decide if validating memory allocator (aka MallocStomp) can be used on the current plVatform.
diff --git a/Engine/Source/Runtime/Core/Private/Misc/ConfigCacheIni.cpp b/Engine/Source/Runtime/Core/Private/Misc/ConfigCacheIni.cpp
--- a/Engine/Source/Runtime/Core/Private/Misc/ConfigCacheIni.cpp
+++ b/Engine/Source/Runtime/Core/Private/Misc/ConfigCacheIni.cpp
@@ -5625,6 +5625,13 @@ bool FConfigCacheIni::CreateGConfigFromSaved(const TCHAR* Filename)
 	for (FConfigBranch& Branch : GConfig->KnownFiles.Branches)
 	{
 		Branch.IniPath = FConfigCacheIni::GetDestIniFilename(*Branch.IniName.ToString(), nullptr, *FPaths::GeneratedConfigDir());
+
+#if ALLOW_INI_OVERRIDE_FROM_COMMANDLINE
+		// process any commandline overrides
+		FConfigFile::OverrideFromCommandline(&Branch.CommandLineOverrides, Branch.IniName.ToString());
+		// and push it into the current values
+		Branch.InMemoryFile.ApplyFile(&Branch.CommandLineOverrides);
+#endif
 	}
 	
 	MemoryReader << ExtraData.Data;
@@ -5741,7 +5748,7 @@ void FConfigCacheIni::InitializeConfigSystem()
 	
 	InitializeConfigRemap();
 	
-#if PLATFORM_SUPPORTS_BINARYCONFIG && PRELOAD_BINARY_CONFIG && !WITH_EDITOR
+#if PLATFORM_SUPPORTS_BINARYCONFIG && PRELOAD_BINARY_CONFIG// && !WITH_EDITOR && !WITH_SERVER_CODE
 	// attempt to load from staged binary config data
 	if (!FParse::Param(FCommandLine::Get(), TEXT("textconfig")) &&
 		FConfigCacheIni::CreateGConfigFromSaved(nullptr))
@@ -5752,69 +5759,70 @@ void FConfigCacheIni::InitializeConfigSystem()
 		Context.Load(TEXT("GameUserSettings"), GGameUserSettingsIni);
 
 		LoadRemainingConfigFiles(Context);
-		return;
 	}
+	else
 #endif
-
-	// Bootstrap the Ini config cache
-	FString IniBootstrapFilename;
-	if (FParse::Value( FCommandLine::Get(), TEXT("IniBootstrap="), IniBootstrapFilename))
 	{
-		TRACE_CPUPROFILER_EVENT_SCOPE(IniBootstrap);
-		TArray<uint8> FileContent;
-		if (FFileHelper::LoadFileToArray(FileContent, *IniBootstrapFilename, FILEREAD_Silent))
-		{
-			FMemoryReader MemoryReader(FileContent, true);
-			GConfig = new FConfigCacheIni(EConfigCacheType::Temporary, NAME_None /*Platform*/, true /* bInGloballyRegistered */);
-			GConfig->SerializeStateForBootstrap_Impl(MemoryReader);
-			GConfig->bIsReadyForUse = true;
-			TRACE_CPUPROFILER_EVENT_SCOPE(ConfigReadyForUseBroadcast);
-			FCoreDelegates::TSConfigReadyForUse().Broadcast();
-			return;
-		}
-		else
+		// Bootstrap the Ini config cache
+		FString IniBootstrapFilename;
+		if (FParse::Value(FCommandLine::Get(), TEXT("IniBootstrap="), IniBootstrapFilename))
 		{
-			FPlatformMisc::LowLevelOutputDebugStringf(TEXT("Unable to bootstrap from archive %s, will fallback on normal initialization\n"), *IniBootstrapFilename);
+			TRACE_CPUPROFILER_EVENT_SCOPE(IniBootstrap);
+			TArray<uint8> FileContent;
+			if (FFileHelper::LoadFileToArray(FileContent, *IniBootstrapFilename, FILEREAD_Silent))
+			{
+				FMemoryReader MemoryReader(FileContent, true);
+				GConfig = new FConfigCacheIni(EConfigCacheType::Temporary, NAME_None /*Platform*/, true /* bInGloballyRegistered */);
+				GConfig->SerializeStateForBootstrap_Impl(MemoryReader);
+				GConfig->bIsReadyForUse = true;
+				TRACE_CPUPROFILER_EVENT_SCOPE(ConfigReadyForUseBroadcast);
+				FCoreDelegates::TSConfigReadyForUse().Broadcast();
+				return;
+			}
+			else
+			{
+				FPlatformMisc::LowLevelOutputDebugStringf(TEXT("Unable to bootstrap from archive %s, will fallback on normal initialization\n"), *IniBootstrapFilename);
+			}
 		}
-	}
 
-	// Perform any upgrade we need before we load any configuration files
-	FConfigManifest::UpgradeFromPreviousVersions();
+		// Perform any upgrade we need before we load any configuration files
+		FConfigManifest::UpgradeFromPreviousVersions();
 
-	// create GConfig
-	GConfig = new FConfigCacheIni(EConfigCacheType::DiskBacked, FPlatformProperties::IniPlatformName(), true /* bInGloballyRegistered */);
+		// create GConfig
+		GConfig = new FConfigCacheIni(EConfigCacheType::DiskBacked, FPlatformProperties::IniPlatformName(), true /* bInGloballyRegistered */);
 
-	// create a context object that we will use for all of the main ini files
-	FConfigContext Context = FConfigContext::ReadIntoGConfig();
+		// create a context object that we will use for all of the main ini files
+		FConfigContext Context = FConfigContext::ReadIntoGConfig();
 
-	// load in the default ini files
-	bool bEngineConfigCreated = InitializeKnownConfigFiles(Context);
+		// load in the default ini files
+		bool bEngineConfigCreated = InitializeKnownConfigFiles(Context);
 
-	// verify if needed
-	const bool bIsGamelessExe = !FApp::HasProjectName();
-	if ( !bIsGamelessExe )
-	{
-		// Now check and see if our game is correct if this is a game agnostic binary
-		if (GIsGameAgnosticExe && !bEngineConfigCreated)
+		// verify if needed
+		const bool bIsGamelessExe = !FApp::HasProjectName();
+		if (!bIsGamelessExe)
 		{
-			const FText AbsolutePath = FText::FromString( IFileManager::Get().ConvertToAbsolutePathForExternalAppForRead(*FPaths::GetPath(GEngineIni)) );
-			//@todo this is too early to localize
-			const FText Message = FText::Format( NSLOCTEXT("Core", "FirstCmdArgMustBeGameName", "'{0}' must exist and contain a DefaultEngine.ini."), AbsolutePath );
-			if (!GIsBuildMachine)
+			// Now check and see if our game is correct if this is a game agnostic binary
+			if (GIsGameAgnosticExe && !bEngineConfigCreated)
 			{
-				FMessageDialog::Open(EAppMsgType::Ok, Message);
-			}
-			FApp::SetProjectName(TEXT("")); // this disables part of the crash reporter to avoid writing log files to a bogus directory
-			if (!GIsBuildMachine)
-			{
-				exit(1);
+				const FText AbsolutePath = FText::FromString(IFileManager::Get().ConvertToAbsolutePathForExternalAppForRead(*FPaths::GetPath(GEngineIni)));
+				//@todo this is too early to localize
+				const FText Message = FText::Format(NSLOCTEXT("Core", "FirstCmdArgMustBeGameName", "'{0}' must exist and contain a DefaultEngine.ini."), AbsolutePath);
+				if (!GIsBuildMachine)
+				{
+					FMessageDialog::Open(EAppMsgType::Ok, Message);
+				}
+				FApp::SetProjectName(TEXT("")); // this disables part of the crash reporter to avoid writing log files to a bogus directory
+				if (!GIsBuildMachine)
+				{
+					exit(1);
+				}
+				UE_LOG(LogInit, Fatal, TEXT("%s"), *Message.ToString());
 			}
-			UE_LOG(LogInit, Fatal,TEXT("%s"), *Message.ToString());
 		}
-	}
 
-	// load editor, etc config files
-	LoadRemainingConfigFiles(Context);
+		// load editor, etc config files
+		LoadRemainingConfigFiles(Context);
+	}
 
 	FCoreDelegates::TSOnConfigSectionsChanged().AddStatic(OnConfigSectionsChanged);
 
