Fixed LLM Untracked stats for platforms that use MB2 with VLPA as it caches freed memory internally and some of that memory can't be returned to kernel immediatelly so it ends up inflating Untracked stat

diff --git a/Engine/Source/Runtime/Core/Private/HAL/LowLevelMemTracker.cpp b/Engine/Source/Runtime/Core/Private/HAL/LowLevelMemTracker.cpp
--- a/Engine/Source/Runtime/Core/Private/HAL/LowLevelMemTracker.cpp
+++ b/Engine/Source/Runtime/Core/Private/HAL/LowLevelMemTracker.cpp
@@ -2023,7 +2023,10 @@ void FLowLevelMemTracker::TickInternal()
 	// remove the MemoryUsageCurrentOverhead from the "Total" for the default LLM as it's not something anyone needs
 	// to investigate when finding what to reduce the platform LLM will have the info 
 	const int64 DefaultProcessMemory = PlatformProcessMemory - (!FPlatformMemory::TracksLLMAllocations() ? MemoryUsageCurrentOverhead : 0);
-	const int64 DefaultUntracked = FMath::Max<int64>(0, DefaultProcessMemory - TrackedTotal);
+	const int64 ImmediatelyFreeableCachedMemory = GMalloc ? GMalloc->GetImmediatelyFreeableCachedMemorySize() : 0;
+	// we need to subtract diff between total cached free mem and immediately freeable cached memory as it is not accounted anywhere in the LLM and would artificially inflate Untracked.
+	// while any new allocations from that cached, but not immediately freeable memory would be properly accounted by the LLM
+	const int64 DefaultUntracked = FMath::Max<int64>(0, DefaultProcessMemory - TrackedTotal - CachedFreeMemory + ImmediatelyFreeableCachedMemory);
 	DefaultTracker.SetTagAmountInUpdate(FindOrAddTagData(ELLMTag::Total), DefaultProcessMemory, false);
 	DefaultTracker.SetTagAmountInUpdate(FindOrAddTagData(ELLMTag::Untracked), DefaultUntracked, false);
 
