Reenable apply ini if not Editor or Server that was disabled in 37977614

diff --git a/Engine/Source/Runtime/Core/Private/Misc/ConfigCacheIni.cpp b/Engine/Source/Runtime/Core/Private/Misc/ConfigCacheIni.cpp
--- a/Engine/Source/Runtime/Core/Private/Misc/ConfigCacheIni.cpp
+++ b/Engine/Source/Runtime/Core/Private/Misc/ConfigCacheIni.cpp
@@ -5748,7 +5748,7 @@ void FConfigCacheIni::InitializeConfigSystem()
 	
 	InitializeConfigRemap();
 	
-#if PLATFORM_SUPPORTS_BINARYCONFIG && PRELOAD_BINARY_CONFIG// && !WITH_EDITOR && !WITH_SERVER_CODE
+#if PLATFORM_SUPPORTS_BINARYCONFIG && PRELOAD_BINARY_CONFIG && !WITH_EDITOR && !WITH_SERVER_CODE
 	// attempt to load from staged binary config data
 	if (!FParse::Param(FCommandLine::Get(), TEXT("textconfig")) &&
 		FConfigCacheIni::CreateGConfigFromSaved(nullptr))
