MB2/3 minor refactoring mostly to improve code and propagate values into functions, but it also resulted in minor perf improvement

diff --git a/Engine/Source/Runtime/Core/Private/HAL/MallocBinned2.cpp b/Engine/Source/Runtime/Core/Private/HAL/MallocBinned2.cpp
--- a/Engine/Source/Runtime/Core/Private/HAL/MallocBinned2.cpp
+++ b/Engine/Source/Runtime/Core/Private/HAL/MallocBinned2.cpp
@@ -31,6 +31,8 @@ PRAGMA_DISABLE_UNSAFE_TYPECAST_WARNINGS
 #	define GBinned2MoveOSFreesOffTimeCriticalThreads	UE_DEFAULT_GBinned2MoveOSFreesOffTimeCriticalThreads
 #endif
 
+namespace
+{
 #if UE_MB2_ALLOCATOR_STATS
 	std::atomic<int64> AllocatedOSSmallPoolMemory(0);
 
@@ -49,7 +51,7 @@ PRAGMA_DISABLE_UNSAFE_TYPECAST_WARNINGS
 // Bin sizes are based around getting the maximum amount of allocations per pool, with as little alignment waste as possible.
 // Bin sizes should be close to even divisors of the system page size, and well distributed.
 // They must be 16-byte aligned as well.
-static constexpr uint16 SmallBinSizes[] =
+constexpr uint16 SmallBinSizesInternal[] =
 {
 	16, 32, 48, 64, 80, 96, 112, 128, 144, 160, 176, 192, 208, // +16
 //	Bin		|	Divider | Slack leftover per page (64KB)
@@ -95,9 +97,11 @@ static constexpr uint16 SmallBinSizes[] =
 #endif
 };
 
-MS_ALIGN(PLATFORM_CACHE_LINE_SIZE) static uint8 UnusedAlignPadding[PLATFORM_CACHE_LINE_SIZE] GCC_ALIGN(PLATFORM_CACHE_LINE_SIZE) = { 0 };
-uint16 FMallocBinned2::SmallBinSizesReversed[UE_MB2_SMALL_POOL_COUNT] = { 0 };
-uint32 FMallocBinned2::PageSize = 0;
+static uint32 GMB2PageSize = 0;
+
+} //~nameless namespace
+
+uint16 FMallocBinned2::SmallBinSizes[UE_MB2_SMALL_POOL_COUNT] = { 0 };
 FMallocBinned2* FMallocBinned2::MallocBinned2 = nullptr;
 // Mapping of sizes to small table indices
 uint8 FMallocBinned2::MemSizeToPoolIndex[1 + (UE_MB2_MAX_SMALL_POOL_SIZE >> UE_MBC_MIN_SMALL_POOL_ALIGNMENT_SHIFT)] = { 0 };
@@ -275,7 +279,7 @@ struct FMallocBinned2::Private
 				Free->BinSize = InBinSize;
 				Free->CanaryAndForkState = Allocator.CurrentCanary;
 				Free->PoolIndex = InPoolIndex;
-				NodePool->FirstFreeBlock   = Free;
+				NodePool->FirstFreeBlock = Free;
 
 				// Free this pool.
 				check(NodePool->Taken >= 1);
@@ -288,10 +292,10 @@ struct FMallocBinned2::Private
 					NodePool->Unlink();
 					{
 						FScopeLock InnerLock(&Allocator.ExternalAllocMutex);
-						Allocator.CachedOSPageAllocator.Free(BasePtrOfNode, Allocator.PageSize);
+						Allocator.CachedOSPageAllocator.Free(BasePtrOfNode, GMB2PageSize);
 					}
 #if UE_MB2_ALLOCATOR_STATS
-					AllocatedOSSmallPoolMemory.fetch_sub((int64)Allocator.PageSize, std::memory_order_relaxed);
+					AllocatedOSSmallPoolMemory.fetch_sub((int64)GMB2PageSize, std::memory_order_relaxed);
 #endif
 				}
 
@@ -302,9 +306,9 @@ struct FMallocBinned2::Private
 		}
 	}
 
+#if BINNED2_FORK_SUPPORT
 	static void CheckThreadFreeBlockListsForFork()
 	{
-#if BINNED2_FORK_SUPPORT
 		if(GMallocBinnedPerThreadCaches)
 		{
 			FScopeLock Lock(&GetFreeBlockListsRegistrationMutex());
@@ -312,8 +316,8 @@ struct FMallocBinned2::Private
 			UE_CLOG(List.Num() == 1 && List[0] != FPerThreadFreeBlockLists::Get(), LogMemory, Fatal, TEXT("There was a thread-local free list at fork time which did not belong to the main forking thread. No other threads should be alive at fork time. If threads are spawned before forking, they must be killed and FMallocBinned2::ClearAndDisableTLSCachesOnCurrentThread() must be called."));
 			UE_CLOG(List.Num() > 1, LogMemory, Fatal, TEXT("There were multiple thread-local free lists at fork time. No other threads should be alive at fork time. If threads are spawned before forking, they must be killed and FMallocBinned2::ClearAndDisableTLSCachesOnCurrentThread() must be called."));
 		}
-#endif
 	}
+#endif
 };
 
 MallocBinnedPrivate::TGlobalRecycler<UE_MB2_SMALL_POOL_COUNT> FMallocBinned2::Private::GGlobalRecycler;
@@ -348,36 +352,33 @@ void FMallocBinned2::FPoolList::LinkToFront(FPoolInfo* Pool)
 
 FMallocBinned2::FPoolInfo& FMallocBinned2::FPoolList::PushNewPoolToFront(FMallocBinned2& Allocator, uint32 InBinSize, uint32 InPoolIndex)
 {
-	const uint32 LocalPageSize = Allocator.PageSize;
-
-	// Allocate memory.
 	LLM_PLATFORM_SCOPE(ELLMTag::FMalloc);
 
 	void* FreePtr;
 	{
 		FScopeLock Lock(&Allocator.ExternalAllocMutex);
-		FreePtr = Allocator.CachedOSPageAllocator.Allocate(LocalPageSize, FMemory::AllocationHints::SmallPool);
+		FreePtr = Allocator.CachedOSPageAllocator.Allocate(GMB2PageSize, FMemory::AllocationHints::SmallPool);
 	}
 	if (!FreePtr)
 	{
-		OutOfMemory(LocalPageSize);
+		OutOfMemory(GMB2PageSize);
 	}
 
 #if !UE_USE_VERYLARGEPAGEALLOCATOR || !UE_MB2_BOOKKEEPING_AT_THE_END_OF_LARGEBLOCK
-	FFreeBlock* Free = new (FreePtr) FFreeBlock(LocalPageSize, InBinSize, InPoolIndex, Allocator.CurrentCanary);
-	check(IsAligned(Free, LocalPageSize));
+	FFreeBlock* Free = new (FreePtr) FFreeBlock(GMB2PageSize, InBinSize, InPoolIndex, Allocator.CurrentCanary);
+	check(IsAligned(Free, GMB2PageSize));
 #else
 	FFreeBlock* FreeBlockPtr = GetPoolHeaderFromPointer(FreePtr);
-	FFreeBlock* Free = new (FreeBlockPtr) FFreeBlock(LocalPageSize, InBinSize, InPoolIndex, Allocator.CurrentCanary);
+	FFreeBlock* Free = new (FreeBlockPtr) FFreeBlock(GMB2PageSize, InBinSize, InPoolIndex, Allocator.CurrentCanary);
 #endif
 
 #if UE_MB2_ALLOCATOR_STATS
-	AllocatedOSSmallPoolMemory.fetch_add((int64)LocalPageSize, std::memory_order_relaxed);
+	AllocatedOSSmallPoolMemory.fetch_add((int64)GMB2PageSize, std::memory_order_relaxed);
 #endif
 
 	FScopeLock Lock(&Allocator.ExternalAllocMutex);
 	// Create pool
-	FPoolInfo* Result = Internal::GetOrCreatePoolInfo(Allocator, Free, FPoolInfo::ECanary::FirstFreeBlockIsPtr, false);
+	FPoolInfo* Result = Internal::GetOrCreatePoolInfo(Allocator, Free, FPoolInfo::ECanary::FirstFreeBlockIsPtr);
 	Result->Link(Front);
 	Result->Taken          = 0;
 	Result->FirstFreeBlock = Free;
@@ -390,36 +391,33 @@ FMallocBinned2::FMallocBinned2()
 	static bool bOnce = false;
 	check(!bOnce); // this is now a singleton-like thing and you cannot make multiple copies
 	bOnce = true;
+
+	FMemory::Memcpy(SmallBinSizes, SmallBinSizesInternal);
 	
-	static_assert(sizeof(FFreeBlock) <= SmallBinSizes[0], "sizeof(FFreeBlock)  must be fit in smallest allocation size handled by FMallocBinned2.");
-	static_assert(sizeof(FMallocBinned2::FBundleNode) <= SmallBinSizes[0], "sizeof(FBundleNode) must fit in smallest allocation size handled by FMallocBinned2.");
+	static_assert(sizeof(FFreeBlock) <= SmallBinSizesInternal[0], "sizeof(FFreeBlock)  must be fit in smallest allocation size handled by FMallocBinned2.");
+	static_assert(sizeof(FMallocBinned2::FBundleNode) <= SmallBinSizesInternal[0], "sizeof(FBundleNode) must fit in smallest allocation size handled by FMallocBinned2.");
 
-	for (uint32 Index = 0; Index != UE_MB2_SMALL_POOL_COUNT; ++Index)
-	{
-		const uint32 Partner = UE_MB2_SMALL_POOL_COUNT - Index - 1;
-		SmallBinSizesReversed[Index] = SmallBinSizes[Partner];
-	}
 	const FGenericPlatformMemoryConstants Constants = FPlatformMemory::GetConstants();
-	PageSize = Constants.BinnedPageSize;
-	OsAllocationGranularity = Constants.BinnedAllocationGranularity ? Constants.BinnedAllocationGranularity : PageSize;
-	NumPoolsPerPage = PageSize / sizeof(FPoolInfo);
-	PtrToPoolMapping.Init(PageSize, NumPoolsPerPage, Constants.AddressStart, Constants.AddressLimit);
-
-	checkf(FMath::IsPowerOfTwo(PageSize), TEXT("OS page size must be a power of two"));
-	checkf(Constants.AddressLimit > PageSize, TEXT("OS address limit must be greater than the page size")); // Check to catch 32 bit overflow in AddressLimit
-	checkf(PageSize % UE_MB2_LARGE_ALLOC == 0, TEXT("OS page size must be a multiple of UE_MB2_LARGE_ALLOC"));
-
-	static_assert(SmallBinSizes[UE_MB2_SMALL_POOL_COUNT - 1] == UE_MB2_MAX_SMALL_POOL_SIZE, "UE_MB2_MAX_SMALL_POOL_SIZE must equal the largest bin size");
-	static_assert(sizeof(FMallocBinned2::FFreeBlock) <= SmallBinSizes[0], "Pool header must be able to fit into the smallest bin");
-	static_assert(UE_ARRAY_COUNT(SmallBinSizes) == UE_MB2_SMALL_POOL_COUNT, "Number of bins in SmallBinSizes must match UE_MB2_SMALL_POOL_COUNT");
-	static_assert(UE_ARRAY_COUNT(SmallBinSizes) <= 256, "Number of bins in SmallBinSizes fit in one byte");
+	GMB2PageSize = Constants.BinnedPageSize;
+	OsAllocationGranularity = Constants.BinnedAllocationGranularity ? Constants.BinnedAllocationGranularity : GMB2PageSize;
+	NumPoolsPerPage = GMB2PageSize / sizeof(FPoolInfo);
+	PtrToPoolMapping.Init(GMB2PageSize, NumPoolsPerPage, Constants.AddressStart, Constants.AddressLimit);
+
+	checkf(FMath::IsPowerOfTwo(GMB2PageSize), TEXT("OS page size must be a power of two"));
+	checkf(Constants.AddressLimit > GMB2PageSize, TEXT("OS address limit must be greater than the page size")); // Check to catch 32 bit overflow in AddressLimit
+	checkf(GMB2PageSize % UE_MB2_LARGE_ALLOC == 0, TEXT("OS page size must be a multiple of UE_MB2_LARGE_ALLOC"));
+
+	static_assert(SmallBinSizesInternal[UE_MB2_SMALL_POOL_COUNT - 1] == UE_MB2_MAX_SMALL_POOL_SIZE, "UE_MB2_MAX_SMALL_POOL_SIZE must equal the largest bin size");
+	static_assert(sizeof(FMallocBinned2::FFreeBlock) <= SmallBinSizesInternal[0], "Pool header must be able to fit into the smallest bin");
+	static_assert(UE_ARRAY_COUNT(SmallBinSizesInternal) == UE_MB2_SMALL_POOL_COUNT, "Number of bins in SmallBinSizesInternal must match UE_MB2_SMALL_POOL_COUNT");
+	static_assert(UE_ARRAY_COUNT(SmallBinSizesInternal) <= 256, "Number of bins in SmallBinSizesInternal fit in one byte");
 	static_assert(sizeof(FFreeBlock) <= UE_MBC_MIN_SMALL_POOL_ALIGNMENT, "Free block struct must be small enough to fit into the smallest bin.");
 
 	// Init pool tables.
 	for (uint32 Index = 0; Index != UE_MB2_SMALL_POOL_COUNT; ++Index)
 	{
 		checkf(Index == 0 || SmallBinSizes[Index - 1] < SmallBinSizes[Index], TEXT("Small bin sizes must be strictly increasing"));
-		checkf(SmallBinSizes[Index] <= PageSize, TEXT("Small bin size must be small enough to fit into a page"));
+		checkf(SmallBinSizes[Index] <= GMB2PageSize, TEXT("Small bin size must be small enough to fit into a page"));
 		checkf(SmallBinSizes[Index] % UE_MBC_MIN_SMALL_POOL_ALIGNMENT == 0, TEXT("Small bin size must be a multiple of UE_MBC_MINIMUM_ALIGNMENT"));
 
 		SmallPoolTables[Index].BinSize = SmallBinSizes[Index];
@@ -440,13 +438,6 @@ FMallocBinned2::FMallocBinned2()
 		*IndexEntry++ = uint8(PoolIndex);
 	}
 
-	// now reverse the pool sizes for cache coherency
-	for (uint32 Index = 0; Index != UE_MB2_SMALL_POOL_COUNT; ++Index)
-	{
-		const uint32 Partner = UE_MB2_SMALL_POOL_COUNT - Index - 1;
-		SmallBinSizesReversed[Index] = SmallBinSizes[Partner];
-	}
-
 	AllocateHashBuckets();
 
 	MallocBinned2 = this;
@@ -620,7 +611,7 @@ void* FMallocBinned2::MallocExternalLarge(SIZE_T Size, uint32 Alignment)
 	Size = Align(FMath::Max((SIZE_T)1, Size), Alignment);
 
 	checkf(FMath::IsPowerOfTwo(Alignment), TEXT("Invalid Malloc alignment: '%u' is not a power of two"), Alignment);
-	checkf(Alignment <= PageSize, TEXT("Invalid Malloc alignment: '%u' is greater than the page size '%u'"), Alignment, PageSize);
+	checkf(Alignment <= GMB2PageSize, TEXT("Invalid Malloc alignment: '%u' is greater than the page size '%u'"), Alignment, GMB2PageSize);
 	checkf(IsSupportedSize(Size), TEXT("Invalid Malloc size: '%" SIZE_T_FMT "'"), Size);
 
 	const SIZE_T AlignedSize = Align(Size, OsAllocationGranularity);
@@ -642,11 +633,11 @@ void* FMallocBinned2::MallocExternalLarge(SIZE_T Size, uint32 Alignment)
 			OutOfMemory(AlignedSize);
 		}
 
-		Pool = Internal::GetOrCreatePoolInfo(*this, Result, FPoolInfo::ECanary::FirstFreeBlockIsOSAllocSize, false);
+		Pool = Internal::GetOrCreatePoolInfo(*this, Result, FPoolInfo::ECanary::FirstFreeBlockIsOSAllocSize);
 	}
 
 	UE_CLOG(!IsAligned(Result, Alignment) ,LogMemory, Fatal, TEXT("FMallocBinned2 alignment was too large for OS. Alignment=%d   Ptr=%p"), Alignment, Result);
-	check(IsAligned(Result, PageSize) && IsOSAllocation(Result));
+	check(IsAligned(Result, GMB2PageSize) && IsOSAllocation(Result));
 
 #if UE_MB2_ALLOCATOR_STATS
 	AllocatedLargePoolMemory.fetch_add(Size, std::memory_order_relaxed);
@@ -748,7 +739,7 @@ void* FMallocBinned2::ReallocExternal(void* Ptr, SIZE_T NewSize, uint32 Alignmen
 	NewSize = Align(FMath::Max((SIZE_T)1, NewSize), Alignment);
 
 	checkf(FMath::IsPowerOfTwo(Alignment), TEXT("Invalid Realloc alignment: '%u' is not a power of two"), Alignment);
-	checkf(Alignment <= PageSize, TEXT("Invalid Realloc alignment: '%u' is greater than the page size '%u'"), Alignment, PageSize);
+	checkf(Alignment <= GMB2PageSize, TEXT("Invalid Realloc alignment: '%u' is greater than the page size '%u'"), Alignment, GMB2PageSize);
 	checkf(IsSupportedSize(NewSize), TEXT("Invalid Realloc size: '%" SIZE_T_FMT "'"), NewSize);
 
 #if UE_MB2_ALLOCATOR_STATS
@@ -943,7 +934,7 @@ void FMallocBinned2::DumpAllocatorStats(class FOutputDevice& Ar)
 	const uint64 OSPageAllocatorCachedFreeSize           = CachedOSPageAllocator.GetCachedFreeTotal();
 
 	Ar.Logf(TEXT("FMallocBinned2 Mem report"));
-	Ar.Logf(TEXT("Constants.BinnedPageSize = %d"), int32(PageSize));
+	Ar.Logf(TEXT("Constants.BinnedPageSize = %d"), int32(GMB2PageSize));
 	Ar.Logf(TEXT("Constants.BinnedAllocationGranularity = %d"), int32(OsAllocationGranularity));
 	Ar.Logf(TEXT("Small Pool Allocations: %fmb  (including bin size padding)"), ((double)TotalAllocatedSmallPoolMemory) / (1024.0f * 1024.0f));
 	Ar.Logf(TEXT("Small Pool OS Allocated: %fmb"), ((double)AllocatedOSSmallPoolMemory.load(std::memory_order_relaxed)) / (1024.0f * 1024.0f));
diff --git a/Engine/Source/Runtime/Core/Private/HAL/MallocBinned3.cpp b/Engine/Source/Runtime/Core/Private/HAL/MallocBinned3.cpp
--- a/Engine/Source/Runtime/Core/Private/HAL/MallocBinned3.cpp
+++ b/Engine/Source/Runtime/Core/Private/HAL/MallocBinned3.cpp
@@ -4,7 +4,7 @@
 
 PRAGMA_DISABLE_UNSAFE_TYPECAST_WARNINGS
 
-#if PLATFORM_64BITS && PLATFORM_HAS_FPlatformVirtualMemoryBlock
+#if PLATFORM_HAS_FPlatformVirtualMemoryBlock
 #include "Misc/ScopeLock.h"
 #include "Templates/Function.h"
 #include "GenericPlatform/GenericPlatformProcess.h"
@@ -16,18 +16,17 @@ PRAGMA_DISABLE_UNSAFE_TYPECAST_WARNINGS
 
 
 #if UE_MB3_USE_CACHED_PAGE_ALLOCATOR_FOR_LARGE_ALLOCS
-#include "HAL/Allocators/CachedOSPageAllocator.h"
-#define UE_MB3_MAX_CACHED_OS_FREES (64)
-#define UE_MB3_MAX_CACHED_OS_FREES_BYTE_LIMIT (64*1024*1024)
+#	include "HAL/Allocators/CachedOSPageAllocator.h"
+#	define UE_MB3_MAX_CACHED_OS_FREES (64)
+#	define UE_MB3_MAX_CACHED_OS_FREES_BYTE_LIMIT (64*1024*1024)
 
-typedef TCachedOSPageAllocator<UE_MB3_MAX_CACHED_OS_FREES, UE_MB3_MAX_CACHED_OS_FREES_BYTE_LIMIT> TBinned3CachedOSPageAllocator;
-
-TBinned3CachedOSPageAllocator& GetCachedOSPageAllocator()
-{
-	static TBinned3CachedOSPageAllocator Singleton;
-	return Singleton;
-}
+	typedef TCachedOSPageAllocator<UE_MB3_MAX_CACHED_OS_FREES, UE_MB3_MAX_CACHED_OS_FREES_BYTE_LIMIT> TBinned3CachedOSPageAllocator;
 
+	TBinned3CachedOSPageAllocator& GetCachedOSPageAllocator()
+	{
+		static TBinned3CachedOSPageAllocator Singleton;
+		return Singleton;
+	}
 #endif
 
 #if UE_MB3_ALLOCATOR_STATS
@@ -54,7 +53,7 @@ TBinned3CachedOSPageAllocator& GetCachedOSPageAllocator()
 	TAtomic<int32> MemoryRangeFreeTotalCount(0);
 #endif
 
-uint16 FMallocBinned3::SmallBinSizesReversedShifted[UE_MB3_SMALL_POOL_COUNT + 1] = { 0 };
+uint16 FMallocBinned3::SmallBinSizesShifted[UE_MB3_SMALL_POOL_COUNT + 1] = { 0 };
 
 #if !BINNED3_USE_SEPARATE_VM_PER_POOL
 	uint8* FMallocBinned3::Binned3BaseVMPtr = nullptr;
@@ -68,6 +67,44 @@ FMallocBinned3* FMallocBinned3::MallocBinned3 = nullptr;
 // Mapping of sizes to small table indices
 uint8 FMallocBinned3::MemSizeToPoolIndex[1 + (UE_MB3_MAX_SMALL_POOL_SIZE >> UE_MBC_MIN_SMALL_POOL_ALIGNMENT_SHIFT)] = { 0 };
 
+template <class T>
+static void SetCanaryInternal(T& PoolInfo, typename T::ECanary ShouldBe, bool bPreexisting, bool bGuarnteedToBeNew)
+{
+	if (bPreexisting)
+	{
+		if (bGuarnteedToBeNew)
+		{
+			UE_LOG(LogMemory, Fatal, TEXT("MallocBinned3 Corruption Canary was 0x%x, should be 0x%x. This block is both preexisting and guaranteed to be new; which makes no sense."), int32(PoolInfo.Canary), int32(ShouldBe));
+		}
+		if (ShouldBe == T::ECanary::Unassigned)
+		{
+			if (PoolInfo.Canary != T::ECanary::Assigned)
+			{
+				UE_LOG(LogMemory, Fatal, TEXT("MallocBinned3 Corruption Canary was 0x%x, will be 0x%x because this block should be preexisting and in use."), int32(PoolInfo.Canary), int32(ShouldBe));
+			}
+		}
+		else if (PoolInfo.Canary != ShouldBe)
+		{
+			UE_LOG(LogMemory, Fatal, TEXT("MallocBinned3 Corruption Canary was 0x%x, should be 0x%x because this block should be preexisting."), int32(PoolInfo.Canary), int32(ShouldBe));
+		}
+	}
+	else
+	{
+		if (bGuarnteedToBeNew)
+		{
+			if (PoolInfo.Canary != T::ECanary::Unassigned)
+			{
+				UE_LOG(LogMemory, Fatal, TEXT("MallocBinned3 Corruption Canary was 0x%x, will be 0x%x. This block is guaranteed to be new yet is it already assigned."), int32(PoolInfo.Canary), int32(ShouldBe));
+			}
+		}
+		else if (PoolInfo.Canary != ShouldBe && PoolInfo.Canary != T::ECanary::Unassigned)
+		{
+			UE_LOG(LogMemory, Fatal, TEXT("MallocBinned3 Corruption Canary was 0x%x, will be 0x%x does not have an expected value."), int32(PoolInfo.Canary), int32(ShouldBe));
+		}
+	}
+	PoolInfo.Canary = ShouldBe;
+}
+
 /** Information about a piece of free memory. */
 struct FFreeBlock
 {
@@ -128,17 +165,17 @@ struct FMallocBinned3::FPoolInfoSmall		//This is more like BlockInfoSmall as it
 {
 	enum ECanary
 	{
-		SmallUnassigned = 0x3,
-		SmallAssigned = 0x1
+		Unassigned = 0x3,
+		Assigned = 0x1
 	};
 
 	uint32 Canary : 2;
 	uint32 Taken : 15;
-	uint32 NoFirstFreeIndex : 1;	// if a small block info has a FirstFreeIndex as it's only 14 bits and can't store UINT_MAX to say there's no first free index
+	uint32 NoFirstFreeIndex : 1;	// if set to 0 means there are free bins and we can lookup FirstFreeIndex, otherwise means the block is exhausted
 	uint32 FirstFreeIndex : 14;
 
 	FPoolInfoSmall()
-		: Canary(ECanary::SmallUnassigned)
+		: Canary(ECanary::Unassigned)
 		, Taken(0)
 		, NoFirstFreeIndex(1)
 		, FirstFreeIndex(0)
@@ -156,44 +193,12 @@ struct FMallocBinned3::FPoolInfoSmall		//This is more like BlockInfoSmall as it
 
 	void SetCanary(ECanary ShouldBe, bool bPreexisting, bool bGuarnteedToBeNew)
 	{
-		if (bPreexisting)
-		{
-			if (bGuarnteedToBeNew)
-			{
-				UE_LOG(LogMemory, Fatal, TEXT("MallocBinned3 Corruption Canary was 0x%x, should be 0x%x. This block is both preexisting and guaranteed to be new; which makes no sense."), int32(Canary), int32(ShouldBe));
-			}
-			if (ShouldBe == ECanary::SmallUnassigned)
-			{
-				if (Canary != ECanary::SmallAssigned)
-				{
-					UE_LOG(LogMemory, Fatal, TEXT("MallocBinned3 Corruption Canary was 0x%x, will be 0x%x because this block should be preexisting and in use."), int32(Canary), int32(ShouldBe));
-				}
-			}
-			else if (Canary != ShouldBe)
-			{
-				UE_LOG(LogMemory, Fatal, TEXT("MallocBinned3 Corruption Canary was 0x%x, should be 0x%x because this block should be preexisting."), int32(Canary), int32(ShouldBe));
-			}
-		}
-		else
-		{
-			if (bGuarnteedToBeNew)
-			{
-				if (Canary != ECanary::SmallUnassigned)
-				{
-					UE_LOG(LogMemory, Fatal, TEXT("MallocBinned3 Corruption Canary was 0x%x, will be 0x%x. This block is guaranteed to be new yet is it already assigned."), int32(Canary), int32(ShouldBe));
-				}
-			}
-			else if (Canary != ShouldBe && Canary != ECanary::SmallUnassigned)
-			{
-				UE_LOG(LogMemory, Fatal, TEXT("MallocBinned3 Corruption Canary was 0x%x, will be 0x%x does not have an expected value."), int32(Canary), int32(ShouldBe));
-			}
-		}
-		Canary = ShouldBe;
+		SetCanaryInternal(*this, ShouldBe, bPreexisting, bGuarnteedToBeNew);
 	}
 
 	bool HasFreeBin() const
 	{
-		CheckCanary(ECanary::SmallAssigned);
+		CheckCanary(ECanary::Assigned);
 		return !NoFirstFreeIndex;
 	}
 
@@ -224,7 +229,7 @@ struct FMallocBinned3::FPoolInfoSmall		//This is more like BlockInfoSmall as it
 };
 
 FMallocBinned3::FPoolInfo::FPoolInfo()
-	: Canary(ECanary::LargeUnassigned)
+	: Canary(ECanary::Unassigned)
 	, AllocSize(0)
 	, VMSizeDivVirtualSizeAlignment(0)
 	, CommitSize(0)
@@ -241,39 +246,7 @@ void FMallocBinned3::FPoolInfo::CheckCanary(ECanary ShouldBe) const
 
 void FMallocBinned3::FPoolInfo::SetCanary(ECanary ShouldBe, bool bPreexisting, bool bGuarnteedToBeNew)
 {
-	if (bPreexisting)
-	{
-		if (bGuarnteedToBeNew)
-		{
-			UE_LOG(LogMemory, Fatal, TEXT("MallocBinned3 Corruption Canary was 0x%x, should be 0x%x. This block is both preexisting and guaranteed to be new; which makes no sense."), int32(Canary), int32(ShouldBe));
-		}
-		if (ShouldBe == ECanary::LargeUnassigned)
-		{
-			if (Canary != ECanary::LargeAssigned)
-			{
-				UE_LOG(LogMemory, Fatal, TEXT("MallocBinned3 Corruption Canary was 0x%x, will be 0x%x because this block should be preexisting and in use."), int32(Canary), int32(ShouldBe));
-			}
-		}
-		else if (Canary != ShouldBe)
-		{
-			UE_LOG(LogMemory, Fatal, TEXT("MallocBinned3 Corruption Canary was 0x%x, should be 0x%x because this block should be preexisting."), int32(Canary), int32(ShouldBe));
-		}
-	}
-	else
-	{
-		if (bGuarnteedToBeNew)
-		{
-			if (Canary != ECanary::LargeUnassigned)
-			{
-				UE_LOG(LogMemory, Fatal, TEXT("MallocBinned3 Corruption Canary was 0x%x, will be 0x%x. This block is guaranteed to be new yet is it already assigned."), int32(Canary), int32(ShouldBe));
-			}
-		}
-		else if (Canary != ShouldBe && Canary != ECanary::LargeUnassigned)
-		{
-			UE_LOG(LogMemory, Fatal, TEXT("MallocBinned3 Corruption Canary was 0x%x, will be 0x%x does not have an expected value."), int32(Canary), int32(ShouldBe));
-		}
-	}
-	Canary = ShouldBe;
+	SetCanaryInternal(*this, ShouldBe, bPreexisting, bGuarnteedToBeNew);
 }
 
 uint32 FMallocBinned3::FPoolInfo::GetOSRequestedBytes() const
@@ -288,20 +261,20 @@ uint32 FMallocBinned3::FPoolInfo::GetOsCommittedBytes() const
 
 uint32 FMallocBinned3::FPoolInfo::GetOsVMPages() const
 {
-	CheckCanary(ECanary::LargeAssigned);
+	CheckCanary(ECanary::Assigned);
 	return VMSizeDivVirtualSizeAlignment;
 }
 
 void FMallocBinned3::FPoolInfo::SetOSAllocationSize(uint32 InRequestedBytes)
 {
-	CheckCanary(ECanary::LargeAssigned);
+	CheckCanary(ECanary::Assigned);
 	AllocSize = InRequestedBytes;
 	check(AllocSize > 0 && CommitSize >= AllocSize && VMSizeDivVirtualSizeAlignment * FPlatformMemory::FPlatformVirtualMemoryBlock::GetVirtualSizeAlignment() >= CommitSize);
 }
 
 void FMallocBinned3::FPoolInfo::SetOSAllocationSizes(uint32 InRequestedBytes, UPTRINT InCommittedBytes, uint32 InVMSizeDivVirtualSizeAlignment)
 {
-	CheckCanary(ECanary::LargeAssigned);
+	CheckCanary(ECanary::Assigned);
 	AllocSize = InRequestedBytes;
 	CommitSize = InCommittedBytes;
 	VMSizeDivVirtualSizeAlignment = InVMSizeDivVirtualSizeAlignment;
@@ -336,7 +309,7 @@ struct FMallocBinned3::Private
 			bGuaranteedToBeNew = true;
 			Allocator.SmallPoolTables[InPoolIndex].NumEverUsedBlocks = BlockIndex + 1;
 		}
-		Result->SetCanary(FPoolInfoSmall::ECanary::SmallAssigned, false, bGuaranteedToBeNew);
+		Result->SetCanary(FPoolInfoSmall::ECanary::Assigned, false, bGuaranteedToBeNew);
 		return Result;
 	}
 
@@ -368,7 +341,7 @@ struct FMallocBinned3::Private
 				}
 				FPoolInfoSmall* NodePool = &NodePoolBlock[OutBlockIndex % Allocator.SmallPoolInfosPerPlatformPage];
 
-				NodePool->CheckCanary(FPoolInfoSmall::ECanary::SmallAssigned);
+				NodePool->CheckCanary(FPoolInfoSmall::ECanary::Assigned);
 
 				const bool bWasExhaused = NodePool->NoFirstFreeIndex;
 
@@ -387,7 +360,7 @@ struct FMallocBinned3::Private
 				check(NodePool->Taken >= 1);
 				if (--NodePool->Taken == 0)
 				{
-					NodePool->SetCanary(FPoolInfoSmall::ECanary::SmallUnassigned, true, false);
+					NodePool->SetCanary(FPoolInfoSmall::ECanary::Unassigned, true, false);
 					Table.BlocksAllocatedBits.FreeBit(OutBlockIndex);
 
 					const uint64 AllocSize = Allocator.SmallPoolTables[InPoolIndex].NumMemoryPagesPerBlock * Allocator.OsAllocationGranularity;
@@ -448,7 +421,7 @@ FMallocBinned3::FPoolInfoSmall* FMallocBinned3::PushNewPoolToFront(FMallocBinned
 	check(IsAligned(Free, OsAllocationGranularity));
 	// Create pool
 	FPoolInfoSmall* Result = Private::GetOrCreatePoolInfoSmall(*this, InPoolIndex, BlockIndex);
-	Result->CheckCanary(FPoolInfoSmall::ECanary::SmallAssigned);
+	Result->CheckCanary(FPoolInfoSmall::ECanary::Assigned);
 	Result->Taken = 0;
 	Result->FirstFreeIndex = 0;
 	Result->NoFirstFreeIndex = 0;
@@ -589,6 +562,8 @@ FMallocBinned3::FMallocBinned3()
 		SmallPoolTables[Index].TotalFreeCount.store(0, std::memory_order_relaxed);
 #endif
 
+		SmallBinSizesShifted[Index + 1] = (SizeTable[Index].BinSize >> UE_MBC_MIN_SMALL_POOL_ALIGNMENT_SHIFT);
+
 		const int64 TotalNumberOfBlocks = UE_MB3_MAX_MEMORY_PER_POOL_SIZE / (SizeTable[Index].NumMemoryPagesPerBlock * OsAllocationGranularity);
 		const uint32 Size = Align(sizeof(FPoolInfoSmall**) * (TotalNumberOfBlocks + SmallPoolInfosPerPlatformPage - 1) / SmallPoolInfosPerPlatformPage, PLATFORM_CACHE_LINE_SIZE);
 		RequiredMetaMem += Size;
@@ -640,13 +615,6 @@ FMallocBinned3::FMallocBinned3()
 		*IndexEntry++ = uint8(PoolIndex);
 	}
 
-	// now reverse the pool sizes for cache coherency
-	for (uint32 Index = 0; Index != UE_MB3_SMALL_POOL_COUNT; ++Index)
-	{
-		uint32 Partner = UE_MB3_SMALL_POOL_COUNT - Index - 1;
-		SmallBinSizesReversedShifted[Index] = (SizeTable[Partner].BinSize >> UE_MBC_MIN_SMALL_POOL_ALIGNMENT_SHIFT);
-	}
-
 	AllocateHashBuckets();
 
 	MallocBinned3 = this;
@@ -766,7 +734,7 @@ void* FMallocBinned3::MallocExternal(SIZE_T Size, uint32 Alignment)
 		{
 			Pool = PushNewPoolToFront(Table, Table.BinSize, PoolIndex, BlockIndex);
 			
-			//Indicates that we run out of Pool memory (512 MB) for this bin type
+			// Indicates that we run out of reserved virtual memory for the pool (UE_MB3_MAX_MEMORY_PER_POOL_SIZE) for this bin type
 			if (!Pool)
 			{
 				if ((PoolIndex + 1) < UE_MB3_SMALL_POOL_COUNT)
@@ -866,7 +834,7 @@ void* FMallocBinned3::MallocExternal(SIZE_T Size, uint32 Alignment)
 #endif
 
 	// Create pool.
-	FPoolInfo* Pool = Internal::GetOrCreatePoolInfo(*this, Result, FPoolInfo::ECanary::LargeAssigned, false);
+	FPoolInfo* Pool = Internal::GetOrCreatePoolInfo(*this, Result, FPoolInfo::ECanary::Assigned);
 	check(Size > 0 && Size <= AlignedSize && AlignedSize >= FPlatformMemory::FPlatformVirtualMemoryBlock::GetCommitAlignment());
 #if UE_MB3_USE_CACHED_PAGE_ALLOCATOR_FOR_LARGE_ALLOCS
 	Pool->SetOSAllocationSizes(Size, AlignedSize, AlignedSize / FPlatformMemory::FPlatformVirtualMemoryBlock::GetCommitAlignment());
@@ -879,16 +847,16 @@ void* FMallocBinned3::MallocExternal(SIZE_T Size, uint32 Alignment)
 
 void* FMallocBinned3::ReallocExternal(void* Ptr, SIZE_T NewSize, uint32 Alignment)
 {
+	const uint64 PoolIndex = PoolIndexFromPtr(Ptr);
 	if (NewSize == 0)
 	{
-		FMallocBinned3::FreeExternal(Ptr);
+		FMallocBinned3::FreeExternal(Ptr, PoolIndex);
 		return nullptr;
 	}
 	static_assert(DEFAULT_ALIGNMENT <= UE_MBC_MIN_SMALL_POOL_ALIGNMENT, "DEFAULT_ALIGNMENT is assumed to be zero"); // used below
 	check(FMath::IsPowerOfTwo(Alignment));
 	check(Alignment <= OsAllocationGranularity);
 
-	const uint64 PoolIndex = PoolIndexFromPtr(Ptr);
 	if (PoolIndex < UE_MB3_SMALL_POOL_COUNT)
 	{
 		check(Ptr); // null is an OS allocation because it will not fall in our VM block
@@ -907,7 +875,7 @@ void* FMallocBinned3::ReallocExternal(void* Ptr, SIZE_T NewSize, uint32 Alignmen
 		// Reallocate and copy the data across
 		void* Result = FMallocBinned3::MallocExternal(NewSize, Alignment);
 		FMemory::Memcpy(Result, Ptr, FMath::Min<SIZE_T>(NewSize, BinSize));
-		FMallocBinned3::FreeExternal(Ptr);
+		FMallocBinned3::FreeExternal(Ptr, PoolIndex);
 		return Result;
 	}
 	if (!Ptr)
@@ -938,7 +906,7 @@ void* FMallocBinned3::ReallocExternal(void* Ptr, SIZE_T NewSize, uint32 Alignmen
 		void* Result = FMallocBinned3::MallocExternal(NewSize, Alignment);
 		SIZE_T CopySize = FMath::Min<SIZE_T>(NewSize, PoolOSRequestedBytes);
 		FMemory::Memcpy(Result, Ptr, CopySize);
-		FMallocBinned3::FreeExternal(Ptr);
+		FMallocBinned3::FreeExternal(Ptr, PoolIndex);
 		return Result;
 	}
 
@@ -952,11 +920,10 @@ void* FMallocBinned3::ReallocExternal(void* Ptr, SIZE_T NewSize, uint32 Alignmen
 	return Ptr;
 }
 
-void FMallocBinned3::FreeExternal(void* Ptr)
+void FMallocBinned3::FreeExternal(void* Ptr, uint64 PoolIndex)
 {
 	NOALLOC_SCOPE_CYCLE_COUNTER(STAT_FMallocBinned3_FreeExternal);
 
-	const uint64 PoolIndex = PoolIndexFromPtr(Ptr);
 	if (PoolIndex < UE_MB3_SMALL_POOL_COUNT)
 	{
 		check(Ptr); // null is an OS allocation because it will not fall in our VM block
@@ -1018,7 +985,7 @@ void FMallocBinned3::FreeExternal(void* Ptr)
 #endif
 
 			checkf(PoolOSRequestedBytes <= PoolOsBytes, TEXT("FMallocBinned3::FreeExternal %u %u"), int32(PoolOSRequestedBytes), int32(PoolOsBytes));
-			Pool->SetCanary(FPoolInfo::ECanary::LargeUnassigned, true, false);
+			Pool->SetCanary(FPoolInfo::ECanary::Unassigned, true, false);
 		}
 
 		// Free an OS allocation.
@@ -1182,6 +1149,7 @@ void FMallocBinned3::DumpAllocatorStats(class FOutputDevice& Ar)
 	Ar.Logf(TEXT("Allocator Stats for Binned3 are not in this build set UE_MB3_ALLOCATOR_STATS 1 in MallocBinned3.cpp"));
 #endif
 }
+
 #if !UE_MB3_INLINE
 	#if PLATFORM_USES_FIXED_GMalloc_CLASS && !FORCE_ANSI_ALLOCATOR && USE_MALLOC_BINNED3
 		//#define FMEMORY_INLINE_FUNCTION_DECORATOR  FORCEINLINE
@@ -1189,6 +1157,7 @@ void FMallocBinned3::DumpAllocatorStats(class FOutputDevice& Ar)
 		#include "FMemory.inl"
 	#endif
 #endif
-#endif
+
+#endif //~PLATFORM_HAS_FPlatformVirtualMemoryBlock
 
 PRAGMA_RESTORE_UNSAFE_TYPECAST_WARNINGS
diff --git a/Engine/Source/Runtime/Core/Public/HAL/MallocBinned2.h b/Engine/Source/Runtime/Core/Public/HAL/MallocBinned2.h
--- a/Engine/Source/Runtime/Core/Public/HAL/MallocBinned2.h
+++ b/Engine/Source/Runtime/Core/Public/HAL/MallocBinned2.h
@@ -27,10 +27,10 @@ struct FGenericMemoryStats;
 #define UE_MB2_LARGE_ALLOC					65536		// Alignment of OS-allocated pointer - pool-allocated pointers will have a non-aligned pointer
 
 #if AGGRESSIVE_MEMORY_SAVING
-#	define UE_MB2_MAX_SMALL_POOL_SIZE		(13104)		// Maximum bin size in SmallBinSizes in cpp file
+#	define UE_MB2_MAX_SMALL_POOL_SIZE		(13104)		// Maximum bin size in SmallBinSizesInternal in cpp file
 #	define UE_MB2_SMALL_POOL_COUNT			48
 #else
-#	define UE_MB2_MAX_SMALL_POOL_SIZE		(32768-16)	// Maximum bin size in SmallBinSizes in cpp file
+#	define UE_MB2_MAX_SMALL_POOL_SIZE		(32768-16)	// Maximum bin size in SmallBinSizesInternal in cpp file
 #	define UE_MB2_SMALL_POOL_COUNT			51
 #endif
 
@@ -545,9 +545,8 @@ public:
 	/** Dumps current allocator stats to the log. */
 	CORE_API virtual void DumpAllocatorStats(class FOutputDevice& Ar) override;
 	
-	static CORE_API uint16 SmallBinSizesReversed[UE_MB2_SMALL_POOL_COUNT]; // this is reversed to get the smallest elements on our main cache line
+	static CORE_API uint16 SmallBinSizes[UE_MB2_SMALL_POOL_COUNT];
 	static CORE_API FMallocBinned2* MallocBinned2;
-	static CORE_API uint32 PageSize;
 	// Mapping of sizes to small table indices
 	static CORE_API uint8 MemSizeToPoolIndex[1 + (UE_MB2_MAX_SMALL_POOL_SIZE >> UE_MBC_MIN_SMALL_POOL_ALIGNMENT_SHIFT)];
 
@@ -556,7 +555,7 @@ public:
 
 	FORCEINLINE uint32 PoolIndexToBinSize(uint32 PoolIndex) const
 	{
-		return SmallBinSizesReversed[UE_MB2_SMALL_POOL_COUNT - PoolIndex - 1];
+		return SmallBinSizes[PoolIndex];
 	}
 
 	void FreeBundles(FBundleNode* Bundles, uint32 PoolIndex);
diff --git a/Engine/Source/Runtime/Core/Public/HAL/MallocBinned3.h b/Engine/Source/Runtime/Core/Public/HAL/MallocBinned3.h
--- a/Engine/Source/Runtime/Core/Public/HAL/MallocBinned3.h
+++ b/Engine/Source/Runtime/Core/Public/HAL/MallocBinned3.h
@@ -4,7 +4,7 @@
 
 #include "CoreTypes.h"
 
-#if PLATFORM_64BITS && PLATFORM_HAS_FPlatformVirtualMemoryBlock
+#if PLATFORM_HAS_FPlatformVirtualMemoryBlock
 #include "HAL/Allocators/CachedOSPageAllocator.h"
 #include "HAL/Allocators/PooledVirtualMemoryAllocator.h"
 #include "HAL/CriticalSection.h"
@@ -94,8 +94,8 @@ public:
 	{
 		enum class ECanary : uint32
 		{
-			LargeUnassigned = 0x39431234,
-			LargeAssigned = 0x17ea5678,
+			Unassigned = 0x39431234,
+			Assigned = 0x17ea5678,
 		};
 
 		FPoolInfo();
@@ -116,8 +116,8 @@ public:
 		void SetOSAllocationSize(uint32 InRequestedBytes);
 		void SetOSAllocationSizes(uint32 InRequestedBytes, UPTRINT InCommittedBytes, uint32 InVMSizeDivVirtualSizeAlignment);
 
-	private:
 		ECanary	Canary;
+	private:
 		uint32 AllocSize;						// Number of bytes allocated
 		uint32 VMSizeDivVirtualSizeAlignment;	// Number of VM bytes allocated aligned for OS
 		uint32 CommitSize;						// Number of bytes committed by the OS
@@ -136,8 +136,8 @@ private:
 		//uint32 BlockSize;					// Size of a block. A block contains N bins in a way to minimize tail memory waste
 		uint32 NumMemoryPagesPerBlock;		// Num memory pages needed to allocate one block. 
 
-		FBitTree BlocksAllocatedBits;		// One bits in here mean the virtual memory is committed
-		FBitTree BlocksExhaustedBits;		// One bit in here means the pool is completely full
+		FBitTree BlocksAllocatedBits;		// A bit per per block, one indicates a committed block
+		FBitTree BlocksExhaustedBits;		// One here means that the block is completely full
 
 		uint32 NumEverUsedBlocks;
 		FPoolInfoSmall** PoolInfos;			// Book keeping info about every allocated area for this bin's pool
@@ -304,12 +304,12 @@ public:
 			if (Lists)
 			{
 				const uint32 PoolIndex = BoundSizeToPoolIndex(Size, MemSizeToPoolIndex);
-				const uint32 BinSize = PoolIndexToBinSize(PoolIndex);
 				Result = Lists->Malloc(PoolIndex);
 #if UE_MB3_ALLOCATOR_STATS
 				if (Result)
 				{
 					SmallPoolTables[PoolIndex].HeadEndAlloc(Size);
+					const uint32 BinSize = PoolIndexToBinSize(PoolIndex);
 					Lists->AllocatedMemory += BinSize;
 				}
 #endif
@@ -395,7 +395,7 @@ public:
 			if (Lists)
 			{
 				const int32 BinSize = PoolIndexToBinSize(PoolIndex);
-				if (Lists->Free(Ptr, PoolIndex, BinSize))
+				if (Lists->Free(Ptr, PoolIndex, BinSize))	//TODO: move BinSize into the FPerThreadFreeBlockLists list itself to avoid doing a lookup here
 				{
 #if UE_MB3_ALLOCATOR_STATS
 					SmallPoolTables[PoolIndex].HeadEndFree();
@@ -405,7 +405,7 @@ public:
 				}
 			}
 		}
-		FreeExternal(Ptr);
+		FreeExternal(Ptr, PoolIndex);
 	}
 
 	FORCEINLINE bool GetSmallAllocationSize(void* Ptr, SIZE_T& SizeOut) const
@@ -441,7 +441,7 @@ public:
 
 	CORE_API void* MallocExternal(SIZE_T Size, uint32 Alignment);
 	CORE_API void* ReallocExternal(void* Ptr, SIZE_T NewSize, uint32 Alignment);
-	CORE_API void FreeExternal(void *Ptr);
+	CORE_API void FreeExternal(void *Ptr, uint64 PoolIndex);
 
 	CORE_API virtual void GetAllocatorStats( FGenericMemoryStats& out_Stats ) override;
 	/** Dumps current allocator stats to the log. */
@@ -449,7 +449,7 @@ public:
 	
 	// +1 enables PoolIndexToBinSize(~0u / -1) dummy access that helps avoid PoolIndex == 0 branching in Realloc(),
 	// see ((!PoolIndex) | (NewSize > PoolIndexToBinSize(static_cast<uint32>(PoolIndex - 1)))
-	static uint16 SmallBinSizesReversedShifted[UE_MB3_SMALL_POOL_COUNT + 1]; // this is reversed to get the smallest elements on our main cache line
+	static uint16 SmallBinSizesShifted[UE_MB3_SMALL_POOL_COUNT + 1];
 	static FMallocBinned3* MallocBinned3;
 
 #if !BINNED3_USE_SEPARATE_VM_PER_POOL
@@ -466,7 +466,7 @@ public:
 
 	FORCEINLINE uint32 PoolIndexToBinSize(uint32 PoolIndex) const
 	{
-		return uint32(SmallBinSizesReversedShifted[UE_MB3_SMALL_POOL_COUNT - PoolIndex - 1]) << UE_MBC_MIN_SMALL_POOL_ALIGNMENT_SHIFT;
+		return uint32(SmallBinSizesShifted[PoolIndex + 1]) << UE_MBC_MIN_SMALL_POOL_ALIGNMENT_SHIFT;
 	}
 
 	void FreeBundles(FBundleNode* Bundles, uint32 PoolIndex);
@@ -490,5 +490,5 @@ PRAGMA_RESTORE_UNSAFE_TYPECAST_WARNINGS
 		#include "FMemory.inl" // IWYU pragma: export
 	#endif
 #endif
-#endif
 
+#endif	//~PLATFORM_HAS_FPlatformVirtualMemoryBlock
diff --git a/Engine/Source/Runtime/Core/Public/HAL/MallocBinnedCommon.h b/Engine/Source/Runtime/Core/Public/HAL/MallocBinnedCommon.h
--- a/Engine/Source/Runtime/Core/Public/HAL/MallocBinnedCommon.h
+++ b/Engine/Source/Runtime/Core/Public/HAL/MallocBinnedCommon.h
@@ -595,7 +595,7 @@ protected:
 		/**
 		* Gets the PoolInfo for a large block memory address. If no valid info exists one is created.
 		*/
-		static PoolInfo* GetOrCreatePoolInfo(AllocType& Allocator, void* InPtr, typename PoolInfo::ECanary Kind, bool bPreexisting)
+		static PoolInfo* GetOrCreatePoolInfo(AllocType& Allocator, void* InPtr, typename PoolInfo::ECanary Kind)
 		{
 			/**
 			* Creates an array of PoolInfo structures for tracking allocations.
@@ -631,13 +631,13 @@ protected:
 				{
 					Collision->BucketIndex = BucketIndexCollision;
 					Collision->FirstPool = CreatePoolArray(Allocator.NumPoolsPerPage);
-					Collision->FirstPool[PoolIndex].SetCanary(Kind, bPreexisting, true);
+					Collision->FirstPool[PoolIndex].SetCanary(Kind, false, true);
 					return &Collision->FirstPool[PoolIndex];
 				}
 
 				if (Collision->BucketIndex == BucketIndexCollision)
 				{
-					Collision->FirstPool[PoolIndex].SetCanary(Kind, bPreexisting, false);
+					Collision->FirstPool[PoolIndex].SetCanary(Kind, false, false);
 					return &Collision->FirstPool[PoolIndex];
 				}
 
@@ -672,11 +672,11 @@ protected:
 			if (!NewBucket->FirstPool)
 			{
 				NewBucket->FirstPool = CreatePoolArray(Allocator.NumPoolsPerPage);
-				NewBucket->FirstPool[PoolIndex].SetCanary(Kind, bPreexisting, true);
+				NewBucket->FirstPool[PoolIndex].SetCanary(Kind, false, true);
 			}
 			else
 			{
-				NewBucket->FirstPool[PoolIndex].SetCanary(Kind, bPreexisting, false);
+				NewBucket->FirstPool[PoolIndex].SetCanary(Kind, false, false);
 			}
 
 			NewBucket->BucketIndex = BucketIndexCollision;
