Removed check from TArray's serialization path.

diff --git a/Engine/Source/Runtime/Core/Public/Containers/Array.h b/Engine/Source/Runtime/Core/Public/Containers/Array.h
--- a/Engine/Source/Runtime/Core/Public/Containers/Array.h
+++ b/Engine/Source/Runtime/Core/Public/Containers/Array.h
@@ -3925,8 +3925,6 @@ struct TArrayPrivateFriend
 			return Ar;
 		}
 
-		check(SerializeNum >= 0);
-
 		if (Ar.IsError() || SerializeNum < 0 || !ensure(!Ar.IsNetArchive() || SerializeNum <= MaxNetArraySerialize))
 		{
 			Ar.SetError();
