[AutoRTFM] Improve efficiency of TaskArray bit masks.

Rather than calling `RemoveAt(0)` on the bit array, we now track an `InitialIndex` which can be advanced forward as needed.

I have cloned the OnAbort benchmarks into benchmarks for PushOnCommit/PopOnCommit. The benchmark demonstrates a 25-40% improvement with the new code. The PushOnAbort/PopOnAbort benchmarks fluctuated ?? 2% and are roughly the same performance as before.

Before (ms) After (ms)  Ratio
-------------------------------------------------------------------------------
Benchmarks.PopOnCommitHandler
-------------------------------------------------------------------------------
PopOnCommitSingleKey	  67.8226     44.3185     65.34%
PopOnCommitMultiKey       117.775     88.5318     75.17%
PopOnCommitShortSingleKey 95.0314     68.1203     71.68%
PopOnCommitLongSingleKey  [REDACTED_PHONE]%
PopOnCommitShortMultiKey  153.458     111.245     72.49%
PopOnCommitLongMultiKey   163.764     113.313     69.19%

-------------------------------------------------------------------------------
Benchmarks.PopAllOnCommitHandlers
-------------------------------------------------------------------------------
Pop                       [REDACTED_PHONE]%
PopAll                    62.7512     37.4476     59.68%

diff --git a/Engine/Source/Programs/AutoRTFMTests/Private/Benchmarks.cpp b/Engine/Source/Programs/AutoRTFMTests/Private/Benchmarks.cpp
--- a/Engine/Source/Programs/AutoRTFMTests/Private/Benchmarks.cpp
+++ b/Engine/Source/Programs/AutoRTFMTests/Private/Benchmarks.cpp
@@ -534,3 +534,186 @@ TEST_CASE("Benchmarks.PopAllOnAbortHandlers")
 		REQUIRE(!bHit);
 	};
 }
+
+TEST_CASE("Benchmarks.PopOnCommitHandler")
+{
+	constexpr unsigned int Count = 16 * 128;
+
+	BENCHMARK("PopOnCommitSingleKey")
+	{
+		bool bHit = false;
+		AutoRTFM::Testing::Commit([&]
+		{
+			for (unsigned int Index = 0; Index < Count; Index++)
+			{
+				AutoRTFM::PushOnCommitHandler(&bHit, [&] { bHit = true; });
+				AutoRTFM::PopOnCommitHandler(&bHit);
+			}
+		});
+		REQUIRE(!bHit);
+	};
+
+	BENCHMARK("PopOnCommitMultiKey")
+	{
+		bool bHits[Count] = { false };
+		AutoRTFM::Testing::Commit([&]
+		{
+			for (unsigned int Index = 0; Index < Count; Index++)
+			{
+				AutoRTFM::PushOnCommitHandler(&bHits[Index], [&bHits, Index] { bHits[Index] = true; });
+				AutoRTFM::PopOnCommitHandler(&bHits[Index]);
+			}
+		});
+
+		for (unsigned int Index = 0; Index < Count; Index++)
+		{
+			REQUIRE(!bHits[Index]);
+		}
+	};
+
+	BENCHMARK("PopOnCommitShortSingleKey")
+	{
+		bool bHit = false;
+		unsigned int CallbackCounter = 0;
+		AutoRTFM::Testing::Commit([&]
+		{
+			for (unsigned int Index = 0; Index < Count; Index++)
+			{
+				AutoRTFM::OnCommit([&] { CallbackCounter++; });
+			}
+
+			for (unsigned int Index = 0; Index < Count; Index++)
+			{
+				AutoRTFM::PushOnCommitHandler(&bHit, [&] { bHit = true; });
+			}
+
+			for (unsigned int Index = 0; Index < Count; Index++)
+			{
+				AutoRTFM::PopOnCommitHandler(&bHit);
+			}
+		});
+		REQUIRE(!bHit);
+		REQUIRE(Count == CallbackCounter);
+	};
+
+	BENCHMARK("PopOnCommitLongSingleKey")
+	{
+		bool bHit = false;
+		unsigned int CallbackCounter = 0;
+		AutoRTFM::Testing::Commit([&]
+		{
+			for (unsigned int Index = 0; Index < Count; Index++)
+			{
+				AutoRTFM::PushOnCommitHandler(&bHit, [&] { bHit = true; });
+			}
+
+			for (unsigned int Index = 0; Index < Count; Index++)
+			{
+				AutoRTFM::OnCommit([&] { CallbackCounter++; });
+			}
+
+			for (unsigned int Index = 0; Index < Count; Index++)
+			{
+				AutoRTFM::PopOnCommitHandler(&bHit);
+			}
+		});
+		REQUIRE(!bHit);
+		REQUIRE(Count == CallbackCounter);
+	};
+
+	BENCHMARK("PopOnCommitShortMultiKey")
+	{
+		bool bHits[Count] = { false };
+		unsigned int CallbackCounter = 0;
+		AutoRTFM::Testing::Commit([&]
+		{
+			for (unsigned int Index = 0; Index < Count; Index++)
+			{
+				AutoRTFM::OnCommit([&] { CallbackCounter++; });
+			}
+
+			for (unsigned int Index = 0; Index < Count; Index++)
+			{
+				AutoRTFM::PushOnCommitHandler(&bHits[Index], [&bHits, Index] { bHits[Index] = true; });
+			}
+
+			for (unsigned int Index = 0; Index < Count; Index++)
+			{
+				AutoRTFM::PopOnCommitHandler(&bHits[Index]);
+			}
+		});
+
+		for (unsigned int Index = 0; Index < Count; Index++)
+		{
+			REQUIRE(!bHits[Index]);
+		}
+		REQUIRE(Count == CallbackCounter);
+	};
+
+	BENCHMARK("PopOnCommitLongMultiKey")
+	{
+		bool bHits[Count] = { false };
+		unsigned int CallbackCounter = 0;
+		AutoRTFM::Testing::Commit([&]
+		{
+			for (unsigned int Index = 0; Index < Count; Index++)
+			{
+				AutoRTFM::PushOnCommitHandler(&bHits[Index], [&bHits, Index] { bHits[Index] = true; });
+			}
+
+			for (unsigned int Index = 0; Index < Count; Index++)
+			{
+				AutoRTFM::OnCommit([&] { CallbackCounter++; });
+			}
+
+			for (unsigned int Index = 0; Index < Count; Index++)
+			{
+				AutoRTFM::PopOnCommitHandler(&bHits[Index]);
+			}
+		});
+
+		for (unsigned int Index = 0; Index < Count; Index++)
+		{
+			REQUIRE(!bHits[Index]);
+		}
+		REQUIRE(Count == CallbackCounter);
+	};
+}
+
+TEST_CASE("Benchmarks.PopAllOnCommitHandlers")
+{
+	constexpr unsigned int Count = 16 * 128;
+
+	BENCHMARK("Pop")
+	{
+		bool bHit = false;
+		AutoRTFM::Testing::Commit([&]
+		{
+			for (unsigned int Index = 0; Index < Count; Index++)
+			{
+				AutoRTFM::PushOnCommitHandler(&bHit, [&] { bHit = true; });
+			}
+
+			for (unsigned int Index = 0; Index < Count; Index++)
+			{
+				AutoRTFM::PopOnCommitHandler(&bHit);
+			}
+		});
+		REQUIRE(!bHit);
+	};
+
+	BENCHMARK("PopAll")
+	{
+		bool bHit = false;
+		AutoRTFM::Testing::Commit([&]
+		{
+			for (unsigned int Index = 0; Index < Count; Index++)
+			{
+				AutoRTFM::PushOnCommitHandler(&bHit, [&] { bHit = true; });
+			}
+
+			AutoRTFM::PopAllOnCommitHandlers(&bHit);
+		});
+		REQUIRE(!bHit);
+	};
+}
diff --git a/Engine/Source/Runtime/Core/Private/AutoRTFM/TaskArray.h b/Engine/Source/Runtime/Core/Private/AutoRTFM/TaskArray.h
--- a/Engine/Source/Runtime/Core/Private/AutoRTFM/TaskArray.h
+++ b/Engine/Source/Runtime/Core/Private/AutoRTFM/TaskArray.h
@@ -36,7 +36,38 @@ class TTaskArray final
 {
 	using SKey = const void*;
 
-	using BitMaskPerKeyType = TBitArray<>;
+	struct FBitMaskPerKeyType 
+	{
+		TBitArray<> Bits;
+		int32 InitialIndex = 0;
+
+		void RemoveFirst()
+		{
+			checkSlow(InitialIndex < Bits.Num());
+			InitialIndex++;
+		}
+
+		void RemoveLast()
+		{
+			checkSlow(InitialIndex < Bits.Num());
+			Bits.SetNumUninitialized(Bits.Num() - 1);
+		}
+
+		bool BitAt(int32 Index) const
+		{
+			return Bits[InitialIndex + Index];
+		}
+
+		int32 Num() const
+		{
+			return Bits.Num() - InitialIndex;
+		}
+
+		bool IsEmpty() const
+		{
+			return InitialIndex == Bits.Num();
+		}
+	};
 
 	using PairType = TTaskArrayKeyValuePair<SKey*, InElementType*>;
 
@@ -175,44 +206,48 @@ public:
     {
 		ASSERT(Key != SKey{});
 		Latest.Add(MoveTemp(Key), MoveTemp(Value));
-		BitMasksPerKey.FindOrAdd(Key).Add(false);
+		BitMasksPerKey.FindOrAdd(Key).Bits.Add(false);
     }
 
 	void AddKeyed(SKey Key, const InElementType& Value)
     {
 		ASSERT(Key != SKey{});
 		Latest.Add(MoveTemp(Key), Value);
-		BitMasksPerKey.FindOrAdd(Key).Add(false);
+		BitMasksPerKey.FindOrAdd(Key).Bits.Add(false);
     }
 
 	bool DeleteKey(SKey Key)
 	{
 		ASSERT(Key != SKey{});
 
-		BitMaskPerKeyType* const KeyedElement = BitMasksPerKey.Find(Key);
+		FBitMaskPerKeyType* const KeyedElement = BitMasksPerKey.Find(Key);
 
 		if (UNLIKELY(nullptr == KeyedElement))
 		{
 			return false;
 		}
 
-		return INDEX_NONE != KeyedElement->FindAndSetLastZeroBit();
+		return KeyedElement->Bits.FindAndSetLastZeroBit() >= KeyedElement->InitialIndex;
 	}
 
 	bool DeleteAllMatchingKeys(SKey Key)
 	{
 		ASSERT(Key != SKey{});
 
-		BitMaskPerKeyType* const KeyedElement = BitMasksPerKey.Find(Key);
+		FBitMaskPerKeyType* const KeyedElement = BitMasksPerKey.Find(Key);
 
 		if (UNLIKELY(nullptr == KeyedElement))
 		{
 			return false;
 		}
 
-		KeyedElement->SetRange(0, KeyedElement->Num(), true);
+		if (KeyedElement->Num() == 0)
+		{
+			return false;
+		}
 
-		return !KeyedElement->IsEmpty();
+		KeyedElement->Bits.SetRange(KeyedElement->InitialIndex, KeyedElement->Num(), true);
+		return true;
 	}
 
     void AddAll(TTaskArray&& Other)
@@ -233,12 +268,13 @@ public:
 
 		for (auto& Pair : Other.BitMasksPerKey)
 		{
-			BitMaskPerKeyType& KeyedElement = BitMasksPerKey.FindOrAdd(Pair.Key);
-			
-			for (FBitReference Bit : Pair.Value)
-			{
-				KeyedElement.Add(Bit);
-			}
+			const SKey& OtherKey = Pair.Key;
+			FBitMaskPerKeyType& KeyedElement = BitMasksPerKey.FindOrAdd(OtherKey);
+
+			const FBitMaskPerKeyType& OtherBitMask = Pair.Value;
+			int OtherNumBits = OtherBitMask.Num();
+
+			KeyedElement.Bits.AddRange(OtherBitMask.Bits, OtherBitMask.Num(), OtherBitMask.InitialIndex);
 		}
 		Other.BitMasksPerKey.Reset();
     }
@@ -315,16 +351,16 @@ private:
 	{
 		if (Pair.Key)
 		{
-			BitMaskPerKeyType* const KeyedElement = BitMasksPerKey.Find(*Pair.Key);
+			FBitMaskPerKeyType* const KeyedElement = BitMasksPerKey.Find(*Pair.Key);
 
-			if (0 < KeyedElement->Num())
+			if (KeyedElement->Num() > 0)
 			{
-				const bool bFirst = (*KeyedElement)[0];
-				KeyedElement->RemoveAt(0);
+				const bool bFirst = KeyedElement->BitAt(0);
+				KeyedElement->RemoveFirst();
 
 				if (bFirst)
 				{
-					static TFunction<void()> EmptyFunctor = [] {};
+					static const TFunction<void()> EmptyFunctor = [] {};
 
 					// We have a deferred signal to pop this element, so do it now.
 					*Pair.Key = SKey{};
@@ -343,18 +379,18 @@ private:
 	{
 		if (Pair.Key)
 		{
-			BitMaskPerKeyType* const KeyedElement = BitMasksPerKey.Find(*Pair.Key);
+			FBitMaskPerKeyType* const KeyedElement = BitMasksPerKey.Find(*Pair.Key);
 
 			const int32 Num = KeyedElement->Num();
 
-			if (0 < Num)
+			if (Num > 0)
 			{
-				const bool bLast = (*KeyedElement)[Num - 1];
-				KeyedElement->RemoveAt(Num - 1);
+				const bool bLast = KeyedElement->BitAt(Num - 1);
+				KeyedElement->RemoveLast();
 
 				if (bLast)
 				{
-					static TFunction<void()> EmptyFunctor = [] {};
+					static const TFunction<void()> EmptyFunctor = [] {};
 
 					// We have a deferred signal to pop this element, so do it now.
 					*Pair.Key = SKey{};
@@ -371,7 +407,7 @@ private:
     
     FInternalArray Latest;
     TArray<FInternalArray> Stash;
-	TMap<SKey, BitMaskPerKeyType> BitMasksPerKey;
+	TMap<SKey, FBitMaskPerKeyType> BitMasksPerKey;
 };
 
 } // namespace AutoRTFM
