[ManualResetEvent]
- Add missing barrier when spurious wake up happen

#rnx

diff --git a/Engine/Source/Runtime/Core/Private/Microsoft/MicrosoftPlatformManualResetEvent.cpp b/Engine/Source/Runtime/Core/Private/Microsoft/MicrosoftPlatformManualResetEvent.cpp
--- a/Engine/Source/Runtime/Core/Private/Microsoft/MicrosoftPlatformManualResetEvent.cpp
+++ b/Engine/Source/Runtime/Core/Private/Microsoft/MicrosoftPlatformManualResetEvent.cpp
@@ -25,7 +25,7 @@ bool FMicrosoftPlatformManualResetEvent::WaitUntil(FMonotonicTimePoint WaitTime)
 
 		for (;;)
 		{
-			if (WaitOnAddress(&bWait, &bLocalWait, sizeof(bool), INFINITE) && !bWait.load(std::memory_order_relaxed))
+			if (WaitOnAddress(&bWait, &bLocalWait, sizeof(bool), INFINITE) && !bWait.load(std::memory_order_acquire))
 			{
 				return true;
 			}
@@ -44,7 +44,7 @@ bool FMicrosoftPlatformManualResetEvent::WaitUntil(FMonotonicTimePoint WaitTime)
 				return false;
 			}
 			const DWORD WaitMs = DWORD(FPlatformMath::CeilToInt64(WaitSpan.ToMilliseconds()));
-			if (WaitOnAddress(&bWait, &bLocalWait, sizeof(bool), WaitMs) && !bWait.load(std::memory_order_relaxed))
+			if (WaitOnAddress(&bWait, &bLocalWait, sizeof(bool), WaitMs) && !bWait.load(std::memory_order_acquire))
 			{
 				return true;
 			}
@@ -54,7 +54,7 @@ bool FMicrosoftPlatformManualResetEvent::WaitUntil(FMonotonicTimePoint WaitTime)
 
 void FMicrosoftPlatformManualResetEvent::Notify()
 {
-	bWait.store(false, std::memory_order_relaxed);
+	bWait.store(false, std::memory_order_release);
 	WakeByAddressSingle((void*)&bWait);
 }
 
