Localization:
Updated FTextHistory to have thread safe revision counters.

#rnx

diff --git a/Engine/Source/Runtime/Core/Private/Internationalization/TextHistory.cpp b/Engine/Source/Runtime/Core/Private/Internationalization/TextHistory.cpp
--- a/Engine/Source/Runtime/Core/Private/Internationalization/TextHistory.cpp
+++ b/Engine/Source/Runtime/Core/Private/Internationalization/TextHistory.cpp
@@ -681,6 +681,22 @@ void WriteTextFormatToBuffer(FString& Buffer, const FString& TokenMarker, const
 ///////////////////////////////////////
 // FTextHistory
 
+uint16 FTextHistory::GetGlobalHistoryRevision() const
+{
+	uint32 PackedRevision = GlobalAndLocalRevisions.load();
+	// We store global revision in higher 16 bits 
+	uint16 GlobalRevision = static_cast<uint16>(PackedRevision >> 16);
+	return GlobalRevision;
+}
+
+uint16 FTextHistory::GetLocalHistoryRevision() const
+{
+	uint32 PackedRevision = GlobalAndLocalRevisions.load();
+	// Local revision is stored in the lower 16 bits 
+	uint16 LocalRevision = static_cast<uint16>(PackedRevision & 0x0000ffff);
+	return LocalRevision;
+}
+
 void FTextHistory::UpdateDisplayStringIfOutOfDate()
 {
 	if (CanUpdateDisplayString())
@@ -689,32 +705,43 @@ void FTextHistory::UpdateDisplayStringIfOutOfDate()
 		uint16 CurrentLocalRevision = 0;
 		FTextLocalizationManager::Get().GetTextRevisions(GetTextId(), CurrentGlobalRevision, CurrentLocalRevision);
 
+		uint32 GlobalAndLocalRevisionsCopy = GlobalAndLocalRevisions.load();
+		uint16 GlobalRevision = static_cast<uint16>(GlobalAndLocalRevisionsCopy >> 16);
+		uint16 LocalRevision = static_cast<uint16>(GlobalAndLocalRevisionsCopy & 0x0000ffff);
 		if (GlobalRevision != CurrentGlobalRevision || LocalRevision != CurrentLocalRevision)
 		{
-			GlobalRevision = CurrentGlobalRevision;
-			LocalRevision = CurrentLocalRevision;
-
-			UpdateDisplayString();
+			uint32 UpdatedGlobalAndLocalRevision = static_cast<uint32>(CurrentLocalRevision);
+			UpdatedGlobalAndLocalRevision |= static_cast<uint32>(CurrentGlobalRevision) << 16;
+			// We ensure that no other thread has updated any of the revision counters 
+			bool bUpdatedByThisThread = GlobalAndLocalRevisions.compare_exchange_strong(GlobalAndLocalRevisionsCopy, UpdatedGlobalAndLocalRevision);
+			if (bUpdatedByThisThread)
+			{
+				UpdateDisplayString();
+			}
+			// Else another thread had already updated the revision counters and updated the text. No need to do anything 
 		}
 	}
 }
 
 void FTextHistory::MarkDisplayStringOutOfDate()
 {
-	GlobalRevision = 0;
-	LocalRevision = 0;
+	GlobalAndLocalRevisions.store(0);
 }
 
 void FTextHistory::MarkDisplayStringUpToDate()
 {
 	if (CanUpdateDisplayString())
 	{
+		uint16 GlobalRevision = 0;
+		uint16 LocalRevision = 0;
 		FTextLocalizationManager::Get().GetTextRevisions(GetTextId(), GlobalRevision, LocalRevision);
+		uint32 UpdatedGlobalAndLocalRevision = static_cast<uint32>(LocalRevision);
+		UpdatedGlobalAndLocalRevision |= static_cast<uint32>(GlobalRevision) << 16;
+		GlobalAndLocalRevisions.store(UpdatedGlobalAndLocalRevision);
 	}
 	else
 	{
-		GlobalRevision = 0;
-		LocalRevision = 0;
+		GlobalAndLocalRevisions.store(0);
 	}
 }
 
diff --git a/Engine/Source/Runtime/Core/Private/Internationalization/TextHistory.h b/Engine/Source/Runtime/Core/Private/Internationalization/TextHistory.h
--- a/Engine/Source/Runtime/Core/Private/Internationalization/TextHistory.h
+++ b/Engine/Source/Runtime/Core/Private/Internationalization/TextHistory.h
@@ -16,6 +16,7 @@
 #include "Templates/SharedPointerInternals.h"
 #include "Templates/UnrealTemplate.h"
 #include "UObject/NameTypes.h"
+#include <atomic>
 
 class ITextGenerator;
 struct FDecimalNumberFormattingRules;
@@ -155,8 +156,8 @@ public:
 	//~ ITextData
 	virtual const FString& GetSourceString() const override { return GetDisplayString(); }
 	virtual FTextConstDisplayStringPtr GetLocalizedString() const override { return nullptr; }
-	virtual uint16 GetGlobalHistoryRevision() const override final { return GlobalRevision; }
-	virtual uint16 GetLocalHistoryRevision() const override final { return LocalRevision; }
+	virtual uint16 GetGlobalHistoryRevision() const override final;
+	virtual uint16 GetLocalHistoryRevision() const override final;
 	virtual const FTextHistory& GetTextHistory() const override final { return *this; }
 	virtual FTextHistory& GetMutableTextHistory() override final { return *this; }
 
@@ -235,11 +236,12 @@ protected:
 	void MarkDisplayStringUpToDate();
 
 private:
-	/** Global revision index of this history, rebuilds when it is out of sync with the FTextLocalizationManager */
-	uint16 GlobalRevision = 0;
-
-	/** Local revision index of this history, rebuilds when it is out of sync with the FTextLocalizationManager */
-	uint16 LocalRevision = 0;
+	/**
+	 * The global and local revision indices of this text history  packed into a std::atomic<uint32> for thread safety.
+	 * Global revision index of this history, rebuilds this text history when it is out of sync with the FTextLocalizationManager. Stored in the higher 16 bits
+	 * Local revision index of this history, rebuilds this text history when it is out of sync with the FTextLocalizationManager. Stored in the lower 16 bits.
+	 */
+	std::atomic<uint32> GlobalAndLocalRevisions {0};
 };
 
 /** A potentially localized piece of source text (may have a TextId). */
