[Insights] Fixed max number of callstack frames traced to be 255 (instead of 256; some platforms still limit this to max 64 frames). Added UE_CALLSTACK_TRACE_MAX_FRAMES to override this if necessary. Memory Insights analysis would already limit this to max 255 frames (length of callstack is stored as an uint8).

diff --git a/Engine/Source/Runtime/Core/Core.Build.cs b/Engine/Source/Runtime/Core/Core.Build.cs
--- a/Engine/Source/Runtime/Core/Core.Build.cs
+++ b/Engine/Source/Runtime/Core/Core.Build.cs
@@ -156,6 +156,7 @@ public class Core : ModuleRules
 					PublicDefinitions.Add("UE_MEMORY_TRACE_AVAILABLE=1");
 					PublicDefinitions.Add("UE_MEMORY_TAGS_TRACE_ENABLED=1");
 					PublicDefinitions.Add("UE_CALLSTACK_TRACE_ENABLED=1");
+					PublicDefinitions.Add("UE_CALLSTACK_TRACE_MAX_FRAMES=64");
 				}
 			}
 		}
@@ -172,6 +173,7 @@ public class Core : ModuleRules
 				PublicDefinitions.Add("UE_MEMORY_TRACE_AVAILABLE=1");
 				PublicDefinitions.Add("UE_MEMORY_TAGS_TRACE_ENABLED=1");
 				PublicDefinitions.Add("UE_CALLSTACK_TRACE_ENABLED=1");
+				PublicDefinitions.Add("UE_CALLSTACK_TRACE_MAX_FRAMES=64");
 				PrivateDefinitions.Add("UE_CALLSTACK_TRACE_ANDROID_USE_STACK_FRAMES_WALKING=1");
 
 				// Support for memory tracing libc.so malloc
diff --git a/Engine/Source/Runtime/Core/Private/Android/AndroidCallstackTrace.cpp b/Engine/Source/Runtime/Core/Private/Android/AndroidCallstackTrace.cpp
--- a/Engine/Source/Runtime/Core/Private/Android/AndroidCallstackTrace.cpp
+++ b/Engine/Source/Runtime/Core/Private/Android/AndroidCallstackTrace.cpp
@@ -87,7 +87,7 @@ FORCEINLINE uint32 FBacktracer::GetBacktraceId(uint64* StackFrames, uint32 NumSt
 FORCEINLINE uint32 FBacktracer::GetBacktraceId(void* ReturnAddress)
 {
 #if !UE_BUILD_SHIPPING
-	uint64 StackFrames[64];
+	uint64 StackFrames[UE_CALLSTACK_TRACE_MAX_FRAMES];
 
 #if UE_CALLSTACK_TRACE_ANDROID_USE_STACK_FRAMES_WALKING
 	uint32 NumStackFrames = FPlatformStackWalk::CaptureStackBackTraceViaFramePointerWalking(StackFrames, UE_ARRAY_COUNT(StackFrames));
diff --git a/Engine/Source/Runtime/Core/Private/IOS/IOSCallstackTrace.cpp b/Engine/Source/Runtime/Core/Private/IOS/IOSCallstackTrace.cpp
--- a/Engine/Source/Runtime/Core/Private/IOS/IOSCallstackTrace.cpp
+++ b/Engine/Source/Runtime/Core/Private/IOS/IOSCallstackTrace.cpp
@@ -81,7 +81,7 @@ FORCEINLINE uint32 FBacktracer::GetBacktraceId(uint64* StackFrames, uint32 NumSt
 FORCEINLINE uint32 FBacktracer::GetBacktraceId(void* ReturnAddress)
 {
 #if !UE_BUILD_SHIPPING
-	uint64 StackFrames[64];
+	uint64 StackFrames[UE_CALLSTACK_TRACE_MAX_FRAMES];
 
 	uint32 NumStackFrames = FPlatformStackWalk::CaptureStackBackTrace(StackFrames, UE_ARRAY_COUNT(StackFrames));
 
diff --git a/Engine/Source/Runtime/Core/Private/ProfilingDebugging/Microsoft/WindowsCallstackTrace.cpp b/Engine/Source/Runtime/Core/Private/ProfilingDebugging/Microsoft/WindowsCallstackTrace.cpp
--- a/Engine/Source/Runtime/Core/Private/ProfilingDebugging/Microsoft/WindowsCallstackTrace.cpp
+++ b/Engine/Source/Runtime/Core/Private/ProfilingDebugging/Microsoft/WindowsCallstackTrace.cpp
@@ -595,7 +595,7 @@ const FBacktracer::FFunction* FBacktracer::LookupFunction(UPTRINT Address, FLook
 uint32 FBacktracer::GetBacktraceId(void* AddressOfReturnAddress) 
 {
 	FLookupState LookupState = {};
-	uint64 Frames[256];
+	uint64 Frames[UE_CALLSTACK_TRACE_MAX_FRAMES];
 
 	UPTRINT* StackPointer = (UPTRINT*)AddressOfReturnAddress;
 
@@ -795,7 +795,7 @@ uint32 FBacktracer::GetBacktraceId(void* AddressOfReturnAddress)
 uint32 FBacktracer::GetBacktraceId(uint64 ReturnAddress)
 {
 #if !UE_BUILD_SHIPPING
-	uint64 StackFrames[256];
+	uint64 StackFrames[UE_CALLSTACK_TRACE_MAX_FRAMES];
 	int32 NumStackFrames = FPlatformStackWalk::CaptureStackBackTrace(StackFrames, UE_ARRAY_COUNT(StackFrames));
 	if (NumStackFrames > 0)
 	{
diff --git a/Engine/Source/Runtime/Core/Private/ProfilingDebugging/Unix/UnixCallstackTrace.cpp b/Engine/Source/Runtime/Core/Private/ProfilingDebugging/Unix/UnixCallstackTrace.cpp
--- a/Engine/Source/Runtime/Core/Private/ProfilingDebugging/Unix/UnixCallstackTrace.cpp
+++ b/Engine/Source/Runtime/Core/Private/ProfilingDebugging/Unix/UnixCallstackTrace.cpp
@@ -50,7 +50,7 @@ FBacktracer* FBacktracer::Get()
 uint32 FBacktracer::GetBacktraceId(void* ReturnAddress)
 {
 #if !UE_BUILD_SHIPPING
-	uint64 StackFrames[256];
+	uint64 StackFrames[UE_CALLSTACK_TRACE_MAX_FRAMES];
 	int32 NumStackFrames = FPlatformStackWalk::CaptureStackBackTrace(StackFrames, UE_ARRAY_COUNT(StackFrames));
 	if (NumStackFrames > 0)
 	{
diff --git a/Engine/Source/Runtime/Core/Public/ProfilingDebugging/CallstackTrace.h b/Engine/Source/Runtime/Core/Public/ProfilingDebugging/CallstackTrace.h
--- a/Engine/Source/Runtime/Core/Public/ProfilingDebugging/CallstackTrace.h
+++ b/Engine/Source/Runtime/Core/Public/ProfilingDebugging/CallstackTrace.h
@@ -19,6 +19,10 @@
 	#define UE_CALLSTACK_TRACE_ENABLED 0
 #endif
 
+#if !defined(UE_CALLSTACK_TRACE_MAX_FRAMES)
+	#define UE_CALLSTACK_TRACE_MAX_FRAMES 255
+#endif
+
 ////////////////////////////////////////////////////////////////////////////////
 #if UE_CALLSTACK_TRACE_ENABLED
 
