[Insights] Fixed consistency regarding usage of the Metadata trace channel (for AssetMetadata and IoStore).

diff --git a/Engine/Source/Runtime/Core/Public/ProfilingDebugging/AssetMetadataTrace.h b/Engine/Source/Runtime/Core/Public/ProfilingDebugging/AssetMetadataTrace.h
--- a/Engine/Source/Runtime/Core/Public/ProfilingDebugging/AssetMetadataTrace.h
+++ b/Engine/Source/Runtime/Core/Public/ProfilingDebugging/AssetMetadataTrace.h
@@ -54,12 +54,12 @@ UE_TRACE_EVENT_END()
  * Utility macro to create an asset scope by specifying object name, class name and package name explicitly
  */
 #define UE_TRACE_METADATA_SCOPE_ASSET_FNAME(ObjectName, ObjClassName, PackageName) \
-	auto MetaNameRef = UE_TRACE_CHANNELEXPR_IS_ENABLED(MetadataChannel) ? FStringTrace::GetNameRef(ObjectName) : UE::Trace::FEventRef32(0,0);\
-	auto ClassNameRef = UE_TRACE_CHANNELEXPR_IS_ENABLED(MetadataChannel) ? FStringTrace::GetNameRef(ObjClassName) : UE::Trace::FEventRef32(0,0);\
-	auto PackageNameRef = UE_TRACE_CHANNELEXPR_IS_ENABLED(MetadataChannel) ? FStringTrace::GetNameRef(PackageName) : UE::Trace::FEventRef32(0,0);\
-	UE_TRACE_METADATA_SCOPE(Asset, AssetMetadataChannel)\
-		<< Asset.Name(MetaNameRef)\
-		<< Asset.Class(ClassNameRef)\
+	auto MetaNameRef = bool(MetadataChannel) && bool(AssetMetadataChannel) ? FStringTrace::GetNameRef(ObjectName) : UE::Trace::FEventRef32(0,0); \
+	auto ClassNameRef = bool(MetadataChannel) && bool(AssetMetadataChannel) ? FStringTrace::GetNameRef(ObjClassName) : UE::Trace::FEventRef32(0,0); \
+	auto PackageNameRef = bool(MetadataChannel) && bool(AssetMetadataChannel) ? FStringTrace::GetNameRef(PackageName) : UE::Trace::FEventRef32(0,0); \
+	UE_TRACE_METADATA_SCOPE(Asset, AssetMetadataChannel) \
+		<< Asset.Name(MetaNameRef) \
+		<< Asset.Class(ClassNameRef) \
 		<< Asset.Package(PackageNameRef);
 
 /*
@@ -74,9 +74,9 @@ UE_TRACE_EVENT_END()
  */
 #define UE_TRACE_PACKAGE_NAME(UPackageId, PackageName) \
 	if (!PackageName.IsNone()) \
-	{\
-		auto PackageNameRef = UE_TRACE_CHANNELEXPR_IS_ENABLED(MetadataChannel) ? FStringTrace::GetNameRef(PackageName) : UE::Trace::FEventRef32(0,0);\
-		UE_TRACE_LOG(Package, PackageMapping, AssetMetadataChannel ) \
+	{ \
+		auto PackageNameRef = bool(MetadataChannel) && bool(AssetMetadataChannel) ? FStringTrace::GetNameRef(PackageName) : UE::Trace::FEventRef32(0,0); \
+		UE_TRACE_LOG(Package, PackageMapping, MetadataChannel | AssetMetadataChannel) \
 			<< PackageMapping.Id(UPackageId.Value()) \
 			<< PackageMapping.Package(PackageNameRef); \
 	}
diff --git a/Engine/Source/Runtime/Core/Public/ProfilingDebugging/IoStoreTrace.h b/Engine/Source/Runtime/Core/Public/ProfilingDebugging/IoStoreTrace.h
--- a/Engine/Source/Runtime/Core/Public/ProfilingDebugging/IoStoreTrace.h
+++ b/Engine/Source/Runtime/Core/Public/ProfilingDebugging/IoStoreTrace.h
@@ -62,7 +62,7 @@ UE_TRACE_CHANNEL_EXTERN(IoStoreChannel, CORE_API);
 	UE_TRACE_METADATA_EVENT_END()
 
 	#define TRACE_IOSTORE_METADATA_SCOPE_TAG(TagName) \
-		auto TagNameRef = UE_TRACE_CHANNELEXPR_IS_ENABLED(MetadataChannel|IoStoreChannel) ? FStringTrace::GetNameRef(TagName) : UE::Trace::FEventRef32(0,0);\
+		auto TagNameRef = bool(MetadataChannel) && bool(IoStoreChannel) ? FStringTrace::GetNameRef(TagName) : UE::Trace::FEventRef32(0,0); \
 		UE_TRACE_METADATA_SCOPE(IoStoreTag, IoStoreChannel) \
 			<< IoStoreTag.Tag(TagNameRef);
 #else
diff --git a/Engine/Source/Runtime/Core/Public/ProfilingDebugging/MetadataTrace.h b/Engine/Source/Runtime/Core/Public/ProfilingDebugging/MetadataTrace.h
--- a/Engine/Source/Runtime/Core/Public/ProfilingDebugging/MetadataTrace.h
+++ b/Engine/Source/Runtime/Core/Public/ProfilingDebugging/MetadataTrace.h
@@ -7,9 +7,9 @@
 
 # if !defined(UE_TRACE_METADATA_ENABLED)
 	#if !IS_PROGRAM
-		#define	UE_TRACE_METADATA_ENABLED UE_TRACE_ENABLED
+		#define UE_TRACE_METADATA_ENABLED UE_TRACE_ENABLED
 	#else
-		#define	UE_TRACE_METADATA_ENABLED 0
+		#define UE_TRACE_METADATA_ENABLED 0
 	#endif
 #endif //defined(UE_TRACE_METADATA_ENABLED
 
@@ -54,7 +54,7 @@ public:
  * A scope to push/pop data on the metadata stack.
  */
 #define UE_TRACE_METADATA_SCOPE(ScopeName, ChannelExpr) \
-	UE_TRACE_LOG_SCOPED(Metadata, ScopeName, ChannelExpr)
+	UE_TRACE_LOG_SCOPED(Metadata, ScopeName, MetadataChannel | ChannelExpr)
 
 /**
  * Temporarily clears the current metadata stack and starts a new one. When this scope is exited the previous
