Fix FTextHistory to work in transactions.

diff --git a/Engine/Plugins/Experimental/LocalizableMessage/Source/LocalizableMessage/Private/LocalizableMessageProcessor.cpp b/Engine/Plugins/Experimental/LocalizableMessage/Source/LocalizableMessage/Private/LocalizableMessageProcessor.cpp
--- a/Engine/Plugins/Experimental/LocalizableMessage/Source/LocalizableMessage/Private/LocalizableMessageProcessor.cpp
+++ b/Engine/Plugins/Experimental/LocalizableMessage/Source/LocalizableMessage/Private/LocalizableMessageProcessor.cpp
@@ -65,7 +65,7 @@ FText FLocalizableMessageProcessor::Localize(const FLocalizableMessage& Message,
 		FTextId TextId;
 		if (UE::LocalizableMessageTextInterop::MessageKeyToTextId(Message.Key, TextId))
 		{
-			UE_AUTORTFM_OPEN{ LocalizedText = FText::AsLocalizable_Advanced(TextId.GetNamespace(), TextId.GetKey(), Message.DefaultText); };
+			LocalizedText = FText::AsLocalizable_Advanced(TextId.GetNamespace(), TextId.GetKey(), Message.DefaultText);
 		}
 	}
 
diff --git a/Engine/Source/Runtime/Core/Private/Internationalization/TextCache.cpp b/Engine/Source/Runtime/Core/Private/Internationalization/TextCache.cpp
--- a/Engine/Source/Runtime/Core/Private/Internationalization/TextCache.cpp
+++ b/Engine/Source/Runtime/Core/Private/Internationalization/TextCache.cpp
@@ -9,15 +9,6 @@
 #include "Misc/CString.h"
 #include "Misc/LazySingleton.h"
 #include "Misc/ScopeLock.h"
-#include "AutoRTFM/AutoRTFM.h"
-
-// Notes on the use of UE_AUTORTFM_ALWAYS_OPEN and UE_AUTORTFM_NOAUTORTFM:
-// It is currently unsafe to use the cache in the open while an uncommitted
-// transaction is in flight that has also touched the cache (#jira SOL-6743).
-// RemoveCache() is not currently reachable from a closed transaction, so these
-// are annotated with UE_AUTORTFM_NOAUTORTFM to prevent new transactional use.
-// FindOrCache() is currently used from transactional code paths, so is
-// annotated with UE_AUTORTFM_ALWAYS_OPEN.
 
 FTextCache& FTextCache::Get()
 {
@@ -36,101 +27,95 @@ FText FTextCache::FindOrCache(const TCHAR* InTextLiteral, const FTextId& InTextI
 
 FText FTextCache::FindOrCache(FStringView InTextLiteral, const FTextId& InTextId)
 {
-	return AutoRTFM::Open([&]
-	{
-		LLM_SCOPE(ELLMTag::Localization);
+	LLM_SCOPE(ELLMTag::Localization);
 
-		// First try and find a cached instance
-		{
-			FText* ReturnFoundText = nullptr;
+	// First try and find a cached instance
+	{
+		FText* ReturnFoundText = nullptr;
 
-			FScopeLock Lock(&CachedTextCS);
+		UE::TScopeLock Lock(CachedTextCS);
 
-			FText* FoundText = CachedText.Find(InTextId);
-			if (FoundText)
+		FText* FoundText = CachedText.Find(InTextId);
+		if (FoundText)
+		{
+			const FString* FoundTextLiteral = FTextInspector::GetSourceString(*FoundText);
+			if (FoundTextLiteral && InTextLiteral.Equals(*FoundTextLiteral, ESearchCase::CaseSensitive))
 			{
-				const FString* FoundTextLiteral = FTextInspector::GetSourceString(*FoundText);
-				if (FoundTextLiteral && InTextLiteral.Equals(*FoundTextLiteral, ESearchCase::CaseSensitive))
-				{
-					ReturnFoundText = FoundText;
-				}
+				ReturnFoundText = FoundText;
 			}
+		}
 
-			if (ReturnFoundText)
-			{
-				return *ReturnFoundText;
-			}
+		if (ReturnFoundText)
+		{
+			return *ReturnFoundText;
 		}
+	}
 
-		// Not currently cached, make a new instance...
-		FText NewText = FText(FString(InTextLiteral), InTextId.GetNamespace(), InTextId.GetKey(), ETextFlag::Immutable);
+	// Not currently cached, make a new instance...
+	FText NewText = FText(FString(InTextLiteral), InTextId.GetNamespace(), InTextId.GetKey(), ETextFlag::Immutable);
 
-		// ... and add it to the cache
-		FScopeLock Lock(&CachedTextCS);
+	// ... and add it to the cache
+	UE::TScopeLock Lock(CachedTextCS);
 
-		CachedText.Emplace(InTextId, NewText);
+	CachedText.Emplace(InTextId, NewText);
 
-		return NewText;
-	});
+	return NewText;
 }
 
 FText FTextCache::FindOrCache(FString&& InTextLiteral, const FTextId& InTextId)
 {
-	return AutoRTFM::Open([&]
-	{
-		LLM_SCOPE(ELLMTag::Localization);
+	LLM_SCOPE(ELLMTag::Localization);
 
-		// First try and find a cached instance
-		{
-			FText* ReturnFoundText = nullptr;
+	// First try and find a cached instance
+	{
+		FText* ReturnFoundText = nullptr;
 
-			FScopeLock Lock(&CachedTextCS);
+		UE::TScopeLock Lock(CachedTextCS);
 
-			FText* FoundText = CachedText.Find(InTextId);
-			if (FoundText)
+		FText* FoundText = CachedText.Find(InTextId);
+		if (FoundText)
+		{
+			const FString* FoundTextLiteral = FTextInspector::GetSourceString(*FoundText);
+			if (FoundTextLiteral && InTextLiteral.Equals(*FoundTextLiteral, ESearchCase::CaseSensitive))
 			{
-				const FString* FoundTextLiteral = FTextInspector::GetSourceString(*FoundText);
-				if (FoundTextLiteral && InTextLiteral.Equals(*FoundTextLiteral, ESearchCase::CaseSensitive))
-				{
-					ReturnFoundText = FoundText;
-				}
+				ReturnFoundText = FoundText;
 			}
+		}
 
-			if (ReturnFoundText)
-			{
-				return *ReturnFoundText;
-			}
+		if (ReturnFoundText)
+		{
+			return *ReturnFoundText;
 		}
+	}
 
-		// Not currently cached, make a new instance...
-		FText NewText = FText(MoveTemp(InTextLiteral), InTextId.GetNamespace(), InTextId.GetKey(), ETextFlag::Immutable);
+	// Not currently cached, make a new instance...
+	FText NewText = FText(MoveTemp(InTextLiteral), InTextId.GetNamespace(), InTextId.GetKey(), ETextFlag::Immutable);
 
-		// ... and add it to the cache
-		FScopeLock Lock(&CachedTextCS);
+	// ... and add it to the cache
+	UE::TScopeLock Lock(CachedTextCS);
 
-		CachedText.Emplace(InTextId, NewText);
+	CachedText.Emplace(InTextId, NewText);
 
-		return NewText;
-	});
+	return NewText;
 }
 
-UE_AUTORTFM_NOAUTORTFM void FTextCache::RemoveCache(const FTextId& InTextId)
+void FTextCache::RemoveCache(const FTextId& InTextId)
 {
 	return RemoveCache(MakeArrayView(&InTextId, 1));
 }
 
-UE_AUTORTFM_NOAUTORTFM void FTextCache::RemoveCache(TArrayView<const FTextId> InTextIds)
+void FTextCache::RemoveCache(TArrayView<const FTextId> InTextIds)
 {
-	FScopeLock Lock(&CachedTextCS);
+	UE::TScopeLock Lock(CachedTextCS);
 	for (const FTextId& TextId : InTextIds)
 	{
 		CachedText.Remove(TextId);
 	}
 }
 
-UE_AUTORTFM_NOAUTORTFM void FTextCache::RemoveCache(const TSet<FTextId>& InTextIds)
+void FTextCache::RemoveCache(const TSet<FTextId>& InTextIds)
 {
-	FScopeLock Lock(&CachedTextCS);
+	UE::TScopeLock Lock(CachedTextCS);
 	for (const FTextId& TextId : InTextIds)
 	{
 		CachedText.Remove(TextId);
diff --git a/Engine/Source/Runtime/Core/Private/Internationalization/TextCache.h b/Engine/Source/Runtime/Core/Private/Internationalization/TextCache.h
--- a/Engine/Source/Runtime/Core/Private/Internationalization/TextCache.h
+++ b/Engine/Source/Runtime/Core/Private/Internationalization/TextCache.h
@@ -3,9 +3,9 @@
 #pragma once
 
 #include "Containers/Map.h"
-#include "HAL/CriticalSection.h"
 #include "HAL/Platform.h"
 #include "Internationalization/Text.h"
+#include "Misc/TransactionallySafeCriticalSection.h"
 
 class FTextId;
 
@@ -41,5 +41,5 @@ public:
 
 private:
 	TMap<FTextId, FText> CachedText;
-	FCriticalSection CachedTextCS;
+	FTransactionallySafeCriticalSection CachedTextCS;
 };
diff --git a/Engine/Source/Runtime/Core/Private/Internationalization/TextHistory.cpp b/Engine/Source/Runtime/Core/Private/Internationalization/TextHistory.cpp
--- a/Engine/Source/Runtime/Core/Private/Internationalization/TextHistory.cpp
+++ b/Engine/Source/Runtime/Core/Private/Internationalization/TextHistory.cpp
@@ -683,17 +683,13 @@ void WriteTextFormatToBuffer(FString& Buffer, const FString& TokenMarker, const
 
 uint16 FTextHistory::GetGlobalHistoryRevision() const
 {
-	uint32 PackedRevision = GlobalAndLocalRevisions.load();
-	// We store global revision in higher 16 bits 
-	uint16 GlobalRevision = static_cast<uint16>(PackedRevision >> 16);
+	UE::TScopeLock _(Mutex);
 	return GlobalRevision;
 }
 
 uint16 FTextHistory::GetLocalHistoryRevision() const
 {
-	uint32 PackedRevision = GlobalAndLocalRevisions.load();
-	// Local revision is stored in the lower 16 bits 
-	uint16 LocalRevision = static_cast<uint16>(PackedRevision & 0x0000ffff);
+	UE::TScopeLock _(Mutex);
 	return LocalRevision;
 }
 
@@ -705,43 +701,43 @@ void FTextHistory::UpdateDisplayStringIfOutOfDate()
 		uint16 CurrentLocalRevision = 0;
 		FTextLocalizationManager::Get().GetTextRevisions(GetTextId(), CurrentGlobalRevision, CurrentLocalRevision);
 
-		uint32 GlobalAndLocalRevisionsCopy = GlobalAndLocalRevisions.load();
-		uint16 GlobalRevision = static_cast<uint16>(GlobalAndLocalRevisionsCopy >> 16);
-		uint16 LocalRevision = static_cast<uint16>(GlobalAndLocalRevisionsCopy & 0x0000ffff);
+		// GlobalRevision and LocalRevision can be updated by concurrent threads!
+		UE::TScopeLock _(Mutex);
+
 		if (GlobalRevision != CurrentGlobalRevision || LocalRevision != CurrentLocalRevision)
 		{
-			uint32 UpdatedGlobalAndLocalRevision = static_cast<uint32>(CurrentLocalRevision);
-			UpdatedGlobalAndLocalRevision |= static_cast<uint32>(CurrentGlobalRevision) << 16;
-			// We ensure that no other thread has updated any of the revision counters 
-			bool bUpdatedByThisThread = GlobalAndLocalRevisions.compare_exchange_strong(GlobalAndLocalRevisionsCopy, UpdatedGlobalAndLocalRevision);
-			if (bUpdatedByThisThread)
-			{
-				UpdateDisplayString();
-			}
-			// Else another thread had already updated the revision counters and updated the text. No need to do anything 
+			GlobalRevision = CurrentGlobalRevision;
+			LocalRevision = CurrentLocalRevision;
+
+			UpdateDisplayString();
 		}
 	}
 }
 
 void FTextHistory::MarkDisplayStringOutOfDate()
 {
-	GlobalAndLocalRevisions.store(0);
+	// GlobalRevision and LocalRevision can be updated by concurrent threads!
+	UE::TScopeLock _(Mutex);
+
+	GlobalRevision = 0;
+	LocalRevision = 0;
 }
 
 void FTextHistory::MarkDisplayStringUpToDate()
 {
-	if (CanUpdateDisplayString())
+	const bool bCanUpdate = CanUpdateDisplayString();
+
+	// GlobalRevision and LocalRevision can be updated by concurrent threads!
+	UE::TScopeLock _(Mutex);
+
+	if (bCanUpdate)
 	{
-		uint16 GlobalRevision = 0;
-		uint16 LocalRevision = 0;
 		FTextLocalizationManager::Get().GetTextRevisions(GetTextId(), GlobalRevision, LocalRevision);
-		uint32 UpdatedGlobalAndLocalRevision = static_cast<uint32>(LocalRevision);
-		UpdatedGlobalAndLocalRevision |= static_cast<uint32>(GlobalRevision) << 16;
-		GlobalAndLocalRevisions.store(UpdatedGlobalAndLocalRevision);
 	}
 	else
 	{
-		GlobalAndLocalRevisions.store(0);
+		GlobalRevision = 0;
+		LocalRevision = 0;
 	}
 }
 
diff --git a/Engine/Source/Runtime/Core/Private/Internationalization/TextHistory.h b/Engine/Source/Runtime/Core/Private/Internationalization/TextHistory.h
--- a/Engine/Source/Runtime/Core/Private/Internationalization/TextHistory.h
+++ b/Engine/Source/Runtime/Core/Private/Internationalization/TextHistory.h
@@ -16,7 +16,7 @@
 #include "Templates/SharedPointerInternals.h"
 #include "Templates/UnrealTemplate.h"
 #include "UObject/NameTypes.h"
-#include <atomic>
+#include "Misc/TransactionallySafeMutex.h"
 
 class ITextGenerator;
 struct FDecimalNumberFormattingRules;
@@ -236,12 +236,12 @@ protected:
 	void MarkDisplayStringUpToDate();
 
 private:
-	/**
-	 * The global and local revision indices of this text history  packed into a std::atomic<uint32> for thread safety.
-	 * Global revision index of this history, rebuilds this text history when it is out of sync with the FTextLocalizationManager. Stored in the higher 16 bits
-	 * Local revision index of this history, rebuilds this text history when it is out of sync with the FTextLocalizationManager. Stored in the lower 16 bits.
-	 */
-	std::atomic<uint32> GlobalAndLocalRevisions {0};
+	/** Global revision index of this history, rebuilds when it is out of sync with the FTextLocalizationManager */
+	uint16 GlobalRevision = 0;
+	/** Local revision index of this history, rebuilds when it is out of sync with the FTextLocalizationManager */
+	uint16 LocalRevision = 0;
+	/** Mutex to prevent races when multiple threads are updating out of date state. */
+	mutable FTransactionallySafeMutex Mutex = FTransactionallySafeMutex();
 };
 
 /** A potentially localized piece of source text (may have a TextId). */
