[RWSpinLock]
- Add a variable sized read-write spinlock along with new templated scope for read-write locks in general

diff --git a/Engine/Source/Runtime/Core/Public/Misc/RWSpinLock.h b/Engine/Source/Runtime/Core/Public/Misc/RWSpinLock.h
new file mode 100644
--- /dev/null
+++ b/Engine/Source/Runtime/Core/Public/Misc/RWSpinLock.h
@@ -0,0 +1,79 @@
+// Copyright Epic Games, Inc. All Rights Reserved.
+
+#pragma once
+
+#include "CoreTypes.h"
+#include "HAL/PlatformProcess.h"
+#include "Math/NumericLimits.h"
+
+#include <atomic>
+
+namespace UE
+{
+	// A read-write lock that doesn't put the thread into a WAIT state but instead repeatedly tries to acquire the lock.
+	// WARNING: Should be used only for very short locks.
+	// Use with `TRWScopeLock`, `TWriteScopeLock` or `TReadScopeLock`.
+	// The SizeType template type dictates the space taken by the spinlock but also its maximum number of possible concurrent readers (i.e. 254 for uint8, etc...).
+	template <typename SizeType = uint32>
+	class TRWSpinLock
+	{
+	public:
+		static_assert(TNumericLimits<SizeType>::Min() >= 0, "Type must be unsigned");
+		UE_NONCOPYABLE(TRWSpinLock);
+
+		TRWSpinLock() = default;
+
+		bool TryWriteLock()
+		{
+			SizeType Expected = 0;
+			return Lock.compare_exchange_strong(Expected, TNumericLimits<SizeType>::Max(), std::memory_order_acquire, std::memory_order_relaxed);
+		}
+
+		void WriteLock() 
+		{
+			while (!TryWriteLock())
+			{
+				// Reduce contention by doing a simple relaxed read to see if we have a chance of being able to lock.
+				while (Lock.load(std::memory_order_relaxed) != 0)
+				{
+					FPlatformProcess::Yield();
+				}
+			}
+		}
+
+		void WriteUnlock() 
+		{
+			Lock.store(0, std::memory_order_release);
+		}
+
+		bool TryReadLock()
+		{
+			SizeType LocalValue = Lock.load(std::memory_order_relaxed);
+			// Check to make sure we don't already have a write lock or that we've not reached the limit of reader locks.
+			if (LocalValue >= TNumericLimits<SizeType>::Max() - 1)
+			{
+				return false;
+			}
+
+			return Lock.compare_exchange_strong(LocalValue, LocalValue + 1, std::memory_order_acquire, std::memory_order_relaxed);
+		}
+
+		void ReadUnlock() 
+		{
+			Lock.fetch_sub(1, std::memory_order_release);
+		}
+
+		void ReadLock()
+		{
+			while (!TryReadLock())
+			{
+				FPlatformProcess::Yield();
+			}
+		}
+
+	private:
+		std::atomic<SizeType> Lock = 0;
+	};
+
+	using FRWSpinLock = TRWSpinLock<uint32>;
+}
diff --git a/Engine/Source/Runtime/Core/Public/Misc/ScopeRWLock.h b/Engine/Source/Runtime/Core/Public/Misc/ScopeRWLock.h
--- a/Engine/Source/Runtime/Core/Public/Misc/ScopeRWLock.h
+++ b/Engine/Source/Runtime/Core/Public/Misc/ScopeRWLock.h
@@ -8,6 +8,88 @@
 #include "AutoRTFM/AutoRTFM.h"
 #include "Templates/Function.h"
 
+namespace UE
+{
+
+// RAII-style scope locking of a synchronisation primitive
+// `MutexType` is required to implement `ReadLock` and `ReadUnlock` methods
+// Example:
+//	{
+//		TReadScopeLock<FRWLock> ScopeLock(RWLock);
+//		...
+//	}
+template<typename MutexType>
+class TReadScopeLock
+{
+public:
+	UE_NONCOPYABLE(TReadScopeLock);
+
+	UE_NODISCARD_CTOR TReadScopeLock(MutexType& InMutex)
+		: Mutex(&InMutex)
+	{
+		check(Mutex);
+		Mutex->ReadLock();
+	}
+
+	~TReadScopeLock()
+	{
+		ReadUnlock();
+	}
+
+	void ReadUnlock()
+	{
+		if (Mutex)
+		{
+			Mutex->ReadUnlock();
+			Mutex = nullptr;
+		}
+	}
+
+private:
+	MutexType* Mutex;
+};
+
+
+// RAII-style scope locking of a synchronisation primitive
+// `MutexType` is required to implement `WriteLock` and `WriteUnlock` methods
+// Example:
+//	{
+//		TWriteScopeLock<FRWLock> ScopeLock(RWLock);
+//		...
+//	}
+template<typename MutexType>
+class TWriteScopeLock
+{
+public:
+	UE_NONCOPYABLE(TWriteScopeLock);
+
+	UE_NODISCARD_CTOR TWriteScopeLock(MutexType& InMutex)
+		: Mutex(&InMutex)
+	{
+		check(Mutex);
+		Mutex->WriteLock();
+	}
+
+	~TWriteScopeLock()
+	{
+		WriteUnlock();
+	}
+
+	void WriteUnlock()
+	{
+		if (Mutex)
+		{
+			Mutex->WriteUnlock();
+			Mutex = nullptr;
+		}
+	}
+
+private:
+	MutexType* Mutex;
+};
+
+} // namespace UE
+
 /** Keeps a FRWLock read-locked while this scope lives */
 class FReadScopeLock
 {
@@ -61,6 +143,56 @@ enum FRWScopeLockType
 	SLT_Write,
 };
 
+namespace UE 
+{
+
+// RAII-style scope locking of a synchronisation primitive
+// `MutexType` is required to implement `ReadLock` `WriteLock` `ReadUnlock` and `WriteUnlock` methods
+// Example:
+//	{
+//		TRWScopeLock<FRWLock> ScopeLock(RWLock, SLT_ReadOnly or SLT_Write);
+//		...
+//	}
+template<typename MutexType>
+class TRWScopeLock
+{
+public:
+	UE_NONCOPYABLE(TRWScopeLock);
+
+	UE_NODISCARD_CTOR TRWScopeLock(MutexType& InMutex, FRWScopeLockType InLockType)
+		: Mutex(&InMutex)
+		, LockType(InLockType)
+	{
+		check(Mutex);
+		if (LockType == SLT_ReadOnly)
+		{
+			Mutex->ReadLock();
+		}
+		else
+		{
+			Mutex->WriteLock();
+		}
+	}
+
+	~TRWScopeLock()
+	{
+		if (LockType == SLT_ReadOnly)
+		{
+			Mutex->ReadUnlock();
+		}
+		else
+		{
+			Mutex->WriteUnlock();
+		}
+	}
+
+private:
+	MutexType* Mutex;
+	FRWScopeLockType LockType;
+};
+
+} // namespace UE
+
 /**
  * Keeps a FRWLock read- or write-locked while this scope lives
  *
