- Support for marking arrays as initialized to empty, so LoadConfig will clear out the array instead of leaving it alone, by using the new ^ config command:
- ^SomeArray=empty
- This will tell GConfig to know that lack of entries means to clear the property in LoadConfig, instead of leaving the array alone (the normal behavior when no entries exist in GConfig for an array)

diff --git a/Engine/Source/Runtime/Core/Private/Misc/ConfigCacheIni.cpp b/Engine/Source/Runtime/Core/Private/Misc/ConfigCacheIni.cpp
--- a/Engine/Source/Runtime/Core/Private/Misc/ConfigCacheIni.cpp
+++ b/Engine/Source/Runtime/Core/Private/Misc/ConfigCacheIni.cpp
@@ -53,6 +53,7 @@ namespace
 		{ '!', FConfigValue::EValueType::Clear },
 		{ '@', FConfigValue::EValueType::ArrayOfStructKey },
 		{ '*', FConfigValue::EValueType::POCArrayOfStructKey },
+		{ '^', FConfigValue::EValueType::InitializeToEmpty },
 	};
 }
 
@@ -1412,6 +1413,17 @@ void FConfigFile::ProcessCommand(FConfigSection* Section, FStringView SectionNam
 			// Remove if present.
 			Section->Remove(Key);
 			Section->CompactStable();
+			
+			// clear any empty initialization, which will reset it to using code defaults if nothing else adds anything else
+			Section->EmptyInitializedKeys.Remove(Key);
+			break;
+		case FConfigValue::EValueType::InitializeToEmpty:
+			// track a key to show uniqueness for arrays of structs
+			Section->EmptyInitializedKeys.Add(Key);
+			
+			// also clear any existing entries
+			Section->Remove(Key);
+			Section->CompactStable();
 			break;
 		case FConfigValue::EValueType::ArrayOfStructKey:
 			// track a key to show uniqueness for arrays of structs
@@ -1622,7 +1634,7 @@ void FillFileFromBuffer(FileType* File, FStringView Buffer, bool bHandleSymbolCo
 					if (bHandleSymbolCommands)
 					{
 						TCHAR Cmd = Start[0];
-						if (Cmd == '+' || Cmd == '-' || Cmd == '.' || Cmd == '!' || Cmd == '@' || Cmd == '*')
+						if (Cmd == '+' || Cmd == '-' || Cmd == '.' || Cmd == '!' || Cmd == '@' || Cmd == '*' || Cmd == '^')
 						{
 							Start++;
 							LineLen--;
@@ -2615,18 +2627,17 @@ int32 FConfigFile::GetArray(const TCHAR* Section, const TCHAR* Key, TArray<FStri
 {
 	Value.Empty();
 	const FConfigSection* Sec = FindSection(Section);
+	FName KeyName(Key);
 	if (Sec != nullptr)
 	{
-		Sec->MultiFind(Key, Value, true);
+		Sec->MultiFind(KeyName, Value, true);
+		// if we have values, or we were initialized to empty
+		return Value.Num() || Sec->EmptyInitializedKeys.Contains(KeyName);
 	}
 #if !UE_BUILD_SHIPPING
-	else
-	{
-		CheckLongSectionNames(Section, this);
-	}
+	CheckLongSectionNames(Section, this);
 #endif
-
-	return Value.Num();
+	return false;
 }
 
 bool FConfigFile::DoesSectionExist(const TCHAR* Section) const
diff --git a/Engine/Source/Runtime/Core/Public/Misc/ConfigCacheIni.h b/Engine/Source/Runtime/Core/Public/Misc/ConfigCacheIni.h
--- a/Engine/Source/Runtime/Core/Public/Misc/ConfigCacheIni.h
+++ b/Engine/Source/Runtime/Core/Public/Misc/ConfigCacheIni.h
@@ -134,6 +134,9 @@ public:
 		Remove,
 		// !Foo=ClearArray
 		Clear,
+		// ^Array=Empty, means that this will clear any existing entries from an array property, unlike
+		// Clear which will leave the array propperty untouched (as nothing is in GConfig for the key) 
+		InitializeToEmpty,
 		// @Array=StructKey
 		ArrayOfStructKey,
 		// *Array=PerObjectConfigStructKey
@@ -466,6 +469,11 @@ public:
 #if UE_WITH_CONFIG_TRACKING
 	TRefCountPtr<UE::ConfigAccessTracking::FSection> SectionAccess;
 #endif
+	
+	// a set of (array) keys that have been initialized to empty - this is needed to we can tell the difference between
+	// wanting the array to be cleared out in LoadConfig, vs not being specified at all. a bool can be initialized to false
+	// or not specified at all, but an array has no way to differentiate without this
+	TSet<FName> EmptyInitializedKeys;
 
 	// dyanmic modification will disable saving for this section
 	bool bCanSave = true;
@@ -588,7 +596,7 @@ private:
 	// This holds per-object config class names, with their ArrayOfStructKeys. Since the POC sections are all unique,
 	// we can't track it just in that section. This is expected to be empty/small
 	TMap<FString, TMap<FName, FString> > PerObjectConfigArrayOfStructKeys;
-
+	
 	// if this is set, then we track changes made to sections for saving or replaying later (currently unused)
 	FConfigCommandStream* ChangeTracker = nullptr;
 	
diff --git a/Engine/Source/Runtime/CoreUObject/Private/UObject/Obj.cpp b/Engine/Source/Runtime/CoreUObject/Private/UObject/Obj.cpp
--- a/Engine/Source/Runtime/CoreUObject/Private/UObject/Obj.cpp
+++ b/Engine/Source/Runtime/CoreUObject/Private/UObject/Obj.cpp
@@ -3192,8 +3192,9 @@ void UObject::LoadConfig(UClass* ConfigClass/*=nullptr*/, const TCHAR* InFilenam
 
 					const int32 Size = Array->Inner->GetElementSize();
 
-					// Only override default properties if there is something to override them with.
-					if (!List.IsEmpty())
+					// Only override default properties if there is something to override them with, or
+					// if the config system had a "initialize to empty" entry
+					if (!List.IsEmpty() || Sec->EmptyInitializedKeys.Contains(KeyName))
 					{
 						ArrayHelper.EmptyAndAddValues(List.Num());
 						for (int32 i = List.Num() - 1, c = 0; i >= 0; i--, c++)
