[AutoRTFM] Avoid calling `GetCurrentTransaction` accessor twice in a row.

Unfortunately, `GetCurrentTransaction` is not just a simple getter; it has additional logic to materialize the transaction if necessary. We can avoid calling it multiple times to prevent redundant materialization checks.

diff --git a/Engine/Source/Runtime/Core/Private/AutoRTFM/ContextInlines.h b/Engine/Source/Runtime/Core/Private/AutoRTFM/ContextInlines.h
--- a/Engine/Source/Runtime/Core/Private/AutoRTFM/ContextInlines.h
+++ b/Engine/Source/Runtime/Core/Private/AutoRTFM/ContextInlines.h
@@ -46,12 +46,11 @@ UE_AUTORTFM_FORCEINLINE void FContext::DidFree(void* LogicalAddress)
 {
     // We can do free's in the open within a transaction *during* when the
     // transaction itself is being destroyed, so we need to check for that case.
-	if (UNLIKELY(!GetCurrentTransaction()))
+	FTransaction* Transaction = GetCurrentTransaction();
+	if (LIKELY(Transaction))
 	{
-		return;
+		Transaction->DidFree(LogicalAddress);
 	}
-
-    GetCurrentTransaction()->DidFree(LogicalAddress);
 }
 
 UE_AUTORTFM_FORCEINLINE bool FContext::AttemptToCommitTransaction(FTransaction* const Transaction)
