add virtual qualifier to ConvertToPlatformPath so when we call OpenMappedEx on IOS we are fixing the path

diff --git a/Engine/Source/Runtime/Core/Public/Apple/ApplePlatformFile.h b/Engine/Source/Runtime/Core/Public/Apple/ApplePlatformFile.h
--- a/Engine/Source/Runtime/Core/Public/Apple/ApplePlatformFile.h
+++ b/Engine/Source/Runtime/Core/Public/Apple/ApplePlatformFile.h
@@ -53,7 +53,7 @@ public:
 
 	virtual bool CopyFile(const TCHAR* To, const TCHAR* From, EPlatformFileRead ReadFlags = EPlatformFileRead::None, EPlatformFileWrite WriteFlags = EPlatformFileWrite::None) override;
 private:
-	FString ConvertToPlatformPath(const FString& Filename, bool bForWrite, bool bIsPublicWrite) { return Filename; }
+	virtual FString ConvertToPlatformPath(const FString& Filename, bool bForWrite, bool bIsPublicWrite) { return Filename; }
 	bool IterateDirectoryCommon(const TCHAR* Directory, const TFunctionRef<bool(struct dirent*)>& Visitor);
 
 	int32 Stat(const TCHAR* Filename, struct stat* OutFileInfo);
diff --git a/Engine/Source/Runtime/Core/Public/IOS/IOSPlatformFile.h b/Engine/Source/Runtime/Core/Public/IOS/IOSPlatformFile.h
--- a/Engine/Source/Runtime/Core/Public/IOS/IOSPlatformFile.h
+++ b/Engine/Source/Runtime/Core/Public/IOS/IOSPlatformFile.h
@@ -55,7 +55,7 @@ public:
 	
 private:
 	bool IterateDirectoryCommon(const TCHAR* Directory, const TFunctionRef<bool(struct dirent*)>& Visitor);
-	FString ConvertToPlatformPath(const FString& Filename, bool bForWrite, bool bIsPublicWrite);
+	virtual FString ConvertToPlatformPath(const FString& Filename, bool bForWrite, bool bIsPublicWrite) override;
 
 	FString ReadDirectory;
 	FString WriteDirectory;
