AndroidPlatformMisc non unity compilation fix

diff --git a/Engine/Source/Runtime/Core/Private/Android/AndroidPlatformMisc.cpp b/Engine/Source/Runtime/Core/Private/Android/AndroidPlatformMisc.cpp
--- a/Engine/Source/Runtime/Core/Private/Android/AndroidPlatformMisc.cpp
+++ b/Engine/Source/Runtime/Core/Private/Android/AndroidPlatformMisc.cpp
@@ -41,6 +41,7 @@
 
 #include "Misc/CoreDelegates.h"
 #include "Async/TaskGraphInterfaces.h"
+#include "Containers/Ticker.h"
 
 #include "FramePro/FrameProProfiler.h"
 #include <sys/mman.h>
