Changed LinearAllocator API to simplify it.
Removed 2MB allocation granularity for cases when VLPA is enabled. Testing showed it didn't improve anything in terms of perf as FNames were moved out of the PersistentLinearAllocator.
Added LinearBlockAllocator for platforms that can't use FVirtualMemoryBlock. It allocates 64KB blocks and serve all allocations from them.
This saves memory for persistent allocations

diff --git a/Engine/Source/Runtime/Core/Private/HAL/UnrealMemory.cpp b/Engine/Source/Runtime/Core/Private/HAL/UnrealMemory.cpp
--- a/Engine/Source/Runtime/Core/Private/HAL/UnrealMemory.cpp
+++ b/Engine/Source/Runtime/Core/Private/HAL/UnrealMemory.cpp
@@ -677,14 +677,9 @@ void* FMemory::MallocPersistentAuxiliary(SIZE_T InSize, uint32 InAlignment)
 	return GetPersistentLinearAllocator().Allocate(InSize, InAlignment);
 }
 
-void FMemory::FreePersistentAuxiliary(void* InPtr)
-{
-	GetPersistentLinearAllocator().TryDeallocate(InPtr, 0);
-}
-
 bool FMemory::IsPersistentAuxiliaryActive()
 {
-	return GetPersistentLinearAllocator().IsInitialized();
+	return GPersistentLinearAllocatorExtends.Size != 0;
 }
 
 SIZE_T FMemory::GetUsedPersistentAuxiliary()
@@ -694,5 +689,5 @@ SIZE_T FMemory::GetUsedPersistentAuxiliary()
 
 
 #if !INLINE_FMEMORY_OPERATION && !PLATFORM_USES_FIXED_GMalloc_CLASS
-#include "HAL/FMemory.inl"
+#	include "HAL/FMemory.inl"
 #endif
diff --git a/Engine/Source/Runtime/Core/Private/Memory/LinearAllocator.cpp b/Engine/Source/Runtime/Core/Private/Memory/LinearAllocator.cpp
--- a/Engine/Source/Runtime/Core/Private/Memory/LinearAllocator.cpp
+++ b/Engine/Source/Runtime/Core/Private/Memory/LinearAllocator.cpp
@@ -5,14 +5,7 @@
 
 CORE_API FPersistentLinearAllocatorExtends GPersistentLinearAllocatorExtends;
 
-static struct FInitializer
-{
-	FInitializer()
-	{
-		GPersistentLinearAllocatorExtends.Address = (uint64)GetPersistentLinearAllocator().GetBasePointer();
-		GPersistentLinearAllocatorExtends.Size = (uint64)GetPersistentLinearAllocator().GetReservedMemorySize();
-	}
-} Initializer;
+static constexpr SIZE_T LinearAllocatorBlockSize = 64 * 1024;
 
 
 #if UE_ENABLE_LINEAR_VIRTUAL_ALLOCATOR
@@ -21,6 +14,16 @@ static struct FInitializer
 #include "Misc/ScopeLock.h"
 #include "HAL/LowLevelMemTracker.h"
 
+struct FPersistentLinearAllocator : public FLinearAllocator
+{
+	FPersistentLinearAllocator(SIZE_T ReserveMemorySize)
+		: FLinearAllocator(ReserveMemorySize)
+	{
+		GPersistentLinearAllocatorExtends.Address = (uint64)VirtualMemory.GetVirtualPointer();
+		GPersistentLinearAllocatorExtends.Size = (uint64)Reserved;
+	}
+};
+
 FLinearAllocator::FLinearAllocator(SIZE_T ReserveMemorySize)
 	: Reserved(ReserveMemorySize)
 {
@@ -29,22 +32,19 @@ FLinearAllocator::FLinearAllocator(SIZE_T ReserveMemorySize)
 		VirtualMemory = VirtualMemory.AllocateVirtual(ReserveMemorySize);
 		if (!VirtualMemory.GetVirtualPointer())
 		{
-			UE_LOG(LogMemory, Warning, TEXT("LinearVirtualMemoryAllocator failed to reserve %u MB and will default to FMemory::Malloc instead"), ReserveMemorySize / 1024 / 1024);
+			UE_LOG(LogMemory, Warning, TEXT("LinearAllocator failed to reserve %u MB and will default to FMemory::Malloc instead"), ReserveMemorySize / 1024 / 1024);
 			Reserved = 0;
 		}
 	}
 	else
 	{
 #if PLATFORM_IOS || PLATFORM_TVOS
-		UE_LOG(LogMemory, Warning, TEXT("LinearVirtualMemoryAllocator requires com.apple.developer.kernel.extended-virtual-addressing entitlement to work"));
+		UE_LOG(LogMemory, Warning, TEXT("LinearAllocator requires com.apple.developer.kernel.extended-virtual-addressing entitlement to work"));
 #else
-		UE_LOG(LogMemory, Warning, TEXT("This platform does not allow to allocate more virtual memory than there is physical memory. LinearVirtualMemoryAllocator will default to FMemory::Malloc instead"));
+		UE_LOG(LogMemory, Warning, TEXT("This platform does not allow to allocate more virtual memory than there is physical memory. LinearAllocator will default to FMemory::Malloc instead"));
 #endif
 		Reserved = 0;
 	}
-
-	GPersistentLinearAllocatorExtends.Address = (uint64)VirtualMemory.GetVirtualPointer();
-	GPersistentLinearAllocatorExtends.Size = (uint64)Reserved;
 }
 
 void* FLinearAllocator::Allocate(SIZE_T Size, uint32 Alignment)
@@ -61,18 +61,8 @@ void* FLinearAllocator::Allocate(SIZE_T Size, uint32 Alignment)
 				if (NewOffset > Committed)
 				{
 					LLM_PLATFORM_SCOPE(ELLMTag::FMalloc);
-					const size_t CommitGranularity = UE_USE_VERYLARGEPAGEALLOCATOR ? 2ul * 1024 * 1024 : 65536ul;
-					const SIZE_T ToCommit = Align(NewOffset - Committed, FMath::Max(VirtualMemory.GetCommitAlignment(), CommitGranularity));
-				#if UE_USE_VERYLARGEPAGEALLOCATOR
-					if (!VirtualMemory.Commit(Committed, ToCommit, false))
-					{
-						// do not try to use linear allocator anymore and fallback to FMemory::Malloc from now on
-						Reserved = Committed;
-						return FMemory::Malloc(Size, Alignment);
-					}
-				#else
+					const SIZE_T ToCommit = Align(NewOffset - Committed, FMath::Max((SIZE_T)VirtualMemory.GetCommitAlignment(), LinearAllocatorBlockSize));
 					VirtualMemory.Commit(Committed, ToCommit);
-				#endif
 					Committed += ToCommit;
 					LLM_IF_ENABLED(FLowLevelMemTracker::Get().OnLowLevelAlloc(ELLMTracker::Platform, (uint8*)VirtualMemory.GetVirtualPointer() + CurrentOffset, ToCommit));
 				}
@@ -87,56 +77,91 @@ void* FLinearAllocator::Allocate(SIZE_T Size, uint32 Alignment)
 		}
 	}
 
-	ExceedsReservation.fetch_add(Size, std::memory_order_relaxed);
+	static bool bOnce = false;
+	if (!bOnce)
+	{
+		UE_LOG(LogMemory, Warning, TEXT("LinearAllocator exceeded %u MB it reserved. Please tune PersistentAllocatorReserveSizeMB setting in [MemoryPools] ini group. Falling back to FMemory::Malloc"), Reserved / 1024 / 1024);
+		bOnce = true;
+	}
 	return FMemory::Malloc(Size, Alignment);
 }
 
-void FLinearAllocator::PreAllocate(SIZE_T Size, uint32 Alignment)
+bool FLinearAllocator::CanFit(SIZE_T Size, uint32 Alignment) const
+{
+	return (Reserved - Align(CurrentOffset, Alignment)) >= Size;
+}
+
+bool FLinearAllocator::ContainsPointer(const void* Ptr) const
 {
-	Alignment = FMath::Max(Alignment, 8u);
+	return (uintptr_t)Ptr - (uintptr_t)VirtualMemory.GetVirtualPointer() < Reserved;
+}
 
-	FScopeLock AutoLock(&Lock);
-	if (CanFit(Size, Alignment))
+#else
+
+FLinearBlockAllocator::FLinearBlockAllocator(SIZE_T)
+{
+	AllocateNewBlock(LinearAllocatorBlockSize);
+	FirstHeader = CurrentHeader;
+}
+
+FLinearBlockAllocator::~FLinearBlockAllocator()
+{
+	while (FirstHeader)
 	{
-		CurrentOffset = Align(CurrentOffset, Alignment);
-		const SIZE_T NewOffset = CurrentOffset + Size;
-		if (NewOffset > Committed)
-		{
-			const SIZE_T ToCommit = Align(NewOffset - Committed, VirtualMemory.GetCommitAlignment());
-			VirtualMemory.Commit(Committed, ToCommit);
-			Committed += ToCommit;
-		}
+		FBlock* Next = FirstHeader->Next;
+		CurrentBlock = (uint8*)FirstHeader + sizeof(FBlock) - CurrentHeader->LinearAllocatorBlockSize;
+		FMemory::Free(CurrentBlock);
+		FirstHeader = Next;
 	}
 }
 
-bool FLinearAllocator::TryDeallocate(void* Ptr, SIZE_T Size)
+void* FLinearBlockAllocator::Allocate(SIZE_T Size, uint32 Alignment)
 {
-	if (ContainsPointer(Ptr))
+	FScopeLock AutoLock(&Lock);
+	if (!CanFit(Size, Alignment))
 	{
-		FScopeLock AutoLock(&Lock);
-		if ((uint8*)Ptr + Size == (uint8*)VirtualMemory.GetVirtualPointer() + CurrentOffset)
-		{
-			LLM(FLowLevelMemTracker::Get().OnLowLevelFree(ELLMTracker::Default, Ptr, ELLMAllocType::FMalloc));
-			CurrentOffset -= Size;
-			return true;
-		}
+		//TODO: if Size >= BlockSize, allocate a new block and link it between curent and previous block
+		AllocateNewBlock(Size + sizeof(FBlock));
+	}
+
+	CurrentOffset = Align(CurrentOffset, Alignment);
+	uint8* Mem = &CurrentBlock[CurrentOffset];
+	CurrentOffset += Size;
+
+	return Mem;
+}
 
-		return false;
+void FLinearBlockAllocator::AllocateNewBlock(SIZE_T Size)
+{
+	Size = Align(Size, LinearAllocatorBlockSize);
+	FBlock* PreviousHeader = CurrentHeader;
+
+	CurrentBlock = (uint8*)FMemory::Malloc(Size);
+	CurrentHeader = (FBlock*)(CurrentBlock + Size - sizeof(FBlock));
+	CurrentHeader->Next = nullptr;
+	CurrentHeader->LinearAllocatorBlockSize = Size;
+
+	if (PreviousHeader)
+	{
+		PreviousHeader->Next = CurrentHeader;
 	}
 
-	FMemory::Free(Ptr);
-	return true;
+	CurrentOffset = 0;
+	CurrentBlockSize = Size - sizeof(FBlock);
+	TotalAllocated += Size;
 }
 
-bool FLinearAllocator::CanFit(SIZE_T Size, uint32 Alignment) const
+bool FLinearBlockAllocator::CanFit(SIZE_T Size, uint32 Alignment) const
 {
-	return (Reserved - Align(CurrentOffset, Alignment)) >= Size;
+	return (CurrentBlockSize - Align(CurrentOffset, Alignment)) >= Size;
 }
 
+typedef FLinearBlockAllocator FPersistentLinearAllocator;
+
+#endif //~UE_ENABLE_LINEAR_VIRTUAL_ALLOCATOR
+
 FLinearAllocator& GetPersistentLinearAllocator()
 {
-	static FLinearAllocator GPersistentLinearAllocator(UE_PERSISTENT_ALLOCATOR_RESERVE_SIZE);
+	static FPersistentLinearAllocator GPersistentLinearAllocator(UE_PERSISTENT_ALLOCATOR_RESERVE_SIZE);
 	return GPersistentLinearAllocator;
-}
-
-#endif //~UE_ENABLE_LINEAR_VIRTUAL_ALLOCATOR
\ No newline at end of file
+}
\ No newline at end of file
diff --git a/Engine/Source/Runtime/Core/Public/HAL/UnrealMemory.h b/Engine/Source/Runtime/Core/Public/HAL/UnrealMemory.h
--- a/Engine/Source/Runtime/Core/Public/HAL/UnrealMemory.h
+++ b/Engine/Source/Runtime/Core/Public/HAL/UnrealMemory.h
@@ -283,7 +283,7 @@ struct FMemory
 	static CORE_API void* MallocPersistentAuxiliary(SIZE_T InSize, uint32 InAlignment = 0);
 
 	UE_DEPRECATED(5.5, "Persistent Auxiliary allocator is obsolete and is replaced by a GetPersistentLinearAllocator()")
-	static CORE_API void FreePersistentAuxiliary(void* InPtr);
+	static inline CORE_API void FreePersistentAuxiliary(void* InPtr) {}
 
 	UE_DEPRECATED(5.5, "Persistent Auxiliary allocator is obsolete and is replaced by a GetPersistentLinearAllocator()")
 	static CORE_API bool IsPersistentAuxiliaryActive();
diff --git a/Engine/Source/Runtime/Core/Public/Memory/LinearAllocator.h b/Engine/Source/Runtime/Core/Public/Memory/LinearAllocator.h
--- a/Engine/Source/Runtime/Core/Public/Memory/LinearAllocator.h
+++ b/Engine/Source/Runtime/Core/Public/Memory/LinearAllocator.h
@@ -4,99 +4,84 @@
 
 #include "HAL/Platform.h"
 #include "HAL/UnrealMemory.h"
+#include "HAL/CriticalSection.h"
 
 #ifndef UE_ENABLE_LINEAR_VIRTUAL_ALLOCATOR
-	#define UE_ENABLE_LINEAR_VIRTUAL_ALLOCATOR PLATFORM_HAS_FPlatformVirtualMemoryBlock
+#	define UE_ENABLE_LINEAR_VIRTUAL_ALLOCATOR PLATFORM_HAS_FPlatformVirtualMemoryBlock
 #endif
 
+
 #if UE_ENABLE_LINEAR_VIRTUAL_ALLOCATOR
-#include <atomic>
-#include "HAL/CriticalSection.h"
 
 struct FLinearAllocator
 {
 	CORE_API FLinearAllocator(SIZE_T ReserveMemorySize);
-
 	FORCEINLINE ~FLinearAllocator()
 	{
 		VirtualMemory.FreeVirtual();
 	}
 
 	CORE_API void* Allocate(SIZE_T Size, uint32 Alignment = 8);
-	CORE_API void  PreAllocate(SIZE_T Size, uint32 Alignment = 8);
-
-	// This will succeed only when deallocating the last allocation
-	CORE_API bool  TryDeallocate(void* Ptr, SIZE_T Size);
-
-	FORCEINLINE bool ContainsPointer(const void* Ptr) const
-	{
-		return (uintptr_t)Ptr - (uintptr_t)VirtualMemory.GetVirtualPointer() < Reserved;
-	}
-
-	FORCEINLINE SIZE_T GetExceedingSize() const
-	{
-		return ExceedsReservation.load(std::memory_order_relaxed);
-	}
-
-	FORCEINLINE bool IsInitialized() const
-	{
-		return Reserved != 0;
-	}
 
 	FORCEINLINE SIZE_T GetAllocatedMemorySize() const
 	{
 		return Committed;
 	}
 
-	FORCEINLINE SIZE_T GetReservedMemorySize() const
-	{
-		return Reserved;
-	}
-
-	FORCEINLINE const void* GetBasePointer() const
-	{
-		return VirtualMemory.GetVirtualPointer();
-	}
-
-private:
+protected:
 	FCriticalSection Lock;
 	FPlatformMemory::FPlatformVirtualMemoryBlock VirtualMemory;
 	SIZE_T Reserved;
 	SIZE_T Committed = 0;
 	SIZE_T CurrentOffset = 0;
-	std::atomic<SIZE_T> ExceedsReservation = 0;
 
 	bool CanFit(SIZE_T Size, uint32 Alignment) const;
+	bool ContainsPointer(const void* Ptr) const;
 };
 
-CORE_API FLinearAllocator& GetPersistentLinearAllocator();
-
 #else
-// stub implementation with most functions being nop
-struct FLinearAllocator
+
+struct FLinearBlockAllocator
 {
-	FORCEINLINE FLinearAllocator(SIZE_T) {}
-
-	FORCEINLINE void*	Allocate(SIZE_T Size, uint32 Alignment = 8) { return FMemory::Malloc(Size, Alignment); }
-	FORCEINLINE void	PreAllocate(SIZE_T, uint32) {}
-	FORCEINLINE bool	TryDeallocate(void* Ptr, SIZE_T Size)	{ FMemory::Free(Ptr); return true; }
-
-	FORCEINLINE bool	ContainsPointer(const void*) const	{ return false; }
-	FORCEINLINE SIZE_T	GetExceedingSize() const			{ return 0; }
-	FORCEINLINE bool	IsInitialized() const				{ return false; }
-	FORCEINLINE SIZE_T	GetAllocatedMemorySize() const		{ return 0; }
-	FORCEINLINE SIZE_T	GetReservedMemorySize() const		{ return 0; }
-	FORCEINLINE const void* GetBasePointer() const			{ return nullptr; }
+	CORE_API FLinearBlockAllocator(SIZE_T);
+	CORE_API ~FLinearBlockAllocator();
+
+	CORE_API void* Allocate(SIZE_T Size, uint32 Alignment = 8);
+
+	FORCEINLINE SIZE_T GetAllocatedMemorySize() const
+	{
+		return TotalAllocated;
+	}
+
+private:
+	struct FBlock
+	{
+		FBlock* Next;
+		SIZE_T BlockSize;
+	};
+
+	FCriticalSection Lock;
+	FBlock* FirstHeader = nullptr;
+	FBlock* CurrentHeader = nullptr;
+	uint8* CurrentBlock = nullptr;
+	SIZE_T CurrentOffset = 0;
+	SIZE_T CurrentBlockSize = 0;
+	SIZE_T TotalAllocated = 0;
+
+	void AllocateNewBlock(SIZE_T Size);
+	bool CanFit(SIZE_T Size, uint32 Alignment) const;
 };
 
-FORCEINLINE FLinearAllocator GetPersistentLinearAllocator() { return FLinearAllocator(0); }
+typedef FLinearBlockAllocator FLinearAllocator;
 
 #endif	//~UE_ENABLE_LINEAR_VIRTUAL_ALLOCATOR
 
+CORE_API FLinearAllocator& GetPersistentLinearAllocator();
+
 struct FPersistentLinearAllocatorExtends
 {
-	uint64 Address;
-	uint64 Size;
+	uint64 Address = 0;
+	uint64 Size = 0;
 };
 
 // Special case for the FPermanentObjectPoolExtents to reduce the amount of pointer dereferencing
diff --git a/Engine/Source/Runtime/CoreUObject/Private/UObject/UObjectAllocator.cpp b/Engine/Source/Runtime/CoreUObject/Private/UObject/UObjectAllocator.cpp
--- a/Engine/Source/Runtime/CoreUObject/Private/UObject/UObjectAllocator.cpp
+++ b/Engine/Source/Runtime/CoreUObject/Private/UObject/UObjectAllocator.cpp
@@ -7,25 +7,11 @@
 #include "UObject/UObjectAllocator.h"
 #include "UObject/UObjectGlobals.h"
 
-DEFINE_LOG_CATEGORY_STATIC(LogUObjectAllocator, Log, All);
-
 /** Global UObjectBase allocator							*/
 COREUOBJECT_API FUObjectAllocator GUObjectAllocator;
 
 static bool GPersistentAllocatorIsDisabled = false;
 
-/**
- * Prints a debugf message to allow tuning
- */
-void FUObjectAllocator::BootMessage()
-{
-	const SIZE_T ExceedingSize = GetPersistentLinearAllocator().GetExceedingSize();
-	if (ExceedingSize > 0)
-	{
-		UE_LOG(LogUObjectAllocator, Warning, TEXT("Persistent memory pool exceeded by %u KB, please tune PersistentAllocatorReserveSizeMB setting in [MemoryPools] ini group."), ExceedingSize / 1024);
-	}
-}
-
 void FUObjectAllocator::DisablePersistentAllocator()
 {
 	GPersistentAllocatorIsDisabled = true;
diff --git a/Engine/Source/Runtime/CoreUObject/Private/UObject/UObjectArray.cpp b/Engine/Source/Runtime/CoreUObject/Private/UObject/UObjectArray.cpp
--- a/Engine/Source/Runtime/CoreUObject/Private/UObject/UObjectArray.cpp
+++ b/Engine/Source/Runtime/CoreUObject/Private/UObject/UObjectArray.cpp
@@ -181,8 +181,6 @@ void FUObjectArray::CloseDisregardForGC()
 			{
 				UE_LOG(LogUObjectArray, Log, TEXT("%i objects are not in the root set, but can never be destroyed because they are in the DisregardForGC set."), NumAlwaysLoadedObjects - NumRootObjects);
 			}
-
-			GUObjectAllocator.BootMessage();
 		}
 	}
 
diff --git a/Engine/Source/Runtime/CoreUObject/Public/UObject/UObjectAllocator.h b/Engine/Source/Runtime/CoreUObject/Public/UObject/UObjectAllocator.h
--- a/Engine/Source/Runtime/CoreUObject/Public/UObject/UObjectAllocator.h
+++ b/Engine/Source/Runtime/CoreUObject/Public/UObject/UObjectAllocator.h
@@ -26,7 +26,8 @@ public:
 	/**
 	 * Prints a debugf message to allow tuning
 	 */
-	COREUOBJECT_API void BootMessage();
+	UE_DEPRECATED(5.6, "BootMessage is obsolete now")
+	COREUOBJECT_API void BootMessage() {}
 
 	/**
 	 * Disables allocation of objects from the persistend allocator
