[RWSpinLock]
- Document recursion support contract

#rnx

diff --git a/Engine/Source/Runtime/Core/Public/Misc/RWSpinLock.h b/Engine/Source/Runtime/Core/Public/Misc/RWSpinLock.h
--- a/Engine/Source/Runtime/Core/Public/Misc/RWSpinLock.h
+++ b/Engine/Source/Runtime/Core/Public/Misc/RWSpinLock.h
@@ -14,6 +14,7 @@ namespace UE
 	// WARNING: Should be used only for very short locks.
 	// Use with `TRWScopeLock`, `TWriteScopeLock` or `TReadScopeLock`.
 	// The SizeType template type dictates the space taken by the spinlock but also its maximum number of possible concurrent readers (i.e. 254 for uint8, etc...).
+	// Read locks support recursion. Write locks don't support recursion even if coming from the same thread currently owning the write lock.
 	template <typename SizeType = uint32>
 	class TRWSpinLock
 	{
