[RaceDetector/Runtime]
- This is the runtime for a race detection system that works both for monolithic and non-monolithic Windows builds.
- The runtime needs to be paired with a Clang compiler that instrument each memory access and calls entrypoints in the runtime.
- The clang compiler is currently out of tree and needs to be set with LLVM_CUSTOM_PATH
ex: RunUBT LyraEditor Win64 Development -Compiler=ClangCustom -EnableInstrumentation
- The instrumentation framework is made in such a way that it is possible to reuse for different tools.
- To active at boot, use -racedetector on the command-line
- To activate from the console, type r.RaceDetector.Activate 10, to activate for 10 seconds
- To break on race detection, use r.RaceDetector.BreakOnRace 1
- Races are reported in the debugger output window and callstacks can be clicked on to go to source.
- Added support for inline frames when generating stack trace on the windows platform

Special thanks to [REDACTED_PERSON] for the initial POC that led to this work.

diff --git a/Engine/Source/Programs/UnrealBuildTool/Configuration/UEBuildModuleCPP.cs b/Engine/Source/Programs/UnrealBuildTool/Configuration/UEBuildModuleCPP.cs
--- a/Engine/Source/Programs/UnrealBuildTool/Configuration/UEBuildModuleCPP.cs
+++ b/Engine/Source/Programs/UnrealBuildTool/Configuration/UEBuildModuleCPP.cs
@@ -678,7 +678,23 @@ namespace UnrealBuildTool
 					"#endif\r\n" +
 					"#if defined(PER_MODULE_INLINE_FILE) && !defined(SUPPRESS_PER_MODULE_INLINE_FILE)\r\n" +
 					"#include PER_MODULE_INLINE_FILE\r\n" +
-					"#endif";
+					"#endif\r\n"
+					;
+
+				if (Target.WindowsPlatform.bEnableInstrumentation)
+				{
+					CompileEnvironment.UserIncludePaths.Add(DirectoryReference.Combine(Unreal.EngineSourceDirectory, "Runtime", "Instrumentation", "Public"));
+					CompileEnvironment.UserIncludePaths.Add(DirectoryReference.Combine(Unreal.EngineSourceDirectory, "Runtime", "Instrumentation", "Internal"));
+					CompileEnvironment.UserIncludePaths.Add(DirectoryReference.Combine(Unreal.EngineSourceDirectory, "Runtime", "Core", "Public"));
+					CompileEnvironment.UserIncludePaths.Add(DirectoryReference.Combine(Unreal.EngineSourceDirectory, "Runtime", "Core", "Internal"));
+
+					// Give the ability for core to hotpatch itself for the per module functions to save another set of JMPs.
+					if (Name == "Core")
+					{ 
+						Content += "#define USE_INSTRUMENTATION_PERMODULE_HOTPATCH 1\r\n";
+					}
+					Content += "#include \"Instrumentation/PerModuleInline.inl\"\r\n";
+				}
 
 				bool bMergeModules = Target.bMergeModules;
 				FileItem PerModuleFile = Graph.CreateIntermediateTextFile(NewDeleteOverrides, Content, bMergeModules);
diff --git a/Engine/Source/Programs/UnrealBuildTool/Platform/Windows/MicrosoftPlatformSDK.cs b/Engine/Source/Programs/UnrealBuildTool/Platform/Windows/MicrosoftPlatformSDK.cs
--- a/Engine/Source/Programs/UnrealBuildTool/Platform/Windows/MicrosoftPlatformSDK.cs
+++ b/Engine/Source/Programs/UnrealBuildTool/Platform/Windows/MicrosoftPlatformSDK.cs
@@ -774,6 +774,19 @@ namespace UnrealBuildTool
 							}
 						}
 					}
+					else if (Compiler == WindowsCompiler.ClangCustom)
+					{
+						// Check for a manual installation to a custom directory
+						string? LlvmPath = Environment.GetEnvironmentVariable("LLVM_CUSTOM_PATH");
+						if (!String.IsNullOrEmpty(LlvmPath))
+						{
+							AddClangToolChain(Compiler, new DirectoryReference(LlvmPath), ToolChains, IsAutoSdk: false, Logger);
+						}
+						else
+						{
+							throw new BuildException("ClangCustom needs LLVM_CUSTOM_PATH environment variable to be defined");
+						}
+					}
 					else if (Compiler == WindowsCompiler.Intel)
 					{
 						DirectoryReference InstallDir = DirectoryReference.Combine(DirectoryReference.GetSpecialFolder(Environment.SpecialFolder.ProgramFiles)!, "Intel", "oneAPI", "compiler");
diff --git a/Engine/Source/Programs/UnrealBuildTool/Platform/Windows/UEBuildWindows.cs b/Engine/Source/Programs/UnrealBuildTool/Platform/Windows/UEBuildWindows.cs
--- a/Engine/Source/Programs/UnrealBuildTool/Platform/Windows/UEBuildWindows.cs
+++ b/Engine/Source/Programs/UnrealBuildTool/Platform/Windows/UEBuildWindows.cs
@@ -127,6 +127,11 @@ namespace UnrealBuildTool
 		/// </summary>
 		ClangRTFM,
 
+		/// <summary>
+		/// Use a custom version of Clang found using LLVM_CUSTOM_PATH environment variable.
+		/// </summary>
+		ClangCustom,
+
 		/// <summary>
 		/// Use the Intel oneAPI C++ compiler
 		/// </summary>
@@ -168,7 +173,7 @@ namespace UnrealBuildTool
 		/// <returns>true if Clang based</returns>
 		public static bool IsClang(this WindowsCompiler Compiler)
 		{
-			return Compiler == WindowsCompiler.Clang || Compiler == WindowsCompiler.ClangRTFM || Compiler == WindowsCompiler.Intel;
+			return Compiler == WindowsCompiler.Clang || Compiler == WindowsCompiler.ClangRTFM || Compiler == WindowsCompiler.ClangCustom || Compiler == WindowsCompiler.Intel;
 		}
 
 		/// <summary>
@@ -461,6 +466,13 @@ namespace UnrealBuildTool
 		[CommandLine("-EnableASan")]
 		public bool bEnableAddressSanitizer = false;
 
+		/// <summary>
+		/// Enables instrumentation.
+		/// </summary>
+		[XmlConfigFile(Category = "BuildConfiguration", Name = "bEnableInstrumentation")]
+		[CommandLine("-EnableInstrumentation")]
+		public bool bEnableInstrumentation = false;
+
 		/// <summary>
 		/// Enables LibFuzzer. Only supported for Visual Studio 2022 17.0.0 and up.
 		/// </summary>
@@ -796,6 +808,7 @@ namespace UnrealBuildTool
 			{
 				case WindowsCompiler.Clang:
 				case WindowsCompiler.ClangRTFM:
+				case WindowsCompiler.ClangCustom:
 				case WindowsCompiler.Intel:
 				case WindowsCompiler.VisualStudio2022:
 					return "2015"; // VS2022 is backwards compatible with VS2015 compiler
@@ -918,6 +931,8 @@ namespace UnrealBuildTool
 
 		public bool bEnableAddressSanitizer => Inner.bEnableAddressSanitizer;
 
+		public bool bEnableInstrumentation => Inner.bEnableInstrumentation;
+
 		public bool bEnableLibFuzzer => Inner.bEnableLibFuzzer;
 
 		public bool bWriteSarif => Inner.bWriteSarif;
@@ -1705,6 +1720,12 @@ namespace UnrealBuildTool
 				}
 			}
 
+			// If instrumentation is active, we need to inject the instrumentation module in every other module
+			if (Target.WindowsPlatform.bEnableInstrumentation && ModuleName != "Instrumentation")
+			{
+				Rules.PrivateDependencyModuleNames.Add("Instrumentation");
+			}
+
 			// allow standalone tools to use target platform modules, without needing Engine
 			if (ModuleName == "TargetPlatform")
 			{
diff --git a/Engine/Source/Programs/UnrealBuildTool/Platform/Windows/VCEnvironment.cs b/Engine/Source/Programs/UnrealBuildTool/Platform/Windows/VCEnvironment.cs
--- a/Engine/Source/Programs/UnrealBuildTool/Platform/Windows/VCEnvironment.cs
+++ b/Engine/Source/Programs/UnrealBuildTool/Platform/Windows/VCEnvironment.cs
@@ -228,6 +228,10 @@ namespace UnrealBuildTool
 			{
 				return FileReference.Combine(GetClangToolPath(CompilerDir), "verse-clang-cl.exe");
 			}
+			else if (Compiler == WindowsCompiler.ClangCustom)
+			{
+				return FileReference.Combine(CompilerDir, "bin", "clang-cl.exe");
+			}
 			else if (Compiler == WindowsCompiler.Intel)
 			{
 				return FileReference.Combine(CompilerDir, "bin", "icx.exe");
@@ -244,7 +248,7 @@ namespace UnrealBuildTool
 			{
 				return FileReference.Combine(UnrealBuildBase.Unreal.EngineDirectory, "Binaries", "Win64", "RadDebugger", Architecture.ToString(), "radlink.exe");
 			}
-			else if ((Compiler == WindowsCompiler.Clang || Compiler == WindowsCompiler.ClangRTFM) && bAllowClangLinker)
+			else if ((Compiler == WindowsCompiler.Clang || Compiler == WindowsCompiler.ClangRTFM || Compiler == WindowsCompiler.ClangCustom) && bAllowClangLinker)
 			{
 				return FileReference.Combine(GetClangToolPath(CompilerDir), "lld-link.exe");
 			}
diff --git a/Engine/Source/Programs/UnrealBuildTool/Platform/Windows/VCToolChain.cs b/Engine/Source/Programs/UnrealBuildTool/Platform/Windows/VCToolChain.cs
--- a/Engine/Source/Programs/UnrealBuildTool/Platform/Windows/VCToolChain.cs
+++ b/Engine/Source/Programs/UnrealBuildTool/Platform/Windows/VCToolChain.cs
@@ -695,6 +695,11 @@ namespace UnrealBuildTool
 				Arguments.Add("/d2ExtendedWarningInfo");
 			}
 
+			if (Target.WindowsPlatform.bEnableInstrumentation)
+			{
+				AddDefinition(Arguments, "USING_INSTRUMENTATION=1");
+			}
+
 			// Address sanitizer
 			if (Target.WindowsPlatform.bEnableAddressSanitizer)
 			{
@@ -1632,8 +1637,8 @@ namespace UnrealBuildTool
 				Arguments.Add("/OPT:NOREF");
 			}
 
-			// Identical COMDAT folding
-			if (Target.WindowsPlatform.bMergeIdenticalCOMDATs && !(Target.WindowsPlatform.Compiler.IsIntel() && Target.WindowsPlatform.bSampleBasedPGO && LinkEnvironment.bPGOProfile))
+			// Identical COMDAT folding. Prevent Identical Code Folding when instrumentation is enabled to avoid thunks to all be folded together and prevent hotpatching.
+			if (Target.WindowsPlatform.bMergeIdenticalCOMDATs && !(Target.WindowsPlatform.Compiler.IsIntel() && Target.WindowsPlatform.bSampleBasedPGO && LinkEnvironment.bPGOProfile) && !Target.WindowsPlatform.bEnableInstrumentation)
 			{
 				Arguments.Add("/OPT:ICF");
 			}
diff --git a/Engine/Source/Runtime/Core/Core.Build.cs b/Engine/Source/Runtime/Core/Core.Build.cs
--- a/Engine/Source/Runtime/Core/Core.Build.cs
+++ b/Engine/Source/Runtime/Core/Core.Build.cs
@@ -293,6 +293,11 @@ public class Core : ModuleRules
 
 		IWYUSupport = IWYUSupport.KeepAsIs;
 
+		if (Target.WindowsPlatform.bEnableInstrumentation)
+		{
+			PrivateDependencyModuleNames.Add("Detours");
+		}
+
 		PrivateDefinitions.Add("UE_PERSISTENT_ALLOCATOR_RESERVE_SIZE=" + GetPersistentAllocatorReserveSize().ToString() + "ULL");
 
 		PublicDefinitions.Add("UE_MERGED_MODULES=" + (Target.bMergeModules ? "1" : "0"));
diff --git a/Engine/Source/Runtime/Core/Internal/Instrumentation/Containers.h b/Engine/Source/Runtime/Core/Internal/Instrumentation/Containers.h
new file mode 100644
--- /dev/null
+++ b/Engine/Source/Runtime/Core/Internal/Instrumentation/Containers.h
@@ -0,0 +1,237 @@
+// Copyright Epic Games, Inc. All Rights Reserved.
+
+#pragma once
+
+#include "Misc/CoreMiscDefines.h"
+
+#if USING_INSTRUMENTATION
+
+#include "CoreTypes.h"
+#include "Containers/Map.h"
+
+#include "Instrumentation/Defines.h"
+
+// ------------------------------------------------------------------------------
+// Allocators (we can't use the Engine default allocators if they're instrumented)
+// ------------------------------------------------------------------------------
+#if PLATFORM_WINDOWS
+class FInstrumentationSafeWinAllocator {
+public:
+	static INSTRUMENTATION_FUNCTION_ATTRIBUTES void* Alloc(SIZE_T Size);
+	static INSTRUMENTATION_FUNCTION_ATTRIBUTES void Free(void* Data);
+	static INSTRUMENTATION_FUNCTION_ATTRIBUTES void* Realloc(void* Data, SIZE_T Size, SIZE_T PreviousSize);
+};
+
+using TInstrumentationSafeBaseAllocator = FInstrumentationSafeWinAllocator;
+
+#define SAFE_OPERATOR_NEW_DELETE() \
+	INSTRUMENTATION_FUNCTION_ATTRIBUTES void * operator new(SIZE_T Size)  \
+	{ \
+		return FInstrumentationSafeWinAllocator::Alloc(Size); \
+	} \
+	INSTRUMENTATION_FUNCTION_ATTRIBUTES void   operator delete(void* Ptr) \
+	{ \
+		FInstrumentationSafeWinAllocator::Free(Ptr);  \
+	} 
+
+#else
+class TInstrumentationSafeMallocAllocator {
+public:
+	INSTRUMENTATION_FUNCTION_ATTRIBUTES static void* Realloc(void* Data, SIZE_T Size, SIZE_T PreviousSize)
+	{
+		return realloc(Data, Size);
+	}
+	INSTRUMENTATION_FUNCTION_ATTRIBUTES static void Free(void* Data)
+	{
+		return free(Data);
+	}
+};
+
+using TInstrumentationSafeBaseAllocator = TInstrumentationSafeMallocAllocator;
+#endif
+
+template <int IndexSize = 32, typename TBaseAllocator = TInstrumentationSafeBaseAllocator>
+class TInstrumentationSafeAllocator {
+public:
+	using SizeType = typename TBitsToSizeType<IndexSize>::Type;
+
+private:
+	using USizeType = std::make_unsigned_t<SizeType>;
+
+public:
+	enum { NeedsElementType = true };
+	enum { RequireRangeCheck = true };
+
+	class ForAnyElementType
+	{
+	public:
+		/** Default constructor. */
+		INSTRUMENTATION_FUNCTION_ATTRIBUTES ForAnyElementType()
+			: Data(nullptr)
+		{
+		}
+
+		/**
+		 * Moves the state of another allocator into this one.
+		 * Moves the state of another allocator into this one.  The allocator can be different.
+		 *
+		 * Assumes that the allocator is currently empty, i.e. memory may be allocated but any existing elements have already been destructed (if necessary).
+		 * @param Other - The allocator to move the state from.  This allocator should be left in a valid empty state.
+		 */
+		FORCEINLINE INSTRUMENTATION_FUNCTION_ATTRIBUTES void MoveToEmpty(ForAnyElementType& Other)
+		{
+			if (Data)
+			{
+				TBaseAllocator::Free(Data);
+			}
+
+			Data = Other.Data;
+			Other.Data = nullptr;
+		}
+
+		/** Destructor. */
+		FORCEINLINE INSTRUMENTATION_FUNCTION_ATTRIBUTES ~ForAnyElementType()
+		{
+			if (Data)
+			{
+				TBaseAllocator::Free(Data);
+			}
+		}
+
+		// FContainerAllocatorInterface
+		FORCEINLINE INSTRUMENTATION_FUNCTION_ATTRIBUTES void* GetAllocation() const
+		{
+			return Data;
+		}
+
+		INSTRUMENTATION_FUNCTION_ATTRIBUTES void ResizeAllocation(SizeType PreviousNumElements, SizeType NumElements, SIZE_T NumBytesPerElement)
+		{
+			// Avoid calling FMemory::Realloc( nullptr, 0 ) as ANSI C mandates returning a valid pointer which is not what we want.
+			if (Data || NumElements)
+			{
+				static_assert(sizeof(SizeType) <= sizeof(SIZE_T), "SIZE_T is expected to handle all possible sizes");
+
+				// Check for under/overflow
+				bool bInvalidResize = NumElements < 0 || NumBytesPerElement < 1 || NumBytesPerElement >(SIZE_T)MAX_int32;
+				if constexpr (sizeof(SizeType) == sizeof(SIZE_T))
+				{
+					bInvalidResize = bInvalidResize || (SIZE_T)(USizeType)NumElements > (SIZE_T)TNumericLimits<SizeType>::Max() / NumBytesPerElement;
+				}
+				if (UNLIKELY(bInvalidResize))
+				{
+					UE::Core::Private::OnInvalidSizedHeapAllocatorNum(IndexSize, NumElements, NumBytesPerElement);
+				}
+
+				Data = TBaseAllocator::Realloc(Data, NumElements * NumBytesPerElement, PreviousNumElements * NumBytesPerElement);
+				if (NumElements)
+				{
+					checkSlow(Data);
+				}
+			}
+		}
+		FORCEINLINE INSTRUMENTATION_FUNCTION_ATTRIBUTES SizeType CalculateSlackReserve(SizeType NumElements, SIZE_T NumBytesPerElement) const
+		{
+			return DefaultCalculateSlackReserve(NumElements, NumBytesPerElement, true);
+		}
+		FORCEINLINE INSTRUMENTATION_FUNCTION_ATTRIBUTES SizeType CalculateSlackReserve(SizeType NumElements, SIZE_T NumBytesPerElement, uint32 AlignmentOfElement) const
+		{
+			return DefaultCalculateSlackReserve(NumElements, NumBytesPerElement, true, (uint32)AlignmentOfElement);
+		}
+		FORCEINLINE INSTRUMENTATION_FUNCTION_ATTRIBUTES SizeType CalculateSlackShrink(SizeType NumElements, SizeType NumAllocatedElements, SIZE_T NumBytesPerElement) const
+		{
+			return DefaultCalculateSlackShrink(NumElements, NumAllocatedElements, NumBytesPerElement, true);
+		}
+		FORCEINLINE INSTRUMENTATION_FUNCTION_ATTRIBUTES SizeType CalculateSlackShrink(SizeType NumElements, SizeType NumAllocatedElements, SIZE_T NumBytesPerElement, uint32 AlignmentOfElement) const
+		{
+			return DefaultCalculateSlackShrink(NumElements, NumAllocatedElements, NumBytesPerElement, true, (uint32)AlignmentOfElement);
+		}
+		FORCEINLINE INSTRUMENTATION_FUNCTION_ATTRIBUTES SizeType CalculateSlackGrow(SizeType NumElements, SizeType NumAllocatedElements, SIZE_T NumBytesPerElement) const
+		{
+			return DefaultCalculateSlackGrow(NumElements, NumAllocatedElements, NumBytesPerElement, true);
+		}
+		FORCEINLINE INSTRUMENTATION_FUNCTION_ATTRIBUTES SizeType CalculateSlackGrow(SizeType NumElements, SizeType NumAllocatedElements, SIZE_T NumBytesPerElement, uint32 AlignmentOfElement) const
+		{
+			return DefaultCalculateSlackGrow(NumElements, NumAllocatedElements, NumBytesPerElement, true, (uint32)AlignmentOfElement);
+		}
+
+		INSTRUMENTATION_FUNCTION_ATTRIBUTES SIZE_T GetAllocatedSize(SizeType NumAllocatedElements, SIZE_T NumBytesPerElement) const
+		{
+			return NumAllocatedElements * NumBytesPerElement;
+		}
+
+		INSTRUMENTATION_FUNCTION_ATTRIBUTES bool HasAllocation() const
+		{
+			return !!Data;
+		}
+
+		INSTRUMENTATION_FUNCTION_ATTRIBUTES SizeType GetInitialCapacity() const
+		{
+			return 0;
+		}
+
+	private:
+		ForAnyElementType(const ForAnyElementType&);
+		ForAnyElementType& operator=(const ForAnyElementType&);
+
+		/** A pointer to the container's elements. */
+		void* Data;
+	};
+
+	template<typename ElementType>
+	class ForElementType : public ForAnyElementType
+	{
+	public:
+		/** Default constructor. */
+		INSTRUMENTATION_FUNCTION_ATTRIBUTES ForElementType()
+		{
+		}
+
+		FORCEINLINE INSTRUMENTATION_FUNCTION_ATTRIBUTES ElementType* GetAllocation() const
+		{
+			return (ElementType*)ForAnyElementType::GetAllocation();
+		}
+	};
+};
+
+template <uint32 NumInlineElements, typename BaseAllocator = TInstrumentationSafeBaseAllocator>
+using TInstrumentationSafeInlineAllocator = TSizedInlineAllocator<NumInlineElements, 32, TInstrumentationSafeAllocator<32, BaseAllocator>>;
+
+template <typename ElementType, typename BaseAllocator = TInstrumentationSafeBaseAllocator>
+using TSafeArray = TArray<ElementType, TInstrumentationSafeAllocator<32, BaseAllocator>>;
+
+template<
+	typename BaseAllocator = TInstrumentationSafeBaseAllocator,
+	uint32   AverageNumberOfElementsPerHashBucket = DEFAULT_NUMBER_OF_ELEMENTS_PER_HASH_BUCKET,
+	uint32   BaseNumberOfHashBuckets = DEFAULT_BASE_NUMBER_OF_HASH_BUCKETS,
+	uint32   MinNumberOfHashedElements = DEFAULT_MIN_NUMBER_OF_HASHED_ELEMENTS
+>
+class TInstrumentationSafeSetAllocator
+{
+public:
+
+	/** Computes the number of hash buckets to use for a given number of elements. */
+	static FORCEINLINE INSTRUMENTATION_FUNCTION_ATTRIBUTES uint32 GetNumberOfHashBuckets(uint32 NumHashedElements)
+	{
+		if (NumHashedElements >= MinNumberOfHashedElements)
+		{
+			return FPlatformMath::RoundUpToPowerOfTwo(NumHashedElements / AverageNumberOfElementsPerHashBucket + BaseNumberOfHashBuckets);
+		}
+
+		return 1;
+	}
+
+	using InSparseArrayAllocator = TSparseArrayAllocator<TInstrumentationSafeAllocator<32, BaseAllocator>, TInstrumentationSafeInlineAllocator<4, BaseAllocator>>;
+	using InHashAllocator = TInstrumentationSafeInlineAllocator<1, BaseAllocator>;
+
+	typedef InSparseArrayAllocator SparseArrayAllocator;
+	typedef InHashAllocator        HashAllocator;
+	
+};
+
+template <typename KeyType, typename ValueType, typename BaseAllocator = TInstrumentationSafeBaseAllocator>
+using TSafeMap = TMap<KeyType, ValueType, TInstrumentationSafeSetAllocator<BaseAllocator>>;
+
+template <typename ElementType, typename BaseAllocator = TInstrumentationSafeBaseAllocator>
+using TSafeSet = TSet<ElementType, DefaultKeyFuncs<ElementType>, TInstrumentationSafeSetAllocator<BaseAllocator>>;
+
+#endif
\ No newline at end of file
diff --git a/Engine/Source/Runtime/Core/Internal/Sanitizer/RaceDetectorInterface.h b/Engine/Source/Runtime/Core/Internal/Sanitizer/RaceDetectorInterface.h
new file mode 100644
--- /dev/null
+++ b/Engine/Source/Runtime/Core/Internal/Sanitizer/RaceDetectorInterface.h
@@ -0,0 +1,89 @@
+// Copyright Epic Games, Inc. All Rights Reserved.
+
+#pragma once
+
+#include "Misc/CoreMiscDefines.h"
+
+#if USING_INSTRUMENTATION
+
+#include "CoreTypes.h"
+#include "Templates/RefCounting.h"
+#include "Sanitizer/RaceDetectorTypes.h"
+
+DECLARE_LOG_CATEGORY_EXTERN(LogRaceDetector, Log, All);
+
+namespace UE::Sanitizer::RaceDetector {
+	namespace Platform {
+		// Performs any platform specific initialization.
+		bool  InitializePlatform();
+		// Performs any platform specific cleanup.
+		bool  CleanupPlatform();
+		// Prepare address space of shadow memory
+		void  InitShadowMemory();
+		// Returns the command line the process was started with.
+		const TCHAR* GetCommandLine();
+		// Returns the page size granularity.
+		UPTRINT GetPageSize();
+		// Returns the base address of the shadow memory address space. 
+		// Note: InitShadowMemory must have been called before, otherwise this is meaningless.
+		UPTRINT GetShadowMemoryBase();
+		// Returns the size of the entire shadow memory address space.
+		// Note: InitShadowMemory must have been called before, otherwise this is meaningless.
+		UPTRINT GetShadowMemorySize();
+		// Returns the number of bytes currently mapped in the shadow memory.
+		uint64 GetShadowMemoryUsage();
+		// Returns whether a particular range is already accessible in shadow memory.
+		bool IsShadowMemoryMapped(UPTRINT Base, UPTRINT Size);
+		// Maps a range in shadow memory so that it is safe to access.
+		void MapShadowMemory(UPTRINT Base, UPTRINT Size);
+		// Unmaps the entire range of shadow memory.
+		void UnmapShadowMemory();
+		// Returns whether a debugger is currently attached to our process.
+		bool IsDebuggerPresent();
+		// Returns whether there is any page currently mapped in shadow memory.
+		bool HasShadowMemoryMapped();
+		// Sends a hint to visual studio to hide first chance exception to reduce noise caused by shadow memory access.
+		void HideFirstChanceExceptionInVisualStudio();
+		// Sleeps for the given amount of milliseconds.
+		void SleepMS(uint32 Milliseconds);
+		// Capture the current callstack.
+		uint16 CaptureStackBackTrace(uint32 FrameToSkip, uint32 FrameToCapture, void** Backtrace);
+		// Allocates a Tls index.
+		uint32 AllocTlsSlot();
+		// Releases a Tls index.
+		void FreeTlsSlot(uint32 Index);
+		// Gets the value of the Tls index for the current thread.
+		void* GetTlsValue(uint32 Index);
+		// Sets the value of the Tls index for the current thread.
+		void SetTlsValue(uint32 Index, void* Value);
+		// Returns the current thread Id.
+		uint32 GetCurrentThreadId();
+
+		// Rewrites the patchable function prefix with the proper jmp to the target.
+		void PrepareTrampoline(void* PatchableFunctionAddress, void* TargetFunctionAddress, bool bUseRETBypass);
+		// Rewrites the patchable function prefix back to its original compiled NOPs.
+		void CleanupTrampoline(void* PatchableFunctionAddress);
+		// Rewrites the 2 first bytes of the function to jump at the beginning of the prefix section.
+		void ActivateTrampoline(void* PatchableFunctionAddress);
+		// Rewrites the 2 first bytes of the function to either do nothing or do an immediate return (RET bypass).
+		void DeactivateTrampoline(void* PatchableFunctionAddress, bool bUseRETBypass);
+		// This needs to be called after trampoline activation or deactivation to make sure it takes effect immediately.
+		void FlushInstructionCache();
+	}
+
+	// Returns whether we should instrument depending on the current context state.
+	bool ShouldInstrument(FContext& Context);
+	// Gets the current thread context, could be nullptr.
+	FContext* GetThreadContext();
+	// Hints the sanitizer that this memory range is being freed.
+	void FreeMemoryRange(void* Ptr, uint64 Size);
+
+	// Makes sure the current thread has a context and returns it.
+	FContext& EnsureCurrentContext();
+	// Releases the current thread context.
+	void ReleaseCurrentContext();
+	// Returns a sync object for the given address, initialize one if there isn't one already.
+	TRefCountPtr<FSyncObject> GetSyncObject(void* SyncAddr);
+}
+
+#endif // USING_INSTRUMENTATION
\ No newline at end of file
diff --git a/Engine/Source/Runtime/Core/Internal/Sanitizer/RaceDetectorTypes.h b/Engine/Source/Runtime/Core/Internal/Sanitizer/RaceDetectorTypes.h
new file mode 100644
--- /dev/null
+++ b/Engine/Source/Runtime/Core/Internal/Sanitizer/RaceDetectorTypes.h
@@ -0,0 +1,464 @@
+// Copyright Epic Games, Inc. All Rights Reserved.
+
+#pragma once
+
+#include "Misc/CoreMiscDefines.h"
+
+#if USING_INSTRUMENTATION
+
+#include "CoreTypes.h"
+
+#include "Misc/EnumClassFlags.h"
+#include "Misc/RWSpinLock.h"
+#include "Misc/ScopeRWLock.h"
+#include "Templates/RefCounting.h"
+
+#include "Sanitizer/Types.h"
+#include "Instrumentation/Types.h"
+#include "Instrumentation/Containers.h"
+
+namespace UE::Sanitizer::RaceDetector {
+
+	using namespace UE::Instrumentation;
+
+	extern volatile bool bDetailedLogGlobal;
+
+	// ------------------------------------------------------------------------------
+	// Clocks.
+	// ------------------------------------------------------------------------------
+	using FClock = uint32;
+	using FContextId = uint8;
+
+	class FClockBank 
+	{
+	public:
+		INSTRUMENTATION_FUNCTION_ATTRIBUTES FClockBank()
+		{
+			Reset();
+		}
+
+		INSTRUMENTATION_FUNCTION_ATTRIBUTES void Acquire(FClockBank& Other, void* ReturnAddress)
+		{
+			for (int32 Index = 0; Index < UE_ARRAY_COUNT(Clocks); ++Index)
+			{
+				if (Other.Clocks[Index] > Clocks[Index])
+				{
+					Clocks[Index] = Other.Clocks[Index];
+					Locations[Index] = ReturnAddress;
+				}
+			}
+		}
+
+		INSTRUMENTATION_FUNCTION_ATTRIBUTES void Release(FClockBank& Other, void* ReturnAddress)
+		{
+			Other.Acquire(*this, ReturnAddress);
+		}
+
+		INSTRUMENTATION_FUNCTION_ATTRIBUTES void AcquireRelease(FClockBank& Other, void* ReturnAddress)
+		{
+			for (int32 Index = 0; Index < UE_ARRAY_COUNT(Clocks); ++Index)
+			{
+				if (Clocks[Index] > Other.Clocks[Index])
+				{
+					Other.Clocks[Index] = Clocks[Index];
+					Other.Locations[Index] = ReturnAddress;
+				}
+				else if (Other.Clocks[Index] > Clocks[Index])
+				{
+					Clocks[Index] = Other.Clocks[Index];
+					Locations[Index] = ReturnAddress;
+				}
+			}
+		}
+
+		INSTRUMENTATION_FUNCTION_ATTRIBUTES void Reset()
+		{
+			memset(Clocks, 0, sizeof(Clocks));
+			memset(Locations, 0, sizeof(Locations));
+		}
+
+		INSTRUMENTATION_FUNCTION_ATTRIBUTES FClock& Get(FContextId ContextId)
+		{
+			return Clocks[ContextId];
+		}
+
+		INSTRUMENTATION_FUNCTION_ATTRIBUTES FCallstackLocation GetLocation(FContextId ContextId)
+		{
+			return FCallstackLocation::FromCallstackFrames(&Locations[ContextId], 1, 2);
+		}
+
+	private:
+		FClock Clocks[256] = {};
+		void* Locations[256] = {};
+	};
+
+
+	// ------------------------------------------------------------------------------
+	// Memory access.
+	// ------------------------------------------------------------------------------
+	enum EMemoryAccessType : uint8 {
+		ACCESS_TYPE_INVALID = 0b0,
+		ACCESS_TYPE_READ = 0b1,
+		ACCESS_TYPE_WRITE = 0b10,
+		ACCESS_TYPE_ATOMIC = 0b100,
+
+		ACCESS_TYPE_ATOMIC_READ = ACCESS_TYPE_ATOMIC | ACCESS_TYPE_READ,
+		ACCESS_TYPE_ATOMIC_WRITE = ACCESS_TYPE_ATOMIC | ACCESS_TYPE_WRITE,
+		ACCESS_TYPE_ATOMIC_READ_WRITE = ACCESS_TYPE_ATOMIC_READ | ACCESS_TYPE_ATOMIC_WRITE
+	};
+	ENUM_CLASS_FLAGS(EMemoryAccessType);
+	INSTRUMENTATION_FUNCTION_ATTRIBUTES inline const TCHAR* AccessTypeToString(EMemoryAccessType AccessType)
+	{
+		switch (AccessType)
+		{
+		case ACCESS_TYPE_READ:
+			return TEXT("Read");
+		case ACCESS_TYPE_WRITE:
+			return TEXT("Write");
+		case ACCESS_TYPE_ATOMIC_READ:
+			return TEXT("AtomicRead");
+		case ACCESS_TYPE_ATOMIC_WRITE:
+			return TEXT("AtomicWrite");
+		case ACCESS_TYPE_ATOMIC_READ_WRITE:
+			return TEXT("AtomicReadWrite");
+		}
+		return TEXT("Unknown");
+	}
+	INSTRUMENTATION_FUNCTION_ATTRIBUTES inline bool IsReadMemoryAccess(EMemoryAccessType AccessType)
+	{
+		return EnumHasAnyFlags(AccessType, EMemoryAccessType::ACCESS_TYPE_READ);
+	}
+	INSTRUMENTATION_FUNCTION_ATTRIBUTES inline bool IsWriteMemoryAccess(EMemoryAccessType AccessType)
+	{
+		return EnumHasAnyFlags(AccessType, EMemoryAccessType::ACCESS_TYPE_WRITE);
+	}
+	INSTRUMENTATION_FUNCTION_ATTRIBUTES inline bool IsAtomicMemoryAccess(EMemoryAccessType AccessType)
+	{
+		return EnumHasAnyFlags(AccessType, EMemoryAccessType::ACCESS_TYPE_ATOMIC);
+	}
+
+	struct FMemoryAccess 
+	{
+		INSTRUMENTATION_FUNCTION_ATTRIBUTES FORCEINLINE FMemoryAccess()
+		{
+		}
+
+		// Not sure why yet but Clang is not inlining this by default
+		// Did we mess up the clang compilation settings??
+		INSTRUMENTATION_FUNCTION_ATTRIBUTES FORCEINLINE FMemoryAccess(uint64 InRawValue)
+			: RawValue(InRawValue)
+		{
+		}
+
+		INSTRUMENTATION_FUNCTION_ATTRIBUTES FORCEINLINE FMemoryAccess(FContextId InContextId, FClock InClock, uint8 InOffset, uint8 InSize, EMemoryAccessType InAccessType, bool bInHasSyncObjectAttached)
+			: Clock(InClock)
+			, ContextId(InContextId)
+			, Access((uint8)(((1ul<<InSize) - 1ul) << InOffset))
+			, AccessType(InAccessType)
+			, bHasSyncObjectAttached(bInHasSyncObjectAttached)
+		{
+		}
+		
+		INSTRUMENTATION_FUNCTION_ATTRIBUTES FORCEINLINE uint8 GetOffset() const { return (uint8)FMath::CountTrailingZeros(Access); }
+		INSTRUMENTATION_FUNCTION_ATTRIBUTES FORCEINLINE uint8 GetSize() const { return (uint8)FMath::CountBits(Access >> GetOffset()); }
+		INSTRUMENTATION_FUNCTION_ATTRIBUTES FORCEINLINE bool IsValid() const { return AccessType != EMemoryAccessType::ACCESS_TYPE_INVALID; }
+
+		union
+		{
+			struct
+			{
+				FClock     Clock;
+				FContextId ContextId;
+				// Each bit represent a 1-byte slot used in our 8 byte shadow
+				// and can easily be tested for overlaps with other accesses.
+				uint8      Access; 
+				union
+				{
+					struct
+					{
+						uint8 bIsRead : 1;
+						uint8 bIsWrite : 1;
+						uint8 bIsAtomic : 1;
+					};
+					EMemoryAccessType AccessType : 3;
+				};
+				uint8 bHasSyncObjectAttached : 1;
+			};
+
+			uint64 RawValue;
+		};
+	};
+
+	// Used to store all accesses made by a thread since we don't have 
+	// enough space in the shadow memory to store callstacks.
+	// Storage space can certainly be improved by storing
+	// function entry and exit instead of full callstack every time.
+	// The current setup amount to about 8MB of memory per thread
+	// for a maximum of 256 thread the shadow memory can account for.
+	// This is totally acceptable for race detection purpose for now
+	// and can be revisited later if needed.
+	struct FAccessHistory
+	{
+		SAFE_OPERATOR_NEW_DELETE()
+
+		struct FEntry
+		{
+			void*              Pointer;
+			FMemoryAccess      Access;
+			FCallstackLocation Location;
+			UPTRINT            Data;
+		};
+
+		static const uint32 NumEntry = 128 * 1024;
+		uint32 Index;
+		uint64 ResetCount;
+		double LastReset;
+		uint32 LastResetClock;
+		FEntry Entries[NumEntry];
+	};
+
+	// ------------------------------------------------------------------------------
+	// Race Detector Context.
+	// ------------------------------------------------------------------------------
+	
+	// We use ref-counting because this otherwise might get deleted
+	// by other threads and we'd need to hold a lock while scanning
+	// the history for race report, which would be unnaceptable.
+	struct FContext : public TRefCountingMixin<FContext> {
+		SAFE_OPERATOR_NEW_DELETE()
+
+			INSTRUMENTATION_FUNCTION_ATTRIBUTES FContext(uint32 InThreadId)
+			: ThreadId(InThreadId)
+		{
+		}
+
+		uint32 GlobalEpoch = 0;
+		// When we activate tracing, we need to recapture the current stack
+		uint32 StackEpoch = 0;
+		// Unassigned until the first memory access
+		FContextId ContextId = 0;
+		// Avoid reading another TLS value for this
+		uint32 ThreadId;
+		// Prevents recursion for instrumentation
+		uint32 InstrumentationDepth = 0;
+		// Prevents recursion for detoured intrumentation
+		uint32 WinInstrumentationDepth = 0;
+		// Used to avoid instrumenting CreateThread while inside a higher level thread creation function (i.e. beginthreadx)
+		uint32 ThreadCreationDepth = 0; 
+		// When we want detailed logging for diagnostic purpose
+		uint32 DetailedLogDepth = 0;
+		// When we want detailed log only on a single address
+		void* DetailedLogAddress = nullptr; 
+		// Clock used while waiting to get a context id assigned.
+		FClock StandbyClock;
+		// Each thread holds a bank of clocks to synchronize with every other context.
+		FClockBank ClockBank;
+		// We need to keep the callstack for each thread 
+		static constexpr SIZE_T MaxCallstackSize = 512;
+		void* CurrentCallstack[MaxCallstackSize];
+		uint16 CurrentCallstackSize = 0;
+		// We use this to pass thread arguments to functions that dont have parameters (i.e. ExitThread).
+		void* ThreadArgs = nullptr;
+		// The clock for this context Id
+		INSTRUMENTATION_FUNCTION_ATTRIBUTES FClock& CurrentClock()
+		{ 
+			// Make sure we use the same value for the comparison and the get in case
+			// this is called from another thread while we're releasing our context id.
+			const FContextId LocalContextId = ContextId;
+			return LocalContextId == 0 ? StandbyClock : ClockBank.Get(LocalContextId);
+		}
+		INSTRUMENTATION_FUNCTION_ATTRIBUTES void IncrementClock()
+		{
+			CurrentClock()++; 
+			
+			if (UNLIKELY(DetailedLogDepth || bDetailedLogGlobal))
+			{
+				FPlatformMisc::LowLevelOutputDebugStringf(TEXT("[%d] Thread is now at clock %d\n"), ThreadId, CurrentClock());
+			}
+		}
+
+		// Lazily created on first instrumented access
+		TUniquePtr<FAccessHistory> AccessHistory;
+	};
+
+	CORE_API extern FContext* GetThreadContext();
+
+	struct FInstrumentationScope
+	{
+		bool bNeedDecrement = false;
+		INSTRUMENTATION_FUNCTION_ATTRIBUTES FInstrumentationScope()
+		{
+			FContext* Context = GetThreadContext();
+			if (Context)
+			{
+				Context->WinInstrumentationDepth++;
+				bNeedDecrement = true;
+			}
+		}
+
+		INSTRUMENTATION_FUNCTION_ATTRIBUTES ~FInstrumentationScope()
+		{
+			if (bNeedDecrement)
+			{
+				GetThreadContext()->WinInstrumentationDepth--;
+			}
+		}
+	};
+
+	class FSyncObject : public TRefCountingMixin<FSyncObject> {
+	public:
+		SAFE_OPERATOR_NEW_DELETE()
+
+		INSTRUMENTATION_FUNCTION_ATTRIBUTES FSyncObject()
+		{
+			FPlatformAtomics::InterlockedIncrement(&SyncObjectCount);
+		}
+		
+		INSTRUMENTATION_FUNCTION_ATTRIBUTES ~FSyncObject()
+		{
+			FPlatformAtomics::InterlockedDecrement(&SyncObjectCount);
+		}
+
+		INSTRUMENTATION_FUNCTION_ATTRIBUTES void SyncAcquire(FContext& Context, const TFunctionRef<void()>& AtomicOp, void* ReturnAddress, void* SyncAddress, const TCHAR* OpName)
+		{
+			UE::TWriteScopeLock Scope(Lock);
+			SyncAcquireAsSoleOwnerOrReadOwner(Context, ReturnAddress, SyncAddress, OpName);
+			AtomicOp();
+		}
+
+		INSTRUMENTATION_FUNCTION_ATTRIBUTES void SyncRelease(FContext& Context, const TFunctionRef<void()>& AtomicOp, void* ReturnAddress, void* SyncAddress, const TCHAR* OpName)
+		{
+			UE::TWriteScopeLock Scope(Lock);
+			SyncReleaseAsSoleOwner(Context, ReturnAddress, SyncAddress, OpName);
+			AtomicOp();
+		}
+
+		INSTRUMENTATION_FUNCTION_ATTRIBUTES void SyncAcquireRelease(FContext& Context, const TFunctionRef<void()>& AtomicOp, void* ReturnAddress, void* SyncAddress, const TCHAR* OpName)
+		{
+			if (UNLIKELY(Context.DetailedLogDepth || bDetailedLogGlobal))
+			{
+				FPlatformMisc::LowLevelOutputDebugStringf(TEXT("[%d] %s acq_rel of 0x%p from function at 0x%p\n"), Context.ThreadId, OpName, SyncAddress, ReturnAddress);
+			}
+
+			UE::TWriteScopeLock Scope(Lock);
+			Context.ClockBank.AcquireRelease(ClockBank, ReturnAddress);
+			AtomicOp();
+		}
+
+		INSTRUMENTATION_FUNCTION_ATTRIBUTES void SyncWithFailureSupport(FContext& Context, const TFunctionRef<bool()>& AtomicOp, EMemoryAccessType AccessType, FAtomicMemoryOrder SuccessOrder, FAtomicMemoryOrder FailureOrder, void* ReturnAddress, void* SyncAddress, const TCHAR* OpName, const TFunctionRef<void(FAtomicMemoryOrder ActualOrder)>& ActualAccessCallback)
+		{
+			using namespace UE::Instrumentation;
+
+			// We only need a take a write lock when we do a release or acq_rel operation otherwise
+			// it's impossible to test for failure order as the AtomicOp inside the write lock would never fail.
+			// An acquire only operation is safe to run under read-lock since we're reading from the syncobject and writing into the context clockbank which is owned by the current thread.
+			// Per the standard, failure memory order cannot be release nor acq_release, so we don't need to look at the failure order to choose our lock type.
+			// See https://www.open-std.org/jtc1/sc22/wg21/docs/papers/2017/n4659.pdf §32.6.1
+			// #17 Requires: The failure argument shall not be memory_order_release nor memory_order_acq_rel.
+			// 
+			// Also for the time being MSVC atomic implementation always combines both order so we have no way to test our support of different SuccessOrder and FailureOrder until they fix their implementation.
+			UE::TRWScopeLock Scope(Lock, IsAtomicOrderRelease(SuccessOrder) ? SLT_Write : SLT_ReadOnly);
+			bool bSucceeded = AtomicOp();
+
+			FAtomicMemoryOrder Order = bSucceeded ? SuccessOrder : FailureOrder;
+			const TCHAR* OpResult = bSucceeded ? TEXT("success") : TEXT("failure");
+
+			ActualAccessCallback(Order);
+
+			if (IsAtomicOrderRelaxed(Order))
+			{
+				if (UNLIKELY(Context.DetailedLogDepth || bDetailedLogGlobal))
+				{
+					FPlatformMisc::LowLevelOutputDebugStringf(TEXT("[%d] %s %s relaxed of 0x%p from function at 0x%p\n"), Context.ThreadId, OpName, OpResult, SyncAddress, ReturnAddress);
+				}
+
+				// Do nothing in the relaxed case since no barrier is provided.
+				return;
+			}
+			else if (AccessType == ACCESS_TYPE_ATOMIC_READ_WRITE && IsAtomicOrderAcquireRelease(Order))
+			{
+				if (UNLIKELY(Context.DetailedLogDepth || bDetailedLogGlobal))
+				{
+					FPlatformMisc::LowLevelOutputDebugStringf(TEXT("[%d] %s %s acq_rel of 0x%p from function at 0x%p\n"), Context.ThreadId, OpName, OpResult, SyncAddress, ReturnAddress);
+				}
+
+				Context.ClockBank.AcquireRelease(ClockBank, ReturnAddress);
+			}
+			else if ((AccessType & ACCESS_TYPE_ATOMIC_READ) == ACCESS_TYPE_ATOMIC_READ && IsAtomicOrderAcquire(Order))
+			{
+				if (UNLIKELY(Context.DetailedLogDepth || bDetailedLogGlobal))
+				{
+					FPlatformMisc::LowLevelOutputDebugStringf(TEXT("[%d] %s %s acquire 0x%p from function at 0x%p\n"), Context.ThreadId, OpName, OpResult, SyncAddress, ReturnAddress);
+				}
+
+				Context.ClockBank.Acquire(ClockBank, ReturnAddress);
+			}
+			else if ((AccessType & ACCESS_TYPE_ATOMIC_WRITE) == ACCESS_TYPE_ATOMIC_WRITE && IsAtomicOrderRelease(Order))
+			{
+				if (UNLIKELY(Context.DetailedLogDepth || bDetailedLogGlobal))
+				{
+					FPlatformMisc::LowLevelOutputDebugStringf(TEXT("[%d] %s %s release 0x%p from function at 0x%p\n"), Context.ThreadId, OpName, OpResult, SyncAddress, ReturnAddress);
+				}
+
+				Context.ClockBank.Release(ClockBank, ReturnAddress);
+			}
+			else
+			{
+				checkf(false, TEXT("Unexpected memory order"));
+			}
+
+			Context.IncrementClock();
+		}
+
+		// Must be called by a thread that has either this object's spin lock,
+		// or an external lock that is guaranteed to be held.
+		INSTRUMENTATION_FUNCTION_ATTRIBUTES void SyncReleaseAsSoleOwner(FContext& Context, void* ReturnAddress, void* SyncAddress, const TCHAR* OpName)
+		{
+			if (UNLIKELY(Context.DetailedLogDepth || bDetailedLogGlobal))
+			{
+				FPlatformMisc::LowLevelOutputDebugStringf(TEXT("[%d] %s releases 0x%p from function at 0x%p\n"), Context.ThreadId, OpName, SyncAddress, ReturnAddress);
+			}
+
+			Context.ClockBank.Release(ClockBank, ReturnAddress);
+		}
+
+		INSTRUMENTATION_FUNCTION_ATTRIBUTES void SyncAcquireAsSoleOwnerOrReadOwner(FContext& Context, void* ReturnAddress, void* SyncAddress, const TCHAR* OpName)
+		{
+			if (UNLIKELY(Context.DetailedLogDepth || bDetailedLogGlobal))
+			{
+				FPlatformMisc::LowLevelOutputDebugStringf(TEXT("[%d] %s acquires 0x%p from function at 0x%p\n"), Context.ThreadId, OpName, SyncAddress, ReturnAddress);
+		 	}
+
+			Context.ClockBank.Acquire(ClockBank, ReturnAddress);
+		}
+
+		INSTRUMENTATION_FUNCTION_ATTRIBUTES static uint64 GetSyncObjectCount()
+		{
+			return SyncObjectCount;
+		}
+	private:
+		UE::FRWSpinLock Lock;
+		FClockBank ClockBank;
+		static volatile int64 SyncObjectCount;
+	};
+
+	// One SyncObjectBank per 64-bit aligned address
+	struct FSyncObjectBank
+	{
+		SAFE_OPERATOR_NEW_DELETE()
+
+		// Contains 8 sync object for each byte of the 64-bit.
+		TRefCountPtr<FSyncObject> SyncObjects[8];
+	};
+
+	// ------------------------------------------------------------------------------
+	// Shadow memory.
+	// ------------------------------------------------------------------------------
+	struct FShadowMemory 
+	{
+		FMemoryAccess Accesses[4];
+	};
+
+} // UE::Sanitizer::RaceDetector
+
+#endif // USING_INSTRUMENTATION
\ No newline at end of file
diff --git a/Engine/Source/Runtime/Core/Internal/Sanitizer/Types.h b/Engine/Source/Runtime/Core/Internal/Sanitizer/Types.h
new file mode 100644
--- /dev/null
+++ b/Engine/Source/Runtime/Core/Internal/Sanitizer/Types.h
@@ -0,0 +1,218 @@
+// Copyright Epic Games, Inc. All Rights Reserved.
+
+#pragma once
+
+#include "Misc/CoreMiscDefines.h"
+
+#if USING_INSTRUMENTATION
+
+#include "CoreTypes.h"
+
+#include "Instrumentation/Types.h"
+#include "Instrumentation/Containers.h"
+#include "Async/Fundamental/Scheduler.h"
+#include "HAL/PlatformStackWalk.h"
+#include "Hash/CityHash.h"
+
+namespace UE::Sanitizer {
+
+	struct FLocation {
+		FLocation(const FString& ModuleName, const FString& FunctionName, const FString& Filename, uint32 Line, UPTRINT ProgramCounter)
+			: ModuleName(ModuleName), FunctionName(FunctionName), Filename(Filename), Line(Line), ProgramCounter(ProgramCounter)
+		{
+		}
+
+		const FString ModuleName;
+		const FString FunctionName;
+		const FString Filename;
+		const uint32  Line;
+		const UPTRINT ProgramCounter;
+
+		int32 GetAlignment() const
+		{
+			return Filename.Len() + GetDigits(Line);
+		}
+
+		FString ToString(int32 Alignment = 0) const
+		{
+			Alignment = FMath::Max(0, Alignment - GetAlignment());
+
+			// Format in a way that visual studio can understand so we can click the file and goto source.
+			return FString::Printf(TEXT("%s (%d):%s %s 0x%p"), *Filename, Line, FCString::Spc(Alignment), *FunctionName, (void*)ProgramCounter);
+		}
+
+	private:
+		static int32 GetDigits(int32 Value)
+		{
+			int32 Digits = 0;
+			do
+			{
+				Value /= 10;
+				++Digits;
+			} while (Value != 0);
+
+			return Digits;
+		};
+	};
+
+	struct FFullLocation {
+		TArray<FLocation> Locations;
+
+		int32 Num() const { return Locations.Num(); }
+
+		int32 GetAlignment() const
+		{
+			int32 MaxAlign = 0;
+			for (const FLocation& Location : Locations)
+			{
+				MaxAlign = FMath::Max(MaxAlign, Location.GetAlignment());
+			}
+
+			return MaxAlign;
+		}
+
+		FString ToString(int32 Alignment = 0) const
+		{
+			if (Alignment == 0)
+			{
+				Alignment = GetAlignment();
+			}
+
+			TStringBuilder<4096> Buffer;
+			for (int32 Index = 0; Index < Locations.Num();)
+			{
+				const FLocation& Location = Locations[Index];
+				Buffer.Append(Location.ToString(Alignment));
+				if (++Index < Locations.Num())
+				{
+					Buffer.AppendChar(TEXT('\n'));
+				}
+			}
+			return Buffer.ToString();
+		}
+	};
+
+
+	// ------------------------------------------------------------------------------
+	// Callstack.
+	// ------------------------------------------------------------------------------
+	class FCallstackLocation {
+	private:
+		static constexpr uint32 NumCallstackFrames = 5;
+
+		static constexpr UPTRINT NoCallstackMarker = 0;
+		static constexpr UPTRINT OverflowCallstackMarker = 0x1;
+
+		static constexpr const TCHAR* OverflowString = TEXT("[Uncaptured frames, callstack too deep]");
+
+	public:
+		FCallstackLocation()
+		{
+			Callstack[0] = NoCallstackMarker;
+		}
+
+		static FCallstackLocation FromCallstackFrames(const void* InCallstack, uint32 NumFrames, uint32 MaxNumFrames)
+		{
+			FCallstackLocation Location;
+
+			// Set overflow marker if needed.
+			uint32 NumMaxFrames = NumCallstackFrames;
+			if (NumFrames >= MaxNumFrames)
+			{
+				Location.Callstack[NumCallstackFrames - 1] = OverflowCallstackMarker;
+				NumMaxFrames--;
+			}
+
+			NumFrames = FMath::Min(NumFrames, MaxNumFrames);
+
+			uint32 NumPopulated = FMath::Min(NumFrames, NumMaxFrames);
+			const UPTRINT* InCallstackPtrs = (const UPTRINT*)InCallstack;
+			InCallstackPtrs += NumFrames - NumPopulated;
+
+			memcpy(Location.Callstack, InCallstackPtrs, sizeof(UPTRINT) * NumPopulated);
+			if (NumPopulated < NumMaxFrames)
+			{
+				memset(Location.Callstack + NumPopulated, 0, sizeof(UPTRINT) * (NumMaxFrames - NumPopulated));
+			}
+			return Location;
+		}
+
+		FFullLocation GetFullLocation() const
+		{
+			FFullLocation Loc;
+			if (Callstack[0] == NoCallstackMarker)
+			{
+				return Loc;
+			}
+
+			// Resolving symbols can cause us to enter a wait that could try to start a new thread
+			// and wait for it to be started. We can't allow that since we could deadlock
+			// if we're reporting a race while having a lock that the new thread might also need
+			// during its initialization. (i.e. Registering new FNames)
+			LowLevelTasks::Private::FOversubscriptionAllowedScope AllowOversubscription(false);
+
+			for (int32 Idx = NumCallstackFrames - 1; Idx >= 0; --Idx)
+			{
+				UPTRINT Frame = Callstack[Idx];
+				if (Frame == NoCallstackMarker)
+				{
+					continue;
+				}
+				else if (Frame == OverflowCallstackMarker)
+				{
+					Loc.Locations.Emplace(OverflowString, OverflowString, OverflowString, 0, 0);
+				}
+				else
+				{
+					const bool bIncludeInlineFrames = true;
+					// We always record the return address of our functions using FuncEntry. What we need to properly decode
+					// inline frames is something that points on the address of the actual call, not the return address, otherwise
+					// we end up decoding inline frames for what comes after we return from the function, which doesn't make sense.
+					// By using Frame - 1, we make sure that we actually point inside the chunk of assembly that represents the function call and it is enough
+					// to make symbol decode correct since we only need to point inside a valid range of a symbol entry, not directly at the start of it.
+					FPlatformStackWalk::EnumerateSymbolInfosForProgramCounter(Frame - 1, bIncludeInlineFrames, 
+						[&Loc, Frame](FProgramCounterSymbolInfo& SymbolInfo)
+						{
+							Loc.Locations.Emplace(FString(SymbolInfo.ModuleName), FString(SymbolInfo.FunctionName), FString(SymbolInfo.Filename), SymbolInfo.LineNumber, Frame);
+						}
+					);
+				}
+			}
+
+			return Loc;
+		}
+
+		uint32 GetHash() const
+		{
+			if (Callstack[0] == NoCallstackMarker)
+			{
+				return 0;
+			}
+
+			return CityHash32((const char*)Callstack, sizeof(UPTRINT) * NumCallstackFrames);
+		}
+
+		UPTRINT GetLastFrame() const
+		{
+			for (int32 Idx = NumCallstackFrames - 1; Idx >= 0; --Idx)
+			{
+				UPTRINT Frame = Callstack[Idx];
+				if (Frame != NoCallstackMarker &&
+					Frame != OverflowCallstackMarker)
+				{
+					return Frame;
+				}
+			}
+
+			return 0;
+		}
+
+	private:
+		// Order is outer frames to inner frames.
+		// For example, Main -> Fn1 -> Fn2 -> FnLeaf.
+		UPTRINT Callstack[NumCallstackFrames];
+	};
+
+} // namespace UE::Sanitizer
+
+#endif
\ No newline at end of file
diff --git a/Engine/Source/Runtime/Core/Private/Instrumentation/Containers.cpp b/Engine/Source/Runtime/Core/Private/Instrumentation/Containers.cpp
new file mode 100644
--- /dev/null
+++ b/Engine/Source/Runtime/Core/Private/Instrumentation/Containers.cpp
@@ -0,0 +1,42 @@
+// Copyright Epic Games, Inc. All Rights Reserved.
+
+#include "Instrumentation/Containers.h"
+
+#if PLATFORM_WINDOWS && USING_INSTRUMENTATION
+
+#include "Windows.h"
+
+void* FInstrumentationSafeWinAllocator::Alloc(size_t Size)
+{
+	return VirtualAlloc(NULL, Size, MEM_RESERVE | MEM_COMMIT, PAGE_READWRITE);
+}
+
+void FInstrumentationSafeWinAllocator::Free(void* Data)
+{
+	VirtualFree(Data, 0, MEM_RELEASE);
+}
+
+void* FInstrumentationSafeWinAllocator::Realloc(void* Data, size_t Size, size_t PreviousSize)
+{
+	if (Data == nullptr)
+	{
+		return Alloc(Size);
+	}
+
+	if (Size == 0)
+	{
+		Free(Data);
+		return nullptr;
+	}
+
+	uint8* NewMem = (uint8*)Alloc(Size);
+	if (NewMem)
+	{
+		memcpy(NewMem, Data, PreviousSize);
+		Free(Data);
+	}
+
+	return NewMem;
+}
+
+#endif // PLATFORM_WINDOWS && USING_INSTRUMENTATION
\ No newline at end of file
diff --git a/Engine/Source/Runtime/Core/Private/Instrumentation/EntryPoints.cpp b/Engine/Source/Runtime/Core/Private/Instrumentation/EntryPoints.cpp
new file mode 100644
--- /dev/null
+++ b/Engine/Source/Runtime/Core/Private/Instrumentation/EntryPoints.cpp
@@ -0,0 +1,8 @@
+// Copyright Epic Games, Inc. All Rights Reserved.
+
+#include "Misc/CoreMiscDefines.h"
+
+// For monolithic build, we include the per module inline only once in the core module.
+#if IS_MONOLITHIC && USING_INSTRUMENTATION
+#include "Instrumentation/PerModuleInline.inl"
+#endif
\ No newline at end of file
diff --git a/Engine/Source/Runtime/Core/Private/Sanitizer/RaceDetector.cpp b/Engine/Source/Runtime/Core/Private/Sanitizer/RaceDetector.cpp
new file mode 100644
--- /dev/null
+++ b/Engine/Source/Runtime/Core/Private/Sanitizer/RaceDetector.cpp
@@ -0,0 +1,1931 @@
+// Copyright Epic Games, Inc. All Rights Reserved.
+
+#include "Sanitizer/RaceDetector.h"
+
+#if USING_INSTRUMENTATION
+
+#include "Sanitizer/RaceDetectorTypes.h"
+#include "Sanitizer/RaceDetectorInterface.h"
+
+#include "Instrumentation/Types.h"
+#include "Instrumentation/Containers.h"
+#include "Instrumentation/EntryPoints.h"
+
+#include "CoreTypes.h"
+#include "Trace/Trace.h"
+#include "Misc/SpinLock.h"
+#include "Misc/CString.h"
+#include "Misc/ScopeLock.h"
+#include "Misc/Paths.h"
+#include "Async/Mutex.h"
+#include "Misc/Parse.h"
+#include "Math/NumericLimits.h"
+#include "Algo/Reverse.h"
+#include "Misc/CommandLine.h"
+#include "Async/UniqueLock.h"
+#include "HAL/FileManager.h"
+#include "HAL/ThreadManager.h"
+#include "Tasks/Task.h"
+
+#if PLATFORM_WINDOWS
+#include "Windows/WindowsHWrapper.h"
+#endif
+
+#include <atomic>
+
+DEFINE_LOG_CATEGORY(LogRaceDetector);
+
+UE_TRACE_CHANNEL_DEFINE(RaceDetectorChannel);
+
+#define WITH_RACEDETECTOR_CHECK 0
+#define WITH_RACEDETECTOR_DEBUG 0
+
+#if WITH_RACEDETECTOR_DEBUG
+UE_DISABLE_OPTIMIZATION
+#endif
+
+static uint32 GRaceDetectorThreadFilter = 0;
+static bool GRaceDetectorFilterDuplicates = true;
+static FAutoConsoleVariableRef CVarReportSameRaceOnce(
+	TEXT("r.RaceDetector.FilterDuplicates"),
+	GRaceDetectorFilterDuplicates,
+	TEXT("Whether to report the same race only once per application lifetime."),
+	ECVF_Default);
+
+static int32 GRaceDetectorActivate = 0;
+static FAutoConsoleVariableRef CVarRaceDetectorActivate(
+	TEXT("r.RaceDetector.Activate"),
+	GRaceDetectorActivate,
+	TEXT("Activate race detection for that many seconds as it most likely makes the engine non-responsive and can't be easily turned off."),
+	FConsoleVariableDelegate::CreateLambda([](IConsoleVariable*) INSTRUMENTATION_FUNCTION_ATTRIBUTES
+		{
+			UE::Sanitizer::RaceDetector::ToggleRaceDetectionUntil(UE::FTimeout(FTimespan::FromSeconds(GRaceDetectorActivate)));
+		}),
+	ECVF_Default);
+
+static int32 GRaceDetectorGlobalDetailedLog = 0;
+static FAutoConsoleVariableRef CVarRaceDetectorGlobalDetailedLog(
+	TEXT("r.RaceDetector.GlobalDetailedLog"),
+	GRaceDetectorGlobalDetailedLog,
+	TEXT("Activate very detailed logging globally on all memory access."),
+	FConsoleVariableDelegate::CreateLambda([](IConsoleVariable*) INSTRUMENTATION_FUNCTION_ATTRIBUTES
+		{
+			UE::Sanitizer::RaceDetector::ToggleGlobalDetailedLog(GRaceDetectorGlobalDetailedLog);
+		}),
+	ECVF_Default);
+
+static bool GRaceDetectorIgnoreRaceWithoutSecondStack = true;
+static FAutoConsoleVariableRef CVarIgnoreRaceWithoutSecondStack(
+	TEXT("r.RaceDetector.IgnoreRaceWithoutSecondStack"),
+	GRaceDetectorIgnoreRaceWithoutSecondStack,
+	TEXT("Determines if races without a second stack will still be shown.\n")
+	TEXT("This generally happens when races are far enough that maybe it is not a concern after all.\n"),
+	ECVF_Default);
+
+static int32 GRaceDetectorMaxMemoryUsage = 16;
+static FAutoConsoleVariableRef CVarRaceDetectorMaxMemoryUsage(
+	TEXT("r.RaceDetector.MaxMemoryUsage"),
+	GRaceDetectorMaxMemoryUsage,
+	TEXT("How many gigabytes that the race detector is allowed to use.\n")
+	TEXT("The lower the limit, the higher the chance of missing some race conditions.\n"),
+	ECVF_Default);
+
+static int32 GRaceDetectorBreakOnRace = 0;
+static FAutoConsoleVariableRef CVarRaceDetectorBreakOnRace(
+	TEXT("r.RaceDetector.BreakOnRace"),
+	GRaceDetectorBreakOnRace,
+	TEXT("Debugbreak on race detection if the debugger is attached\n")
+	TEXT("	1 - Break once\n")
+	TEXT("	2 - Break always\n")
+	TEXT("	3 - Break only on race detected while detailed logging is active\n")
+	TEXT("	4 - Break only when the second callstack is missing\n"),
+	ECVF_Default);
+
+using namespace UE::Instrumentation;
+
+namespace UE::Sanitizer::RaceDetector {
+	
+	bool bRuntimeInitialized = false;
+
+	UPTRINT ShadowMemoryBase = 0;
+	UPTRINT ShadowMemorySize = 0;
+	UPTRINT ShadowMemoryEnd = 0;
+	uint32  ContextTLSIndex  = UINT32_MAX;
+	volatile uint32 GlobalEpoch = 0;
+	volatile int64 FSyncObject::SyncObjectCount = 0;
+	volatile FContextId CurrentContextId = 1;
+
+	UE::FSpinLock RaceReportsLock;
+	TArray<FString> RaceReports;
+
+	// Lock for both ContextMapping and FreeContexts
+	UE::FSpinLock ContextMappingLock; 
+	TSafeMap<FContextId, TRefCountPtr<FContext>> ContextMapping;
+	TSafeArray<FContextId> FreeContexts;
+
+	// Used for filtering duplicates
+	TSafeSet<uint64> RaceHashes;
+	UE::FSpinLock RaceHashesLock;
+	UE::FTimeout RaceDetectorTimeout = UE::FTimeout::Never();
+	
+	volatile bool bIsResettingShadow = false;
+	volatile bool bIsDebuggerPresent = false;
+	volatile bool bDetailedLogGlobal = false;
+
+	// Can't be an atomic as it might cause infinite recursion in debug when atomics
+	// are unoptimized and we end up with a function call to an instrumented function
+	// going directly back to the ShouldInstrument call. std::atomics that are tested
+	// after InstrumentationDepth should be OK.
+	volatile bool bRaceDetectorActive = false;
+	INSTRUMENTATION_FUNCTION_ATTRIBUTES bool ShouldInstrument(FContext& Context)
+	{
+		return bRaceDetectorActive && !Context.WinInstrumentationDepth && !Context.InstrumentationDepth && !bIsResettingShadow;
+	}
+
+	FORCEINLINE INSTRUMENTATION_FUNCTION_ATTRIBUTES void* GetTlsValueFast(uint32 Index)
+	{
+		return (void*)__readgsqword(0x1480lu + (unsigned long)Index * sizeof(void*));
+	}
+
+	FORCEINLINE INSTRUMENTATION_FUNCTION_ATTRIBUTES FContext* GetThreadContext()
+	{
+		return (FContext*)GetTlsValueFast(ContextTLSIndex);
+	}
+
+	INSTRUMENTATION_FUNCTION_ATTRIBUTES uint64 GetSyncObjectMemoryUsage()
+	{
+		return FSyncObject::GetSyncObjectCount() * Align(sizeof(FSyncObject), Platform::GetPageSize());
+	}
+
+	// ------------------------------------------------------------------------------
+	// Callstack management.
+	// ------------------------------------------------------------------------------
+	FORCEINLINE INSTRUMENTATION_FUNCTION_ATTRIBUTES FCallstackLocation GetCurrentCallstackLocation()
+	{
+		FContext& Context = *GetThreadContext();
+
+		return FCallstackLocation::FromCallstackFrames(Context.CurrentCallstack, Context.CurrentCallstackSize, FContext::MaxCallstackSize);
+	}
+
+	FORCEINLINE INSTRUMENTATION_FUNCTION_ATTRIBUTES void AddCallstackFrame(FContext& Context, void* ReturnAddress)
+	{
+		// Don't capture extremely deep callstacks.
+		if (Context.CurrentCallstackSize < FContext::MaxCallstackSize)
+		{
+			Context.CurrentCallstack[Context.CurrentCallstackSize] = ReturnAddress;
+		}
+
+		Context.CurrentCallstackSize++;
+	}
+
+	FORCEINLINE INSTRUMENTATION_FUNCTION_ATTRIBUTES void RemoveCallstackFrame(FContext& Context)
+	{
+		Context.CurrentCallstackSize--;
+	}
+
+	INSTRUMENTATION_FUNCTION_ATTRIBUTES void RegisterContext(FContext& Context)
+	{
+		Context.AccessHistory = MakeUnique<FAccessHistory>();
+
+		UE::TScopeLock Scope(ContextMappingLock);
+		Context.GlobalEpoch = GlobalEpoch;
+		if (CurrentContextId < TNumericLimits<FContextId>::Max())
+		{
+			Context.ContextId = CurrentContextId++;
+		}
+		else if (FreeContexts.Num() > 0)
+		{
+			// Reuse the oldest context id to have a chance to find races
+			// for threads that just exited.
+			Context.ContextId = FreeContexts[0];
+			FreeContexts.RemoveAt(0, EAllowShrinking::No);
+
+			// Now we can get rid of the old instrumentation context
+			TRefCountPtr<FContext> OldContext;
+			ContextMapping.RemoveAndCopyValue(Context.ContextId, OldContext);
+			// Acquire clocks from the old context since other threads might already have entries for our context id.
+			if (OldContext->GlobalEpoch == Context.GlobalEpoch)
+			{
+				// We acquire here since we were already given a set of vector clock from the thread that spawned us
+				// and we need to union it with the context we're inheriting from.
+				Context.ClockBank.Acquire(OldContext->ClockBank, _ReturnAddress());
+			}
+
+			if (Context.DetailedLogDepth || bDetailedLogGlobal)
+			{
+				FPlatformMisc::LowLevelOutputDebugStringf(TEXT("[%d] Recycling contextid %d from old thread %d\n"), Context.ThreadId, Context.ContextId, OldContext->ThreadId);
+			}
+		}
+		else
+		{
+			UE_LOG(LogRaceDetector, Fatal, TEXT("Too many threads active at once"));
+		}
+		
+		if (Context.DetailedLogDepth || bDetailedLogGlobal)
+		{
+			FPlatformMisc::LowLevelOutputDebugStringf(TEXT("[%d] Registered with contextid %d\n"), Context.ThreadId, Context.ContextId);
+		}
+
+		Context.CurrentClock() = FMath::Max(Context.StandbyClock, Context.CurrentClock());
+
+		check(Context.ContextId != 0);
+		ContextMapping.Add(Context.ContextId, &Context);
+		Context.IncrementClock();
+		check(Context.CurrentClock() > 0);
+	}
+
+	FORCEINLINE INSTRUMENTATION_FUNCTION_ATTRIBUTES void EnsureContextRegistered(FContext& Context)
+	{
+		if (Context.ContextId == 0)
+		{
+			RegisterContext(Context);
+		}
+	}
+
+	INSTRUMENTATION_FUNCTION_ATTRIBUTES void ReleaseContext(FContext* Context)
+	{
+		if (Context != nullptr)
+		{
+			if (Context->InstrumentationDepth != 0)
+			{
+				UE_LOG(LogRaceDetector, Fatal, TEXT("Trying to release an instrumentation context still in use"));
+			}
+
+			// Check if the context has ever been registered
+			if (Context->ContextId)
+			{
+				// Backup the current clock as we're releasing our context id and won't have access to the clock afterward.
+				Context->StandbyClock = Context->CurrentClock();
+				// We don't destroy the context here as we want to leave a chance
+				// to find race in short lived threads or threads that race just
+				// before exiting. We'll recycle it when its ContextId gets reused.
+				UE::TScopeLock Scope(ContextMappingLock);
+				FreeContexts.Push(Context->ContextId);
+				Context->ContextId = 0;
+			}
+		}
+	}
+
+	INSTRUMENTATION_FUNCTION_ATTRIBUTES TRefCountPtr<FContext> GetContextById(FContextId ContextId)
+	{
+		UE::TScopeLock Scope(ContextMappingLock);
+		return ContextMapping.FindRef(ContextId);
+	}
+
+	FORCEINLINE INSTRUMENTATION_FUNCTION_ATTRIBUTES UPTRINT GetAlignedAddress(UPTRINT Ptr)
+	{
+		return UPTRINT(Ptr) & ~UPTRINT(0b111);
+	}
+
+	FORCEINLINE INSTRUMENTATION_FUNCTION_ATTRIBUTES void* GetAlignedPointer(void* Ptr)
+	{
+		return (void*)GetAlignedAddress((UPTRINT)Ptr);
+	}
+
+	FORCEINLINE INSTRUMENTATION_FUNCTION_ATTRIBUTES UPTRINT GetShadowMemoryAddress(UPTRINT Ptr)
+	{
+		static_assert(sizeof(FShadowMemory) <= sizeof(UPTRINT) * 4);
+
+		// Map the higher address space as a continuation of the lower one by just removing the shadow size from it.
+		if (Ptr >= ShadowMemoryEnd)
+		{
+			Ptr -= ShadowMemorySize;
+		}
+
+		return ShadowMemoryBase + (Ptr >> 3) * sizeof(FShadowMemory);
+	}
+
+	INSTRUMENTATION_FUNCTION_ATTRIBUTES FShadowMemory* GetShadowMemory(UPTRINT Ptr)
+	{
+		UPTRINT ShadowMemoryAddress = GetShadowMemoryAddress(Ptr);
+		// When the debugger is present, we map shadow memory before using it since its the fastest method.
+		// If the debugger is not present we let the page fault handler take care of page faults,
+		// because its even faster (i.e. less memory lookups / cache misses).
+		if (bIsDebuggerPresent)
+		{
+			Platform::MapShadowMemory(ShadowMemoryAddress, sizeof(FShadowMemory));
+		}
+		return (FShadowMemory*)ShadowMemoryAddress;
+	}
+
+	// Reduce lock contention by striping the addresses into buckets.
+	constexpr uint32 STRIPING_BITS = 8;
+	constexpr uint32 STRIPING_SIZE = 1 << STRIPING_BITS;
+	constexpr uint32 STRIPING_MASK = STRIPING_SIZE - 1;
+
+	struct FSyncBanks
+	{
+		TSafeMap<UPTRINT, FSyncObjectBank> Map;
+		UE::FSpinLock Lock;
+	};
+	FSyncBanks StripedSyncBanks[STRIPING_SIZE];
+	
+	INSTRUMENTATION_FUNCTION_ATTRIBUTES TRefCountPtr<FSyncObject> GetSyncObject(void* SyncAddr)
+	{
+		UPTRINT Index    = UPTRINT(SyncAddr) & 7;
+		UPTRINT BankAddr = UPTRINT(SyncAddr) >> 3;
+
+		FSyncBanks& SyncBanks = StripedSyncBanks[BankAddr & STRIPING_MASK];
+		
+		UE::TScopeLock Scope(SyncBanks.Lock);
+		FSyncObjectBank* Bank = SyncBanks.Map.Find(BankAddr);
+		if (!Bank)
+		{
+			Bank = &SyncBanks.Map.Emplace(BankAddr, FSyncObjectBank());
+		}
+		TRefCountPtr<FSyncObject>& SyncObject = Bank->SyncObjects[Index];
+		if (SyncObject == nullptr)
+		{
+			SyncObject = new FSyncObject();
+		}
+
+		return SyncObject;
+	}
+
+	INSTRUMENTATION_FUNCTION_ATTRIBUTES bool HasInstrumentedAtomics()
+	{
+		FInstrumentationScope InstrumentationScope;
+		for (uint32 Index = 0; Index < STRIPING_SIZE; ++Index)
+		{
+			FSyncBanks& SyncBanks = StripedSyncBanks[Index];
+
+			TSafeMap<UPTRINT, FSyncObjectBank> OldMap;
+			{
+				UE::TScopeLock Scope(SyncBanks.Lock);
+				if (SyncBanks.Map.Num())
+				{
+					return true;
+				}
+			}
+		}
+
+		return false;
+	}
+
+	INSTRUMENTATION_FUNCTION_ATTRIBUTES void ResetInstrumentedAtomics()
+	{
+		FInstrumentationScope InstrumentationScope;
+		for (uint32 Index = 0; Index < STRIPING_SIZE; ++Index)
+		{
+			FSyncBanks& SyncBanks = StripedSyncBanks[Index];
+
+			TSafeMap<UPTRINT, FSyncObjectBank> OldMap;
+			{
+				UE::TScopeLock Scope(SyncBanks.Lock);
+				OldMap = MoveTemp(SyncBanks.Map);
+			}
+		}
+	}
+	
+	INSTRUMENTATION_FUNCTION_ATTRIBUTES void FreeMemoryRange(void* Ptr, uint64 Size)
+	{
+		if (!bRaceDetectorActive || bIsResettingShadow)
+		{
+			return;
+		}
+
+		UPTRINT BankAddr = UPTRINT(Ptr) >> 3;
+		UPTRINT ShadowAddr = GetShadowMemoryAddress((UPTRINT)Ptr);
+		for (uint64 Index = 0; Index < Size; Index += 8, ++BankAddr, ShadowAddr += sizeof(FShadowMemory))
+		{
+			// We only need to verify this once in every PageSize so we could optimize this further if needed.
+			if (Platform::IsShadowMemoryMapped(ShadowAddr, sizeof(FShadowMemory)))
+			{
+				FShadowMemory* Shadow = (FShadowMemory*)ShadowAddr;
+
+				// Best effort to clear sync objects if we know some were attached.
+				// If we miss any, they'll be cleaned up on the next shadow reset anyway.
+				const bool bHasSyncObject =
+					Shadow->Accesses[0].bHasSyncObjectAttached ||
+					Shadow->Accesses[1].bHasSyncObjectAttached ||
+					Shadow->Accesses[2].bHasSyncObjectAttached ||
+					Shadow->Accesses[3].bHasSyncObjectAttached;
+
+				if (bHasSyncObject)
+				{
+					FSyncBanks& SyncBanks = StripedSyncBanks[BankAddr & STRIPING_MASK];
+					{
+						UE::TScopeLock ScopeLock(SyncBanks.Lock);
+						SyncBanks.Map.Remove(BankAddr);
+					}
+				}
+
+				Shadow->Accesses[0].RawValue = 0;
+				Shadow->Accesses[1].RawValue = 0;
+				Shadow->Accesses[2].RawValue = 0;
+				Shadow->Accesses[3].RawValue = 0;
+			}
+		}
+	}
+	
+	FORCEINLINE INSTRUMENTATION_FUNCTION_ATTRIBUTES uint8 GetOffsetInBytes(void* Ptr)
+	{
+		UPTRINT OriginalAddr = (UPTRINT)Ptr;
+		UPTRINT FinalAddr = OriginalAddr & ~UPTRINT(0b111);
+		return (uint8)(OriginalAddr - FinalAddr);
+	}
+
+	FORCEINLINE INSTRUMENTATION_FUNCTION_ATTRIBUTES void ResetShadowMemory()
+	{
+		Platform::UnmapShadowMemory();
+	}
+
+	TOptional<TRaceCallbackFn> RaceCallbackFn;
+	INSTRUMENTATION_FUNCTION_ATTRIBUTES void ReportRace(FContext& Context, void* Pointer, const FMemoryAccess& CurrentAccess, const FMemoryAccess& RacingAccess, const FCallstackLocation& FirstLocation)
+	{
+		TRefCountPtr<FContext> RacingContext = GetContextById(RacingAccess.ContextId);
+		// The thread might have been destroyed
+		if (!RacingContext.IsValid())
+		{ 
+			return;
+		}
+
+		const uint32 FirstThreadId = Context.ThreadId;
+		const uint32 SecondThreadId = RacingContext->ThreadId;
+
+		// Dumb scan of the history to find the matching race so we can get its location.
+		FCallstackLocation SecondLocation;
+		FAccessHistory& SecondAccessHistory = *RacingContext->AccessHistory;
+		bool bFoundLocation = false;
+
+		const void* AlignedPointer = GetAlignedPointer(Pointer);
+		for (FAccessHistory::FEntry& Entry : SecondAccessHistory.Entries)
+		{
+			if (Entry.Pointer == AlignedPointer &&
+				Entry.Access.RawValue == RacingAccess.RawValue
+				)
+			{
+				SecondLocation = Entry.Location;
+				bFoundLocation = true;
+				break;
+			}
+		}
+
+		// Resolving symbols can cause us to enter a wait that could try to start a new thread
+		// and wait for it to be started. We can't allow that since we could deadlock
+		// if we're reporting a race while having a lock that the new thread might also need
+		// during its initialization. (i.e. Registering new FNames)
+		LowLevelTasks::Private::FOversubscriptionAllowedScope AllowOversubscription(false);
+
+		if (RaceCallbackFn)
+		{
+			FFullLocation FirstFullLocation = FirstLocation.GetFullLocation();
+			FFullLocation SecondFullLocation = SecondLocation.GetFullLocation();
+			RaceCallbackFn.GetValue()((UPTRINT)Pointer, FirstThreadId, SecondThreadId, FirstFullLocation, SecondFullLocation);
+		}
+		else
+		{
+			if (!bFoundLocation && GRaceDetectorIgnoreRaceWithoutSecondStack)
+			{
+				return;
+			}
+
+			if (GRaceDetectorFilterDuplicates)
+			{
+				uint64 FirstLocationLastFrame = FirstLocation.GetLastFrame();
+				uint64 SecondLocationLastFrame = SecondLocation.GetLastFrame();
+				uint64 LocationHash = FirstLocationLastFrame ^ SecondLocationLastFrame;
+				bool bAlreadyFound = false;
+				{
+					UE::TUniqueLock Lock(RaceHashesLock);
+					RaceHashes.FindOrAdd(LocationHash, &bAlreadyFound);
+				}
+				if (bAlreadyFound)
+				{
+					return;
+				}
+			}
+
+			if (GRaceDetectorThreadFilter)
+			{
+				if (Context.ThreadId != GRaceDetectorThreadFilter &&
+					RacingContext->ThreadId != GRaceDetectorThreadFilter)
+				{
+					return;
+				}
+			}
+
+			FFullLocation FirstFullLocation = FirstLocation.GetFullLocation();
+			FFullLocation SecondFullLocation = SecondLocation.GetFullLocation();
+
+			const FString& FirstThreadName = FThreadManager::Get().GetThreadName(FirstThreadId);
+			const FString& SecondThreadName = FThreadManager::Get().GetThreadName(SecondThreadId);
+
+			FString NotFoundMessage;
+			if (!bFoundLocation)
+			{
+				NotFoundMessage = FString::Printf(
+					TEXT("Location not found in access history.\n")
+					TEXT("Number of history rollovers: %llu\n")
+					TEXT("Last history rollover: %.02f seconds ago at clock %d"),
+					SecondAccessHistory.ResetCount,
+					FPlatformTime::Seconds() - SecondAccessHistory.LastReset,
+					SecondAccessHistory.LastResetClock
+				);
+			}
+			
+			FClock LastSyncClock = Context.ClockBank.Get(RacingAccess.ContextId);
+			FFullLocation LastSyncLocation = Context.ClockBank.GetLocation(RacingAccess.ContextId).GetFullLocation();
+
+			const int32 Alignment = FMath::Max3(FirstFullLocation.GetAlignment(), LastSyncLocation.GetAlignment(), SecondFullLocation.GetAlignment());
+
+			TStringBuilder<4096> Report;
+			Report.Appendf(
+				TEXT("=========================================") LINE_TERMINATOR
+				TEXT("WARNING: RaceDetector: data race detected") LINE_TERMINATOR
+				TEXT("%s of size %d at %p made at clock %lu by thread %s (%d) which is now at clock %lu:") LINE_TERMINATOR 
+				TEXT("%s") LINE_TERMINATOR
+				TEXT("Previous %s of size %d at %p at clock %lu by thread %s (%d) which is now at clock %lu:") LINE_TERMINATOR
+				TEXT("%s") LINE_TERMINATOR
+				TEXT("Last known sync clock between both threads is %lu:") LINE_TERMINATOR
+				TEXT("%s") LINE_TERMINATOR
+				TEXT("=========================================") LINE_TERMINATOR,
+				AccessTypeToString(CurrentAccess.AccessType),
+				CurrentAccess.GetSize(),
+				(void*)((UPTRINT)AlignedPointer + CurrentAccess.GetOffset()),
+				CurrentAccess.Clock,
+				*FirstThreadName,
+				FirstThreadId,
+				Context.CurrentClock(),
+				*FirstFullLocation.ToString(Alignment),
+				AccessTypeToString(RacingAccess.AccessType),
+				RacingAccess.GetSize(),
+				(void*)((UPTRINT)AlignedPointer + RacingAccess.GetOffset()),
+				RacingAccess.Clock,
+				*SecondThreadName,
+				SecondThreadId,
+				RacingContext->CurrentClock(),
+				bFoundLocation ? *SecondFullLocation.ToString(Alignment) : *NotFoundMessage,
+				LastSyncClock,
+				*LastSyncLocation.ToString(Alignment)
+			);
+
+			FPlatformMisc::LowLevelOutputDebugString(Report.ToString());
+
+			// We do as little as possible from within the race since it might cause reentrancy depending on where the race happens.
+			// Send the report to be properly logged to file in another thread.
+			UE::TScopeLock Lock(RaceReportsLock);
+			RaceReports.Emplace(Report.ToString());
+		}
+		
+		if (GRaceDetectorBreakOnRace && FPlatformMisc::IsDebuggerPresent())
+		{
+			if ((GRaceDetectorBreakOnRace == 3 && !Context.DetailedLogDepth) ||
+				(GRaceDetectorBreakOnRace == 4 && bFoundLocation))
+			{
+				return;
+			}
+
+			// Reset for break once.
+			if (GRaceDetectorBreakOnRace == 1)
+			{
+				GRaceDetectorBreakOnRace = 0;
+			}
+
+			PLATFORM_BREAK();
+		}
+	}
+
+	INSTRUMENTATION_FUNCTION_ATTRIBUTES void SetRaceCallbackFn(TRaceCallbackFn CallbackFn)
+	{
+		RaceCallbackFn = CallbackFn;
+	}
+
+	INSTRUMENTATION_FUNCTION_ATTRIBUTES void ResetRaceCallbackFn()
+	{
+		RaceCallbackFn.Reset();
+	}
+
+	INSTRUMENTATION_FUNCTION_ATTRIBUTES void InstrumentAccess(FContext& Context, void* Ptr, uint8 Size, EMemoryAccessType AccessType, FAtomicMemoryOrder Order, void* ReturnAddress, const TCHAR* OpName, bool bHasSyncObjectAttached)
+	{
+#if WITH_RACEDETECTOR_CHECK
+		// Validate size and alignment since it should have been done before entering this function.
+		check(Size <= 8 && GetShadowMemoryAddress((UPTRINT)Ptr) == GetShadowMemoryAddress((UPTRINT)Ptr + Size - 1));
+#endif
+		// Grab a ContextId on first instrumented access to avoid wasting our precious 256 available context
+		// on thirdparty and OS threads that never enter into instrumented code.
+		EnsureContextRegistered(Context);
+
+		const uint32 Epoch = GlobalEpoch;
+		const uint8 Offset = GetOffsetInBytes(Ptr);
+		const FMemoryAccess CurrentAccess(Context.ContextId, Context.CurrentClock(), Offset, Size, AccessType, bHasSyncObjectAttached);
+
+		// GetShadowMemory will contain the cost of committing pages to shadow memory
+		FShadowMemory* Shadow = GetShadowMemory((UPTRINT)Ptr);
+
+		// Use short variable name here to improve readability further down.
+		FMemoryAccess* S = Shadow->Accesses;
+
+		// This is very practical to scope small section of code that we want to understand
+		// exactly what's going on inside.
+		if (UNLIKELY(Context.DetailedLogDepth || bDetailedLogGlobal))
+		{
+			if (Context.DetailedLogAddress == nullptr ||
+				Context.DetailedLogAddress == Ptr)
+			{
+				FPlatformMisc::LowLevelOutputDebugStringf(
+					TEXT("[%d] %s / %s of size %d offset %d (%s) at %p (ctx:%d clk:%u)")
+					TEXT(" S0 a:%d s:%d o:%d ctx:%d clk:%u")
+					TEXT(" S1 a:%d s:%d o:%d ctx:%d clk:%u")
+					TEXT(" S2 a:%d s:%d o:%d ctx:%d clk:%u")
+					TEXT(" S3 a:%d s:%d o:%d ctx:%d clk:%u\n"),
+					Context.ThreadId,
+					OpName,
+					AccessTypeToString(AccessType),
+					Size,
+					Offset,
+					LexToString(Order),
+					Ptr,
+					CurrentAccess.ContextId,
+					CurrentAccess.Clock,
+					S[0].AccessType, S[0].GetSize(), S[0].GetOffset(), S[0].ContextId, S[0].Clock,
+					S[1].AccessType, S[1].GetSize(), S[1].GetOffset(), S[1].ContextId, S[1].Clock,
+					S[2].AccessType, S[2].GetSize(), S[2].GetOffset(), S[2].ContextId, S[2].Clock,
+					S[3].AccessType, S[3].GetSize(), S[3].GetOffset(), S[3].ContextId, S[3].Clock
+				);
+			}
+		}
+
+		if (S[0].RawValue == CurrentAccess.RawValue ||
+			S[1].RawValue == CurrentAccess.RawValue ||
+			S[2].RawValue == CurrentAccess.RawValue ||
+			S[3].RawValue == CurrentAccess.RawValue)
+		{
+			// Skip everything including history logging if our access is already in the shadow.
+			return;
+		}
+
+		// Write to our history first so that any race detected from the shadow
+		// can be resolved.
+		FAccessHistory& AccessHistory = *Context.AccessHistory;
+		const uint32 HistoryIndex = AccessHistory.Index++;
+		if (UNLIKELY(AccessHistory.Index >= AccessHistory.NumEntry))
+		{
+			AccessHistory.ResetCount++;
+			AccessHistory.Index = 0;
+			AccessHistory.LastReset = FPlatformTime::Seconds();
+			AccessHistory.LastResetClock = CurrentAccess.Clock;
+		}
+		FAccessHistory::FEntry& HistoryEntry = AccessHistory.Entries[HistoryIndex];
+		HistoryEntry.Access.RawValue = CurrentAccess.RawValue;
+		// This is expensive. We need to optimize history so that we don't need to keep full callstack location
+		// and instead replay the event to discover the callstack along the way.
+		HistoryEntry.Location = GetCurrentCallstackLocation();
+		HistoryEntry.Pointer = GetAlignedPointer(Ptr);
+
+		bool bSaveNeeded = true;
+		FMemoryAccess PreviousAccess;
+		// We don't care about ordering but each 64-bit needs to keep their integrity.
+		// This will remain thread-safe as long as we are only reading and writing whole 64-bit words.
+		for (int32 Index = 0; Index < 4; ++Index)
+		{
+			uint64& RawValue = Shadow->Accesses[Index].RawValue;
+
+			// Read the whole 64-bit and store it locally to keep things atomic and thread-safe.
+			PreviousAccess.RawValue = RawValue;
+
+			// We fill the slots in order so we can early out.
+			// We can assume other slots are empty if we find one empty.
+			if (PreviousAccess.AccessType == EMemoryAccessType::ACCESS_TYPE_INVALID)
+			{
+				if (bSaveNeeded)
+				{
+					// Overwrite the whole 64-bit word to keep things atomic and thread-safe.
+					RawValue = CurrentAccess.RawValue;
+				}
+				return;
+			}
+
+#if WITH_RACEDETECTOR_CHECK
+			// Should never happen to have identical values since we verified at entry.
+			check(PreviousAccess.RawValue != CurrentAccess.RawValue);
+#endif
+
+			// Check if there is any overlap with the previous access
+			// Each bit in Access represets 1-byte in memory so we can just
+			// AND both access to know if we have common bytes being accessed. 
+			if (!(CurrentAccess.Access & PreviousAccess.Access))
+			{
+				continue;
+			}
+
+			// If we already have a slot, upgrade it if possible to avoid spilling to too many slots.
+			if (PreviousAccess.ContextId == CurrentAccess.ContextId)
+			{
+				// Ultimately, what we want in the slot is what is most susceptible to cause a race.
+				// We try to keep non-atomic first, then write first.. then finally reads.
+				if (bSaveNeeded &&
+					(PreviousAccess.bIsAtomic > CurrentAccess.bIsAtomic ||
+						CurrentAccess.bIsWrite > PreviousAccess.bIsWrite))
+				{
+					RawValue = CurrentAccess.RawValue;
+					bSaveNeeded = false;
+				}
+				continue;
+			}
+
+			// Check for obvious correct case where we're not racing.
+			const bool bBothRead   = !(PreviousAccess.bIsWrite | CurrentAccess.bIsWrite);
+			const bool bBothAtomic = PreviousAccess.bIsAtomic & CurrentAccess.bIsAtomic;
+			if (bBothRead | bBothAtomic)
+			{
+				continue;
+			}
+
+			// Verify that the clock we got in our bank for the previous access context has been
+			// synchronized since the last access... if not it means a barrier is missing and we have a race.
+			FClock Clock = Context.ClockBank.Get(PreviousAccess.ContextId);
+			if (Clock >= PreviousAccess.Clock)
+			{
+				continue;
+			}
+
+			// Reset the shadow to avoid reporting the same race multiple times
+			// Use the first reset as a sync point between multiple threads that might want to report the same race.
+			bool bWonReportingRace = FPlatformAtomics::InterlockedExchange((volatile int64*)&Shadow->Accesses[0].RawValue, 0);
+			Shadow->Accesses[1].RawValue = 0;
+			Shadow->Accesses[2].RawValue = 0;
+			Shadow->Accesses[3].RawValue = 0;
+
+			// We can find false positives when the race detector is shutting down or during shadow resets
+			// because some function will stop instrumenting before others, etc...
+			// Just ignore any race we have found if anything seems incoherent.
+			if (bWonReportingRace && bRaceDetectorActive && Epoch == GlobalEpoch && !bIsResettingShadow)
+			{
+				ReportRace(Context, Ptr, CurrentAccess, PreviousAccess, GetCurrentCallstackLocation());
+			}
+
+			// Just return now as we're in the race reporting case where we don't need to save our access anymore.
+			return;
+		}
+
+		if (bSaveNeeded)
+		{
+			// We haven't saved our access yet.
+			// Use history index as a 'random' position to avoid expensive computations.
+			Shadow->Accesses[HistoryIndex & 3].RawValue = CurrentAccess.RawValue;
+		}
+	}
+
+	INSTRUMENTATION_FUNCTION_ATTRIBUTES FORCEINLINE void InstrumentAccess(FContext& Context, void* Ptr, uint8 Size, EMemoryAccessType AccessType, FAtomicMemoryOrder Order, void* ReturnAddress, const TCHAR* OpName)
+	{
+		const bool bHasSyncObjectAttached = !IsAtomicOrderRelaxed(Order);
+		InstrumentAccess(Context, Ptr, Size, AccessType, Order, ReturnAddress, OpName, bHasSyncObjectAttached);
+	}
+
+	// Handles alignment and sizes bigger than a shadow cell
+	INSTRUMENTATION_FUNCTION_ATTRIBUTES void InstrumentAccessPreamble(FContext& Context, UPTRINT Ptr, uint32 Size, EMemoryAccessType AccessType, FAtomicMemoryOrder Order, void* ReturnAddress, const TCHAR* OpName)
+	{
+		UPTRINT AdjustedSize = FMath::Min((UPTRINT)Size, GetAlignedAddress(Ptr) + 8 - Ptr);
+		InstrumentAccess(Context, (void*)Ptr, static_cast<uint8>(AdjustedSize), AccessType, Order, ReturnAddress, OpName);
+		Ptr += AdjustedSize;
+		Size -= AdjustedSize;
+
+		if (Size)
+		{
+			while (Size >= 8)
+			{
+				InstrumentAccess(Context, (void*)Ptr, static_cast<uint8>(8), AccessType, Order, ReturnAddress, OpName);
+				Ptr += 8;
+				Size -= 8;
+			}
+
+			if (Size)
+			{
+				InstrumentAccess(Context, (void*)Ptr, static_cast<uint8>(Size), AccessType, Order, ReturnAddress, OpName);
+			}
+		}
+	}
+
+	INSTRUMENTATION_FUNCTION_ATTRIBUTES void InstrumentAtomicAccess(FContext& Context, void* AtomicAddr, uint8 Size, EMemoryAccessType AccessType, FAtomicMemoryOrder Order, const TFunctionRef<void()>& AtomicOp, void* ReturnAddress, const TCHAR* OpName)
+	{
+#if WITH_RACEDETECTOR_CHECK
+		check(FMath::Min((UPTRINT)Size, GetAlignedAddress((UPTRINT)AtomicAddr) + 8 - (UPTRINT)AtomicAddr) == (UPTRINT)Size);
+#endif
+		InstrumentAccess(Context, AtomicAddr, Size, AccessType, Order, ReturnAddress, OpName);
+		if (IsAtomicOrderRelaxed(Order))
+		{
+			AtomicOp();
+			return;
+		}
+
+		TRefCountPtr<FSyncObject> Atomic = GetSyncObject(AtomicAddr);
+
+		if (AccessType == ACCESS_TYPE_ATOMIC_READ_WRITE && IsAtomicOrderAcquireRelease(Order))
+		{
+			Atomic->SyncAcquireRelease(Context, AtomicOp, ReturnAddress, AtomicAddr, OpName);
+		}
+		else if ((AccessType & ACCESS_TYPE_ATOMIC_READ) == ACCESS_TYPE_ATOMIC_READ && IsAtomicOrderAcquire(Order))
+		{
+			Atomic->SyncAcquire(Context, AtomicOp, ReturnAddress, AtomicAddr, OpName);
+		}
+		else if ((AccessType & ACCESS_TYPE_ATOMIC_WRITE) == ACCESS_TYPE_ATOMIC_WRITE && IsAtomicOrderRelease(Order))
+		{
+			Atomic->SyncRelease(Context, AtomicOp, ReturnAddress, AtomicAddr, OpName);
+		}
+		else
+		{
+			checkf(false, TEXT("Unexpected memory order"));
+		}
+
+		Context.IncrementClock();
+	}
+
+	INSTRUMENTATION_FUNCTION_ATTRIBUTES void InstrumentAtomicAccess(FContext& Context, void* AtomicAddr, uint8 Size, EMemoryAccessType AccessType, FAtomicMemoryOrder SuccessOrder, FAtomicMemoryOrder FailureOrder, const TFunctionRef<bool()>& AtomicOp, void* ReturnAddress, const TCHAR* OpName)
+	{
+#if WITH_RACEDETECTOR_CHECK
+		check(FMath::Min((UPTRINT)Size, GetAlignedAddress((UPTRINT)AtomicAddr) + 8 - (UPTRINT)AtomicAddr) == (UPTRINT)Size);
+#endif
+		if (IsAtomicOrderRelaxed(SuccessOrder) && IsAtomicOrderRelaxed(FailureOrder))
+		{
+			// Both orders are relaxed, forward either one of them.
+			InstrumentAccess(Context, AtomicAddr, Size, AccessType, SuccessOrder, ReturnAddress, OpName);
+			
+			AtomicOp();
+		}
+		else
+		{
+			GetSyncObject(AtomicAddr)->SyncWithFailureSupport(Context, AtomicOp, AccessType, SuccessOrder, FailureOrder, ReturnAddress, AtomicAddr, OpName,
+				[&Context, AtomicAddr, Size, AccessType, ReturnAddress, &SuccessOrder, OpName](FAtomicMemoryOrder ActualOrder)
+				{
+					const bool bHasSyncObjectAttached = true;
+					InstrumentAccess(Context, AtomicAddr, Size, AccessType, ActualOrder, ReturnAddress, OpName, bHasSyncObjectAttached);
+				}
+			);
+		}
+	}
+
+	INSTRUMENTATION_FUNCTION_ATTRIBUTES void DumpContext()
+	{
+		FInstrumentationScope InstrumentationScope;
+		FContext& Context = EnsureCurrentContext();
+
+		TStringBuilder<4096> OtherClocks;
+		for (FContextId ContextId = 0; ContextId < 256; ++ContextId)
+		{
+			if (Context.ClockBank.Get(ContextId) != 0)
+			{
+				OtherClocks.Appendf(TEXT("[%d=%d]"), ContextId, Context.ClockBank.Get(ContextId));
+			}
+		}
+
+		FPlatformMisc::LowLevelOutputDebugStringf(
+			TEXT("Thread %s (%d), context %d, clock %d, other clocks %s\n"), 
+			*FThreadManager::Get().GetThreadName(Context.ThreadId),
+			Context.ThreadId, 
+			Context.ContextId, 
+			Context.CurrentClock(), 
+			*OtherClocks.ToString());
+	}
+
+	INSTRUMENTATION_FUNCTION_ATTRIBUTES void DumpContextDetailed()
+	{
+		FInstrumentationScope InstrumentationScope;
+		FContext& Context = EnsureCurrentContext();
+
+		FPlatformMisc::LowLevelOutputDebugStringf(
+			TEXT("Thread %s (%d), context %d, clock %d\n"),
+			*FThreadManager::Get().GetThreadName(Context.ThreadId),
+			Context.ThreadId, 
+			Context.ContextId,
+			Context.CurrentClock()
+		);
+
+		for (FContextId ContextId = 0; ContextId < 256; ++ContextId)
+		{
+			if (Context.ClockBank.Get(ContextId) != 0)
+			{
+				FPlatformMisc::LowLevelOutputDebugStringf(
+					TEXT("  [%d=%d]\n%s\n"), 
+					ContextId, 
+					Context.ClockBank.Get(ContextId),
+					*Context.ClockBank.GetLocation(ContextId).GetFullLocation().ToString()
+				);
+			}
+		}
+	}
+	
+	INSTRUMENTATION_FUNCTION_ATTRIBUTES void UnhookInstrumentation();
+	INSTRUMENTATION_FUNCTION_ATTRIBUTES void HookInstrumentation();
+
+	INSTRUMENTATION_FUNCTION_ATTRIBUTES void ResetShadow()
+	{
+		static UE::FMutex Lock;
+		if (!Lock.TryLock())
+		{
+			FPlatformMisc::LowLevelOutputDebugStringf(TEXT("[%d] Skip duplicate reset shadow\n"), Platform::GetCurrentThreadId());
+			return;
+		}
+		ON_SCOPE_EXIT
+		{
+			Lock.Unlock();
+		};
+
+		FPlatformMisc::LowLevelOutputDebugStringf(TEXT("[%d] Reset Shadow Memory Started\n"), Platform::GetCurrentThreadId());
+
+		bIsResettingShadow = true;
+
+		if (bRaceDetectorActive)
+		{
+			UnhookInstrumentation();
+		}
+
+		// In case there are some race while unmapping we will loop until
+		// both collections are empty and coherent to together to avoid false positives.
+		do
+		{
+			ResetShadowMemory();
+			ResetInstrumentedAtomics();
+
+		} while (Platform::HasShadowMemoryMapped() || HasInstrumentedAtomics());
+
+		FPlatformMisc::LowLevelOutputDebugStringf(TEXT("[%d] Reset Shadow Memory Ended\n"), Platform::GetCurrentThreadId());
+		
+		if (bRaceDetectorActive)
+		{
+			GlobalEpoch++;
+
+			HookInstrumentation();
+		}
+		bIsResettingShadow = false;
+	}
+
+	INSTRUMENTATION_FUNCTION_ATTRIBUTES void HandleReports()
+	{
+		static TUniquePtr<FArchive> RaceDetectorLog;
+
+		TArray<FString> LocalReports;
+		{
+			UE::TScopeLock Lock(RaceReportsLock);
+			LocalReports = MoveTemp(RaceReports);
+		}
+
+		if (LocalReports.Num())
+		{
+			if (RaceDetectorLog == nullptr)
+			{
+				const FString LogDir = GIsBuildMachine ? FPaths::Combine(*FPaths::EngineDir(), TEXT("Programs"), TEXT("AutomationTool"), TEXT("Saved"), TEXT("Logs")) : FPaths::ProjectLogDir();
+				const FString StateLogOutputFilename = FPaths::Combine(LogDir, TEXT("Sanitizer"),
+					FString::Printf(TEXT("RaceDetector-%08x-%s.log"), FPlatformProcess::GetCurrentProcessId(), *FDateTime::Now().ToIso8601().Replace(TEXT(":"), TEXT("."))));
+				RaceDetectorLog = TUniquePtr<FArchive>(IFileManager::Get().CreateFileWriter(*StateLogOutputFilename, FILEWRITE_AllowRead));
+			}
+
+			for (FString& Report : LocalReports)
+			{
+				RaceDetectorLog->Serialize(TCHAR_TO_ANSI(*Report), Report.Len());
+			}
+
+			RaceDetectorLog->Flush();
+		}
+	}
+
+	INSTRUMENTATION_FUNCTION_ATTRIBUTES void SanitizerThreadRun(volatile bool& bContinue)
+	{
+		FInstrumentationScope Scope;
+
+		bool bWasDebuggerPresent = false;
+		int32 MemoryUsageStatIterations = 0;
+		while (bContinue)
+		{
+			Platform::SleepMS(1000);
+
+			bIsDebuggerPresent = Platform::IsDebuggerPresent();
+			if (bIsDebuggerPresent && !bWasDebuggerPresent)
+			{
+				Platform::HideFirstChanceExceptionInVisualStudio();
+			}
+			bWasDebuggerPresent = bIsDebuggerPresent;
+
+			if (bRaceDetectorActive)
+			{
+				if (RaceDetectorTimeout.IsExpired())
+				{
+					ToggleRaceDetection(false);
+					continue;
+				}
+
+				double ShadowMemoryUsage = double(Platform::GetShadowMemoryUsage()) / (1024 *1024 *1024);
+				double SyncObjectMemoryUsage = double(GetSyncObjectMemoryUsage()) / (1024 * 1024 * 1024);
+
+				if (++MemoryUsageStatIterations > 10)
+				{
+					FPlatformMisc::LowLevelOutputDebugStringf(
+						TEXT("Sanitizer Memory Usage (Shadow : %.04f GB, SyncObjects : %.04f GB)\n"),
+						ShadowMemoryUsage,
+						SyncObjectMemoryUsage
+					);
+					MemoryUsageStatIterations = 0;
+				}
+
+				if (GRaceDetectorMaxMemoryUsage && (SyncObjectMemoryUsage + ShadowMemoryUsage) > GRaceDetectorMaxMemoryUsage)
+				{
+					ResetShadow();
+				}
+			}
+
+			HandleReports();
+		}
+	}
+
+	/**
+	 * FMalloc proxy removes instrumentation for allocators 
+	 * and free the shadow range to avoid detecting races in freed
+	 * memory.
+	 */
+	class FMallocInstrumentation : public FMalloc
+	{
+	private:
+		/** Malloc we're based on, aka using under the hood */
+		FMalloc* InnerMalloc;
+
+	public:
+		SAFE_OPERATOR_NEW_DELETE();
+
+		INSTRUMENTATION_FUNCTION_ATTRIBUTES explicit FMallocInstrumentation(FMalloc* InMalloc)
+			: InnerMalloc(InMalloc)
+		{
+			checkf(InnerMalloc, TEXT("FMallocInstrumentation is used without a valid malloc!"));
+		}
+
+		INSTRUMENTATION_FUNCTION_ATTRIBUTES virtual void InitializeStatsMetadata() override
+		{
+			FInstrumentationScope Scope;
+			InnerMalloc->InitializeStatsMetadata();
+		}
+
+		INSTRUMENTATION_FUNCTION_ATTRIBUTES virtual void* Malloc(SIZE_T Size, uint32 Alignment) override
+		{
+			FInstrumentationScope Scope;
+			return InnerMalloc->Malloc(Size, Alignment);
+		}
+
+		INSTRUMENTATION_FUNCTION_ATTRIBUTES virtual void* Realloc(void* InPtr, SIZE_T NewSize, uint32 Alignment) override
+		{
+			FInstrumentationScope Scope;
+
+			if (!bRaceDetectorActive)
+			{
+				return InnerMalloc->Realloc(InPtr, NewSize, Alignment);
+			}
+
+			// We have to always allocate new blocs in order to invalidate the old memory range before it can be reused.
+			void* NewPtr = nullptr;
+			if (NewSize)
+			{
+				NewPtr = InnerMalloc->Malloc(NewSize, Alignment);
+			}
+			
+			if (InPtr)
+			{
+				SIZE_T OldSize = 0;
+				InnerMalloc->GetAllocationSize(InPtr, OldSize);
+
+				if (OldSize)
+				{
+					if (NewPtr)
+					{
+						FMemory::Memcpy(NewPtr, InPtr, FMath::Min(OldSize, NewSize));
+					}
+					FreeMemoryRange(InPtr, OldSize);
+				}
+				InnerMalloc->Free(InPtr);
+			}
+
+			return NewPtr;
+		}
+
+		INSTRUMENTATION_FUNCTION_ATTRIBUTES virtual void Free(void* Ptr) override
+		{
+			FInstrumentationScope Scope;
+			SIZE_T Size = 0;
+			InnerMalloc->GetAllocationSize(Ptr, Size);
+			if (Ptr && Size)
+			{
+				FreeMemoryRange(Ptr, Size);
+			}
+			return InnerMalloc->Free(Ptr);
+		}
+
+		INSTRUMENTATION_FUNCTION_ATTRIBUTES virtual SIZE_T QuantizeSize(SIZE_T Count, uint32 Alignment) override
+		{
+			FInstrumentationScope Scope;
+			return InnerMalloc->QuantizeSize(Count, Alignment);
+		}
+
+		INSTRUMENTATION_FUNCTION_ATTRIBUTES virtual void UpdateStats() override
+		{
+			FInstrumentationScope Scope;
+			InnerMalloc->UpdateStats();
+		}
+
+		INSTRUMENTATION_FUNCTION_ATTRIBUTES virtual void GetAllocatorStats(FGenericMemoryStats& out_Stats) override
+		{
+			FInstrumentationScope Scope;
+			InnerMalloc->GetAllocatorStats(out_Stats);
+		}
+
+		INSTRUMENTATION_FUNCTION_ATTRIBUTES virtual void DumpAllocatorStats(class FOutputDevice& Ar) override
+		{
+			FInstrumentationScope Scope;
+			InnerMalloc->DumpAllocatorStats(Ar);
+		}
+
+		INSTRUMENTATION_FUNCTION_ATTRIBUTES virtual bool IsInternallyThreadSafe() const override
+		{
+			FInstrumentationScope Scope;
+			return InnerMalloc->IsInternallyThreadSafe();
+		}
+
+		INSTRUMENTATION_FUNCTION_ATTRIBUTES virtual bool ValidateHeap() override
+		{
+			FInstrumentationScope Scope;
+			return InnerMalloc->ValidateHeap();
+		}
+
+#if UE_ALLOW_EXEC_COMMANDS
+		INSTRUMENTATION_FUNCTION_ATTRIBUTES virtual bool Exec(UWorld* InWorld, const TCHAR* Cmd, FOutputDevice& Ar) override
+		{
+			FInstrumentationScope Scope;
+			return InnerMalloc->Exec(InWorld, Cmd, Ar);
+		}
+#endif // UE_ALLOW_EXEC_COMMANDS
+
+		INSTRUMENTATION_FUNCTION_ATTRIBUTES virtual bool GetAllocationSize(void* Ptr, SIZE_T& SizeOut) override
+		{
+			FInstrumentationScope Scope;
+			return InnerMalloc->GetAllocationSize(Ptr, SizeOut);
+		}
+
+		INSTRUMENTATION_FUNCTION_ATTRIBUTES virtual const TCHAR* GetDescriptiveName() override
+		{
+			FInstrumentationScope Scope;
+			return InnerMalloc->GetDescriptiveName();
+		}
+
+		INSTRUMENTATION_FUNCTION_ATTRIBUTES virtual void Trim(bool bTrimThreadCaches) override
+		{
+			FInstrumentationScope Scope;
+			InnerMalloc->Trim(bTrimThreadCaches);
+		}
+
+		INSTRUMENTATION_FUNCTION_ATTRIBUTES virtual void SetupTLSCachesOnCurrentThread() override
+		{
+			FInstrumentationScope Scope;
+			InnerMalloc->SetupTLSCachesOnCurrentThread();
+		}
+
+		INSTRUMENTATION_FUNCTION_ATTRIBUTES virtual void MarkTLSCachesAsUsedOnCurrentThread() override
+		{
+			FInstrumentationScope Scope;
+			InnerMalloc->MarkTLSCachesAsUsedOnCurrentThread();
+		}
+
+		INSTRUMENTATION_FUNCTION_ATTRIBUTES virtual void MarkTLSCachesAsUnusedOnCurrentThread() override
+		{
+			FInstrumentationScope Scope;
+			InnerMalloc->MarkTLSCachesAsUnusedOnCurrentThread();
+		}
+
+		INSTRUMENTATION_FUNCTION_ATTRIBUTES virtual void ClearAndDisableTLSCachesOnCurrentThread() override
+		{
+			FInstrumentationScope Scope;
+			InnerMalloc->ClearAndDisableTLSCachesOnCurrentThread();
+		}
+
+		INSTRUMENTATION_FUNCTION_ATTRIBUTES virtual void OnMallocInitialized() override
+		{
+			FInstrumentationScope Scope;
+			InnerMalloc->OnMallocInitialized();
+		}
+
+		INSTRUMENTATION_FUNCTION_ATTRIBUTES virtual void OnPreFork() override
+		{
+			FInstrumentationScope Scope;
+			InnerMalloc->OnPreFork();
+		}
+
+		INSTRUMENTATION_FUNCTION_ATTRIBUTES virtual void OnPostFork() override
+		{
+			FInstrumentationScope Scope;
+			InnerMalloc->OnPostFork();
+		}
+
+		// FMalloc interface end
+	};
+	
+	INSTRUMENTATION_FUNCTION_ATTRIBUTES void InitMemoryAllocator()
+	{
+		while (true)
+		{
+			FMalloc* LocalGMalloc = GMalloc;
+			FMalloc* Proxy = new FMallocInstrumentation(LocalGMalloc);
+			if (FPlatformAtomics::InterlockedCompareExchangePointer((void**)&GMalloc, Proxy, LocalGMalloc) == LocalGMalloc)
+			{
+				return;
+			}
+			delete Proxy;
+		}
+	}
+
+	INSTRUMENTATION_FUNCTION_ATTRIBUTES void PopulateHotPatchFunctions();
+	INSTRUMENTATION_FUNCTION_ATTRIBUTES void PrepareHotPatchFunctions();
+	INSTRUMENTATION_FUNCTION_ATTRIBUTES void HookAlwaysOnInstrumentation();
+	INSTRUMENTATION_FUNCTION_ATTRIBUTES bool Initialize()
+	{
+		if (bRuntimeInitialized)
+		{
+			return true;
+		}
+
+		PopulateHotPatchFunctions();
+
+		PrepareHotPatchFunctions();
+
+		if (!Platform::InitializePlatform())
+		{
+			return false;
+		}
+
+		Platform::InitShadowMemory();
+		
+		ShadowMemoryBase = Platform::GetShadowMemoryBase();
+		ShadowMemorySize = Platform::GetShadowMemorySize();
+		ShadowMemoryEnd = ShadowMemoryBase + ShadowMemorySize;
+
+		InitMemoryAllocator();
+		
+		HookAlwaysOnInstrumentation();
+
+		bRuntimeInitialized = true;
+		
+		const TCHAR* CommandLine = Platform::GetCommandLine();
+		if (FCString::Stristr(CommandLine, TEXT("-racedetector")))
+		{
+			ToggleRaceDetection(true);
+		}
+
+		return true;
+	}
+
+	INSTRUMENTATION_FUNCTION_ATTRIBUTES FContext& EnsureCurrentContext()
+	{
+		check(ContextTLSIndex != UINT32_MAX);
+		FContext* Context = GetThreadContext();
+		if (Context == nullptr)
+		{
+			Context = new FContext(Platform::GetCurrentThreadId());
+
+			// Refcount owned by the thread itself.
+			Context->AddRef();
+			Platform::SetTlsValue(ContextTLSIndex, Context);
+			if (GetTlsValueFast(ContextTLSIndex) != Context)
+			{
+				UE_LOG(LogRaceDetector, Fatal, TEXT("GetTlsValueFast implementation is invalid"));
+			}
+		}
+
+		return *Context;
+	}
+
+	INSTRUMENTATION_FUNCTION_ATTRIBUTES void ReleaseCurrentContext()
+	{
+		if (GetThreadContext())
+		{
+			ReleaseContext(GetThreadContext());
+			Platform::SetTlsValue(ContextTLSIndex, 0);
+		}
+	}
+
+	INSTRUMENTATION_FUNCTION_ATTRIBUTES void CleanupHotPatchFunctions();
+	INSTRUMENTATION_FUNCTION_ATTRIBUTES void ReleaseCurrentContext();
+	INSTRUMENTATION_FUNCTION_ATTRIBUTES void UnhookAlwaysOnInstrumentation();
+	INSTRUMENTATION_FUNCTION_ATTRIBUTES bool Shutdown()
+	{
+		if (!bRuntimeInitialized)
+		{
+			return true;
+		}
+		
+		ToggleRaceDetection(false);
+
+		UnhookAlwaysOnInstrumentation();
+
+		CleanupHotPatchFunctions();
+
+		if (!Platform::CleanupPlatform())
+		{
+			return false;
+		}
+
+		if (ContextTLSIndex != UINT32_MAX)
+		{
+			ReleaseCurrentContext();
+			FPlatformTLS::FreeTlsSlot(ContextTLSIndex);
+			ContextTLSIndex = UINT32_MAX;
+		}
+
+		bRuntimeInitialized = false;
+		return true;
+	}
+
+	INSTRUMENTATION_FUNCTION_ATTRIBUTES void ToggleFilterDuplicateRaces(bool bEnable)
+	{
+		GRaceDetectorFilterDuplicates = bEnable;
+	}
+
+	INSTRUMENTATION_FUNCTION_ATTRIBUTES void ToggleRaceDetectionUntil(UE::FTimeout Timeout)
+	{
+		RaceDetectorTimeout = Timeout;
+		ToggleRaceDetection(true);
+	}
+	
+	INSTRUMENTATION_FUNCTION_ATTRIBUTES void ToggleFilterOtherThreads(bool bEnable)
+	{
+		GRaceDetectorThreadFilter = Platform::GetCurrentThreadId();
+	}
+
+	INSTRUMENTATION_FUNCTION_ATTRIBUTES void ToggleRaceDetection(bool bEnable)
+	{
+		if (bRaceDetectorActive == bEnable)
+		{
+			return;
+		}
+
+		// Just make sure we're not initializing the stack from an access that comes from this function
+		FInstrumentationScope InstrumentationScope;
+
+		bRaceDetectorActive = bEnable;
+		if (!bEnable)
+		{
+			UnhookInstrumentation();
+
+			ResetShadow();
+		}
+		else
+		{
+			GlobalEpoch++;
+
+			HookInstrumentation();
+		}
+
+		UE_LOG(LogRaceDetector, Log, TEXT("Race detector has been toggled %s"), bEnable ? TEXT("on") : TEXT("off"));
+	}
+
+	INSTRUMENTATION_FUNCTION_ATTRIBUTES void ToggleThreadDetailedLog(bool bEnabled)
+	{
+		FContext& Context = EnsureCurrentContext();
+		if (bEnabled)
+		{
+			Context.DetailedLogDepth++;
+		}
+		else if (Context.DetailedLogDepth > 0)
+		{
+			Context.DetailedLogDepth--;
+		}
+	}
+
+	INSTRUMENTATION_FUNCTION_ATTRIBUTES void ToggleGlobalDetailedLog(bool bEnabled)
+	{
+		bDetailedLogGlobal = bEnabled;
+	}
+
+	INSTRUMENTATION_FUNCTION_ATTRIBUTES void ToggleFilterDetailedLogOnAddress(void* Address)
+	{
+		FContext& Context = EnsureCurrentContext();
+		Context.DetailedLogAddress = Address;
+	}
+
+
+} // UE::Sanitizer::RaceDetector
+
+using namespace UE::Sanitizer;
+using namespace UE::Sanitizer::RaceDetector;
+
+// We don't need to test for bRaceDetectorActive or bIsResettingShadow here
+// since we unhook the instrumentation instead.
+#define BEGIN_HANDLE_INSTRUMENTATION(FuncDoIfNotInstrumenting) \
+	FContext* ContextPtr = (FContext*)GetTlsValueFast(ContextTLSIndex); \
+	if (ContextPtr == nullptr || (ContextPtr->WinInstrumentationDepth | ContextPtr->InstrumentationDepth) != 0) \
+	{ \
+		FuncDoIfNotInstrumenting; \
+	} \
+	FContext& Context = *ContextPtr; \
+	Context.InstrumentationDepth++; \
+
+#define FINISH_HANDLE_INSTRUMENTATION() \
+	Context.InstrumentationDepth--;
+
+#define INSTRUMENTATION_CORE_FUNCTION_PREFIX CORE_API INSTRUMENTATION_FUNCTION_ATTRIBUTES
+
+extern "C" {
+
+	INSTRUMENTATION_CORE_FUNCTION_PREFIX void __RaceDetector__AnnotateHappensBefore(const char* f, int l, void* addr)
+	{
+		BEGIN_HANDLE_INSTRUMENTATION(return);
+
+		InstrumentAccess(Context, addr, sizeof(UPTRINT), ACCESS_TYPE_ATOMIC_READ, FAtomicMemoryOrder::MEMORY_ORDER_RELEASE, _ReturnAddress(), TEXT("AnnotateHappensBefore"));
+
+		TRefCountPtr<FSyncObject> Atomic = GetSyncObject(addr);
+		Atomic->SyncRelease(Context, []() {}, _ReturnAddress(), addr, TEXT("AnnotateHappensBefore"));
+		Context.IncrementClock();
+
+		FINISH_HANDLE_INSTRUMENTATION();
+	}
+
+	INSTRUMENTATION_CORE_FUNCTION_PREFIX void __RaceDetector__AnnotateHappensAfter(const char* f, int l, void* addr)
+	{
+		BEGIN_HANDLE_INSTRUMENTATION(return);
+
+		InstrumentAccess(Context, addr, sizeof(UPTRINT), ACCESS_TYPE_ATOMIC_READ, FAtomicMemoryOrder::MEMORY_ORDER_ACQUIRE, _ReturnAddress(), TEXT("AnnotateHappensBefore"));
+
+		TRefCountPtr<FSyncObject> Atomic = GetSyncObject(addr);
+		Atomic->SyncAcquire(Context, []() {}, _ReturnAddress(), addr, TEXT("AnnotateHappensAfter"));
+		Context.IncrementClock();
+
+		FINISH_HANDLE_INSTRUMENTATION();
+	}
+
+	INSTRUMENTATION_CORE_FUNCTION_PREFIX void __RaceDetector__Instrument_FuncEntry(void* ReturnAddress)
+	{
+		BEGIN_HANDLE_INSTRUMENTATION(return);
+
+		AddCallstackFrame(Context, ReturnAddress);
+
+		FINISH_HANDLE_INSTRUMENTATION();
+	}
+
+	INSTRUMENTATION_CORE_FUNCTION_PREFIX void __RaceDetector__Instrument_FuncExit()
+	{
+		BEGIN_HANDLE_INSTRUMENTATION(return);
+
+		RemoveCallstackFrame(Context);
+
+		FINISH_HANDLE_INSTRUMENTATION();
+	}
+
+	INSTRUMENTATION_CORE_FUNCTION_PREFIX __attribute__((hot)) void __RaceDetector__Instrument_Store(uint64 Address, uint32 Size)
+	{
+		BEGIN_HANDLE_INSTRUMENTATION(return);
+
+		AddCallstackFrame(Context, _ReturnAddress());
+		InstrumentAccessPreamble(Context, Address, Size, EMemoryAccessType::ACCESS_TYPE_WRITE, FAtomicMemoryOrder::MEMORY_ORDER_RELAXED, _ReturnAddress(), TEXT("Store"));
+		RemoveCallstackFrame(Context);
+
+		FINISH_HANDLE_INSTRUMENTATION();
+	}
+
+	INSTRUMENTATION_CORE_FUNCTION_PREFIX __attribute__((hot)) void __RaceDetector__Instrument_Load(uint64 Address, uint32 Size)
+	{
+		BEGIN_HANDLE_INSTRUMENTATION(return);
+
+		AddCallstackFrame(Context, _ReturnAddress());
+		InstrumentAccessPreamble(Context, Address, Size, EMemoryAccessType::ACCESS_TYPE_READ, FAtomicMemoryOrder::MEMORY_ORDER_RELAXED, _ReturnAddress(), TEXT("Load"));
+		RemoveCallstackFrame(Context);
+
+		FINISH_HANDLE_INSTRUMENTATION();
+	}
+
+	INSTRUMENTATION_CORE_FUNCTION_PREFIX void __RaceDetector__Instrument_StoreRange(uint64 Address, uint32 Size)
+	{
+		if (Size == 0)
+		{
+			return;
+		}
+
+		BEGIN_HANDLE_INSTRUMENTATION(return);
+		
+		AddCallstackFrame(Context, _ReturnAddress());
+		InstrumentAccessPreamble(Context, Address, Size, EMemoryAccessType::ACCESS_TYPE_WRITE, FAtomicMemoryOrder::MEMORY_ORDER_RELAXED, _ReturnAddress(), TEXT("StoreRange"));
+		RemoveCallstackFrame(Context);
+
+		FINISH_HANDLE_INSTRUMENTATION();
+	}
+
+	INSTRUMENTATION_CORE_FUNCTION_PREFIX void __RaceDetector__Instrument_LoadRange(uint64 Address, uint32 Size)
+	{
+		if (Size == 0)
+		{
+			return;
+		}
+
+		BEGIN_HANDLE_INSTRUMENTATION(return);
+
+		AddCallstackFrame(Context, _ReturnAddress());
+		InstrumentAccessPreamble(Context, Address, Size, EMemoryAccessType::ACCESS_TYPE_READ, FAtomicMemoryOrder::MEMORY_ORDER_RELAXED, _ReturnAddress(), TEXT("LoadRange"));
+		RemoveCallstackFrame(Context);
+
+		FINISH_HANDLE_INSTRUMENTATION();
+	}
+}
+
+template <typename AtomicType>
+FORCEINLINE INSTRUMENTATION_FUNCTION_ATTRIBUTES AtomicType InstrumentAtomicLoad(AtomicType* Atomic, FAtomicMemoryOrder Order)
+{
+	AtomicType Ret;
+	auto AtomicOp = [Atomic, &Ret]() {
+		Ret = UE::Core::Private::Atomic::Load(Atomic);
+		};
+
+	BEGIN_HANDLE_INSTRUMENTATION(AtomicOp(); return Ret);
+	AddCallstackFrame(Context, _ReturnAddress());
+	InstrumentAtomicAccess(Context, Atomic, sizeof(AtomicType), EMemoryAccessType::ACCESS_TYPE_ATOMIC_READ, Order, AtomicOp, _ReturnAddress(), TEXT("AtomicLoad"));
+	RemoveCallstackFrame(Context);
+	FINISH_HANDLE_INSTRUMENTATION()
+	return Ret;
+}
+
+template <typename AtomicType>
+FORCEINLINE INSTRUMENTATION_FUNCTION_ATTRIBUTES void InstrumentAtomicStore(AtomicType* Atomic, AtomicType Val, FAtomicMemoryOrder Order)
+{
+	auto AtomicOp = [Atomic, Val]() {
+		UE::Core::Private::Atomic::Store(Atomic, Val);
+		};
+
+	BEGIN_HANDLE_INSTRUMENTATION(AtomicOp(); return);
+	AddCallstackFrame(Context, _ReturnAddress());
+	InstrumentAtomicAccess(Context, Atomic, sizeof(AtomicType), EMemoryAccessType::ACCESS_TYPE_ATOMIC_WRITE, Order, AtomicOp, _ReturnAddress(), TEXT("AtomicStore"));
+	RemoveCallstackFrame(Context);
+	FINISH_HANDLE_INSTRUMENTATION();
+}
+
+template <typename AtomicType>
+FORCEINLINE INSTRUMENTATION_FUNCTION_ATTRIBUTES AtomicType InstrumentAtomicExchange(AtomicType* Atomic, AtomicType Val, FAtomicMemoryOrder Order)
+{
+	AtomicType Ret;
+	auto AtomicOp = [Atomic, Val, &Ret]() {
+		Ret = UE::Core::Private::Atomic::Exchange(Atomic, Val);
+		};
+
+	BEGIN_HANDLE_INSTRUMENTATION(AtomicOp(); return Ret);
+	AddCallstackFrame(Context, _ReturnAddress());
+	InstrumentAtomicAccess(Context, Atomic, sizeof(AtomicType), EMemoryAccessType::ACCESS_TYPE_ATOMIC_READ_WRITE, Order, AtomicOp, _ReturnAddress(), TEXT("AtomicExchange"));
+	RemoveCallstackFrame(Context);
+	FINISH_HANDLE_INSTRUMENTATION();
+	return Ret;
+}
+
+template <typename AtomicType>
+FORCEINLINE INSTRUMENTATION_FUNCTION_ATTRIBUTES AtomicType InstrumentAtomicCompareExchange(AtomicType* Atomic, AtomicType* Expected, AtomicType Val, FAtomicMemoryOrder SuccessOrder, FAtomicMemoryOrder FailureOrder)
+{
+	AtomicType Ret;
+	auto AtomicOp = [Atomic, Expected, Val, &Ret]() -> bool {
+		Ret = UE::Core::Private::Atomic::CompareExchange(Atomic, *Expected, Val);
+		if (Ret != *Expected)
+		{
+			*Expected = Ret;
+			return false;
+		}
+		return true;
+	};
+
+	BEGIN_HANDLE_INSTRUMENTATION(AtomicOp(); return Ret);
+	AddCallstackFrame(Context, _ReturnAddress());
+	InstrumentAtomicAccess(Context, Atomic, sizeof(AtomicType), EMemoryAccessType::ACCESS_TYPE_ATOMIC_READ_WRITE, SuccessOrder, FailureOrder, AtomicOp, _ReturnAddress(), TEXT("CompareExchange"));
+	RemoveCallstackFrame(Context);
+	FINISH_HANDLE_INSTRUMENTATION();
+	return Ret;
+}
+
+template <typename AtomicType>
+FORCEINLINE INSTRUMENTATION_FUNCTION_ATTRIBUTES AtomicType InstrumentAtomicFetchAdd(AtomicType* Atomic, AtomicType Val, FAtomicMemoryOrder Order)
+{
+	AtomicType Ret;
+	auto AtomicOp = [Atomic, Val, &Ret]() {
+		Ret = UE::Core::Private::Atomic::AddExchange(Atomic, Val);
+		};
+
+	BEGIN_HANDLE_INSTRUMENTATION(AtomicOp(); return Ret);
+	AddCallstackFrame(Context, _ReturnAddress());
+	InstrumentAtomicAccess(Context, Atomic, sizeof(AtomicType), EMemoryAccessType::ACCESS_TYPE_ATOMIC_READ_WRITE, Order, AtomicOp, _ReturnAddress(), TEXT("FetchAdd"));
+	RemoveCallstackFrame(Context);
+	FINISH_HANDLE_INSTRUMENTATION();
+	return Ret;
+}
+
+template <typename AtomicType>
+FORCEINLINE INSTRUMENTATION_FUNCTION_ATTRIBUTES AtomicType InstrumentAtomicFetchSub(AtomicType* Atomic, AtomicType Val, FAtomicMemoryOrder Order)
+{
+	AtomicType Ret;
+	auto AtomicOp = [Atomic, Val, &Ret]() {
+		Ret = UE::Core::Private::Atomic::SubExchange(Atomic, Val);
+		};
+
+	BEGIN_HANDLE_INSTRUMENTATION(AtomicOp(); return Ret);
+	AddCallstackFrame(Context, _ReturnAddress());
+	InstrumentAtomicAccess(Context, Atomic, sizeof(AtomicType), EMemoryAccessType::ACCESS_TYPE_ATOMIC_READ_WRITE, Order, AtomicOp, _ReturnAddress(), TEXT("FetchSub"));
+	RemoveCallstackFrame(Context);
+	FINISH_HANDLE_INSTRUMENTATION();
+	return Ret;
+}
+
+template <typename AtomicType>
+FORCEINLINE INSTRUMENTATION_FUNCTION_ATTRIBUTES AtomicType InstrumentAtomicFetchOr(AtomicType* Atomic, AtomicType Val, FAtomicMemoryOrder Order)
+{
+	AtomicType Ret;
+	auto AtomicOp = [Atomic, Val, &Ret]() {
+		Ret = UE::Core::Private::Atomic::OrExchange(Atomic, Val);
+		};
+
+	BEGIN_HANDLE_INSTRUMENTATION(AtomicOp(); return Ret);
+	AddCallstackFrame(Context, _ReturnAddress());
+	InstrumentAtomicAccess(Context, Atomic, sizeof(AtomicType), EMemoryAccessType::ACCESS_TYPE_ATOMIC_READ_WRITE, Order, AtomicOp, _ReturnAddress(), TEXT("FetchOr"));
+	RemoveCallstackFrame(Context);
+	FINISH_HANDLE_INSTRUMENTATION();
+	return Ret;
+}
+
+template <typename AtomicType>
+FORCEINLINE INSTRUMENTATION_FUNCTION_ATTRIBUTES AtomicType InstrumentAtomicFetchXor(AtomicType* Atomic, AtomicType Val, FAtomicMemoryOrder Order)
+{
+	AtomicType Ret;
+	auto AtomicOp = [Atomic, Val, &Ret]() {
+		Ret = UE::Core::Private::Atomic::XorExchange(Atomic, Val);
+	};
+	
+	BEGIN_HANDLE_INSTRUMENTATION(AtomicOp(); return Ret);
+	AddCallstackFrame(Context, _ReturnAddress());
+	InstrumentAtomicAccess(Context, Atomic, sizeof(AtomicType), EMemoryAccessType::ACCESS_TYPE_ATOMIC_READ_WRITE, Order, AtomicOp, _ReturnAddress(), TEXT("FetchXor"));
+	RemoveCallstackFrame(Context);
+	FINISH_HANDLE_INSTRUMENTATION();
+	return Ret;
+}
+
+template <typename AtomicType>
+FORCEINLINE INSTRUMENTATION_FUNCTION_ATTRIBUTES AtomicType InstrumentAtomicFetchAnd(AtomicType* Atomic, AtomicType Val, FAtomicMemoryOrder Order)
+{
+	AtomicType Ret;
+	auto AtomicOp = [Atomic, Val, &Ret]() {
+		Ret = UE::Core::Private::Atomic::AndExchange(Atomic, Val);
+		};
+	
+	BEGIN_HANDLE_INSTRUMENTATION(AtomicOp(); return Ret);
+	AddCallstackFrame(Context, _ReturnAddress());
+	InstrumentAtomicAccess(Context, Atomic, sizeof(AtomicType), EMemoryAccessType::ACCESS_TYPE_ATOMIC_READ_WRITE, Order, AtomicOp, _ReturnAddress(), TEXT("FetchAnd"));
+	RemoveCallstackFrame(Context);
+	FINISH_HANDLE_INSTRUMENTATION();
+	return Ret;
+}
+
+#define INSTRUMENT_LOAD_FUNC(Type) \
+INSTRUMENTATION_CORE_FUNCTION_PREFIX Type __RaceDetector__Instrument_AtomicLoad_##Type(Type* Atomic, FAtomicMemoryOrder MemoryOrder) \
+{ \
+	return InstrumentAtomicLoad(Atomic, MemoryOrder); \
+}
+
+#define INSTRUMENT_STORE_FUNC(Type) \
+INSTRUMENTATION_CORE_FUNCTION_PREFIX void __RaceDetector__Instrument_AtomicStore_##Type(Type* Atomic, Type Val, FAtomicMemoryOrder MemoryOrder) \
+{ \
+	return InstrumentAtomicStore(Atomic, Val, MemoryOrder); \
+}
+
+#define INSTRUMENT_EXCHANGE_FUNC(Type) \
+INSTRUMENTATION_CORE_FUNCTION_PREFIX Type __RaceDetector__Instrument_AtomicExchange_##Type(Type* Atomic, Type Val, FAtomicMemoryOrder MemoryOrder) \
+{ \
+    return InstrumentAtomicExchange(Atomic, Val, MemoryOrder); \
+}
+
+#define INSTRUMENT_COMPARE_EXCHANGE_FUNC(Type) \
+INSTRUMENTATION_CORE_FUNCTION_PREFIX Type __RaceDetector__Instrument_AtomicCompareExchange_##Type(Type* Atomic, Type* Expected, Type Val, FAtomicMemoryOrder SuccessMemoryOrder, FAtomicMemoryOrder FailureMemoryOrder) \
+{ \
+    return InstrumentAtomicCompareExchange(Atomic, Expected, Val, SuccessMemoryOrder, FailureMemoryOrder); \
+} 
+
+#define INSTRUMENT_RMW_FUNC(Func, Type) \
+INSTRUMENTATION_CORE_FUNCTION_PREFIX Type __RaceDetector__Instrument_Atomic##Func##_##Type(Type* Atomic, Type Val, FAtomicMemoryOrder MemoryOrder) \
+{ \
+	return InstrumentAtomic##Func(Atomic, Val, MemoryOrder); \
+}
+
+// Define the functions with C linkage, referenced by the compiler's instrumentation pass.
+extern "C" {
+	INSTRUMENT_LOAD_FUNC(int8)
+	INSTRUMENT_LOAD_FUNC(int16)
+	INSTRUMENT_LOAD_FUNC(int32)
+	INSTRUMENT_LOAD_FUNC(int64)
+
+	INSTRUMENT_STORE_FUNC(int8)
+	INSTRUMENT_STORE_FUNC(int16)
+	INSTRUMENT_STORE_FUNC(int32)
+	INSTRUMENT_STORE_FUNC(int64)
+
+	INSTRUMENT_EXCHANGE_FUNC(int8)
+	INSTRUMENT_EXCHANGE_FUNC(int16)
+	INSTRUMENT_EXCHANGE_FUNC(int32)
+	INSTRUMENT_EXCHANGE_FUNC(int64)
+
+	INSTRUMENT_COMPARE_EXCHANGE_FUNC(int8)
+	INSTRUMENT_COMPARE_EXCHANGE_FUNC(int16)
+	INSTRUMENT_COMPARE_EXCHANGE_FUNC(int32)
+	INSTRUMENT_COMPARE_EXCHANGE_FUNC(int64)
+
+	INSTRUMENT_RMW_FUNC(FetchAdd, int8)
+	INSTRUMENT_RMW_FUNC(FetchAdd, int16)
+	INSTRUMENT_RMW_FUNC(FetchAdd, int32)
+	INSTRUMENT_RMW_FUNC(FetchAdd, int64)
+	INSTRUMENT_RMW_FUNC(FetchSub, int8)
+	INSTRUMENT_RMW_FUNC(FetchSub, int16)
+	INSTRUMENT_RMW_FUNC(FetchSub, int32)
+	INSTRUMENT_RMW_FUNC(FetchSub, int64)
+	INSTRUMENT_RMW_FUNC(FetchOr, int8)
+	INSTRUMENT_RMW_FUNC(FetchOr, int16)
+	INSTRUMENT_RMW_FUNC(FetchOr, int32)
+	INSTRUMENT_RMW_FUNC(FetchOr, int64)
+	INSTRUMENT_RMW_FUNC(FetchXor, int8)
+	INSTRUMENT_RMW_FUNC(FetchXor, int16)
+	INSTRUMENT_RMW_FUNC(FetchXor, int32)
+	INSTRUMENT_RMW_FUNC(FetchXor, int64)
+	INSTRUMENT_RMW_FUNC(FetchAnd, int8)
+	INSTRUMENT_RMW_FUNC(FetchAnd, int16)
+	INSTRUMENT_RMW_FUNC(FetchAnd, int32)
+	INSTRUMENT_RMW_FUNC(FetchAnd, int64)
+}
+
+
+// We need the pointers to the native instrument functions to hotpatch them.
+extern "C" {
+	void AnnotateHappensBefore(const char* f, int l, void* addr);
+	void AnnotateHappensAfter(const char* f, int l, void* addr);
+	void __Instrument_FuncEntry(void* ReturnAddress);
+	void __Instrument_FuncExit();
+	void __Instrument_StoreRange(uint64 Address, uint32 Size);
+	void __Instrument_LoadRange(uint64 Address, uint32 Size);
+	void __Instrument_Store(uint64 Address, uint32 Size);
+	void __Instrument_Load(uint64 Address, uint32 Size);
+}
+
+namespace UE::Sanitizer::RaceDetector {
+
+// Those are never toggled off after boot.
+TArray<TPair<void*, void*>>  AlwaysOnInstrumentationFunctions;
+// We use our own hotpatching method to hook into these atomically.
+TArray<TPair<void*, void*>>  HotpatchInstrumentationFunctions;
+// We can place the hottest function that do nothing in a fast collection so they get patched by a RET when inactive.
+// This saves more than 1m30s on a 8m15s map load!!
+TArray<TPair<void*, void*>>  HotpatchInstrumentationFunctionsFast;
+
+INSTRUMENTATION_FUNCTION_ATTRIBUTES void PopulateHotPatchFunctions()
+{
+	if (!HotpatchInstrumentationFunctions.IsEmpty())
+	{
+		return;
+	}
+
+	AlwaysOnInstrumentationFunctions.Emplace(__Thunk__Instrument_FuncEntry, __RaceDetector__Instrument_FuncEntry);
+	AlwaysOnInstrumentationFunctions.Emplace(__Thunk__Instrument_FuncExit, __RaceDetector__Instrument_FuncExit);
+
+	// Hook the instrumentation thunks needed for race detector, this is where the calls from all the modules end up.
+	HotpatchInstrumentationFunctionsFast.Emplace(__Thunk__AnnotateHappensBefore, __RaceDetector__AnnotateHappensBefore);
+	HotpatchInstrumentationFunctionsFast.Emplace(__Thunk__AnnotateHappensAfter, __RaceDetector__AnnotateHappensAfter);
+	HotpatchInstrumentationFunctionsFast.Emplace(__Thunk__Instrument_StoreRange, __RaceDetector__Instrument_StoreRange);
+	HotpatchInstrumentationFunctionsFast.Emplace(__Thunk__Instrument_LoadRange, __RaceDetector__Instrument_LoadRange);
+	HotpatchInstrumentationFunctionsFast.Emplace(__Thunk__Instrument_Store, __RaceDetector__Instrument_Store);
+	HotpatchInstrumentationFunctionsFast.Emplace(__Thunk__Instrument_Load, __RaceDetector__Instrument_Load);
+
+	// In non-monolithic, we also hook the core functions directly instead of via the thunks to save on another set of JMPs
+#if !IS_MONOLITHIC
+	AlwaysOnInstrumentationFunctions.Emplace(__Instrument_FuncEntry, __RaceDetector__Instrument_FuncEntry);
+	AlwaysOnInstrumentationFunctions.Emplace(__Instrument_FuncExit, __RaceDetector__Instrument_FuncExit);
+
+	HotpatchInstrumentationFunctionsFast.Emplace(AnnotateHappensBefore, __RaceDetector__AnnotateHappensBefore);
+	HotpatchInstrumentationFunctionsFast.Emplace(AnnotateHappensAfter, __RaceDetector__AnnotateHappensAfter);
+	HotpatchInstrumentationFunctionsFast.Emplace(__Instrument_StoreRange, __RaceDetector__Instrument_StoreRange);
+	HotpatchInstrumentationFunctionsFast.Emplace(__Instrument_LoadRange, __RaceDetector__Instrument_LoadRange);
+	HotpatchInstrumentationFunctionsFast.Emplace(__Instrument_Store, __RaceDetector__Instrument_Store);
+	HotpatchInstrumentationFunctionsFast.Emplace(__Instrument_Load, __RaceDetector__Instrument_Load);
+#endif
+
+#define HOTPATCH_FUNC(Type) \
+	HotpatchInstrumentationFunctions.Emplace(__Thunk__Instrument_AtomicLoad_##Type, __RaceDetector__Instrument_AtomicLoad_##Type); \
+	HotpatchInstrumentationFunctions.Emplace(__Thunk__Instrument_AtomicStore_##Type, __RaceDetector__Instrument_AtomicStore_##Type); \
+	HotpatchInstrumentationFunctions.Emplace(__Thunk__Instrument_AtomicExchange_##Type, __RaceDetector__Instrument_AtomicExchange_##Type); \
+	HotpatchInstrumentationFunctions.Emplace(__Thunk__Instrument_AtomicCompareExchange_##Type, __RaceDetector__Instrument_AtomicCompareExchange_##Type);
+
+#define HOTPATCH_RMW_FUNC(Func, Type) \
+	HotpatchInstrumentationFunctions.Emplace(__Thunk__Instrument_Atomic##Func##_##Type, __RaceDetector__Instrument_Atomic##Func##_##Type);
+
+	HOTPATCH_FUNC(int8)
+	HOTPATCH_FUNC(int16)
+	HOTPATCH_FUNC(int32)
+	HOTPATCH_FUNC(int64)
+
+	HOTPATCH_RMW_FUNC(FetchAdd, int8)
+	HOTPATCH_RMW_FUNC(FetchAdd, int16)
+	HOTPATCH_RMW_FUNC(FetchAdd, int32)
+	HOTPATCH_RMW_FUNC(FetchAdd, int64)
+	HOTPATCH_RMW_FUNC(FetchSub, int8)
+	HOTPATCH_RMW_FUNC(FetchSub, int16)
+	HOTPATCH_RMW_FUNC(FetchSub, int32)
+	HOTPATCH_RMW_FUNC(FetchSub, int64)
+	HOTPATCH_RMW_FUNC(FetchOr, int8)
+	HOTPATCH_RMW_FUNC(FetchOr, int16)
+	HOTPATCH_RMW_FUNC(FetchOr, int32)
+	HOTPATCH_RMW_FUNC(FetchOr, int64)
+	HOTPATCH_RMW_FUNC(FetchXor, int8)
+	HOTPATCH_RMW_FUNC(FetchXor, int16)
+	HOTPATCH_RMW_FUNC(FetchXor, int32)
+	HOTPATCH_RMW_FUNC(FetchXor, int64)
+	HOTPATCH_RMW_FUNC(FetchAnd, int8)
+	HOTPATCH_RMW_FUNC(FetchAnd, int16)
+	HOTPATCH_RMW_FUNC(FetchAnd, int32)
+	HOTPATCH_RMW_FUNC(FetchAnd, int64)
+}
+
+INSTRUMENTATION_FUNCTION_ATTRIBUTES void PrepareHotPatchFunctions()
+{
+	for (const auto [ThunkFn, DetouredFn] : AlwaysOnInstrumentationFunctions)
+	{
+		Platform::PrepareTrampoline(ThunkFn, DetouredFn, false /* bUseRETBypass */);
+	}
+
+	for (const auto [ThunkFn, DetouredFn] : HotpatchInstrumentationFunctions)
+	{
+		Platform::PrepareTrampoline(ThunkFn, DetouredFn, false /* bUseRETBypass */);
+	}
+
+	for (const auto [ThunkFn, DetouredFn] : HotpatchInstrumentationFunctionsFast)
+	{
+		Platform::PrepareTrampoline(ThunkFn, DetouredFn, true /* bUseRETBypass */);
+	}
+}
+
+INSTRUMENTATION_FUNCTION_ATTRIBUTES void CleanupHotPatchFunctions()
+{
+	for (const auto [ThunkFn, DetouredFn] : AlwaysOnInstrumentationFunctions)
+	{
+		Platform::CleanupTrampoline(ThunkFn);
+	}
+
+	for (const auto [ThunkFn, DetouredFn] : HotpatchInstrumentationFunctions)
+	{
+		Platform::CleanupTrampoline(ThunkFn);
+	}
+
+	for (const auto [ThunkFn, DetouredFn] : HotpatchInstrumentationFunctionsFast)
+	{
+		Platform::CleanupTrampoline(ThunkFn);
+	}
+}
+
+INSTRUMENTATION_FUNCTION_ATTRIBUTES void HookAlwaysOnInstrumentation()
+{
+	for (const auto [ThunkFn, DetouredFn] : AlwaysOnInstrumentationFunctions)
+	{
+		Platform::ActivateTrampoline(ThunkFn);
+	}
+
+	Platform::FlushInstructionCache();
+}
+
+INSTRUMENTATION_FUNCTION_ATTRIBUTES void HookInstrumentation()
+{
+	for (const auto [ThunkFn, DetouredFn] : HotpatchInstrumentationFunctions)
+	{
+		Platform::ActivateTrampoline(ThunkFn);
+	}
+
+	for (const auto [ThunkFn, DetouredFn] : HotpatchInstrumentationFunctionsFast)
+	{
+		Platform::ActivateTrampoline(ThunkFn);
+	}
+
+	Platform::FlushInstructionCache();
+}
+
+INSTRUMENTATION_FUNCTION_ATTRIBUTES void UnhookAlwaysOnInstrumentation()
+{
+	for (const auto [ThunkFn, DetouredFn] : AlwaysOnInstrumentationFunctions)
+	{
+		Platform::DeactivateTrampoline(ThunkFn, false /* bUseRETBypass */);
+	}
+
+	Platform::FlushInstructionCache();
+}
+
+INSTRUMENTATION_FUNCTION_ATTRIBUTES void UnhookInstrumentation()
+{
+	for (const auto [ThunkFn, DetouredFn] : HotpatchInstrumentationFunctions)
+	{
+		Platform::DeactivateTrampoline(ThunkFn, false /* bUseRETBypass */);
+	}
+
+	for (const auto [ThunkFn, DetouredFn] : HotpatchInstrumentationFunctionsFast)
+	{
+		Platform::DeactivateTrampoline(ThunkFn, true /* bUseRETBypass */);
+	}
+
+	Platform::FlushInstructionCache();
+}
+
+INSTRUMENTATION_FUNCTION_ATTRIBUTES void InitializeBoot()
+{
+	ContextTLSIndex = Platform::AllocTlsSlot();
+	EnsureCurrentContext();
+}
+
+} // UE::Sanitizer::RaceDetector
+
+#ifdef INSTRUMENTATION_DEBUG
+UE_ENABLE_OPTIMIZATION
+#endif
+
+#endif // USING_INSTRUMENTATION
\ No newline at end of file
diff --git a/Engine/Source/Runtime/Core/Private/Sanitizer/RaceDetectorWindows.cpp b/Engine/Source/Runtime/Core/Private/Sanitizer/RaceDetectorWindows.cpp
new file mode 100644
--- /dev/null
+++ b/Engine/Source/Runtime/Core/Private/Sanitizer/RaceDetectorWindows.cpp
@@ -0,0 +1,1264 @@
+// Copyright Epic Games, Inc. All Rights Reserved.
+
+#include "Misc/CoreMiscDefines.h"
+
+#if PLATFORM_WINDOWS && USING_INSTRUMENTATION
+
+#include "Sanitizer/Types.h"
+#include "Sanitizer/RaceDetector.h"
+#include "Sanitizer/RaceDetectorTypes.h"
+#include "Sanitizer/RaceDetectorInterface.h"
+#include "Instrumentation/Defines.h"
+#include "Instrumentation/Containers.h"
+
+#include "Windows/WindowsHWrapper.h"
+#include "Windows/AllowWindowsPlatformTypes.h"
+
+THIRD_PARTY_INCLUDES_START
+#include <detours/detours.h>
+#include <winternl.h> // for NTSTATUS
+#include <process.h>  // for _beginthreadex
+THIRD_PARTY_INCLUDES_END
+
+extern "C" __declspec(dllimport) void __stdcall GetCurrentThreadStackLimits(void**, void**);
+
+using namespace UE::Sanitizer;
+using namespace UE::Sanitizer::RaceDetector;
+
+namespace UE::Sanitizer::RaceDetector {
+	extern void InitializeBoot();
+	extern void SanitizerThreadRun(volatile bool& bContinue);
+	extern volatile bool bIsResettingShadow;
+}
+
+#define START_WINAPI_INSTRUMENTATION(DoIfNotInstrumenting) \
+	DWORD OriginalError = GetLastError(); \
+	FContext* ContextPtr = GetThreadContext(); \
+	if (ContextPtr == nullptr || !ShouldInstrument(*ContextPtr)) \
+	{ \
+		SetLastError(OriginalError); \
+		DoIfNotInstrumenting; \
+	} \
+	FContext& Context = *ContextPtr; \
+	Context.WinInstrumentationDepth++; \
+	SetLastError(OriginalError);
+
+#define FINISH_WINAPI_INSTRUMENTATION() \
+	Context.WinInstrumentationDepth--;
+
+struct FInstrumentedStartThreadArgs {
+	SAFE_OPERATOR_NEW_DELETE()
+
+	LPVOID RealThreadParameter;
+	LPTHREAD_START_ROUTINE RealStartRoutine;
+	bool bDetailedLog = false;
+
+	// ClockBank inherited from the CreateThread call
+	FClockBank ClockBank;
+
+	std::atomic<HANDLE> ThreadHandle{ nullptr };
+};
+
+// This is actually really important since we need to set back the same
+// Windows error that came back from the true Windows API call since the
+// instrumentation is doing some API calls of its own, we can end up
+// resetting the last errors and causing issues in perceived results.
+struct FLastErrorPreservationScope
+{
+	DWORD OriginalError;
+
+	INSTRUMENTATION_FUNCTION_ATTRIBUTES FLastErrorPreservationScope()
+		: OriginalError(GetLastError())
+	{
+	}
+
+	INSTRUMENTATION_FUNCTION_ATTRIBUTES ~FLastErrorPreservationScope()
+	{
+		SetLastError(OriginalError);
+	}
+};
+
+// Publish clock, so that if somebody waits on the thread handle, they can acquire the clock
+// to establish an happen before/after relationship.
+// Must be called from thread about to end.
+INSTRUMENTATION_FUNCTION_ATTRIBUTES void InstrumentEndThread()
+{
+	HANDLE Handle = nullptr;
+	if (GetThreadContext())
+	{
+		FInstrumentedStartThreadArgs* ThreadArgs = (FInstrumentedStartThreadArgs*)GetThreadContext()->ThreadArgs;
+		if (ThreadArgs)
+		{
+			// This can only be nullptr if this thread was so quick to execute
+			// that the parent thread didn't have a chance to broadcast the
+			// thread handle yet.
+			while ((Handle = ThreadArgs->ThreadHandle) == nullptr)
+			{
+				FPlatformProcess::YieldThread();
+			}
+			delete ThreadArgs;
+			GetThreadContext()->ThreadArgs = nullptr;
+		}
+	}
+
+	START_WINAPI_INSTRUMENTATION(ReleaseCurrentContext(); return);
+	if (Handle)
+	{
+		TRefCountPtr<FSyncObject> Sync = GetSyncObject(Handle);
+		Sync->SyncReleaseAsSoleOwner(Context, _ReturnAddress(), Handle, TEXT("InstrumentEndThread"));
+		Context.IncrementClock();
+	}
+	
+	// Clean the stack to avoid false positives when it gets reused for another thread.
+	void* LowLimit; void* HighLimit;
+	GetCurrentThreadStackLimits(&LowLimit, &HighLimit);
+	UE::Sanitizer::RaceDetector::FreeMemoryRange(LowLimit, (UPTRINT)HighLimit - (UPTRINT)LowLimit);
+
+	ReleaseCurrentContext();
+
+	// We released our context so we can't decrement the count now but we can't decrement it
+	// before releasing the context because it could cause reentrancy.
+	//FINISH_WINAPI_INSTRUMENTATION();
+}
+
+INSTRUMENTATION_FUNCTION_ATTRIBUTES DWORD WINAPI InstrumentedStartThread(LPVOID Param)
+{
+	FInstrumentedStartThreadArgs* ThreadArgs = (FInstrumentedStartThreadArgs*)Param;
+
+	FContext& Context = EnsureCurrentContext();
+	Context.ThreadArgs = ThreadArgs;
+	{
+		FInstrumentationScope InstrumentationScope;
+		Context.ClockBank.Acquire(ThreadArgs->ClockBank, _ReturnAddress());
+		// We don't increment the clock here since we don't have any associated
+		// contextid yet. Clock is incremented when contextid is reserved on first
+		// memory access.
+	}
+
+	if (ThreadArgs->bDetailedLog)
+	{
+		FPlatformMisc::LowLevelOutputDebugStringf(TEXT("[%d] Thread starting at clock %d\n"), Context.ThreadId, Context.CurrentClock());
+	}
+
+	// Call the real function.
+	DWORD Result = ThreadArgs->RealStartRoutine(ThreadArgs->RealThreadParameter);
+
+	InstrumentEndThread();
+
+	return Result;
+}
+
+HANDLE (WINAPI* TrueCreateThread)(LPSECURITY_ATTRIBUTES lpThreadAttributes, SIZE_T dwStackSize, LPTHREAD_START_ROUTINE lpStartAddress, LPVOID lpParameter, DWORD dwCreationFlags, LPDWORD lpThreadId);
+INSTRUMENTATION_FUNCTION_ATTRIBUTES HANDLE WINAPI DetouredCreateThread(LPSECURITY_ATTRIBUTES lpThreadAttributes, SIZE_T dwStackSize, LPTHREAD_START_ROUTINE lpStartAddress, LPVOID lpParameter, DWORD dwCreationFlags, LPDWORD lpThreadId)
+{
+	FContext& Context = EnsureCurrentContext();
+	// If ThreadCreationDepth is non-zero here, we're probably coming from a _beginthreadex call 
+	// which was already instrumented. 
+	// Pass along the arguments as-they-are as they're already the instrumented arguments.
+	if (Context.ThreadCreationDepth > 0)
+	{
+		return TrueCreateThread(lpThreadAttributes, dwStackSize, lpStartAddress, lpParameter, dwCreationFlags, lpThreadId);
+	}
+
+	FInstrumentationScope InstrumentationScope;
+	if (Context.DetailedLogDepth || bDetailedLogGlobal)
+	{
+		FPlatformMisc::LowLevelOutputDebugStringf(TEXT("[%d] Creating thread (current clock %d)\n"), Context.ThreadId, Context.CurrentClock());
+	}
+
+	FInstrumentedStartThreadArgs* Args = new FInstrumentedStartThreadArgs;
+	Args->RealThreadParameter = lpParameter;
+	Args->RealStartRoutine = lpStartAddress;
+	Args->ClockBank = Context.ClockBank;
+	Args->bDetailedLog = Context.DetailedLogDepth;
+	Context.IncrementClock();
+
+	HANDLE Handle = TrueCreateThread(lpThreadAttributes, dwStackSize, InstrumentedStartThread, Args, dwCreationFlags, lpThreadId);
+	if (Handle)
+	{
+		Args->ThreadHandle.store(Handle);
+	}
+	return Handle;
+}
+
+UPTRINT (WINAPI* True_beginthreadex)(void* security, unsigned stack_size, unsigned(__stdcall* start_address)(void*), void* arglist, unsigned initflag, unsigned* thrdaddr);
+INSTRUMENTATION_FUNCTION_ATTRIBUTES UPTRINT WINAPI Detoured_beginthreadex(void* security, unsigned stack_size, unsigned(__stdcall* start_address)(void*), void* arglist, unsigned initflag, unsigned* thrdaddr)
+{
+	FContext& Context = EnsureCurrentContext();
+	FInstrumentationScope InstrumentationScope;
+	if (Context.DetailedLogDepth || bDetailedLogGlobal)
+	{
+		FPlatformMisc::LowLevelOutputDebugStringf(TEXT("[%d] Creating thread (current clock %d)\n"), Context.ThreadId, Context.CurrentClock());
+	}
+
+	FInstrumentedStartThreadArgs* Args = new FInstrumentedStartThreadArgs;
+	Args->RealThreadParameter = arglist;
+	Args->RealStartRoutine = (LPTHREAD_START_ROUTINE)start_address;
+	Args->ClockBank = Context.ClockBank;
+	Args->bDetailedLog = Context.DetailedLogDepth;
+	Context.IncrementClock();
+	
+	Context.ThreadCreationDepth++;
+	UPTRINT Handle = True_beginthreadex(security, stack_size, (unsigned(__stdcall*)(void*))InstrumentedStartThread, Args, initflag, thrdaddr);
+	Context.ThreadCreationDepth--;
+	if (Handle)
+	{
+		Args->ThreadHandle.store((HANDLE)Handle);
+	}
+	return Handle;
+}
+
+void (WINAPI* TrueFreeLibraryAndExitThread)(HMODULE hLibModule, DWORD dwExitCode);
+INSTRUMENTATION_FUNCTION_ATTRIBUTES void WINAPI DetouredFreeLibraryAndExitThread(HMODULE hLibModule, DWORD dwExitCode)
+{
+	InstrumentEndThread();
+	TrueFreeLibraryAndExitThread(hLibModule, dwExitCode);
+}
+
+void(__cdecl* True_Cnd_do_broadcast_at_thread_exit)(void);
+INSTRUMENTATION_FUNCTION_ATTRIBUTES void WINAPI Detoured_Cnd_do_broadcast_at_thread_exit()
+{
+	InstrumentEndThread();
+	True_Cnd_do_broadcast_at_thread_exit();
+}
+
+void (WINAPI* TrueExitThread)(DWORD dwExitCode);
+INSTRUMENTATION_FUNCTION_ATTRIBUTES void WINAPI DetouredExitThread(DWORD dwExitCode)
+{
+	InstrumentEndThread();
+	TrueExitThread(dwExitCode);
+}
+
+INSTRUMENTATION_FUNCTION_ATTRIBUTES void AcquireWaitHandle(HANDLE Handle, void* ReturnAddress, const TCHAR* OpName)
+{
+	FContext& Context = EnsureCurrentContext();
+
+	TRefCountPtr<FSyncObject> Sync = GetSyncObject(Handle);
+	Sync->SyncAcquire(Context, []() {}, ReturnAddress, Handle, OpName);
+	Context.IncrementClock();
+}
+
+DWORD(WINAPI* TrueWaitForSingleObject)(HANDLE hHandle, DWORD dwMilliseconds) = WaitForSingleObject;
+DWORD(WINAPI* TrueWaitForSingleObjectEx)(HANDLE hHandle, DWORD dwMilliseconds, BOOL bAlertable) = WaitForSingleObjectEx;
+
+INSTRUMENTATION_FUNCTION_ATTRIBUTES DWORD WINAPI DetouredWaitForSingleObject(HANDLE hHandle, DWORD dwMilliseconds)
+{
+	// Use the Ex version here to avoid recursing into our other detoured function
+	DWORD Result = TrueWaitForSingleObjectEx(hHandle, dwMilliseconds, FALSE);
+
+	START_WINAPI_INSTRUMENTATION(return Result);
+	FLastErrorPreservationScope LastErrorScope;
+	if (Result == WAIT_OBJECT_0)
+	{
+		AcquireWaitHandle(hHandle, _ReturnAddress(), TEXT("WaitForSingleObject"));
+	}
+
+	FINISH_WINAPI_INSTRUMENTATION();
+	return Result;
+}
+
+INSTRUMENTATION_FUNCTION_ATTRIBUTES DWORD WINAPI DetouredWaitForSingleObjectEx(HANDLE hHandle, DWORD dwMilliseconds, BOOL bAlertable)
+{
+	DWORD Result = TrueWaitForSingleObjectEx(hHandle, dwMilliseconds, bAlertable);
+
+	START_WINAPI_INSTRUMENTATION(return Result);
+	FLastErrorPreservationScope LastErrorScope;
+
+	if (Result == WAIT_OBJECT_0)
+	{
+		AcquireWaitHandle(hHandle, _ReturnAddress(), TEXT("WaitForSingleObjectEx"));
+	}
+
+	FINISH_WINAPI_INSTRUMENTATION();
+	return Result;
+}
+
+BOOL(WINAPI* TrueSetEvent)(HANDLE hHandle) = SetEvent;
+INSTRUMENTATION_FUNCTION_ATTRIBUTES BOOL WINAPI DetouredSetEvent(HANDLE hHandle)
+{
+	START_WINAPI_INSTRUMENTATION(return TrueSetEvent(hHandle));
+
+	// The scope is important to make sure the TRefCountPtr is destroyed before we FINISH_WINAPI_INSTRUMENTATION.
+	{
+		TRefCountPtr<FSyncObject> Sync = GetSyncObject(hHandle);
+		Sync->SyncRelease(Context, []() {}, _ReturnAddress(), hHandle, TEXT("SetEvent"));
+		Context.IncrementClock();
+	}
+	BOOL Result = TrueSetEvent(hHandle);
+
+	FINISH_WINAPI_INSTRUMENTATION();
+	return Result;
+}
+
+BOOL(WINAPI* TrueWaitOnAddress)(volatile VOID* Address, PVOID CompareAddress, SIZE_T AddressSize, DWORD dwMilliseconds) = WaitOnAddress;
+INSTRUMENTATION_FUNCTION_ATTRIBUTES BOOL WINAPI DetouredWaitOnAddress(volatile VOID* Address, PVOID CompareAddress, SIZE_T AddressSize, DWORD dwMilliseconds)
+{
+	BOOL Result = TrueWaitOnAddress(Address, CompareAddress, AddressSize, dwMilliseconds);
+
+	START_WINAPI_INSTRUMENTATION(return Result);
+	FLastErrorPreservationScope LastErrorScope;
+
+	if (Result)
+	{
+		// Prevent spurious wake ups from providing a barrier by making sure the value has changed
+		const uint8* Ptr1 = (const uint8*)Address;
+		const uint8* Ptr2 = (const uint8*)CompareAddress;
+
+		bool bIsValueChanged = false;
+		for (SIZE_T Index = 0; Index < AddressSize; ++Index)
+		{
+			if (Ptr1[Index] != Ptr2[Index])
+			{
+				bIsValueChanged = true;
+				break;
+			}
+		}
+
+		if (bIsValueChanged)
+		{
+			TRefCountPtr<FSyncObject> Sync = GetSyncObject((void*)Address);
+			Sync->SyncAcquire(Context, []() {}, _ReturnAddress(), (void*)Address, TEXT("WaitOnAddress"));
+			Context.IncrementClock();
+		}
+	}
+
+	FINISH_WINAPI_INSTRUMENTATION();
+	return Result;
+}
+
+void (WINAPI* TrueWakeByAddressSingle)(PVOID Address) = WakeByAddressSingle;
+INSTRUMENTATION_FUNCTION_ATTRIBUTES void WINAPI DetouredWakeByAddressSingle(PVOID Address)
+{
+	START_WINAPI_INSTRUMENTATION(TrueWakeByAddressSingle(Address); return);
+
+	{
+		FLastErrorPreservationScope LastErrorScope;
+		TRefCountPtr<FSyncObject> Sync = GetSyncObject(Address);
+		Sync->SyncRelease(Context, []() {}, _ReturnAddress(), Address, TEXT("WakeByAddressSingle"));
+		Context.IncrementClock();
+	}
+
+	TrueWakeByAddressSingle(Address);
+
+	FINISH_WINAPI_INSTRUMENTATION();
+}
+
+
+BOOL(WINAPI* TrueTryEnterCriticalSection)(CRITICAL_SECTION* lpCriticalSection) = TryEnterCriticalSection;
+INSTRUMENTATION_FUNCTION_ATTRIBUTES BOOL WINAPI DetouredTryEnterCriticalSection(CRITICAL_SECTION* lpCriticalSection)
+{
+	BOOL Result = TrueTryEnterCriticalSection(lpCriticalSection);
+
+	START_WINAPI_INSTRUMENTATION(return Result);
+	FLastErrorPreservationScope LastErrorScope;
+
+	// If we entered the critical section, we are now the sole owner.
+	if (Result)
+	{
+		TRefCountPtr<FSyncObject> Sync = GetSyncObject(lpCriticalSection);
+		Sync->SyncAcquireAsSoleOwnerOrReadOwner(Context, _ReturnAddress(), lpCriticalSection, TEXT("TryEnterCriticalSection"));
+		Context.IncrementClock();
+	}
+
+	FINISH_WINAPI_INSTRUMENTATION();
+	return Result;
+}
+
+void(WINAPI* TrueEnterCriticalSection)(CRITICAL_SECTION* lpCriticalSection) = EnterCriticalSection;
+INSTRUMENTATION_FUNCTION_ATTRIBUTES void WINAPI DetouredEnterCriticalSection(CRITICAL_SECTION* lpCriticalSection)
+{
+	TrueEnterCriticalSection(lpCriticalSection);
+
+	START_WINAPI_INSTRUMENTATION(return);
+	FLastErrorPreservationScope LastErrorScope;
+
+	// The scope is important to make sure the TRefCountPtr is destroyed before we FINISH_WINAPI_INSTRUMENTATION.
+	{
+		// We are now the sole owner.
+		TRefCountPtr<FSyncObject> Sync = GetSyncObject(lpCriticalSection);
+		Sync->SyncAcquireAsSoleOwnerOrReadOwner(Context, _ReturnAddress(), lpCriticalSection, TEXT("EnterCriticalSection"));
+		Context.IncrementClock();
+	}
+
+	FINISH_WINAPI_INSTRUMENTATION();
+}
+
+void (WINAPI* TrueLeaveCriticalSection)(CRITICAL_SECTION* lpCriticalSection) = LeaveCriticalSection;
+INSTRUMENTATION_FUNCTION_ATTRIBUTES void WINAPI DetouredLeaveCriticalSection(CRITICAL_SECTION* lpCriticalSection)
+{
+	START_WINAPI_INSTRUMENTATION(TrueLeaveCriticalSection(lpCriticalSection); return);
+
+	{
+		FLastErrorPreservationScope LastErrorScope;
+		// Assume this is called with the right semantics, i.e. that we are actual
+		// owners of the critical section.
+		TRefCountPtr<FSyncObject> Sync = GetSyncObject(lpCriticalSection);
+		Sync->SyncReleaseAsSoleOwner(Context, _ReturnAddress(), lpCriticalSection, TEXT("LeaveCriticalSection"));
+		Context.IncrementClock();
+	}
+	TrueLeaveCriticalSection(lpCriticalSection);
+
+	FINISH_WINAPI_INSTRUMENTATION();
+}
+
+// SRW locks.
+void (WINAPI* TrueAcquireSRWLockShared)(PSRWLOCK SRWLock) = AcquireSRWLockShared;
+INSTRUMENTATION_FUNCTION_ATTRIBUTES void WINAPI DetouredAcquireSRWLockShared(PSRWLOCK SRWLock)
+{
+	TrueAcquireSRWLockShared(SRWLock);
+
+	START_WINAPI_INSTRUMENTATION(return);
+	FLastErrorPreservationScope LastErrorScope;
+
+	// The scope is important to make sure the TRefCountPtr is destroyed before we FINISH_WINAPI_INSTRUMENTATION.
+	{
+		TRefCountPtr<FSyncObject> Sync = GetSyncObject(SRWLock);
+		Sync->SyncAcquireAsSoleOwnerOrReadOwner(Context, _ReturnAddress(), SRWLock, TEXT("AcquireSRWLockShared"));
+		Context.IncrementClock();
+	}
+
+	FINISH_WINAPI_INSTRUMENTATION();
+}
+
+BOOLEAN(WINAPI* TrueTryAcquireSRWLockShared)(PSRWLOCK SRWLock) = TryAcquireSRWLockShared;
+INSTRUMENTATION_FUNCTION_ATTRIBUTES BOOLEAN WINAPI DetouredTryAcquireSRWLockShared(PSRWLOCK SRWLock)
+{
+	BOOLEAN Result = TrueTryAcquireSRWLockShared(SRWLock);
+	
+	START_WINAPI_INSTRUMENTATION(return Result);
+
+	FLastErrorPreservationScope LastErrorScope;
+	if (Result)
+	{
+		TRefCountPtr<FSyncObject> Sync = GetSyncObject(SRWLock);
+		Sync->SyncAcquireAsSoleOwnerOrReadOwner(Context, _ReturnAddress(), SRWLock, TEXT("TryAcquireSRWLockShared"));
+		Context.IncrementClock();
+	}
+
+	FINISH_WINAPI_INSTRUMENTATION();
+	return Result;
+}
+
+void (WINAPI* TrueReleaseSRWLockShared)(PSRWLOCK SRWLock) = ReleaseSRWLockShared;
+INSTRUMENTATION_FUNCTION_ATTRIBUTES void WINAPI DetouredReleaseSRWLockShared(PSRWLOCK SRWLock)
+{
+	START_WINAPI_INSTRUMENTATION(TrueReleaseSRWLockShared(SRWLock); return);
+
+	{
+		FLastErrorPreservationScope LastErrorScope;
+		TRefCountPtr<FSyncObject> Sync = GetSyncObject(SRWLock);
+		Sync->SyncRelease(Context, []() {}, _ReturnAddress(), SRWLock, TEXT("ReleaseSRWLockShared"));
+		Context.IncrementClock();
+	}
+
+	TrueReleaseSRWLockShared(SRWLock);
+
+	FINISH_WINAPI_INSTRUMENTATION();
+}
+
+void (WINAPI* TrueAcquireSRWLockExclusive)(PSRWLOCK SRWLock) = AcquireSRWLockExclusive;
+INSTRUMENTATION_FUNCTION_ATTRIBUTES void WINAPI DetouredAcquireSRWLockExclusive(PSRWLOCK SRWLock)
+{
+	TrueAcquireSRWLockExclusive(SRWLock);
+
+	START_WINAPI_INSTRUMENTATION(return);
+	FLastErrorPreservationScope LastErrorScope;
+
+	// The scope is important to make sure the TRefCountPtr is destroyed before we FINISH_WINAPI_INSTRUMENTATION.
+	{
+		TRefCountPtr<FSyncObject> Sync = GetSyncObject(SRWLock);
+		Sync->SyncAcquireAsSoleOwnerOrReadOwner(Context, _ReturnAddress(), SRWLock, TEXT("AcquireSRWLockExclusive"));
+		Context.IncrementClock();
+	}
+
+	FINISH_WINAPI_INSTRUMENTATION();
+}
+
+BOOLEAN(WINAPI* TrueTryAcquireSRWLockExclusive)(PSRWLOCK SRWLock) = TryAcquireSRWLockExclusive;
+INSTRUMENTATION_FUNCTION_ATTRIBUTES BOOLEAN WINAPI DetouredTryAcquireSRWLockExclusive(PSRWLOCK SRWLock)
+{
+	BOOLEAN Result = TrueTryAcquireSRWLockExclusive(SRWLock);
+
+	START_WINAPI_INSTRUMENTATION(return Result);
+	FLastErrorPreservationScope LastErrorScope;
+
+	if (Result)
+	{
+		TRefCountPtr<FSyncObject> Sync = GetSyncObject(SRWLock);
+		Sync->SyncAcquireAsSoleOwnerOrReadOwner(Context, _ReturnAddress(), SRWLock, TEXT("TryAcquireSRWLockExclusive"));
+		Context.IncrementClock();
+	}
+
+	FINISH_WINAPI_INSTRUMENTATION();
+	return Result;
+}
+
+void (WINAPI* TrueReleaseSRWLockExclusive)(PSRWLOCK SRWLock) = ReleaseSRWLockExclusive;
+INSTRUMENTATION_FUNCTION_ATTRIBUTES void WINAPI DetouredReleaseSRWLockExclusive(PSRWLOCK SRWLock)
+{
+	START_WINAPI_INSTRUMENTATION(TrueReleaseSRWLockExclusive(SRWLock); return);
+
+	{
+		FLastErrorPreservationScope LastErrorScope;
+		TRefCountPtr<FSyncObject> Sync = GetSyncObject(SRWLock);
+		Sync->SyncReleaseAsSoleOwner(Context, _ReturnAddress(), SRWLock, TEXT("ReleaseSRWLockExclusive"));
+		Context.IncrementClock();
+	}
+
+	TrueReleaseSRWLockExclusive(SRWLock);
+
+	FINISH_WINAPI_INSTRUMENTATION();
+}
+
+typedef PVOID (STDAPICALLTYPE RTLALLOCATEHEAP)(
+	[in]           PVOID  HeapHandle,
+	[in, optional] ULONG  Flags,
+	[in]           SIZE_T Size
+);
+
+RTLALLOCATEHEAP* TrueRtlAllocateHeap = nullptr;
+INSTRUMENTATION_FUNCTION_ATTRIBUTES LPVOID STDAPICALLTYPE DetouredRtlAllocateHeap(HANDLE hHeap, DWORD dwFlags, SIZE_T dwBytes)
+{
+	START_WINAPI_INSTRUMENTATION(return TrueRtlAllocateHeap(hHeap, dwFlags, dwBytes););
+	LPVOID Result = TrueRtlAllocateHeap(hHeap, dwFlags, dwBytes);
+	FINISH_WINAPI_INSTRUMENTATION();
+	return Result;
+}
+
+typedef NTSTATUS(STDAPICALLTYPE RTLFREEHEAP)
+(
+	IN PVOID HeapHandle,
+	IN ULONG Flags,								/* optional */
+	IN PVOID MemoryPointer
+);
+
+RTLFREEHEAP* TrueRtlFreeHeap = nullptr;
+INSTRUMENTATION_FUNCTION_ATTRIBUTES NTSTATUS STDAPICALLTYPE DetouredRtlFreeHeap(HANDLE hHeap, DWORD dwFlags, LPVOID lpMem)
+{
+	START_WINAPI_INSTRUMENTATION(return TrueRtlFreeHeap(hHeap, dwFlags, lpMem););
+
+	SIZE_T Size = 0;
+	if (lpMem)
+	{
+		Size = HeapSize(hHeap, 0, lpMem);
+	}
+	
+	// Free the memory range before freeing the memory to avoid
+	// another thread racing with the same address range.
+	if (lpMem && Size)
+	{
+		FreeMemoryRange(lpMem, Size);
+	}
+
+	NTSTATUS Result = TrueRtlFreeHeap(hHeap, dwFlags, lpMem);
+
+	FINISH_WINAPI_INSTRUMENTATION();
+
+	return Result;
+}
+
+typedef NTSYSAPI PVOID (NTAPI RTLREALLOCATEHEAP)(
+	IN PVOID                HeapHandle,
+	IN ULONG                Flags,
+	IN PVOID                MemoryPointer,
+	IN ULONG                Size);
+
+RTLREALLOCATEHEAP* TrueRtlReallocateHeap = nullptr;
+INSTRUMENTATION_FUNCTION_ATTRIBUTES PVOID NTAPI DetouredRtlReallocateHeap(PVOID hHeap, ULONG dwFlags, PVOID lpMem, ULONG dwBytes)
+{
+	START_WINAPI_INSTRUMENTATION(return TrueRtlReallocateHeap(hHeap, dwFlags, lpMem, dwBytes));
+
+	// We have to always allocate new blocs in order to invalidate the old memory range before it can get reused.
+	void* NewPtr = TrueRtlAllocateHeap(hHeap, dwFlags, dwBytes);
+
+	if (lpMem)
+	{
+		ULONG OldSize = HeapSize(hHeap, 0, lpMem);
+
+		if (NewPtr)
+		{
+			FMemory::Memcpy(NewPtr, lpMem, FMath::Min(OldSize, dwBytes));
+		}
+
+		FreeMemoryRange(lpMem, OldSize);
+		TrueRtlFreeHeap(hHeap, dwFlags, lpMem);
+	}
+	
+	FINISH_WINAPI_INSTRUMENTATION();
+
+	return NewPtr;
+}
+
+// Functions we need to detour for race detector.
+TArray<TPair<PVOID*, void*>> DetouredFunctions;
+
+INSTRUMENTATION_FUNCTION_ATTRIBUTES void PopulateDetouredFunctions()
+{
+	if (!DetouredFunctions.IsEmpty())
+	{
+		return;
+	}
+
+	TrueWaitForSingleObject = ::WaitForSingleObject;
+	DetouredFunctions.Emplace(&(PVOID&)TrueWaitForSingleObject, DetouredWaitForSingleObject);
+	TrueWaitForSingleObjectEx = ::WaitForSingleObjectEx;
+	DetouredFunctions.Emplace(&(PVOID&)TrueWaitForSingleObjectEx, DetouredWaitForSingleObjectEx);
+	TrueSetEvent = ::SetEvent;
+	DetouredFunctions.Emplace(&(PVOID&)TrueSetEvent, DetouredSetEvent);
+
+	TrueWaitOnAddress = ::WaitOnAddress;
+	DetouredFunctions.Emplace(&(PVOID&)TrueWaitOnAddress, DetouredWaitOnAddress);
+	TrueWakeByAddressSingle = ::WakeByAddressSingle;
+	DetouredFunctions.Emplace(&(PVOID&)TrueWakeByAddressSingle, DetouredWakeByAddressSingle);
+
+	TrueLeaveCriticalSection = ::LeaveCriticalSection;
+	DetouredFunctions.Emplace(&(PVOID&)TrueLeaveCriticalSection, DetouredLeaveCriticalSection);
+	TrueTryEnterCriticalSection = ::TryEnterCriticalSection;
+	DetouredFunctions.Emplace(&(PVOID&)TrueTryEnterCriticalSection, DetouredTryEnterCriticalSection);
+	TrueEnterCriticalSection = ::EnterCriticalSection;
+	DetouredFunctions.Emplace(&(PVOID&)TrueEnterCriticalSection, DetouredEnterCriticalSection);
+
+	TrueAcquireSRWLockShared = ::AcquireSRWLockShared;
+	DetouredFunctions.Emplace(&(PVOID&)TrueAcquireSRWLockShared, DetouredAcquireSRWLockShared);
+	TrueTryAcquireSRWLockShared = ::TryAcquireSRWLockShared;
+	DetouredFunctions.Emplace(&(PVOID&)TrueTryAcquireSRWLockShared, DetouredTryAcquireSRWLockShared);
+	TrueReleaseSRWLockShared = ::ReleaseSRWLockShared;
+	DetouredFunctions.Emplace(&(PVOID&)TrueReleaseSRWLockShared, DetouredReleaseSRWLockShared);
+	TrueAcquireSRWLockExclusive = ::AcquireSRWLockExclusive;
+	DetouredFunctions.Emplace(&(PVOID&)TrueAcquireSRWLockExclusive, DetouredAcquireSRWLockExclusive);
+	TrueTryAcquireSRWLockExclusive = ::TryAcquireSRWLockExclusive;
+	DetouredFunctions.Emplace(&(PVOID&)TrueTryAcquireSRWLockExclusive, DetouredTryAcquireSRWLockExclusive);
+	TrueReleaseSRWLockExclusive = ::ReleaseSRWLockExclusive;
+	DetouredFunctions.Emplace(&(PVOID&)TrueReleaseSRWLockExclusive, DetouredReleaseSRWLockExclusive);
+
+	// Some allocations go directly through the ansi allocator so we need to instrument them
+	// to avoid false positives.
+	TrueRtlFreeHeap = (RTLFREEHEAP*)GetProcAddress(GetModuleHandleA("ntdll"), "RtlFreeHeap");
+	TrueRtlReallocateHeap = (RTLREALLOCATEHEAP*)GetProcAddress(GetModuleHandleA("ntdll"), "RtlReAllocateHeap");
+	TrueRtlAllocateHeap = (RTLALLOCATEHEAP*)GetProcAddress(GetModuleHandleA("ntdll"), "RtlAllocateHeap");
+
+	DetouredFunctions.Emplace(&(PVOID&)TrueRtlFreeHeap, DetouredRtlFreeHeap);
+	DetouredFunctions.Emplace(&(PVOID&)TrueRtlReallocateHeap, DetouredRtlReallocateHeap);
+	DetouredFunctions.Emplace(&(PVOID&)TrueRtlAllocateHeap, DetouredRtlAllocateHeap);
+}
+
+namespace UE::Sanitizer::RaceDetector::Platform {
+
+	UPTRINT ShadowBase = 0;
+	UPTRINT ShadowSize = 0;
+	UPTRINT ShadowEnd = 0;
+	UPTRINT ShadowBitmapBase = 0;
+	UPTRINT ShadowBitmapSize = 0;
+	UPTRINT ShadowBitmapEnd = 0;
+	UPTRINT DirtyShadowBitmapBase = 0;
+	UPTRINT DirtyShadowBitmapSize = 0;
+	UPTRINT DirtyShadowBitmapEnd = 0;
+	UPTRINT PageSize = 0;
+	UPTRINT PageSizeBitShift = 0;
+	UPTRINT PageSizeMask = 0;
+	HANDLE  SanitizerThreadHandle = 0;
+	std::atomic<bool> bHasShadowMemoryMapped = false;
+	volatile bool bSanitizerThreadContinue = true;
+
+	INSTRUMENTATION_FUNCTION_ATTRIBUTES DWORD WINAPI SanitizerThreadProc(LPVOID lpParameter)
+	{
+		SanitizerThreadRun(bSanitizerThreadContinue);
+		return 0;
+	}
+	
+	INSTRUMENTATION_FUNCTION_ATTRIBUTES bool HasShadowMemoryMapped()
+	{
+		return bHasShadowMemoryMapped.load(std::memory_order_relaxed);
+	}
+
+	INSTRUMENTATION_FUNCTION_ATTRIBUTES void SleepMS(uint32 MilliSeconds)
+	{
+		Sleep(MilliSeconds);
+	}
+
+	INSTRUMENTATION_FUNCTION_ATTRIBUTES void PrepareTrampoline(void* ThunkFunctionAddress, void* DestFunctionAddress, bool bUseRETBypass)
+	{
+		UPTRINT ThunkFunctionStart = (UPTRINT)ThunkFunctionAddress;
+
+		uint8* PrefixStart = (uint8*)(ThunkFunctionStart - INSTRUMENTATION_HOTPATCH_PREFIX_NOPS);
+		for (int32 Index = 0; Index < INSTRUMENTATION_HOTPATCH_PREFIX_NOPS; ++Index)
+		{
+			// Verify that we have NOPs to store the trampoline. 
+			if (PrefixStart[Index] != 0x90)
+			{
+
+				UE_LOG(LogRaceDetector, Fatal, TEXT("The function at address %p doesn't have a patchable prefix or was already patched by another system"), ThunkFunctionAddress);
+				return;
+			}
+		}
+
+		uint16* FunctionStart = (uint16*)(ThunkFunctionStart);
+		if (*FunctionStart != 0xC3C3 /* RETs */ && *FunctionStart != 0x9066 /* Two byte NOP */)
+		{
+			UE_LOG(LogRaceDetector, Fatal, TEXT("The function at address %p doesn't have a patchable entry or was already patched by another system"), ThunkFunctionAddress);
+			return;
+		}
+
+		unsigned long OldProtection;
+		if (!VirtualProtect(reinterpret_cast<void*>(PrefixStart), INSTRUMENTATION_HOTPATCH_TOTAL_NOPS, PAGE_EXECUTE_READWRITE, &OldProtection))
+		{
+			UE_LOG(LogRaceDetector, Fatal, TEXT("Unable to change page protection for hotpatching at %p (error %d)"), ThunkFunctionAddress, GetLastError());
+		}
+
+		// Set up the unconditional absolute jump
+		uint8* Trampoline = (uint8*)(ThunkFunctionStart - INSTRUMENTATION_HOTPATCH_PREFIX_NOPS);
+		*(uint16*)(Trampoline + 0)  = 0xB848; /* MOV RAX, imm64  */
+		*(uint64*)(Trampoline + 2)  = (uint64)DestFunctionAddress;
+		*(uint16*)(Trampoline + 10) = 0xE0FF; /* JMP RAX */
+
+		/* Replace NOP by RETs to further improve perf when instructed to do so */
+		if (bUseRETBypass)
+		{
+			*(uint16*)ThunkFunctionStart = 0xC3C3; 
+		}
+
+		// Restore old protection
+		if (!VirtualProtect(reinterpret_cast<void*>(PrefixStart), INSTRUMENTATION_HOTPATCH_TOTAL_NOPS, OldProtection, &OldProtection))
+		{
+			UE_LOG(LogRaceDetector, Fatal, TEXT("Unable to restore page protection for hotpatching at %p (error %d)"), ThunkFunctionAddress, GetLastError());
+		}
+	}
+
+	INSTRUMENTATION_FUNCTION_ATTRIBUTES void CleanupTrampoline(void* ThunkFunctionAddress)
+	{
+		UPTRINT ThunkFunctionStart = (UPTRINT)ThunkFunctionAddress;
+
+		uint8* PrefixStart = (uint8*)(ThunkFunctionStart - INSTRUMENTATION_HOTPATCH_PREFIX_NOPS);
+
+		unsigned long OldProtection;
+		if (!VirtualProtect(PrefixStart, INSTRUMENTATION_HOTPATCH_TOTAL_NOPS, PAGE_EXECUTE_READWRITE, &OldProtection))
+		{
+			UE_LOG(LogRaceDetector, Fatal, TEXT("Unable to change page protection for hotpatching at %p (error %d)"), ThunkFunctionAddress, GetLastError());
+		}
+
+		// Restore all the NOPs that were overwritten
+		memset(PrefixStart, 0x90, INSTRUMENTATION_HOTPATCH_PREFIX_NOPS);
+
+		// Restore old protection
+		if (!VirtualProtect(PrefixStart, INSTRUMENTATION_HOTPATCH_TOTAL_NOPS, OldProtection, &OldProtection))
+		{
+			UE_LOG(LogRaceDetector, Fatal, TEXT("Unable to restore page protection for hotpatching at %p (error %d)"), ThunkFunctionAddress, GetLastError());
+		}
+	}
+
+	INSTRUMENTATION_FUNCTION_ATTRIBUTES void ActivateTrampoline(void* ThunkFunctionAddress)
+	{
+		UPTRINT ThunkFunctionStart = (UPTRINT)ThunkFunctionAddress;
+
+		unsigned long OldProtection;
+		if (!VirtualProtect(reinterpret_cast<void*>(ThunkFunctionStart - INSTRUMENTATION_HOTPATCH_PREFIX_NOPS), INSTRUMENTATION_HOTPATCH_TOTAL_NOPS, PAGE_EXECUTE_READWRITE, &OldProtection))
+		{
+			UE_LOG(LogRaceDetector, Fatal, TEXT("Unable to change page protection for hotpatching at %p (error %d)"), ThunkFunctionAddress, GetLastError());
+		}
+
+		// Set up a relative jump to the beginning of the prefix section
+		// This is actually atomic and won't cause any crashes during the transition.
+		uint16 JmpRelativeOperand = 0xFE - INSTRUMENTATION_HOTPATCH_PREFIX_NOPS;
+		*(uint16*)ThunkFunctionStart = (uint16)(JmpRelativeOperand << 8) | 0xEB; /* JMP -INSTRUMENTATION_HOTPATCH_PREFIX_NOPS */
+
+		// Restore old protection
+		if (!VirtualProtect(reinterpret_cast<void*>(ThunkFunctionStart - INSTRUMENTATION_HOTPATCH_PREFIX_NOPS), INSTRUMENTATION_HOTPATCH_TOTAL_NOPS, OldProtection, &OldProtection))
+		{
+			UE_LOG(LogRaceDetector, Fatal, TEXT("Unable to restore page protection for hotpatching at %p (error %d)"), ThunkFunctionAddress, GetLastError());
+		}
+	}
+
+	INSTRUMENTATION_FUNCTION_ATTRIBUTES void DeactivateTrampoline(void* ThunkFunctionAddress, bool bUseRETBypass)
+	{
+		UPTRINT ThunkFunctionStart = (UPTRINT)ThunkFunctionAddress;
+
+		unsigned long OldProtection;
+		if (!VirtualProtect(reinterpret_cast<void*>(ThunkFunctionStart - INSTRUMENTATION_HOTPATCH_PREFIX_NOPS), INSTRUMENTATION_HOTPATCH_TOTAL_NOPS, PAGE_EXECUTE_READWRITE, &OldProtection))
+		{
+			UE_LOG(LogRaceDetector, Fatal, TEXT("Unable to change page protection for hotpatching at %p (error %d)"), ThunkFunctionAddress, GetLastError());
+		}
+
+		/* Replace NOP by RETs to further improve perf when instructed to do so */
+		if (bUseRETBypass)
+		{
+			*(uint16*)ThunkFunctionStart = 0xC3C3;
+		}
+		else
+		{
+			// This is actually atomic a wont cause any crashes during the transition.
+			*(uint16*)ThunkFunctionStart = 0x9066; /* Replace by 2-byte NOP */
+		}
+
+		// Restore old protection
+		if (!VirtualProtect(reinterpret_cast<void*>(ThunkFunctionStart - INSTRUMENTATION_HOTPATCH_PREFIX_NOPS), INSTRUMENTATION_HOTPATCH_TOTAL_NOPS, OldProtection, &OldProtection))
+		{
+			UE_LOG(LogRaceDetector, Fatal, TEXT("Unable to restore page protection for hotpatching at %p (error %d)"), ThunkFunctionAddress, GetLastError());
+		}
+	}
+	
+	INSTRUMENTATION_FUNCTION_ATTRIBUTES void HideFirstChanceExceptionInVisualStudio()
+	{
+		// Tell MSVC to avoid flooding the debug output with first chance exception
+		// since we're going to use an exception handler to commit shadow memory
+		// on demand.
+		__try
+		{
+			RaiseException(0x0E0736170, 0, 0, 0);
+		}
+		__except (EXCEPTION_EXECUTE_HANDLER)
+			CA_SUPPRESS(6322)
+		{
+		}
+	}
+
+	INSTRUMENTATION_FUNCTION_ATTRIBUTES void FlushInstructionCache()
+	{
+		::FlushInstructionCache(GetCurrentProcess(), 0, 0);
+	}
+
+	INSTRUMENTATION_FUNCTION_ATTRIBUTES uint32 GetCurrentThreadId()
+	{
+		return ::GetCurrentThreadId();
+	}
+
+	INSTRUMENTATION_FUNCTION_ATTRIBUTES void* GetTlsValue(uint32 Index)
+	{
+		return ::TlsGetValue(Index);
+	}
+
+	INSTRUMENTATION_FUNCTION_ATTRIBUTES void SetTlsValue(uint32 Index, void* Value)
+	{
+		::TlsSetValue(Index, Value);
+	}
+
+	INSTRUMENTATION_FUNCTION_ATTRIBUTES uint32 AllocTlsSlot()
+	{
+		return ::TlsAlloc();
+	}
+
+	INSTRUMENTATION_FUNCTION_ATTRIBUTES void FreeTlsSlot(uint32 Index)
+	{
+		::TlsFree(Index);
+	}
+
+	INSTRUMENTATION_FUNCTION_ATTRIBUTES bool IsDebuggerPresent()
+	{
+		return ::IsDebuggerPresent();
+	}
+
+	INSTRUMENTATION_FUNCTION_ATTRIBUTES void InitShadowMemory()
+	{
+		UPTRINT BiggestFreeRegionBase = 0;
+		UPTRINT BiggestFreeRegionSize = 0;
+
+		UPTRINT BaseAddress = 0;
+
+		SYSTEM_INFO si;
+		GetSystemInfo(&si);
+
+		PageSize = si.dwPageSize;
+		PageSizeBitShift = FPlatformMath::CeilLogTwo(PageSize);
+		PageSizeMask = ~((1 << PageSizeBitShift) - 1);
+
+		// Scan for the biggest free address space block we can find
+		MEMORY_BASIC_INFORMATION MemoryBasicInfo;
+		while (VirtualQuery((void*)BaseAddress, &MemoryBasicInfo, sizeof(MemoryBasicInfo)))
+		{
+			if (MemoryBasicInfo.State == MEM_FREE)
+			{
+				if (MemoryBasicInfo.RegionSize > BiggestFreeRegionSize)
+				{
+					BiggestFreeRegionBase = (UPTRINT)MemoryBasicInfo.BaseAddress;
+					BiggestFreeRegionSize = MemoryBasicInfo.RegionSize;
+				}
+			}
+
+			BaseAddress = (UPTRINT)MemoryBasicInfo.BaseAddress + MemoryBasicInfo.RegionSize;
+		}
+
+		// -----------
+		// Low Application Address Space
+		// -----------
+		// Shadow
+		// -----------
+		// High Application Address Space
+		// -----------
+		// We need to split the total address space in a way that we always have
+		// 4 times the shadow address space as the total address space usable by the application.
+		UPTRINT TotalAddressSpace = BaseAddress;
+
+		// Reserve 7/8 of the space for the shadow so that we leave the application with 1/8 of the total address space
+		ShadowBase = Align(BiggestFreeRegionBase, 1024 * 1024 * 1024);
+		ShadowSize = 7 * (TotalAddressSpace >> 3);
+
+		// Because page faults are super slow when attached with a debugger
+		// we're going to use a bitmap view where each 4KB page of the shadow bitmap
+		// is a bit. When that bit is 1, the page has already been committed so we
+		// don't need to call VirtualAlloc for that page.
+		// The required memory for this scheme is a maximum of 4GB support 128TB of address space.
+		ShadowBitmapSize = Align((ShadowSize / PageSize) >> 3, PageSize);
+		ShadowBitmapBase = ShadowBase;
+		ShadowBitmapEnd  = ShadowBitmapBase + ShadowBitmapSize;
+
+		// Because zeroing 4GB is very slow when unmapping shadow.
+		// We keep another 128KB of bits to know which page in the ShadowBitmap have been commited
+		// so we just have to zero those page instead of uncommitting the whole 4GB and recommitting it.
+		DirtyShadowBitmapBase = ShadowBitmapEnd;
+		DirtyShadowBitmapSize = (ShadowBitmapSize / PageSize) >> 3;
+		DirtyShadowBitmapEnd  = DirtyShadowBitmapBase + DirtyShadowBitmapSize;
+
+		// Put the real shadow memory after the bitmaps
+		ShadowBase = DirtyShadowBitmapEnd;
+		ShadowEnd = ShadowBase + ShadowSize;
+
+		if (!VirtualAlloc((LPVOID)ShadowBitmapBase, ShadowSize + ShadowBitmapSize + DirtyShadowBitmapSize, MEM_RESERVE, PAGE_READWRITE))
+		{
+			UE_LOG(LogRaceDetector, Fatal, TEXT("Failed to reserve shadow memory (err: %d)"), GetLastError());
+		}
+
+		if (!VirtualAlloc((LPVOID)ShadowBitmapBase, ShadowBitmapSize + DirtyShadowBitmapSize, MEM_COMMIT, PAGE_READWRITE))
+		{
+			UE_LOG(LogRaceDetector, Fatal, TEXT("Failed to commit shadow memory bitmap (err: %d)"), GetLastError());
+		}
+	}
+
+	INSTRUMENTATION_FUNCTION_ATTRIBUTES void DirtyShadowBitmapPage(UPTRINT ShadowBitmapAddress)
+	{
+		UPTRINT BitmapPage = (ShadowBitmapAddress - ShadowBitmapBase) >> PageSizeBitShift;
+
+		uint8* DirtyBitmapPtr = (uint8*)(DirtyShadowBitmapBase + (BitmapPage >> 3));
+		uint8 PageBit = (uint8)(1 << (BitmapPage & 7));
+		checkSlow((UPTRINT)DirtyBitmapPtr < DirtyShadowBitmapEnd && (UPTRINT)DirtyBitmapPtr >= DirtyShadowBitmapBase);
+		if ((*DirtyBitmapPtr & PageBit) == 0)
+		{
+			std::atomic<uint8>* AtomicDirtyBitmapPtr = (std::atomic<uint8>*)DirtyBitmapPtr;
+			AtomicDirtyBitmapPtr->fetch_or(PageBit, std::memory_order_relaxed);
+		}
+	}
+
+	INSTRUMENTATION_FUNCTION_ATTRIBUTES void ResetShadowBitmap()
+	{
+		// Scan all dirty bits in the 128KB range so that we know exactly which part of the 4GB shadow bitmap 
+		// we need to zero out. This is actually a lot faster than uncommitting and recommitting the entire
+		// range by the OS.
+		constexpr int8 BitCount = sizeof(UPTRINT) << 3;
+		for (UPTRINT DirtyPtr = DirtyShadowBitmapBase; DirtyPtr < DirtyShadowBitmapEnd; DirtyPtr += sizeof(UPTRINT))
+		{
+			UPTRINT DirtyBits = *(UPTRINT*)DirtyPtr;
+			if (DirtyBits)
+			{
+				UPTRINT DirtyPage = ((DirtyPtr - DirtyShadowBitmapBase) << 3) << PageSizeBitShift;
+
+				do
+				{
+					if (DirtyBits & 1)
+					{
+						UPTRINT ShadowBitmapAddr = ShadowBitmapBase + DirtyPage;
+						memset((void*)ShadowBitmapAddr, 0, PageSize);
+					}
+
+					DirtyBits >>= 1;
+					DirtyPage += PageSize;
+				} while (DirtyBits);
+
+				*(UPTRINT*)DirtyPtr = 0;
+			}
+		}
+	}
+
+	INSTRUMENTATION_FUNCTION_ATTRIBUTES UPTRINT GetPageSize()
+	{
+		return PageSize;
+	}
+
+	INSTRUMENTATION_FUNCTION_ATTRIBUTES bool IsShadowMemoryMapped(UPTRINT Base, UPTRINT Size)
+	{
+		UPTRINT AlignedStart = AlignDown(Base, PageSize);
+		UPTRINT AlignedEnd = Align(Base + Size, PageSize);
+
+		UPTRINT ShadowPageStart = ((AlignedStart - ShadowBase) >> PageSizeBitShift);
+		UPTRINT ShadowPageEnd = ((AlignedEnd - ShadowBase) >> PageSizeBitShift);
+
+		for (UPTRINT ShadowPage = ShadowPageStart; ShadowPage < ShadowPageEnd; ++ShadowPage)
+		{
+			uint8* BitmapPtr = (uint8*)(ShadowBitmapBase + (ShadowPage >> 3));
+			uint8 PageBit = (uint8)(1 << (ShadowPage & 7));
+			if ((*BitmapPtr & PageBit))
+			{
+				return true;
+			}
+		}
+
+		return false;
+	}
+
+	INSTRUMENTATION_FUNCTION_ATTRIBUTES const TCHAR* GetCommandLine()
+	{
+		return ::GetCommandLineW();
+	}
+
+	INSTRUMENTATION_FUNCTION_ATTRIBUTES uint16 CaptureStackBackTrace(uint32 FrameToSkip, uint32 FrameToCapture, void** Backtrace)
+	{
+		return ::RtlCaptureStackBackTrace(FrameToSkip, FrameToCapture, Backtrace, 0);
+	}
+
+	// This is how shadow memory typically works, you just commit the page when there's a page fault.
+	// But under the debugger, the kernel will send events to the debugger first and all these
+	// events are taking a lot of time to process...Waaaayyy too much time. 
+	// So we only rely on this exception handler when the debugger is not attached or during shadow resets.
+	INSTRUMENTATION_FUNCTION_ATTRIBUTES static LONG CALLBACK ShadowExceptionHandler(PEXCEPTION_POINTERS ExceptionPointers)
+	{
+		if (ExceptionPointers->ExceptionRecord->ExceptionCode == EXCEPTION_ACCESS_VIOLATION &&
+			ExceptionPointers->ExceptionRecord->NumberParameters >= 2)
+		{
+			const uintptr_t Address = (uintptr_t)ExceptionPointers->ExceptionRecord->ExceptionInformation[1];
+
+			// Only handle exception inside the shadow memory address space.
+			if (Address >= ShadowBase && Address < ShadowBase + ShadowSize)
+			{
+				// If we're resetting the shadow memory, we delay resuming threads until the reset is done.
+				// because we don't want to mess with the bitmap while it's being cleaned up.
+				while (bIsResettingShadow)
+				{
+					Sleep(0);
+				}
+
+				uintptr_t BaseAddress = AlignDown(Address, PageSize);
+				uintptr_t Result = (uintptr_t)::VirtualAlloc((LPVOID)BaseAddress, PageSize, MEM_COMMIT, PAGE_READWRITE);
+				if (Result == BaseAddress)
+				{
+					bHasShadowMemoryMapped.store(true, std::memory_order_relaxed);
+
+					// Mark the page as committed in the bitmap. We use the bitmap when the debugger is
+					// present and to track memory usage of the shadow.
+					UPTRINT ShadowPage = (BaseAddress - ShadowBase) >> PageSizeBitShift;
+
+					uint8* BitmapPtr = (uint8*)(ShadowBitmapBase + (ShadowPage >> 3));
+					uint8 PageBit = (uint8)(1 << (ShadowPage & 7));
+					
+					std::atomic<uint8>* AtomicBitmapPtr = (std::atomic<uint8>*)BitmapPtr;
+					AtomicBitmapPtr->fetch_or(PageBit, std::memory_order_relaxed);
+
+					DirtyShadowBitmapPage((UPTRINT)BitmapPtr);
+
+					return EXCEPTION_CONTINUE_EXECUTION;
+				}
+			}
+		}
+
+		return EXCEPTION_CONTINUE_SEARCH;
+	}
+
+	// Commit shadow memory so that we can read and write to it.
+	INSTRUMENTATION_FUNCTION_ATTRIBUTES void MapShadowMemory(UPTRINT Base, UPTRINT Size)
+	{
+		UPTRINT AlignedStart = AlignDown(Base, PageSize);
+		UPTRINT AlignedEnd = Align(Base + Size, PageSize);
+
+		UPTRINT ShadowPageStart = ((AlignedStart - ShadowBase) >> PageSizeBitShift);
+		UPTRINT ShadowPageEnd = ((AlignedEnd - ShadowBase) >> PageSizeBitShift);
+
+		// Do not apply any modifications on the bitmap before allocating memory
+		// since other threads might also want the same memory.
+		bool bNeedMapping = false;
+		for (UPTRINT ShadowPage = ShadowPageStart; ShadowPage < ShadowPageEnd; ++ShadowPage)
+		{
+			uint8* BitmapPtr = (uint8*)(ShadowBitmapBase + (ShadowPage >> 3));
+			uint8 PageBit = (uint8)(1 << (ShadowPage & 7));
+			if ((*BitmapPtr & PageBit) == 0)
+			{
+				// We found at least one uncommitted region, bail out.
+				bNeedMapping = true;
+				break;
+			}
+		}
+
+		if (bNeedMapping)
+		{
+			VirtualAlloc((void*)AlignedStart, AlignedEnd - AlignedStart, MEM_COMMIT, PAGE_READWRITE);
+
+			bHasShadowMemoryMapped.store(true, std::memory_order_relaxed);
+
+			// Mark all the committed page in the bitmap in a thread-safe manner.
+			// We could unroll this to make it faster but this is not on the critical path anyway.
+			for (UPTRINT ShadowPage = ShadowPageStart; ShadowPage < ShadowPageEnd; ShadowPage++)
+			{
+				uint8* BitmapPtr = (uint8*)(ShadowBitmapBase + (ShadowPage >> 3));
+				uint8 PageBit = (uint8)(1 << (ShadowPage & 7));
+				if ((*BitmapPtr & PageBit) == 0)
+				{
+					std::atomic<uint8>* AtomicBitmapPtr = (std::atomic<uint8>*)BitmapPtr;
+					AtomicBitmapPtr->fetch_or(PageBit, std::memory_order_relaxed);
+
+					DirtyShadowBitmapPage((UPTRINT)BitmapPtr);
+				}
+			}
+		}
+	}
+
+	INSTRUMENTATION_FUNCTION_ATTRIBUTES void UnmapShadowMemory()
+	{
+		TRACE_CPUPROFILER_EVENT_SCOPE(UE::Sanitizer::UnmapShadowMemory);
+
+		// This gets set during mapping and mapping can happen even while we unmap
+		// so we need to reset the flag first.
+		bHasShadowMemoryMapped = false;
+		// Decommit the shadow so it goes back to zeros
+		VirtualFree((LPVOID*)ShadowBase, ShadowSize, MEM_DECOMMIT);
+		ResetShadowBitmap();
+	}
+
+	INSTRUMENTATION_FUNCTION_ATTRIBUTES UPTRINT GetShadowMemoryBase()
+	{
+		return ShadowBase;
+	}
+
+	INSTRUMENTATION_FUNCTION_ATTRIBUTES UPTRINT GetShadowMemorySize()
+	{
+		return ShadowSize;
+	}
+
+	INSTRUMENTATION_FUNCTION_ATTRIBUTES uint64 GetShadowMemoryUsage()
+	{
+		uint64 ShadowPageCount = 0;
+
+		// Scan all dirty bits in the 128KB range so that we know exactly which part of the 4GB shadow bitmap 
+		// we need to access.
+		constexpr int8 BitCount = sizeof(UPTRINT) << 3;
+		for (UPTRINT DirtyPtr = DirtyShadowBitmapBase; DirtyPtr < DirtyShadowBitmapEnd; DirtyPtr += sizeof(UPTRINT))
+		{
+			UPTRINT DirtyBits = *(UPTRINT*)DirtyPtr;
+			if (DirtyBits)
+			{
+				UPTRINT DirtyPage = ((DirtyPtr - DirtyShadowBitmapBase) << 3) << PageSizeBitShift;
+
+				do
+				{
+					if (DirtyBits & 1)
+					{
+						UPTRINT ShadowBitmapAddr = ShadowBitmapBase + DirtyPage;
+
+						// Scan the whole 4KB of bits to know how many bits are used
+						// which tells us how many shadow pages are committed.
+						for (UPTRINT Addr = 0; Addr < PageSize; Addr += 8)
+						{
+							uint64* Bits = (uint64*)(ShadowBitmapAddr + Addr);
+
+							ShadowPageCount += FMath::CountBits(*Bits);
+						}
+					}
+
+					DirtyBits >>= 1;
+					DirtyPage += PageSize;
+				} while (DirtyBits);
+			}
+		}
+
+		return ShadowPageCount << PageSizeBitShift;
+	}
+
+	INSTRUMENTATION_FUNCTION_ATTRIBUTES bool InitializePlatform()
+	{
+		AddVectoredExceptionHandler(TRUE, &ShadowExceptionHandler);
+		
+		HideFirstChanceExceptionInVisualStudio();
+
+		PopulateDetouredFunctions();
+
+		SanitizerThreadHandle = CreateThread(
+			nullptr,
+			0,
+			SanitizerThreadProc,
+			0,
+			0,
+			0
+		);
+
+		DetourTransactionBegin();
+		DetourUpdateThread(GetCurrentThread());
+
+		for (const auto [RealFn, DetouredFn] : DetouredFunctions)
+		{
+			DetourAttach(RealFn, DetouredFn);
+		}
+
+		return DetourTransactionCommit() == NO_ERROR;
+	}
+
+	INSTRUMENTATION_FUNCTION_ATTRIBUTES bool CleanupPlatform()
+	{
+		bSanitizerThreadContinue = false;
+		WaitForSingleObject(SanitizerThreadHandle, INFINITE);
+		CloseHandle(SanitizerThreadHandle);
+
+		DetourTransactionBegin();
+		DetourUpdateThread(GetCurrentThread());
+
+		for (const auto [RealFn, DetouredFn] : DetouredFunctions)
+		{
+			DetourDetach(RealFn, DetouredFn);
+		}
+
+		return DetourTransactionCommit() == NO_ERROR;
+	}
+}
+
+// When building in non monolithic we need to hook ourselves
+// as fast as possible so do it here.
+INSTRUMENTATION_FUNCTION_ATTRIBUTES BOOL WINAPI DllMain(
+	HINSTANCE HinstDLL,
+	DWORD     dwReason,
+	LPVOID    LpvReserved
+)
+{
+	switch (dwReason)
+	{
+	case DLL_PROCESS_ATTACH:
+		UE::Sanitizer::RaceDetector::Initialize();
+		break;
+	case DLL_PROCESS_DETACH:
+		UE::Sanitizer::RaceDetector::Shutdown();
+		break;
+	}
+
+	return TRUE;
+}
+
+INSTRUMENTATION_FUNCTION_ATTRIBUTES int DetourBootFunctions()
+{
+	DetourTransactionBegin();
+	DetourUpdateThread(GetCurrentThread());
+
+	TrueCreateThread = ::CreateThread;
+	DetourAttach(&TrueCreateThread, DetouredCreateThread);
+	True_beginthreadex = ::_beginthreadex;
+	DetourAttach(&True_beginthreadex, Detoured_beginthreadex);
+	TrueExitThread = ::ExitThread;
+	DetourAttach(&TrueExitThread, DetouredExitThread);
+	TrueFreeLibraryAndExitThread = ::FreeLibraryAndExitThread;
+	DetourAttach(&TrueFreeLibraryAndExitThread, DetouredFreeLibraryAndExitThread);
+	True_Cnd_do_broadcast_at_thread_exit = ::_Cnd_do_broadcast_at_thread_exit;
+	DetourAttach(&True_Cnd_do_broadcast_at_thread_exit, Detoured_Cnd_do_broadcast_at_thread_exit);
+
+	LONG Result = DetourTransactionCommit();
+	if (Result != NO_ERROR)
+	{
+		UE_LOG(LogRaceDetector, Fatal, TEXT("Could not install detoured boot functions (error %d)"), Result);
+		return Result;
+	}
+
+	return 0;
+}
+
+extern "C" INSTRUMENTATION_FUNCTION_ATTRIBUTES int PreInit()
+{
+	InitializeBoot();
+	return DetourBootFunctions();
+}
+
+#pragma section(".CRT$XCT",long,read)
+__declspec(allocate(".CRT$XCT")) int (*SanitizerPreStaticInitFn)() = PreInit;
+
+#include "Windows/HideWindowsPlatformTypes.h"
+
+#endif // PLATFORM_WINDOWS && USING_INSTRUMENTATION
diff --git a/Engine/Source/Runtime/Core/Private/Windows/WindowsPlatformMisc.cpp b/Engine/Source/Runtime/Core/Private/Windows/WindowsPlatformMisc.cpp
--- a/Engine/Source/Runtime/Core/Private/Windows/WindowsPlatformMisc.cpp
+++ b/Engine/Source/Runtime/Core/Private/Windows/WindowsPlatformMisc.cpp
@@ -2604,7 +2604,12 @@ int32 FWindowsPlatformMisc::NumberOfWorkerThreadsToSpawn()
 {	
 	static int32 MaxServerWorkerThreads = 4;
 
+// Limit the number of threads as we have a limit of 256 context id in the shadow memory structure.
+#if USING_INSTRUMENTATION
+	int32 MaxWorkerThreads = 32;
+#else
 	int32 MaxWorkerThreads = INT32_MAX;
+#endif
 
 	int32 NumberOfCores = FWindowsPlatformMisc::NumberOfCores();
 	int32 NumberOfCoresIncludingHyperthreads = FWindowsPlatformMisc::NumberOfCoresIncludingHyperthreads();
diff --git a/Engine/Source/Runtime/Core/Private/Windows/WindowsPlatformStackWalk.cpp b/Engine/Source/Runtime/Core/Private/Windows/WindowsPlatformStackWalk.cpp
--- a/Engine/Source/Runtime/Core/Private/Windows/WindowsPlatformStackWalk.cpp
+++ b/Engine/Source/Runtime/Core/Private/Windows/WindowsPlatformStackWalk.cpp
@@ -456,7 +456,13 @@ UE_AUTORTFM_ALWAYS_OPEN uint32 FWindowsPlatformStackWalk::CaptureStackBackTrace(
 	return Depth;
 }
 
-UE_AUTORTFM_ALWAYS_OPEN void FWindowsPlatformStackWalk::ProgramCounterToSymbolInfo( uint64 ProgramCounter, FProgramCounterSymbolInfo& out_SymbolInfo )
+UE_AUTORTFM_ALWAYS_OPEN void FWindowsPlatformStackWalk::ProgramCounterToSymbolInfo( uint64 ProgramCounter, FProgramCounterSymbolInfo& OutSymbolInfo)
+{
+	const bool bIncludeInlineFrames = false;
+	EnumerateSymbolInfosForProgramCounter(ProgramCounter, bIncludeInlineFrames, [&OutSymbolInfo](FProgramCounterSymbolInfo& SymbolInfo) { OutSymbolInfo = SymbolInfo; });
+}
+
+UE_AUTORTFM_ALWAYS_OPEN void FWindowsPlatformStackWalk::EnumerateSymbolInfosForProgramCounter(uint64 ProgramCounter, bool bIncludeInlineFrames, TFunctionRef<void (FProgramCounterSymbolInfo& OutSymbolInfo)> Callback)
 {
 	UE::TUniqueLock GlobalLock(GStackWalkingLock);
 
@@ -469,10 +475,10 @@ UE_AUTORTFM_ALWAYS_OPEN void FWindowsPlatformStackWalk::ProgramCounterToSymbolIn
 	LoadSymbolsForModuleByAddress(ProgramCounter, GetSymbolSearchPath(), bShouldReloadModuleMissingDebugSymbols);
 #endif
 
+	FProgramCounterSymbolInfo SymbolInfo;
 	// Set the program counter.
-	out_SymbolInfo.ProgramCounter = ProgramCounter;
+	SymbolInfo.ProgramCounter = ProgramCounter;
 
-	uint32 LastError = 0;
 	HANDLE ProcessHandle = GProcessHandle;
 
 	// Initialize symbol.
@@ -481,77 +487,50 @@ UE_AUTORTFM_ALWAYS_OPEN void FWindowsPlatformStackWalk::ProgramCounterToSymbolIn
 	Symbol->SizeOfStruct = sizeof(SYMBOL_INFO);
 	Symbol->MaxNameLen = FProgramCounterSymbolInfo::MAX_NAME_LENGTH;
 
-	// Get function name.
-	if( SymFromAddr( ProcessHandle, ProgramCounter, nullptr, Symbol ) )
-	{
-		// Skip any funky chars in the beginning of a function name.
-		int32 Offset = 0;
-		while( Symbol->Name[Offset] < 32 || Symbol->Name[Offset] > 127 )
-		{
-			Offset++;
-		}
-
-		// Write out function name.
-		FCStringAnsi::Strncpy( out_SymbolInfo.FunctionName, Symbol->Name + Offset, FProgramCounterSymbolInfo::MAX_NAME_LENGTH ); 
-		FCStringAnsi::Strncat( out_SymbolInfo.FunctionName, "()", FProgramCounterSymbolInfo::MAX_NAME_LENGTH );
-	}
-	else
-	{
-		// No symbol found for this address.
-		LastError = GetLastError();
-	}
-
-	// Get filename and line number.
-	IMAGEHLP_LINE64	ImageHelpLine = {0};
-	ImageHelpLine.SizeOfStruct = sizeof( ImageHelpLine );
-	if( SymGetLineFromAddr64( ProcessHandle, ProgramCounter, (::DWORD *)&out_SymbolInfo.SymbolDisplacement, &ImageHelpLine ) )
-	{
-		FCStringAnsi::Strncpy( out_SymbolInfo.Filename, ImageHelpLine.FileName, FProgramCounterSymbolInfo::MAX_NAME_LENGTH );
-		out_SymbolInfo.LineNumber = ImageHelpLine.LineNumber;
-	}
-	else
-	{
-		LastError = GetLastError();
-	}
-
 	// Get module name.
-	IMAGEHLP_MODULE64 ImageHelpModule = {0};
-	ImageHelpModule.SizeOfStruct = sizeof( ImageHelpModule );
-	if( SymGetModuleInfo64( ProcessHandle, ProgramCounter, &ImageHelpModule) )
+	IMAGEHLP_MODULE64 ImageHelpModule = { 0 };
+	ImageHelpModule.SizeOfStruct = sizeof(ImageHelpModule);
+	if (SymGetModuleInfo64(ProcessHandle, ProgramCounter, &ImageHelpModule))
 	{
 		// Write out module information.
-		FCStringAnsi::Strncpy( out_SymbolInfo.ModuleName, ImageHelpModule.ImageName, FProgramCounterSymbolInfo::MAX_NAME_LENGTH );
+		FCStringAnsi::Strncpy(SymbolInfo.ModuleName, ImageHelpModule.ImageName, FProgramCounterSymbolInfo::MAX_NAME_LENGTH);
 	}
-	else
+
+	if (bIncludeInlineFrames)
 	{
-		LastError = GetLastError();
+		// Check to see if we need to process inline frames
+		DWORD NumInlineFrames = SymAddrIncludeInlineTrace(ProcessHandle, ProgramCounter);
+		if (NumInlineFrames != 0)
+		{
+			DWORD InlineContext = 0; DWORD FrameIndex = 0;
+			if (SymQueryInlineTrace(ProcessHandle, ProgramCounter, 0, ProgramCounter, ProgramCounter, &InlineContext, &FrameIndex))
+			{
+				for (DWORD Index = 0; Index < NumInlineFrames; ++Index, ++InlineContext)
+				{
+					DWORD64 Displacement = 0;
+					if (SymFromInlineContext(ProcessHandle, ProgramCounter, InlineContext, &Displacement, Symbol))
+					{
+						FCStringAnsi::Strncpy(SymbolInfo.FunctionName, "[Inline Frame] ", FProgramCounterSymbolInfo::MAX_NAME_LENGTH);
+						FCStringAnsi::Strncat(SymbolInfo.FunctionName, Symbol->Name, FProgramCounterSymbolInfo::MAX_NAME_LENGTH);
+						FCStringAnsi::Strncat(SymbolInfo.FunctionName, "()", FProgramCounterSymbolInfo::MAX_NAME_LENGTH);
+
+						SymbolInfo.OffsetInModule = ProgramCounter - Symbol->ModBase;
+
+						IMAGEHLP_LINE64 Line{ 0 };
+						Line.SizeOfStruct = sizeof(IMAGEHLP_LINE64);
+						if (SymGetLineFromInlineContext(ProcessHandle, ProgramCounter, InlineContext, 0, (::DWORD*)&SymbolInfo.SymbolDisplacement, &Line))
+						{
+							SymbolInfo.LineNumber = Line.LineNumber;
+							FCStringAnsi::Strncpy(SymbolInfo.Filename, Line.FileName, FProgramCounterSymbolInfo::MAX_NAME_LENGTH);
+							Callback(SymbolInfo);
+						}
+					}
+				}
+			}
+		}
 	}
-}
-
-UE_AUTORTFM_ALWAYS_OPEN void FWindowsPlatformStackWalk::ProgramCounterToSymbolInfoEx(uint64 ProgramCounter, FProgramCounterSymbolInfoEx& out_SymbolInfo)
-{
-	UE::TUniqueLock GlobalLock(GStackWalkingLock);
-
-#if ON_DEMAND_SYMBOL_LOADING
-	// Load symbols for the module
-	bool bShouldReloadModuleMissingDebugSymbols = !FPlatformProperties::IsMonolithicBuild() && FPlatformStackWalk::WantsDetailedCallstacksInNonMonolithicBuilds();
-	LoadSymbolsForModuleByAddress(ProgramCounter, GetSymbolSearchPath(), bShouldReloadModuleMissingDebugSymbols);
-#endif
-
-	// Set the program counter.
-	out_SymbolInfo.ProgramCounter = ProgramCounter;
-
-	uint32 LastError = 0;
-	HANDLE ProcessHandle = GProcessHandle;
-
-	// Initialize symbol.
-	ANSICHAR SymbolBuffer[sizeof(SYMBOL_INFO) + FProgramCounterSymbolInfo::MAX_NAME_LENGTH] = { 0 };
-	SYMBOL_INFO* Symbol = (SYMBOL_INFO*)SymbolBuffer;
-	Symbol->SizeOfStruct = sizeof(SYMBOL_INFO);
-	Symbol->MaxNameLen = FProgramCounterSymbolInfo::MAX_NAME_LENGTH;
 
 	// Get function name.
-	ANSICHAR FunctionName[FProgramCounterSymbolInfo::MAX_NAME_LENGTH] = { 0 };
 	if (SymFromAddr(ProcessHandle, ProgramCounter, nullptr, Symbol))
 	{
 		// Skip any funky chars in the beginning of a function name.
@@ -560,47 +539,38 @@ UE_AUTORTFM_ALWAYS_OPEN void FWindowsPlatformStackWalk::ProgramCounterToSymbolIn
 		{
 			Offset++;
 		}
+		
+		SymbolInfo.OffsetInModule = ProgramCounter - Symbol->ModBase;
 
 		// Write out function name.
-		FCStringAnsi::Strncpy(FunctionName, Symbol->Name + Offset, FProgramCounterSymbolInfo::MAX_NAME_LENGTH);
-		FCStringAnsi::Strncat(FunctionName, "()", FProgramCounterSymbolInfo::MAX_NAME_LENGTH);
-		out_SymbolInfo.FunctionName = ANSI_TO_TCHAR(FunctionName);
-	}
-	else
-	{
-		// No symbol found for this address.
-		LastError = GetLastError();
+		FCStringAnsi::Strncpy(SymbolInfo.FunctionName, Symbol->Name + Offset, FProgramCounterSymbolInfo::MAX_NAME_LENGTH);
+		FCStringAnsi::Strncat(SymbolInfo.FunctionName, "()", FProgramCounterSymbolInfo::MAX_NAME_LENGTH);
 	}
 
 	// Get filename and line number.
-	ANSICHAR Filename[FProgramCounterSymbolInfo::MAX_NAME_LENGTH] = { 0 };
-	IMAGEHLP_LINE64	ImageHelpLine = { 0 };
+	IMAGEHLP_LINE64	ImageHelpLine = {0};
 	ImageHelpLine.SizeOfStruct = sizeof(ImageHelpLine);
-	if (SymGetLineFromAddr64(ProcessHandle, ProgramCounter, (::DWORD*) & out_SymbolInfo.SymbolDisplacement, &ImageHelpLine))
-	{
-		FCStringAnsi::Strncpy(Filename, ImageHelpLine.FileName, FProgramCounterSymbolInfo::MAX_NAME_LENGTH);
-		out_SymbolInfo.Filename = ANSI_TO_TCHAR(Filename);
-		out_SymbolInfo.LineNumber = ImageHelpLine.LineNumber;
-	}
-	else
+	if (SymGetLineFromAddr64(ProcessHandle, ProgramCounter, (::DWORD*)&SymbolInfo.SymbolDisplacement, &ImageHelpLine))
 	{
-		LastError = GetLastError();
+		FCStringAnsi::Strncpy(SymbolInfo.Filename, ImageHelpLine.FileName, FProgramCounterSymbolInfo::MAX_NAME_LENGTH);
+		SymbolInfo.LineNumber = ImageHelpLine.LineNumber;
 	}
 
-	// Get module name.
-	ANSICHAR ModuleName[FProgramCounterSymbolInfo::MAX_NAME_LENGTH] = { 0 };
-	IMAGEHLP_MODULE64 ImageHelpModule = { 0 };
-	ImageHelpModule.SizeOfStruct = sizeof(ImageHelpModule);
-	if (SymGetModuleInfo64(ProcessHandle, ProgramCounter, &ImageHelpModule))
-	{
-		// Write out module information.
-		FCStringAnsi::Strncpy(ModuleName, ImageHelpModule.ImageName, FProgramCounterSymbolInfo::MAX_NAME_LENGTH);
-		out_SymbolInfo.ModuleName = ANSI_TO_TCHAR(ModuleName);
-	}
-	else
-	{
-		LastError = GetLastError();
-	}
+	Callback(SymbolInfo);
+}
+
+UE_AUTORTFM_ALWAYS_OPEN void FWindowsPlatformStackWalk::ProgramCounterToSymbolInfoEx(uint64 ProgramCounter, FProgramCounterSymbolInfoEx& OutSymbolInfoEx)
+{
+	FProgramCounterSymbolInfo SymbolInfo;
+	ProgramCounterToSymbolInfo(ProgramCounter, SymbolInfo);
+
+	OutSymbolInfoEx.ProgramCounter = ProgramCounter;
+	OutSymbolInfoEx.LineNumber = SymbolInfo.LineNumber;
+	OutSymbolInfoEx.SymbolDisplacement = SymbolInfo.SymbolDisplacement;
+	OutSymbolInfoEx.OffsetInModule = SymbolInfo.OffsetInModule;
+	OutSymbolInfoEx.FunctionName = StringCast<TCHAR, FProgramCounterSymbolInfo::MAX_NAME_LENGTH>(SymbolInfo.FunctionName);
+	OutSymbolInfoEx.ModuleName = StringCast<TCHAR, FProgramCounterSymbolInfo::MAX_NAME_LENGTH>(SymbolInfo.ModuleName);
+	OutSymbolInfoEx.Filename = StringCast<TCHAR, FProgramCounterSymbolInfo::MAX_NAME_LENGTH>(SymbolInfo.Filename);
 }
 
 /**
diff --git a/Engine/Source/Runtime/Core/Public/Clang/ClangPlatformCodeAnalysis.h b/Engine/Source/Runtime/Core/Public/Clang/ClangPlatformCodeAnalysis.h
--- a/Engine/Source/Runtime/Core/Public/Clang/ClangPlatformCodeAnalysis.h
+++ b/Engine/Source/Runtime/Core/Public/Clang/ClangPlatformCodeAnalysis.h
@@ -5,6 +5,9 @@
 
 // HEADER_UNIT_UNSUPPORTED - Clang not supporting header units
 
+// Guard to prevent non-unity incremental checks using MSVC to try to compile this header.
+#if defined(__clang__)
+
 // Code analysis features
 #if defined( __clang_analyzer__ )
 	#define USING_CODE_ANALYSIS 1
@@ -112,11 +115,25 @@
 	#define USING_UNDEFINED_BEHAVIOR_SANITISER 0
 #endif
 
-#if USING_THREAD_SANITISER
+#ifndef USING_INSTRUMENTATION
+	#define USING_INSTRUMENTATION 0
+#endif
+
+#if USING_INSTRUMENTATION
+	// For instrumentation purpose, it is important that we do not skip any memory ops
+	// even if marked as TSAN_SAFE.
+	#define TSAN_SAFE
+#endif
+
+// We do want all the additional annotations for the instrumentation too.
+#if USING_THREAD_SANITISER || USING_INSTRUMENTATION
 
+	// This might already have been defined in the instrumentation block above
+#ifndef TSAN_SAFE
 	// Function attribute to disable thread sanitiser validation on specific functions that assume non-atomic load/stores are implicitly atomic
 	// This is only safe for int32/uint32/int64/uint64/uintptr_t/intptr_t/void* types on x86/x64 strongly-consistent memory systems.
 	#define TSAN_SAFE __attribute__((no_sanitize("thread")))
+#endif
 
 	// Thread-sanitiser annotation functions.
 	#ifdef __cplusplus
@@ -137,9 +154,36 @@
 	// Because annotating the global bools is tiresome...
 	#ifdef __cplusplus
 	    #include <atomic>
-	    #define TSAN_ATOMIC(Type) std::atomic<Type>
+
+		// We need something that is defaulting to relaxed to avoid introducing additional barrier
+		// that cause issues to go unnoticed.
+		template <typename T>
+		class TTSANSafeValue
+		{
+			std::atomic<T> Value;
+
+		public:
+			__forceinline TTSANSafeValue() = default;
+			__forceinline TTSANSafeValue(T InValue) : Value(InValue) { }
+			__forceinline operator T() const { return Value.load(std::memory_order_relaxed); }
+			__forceinline void operator=(T InValue) { Value.store(InValue, std::memory_order_relaxed); }
+			__forceinline T operator++() { return Value.fetch_add(1, std::memory_order_relaxed) + 1; }
+			__forceinline T operator++(int) { return Value.fetch_add(1, std::memory_order_relaxed); }
+			__forceinline T operator--() { return Value.fetch_sub(1, std::memory_order_relaxed) - 1; }
+			__forceinline T operator--(int) { return Value.fetch_sub(1, std::memory_order_relaxed); }
+
+			template <typename OtherType>
+			__forceinline T operator+=(OtherType InValue) { return Value.fetch_add(InValue, std::memory_order_relaxed) + InValue; }
+
+			template <typename OtherType>
+			__forceinline T operator-=(OtherType InValue) { return Value.fetch_sub(InValue, std::memory_order_relaxed) - InValue; }
+		};
+
+	    #define TSAN_ATOMIC(Type) TTSANSafeValue<Type>
 	#endif
 
 #endif
 
+#endif // __clang__
+
 // IWYU pragma: end_exports
diff --git a/Engine/Source/Runtime/Core/Public/Misc/Build.h b/Engine/Source/Runtime/Core/Public/Misc/Build.h
--- a/Engine/Source/Runtime/Core/Public/Misc/Build.h
+++ b/Engine/Source/Runtime/Core/Public/Misc/Build.h
@@ -429,7 +429,7 @@
 #ifndef ALLOW_HANG_DETECTION
 	#define ALLOW_HANG_DETECTION 1
 #endif
-#define USE_HANG_DETECTION (ALLOW_HANG_DETECTION && !WITH_EDITORONLY_DATA && !IS_PROGRAM && !UE_BUILD_DEBUG && !ENABLE_PGO_PROFILE)
+#define USE_HANG_DETECTION (ALLOW_HANG_DETECTION && !WITH_EDITORONLY_DATA && !IS_PROGRAM && !UE_BUILD_DEBUG && !ENABLE_PGO_PROFILE && !USING_THREAD_SANITISER && !USING_INSTRUMENTATION)
 
 // Controls the creation of a thread for detecting hitches (FGameThreadHitchHeartBeat). This is subject to other criteria, USE_HITCH_DETECTION
 #ifndef ALLOW_HITCH_DETECTION
@@ -446,7 +446,7 @@
 	#define USE_PER_MODULE_UOBJECT_BOOTSTRAP 0
 #endif
 
-#define USE_HITCH_DETECTION (ALLOW_HITCH_DETECTION && !WITH_EDITORONLY_DATA && !IS_PROGRAM && !UE_BUILD_DEBUG)
+#define USE_HITCH_DETECTION (ALLOW_HITCH_DETECTION && !WITH_EDITORONLY_DATA && !IS_PROGRAM && !UE_BUILD_DEBUG && !USING_THREAD_SANITISER && !USING_INSTRUMENTATION)
 
 // Controls whether shipping builds create backups of the most recent log file.
 // All other configurations always create backups.
diff --git a/Engine/Source/Runtime/Core/Public/Misc/CoreMiscDefines.h b/Engine/Source/Runtime/Core/Public/Misc/CoreMiscDefines.h
--- a/Engine/Source/Runtime/Core/Public/Misc/CoreMiscDefines.h
+++ b/Engine/Source/Runtime/Core/Public/Misc/CoreMiscDefines.h
@@ -127,11 +127,15 @@
 	#define CA_CONSTANT_IF(Condition) if (Condition)
 #endif
 
+#ifndef USING_INSTRUMENTATION
+	#define USING_INSTRUMENTATION 0
+#endif
+
 #ifndef USING_THREAD_SANITISER
 	#define USING_THREAD_SANITISER 0
 #endif
 
-#if USING_THREAD_SANITISER
+#if USING_THREAD_SANITISER || USING_INSTRUMENTATION
 	#if !defined( TSAN_SAFE ) || !defined( TSAN_BEFORE ) || !defined( TSAN_AFTER ) || !defined( TSAN_ATOMIC )
 		#error Thread Sanitiser macros are not configured correctly for this platform
 	#endif
diff --git a/Engine/Source/Runtime/Core/Public/Sanitizer/RaceDetector.h b/Engine/Source/Runtime/Core/Public/Sanitizer/RaceDetector.h
new file mode 100644
--- /dev/null
+++ b/Engine/Source/Runtime/Core/Public/Sanitizer/RaceDetector.h
@@ -0,0 +1,52 @@
+// Copyright Epic Games, Inc. All Rights Reserved.
+
+#pragma once
+
+#include "Misc/CoreMiscDefines.h"
+
+#if USING_INSTRUMENTATION
+
+#include "Sanitizer/Types.h"
+#include "Misc/Timeout.h"
+
+namespace UE::Sanitizer::RaceDetector
+{
+	// Setup everything that is required for the race detector.
+	CORE_API bool Initialize();
+	// Cleanup everything about the race detector that was setup in Initialize.
+	CORE_API bool Shutdown();
+
+	typedef TFunctionRef<void(uint64 RaceAddress, uint32 FirstThreadId, uint32 SecondThreadId, const FFullLocation& FirstLocation, const FFullLocation& SecondLocation)> TRaceCallbackFn;
+	// Set a callback that will be called whenever a race is detected (filters will not apply).
+	// NOTE: It is not safe to perform any allocation within the callback since it will be
+	//       called from instrumentation code.
+	CORE_API void SetRaceCallbackFn(TRaceCallbackFn CallbackFn);
+	// Remove any callback that was previously set.
+	CORE_API void ResetRaceCallbackFn();
+
+	// Dump the current thread context to help with debugging.
+	CORE_API void DumpContext();
+	// Dump the current thread context but with additional info.
+	CORE_API void DumpContextDetailed();
+
+	// Means that races that doesn't involve the current thread will be filtered.
+	CORE_API void ToggleFilterOtherThreads(bool bEnable);
+
+	// Toggle race detection until a call to the same function that toggles it back off.
+	CORE_API void ToggleRaceDetection(bool bEnable);
+	// Toggle filter on duplicate races. Does not apply when a callback is set.
+	CORE_API void ToggleFilterDuplicateRaces(bool bEnable);
+	// Activate race detection for certain time after which it will be automatically turned off.
+	CORE_API void ToggleRaceDetectionUntil(UE::FTimeout Timeout);
+
+	// Toggle detailed logging globally for all threads. Very slow!
+	CORE_API void ToggleGlobalDetailedLog(bool bEnable);
+	// Toggle detailed logging for the current thread.
+	CORE_API void ToggleThreadDetailedLog(bool bEnable);
+	// Toggle detailed logging filter to only log activity for a specific address.
+	// NOTE: Set to nullptr to remove filtering.
+	CORE_API void ToggleFilterDetailedLogOnAddress(void* Ptr);
+	
+} // UE::Sanitizer::RaceDetector
+
+#endif
\ No newline at end of file
diff --git a/Engine/Source/Runtime/Core/Public/Templates/Atomic.h b/Engine/Source/Runtime/Core/Public/Templates/Atomic.h
--- a/Engine/Source/Runtime/Core/Public/Templates/Atomic.h
+++ b/Engine/Source/Runtime/Core/Public/Templates/Atomic.h
@@ -13,7 +13,7 @@
 // `TAtomic` is planned for deprecation. Please use `std::atomic`
 ////////////////////////////////////////////////////////////////////////////////////////////////////////////////////////////
 
-#define USE_DEPRECATED_TATOMIC 1
+#define USE_DEPRECATED_TATOMIC 0
 
 #if !defined(USE_DEPRECATED_TATOMIC)
 #define USE_DEPRECATED_TATOMIC 0
diff --git a/Engine/Source/Runtime/Core/Public/Windows/WindowsPlatformAtomics.h b/Engine/Source/Runtime/Core/Public/Windows/WindowsPlatformAtomics.h
--- a/Engine/Source/Runtime/Core/Public/Windows/WindowsPlatformAtomics.h
+++ b/Engine/Source/Runtime/Core/Public/Windows/WindowsPlatformAtomics.h
@@ -7,6 +7,19 @@
 #include "Windows/WindowsSystemIncludes.h"
 #include <intrin.h>
 
+#include <atomic>
+
+template <typename T>
+class TStdAtomicSizeAssert {
+	static_assert(sizeof(T) == sizeof(std::atomic<T>), "Unexpected atomic size");
+	static_assert( alignof(T) == alignof(std::atomic<T>), "Unexpected atomic alignment");
+};
+
+inline TStdAtomicSizeAssert<int8> Assert8;
+inline TStdAtomicSizeAssert<int16> Assert16;
+inline TStdAtomicSizeAssert<int32> Assert32;
+inline TStdAtomicSizeAssert<int64> Assert64;
+
 /**
  * Windows implementation of the Atomics OS functions
  */
@@ -283,43 +296,43 @@ struct FWindowsPlatformAtomics
 
 	static FORCEINLINE int8 AtomicRead(volatile const int8* Src)
 	{
-		return InterlockedCompareExchange((int8*)Src, 0, 0);
+		return ((std::atomic<int8>*)(Src))->load();
 	}
 
 	static FORCEINLINE int16 AtomicRead(volatile const int16* Src)
 	{
-		return InterlockedCompareExchange((int16*)Src, 0, 0);
+		return ((std::atomic<int16>*)(Src))->load();
 	}
 
 	static FORCEINLINE int32 AtomicRead(volatile const int32* Src)
 	{
-		return InterlockedCompareExchange((int32*)Src, 0, 0);
+		return ((std::atomic<int32>*)(Src))->load();
 	}
 
 	static FORCEINLINE int64 AtomicRead(volatile const int64* Src)
 	{
-		return InterlockedCompareExchange((int64*)Src, 0, 0);
+		return ((std::atomic<int64>*)(Src))->load();
 	}
 
 	static FORCEINLINE int8 AtomicRead_Relaxed(volatile const int8* Src)
 	{
-		return *Src;
+		return ((std::atomic<int8>*)(Src))->load(std::memory_order_relaxed);
 	}
 
 	static FORCEINLINE int16 AtomicRead_Relaxed(volatile const int16* Src)
 	{
-		return *Src;
+		return ((std::atomic<int16>*)(Src))->load(std::memory_order_relaxed);
 	}
 
 	static FORCEINLINE int32 AtomicRead_Relaxed(volatile const int32* Src)
 	{
-		return *Src;
+		return ((std::atomic<int32>*)(Src))->load(std::memory_order_relaxed);
 	}
 
 	static FORCEINLINE int64 AtomicRead_Relaxed(volatile const int64* Src)
 	{
 		#if PLATFORM_64BITS
-			return *Src;
+			return ((std::atomic<int64>*)(Src))->load(std::memory_order_relaxed);
 		#else
 			return InterlockedCompareExchange((volatile int64*)Src, 0, 0);
 		#endif
@@ -327,43 +340,43 @@ struct FWindowsPlatformAtomics
 
 	static FORCEINLINE void AtomicStore(volatile int8* Src, int8 Val)
 	{
-		InterlockedExchange(Src, Val);
+		((std::atomic<int8>*)(Src))->store(Val);
 	}
 
 	static FORCEINLINE void AtomicStore(volatile int16* Src, int16 Val)
 	{
-		InterlockedExchange(Src, Val);
+		((std::atomic<int16>*)(Src))->store(Val);
 	}
 
 	static FORCEINLINE void AtomicStore(volatile int32* Src, int32 Val)
 	{
-		InterlockedExchange(Src, Val);
+		((std::atomic<int32>*)(Src))->store(Val);
 	}
 
 	static FORCEINLINE void AtomicStore(volatile int64* Src, int64 Val)
 	{
-		InterlockedExchange(Src, Val);
+		((std::atomic<int64>*)(Src))->store(Val);
 	}
 
 	static FORCEINLINE void AtomicStore_Relaxed(volatile int8* Src, int8 Val)
 	{
-		*Src = Val;
+		((std::atomic<int8>*)(Src))->store(Val, std::memory_order_relaxed);
 	}
 
 	static FORCEINLINE void AtomicStore_Relaxed(volatile int16* Src, int16 Val)
 	{
-		*Src = Val;
+		((std::atomic<int16>*)(Src))->store(Val, std::memory_order_relaxed);
 	}
 
 	static FORCEINLINE void AtomicStore_Relaxed(volatile int32* Src, int32 Val)
 	{
-		*Src = Val;
+		((std::atomic<int32>*)(Src))->store(Val, std::memory_order_relaxed);
 	}
 
 	static FORCEINLINE void AtomicStore_Relaxed(volatile int64* Src, int64 Val)
 	{
 		#if PLATFORM_64BITS
-			*Src = Val;
+			((std::atomic<int64>*)(Src))->store(Val, std::memory_order_relaxed);
 		#else
 			InterlockedExchange(Src, Val);
 		#endif
diff --git a/Engine/Source/Runtime/Core/Public/Windows/WindowsPlatformStackWalk.h b/Engine/Source/Runtime/Core/Public/Windows/WindowsPlatformStackWalk.h
--- a/Engine/Source/Runtime/Core/Public/Windows/WindowsPlatformStackWalk.h
+++ b/Engine/Source/Runtime/Core/Public/Windows/WindowsPlatformStackWalk.h
@@ -16,6 +16,8 @@ struct FWindowsPlatformStackWalk
 	
 	static CORE_API TArray<FProgramCounterSymbolInfo> GetStack(int32 IgnoreCount, int32 MaxDepth = 100, void* Context = nullptr);
 
+	static CORE_API void EnumerateSymbolInfosForProgramCounter(uint64 ProgramCounter, bool bIncludeInlineFrames, TFunctionRef<void(FProgramCounterSymbolInfo& OutSymbolInfo)> Callback);
+
 	static CORE_API void ProgramCounterToSymbolInfo( uint64 ProgramCounter, FProgramCounterSymbolInfo& out_SymbolInfo );
 	static CORE_API void ProgramCounterToSymbolInfoEx( uint64 ProgramCounter, FProgramCounterSymbolInfoEx& out_SymbolInfo );
 	CORE_API FORCENOINLINE static uint32 CaptureStackBackTrace( uint64* BackTrace, uint32 MaxDepth, void* Context = nullptr );  // FORCENOINLINE so it can be counted during StackTrace
diff --git a/Engine/Source/Runtime/Core/Tests/Sanitizer/RaceDetectorTests.cpp b/Engine/Source/Runtime/Core/Tests/Sanitizer/RaceDetectorTests.cpp
new file mode 100644
--- /dev/null
+++ b/Engine/Source/Runtime/Core/Tests/Sanitizer/RaceDetectorTests.cpp
@@ -0,0 +1,760 @@
+// Copyright Epic Games, Inc. All Rights Reserved.
+
+#include "CoreTypes.h"
+#include "Misc/AutomationTest.h"
+#include "Tests/Benchmark.h"
+#include "Tasks/Pipe.h"
+#include "HAL/Thread.h"
+#include "Misc/SpinLock.h"
+#include "Misc/ScopeLock.h"
+#include "Async/Mutex.h"
+#include "Async/UniqueLock.h"
+#include "Async/ParallelFor.h"
+#include "Tests/TestHarnessAdapter.h"
+#include "HAL/ThreadHeartBeat.h"
+#include "HAL/PlatformManualResetEvent.h"
+#include "ProfilingDebugging/ScopedTimers.h"
+#include "Containers/UnrealString.h"
+#include "Sanitizer/RaceDetector.h"
+
+#include <atomic>
+#include <thread>
+
+#if WITH_TESTS && USING_INSTRUMENTATION
+
+namespace UE { namespace TasksTests
+{
+	using namespace Tasks;
+	using namespace UE::Sanitizer;
+	using namespace UE::Sanitizer::RaceDetector;
+
+	// A "fake" spinning wait that makes sure a "Wait" doesn't make the caller thread
+	// retract the task. Used to test different scenarios.
+	void TestWait(FTask& Task)
+	{
+		while (!Task.IsCompleted())
+		{
+			FPlatformProcess::Yield();
+		}
+	}
+
+	void TestWait(FGraphEventRef& GraphEvent)
+	{
+		while (!GraphEvent->IsComplete())
+		{
+			FPlatformProcess::Yield();
+		}
+	}
+
+	struct FDataRace {
+		void* Address;
+		FString FirstThreadName;
+		FString SecondThreadName;
+		FFullLocation FirstLocation;
+		FFullLocation SecondLocation;
+
+		bool operator==(const FDataRace& Other) const
+		{
+			return Address == Other.Address &&
+				((FirstThreadName == Other.FirstThreadName && SecondThreadName == Other.SecondThreadName) ||
+				(SecondThreadName == Other.FirstThreadName && FirstThreadName == Other.SecondThreadName));
+		}
+	};
+
+	class FRaceCollectorBase {
+	public:
+	
+		bool IsEmpty()
+		{
+			return Races.IsEmpty();
+		}
+
+		void Reset()
+		{
+			Races.Reset();
+		}
+
+		bool Contains(void* RaceAddress)
+		{
+			for (auto& Race : Races)
+			{
+				if (Race.Address == RaceAddress)
+				{
+					return true;
+				}
+			}
+			return false;
+		}
+
+		bool Contains(void* RaceAddress, const FString& FirstTaskName, const FString& SecondTaskName)
+		{
+			FDataRace Check{ RaceAddress, FirstTaskName, SecondTaskName };
+			for (auto& Race : Races)
+			{
+				if (Race == Check)
+					return true;
+			}
+			return false;
+		}
+
+		uint32 NumRacesForAddress(void* RaceAddress)
+		{
+			uint32 Num = 0;
+			for (auto& Race : Races)
+			{
+				if (Race.Address == RaceAddress)
+				{
+					Num++;
+				}
+			}
+			return Num;
+		}
+
+	protected:
+		UE::FSpinLock RaceLock;
+		TArray<FDataRace> Races;
+	};
+
+	class FThreadRaceCollector : public FRaceCollectorBase {
+	public:
+		FThreadRaceCollector()
+		{
+			SetRaceCallbackFn(*this);
+		}
+
+		~FThreadRaceCollector()
+		{
+			ResetRaceCallbackFn();
+		}
+
+		void operator()(uint64 RaceAddress, uint32 FirstThreadId, uint32 SecondThreadId, const FFullLocation& FirstLocation, const FFullLocation& SecondLocation)
+		{
+			UE::TUniqueLock Lock(RaceLock);
+			Races.Add(FDataRace{ (void*)RaceAddress, FString::Printf(TEXT("%d"), FirstThreadId), FString::Printf(TEXT("%d"), SecondThreadId), FirstLocation, SecondLocation });
+		}
+	};
+
+	void Test();
+	void TestLockFree(int32 OuterIters = 3);
+
+	TEST_CASE_NAMED(FRaceDetectorThreadTest, "System::Core::Sanitizer::RaceDetector::Threads", "[ApplicationContextMask][EngineFilter]")
+	{
+		CHECK(UE::Sanitizer::RaceDetector::Initialize());
+
+		SECTION("Manual test")
+		{
+			//Test();
+		}
+
+		SECTION("Task with prerequisite")
+		{
+			ToggleRaceDetection(false);
+			FThreadRaceCollector Collector;
+
+			for (uint32 i = 0; i < 10000; ++i)
+			{
+				Collector.Reset();
+				ToggleRaceDetection(true);
+
+				int x = 0, y = 0;
+
+				FGraphEventRef LegacyTask = FFunctionGraphTask::CreateAndDispatchWhenReady([&]
+				{
+					x = 1;
+					y = 1;
+				}, TStatId(), nullptr, ENamedThreads::AnyHiPriThreadHiPriTask);
+
+				{
+					FGraphEventArray Prereqs;
+					Prereqs.Add(LegacyTask);
+					LegacyTask = FFunctionGraphTask::CreateAndDispatchWhenReady([&]
+					{
+						x = 2;
+						y = 2;
+					}, TStatId(), &Prereqs, ENamedThreads::AnyHiPriThreadHiPriTask);
+				}
+
+				LegacyTask->Wait();
+
+				ToggleRaceDetection(false);
+
+				CHECK(!Collector.Contains(&x));
+				CHECK(!Collector.Contains(&y));
+			}
+			
+			ToggleRaceDetection(false);
+		}
+
+		SECTION("Standard threads (race)")
+		{
+			ToggleRaceDetection(false);
+
+			FThreadRaceCollector Collector;
+
+			ToggleRaceDetection(true);
+
+			int x = 0, y = 0, z = 0, w = 0;
+			std::atomic<int> sync = 0;
+			std::thread t1{ [&]() {
+				x = 1;
+				y = 1;
+				z = 1;
+			} };
+
+			w = 2;
+
+			std::thread t2{ [&]() {
+				w = 3;
+				z = 3;
+			} };
+
+			x = 2;
+
+			t1.join();
+			t2.join();
+
+			y = 2;
+
+			ToggleRaceDetection(false);
+
+			CHECK(Collector.Contains(&x));
+			CHECK(!Collector.Contains(&y));
+			CHECK(Collector.Contains(&z));
+			CHECK(!Collector.Contains(&w));	
+		}
+
+		SECTION("Standard threads (atomic synchronization)")
+		{
+			ToggleRaceDetection(false);
+			FThreadRaceCollector Collector;
+
+			ToggleRaceDetection(true);
+
+			int x = 0, y = 0, result = 0;
+			std::atomic<int> sync = 0;
+			std::thread t1{ [&]() {
+				x = 10;
+
+				int expected = 0;
+				if (!sync.compare_exchange_strong(expected, 1))
+				{
+					result = y;
+				}
+			}};
+
+			y = 20;
+
+			int expected = 0;
+			if (!sync.compare_exchange_strong(expected, 1))
+			{
+				result = x;
+			}
+
+			t1.join();
+
+			ToggleRaceDetection(false);
+
+			CHECK(!Collector.Contains(&x));
+			CHECK(!Collector.Contains(&y));
+			CHECK(!Collector.Contains(&result));
+		}
+
+		SECTION("Standard threads (event synchronization)")
+		{
+			ToggleRaceDetection(false);
+
+			FThreadRaceCollector Collector;
+
+			ToggleRaceDetection(true);
+
+			FEvent* Event = FPlatformProcess::GetSynchEventFromPool(false);
+			Event->Reset();
+
+			int x = 0, y= 0, result = 0;
+			std::thread t1{ [&]() {
+				x = 10;
+				Event->Trigger();
+				y = 20;
+			} };
+
+			Event->Wait();
+			result = x;
+			result += y;
+			
+			t1.join();
+
+			ToggleRaceDetection(false);
+
+			CHECK(!Collector.Contains(&x));
+			CHECK(Collector.Contains(&y));
+			CHECK(result); // avoid the result and its operations to be optimized out.
+		}
+
+		SECTION("Standard threads (manual reset event synchronization)")
+		{
+			ToggleRaceDetection(false);
+			FThreadRaceCollector Collector;
+
+			ToggleRaceDetection(true);
+
+			UE::FPlatformManualResetEvent Event;
+			Event.Reset();
+
+			int x = 0, y = 0, result = 0;
+			std::thread t1{ [&]() {
+				x = 10;
+				Event.Notify();
+				y = 20;
+			} };
+
+			Event.Wait();
+			result = x;
+			result += y;
+
+			t1.join();
+
+			ToggleRaceDetection(false);
+
+			CHECK(!Collector.Contains(&x));
+			CHECK(Collector.Contains(&y));
+			CHECK(result); // avoid the result and its operations to be optimized out.
+		}
+		
+		SECTION("Modern tasks")
+		{
+			ToggleRaceDetection(false);
+
+			FThreadRaceCollector Collector;
+
+			ToggleRaceDetection(true);
+
+			int x = 0, y = 0;
+			FTask TaskA = UE::Tasks::Launch(TEXT("A"), [&]() {
+				x = 3;
+				y = 3;
+			});
+			x = 5;
+			TestWait(TaskA);
+
+			y = 5;
+
+			ToggleRaceDetection(false);
+			CHECK(Collector.Contains(&x));
+			CHECK(!Collector.Contains(&y));
+		}
+		
+		SECTION("Legacy tasks")
+		{
+			ToggleRaceDetection(false);
+
+			FThreadRaceCollector Collector;
+
+			ToggleRaceDetection(true);
+
+			int x = 0, y = 0;
+			FGraphEventRef LegacyTask = FFunctionGraphTask::CreateAndDispatchWhenReady([&]
+			{
+				x = 3;
+				y = 3;
+			}, TStatId(), nullptr, ENamedThreads::AnyHiPriThreadHiPriTask);
+
+			x = 5;
+
+			TestWait(LegacyTask);
+
+			y = 5;
+
+			ToggleRaceDetection(false);
+
+			CHECK(Collector.Contains(&x));
+			CHECK(!Collector.Contains(&y));
+		}
+
+		SECTION("Two async threads (race)")
+		{
+			ToggleRaceDetection(false);
+
+			FThreadRaceCollector Collector;
+
+			ToggleRaceDetection(true);
+
+			int x = 0;
+			TFuture<void> FutureA = AsyncThread([&x]() {
+				x = 1;
+			});
+
+			x = 2;
+
+			FutureA.Get();
+
+			ToggleRaceDetection(false);
+
+			CHECK(Collector.Contains(&x));
+		}
+
+		SECTION("Async threads (atomic synchronization)")
+		{
+			ToggleRaceDetection(false);
+
+			FThreadRaceCollector Collector;
+
+			int x = 0;
+			std::atomic<bool> sync = false;
+			TFuture<void> FutureA = AsyncThread([&x, &sync]() {
+				ToggleRaceDetection(true);
+				x = 1;
+				sync.store(true);
+				ToggleRaceDetection(false);
+			});
+
+			FutureA.Get();
+
+			TFuture<void> FutureB = AsyncThread([&x, &sync]() {
+				ToggleRaceDetection(true);
+				if (sync.load())
+				{
+					x = 2;
+				}
+				ToggleRaceDetection(false);
+			});
+
+			
+			FutureB.Get();
+
+			CHECK(!Collector.Contains(&x));
+		}
+
+		SECTION("Async threads (atomic access)")
+		{
+			ToggleRaceDetection(false);
+
+			FThreadRaceCollector Collector;
+
+			ToggleRaceDetection(true);
+
+			std::atomic<uint32> x = 0;
+			TFuture<void> FutureA = AsyncThread([&x]() {
+				x = 1;
+			});
+
+			TFuture<void> FutureB = AsyncThread([&x]() {
+				x = 2;
+			});
+
+			FutureA.Get();
+			FutureB.Get();
+
+			ToggleRaceDetection(false);
+
+			CHECK(!Collector.Contains(&x));
+		}
+	}
+
+	struct FTestStruct
+	{
+		int32 Index;
+		int32 Constant;
+		FTestStruct(int32 InIndex)
+			: Index(InIndex)
+			, Constant(0xfe05abcd)
+		{
+		}
+	};
+
+	struct FTestRigFIFO
+	{
+		FLockFreePointerFIFOBase<FTestStruct, PLATFORM_CACHE_LINE_SIZE> Test1;
+		FLockFreePointerFIFOBase<FTestStruct, 8> Test2;
+		FLockFreePointerFIFOBase<FTestStruct, 8, 1 << 4> Test3;
+	};
+
+	struct FTestRigLIFO
+	{
+		FLockFreePointerListLIFOBase<FTestStruct, PLATFORM_CACHE_LINE_SIZE> Test1;
+		FLockFreePointerListLIFOBase<FTestStruct, 8> Test2;
+		FLockFreePointerListLIFOBase<FTestStruct, 8, 1 << 4> Test3;
+	};
+
+
+	void Test()
+	{
+		ToggleRaceDetection(false);
+		FThreadRaceCollector Collector;
+		ToggleRaceDetection(true);
+
+		TestLockFree();
+
+		ToggleRaceDetection(false);
+		Collector.Contains(nullptr);
+	}
+
+	void TestLockFree(int32 OuterIters)
+	{
+		if (!FTaskGraphInterface::IsMultithread())
+		{
+			UE_LOG(LogConsoleResponse, Display, TEXT("WARNING: TestLockFree disabled for non multi-threading platforms"));
+			return;
+		}
+
+		const int32 NumWorkers = FTaskGraphInterface::Get().GetNumWorkerThreads();
+		// If we have too many threads active at once, they become too slow due to contention.  Set a reasonable maximum for how many are required to guarantee correctness of our LockFreePointers.
+		const int32 MaxWorkersForTest = 5;
+		const int32 MinWorkersForTest = 2; // With less than two threads we're not testing threading at all, so the test is pointless.
+		if (NumWorkers < MinWorkersForTest)
+		{
+			UE_LOG(LogConsoleResponse, Display, TEXT("WARNING: TestLockFree disabled for current machine because of not enough worker threads.  Need %d, have %d."), MinWorkersForTest, NumWorkers);
+			return;
+		}
+
+		const uint32 NumWorkersForTest = static_cast<uint32>(FMath::Clamp(NumWorkers, MinWorkersForTest, MaxWorkersForTest));
+		auto RunWorkersSynchronous = [NumWorkersForTest](const TFunction<void(uint32)>& WorkerTask)
+		{
+			FGraphEventArray Tasks;
+			for (uint32 Index = 0; Index < NumWorkersForTest; Index++)
+			{
+				TUniqueFunction<void()> WorkerTaskWithIndex{ [Index, &WorkerTask] { WorkerTask(Index); } };
+				Tasks.Add(FFunctionGraphTask::CreateAndDispatchWhenReady(MoveTemp(WorkerTaskWithIndex), TStatId{}, nullptr, ENamedThreads::AnyNormalThreadHiPriTask));
+			}
+			FTaskGraphInterface::Get().WaitUntilTasksComplete(MoveTemp(Tasks));
+		};
+
+		for (int32 Iter = 0; Iter < OuterIters; Iter++)
+		{
+			{
+				UE_LOG(LogTemp, Display, TEXT("******************************* Iter FIFO %d"), Iter);
+				FTestRigFIFO Rig;
+				for (int32 Index = 0; Index < 1000; Index++)
+				{
+					Rig.Test1.Push(new FTestStruct(Index));
+				}
+				TFunction<void(uint32)> Broadcast =
+					[&Rig](uint32 WorkerIndex)
+				{
+					FRandomStream Stream(((int32)WorkerIndex) * 7 + 13);
+					for (int32 Index = 0; Index < 100000; Index++)
+					{
+						if (Stream.FRand() < .03f)
+						{
+							TArray<FTestStruct*> Items;
+							{
+								float r = Stream.FRand();
+								if (r < .33f)
+								{
+									Rig.Test1.PopAll(Items);
+								}
+								else if (r < .66f)
+								{
+									Rig.Test2.PopAll(Items);
+								}
+								else
+								{
+									Rig.Test3.PopAll(Items);
+								}
+							}
+							for (FTestStruct* Item : Items)
+							{
+								float r = Stream.FRand();
+								if (r < .33f)
+								{
+									Rig.Test1.Push(Item);
+								}
+								else if (r < .66f)
+								{
+									Rig.Test2.Push(Item);
+								}
+								else
+								{
+									Rig.Test3.Push(Item);
+								}
+							}
+						}
+						else
+						{
+							FTestStruct* Item;
+							{
+								float r = Stream.FRand();
+								if (r < .33f)
+								{
+									Item = Rig.Test1.Pop();
+								}
+								else if (r < .66f)
+								{
+									Item = Rig.Test2.Pop();
+								}
+								else
+								{
+									Item = Rig.Test3.Pop();
+								}
+							}
+							if (Item)
+							{
+								float r = Stream.FRand();
+								if (r < .33f)
+								{
+									Rig.Test1.Push(Item);
+								}
+								else if (r < .66f)
+								{
+									Rig.Test2.Push(Item);
+								}
+								else
+								{
+									Rig.Test3.Push(Item);
+								}
+							}
+						}
+					}
+				};
+				RunWorkersSynchronous(Broadcast);
+
+				TArray<FTestStruct*> Items;
+				Rig.Test1.PopAll(Items);
+				Rig.Test2.PopAll(Items);
+				Rig.Test3.PopAll(Items);
+
+				checkf(Items.Num() == 1000, TEXT("Items %d"), Items.Num());
+
+				for (int32 LookFor = 0; LookFor < 1000; LookFor++)
+				{
+					bool bFound = false;
+					for (int32 Index = 0; Index < 1000; Index++)
+					{
+						if (Items[Index]->Index == LookFor && Items[Index]->Constant == 0xfe05abcd)
+						{
+							check(!bFound);
+							bFound = true;
+						}
+					}
+					check(bFound);
+				}
+				for (FTestStruct* Item : Items)
+				{
+					delete Item;
+				}
+
+				UE_LOG(LogTemp, Display, TEXT("******************************* Pass FTestRigFIFO"));
+
+			}
+			{
+				UE_LOG(LogTemp, Display, TEXT("******************************* Iter LIFO %d"), Iter);
+				FTestRigLIFO Rig;
+				for (int32 Index = 0; Index < 1000; Index++)
+				{
+					Rig.Test1.Push(new FTestStruct(Index));
+				}
+				TFunction<void(uint32)> Broadcast =
+					[&Rig](uint32 WorkerIndex)
+				{
+					FRandomStream Stream(((int32)WorkerIndex) * 7 + 13);
+					for (int32 Index = 0; Index < 100000; Index++)
+					{
+						if (Index % 200000 == 1)
+						{
+							//UE_LOG(LogTemp, Log, TEXT("%8d iters thread=%d"), Index, int32(WorkerIndex));
+						}
+						if (Stream.FRand() < .03f)
+						{
+							TArray<FTestStruct*> Items;
+							{
+								float r = Stream.FRand();
+								if (r < .33f)
+								{
+									Rig.Test1.PopAll(Items);
+								}
+								else if (r < .66f)
+								{
+									Rig.Test2.PopAll(Items);
+								}
+								else
+								{
+									Rig.Test3.PopAll(Items);
+								}
+							}
+							for (FTestStruct* Item : Items)
+							{
+								float r = Stream.FRand();
+								if (r < .33f)
+								{
+									Rig.Test1.Push(Item);
+								}
+								else if (r < .66f)
+								{
+									Rig.Test2.Push(Item);
+								}
+								else
+								{
+									Rig.Test3.Push(Item);
+								}
+							}
+						}
+						else
+						{
+							FTestStruct* Item;
+							{
+								float r = Stream.FRand();
+								if (r < .33f)
+								{
+									Item = Rig.Test1.Pop();
+								}
+								else if (r < .66f)
+								{
+									Item = Rig.Test2.Pop();
+								}
+								else
+								{
+									Item = Rig.Test3.Pop();
+								}
+							}
+							if (Item)
+							{
+								float r = Stream.FRand();
+								if (r < .33f)
+								{
+									Rig.Test1.Push(Item);
+								}
+								else if (r < .66f)
+								{
+									Rig.Test2.Push(Item);
+								}
+								else
+								{
+									Rig.Test3.Push(Item);
+								}
+							}
+						}
+					}
+				};
+				RunWorkersSynchronous(Broadcast);
+
+				TArray<FTestStruct*> Items;
+				Rig.Test1.PopAll(Items);
+				Rig.Test2.PopAll(Items);
+				Rig.Test3.PopAll(Items);
+
+				checkf(Items.Num() == 1000, TEXT("Items %d"), Items.Num());
+
+				for (int32 LookFor = 0; LookFor < 1000; LookFor++)
+				{
+					bool bFound = false;
+					for (int32 Index = 0; Index < 1000; Index++)
+					{
+						if (Items[Index]->Index == LookFor && Items[Index]->Constant == 0xfe05abcd)
+						{
+							check(!bFound);
+							bFound = true;
+						}
+					}
+					check(bFound);
+				}
+				for (FTestStruct* Item : Items)
+				{
+					delete Item;
+				}
+
+				UE_LOG(LogTemp, Display, TEXT("******************************* Pass FTestRigLIFO"));
+			}
+		}
+	}
+}}
+
+#endif // WITH_TESTS
diff --git a/Engine/Source/Runtime/Instrumentation/Instrumentation.Build.cs b/Engine/Source/Runtime/Instrumentation/Instrumentation.Build.cs
new file mode 100644
--- /dev/null
+++ b/Engine/Source/Runtime/Instrumentation/Instrumentation.Build.cs
@@ -0,0 +1,18 @@
+// Copyright Epic Games, Inc. All Rights Reserved.
+
+using UnrealBuildTool;
+using System;
+using System.IO;
+
+public class Instrumentation : ModuleRules
+{
+	public Instrumentation(ReadOnlyTargetRules Target) : base(Target)
+	{
+		bRequiresImplementModule = false;
+		PrivateIncludePathModuleNames.Add("Core");
+
+		PrivateDefinitions.Add("SUPPRESS_PER_MODULE_INLINE_FILE"); // This module does not use core's standard operator new/delete overloads
+		
+		bDisableAutoRTFMInstrumentation = true;
+	}
+}
diff --git a/Engine/Source/Runtime/Instrumentation/Internal/Instrumentation/Defines.h b/Engine/Source/Runtime/Instrumentation/Internal/Instrumentation/Defines.h
new file mode 100644
--- /dev/null
+++ b/Engine/Source/Runtime/Instrumentation/Internal/Instrumentation/Defines.h
@@ -0,0 +1,11 @@
+// Copyright Epic Games, Inc. All Rights Reserved.
+
+#pragma once
+
+#define INSTRUMENTATION_FUNCTION_ATTRIBUTES __declspec(safebuffers) __attribute__((disable_sanitizer_instrumentation))
+
+// Reserves 12 NOPs before the function and 2 NOPs at function entry so that we can patch atomically
+#define INSTRUMENTATION_HOTPATCH_TOTAL_NOPS  14
+#define INSTRUMENTATION_HOTPATCH_PREFIX_NOPS 12
+
+#define INSTRUMENTATION_FUNCTION_HOTPATCHABLE __attribute__((patchable_function_entry(INSTRUMENTATION_HOTPATCH_TOTAL_NOPS, INSTRUMENTATION_HOTPATCH_PREFIX_NOPS)))
\ No newline at end of file
diff --git a/Engine/Source/Runtime/Instrumentation/Internal/Instrumentation/EntryPoints.h b/Engine/Source/Runtime/Instrumentation/Internal/Instrumentation/EntryPoints.h
new file mode 100644
--- /dev/null
+++ b/Engine/Source/Runtime/Instrumentation/Internal/Instrumentation/EntryPoints.h
@@ -0,0 +1,84 @@
+// Copyright Epic Games, Inc. All Rights Reserved.
+
+#include "Misc/CoreMiscDefines.h"
+
+#if USING_INSTRUMENTATION
+
+#include "Instrumentation/Types.h"
+
+extern "C"
+{
+	INSTRUMENTATION_API void __Thunk__AnnotateHappensBefore(const char* f, int l, void* addr);
+	INSTRUMENTATION_API void __Thunk__AnnotateHappensAfter(const char* f, int l, void* addr);
+	INSTRUMENTATION_API void __Thunk__Instrument_FuncEntry(void* ReturnAddress);
+	INSTRUMENTATION_API void __Thunk__Instrument_FuncExit();
+	INSTRUMENTATION_API void __Thunk__Instrument_Store(uint64 Address, uint32 Size);
+	INSTRUMENTATION_API void __Thunk__Instrument_Load(uint64 Address, uint32 Size);
+	INSTRUMENTATION_API void __Thunk__Instrument_StoreRange(uint64 Address, uint32 Size);
+	INSTRUMENTATION_API void __Thunk__Instrument_LoadRange(uint64 Address, uint32 Size);
+
+	#define INSTRUMENT_LOAD_FUNC(Type) \
+	INSTRUMENTATION_API Type __Thunk__Instrument_AtomicLoad_##Type(Type* Atomic, UE::Instrumentation::FAtomicMemoryOrder MemoryOrder);
+
+	#define INSTRUMENT_STORE_FUNC(Type) \
+	INSTRUMENTATION_API void __Thunk__Instrument_AtomicStore_##Type(Type* Atomic, Type Val, UE::Instrumentation::FAtomicMemoryOrder MemoryOrder);
+
+	#define INSTRUMENT_EXCHANGE_FUNC(Type) \
+	INSTRUMENTATION_API Type __Thunk__Instrument_AtomicExchange_##Type(Type* Atomic, Type Val, UE::Instrumentation::FAtomicMemoryOrder MemoryOrder);
+
+	#define INSTRUMENT_COMPARE_EXCHANGE_FUNC(Type) \
+	INSTRUMENTATION_API Type __Thunk__Instrument_AtomicCompareExchange_##Type(Type* Atomic, Type* Expected, Type Val, UE::Instrumentation::FAtomicMemoryOrder SuccessMemoryOrder, UE::Instrumentation::FAtomicMemoryOrder FailureMemoryOrder);
+
+	#define INSTRUMENT_RMW_FUNC(Func, Type) \
+	INSTRUMENTATION_API Type __Thunk__Instrument_Atomic##Func##_##Type(Type* Atomic, Type Val, UE::Instrumentation::FAtomicMemoryOrder MemoryOrder);
+
+	// Define the functions with C linkage, referenced by the compiler's instrumentation pass.
+	INSTRUMENT_LOAD_FUNC(int8)
+	INSTRUMENT_LOAD_FUNC(int16)
+	INSTRUMENT_LOAD_FUNC(int32)
+	INSTRUMENT_LOAD_FUNC(int64)
+
+	INSTRUMENT_STORE_FUNC(int8)
+	INSTRUMENT_STORE_FUNC(int16)
+	INSTRUMENT_STORE_FUNC(int32)
+	INSTRUMENT_STORE_FUNC(int64)
+
+	INSTRUMENT_EXCHANGE_FUNC(int8)
+	INSTRUMENT_EXCHANGE_FUNC(int16)
+	INSTRUMENT_EXCHANGE_FUNC(int32)
+	INSTRUMENT_EXCHANGE_FUNC(int64)
+
+	INSTRUMENT_COMPARE_EXCHANGE_FUNC(int8)
+	INSTRUMENT_COMPARE_EXCHANGE_FUNC(int16)
+	INSTRUMENT_COMPARE_EXCHANGE_FUNC(int32)
+	INSTRUMENT_COMPARE_EXCHANGE_FUNC(int64)
+
+	INSTRUMENT_RMW_FUNC(FetchAdd, int8)
+	INSTRUMENT_RMW_FUNC(FetchAdd, int16)
+	INSTRUMENT_RMW_FUNC(FetchAdd, int32)
+	INSTRUMENT_RMW_FUNC(FetchAdd, int64)
+	INSTRUMENT_RMW_FUNC(FetchSub, int8)
+	INSTRUMENT_RMW_FUNC(FetchSub, int16)
+	INSTRUMENT_RMW_FUNC(FetchSub, int32)
+	INSTRUMENT_RMW_FUNC(FetchSub, int64)
+	INSTRUMENT_RMW_FUNC(FetchOr, int8)
+	INSTRUMENT_RMW_FUNC(FetchOr, int16)
+	INSTRUMENT_RMW_FUNC(FetchOr, int32)
+	INSTRUMENT_RMW_FUNC(FetchOr, int64)
+	INSTRUMENT_RMW_FUNC(FetchXor, int8)
+	INSTRUMENT_RMW_FUNC(FetchXor, int16)
+	INSTRUMENT_RMW_FUNC(FetchXor, int32)
+	INSTRUMENT_RMW_FUNC(FetchXor, int64)
+	INSTRUMENT_RMW_FUNC(FetchAnd, int8)
+	INSTRUMENT_RMW_FUNC(FetchAnd, int16)
+	INSTRUMENT_RMW_FUNC(FetchAnd, int32)
+	INSTRUMENT_RMW_FUNC(FetchAnd, int64)
+
+	#undef INSTRUMENT_LOAD_FUNC
+	#undef INSTRUMENT_STORE_FUNC
+	#undef INSTRUMENT_EXCHANGE_FUNC
+	#undef INSTRUMENT_COMPARE_EXCHANGE_FUNC
+	#undef INSTRUMENT_RMW_FUNC
+}
+
+#endif // USING_INSTRUMENTATION
\ No newline at end of file
diff --git a/Engine/Source/Runtime/Instrumentation/Internal/Instrumentation/PerModuleInline.inl b/Engine/Source/Runtime/Instrumentation/Internal/Instrumentation/PerModuleInline.inl
new file mode 100644
--- /dev/null
+++ b/Engine/Source/Runtime/Instrumentation/Internal/Instrumentation/PerModuleInline.inl
@@ -0,0 +1,168 @@
+// Copyright Epic Games, Inc. All Rights Reserved.
+
+#include "Misc/CoreMiscDefines.h"
+
+#if USING_INSTRUMENTATION
+
+#include "Instrumentation/Types.h"
+#include "Instrumentation/Defines.h"
+#include "Instrumentation/EntryPoints.h"
+
+#if defined(USE_INSTRUMENTATION_PERMODULE_HOTPATCH) && USE_INSTRUMENTATION_PERMODULE_HOTPATCH
+	#define INSTRUMENTATION_PERMODULE_FUNCTION_ATTRIBUTES INSTRUMENTATION_FUNCTION_ATTRIBUTES INSTRUMENTATION_FUNCTION_HOTPATCHABLE	
+#else
+	#define INSTRUMENTATION_PERMODULE_FUNCTION_ATTRIBUTES INSTRUMENTATION_FUNCTION_ATTRIBUTES	
+#endif
+
+// Define the functions with C linkage, referenced by the compiler's instrumentation pass.
+// We redirect everything to the instrumentation layer thunks so that every module instrumentation ends up at the same place.
+extern "C" {
+
+	INSTRUMENTATION_PERMODULE_FUNCTION_ATTRIBUTES void AnnotateHappensBefore(const char* f, int l, void* addr)
+	{
+		__Thunk__AnnotateHappensBefore(f, l, addr);
+	}
+
+	INSTRUMENTATION_PERMODULE_FUNCTION_ATTRIBUTES void AnnotateHappensAfter(const char* f, int l, void* addr)
+	{
+		__Thunk__AnnotateHappensAfter(f, l, addr);
+	}
+
+	INSTRUMENTATION_PERMODULE_FUNCTION_ATTRIBUTES void __Instrument_FuncEntry(void* ReturnAddress)
+	{
+		__Thunk__Instrument_FuncEntry(ReturnAddress);
+	}
+
+	INSTRUMENTATION_PERMODULE_FUNCTION_ATTRIBUTES void __Instrument_FuncExit()
+	{
+		__Thunk__Instrument_FuncExit();
+	}
+
+	INSTRUMENTATION_PERMODULE_FUNCTION_ATTRIBUTES void __Instrument_Store(uint64 Address, uint32 Size)
+	{
+		__Thunk__Instrument_Store(Address, Size);
+	}
+
+	INSTRUMENTATION_PERMODULE_FUNCTION_ATTRIBUTES void __Instrument_Load(uint64 Address, uint32 Size)
+	{
+		__Thunk__Instrument_Load(Address, Size);
+	}
+
+	INSTRUMENTATION_PERMODULE_FUNCTION_ATTRIBUTES void __Instrument_StoreRange(uint64 Address, uint32 Size)
+	{
+		__Thunk__Instrument_StoreRange(Address, Size);
+	}
+
+	INSTRUMENTATION_PERMODULE_FUNCTION_ATTRIBUTES void __Instrument_LoadRange(uint64 Address, uint32 Size)
+	{
+		__Thunk__Instrument_LoadRange(Address, Size);
+	}
+
+#define INSTRUMENT_LOAD_FUNC_IMPL(Type) \
+INSTRUMENTATION_PERMODULE_FUNCTION_ATTRIBUTES Type __Instrument_AtomicLoad_##Type(Type* Atomic, UE::Instrumentation::FAtomicMemoryOrder MemoryOrder) \
+{ \
+	return __Thunk__Instrument_AtomicLoad_##Type(Atomic, MemoryOrder); \
+}
+
+#define INSTRUMENT_STORE_FUNC_IMPL(Type) \
+INSTRUMENTATION_PERMODULE_FUNCTION_ATTRIBUTES void __Instrument_AtomicStore_##Type(Type* Atomic, Type Val, UE::Instrumentation::FAtomicMemoryOrder MemoryOrder) \
+{ \
+	__Thunk__Instrument_AtomicStore_##Type(Atomic, Val, MemoryOrder); \
+}
+
+#define INSTRUMENT_EXCHANGE_FUNC_IMPL(Type) \
+INSTRUMENTATION_PERMODULE_FUNCTION_ATTRIBUTES Type __Instrument_AtomicExchange_##Type(Type* Atomic, Type Val, UE::Instrumentation::FAtomicMemoryOrder MemoryOrder) \
+{ \
+	return __Thunk__Instrument_AtomicExchange_##Type(Atomic, Val, MemoryOrder); \
+}
+
+#define INSTRUMENT_COMPARE_EXCHANGE_FUNC_IMPL(Type) \
+INSTRUMENTATION_PERMODULE_FUNCTION_ATTRIBUTES Type __Instrument_AtomicCompareExchange_##Type(Type* Atomic, Type* Expected, Type Val, UE::Instrumentation::FAtomicMemoryOrder SuccessMemoryOrder, UE::Instrumentation::FAtomicMemoryOrder FailureMemoryOrder) \
+{ \
+	return __Thunk__Instrument_AtomicCompareExchange_##Type(Atomic, Expected, Val, SuccessMemoryOrder, FailureMemoryOrder); \
+} 
+
+#define INSTRUMENT_FETCHADD_FUNC_IMPL(Type) \
+INSTRUMENTATION_PERMODULE_FUNCTION_ATTRIBUTES Type __Instrument_AtomicFetchAdd_##Type(Type* Atomic, Type Val, UE::Instrumentation::FAtomicMemoryOrder MemoryOrder) \
+{ \
+	return __Thunk__Instrument_AtomicFetchAdd_##Type(Atomic, Val, MemoryOrder); \
+} 
+
+#define INSTRUMENT_FETCHSUB_FUNC_IMPL(Type) \
+INSTRUMENTATION_PERMODULE_FUNCTION_ATTRIBUTES Type __Instrument_AtomicFetchSub_##Type(Type* Atomic, Type Val, UE::Instrumentation::FAtomicMemoryOrder MemoryOrder) \
+{ \
+	return __Thunk__Instrument_AtomicFetchSub_##Type(Atomic, Val, MemoryOrder); \
+} 
+
+#define INSTRUMENT_FETCHOR_FUNC_IMPL(Type) \
+INSTRUMENTATION_PERMODULE_FUNCTION_ATTRIBUTES Type __Instrument_AtomicFetchOr_##Type(Type* Atomic, Type Val, UE::Instrumentation::FAtomicMemoryOrder MemoryOrder) \
+{ \
+	return __Thunk__Instrument_AtomicFetchOr_##Type(Atomic, Val, MemoryOrder); \
+} 
+
+#define INSTRUMENT_FETCHXOR_FUNC_IMPL(Type) \
+INSTRUMENTATION_PERMODULE_FUNCTION_ATTRIBUTES Type __Instrument_AtomicFetchXor_##Type(Type* Atomic, Type Val, UE::Instrumentation::FAtomicMemoryOrder MemoryOrder) \
+{ \
+	return __Thunk__Instrument_AtomicFetchXor_##Type(Atomic, Val, MemoryOrder); \
+} 
+
+#define INSTRUMENT_FETCHAND_FUNC_IMPL(Type) \
+INSTRUMENTATION_PERMODULE_FUNCTION_ATTRIBUTES Type __Instrument_AtomicFetchAnd_##Type(Type* Atomic, Type Val, UE::Instrumentation::FAtomicMemoryOrder MemoryOrder) \
+{ \
+	return __Thunk__Instrument_AtomicFetchAnd_##Type(Atomic, Val, MemoryOrder); \
+} 
+
+INSTRUMENT_LOAD_FUNC_IMPL(int8)
+INSTRUMENT_LOAD_FUNC_IMPL(int16)
+INSTRUMENT_LOAD_FUNC_IMPL(int32)
+INSTRUMENT_LOAD_FUNC_IMPL(int64)
+
+INSTRUMENT_STORE_FUNC_IMPL(int8)
+INSTRUMENT_STORE_FUNC_IMPL(int16)
+INSTRUMENT_STORE_FUNC_IMPL(int32)
+INSTRUMENT_STORE_FUNC_IMPL(int64)
+
+INSTRUMENT_EXCHANGE_FUNC_IMPL(int8)
+INSTRUMENT_EXCHANGE_FUNC_IMPL(int16)
+INSTRUMENT_EXCHANGE_FUNC_IMPL(int32)
+INSTRUMENT_EXCHANGE_FUNC_IMPL(int64)
+
+INSTRUMENT_COMPARE_EXCHANGE_FUNC_IMPL(int8)
+INSTRUMENT_COMPARE_EXCHANGE_FUNC_IMPL(int16)
+INSTRUMENT_COMPARE_EXCHANGE_FUNC_IMPL(int32)
+INSTRUMENT_COMPARE_EXCHANGE_FUNC_IMPL(int64)
+
+INSTRUMENT_FETCHADD_FUNC_IMPL(int8)
+INSTRUMENT_FETCHADD_FUNC_IMPL(int16)
+INSTRUMENT_FETCHADD_FUNC_IMPL(int32)
+INSTRUMENT_FETCHADD_FUNC_IMPL(int64)
+INSTRUMENT_FETCHSUB_FUNC_IMPL(int8)
+INSTRUMENT_FETCHSUB_FUNC_IMPL(int16)
+INSTRUMENT_FETCHSUB_FUNC_IMPL(int32)
+INSTRUMENT_FETCHSUB_FUNC_IMPL(int64)
+INSTRUMENT_FETCHOR_FUNC_IMPL(int8)
+INSTRUMENT_FETCHOR_FUNC_IMPL(int16)
+INSTRUMENT_FETCHOR_FUNC_IMPL(int32)
+INSTRUMENT_FETCHOR_FUNC_IMPL(int64)
+INSTRUMENT_FETCHXOR_FUNC_IMPL(int8)
+INSTRUMENT_FETCHXOR_FUNC_IMPL(int16)
+INSTRUMENT_FETCHXOR_FUNC_IMPL(int32)
+INSTRUMENT_FETCHXOR_FUNC_IMPL(int64)
+INSTRUMENT_FETCHAND_FUNC_IMPL(int8)
+INSTRUMENT_FETCHAND_FUNC_IMPL(int16)
+INSTRUMENT_FETCHAND_FUNC_IMPL(int32)
+INSTRUMENT_FETCHAND_FUNC_IMPL(int64)
+
+#undef INSTRUMENT_LOAD_FUNC_IMPL
+#undef INSTRUMENT_STORE_FUNC_IMPL
+#undef INSTRUMENT_EXCHANGE_FUNC_IMPL
+#undef INSTRUMENT_COMPARE_EXCHANGE_FUNC_IMPL
+#undef INSTRUMENT_FETCHADD_FUNC_IMPL
+#undef INSTRUMENT_FETCHSUB_FUNC_IMPL
+#undef INSTRUMENT_FETCHOR_FUNC_IMPL
+#undef INSTRUMENT_FETCHXOR_FUNC_IMPL
+#undef INSTRUMENT_FETCHAND_FUNC_IMPL
+
+}
+
+#endif // USING_INSTRUMENTATION
\ No newline at end of file
diff --git a/Engine/Source/Runtime/Instrumentation/Private/Instrumentation/EntryPoints.cpp b/Engine/Source/Runtime/Instrumentation/Private/Instrumentation/EntryPoints.cpp
new file mode 100644
--- /dev/null
+++ b/Engine/Source/Runtime/Instrumentation/Private/Instrumentation/EntryPoints.cpp
@@ -0,0 +1,150 @@
+// Copyright Epic Games, Inc. All Rights Reserved.
+
+#include "Instrumentation/EntryPoints.h"
+
+#if USING_INSTRUMENTATION
+
+#include "Instrumentation/Defines.h"
+#include "Templates/Atomic.h"
+#include "CoreTypes.h"
+
+#define INSTRUMENTATION_THUNK_ATTRIBUTES INSTRUMENTATION_API INSTRUMENTATION_FUNCTION_ATTRIBUTES INSTRUMENTATION_FUNCTION_HOTPATCHABLE
+
+// Define the functions with C linkage, referenced by the compiler's instrumentation pass.
+extern "C" {
+
+INSTRUMENTATION_THUNK_ATTRIBUTES void __Thunk__AnnotateHappensBefore(const char* f, int l, void* addr)
+{
+}
+
+INSTRUMENTATION_THUNK_ATTRIBUTES void __Thunk__AnnotateHappensAfter(const char* f, int l, void* addr)
+{
+}
+
+INSTRUMENTATION_THUNK_ATTRIBUTES void __Thunk__Instrument_FuncEntry(void* ReturnAddress)
+{
+}
+
+INSTRUMENTATION_THUNK_ATTRIBUTES void __Thunk__Instrument_FuncExit()
+{
+}
+
+INSTRUMENTATION_THUNK_ATTRIBUTES void __Thunk__Instrument_Store(uint64 Address, uint32 Size)
+{
+}
+
+INSTRUMENTATION_THUNK_ATTRIBUTES void __Thunk__Instrument_Load(uint64 Address, uint32 Size)
+{
+}
+
+INSTRUMENTATION_THUNK_ATTRIBUTES void __Thunk__Instrument_StoreRange(uint64 Address, uint32 Size)
+{
+}
+
+INSTRUMENTATION_THUNK_ATTRIBUTES void __Thunk__Instrument_LoadRange(uint64 Address, uint32 Size)
+{
+}
+
+#define INSTRUMENT_LOAD_FUNC_THUNK(Type) \
+INSTRUMENTATION_THUNK_ATTRIBUTES Type __Thunk__Instrument_AtomicLoad_##Type(Type* Atomic, UE::Instrumentation::FAtomicMemoryOrder MemoryOrder) \
+{ \
+	return UE::Core::Private::Atomic::Load(Atomic); \
+} 
+
+#define INSTRUMENT_STORE_FUNC_THUNK(Type) \
+INSTRUMENTATION_THUNK_ATTRIBUTES void __Thunk__Instrument_AtomicStore_##Type(Type* Atomic, Type Val, UE::Instrumentation::FAtomicMemoryOrder MemoryOrder) \
+{ \
+	UE::Core::Private::Atomic::Store(Atomic, Val); \
+} 
+
+#define INSTRUMENT_EXCHANGE_FUNC_THUNK(Type) \
+INSTRUMENTATION_THUNK_ATTRIBUTES Type __Thunk__Instrument_AtomicExchange_##Type(Type* Atomic, Type Val, UE::Instrumentation::FAtomicMemoryOrder MemoryOrder) \
+{ \
+	return UE::Core::Private::Atomic::Exchange(Atomic, Val); \
+} 
+
+#define INSTRUMENT_COMPARE_EXCHANGE_FUNC_THUNK(Type) \
+INSTRUMENTATION_THUNK_ATTRIBUTES Type __Thunk__Instrument_AtomicCompareExchange_##Type(Type* Atomic, Type* Expected, Type Val, UE::Instrumentation::FAtomicMemoryOrder SuccessMemoryOrder, UE::Instrumentation::FAtomicMemoryOrder FailureMemoryOrder) \
+{ \
+	Type Ret = UE::Core::Private::Atomic::CompareExchange(Atomic, *Expected, Val); \
+	if (Ret != *Expected) \
+	{ \
+		*Expected = Ret; \
+	} \
+	return Ret; \
+} 
+
+#define INSTRUMENT_FETCHADD_FUNC_THUNK(Type) \
+INSTRUMENTATION_THUNK_ATTRIBUTES Type __Thunk__Instrument_AtomicFetchAdd_##Type(Type* Atomic, Type Val, UE::Instrumentation::FAtomicMemoryOrder MemoryOrder) \
+{ \
+	return UE::Core::Private::Atomic::AddExchange(Atomic, Val); \
+} 
+
+#define INSTRUMENT_FETCHSUB_FUNC_THUNK(Type) \
+INSTRUMENTATION_THUNK_ATTRIBUTES Type __Thunk__Instrument_AtomicFetchSub_##Type(Type* Atomic, Type Val, UE::Instrumentation::FAtomicMemoryOrder MemoryOrder) \
+{ \
+	return UE::Core::Private::Atomic::SubExchange(Atomic, Val); \
+} 
+
+#define INSTRUMENT_FETCHOR_FUNC_THUNK(Type) \
+INSTRUMENTATION_THUNK_ATTRIBUTES Type __Thunk__Instrument_AtomicFetchOr_##Type(Type* Atomic, Type Val, UE::Instrumentation::FAtomicMemoryOrder MemoryOrder) \
+{ \
+	return UE::Core::Private::Atomic::OrExchange(Atomic, Val); \
+} 
+
+#define INSTRUMENT_FETCHXOR_FUNC_THUNK(Type) \
+INSTRUMENTATION_THUNK_ATTRIBUTES Type __Thunk__Instrument_AtomicFetchXor_##Type(Type* Atomic, Type Val, UE::Instrumentation::FAtomicMemoryOrder MemoryOrder) \
+{ \
+	return UE::Core::Private::Atomic::XorExchange(Atomic, Val); \
+} 
+
+#define INSTRUMENT_FETCHAND_FUNC_THUNK(Type) \
+INSTRUMENTATION_THUNK_ATTRIBUTES Type __Thunk__Instrument_AtomicFetchAnd_##Type(Type* Atomic, Type Val, UE::Instrumentation::FAtomicMemoryOrder MemoryOrder) \
+{ \
+	return UE::Core::Private::Atomic::AndExchange(Atomic, Val); \
+} 
+
+INSTRUMENT_LOAD_FUNC_THUNK(int8)
+INSTRUMENT_LOAD_FUNC_THUNK(int16)
+INSTRUMENT_LOAD_FUNC_THUNK(int32)
+INSTRUMENT_LOAD_FUNC_THUNK(int64)
+
+INSTRUMENT_STORE_FUNC_THUNK(int8)
+INSTRUMENT_STORE_FUNC_THUNK(int16)
+INSTRUMENT_STORE_FUNC_THUNK(int32)
+INSTRUMENT_STORE_FUNC_THUNK(int64)
+
+INSTRUMENT_EXCHANGE_FUNC_THUNK(int8)
+INSTRUMENT_EXCHANGE_FUNC_THUNK(int16)
+INSTRUMENT_EXCHANGE_FUNC_THUNK(int32)
+INSTRUMENT_EXCHANGE_FUNC_THUNK(int64)
+
+INSTRUMENT_COMPARE_EXCHANGE_FUNC_THUNK(int8)
+INSTRUMENT_COMPARE_EXCHANGE_FUNC_THUNK(int16)
+INSTRUMENT_COMPARE_EXCHANGE_FUNC_THUNK(int32)
+INSTRUMENT_COMPARE_EXCHANGE_FUNC_THUNK(int64)
+
+INSTRUMENT_FETCHADD_FUNC_THUNK(int8)
+INSTRUMENT_FETCHADD_FUNC_THUNK(int16)
+INSTRUMENT_FETCHADD_FUNC_THUNK(int32)
+INSTRUMENT_FETCHADD_FUNC_THUNK(int64)
+INSTRUMENT_FETCHSUB_FUNC_THUNK(int8)
+INSTRUMENT_FETCHSUB_FUNC_THUNK(int16)
+INSTRUMENT_FETCHSUB_FUNC_THUNK(int32)
+INSTRUMENT_FETCHSUB_FUNC_THUNK(int64)
+INSTRUMENT_FETCHOR_FUNC_THUNK(int8)
+INSTRUMENT_FETCHOR_FUNC_THUNK(int16)
+INSTRUMENT_FETCHOR_FUNC_THUNK(int32)
+INSTRUMENT_FETCHOR_FUNC_THUNK(int64)
+INSTRUMENT_FETCHXOR_FUNC_THUNK(int8)
+INSTRUMENT_FETCHXOR_FUNC_THUNK(int16)
+INSTRUMENT_FETCHXOR_FUNC_THUNK(int32)
+INSTRUMENT_FETCHXOR_FUNC_THUNK(int64)
+INSTRUMENT_FETCHAND_FUNC_THUNK(int8)
+INSTRUMENT_FETCHAND_FUNC_THUNK(int16)
+INSTRUMENT_FETCHAND_FUNC_THUNK(int32)
+INSTRUMENT_FETCHAND_FUNC_THUNK(int64)
+
+}
+
+#endif // USING_INSTRUMENTATION
\ No newline at end of file
diff --git a/Engine/Source/Runtime/Instrumentation/Public/Instrumentation/Types.h b/Engine/Source/Runtime/Instrumentation/Public/Instrumentation/Types.h
new file mode 100644
--- /dev/null
+++ b/Engine/Source/Runtime/Instrumentation/Public/Instrumentation/Types.h
@@ -0,0 +1,94 @@
+// Copyright Epic Games, Inc. All Rights Reserved.
+
+#pragma once
+
+#include "Misc/CoreMiscDefines.h"
+
+#if USING_INSTRUMENTATION
+
+#include "Instrumentation/Defines.h"
+
+#include "CoreTypes.h"
+
+#include <atomic>
+
+namespace UE::Instrumentation
+{
+
+	// ------------------------------------------------------------------------------
+	// Atomics.
+	// ------------------------------------------------------------------------------
+	enum FAtomicMemoryOrder : int8 {
+		MEMORY_ORDER_RELAXED,
+		MEMORY_ORDER_CONSUME,
+		MEMORY_ORDER_ACQUIRE,
+		MEMORY_ORDER_RELEASE,
+		MEMORY_ORDER_ACQ_REL,
+		MEMORY_ORDER_SEQ_CST
+	};
+
+	INSTRUMENTATION_FUNCTION_ATTRIBUTES inline const TCHAR* LexToString(FAtomicMemoryOrder Order)
+	{
+		switch (Order)
+		{
+		case MEMORY_ORDER_RELAXED:
+			return TEXT("relaxed");
+		case MEMORY_ORDER_CONSUME:
+			return TEXT("consume");
+		case MEMORY_ORDER_ACQUIRE:
+			return TEXT("acquire");
+		case MEMORY_ORDER_RELEASE:
+			return TEXT("release");
+		case MEMORY_ORDER_ACQ_REL:
+			return TEXT("acq_rel");
+		case MEMORY_ORDER_SEQ_CST:
+			return TEXT("seq_cst");
+		default:
+			return TEXT("unknown");
+		}
+	};
+
+	INSTRUMENTATION_FUNCTION_ATTRIBUTES inline std::memory_order ToStdMemoryOrder(FAtomicMemoryOrder MemoryOrder)
+	{
+		switch (MemoryOrder)
+		{
+		case MEMORY_ORDER_RELAXED:
+			return std::memory_order_relaxed;
+		case MEMORY_ORDER_CONSUME:
+			return std::memory_order_consume;
+		case MEMORY_ORDER_ACQUIRE:
+			return std::memory_order_acquire;
+		case MEMORY_ORDER_RELEASE:
+			return std::memory_order_release;
+		case MEMORY_ORDER_ACQ_REL:
+			return std::memory_order_acq_rel;
+		default:
+			return std::memory_order_seq_cst;
+		}
+	}
+
+	INSTRUMENTATION_FUNCTION_ATTRIBUTES inline bool IsAtomicOrderAcquire(FAtomicMemoryOrder Order)
+	{
+		return Order == FAtomicMemoryOrder::MEMORY_ORDER_ACQ_REL || Order == FAtomicMemoryOrder::MEMORY_ORDER_ACQUIRE
+			|| Order == FAtomicMemoryOrder::MEMORY_ORDER_CONSUME || Order == FAtomicMemoryOrder::MEMORY_ORDER_SEQ_CST;
+	}
+
+	INSTRUMENTATION_FUNCTION_ATTRIBUTES inline bool IsAtomicOrderRelease(FAtomicMemoryOrder Order)
+	{
+		return Order == FAtomicMemoryOrder::MEMORY_ORDER_ACQ_REL || Order == FAtomicMemoryOrder::MEMORY_ORDER_RELEASE
+			|| Order == FAtomicMemoryOrder::MEMORY_ORDER_SEQ_CST;
+	}
+
+	INSTRUMENTATION_FUNCTION_ATTRIBUTES inline bool IsAtomicOrderAcquireRelease(FAtomicMemoryOrder Order)
+	{
+		return Order == FAtomicMemoryOrder::MEMORY_ORDER_ACQ_REL || Order == FAtomicMemoryOrder::MEMORY_ORDER_SEQ_CST;
+	}
+
+	INSTRUMENTATION_FUNCTION_ATTRIBUTES FORCEINLINE bool IsAtomicOrderRelaxed(FAtomicMemoryOrder Order)
+	{
+		return Order == FAtomicMemoryOrder::MEMORY_ORDER_RELAXED;
+	}
+
+} // UE::Instrumentation
+
+#endif // USING_INSTRUMENTATION
\ No newline at end of file
diff --git a/Engine/Source/Runtime/Launch/Private/Windows/LaunchWindows.cpp b/Engine/Source/Runtime/Launch/Private/Windows/LaunchWindows.cpp
--- a/Engine/Source/Runtime/Launch/Private/Windows/LaunchWindows.cpp
+++ b/Engine/Source/Runtime/Launch/Private/Windows/LaunchWindows.cpp
@@ -4,6 +4,7 @@
 
 #include "Misc/App.h"
 #include "Misc/OutputDeviceError.h"
+#include "Sanitizer/RaceDetector.h"
 #include "LaunchEngineLoop.h"
 #include "HAL/ExceptionHandling.h"
 #include "HAL/PlatformMallocCrash.h"
@@ -314,6 +315,23 @@ LAUNCH_API void LaunchWindowsShutdown()
 
 int32 WINAPI WinMain(_In_ HINSTANCE hInInstance, _In_opt_ HINSTANCE hPrevInstance, _In_ char* pCmdLine, _In_ int32 nCmdShow)
 {
+#if USING_INSTRUMENTATION
+
+	if (!UE::Sanitizer::RaceDetector::Initialize())
+	{
+		return 1;
+	}
+
+	ON_SCOPE_EXIT
+	{
+		if (!UE::Sanitizer::RaceDetector::Shutdown())
+		{
+			return 1;
+		}
+	};
+
+#endif
+
 	int32 Result = LaunchWindowsStartup(hInInstance, hPrevInstance, pCmdLine, nCmdShow, nullptr);
 	LaunchWindowsShutdown();
 	return Result;
