Fixed an assert in MallocBinnedCommon.cpp for when UE_MBC_MAX_LISTED_SMALL_POOL_SIZE is not aligned to a page size. This should only be applicable for cases when UE_MBC_MAX_LISTED_SMALL_POOL_SIZE != UE_MB3_MAX_SMALL_POOL_SIZE

diff --git a/Engine/Source/Runtime/Core/Private/HAL/MallocBinnedCommon.cpp b/Engine/Source/Runtime/Core/Private/HAL/MallocBinnedCommon.cpp
--- a/Engine/Source/Runtime/Core/Private/HAL/MallocBinnedCommon.cpp
+++ b/Engine/Source/Runtime/Core/Private/HAL/MallocBinnedCommon.cpp
@@ -140,7 +140,7 @@ uint8 FSizeTableEntry::FillSizeTable(uint64 PlatformPageSize, FSizeTableEntry* S
 
 	Algo::Sort(MakeArrayView(SizeTable, Index));
 	check(SizeTable[Index - 1].BinSize == UE_MBC_MAX_LISTED_SMALL_POOL_SIZE);
-	check(IsAligned(UE_MBC_MAX_LISTED_SMALL_POOL_SIZE, BasePageSize));
+	check((UE_MBC_MAX_LISTED_SMALL_POOL_SIZE == MaxSize) || IsAligned(UE_MBC_MAX_LISTED_SMALL_POOL_SIZE, BasePageSize));
 	for (uint32 Size = UE_MBC_MAX_LISTED_SMALL_POOL_SIZE + BasePageSize; Size <= MaxSize; Size += SizeIncrement)
 	{
 		SizeTable[Index++] = FSizeTableEntry(Size, PlatformPageSize, Size / BasePageSize, BasePageSize);
