Use a force inline prefix that is supported by native clang

diff --git a/Engine/Source/Runtime/Core/Public/Clang/ClangPlatformCodeAnalysis.h b/Engine/Source/Runtime/Core/Public/Clang/ClangPlatformCodeAnalysis.h
--- a/Engine/Source/Runtime/Core/Public/Clang/ClangPlatformCodeAnalysis.h
+++ b/Engine/Source/Runtime/Core/Public/Clang/ClangPlatformCodeAnalysis.h
@@ -163,20 +163,20 @@
 			std::atomic<T> Value;
 
 		public:
-			__forceinline TTSANSafeValue() = default;
-			__forceinline TTSANSafeValue(T InValue) : Value(InValue) { }
-			__forceinline operator T() const { return Value.load(std::memory_order_relaxed); }
-			__forceinline void operator=(T InValue) { Value.store(InValue, std::memory_order_relaxed); }
-			__forceinline T operator++() { return Value.fetch_add(1, std::memory_order_relaxed) + 1; }
-			__forceinline T operator++(int) { return Value.fetch_add(1, std::memory_order_relaxed); }
-			__forceinline T operator--() { return Value.fetch_sub(1, std::memory_order_relaxed) - 1; }
-			__forceinline T operator--(int) { return Value.fetch_sub(1, std::memory_order_relaxed); }
+			__attribute__((always_inline)) TTSANSafeValue() = default;
+			__attribute__((always_inline)) TTSANSafeValue(T InValue) : Value(InValue) { }
+			__attribute__((always_inline)) operator T() const { return Value.load(std::memory_order_relaxed); }
+			__attribute__((always_inline)) void operator=(T InValue) { Value.store(InValue, std::memory_order_relaxed); }
+			__attribute__((always_inline)) T operator++() { return Value.fetch_add(1, std::memory_order_relaxed) + 1; }
+			__attribute__((always_inline)) T operator++(int) { return Value.fetch_add(1, std::memory_order_relaxed); }
+			__attribute__((always_inline)) T operator--() { return Value.fetch_sub(1, std::memory_order_relaxed) - 1; }
+			__attribute__((always_inline)) T operator--(int) { return Value.fetch_sub(1, std::memory_order_relaxed); }
 
 			template <typename OtherType>
-			__forceinline T operator+=(OtherType InValue) { return Value.fetch_add(InValue, std::memory_order_relaxed) + InValue; }
+			__attribute__((always_inline)) T operator+=(OtherType InValue) { return Value.fetch_add(InValue, std::memory_order_relaxed) + InValue; }
 
 			template <typename OtherType>
-			__forceinline T operator-=(OtherType InValue) { return Value.fetch_sub(InValue, std::memory_order_relaxed) - InValue; }
+			__attribute__((always_inline)) T operator-=(OtherType InValue) { return Value.fetch_sub(InValue, std::memory_order_relaxed) - InValue; }
 		};
 
 	    #define TSAN_ATOMIC(Type) TTSANSafeValue<Type>
