[Insights] Fixed TSAN data race re LLM stats emitted before LLM is configured. Also fixed issue re stats and counters emitting memory events before LLM is configured.

#rnx

diff --git a/Engine/Source/Runtime/Core/Private/HAL/LowLevelMemTracker.cpp b/Engine/Source/Runtime/Core/Private/HAL/LowLevelMemTracker.cpp
--- a/Engine/Source/Runtime/Core/Private/HAL/LowLevelMemTracker.cpp
+++ b/Engine/Source/Runtime/Core/Private/HAL/LowLevelMemTracker.cpp
@@ -19,6 +19,7 @@
 #include "Misc/Fork.h"
 #include "Misc/OutputDeviceRedirector.h"
 #include "Misc/Paths.h"
+#include "Misc/ScopeExit.h"
 #include "Misc/ScopeLock.h"
 #include "Misc/ScopeRWLock.h"
 #include "Misc/StringBuilder.h"
@@ -2255,7 +2256,11 @@ void FLowLevelMemTracker::ProcessCommandLine(const TCHAR* CmdLine)
 		}
 		return;
 	}
-	bConfigurationComplete = true;
+
+	ON_SCOPE_EXIT
+	{
+		bConfigurationComplete = true;
+	};
 
 	if (bShouldDisable)
 	{
diff --git a/Engine/Source/Runtime/Core/Private/ProfilingDebugging/TraceAuxiliary.cpp b/Engine/Source/Runtime/Core/Private/ProfilingDebugging/TraceAuxiliary.cpp
--- a/Engine/Source/Runtime/Core/Private/ProfilingDebugging/TraceAuxiliary.cpp
+++ b/Engine/Source/Runtime/Core/Private/ProfilingDebugging/TraceAuxiliary.cpp
@@ -890,8 +890,6 @@ void FTraceAuxiliaryImpl::StartWorkerThread()
 void TraceAuxiliaryEmitStats(const UE::Trace::FStatistics& Stats)
 {
 	uint64 TotalMemoryUsed = Stats.BlockPoolAllocated + Stats.FixedBufferAllocated + Stats.SharedBufferAllocated + Stats.CacheAllocated;
-	//ensure(TotalMemoryUsed == Stats.MemoryUsed);
-	//ensure(Stats.BytesEmitted >= Stats.BytesTraced);
 
 #if UE_TRACEAUX_EMIT_STAT_COUNTERS
 	SET_MEMORY_STAT(STAT_Trace_MemoryUsed, Stats.MemoryUsed);
@@ -940,6 +938,15 @@ void TraceAuxiliaryOnTraceUpdateCallback()
 	// This is called each time the Trace system updates.
 	// If the trace system uses a worker thread, this is called from the worker thread.
 
+#if ENABLE_LOW_LEVEL_MEM_TRACKER
+	if (!FLowLevelMemTracker::Get().IsConfigured())
+	{
+		// a) Avoids emitting LLM stats before LLM is configured.
+		// b) Avoids allocating memory (through the STAT or CounterTrace APIs) before LLM is configured.
+		return;
+	}
+#endif
+
 #if UE_TRACEAUX_EMIT_STAT_COUNTERS || UE_TRACEAUX_EMIT_TRACE_COUNTERS || UE_TRACEAUX_EMIT_LLM_STATS
 	UE::Trace::FStatistics Stats;
 	UE::Trace::GetStatistics(Stats);
diff --git a/Engine/Source/Runtime/Core/Public/HAL/LowLevelMemTracker.h b/Engine/Source/Runtime/Core/Public/HAL/LowLevelMemTracker.h
--- a/Engine/Source/Runtime/Core/Public/HAL/LowLevelMemTracker.h
+++ b/Engine/Source/Runtime/Core/Public/HAL/LowLevelMemTracker.h
@@ -750,6 +750,12 @@ public:
 	{
 		return bFullyInitialised;
 	}
+
+	FORCEINLINE bool IsConfigured() const
+	{
+		return bConfigurationComplete;
+	}
+
 	/**
 	 * Initialize the data required for tracking allocations if it has not already been done.
 	 * This can occur after construction of the singleton FLowLevelMemTracker. 
@@ -865,7 +871,7 @@ private:
 	bool bInitialisedTracking;
 	bool bIsBootstrapping;
 	bool bFullyInitialised;
-	bool bConfigurationComplete;
+	std::atomic<bool> bConfigurationComplete;
 	bool bTagAdded;
 	bool bAutoPublish;
 	bool bPublishSingleFrame;
