Fixed TSet::EmplaceByHash

diff --git a/Engine/Source/Runtime/Core/Public/Containers/Set.h b/Engine/Source/Runtime/Core/Public/Containers/Set.h
--- a/Engine/Source/Runtime/Core/Public/Containers/Set.h
+++ b/Engine/Source/Runtime/Core/Public/Containers/Set.h
@@ -781,7 +781,7 @@ public:
 
 		SizeType NewHashIndex = ElementAllocation.Index;
 
-		if (!TryReplaceExisting(KeyHash, Element, ElementAllocation.Index, bIsAlreadyInSetPtr))
+		if (!TryReplaceExisting(KeyHash, Element, NewHashIndex, bIsAlreadyInSetPtr))
 		{
 			RehashOrLink(KeyHash, Element, NewHashIndex);
 		}
