Replaced a custom string allocator for TextKeys by the PersistentLinearAllocator

diff --git a/Engine/Source/Runtime/Core/Private/Internationalization/TextKey.cpp b/Engine/Source/Runtime/Core/Private/Internationalization/TextKey.cpp
--- a/Engine/Source/Runtime/Core/Private/Internationalization/TextKey.cpp
+++ b/Engine/Source/Runtime/Core/Private/Internationalization/TextKey.cpp
@@ -15,6 +15,7 @@
 #include "Misc/LazySingleton.h"
 #include "Misc/ScopeRWLock.h"
 #include "Modules/VisualizerDebuggingState.h"
+#include "Memory/LinearAllocator.h"
 
 DEFINE_LOG_CATEGORY_STATIC(LogTextKey, Log, All);
 
@@ -27,14 +28,6 @@ DEFINE_LOG_CATEGORY_STATIC(LogTextKey, Log, All);
 	#define UE_TEXTKEY_SPLIT_GUID (1)
 #endif
 
-#ifndef UE_TEXTKEY_SLAB_ALLOCATOR_SLAB_SIZE
-	#define UE_TEXTKEY_SLAB_ALLOCATOR_SLAB_SIZE (32768)
-#endif
-
-#ifndef UE_TEXTKEY_GUID_CHUNK_SIZE_BYTES
-	#define UE_TEXTKEY_GUID_CHUNK_SIZE_BYTES (32768)
-#endif
-
 #ifndef UE_TEXTKEY_ELEMENTS_CHUNK_SIZE_BYTES
 	#define UE_TEXTKEY_ELEMENTS_CHUNK_SIZE_BYTES (32768)
 #endif
@@ -343,13 +336,13 @@ private:
 			LLM_SCOPE_BYNAME(TEXT("Localization/TextKeys"));
 			if (KeyData.IsStringView())
 			{
-				FTextKeyStringView NewStr = StringAllocations.Add(KeyData.AsStringView());
+				FTextKeyStringView NewStr = FStaticAllocator::Add(KeyData.AsStringView());
 				return KeyDataAllocations.Add(FKeyData(NewStr, KeyData.StrHash));
 			}
 			else
 			{
-				const FGuid& NewGuid = GuidAllocations[GuidAllocations.AddElement(KeyData.AsGuid())];
-				return KeyDataAllocations.Add(FKeyData(NewGuid, KeyData.StrHash));
+				const FGuid* NewGuid = FStaticAllocator::Add(KeyData.AsGuid());
+				return KeyDataAllocations.Add(FKeyData(*NewGuid, KeyData.StrHash));
 			}
 		}
 	}
@@ -462,89 +455,26 @@ private:
 		int32* Hash = nullptr;
 	};
 
-	class FStringAllocator
+	struct FStaticAllocator
 	{
-	public:
-		FStringAllocator() = default;
-
-		~FStringAllocator()
+		static FTextKeyStringView Add(FTextKeyStringView InStr)
 		{
-			uint64 TotalNumElementsWasted = 0;
-			for (FStringSlab& Slab : Slabs)
-			{
-				TotalNumElementsWasted += (SlabSizeInElements - Slab.NumElementsUsed);
-				FMemory::Free(Slab.Allocation);
-			}
-			//UE_LOG(LogLocalization, Log, TEXT("FTextKey slab allocator allocated %d slabs (of %d elements) and wasted %d elements (%d bytes)"), Slabs.Num(), SlabSizeInElements, TotalNumElementsWasted, TotalNumElementsWasted * sizeof(FTextKeyCharType));
-		}
-
-		FTextKeyStringView Add(FTextKeyStringView InStr)
-		{
-			const int32 NumSlabElementsNeeded = InStr.Len();
-			FStringSlab& Slab = GetSlab(NumSlabElementsNeeded);
-
-			FTextKeyCharType* StringPtr = Slab.Allocation + Slab.NumElementsUsed;
-			Slab.NumElementsUsed += NumSlabElementsNeeded;
+			const int32 NumBytes = InStr.Len() * sizeof(FTextKeyCharType);
+			void* StringPtr = GetPersistentLinearAllocator().Allocate(NumBytes, alignof(FTextKeyCharType));
+			FMemory::Memcpy(StringPtr, InStr.GetData(), NumBytes);
 
-			FMemory::Memcpy(StringPtr, InStr.GetData(), InStr.Len() * sizeof(FTextKeyCharType));
-
-			return FTextKeyStringView(StringPtr, InStr.Len());
+			return FTextKeyStringView((FTextKeyCharType*)StringPtr, InStr.Len());
 		}
 
-	private:
-		static constexpr int32 SlabSizeInElements = UE_TEXTKEY_SLAB_ALLOCATOR_SLAB_SIZE;
-
-		struct FStringSlab
+		static FGuid* Add(const FGuid& InGuid)
 		{
-			FTextKeyCharType* Allocation = nullptr;
-			int32 NumElementsUsed = 0;
-		};
-
-		FStringSlab& GetSlab(const int32 NumSlabElementsNeeded)
-		{
-			if (Slabs.Num() > 0)
-			{
-				// Always try the last slab first
-				if (FStringSlab& Slab = Slabs.Last();
-					(Slab.NumElementsUsed + NumSlabElementsNeeded) <= SlabSizeInElements)
-				{
-					return Slab;
-				}
-
-				if (Slabs.Num() > 1)
-				{
-					// We only add to the last slab in the array, so if we've run out of space, merge it back into the array based 
-					// on its current used size (sorted most used first), and then check to see if the new last slab has space
-					FStringSlab SlabToMerge = Slabs.Pop(EAllowShrinking::No);
-					const int32 MergeIndex = Algo::LowerBound(Slabs, SlabToMerge, [](const FStringSlab& A, const FStringSlab& B)
-					{
-						return A.NumElementsUsed > B.NumElementsUsed;
-					});
-					Slabs.Insert(SlabToMerge, MergeIndex);
-					if (MergeIndex != Slabs.Num())
-					{
-						if (FStringSlab& Slab = Slabs.Last();
-							(Slab.NumElementsUsed + NumSlabElementsNeeded) <= SlabSizeInElements)
-						{
-							return Slab;
-						}
-					}
-				}
-			}
-
-			// If no slabs have space then just allocate a new one
-			checkf(NumSlabElementsNeeded <= SlabSizeInElements, TEXT("Tried to allocate a FTextKey string of %d elements, which is larger than the allowed slab size of %d elements!"), NumSlabElementsNeeded, SlabSizeInElements);
-			FStringSlab& Slab = Slabs.AddDefaulted_GetRef();
-			Slab.Allocation = static_cast<FTextKeyCharType*>(FMemory::Malloc(SlabSizeInElements * sizeof(FTextKeyCharType), 1));
-			return Slab;
+			FGuid* GuidPtr = (FGuid*)GetPersistentLinearAllocator().Allocate(sizeof(InGuid), alignof(FGuid));
+			*GuidPtr = InGuid;
+			return GuidPtr;
 		}
-
-		TArray<FStringSlab> Slabs;
 	};
 
 	mutable FRWLock DataRW;
-	FStringAllocator StringAllocations;
-	TChunkedArray<FGuid, UE_TEXTKEY_GUID_CHUNK_SIZE_BYTES> GuidAllocations;
 	FKeyDataAllocator KeyDataAllocations;
 
 	// Sparse TCHAR state; built on-demand by anything still using the deprecated FTextKey::GetChars function
