Added an ability to initialize MB3 with an external memory block

diff --git a/Engine/Source/Runtime/Core/Private/HAL/MallocBinned3.cpp b/Engine/Source/Runtime/Core/Private/HAL/MallocBinned3.cpp
--- a/Engine/Source/Runtime/Core/Private/HAL/MallocBinned3.cpp
+++ b/Engine/Source/Runtime/Core/Private/HAL/MallocBinned3.cpp
@@ -442,8 +442,7 @@ FMallocBinned3::FPoolInfoSmall* FMallocBinned3::GetFrontPool(FPoolTable& Table,
 	return Private::GetOrCreatePoolInfoSmall(*this, InPoolIndex, OutBlockIndex);
 }
 
-
-FMallocBinned3::FMallocBinned3()
+FMallocBinned3::FMallocBinned3(const FPlatformMemory::FPlatformVirtualMemoryBlock* ExternalMemoryBlock)
 {
 	static bool bOnce = false;
 	check(!bOnce); // this is now a singleton-like thing and you cannot make multiple copies
@@ -455,7 +454,8 @@ FMallocBinned3::FMallocBinned3()
 	// First thing we try to allocate address space for bins as it might help us to move forward Constants.AddressStart and reduce the amount of available address space for the Large OS Allocs
 	// Available address space is used to reserve hash map that can address all of that range, so less addressable space means less memory is allocated for book keeping
 #if !BINNED3_USE_SEPARATE_VM_PER_POOL
-	Binned3BaseVMBlock = FPlatformMemory::FPlatformVirtualMemoryBlock::AllocateVirtual(UE_MB3_SMALL_POOL_COUNT * UE_MB3_MAX_MEMORY_PER_POOL_SIZE, OsAllocationGranularity);
+	Binned3BaseVMBlock = ExternalMemoryBlock ? *ExternalMemoryBlock : FPlatformMemory::FPlatformVirtualMemoryBlock::AllocateVirtual(UE_MB3_SMALL_POOL_COUNT * UE_MB3_MAX_MEMORY_PER_POOL_SIZE, OsAllocationGranularity);
+	check(Binned3BaseVMBlock.GetActualSize() == UE_MB3_SMALL_POOL_COUNT * UE_MB3_MAX_MEMORY_PER_POOL_SIZE);
 	Binned3BaseVMPtr = (uint8*)Binned3BaseVMBlock.GetVirtualPointer();
 	check(IsAligned(Binned3BaseVMPtr, OsAllocationGranularity));
 	verify(Binned3BaseVMPtr);
diff --git a/Engine/Source/Runtime/Core/Public/HAL/MallocBinned3.h b/Engine/Source/Runtime/Core/Public/HAL/MallocBinned3.h
--- a/Engine/Source/Runtime/Core/Public/HAL/MallocBinned3.h
+++ b/Engine/Source/Runtime/Core/Public/HAL/MallocBinned3.h
@@ -286,7 +286,7 @@ private:
 
 public:
 
-	CORE_API FMallocBinned3();
+	CORE_API FMallocBinned3(const FPlatformMemory::FPlatformVirtualMemoryBlock* ExternalMemoryBlock = nullptr);
 
 	CORE_API virtual ~FMallocBinned3();
 
