Fixed algorithms to take either an rvalue range, so views can be passed as temporaries, or const lvalue refs if the range isn't modified to produce fewer template instantiations.

diff --git a/Engine/Source/Developer/DerivedDataCache/Private/DerivedDataCacheStats.cpp b/Engine/Source/Developer/DerivedDataCache/Private/DerivedDataCacheStats.cpp
--- a/Engine/Source/Developer/DerivedDataCache/Private/DerivedDataCacheStats.cpp
+++ b/Engine/Source/Developer/DerivedDataCache/Private/DerivedDataCacheStats.cpp
@@ -339,7 +339,7 @@ void FCacheStoreStats::SetTotalPhysicalSize(uint64 InTotalPhysicalSize)
 FCacheBucketStats& FCacheStats::GetBucket(FCacheBucket Bucket)
 {
 	TUniqueLock Lock(Mutex);
-	const int32 Index = Algo::LowerBoundBy(BucketStats, Bucket, [](TUniquePtr<FCacheBucketStats>& Stats) { return Stats->Bucket; });
+	const int32 Index = Algo::LowerBoundBy(BucketStats, Bucket, &FCacheBucketStats::Bucket);
 	if (!BucketStats.IsValidIndex(Index) || BucketStats[Index]->Bucket != Bucket)
 	{
 		FCacheBucketStats* Stats = new FCacheBucketStats;
diff --git a/Engine/Source/Runtime/Core/Public/Algo/BinarySearch.h b/Engine/Source/Runtime/Core/Public/Algo/BinarySearch.h
--- a/Engine/Source/Runtime/Core/Public/Algo/BinarySearch.h
+++ b/Engine/Source/Runtime/Core/Public/Algo/BinarySearch.h
@@ -89,12 +89,12 @@ namespace Algo
 	 * @returns Position of the first element >= Value, may be position after last element in range
 	 */
 	template <typename RangeType, typename ValueType, typename SortPredicateType>
-	FORCEINLINE auto LowerBound(RangeType& Range, const ValueType& Value, SortPredicateType SortPredicate) -> decltype(GetNum(Range))
+	FORCEINLINE auto LowerBound(const RangeType& Range, const ValueType& Value, SortPredicateType SortPredicate) -> decltype(GetNum(Range))
 	{
 		return AlgoImpl::LowerBoundInternal(GetData(Range), GetNum(Range), Value, FIdentityFunctor(), SortPredicate);
 	}
 	template <typename RangeType, typename ValueType>
-	FORCEINLINE auto LowerBound(RangeType& Range, const ValueType& Value) -> decltype(GetNum(Range))
+	FORCEINLINE auto LowerBound(const RangeType& Range, const ValueType& Value) -> decltype(GetNum(Range))
 	{
 		return AlgoImpl::LowerBoundInternal(GetData(Range), GetNum(Range), Value, FIdentityFunctor(), TLess<>());
 	}
@@ -110,12 +110,12 @@ namespace Algo
 	 * @returns Position of the first element >= Value, may be position after last element in range
 	 */
 	template <typename RangeType, typename ValueType, typename ProjectionType, typename SortPredicateType>
-	FORCEINLINE auto LowerBoundBy(RangeType& Range, const ValueType& Value, ProjectionType Projection, SortPredicateType SortPredicate) -> decltype(GetNum(Range))
+	FORCEINLINE auto LowerBoundBy(const RangeType& Range, const ValueType& Value, ProjectionType Projection, SortPredicateType SortPredicate) -> decltype(GetNum(Range))
 	{
 		return AlgoImpl::LowerBoundInternal(GetData(Range), GetNum(Range), Value, Projection, SortPredicate);
 	}
 	template <typename RangeType, typename ValueType, typename ProjectionType>
-	FORCEINLINE auto LowerBoundBy(RangeType& Range, const ValueType& Value, ProjectionType Projection) -> decltype(GetNum(Range))
+	FORCEINLINE auto LowerBoundBy(const RangeType& Range, const ValueType& Value, ProjectionType Projection) -> decltype(GetNum(Range))
 	{
 		return AlgoImpl::LowerBoundInternal(GetData(Range), GetNum(Range), Value, Projection, TLess<>());
 	}
@@ -130,12 +130,12 @@ namespace Algo
 	 * @returns Position of the first element > Value, may be past end of range
 	 */
 	template <typename RangeType, typename ValueType, typename SortPredicateType>
-	FORCEINLINE auto UpperBound(RangeType& Range, const ValueType& Value, SortPredicateType SortPredicate) -> decltype(GetNum(Range))
+	FORCEINLINE auto UpperBound(const RangeType& Range, const ValueType& Value, SortPredicateType SortPredicate) -> decltype(GetNum(Range))
 	{
 		return AlgoImpl::UpperBoundInternal(GetData(Range), GetNum(Range), Value, FIdentityFunctor(), SortPredicate);
 	}
 	template <typename RangeType, typename ValueType>
-	FORCEINLINE auto UpperBound(RangeType& Range, const ValueType& Value) -> decltype(GetNum(Range))
+	FORCEINLINE auto UpperBound(const RangeType& Range, const ValueType& Value) -> decltype(GetNum(Range))
 	{
 		return AlgoImpl::UpperBoundInternal(GetData(Range), GetNum(Range), Value, FIdentityFunctor(), TLess<>());
 	}
@@ -151,12 +151,12 @@ namespace Algo
 	 * @returns Position of the first element > Value, may be past end of range
 	 */
 	template <typename RangeType, typename ValueType, typename ProjectionType, typename SortPredicateType>
-	FORCEINLINE auto UpperBoundBy(RangeType& Range, const ValueType& Value, ProjectionType Projection, SortPredicateType SortPredicate) -> decltype(GetNum(Range))
+	FORCEINLINE auto UpperBoundBy(const RangeType& Range, const ValueType& Value, ProjectionType Projection, SortPredicateType SortPredicate) -> decltype(GetNum(Range))
 	{
 		return AlgoImpl::UpperBoundInternal(GetData(Range), GetNum(Range), Value, Projection, SortPredicate);
 	}
 	template <typename RangeType, typename ValueType, typename ProjectionType>
-	FORCEINLINE auto UpperBoundBy(RangeType& Range, const ValueType& Value, ProjectionType Projection) -> decltype(GetNum(Range))
+	FORCEINLINE auto UpperBoundBy(const RangeType& Range, const ValueType& Value, ProjectionType Projection) -> decltype(GetNum(Range))
 	{
 		return AlgoImpl::UpperBoundInternal(GetData(Range), GetNum(Range), Value, Projection, TLess<>());
 	}
@@ -170,7 +170,7 @@ namespace Algo
 	 * @return Index of found element, or INDEX_NONE
 	 */
 	template <typename RangeType, typename ValueType, typename SortPredicateType>
-	FORCEINLINE auto BinarySearch(RangeType& Range, const ValueType& Value, SortPredicateType SortPredicate) -> decltype(GetNum(Range))
+	FORCEINLINE auto BinarySearch(const RangeType& Range, const ValueType& Value, SortPredicateType SortPredicate) -> decltype(GetNum(Range))
 	{
 		auto CheckIndex = LowerBound(Range, Value, SortPredicate);
 		if (CheckIndex < GetNum(Range))
@@ -185,7 +185,7 @@ namespace Algo
 		return INDEX_NONE;
 	}
 	template <typename RangeType, typename ValueType>
-	FORCEINLINE auto BinarySearch(RangeType& Range, const ValueType& Value)
+	FORCEINLINE auto BinarySearch(const RangeType& Range, const ValueType& Value)
 	{
 		return BinarySearch(Range, Value, TLess<>());
 	}
@@ -200,7 +200,7 @@ namespace Algo
 	 * @return Index of found element, or INDEX_NONE
 	 */
 	template <typename RangeType, typename ValueType, typename ProjectionType, typename SortPredicateType>
-	FORCEINLINE auto BinarySearchBy(RangeType& Range, const ValueType& Value, ProjectionType Projection, SortPredicateType SortPredicate) -> decltype(GetNum(Range))
+	FORCEINLINE auto BinarySearchBy(const RangeType& Range, const ValueType& Value, ProjectionType Projection, SortPredicateType SortPredicate) -> decltype(GetNum(Range))
 	{
 		auto CheckIndex = LowerBoundBy(Range, Value, Projection, SortPredicate);
 		if (CheckIndex < GetNum(Range))
@@ -215,7 +215,7 @@ namespace Algo
 		return INDEX_NONE;
 	}
 	template <typename RangeType, typename ValueType, typename ProjectionType>
-	FORCEINLINE auto BinarySearchBy(RangeType& Range, const ValueType& Value, ProjectionType Projection)
+	FORCEINLINE auto BinarySearchBy(const RangeType& Range, const ValueType& Value, ProjectionType Projection)
 	{
 		return BinarySearchBy(Range, Value, Projection, TLess<>());
 	}
diff --git a/Engine/Source/Runtime/Core/Public/Algo/Compare.h b/Engine/Source/Runtime/Core/Public/Algo/Compare.h
--- a/Engine/Source/Runtime/Core/Public/Algo/Compare.h
+++ b/Engine/Source/Runtime/Core/Public/Algo/Compare.h
@@ -12,7 +12,7 @@
 namespace Algo::Private
 {
 	template <typename InAT, typename InBT, typename ProjectionT, typename PredicateT>
-	constexpr bool Compare(InAT&& InputA, InBT&& InputB, ProjectionT Projection, PredicateT Predicate)
+	constexpr bool Compare(const InAT& InputA, const InBT& InputB, ProjectionT Projection, PredicateT Predicate)
 	{
 		const SIZE_T SizeA = GetNum(InputA);
 		const SIZE_T SizeB = GetNum(InputB);
@@ -48,9 +48,9 @@ namespace Algo
 	 * @return Whether the containers are the same size and operator== returned true for every pair of elements.
 	 */
 	template <typename InAT, typename InBT>
-	constexpr bool Compare(InAT&& InputA, InBT&& InputB)
+	constexpr bool Compare(const InAT& InputA, const InBT& InputB)
 	{
-		return Private::Compare(Forward<InAT>(InputA), Forward<InBT>(InputB), FIdentityFunctor(), TEqualTo<>());
+		return Private::Compare(InputA, InputB, FIdentityFunctor(), TEqualTo<>());
 	}
 
 	/**
@@ -63,16 +63,9 @@ namespace Algo
 	 * @return Whether the containers are the same size and the predicate returned true for every pair of elements.
 	 */
 	template <typename InAT, typename InBT, typename PredicateT>
-	constexpr bool Compare(InAT&& InputA, InBT&& InputB, PredicateT Predicate)
+	constexpr bool Compare(const InAT& InputA, const InBT& InputB, PredicateT Predicate)
 	{
-		return Private::Compare(Forward<InAT>(InputA), Forward<InBT>(InputB), FIdentityFunctor(), MoveTemp(Predicate));
-	}
-
-	template <typename InAT, typename InBT, typename PredicateT>
-	UE_DEPRECATED(5.0, "CompareByPredicate has been renamed to Compare.")
-	constexpr bool CompareByPredicate(InAT&& InputA, InBT&& InputB, PredicateT Predicate)
-	{
-		return Compare(Forward<InAT>(InputA), Forward<InBT>(InputB), MoveTemp(Predicate));
+		return Private::Compare(InputA, InputB, FIdentityFunctor(), MoveTemp(Predicate));
 	}
 
 	/**
@@ -85,9 +78,9 @@ namespace Algo
 	 * @return Whether the containers are the same size and operator== returned true for every pair of elements.
 	 */
 	template <typename InAT, typename InBT, typename ProjectionT>
-	constexpr bool CompareBy(InAT&& InputA, InBT&& InputB, ProjectionT Projection)
+	constexpr bool CompareBy(const InAT& InputA, const InBT& InputB, ProjectionT Projection)
 	{
-		return Private::Compare(Forward<InAT>(InputA), Forward<InBT>(InputB), MoveTemp(Projection), TEqualTo<>());
+		return Private::Compare(InputA, InputB, MoveTemp(Projection), TEqualTo<>());
 	}
 
 	/**
@@ -101,9 +94,9 @@ namespace Algo
 	 * @return Whether the containers are the same size and the predicate returned true for every pair of elements.
 	 */
 	template <typename InAT, typename InBT, typename ProjectionT, typename PredicateT>
-	constexpr bool CompareBy(InAT&& InputA, InBT&& InputB, ProjectionT Projection, PredicateT Predicate)
+	constexpr bool CompareBy(const InAT& InputA, const InBT& InputB, ProjectionT Projection, PredicateT Predicate)
 	{
-		return Private::Compare(Forward<InAT>(InputA), Forward<InBT>(InputB), MoveTemp(Projection), MoveTemp(Predicate));
+		return Private::Compare(InputA, InputB, MoveTemp(Projection), MoveTemp(Predicate));
 	}
 
 	/**
diff --git a/Engine/Source/Runtime/Core/Public/Algo/ForEach.h b/Engine/Source/Runtime/Core/Public/Algo/ForEach.h
--- a/Engine/Source/Runtime/Core/Public/Algo/ForEach.h
+++ b/Engine/Source/Runtime/Core/Public/Algo/ForEach.h
@@ -15,7 +15,7 @@ namespace Algo
 	 * @param  Callable   Callable object
 	 */
 	template <typename InT, typename PredicateT, typename CallableT>
-	FORCEINLINE void ForEachIf(InT& Input, PredicateT Predicate, CallableT Callable)
+	FORCEINLINE void ForEachIf(InT&& Input, PredicateT Predicate, CallableT Callable)
 	{
 		for (auto& Value : Input)
 		{
@@ -33,7 +33,7 @@ namespace Algo
 	 * @param  Callable   Callable object
 	 */
 	template <typename InT, typename CallableT>
-	FORCEINLINE void ForEach(InT& Input, CallableT Callable)
+	FORCEINLINE void ForEach(InT&& Input, CallableT Callable)
 	{
 		for (auto& Value : Input)
 		{
diff --git a/Engine/Source/Runtime/Core/Public/Algo/Heapify.h b/Engine/Source/Runtime/Core/Public/Algo/Heapify.h
--- a/Engine/Source/Runtime/Core/Public/Algo/Heapify.h
+++ b/Engine/Source/Runtime/Core/Public/Algo/Heapify.h
@@ -17,7 +17,7 @@ namespace Algo
 	 * @param Range	The range to heapify.
 	 */
 	template <typename RangeType>
-	FORCEINLINE void Heapify(RangeType& Range)
+	FORCEINLINE void Heapify(RangeType&& Range)
 	{
 		AlgoImpl::HeapifyInternal(GetData(Range), GetNum(Range), FIdentityFunctor(), TLess<>());
 	}
@@ -29,7 +29,7 @@ namespace Algo
 	 * @param Predicate	A binary predicate object used to specify if one element should precede another.
 	 */
 	template <typename RangeType, typename PredicateType>
-	FORCEINLINE void Heapify(RangeType& Range, PredicateType Predicate)
+	FORCEINLINE void Heapify(RangeType&& Range, PredicateType Predicate)
 	{
 		AlgoImpl::HeapifyInternal(GetData(Range), GetNum(Range), FIdentityFunctor(), Predicate);
 	}
@@ -42,7 +42,7 @@ namespace Algo
 	 * @param Projection	The projection to apply to the elements.
 	 */
 	template <typename RangeType, typename ProjectionType, typename PredicateType>
-	FORCEINLINE void HeapifyBy(RangeType& Range, ProjectionType Projection)
+	FORCEINLINE void HeapifyBy(RangeType&& Range, ProjectionType Projection)
 	{
 		AlgoImpl::HeapifyInternal(GetData(Range), GetNum(Range), Projection, TLess<>());
 	}
@@ -55,7 +55,7 @@ namespace Algo
 	 * @param Predicate		A binary predicate object used to specify if one element should precede another.
 	 */
 	template <typename RangeType, typename ProjectionType, typename PredicateType>
-	FORCEINLINE void HeapifyBy(RangeType& Range, ProjectionType Projection, PredicateType Predicate)
+	FORCEINLINE void HeapifyBy(RangeType&& Range, ProjectionType Projection, PredicateType Predicate)
 	{
 		AlgoImpl::HeapifyInternal(GetData(Range), GetNum(Range), Projection, Predicate);
 	}
diff --git a/Engine/Source/Runtime/Core/Public/Algo/Includes.h b/Engine/Source/Runtime/Core/Public/Algo/Includes.h
--- a/Engine/Source/Runtime/Core/Public/Algo/Includes.h
+++ b/Engine/Source/Runtime/Core/Public/Algo/Includes.h
@@ -52,7 +52,7 @@ namespace Algo
 	 * @return True if RangeB is a subsequence of RangeA, false otherwise.
 	 */
 	template <typename RangeTypeA, typename RangeTypeB>
-	constexpr bool Includes(RangeTypeA&& RangeA, RangeTypeB&& RangeB)
+	constexpr bool Includes(const RangeTypeA& RangeA, const RangeTypeB& RangeB)
 	{
 		return Private::Includes(GetData(RangeA), GetNum(RangeA), GetData(RangeB), GetNum(RangeB), FIdentityFunctor(), TLess<>());
 	}
@@ -68,7 +68,7 @@ namespace Algo
 	 * @return True if RangeB is a subsequence of RangeA according to the provided predicate, false otherwise.
 	 */
 	template <typename RangeTypeA, typename RangeTypeB, typename SortPredicateType>
-	constexpr bool Includes(RangeTypeA&& RangeA, RangeTypeB&& RangeB, SortPredicateType SortPredicate)
+	constexpr bool Includes(const RangeTypeA& RangeA, const RangeTypeB& RangeB, SortPredicateType SortPredicate)
 	{
 		return Private::Includes(GetData(RangeA), GetNum(RangeA), GetData(RangeB), GetNum(RangeB), FIdentityFunctor(), MoveTemp(SortPredicate));
 	}
@@ -84,7 +84,7 @@ namespace Algo
 	 * @return True if RangeB is a subsequence of RangeA, based on the comparison of projected elements, false otherwise.
 	 */
 	template <typename RangeTypeA, typename RangeTypeB, typename ProjectionType>
-	constexpr bool IncludesBy(RangeTypeA&& RangeA, RangeTypeB&& RangeB, ProjectionType Projection)
+	constexpr bool IncludesBy(const RangeTypeA& RangeA, const RangeTypeB& RangeB, ProjectionType Projection)
 	{
 		return Private::Includes(GetData(RangeA), GetNum(RangeA), GetData(RangeB), GetNum(RangeB), MoveTemp(Projection), TLess<>());
 	}
@@ -101,7 +101,7 @@ namespace Algo
 	 * @return True if RangeB is a subsequence of RangeA according to the comparison of projected elements using the provided predicate, false otherwise.
 	 */
 	template <typename RangeTypeA, typename RangeTypeB, typename ProjectionType, typename SortPredicateType>
-	constexpr bool IncludesBy(RangeTypeA&& RangeA, RangeTypeB&& RangeB, ProjectionType Projection, SortPredicateType SortPredicate)
+	constexpr bool IncludesBy(const RangeTypeA& RangeA, const RangeTypeB& RangeB, ProjectionType Projection, SortPredicateType SortPredicate)
 	{
 		return Private::Includes(GetData(RangeA), GetNum(RangeA), GetData(RangeB), GetNum(RangeB), MoveTemp(Projection), MoveTemp(SortPredicate));
 	}
diff --git a/Engine/Source/Runtime/Core/Public/Algo/IndexOf.h b/Engine/Source/Runtime/Core/Public/Algo/IndexOf.h
--- a/Engine/Source/Runtime/Core/Public/Algo/IndexOf.h
+++ b/Engine/Source/Runtime/Core/Public/Algo/IndexOf.h
@@ -9,7 +9,7 @@
 namespace AlgoImpl
 {
 	template <typename RangeType, typename ValueType, typename ProjectionType>
-	auto IndexOfBy(RangeType&& Range, const ValueType& Value, ProjectionType Proj)
+	auto IndexOfBy(const RangeType& Range, const ValueType& Value, ProjectionType Proj)
 	{
 		auto Num = GetNum(Range);
 		auto Data = GetData(Range);
@@ -27,7 +27,7 @@ namespace AlgoImpl
 	}
 
 	template <typename RangeType, typename PredicateType>
-	auto IndexOfByPredicate(RangeType&& Range, PredicateType Pred)
+	auto IndexOfByPredicate(const RangeType& Range, PredicateType Pred)
 	{
 		auto Num = GetNum(Range);
 		auto Data = GetData(Range);
@@ -56,10 +56,10 @@ namespace Algo
 	 * @return The index of the first element found, or (SizeType)-1 if none was found (where SizeType = decltype(GetNum(Range)))
 	 */
 	template <typename RangeType, typename ValueType>
-	FORCEINLINE auto IndexOf(RangeType&& Range, const ValueType& Value)
+	FORCEINLINE auto IndexOf(const RangeType& Range, const ValueType& Value)
 	{
 		static_assert(TIsContiguousContainer<RangeType>::Value, "Must only be used with contiguous containers");
-		return AlgoImpl::IndexOfBy(Forward<RangeType>(Range), Value, FIdentityFunctor());
+		return AlgoImpl::IndexOfBy(Range, Value, FIdentityFunctor());
 	}
 
 	/**
@@ -72,10 +72,10 @@ namespace Algo
 	 * @return The index of the first element found, or (SizeType)-1 if none was found (where SizeType = decltype(GetNum(Range)))
 	 */
 	template <typename RangeType, typename ValueType, typename ProjectionType>
-	FORCEINLINE auto IndexOfBy(RangeType&& Range, const ValueType& Value, ProjectionType Proj)
+	FORCEINLINE auto IndexOfBy(const RangeType& Range, const ValueType& Value, ProjectionType Proj)
 	{
 		static_assert(TIsContiguousContainer<RangeType>::Value, "Must only be used with contiguous containers");
-		return AlgoImpl::IndexOfBy(Forward<RangeType>(Range), Value, MoveTemp(Proj));
+		return AlgoImpl::IndexOfBy(Range, Value, MoveTemp(Proj));
 	}
 
 	/**
@@ -87,9 +87,9 @@ namespace Algo
 	 * @return The index of the first element found, or (SizeType)-1 if none was found (where SizeType = decltype(GetNum(Range)))
 	 */
 	template <typename RangeType, typename PredicateType>
-	FORCEINLINE auto IndexOfByPredicate(RangeType&& Range, PredicateType Pred)
+	FORCEINLINE auto IndexOfByPredicate(const RangeType& Range, PredicateType Pred)
 	{
 		static_assert(TIsContiguousContainer<RangeType>::Value, "Must only be used with contiguous containers");
-		return AlgoImpl::IndexOfByPredicate(Forward<RangeType>(Range), MoveTemp(Pred));
+		return AlgoImpl::IndexOfByPredicate(Range, MoveTemp(Pred));
 	}
 }
diff --git a/Engine/Source/Runtime/Core/Public/Algo/IsHeap.h b/Engine/Source/Runtime/Core/Public/Algo/IsHeap.h
--- a/Engine/Source/Runtime/Core/Public/Algo/IsHeap.h
+++ b/Engine/Source/Runtime/Core/Public/Algo/IsHeap.h
@@ -22,7 +22,7 @@ namespace AlgoImpl
 	 * @return	returns		true if the range is a min-heap
 	 */
 	template <typename RangeValueType, typename IndexType, typename ProjectionType, typename PredicateType>
-	bool IsHeapInternal(RangeValueType* Heap, IndexType Num, ProjectionType Projection, PredicateType Predicate)
+	bool IsHeapInternal(const RangeValueType* Heap, IndexType Num, ProjectionType Projection, PredicateType Predicate)
 	{
 		for (IndexType Index = 1; Index < Num; Index++)
 		{
@@ -47,7 +47,7 @@ namespace Algo
 	 * @return	returns	true if the range is a min-heap
 	 */
 	template <typename RangeType>
-	FORCEINLINE bool IsHeap(RangeType& Range)
+	FORCEINLINE bool IsHeap(const RangeType& Range)
 	{
 		return AlgoImpl::IsHeapInternal(GetData(Range), GetNum(Range), FIdentityFunctor(), TLess<>());
 	}
@@ -61,7 +61,7 @@ namespace Algo
 	 * @return	returns		true if the range is a min-heap
 	 */
 	template <typename RangeType, typename PredicateType>
-	FORCEINLINE bool IsHeap(RangeType& Range, PredicateType Predicate)
+	FORCEINLINE bool IsHeap(const RangeType& Range, PredicateType Predicate)
 	{
 		return AlgoImpl::IsHeapInternal(GetData(Range), GetNum(Range), FIdentityFunctor(), MoveTemp(Predicate));
 	}
@@ -75,7 +75,7 @@ namespace Algo
 	 * @return	returns		true if the range is a min-heap
 	 */
 	template <typename RangeType, typename ProjectionType>
-	FORCEINLINE bool IsHeapBy(RangeType& Range, ProjectionType Projection)
+	FORCEINLINE bool IsHeapBy(const RangeType& Range, ProjectionType Projection)
 	{
 		return AlgoImpl::IsHeapInternal(GetData(Range), GetNum(Range), MoveTemp(Projection), TLess<>());
 	}
@@ -90,7 +90,7 @@ namespace Algo
 	 * @return	returns		true if the range is a min-heap
 	 */
 	template <typename RangeType, typename ProjectionType, typename PredicateType>
-	FORCEINLINE bool IsHeapBy(RangeType& Range, ProjectionType Projection, PredicateType Predicate)
+	FORCEINLINE bool IsHeapBy(const RangeType& Range, ProjectionType Projection, PredicateType Predicate)
 	{
 		return AlgoImpl::IsHeapInternal(GetData(Range), GetNum(Range), MoveTemp(Projection), MoveTemp(Predicate));
 	}
diff --git a/Engine/Source/Runtime/Core/Public/Algo/KahnTopologicalSort.h b/Engine/Source/Runtime/Core/Public/Algo/KahnTopologicalSort.h
--- a/Engine/Source/Runtime/Core/Public/Algo/KahnTopologicalSort.h
+++ b/Engine/Source/Runtime/Core/Public/Algo/KahnTopologicalSort.h
@@ -60,7 +60,7 @@ namespace Algo
 
 	/** Public entrypoint. Implements Algo::TopologicalSort using the Kahn Topological Sort algorithm. */
 	template <typename RangeType, typename GetElementDependenciesType>
-	bool KahnTopologicalSort(RangeType& UniqueRange, GetElementDependenciesType GetElementDependencies,
+	bool KahnTopologicalSort(RangeType&& UniqueRange, GetElementDependenciesType GetElementDependencies,
 		ETopologicalSort Flags)
 	{
 		using namespace AlgoImpl;
diff --git a/Engine/Source/Runtime/Core/Public/Algo/MaxElement.h b/Engine/Source/Runtime/Core/Public/Algo/MaxElement.h
--- a/Engine/Source/Runtime/Core/Public/Algo/MaxElement.h
+++ b/Engine/Source/Runtime/Core/Public/Algo/MaxElement.h
@@ -37,7 +37,7 @@ namespace Algo
 	 * @return A pointer to the maximum element, or nullptr if the range was empty.
 	 */
 	template <typename RangeType>
-	FORCEINLINE auto MaxElement(RangeType& Range)
+	FORCEINLINE auto MaxElement(RangeType&& Range)
 		-> decltype(AlgoImpl::MaxElementBy(Range, FIdentityFunctor(), TLess<>()))
 	{
 		return AlgoImpl::MaxElementBy(Range, FIdentityFunctor(), TLess<>());
@@ -52,7 +52,7 @@ namespace Algo
 	 * @return A pointer to the maximum element, or nullptr if the range was empty.
 	 */
 	template <typename RangeType, typename ComparatorType>
-	FORCEINLINE auto MaxElement(RangeType& Range, ComparatorType Comp)
+	FORCEINLINE auto MaxElement(RangeType&& Range, ComparatorType Comp)
 		-> decltype(AlgoImpl::MaxElementBy(Range, FIdentityFunctor(), MoveTemp(Comp)))
 	{
 		return AlgoImpl::MaxElementBy(Range, FIdentityFunctor(), MoveTemp(Comp));
@@ -67,7 +67,7 @@ namespace Algo
 	 * @return A pointer to the maximum element, or nullptr if the range was empty.
 	 */
 	template <typename RangeType, typename ProjectionType>
-	FORCEINLINE auto MaxElementBy(RangeType& Range, ProjectionType Proj)
+	FORCEINLINE auto MaxElementBy(RangeType&& Range, ProjectionType Proj)
 		-> decltype(AlgoImpl::MaxElementBy(Range, MoveTemp(Proj), TLess<>()))
 	{
 		return AlgoImpl::MaxElementBy(Range, MoveTemp(Proj), TLess<>());
@@ -83,7 +83,7 @@ namespace Algo
 	 * @return A pointer to the maximum element, or nullptr if the range was empty.
 	 */
 	template <typename RangeType, typename ProjectionType, typename ComparatorType>
-	FORCEINLINE auto MaxElementBy(RangeType& Range, ProjectionType Proj, ComparatorType Comp)
+	FORCEINLINE auto MaxElementBy(RangeType&& Range, ProjectionType Proj, ComparatorType Comp)
 		-> decltype(AlgoImpl::MaxElementBy(Range, MoveTemp(Proj), MoveTemp(Comp)))
 	{
 		return AlgoImpl::MaxElementBy(Range, MoveTemp(Proj), MoveTemp(Comp));
diff --git a/Engine/Source/Runtime/Core/Public/Algo/MinElement.h b/Engine/Source/Runtime/Core/Public/Algo/MinElement.h
--- a/Engine/Source/Runtime/Core/Public/Algo/MinElement.h
+++ b/Engine/Source/Runtime/Core/Public/Algo/MinElement.h
@@ -37,7 +37,7 @@ namespace Algo
 	 * @return A pointer to the minimum element, or nullptr if the range was empty.
 	 */
 	template <typename RangeType>
-	FORCEINLINE auto MinElement(RangeType& Range)
+	FORCEINLINE auto MinElement(RangeType&& Range)
 		-> decltype(AlgoImpl::MinElementBy(Range, FIdentityFunctor(), TLess<>()))
 	{
 		return AlgoImpl::MinElementBy(Range, FIdentityFunctor(), TLess<>());
@@ -52,7 +52,7 @@ namespace Algo
 	 * @return A pointer to the minimum element, or nullptr if the range was empty.
 	 */
 	template <typename RangeType, typename ComparatorType>
-	FORCEINLINE auto MinElement(RangeType& Range, ComparatorType Comp)
+	FORCEINLINE auto MinElement(RangeType&& Range, ComparatorType Comp)
 		-> decltype(AlgoImpl::MinElementBy(Range, FIdentityFunctor(), MoveTemp(Comp)))
 	{
 		return AlgoImpl::MinElementBy(Range, FIdentityFunctor(), MoveTemp(Comp));
@@ -67,7 +67,7 @@ namespace Algo
 	 * @return A pointer to the minimum element, or nullptr if the range was empty.
 	 */
 	template <typename RangeType, typename ProjectionType>
-	FORCEINLINE auto MinElementBy(RangeType& Range, ProjectionType Proj)
+	FORCEINLINE auto MinElementBy(RangeType&& Range, ProjectionType Proj)
 		-> decltype(AlgoImpl::MinElementBy(Range, MoveTemp(Proj), TLess<>()))
 	{
 		return AlgoImpl::MinElementBy(Range, MoveTemp(Proj), TLess<>());
@@ -83,7 +83,7 @@ namespace Algo
 	 * @return A pointer to the minimum element, or nullptr if the range was empty.
 	 */
 	template <typename RangeType, typename ProjectionType, typename ComparatorType>
-	FORCEINLINE auto MinElementBy(RangeType& Range, ProjectionType Proj, ComparatorType Comp)
+	FORCEINLINE auto MinElementBy(RangeType&& Range, ProjectionType Proj, ComparatorType Comp)
 		-> decltype(AlgoImpl::MinElementBy(Range, MoveTemp(Proj), MoveTemp(Comp)))
 	{
 		return AlgoImpl::MinElementBy(Range, MoveTemp(Proj), MoveTemp(Comp));
diff --git a/Engine/Source/Runtime/Core/Public/Algo/Mismatch.h b/Engine/Source/Runtime/Core/Public/Algo/Mismatch.h
--- a/Engine/Source/Runtime/Core/Public/Algo/Mismatch.h
+++ b/Engine/Source/Runtime/Core/Public/Algo/Mismatch.h
@@ -13,7 +13,7 @@
 namespace Algo::Private
 {
 	template <typename InputAType, typename InputBType, typename ProjectionType, typename PredicateType>
-	constexpr auto Mismatch(InputAType&& InputA, InputBType&& InputB, ProjectionType Projection, PredicateType Predicate)
+	constexpr auto Mismatch(const InputAType& InputA, const InputBType& InputB, ProjectionType Projection, PredicateType Predicate)
 	{
 		using IndexType = std::common_type_t<decltype(GetNum(InputA)), decltype(GetNum(InputB))>;
 
@@ -55,9 +55,9 @@ namespace Algo
 	 * @return The index at which the containers differ.
 	 */
 	template <typename InputAType, typename InputBType>
-	constexpr auto Mismatch(InputAType&& InputA, InputBType&& InputB) -> std::common_type_t<decltype(GetNum(InputA)), decltype(GetNum(InputB))>
+	constexpr auto Mismatch(const InputAType& InputA, const InputBType& InputB) -> std::common_type_t<decltype(GetNum(InputA)), decltype(GetNum(InputB))>
 	{
-		return Private::Mismatch(Forward<InputAType>(InputA), Forward<InputBType>(InputB), FIdentityFunctor(), TEqualTo<>());
+		return Private::Mismatch(InputA, InputB, FIdentityFunctor(), TEqualTo<>());
 	}
 
 	/**
@@ -70,9 +70,9 @@ namespace Algo
 	 * @return The index at which the containers differ.
 	 */
 	template <typename InputAType, typename InputBType, typename PredicateType>
-	constexpr auto Mismatch(InputAType&& InputA, InputBType&& InputB, PredicateType Predicate) -> std::common_type_t<decltype(GetNum(InputA)), decltype(GetNum(InputB))>
+	constexpr auto Mismatch(const InputAType& InputA, const InputBType& InputB, PredicateType Predicate) -> std::common_type_t<decltype(GetNum(InputA)), decltype(GetNum(InputB))>
 	{
-		return Private::Mismatch(Forward<InputAType>(InputA), Forward<InputBType>(InputB), FIdentityFunctor(), MoveTemp(Predicate));
+		return Private::Mismatch(InputA, InputB, FIdentityFunctor(), MoveTemp(Predicate));
 	}
 
 	/**
@@ -85,9 +85,9 @@ namespace Algo
 	 * @return The index at which the containers differ.
 	 */
 	template <typename InputAType, typename InputBType, typename ProjectionType>
-	constexpr auto MismatchBy(InputAType&& InputA, InputBType&& InputB, ProjectionType Projection) -> std::common_type_t<decltype(GetNum(InputA)), decltype(GetNum(InputB))>
+	constexpr auto MismatchBy(const InputAType& InputA, const InputBType& InputB, ProjectionType Projection) -> std::common_type_t<decltype(GetNum(InputA)), decltype(GetNum(InputB))>
 	{
-		return Private::Mismatch(Forward<InputAType>(InputA), Forward<InputBType>(InputB), MoveTemp(Projection), TEqualTo<>());
+		return Private::Mismatch(InputA, InputB, MoveTemp(Projection), TEqualTo<>());
 	}
 
 	/**
@@ -101,8 +101,8 @@ namespace Algo
 	 * @return The index at which the containers differ.
 	 */
 	template <typename InputAType, typename InputBType, typename ProjectionType, typename PredicateType>
-	constexpr auto MismatchBy(InputAType&& InputA, InputBType&& InputB, ProjectionType Projection, PredicateType Predicate) -> std::common_type_t<decltype(GetNum(InputA)), decltype(GetNum(InputB))>
+	constexpr auto MismatchBy(const InputAType& InputA, const InputBType& InputB, ProjectionType Projection, PredicateType Predicate) -> std::common_type_t<decltype(GetNum(InputA)), decltype(GetNum(InputB))>
 	{
-		return Private::Mismatch(Forward<InputAType>(InputA), Forward<InputBType>(InputB), MoveTemp(Projection), MoveTemp(Predicate));
+		return Private::Mismatch(InputA, InputB, MoveTemp(Projection), MoveTemp(Predicate));
 	}
 }
diff --git a/Engine/Source/Runtime/Core/Public/Algo/RandomShuffle.h b/Engine/Source/Runtime/Core/Public/Algo/RandomShuffle.h
--- a/Engine/Source/Runtime/Core/Public/Algo/RandomShuffle.h
+++ b/Engine/Source/Runtime/Core/Public/Algo/RandomShuffle.h
@@ -13,7 +13,7 @@ namespace Algo
 	 * @param  Range  Any contiguous container.
 	 */
 	template <typename RangeType>
-	void RandomShuffle(RangeType& Range)
+	void RandomShuffle(RangeType&& Range)
 	{
 		auto Data = GetData(Range);
 
diff --git a/Engine/Source/Runtime/Core/Public/Algo/RemoveIf.h b/Engine/Source/Runtime/Core/Public/Algo/RemoveIf.h
--- a/Engine/Source/Runtime/Core/Public/Algo/RemoveIf.h
+++ b/Engine/Source/Runtime/Core/Public/Algo/RemoveIf.h
@@ -77,7 +77,7 @@ namespace Algo
 	 * @return The index of the first element after those which were not removed.
 	 */
 	template <typename RangeType, typename Predicate>
-	[[nodiscard]] int32 StableRemoveIf(RangeType& Range, Predicate Pred)
+	[[nodiscard]] int32 StableRemoveIf(RangeType&& Range, Predicate Pred)
 	{
 		auto* First = GetData(Range);
 		auto* Last  = First + GetNum(Range);
diff --git a/Engine/Source/Runtime/Core/Public/Algo/Reverse.h b/Engine/Source/Runtime/Core/Public/Algo/Reverse.h
--- a/Engine/Source/Runtime/Core/Public/Algo/Reverse.h
+++ b/Engine/Source/Runtime/Core/Public/Algo/Reverse.h
@@ -48,7 +48,7 @@ namespace Algo
 	 * @param  Container  The container to reverse
 	 */
 	template <typename ContainerType>
-	FORCEINLINE void Reverse(ContainerType& Container)
+	FORCEINLINE void Reverse(ContainerType&& Container)
 	{
 		return AlgoImpl::Reverse(Container.GetData(), Container.Num());
 	}
diff --git a/Engine/Source/Runtime/Core/Public/Algo/Rotate.h b/Engine/Source/Runtime/Core/Public/Algo/Rotate.h
--- a/Engine/Source/Runtime/Core/Public/Algo/Rotate.h
+++ b/Engine/Source/Runtime/Core/Public/Algo/Rotate.h
@@ -56,7 +56,7 @@ namespace Algo
 	 * @return The new index of the element that was previously at the start of the range.
 	 */
 	template <typename RangeType>
-	FORCEINLINE int32 Rotate(RangeType& Range, int32 Count)
+	FORCEINLINE int32 Rotate(RangeType&& Range, int32 Count)
 	{
 		return AlgoImpl::RotateInternal(GetData(Range), GetNum(Range), Count);
 	}
diff --git a/Engine/Source/Runtime/Core/Public/Algo/TopologicalSort.h b/Engine/Source/Runtime/Core/Public/Algo/TopologicalSort.h
--- a/Engine/Source/Runtime/Core/Public/Algo/TopologicalSort.h
+++ b/Engine/Source/Runtime/Core/Public/Algo/TopologicalSort.h
@@ -28,9 +28,9 @@ namespace Algo
 	 * @return True if succeeded, false if failed due to cycles.
 	 */
 	template <typename RangeType, typename GetElementDependenciesType>
-	FORCEINLINE bool TopologicalSort(RangeType& UniqueRange, GetElementDependenciesType GetElementDependencies,
+	FORCEINLINE bool TopologicalSort(RangeType&& UniqueRange, GetElementDependenciesType GetElementDependencies,
 		ETopologicalSort Flags = ETopologicalSort::None)
 	{
-		return KahnTopologicalSort(UniqueRange, Forward<GetElementDependenciesType>(GetElementDependencies), Flags);
+		return KahnTopologicalSort(Forward<RangeType>(UniqueRange), Forward<GetElementDependenciesType>(GetElementDependencies), Flags);
 	}
 }
\ No newline at end of file
