Pak file mapping lifetime fix.

There were cases with mapped regions freed after the source mapped pak file.
Added a mapped region proxy to ensure that those are always not outliving the file handle.

Apple's FMappedFileHandle outstanding regions check is not fatal in shipping builds.
Also made NumOutstandingRegions management thread-safe and changed logging to be non-fatal.

This is a manual merge of CL 38095918, but manually merged due to unification refactoring done for IOSPlatformFile and ApplePlatformFile (CL 37960511).

#rnx

diff --git a/Engine/Source/Runtime/Core/Private/Apple/ApplePlatformFile.cpp b/Engine/Source/Runtime/Core/Private/Apple/ApplePlatformFile.cpp
--- a/Engine/Source/Runtime/Core/Private/Apple/ApplePlatformFile.cpp
+++ b/Engine/Source/Runtime/Core/Private/Apple/ApplePlatformFile.cpp
@@ -522,7 +522,7 @@ class FMappedFileHandle final : public IMappedFileHandle
 {
 	const uint8* MappedPtr;
 	FString Filename;
-	int32 NumOutstandingRegions;
+	std::atomic<int32> NumOutstandingRegions;
 	int FileHandle;
 	
 public:
@@ -541,7 +541,16 @@ public:
 	~FMappedFileHandle()
 	{
 		// can't delete the file before you delete all outstanding regions
-		UE_CLOG(NumOutstandingRegions, LogHAL, Fatal, TEXT("Cleaning mapped file with alive mapped regions."));
+		UE_CLOG(
+				NumOutstandingRegions.load(std::memory_order_relaxed),
+				LogHAL,
+#if UE_BUILD_SHIPPING
+				Error,
+#else
+				Fatal,
+#endif
+				TEXT("Cleaning mapped file with alive mapped regions: %s"),
+				*Filename);	
 		close(FileHandle);
 	}
 	
@@ -579,15 +588,18 @@ public:
 		// create a mapping for this range
 		const uint8* MapPtr = AlignedMapPtr + Offset - AlignedOffset;
 		FMappedFileRegion* Result = new FMappedFileRegion(MapPtr, AlignedMapPtr, BytesToMap, AlignedSize, Filename, Offset, this);
-		NumOutstandingRegions++;
+
+		NumOutstandingRegions.fetch_add(1, std::memory_order_relaxed);
+
 		return Result;
 	}
 	
 	void UnMap(FMappedFileRegion* Region)
 	{
 		LLM_PLATFORM_SCOPE(ELLMTag::PlatformMMIO);
-		check(NumOutstandingRegions > 0);
-		NumOutstandingRegions--;
+		
+		int32 OldNumOutstandingRegions = NumOutstandingRegions.fetch_sub(1, std::memory_order_acq_rel);
+		check(OldNumOutstandingRegions > 0);	
 		
 		LLM_IF_ENABLED(FLowLevelMemTracker::Get().OnLowLevelFree(ELLMTracker::Platform, (void*)Region->AlignedPtr));
 		int Res = munmap((void*)Region->AlignedPtr, Region->AlignedSize);
diff --git a/Engine/Source/Runtime/PakFile/Private/IPlatformFilePak.cpp b/Engine/Source/Runtime/PakFile/Private/IPlatformFilePak.cpp
--- a/Engine/Source/Runtime/PakFile/Private/IPlatformFilePak.cpp
+++ b/Engine/Source/Runtime/PakFile/Private/IPlatformFilePak.cpp
@@ -3903,6 +3903,25 @@ public:
 	}
 };
 
+class FMappedRegionPakProxy final : public IMappedFileRegion
+{
+	TSharedPtr<IMappedFileHandle> Parent;
+	TUniquePtr<IMappedFileRegion> LowerLevel;
+
+public:
+	FMappedRegionPakProxy(TSharedPtr<IMappedFileHandle> InParent, TUniquePtr<IMappedFileRegion> InLowerLevel, const FString& InDebugFilename, size_t InDebugOffsetIntoFile)
+		: IMappedFileRegion(InLowerLevel->GetMappedPtr(), InLowerLevel->GetMappedSize(), InDebugFilename, InDebugOffsetIntoFile)
+		, Parent(InParent)
+		, LowerLevel(MoveTemp(InLowerLevel))
+	{
+	}
+	
+	virtual void PreloadHint(int64 PreloadOffset = 0, int64 BytesToPreload = MAX_int64) override
+	{
+		LowerLevel->PreloadHint(PreloadOffset, BytesToPreload);
+	}
+};
+
 class FMappedFilePakProxy final : public IMappedFileHandle
 {
 	TSharedPtr<IMappedFileHandle> LowerLevel;
@@ -3927,7 +3946,8 @@ public:
 		check(BytesToMap > 0); // don't map zero bytes
 		//check(Offset + BytesToMap <= GetFileSize()); // don't map zero bytes and don't map off the end of the (virtual) file
 		//check(Offset + OffsetInPak + BytesToMap <= PakSize); // don't map zero bytes and don't map off the end of the (real) file
-		return LowerLevel->MapRegion(Offset + OffsetInPak, BytesToMap, Flags);
+		TUniquePtr<IMappedFileRegion> LowerLevelMappedRegion = TUniquePtr<IMappedFileRegion>(LowerLevel->MapRegion(Offset + OffsetInPak, BytesToMap, Flags));
+		return new FMappedRegionPakProxy(LowerLevel, MoveTemp(LowerLevelMappedRegion), DebugFilename, Offset + OffsetInPak);
 	}
 };
 
