Added a 64KB constant to the size of the persistent allocator reserve to prevent platform with transparent large pages from using 2mb pages

diff --git a/Engine/Config/BaseEngine.ini b/Engine/Config/BaseEngine.ini
--- a/Engine/Config/BaseEngine.ini
+++ b/Engine/Config/BaseEngine.ini
@@ -2335,8 +2335,7 @@ bRecompileOnLoad=true
 
 [MemoryPools]
 FLightPrimitiveInteractionInitialBlockSize=512
-; 511 to prevent platforms with transparent large pages from using 2 mb pages for linear allocator as it improved nothing
-PersistentAllocatorReserveSizeMB=511
+PersistentAllocatorReserveSizeMB=512
 
 [ProcessLimits]
 VirtualMemoryLimitInKB=0
diff --git a/Engine/Source/Runtime/Core/Private/Memory/LinearAllocator.cpp b/Engine/Source/Runtime/Core/Private/Memory/LinearAllocator.cpp
--- a/Engine/Source/Runtime/Core/Private/Memory/LinearAllocator.cpp
+++ b/Engine/Source/Runtime/Core/Private/Memory/LinearAllocator.cpp
@@ -23,7 +23,7 @@ struct FPersistentLinearAllocator : public FLinearAllocator
 		GPersistentLinearAllocatorExtends.Size = (uint64)Reserved;
 	}
 };
-
+ 
 FLinearAllocator::FLinearAllocator(SIZE_T ReserveMemorySize)
 	: Reserved(ReserveMemorySize)
 {
@@ -162,6 +162,7 @@ typedef FLinearBlockAllocator FPersistentLinearAllocator;
 
 FLinearAllocator& GetPersistentLinearAllocator()
 {
-	static FPersistentLinearAllocator GPersistentLinearAllocator(UE_PERSISTENT_ALLOCATOR_RESERVE_SIZE);
+	// We have to make sure that PersistentLinearAllocator always reserves the amount of memory that's not multiple of 2 MB as it causes issues on platforms with transparent large pages
+	static FPersistentLinearAllocator GPersistentLinearAllocator(UE_PERSISTENT_ALLOCATOR_RESERVE_SIZE + 64 * 1024);
 	return GPersistentLinearAllocator;
 }
\ No newline at end of file
