[BuildSettings/Core]
* Moved PersistentAllocatorReserveSize from Core to BuildSettings since it will change between projects and machines and we want all those things to be stored in one small module

diff --git a/Engine/Source/Runtime/BuildSettings/BuildSettings.Build.cs b/Engine/Source/Runtime/BuildSettings/BuildSettings.Build.cs
--- a/Engine/Source/Runtime/BuildSettings/BuildSettings.Build.cs
+++ b/Engine/Source/Runtime/BuildSettings/BuildSettings.Build.cs
@@ -1,8 +1,10 @@
 // Copyright Epic Games, Inc. All Rights Reserved.
 
+using EpicGames.Core;
 using UnrealBuildTool;
 using System;
 using System.Linq;
+using Microsoft.Extensions.Logging;
 
 public class BuildSettings : ModuleRules
 {
@@ -57,5 +59,34 @@ public class BuildSettings : ModuleRules
 				PrivateDefinitions.Add(String.Format("UE_LIVE_CODING_PROJECT=\"{0}\"", Target.ProjectFile.FullName.Replace("\\", "\\\\")));
 			}
 		}
+
+		PrivateDefinitions.Add("UE_PERSISTENT_ALLOCATOR_RESERVE_SIZE=" + GetPersistentAllocatorReserveSize().ToString() + "ULL");
+
+	}
+
+	private ulong ReadPersistentAllocatorReserveSize(ConfigHierarchyType ConfigHierarchy)
+	{
+		ConfigHierarchy Ini = ConfigCache.ReadHierarchy(ConfigHierarchy, DirectoryReference.FromFile(Target.ProjectFile), Target.Platform);
+		int SizeOfPermanentObjectPool = 0;
+		if (Ini.GetInt32("/Script/Engine.GarbageCollectionSettings", "gc.SizeOfPermanentObjectPool", out SizeOfPermanentObjectPool))
+		{
+			Target.Logger.LogWarning("/Script/Engine.GarbageCollectionSettings, gc.SizeOfPermanentObjectPool ini for Project {0} setting was deprecated in a favor of MemoryPools, PersistentAllocatorReserveSizeMB", Target.ProjectFile.ToString());
+		}
+		int PersistentAllocatorReserveSizeMB = 0;
+		Ini.GetInt32("MemoryPools", "PersistentAllocatorReserveSizeMB", out PersistentAllocatorReserveSizeMB);
+		return (ulong)PersistentAllocatorReserveSizeMB * 1024 * 1024;
+	}
+
+	private ulong GetPersistentAllocatorReserveSize()
+	{
+		ConfigHierarchyType ConfigHierarchy = Target.Type == TargetType.Editor ? ConfigHierarchyType.Editor : ConfigHierarchyType.Engine;
+		ulong PersistentAllocatorReserveSizeMB = ReadPersistentAllocatorReserveSize(ConfigHierarchy);
+		if (PersistentAllocatorReserveSizeMB == 0 && ConfigHierarchy == ConfigHierarchyType.Editor)
+		{
+			// if Editor doesn't define PersistentAllocatorReserveSizeMB, read the one from the Engine and simply double it
+			PersistentAllocatorReserveSizeMB = ReadPersistentAllocatorReserveSize(ConfigHierarchyType.Engine) * 2;
+		}
+
+		return PersistentAllocatorReserveSizeMB;
 	}
 }
diff --git a/Engine/Source/Runtime/BuildSettings/Private/BuildSettings.cpp b/Engine/Source/Runtime/BuildSettings/Private/BuildSettings.cpp
--- a/Engine/Source/Runtime/BuildSettings/Private/BuildSettings.cpp
+++ b/Engine/Source/Runtime/BuildSettings/Private/BuildSettings.cpp
@@ -106,4 +106,9 @@ namespace BuildSettings
 		return nullptr;
 		#endif
 	}
+
+	uint64 GetPersistentAllocatorReserveSize()
+	{
+		return UE_PERSISTENT_ALLOCATOR_RESERVE_SIZE;
+	}
 }
diff --git a/Engine/Source/Runtime/BuildSettings/Public/BuildSettings.h b/Engine/Source/Runtime/BuildSettings/Public/BuildSettings.h
--- a/Engine/Source/Runtime/BuildSettings/Public/BuildSettings.h
+++ b/Engine/Source/Runtime/BuildSettings/Public/BuildSettings.h
@@ -132,4 +132,10 @@ namespace BuildSettings
 	 * Returns the live coding project. Only set in monolithic builds
 	 */
 	BUILDSETTINGS_API const TCHAR* GetLiveCodingProject();
+
+	/**
+	 * Persistent linear allocator reserve size. Stored here to prevent rebuilds of Core when switching between projects
+	 */
+	BUILDSETTINGS_API uint64 GetPersistentAllocatorReserveSize();
+
 }
diff --git a/Engine/Source/Runtime/Core/Core.Build.cs b/Engine/Source/Runtime/Core/Core.Build.cs
--- a/Engine/Source/Runtime/Core/Core.Build.cs
+++ b/Engine/Source/Runtime/Core/Core.Build.cs
@@ -3,8 +3,6 @@
 using UnrealBuildTool;
 using System;
 using System.IO;
-using EpicGames.Core;
-using Microsoft.Extensions.Logging;
 
 public class Core : ModuleRules
 {
@@ -298,8 +296,6 @@ public class Core : ModuleRules
 			PrivateDependencyModuleNames.Add("Detours");
 		}
 
-		PrivateDefinitions.Add("UE_PERSISTENT_ALLOCATOR_RESERVE_SIZE=" + GetPersistentAllocatorReserveSize().ToString() + "ULL");
-
 		PublicDefinitions.Add("UE_MERGED_MODULES=" + (Target.bMergeModules ? "1" : "0"));
 	}
 
@@ -319,30 +315,4 @@ public class Core : ModuleRules
 			 Target.Platform == UnrealTargetPlatform.LinuxArm64 ||
 			 Target.Platform == UnrealTargetPlatform.Win64);
 	}
-
-	private ulong ReadPersistentAllocatorReserveSize(ConfigHierarchyType ConfigHierarchy)
-	{
-		ConfigHierarchy Ini = ConfigCache.ReadHierarchy(ConfigHierarchy, DirectoryReference.FromFile(Target.ProjectFile), Target.Platform);
-		int SizeOfPermanentObjectPool = 0;
-		if (Ini.GetInt32("/Script/Engine.GarbageCollectionSettings", "gc.SizeOfPermanentObjectPool", out SizeOfPermanentObjectPool))
-		{
-			Target.Logger.LogWarning("/Script/Engine.GarbageCollectionSettings, gc.SizeOfPermanentObjectPool ini for Project {0} setting was deprecated in a favor of MemoryPools, PersistentAllocatorReserveSizeMB", Target.ProjectFile.ToString());
-		}
-		int PersistentAllocatorReserveSizeMB = 0;
-		Ini.GetInt32("MemoryPools", "PersistentAllocatorReserveSizeMB", out PersistentAllocatorReserveSizeMB);
-		return (ulong)PersistentAllocatorReserveSizeMB * 1024 * 1024;
-	}
-
-	private ulong GetPersistentAllocatorReserveSize()
-	{
-		ConfigHierarchyType ConfigHierarchy = Target.Type == TargetType.Editor ? ConfigHierarchyType.Editor : ConfigHierarchyType.Engine;
-		ulong PersistentAllocatorReserveSizeMB = ReadPersistentAllocatorReserveSize(ConfigHierarchy);
-		if (PersistentAllocatorReserveSizeMB == 0 && ConfigHierarchy == ConfigHierarchyType.Editor)
-		{
-			// if Editor doesn't define PersistentAllocatorReserveSizeMB, read the one from the Engine and simply double it
-			PersistentAllocatorReserveSizeMB = ReadPersistentAllocatorReserveSize(ConfigHierarchyType.Engine) * 2;
-		}
-
-		return PersistentAllocatorReserveSizeMB;
-	}
 }
diff --git a/Engine/Source/Runtime/Core/Private/Memory/LinearAllocator.cpp b/Engine/Source/Runtime/Core/Private/Memory/LinearAllocator.cpp
--- a/Engine/Source/Runtime/Core/Private/Memory/LinearAllocator.cpp
+++ b/Engine/Source/Runtime/Core/Private/Memory/LinearAllocator.cpp
@@ -1,7 +1,7 @@
 // Copyright Epic Games, Inc. All Rights Reserved.
 
 #include "Memory/LinearAllocator.h"
-
+#include "BuildSettings.h"
 
 CORE_API FPersistentLinearAllocatorExtends GPersistentLinearAllocatorExtends;
 
@@ -163,6 +163,6 @@ typedef FLinearBlockAllocator FPersistentLinearAllocator;
 FLinearAllocator& GetPersistentLinearAllocator()
 {
 	// We have to make sure that PersistentLinearAllocator always reserves the amount of memory that's not multiple of 2 MB as it causes issues on platforms with transparent large pages
-	static FPersistentLinearAllocator GPersistentLinearAllocator(UE_PERSISTENT_ALLOCATOR_RESERVE_SIZE + 64 * 1024);
+	static FPersistentLinearAllocator GPersistentLinearAllocator(BuildSettings::GetPersistentAllocatorReserveSize() + 64 * 1024);
 	return GPersistentLinearAllocator;
 }
\ No newline at end of file
