[AutoRTFM] Register the "true" imported function for closed variants.

The standard library AutoRTFM mappings may be registered in a DLL that goes via its own import stubs.
If we don't register the "true" (redirected) functions, then lookups from other function will attempt to lookup the true function, but as this was never registered the closed lookup will fail.

diff --git a/Engine/Source/Runtime/Core/Private/AutoRTFM/FunctionMap.cpp b/Engine/Source/Runtime/Core/Private/AutoRTFM/FunctionMap.cpp
--- a/Engine/Source/Runtime/Core/Private/AutoRTFM/FunctionMap.cpp
+++ b/Engine/Source/Runtime/Core/Private/AutoRTFM/FunctionMap.cpp
@@ -11,13 +11,38 @@
 #include <stdlib.h>
 #include <string.h>
 
-#include "Containers/Array.h"
-
 #include "AutoRTFM/AutoRTFMConstants.h"
 
 namespace AutoRTFM
 {
 
+// Attempts to lookup the "true" function from a dynamically linked import function thunk.
+// Returns the pointer to the true function, or nullptr if the function cannot 
+// be resolved (or is not an import thunk).
+static inline void* FollowRelocation(void* Function)
+{
+#if PLATFORM_WINDOWS
+#if defined(_x86_64) || defined(__x86_64__) || defined(_M_X64) || defined(_M_AMD64)
+	// Note: Windows has multiple ways to perform relocations:
+	// https://learn.microsoft.com/en-us/windows/win32/debug/pe-format#base-relocation-types
+	// This handles the only relocation mode currently observed, but may need to
+	// support more modes in the future.
+	std::byte* Ptr = reinterpret_cast<std::byte*>(Function);
+	uint16 U16;
+	memcpy(&U16, Ptr, 2);
+	if (U16 == 0x25ff) // jmp qword ptr [rip+<relative-address>]
+	{
+		uint32 RelativeAddress;
+		memcpy(&RelativeAddress, Ptr+2, 4);
+		Ptr += RelativeAddress + 6; // apply relative offset
+		memcpy(&Ptr, Ptr, 8);	   // load pointer
+		return Ptr;
+	}
+#endif  // is-x64
+#endif  // PLATFORM_WINDOWS
+	return nullptr;
+}
+
 FFunctionMap* AllocateFunctionMap(size_t Size)
 {
     ASSERT(Size && !(Size & (Size - 1)));
@@ -88,45 +113,28 @@ void FunctionMapAddImpl(void* OldFunction, void* NewFunction)
 
 void FunctionMapAdd(void* OldFunction, void* NewFunction)
 {
-    if (!OldFunction)
-    {
-        // Silently ignore null OldFunction to make weak linking easy.
-        return;
-    }
+	if (!OldFunction)
+	{
+		// Silently ignore null OldFunction to make weak linking easy.
+		return;
+	}
 
-    InitializeGlobalDataIfNecessary();
+	InitializeGlobalDataIfNecessary();
 
-    TLockGuard<FFastLock> LockGuard(GlobalData->FunctionMapLock);
-    TArray<void*> Functions;
-    Functions.Push(OldFunction);
-    while (!Functions.IsEmpty())
-    {
-        void* Function = Functions.Last();
-        Functions.Pop();
-        FunctionMapAddImpl(Function, NewFunction);
-    }
+	TLockGuard<FFastLock> LockGuard(GlobalData->FunctionMapLock);
+	FunctionMapAddImpl(OldFunction, NewFunction);
+	if (void* OldRelocated = FollowRelocation(OldFunction))
+	{
+		FunctionMapAddImpl(OldRelocated, NewFunction);
+	}
 }
 
 void* FunctionMapLookupTryFollowingRelocation(void* Function, const char* Where)
 {
-#if defined(_x86_64) || defined(__x86_64__) || defined(_M_X64) || defined(_M_AMD64)
-	// Note: Windows has multiple ways to perform relocations:
-	// https://learn.microsoft.com/en-us/windows/win32/debug/pe-format#base-relocation-types
-	// This handles the only relocation mode currently observed, but may need to
-	// support more modes in the future.
-	std::byte* Ptr = reinterpret_cast<std::byte*>(Function);
-	uint16 U16;
-	memcpy(&U16, Ptr, 2);
-	if (U16 == 0x25ff) // jmp qword ptr [rip+<relative-address>]
+	if (void* Ptr = FollowRelocation(Function); LIKELY(Ptr))
 	{
-		uint32 RelativeAddress;
-		memcpy(&RelativeAddress, Ptr+2, 4);
-		Ptr += RelativeAddress + 6; // apply relative offset
-		memcpy(&Ptr, Ptr, 8);       // load pointer
 		AUTORTFM_MUST_TAIL return FunctionMapLookup(Ptr, Where);
 	}
-#endif
-
 	AUTORTFM_MUST_TAIL return FunctionMapReportError(Function, Where);
 }
 
