CocoaThread - Remove confusions that we are asking ThreadSource to do anything.
When using CFRunLoopRunInMode we are asking the current thread's RunLoop to process element(s) and not one from another thread.

#tests Local Testing, this is functionally identical

diff --git a/Engine/Source/Runtime/Core/Private/Mac/CocoaThread.cpp b/Engine/Source/Runtime/Core/Private/Mac/CocoaThread.cpp
--- a/Engine/Source/Runtime/Core/Private/Mac/CocoaThread.cpp
+++ b/Engine/Source/Runtime/Core/Private/Mac/CocoaThread.cpp
@@ -104,11 +104,6 @@ public:
 		CFRunLoopWakeUp(TargetRunLoop);
 	}
 	
-	void RunInMode(CFStringRef WaitMode)
-	{
-		CFRunLoopRunInMode(WaitMode, 0, true);
-	}
-	
 	void Process(CFStringRef Mode)
 	{
 		TArray<FCocoaRunLoopTask*> NewTasks;
@@ -421,7 +416,7 @@ static void PerformBlockOnThread(FCocoaRunLoopSource& ThreadSource, dispatch_blo
 		
 		do {
 			ThreadSource.Wake();
-			ThreadSource.RunInMode( (CFStringRef)WaitMode );
+			CFRunLoopRunInMode((CFStringRef)WaitMode, 0, true);
 		} while (dispatch_semaphore_wait(Semaphore, dispatch_time(0, 100000ull)));
 		
 		Block_release(ExecuteBlock);
