Stats - Support flags in lightweight stats

This re-uses existing declaration macros (that were previously only utilized when STATS is 1) to define lightweight structs containing information about the stat such as the flags and group. This allows us to do things like properly filter out verbose stats when profiling in the Test config (where STATS is 0).

#tests replay run, local unit tests, local traces
[RN] Support verbosity flags in lightweight stat macros

diff --git a/Engine/Plugins/Animation/MLDeformer/MLDeformerFramework/Source/MLDeformerFramework/Public/MLDeformerModelInstance.h b/Engine/Plugins/Animation/MLDeformer/MLDeformerFramework/Source/MLDeformerFramework/Public/MLDeformerModelInstance.h
--- a/Engine/Plugins/Animation/MLDeformer/MLDeformerFramework/Source/MLDeformerFramework/Public/MLDeformerModelInstance.h
+++ b/Engine/Plugins/Animation/MLDeformer/MLDeformerFramework/Source/MLDeformerFramework/Public/MLDeformerModelInstance.h
@@ -13,10 +13,8 @@ class UMLDeformerModel;
 class USkeletalMeshComponent;
 class UMLDeformerComponent;
 
-#if STATS
 DECLARE_STATS_GROUP(TEXT("MLDeformer"), STATGROUP_MLDeformer, STATCAT_Advanced);
 DECLARE_CYCLE_STAT_EXTERN(TEXT("ML Deformer Inference"), STAT_MLDeformerInference, STATGROUP_MLDeformer, );
-#endif
 /**
  * An instance of the ML Deformer model.
  * The ML Deformer model contains shared data, while this instance contains data unique to the actor it is being applied to.
diff --git a/Engine/Plugins/Compression/OodleNetwork/Source/Public/OodleNetworkHandlerComponent.h b/Engine/Plugins/Compression/OodleNetwork/Source/Public/OodleNetworkHandlerComponent.h
--- a/Engine/Plugins/Compression/OodleNetwork/Source/Public/OodleNetworkHandlerComponent.h
+++ b/Engine/Plugins/Compression/OodleNetwork/Source/Public/OodleNetworkHandlerComponent.h
@@ -42,8 +42,6 @@ enum class EOodleNetworkEnableMode : uint8
 
 /** Stats */
 
-#if STATS
-
 #if !UE_BUILD_SHIPPING
 DECLARE_DWORD_ACCUMULATOR_STAT_EXTERN(TEXT("Reserved Oodle (bits)"), STAT_PacketReservedOodle, STATGROUP_Packet, );
 #endif
@@ -77,8 +75,6 @@ DECLARE_MEMORY_STAT_EXTERN(TEXT("Oodle Dictionary Bytes"), STAT_Oodle_Dictionary
 DECLARE_MEMORY_STAT_EXTERN(TEXT("Oodle Shared Bytes"), STAT_Oodle_SharedBytes, STATGROUP_OodleNetwork, );
 DECLARE_MEMORY_STAT_EXTERN(TEXT("Oodle State Bytes"), STAT_Oodle_StateBytes, STATGROUP_OodleNetwork, );
 
-#endif // STATS
-
 
 /** Globals */
 
diff --git a/Engine/Source/Runtime/Core/Public/Stats/LightweightStats.h b/Engine/Source/Runtime/Core/Public/Stats/LightweightStats.h
--- a/Engine/Source/Runtime/Core/Public/Stats/LightweightStats.h
+++ b/Engine/Source/Runtime/Core/Public/Stats/LightweightStats.h
@@ -20,32 +20,65 @@ struct TStatId
 {
 	const PROFILER_CHAR* StatString;
 
-	FORCEINLINE TStatId()
+	FORCEINLINE constexpr TStatId()
 		: StatString(nullptr)
 	{
 	}
 
-	FORCEINLINE TStatId(const PROFILER_CHAR* InString)
+	FORCEINLINE constexpr TStatId(const PROFILER_CHAR* InString)
 		: StatString(InString)
 	{
 	}
 
-	FORCEINLINE bool IsValidStat() const
+	FORCEINLINE constexpr bool IsValidStat() const
 	{
 		return StatString != nullptr;
 	}
 
-	FORCEINLINE bool operator==(TStatId Other) const
+	FORCEINLINE constexpr bool operator==(TStatId Other) const
 	{
 		return StatString == Other.StatString;
 	}
 
-	FORCEINLINE bool operator!=(TStatId Other) const
+	FORCEINLINE constexpr bool operator!=(TStatId Other) const
 	{
 		return StatString != Other.StatString;
 	}
 };
 
+namespace UE::Stats::Private
+{
+	// We cannot add members to TStatId as one is stored per uobject which results in a large memory overhead,
+	// so instead we utilze a proxy object internally that can implictly convert back into a TStatId when needed.
+	struct TStatIdProxy
+	{
+		TStatId StatId;
+		const PROFILER_CHAR* Group = nullptr;
+		EStatFlags Flags = EStatFlags::None;
+
+		constexpr TStatIdProxy() = default;
+
+		// Allow implicit conversions from nullptr to support conditional scopes
+		constexpr TStatIdProxy(TYPE_OF_NULLPTR)
+		{
+		}
+
+		constexpr explicit TStatIdProxy(const PROFILER_CHAR* InString, const PROFILER_CHAR* InGroup = nullptr, EStatFlags InFlags = EStatFlags::None)
+			: StatId(InString)
+			, Group(InGroup)
+			, Flags(InFlags)
+		{
+		}
+
+		// Allow implicit conversions to TStatId to keep macros like GET_STATID compatible
+		// with code that takes a TStatId.
+		constexpr operator TStatId() const
+		{
+			return StatId;
+		}
+	};
+} // namespace UE::Stats::Private
+
 class FScopeCycleCounter
 {
 public:
@@ -57,7 +90,9 @@ public:
 
 		bPop(false)
 	{
-		if (GCycleStatsShouldEmitNamedEvents && InStatId.IsValidStat())
+		if (GCycleStatsShouldEmitNamedEvents &&
+			InStatId.IsValidStat() &&
+			(GShouldEmitVerboseNamedEvents || !EnumHasAnyFlags(StatFlags, EStatFlags::Verbose)))
 		{
 			// Memory validation disabled as this is hot code.
 			UE_AUTORTFM_OPEN_NO_VALIDATION
@@ -74,6 +109,11 @@ public:
 	{
 	}
 
+	FORCEINLINE explicit FScopeCycleCounter(const UE::Stats::Private::TStatIdProxy& InStatIdProxy)
+		: FScopeCycleCounter(InStatIdProxy.StatId, InStatIdProxy.Flags, false)
+	{
+	}
+
 	FORCEINLINE ~FScopeCycleCounter()
 	{
 		AutoRTFM::PopOnAbortHandler(this);
@@ -94,6 +134,93 @@ private:
 	bool bPop;
 };
 
+// ----------------------------------------------------------------
+// Internal Helpers
+// ----------------------------------------------------------------
+namespace UE::Stats::Private
+{
+	// Fallback stat data for StatId's that don't have a custom stat struct defined.
+	struct FDefaultStatData
+	{
+		static constexpr const PROFILER_CHAR* GetGroup()
+		{
+			return nullptr;
+		}
+		static constexpr EStatFlags GetFlags()
+		{
+			return EStatFlags::None;
+		}
+	};
+
+	template<typename StatDataType, bool AllowDefault, typename = void>
+	struct TStatDataHelper
+	{
+		// This error means a DECLARE_* macro is missing for this stat for this compile configuration.
+		// The Declaration macros can be found in Stats.h.
+		static_assert(AllowDefault, "Default StatId not permitted. Please provide a declaration for this stat via one of the DECLARE_ Stat macros.");
+		using Type = FDefaultStatData;
+	};
+	// If StatDataType is a complete type such that we can call sizeof on it, this specialization is used
+	template<typename StatDataType, bool AllowDefault>
+	struct TStatDataHelper<StatDataType, AllowDefault, std::void_t<decltype(sizeof(StatDataType))>>
+	{
+		using Type = StatDataType;
+	};
+
+	// Helper defined as either the custom stat data struct if it's defined, or the default (FDefaultStatData).
+	template<typename StatDataType, bool AllowDefault>
+	using TStatDataType = typename TStatDataHelper<StatDataType, AllowDefault>::Type;
+
+
+	template<typename AllowDefaultStatToken, typename = void>
+	struct TAllowDefaultStatHelper
+	{
+		static constexpr bool Value = false;
+	};
+	template<typename AllowDefaultStatToken>
+	struct TAllowDefaultStatHelper<AllowDefaultStatToken, std::void_t<decltype(sizeof(AllowDefaultStatToken))>>
+	{
+		// A token has been defined for this stat, so it's allowed.
+		static constexpr bool Value = true;
+	};
+
+	// Determines if falling back to FDefaultStatData is valid for a given stat.
+	// To be allowed, INTERNAL_ALLOW_DEFAULT_STAT must be defined for the stat.
+	template<typename AllowDefaultStatToken>
+	constexpr bool AllowDefaultForStat = TAllowDefaultStatHelper<AllowDefaultStatToken>::Value;
+
+	template<typename StatDataType>
+	constexpr TStatIdProxy MakeStatId(const PROFILER_CHAR* StatName)
+	{
+		return TStatIdProxy(StatName, StatDataType::GetGroup(), StatDataType::GetFlags());
+	}
+} // namespace UE::Stats::Private
+
+// ----------------------------------------------------------------
+// Helper Macros
+// ----------------------------------------------------------------
+
+// Specifies that a custom FStat struct doesn't need to be defined for this stat and we can fallback to the default.
+// This is for internal use only currently.
+#define UE_INTERNAL_ALLOW_DEFAULT_STAT(Stat)\
+	struct FAllowDefaultStat_##Stat{};
+
+// Returns either the custom stat data struct for a stat, or FDefaultStatData if one doesn't exist.
+#define UE_INTERNAL_GET_STATDATA(Stat)\
+	UE::Stats::Private::TStatDataType<struct FStat_##Stat, UE::Stats::Private::AllowDefaultForStat<struct FAllowDefaultStat_##Stat>>
+
+// Returns either the custom stat data struct for a stat, or FDefaultStatData if one doesn't exist.
+// This version will not emit a compilation error if a default is not allowed to handled special cases.
+#define UE_INTERNAL_GET_STATDATA_ALLOW_DEFAULT(Stat)\
+	UE::Stats::Private::TStatDataType<struct FStat_##Stat, true>
+
+// Helper to evaluate if a stat is marked as verbose.
+#define UE_IS_STAT_VERBOSE(Stat)\
+	EnumHasAnyFlags(UE_INTERNAL_GET_STATDATA(Stat)::GetFlags(), EStatFlags::Verbose)
+
+#define UE_IS_STAT_VERBOSE_ALLOW_DEFAULT(Stat)\
+	EnumHasAnyFlags(UE_INTERNAL_GET_STATDATA_ALLOW_DEFAULT(Stat)::GetFlags(), EStatFlags::Verbose)
+
 #elif !STATS // UE_USE_LIGHTWEIGHT_STATS
 
 struct TStatId {};
diff --git a/Engine/Source/Runtime/Core/Public/Stats/Stats.h b/Engine/Source/Runtime/Core/Public/Stats/Stats.h
--- a/Engine/Source/Runtime/Core/Public/Stats/Stats.h
+++ b/Engine/Source/Runtime/Core/Public/Stats/Stats.h
@@ -500,33 +500,102 @@ struct FStat_##StatName\
 // profiling namedevents if enabled.
 // ----------------------------------------------------------------------------
 
+// ----------------------------------------------------------------
+// Stat Declarations
+// ----------------------------------------------------------------
+
+#define DECLARE_STAT(StatName, GroupName, StatFlags)\
+	struct FStat_##StatName\
+	{\
+		static constexpr const PROFILER_CHAR* GetGroup()\
+		{\
+			return ANSI_TO_PROFILING(#GroupName);\
+		}\
+		static constexpr EStatFlags GetFlags()\
+		{\
+			return StatFlags;\
+		}\
+	};
+
+// Replace declarations with lightweight instances
+#define DECLARE_CYCLE_STAT(CounterName, StatId, GroupId) \
+	DECLARE_STAT(StatId, GroupId, EStatFlags::None);
+
+#define DECLARE_CYCLE_STAT_WITH_FLAGS(CounterName, StatId, GroupId, StatFlags) \
+	DECLARE_STAT(StatId, GroupId, StatFlags);
+
+#define DECLARE_CYCLE_STAT_EXTERN(CounterName, StatId, GroupId, APIX) \
+	DECLARE_STAT(StatId, GroupId, EStatFlags::None);
+
+#define DECLARE_CYCLE_STAT_WITH_FLAGS_EXTERN(CounterName, StatId, GroupId, StatFlags, APIX) \
+	DECLARE_STAT(StatId, GroupId, StatFlags);
+
+#define DECLARE_MEMORY_STAT_EXTERN(CounterName, StatId, GroupId, API)\
+	DECLARE_STAT(StatId, GroupId, EStatFlags::None);
+
+#define DECLARE_MEMORY_STAT_POOL_EXTERN(CounterName, StatId, GroupId, Pool, API)\
+	DECLARE_STAT(StatId, GroupId, EStatFlags::None);
+
+#define DECLARE_MEMORY_STAT(CounterName,StatId,GroupId)\
+	DECLARE_STAT(StatId, GroupId, EStatFlags::None);
+
+#define DECLARE_MEMORY_STAT_POOL(CounterName,StatId,GroupId,Pool)\
+	DECLARE_STAT(StatId, GroupId, EStatFlags::None);
+
+// Some systems like the GPU profiler utilize counter declarations to declare stats used with SCOPE_CYCLE_COUNTER.
+// Since this is a special case, allow falling back to the default stat data for these.
+#define DECLARE_FLOAT_COUNTER_STAT(CounterName,StatId,GroupId)					UE_INTERNAL_ALLOW_DEFAULT_STAT(StatId);
+#define DECLARE_FLOAT_COUNTER_STAT_EXTERN(CounterName,StatId,GroupId, API)		UE_INTERNAL_ALLOW_DEFAULT_STAT(StatId);
+#define DECLARE_FLOAT_ACCUMULATOR_STAT(CounterName,StatId,GroupId)				UE_INTERNAL_ALLOW_DEFAULT_STAT(StatId);
+#define DECLARE_FLOAT_ACCUMULATOR_STAT_EXTERN(CounterName,StatId,GroupId, API)	UE_INTERNAL_ALLOW_DEFAULT_STAT(StatId);
+#define DECLARE_DWORD_COUNTER_STAT(CounterName,StatId,GroupId)					UE_INTERNAL_ALLOW_DEFAULT_STAT(StatId);
+#define DECLARE_DWORD_COUNTER_STAT_EXTERN(CounterName,StatId,GroupId, API)		UE_INTERNAL_ALLOW_DEFAULT_STAT(StatId);
+#define DECLARE_DWORD_ACCUMULATOR_STAT(CounterName,StatId,GroupId)				UE_INTERNAL_ALLOW_DEFAULT_STAT(StatId);
+#define DECLARE_DWORD_ACCUMULATOR_STAT_EXTERN(CounterName,StatId,GroupId, API)	UE_INTERNAL_ALLOW_DEFAULT_STAT(StatId);
+#define DECLARE_FNAME_STAT(CounterName,StatId,GroupId, API)						UE_INTERNAL_ALLOW_DEFAULT_STAT(StatId);
+#define DECLARE_PTR_STAT(CounterName,StatId,GroupId)							UE_INTERNAL_ALLOW_DEFAULT_STAT(StatId);
+#define DECLARE_FNAME_STAT_EXTERN(CounterName,StatId,GroupId, API)				UE_INTERNAL_ALLOW_DEFAULT_STAT(StatId);
+#define DECLARE_PTR_STAT_EXTERN(CounterName,StatId,GroupId, API)				UE_INTERNAL_ALLOW_DEFAULT_STAT(StatId);
+
+// ----------------------------------------------------------------
+// Scope Counters
+// ----------------------------------------------------------------
+
 #define SCOPE_CYCLE_COUNTER_TO_TRACE(StatString, StatName, Condition) \
 	TRACE_CPUPROFILER_EVENT_DECLARE(PREPROCESSOR_JOIN(PREPROCESSOR_JOIN(__Decl_, StatName), __LINE__)); \
 	TRACE_CPUPROFILER_EVENT_SCOPE_USE(PREPROCESSOR_JOIN(PREPROCESSOR_JOIN(__Decl_, StatName), __LINE__), StatString, PREPROCESSOR_JOIN(PREPROCESSOR_JOIN(__Scope_, StatName), __LINE__), Condition && GCycleStatsShouldEmitNamedEvents);
 
+// Note: Since these are defining the stat inline, no custom stat data will exist.
 #define DECLARE_SCOPE_CYCLE_COUNTER(CounterName,Stat,GroupId) \
-	FScopeCycleCounter StatNamedEventsScope_##Stat(TStatId(ANSI_TO_PROFILING(#Stat))); \
+	FScopeCycleCounter StatNamedEventsScope_##Stat(UE::Stats::Private::TStatIdProxy(ANSI_TO_PROFILING(#Stat), ANSI_TO_PROFILING(#GroupId))); \
 	SCOPE_CYCLE_COUNTER_TO_TRACE(CounterName, Stat, true);
 
 #define QUICK_SCOPE_CYCLE_COUNTER(Stat) \
-	FScopeCycleCounter StatNamedEventsScope_##Stat(TStatId(ANSI_TO_PROFILING(#Stat))); \
+	FScopeCycleCounter StatNamedEventsScope_##Stat(UE::Stats::Private::TStatIdProxy(ANSI_TO_PROFILING(#Stat))); \
 	SCOPE_CYCLE_COUNTER_TO_TRACE(#Stat, Stat, true);
 
+// Note: FScopeCycleCounter is what emits the trace events to external profilers
+// while SCOPE_CYCLE_COUNTER_TO_TRACE handles emitting it to insights.
+// FScopeCycleCounter handles verbosity checks itself.
 #define SCOPE_CYCLE_COUNTER(Stat) \
-	FScopeCycleCounter StatNamedEventsScope_##Stat(TStatId(ANSI_TO_PROFILING(#Stat))); \
-	SCOPE_CYCLE_COUNTER_TO_TRACE(#Stat, Stat, true);
+	FScopeCycleCounter StatNamedEventsScope_##Stat(GET_STATID(Stat)); \
+	SCOPE_CYCLE_COUNTER_TO_TRACE(#Stat, Stat, (GShouldEmitVerboseNamedEvents || !UE_IS_STAT_VERBOSE(Stat)));
 
+// Note: This is often used with non declared stats so we always allow default stat data.
 #define SCOPE_CYCLE_COUNTER_STATID(StatId) \
 	FScopeCycleCounter StatNamedEventsScope_STATID(StatId); \
-	TRACE_CPUPROFILER_EVENT_SCOPE_TEXT_CONDITIONAL(StatId.StatString, StatId.StatString != nullptr && GCycleStatsShouldEmitNamedEvents > 0);
+	TRACE_CPUPROFILER_EVENT_SCOPE_TEXT_CONDITIONAL(StatId.StatString, StatId.StatString != nullptr && GCycleStatsShouldEmitNamedEvents > 0 && (GShouldEmitVerboseNamedEvents || !UE_IS_STAT_VERBOSE_ALLOW_DEFAULT(Stat)));
 
 #define CONDITIONAL_SCOPE_CYCLE_COUNTER(Stat,bCondition) \
-	FScopeCycleCounter StatNamedEventsScope_##Stat(bCondition ? ANSI_TO_PROFILING(#Stat) : nullptr); \
-	SCOPE_CYCLE_COUNTER_TO_TRACE(#Stat, Stat, bCondition);
+	FScopeCycleCounter StatNamedEventsScope_##Stat(bCondition ? GET_STATID(Stat) : nullptr); \
+	SCOPE_CYCLE_COUNTER_TO_TRACE(#Stat, Stat, bCondition && (GShouldEmitVerboseNamedEvents || !UE_IS_STAT_VERBOSE(Stat)));
 
-#define RETURN_QUICK_DECLARE_CYCLE_STAT(StatId,GroupId) return TStatId(ANSI_TO_PROFILING(#StatId));
+#define RETURN_QUICK_DECLARE_CYCLE_STAT(StatId,GroupId) return QUICK_USE_CYCLE_STAT(StatId, GroupId);
 
-#define GET_STATID(Stat) (TStatId(ANSI_TO_PROFILING(#Stat)))
+#define QUICK_USE_CYCLE_STAT(StatId,GroupId) UE::Stats::Private::TStatIdProxy(ANSI_TO_PROFILING(#StatId), ANSI_TO_PROFILING(#GroupId))
+
+// Creates a TStatIdProxy on the fly, optionally populating the group/flags from the custom stat data struct if it exists.
+#define GET_STATID(Stat) (UE::Stats::Private::MakeStatId<UE_INTERNAL_GET_STATDATA(Stat)>(ANSI_TO_PROFILING(#Stat)))
 
 
 #elif USE_LIGHTWEIGHT_STATS_FOR_HITCH_DETECTION && USE_HITCH_DETECTION
@@ -580,6 +649,8 @@ struct FStat_##StatName\
 #define SCOPE_SECONDS_ACCUMULATOR(Stat)
 #define SCOPE_MS_ACCUMULATOR(Stat)
 #define DEFINE_STAT(Stat)
+
+#if !UE_USE_LIGHTWEIGHT_STATS
 #define QUICK_USE_CYCLE_STAT(StatId,GroupId) TStatId()
 #define DECLARE_CYCLE_STAT(CounterName,StatId,GroupId)
 #define DECLARE_CYCLE_STAT_WITH_FLAGS(CounterName,StatId,GroupId,StatFlags)
@@ -601,6 +672,7 @@ struct FStat_##StatName\
 #define DECLARE_PTR_STAT_EXTERN(CounterName,StatId,GroupId, API)
 #define DECLARE_MEMORY_STAT_EXTERN(CounterName,StatId,GroupId, API)
 #define DECLARE_MEMORY_STAT_POOL_EXTERN(CounterName,StatId,GroupId,Pool, API)
+#endif // !UE_USE_LIGHTWEIGHT_STATS
 
 #define DECLARE_STATS_GROUP(GroupDesc,GroupId,GroupCat)
 #define DECLARE_STATS_GROUP_VERBOSE(GroupDesc,GroupId,GroupCat)
diff --git a/Engine/Source/Runtime/CoreUObject/Public/UObject/UObjectBaseUtility.h b/Engine/Source/Runtime/CoreUObject/Public/UObject/UObjectBaseUtility.h
--- a/Engine/Source/Runtime/CoreUObject/Public/UObject/UObjectBaseUtility.h
+++ b/Engine/Source/Runtime/CoreUObject/Public/UObject/UObjectBaseUtility.h
@@ -1022,7 +1022,7 @@ public:
 #define SCOPE_CYCLE_UOBJECT(Name, Object) \
 	FScopeCycleCounterUObject ObjCycleCount_##Name(Object);
 
-#elif ENABLE_STATNAMEDEVENTS
+#elif UE_USE_LIGHTWEIGHT_STATS
 
 class FScopeCycleCounterUObject
 {
diff --git a/Engine/Source/Runtime/SlateCore/Public/Widgets/SWidgetUtils.h b/Engine/Source/Runtime/SlateCore/Public/Widgets/SWidgetUtils.h
--- a/Engine/Source/Runtime/SlateCore/Public/Widgets/SWidgetUtils.h
+++ b/Engine/Source/Runtime/SlateCore/Public/Widgets/SWidgetUtils.h
@@ -51,7 +51,7 @@ struct FScopeCycleCounterSWidget : public FCycleCounter
 #define SCOPE_CYCLE_SWIDGET(Object) \
 	FScopeCycleCounterSWidget ANONYMOUS_VARIABLE(SlateWidgetCycleCount_) (Object);
 
-#elif ENABLE_STATNAMEDEVENTS
+#elif UE_USE_LIGHTWEIGHT_STATS
 
 struct FScopeCycleCounterSWidget
 {
