[AutoRTFM] Log with Fatal severity for memory validation errors with -AutoRTFMAssertOnError=1

diff --git a/Engine/Source/Runtime/Core/Private/AutoRTFM/Transaction.cpp b/Engine/Source/Runtime/Core/Private/AutoRTFM/Transaction.cpp
--- a/Engine/Source/Runtime/Core/Private/AutoRTFM/Transaction.cpp
+++ b/Engine/Source/Runtime/Core/Private/AutoRTFM/Transaction.cpp
@@ -450,16 +450,23 @@ void FTransaction::ValidateWriteHash() const
 	static constexpr TCHAR Message[] =
 		TEXT("Memory modified in a transaction was also modified in an call to AutoRTFM::Open(). ")
 		TEXT("This may lead to memory corruption if the transaction is aborted.");
-	if (CurrentMemoryValidationLevel == EMemoryValidationLevel::Warn)
+	if (UNLIKELY(OldHash != NewHash))
 	{
-		if (UNLIKELY(OldHash != NewHash))
+		if (CurrentMemoryValidationLevel == EMemoryValidationLevel::Warn)
 		{
 			UE_LOG(LogAutoRTFM, Warning, Message);
 		}
-	}
-	else
-	{
-		ensureMsgf(OldHash == NewHash, Message);
+		else
+		{
+			if (ForTheRuntime::IsAutoRTFMAssertOnError())
+			{
+				UE_LOG(LogAutoRTFM, Fatal, Message);
+			}
+			else
+			{
+				ensureMsgf(OldHash == NewHash, Message);
+			}
+		}
 	}
 }
 
