Make the `Key` of on-commit and on-abort handlers be checked for stack-locality before running.

This is crucial because we could have a stack-local object that we are attempting to access that has already gone away on the stack!

diff --git a/Engine/Source/Programs/AutoRTFMTests/Private/Abort.cpp b/Engine/Source/Programs/AutoRTFMTests/Private/Abort.cpp
--- a/Engine/Source/Programs/AutoRTFMTests/Private/Abort.cpp
+++ b/Engine/Source/Programs/AutoRTFMTests/Private/Abort.cpp
@@ -684,6 +684,67 @@ TEST_CASE("Abort.PushOnAbortHandler_Order")
 			REQUIRE(42 == Value);
 		}
 	}
+
+	SECTION("Handler Ran Because Key Is Stack Local To Current Aborting Transaction")
+	{
+		bool bHit = false;
+		AutoRTFM::Testing::Abort([&]
+		{
+			int Foo = 42;
+			AutoRTFM::PushOnAbortHandler(&Foo, [&] { bHit = true; });
+			AutoRTFM::AbortTransaction();
+		});
+		REQUIRE(bHit);
+	}
+
+	SECTION("Handler Ran Because Key Is Stack Local In Aborting Child")
+	{
+		bool bHit = false;
+		AutoRTFM::Testing::Commit([&]
+		{
+			AutoRTFM::Testing::Abort([&]
+			{
+				int Foo = 42;
+				AutoRTFM::PushOnAbortHandler(&Foo, [&] { bHit = true; });	
+				AutoRTFM::AbortTransaction();
+			});
+		});
+		REQUIRE(bHit);
+	}
+
+	SECTION("Handler Ran Because Key Is Stack Local In Aborting Parent")
+	{
+		// Need to disable retries because the inner commit would retry and run the on-abort handler.
+		AutoRTFMTestUtils::FScopedRetry _(AutoRTFM::ForTheRuntime::EAutoRTFMRetryTransactionState::NoRetry);
+
+		bool bHit = false;
+		AutoRTFM::Testing::Abort([&]
+		{
+			int Foo = 42;
+			AutoRTFM::Testing::Commit([&]
+			{
+				AutoRTFM::PushOnAbortHandler(&Foo, [&] { bHit = true; });	
+			});
+			AutoRTFM::AbortTransaction();
+		});
+		REQUIRE(bHit);
+	}
+
+	SECTION("Handler Ignored Because Key Is Stack Local In Committing Child")
+	{
+		// Need to disable retries because the inner commit would retry and run the on-abort handler.
+		AutoRTFMTestUtils::FScopedRetry _(AutoRTFM::ForTheRuntime::EAutoRTFMRetryTransactionState::NoRetry);
+
+		AutoRTFM::Testing::Abort([&]
+		{
+			AutoRTFM::Testing::Commit([&]
+			{
+				int Foo = 42;
+				AutoRTFM::PushOnAbortHandler(&Foo, [] { FAIL("Unreachable!"); });	
+			});
+			AutoRTFM::AbortTransaction();
+		});
+	}
 }
 
 TEST_CASE("Abort.OnAbortTiming")
diff --git a/Engine/Source/Programs/AutoRTFMTests/Private/Commit.cpp b/Engine/Source/Programs/AutoRTFMTests/Private/Commit.cpp
--- a/Engine/Source/Programs/AutoRTFMTests/Private/Commit.cpp
+++ b/Engine/Source/Programs/AutoRTFMTests/Private/Commit.cpp
@@ -1,6 +1,7 @@
 // Copyright Epic Games, Inc. All Rights Reserved.
 
 #include "Catch2Includes.h"
+#include "Utils.h"
 #include <AutoRTFM/AutoRTFMTesting.h>
 #include <map>
 #include <vector>
@@ -424,4 +425,37 @@ TEST_CASE("Commit.PushOnCommitHandler_Order")
 			REQUIRE(42 == Value);
 		}
 	}
+
+	SECTION("Handler Ignored Because Key Is Stack Local")
+	{
+		AutoRTFM::Testing::Commit([&]
+		{
+			int Foo = 42;
+			AutoRTFM::PushOnCommitHandler(&Foo, [] { FAIL("Unreachable!"); });
+		});
+	}
+
+	SECTION("Handler Ignored Because Key Is Stack Local In Committing Child")
+	{
+		AutoRTFM::Testing::Commit([&]
+		{
+			AutoRTFM::Testing::Commit([&]
+			{
+				int Foo = 42;
+				AutoRTFM::PushOnCommitHandler(&Foo, [] { FAIL("Unreachable!"); });	
+			});
+		});
+	}
+
+	SECTION("Handler Ignored Because Key Is Stack Local In Parent")
+	{
+		AutoRTFM::Testing::Commit([&]
+		{
+			int Foo = 42;
+			AutoRTFM::Testing::Commit([&]
+			{
+				AutoRTFM::PushOnCommitHandler(&Foo, [] { FAIL("Unreachable!"); });	
+			});
+		});
+	}
 }
diff --git a/Engine/Source/Programs/AutoRTFMTests/Private/Utils.h b/Engine/Source/Programs/AutoRTFMTests/Private/Utils.h
--- a/Engine/Source/Programs/AutoRTFMTests/Private/Utils.h
+++ b/Engine/Source/Programs/AutoRTFMTests/Private/Utils.h
@@ -10,9 +10,9 @@ namespace AutoRTFMTestUtils
 {
 
 // Temporarily changes the AutoRTFM retry mode for the lifetime of the FScopedRetry object.
-struct FScopedRetry
+struct FScopedRetry final
 {
-	FScopedRetry(AutoRTFM::ForTheRuntime::EAutoRTFMRetryTransactionState NewRetry)
+	explicit FScopedRetry(AutoRTFM::ForTheRuntime::EAutoRTFMRetryTransactionState NewRetry)
 		: OldRetry(AutoRTFM::ForTheRuntime::GetRetryTransaction())
 	{
 		AutoRTFM::ForTheRuntime::SetRetryTransaction(NewRetry);
@@ -23,6 +23,9 @@ struct FScopedRetry
 		AutoRTFM::ForTheRuntime::SetRetryTransaction(OldRetry);
 	}
 
+	UE_NONCOPYABLE(FScopedRetry)
+
+private:
 	AutoRTFM::ForTheRuntime::EAutoRTFMRetryTransactionState OldRetry;
 };
 
diff --git a/Engine/Source/Runtime/Core/Private/AutoRTFM/Context.cpp b/Engine/Source/Runtime/Core/Private/AutoRTFM/Context.cpp
--- a/Engine/Source/Runtime/Core/Private/AutoRTFM/Context.cpp
+++ b/Engine/Source/Runtime/Core/Private/AutoRTFM/Context.cpp
@@ -455,7 +455,7 @@ ETransactionResult FContext::Transact(void (*InstrumentedFunction)(void*), void*
 				{
 					// We skip trying to commit this time, and instead re-run the transaction.
 					Status = EContextStatus::AbortedByFailedLockAcquisition;
-					CurrentTransaction->AbortWithoutThrowing();
+					CurrentTransaction->AbortWithoutThrowing(EAbortWithoutThrowingFlush::Flush);
 					ClearTransactionStatus();
 
 					// We've tried to run at least once if we get here!
diff --git a/Engine/Source/Runtime/Core/Private/AutoRTFM/TaskArray.h b/Engine/Source/Runtime/Core/Private/AutoRTFM/TaskArray.h
--- a/Engine/Source/Runtime/Core/Private/AutoRTFM/TaskArray.h
+++ b/Engine/Source/Runtime/Core/Private/AutoRTFM/TaskArray.h
@@ -250,6 +250,18 @@ public:
 		return true;
 	}
 
+	template<typename FunctorType>
+	void DeleteKeyedIf(const FunctorType& Func)
+	{
+		for (auto& Pair : BitMasksPerKey)
+		{
+			if (Func(Pair.Key))
+			{
+				Pair.Value.Bits.SetRange(Pair.Value.InitialIndex, Pair.Value.Num(), true);
+			}
+		}
+	}
+
     void AddAll(TTaskArray&& Other)
     {
         Canonicalize();
diff --git a/Engine/Source/Runtime/Core/Private/AutoRTFM/Transaction.cpp b/Engine/Source/Runtime/Core/Private/AutoRTFM/Transaction.cpp
--- a/Engine/Source/Runtime/Core/Private/AutoRTFM/Transaction.cpp
+++ b/Engine/Source/Runtime/Core/Private/AutoRTFM/Transaction.cpp
@@ -85,7 +85,7 @@ bool FTransaction::IsFresh() const
 		&& DeferredPopAllOnAbortHandlers.IsEmpty();
 }
 
-void FTransaction::AbortWithoutThrowing()
+void FTransaction::AbortWithoutThrowing(EAbortWithoutThrowingFlush State)
 {
 	UE_LOG(LogAutoRTFM, Verbose, TEXT("Aborting '%hs'!"), GetContextStatusName(Context->GetStatus()));
 
@@ -105,6 +105,13 @@ void FTransaction::AbortWithoutThrowing()
     CommitTasks.Reset();
 
     Undo();
+
+	// Only used by the aborteroonie so it is definitely not in the hot path.
+	if (UNLIKELY(EAbortWithoutThrowingFlush::Flush == State))
+	{
+		AbortTasks.DeleteKeyedIf([&](const void* const Key) { return this->IsOnStack(Key); });
+	}
+
 	AbortTasks.ForEachBackward([&](TFunction<void()>& Task)
     { 
         // Call and then reset each of the tasks in reverse order.
@@ -267,6 +274,9 @@ void FTransaction::CommitNested()
 	}
 	DeferredPopAllOnAbortHandlers.Reset();
 
+	CommitTasks.DeleteKeyedIf([&](const void* const Key) { return this->IsOnStack(Key); });
+	AbortTasks.DeleteKeyedIf([&](const void* const Key) { return this->IsOnStack(Key); });
+
     Parent->CommitTasks.AddAll(MoveTemp(CommitTasks));
     Parent->AbortTasks.AddAll(MoveTemp(AbortTasks));
 
@@ -282,6 +292,8 @@ bool FTransaction::AttemptToCommitOuterNest()
 	UE_LOG(LogAutoRTFM, Verbose, TEXT("Running commit tasks..."));
 
     AbortTasks.Reset();
+	
+	CommitTasks.DeleteKeyedIf([&](const void* const Key) { return this->IsOnStack(Key); });
 
     CommitTasks.ForEachForward([] (TFunction<void()>& Task)
     { 
diff --git a/Engine/Source/Runtime/Core/Private/AutoRTFM/Transaction.h b/Engine/Source/Runtime/Core/Private/AutoRTFM/Transaction.h
--- a/Engine/Source/Runtime/Core/Private/AutoRTFM/Transaction.h
+++ b/Engine/Source/Runtime/Core/Private/AutoRTFM/Transaction.h
@@ -18,6 +18,18 @@ namespace AutoRTFM
 {
 class FContext;
 
+enum class EAbortWithoutThrowingFlush : bool
+{
+	// The default - since aborts happen eagerly when we call `AbortTransaction`
+	// we want to run on-abort handlers that are still on the stack.
+	NoFlush = false,
+
+	// Only used by the Aborteroonie - we call `AbortWithoutThrowing` after the
+	// transactions stack has already disappeared, and so it is thus unsafe to
+	// call potentially stack local on-abort handlers.
+	Flush = true,
+};
+
 class FTransaction final
 {
 public:
@@ -101,7 +113,8 @@ public:
 	void PopAllDeferUntilAbortHandlers(const void* Key);
 
     [[noreturn]] void AbortAndThrow();
-    void AbortWithoutThrowing();
+
+    void AbortWithoutThrowing(EAbortWithoutThrowingFlush State = EAbortWithoutThrowingFlush::NoFlush);
     bool AttemptToCommit();
 
     // Record that a write is about to occur at the given LogicalAddress of Size bytes.
diff --git a/Engine/Source/Runtime/Core/Public/AutoRTFM/AutoRTFM.h b/Engine/Source/Runtime/Core/Public/AutoRTFM/AutoRTFM.h
--- a/Engine/Source/Runtime/Core/Public/AutoRTFM/AutoRTFM.h
+++ b/Engine/Source/Runtime/Core/Public/AutoRTFM/AutoRTFM.h
@@ -878,7 +878,9 @@ template<typename FunctorType> static UE_AUTORTFM_FORCEINLINE void OnAbort(Funct
 // Register a handler for transaction commit. Takes a key parameter so that
 // the handler can be unregistered (see `PopOnCommitHandler`). This is useful
 // for scoped mutations that need an abort handler present unless execution
-// reaches the end of the relevant scope.
+// reaches the end of the relevant scope. If `Key` was stack local to a
+// committing transaction then the runtime will not run the handler as the data
+// backing the handler (normally a `this` object) has already been destroyed.
 template<typename FunctorType> static UE_AUTORTFM_FORCEINLINE void PushOnCommitHandler(const void* Key, FunctorType&& Work)
 {
 	if (autortfm_is_closed())
@@ -890,7 +892,9 @@ template<typename FunctorType> static UE_AUTORTFM_FORCEINLINE void PushOnCommitH
 // Register a handler for transaction commit. Takes a key parameter so that
 // the handler can be unregistered (see `PopOnCommitHandler`). This is useful
 // for scoped mutations that need an abort handler present unless execution
-// reaches the end of the relevant scope.
+// reaches the end of the relevant scope. If `Key` was stack local to a
+// committing transaction then the runtime will not run the handler as the data
+// backing the handler (normally a `this` object) has already been destroyed.
 template<typename FunctorType> static UE_AUTORTFM_FORCEINLINE void PushOnCommitHandler(const void*, FunctorType&&) {}
 #endif
 
@@ -926,7 +930,10 @@ static UE_AUTORTFM_FORCEINLINE void PopAllOnCommitHandlers(const void*) {}
 // Register a handler for transaction abort. Takes a key parameter so that
 // the handler can be unregistered (see `PopOnAbortHandler`). This is useful
 // for scoped mutations that need an abort handler present unless execution
-// reaches the end of the relevant scope.
+// reaches the end of the relevant scope. If `Key` was stack-local to a
+// nested child transaction then the runtime will not run the handler as the
+// data backing the handler (normally a `this` object) has already been
+// destroyed.
 template<typename FunctorType> static UE_AUTORTFM_FORCEINLINE void PushOnAbortHandler(const void* Key, FunctorType&& Work)
 {
 	if (autortfm_is_closed())
@@ -938,7 +945,10 @@ template<typename FunctorType> static UE_AUTORTFM_FORCEINLINE void PushOnAbortHa
 // Register a handler for transaction abort. Takes a key parameter so that
 // the handler can be unregistered (see `PopOnAbortHandler`). This is useful
 // for scoped mutations that need an abort handler present unless execution
-// reaches the end of the relevant scope.
+// reaches the end of the relevant scope. If `Key` was stack-local to a
+// nested child transaction then the runtime will not run the handler as the
+// data backing the handler (normally a `this` object) has already been
+// destroyed.
 template<typename FunctorType> static UE_AUTORTFM_FORCEINLINE void PushOnAbortHandler(const void* Key, FunctorType&&) {}
 #endif
 
