[AutoRTFM] Remove TTaggedPtr specialization for void

As far as I can tell, this specialization exists for `T& operator*()`, which was never used.
The specialization was an entire copy of the class, and inevitably they have diverged in subtle ways.
Remove the operator and the specialization.

diff --git a/Engine/Source/Runtime/Core/Private/AutoRTFM/TaggedPtr.h b/Engine/Source/Runtime/Core/Private/AutoRTFM/TaggedPtr.h
--- a/Engine/Source/Runtime/Core/Private/AutoRTFM/TaggedPtr.h
+++ b/Engine/Source/Runtime/Core/Private/AutoRTFM/TaggedPtr.h
@@ -40,13 +40,13 @@ namespace AutoRTFM
 			Reset();
 		}
 
-		UE_AUTORTFM_FORCEINLINE explicit TTaggedPtr(T* const Pointer) : Payload(reinterpret_cast<uintptr_t>(Pointer))
+		UE_AUTORTFM_FORCEINLINE TTaggedPtr(T* const Pointer) : Payload(reinterpret_cast<uintptr_t>(Pointer))
 		{
 			ASSERT(0 == (Payload & ~PointerMask));
 		}
 
-		UE_AUTORTFM_FORCEINLINE explicit TTaggedPtr(const TTaggedPtr& Other) : Payload(Other.Payload) {}
-		UE_AUTORTFM_FORCEINLINE explicit TTaggedPtr(const TTaggedPtr&& Other) : Payload(Other.Payload) {}
+		UE_AUTORTFM_FORCEINLINE TTaggedPtr(const TTaggedPtr& Other) : Payload(Other.Payload) {}
+		UE_AUTORTFM_FORCEINLINE TTaggedPtr(const TTaggedPtr&& Other) : Payload(Other.Payload) {}
 
 		UE_AUTORTFM_FORCEINLINE TTaggedPtr& operator=(T* const Pointer)
 		{
@@ -70,11 +70,6 @@ namespace AutoRTFM
 			return *this;
 		}
 
-		UE_AUTORTFM_FORCEINLINE T& operator*() const
-	    {
-	    	return *reinterpret_cast<T*>(Payload & PointerMask);
-	    }
-
 	    UE_AUTORTFM_FORCEINLINE T* operator->() const
 	    {
 	    	return reinterpret_cast<T*>(Payload & PointerMask);
@@ -85,101 +80,6 @@ namespace AutoRTFM
 	    	return reinterpret_cast<T*>(Payload & PointerMask);
 	    }
 
-	    UE_AUTORTFM_FORCEINLINE void SetTopTag(const TopTagT Tag)
-	    {
-	    	*reinterpret_cast<TopTagT*>(reinterpret_cast<uint8_t*>(&Payload) + OffsetOfTopTag) = MoveTemp(Tag);
-	    }
-
-	    UE_AUTORTFM_FORCEINLINE TopTagT GetTopTag() const
-	    {
-	    	return *reinterpret_cast<const TopTagT*>(reinterpret_cast<const uint8_t*>(&Payload) + OffsetOfTopTag);
-	    }
-
-	    UE_AUTORTFM_FORCEINLINE void Reset()
-	    {
-	    	Payload = 0;
-	    }
-
-	    UE_AUTORTFM_FORCEINLINE uintptr_t GetPayload() const
-	    {
-	    	return Payload;
-	    }
-
-	private:
-		uintptr_t Payload;
-	};
-
-	template<typename TopTagT> class TTaggedPtr<void, TopTagT> final
-	{
-		// Our tagged pointer assumes 64-bit pointers.
-		static_assert(sizeof(uintptr_t) == sizeof(uint64_t));
-
-		// For our tagged pointers we assume the bottom 48-bits are used, and
-		// so the top 16-bits are free for our tag.
-		static_assert(sizeof(TopTagT) == sizeof(uint16_t));
-
-		static constexpr unsigned PointerBits = 48;
-		static constexpr uintptr_t PointerMask = (static_cast<uintptr_t>(1) << PointerBits) - 1;
-		static constexpr unsigned TopTagBits = sizeof(TopTagT);
-		static constexpr uintptr_t TopTagMask = (static_cast<uintptr_t>(1) << TopTagBits) - 1;
-
-	public:
-		struct Hasher final
-		{
-		    size_t operator()(const TTaggedPtr<void, TopTagT>& TaggedPtr) const
-		    {
-		      return TaggedPtr.Payload;
-		    }
-		};
-
-		// Assumes little-endian, so top tag is the last 2-bytes in memory.
-		static constexpr size_t OffsetOfTopTag = PointerBits / 8;
-
-		UE_AUTORTFM_FORCEINLINE TTaggedPtr()
-		{
-			Reset();
-		}
-
-		UE_AUTORTFM_FORCEINLINE TTaggedPtr(void* const Pointer) : Payload(reinterpret_cast<uintptr_t>(Pointer))
-		{
-			ASSERT(0 == (Payload & ~PointerMask));
-		}
-
-		UE_AUTORTFM_FORCEINLINE TTaggedPtr(const TTaggedPtr& Other) : Payload(Other.Payload) {}
-		UE_AUTORTFM_FORCEINLINE TTaggedPtr(const TTaggedPtr&& Other) : Payload(Other.Payload) {}
-
-		UE_AUTORTFM_FORCEINLINE TTaggedPtr& operator=(void* const Pointer)
-		{
-			const uintptr_t NewPayload = reinterpret_cast<uintptr_t>(Pointer);
-
-			ASSERT(0 == (NewPayload & ~PointerMask));
-			Payload &= ~PointerMask;
-			Payload |= NewPayload;
-			return *this;
-		}
-
-		UE_AUTORTFM_FORCEINLINE TTaggedPtr& operator=(const TTaggedPtr& Other)
-		{
-			Payload = Other.Payload;
-			return *this;
-		}
-
-		UE_AUTORTFM_FORCEINLINE TTaggedPtr& operator=(TTaggedPtr&& Other)
-		{
-			Payload = Other.Payload;
-			return *this;
-		}
-
-		UE_AUTORTFM_FORCEINLINE bool operator==(const TTaggedPtr& Other) const
-		{
-			return Payload == Other.Payload;
-		}
-
-		UE_AUTORTFM_FORCEINLINE void* Get() const
-	    {
-	    	return reinterpret_cast<void*>(Payload & PointerMask);
-	    }
-
 	    UE_AUTORTFM_FORCEINLINE void SetTopTag(const TopTagT Tag)
 	    {
 	    	*reinterpret_cast<TopTagT*>(reinterpret_cast<uint8_t*>(&Payload) + OffsetOfTopTag) = Tag;
@@ -203,4 +103,5 @@ namespace AutoRTFM
 	private:
 		uintptr_t Payload;
 	};
+
 }
