[AutoRTFM] Run some FText unit tests inside the AutoRTFM engine test suite.

FText doesn't have a lot of AutoRTFM-specific behavior, but we are seeing some strange memory corruption within FText that only reproduces on Linux + AutoRTFM. Running the unit tests while instrumented under AutoRTFM seems like a reasonable defensive measure.

diff --git a/Engine/Source/Runtime/Core/Tests/Internationalization/TextTest.cpp b/Engine/Source/Runtime/Core/Tests/Internationalization/TextTest.cpp
--- a/Engine/Source/Runtime/Core/Tests/Internationalization/TextTest.cpp
+++ b/Engine/Source/Runtime/Core/Tests/Internationalization/TextTest.cpp
@@ -4,8 +4,10 @@
 #include "Templates/UnrealTemplate.h"
 #include "Containers/Array.h"
 #include "Containers/UnrealString.h"
+#include "Misc/AutomationTest.h"
 #include "Misc/DateTime.h"
 #include "CoreGlobals.h"
+#include "AutoRTFM/AutoRTFM.h"
 #include "Internationalization/Text.h"
 #include "Internationalization/Culture.h"
 #include "Internationalization/FastDecimalFormat.h"
@@ -1268,7 +1270,7 @@ TEST_CASE_NAMED(FTextStringificationTest, "System::Core::Misc::TextStringificati
 	I18N.RestoreCultureState(OriginalCultureState);
 }
 
-TEST_CASE_NAMED(FTextFormatArgModifierTest, "System::Core::Misc::TextFormatArgModifiers", "[.][EditorContext][ClientContext][EngineFilter]")
+static void RunTestCasesForSystem_Core_Misc_TextFormatArgModifiers()
 {
 	auto EnsureValidResult = [&](const FString& InResult, const FString& InExpected, const FString& InName, const FString& InDescription)
 	{
@@ -1385,11 +1387,15 @@ TEST_CASE_NAMED(FTextFormatArgModifierTest, "System::Core::Misc::TextFormatArgMo
 		}
 	}
 
-
 	// Restore original culture
 	I18N.RestoreCultureState(OriginalCultureState);
 }
 
+TEST_CASE_NAMED(FTextFormatArgModifierTest, "System::Core::Misc::TextFormatArgModifiers", "[.][EditorContext][ClientContext][EngineFilter]")
+{
+	RunTestCasesForSystem_Core_Misc_TextFormatArgModifiers();
+}
+
 #if UE_ENABLE_ICU
 
 TEST_CASE_NAMED(FICUSanitizationTest, "System::Core::Misc::ICUSanitization", "[.][EditorContext][ClientContext][EngineFilter]")
@@ -1553,6 +1559,52 @@ TEST_CASE_NAMED(FICUTextTest, "System::Core::Misc::ICUText", "[.][EditorContext]
 
 #endif // #if UE_ENABLE_ICU
 
+#if WITH_DEV_AUTOMATION_TESTS
+
+IMPLEMENT_SIMPLE_AUTOMATION_TEST(FAutoRTFMTextTests, "AutoRTFM + FText", \
+								 EAutomationTestFlags::EngineFilter | EAutomationTestFlags::ClientContext | \
+								 EAutomationTestFlags::ServerContext | EAutomationTestFlags::CommandletContext)
+
+bool FAutoRTFMTextTests::RunTest(const FString& Parameters)
+{
+	if (!AutoRTFM::ForTheRuntime::IsAutoRTFMRuntimeEnabled())
+	{
+		ExecutionInfo.AddEvent(FAutomationEvent(EAutomationEventType::Info, TEXT("SKIPPED 'FAutoRTFMTextTests' test. AutoRTFM disabled.")));
+		return true;
+	}
+
+	// Verify that FText works properly when used in an AutoRTFM transaction, both committing and aborting.
+	// System::Core::Misc::TextStringification is not included, because this test uses the StringTableRegistry.
+	// (The StringTableRegistry is not yet AutoRTFM-safe.)
+	AutoRTFM::Commit([]
+	{
+		FTextNumericParsingTestClass Instance;
+		Instance.TextNumericParsingTest();
+	});
+
+	AutoRTFM::Transact([]
+	{
+		FTextNumericParsingTestClass Instance;
+		Instance.TextNumericParsingTest();
+		AutoRTFM::AbortTransaction();
+	});
+
+	AutoRTFM::Commit([]
+	{
+		RunTestCasesForSystem_Core_Misc_TextFormatArgModifiers();
+	});
+
+	AutoRTFM::Transact([]
+	{
+		RunTestCasesForSystem_Core_Misc_TextFormatArgModifiers();
+		AutoRTFM::AbortTransaction();
+	});
+
+	return true;
+}
+
+#endif // #if WITH_DEV_AUTOMATION_TESTS
+
 #undef LOCTEXT_NAMESPACE 
 
 #endif //WITH_TESTS
\ No newline at end of file
