Thumbnail FImage accessors ; prefer using FImage for image data, not Arrays

diff --git a/Engine/Source/Editor/UnrealEd/Private/ThumbnailExternalCache.cpp b/Engine/Source/Editor/UnrealEd/Private/ThumbnailExternalCache.cpp
--- a/Engine/Source/Editor/UnrealEd/Private/ThumbnailExternalCache.cpp
+++ b/Engine/Source/Editor/UnrealEd/Private/ThumbnailExternalCache.cpp
@@ -28,38 +28,26 @@ namespace ThumbnailExternalCache
 
 	void ResizeThumbnailImage(FObjectThumbnail& Thumbnail, const int32 NewWidth, const int32 NewHeight)
 	{
-		TArray<uint8> DestData;
-		DestData.AddUninitialized(NewWidth * NewHeight * sizeof(FColor));
+		FImageView SrcImage = Thumbnail.GetImage();
 
-		// Force decompress if needed
-		Thumbnail.GetUncompressedImageData();
-		TArray<uint8>& UncompressedImageData = Thumbnail.AccessImageData();
+		FImage DestImage;
+		FImageCore::ResizeImageAllocDest(SrcImage,DestImage,NewWidth,NewHeight);
 
-		const bool bLinearSpace = true;
-		const bool bForceOpaqueOutput = false;
-		const TArrayView<FColor> SrcDataView(reinterpret_cast<FColor*>(UncompressedImageData.GetData()), UncompressedImageData.Num() / sizeof(FColor));
-		const TArrayView<FColor> DestDataView(reinterpret_cast<FColor*>(DestData.GetData()), DestData.Num() / sizeof(FColor));
-		FImageUtils::ImageResize(Thumbnail.GetImageWidth(), Thumbnail.GetImageHeight(), SrcDataView, NewWidth, NewHeight, DestDataView, bLinearSpace, bForceOpaqueOutput);
-
-		UncompressedImageData = MoveTemp(DestData);
-		Thumbnail.SetImageSize(NewWidth, NewHeight);
-
-		// Invalidate compressed data so it will be recompressed
-		Thumbnail.AccessCompressedImageData().Reset();
+		Thumbnail.SetImage(DestImage,true);
 	}
 
 	// Return true if was resized
 	bool ResizeThumbnailIfNeeded(FObjectThumbnail& Thumbnail, const int32 MaxImageSize)
 	{
-		const int32 Width = Thumbnail.GetImageWidth();
+		const int32 Width  = Thumbnail.GetImageWidth();
 		const int32 Height = Thumbnail.GetImageHeight();
 
 		// Resize if larger than maximum size
 		if (Width > MaxImageSize || Height > MaxImageSize)
 		{
 			const double ShrinkModifier = (double)FMath::Max<int32>(Width, Height) / (double)MaxImageSize;
-			const int32 NewWidth = int32((double)Width / ShrinkModifier);
-			const int32 NewHeight = int32((double)Height / ShrinkModifier);
+			const int32 NewWidth  = FMath::RoundToInt((double)Width / ShrinkModifier);
+			const int32 NewHeight = FMath::RoundToInt((double)Height / ShrinkModifier);
 			ResizeThumbnailImage(Thumbnail, NewWidth, NewHeight);
 
 			return true;
diff --git a/Engine/Source/Runtime/Core/Core.Build.cs b/Engine/Source/Runtime/Core/Core.Build.cs
--- a/Engine/Source/Runtime/Core/Core.Build.cs
+++ b/Engine/Source/Runtime/Core/Core.Build.cs
@@ -37,7 +37,13 @@ public class Core : ModuleRules
 		PrivateIncludePathModuleNames.AddRange(
 			new string[] {
 				"DerivedDataCache",
-				"TargetPlatform",
+				"TargetPlatform"
+			}
+			);
+
+		PublicIncludePathModuleNames.AddRange(
+			new string[] {
+				"ImageCore"
 			}
 			);
 
diff --git a/Engine/Source/Runtime/Core/Public/Misc/ObjectThumbnail.h b/Engine/Source/Runtime/Core/Public/Misc/ObjectThumbnail.h
--- a/Engine/Source/Runtime/Core/Public/Misc/ObjectThumbnail.h
+++ b/Engine/Source/Runtime/Core/Public/Misc/ObjectThumbnail.h
@@ -8,6 +8,7 @@
 #include "Serialization/Archive.h"
 #include "Serialization/StructuredArchive.h"
 #include "UObject/NameTypes.h"
+#include "ImageCore.h"
 
 /**
  * Thumbnail compression interface for packages.  The engine registers a class that can compress and
@@ -160,9 +161,67 @@ public:
 	/** Returns thumbnail compressor that would be used on current uncompressed image data. */
 	CORE_API FThumbnailCompressionInterface* ChooseNewCompressor() const;
 
-	/** Returns uncompressed image data, decompressing it on demand if needed. */
+	/** Returns uncompressed image data, decompressing it on demand if needed.
+	
+		Prefer GetImage and use FImage/FImageView for image data.
+	
+		Not actually a const function, may change the "ImageData" member.
+	*/
 	CORE_API const TArray< uint8 >& GetUncompressedImageData() const;
 
+	/** Returns uncompressed image data as an FImageView, decompressing it on demand if needed.
+
+		Note the FImageView does not have a copy of the data, it points at the Thumbnail UncompressedImageData.
+		If the UncompressedImageData is freed or changed, the FImageView will be affected.
+		
+		Not actually a const function, may change the "ImageData" member.
+	*/
+	FImageView GetImage() const
+	{
+		const TArray< uint8 >& Data = FObjectThumbnail::GetUncompressedImageData();
+		if ( Data.Num() == 0 )
+		{
+			return FImageView();
+		}
+
+		return FImageView( (void *)&Data[0],ImageWidth,ImageHeight,ERawImageFormat::BGRA8 );
+	}
+
+	/** Set Thumbnail uncompressed image data.
+		Existing compressed data, if any, is freed.
+	
+		if bMove is true, the Image data is moved into the Thumbnail
+		if bMove is false, the Image is copied into the thumbnail and not changed
+
+		Thumbnails are always BGRA8-sRGB.  The passed in Image can be other formats and conversion will be done if needed.
+	*/
+	void SetImage(FImage & Image, bool bMove)
+	{
+		if ( Image.Format != ERawImageFormat::BGRA8 || Image.GammaSpace != EGammaSpace::sRGB )
+		{
+			// Image must be converted to BGRA8 to store in Thumbnail
+			FImage ImageBGRA8;
+			Image.CopyTo(ImageBGRA8,ERawImageFormat::BGRA8,EGammaSpace::sRGB);
+			SetImage(ImageBGRA8,true);
+			return;
+		}
+
+		ImageWidth  = Image.SizeX;
+		ImageHeight = Image.SizeY;
+		CompressedImageData.Reset();
+		bIsJPEG = false;
+
+		if ( bMove )
+		{
+			ImageData = MoveTemp(Image.RawData);
+			Image.RawData.Reset();
+		}
+		else
+		{
+			ImageData = Image.RawData;
+		}
+	}
+
 	/** Serializers */
 	CORE_API void Serialize(FArchive& Ar);
 	CORE_API void Serialize(FStructuredArchive::FSlot Slot);
