Explicitly set the size of EExtendedTaskPriority to 8 bits to match the normal ETaskPriority, so both can be efficiently packed into an int

diff --git a/Engine/Source/Runtime/Core/Public/Tasks/TaskPrivate.h b/Engine/Source/Runtime/Core/Public/Tasks/TaskPrivate.h
--- a/Engine/Source/Runtime/Core/Public/Tasks/TaskPrivate.h
+++ b/Engine/Source/Runtime/Core/Public/Tasks/TaskPrivate.h
@@ -56,7 +56,7 @@ namespace UE::Tasks
 	using LowLevelTasks::ToTaskPriority;
 
 	// special task priorities for tasks that are never sent to the scheduler
-	enum class EExtendedTaskPriority
+	enum class EExtendedTaskPriority : int8
 	{
 		None,
 		Inline, // a task priority for "inline" task execution - a task is executed "inline" by the thread that unlocked it, w/o scheduling
