[AutoRTFM] Use a transactionally-safe RW lock for text argument modifiers.

The text formatter can run inside of a transaction, and requires a RW lock when text argument modifiers are registered.

diff --git a/Engine/Source/Runtime/Core/Private/Internationalization/TextFormatter.cpp b/Engine/Source/Runtime/Core/Private/Internationalization/TextFormatter.cpp
--- a/Engine/Source/Runtime/Core/Private/Internationalization/TextFormatter.cpp
+++ b/Engine/Source/Runtime/Core/Private/Internationalization/TextFormatter.cpp
@@ -1007,19 +1007,19 @@ FTextFormatter& FTextFormatter::Get()
 
 void FTextFormatter::RegisterTextArgumentModifier(const FTextFormatString& InKeyword, FCompileTextArgumentModifierFuncPtr InCompileFunc)
 {
-	FWriteScopeLock Lock(TextArgumentModifiersRW);
+	UE::TWriteScopeLock Lock(TextArgumentModifiersRW);
 	TextArgumentModifiers.Add(InKeyword, MoveTemp(InCompileFunc));
 }
 
 void FTextFormatter::UnregisterTextArgumentModifier(const FTextFormatString& InKeyword)
 {
-	FWriteScopeLock Lock(TextArgumentModifiersRW);
+	UE::TWriteScopeLock Lock(TextArgumentModifiersRW);
 	TextArgumentModifiers.Remove(InKeyword);
 }
 
 FTextFormatter::FCompileTextArgumentModifierFuncPtr FTextFormatter::FindTextArgumentModifier(const FTextFormatString& InKeyword) const
 {
-	FReadScopeLock Lock(TextArgumentModifiersRW);
+	UE::TReadScopeLock Lock(TextArgumentModifiersRW);
 	return TextArgumentModifiers.FindRef(InKeyword);
 }
 
diff --git a/Engine/Source/Runtime/Core/Public/Internationalization/TextFormatter.h b/Engine/Source/Runtime/Core/Public/Internationalization/TextFormatter.h
--- a/Engine/Source/Runtime/Core/Public/Internationalization/TextFormatter.h
+++ b/Engine/Source/Runtime/Core/Public/Internationalization/TextFormatter.h
@@ -6,10 +6,10 @@
 #include "Containers/SortedMap.h"
 #include "Containers/UnrealString.h"
 #include "CoreTypes.h"
-#include "HAL/CriticalSection.h"
 #include "Internationalization/ITextFormatArgumentModifier.h"
 #include "Internationalization/Text.h"
 #include "Misc/ExpressionParserTypes.h"
+#include "Misc/TransactionallySafeRWLock.h"
 #include "Templates/Function.h"
 #include "Templates/SharedPointer.h"
 
@@ -108,5 +108,5 @@ private:
 	TSortedMap<FTextFormatString, FCompileTextArgumentModifierFuncPtr> TextArgumentModifiers;
 
 	/** RW lock protecting the argument modifiers map from being modified concurrently */
-	mutable FRWLock TextArgumentModifiersRW;
+	mutable FTransactionallySafeRWLock TextArgumentModifiersRW;
 };
