[Core] Added LiveCoding support/reload to FDelayedAutoRegisterHelper

diff --git a/Engine/Source/Developer/Windows/LiveCoding/Private/LiveCodingModule.cpp b/Engine/Source/Developer/Windows/LiveCoding/Private/LiveCodingModule.cpp
--- a/Engine/Source/Developer/Windows/LiveCoding/Private/LiveCodingModule.cpp
+++ b/Engine/Source/Developer/Windows/LiveCoding/Private/LiveCodingModule.cpp
@@ -917,6 +917,8 @@ void FLiveCodingModule::AttemptSyncLivePatching()
 			NewObjects.Empty();
 #endif
 
+			FDelayedAutoRegisterHelper::RunAndClearDelayedAutoRegisterDelegates(EDelayedRegisterRunPhase::LiveCodingReload);
+
 			OnPatchCompleteDelegate.Broadcast();
 			GHasLoadedPatch = false;
 
diff --git a/Engine/Source/Runtime/Core/Private/Misc/DelayedAutoRegister.cpp b/Engine/Source/Runtime/Core/Private/Misc/DelayedAutoRegister.cpp
--- a/Engine/Source/Runtime/Core/Private/Misc/DelayedAutoRegister.cpp
+++ b/Engine/Source/Runtime/Core/Private/Misc/DelayedAutoRegister.cpp
@@ -10,7 +10,7 @@
 #include "Misc/AssertionMacros.h"
 #include "Templates/Function.h"
 
-DECLARE_MULTICAST_DELEGATE(FDelayedAutoRegisterDelegate);
+DECLARE_MULTICAST_DELEGATE_OneParam(FDelayedAutoRegisterDelegate, EDelayedRegisterRunPhase);
 
 static TSet<EDelayedRegisterRunPhase> GPhasesAlreadyRun;
 
@@ -20,12 +20,31 @@ static FDelayedAutoRegisterDelegate& GetDelayedAutoRegisterDelegate(EDelayedRegi
 	return Singleton[(uint8)Phase];
 }
 
-FDelayedAutoRegisterHelper::FDelayedAutoRegisterHelper(EDelayedRegisterRunPhase Phase, TFunction<void()> RegistrationFunction)
+FDelayedAutoRegisterHelper::FDelayedAutoRegisterHelper(EDelayedRegisterRunPhase Phase, TFunction<void()> RegistrationFunction, const bool bRerunOnLiveCodingReload)
+	: FDelayedAutoRegisterHelper(
+		Phase,
+		[RegistrationFunction](const EDelayedRegisterRunPhase) { RegistrationFunction(); },
+		bRerunOnLiveCodingReload)
 {
+}
+
+FDelayedAutoRegisterHelper::FDelayedAutoRegisterHelper(
+	EDelayedRegisterRunPhase Phase,
+	TFunction<void(const EDelayedRegisterRunPhase)> RegistrationFunction,
+	const bool bRerunOnLiveCodingReload)
+{
+#if WITH_EDITOR && WITH_LIVE_CODING
+	// We can use both the provided phase AND LiveCodingReload
+	if (bRerunOnLiveCodingReload && Phase != EDelayedRegisterRunPhase::LiveCodingReload)
+	{
+		GetDelayedAutoRegisterDelegate(EDelayedRegisterRunPhase::LiveCodingReload).AddLambda(RegistrationFunction);
+	}
+#endif
+
 	// if the phase has already passed, we just run the function immediately
-	if (GPhasesAlreadyRun.Contains(Phase))
+	if (Phase < EDelayedRegisterRunPhase::NumRunOncePhases && GPhasesAlreadyRun.Contains(Phase))
 	{
-		RegistrationFunction();
+		RegistrationFunction(Phase);
 	}
 	else
 	{
@@ -33,13 +52,20 @@ FDelayedAutoRegisterHelper::FDelayedAutoRegisterHelper(EDelayedRegisterRunPhase
 	}
 }
 
-void FDelayedAutoRegisterHelper::RunAndClearDelayedAutoRegisterDelegates(EDelayedRegisterRunPhase Phase)
+void FDelayedAutoRegisterHelper::RunAndClearDelayedAutoRegisterDelegates(EDelayedRegisterRunPhase RunPhase)
 {
-	checkf(!GPhasesAlreadyRun.Contains(Phase), TEXT("Delayed Startup phase %d has already run - it is not expected to be run again!"), (int32)Phase);
+	const bool bIsRunOncePhase = RunPhase < EDelayedRegisterRunPhase::NumRunOncePhases;
+	if (bIsRunOncePhase)
+	{
+		checkf(!GPhasesAlreadyRun.Contains(RunPhase), TEXT("Delayed Startup phase %d has already run - it is not expected to be run again!"), (int32)RunPhase);
+	}
 
 	// run all the delayed functions!
-	GetDelayedAutoRegisterDelegate(Phase).Broadcast();
-	GetDelayedAutoRegisterDelegate(Phase).Clear();
+	GetDelayedAutoRegisterDelegate(RunPhase).Broadcast(RunPhase);
 
-	GPhasesAlreadyRun.Add(Phase);
+	if (bIsRunOncePhase)
+	{
+		GetDelayedAutoRegisterDelegate(RunPhase).Clear();
+		GPhasesAlreadyRun.Add(RunPhase);
+	}
 }
diff --git a/Engine/Source/Runtime/Core/Public/Misc/DelayedAutoRegister.h b/Engine/Source/Runtime/Core/Public/Misc/DelayedAutoRegister.h
--- a/Engine/Source/Runtime/Core/Public/Misc/DelayedAutoRegister.h
+++ b/Engine/Source/Runtime/Core/Public/Misc/DelayedAutoRegister.h
@@ -20,13 +20,20 @@ enum class EDelayedRegisterRunPhase : uint8
 	DeviceProfileManagerReady,
 	EndOfEngineInit,
 
+	NumRunOncePhases, // Phases before this are run once
+
+	LiveCodingReload,
+
 	NumPhases,
 };
 
 struct FDelayedAutoRegisterHelper
 {
 
-	CORE_API FDelayedAutoRegisterHelper(EDelayedRegisterRunPhase RunPhase, TFunction<void()> RegistrationFunction);
+	CORE_API FDelayedAutoRegisterHelper(EDelayedRegisterRunPhase RunPhase, TFunction<void()> RegistrationFunction, const bool bRerunOnLiveCodingReload = false);
+
+	/** Overload that provides the phase being registered at the time the function is called.  */
+	CORE_API FDelayedAutoRegisterHelper(EDelayedRegisterRunPhase RunPhase, TFunction<void(const EDelayedRegisterRunPhase)> RegistrationFunction, const bool bRerunOnLiveCodingReload = false);
 
 	static CORE_API void RunAndClearDelayedAutoRegisterDelegates(EDelayedRegisterRunPhase RunPhase);
 };
