[AutoRTFM] Replace FTransactionallySafeWriteScopeLock with UE::TWriteScopeLock.

FTransactionallySafeWriteScopeLock is redundant, as it is functionally identical to a UE::TWriteScopeLock.

diff --git a/Engine/Plugins/Experimental/ChaosModularVehicle/Source/ChaosModularVehicleEngine/Public/ChaosModularVehicle/ModularVehicleSimulationCU.h b/Engine/Plugins/Experimental/ChaosModularVehicle/Source/ChaosModularVehicleEngine/Public/ChaosModularVehicle/ModularVehicleSimulationCU.h
--- a/Engine/Plugins/Experimental/ChaosModularVehicle/Source/ChaosModularVehicleEngine/Public/ChaosModularVehicle/ModularVehicleSimulationCU.h
+++ b/Engine/Plugins/Experimental/ChaosModularVehicle/Source/ChaosModularVehicleEngine/Public/ChaosModularVehicle/ModularVehicleSimulationCU.h
@@ -8,6 +8,7 @@
 #include "ModularVehicleBuilder.h"
 #include "Chaos/ParticleHandle.h"
 #include "Chaos/ParticleHandleFwd.h"
+#include "Misc/ScopeRWLock.h"
 #include "Misc/TransactionallySafeRWLock.h"
 
 DECLARE_LOG_CATEGORY_EXTERN(LogModularVehicleSim, Log, All);
@@ -56,7 +57,7 @@ public:
 	void Terminate();
 	void SetInputMappings(const FInputNameMap& InNameMap)
 	{ 
-		FTransactionallySafeWriteScopeLock InputConfigLock(InputConfigurationLock);
+		UE::TWriteScopeLock InputConfigLock(InputConfigurationLock);
 		InputNameMap = InNameMap; 
 	}
 	void CacheRootParticle(IPhysicsProxyBase* Proxy);
diff --git a/Engine/Source/Programs/AutoRTFMTests/Private/UECoreTests.cpp b/Engine/Source/Programs/AutoRTFMTests/Private/UECoreTests.cpp
--- a/Engine/Source/Programs/AutoRTFMTests/Private/UECoreTests.cpp
+++ b/Engine/Source/Programs/AutoRTFMTests/Private/UECoreTests.cpp
@@ -1576,7 +1576,7 @@ TEST_CASE("UECore.TransactionallySafeWriteScopeLock")
 
 		AutoRTFM::ETransactionResult Result = AutoRTFM::Transact([&]
 			{
-				FTransactionallySafeWriteScopeLock Lock(RWLock);
+				UE::TWriteScopeLock Lock(RWLock);
 				AutoRTFM::AbortTransaction();
 			});
 
@@ -1584,7 +1584,7 @@ TEST_CASE("UECore.TransactionallySafeWriteScopeLock")
 
 		Result = AutoRTFM::Transact([&]
 			{
-				FTransactionallySafeWriteScopeLock Lock(RWLock);
+				UE::TWriteScopeLock Lock(RWLock);
 			});
 
 		REQUIRE(AutoRTFM::ETransactionResult::Committed == Result);
@@ -1595,7 +1595,7 @@ TEST_CASE("UECore.TransactionallySafeWriteScopeLock")
 		AutoRTFM::ETransactionResult Result = AutoRTFM::Transact([&]
 			{
 				FTransactionallySafeRWLock RWLock;
-				FTransactionallySafeWriteScopeLock Lock(RWLock);
+				UE::TWriteScopeLock Lock(RWLock);
 				AutoRTFM::AbortTransaction();
 			});
 
@@ -1604,7 +1604,7 @@ TEST_CASE("UECore.TransactionallySafeWriteScopeLock")
 		Result = AutoRTFM::Transact([&]
 			{
 				FTransactionallySafeRWLock RWLock;
-				FTransactionallySafeWriteScopeLock Lock(RWLock);
+				UE::TWriteScopeLock Lock(RWLock);
 			});
 
 		REQUIRE(AutoRTFM::ETransactionResult::Committed == Result);
@@ -1620,7 +1620,7 @@ TEST_CASE("UECore.TransactionallySafeWriteScopeLock")
 
 				InnerResult = AutoRTFM::Transact([&]
 					{
-						FTransactionallySafeWriteScopeLock Lock(RWLock);
+						UE::TWriteScopeLock Lock(RWLock);
 						AutoRTFM::CascadingAbortTransaction();
 					});
 			});
@@ -1633,7 +1633,7 @@ TEST_CASE("UECore.TransactionallySafeWriteScopeLock")
 
 				InnerResult = AutoRTFM::Transact([&]
 					{
-						FTransactionallySafeWriteScopeLock Lock(RWLock);
+						UE::TWriteScopeLock Lock(RWLock);
 						AutoRTFM::AbortTransaction();
 					});
 			});
@@ -1647,7 +1647,7 @@ TEST_CASE("UECore.TransactionallySafeWriteScopeLock")
 
 				InnerResult = AutoRTFM::Transact([&]
 					{
-						FTransactionallySafeWriteScopeLock Lock(RWLock);
+						UE::TWriteScopeLock Lock(RWLock);
 					});
 
 				AutoRTFM::AbortTransaction();
@@ -1661,7 +1661,7 @@ TEST_CASE("UECore.TransactionallySafeWriteScopeLock")
 
 				InnerResult = AutoRTFM::Transact([&]
 					{
-						FTransactionallySafeWriteScopeLock Lock(RWLock);
+						UE::TWriteScopeLock Lock(RWLock);
 					});
 			});
 
diff --git a/Engine/Source/Runtime/AssetRegistry/Private/AssetRegistryImpl.h b/Engine/Source/Runtime/AssetRegistry/Private/AssetRegistryImpl.h
--- a/Engine/Source/Runtime/AssetRegistry/Private/AssetRegistryImpl.h
+++ b/Engine/Source/Runtime/AssetRegistry/Private/AssetRegistryImpl.h
@@ -14,6 +14,7 @@
 #include "Misc/Optional.h"
 #include "Misc/TransactionallySafeRWLock.h"
 #include "Misc/TransactionallySafeRWScopeLock.h"
+#include "Misc/ScopeRWLock.h"
 #include "ModuleDescriptor.h"
 #include "PackageDependencyData.h"
 #include "Templates/Function.h"
@@ -45,7 +46,7 @@ namespace UE::AssetRegistry::Impl { struct FTickContext; }
 namespace UE::AssetRegistry::Impl { struct FClassInheritanceContext; }
 namespace UE::AssetRegistry { template <typename TScopeLockType> class TRWScopeLockWithPriority; }
 namespace UE::AssetRegistry { using FInterfaceReadScopeLock = TRWScopeLockWithPriority<FTransactionallySafeReadScopeLock>; }
-namespace UE::AssetRegistry { using FInterfaceWriteScopeLock = TRWScopeLockWithPriority<FTransactionallySafeWriteScopeLock>; }
+namespace UE::AssetRegistry { using FInterfaceWriteScopeLock = TRWScopeLockWithPriority<UE::TWriteScopeLock<FTransactionallySafeRWLock>>; }
 namespace UE::AssetRegistry { using FInterfaceRWScopeLock = class FRWScopeLockWithPriority; }
 
 #if WITH_EDITOR
diff --git a/Engine/Source/Runtime/Core/Private/Misc/ConfigCacheIni.cpp b/Engine/Source/Runtime/Core/Private/Misc/ConfigCacheIni.cpp
--- a/Engine/Source/Runtime/Core/Private/Misc/ConfigCacheIni.cpp
+++ b/Engine/Source/Runtime/Core/Private/Misc/ConfigCacheIni.cpp
@@ -24,6 +24,8 @@
 #include "Misc/StringBuilder.h"
 #include "Misc/PathViews.h"
 #include "Misc/ScopeLock.h"
+#include "Misc/ScopeRWLock.h"
+#include "Misc/TransactionallySafeRWLock.h"
 #include "ProfilingDebugging/CpuProfilerTrace.h"
 #include "ProfilingDebugging/AssetMetadataTrace.h"
 #include "Serialization/MemoryReader.h"
@@ -33,7 +35,6 @@
 #include "Async/ParallelFor.h"
 #include "Misc/OutputDeviceRedirector.h"
 #include "Logging/MessageLog.h"
-#include "Misc/TransactionallySafeRWLock.h"
 #include <limits>
 
 namespace
@@ -961,7 +962,7 @@ FConfigFile::FConfigFile(FConfigFile&& Other)
 
 FConfigFile& FConfigFile::operator=(const FConfigFile& Other)
 {
-	FTransactionallySafeWriteScopeLock ScopeLock(ConfigFileMapLock);
+	UE::TWriteScopeLock ScopeLock(ConfigFileMapLock);
 	this->FConfigFileMap::operator=(Other);
 	Dirty = Other.Dirty;
 	NoSave = Other.NoSave;
@@ -1008,7 +1009,7 @@ FConfigFile& FConfigFile::operator=(const FConfigFile& Other)
 
 FConfigFile& FConfigFile::operator=(FConfigFile&& Other)
 {
-	FTransactionallySafeWriteScopeLock ScopeLock(ConfigFileMapLock);
+	UE::TWriteScopeLock ScopeLock(ConfigFileMapLock);
 	this->FConfigFileMap::operator=(MoveTemp(Other));
 	Dirty = Other.Dirty;
 	NoSave = Other.NoSave;
@@ -1146,7 +1147,7 @@ void FConfigFile::Shrink()
 	if (IsInGameThread()) GConfigShrinkTime -= FPlatformTime::Seconds();
 #endif
 
-	FTransactionallySafeWriteScopeLock ScopeLock(ConfigFileMapLock);
+	UE::TWriteScopeLock ScopeLock(ConfigFileMapLock);
 	FConfigFileMap::Shrink();
 
 	for (FConfigFileMap::TIterator It(*this); It; ++It)
diff --git a/Engine/Source/Runtime/Core/Private/Serialization/CustomVersion.cpp b/Engine/Source/Runtime/Core/Private/Serialization/CustomVersion.cpp
--- a/Engine/Source/Runtime/Core/Private/Serialization/CustomVersion.cpp
+++ b/Engine/Source/Runtime/Core/Private/Serialization/CustomVersion.cpp
@@ -189,7 +189,7 @@ FCustomVersionContainer FCurrentCustomVersions::GetAll()
 		}
 	}
 
-	FTransactionallySafeWriteScopeLock Scope(Registry.Lock);
+	UE::TWriteScopeLock Scope(Registry.Lock);
 	Registry.RegisterQueue();
 	return Registry.Registered;
 }
@@ -240,7 +240,7 @@ void FCurrentCustomVersions::Register(const FGuid& Key, int32 Version, const TCH
 {
 	FStaticCustomVersionRegistry& Registry = FStaticCustomVersionRegistry::Get();
 
-	FTransactionallySafeWriteScopeLock Scope(Registry.Lock);
+	UE::TWriteScopeLock Scope(Registry.Lock);
 	check(Registry.Queue.Find(Key) == nullptr);
 	Registry.Queue.Add(Key, { Version, Name, ValidatorFunc });
 }
@@ -249,7 +249,7 @@ void FCurrentCustomVersions::Unregister(const FGuid& Key)
 {
 	FStaticCustomVersionRegistry& Registry = FStaticCustomVersionRegistry::Get();
 
-	FTransactionallySafeWriteScopeLock Scope(Registry.Lock);
+	UE::TWriteScopeLock Scope(Registry.Lock);
 	Registry.Unregister(Key);
 }
 
@@ -270,7 +270,7 @@ const FCustomVersionContainer& FCustomVersionContainer::GetRegistered()
 	FStaticCustomVersionRegistry& Registry = FStaticCustomVersionRegistry::Get();
 
 	// Even though returning a reference isn't thread-safe, we can still synchronize access to the queue
-	FTransactionallySafeWriteScopeLock Scope(Registry.Lock);
+	UE::TWriteScopeLock Scope(Registry.Lock);
 	Registry.RegisterQueue();
 	return Registry.Registered;
 }
diff --git a/Engine/Source/Runtime/Core/Public/Misc/ConfigCacheIni.h b/Engine/Source/Runtime/Core/Public/Misc/ConfigCacheIni.h
--- a/Engine/Source/Runtime/Core/Public/Misc/ConfigCacheIni.h
+++ b/Engine/Source/Runtime/Core/Public/Misc/ConfigCacheIni.h
@@ -31,6 +31,7 @@
 #include "Misc/ConfigAccessTracking.h"
 #include "Misc/ConfigTypes.h"
 #include "Misc/Paths.h"
+#include "Misc/ScopeRWLock.h"
 #include "Misc/TransactionallySafeRWScopeLock.h"
 #include "Serialization/Archive.h"
 #include "Serialization/StructuredArchive.h"
@@ -654,7 +655,7 @@ public:
 
 	FORCEINLINE void Empty(int32 ExpectedNumElements = 0)
 	{
-		FTransactionallySafeWriteScopeLock ScopeLock(ConfigFileMapLock);
+		UE::TWriteScopeLock ScopeLock(ConfigFileMapLock);
 		FConfigFileMap::Empty(ExpectedNumElements);
 	}
 
@@ -678,50 +679,50 @@ public:
 
 	FORCEINLINE int32 Remove(KeyConstPointerType InKey)
 	{
-		FTransactionallySafeWriteScopeLock ScopeLock(ConfigFileMapLock);
+		UE::TWriteScopeLock ScopeLock(ConfigFileMapLock);
 		return FConfigFileMap::Remove(InKey);
 	}
 
 	FORCEINLINE ValueType& Add(const KeyType& InKey, const ValueType& InValue)
 	{
-		FTransactionallySafeWriteScopeLock ScopeLock(ConfigFileMapLock);
+		UE::TWriteScopeLock ScopeLock(ConfigFileMapLock);
 		return FConfigFileMap::Add(InKey, InValue);
 	}
 
 	FORCEINLINE ValueType& Add(const KeyType&  InKey, ValueType&& InValue)
 	{
-		FTransactionallySafeWriteScopeLock ScopeLock(ConfigFileMapLock);
+		UE::TWriteScopeLock ScopeLock(ConfigFileMapLock);
 		return FConfigFileMap::Add(InKey, MoveTempIfPossible(InValue));
 	}
 
 	FORCEINLINE ValueType& Add(KeyType&& InKey, const ValueType&  InValue)
 	{
-		FTransactionallySafeWriteScopeLock ScopeLock(ConfigFileMapLock);
+		UE::TWriteScopeLock ScopeLock(ConfigFileMapLock);
 		return FConfigFileMap::Add(MoveTempIfPossible(InKey), InValue);
 	}
 
 	FORCEINLINE ValueType& Add(KeyType&& InKey,	ValueType&& InValue)
 	{
-		FTransactionallySafeWriteScopeLock ScopeLock(ConfigFileMapLock);
+		UE::TWriteScopeLock ScopeLock(ConfigFileMapLock);
 		return FConfigFileMap::Add(MoveTempIfPossible(InKey), MoveTempIfPossible(InValue));
 	}
 	
 	FORCEINLINE void Append(TMap<FString, FConfigSection> Other)
 	{
-		FTransactionallySafeWriteScopeLock ScopeLock(ConfigFileMapLock);
+		UE::TWriteScopeLock ScopeLock(ConfigFileMapLock);
 		FConfigFileMap::Append(MoveTemp(Other));
 	}
 
 	FORCEINLINE void Reset()
 	{
-		FTransactionallySafeWriteScopeLock ScopeLock(ConfigFileMapLock);
+		UE::TWriteScopeLock ScopeLock(ConfigFileMapLock);
 		FConfigFileMap::Reset();
 	}
 
 	UE_DEPRECATED(5.4, "Use FindOrAddConfigSection, and/or use the new AddToSection, etc APIs to modify sections without retrieving the section. See top of ConfigCacheIni.h for more info.")
 	FORCEINLINE ValueType& FindOrAdd(const FString& Key)
 	{
-		FTransactionallySafeWriteScopeLock ScopeLock(ConfigFileMapLock);
+		UE::TWriteScopeLock ScopeLock(ConfigFileMapLock);
 		return FConfigFileMap::FindOrAdd(Key);
 	}
 
diff --git a/Engine/Source/Runtime/Core/Public/Misc/TransactionallySafeRWScopeLock.h b/Engine/Source/Runtime/Core/Public/Misc/TransactionallySafeRWScopeLock.h
--- a/Engine/Source/Runtime/Core/Public/Misc/TransactionallySafeRWScopeLock.h
+++ b/Engine/Source/Runtime/Core/Public/Misc/TransactionallySafeRWScopeLock.h
@@ -26,23 +26,3 @@ private:
 	FTransactionallySafeRWLock& Lock;
 };
 
-// A transactionally safe write scope lock - that uses a transactionally safe read-write lock to back it.
-class FTransactionallySafeWriteScopeLock
-{
-public:
-	UE_NODISCARD_CTOR explicit FTransactionallySafeWriteScopeLock(FTransactionallySafeRWLock& InLock)
-		: Lock(InLock)
-	{
-		Lock.WriteLock();
-	}
-
-	~FTransactionallySafeWriteScopeLock()
-	{
-		Lock.WriteUnlock();
-	}
-
-private:
-	UE_NONCOPYABLE(FTransactionallySafeWriteScopeLock);
-
-	FTransactionallySafeRWLock& Lock;
-};
diff --git a/Engine/Source/Runtime/CoreUObject/Private/UObject/LinkerInstancingContext.cpp b/Engine/Source/Runtime/CoreUObject/Private/UObject/LinkerInstancingContext.cpp
--- a/Engine/Source/Runtime/CoreUObject/Private/UObject/LinkerInstancingContext.cpp
+++ b/Engine/Source/Runtime/CoreUObject/Private/UObject/LinkerInstancingContext.cpp
@@ -8,6 +8,7 @@
 #include "Containers/UnrealString.h"
 #include "Templates/Tuple.h"
 #include "Misc/AutomationTest.h"
+#include "Misc/ScopeRWLock.h"
 #include "Misc/TransactionallySafeRWLock.h"
 #include "Misc/TransactionallySafeRWScopeLock.h"
 
@@ -51,7 +52,7 @@ public:
 		
 	void AddPackageMapping(FName Original, FName Instanced)
 	{
-		FTransactionallySafeWriteScopeLock ScopeLock(Lock);
+		UE::TWriteScopeLock ScopeLock(Lock);
 		LLM_SCOPE_BYTAG(Loading_LinkerInstancingContext);
 		InstancedPackageMap.AddPackageMapping(Original, Instanced);
 	}
@@ -93,27 +94,27 @@ public:
 
 	void EnableAutomationTest()
 	{
-		FTransactionallySafeWriteScopeLock ScopeLock(Lock);
+		UE::TWriteScopeLock ScopeLock(Lock);
 		InstancedPackageMap.EnableAutomationTest();
 	}
 
 	void BuildPackageMapping(FName Original, FName Instanced, bool bInSoftObjectPathRemappingEnabled)
 	{
-		FTransactionallySafeWriteScopeLock ScopeLock(Lock);
+		UE::TWriteScopeLock ScopeLock(Lock);
 		LLM_SCOPE_BYTAG(Loading_LinkerInstancingContext);
 		InstancedPackageMap.BuildPackageMapping(Original, Instanced, bInSoftObjectPathRemappingEnabled);
 	}
 
 	void AddTag(FName NewTag)
 	{
-		FTransactionallySafeWriteScopeLock ScopeLock(Lock);
+		UE::TWriteScopeLock ScopeLock(Lock);
 		LLM_SCOPE_BYTAG(Loading_LinkerInstancingContext);
 		Tags.Add(NewTag);
 	}
 
 	void AppendTags(const TSet<FName>& NewTags)
 	{
-		FTransactionallySafeWriteScopeLock ScopeLock(Lock);
+		UE::TWriteScopeLock ScopeLock(Lock);
 		LLM_SCOPE_BYTAG(Loading_LinkerInstancingContext);
 		Tags.Append(NewTags);
 	}
@@ -136,14 +137,14 @@ public:
 
 	void AddPathMapping(FSoftObjectPath Original, FSoftObjectPath Instanced)
 	{
-		FTransactionallySafeWriteScopeLock ScopeLock(Lock);
+		UE::TWriteScopeLock ScopeLock(Lock);
 		LLM_SCOPE_BYTAG(Loading_LinkerInstancingContext);
 		PathMapping.Emplace(Original.GetAssetPath(), Instanced.GetAssetPath());
 	}
 
 	void SetPackageMappingFunc(TFunction<FName(FName)> InInstancedPackageMapFunc)
 	{
-		FTransactionallySafeWriteScopeLock ScopeLock(Lock);
+		UE::TWriteScopeLock ScopeLock(Lock);
 		InstancedPackageMapFunc = MoveTemp(InInstancedPackageMapFunc);
 	}
 
@@ -169,7 +170,7 @@ public:
 
 		if (bAddPackageMapping)
 		{
-			FTransactionallySafeWriteScopeLock ScopeLock(Lock);
+			UE::TWriteScopeLock ScopeLock(Lock);
 			LLM_SCOPE_BYTAG(Loading_LinkerInstancingContext);
 			InstancedPackageMap.AddPackageMapping(PackageName, RemappedPackageName);
 		}
diff --git a/Engine/Source/Runtime/CoreUObject/Public/UObject/UObjectAnnotation.h b/Engine/Source/Runtime/CoreUObject/Public/UObject/UObjectAnnotation.h
--- a/Engine/Source/Runtime/CoreUObject/Public/UObject/UObjectAnnotation.h
+++ b/Engine/Source/Runtime/CoreUObject/Public/UObject/UObjectAnnotation.h
@@ -10,6 +10,7 @@
 #include "UObject/UObjectArray.h"
 #include "UObject/UObjectBaseUtility.h"
 #include "Misc/ScopeLock.h"
+#include "Misc/ScopeRWLock.h"
 #include "Misc/TransactionallySafeCriticalSection.h"
 #include "Misc/TransactionallySafeRWLock.h"
 #include "Misc/TransactionallySafeRWScopeLock.h"
@@ -762,13 +763,13 @@ public:
 	 */
 	void AddAnnotation(int32 Index, const TAnnotation& Annotation)
 	{
-		FTransactionallySafeWriteScopeLock AnnotationArrayLock(AnnotationArrayCritical);
+		UE::TWriteScopeLock AnnotationArrayLock(AnnotationArrayCritical);
 		AddAnnotationInternal(Index, Annotation);
 	}
 
 	void AddAnnotation(int32 Index, TAnnotation&& Annotation)
 	{
-		FTransactionallySafeWriteScopeLock AnnotationArrayLock(AnnotationArrayCritical);
+		UE::TWriteScopeLock AnnotationArrayLock(AnnotationArrayCritical);
 		AddAnnotationInternal(Index, MoveTemp(Annotation));
 	}
 
@@ -785,7 +786,7 @@ public:
 	 */
 	TAnnotation& AddOrGetAnnotation(int32 Index, TFunctionRef<TAnnotation()> NewAnnotationFn)
 	{		
-		FTransactionallySafeWriteScopeLock AnnotationArrayLock(AnnotationArrayCritical);
+		UE::TWriteScopeLock AnnotationArrayLock(AnnotationArrayCritical);
 		
 		if (NumAnnotations == 0 && Chunks.Num() == 0)
 		{
@@ -838,7 +839,7 @@ public:
 	 */
 	void RemoveAnnotation(int32 Index)
 	{
-		FTransactionallySafeWriteScopeLock AnnotationArrayLock(AnnotationArrayCritical);
+		UE::TWriteScopeLock AnnotationArrayLock(AnnotationArrayCritical);
 		FreeAnnotation(Index);
 	}
 
@@ -932,7 +933,7 @@ public:
 	void RemoveAllAnnotations()
 	{
 		bool bHadAnnotations = (NumAnnotations > 0);	
-		FTransactionallySafeWriteScopeLock AnnotationArrayLock(AnnotationArrayCritical);
+		UE::TWriteScopeLock AnnotationArrayLock(AnnotationArrayCritical);
 		FreeAllAnnotations();
 		if (bHadAnnotations)
 		{
@@ -951,7 +952,7 @@ public:
 	 */
 	void TrimAnnotations()
 	{
-		FTransactionallySafeWriteScopeLock AnnotationArrayLock(AnnotationArrayCritical);
+		UE::TWriteScopeLock AnnotationArrayLock(AnnotationArrayCritical);
 		for (TAnnotationChunk& Chunk : Chunks)
 		{
 			if (Chunk.Num == 0 && Chunk.Items)
diff --git a/Engine/Source/Runtime/Engine/Private/Animation/Skeleton.cpp b/Engine/Source/Runtime/Engine/Private/Animation/Skeleton.cpp
--- a/Engine/Source/Runtime/Engine/Private/Animation/Skeleton.cpp
+++ b/Engine/Source/Runtime/Engine/Private/Animation/Skeleton.cpp
@@ -24,6 +24,7 @@
 #include "Logging/MessageLog.h"
 #include "ComponentReregisterContext.h"
 #include "Engine/SkeletalMeshSocket.h"
+#include "Misc/ScopeRWLock.h"
 #include "Misc/ScopedSlowTask.h"
 #include "Animation/AnimBlueprint.h"
 #include "UObject/AnimObjectVersion.h"
@@ -755,7 +756,7 @@ bool USkeleton::IsCompatibleMesh(const USkinnedAsset* InSkinnedAsset, bool bDoPa
 void USkeleton::ClearCacheData()
 {
 	{
-		FTransactionallySafeWriteScopeLock Lock(SkinnedAssetLinkupCacheLock);
+		UE::TWriteScopeLock Lock(SkinnedAssetLinkupCacheLock);
 		SkinnedAssetLinkupCache.Empty();
 	}
 
@@ -790,7 +791,7 @@ const FSkeletonToMeshLinkup& USkeleton::AddMeshLinkupData(const USkinnedAsset* I
 	const uint32 SkinnedAssetHash = GetTypeHash(SkinnedAssetKey);
 
 	{
-		FTransactionallySafeWriteScopeLock Lock(SkinnedAssetLinkupCacheLock);
+		UE::TWriteScopeLock Lock(SkinnedAssetLinkupCacheLock);
 		if (const TUniquePtr<FSkeletonToMeshLinkup>* SkeletonToMeshLinkupPtr = SkinnedAssetLinkupCache.FindByHash(SkinnedAssetHash, SkinnedAssetKey))
 		{
 			// While we were building the linkup data, another thread beat us to it
@@ -838,7 +839,7 @@ int32 USkeleton::GetMeshLinkupIndex(const USkinnedAsset* InSkinnedAsset)
 void USkeleton::RemoveLinkup(const USkinnedAsset* InSkinnedAsset)
 {
 	{
-		FTransactionallySafeWriteScopeLock Lock(SkinnedAssetLinkupCacheLock);
+		UE::TWriteScopeLock Lock(SkinnedAssetLinkupCacheLock);
 		SkinnedAssetLinkupCache.Remove(InSkinnedAsset);
 	}
 
diff --git a/Engine/Source/Runtime/Engine/Private/VisualLogger/VisualLogger.cpp b/Engine/Source/Runtime/Engine/Private/VisualLogger/VisualLogger.cpp
--- a/Engine/Source/Runtime/Engine/Private/VisualLogger/VisualLogger.cpp
+++ b/Engine/Source/Runtime/Engine/Private/VisualLogger/VisualLogger.cpp
@@ -4,6 +4,7 @@
 #include "EngineStats.h"
 #include "Containers/StringFwd.h"
 #include "Misc/CommandLine.h"
+#include "Misc/ScopeRWLock.h"
 #include "Modules/ModuleManager.h"
 #include "AI/NavDataGenerator.h" // IWYU pragma: keep
 #include "AI/NavigationSystemBase.h"
@@ -1059,7 +1060,7 @@ UObject* FVisualLogger::FindRedirectionInternal(const UObject* Object) const
 
 void FVisualLogger::CleanupRedirects()
 {
-	FTransactionallySafeWriteScopeLock Lock(RedirectRWLock);
+	UE::TWriteScopeLock Lock(RedirectRWLock);
 	for (auto It = ChildToOwnerMap.CreateIterator(); It; ++It)
 	{
 		if(!It.Value().IsValid())
diff --git a/Engine/Source/Runtime/Engine/Public/VisualLogger/VisualLogger.h b/Engine/Source/Runtime/Engine/Public/VisualLogger/VisualLogger.h
--- a/Engine/Source/Runtime/Engine/Public/VisualLogger/VisualLogger.h
+++ b/Engine/Source/Runtime/Engine/Public/VisualLogger/VisualLogger.h
@@ -14,6 +14,7 @@
 #include "Containers/Ticker.h"
 #include "Misc/TransactionallySafeRWLock.h"
 #include "Misc/TransactionallySafeRWScopeLock.h"
+#include "Misc/ScopeRWLock.h"
 
 #if ENABLE_VISUAL_LOG
 
@@ -733,7 +734,7 @@ public:
 		FVisualLogger& Logger = FVisualLogger::Get();
 		UObject* NewRedirection = nullptr;
 		{
-			FTransactionallySafeWriteScopeLock Lock(Logger.RedirectRWLock);
+			UE::TWriteScopeLock Lock(Logger.RedirectRWLock);
 			NewRedirection = Logger.RedirectInternal(FromObject, ToObject);
 		}
 		UE_CVLOG(FromObject != nullptr && NewRedirection != nullptr, FromObject, LogVisual, Log, TEXT("Redirected '%s' to '%s'"), *FromObject->GetName(), *NewRedirection->GetName());
diff --git a/Engine/Source/Runtime/MovieScene/Private/MovieScenePlayer.cpp b/Engine/Source/Runtime/MovieScene/Private/MovieScenePlayer.cpp
--- a/Engine/Source/Runtime/MovieScene/Private/MovieScenePlayer.cpp
+++ b/Engine/Source/Runtime/MovieScene/Private/MovieScenePlayer.cpp
@@ -15,6 +15,7 @@
 #include "MovieSceneSequenceID.h"
 #include "UniversalObjectLocatorResolveParameterBuffer.inl"
 #include "MovieSceneBindingReferences.h"
+#include "Misc/ScopeRWLock.h"
 #include "Misc/TransactionallySafeRWLock.h"
 #include "Misc/TransactionallySafeRWScopeLock.h"
 
@@ -56,7 +57,7 @@ PRAGMA_DISABLE_DEPRECATION_WARNINGS
 IMovieScenePlayer::IMovieScenePlayer()
 	: BindingOverrides(StaticBindingOverrides.BindingOverrides)
 {
-	FTransactionallySafeWriteScopeLock ScopeLock(UE::MovieScene::GGlobalPlayerRegistryLock);
+	UE::TWriteScopeLock ScopeLock(UE::MovieScene::GGlobalPlayerRegistryLock);
 
 	UE::MovieScene::GGlobalPlayerRegistry.Shrink();
 	UniqueIndex = UE::MovieScene::GGlobalPlayerRegistry.Add(this);
@@ -69,7 +70,7 @@ PRAGMA_ENABLE_DEPRECATION_WARNINGS
 PRAGMA_DISABLE_DEPRECATION_WARNINGS
 IMovieScenePlayer::~IMovieScenePlayer()
 {	
-	FTransactionallySafeWriteScopeLock ScopeLock(UE::MovieScene::GGlobalPlayerRegistryLock);
+	UE::TWriteScopeLock ScopeLock(UE::MovieScene::GGlobalPlayerRegistryLock);
 
 	UE::MovieScene::GGlobalPlayerUpdateFlags[UniqueIndex] = 0;
 	UE::MovieScene::GGlobalPlayerRegistry.RemoveAt(UniqueIndex, 1);
