[core] fix compilation glitch when building without UE_VALIDATE_FORMAT_STRINGS

diff --git a/Engine/Source/Runtime/Core/Public/Containers/UnrealString.h.inl b/Engine/Source/Runtime/Core/Public/Containers/UnrealString.h.inl
--- a/Engine/Source/Runtime/Core/Public/Containers/UnrealString.h.inl
+++ b/Engine/Source/Runtime/Core/Public/Containers/UnrealString.h.inl
@@ -1562,9 +1562,6 @@ public:
 	/** Removes spaces from the string.  I.E. "Spaces Are Cool" --> "SpacesAreCool". */
 	CORE_API void RemoveSpacesInline();
 
-
-#if UE_VALIDATE_FORMAT_STRINGS
-
 	/**
 	 * Constructs a string similarly to how classic sprintf works.
 	 *
@@ -1590,29 +1587,6 @@ public:
 		return *this;
 	}
 
-#else
-
-	template <typename FmtType, typename... Types>
-	[[nodiscard]] static UE_STRING_CLASS Printf(const FmtType& Fmt, Types... Args)
-	{
-		static_assert(TIsArrayOrRefOfTypeByPredicate<FmtType, TIsCharEncodingCompatibleWithPrintfFmt>::Value, "Formatting string must be a literal " PREPROCESSOR_TO_STRING(UE_STRING_PRINTF_FMT_CHARTYPE) " array.");
-		static_assert((TIsValidVariadicFunctionArg<Types>::Value && ...), "Invalid argument(s) passed to Printf");
-
-		return PrintfImpl((const ElementType*)Fmt, Args...);
-	}
-
-	template <typename FmtType, typename... Types>
-	UE_STRING_CLASS& Appendf(const FmtType& Fmt, Types... Args)
-	{
-		static_assert(TIsArrayOrRefOfTypeByPredicate<FmtType, TIsCharEncodingCompatibleWithPrintfFmt>::Value, "Formatting string must be a literal " PREPROCESSOR_TO_STRING(UE_STRING_PRINTF_FMT_CHARTYPE) " array.");
-		static_assert((TIsValidVariadicFunctionArg<Types>::Value && ...), "Invalid argument(s) passed to Appendf");
-
-		AppendfImpl(*this, (const ElementType*)Fmt, Args...);
-		return *this;
-	}
-
-#endif // UE_VALIDATE_FORMAT_STRINGS
-
 private:
 	static CORE_API UE_STRING_CLASS VARARGS PrintfImpl(const ElementType* Fmt, ...);
 	static CORE_API void VARARGS AppendfImpl(UE_STRING_CLASS& AppendToMe, const ElementType* Fmt, ...);
diff --git a/Engine/Source/Runtime/Core/Public/String/FormatStringSan.h b/Engine/Source/Runtime/Core/Public/String/FormatStringSan.h
--- a/Engine/Source/Runtime/Core/Public/String/FormatStringSan.h
+++ b/Engine/Source/Runtime/Core/Public/String/FormatStringSan.h
@@ -501,7 +501,7 @@ namespace UE::Core::Private
 		{
 			const CharType* FormatString;
 
-#if defined(__cpp_consteval)
+#if UE_VALIDATE_FORMAT_STRINGS && defined(__cpp_consteval)
 			template <int32 N>
 			consteval TCheckedFormatStringPrivate(const CharType (&Fmt)[N]) : FormatString{Fmt}
 			{
@@ -542,6 +542,7 @@ namespace UE::Core::Private
 			template <int32 N>
 			TCheckedFormatStringPrivate(const CharType(&Fmt)[N]) : FormatString{Fmt}
 			{
+				static_assert((TIsValidVariadicFunctionArg<ArgTypes>::Value && ...), "Invalid argument(s) passed to Printf");
 			}
 
 #endif // defined(__cpp_consteval)
