[Scheduler/Unittests]
- Fix potential failure if the global task was stealed by another thread
- Fix task that could get out of scope before being marked as completed

#rnx

diff --git a/Engine/Source/Runtime/Core/Private/Async/Fundamental/UnitTests.cpp b/Engine/Source/Runtime/Core/Private/Async/Fundamental/UnitTests.cpp
--- a/Engine/Source/Runtime/Core/Private/Async/Fundamental/UnitTests.cpp
+++ b/Engine/Source/Runtime/Core/Private/Async/Fundamental/UnitTests.cpp
@@ -535,6 +535,9 @@ namespace Tasks2Tests
 	{
 		using namespace LowLevelTasks;
 
+		// Avoid any kind of task stealing.
+		LowLevelTasks::FScheduler::Get().RestartWorkers(0, 1);
+
 		FTask RootTask;
 		FTask LocalTask;
 		FTask GlobalTask;
@@ -542,24 +545,21 @@ namespace Tasks2Tests
 		std::atomic_bool GlobalRun{ false };
 		std::atomic_bool LocalRun{ false };
 
-		UE::FManualResetEvent FinishEvent;
-
 		// This test launch 2 tasks
 		// One global that is high prio and supposed to run before the local
 		// One local that is normal prio which is supposed to run after the global one
 
 		LocalTask.Init(TEXT("LocalTask"), ETaskPriority::BackgroundNormal,
 			[&]()
-			{ 
-				LocalRun = true; 
+			{
+				LocalRun = true;
 				verify(GlobalRun == true);
-				FinishEvent.Notify();
 			}
 		);
 
-		GlobalTask.Init(TEXT("GlobalTask"), ETaskPriority::BackgroundHigh, 
+		GlobalTask.Init(TEXT("GlobalTask"), ETaskPriority::BackgroundHigh,
 			[&]()
-			{ 
+			{
 				GlobalRun = true;
 				verify(LocalRun == false);
 			}
@@ -569,17 +569,20 @@ namespace Tasks2Tests
 			TEXT("RootTask"), ETaskPriority::BackgroundNormal,
 			[&]()
 			{
-				// Both task are launched from inside a worker with no addition wake up so that both task should end up being run in the same worker
+				// Both task are launched from inside a worker with no additional wake up so that both task should end up being run in the same worker
 				verify(TryLaunch(LocalTask, EQueuePreference::LocalQueuePreference, false /*bWakeUpWorker*/));
 				verify(TryLaunch(GlobalTask, EQueuePreference::GlobalQueuePreference, false /*bWakeUpWorker*/));
 			});
 
 		TryLaunch(RootTask);
 
-		verify(FinishEvent.WaitFor(UE::FMonotonicTimeSpan::FromSeconds(1)));
+		while (!LocalTask.IsCompleted()) {};
+		while (!GlobalTask.IsCompleted()) {};
 		verify(GlobalRun == true);
 		verify(LocalRun == true);
 
+		LowLevelTasks::FScheduler::Get().RestartWorkers();
+
 		return true;
 	};
 
