[AutoRTFM] Add support for disabling memory validation for a specific write.

This will be required to correctly handle UObject in-place reconstruction, while keeping the validator happy.

This CL extracts some of the bit packing in FWriteLogEntry with a new FMemoryLocation struct, shared between the write log and the hit set.

diff --git a/Engine/Source/Programs/AutoRTFMTests/Private/Open.cpp b/Engine/Source/Programs/AutoRTFMTests/Private/Open.cpp
--- a/Engine/Source/Programs/AutoRTFMTests/Private/Open.cpp
+++ b/Engine/Source/Programs/AutoRTFMTests/Private/Open.cpp
@@ -260,6 +260,24 @@ TEST_CASE("Open.Collision")
             REQUIRE(WarningContext.GetWarnings().empty());
         }
 
+        SECTION("Transact(Open(RecordOpenWriteNoMemoryValidation()), Open(Assign))")
+        {
+            AutoRTFM::Transact([&] ()
+            {
+                AutoRTFM::Open([&] () 
+                {
+                    AutoRTFM::RecordOpenWriteNoMemoryValidation(&I);
+                });
+                AutoRTFM::Open([&] () 
+                {
+                    I = 42;
+                });
+            });
+
+            REQUIRE(I == 42);
+            REQUIRE(WarningContext.GetWarnings().empty());
+        }
+
         SECTION("Transact(Open(Assign), Assign)")
         {
             AutoRTFM::Transact([&] ()
@@ -505,4 +523,46 @@ TEST_CASE("Open.Collision")
         REQUIRE_THAT(WarningContext.GetWarnings(),
             Catch::Matchers::VectorContains(FString(AutoRTFMTestUtils::kMemoryModifiedWarning)));
     }
+
+    SECTION("Transact(Open(RecordOpenWrite()), Open(Assign))")
+    {
+        AutoRTFM::Transact([&] ()
+        {
+            AutoRTFM::Open([&] () 
+            {
+                AutoRTFM::RecordOpenWrite(&I);
+            });
+            AutoRTFM::Open([&] () 
+            {
+                I = 42;
+            });
+        });
+
+        REQUIRE(I == 42);
+        REQUIRE_THAT(WarningContext.GetWarnings(),
+            Catch::Matchers::VectorContains(FString(AutoRTFMTestUtils::kMemoryModifiedWarning)));
+    }
+
+    SECTION("Transact(Open(RecordOpenWriteNoMemoryValidation()), Open(RecordOpenWrite()), Open(Assign))")
+    {
+        AutoRTFM::Transact([&] ()
+        {
+            AutoRTFM::Open([&] () 
+            {
+                AutoRTFM::RecordOpenWriteNoMemoryValidation(&I);
+            });
+            AutoRTFM::Open([&] () 
+            {
+                AutoRTFM::RecordOpenWrite(&I);
+            });
+            AutoRTFM::Open([&] () 
+            {
+                I = 42;
+            });
+        });
+
+        REQUIRE_THAT(WarningContext.GetWarnings(),
+            Catch::Matchers::VectorContains(FString(AutoRTFMTestUtils::kMemoryModifiedWarning)));
+    }
+
 }
diff --git a/Engine/Source/Runtime/Core/Private/AutoRTFM/API.cpp b/Engine/Source/Runtime/Core/Private/AutoRTFM/API.cpp
--- a/Engine/Source/Runtime/Core/Private/AutoRTFM/API.cpp
+++ b/Engine/Source/Runtime/Core/Private/AutoRTFM/API.cpp
@@ -452,7 +452,7 @@ AUTORTFM_CORE_API void RTFM_autortfm_open(autortfm_memory_validation_level Memor
 	}
 }
 
-AUTORTFM_CORE_API void RTFM_autortfm_record_open_write(void*, size_t)
+AUTORTFM_CORE_API void RTFM_autortfm_record_open_write_err(void*, size_t)
 {
 	UE_LOG(LogAutoRTFM, Fatal, TEXT("The function `autortfm_record_open_write` was called from closed code."));
 }
@@ -776,15 +776,29 @@ extern "C" void autortfm_record_open_write(void* Ptr, size_t Size)
 {
 	if (autortfm_is_closed())
 	{
-		return RTFM_autortfm_record_open_write(Ptr, Size);
+		return RTFM_autortfm_record_open_write_err(Ptr, Size);
 	}
-	else if (!FContext::Get()->IsTransactional())
+	else if (FContext::Get()->IsTransactional())
 	{
-		return;
+		if (FTransaction* CurrentTransaction = FContext::Get()->GetCurrentTransaction())
+		{
+			CurrentTransaction->RecordWrite(Ptr, Size);
+		}
 	}
-	else if (FTransaction* CurrentTransaction = FContext::Get()->GetCurrentTransaction())
+}
+
+extern "C" void autortfm_record_open_write_no_memory_validation(void* Ptr, size_t Size)
+{
+	if (autortfm_is_closed())
+	{
+		return RTFM_autortfm_record_open_write_err(Ptr, Size);
+	}
+	else if (FContext::Get()->IsTransactional())
 	{
-		CurrentTransaction->RecordWrite(Ptr, Size);
+		if (FTransaction* CurrentTransaction = FContext::Get()->GetCurrentTransaction())
+		{
+			CurrentTransaction->RecordWrite(Ptr, Size, /* bNoMemoryValidation */ true);
+		}
 	}
 }
 
diff --git a/Engine/Source/Runtime/Core/Private/AutoRTFM/HitSet.h b/Engine/Source/Runtime/Core/Private/AutoRTFM/HitSet.h
--- a/Engine/Source/Runtime/Core/Private/AutoRTFM/HitSet.h
+++ b/Engine/Source/Runtime/Core/Private/AutoRTFM/HitSet.h
@@ -2,7 +2,7 @@
 
 #pragma once
 
-#include "TaggedPtr.h"
+#include "MemoryLocation.h"
 #include "Utils.h"
 
 namespace AutoRTFM
@@ -10,7 +10,7 @@ namespace AutoRTFM
 	class FHitSet final
 	{
 	public:
-		using Key = TTaggedPtr<void>;
+		using Key = FMemoryLocation;
 
 	private:
 		// TODO: Revisit a good probe depth for the hashset.
diff --git a/Engine/Source/Runtime/Core/Private/AutoRTFM/MemoryLocation.h b/Engine/Source/Runtime/Core/Private/AutoRTFM/MemoryLocation.h
new file mode 100644
--- /dev/null
+++ b/Engine/Source/Runtime/Core/Private/AutoRTFM/MemoryLocation.h
@@ -0,0 +1,27 @@
+// Copyright Epic Games, Inc. All Rights Reserved.
+
+#pragma once
+
+#include "HAL/Platform.h"
+#include "TaggedPtr.h"
+
+namespace AutoRTFM
+{
+	struct FMemoryLocationTag
+	{
+		// Number of bits to use for the Size field.
+		static constexpr int SizeBits = 14;
+		// Maximum value that can be stored in Size.
+		static constexpr uint16 MaxSize = (1u << SizeBits) - 1;
+
+		// Size of the write.
+		uint16 Size : SizeBits;
+		// If set, then the FMemoryLocation's pointer holds the original data.
+		uint16 bIsSmall : 1;
+		// If set, then the write shouldn't be considered by the memory validator.
+		uint16 bNoMemoryValidation : 1;
+	};
+	static_assert(sizeof(FMemoryLocationTag) == 2);
+	
+	using FMemoryLocation = TTaggedPtr<void, FMemoryLocationTag>;
+}
diff --git a/Engine/Source/Runtime/Core/Private/AutoRTFM/TaggedPtr.h b/Engine/Source/Runtime/Core/Private/AutoRTFM/TaggedPtr.h
--- a/Engine/Source/Runtime/Core/Private/AutoRTFM/TaggedPtr.h
+++ b/Engine/Source/Runtime/Core/Private/AutoRTFM/TaggedPtr.h
@@ -80,12 +80,12 @@ namespace AutoRTFM
 	    	return reinterpret_cast<T*>(Payload & PointerMask);
 	    }
 
-	    UE_AUTORTFM_FORCEINLINE void SetTopTag(const TopTagT Tag)
+	    UE_AUTORTFM_FORCEINLINE TopTagT& TopTag()
 	    {
-	    	*reinterpret_cast<TopTagT*>(reinterpret_cast<uint8_t*>(&Payload) + OffsetOfTopTag) = Tag;
+	    	return *reinterpret_cast<TopTagT*>(reinterpret_cast<uint8_t*>(&Payload) + OffsetOfTopTag);
 	    }
 
-	    UE_AUTORTFM_FORCEINLINE TopTagT GetTopTag() const
+	    UE_AUTORTFM_FORCEINLINE const TopTagT& TopTag() const
 	    {
 	    	return *reinterpret_cast<const TopTagT*>(reinterpret_cast<const uint8_t*>(&Payload) + OffsetOfTopTag);
 	    }
diff --git a/Engine/Source/Runtime/Core/Private/AutoRTFM/Transaction.cpp b/Engine/Source/Runtime/Core/Private/AutoRTFM/Transaction.cpp
--- a/Engine/Source/Runtime/Core/Private/AutoRTFM/Transaction.cpp
+++ b/Engine/Source/Runtime/Core/Private/AutoRTFM/Transaction.cpp
@@ -233,7 +233,7 @@ void FTransaction::CommitNested()
 		Parent->WriteLog.Push(Write);
 
 		FHitSet::Key HitSetEntry(Write.GetOriginal());
-		HitSetEntry.SetTopTag(static_cast<uint16_t>(Write.GetSize()));
+		HitSetEntry.TopTag().Size = static_cast<uint16_t>(Write.GetSize());
 
         Parent->HitSet.FindOrTryInsert(HitSetEntry);
     }
@@ -503,6 +503,11 @@ FTransaction::TWriteHash FTransaction::CalculateNestedWriteHash(size_t NumWriteE
 		}
 
 		const FWriteLogEntry& Entry = *Iter;
+		if (UNLIKELY(Entry.IgnoreForValidation()))
+		{
+			continue; // Next write
+		}
+
 		const uint8* Original = reinterpret_cast<const uint8*>(Entry.GetOriginal());
 		const size_t Size = Entry.GetSize();
 		if (Size > sizeof(FSmallValue))
diff --git a/Engine/Source/Runtime/Core/Private/AutoRTFM/Transaction.h b/Engine/Source/Runtime/Core/Private/AutoRTFM/Transaction.h
--- a/Engine/Source/Runtime/Core/Private/AutoRTFM/Transaction.h
+++ b/Engine/Source/Runtime/Core/Private/AutoRTFM/Transaction.h
@@ -118,8 +118,8 @@ public:
     bool AttemptToCommit();
 
     // Record that a write is about to occur at the given LogicalAddress of Size bytes.
-    void RecordWrite(void* LogicalAddress, size_t Size);
-    void RecordWriteMaxPageSized(void* LogicalAddress, size_t Size);
+    void RecordWrite(void* LogicalAddress, size_t Size, bool bNoMemoryValidation = false);
+    void RecordWriteMaxPageSized(void* LogicalAddress, size_t Size, bool bNoMemoryValidation = false);
     template<unsigned SIZE> void RecordWrite(void* LogicalAddress);
 	template<unsigned SIZE> void RecordWriteInsertedSlow(void* LogicalAddress);
 	template<unsigned SIZE> void RecordWriteNotInsertedSlow(void* LogicalAddress);
diff --git a/Engine/Source/Runtime/Core/Private/AutoRTFM/TransactionInlines.h b/Engine/Source/Runtime/Core/Private/AutoRTFM/TransactionInlines.h
--- a/Engine/Source/Runtime/Core/Private/AutoRTFM/TransactionInlines.h
+++ b/Engine/Source/Runtime/Core/Private/AutoRTFM/TransactionInlines.h
@@ -3,11 +3,13 @@
 #pragma once
 
 #include "Context.h"
+#include "MemoryLocation.h"
 #include "ScopedGuard.h"
 #include "Transaction.h"
 
 #include "HAL/Platform.h"
 #include "Utils.h"
+#include "WriteLog.h"
 
 #if PLATFORM_HAS_ASAN_INCLUDE
 #	if defined(__SANITIZE_ADDRESS__)
@@ -95,15 +97,21 @@ UE_AUTORTFM_FORCEINLINE bool FTransaction::ShouldRecordWrite(void* LogicalAddres
     return false;
 }
 
-AUTORTFM_NO_ASAN UE_AUTORTFM_FORCEINLINE void FTransaction::RecordWriteMaxPageSized(void* LogicalAddress, size_t Size)
+AUTORTFM_NO_ASAN UE_AUTORTFM_FORCEINLINE
+void FTransaction::RecordWriteMaxPageSized(void* LogicalAddress, size_t Size, bool bNoMemoryValidation /* = false */)
 {
     void* CopyAddress = WriteLogBumpAllocator.Allocate(Size);
     memcpy(CopyAddress, LogicalAddress, Size);
 
-    WriteLog.Push(FWriteLogEntry(LogicalAddress, Size, CopyAddress));
+    FMemoryLocation OriginalAndFlags(LogicalAddress);
+    FMemoryLocationTag& Tag = OriginalAndFlags.TopTag();
+    Tag.bNoMemoryValidation = bNoMemoryValidation;
+    Tag.Size = static_cast<uint16_t>(Size);
+    WriteLog.Push(FWriteLogEntry(OriginalAndFlags, CopyAddress));
 }
 
-AUTORTFM_NO_ASAN UE_AUTORTFM_FORCEINLINE void FTransaction::RecordWrite(void* LogicalAddress, size_t Size)
+AUTORTFM_NO_ASAN UE_AUTORTFM_FORCEINLINE
+void FTransaction::RecordWrite(void* LogicalAddress, size_t Size, bool bNoMemoryValidation /* = false */)
 {
     if (UNLIKELY(0 == Size))
     {
@@ -124,7 +132,9 @@ AUTORTFM_NO_ASAN UE_AUTORTFM_FORCEINLINE void FTransaction::RecordWrite(void* Lo
     if (Size <= 16)
     {
         FMemoryLocation Key(LogicalAddress);
-        Key.SetTopTag(static_cast<uint16_t>(Size));
+        FMemoryLocationTag& Tag = Key.TopTag();
+        Tag.bNoMemoryValidation = bNoMemoryValidation;
+        Tag.Size = static_cast<uint16_t>(Size);
 
         if (HitSet.FindOrTryInsert(Key))
         {
@@ -153,7 +163,7 @@ AUTORTFM_NO_ASAN UE_AUTORTFM_FORCEINLINE void FTransaction::RecordWrite(void* Lo
     }
 
     // Remainder at the end of the memcpy.
-    RecordWriteMaxPageSized(Address + I, Size - I);
+    RecordWriteMaxPageSized(Address + I, Size - I, bNoMemoryValidation);
 }
 
 template<unsigned SIZE> AUTORTFM_NO_ASAN UE_AUTORTFM_FORCEINLINE void FTransaction::RecordWrite(void* LogicalAddress)
@@ -167,7 +177,7 @@ template<unsigned SIZE> AUTORTFM_NO_ASAN UE_AUTORTFM_FORCEINLINE void FTransacti
     }
 
     FMemoryLocation Key(LogicalAddress);
-    Key.SetTopTag(static_cast<uint16_t>(SIZE));
+    Key.TopTag().Size = static_cast<uint16_t>(SIZE);
 
 	switch (HitSet.FindOrTryInsertNoResize(Key))
 	{
@@ -184,7 +194,7 @@ template<unsigned SIZE> AUTORTFM_NO_ASAN UE_AUTORTFM_FORCEINLINE void FTransacti
 template<unsigned SIZE> AUTORTFM_NO_ASAN FORCENOINLINE void FTransaction::RecordWriteNotInsertedSlow(void* LogicalAddress)
 {
 	FMemoryLocation Key(LogicalAddress);
-	Key.SetTopTag(static_cast<uint16_t>(SIZE));
+	Key.TopTag().Size = static_cast<uint16_t>(SIZE);
 
 	if (HitSet.FindOrTryInsert(Key))
 	{
diff --git a/Engine/Source/Runtime/Core/Private/AutoRTFM/WriteLog.h b/Engine/Source/Runtime/Core/Private/AutoRTFM/WriteLog.h
--- a/Engine/Source/Runtime/Core/Private/AutoRTFM/WriteLog.h
+++ b/Engine/Source/Runtime/Core/Private/AutoRTFM/WriteLog.h
@@ -2,58 +2,59 @@
 
 #pragma once
 
-#include "AutoRTFM/HitSet.h"
+#include "MemoryLocation.h"
 #include "Utils.h"
 #include <stdint.h>
 #include <type_traits>
 
 namespace AutoRTFM
 {
-	using FMemoryLocation = FHitSet::Key;
+	class FWriteLogEntry final
+	{
+		FMemoryLocation OriginalAndFlags;
+		uintptr_t Copy;
 
-    class FWriteLogEntry final
-    {
-		static constexpr uintptr_t IsSmallBit = 0x8000;
-
-        FMemoryLocation OriginalAndIsSmallAndSize;
-        uintptr_t Copy;
 	public:
+		UE_AUTORTFM_FORCEINLINE FWriteLogEntry() = default;
+		UE_AUTORTFM_FORCEINLINE FWriteLogEntry(const FWriteLogEntry&) = default;
+		UE_AUTORTFM_FORCEINLINE FWriteLogEntry& operator=(const FWriteLogEntry&) = default;
 
-        UE_AUTORTFM_FORCEINLINE FWriteLogEntry() = default;
-        UE_AUTORTFM_FORCEINLINE FWriteLogEntry(const FWriteLogEntry&) = default;
-        UE_AUTORTFM_FORCEINLINE FWriteLogEntry& operator=(const FWriteLogEntry&) = default;
-
-        UE_AUTORTFM_FORCEINLINE explicit FWriteLogEntry(void* Original, size_t Size, void* Copy) :
-			OriginalAndIsSmallAndSize(Original), Copy(reinterpret_cast<uintptr_t>(Copy))
-        {
-			OriginalAndIsSmallAndSize.SetTopTag(static_cast<uint16_t>(Size));
-        }
+		UE_AUTORTFM_FORCEINLINE FWriteLogEntry(void* Original, size_t Size, void* Copy) :
+			OriginalAndFlags(Original), Copy(reinterpret_cast<uintptr_t>(Copy))
+		{
+			OriginalAndFlags.TopTag().Size = static_cast<uint16_t>(Size);
+		}
 
+		UE_AUTORTFM_FORCEINLINE FWriteLogEntry(FMemoryLocation OriginalAndFlags, void* Copy) :
+			OriginalAndFlags(OriginalAndFlags), Copy(reinterpret_cast<uintptr_t>(Copy))
+		{
+		}
+		
 		template<unsigned int SIZE> UE_AUTORTFM_FORCEINLINE static FWriteLogEntry CreateSmall(void* Original)
 		{
 			static_assert(SIZE <= 8);
 
 			FWriteLogEntry Entry;
 
-			Entry.OriginalAndIsSmallAndSize = Original;
+			Entry.OriginalAndFlags = Original;
 
 			// We store the 8-byte or less write log entry in the write log itself since it is small.
 			uintptr_t Temp = 0;
 			memcpy(&Temp, Original, SIZE);
 			Entry.Copy = Temp;
 
-			size_t Size = IsSmallBit | SIZE;
+			FMemoryLocationTag& Tag = Entry.OriginalAndFlags.TopTag();
+			Tag.Size = SIZE;
+			Tag.bIsSmall = 1;
 
-			ASSERT(Size <= UINT16_MAX);
-
-			Entry.OriginalAndIsSmallAndSize.SetTopTag(static_cast<uint16_t>(Size));
+			ASSERT(Tag.Size <= FMemoryLocationTag::MaxSize);
 
 			return Entry;
 		}
 
 		UE_AUTORTFM_FORCEINLINE void* GetOriginal() const
 		{
-			return OriginalAndIsSmallAndSize.Get();
+			return OriginalAndFlags.Get();
 		}
 
 		UE_AUTORTFM_FORCEINLINE void* GetCopy()
@@ -82,14 +83,19 @@ namespace AutoRTFM
 
 		UE_AUTORTFM_FORCEINLINE size_t GetSize() const
 		{
-			return OriginalAndIsSmallAndSize.GetTopTag() & ~IsSmallBit;
+			return OriginalAndFlags.TopTag().Size;
 		}
 
 		UE_AUTORTFM_FORCEINLINE bool IsSmall() const
 		{
-			return OriginalAndIsSmallAndSize.GetTopTag() & IsSmallBit;
+			return OriginalAndFlags.TopTag().bIsSmall;
 		}
-    };
+
+		UE_AUTORTFM_FORCEINLINE bool IgnoreForValidation() const
+		{
+			return OriginalAndFlags.TopTag().bNoMemoryValidation;
+		}
+	};
 
     class FWriteLog final
     {
diff --git a/Engine/Source/Runtime/Core/Public/AutoRTFM/AutoRTFM.h b/Engine/Source/Runtime/Core/Public/AutoRTFM/AutoRTFM.h
--- a/Engine/Source/Runtime/Core/Public/AutoRTFM/AutoRTFM.h
+++ b/Engine/Source/Runtime/Core/Public/AutoRTFM/AutoRTFM.h
@@ -289,6 +289,16 @@ static UE_AUTORTFM_FORCEINLINE void autortfm_record_open_write(void* Ptr, size_t
 }
 #endif
 
+#if UE_AUTORTFM_ENABLED
+UE_AUTORTFM_API void autortfm_record_open_write_no_memory_validation(void* Ptr, size_t Size);
+#else
+static UE_AUTORTFM_FORCEINLINE void autortfm_record_open_write_no_memory_validation(void* Ptr, size_t Size)
+{
+	UE_AUTORTFM_UNUSED(Ptr);
+	UE_AUTORTFM_UNUSED(Size);
+}
+#endif
+
 #if UE_AUTORTFM_ENABLED
 UE_AUTORTFM_API void autortfm_register_open_function(void* original_function, void* new_function);
 #else
@@ -1013,6 +1023,18 @@ template<typename Type> static UE_AUTORTFM_FORCEINLINE void RecordOpenWrite(Type
 	autortfm_record_open_write(Ptr, sizeof(Type));
 }
 
+// Same as RecordOpenWrite() but marks the write as ignorable by the memory validator.
+static UE_AUTORTFM_FORCEINLINE void RecordOpenWriteNoMemoryValidation(void* Ptr, size_t Size)
+{
+	autortfm_record_open_write_no_memory_validation(Ptr, Size);
+}
+
+// Same as RecordOpenWrite() but marks the write as ignorable by the memory validator.
+template<typename Type> static UE_AUTORTFM_FORCEINLINE void RecordOpenWriteNoMemoryValidation(Type* Ptr)
+{
+	autortfm_record_open_write_no_memory_validation(Ptr, sizeof(Type));
+}
+
 // Report that a unreachable codepath is being hit. Used to manually ban certain codepaths
 // from being transactionally safe.
 static UE_AUTORTFM_FORCEINLINE void Unreachable(const UE_AUTORTFM_TCHAR* const Message = nullptr)
