[Async]
- Add a way to send tasks to the ticker using the Async API for discoverability

#rnx

diff --git a/Engine/Source/Runtime/Core/Public/Async/Async.h b/Engine/Source/Runtime/Core/Public/Async/Async.h
--- a/Engine/Source/Runtime/Core/Public/Async/Async.h
+++ b/Engine/Source/Runtime/Core/Public/Async/Async.h
@@ -16,6 +16,7 @@
 #include "Misc/Fork.h"
 #include "Misc/IQueuedWork.h"
 #include "Misc/QueuedThreadPool.h"
+#include "Containers/Ticker.h"
 #include "Stats/Stats.h"
 #include "Templates/Function.h"
 #include "Templates/UnrealTemplate.h"
@@ -28,9 +29,15 @@ enum class EAsyncExecution
 	/** Execute in Task Graph (for short running tasks). */
 	TaskGraph,
 
-	/** Execute in Task Graph on the main thread (for short running tasks). */
+	/** Execute in Task Graph on the main thread (for short running tasks that are safe to run anywhere including inside other waits). 
+	*   Note: this could end up running inside the GC or inside PostLoads, code with delegates that can't be known to be safe in advance
+	*         should opt for TaskGraphMainTick instead.
+	*/
 	TaskGraphMainThread,
 
+	/** Execute inside a Tick on the main thread (for short running tasks that needs to run at a safe place). */
+	TaskGraphMainTick,
+
 	/** Execute in separate thread if supported (for long running tasks). */
 	Thread,
 
@@ -305,7 +312,15 @@ auto Async(EAsyncExecution Execution, CallableType&& Callable, TUniqueFunction<v
 			TGraphTask<TAsyncGraphTask<ResultType>>::CreateTask().ConstructAndDispatchWhenReady(MoveTemp(Function), MoveTemp(Promise), Execution == EAsyncExecution::TaskGraph ? ENamedThreads::AnyThread : ENamedThreads::GameThread);
 		}
 		break;
-	
+	case EAsyncExecution::TaskGraphMainTick:
+		FTSTicker::GetCoreTicker().AddTicker(TEXT("Async_TaskGraphMainTick"), 0.0f,
+			[Function = MoveTemp(Function), Promise = MoveTemp(Promise)](float) mutable
+			{
+				SetPromise(Promise, Function);
+				return false;
+			}
+		);
+		break;
 	case EAsyncExecution::Thread:
 		if (FPlatformProcess::SupportsMultithreading())
 		{
