[core] add missing header

diff --git a/Engine/Source/Runtime/Core/Public/String/FormatStringSan.h b/Engine/Source/Runtime/Core/Public/String/FormatStringSan.h
--- a/Engine/Source/Runtime/Core/Public/String/FormatStringSan.h
+++ b/Engine/Source/Runtime/Core/Public/String/FormatStringSan.h
@@ -5,6 +5,7 @@
 #include "CoreTypes.h"
 #include "Templates/Requires.h"
 #include "Templates/Identity.h"
+#include "Templates/IsValidVariadicFunctionArg.h"
 #include "Containers/ContainersFwd.h"
 
 #define UE_CHECK_FORMAT_STRING(Fmt, ...) do {														\
