ObjectThumbnail::SetImage refactor to use FImage move constructor instead of bMove bool

diff --git a/Engine/Source/Editor/UnrealEd/Private/ThumbnailExternalCache.cpp b/Engine/Source/Editor/UnrealEd/Private/ThumbnailExternalCache.cpp
--- a/Engine/Source/Editor/UnrealEd/Private/ThumbnailExternalCache.cpp
+++ b/Engine/Source/Editor/UnrealEd/Private/ThumbnailExternalCache.cpp
@@ -33,7 +33,7 @@ namespace ThumbnailExternalCache
 		FImage DestImage;
 		FImageCore::ResizeImageAllocDest(SrcImage,DestImage,NewWidth,NewHeight);
 
-		Thumbnail.SetImage(DestImage,true);
+		Thumbnail.SetImage( MoveTemp(DestImage) );
 	}
 
 	// Return true if was resized
diff --git a/Engine/Source/Runtime/Core/Public/Misc/ObjectThumbnail.h b/Engine/Source/Runtime/Core/Public/Misc/ObjectThumbnail.h
--- a/Engine/Source/Runtime/Core/Public/Misc/ObjectThumbnail.h
+++ b/Engine/Source/Runtime/Core/Public/Misc/ObjectThumbnail.h
@@ -187,39 +187,35 @@ public:
 		return FImageView( (void *)&Data[0],ImageWidth,ImageHeight,ERawImageFormat::BGRA8 );
 	}
 
-	/** Set Thumbnail uncompressed image data.
+	/** Copy Image into Thumbnail uncompressed image data.
 		Existing compressed data, if any, is freed.
 	
-		if bMove is true, the Image data is moved into the Thumbnail
-		if bMove is false, the Image is copied into the thumbnail and not changed
-
 		Thumbnails are always BGRA8-sRGB.  The passed in Image can be other formats and conversion will be done if needed.
 	*/
-	void SetImage(FImage & Image, bool bMove)
+	void SetImage(const FImageView & Image)
 	{
-		if ( Image.Format != ERawImageFormat::BGRA8 || Image.GammaSpace != EGammaSpace::sRGB )
-		{
-			// Image must be converted to BGRA8 to store in Thumbnail
-			FImage ImageBGRA8;
-			Image.CopyTo(ImageBGRA8,ERawImageFormat::BGRA8,EGammaSpace::sRGB);
-			SetImage(ImageBGRA8,true);
-			return;
-		}
+		// Image must be converted to BGRA8 to store in Thumbnail
+		// if Image is already BGRA8-SRGB then this is just a memcpy
+		//	which is what we need anyway to copy the bytes into a new array
+		//	so there is no inefficiency in just always using the FImage copy here
+		FImage ImageBGRA8;
+		Image.CopyTo(ImageBGRA8,ERawImageFormat::BGRA8,EGammaSpace::sRGB);
+		SetImage( MoveTemp(ImageBGRA8) );
+	}
+	
+	/** Move Image data into Thumbnail (convert to BGRA8-SRGB if necessary) */
+	void SetImage(FImage && Image)
+	{
+		// change format if needed, nop if not
+		// MoveFrom is discardable so it's okay if we just change it in place
+		Image.ChangeFormat(ERawImageFormat::BGRA8,EGammaSpace::sRGB);
 
 		ImageWidth  = Image.SizeX;
 		ImageHeight = Image.SizeY;
 		CompressedImageData.Reset();
 		bIsJPEG = false;
-
-		if ( bMove )
-		{
-			ImageData = MoveTemp(Image.RawData);
-			Image.RawData.Reset();
-		}
-		else
-		{
-			ImageData = Image.RawData;
-		}
+		ImageData = MoveTemp(Image.RawData); // array move
+		Image.Reset();
 	}
 
 	/** Serializers */
