Preventing TScriptDelegate::IsCompactable()  from resolving remote objects

diff --git a/Engine/Source/Runtime/Core/Public/UObject/ScriptDelegates.h b/Engine/Source/Runtime/Core/Public/UObject/ScriptDelegates.h
--- a/Engine/Source/Runtime/Core/Public/UObject/ScriptDelegates.h
+++ b/Engine/Source/Runtime/Core/Public/UObject/ScriptDelegates.h
@@ -226,7 +226,12 @@ public:
 	inline bool IsCompactable() const
 	{
 		FReadAccessScope ReadScope = GetReadAccessScope();
-		return FunctionName == NAME_None || !Object.Get(true);
+		return FunctionName == NAME_None || 
+#if UE_WITH_REMOTE_OBJECT_HANDLE
+			!Object.IsValid(true);
+#else
+			!Object.Get(true);
+#endif
 	}
 
 	/**
