[Insights] Added more options to set value for a trace counter:
* TRACE_COUNTER_SET_IF_DIFFERENT(CounterName, Value)
* TRACE_COUNTER_ADD_IF_NOT_ZERO(CounterName, Value)
* TRACE_COUNTER_SUBTRACT_IF_NOT_ZERO(CounterName, Value)
* TRACE_COUNTER_SET_ALWAYS(CounterName, Value)
* TRACE_COUNTER_ADD_ALWAYS(CounterName, Value)
* TRACE_COUNTER_SUBTRACT_ALWAYS(CounterName, Value)
These can be used no matter how counter is created ("checked" or "unchecked").

diff --git a/Engine/Source/Runtime/Core/Public/ProfilingDebugging/CountersTrace.h b/Engine/Source/Runtime/Core/Public/ProfilingDebugging/CountersTrace.h
--- a/Engine/Source/Runtime/Core/Public/ProfilingDebugging/CountersTrace.h
+++ b/Engine/Source/Runtime/Core/Public/ProfilingDebugging/CountersTrace.h
@@ -121,6 +121,29 @@ struct FCountersTrace
 			}
 		}
 
+		void SetIfDifferent(ValueType InValue)
+		{
+			if (Value != InValue)
+			{
+				Value = InValue;
+				if (UE_TRACE_CHANNELEXPR_IS_ENABLED(CountersChannel))
+				{
+					LateInit();
+					OutputSetValue(uint16(CounterId), Value);
+				}
+			}
+		}
+
+		void SetAlways(ValueType InValue)
+		{
+			Value = InValue;
+			if (UE_TRACE_CHANNELEXPR_IS_ENABLED(CountersChannel))
+			{
+				LateInit();
+				OutputSetValue(uint16(CounterId), Value);
+			}
+		}
+
 		void Add(ValueType InValue)
 		{
 			if (bUnchecked || InValue != 0)
@@ -134,6 +157,29 @@ struct FCountersTrace
 			}
 		}
 
+		void AddIfNotZero(ValueType InValue)
+		{
+			if (InValue != 0)
+			{
+				Value += InValue;
+				if (UE_TRACE_CHANNELEXPR_IS_ENABLED(CountersChannel))
+				{
+					LateInit();
+					OutputSetValue(uint16(CounterId), Value);
+				}
+			}
+		}
+
+		void AddAlways(ValueType InValue)
+		{
+			Value += InValue;
+			if (UE_TRACE_CHANNELEXPR_IS_ENABLED(CountersChannel))
+			{
+				LateInit();
+				OutputSetValue(uint16(CounterId), Value);
+			}
+		}
+
 		void Subtract(ValueType InValue)
 		{
 			if (bUnchecked || InValue != 0)
@@ -147,6 +193,29 @@ struct FCountersTrace
 			}
 		}
 
+		void SubtractIfNotZero(ValueType InValue)
+		{
+			if (InValue != 0)
+			{
+				Value -= InValue;
+				if (UE_TRACE_CHANNELEXPR_IS_ENABLED(CountersChannel))
+				{
+					LateInit();
+					OutputSetValue(uint16(CounterId), Value);
+				}
+			}
+		}
+
+		void SubtractAlways(ValueType InValue)
+		{
+			Value -= InValue;
+			if (UE_TRACE_CHANNELEXPR_IS_ENABLED(CountersChannel))
+			{
+				LateInit();
+				OutputSetValue(uint16(CounterId), Value);
+			}
+		}
+
 		void Increment()
 		{
 			++Value;
@@ -313,6 +382,8 @@ struct FCountersTrace
 //////////////////////////////////////////////////
 // Counter Operations
 
+// A value that does not change will be traced (or not) depending on how the counter is created.
+
 #define TRACE_COUNTER_SET(CounterName, Value) \
 	PREPROCESSOR_JOIN(__GTraceCounter, CounterName).Set(Value);
 
@@ -328,6 +399,30 @@ struct FCountersTrace
 #define TRACE_COUNTER_DECREMENT(CounterName) \
 	PREPROCESSOR_JOIN(__GTraceCounter, CounterName).Decrement();
 
+// _IF_DIFFERENT / _IF_NOT_ZERO
+// It will not trace a value that doesn't change (no matter how the counter was created).
+
+#define TRACE_COUNTER_SET_IF_DIFFERENT(CounterName, Value) \
+	PREPROCESSOR_JOIN(__GTraceCounter, CounterName).SetIfDifferent(Value);
+
+#define TRACE_COUNTER_ADD_IF_NOT_ZERO(CounterName, Value) \
+	PREPROCESSOR_JOIN(__GTraceCounter, CounterName).AddIfNotZero(Value);
+
+#define TRACE_COUNTER_SUBTRACT_IF_NOT_ZERO(CounterName, Value) \
+	PREPROCESSOR_JOIN(__GTraceCounter, CounterName).SubtractIfNotZero(Value);
+
+// _ALWAYS
+// It will trace even if value doesn't change (no matter how the counter was created).
+
+#define TRACE_COUNTER_SET_ALWAYS(CounterName, Value) \
+	PREPROCESSOR_JOIN(__GTraceCounter, CounterName).SetAlways(Value);
+
+#define TRACE_COUNTER_ADD_ALWAYS(CounterName, Value) \
+	PREPROCESSOR_JOIN(__GTraceCounter, CounterName).AddAlways(Value);
+
+#define TRACE_COUNTER_SUBTRACT_ALWAYS(CounterName, Value) \
+	PREPROCESSOR_JOIN(__GTraceCounter, CounterName).SubtractAlways(Value);
+
 #else // COUNTERSTRACE_ENABLED
 
 #define TRACE_INT_VALUE(CounterDisplayName, Value)
@@ -378,4 +473,12 @@ struct FCountersTrace
 #define TRACE_COUNTER_INCREMENT(CounterName)
 #define TRACE_COUNTER_DECREMENT(CounterName)
 
+#define TRACE_COUNTER_SET_IF_DIFFERENT(CounterName, Value)
+#define TRACE_COUNTER_ADD_IF_NOT_ZERO(CounterName, Value)
+#define TRACE_COUNTER_SUBTRACT_IF_NOT_ZERO(CounterName, Value)
+
+#define TRACE_COUNTER_SET_ALWAYS(CounterName, Value)
+#define TRACE_COUNTER_ADD_ALWAYS(CounterName, Value)
+#define TRACE_COUNTER_SUBTRACT_ALWAYS(CounterName, Value)
+
 #endif // COUNTERSTRACE_ENABLED
