Make insights tracing work under AutoRTFM.

diff --git a/Engine/Source/Programs/AutoRTFMTests/AutoRTFMTests.Target.cs b/Engine/Source/Programs/AutoRTFMTests/AutoRTFMTests.Target.cs
--- a/Engine/Source/Programs/AutoRTFMTests/AutoRTFMTests.Target.cs
+++ b/Engine/Source/Programs/AutoRTFMTests/AutoRTFMTests.Target.cs
@@ -62,5 +62,8 @@ public class AutoRTFMTestsTarget : TargetRules
 
 		GlobalDefinitions.Add("MALLOC_LEAKDETECTION=1");
 		GlobalDefinitions.Add("PLATFORM_USES_FIXED_GMalloc_CLASS=0");
+
+		// For testing insights, this sets `UE_TRACE_ENABLED`.
+		bEnableTrace = true;
 	}
 }
diff --git a/Engine/Source/Programs/AutoRTFMTests/Private/InsightsTests.cpp b/Engine/Source/Programs/AutoRTFMTests/Private/InsightsTests.cpp
new file mode 100644
--- /dev/null
+++ b/Engine/Source/Programs/AutoRTFMTests/Private/InsightsTests.cpp
@@ -0,0 +1,455 @@
+// Copyright Epic Games, Inc. All Rights Reserved.
+
+#include "API.h"
+#include "Catch2Includes.h"
+#include "CoreTypes.h"
+#include "ProfilingDebugging/CpuProfilerTrace.h"
+
+#include <AutoRTFM/AutoRTFMTesting.h>
+
+UE_TRACE_CHANNEL_DEFINE(AutoRTFMInsightsChannel)
+
+TEST_CASE("Insights")
+{
+	CpuChannel.Toggle(true);
+	REQUIRE(UE_TRACE_CHANNELEXPR_IS_ENABLED(CpuChannel));
+
+	AutoRTFMInsightsChannel.Toggle(true);
+	REQUIRE(UE_TRACE_CHANNELEXPR_IS_ENABLED(AutoRTFMInsightsChannel));
+
+	SECTION("TRACE_CPUPROFILER_EVENT_DECLARE")
+	{
+		AutoRTFM::Testing::Abort([&]
+		{
+			TRACE_CPUPROFILER_EVENT_DECLARE(SomeEvent);
+			AutoRTFM::AbortTransaction();
+		});
+
+		AutoRTFM::Testing::Commit([&]
+		{
+			TRACE_CPUPROFILER_EVENT_DECLARE(SomeEvent);
+		});
+	}
+
+	SECTION("TRACE_CPUPROFILER_EVENT_SCOPE_USE")
+	{
+		TRACE_CPUPROFILER_EVENT_DECLARE(SomeEvent);
+		AutoRTFM::Testing::Commit([&]
+		{
+			TRACE_CPUPROFILER_EVENT_SCOPE_USE(SomeEvent, ANSI_TO_TCHAR("Wowwee"), _, false);
+		});
+
+		AutoRTFM::Testing::Commit([&]
+		{
+			TRACE_CPUPROFILER_EVENT_SCOPE_USE(SomeEvent, ANSI_TO_TCHAR("Wowwee"), _, true);
+		});
+
+		SECTION("In a child transaction")
+		{
+			AutoRTFM::Testing::Commit([&]
+			{
+				AutoRTFM::Testing::Commit([&]
+				{
+					TRACE_CPUPROFILER_EVENT_SCOPE_USE(SomeEvent, ANSI_TO_TCHAR("Wowwee"), _, true);
+				});
+			});
+		}
+
+		AutoRTFM::Testing::Abort([&]
+		{
+			TRACE_CPUPROFILER_EVENT_SCOPE_USE(SomeEvent, ANSI_TO_TCHAR("Wowwee"), _, true);
+			AutoRTFM::AbortTransaction();
+		});
+	}
+
+	SECTION("TRACE_CPUPROFILER_EVENT_SCOPE_USE")
+	{
+		TRACE_CPUPROFILER_EVENT_DECLARE(SomeEvent);
+		AutoRTFM::Testing::Commit([&]
+		{
+			TRACE_CPUPROFILER_EVENT_SCOPE_USE_ON_CHANNEL(SomeEvent, ANSI_TO_TCHAR("Wowwee"), _, AutoRTFMInsightsChannel, false);
+		});
+
+		AutoRTFM::Testing::Commit([&]
+		{
+			TRACE_CPUPROFILER_EVENT_SCOPE_USE_ON_CHANNEL(SomeEvent, ANSI_TO_TCHAR("Wowwee"), _, AutoRTFMInsightsChannel, true);
+		});
+
+		SECTION("In a child transaction")
+		{
+			AutoRTFM::Testing::Commit([&]
+			{
+				AutoRTFM::Testing::Commit([&]
+				{
+					TRACE_CPUPROFILER_EVENT_SCOPE_USE_ON_CHANNEL(SomeEvent, ANSI_TO_TCHAR("Wowwee"), _, AutoRTFMInsightsChannel, true);
+				});
+			});
+		}
+
+		AutoRTFM::Testing::Abort([&]
+		{
+			TRACE_CPUPROFILER_EVENT_SCOPE_USE_ON_CHANNEL(SomeEvent, ANSI_TO_TCHAR("Wowwee"), _, AutoRTFMInsightsChannel, true);
+			AutoRTFM::AbortTransaction();
+		});
+	}
+
+	SECTION("TRACE_CPUPROFILER_EVENT_SCOPE_STR")
+	{
+		AutoRTFM::Testing::Commit([&]
+		{
+			TRACE_CPUPROFILER_EVENT_SCOPE_STR(ANSI_TO_TCHAR("Wowwee"));
+		});
+
+		SECTION("In a child transaction")
+		{
+			AutoRTFM::Testing::Commit([&]
+			{
+				AutoRTFM::Testing::Commit([&]
+				{
+					TRACE_CPUPROFILER_EVENT_SCOPE_STR(ANSI_TO_TCHAR("Wowwee"));
+				});
+			});
+		}
+
+		AutoRTFM::Testing::Abort([&]
+		{
+			TRACE_CPUPROFILER_EVENT_SCOPE_STR(ANSI_TO_TCHAR("Wowwee"));
+			AutoRTFM::AbortTransaction();
+		});
+	}
+
+	SECTION("TRACE_CPUPROFILER_EVENT_SCOPE_STR_CONDITIONAL")
+	{
+		AutoRTFM::Testing::Commit([&]
+		{
+			TRACE_CPUPROFILER_EVENT_SCOPE_STR_CONDITIONAL(ANSI_TO_TCHAR("Wowwee"), false);
+		});
+
+		AutoRTFM::Testing::Commit([&]
+		{
+			TRACE_CPUPROFILER_EVENT_SCOPE_STR_CONDITIONAL(ANSI_TO_TCHAR("Wowwee"), true);
+		});
+
+		SECTION("In a child transaction")
+		{
+			AutoRTFM::Testing::Commit([&]
+			{
+				AutoRTFM::Testing::Commit([&]
+				{
+					TRACE_CPUPROFILER_EVENT_SCOPE_STR_CONDITIONAL(ANSI_TO_TCHAR("Wowwee"), true);
+				});
+			});
+		}
+
+		AutoRTFM::Testing::Abort([&]
+		{
+			TRACE_CPUPROFILER_EVENT_SCOPE_STR_CONDITIONAL(ANSI_TO_TCHAR("Wowwee"), true);
+			AutoRTFM::AbortTransaction();
+		});
+	}
+
+	SECTION("TRACE_CPUPROFILER_EVENT_SCOPE_ON_CHANNEL_STR")
+	{
+		AutoRTFM::Testing::Commit([&]
+		{
+			TRACE_CPUPROFILER_EVENT_SCOPE_ON_CHANNEL_STR(ANSI_TO_TCHAR("Wowwee"), AutoRTFMInsightsChannel);
+		});
+
+		SECTION("In a child transaction")
+		{
+			AutoRTFM::Testing::Commit([&]
+			{
+				AutoRTFM::Testing::Commit([&]
+				{
+					TRACE_CPUPROFILER_EVENT_SCOPE_ON_CHANNEL_STR(ANSI_TO_TCHAR("Wowwee"), AutoRTFMInsightsChannel);
+				});
+			});
+		}
+
+		AutoRTFM::Testing::Abort([&]
+		{
+			TRACE_CPUPROFILER_EVENT_SCOPE_ON_CHANNEL_STR(ANSI_TO_TCHAR("Wowwee"), AutoRTFMInsightsChannel);
+			AutoRTFM::AbortTransaction();
+		});
+	}
+
+	SECTION("TRACE_CPUPROFILER_EVENT_SCOPE_ON_CHANNEL_STR_CONDITIONAL")
+	{
+		AutoRTFM::Testing::Commit([&]
+		{
+			TRACE_CPUPROFILER_EVENT_SCOPE_ON_CHANNEL_STR_CONDITIONAL(ANSI_TO_TCHAR("Wowwee"), AutoRTFMInsightsChannel, false);
+		});
+
+		AutoRTFM::Testing::Commit([&]
+		{
+			TRACE_CPUPROFILER_EVENT_SCOPE_ON_CHANNEL_STR_CONDITIONAL(ANSI_TO_TCHAR("Wowwee"), AutoRTFMInsightsChannel, true);
+		});
+
+		SECTION("In a child transaction")
+		{
+			AutoRTFM::Testing::Commit([&]
+			{
+				AutoRTFM::Testing::Commit([&]
+				{
+					TRACE_CPUPROFILER_EVENT_SCOPE_ON_CHANNEL_STR_CONDITIONAL(ANSI_TO_TCHAR("Wowwee"), AutoRTFMInsightsChannel, true);
+				});
+			});
+		}
+
+		AutoRTFM::Testing::Abort([&]
+		{
+			TRACE_CPUPROFILER_EVENT_SCOPE_ON_CHANNEL_STR_CONDITIONAL(ANSI_TO_TCHAR("Wowwee"), AutoRTFMInsightsChannel, true);
+			AutoRTFM::AbortTransaction();
+		});
+	}
+	
+	SECTION("TRACE_CPUPROFILER_EVENT_SCOPE")
+	{
+		AutoRTFM::Testing::Commit([&]
+		{
+			TRACE_CPUPROFILER_EVENT_SCOPE(Wowwee);
+		});
+
+		SECTION("In a child transaction")
+		{
+			AutoRTFM::Testing::Commit([&]
+			{
+				AutoRTFM::Testing::Commit([&]
+				{
+					TRACE_CPUPROFILER_EVENT_SCOPE(Wowwee);
+				});
+			});
+		}
+
+		AutoRTFM::Testing::Abort([&]
+		{
+			TRACE_CPUPROFILER_EVENT_SCOPE(Wowwee);
+			AutoRTFM::AbortTransaction();
+		});
+	}
+	
+	SECTION("TRACE_CPUPROFILER_EVENT_SCOPE_CONDITIONAL")
+	{
+		AutoRTFM::Testing::Commit([&]
+		{
+			TRACE_CPUPROFILER_EVENT_SCOPE_CONDITIONAL(Wowwee, false);
+		});
+
+		AutoRTFM::Testing::Commit([&]
+		{
+			TRACE_CPUPROFILER_EVENT_SCOPE_CONDITIONAL(Wowwee, true);
+		});
+
+		SECTION("In a child transaction")
+		{
+			AutoRTFM::Testing::Commit([&]
+			{
+				AutoRTFM::Testing::Commit([&]
+				{
+					TRACE_CPUPROFILER_EVENT_SCOPE_CONDITIONAL(Wowwee, true);
+				});
+			});
+		}
+
+		AutoRTFM::Testing::Abort([&]
+		{
+			TRACE_CPUPROFILER_EVENT_SCOPE_CONDITIONAL(Wowwee, true);
+			AutoRTFM::AbortTransaction();
+		});
+	}
+	
+	SECTION("TRACE_CPUPROFILER_EVENT_SCOPE_ON_CHANNEL")
+	{
+		AutoRTFM::Testing::Commit([&]
+		{
+			TRACE_CPUPROFILER_EVENT_SCOPE_ON_CHANNEL(Wowwee, AutoRTFMInsightsChannel);
+		});
+
+		SECTION("In a child transaction")
+		{
+			AutoRTFM::Testing::Commit([&]
+			{
+				AutoRTFM::Testing::Commit([&]
+				{
+					TRACE_CPUPROFILER_EVENT_SCOPE_ON_CHANNEL(Wowwee, AutoRTFMInsightsChannel);
+				});
+			});
+		}
+
+		AutoRTFM::Testing::Abort([&]
+		{
+			TRACE_CPUPROFILER_EVENT_SCOPE_ON_CHANNEL(Wowwee, AutoRTFMInsightsChannel);
+			AutoRTFM::AbortTransaction();
+		});
+	}
+
+	SECTION("TRACE_CPUPROFILER_EVENT_SCOPE_ON_CHANNEL_CONDITIONAL")
+	{
+		AutoRTFM::Testing::Commit([&]
+		{
+			TRACE_CPUPROFILER_EVENT_SCOPE_ON_CHANNEL_CONDITIONAL(Wowwee, AutoRTFMInsightsChannel, false);
+		});
+
+		AutoRTFM::Testing::Commit([&]
+		{
+			TRACE_CPUPROFILER_EVENT_SCOPE_ON_CHANNEL_CONDITIONAL(Wowwee, AutoRTFMInsightsChannel, true);
+		});
+
+		SECTION("In a child transaction")
+		{
+			AutoRTFM::Testing::Commit([&]
+			{
+				AutoRTFM::Testing::Commit([&]
+				{
+					TRACE_CPUPROFILER_EVENT_SCOPE_ON_CHANNEL_CONDITIONAL(Wowwee, AutoRTFMInsightsChannel, true);
+				});
+			});
+		}
+
+		AutoRTFM::Testing::Abort([&]
+		{
+			TRACE_CPUPROFILER_EVENT_SCOPE_ON_CHANNEL_CONDITIONAL(Wowwee, AutoRTFMInsightsChannel, true);
+			AutoRTFM::AbortTransaction();
+		});
+	}
+
+	SECTION("TRACE_CPUPROFILER_EVENT_SCOPE_TEXT")
+	{
+		AutoRTFM::Testing::Commit([&]
+		{
+			FString Name("Wowwee");
+			TRACE_CPUPROFILER_EVENT_SCOPE_TEXT(*Name);
+		});
+
+		SECTION("In a child transaction")
+		{
+			AutoRTFM::Testing::Commit([&]
+			{
+				AutoRTFM::Testing::Commit([&]
+				{
+					FString Name("Wowwee");
+					TRACE_CPUPROFILER_EVENT_SCOPE_TEXT(*Name);
+				});
+			});
+		}
+
+		AutoRTFM::Testing::Abort([&]
+		{
+			FString Name("Wowwee");
+			TRACE_CPUPROFILER_EVENT_SCOPE_TEXT(*Name);
+			AutoRTFM::AbortTransaction();
+		});
+	}
+
+	SECTION("TRACE_CPUPROFILER_EVENT_SCOPE_TEXT_ON_CHANNEL")
+	{
+		AutoRTFM::Testing::Commit([&]
+		{
+			FString Name("Wowwee");
+			TRACE_CPUPROFILER_EVENT_SCOPE_TEXT_ON_CHANNEL(*Name, AutoRTFMInsightsChannel);
+		});
+
+		SECTION("In a child transaction")
+		{
+			AutoRTFM::Testing::Commit([&]
+			{
+				AutoRTFM::Testing::Commit([&]
+				{
+					FString Name("Wowwee");
+					TRACE_CPUPROFILER_EVENT_SCOPE_TEXT_ON_CHANNEL(*Name, AutoRTFMInsightsChannel);
+				});
+			});
+		}
+
+		AutoRTFM::Testing::Abort([&]
+		{
+			FString Name("Wowwee");
+			TRACE_CPUPROFILER_EVENT_SCOPE_TEXT_ON_CHANNEL(*Name, AutoRTFMInsightsChannel);
+			AutoRTFM::AbortTransaction();
+		});
+	}
+
+	SECTION("TRACE_CPUPROFILER_EVENT_SCOPE_TEXT_ON_CHANNEL_CONDITIONAL")
+	{
+		AutoRTFM::Testing::Commit([&]
+		{
+			FString Name("Wowwee");
+			TRACE_CPUPROFILER_EVENT_SCOPE_TEXT_ON_CHANNEL_CONDITIONAL(*Name, AutoRTFMInsightsChannel, false);
+		});
+
+		AutoRTFM::Testing::Commit([&]
+		{
+			FString Name("Wowwee");
+			TRACE_CPUPROFILER_EVENT_SCOPE_TEXT_ON_CHANNEL_CONDITIONAL(*Name, AutoRTFMInsightsChannel, true);
+		});
+
+		SECTION("In a child transaction")
+		{
+			AutoRTFM::Testing::Commit([&]
+			{
+				AutoRTFM::Testing::Commit([&]
+				{
+					FString Name("Wowwee");
+					TRACE_CPUPROFILER_EVENT_SCOPE_TEXT_ON_CHANNEL_CONDITIONAL(*Name, AutoRTFMInsightsChannel, true);
+				});
+			});
+		}
+
+		AutoRTFM::Testing::Abort([&]
+		{
+			FString Name("Wowwee");
+			TRACE_CPUPROFILER_EVENT_SCOPE_TEXT_ON_CHANNEL_CONDITIONAL(*Name, AutoRTFMInsightsChannel, true);
+			AutoRTFM::AbortTransaction();
+		});
+	}
+
+	SECTION("TRACE_CPUPROFILER_EVENT_FLUSH")
+	{
+		AutoRTFM::Testing::Commit([&]
+		{
+			TRACE_CPUPROFILER_EVENT_FLUSH();
+		});
+
+		AutoRTFM::Testing::Abort([&]
+		{
+			TRACE_CPUPROFILER_EVENT_FLUSH();
+			AutoRTFM::AbortTransaction();
+		});
+	}
+
+	SECTION("TRACE_CPUPROFILER_EVENT_MANUAL_IS_ENABLED")
+	{
+		AutoRTFM::Testing::Commit([&]
+		{
+			TRACE_CPUPROFILER_EVENT_MANUAL_IS_ENABLED();
+		});
+
+		AutoRTFM::Testing::Abort([&]
+		{
+			TRACE_CPUPROFILER_EVENT_MANUAL_IS_ENABLED();
+			AutoRTFM::AbortTransaction();
+		});
+	}
+
+	SECTION("TRACE_CPUPROFILER_EVENT_MANUAL_START + TRACE_CPUPROFILER_EVENT_MANUAL_END")
+	{
+		AutoRTFM::Testing::Commit([&]
+		{
+			TRACE_CPUPROFILER_EVENT_MANUAL_START(TEXT("Wowwee"));
+			TRACE_CPUPROFILER_EVENT_MANUAL_END();
+		});
+
+		AutoRTFM::Testing::Abort([&]
+		{
+			TRACE_CPUPROFILER_EVENT_MANUAL_START(TEXT("Wowwee"));
+			TRACE_CPUPROFILER_EVENT_MANUAL_END();
+			AutoRTFM::AbortTransaction();
+		});
+
+		AutoRTFM::Testing::Abort([&]
+		{
+			TRACE_CPUPROFILER_EVENT_MANUAL_START(TEXT("Wowwee"));
+			AutoRTFM::AbortTransaction();
+		});
+	}
+}
diff --git a/Engine/Source/Runtime/Core/Private/ProfilingDebugging/CpuProfilerTrace.cpp b/Engine/Source/Runtime/Core/Private/ProfilingDebugging/CpuProfilerTrace.cpp
--- a/Engine/Source/Runtime/Core/Private/ProfilingDebugging/CpuProfilerTrace.cpp
+++ b/Engine/Source/Runtime/Core/Private/ProfilingDebugging/CpuProfilerTrace.cpp
@@ -133,7 +133,6 @@ FCpuProfilerTraceInternal::FThreadBuffer* FCpuProfilerTraceInternal::CreateThrea
 	return ThreadBuffer;
 }
 
-UE_AUTORTFM_ALWAYS_OPEN
 void FCpuProfilerTraceInternal::FlushThreadBuffer(FThreadBuffer* InThreadBuffer)
 {
 	UE_TRACE_LOG(CpuProfiler, EventBatchV2, true)
@@ -163,14 +162,77 @@ void FCpuProfilerTraceInternal::FlushThreadBuffer(FThreadBuffer* InThreadBuffer)
 		FCpuProfilerTraceInternal::FlushThreadBuffer(ThreadBuffer); \
 	}
 
+#define INSIGHTS_AUTORTFM_UNREACHABLE() \
+	do { if (AutoRTFM::IsClosed()) { AutoRTFM::Unreachable(TEXT("Unreachable transactional codepath in FCpuProfilerTrace")); } } while (0)
+
+uint8 FCpuProfilerTrace::OnAbortKey = 42;
+
+void FCpuProfilerTrace::GetOrCreateSpecId(uint32& InOutSpecId, const ANSICHAR* InEventString, const ANSICHAR* File, uint32 Line)
+{
+	AutoRTFM::Open([&]
+	{
+		// We only do relaxed here to avoid barrier cost as the worst case that can happen is multiple threads could each create an event type.
+		// At some point the last thread in the race will set the output event and no more thread will try to create new ones from then on.
+		// We don't care which event type wins as long as all threads eventually converge and stop creating new ones.
+		if (FPlatformAtomics::AtomicRead_Relaxed((volatile int32*)&InOutSpecId) == 0)
+		{
+			FPlatformAtomics::AtomicStore_Relaxed((volatile int32*)&InOutSpecId, FCpuProfilerTrace::OutputEventType(InEventString, File, Line));
+		}
+	});
+}
+
+void FCpuProfilerTrace::GetOrCreateSpecId(uint32& InOutSpecId, const TCHAR* InEventString, const ANSICHAR* File, uint32 Line)
+{
+	AutoRTFM::Open([&]
+	{
+		// We only do relaxed here to avoid barrier cost as the worst case that can happen is multiple threads could each create an event type.
+		// At some point the last thread in the race will set the output event and no more thread will try to create new ones from then on.
+		// We don't care which event type wins as long as all threads eventually converge and stop creating new ones.
+		if (FPlatformAtomics::AtomicRead_Relaxed((volatile int32*)&InOutSpecId) == 0)
+		{
+			FPlatformAtomics::AtomicStore_Relaxed((volatile int32*)&InOutSpecId, FCpuProfilerTrace::OutputEventType(InEventString, File, Line));
+		}
+	});
+}
+
 void FCpuProfilerTrace::OutputBeginEvent(uint32 SpecId)
 {
+	if (AutoRTFM::IsClosed())
+	{
+		AutoRTFM::Open([&]
+		{
+			OutputBeginEvent(SpecId);
+		});
+
+		AutoRTFM::PushOnAbortHandler(&FCpuProfilerTrace::OnAbortKey, []
+		{
+			OutputEndEvent();
+		});
+
+		return;
+	}
+
 	CPUPROFILERTRACE_OUTPUTBEGINEVENT_PROLOGUE();
 	CPUPROFILERTRACE_OUTPUTBEGINEVENT_EPILOGUE();
 }
 
 void FCpuProfilerTrace::OutputBeginDynamicEvent(const ANSICHAR* Name, const ANSICHAR* File, uint32 Line)
 {
+	if (AutoRTFM::IsClosed())
+	{
+		AutoRTFM::Open([&]
+		{
+			OutputBeginDynamicEvent(Name, File, Line);
+		});
+
+		AutoRTFM::PushOnAbortHandler(&FCpuProfilerTrace::OnAbortKey, []
+		{
+			OutputEndEvent();
+		});
+
+		return;
+	}
+
 	CPUPROFILERTRACE_OUTPUTBEGINEVENT_PROLOGUE();
 	uint32 SpecId = OutputDynamicEventType(Name, File, Line);
 	CPUPROFILERTRACE_OUTPUTBEGINEVENT_EPILOGUE();
@@ -178,6 +240,21 @@ void FCpuProfilerTrace::OutputBeginDynamicEvent(const ANSICHAR* Name, const ANSI
 
 void FCpuProfilerTrace::OutputBeginDynamicEvent(const TCHAR* Name, const ANSICHAR* File, uint32 Line)
 {
+	if (AutoRTFM::IsClosed())
+	{
+		AutoRTFM::Open([&]
+		{
+			OutputBeginDynamicEvent(Name, File, Line);
+		});
+
+		AutoRTFM::PushOnAbortHandler(&FCpuProfilerTrace::OnAbortKey, []
+		{
+			OutputEndEvent();
+		});
+
+		return;
+	}
+
 	CPUPROFILERTRACE_OUTPUTBEGINEVENT_PROLOGUE();
 	uint32 SpecId = OutputDynamicEventType(Name, File, Line);
 	CPUPROFILERTRACE_OUTPUTBEGINEVENT_EPILOGUE();
@@ -185,6 +262,21 @@ void FCpuProfilerTrace::OutputBeginDynamicEvent(const TCHAR* Name, const ANSICHA
 
 void FCpuProfilerTrace::OutputBeginDynamicEvent(const FName Name, const ANSICHAR* File, uint32 Line)
 {
+	if (AutoRTFM::IsClosed())
+	{
+		AutoRTFM::Open([&]
+		{
+			OutputBeginDynamicEvent(Name, File, Line);
+		});
+
+		AutoRTFM::PushOnAbortHandler(&FCpuProfilerTrace::OnAbortKey, []
+		{
+			OutputEndEvent();
+		});
+
+		return;
+	}
+
 	CPUPROFILERTRACE_OUTPUTBEGINEVENT_PROLOGUE();
 	uint32 SpecId = OutputDynamicEventType(Name, File, Line);
 	CPUPROFILERTRACE_OUTPUTBEGINEVENT_EPILOGUE();
@@ -192,6 +284,21 @@ void FCpuProfilerTrace::OutputBeginDynamicEvent(const FName Name, const ANSICHAR
 
 void FCpuProfilerTrace::OutputBeginDynamicEventWithId(const FName Id, const ANSICHAR* Name, const ANSICHAR* File, uint32 Line)
 {
+	if (AutoRTFM::IsClosed())
+	{
+		AutoRTFM::Open([&]
+		{
+			OutputBeginDynamicEventWithId(Id, Name, File, Line);
+		});
+
+		AutoRTFM::PushOnAbortHandler(&FCpuProfilerTrace::OnAbortKey, []
+		{
+			OutputEndEvent();
+		});
+
+		return;
+	}
+
 	CPUPROFILERTRACE_OUTPUTBEGINEVENT_PROLOGUE();
 	uint32 SpecId = OutputDynamicEventTypeWithId(Id, Name, File, Line);
 	CPUPROFILERTRACE_OUTPUTBEGINEVENT_EPILOGUE();
@@ -199,6 +306,21 @@ void FCpuProfilerTrace::OutputBeginDynamicEventWithId(const FName Id, const ANSI
 
 void FCpuProfilerTrace::OutputBeginDynamicEventWithId(const FName Id, const TCHAR* Name, const ANSICHAR* File, uint32 Line)
 {
+	if (AutoRTFM::IsClosed())
+	{
+		AutoRTFM::Open([&]
+		{
+			OutputBeginDynamicEventWithId(Id, Name, File, Line);
+		});
+
+		AutoRTFM::PushOnAbortHandler(&FCpuProfilerTrace::OnAbortKey, []
+		{
+			OutputEndEvent();
+		});
+
+		return;
+	}
+
 	CPUPROFILERTRACE_OUTPUTBEGINEVENT_PROLOGUE();
 	uint32 SpecId = OutputDynamicEventTypeWithId(Id, Name, File, Line);
 	CPUPROFILERTRACE_OUTPUTBEGINEVENT_EPILOGUE();
@@ -206,6 +328,7 @@ void FCpuProfilerTrace::OutputBeginDynamicEventWithId(const FName Id, const TCHA
 
 void FCpuProfilerTrace::OutputResumeEvent(uint64 SpecId, uint32& TimerScopeDepth)
 {
+	INSIGHTS_AUTORTFM_UNREACHABLE();
 	FCpuProfilerTraceInternal::NestedTimerScopeDepths.Push({&TimerScopeDepth, FCpuProfilerTraceInternal::ThreadDepth});
 	FCpuProfilerTraceInternal::ThreadDepth = FCpuProfilerTraceInternal::ThreadDepth + TimerScopeDepth;
 
@@ -230,6 +353,7 @@ void FCpuProfilerTrace::OutputResumeEvent(uint64 SpecId, uint32& TimerScopeDepth
 
 void FCpuProfilerTrace::OutputSuspendEvent()
 {
+	INSIGHTS_AUTORTFM_UNREACHABLE();
 	auto [TimerScopeDepth, SavedThreadDepth] = FCpuProfilerTraceInternal::NestedTimerScopeDepths.Pop();
 	*TimerScopeDepth = FCpuProfilerTraceInternal::ThreadDepth - SavedThreadDepth;
 	FCpuProfilerTraceInternal::ThreadDepth = SavedThreadDepth;
@@ -257,6 +381,23 @@ void FCpuProfilerTrace::OutputSuspendEvent()
 
 void FCpuProfilerTrace::OutputEndEvent()
 {
+	if (AutoRTFM::IsClosed())
+	{
+		// For this to work correctly, the event *must* have been begun within the same transaction
+		// as the end event was called. Otherwise the following could happen within a transaction:
+		// - `OutputEndEvent` is called, which happens immediately.
+		// - The transaction aborts.
+		// - But we've unconditionally ended the event!
+
+		AutoRTFM::Open([&]
+		{
+			OutputEndEvent();
+		});
+
+		AutoRTFM::PopOnAbortHandler(&FCpuProfilerTrace::OnAbortKey);
+
+		return;
+	}
 	--FCpuProfilerTraceInternal::ThreadDepth;
 	FCpuProfilerTraceInternal::FThreadBuffer* ThreadBuffer = FCpuProfilerTraceInternal::ThreadBuffer;
 	if (!ThreadBuffer)
@@ -281,9 +422,16 @@ uint32 FCpuProfilerTraceInternal::GetNextSpecId()
 	return (NextSpecId++) + 1;
 }
 
-UE_AUTORTFM_ALWAYS_OPEN
 uint32 FCpuProfilerTrace::OutputEventType(const TCHAR* Name, const ANSICHAR* File, uint32 Line)
 {
+	if (AutoRTFM::IsClosed())
+	{
+		return AutoRTFM::Open([&]
+		{
+			return OutputEventType(Name, File, Line);
+		});
+	}
+
 	uint32 SpecId = FCpuProfilerTraceInternal::GetNextSpecId();
 	uint16 NameLen = uint16(FCString::Strlen(Name));
 	uint32 DataSize = NameLen * sizeof(ANSICHAR); // EventSpec.Name is traced as UE::Trace::AnsiString
@@ -302,9 +450,16 @@ uint32 FCpuProfilerTrace::OutputEventType(const TCHAR* Name, const ANSICHAR* Fil
 	return SpecId;
 }
 
-UE_AUTORTFM_ALWAYS_OPEN
 uint32 FCpuProfilerTrace::OutputEventType(const ANSICHAR* Name, const ANSICHAR* File, uint32 Line)
 {
+	if (AutoRTFM::IsClosed())
+	{
+		return AutoRTFM::Open([&]
+		{
+			return OutputEventType(Name, File, Line);
+		});
+	}
+
 	uint32 SpecId = FCpuProfilerTraceInternal::GetNextSpecId();
 	uint16 NameLen = uint16(strlen(Name));
 	uint32 DataSize = NameLen * sizeof(ANSICHAR);
@@ -323,9 +478,16 @@ uint32 FCpuProfilerTrace::OutputEventType(const ANSICHAR* Name, const ANSICHAR*
 	return SpecId;
 }
 
-UE_AUTORTFM_ALWAYS_OPEN
 uint32 FCpuProfilerTrace::OutputEventType(const FName Name, const ANSICHAR* File, uint32 Line)
 {
+	if (AutoRTFM::IsClosed())
+	{
+		return AutoRTFM::Open([&]
+		{
+			return OutputEventType(Name, File, Line);
+		});
+	}
+
 	const FNameEntry* NameEntry = Name.GetDisplayNameEntry();
 	if (NameEntry->IsWide())
 	{
@@ -341,9 +503,16 @@ uint32 FCpuProfilerTrace::OutputEventType(const FName Name, const ANSICHAR* File
 	}
 }
 
-UE_AUTORTFM_ALWAYS_OPEN
 uint32 FCpuProfilerTrace::OutputDynamicEventType(const ANSICHAR* Name, const ANSICHAR* File, uint32 Line)
 {
+	if (AutoRTFM::IsClosed())
+	{
+		return AutoRTFM::Open([&]
+		{
+			return OutputDynamicEventType(Name, File, Line);
+		});
+	}
+
 	FCpuProfilerTraceInternal::FThreadBuffer* ThreadBuffer = FCpuProfilerTraceInternal::ThreadBuffer;
 	if (!ThreadBuffer)
 	{
@@ -362,9 +531,16 @@ uint32 FCpuProfilerTrace::OutputDynamicEventType(const ANSICHAR* Name, const ANS
 	return SpecId;
 }
 
-UE_AUTORTFM_ALWAYS_OPEN
 uint32 FCpuProfilerTrace::OutputDynamicEventType(const TCHAR* Name, const ANSICHAR* File, uint32 Line)
 {
+	if (AutoRTFM::IsClosed())
+	{
+		return AutoRTFM::Open([&]
+		{
+			return OutputDynamicEventType(Name, File, Line);
+		});
+	}
+
 	FCpuProfilerTraceInternal::FThreadBuffer* ThreadBuffer = FCpuProfilerTraceInternal::ThreadBuffer;
 	if (!ThreadBuffer)
 	{
@@ -383,9 +559,16 @@ uint32 FCpuProfilerTrace::OutputDynamicEventType(const TCHAR* Name, const ANSICH
 	return SpecId;
 }
 
-UE_AUTORTFM_ALWAYS_OPEN
 uint32 FCpuProfilerTrace::OutputDynamicEventType(const FName Name, const ANSICHAR* File, uint32 Line)
 {
+	if (AutoRTFM::IsClosed())
+	{
+		return AutoRTFM::Open([&]
+		{
+			return OutputDynamicEventType(Name, File, Line);
+		});
+	}
+
 	FCpuProfilerTraceInternal::FThreadBuffer* ThreadBuffer = FCpuProfilerTraceInternal::ThreadBuffer;
 	if (!ThreadBuffer)
 	{
@@ -401,9 +584,16 @@ uint32 FCpuProfilerTrace::OutputDynamicEventType(const FName Name, const ANSICHA
 	return SpecId;
 }
 
-UE_AUTORTFM_ALWAYS_OPEN
 uint32 FCpuProfilerTrace::OutputDynamicEventTypeWithId(const FName Id, const ANSICHAR* Name, const ANSICHAR* File, uint32 Line)
 {
+	if (AutoRTFM::IsClosed())
+	{
+		return AutoRTFM::Open([&]
+		{
+			return OutputDynamicEventTypeWithId(Id, Name, File, Line);
+		});
+	}
+
 	FCpuProfilerTraceInternal::FThreadBuffer* ThreadBuffer = FCpuProfilerTraceInternal::ThreadBuffer;
 	if (!ThreadBuffer)
 	{
@@ -426,9 +616,16 @@ uint32 FCpuProfilerTrace::OutputDynamicEventTypeWithId(const FName Id, const ANS
 	return SpecId;
 }
 
-UE_AUTORTFM_ALWAYS_OPEN
 uint32 FCpuProfilerTrace::OutputDynamicEventTypeWithId(const FName Id, const TCHAR* Name, const ANSICHAR* File, uint32 Line)
 {
+	if (AutoRTFM::IsClosed())
+	{
+		return AutoRTFM::Open([&]
+		{
+			return OutputDynamicEventTypeWithId(Id, Name, File, Line);
+		});
+	}
+
 	FCpuProfilerTraceInternal::FThreadBuffer* ThreadBuffer = FCpuProfilerTraceInternal::ThreadBuffer;
 	if (!ThreadBuffer)
 	{
@@ -453,6 +650,16 @@ uint32 FCpuProfilerTrace::OutputDynamicEventTypeWithId(const FName Id, const TCH
 
 void FCpuProfilerTrace::FlushThreadBuffer()
 {
+	if (AutoRTFM::IsClosed())
+	{
+		AutoRTFM::Open([&]
+		{
+			FlushThreadBuffer();
+		});
+
+		return;
+	}
+
 	FCpuProfilerTraceInternal::FThreadBuffer* ThreadBuffer = FCpuProfilerTraceInternal::ThreadBuffer;
 	if (ThreadBuffer && ThreadBuffer->BufferSize > 0)
 	{
diff --git a/Engine/Source/Runtime/Core/Public/ProfilingDebugging/CpuProfilerTrace.h b/Engine/Source/Runtime/Core/Public/ProfilingDebugging/CpuProfilerTrace.h
--- a/Engine/Source/Runtime/Core/Public/ProfilingDebugging/CpuProfilerTrace.h
+++ b/Engine/Source/Runtime/Core/Public/ProfilingDebugging/CpuProfilerTrace.h
@@ -134,6 +134,24 @@ struct FCpuProfilerTrace
 	*/
 	CORE_API static void FlushThreadBuffer();
 
+	/*
+	 * Get or create a unique event definition id for the given parameters.
+	 * @param InOutSpecId - The event definition id.
+	 * @param Name - The name of event
+	 * @param File - The source filename
+	 * @param Line - The line number in source file
+	 */
+	CORE_API static void GetOrCreateSpecId(uint32& InOutSpecId, const ANSICHAR* InEventString, const ANSICHAR* File, uint32 Line);
+	
+	/*
+	 * Get or create a unique event definition id for the given parameters.
+	 * @param InOutSpecId - The event definition id.
+	 * @param Name - The name of event
+	 * @param File - The source filename
+	 * @param Line - The line number in source file
+	 */
+	CORE_API static void GetOrCreateSpecId(uint32& InOutSpecId, const TCHAR* InEventString, const ANSICHAR* File, uint32 Line);
+
 	class FEventScope
 	{
 	public:
@@ -194,13 +212,7 @@ struct FCpuProfilerTrace
 		{
 			if (bEnabled)
 			{
-				// We only do relaxed here to avoid barrier cost as the worst case that can happen is multiple threads could each create an event type.
-				// At some point the last thread in the race will set the output event and no more thread will try to create new ones from then on.
-				// We don't care which event type wins as long as all threads eventually converge and stop creating new ones.
-				if (FPlatformAtomics::AtomicRead_Relaxed((volatile int32*)&InOutSpecId) == 0)
-				{
-					FPlatformAtomics::AtomicStore_Relaxed((volatile int32*)&InOutSpecId, FCpuProfilerTrace::OutputEventType(InEventString, File, Line));
-				}
+				GetOrCreateSpecId(InOutSpecId, InEventString, File, Line);
 				OutputBeginEvent(FPlatformAtomics::AtomicRead_Relaxed((volatile int32*)&InOutSpecId));
 			}
 		}
@@ -209,13 +221,7 @@ struct FCpuProfilerTrace
 		{
 			if (bEnabled)
 			{
-				// We only do relaxed here to avoid barrier cost as the worst case that can happen is multiple threads could each create an event type.
-				// At some point the last thread in the race will set the output event and no more thread will try to create new ones from then on.
-				// We don't care which event type wins as long as all threads eventually converge and stop creating new ones.
-				if (FPlatformAtomics::AtomicRead_Relaxed((volatile int32*)&InOutSpecId) == 0)
-				{
-					FPlatformAtomics::AtomicStore_Relaxed((volatile int32*)&InOutSpecId, FCpuProfilerTrace::OutputEventType(InEventString, File, Line));
-				}
+				GetOrCreateSpecId(InOutSpecId, InEventString, File, Line);
 				OutputBeginEvent(FPlatformAtomics::AtomicRead_Relaxed((volatile int32*)&InOutSpecId));
 			}
 		}
@@ -271,6 +277,9 @@ struct FCpuProfilerTrace
 
 		bool bEnabled;
 	};
+
+private:
+	static uint8 OnAbortKey;
 };
 
 // Advanced macro for integrating e.g. stats/named events with CPU trace.
