- Platform CVar mapping for cooking:
- There is a mapping of EShaderPlatform -> Platform, so ShaderPlatformToPlatformName can return the correct platform when cooking
- There is a mapping of Platform -> DeviceProfile, so a TargetPlatformControls object can set a unique DeviceProfile to cook with
- It's likely more places in code will need to hook into these mappings (see namespace ConsoleVariablePlatformMapping)
- Added a CVar (disabled currently) to use the above (see cook.cvarcontrol)
- Added Android_XRTargetPlatform classes, and an early version of a basic Android_XR DeviceProfile, to use the Platform->DeviceProfile mapping

diff --git a/Engine/Config/Android/DataDrivenPlatformInfo.ini b/Engine/Config/Android/DataDrivenPlatformInfo.ini
--- a/Engine/Config/Android/DataDrivenPlatformInfo.ini
+++ b/Engine/Config/Android/DataDrivenPlatformInfo.ini
@@ -42,8 +42,8 @@ ActiveIconName=LevelEditor.PreviewMode.AndroidVulkan.Enabled
 InactiveIconName=LevelEditor.PreviewMode.AndroidVulkan.Disabled
 MenuTooltip=NSLOCTEXT("PreviewPlatform", "PreviewMenuTooltip_AndroidVulkan", "Mobile preview using Android Vulkan quality settings.")
 IconText=NSLOCTEXT("PreviewPlatform", "PreviewIconText_AndroidVulkan", "Android Vulkan")
-DeviceProfileName=Android_Preview_Vulkan:Android_Low:Android_Mid:Android_High
-FriendlyName=NSLOCTEXT("PreviewPlatform", "PreviewFriendlyName_GenericVulkanAndroid", "Generic Android Vulkan"):NSLOCTEXT("PreviewPlatform", "PreviewFriendlyName_AndroidVulkanLow","Android Vulkan Low"):NSLOCTEXT("PreviewPlatform", "PreviewFriendlyName_AndroidVulkanMid","Android Vulkan Mid"):NSLOCTEXT("PreviewPlatform", "PreviewFriendlyName_AndroidVulkanHigh","Android Vulkan High")
+DeviceProfileName=Android_Preview_Vulkan:Android_Low:Android_Mid:Android_High:Android_XR
+FriendlyName=NSLOCTEXT("PreviewPlatform", "PreviewFriendlyName_GenericVulkanAndroid", "Generic Android Vulkan"):NSLOCTEXT("PreviewPlatform", "PreviewFriendlyName_AndroidVulkanLow","Android Vulkan Low"):NSLOCTEXT("PreviewPlatform", "PreviewFriendlyName_AndroidVulkanMid","Android Vulkan Mid"):NSLOCTEXT("PreviewPlatform", "PreviewFriendlyName_AndroidVulkanHigh","Android Vulkan High"):NSLOCTEXT("PreviewPlatform", "PreviewFriendlyName_AndroidGenericXR","Android Generic XR")
 
 
 [PreviewPlatform AndroidVulkanSM5]
diff --git a/Engine/Config/BaseDeviceProfiles.ini b/Engine/Config/BaseDeviceProfiles.ini
--- a/Engine/Config/BaseDeviceProfiles.ini
+++ b/Engine/Config/BaseDeviceProfiles.ini
@@ -150,6 +150,7 @@
 +DeviceProfileNameAndTypes=Meta_Quest_Pro,Android
 +DeviceProfileNameAndTypes=Oculus_Quest2,Android
 +DeviceProfileNameAndTypes=Oculus_Quest,Android
++DeviceProfileNameAndTypes=Android_XR,Android
 +DeviceProfileNameAndTypes=HoloLens,HoloLens
 +DeviceProfileNameAndTypes=MagicLeap_Vulkan,Android
 
@@ -1196,6 +1197,12 @@ BaseProfileName=Android_High
 ; enable Vulkan
 +CVars=r.Android.DisableVulkanSupport=0
 
+[Android_XR DeviceProfile]
+DeviceType=Android
+BaseProfileName=Android_Mid
+;+CVars=vr.InstancedStereo=true
+;+CVars=vr.MobileMultiView=true
+
 [Meta_Quest_3 DeviceProfile]
 DeviceType=Android
 BaseProfileName=Meta_Quest_Pro
@@ -1213,7 +1220,7 @@ BaseProfileName=Oculus_Quest
 
 [Oculus_Quest DeviceProfile]
 DeviceType=Android
-BaseProfileName=Android_Mid
+BaseProfileName=Android_XR
 +CVars=xr.VRS.DynamicFoveation=1
 +CVars=r.Android.DisableVulkanSupport=0
 ; disable running Niagara compute shaders on the GPU. This can cause hard to trace crashes.
diff --git a/Engine/Source/Developer/Android/AndroidTargetPlatformControls/Private/AndroidTargetPlatformControls.h b/Engine/Source/Developer/Android/AndroidTargetPlatformControls/Private/AndroidTargetPlatformControls.h
--- a/Engine/Source/Developer/Android/AndroidTargetPlatformControls/Private/AndroidTargetPlatformControls.h
+++ b/Engine/Source/Developer/Android/AndroidTargetPlatformControls/Private/AndroidTargetPlatformControls.h
@@ -247,6 +247,22 @@ public:
 #endif
 };
 
+// XR using the ASTC TPC now
+class FAndroid_XRTargetPlatformControls : public FAndroid_ASTCTargetPlatformControls
+{
+public:
+	FAndroid_XRTargetPlatformControls(bool bIsClient, ITargetPlatformSettings* TargetPlatformSettings)
+		: FAndroid_ASTCTargetPlatformControls(bIsClient, TargetPlatformSettings, TEXT("XR"))
+	{
+	}
+
+	virtual FString CookingDeviceProfileName() const override
+	{
+		return TEXT("Android_XR");
+	}
+};
+
+
 class FAndroid_MultiTargetPlatformControls : public FAndroidTargetPlatformControls
 {
 	TArray<FAndroidTargetPlatformControls*> FormatTargetPlatforms;
diff --git a/Engine/Source/Developer/Android/AndroidTargetPlatformControls/Private/AndroidTargetPlatformControlsModule.cpp b/Engine/Source/Developer/Android/AndroidTargetPlatformControls/Private/AndroidTargetPlatformControlsModule.cpp
--- a/Engine/Source/Developer/Android/AndroidTargetPlatformControls/Private/AndroidTargetPlatformControlsModule.cpp
+++ b/Engine/Source/Developer/Android/AndroidTargetPlatformControls/Private/AndroidTargetPlatformControlsModule.cpp
@@ -37,11 +37,13 @@ protected:
 
 			TMap<EAndroidTextureFormatCategory, ITargetPlatformSettings*> OutMap;
 			ITargetPlatformSettings* OutMulti = nullptr;
+			ITargetPlatformSettings* OutXR = nullptr;
 			IAndroidTargetPlatformSettingsModule* Module = FModuleManager::GetModulePtr<IAndroidTargetPlatformSettingsModule>(PlatformSettingsModuleName);
 			if (Module != nullptr)
 			{
 				Module->GetSinglePlatformSettings(OutMap);
 				OutMulti = Module->GetMultiPlatformSettings();
+				OutXR = Module->GetXRPlatformSettings();
 			}
 
 			for (int32 Type = 0; Type < 2; Type++)
@@ -54,6 +56,9 @@ protected:
 				SinglePlatforms.Add(new FAndroid_ASTCTargetPlatformControls(bIsClient, OutMap[EAndroidTextureFormatCategory::ASTC]));
 				SinglePlatforms.Add(new FAndroid_DXTTargetPlatformControls(bIsClient, OutMap[EAndroidTextureFormatCategory::DXT]));
 				SinglePlatforms.Add(new FAndroid_ETC2TargetPlatformControls(bIsClient, OutMap[EAndroidTextureFormatCategory::ETC2]));
+				
+				// XR target platform
+				SinglePlatforms.Add(new FAndroid_XRTargetPlatformControls(bIsClient, OutXR));
 
 				// these are used in NotifyMultiSelectedFormatsChanged, so track in another array
 				MultiPlatforms.Add(new FAndroid_MultiTargetPlatformControls(bIsClient, OutMulti));
diff --git a/Engine/Source/Developer/Android/AndroidTargetPlatformSettings/Private/AndroidTargetPlatformSettingsModule.cpp b/Engine/Source/Developer/Android/AndroidTargetPlatformSettings/Private/AndroidTargetPlatformSettingsModule.cpp
--- a/Engine/Source/Developer/Android/AndroidTargetPlatformSettings/Private/AndroidTargetPlatformSettingsModule.cpp
+++ b/Engine/Source/Developer/Android/AndroidTargetPlatformSettings/Private/AndroidTargetPlatformSettingsModule.cpp
@@ -38,6 +38,11 @@ public:
 		return MultiTargetPlatformSettings;
 	}
 
+	ITargetPlatformSettings* GetXRPlatformSettings()
+	{
+		return XRTargetPlatformSettings;
+	}
+
 protected:
 	virtual void GetTargetPlatformSettings(TArray<ITargetPlatformSettings*>& TargetPlatforms)
 	{
@@ -60,9 +65,14 @@ protected:
 		ITargetPlatformSettings* MultiTP = new FAndroid_MultiTargetPlatformSettings();
 		TargetPlatforms.Add(MultiTP);
 		MultiTargetPlatformSettings = MultiTP;
+
+		ITargetPlatformSettings* XRTP = new FAndroid_XRTargetPlatformSettings();
+		TargetPlatforms.Add(XRTP);
+		XRTargetPlatformSettings = XRTP;
 	}
 private:
 	ITargetPlatformSettings* MultiTargetPlatformSettings;
+	ITargetPlatformSettings* XRTargetPlatformSettings;
 	TMap<EAndroidTextureFormatCategory, ITargetPlatformSettings*> SinglePlatformSettings;
 };
 
diff --git a/Engine/Source/Developer/Android/AndroidTargetPlatformSettings/Public/AndroidTargetPlatformSettings.h b/Engine/Source/Developer/Android/AndroidTargetPlatformSettings/Public/AndroidTargetPlatformSettings.h
--- a/Engine/Source/Developer/Android/AndroidTargetPlatformSettings/Public/AndroidTargetPlatformSettings.h
+++ b/Engine/Source/Developer/Android/AndroidTargetPlatformSettings/Public/AndroidTargetPlatformSettings.h
@@ -172,7 +172,7 @@ public:
 class FAndroid_ASTCTargetPlatformSettings : public FAndroidTargetPlatformSettings
 {
 public:
-	FAndroid_ASTCTargetPlatformSettings() : FAndroidTargetPlatformSettings(TEXT("ASTC"))
+	FAndroid_ASTCTargetPlatformSettings(const TCHAR* Name=TEXT("ASTC")) : FAndroidTargetPlatformSettings(Name)
 	{
 	}
 
@@ -216,3 +216,12 @@ public:
 		return FAndroidTargetPlatformSettings::SupportsFeature(Feature);
 	}
 };
+
+class FAndroid_XRTargetPlatformSettings : public FAndroid_ASTCTargetPlatformSettings
+{
+
+public:
+	FAndroid_XRTargetPlatformSettings() : FAndroid_ASTCTargetPlatformSettings(TEXT("XR"))
+	{
+	}
+};
diff --git a/Engine/Source/Developer/Android/AndroidTargetPlatformSettings/Public/IAndroidTargetPlatformSettingsModule.h b/Engine/Source/Developer/Android/AndroidTargetPlatformSettings/Public/IAndroidTargetPlatformSettingsModule.h
--- a/Engine/Source/Developer/Android/AndroidTargetPlatformSettings/Public/IAndroidTargetPlatformSettingsModule.h
+++ b/Engine/Source/Developer/Android/AndroidTargetPlatformSettings/Public/IAndroidTargetPlatformSettingsModule.h
@@ -15,4 +15,5 @@ class IAndroidTargetPlatformSettingsModule : public ITargetPlatformSettingsModul
 public:
 	virtual void GetSinglePlatformSettings(TMap<EAndroidTextureFormatCategory, ITargetPlatformSettings*>& OutMap) = 0;
 	virtual ITargetPlatformSettings* GetMultiPlatformSettings() = 0;
+	virtual ITargetPlatformSettings* GetXRPlatformSettings() = 0;
 };
diff --git a/Engine/Source/Developer/TargetPlatform/Private/TargetPlatformManagerModule.cpp b/Engine/Source/Developer/TargetPlatform/Private/TargetPlatformManagerModule.cpp
--- a/Engine/Source/Developer/TargetPlatform/Private/TargetPlatformManagerModule.cpp
+++ b/Engine/Source/Developer/TargetPlatform/Private/TargetPlatformManagerModule.cpp
@@ -30,6 +30,7 @@
 #include "DesktopPlatformModule.h"
 #include "Interfaces/ITurnkeySupportModule.h"
 #include "DataDrivenShaderPlatformInfo.h"
+#include "RHIShaderFormatDefinitions.inl"
 
 DEFINE_LOG_CATEGORY_STATIC(LogTargetPlatformManager, Log, All);
 
@@ -38,6 +39,17 @@ DEFINE_LOG_CATEGORY_STATIC(LogTargetPlatformManager, Log, All);
 #define AUTOSDKS_ENABLED DDPI_HAS_EXTENDED_PLATFORMINFO_DATA
 #endif
 
+
+static int32 GCookTimeCVarControl = 0;
+static FAutoConsoleVariableRef CVarCookTimeCVarControl(
+	TEXT("cook.cvarcontrol"),
+	GCookTimeCVarControl,
+	TEXT("Controls how cvars for other platforms are managed during cooking:\n"
+		"0: Completely disabled\n"
+		"1: Redirects CVars for the cooking platform to the cooking DP. CURRENTLY ALSO ENABLES redirecting a ShaderPlatform to the cooking platform"
+		"(more to come)"));
+
+
 #if AUTOSDKS_ENABLED
 namespace UE::AutoSDK
 {
@@ -104,6 +116,7 @@ FDelayedAutoRegisterHelper GAutoSDKInit(EDelayedRegisterRunPhase::FileSystemRead
 #endif
 
 
+
 static const size_t MaxPlatformCount = 64;		// In the unlikely event that someone bumps this please note that there's
 												// an implicit assumption that there won't be more than 64 unique target
 												// platforms in the FTargetPlatformSet code since it uses one bit of an
@@ -273,6 +286,36 @@ public:
 			GetShaderFormats();
 		}
 
+		if (GCookTimeCVarControl == 1)
+		{
+			// hook up some global mappings
+			// @todo clear out existing mappings for each Invalidate?
+			for (const ITargetPlatform* Platform : GetActiveTargetPlatforms())
+			{
+				FName PlatformName = *Platform->IniPlatformName();
+				FName DPName = *Platform->CookingDeviceProfileName();
+				// register that when we cook for this platform, we will want to use the given DP when looking up CVar values
+				// if it matches the platform, we don't set it, because that's the default, so don't actually do anything special
+				if (DPName != PlatformName)
+				{
+					ConsoleVariablePlatformMapping::RegisterPlatformToDeviceProfileMapping(PlatformName, DPName);
+				}
+
+				if (Platform != GetRunningTargetPlatform())
+				{
+					// register that when we cook for this platform's shaderplatforms, we will want to use the given platform when looking up cvars
+					TArray<FName> ShaderFormats;
+					Platform->GetAllTargetedShaderFormats(ShaderFormats);
+					for (FName SF : ShaderFormats)
+					{
+						EShaderPlatform SP = ShaderFormatNameToShaderPlatform(SF);
+						ConsoleVariablePlatformMapping::RegisterShaderPlatformToPlatformMapping((int)SP, PlatformName);
+					}
+				}
+			}
+		}
+
+
 		OnTargetPlatformsInvalidated.Broadcast();
 	}
 
diff --git a/Engine/Source/Runtime/Core/Private/HAL/ConsoleManager.cpp b/Engine/Source/Runtime/Core/Private/HAL/ConsoleManager.cpp
--- a/Engine/Source/Runtime/Core/Private/HAL/ConsoleManager.cpp
+++ b/Engine/Source/Runtime/Core/Private/HAL/ConsoleManager.cpp
@@ -1505,17 +1505,36 @@ static FName MakePlatformKey(FName PlatformName, const FString& DeviceProfileNam
 	return *(PlatformName.ToString() + TEXT("/") + (DeviceProfileName.Len() ? DeviceProfileName : PlatformName.ToString()));
 }
 
+static FString CheckForRemappedDP(FName PlatformName, const FString& DeviceProfileName)
+{
+	// check if we have an mapped override value (when cooking, usually) unless we passed in a DPName
+	if (DeviceProfileName == TEXT(""))
+	{
+		FName MappedOverride = ConsoleVariablePlatformMapping::GetPlatformToDeviceProfileMapping(PlatformName);
+		if (MappedOverride != NAME_None)
+		{
+			return MappedOverride.ToString();
+		}
+	}
+
+	return DeviceProfileName;
+}
+
 template<class T>
-bool FConsoleVariableExtendedData<T>::HasPlatformValueVariable(FName PlatformName, const FString& DeviceProfileName)
+bool FConsoleVariableExtendedData<T>::HasPlatformValueVariable(FName PlatformName, const FString& InDeviceProfileName)
 {
+	FString DeviceProfileName = CheckForRemappedDP(PlatformName, InDeviceProfileName);
+
 	// cheap lock here, contention is very rare
 	FRWScopeLock Lock(PlatformValuesLock, SLT_Write);
 	return PlatformValues.Contains(MakePlatformKey(PlatformName, DeviceProfileName));
 }
 
 template<class T>
-TSharedPtr<IConsoleVariable> FConsoleVariableExtendedData<T>::GetPlatformValueVariable(FName PlatformName, const FString& DeviceProfileName)
+TSharedPtr<IConsoleVariable> FConsoleVariableExtendedData<T>::GetPlatformValueVariable(FName PlatformName, const FString& InDeviceProfileName)
 {
+	FString DeviceProfileName = CheckForRemappedDP(PlatformName, InDeviceProfileName);
+
 	// if we have GSpecialDPNameForPremadePlatformKey passed in, we have already gone through the
 	// Load and we have a premade key in PlatformName
 	if (DeviceProfileName != GSpecialDPNameForPremadePlatformKey)
@@ -5036,5 +5055,41 @@ FAutoConsoleCommandDeprecated::FAutoConsoleCommandDeprecated(const TCHAR* Name,
 }
 
 
+namespace ConsoleVariablePlatformMapping
+{
+	TMap<int, FName> ShaderPlatformToPlatformMap;
+	TMap<FName, FName> PlatformToDeviceProfileMap;
+
+	void RegisterShaderPlatformToPlatformMapping(int ShaderPlatform, FName PlatformName)
+	{
+		checkf(!ShaderPlatformToPlatformMap.Contains(ShaderPlatform) || ShaderPlatformToPlatformMap.FindRef(ShaderPlatform) == PlatformName, 
+			TEXT("There's already a platform (%s) targeting ShaderPlatform %d, when setting mapping to %s"), 
+			*ShaderPlatformToPlatformMap.FindRef(ShaderPlatform).ToString(), ShaderPlatform, *PlatformName.ToString());
+
+		ShaderPlatformToPlatformMap.Add(ShaderPlatform, PlatformName);
+	}
+
+	void RegisterPlatformToDeviceProfileMapping(FName PlatformName, FName DeviceProfileName)
+	{
+		checkf(!PlatformToDeviceProfileMap.Contains(PlatformName) || PlatformToDeviceProfileMap.FindRef(PlatformName) == DeviceProfileName,
+			TEXT("There's already a DP (%s) being used for platform %s, when setting mapping to %s"), 
+			*PlatformToDeviceProfileMap.FindRef(PlatformName).ToString(), *PlatformName.ToString(), *DeviceProfileName.ToString());
+
+		PlatformToDeviceProfileMap.Add(PlatformName, DeviceProfileName);
+	}
+
+	FName GetShaderPlatformToPlatformMapping(int ShaderPlatform)
+	{
+		// this will return NAME_None if it's not present
+		return ShaderPlatformToPlatformMap.FindRef(ShaderPlatform);
+	}
+
+	FName GetPlatformToDeviceProfileMapping(FName PlatformName)
+	{
+		// this will return NAME_None if it's not present
+		return PlatformToDeviceProfileMap.FindRef(PlatformName);
+	}
+}
+
 
 #undef LOCTEXT_NAMESPACE
diff --git a/Engine/Source/Runtime/Core/Public/HAL/IConsoleManager.h b/Engine/Source/Runtime/Core/Public/HAL/IConsoleManager.h
--- a/Engine/Source/Runtime/Core/Public/HAL/IConsoleManager.h
+++ b/Engine/Source/Runtime/Core/Public/HAL/IConsoleManager.h
@@ -196,6 +196,20 @@ extern CORE_API const TCHAR* GetConsoleVariableSetByName(EConsoleVariableFlags C
 /** Inverse of GetConsoleVariableSetByName() */
 extern CORE_API EConsoleVariableFlags GetConsoleVariableSetByValue(const TCHAR* SetByName);
 
+namespace ConsoleVariablePlatformMapping
+{
+	// Store a mapping from EShaderPlatform (as int) to a Platform that is currently cooking for that ShaderPlatform
+	CORE_API void RegisterShaderPlatformToPlatformMapping(int ShaderPlatform, FName PlatformName);
+	// Store a mapping from Platform to a DevicveProfile that is currently being used when cooking for that Platform
+	CORE_API void RegisterPlatformToDeviceProfileMapping(FName PlatformName, FName DeviceProfileName);
+
+	// Receives the mapping previously set with RegisterShaderPlatformToPlatformMapping
+	// Returns NAME_None if no mapping exists
+	CORE_API FName GetShaderPlatformToPlatformMapping(int ShaderPlatform);
+	// Receives the mapping previously set with RegisterPlatformToDeviceProfileMapping
+	// Returns NAME_None if no mapping exists
+	CORE_API FName GetPlatformToDeviceProfileMapping(FName PlatformName);
+}
 
 class IConsoleVariable;
 
