Removed the TSharedPtr implementation of FTextDisplayString now that the TRefCountPtr method is mature

diff --git a/Engine/Source/Runtime/Core/Public/Internationalization/LocalizedTextSourceTypes.h b/Engine/Source/Runtime/Core/Public/Internationalization/LocalizedTextSourceTypes.h
--- a/Engine/Source/Runtime/Core/Public/Internationalization/LocalizedTextSourceTypes.h
+++ b/Engine/Source/Runtime/Core/Public/Internationalization/LocalizedTextSourceTypes.h
@@ -83,12 +83,6 @@ struct ELocalizedTextSourcePriority
 	};
 };
 
-#ifndef UE_TEXT_DISPLAYSTRING_USE_REFCOUNTPTR
-	#define UE_TEXT_DISPLAYSTRING_USE_REFCOUNTPTR (1)
-#endif
-
-#if UE_TEXT_DISPLAYSTRING_USE_REFCOUNTPTR
-
 namespace UE::Text::Private
 {
 
@@ -356,17 +350,3 @@ using FTextConstDisplayStringPtr = UE::Text::Private::TDisplayStringPtr<const FS
 {
 	return FTextDisplayStringRef(MakeRefCount<UE::Text::Private::FRefCountedDisplayString>(MoveTemp(InDisplayString)));
 }
-
-#else	// UE_TEXT_DISPLAYSTRING_USE_REFCOUNTPTR
-
-using FTextDisplayStringRef = TSharedRef<FString, ESPMode::ThreadSafe>;
-using FTextDisplayStringPtr = TSharedPtr<FString, ESPMode::ThreadSafe>;
-using FTextConstDisplayStringRef = TSharedRef<const FString, ESPMode::ThreadSafe>;
-using FTextConstDisplayStringPtr = TSharedPtr<const FString, ESPMode::ThreadSafe>;
-
-[[nodiscard]] inline FTextDisplayStringRef MakeTextDisplayString(FString&& InDisplayString)
-{
-	return MakeShared<FString, ESPMode::ThreadSafe>(MoveTemp(InDisplayString));
-}
-
-#endif	// UE_TEXT_DISPLAYSTRING_USE_REFCOUNTPTR
