Make entities private by default by making sure they don't get the RF_Public flag on creation.
Keep the RF_Public flag on entity archetypes and prefab assets.
Updated scene graph unit tests to validate the public flags.

diff --git a/Engine/Source/Runtime/Core/Public/UObject/FortniteMainBranchObjectVersion.h b/Engine/Source/Runtime/Core/Public/UObject/FortniteMainBranchObjectVersion.h
--- a/Engine/Source/Runtime/Core/Public/UObject/FortniteMainBranchObjectVersion.h
+++ b/Engine/Source/Runtime/Core/Public/UObject/FortniteMainBranchObjectVersion.h
@@ -547,6 +547,9 @@ struct FFortniteMainBranchObjectVersion
 		// Added support for relative transform in WorldPartitionActorDesc
 		WorldPartitionActorDescSerializeRelativeTransform,
 
+		// Make sure scene graph entities are not public by default
+		SceneGraphEntitiesPrivateByDefault,
+
 		// -----<new versions can be added above this line>-------------------------------------------------
 		VersionPlusOne,
 		LatestVersion = VersionPlusOne - 1
