Remove old task graph front end code (it was disabled in 5.5) and only support TASKGRAPH_NEW_FRONTEND=1
Deprecate TaskGraphDefinitions.h as it was only needed during the transition

diff --git a/Engine/Source/Runtime/Core/Private/Async/TaskGraph.cpp b/Engine/Source/Runtime/Core/Private/Async/TaskGraph.cpp
--- a/Engine/Source/Runtime/Core/Private/Async/TaskGraph.cpp
+++ b/Engine/Source/Runtime/Core/Private/Async/TaskGraph.cpp
@@ -1335,32 +1335,7 @@ public:
 private:
 	void QueueTask(class FBaseGraphTask* Task, bool bWakeUpWorker, ENamedThreads::Type InThreadToExecuteOn, ENamedThreads::Type InCurrentThreadIfKnown) override
 	{
-		if (ENamedThreads::GetThreadIndex(InThreadToExecuteOn) == ENamedThreads::AnyThread)
-		{
-#if TASKGRAPH_NEW_FRONTEND
-			checkNoEntry();
-#else
-			uint32 ThreadPriority = GetThreadPriorityIndex(InThreadToExecuteOn);
-			check(ThreadPriority < uint32(ENamedThreads::NumThreadPriorities));
-			LowLevelTasks::ETaskPriority Conversion[int(ENamedThreads::NumThreadPriorities)] = { LowLevelTasks::ETaskPriority::Normal, LowLevelTasks::ETaskPriority::High, LowLevelTasks::ETaskPriority::BackgroundNormal };
-			LowLevelTasks::ETaskPriority Priority = Conversion[ThreadPriority];
-
-			if (Priority == LowLevelTasks::ETaskPriority::BackgroundNormal && GetTaskPriority(InThreadToExecuteOn))
-			{
-				Priority = LowLevelTasks::ETaskPriority::BackgroundHigh;
-			}
-
-			Task->GetTaskHandle().Init(TEXT("TaskGraphTask"), Priority, [Task, InThreadToExecuteOn, Deleter(LowLevelTasks::TDeleter<FBaseGraphTask, &FBaseGraphTask::DeleteTask>(Task))]()
-			{
-				Task->Execute(NewTasks, InThreadToExecuteOn, false);
-			});
-	
-			bWakeUpWorker |= LowLevelTasks::FSchedulerTls::IsBusyWaiting();
-
-			verifySlow(LowLevelTasks::TryLaunch(Task->GetTaskHandle(), bWakeUpWorker ? LowLevelTasks::EQueuePreference::GlobalQueuePreference : LowLevelTasks::EQueuePreference::LocalQueuePreference, bWakeUpWorker));
-			return;
-#endif
-		}
+		check(ENamedThreads::GetThreadIndex(InThreadToExecuteOn) != ENamedThreads::AnyThread);
 
 		ENamedThreads::Type CurrentThreadIfKnown;
 		if (ENamedThreads::GetThreadIndex(InCurrentThreadIfKnown) == ENamedThreads::AnyThread)
@@ -1466,7 +1441,6 @@ private:
 		TaskTrace::FWaitingScope WaitingScope(GetTraceIds(Tasks));
 		TRACE_CPUPROFILER_EVENT_SCOPE(WaitUntilTasksComplete);
 
-#if TASKGRAPH_NEW_FRONTEND
 		// try retraction first and only invoke named thread waiting as a last resort
 		bool bAllTasksCompleted = true;
 		for (const FGraphEventRef& Task : Tasks)
@@ -1483,7 +1457,6 @@ private:
 		{
 			return;
 		}
-#endif
 
 		ENamedThreads::Type CurrentThread = CurrentThreadIfKnown;
 		if (ENamedThreads::GetThreadIndex(CurrentThreadIfKnown) == ENamedThreads::AnyThread)
@@ -1574,9 +1547,9 @@ private:
 			InEvent->Trigger();
 			return;
 		}
-#if TASKGRAPH_NEW_FRONTEND
-		// When the new taskgraph frontend is active, FGraphEventArray can be used as prerequisites.
-		// We prefer the new task system because it avoids a potential deadlock situation if all threads are used
+
+		// With the new taskgraph frontend FGraphEventArray can be used as prerequisites.
+		// This which avoids a potential deadlock situation if all threads are used
 		// while the condition is met because this one will execute inline from the thread that completed the prereqs
 		// instead of scheduling yet another task just to trigger the event that may never run if the taskgraph is already full.
 		UE::Tasks::Launch(
@@ -1589,9 +1562,6 @@ private:
 			LowLevelTasks::ETaskPriority::Normal,
 			UE::Tasks::EExtendedTaskPriority::Inline
 		);
-#else
-		TGraphTask<FTriggerEventGraphTask>::CreateTask(&Tasks, CurrentThreadIfKnown).ConstructAndDispatchWhenReady(InEvent, TriggerThread);
-#endif
 	}
 
 	void AddShutdownCallback(TFunction<void()>& Callback) override
@@ -1716,126 +1686,12 @@ bool FTaskGraphInterface::IsMultithread()
 	return FPlatformProcess::SupportsMultithreading() || (FForkProcessHelper::IsForkedMultithreadInstance() && GAllowTaskGraphForkMultithreading);
 }
 
-#if TASKGRAPH_NEW_FRONTEND
-
 FGraphEventImplAllocator& GetGraphEventImplAllocator()
 {
 	static FGraphEventImplAllocator Singleton;
 	return Singleton;
 }
 
-#else
-
-// Statics and some implementations from FBaseGraphTask and FGraphEvent
-
-#if !(UE_BUILD_SHIPPING || UE_BUILD_TEST)
-void FBaseGraphTask::LogPossiblyInvalidSubsequentsTask(const TCHAR* TaskName)
-{
-	UE_LOG(LogTaskGraph, Warning, TEXT("Subsequents of %s look like they contain invalid pointer(s)."), TaskName);
-}
-#endif
-
-static TLockFreeClassAllocator_TLSCache<FGraphEvent, PLATFORM_CACHE_LINE_SIZE>& GetGraphEventAllocator()
-{
-	static TLockFreeClassAllocator_TLSCache<FGraphEvent, PLATFORM_CACHE_LINE_SIZE> Allocator;
-	return Allocator;
-}
-
-FGraphEventRef FGraphEvent::CreateGraphEvent()
-{
-	FGraphEvent* Instance = new(GetGraphEventAllocator().Allocate()) FGraphEvent{};
-	return Instance;
-}
-
-void FGraphEvent::Recycle(FGraphEvent* ToRecycle)
-{
-	GetGraphEventAllocator().Free(ToRecycle);
-}
-
-void FGraphEvent::DispatchSubsequents(ENamedThreads::Type CurrentThreadIfKnown)
-{
-	TArray<FBaseGraphTask*> NewTasks;
-	DispatchSubsequents(NewTasks, CurrentThreadIfKnown);
-}
-
-void FGraphEvent::DispatchSubsequents(TArray<FBaseGraphTask*>& NewTasks, ENamedThreads::Type CurrentThreadIfKnown, bool bInternal/* = false */)
-{
-	if (EventsToWaitFor.Num())
-	{
-		// need to save this first and empty the actual tail of the task might be recycled faster than it is cleared.
-		FGraphEventArray TempEventsToWaitFor;
-		Exchange(EventsToWaitFor, TempEventsToWaitFor);
-
-		bool bSpawnGatherTask = true;
-
-		if (GTestDontCompleteUntilForAlreadyComplete)
-		{
-			bSpawnGatherTask = false;
-			for (FGraphEventRef& Item : TempEventsToWaitFor)
-			{
-				if (!Item->IsComplete())
-				{
-					bSpawnGatherTask = true;
-					break;
-				}
-			}
-		}
-
-		if (bSpawnGatherTask)
-		{
-			// create the Gather...this uses a special version of private CreateTask that "assumes" the subsequent list (which other threads might still be adding too).
-			DECLARE_CYCLE_STAT(TEXT("FNullGraphTask.DontCompleteUntil"), STAT_FNullGraphTask_DontCompleteUntil, STATGROUP_TaskGraphTasks);
-
-			ENamedThreads::Type LocalThreadToDoGatherOn = ENamedThreads::AnyHiPriThreadHiPriTask;
-			if (!GIgnoreThreadToDoGatherOn)
-			{
-				//LocalThreadToDoGatherOn = ThreadToDoGatherOn;
-				ENamedThreads::Type CurrentThreadIndex = ENamedThreads::GetThreadIndex(CurrentThreadIfKnown);
-				if (CurrentThreadIndex <= ENamedThreads::ActualRenderingThread)
-				{
-					LocalThreadToDoGatherOn = CurrentThreadIndex;
-				}
-			}
-
-//#if UE_TASK_TRACE_ENABLED
-//			// regenerate TraceId as we're going to use the same event with another task
-//			TraceId = TaskTrace::GenerateTaskId();
-//#endif
-			TGraphTask<FNullGraphTask>::CreateTask(FGraphEventRef(this), &TempEventsToWaitFor, CurrentThreadIfKnown).ConstructAndDispatchWhenReady(GET_STATID(STAT_FNullGraphTask_DontCompleteUntil), LocalThreadToDoGatherOn);
-			return;
-		}
-	}
-
-	bool bWakeUpWorker = false;
-	TArray<FBaseGraphTask*> PoppedTasks;
-	SubsequentList.PopAllAndClose(PoppedTasks);
-	for (FBaseGraphTask* NewTask : ReverseIterate(PoppedTasks)) // reverse the order since PopAll is implicitly backwards
-	{
-		checkThreadGraph(NewTask);
-		NewTask->ConditionalQueueTask(CurrentThreadIfKnown, bWakeUpWorker);
-	}
-
-	if (!bInternal)
-	{
-		TaskTrace::Launched(GetTraceId(), TEXT("Standalone graph event"), true, ENamedThreads::AnyThread, sizeof(FGraphEvent));
-	}
-	TaskTrace::Completed(GetTraceId());
-}
-
-FGraphEvent::~FGraphEvent()
-{
-#if DO_CHECK
-	if (!IsComplete())
-	{
-		check(SubsequentList.IsClosed());
-	}
-#endif
-	CheckDontCompleteUntilIsEmpty(); // We should not have any wait untils outstanding
-
-	TaskTrace::Destroyed(GetTraceId());
-}
-
-#endif
 
 DECLARE_CYCLE_STAT(TEXT("FBroadcastTask"), STAT_FBroadcastTask, STATGROUP_TaskGraphTasks);
 
@@ -2087,55 +1943,7 @@ static FAutoConsoleCommand TaskThreadPriorityCmd(
 
 int32 WaitForAnyTaskCompleted(const FGraphEventArray& GraphEvents, FTimespan Timeout /*= FTimespan::MaxValue()*/)
 {
-#if TASKGRAPH_NEW_FRONTEND
 	return UE::Tasks::WaitAny(GraphEvents, Timeout);
-#else
-	if (UNLIKELY(GraphEvents.Num() == 0))
-	{
-		return INDEX_NONE;
-	}
-
-	// Avoid memory allocations if any of the events are already completed
-	for (int32 Index = 0; Index < GraphEvents.Num(); ++Index)
-	{
-		if (GraphEvents[Index]->IsComplete())
-		{
-			return Index;
-		}
-	}
-
-	struct FSharedData
-	{
-		UE::FManualResetEvent Event;
-		std::atomic<int32>    CompletedTaskIndex{ 0 };
-	};
-
-	// Shared data usage is important to avoid the variable to go out of scope
-	// before all the task have been run even if we exit after the first event
-	// is triggered.
-	TSharedRef<FSharedData> SharedData = MakeShared<FSharedData>();
-
-	for (int32 Index = 0; Index < GraphEvents.Num(); ++Index)
-	{
-		FFunctionGraphTask::CreateAndDispatchWhenReady(
-			[SharedData, Index]
-			{
-				SharedData->CompletedTaskIndex.store(Index, std::memory_order_relaxed);
-				SharedData->Event.Notify();
-			},
-			TStatId{},
-			GraphEvents[Index],
-			ENamedThreads::AnyHiPriThreadHiPriTask /* Run as soon as possible. */
-		);
-	}
-
-	if (SharedData->Event.WaitFor(UE::FMonotonicTimeSpan::FromMilliseconds(Timeout.GetTotalMilliseconds())))
-	{
-		return SharedData->CompletedTaskIndex.load(std::memory_order_relaxed);
-	}
-
-	return INDEX_NONE;
-#endif
 }
 
 FGraphEventRef AnyTaskCompleted(const FGraphEventArray& GraphEvents)
diff --git a/Engine/Source/Runtime/Core/Private/Tasks/TaskPrivate.cpp b/Engine/Source/Runtime/Core/Private/Tasks/TaskPrivate.cpp
--- a/Engine/Source/Runtime/Core/Private/Tasks/TaskPrivate.cpp
+++ b/Engine/Source/Runtime/Core/Private/Tasks/TaskPrivate.cpp
@@ -24,7 +24,6 @@ namespace UE::Tasks
 		{
 			TaskTrace::Scheduled(GetTraceId());
 
-#if TASKGRAPH_NEW_FRONTEND
 			if (IsNamedThreadTask())
 			{
 				ENamedThreads::Type ConversionMap[] =
@@ -48,7 +47,6 @@ namespace UE::Tasks
 				FTaskGraphInterface::Get().QueueTask(static_cast<FBaseGraphTask*>(this), true, ConversionMap[(int32)ExtendedPriority - (int32)EExtendedTaskPriority::GameThreadNormalPri]);
 				return;
 			}
-#endif
 			
 			// In case a thread is waiting on us to perform retraction, now is the time to try retraction again.
 			// This needs to be before the launch as performing the execution can destroy the task.
@@ -97,13 +95,11 @@ namespace UE::Tasks
 				return false;
 			}
 
-#if TASKGRAPH_NEW_FRONTEND
 			// task retraction is not supported for named thread tasks
 			if (IsNamedThreadTask())
 			{
 				return false;
 			}
-#endif
 
 			// avoid stack overflow. is not expected in a real-life cases but happens in stress tests
 			if (RecursionDepth == 200)
@@ -317,7 +313,6 @@ namespace UE::Tasks
 
 		bool TryWaitOnNamedThread(FTaskBase& Task)
 		{
-#if TASKGRAPH_NEW_FRONTEND
 			// handle waiting only on a named thread and if not called from inside a task
 			FTaskGraphInterface& TaskGraph = FTaskGraphInterface::Get();
 			ENamedThreads::Type CurrentThread = TaskGraph.GetCurrentThreadIfKnown();
@@ -338,7 +333,6 @@ namespace UE::Tasks
 				check(Task.IsCompleted());
 				return true;
 			}
-#endif
 
 			return false;
 		}
@@ -357,7 +351,6 @@ namespace UE::Tasks
 			TEXT("Inline"),
 			TEXT("TaskEvent"),
 
-#if TASKGRAPH_NEW_FRONTEND
 			TEXT("GameThreadNormalPri"),
 			TEXT("GameThreadHiPri"),
 			TEXT("GameThreadNormalPriLocalQueue"),
@@ -372,7 +365,6 @@ namespace UE::Tasks
 			TEXT("RHIThreadHiPri"),
 			TEXT("RHIThreadNormalPriLocalQueue"),
 			TEXT("RHIThreadHiPriLocalQueue")
-#endif
 		};
 		return ExtendedTaskPriorityToStr[(int32)ExtendedPriority];
 	}
@@ -390,7 +382,6 @@ namespace UE::Tasks
 		CONVERT_EXTENDED_TASK_PRIORITY(Inline);
 		CONVERT_EXTENDED_TASK_PRIORITY(TaskEvent);
 
-#if TASKGRAPH_NEW_FRONTEND
 		CONVERT_EXTENDED_TASK_PRIORITY(GameThreadNormalPri);
 		CONVERT_EXTENDED_TASK_PRIORITY(GameThreadHiPri);
 		CONVERT_EXTENDED_TASK_PRIORITY(GameThreadNormalPriLocalQueue);
@@ -405,7 +396,6 @@ namespace UE::Tasks
 		CONVERT_EXTENDED_TASK_PRIORITY(RHIThreadHiPri);
 		CONVERT_EXTENDED_TASK_PRIORITY(RHIThreadNormalPriLocalQueue);
 		CONVERT_EXTENDED_TASK_PRIORITY(RHIThreadHiPriLocalQueue);
-#endif
 
 #undef CONVERT_EXTENDED_TASK_PRIORITY
 
@@ -467,8 +457,6 @@ namespace UE::Tasks
 		{
 			using namespace UE::Tasks;
 
-#if TASKGRAPH_NEW_FRONTEND
-
 			ENamedThreads::Type ThreadIndex = ENamedThreads::GetThreadIndex(ThreadType);
 			if (ThreadIndex != ENamedThreads::AnyThread)
 			{
@@ -497,18 +485,6 @@ namespace UE::Tasks
 			{
 				OutPriority = ETaskPriority::BackgroundHigh;
 			}
-
-#else
-			uint32 ThreadPriority = GetThreadPriorityIndex(ThreadType);
-			check(ThreadPriority < uint32(ENamedThreads::NumThreadPriorities));
-			ETaskPriority Conversion[int(ENamedThreads::NumThreadPriorities)] = { ETaskPriority::Normal, ETaskPriority::High, ETaskPriority::BackgroundNormal };
-			OutPriority = Conversion[ThreadPriority];
-
-			if (OutPriority == ETaskPriority::BackgroundNormal && GetTaskPriority(ThreadType))
-			{
-				OutPriority = ETaskPriority::BackgroundHigh;
-			}
-#endif
 		}
 
 		ENamedThreads::Type TranslatePriority(ETaskPriority Priority)
@@ -527,16 +503,6 @@ namespace UE::Tasks
 			return ConversionMap[(int32)Priority];
 		}
 
-#if !TASKGRAPH_NEW_FRONTEND
-
-		ENamedThreads::Type TranslatePriority(ETaskPriority Priority, EExtendedTaskPriority /*ExtendedPriority*/)
-		{
-			checkSlow(Priority <= ETaskPriority::Count);
-			return TranslatePriority(Priority);
-		}
-
-#else // !TASKGRAPH_NEW_FRONTEND
-
 		// task priority translation from the new API to the old API
 		ENamedThreads::Type TranslatePriority(EExtendedTaskPriority Priority)
 		{
@@ -571,13 +537,10 @@ namespace UE::Tasks
 			return ExtendedPriority < EExtendedTaskPriority::GameThreadNormalPri || ExtendedPriority == EExtendedTaskPriority::Count ? TranslatePriority(Priority) : TranslatePriority(ExtendedPriority);
 		}
 
-#endif // !TASKGRAPH_NEW_FRONTEND
-
 		bool ShouldForceWaitWithNamedThreadsSupport(EExtendedTaskPriority ExtendedPriority)
 		{
 			// We force wait named thread support when we're waiting on a task that must run on the same thread we're currently on.
 			// If we don't do this, it's a guaranteed deadlock.
-#if TASKGRAPH_NEW_FRONTEND
 			const bool bIsNamedThreadTask = ExtendedPriority >= EExtendedTaskPriority::GameThreadNormalPri;
 			if (bIsNamedThreadTask)
 			{
@@ -589,7 +552,6 @@ namespace UE::Tasks
 					return TaskThreadIndex == CurrentThreadIndex;
 				}
 			}
-#endif
 			return false;
 		}
 	}
diff --git a/Engine/Source/Runtime/Core/Public/Async/TaskGraphDefinitions.h b/Engine/Source/Runtime/Core/Public/Async/TaskGraphDefinitions.h
--- a/Engine/Source/Runtime/Core/Public/Async/TaskGraphDefinitions.h
+++ b/Engine/Source/Runtime/Core/Public/Async/TaskGraphDefinitions.h
@@ -1,7 +1,10 @@
 // Copyright Epic Games, Inc. All Rights Reserved.
+
 #pragma once
 
-#if !defined(TASKGRAPH_NEW_FRONTEND)
-#define TASKGRAPH_NEW_FRONTEND 1
-#endif
+#include "CoreTypes.h"
 
+UE_DEPRECATED_HEADER(5.6, "Include Async/TaskGraphFwd.h");
+
+// The old task graph interface was removed in UE 5.6
+#define TASKGRAPH_NEW_FRONTEND 1
diff --git a/Engine/Source/Runtime/Core/Public/Async/TaskGraphFwd.h b/Engine/Source/Runtime/Core/Public/Async/TaskGraphFwd.h
--- a/Engine/Source/Runtime/Core/Public/Async/TaskGraphFwd.h
+++ b/Engine/Source/Runtime/Core/Public/Async/TaskGraphFwd.h
@@ -6,19 +6,12 @@
 
 #pragma once
 
-#include "TaskGraphDefinitions.h"
 #include "Templates/RefCounting.h"
 
 class FBaseGraphTask;
 
-#if TASKGRAPH_NEW_FRONTEND
-
+/** Generic task used by the task graph system */
 using FGraphEvent = FBaseGraphTask;
-using FGraphEventRef = TRefCountPtr<FBaseGraphTask>;
-
-#else
 
-/** Convenience typedef for a reference counted pointer to a graph event **/
-using FGraphEventRef = TRefCountPtr<class FGraphEvent>;
-
-#endif
+/** Convenience typedef for a reference counted pointer to a graph event */
+using FGraphEventRef = TRefCountPtr<FBaseGraphTask>;
diff --git a/Engine/Source/Runtime/Core/Public/Async/TaskGraphInterfaces.h b/Engine/Source/Runtime/Core/Public/Async/TaskGraphInterfaces.h
--- a/Engine/Source/Runtime/Core/Public/Async/TaskGraphInterfaces.h
+++ b/Engine/Source/Runtime/Core/Public/Async/TaskGraphInterfaces.h
@@ -273,10 +273,7 @@ CORE_API TArray<TaskTrace::FId> GetTraceIds(const FGraphEventArray& Tasks);
 class FTaskGraphInterface
 {
 	friend class FBaseGraphTask;
-
-#if TASKGRAPH_NEW_FRONTEND
 	friend UE::Tasks::Private::FTaskBase;
-#endif
 
 	/**
 	 *	Internal function to queue a task
@@ -444,8 +441,6 @@ struct FTaskGraphBlockAllocationTag : FDefaultBlockAllocationTag
 	using Allocator = TBlockAllocationCache<BlockSize, FAlignedAllocator>;
 };
 
-#if TASKGRAPH_NEW_FRONTEND
-
 /** 
  *	Base class for all tasks. A replacement for `FBaseGraphTask` and `FGraphEvent` from the old API, based on `Tasks::Private::FTaskBase` functionality
  **/
@@ -679,682 +674,6 @@ inline FGraphEventRef FBaseGraphTask::CreateGraphEvent()
 	return FGraphEventRef{ GraphEvent, /*bAddRef = */ false };
 }
 
-#else // TASKGRAPH_NEW_FRONTEND
-
-struct FTaskBlockAllocationTag : FDefaultBlockAllocationTag
-{
-	static constexpr uint32 BlockSize = 64 * 1024;
-	static constexpr bool AllowOversizedBlocks = false;
-	static constexpr bool RequiresAccurateSize = false;
-	static constexpr bool InlineBlockAllocation = true;
-	static constexpr const char* TagName = "TaskLinearAllocator";
-
-	using Allocator = TBlockAllocationCache<BlockSize, FAlignedAllocator>;
-};
-
-/** 
- *	Base class for all tasks. 
- *	Tasks go through a very specific life stage progression, and this is verified.
- **/
-
-class FBaseGraphTask : private UE::FInheritedContextBase
-{
-protected:
-	/** 
-	 *	Constructor
-	 *	@param InNumberOfPrerequistitesOutstanding; the number of prerequisites outstanding. We actually add one to this to prevent the task from firing while we are setting up the task
-	 **/
-	FBaseGraphTask(int32 InNumberOfPrerequistitesOutstanding)
-		: ThreadToExecuteOn(ENamedThreads::AnyThread)
-		, NumberOfPrerequistitesOutstanding(InNumberOfPrerequistitesOutstanding + 1) // + 1 is not a prerequisite, it is a lock to prevent it from executing while it is getting prerequisites, one it is safe to execute, call PrerequisitesComplete
-	{
-		checkThreadGraph(LifeStage.Increment() == int32(LS_Contructed));
-		CaptureInheritedContext();
-	}
-	/** 
-	 *	Sets the desired execution thread. This is not part of the constructor because this information may not be known quite yet duiring construction.
-	 *	@param InThreadToExecuteOn; the desired thread to execute on.
-	 **/
-	void SetThreadToExecuteOn(ENamedThreads::Type InThreadToExecuteOn)
-	{
-		ThreadToExecuteOn = InThreadToExecuteOn;
-		checkThreadGraph(LifeStage.Increment() == int32(LS_ThreadSet));
-	}
-
-	/** 
-	 *	Indicates that the prerequisites are set up and that the task can be executed as soon as the prerequisites are finished.
-	 *	@param NumAlreadyFinishedPrequistes; the number of prerequisites that have not been set up because those tasks had already completed.
-	 *	@param bUnlock; if true, let the task execute if it can
-	 **/
-	void PrerequisitesComplete(ENamedThreads::Type CurrentThread, int32 NumAlreadyFinishedPrequistes, bool bUnlock = true)
-	{
-		checkThreadGraph(LifeStage.Increment() == int32(LS_PrequisitesSetup));
-		int32 NumToSub = NumAlreadyFinishedPrequistes + (bUnlock ? 1 : 0); // the +1 is for the "lock" we set up in the constructor
-		if (NumberOfPrerequistitesOutstanding.Subtract(NumToSub) == NumToSub) 
-		{
-			bool bWakeUpWorker = true;
-			QueueTask(CurrentThread, bWakeUpWorker);	
-		}
-	}
-	/** destructor, just checks the life stage **/
-	virtual ~FBaseGraphTask()
-	{
-#if DO_GUARD_SLOW
-		int32 Stage = LifeStage.Increment();
-		checkf(Stage == int32(LS_Deconstucted), TEXT("LifeStage was %d"), Stage);
-#endif
-	}
-
-#if !(UE_BUILD_SHIPPING || UE_BUILD_TEST)
-	/** Logs a task name that may contain invalid subsequents. Debug only. */
-	static void CORE_API LogPossiblyInvalidSubsequentsTask(const TCHAR* TaskName);
-#endif
-
-	/** 
-	 *	An indication that a prerequisite has been completed. Reduces the number of prerequisites by one and if no prerequisites are outstanding, it queues the task for execution.
-	 *	@param CurrentThread; provides the index of the thread we are running on. This is handy for submitting new taks. Can be ENamedThreads::AnyThread if the current thread is unknown.
-	 **/
-	void ConditionalQueueTask(ENamedThreads::Type CurrentThread, bool& bWakeUpWorker)
-	{
-		if (NumberOfPrerequistitesOutstanding.Decrement()==0)
-		{
-			QueueTask(CurrentThread, bWakeUpWorker);
-			bWakeUpWorker = true;
-		}
-	}
-
-	TaskTrace::FId GetTraceId() const
-	{
-#if UE_TASK_TRACE_ENABLED
-		return TraceId.load(std::memory_order_relaxed);
-#else
-		return TaskTrace::InvalidId;
-#endif
-	}
-
-	void SetTraceId(TaskTrace::FId InTraceId)
-	{
-#if UE_TASK_TRACE_ENABLED
-		TraceId = InTraceId;
-#endif
-	}
-
-	LowLevelTasks::FTask& GetTaskHandle()
-	{
-		return TaskHandle;
-	}
-
-	ENamedThreads::Type GetThreadToExecuteOn() const
-	{
-		return ThreadToExecuteOn;
-	}
-
-private:
-	friend class FNamedTaskThread;
-	friend class FTaskThreadBase;
-	friend class FTaskThreadAnyThread;
-	friend class FGraphEvent;
-	friend class FTaskGraphImplementation;
-	friend class FTaskGraphCompatibilityImplementation;
-
-	LowLevelTasks::FTask TaskHandle;
-
-	// Subclass API
-
-	/** 
-	 *	Virtual call to actually execute the task. This will also call the destructor and free any memory in the old backend if bDeleteOnCompletion is set to true.
-	 *	@param CurrentThread; provides the index of the thread we are running on. This is handy for submitting new tasks.
-	 *  @param bDeleteOnCompletion; specifies if the task should delete itself after completion.
-	 **/
-	virtual void ExecuteTask(TArray<FBaseGraphTask*>& NewTasks, ENamedThreads::Type CurrentThread, bool bDeleteOnCompletion)=0;
-
-	/** 
-	*	Virtual call to actually delete the task any memory. This is used for the New Backend.
-	**/
-	virtual void DeleteTask() = 0;
-
-	// API called from other parts of the system
-
-	/** 
-	 *	Called by the system to execute this task after it has been removed from an internal queue.
-	 *	Just checks the life stage and passes off to the virtual ExecuteTask method.
-	 *	@param CurrentThread; provides the index of the thread we are running on. This is handy for submitting new tasks.
-	 *  @param bDeleteOnCompletion; specifies if the task should delete itself after completion.
-	 **/
-	FORCEINLINE void Execute(TArray<FBaseGraphTask*>& NewTasks, ENamedThreads::Type CurrentThread, bool bDeleteOnCompletion)
-	{
-		checkThreadGraph(LifeStage.Increment() == int32(LS_Executing));
-
-		UE::FInheritedContextScope InheritedContextScope = RestoreInheritedContext();
-		ExecuteTask(NewTasks, CurrentThread, bDeleteOnCompletion);
-	}
-
-	// Internal Use
-
-	/** 
-	 *	Queues the task for execution.
-	 *	@param CurrentThread; provides the index of the thread we are running on. This is handy for submitting new taks. Can be ENamedThreads::AnyThread if the current thread is unknown.
-	 **/
-	void QueueTask(ENamedThreads::Type CurrentThreadIfKnown, bool bWakeUpWorker)
-	{
-		checkThreadGraph(LifeStage.Increment() == int32(LS_Queued));
-		TaskTrace::Scheduled(GetTraceId());
-		FTaskGraphInterface::Get().QueueTask(this, bWakeUpWorker, ThreadToExecuteOn, CurrentThreadIfKnown);
-	}
-
-	/**	Thread to execute on, can be ENamedThreads::AnyThread to execute on any unnamed thread **/
-	ENamedThreads::Type			ThreadToExecuteOn;
-	/**	Number of prerequisites outstanding. When this drops to zero, the thread is queued for execution.  **/
-	FThreadSafeCounter			NumberOfPrerequistitesOutstanding; 
-
-
-#if DO_GUARD_SLOW || USING_CODE_ANALYSIS
-	// Life stage verification
-	// Tasks go through 8 steps, in order. In non-final builds, we track them with a thread safe counter and verify that the progression is correct.
-	enum ELifeStage
-	{
-		LS_BaseContructed = 0,
-		LS_Contructed,
-		LS_ThreadSet,
-		LS_PrequisitesSetup,
-		LS_Queued,
-		LS_Executing,
-		LS_Deconstucted,
-	};
-	FThreadSafeCounter			LifeStage;
-
-#endif
-
-#if UE_TASK_TRACE_ENABLED
-	std::atomic<TaskTrace::FId> TraceId { TaskTrace::InvalidId };
-#endif
-};
-
-/** 
- *	A FGraphEvent is a list of tasks waiting for something.
- *	These tasks are call the subsequents.
- *	A graph event is a prerequisite for each of its subsequents.
- *	Graph events have a lifetime managed by reference counting.
- **/
-class FGraphEvent 
-{
-public:
-	/**
-	 *	A factory method to create a graph event. 
-	 *	@return a reference counted pointer to the new graph event. Note this should be stored in a FGraphEventRef or it will be immediately destroyed!
-	**/
-	static CORE_API FGraphEventRef CreateGraphEvent();
-	
-	/**
-	 *	Attempts to a new subsequent task. If this event has already fired, false is returned and action must be taken to ensure that the task will still fire even though this event cannot be a prerequisite (because it is already finished).
-	 *	@return true if the task was successfully set up as a subsequent. false if the event has already fired.
-	**/
-	bool AddSubsequent(class FBaseGraphTask* Subsequent)
-	{
-		bool bSucceeded = SubsequentList.PushIfNotClosed(Subsequent);
-		if (bSucceeded)
-		{
-			TaskTrace::SubsequentAdded(GetTraceId(), Subsequent->GetTraceId());
-		}
-		return bSucceeded;
-	}
-
-	/**
-	 *	Verification function to ensure that nobody was tried to add WaitUntil's outside of the context of execution
-	**/
-	void CheckDontCompleteUntilIsEmpty()
-	{
-		checkThreadGraph(!EventsToWaitFor.Num());
-	}
-
-	/**
-	 *	Delay the firing of this event until the given event fires.
-	 *	CAUTION: This is only legal while executing the task associated with this event.
-	 *	@param EventToWaitFor event to wait for until we fire.
-	**/
-	void DontCompleteUntil(FGraphEventRef EventToWaitFor)
-	{
-		checkThreadGraph(!IsComplete()); // it is not legal to add a DontCompleteUntil after the event has been completed. Basically, this is only legal within a task function.
-		EventsToWaitFor.Emplace(EventToWaitFor);
-		TaskTrace::SubsequentAdded(EventToWaitFor->GetTraceId(), GetTraceId());
-	}
-
-	/**
-	 *	"Complete" the event. This grabs the list of subsequents and atomically closes it. Then for each subsequent it reduces the number of prerequisites outstanding and if that drops to zero, the task is queued.
-	 *	@param CurrentThreadIfKnown if the current thread is known, provide it here. Otherwise it will be determined via TLS if any task ends up being queued.
-	**/
-	CORE_API void DispatchSubsequents(ENamedThreads::Type CurrentThreadIfKnown = ENamedThreads::AnyThread);
-
-	/**
-	 *	"Complete" the event. This grabs the list of subsequents and atomically closes it. Then for each subsequent and for each item in "NewTasks" it reduces the number of prerequisites outstanding and if that drops to zero, the task is queued. 
-	 *	@param NewTasks subsequents to add
-	 *	@param	 CurrentThreadIfKnown if the current thread is known, provide it here. Otherwise it will be determined via TLS if any task ends up being queued.
-	 *	@param bInternal used to distinguish between "internal" graph events that have an accompanying task, and "standalone" graph events created explicitly by `FGraphEvent::CreateGraphEvent()`
-	**/
-	CORE_API void DispatchSubsequents(TArray<FBaseGraphTask*>& NewTasks, ENamedThreads::Type CurrentThreadIfKnown = ENamedThreads::AnyThread, bool bInternal = false);
-
-	/**
-	 *	Determine if the event has been completed. This can be used to poll for completion. 
-	 *	@return true if this event has completed. 
-	 *	CAUTION: If this returns false, the event could still end up completing before this function even returns. In other words, a false return does not mean that event is not yet completed!
-	**/
-	bool IsComplete() const
-	{
-		return SubsequentList.IsClosed();
-	}
-
-	/**
-	 * A convenient short version of `FTaskGraphInterface::WaitUntilTaskCompletes`
-	 */
-	void Wait(ENamedThreads::Type CurrentThreadIfKnown = ENamedThreads::AnyThread)
-	{
-		FTaskGraphInterface::Get().WaitUntilTaskCompletes(this, CurrentThreadIfKnown);
-	}
-
-	/**
-	 * Sets a name for the event for debugging purposes.
-	 */
-	void SetDebugName(const TCHAR* Name)
-	{
-#if !UE_BUILD_SHIPPING && !UE_BUILD_TEST
-		DebugName = Name;
-#endif
-	}
-
-	TaskTrace::FId GetTraceId() const
-	{
-#if UE_TASK_TRACE_ENABLED
-		return TraceId.load(std::memory_order_relaxed);
-#else
-		return TaskTrace::InvalidId;
-#endif
-	}
-
-	// does nothing and is needed only for compatibility with the new frontend
-	FGraphEventRef CreateCompletionHandle()
-	{
-		// nothing to do here as this instance already serves as a completion handle
-		return this;
-	}
-
-private:
-	friend class TRefCountPtr<FGraphEvent>;
-	friend class TLockFreeClassAllocator_TLSCache<FGraphEvent, PLATFORM_CACHE_LINE_SIZE>;
-
-	/** 
-	 *	Internal function to call the destructor and recycle a graph event
-	 *	@param ToRecycle; graph event to recycle
-	**/
-	static CORE_API void Recycle(FGraphEvent* ToRecycle);
-
-	friend struct FGraphEventAndSmallTaskStorage;
-
-	/**
-	 *	Hidden Constructor
-	**/
-	FGraphEvent()
-		: ThreadToDoGatherOn(ENamedThreads::AnyHiPriThreadHiPriTask)
-	{
-	}
-
-	/**
-	 *	Destructor. Verifies we aren't destroying it prematurely. 
-	**/
-	~FGraphEvent();
-
-	// Interface for TRefCountPtr
-
-public:
-	/** 
-	 *	Increases the reference count 
-	 *	@return the new reference count
-	**/
-	uint32 AddRef()
-	{
-		int32 RefCount = ReferenceCount.Increment();
-		checkThreadGraph(RefCount > 0);
-		return RefCount;
-	}
-	/** 
-	 *	Decreases the reference count and destroys the graph event if it is zero.
-	 *	@return the new reference count
-	**/
-	uint32 Release()
-	{
-		int32 RefCount = ReferenceCount.Decrement();
-		checkThreadGraph(RefCount >= 0);
-		if (RefCount == 0)
-		{
-			Recycle(this);
-		}
-		return RefCount;
-	}
-
-	uint32 GetRefCount() const
-	{
-		return ReferenceCount.GetValue();
-	}
-
-private:
-
-	/** Threadsafe list of subsequents for the event **/
-	TClosableLockFreePointerListUnorderedSingleConsumer<FBaseGraphTask, 0>	SubsequentList;
-	/** List of events to wait for until firing. This is not thread safe as it is only legal to fill it in within the context of an executing task. **/
-	FGraphEventArray														EventsToWaitFor;
-	/** Number of outstanding references to this graph event **/
-	FThreadSafeCounter														ReferenceCount;
-	ENamedThreads::Type														ThreadToDoGatherOn;
-
-#if !UE_BUILD_SHIPPING && !UE_BUILD_TEST
-	const TCHAR* DebugName = nullptr;
-#endif
-
-#if UE_TASK_TRACE_ENABLED
-	std::atomic<TaskTrace::FId> TraceId { TaskTrace::GenerateTaskId() };
-#endif
-};
-
-/** 
- The user defined task type can take arguments to a constructor. These arguments (unfortunately) must not be references.
- The API required of TTask:
-
-class FGenericTask
-{
-	TSomeType	SomeArgument;
-public:
-	FGenericTask(TSomeType InSomeArgument) // CAUTION!: Must not use references in the constructor args; use pointers instead if you need by reference
-		: SomeArgument(InSomeArgument)
-	{
-		// Usually the constructor doesn't do anything except save the arguments for use in DoWork or GetDesiredThread.
-	}
-	~FGenericTask()
-	{
-		// you will be destroyed immediately after you execute. Might as well do cleanup in DoWork, but you could also use a destructor.
-	}
-	FORCEINLINE TStatId GetStatId() const
-	{
-		RETURN_QUICK_DECLARE_CYCLE_STAT(FGenericTask, STATGROUP_TaskGraphTasks);
-	}
-
-	[static] ENamedThreads::Type GetDesiredThread()
-	{
-		return ENamedThreads::[named thread or AnyThread];
-	}
-	void DoTask(ENamedThreads::Type CurrentThread, const FGraphEventRef& MyCompletionGraphEvent)
-	{
-		// The arguments are useful for setting up other tasks. 
-		// Do work here, probably using SomeArgument.
-		MyCompletionGraphEvent->DontCompleteUntil(TGraphTask<FSomeChildTask>::CreateTask(NULL,CurrentThread).ConstructAndDispatchWhenReady());
-	}
-};
-**/
-
-
-/** 
- *	TGraphTask
- *	Embeds a user defined task, as exemplified above, for doing the work and provides the functionality for setting up and handling prerequisites and subsequents
- **/
-template<typename TTask>
-class TGraphTask final : public TConcurrentLinearObject<TGraphTask<TTask>, FTaskGraphBlockAllocationTag>, public FBaseGraphTask
-{
-public:
-	/** 
-	 *	This is a helper class returned from the factory. It constructs the embeded task with a set of arguments and sets the task up and makes it ready to execute.
-	 *	The task may complete before these routines even return.
-	 **/
-	class FConstructor
-	{
-	public:
-		/** Passthrough internal task constructor and dispatch. Note! Generally speaking references will not pass through; use pointers */
-		template<typename...T>
-		FGraphEventRef ConstructAndDispatchWhenReady(T&&... Args)
-		{
-			new ((void *)&Owner->TaskStorage) TTask(Forward<T>(Args)...);
-			return Owner->Setup(Prerequisites, CurrentThreadIfKnown);
-		}
-
-		/** Passthrough internal task constructor and hold. */
-		template<typename...T>
-		TGraphTask* ConstructAndHold(T&&... Args)
-		{
-			new ((void *)&Owner->TaskStorage) TTask(Forward<T>(Args)...);
-			return Owner->Hold(Prerequisites, CurrentThreadIfKnown);
-		}
-
-	private:
-		friend TGraphTask;
-
-		/** The task that created me to assist with embeded task construction and preparation. **/
-		TGraphTask*						Owner;
-		/** The list of prerequisites. **/
-		const FGraphEventArray*			Prerequisites;
-		/** If known, the current thread.  ENamedThreads::AnyThread is also fine, and if that is the value, we will determine the current thread, as needed, via TLS. **/
-		ENamedThreads::Type				CurrentThreadIfKnown;
-
-		/** Constructor, simply saves off the arguments for later use after we actually construct the embeded task. **/
-		FConstructor(TGraphTask* InOwner, const FGraphEventArray* InPrerequisites, ENamedThreads::Type InCurrentThreadIfKnown)
-			: Owner(InOwner)
-			, Prerequisites(InPrerequisites)
-			, CurrentThreadIfKnown(InCurrentThreadIfKnown)
-		{
-		}
-		/** Prohibited copy construction **/
-		FConstructor(const FConstructor& Other)
-		{
-			check(0);
-		}
-		/** Prohibited copy **/
-		void operator=(const FConstructor& Other)
-		{
-			check(0);
-		}
-	};
-
-	/** 
-	 *	Factory to create a task and return the helper object to construct the embedded task and set it up for execution.
-	 *	@param Prerequisites; the list of FGraphEvents that must be completed prior to this task executing.
-	 *	@param CurrentThreadIfKnown; provides the index of the thread we are running on. Can be ENamedThreads::AnyThread if the current thread is unknown.
-	 *	@return a temporary helper class which can be used to complete the process.
-	**/
-	static FConstructor CreateTask(const FGraphEventArray* Prerequisites = NULL, ENamedThreads::Type CurrentThreadIfKnown = ENamedThreads::AnyThread)
-	{
-		FGraphEventRef GraphEvent = TTask::GetSubsequentsMode() == ESubsequentsMode::FireAndForget ? NULL : FGraphEvent::CreateGraphEvent();
-
-		int32 NumPrereq = Prerequisites ? Prerequisites->Num() : 0;
-		return FConstructor(new TGraphTask(MoveTemp(GraphEvent), NumPrereq), Prerequisites, CurrentThreadIfKnown);
-	}
-
-	void Unlock(ENamedThreads::Type CurrentThreadIfKnown = ENamedThreads::AnyThread)
-	{
-		TaskTrace::Launched(GetTraceId(), nullptr, Subsequents.IsValid(), ((TTask*)&TaskStorage)->GetDesiredThread(), sizeof(*this));
-
-		bool bWakeUpWorker = true;
-		ConditionalQueueTask(CurrentThreadIfKnown, bWakeUpWorker);
-	}
-
-	FGraphEventRef GetCompletionEvent()
-	{
-		return Subsequents;
-	}
-
-private:
-	friend class FConstructor;
-	friend class FGraphEvent;
-
-	// API derived from FBaseGraphTask
-
-	/** 
-	 *	Virtual call to actually execute the task. 
-	 *	@param CurrentThread; provides the index of the thread we are running on. This is handy for submitting new taks.
-	 *  @param bDeleteOnCompletion; specifies if the task will will delete itself after execution (true) or if DeleteTask has to be called manually (false).
-	 *	Executes the embedded task. 
-	 *  Destroys the embedded task.
-	 *	Dispatches the subsequents.
-	 *	Destroys myself.
-	 **/
-	void ExecuteTask(TArray<FBaseGraphTask*>& NewTasks, ENamedThreads::Type CurrentThread, bool bDeleteOnCompletion) override
-	{
-		checkThreadGraph(TaskConstructed);
-
-		// Fire and forget mode must not have subsequents
-		// Track subsequents mode must have subsequents
-		checkThreadGraph((TTask::GetSubsequentsMode() == ESubsequentsMode::FireAndForget) != IsValidRef(Subsequents));
-
-		if (TTask::GetSubsequentsMode() == ESubsequentsMode::TrackSubsequents)
-		{
-			Subsequents->CheckDontCompleteUntilIsEmpty(); // we can only add wait for tasks while executing the task
-		}
-		
-		TTask& Task = *(TTask*)&TaskStorage;
-		{
-			TaskTrace::FTaskTimingEventScope TaskEventScope(GetTraceId());
-			FScopeCycleCounter Scope(Task.GetStatId(), true);
-			Task.DoTask(CurrentThread, Subsequents);
-			Task.~TTask();
-			checkThreadGraph(ENamedThreads::GetThreadIndex(CurrentThread) <= ENamedThreads::GetRenderThread() || FMemStack::Get().IsEmpty()); // you must mark and pop memstacks if you use them in tasks! Named threads are excepted.
-		}
-		
-		TaskConstructed = false;
-
-		if (TTask::GetSubsequentsMode() == ESubsequentsMode::TrackSubsequents)
-		{
-			FPlatformMisc::MemoryBarrier();
-			Subsequents->DispatchSubsequents(NewTasks, CurrentThread, true);
-		}
-		else
-		{
-			// "fire and forget" tasks don't have an accompanying FGraphEvent that traces completion and destruction
-			TaskTrace::Completed(GetTraceId());
-			TaskTrace::Destroyed(GetTraceId());
-		}
-
-		if (bDeleteOnCompletion)
-		{
-			DeleteTask();
-		}
-	}
-
-	void DeleteTask() final override
-	{
-		delete this;
-	}
-
-	// Internals 
-
-	/** 
-	 *	Private constructor, constructs the base class with the number of prerequisites.
-	 *	@param InSubsequents subsequents to associate with this task. This refernence is destroyed in the process!
-	 *	@param NumberOfPrerequistitesOutstanding the number of prerequisites this task will have when it is built.
-	**/
-	TGraphTask(FGraphEventRef InSubsequents, int32 NumberOfPrerequistitesOutstanding)
-		: FBaseGraphTask(NumberOfPrerequistitesOutstanding)
-		, TaskConstructed(false)
-	{
-		Subsequents.Swap(InSubsequents);
-		SetTraceId(Subsequents.IsValid() ? Subsequents->GetTraceId() : TaskTrace::GenerateTaskId());
-	}
-
-	/** 
-	 *	Private destructor, just checks that the task appears to be completed
-	**/
-	~TGraphTask() override
-	{
-		checkThreadGraph(!TaskConstructed);
-	}
-
-	/** 
-	 *	Call from FConstructor to complete the setup process
-	 *	@param Prerequisites; the list of FGraphEvents that must be completed prior to this task executing.
-	 *	@param CurrentThreadIfKnown; provides the index of the thread we are running on. Can be ENamedThreads::AnyThread if the current thread is unknown.
-	 *	@param bUnlock; if true, task can execute right now if possible
-	 * 
-	 *	Create the completed event
-	 *	Set the thread to execute on based on the embedded task
-	 *	Attempt to add myself as a subsequent to each prerequisite
-	 *	Tell the base task that I am ready to start as soon as my prerequisites are ready.
-	 **/
-	void SetupPrereqs(const FGraphEventArray* Prerequisites, ENamedThreads::Type CurrentThreadIfKnown, bool bUnlock)
-	{
-		checkThreadGraph(!TaskConstructed);
-		TaskConstructed = true;
-		TTask& Task = *(TTask*)&TaskStorage;
-		SetThreadToExecuteOn(Task.GetDesiredThread());
-		int32 AlreadyCompletedPrerequisites = 0;
-		if (Prerequisites)
-		{
-			for (int32 Index = 0; Index < Prerequisites->Num(); Index++)
-			{
-				FGraphEvent* Prerequisite = (*Prerequisites)[Index];
-				if (Prerequisite == nullptr || !Prerequisite->AddSubsequent(this))
-				{
-					AlreadyCompletedPrerequisites++;
-				}
-			}
-		}
-		PrerequisitesComplete(CurrentThreadIfKnown, AlreadyCompletedPrerequisites, bUnlock);
-	}
-
-	/** 
-	 *	Call from FConstructor to complete the setup process
-	 *	@param Prerequisites; the list of FGraphEvents that must be completed prior to this task executing.
-	 *	@param CurrentThreadIfKnown; provides the index of the thread we are running on. Can be ENamedThreads::AnyThread if the current thread is unknown.
-	 *	@return A new graph event which represents the completion of this task.
-	 * 
-	 *	Create the completed event
-	 *	Set the thread to execute on based on the embedded task
-	 *	Attempt to add myself as a subsequent to each prerequisite
-	 *	Tell the base task that I am ready to start as soon as my prerequisites are ready.
-	 **/
-	FGraphEventRef Setup(const FGraphEventArray* Prerequisites = NULL, ENamedThreads::Type CurrentThreadIfKnown = ENamedThreads::AnyThread)
-	{
-		TaskTrace::Launched(GetTraceId(), nullptr, Subsequents.IsValid(), ((TTask*)&TaskStorage)->GetDesiredThread(), sizeof(*this));
-
-		FGraphEventRef ReturnedEventRef = Subsequents; // very important so that this doesn't get destroyed before we return
-		SetupPrereqs(Prerequisites, CurrentThreadIfKnown, true);
-		return ReturnedEventRef;
-	}
-
-	/** 
-	 *	Call from FConstructor to complete the setup process, but doesn't allow the task to dispatch yet
-	 *	@param Prerequisites; the list of FGraphEvents that must be completed prior to this task executing.
-	 *	@param CurrentThreadIfKnown; provides the index of the thread we are running on. Can be ENamedThreads::AnyThread if the current thread is unknown.
-	 *	@return a pointer to the task
-	 * 
-	 *	Create the completed event
-	 *	Set the thread to execute on based on the embedded task
-	 *	Attempt to add myself as a subsequent to each prerequisite
-	 *	Tell the base task that I am ready to start as soon as my prerequisites are ready.
-	 **/
-	TGraphTask* Hold(const FGraphEventArray* Prerequisites = NULL, ENamedThreads::Type CurrentThreadIfKnown = ENamedThreads::AnyThread)
-	{
-		TaskTrace::Created(GetTraceId(), sizeof(*this));
-
-		SetupPrereqs(Prerequisites, CurrentThreadIfKnown, false);
-		return this;
-	}
-
-	/** 
-	 *	Factory to create a gather task which assumes the given subsequent list from some other tasks.
-	 *	This is used to support "WaitFor" during a task execution.
-	 *	@param SubsequentsToAssume; subsequents to "assume" from an existing task
-	 *	@param Prerequisites; the list of FGraphEvents that must be completed prior to dispatching my seubsequents.
-	 *	@param CurrentThreadIfKnown; provides the index of the thread we are running on. Can be ENamedThreads::AnyThread if the current thread is unknown.
-	**/
-	static FConstructor CreateTask(FGraphEventRef SubsequentsToAssume, const FGraphEventArray* Prerequisites = NULL, ENamedThreads::Type CurrentThreadIfKnown = ENamedThreads::AnyThread)
-	{
-		return FConstructor(new TGraphTask(SubsequentsToAssume, Prerequisites ? Prerequisites->Num() : 0), Prerequisites, CurrentThreadIfKnown);
-	}
-
-	/** An aligned bit of storage to hold the embedded task **/
-	TAlignedBytes<sizeof(TTask),alignof(TTask)> TaskStorage;
-	/** Used to sanity check the state of the object **/
-	bool						TaskConstructed;
-	/** A reference counted pointer to the completion event which lists the tasks that have me as a prerequisite. **/
-	FGraphEventRef				Subsequents;
-};
-
-#endif // TASKGRAPH_NEW_FRONTEND
-
 // Blocks the current thread until any of the given tasks is completed.
 // Is slightly more efficient than `AnyTaskCompleted()->Wait()` and supports timeout while `FGraphEvent::Wait()` doesn't.
 // Returns the index of the first completed task, or `INDEX_NONE` on timeout.
diff --git a/Engine/Source/Runtime/Core/Public/Tasks/TaskPrivate.h b/Engine/Source/Runtime/Core/Public/Tasks/TaskPrivate.h
--- a/Engine/Source/Runtime/Core/Public/Tasks/TaskPrivate.h
+++ b/Engine/Source/Runtime/Core/Public/Tasks/TaskPrivate.h
@@ -62,7 +62,7 @@ namespace UE::Tasks
 		Inline, // a task priority for "inline" task execution - a task is executed "inline" by the thread that unlocked it, w/o scheduling
 		TaskEvent, // a task priority used by task events, allows to shortcut task execution
 
-#if TASKGRAPH_NEW_FRONTEND
+
 		// for integration with named threads
 		GameThreadNormalPri,
 		GameThreadHiPri,
@@ -78,7 +78,6 @@ namespace UE::Tasks
 		RHIThreadHiPri,
 		RHIThreadNormalPriLocalQueue,
 		RHIThreadHiPriLocalQueue,
-#endif
 
 		Count
 	};
@@ -207,12 +206,10 @@ namespace UE::Tasks
 				return FPlatformTLS::GetCurrentThreadId() != ExecutingThreadId.load(std::memory_order_relaxed);
 			}
 
-#if TASKGRAPH_NEW_FRONTEND
 			bool IsNamedThreadTask() const
 			{
 				return ExtendedPriority >= EExtendedTaskPriority::GameThreadNormalPri;
 			}
-#endif
 
 			ETaskPriority GetPriority() const
 			{
diff --git a/Engine/Source/Runtime/Core/Tests/Tasks/TasksTest.cpp b/Engine/Source/Runtime/Core/Tests/Tasks/TasksTest.cpp
--- a/Engine/Source/Runtime/Core/Tests/Tasks/TasksTest.cpp
+++ b/Engine/Source/Runtime/Core/Tests/Tasks/TasksTest.cpp
@@ -364,7 +364,6 @@ namespace UE { namespace TasksTests
 			Prerequisites(TArrayView<FTask>{});
 		}
 
-#if TASKGRAPH_NEW_FRONTEND
 		{	// a basic test for a named thread task
 			FTask GTTask = Launch
 			(
@@ -375,7 +374,6 @@ namespace UE { namespace TasksTests
 			);
 			GTTask.Wait();
 		}
-#endif
 
 		//{	// Cancelled task is completed only after nested
 		//	FTaskEvent ParentBlocker{ UE_SOURCE_LOCATION };
diff --git a/Engine/Source/Runtime/Engine/Private/Components/SkeletalMeshComponent.cpp b/Engine/Source/Runtime/Engine/Private/Components/SkeletalMeshComponent.cpp
--- a/Engine/Source/Runtime/Engine/Private/Components/SkeletalMeshComponent.cpp
+++ b/Engine/Source/Runtime/Engine/Private/Components/SkeletalMeshComponent.cpp
@@ -4372,9 +4372,7 @@ bool USkeletalMeshComponent::HandleExistingParallelEvaluationTask(bool bBlockOnT
 {
 	// We are already processing eval and we are not inside it
 	if (IsRunningParallelEvaluation()
-#if TASKGRAPH_NEW_FRONTEND 
 		&& ParallelAnimationEvaluationTask->IsAwaitable()
-#endif
 	)
 	{
 		if (bBlockOnTask)
diff --git a/Engine/Source/Runtime/Engine/Private/Streaming/StreamingManagerTexture.cpp b/Engine/Source/Runtime/Engine/Private/Streaming/StreamingManagerTexture.cpp
--- a/Engine/Source/Runtime/Engine/Private/Streaming/StreamingManagerTexture.cpp
+++ b/Engine/Source/Runtime/Engine/Private/Streaming/StreamingManagerTexture.cpp
@@ -225,9 +225,7 @@ FRenderAssetStreamingManager::~FRenderAssetStreamingManager()
 TArray<FStreamingRenderAsset>& FRenderAssetStreamingManager::GetStreamingRenderAssetsAsyncSafe()
 {
 	if (StreamingRenderAssetsSyncEvent.IsValid()
-#if TASKGRAPH_NEW_FRONTEND
 		&& StreamingRenderAssetsSyncEvent->IsAwaitable()
-#endif
 	)
 	{
 		StreamingRenderAssetsSyncEvent->Wait(ENamedThreads::GameThread);
diff --git a/Engine/Source/Runtime/Engine/Private/TickTaskManager.cpp b/Engine/Source/Runtime/Engine/Private/TickTaskManager.cpp
--- a/Engine/Source/Runtime/Engine/Private/TickTaskManager.cpp
+++ b/Engine/Source/Runtime/Engine/Private/TickTaskManager.cpp
@@ -311,12 +311,7 @@ public:
 };
 
 // Declare the internal async task that is used for scheduling tick tasks with the task graph backend
-// TODO remove this when old frontend is removed. Batching does not work with the old frontend
-#if TASKGRAPH_NEW_FRONTEND
 using FTickGraphTask = FBaseGraphTask;
-#else
-using FTickGraphTask = TGraphTask<FTickFunctionTask>;
-#endif
 
 /** Info used to execute a batch tick */
 struct FTickBatchInfo
@@ -657,8 +652,6 @@ public:
 	/** This will add to an existing batch, create a new batch, or just spawn a single task and return null */
 	FTickBatchInfo* QueueOrBatchTickTask(TArray<FTickFunction*>& Prerequisites, FTickFunction* TickFunction, const FTickContext& TickContext)
 	{
-		// Batching is not supported with the old frontend
-#if TASKGRAPH_NEW_FRONTEND
 		if (bAllowBatchedTicksForFrame)
 		{
 			FTickGroupCondition Condition = FTickGroupCondition(TickFunction);
@@ -728,7 +721,6 @@ public:
 				return BatchInfo;
 			}
 		}
-#endif //TASKGRAPH_NEW_FRONTEND
 
 		// No batching, create a single task
 		// FGraphEventArray array has some inline members so it is faster to not explicitly reserve space
