SimpleWaveReader to read .wav files from files written with the SimpleWaveWriter

diff --git a/Engine/Source/Runtime/Core/Private/Audio/LogAudioCore.cpp b/Engine/Source/Runtime/Core/Private/Audio/LogAudioCore.cpp
new file mode 100644
--- /dev/null
+++ b/Engine/Source/Runtime/Core/Private/Audio/LogAudioCore.cpp
@@ -0,0 +1,5 @@
+﻿// Copyright Epic Games, Inc. All Rights Reserved.
+
+#include "Audio/LogAudioCore.h"
+
+DEFINE_LOG_CATEGORY(LogAudioCore)
\ No newline at end of file
diff --git a/Engine/Source/Runtime/Core/Private/Audio/LogAudioCore.h b/Engine/Source/Runtime/Core/Private/Audio/LogAudioCore.h
new file mode 100644
--- /dev/null
+++ b/Engine/Source/Runtime/Core/Private/Audio/LogAudioCore.h
@@ -0,0 +1,6 @@
+﻿// Copyright Epic Games, Inc. All Rights Reserved.
+#pragma once
+
+#include "Logging/LogMacros.h"
+
+DECLARE_LOG_CATEGORY_EXTERN(LogAudioCore, Log, Log);
\ No newline at end of file
diff --git a/Engine/Source/Runtime/Core/Private/Audio/SimpleWaveFormat.h b/Engine/Source/Runtime/Core/Private/Audio/SimpleWaveFormat.h
new file mode 100644
--- /dev/null
+++ b/Engine/Source/Runtime/Core/Private/Audio/SimpleWaveFormat.h
@@ -0,0 +1,21 @@
+﻿// Copyright Epic Games, Inc. All Rights Reserved.
+#pragma once
+
+#include "HAL/Platform.h"
+
+namespace Audio
+{
+	// Local definition so we don't depend on platform includes.
+	enum EFormatType { IEEE_FLOAT = 0x3 }; // WAVE_FORMAT_IEEE_FLOAT
+	struct FWaveFormatEx
+	{
+		uint16	FormatTag;
+		uint16	NumChannels;
+		uint32	NumSamplesPerSec;
+		uint32	AverageBytesPerSec;
+		uint16	BlockAlign;
+		uint16	NumBitsPerSample;
+		uint16	Size;
+	};
+
+}
\ No newline at end of file
diff --git a/Engine/Source/Runtime/Core/Private/Audio/SimpleWaveReader.cpp b/Engine/Source/Runtime/Core/Private/Audio/SimpleWaveReader.cpp
new file mode 100644
--- /dev/null
+++ b/Engine/Source/Runtime/Core/Private/Audio/SimpleWaveReader.cpp
@@ -0,0 +1,171 @@
+﻿// Copyright Epic Games, Inc. All Rights Reserved.
+#include "Audio/SimpleWaveReader.h"
+#include "Audio/SimpleWaveFormat.h"
+#include "LogAudioCore.h"
+#include "HAL/Platform.h"
+
+namespace Audio
+{
+	FSimpleWaveReader::FSimpleWaveReader(TUniquePtr<FArchive>&& InInputStream)
+		: InputStream{ MoveTemp(InInputStream) }
+	{
+		bIsDataValid = ReadHeader(SampleRate, NumChannels, DataSize);
+	}
+
+	FSimpleWaveReader::~FSimpleWaveReader()
+	{
+	}
+
+	void FSimpleWaveReader::SeekToFrame(uint32 FrameIndex) const
+	{
+		check(bIsDataValid);
+		int64 SeekPos = DataStartPos + FrameIndex * NumChannels * sizeof(float);
+		InputStream->Seek(SeekPos);
+	}
+	
+	bool FSimpleWaveReader::Read(TArrayView<const float> OutBuffer, int64& OutNumSamplesRead) const
+	{
+		check(bIsDataValid);
+		
+		// buffer is required to be a multiple of the number of channels to read this file
+		check(OutBuffer.Num() % NumChannels == 0);
+
+		const int64 BufferDataSize = OutBuffer.NumBytes();
+		
+		FMemory::Memzero((void*)OutBuffer.GetData(), BufferDataSize);
+
+		if (InputStream->AtEnd())
+		{
+			OutNumSamplesRead = 0;
+			return true;
+		}
+	
+		const int64 DataRemainingSize = DataSize - (InputStream->Tell() - DataStartPos);
+		const int64 ReadSize = FMath::Min(BufferDataSize, DataRemainingSize);
+		
+		InputStream->Serialize((void*)OutBuffer.GetData(), ReadSize);
+
+		OutNumSamplesRead = ReadSize / sizeof(float);
+
+		return InputStream->AtEnd();
+	}
+
+	bool FSimpleWaveReader::IsDataValid() const
+	{
+		return bIsDataValid;
+	}
+
+	uint32 FSimpleWaveReader::GetSampleRate() const
+	{
+		check(bIsDataValid);
+		return SampleRate;
+	}
+
+	uint16 FSimpleWaveReader::GetNumChannels() const
+	{
+		check(bIsDataValid);
+		return NumChannels;
+	}
+
+	uint32 FSimpleWaveReader::GetNumSamples() const
+	{
+		check(bIsDataValid);
+		return DataSize / sizeof(float);
+	}
+
+	bool FSimpleWaveReader::ReadHeader(uint32& OutSampleRate, uint16& OutNumChannels, uint32& OutDataSize)
+	{
+		InputStream->Seek(0);
+		
+		int32 ID;
+		*InputStream << ID;
+		if (ID != 'FFIR')
+		{
+			UE_LOG(LogAudioCore, Error, TEXT("[%hs]: Unexpected Format: Expected 'RIFF'"), __func__);
+			return false;
+		}
+		
+		int32 RiffChunkSize = 0;
+		*InputStream << RiffChunkSize;
+		
+		*InputStream << ID;
+		
+		if (ID != 'EVAW')
+		{
+			UE_LOG(LogAudioCore, Error, TEXT("[%hs]: Unexpected Format - Expected 'WAVE', was '%hs'"), __func__, (ANSICHAR*)&ID);
+			return false;
+		}
+		
+		*InputStream << ID;
+		if (ID != ' tmf')
+		{
+			UE_LOG(LogAudioCore, Error, TEXT("[%hs]: Unexpected chunk - Expected 'fmt ', was '%hs'"), __func__, (ANSICHAR*)&ID);
+			return false;
+		}
+		
+		FWaveFormatEx Fmt = { 0 };
+		int32 FmtSize = sizeof(Fmt);
+		*InputStream << FmtSize;
+		InputStream->Serialize((void*)&Fmt, FmtSize);
+
+		uint16 NumBitsPerSample = sizeof(float) * 8;
+		if (Fmt.NumBitsPerSample != NumBitsPerSample)
+		{
+			UE_LOG(LogAudioCore, Error, TEXT("[%hs]: Expected NumBitsPerSample to be %d was %d"), __func__, NumBitsPerSample, Fmt.NumBitsPerSample);
+			return false;
+		}
+
+		uint16 BlockAlign = (uint16)(Fmt.NumBitsPerSample * Fmt.NumChannels) / 8;
+		if (Fmt.BlockAlign != BlockAlign)
+		{
+			UE_LOG(LogAudioCore, Error, TEXT("[%hs]: Expected BlockAlign to be %d was %d"), __func__, BlockAlign, Fmt.BlockAlign);
+			return false;
+		}
+
+		uint32 AverageBytesPerSec = Fmt.BlockAlign * Fmt.NumSamplesPerSec;
+		if (Fmt.AverageBytesPerSec != AverageBytesPerSec)
+		{
+			UE_LOG(LogAudioCore, Error, TEXT("[%hs]: Expected AverageBytesPerSec to be %d was %d"), __func__, AverageBytesPerSec, Fmt.AverageBytesPerSec);
+			return false;
+		}
+		
+		if (Fmt.FormatTag != EFormatType::IEEE_FLOAT)
+		{
+			UE_LOG(LogAudioCore, Error, TEXT("[%hs]: Expected FormatTag to be %d was %d"), __func__, EFormatType::IEEE_FLOAT, Fmt.FormatTag);
+			return false;
+		}
+		
+		*InputStream << ID;
+		if (ID != 'atad')
+		{
+			UE_LOG(LogAudioCore, Error, TEXT("[%hs]: Unexpected chunk - expected 'data', was '%hs'"), __func__, (ANSICHAR*)&ID);
+			return false;
+		}
+		
+		*InputStream << OutDataSize;
+
+		if (OutDataSize == 0)
+		{
+			UE_LOG(LogAudioCore, Error, TEXT("[%hs]: File contains no wav data"), __func__);
+			return false;
+		}
+		
+		if (Fmt.NumBitsPerSample == 0)
+		{
+			UE_LOG(LogAudioCore, Error, TEXT("[%hs]: Invalid Sample Rate %d"), __func__, Fmt.NumBitsPerSample);
+			return false;
+		}
+		
+		if (Fmt.NumChannels == 0)
+		{
+			UE_LOG(LogAudioCore, Error, TEXT("[%hs]: Invalid Num Channels %d"), __func__, Fmt.NumChannels);
+			return false;
+		}
+
+		OutSampleRate = Fmt.NumSamplesPerSec;
+		OutNumChannels = Fmt.NumChannels;
+		
+		DataStartPos = InputStream->Tell();
+		return true;
+	}
+}
\ No newline at end of file
diff --git a/Engine/Source/Runtime/Core/Private/Audio/SimpleWaveWriter.cpp b/Engine/Source/Runtime/Core/Private/Audio/SimpleWaveWriter.cpp
--- a/Engine/Source/Runtime/Core/Private/Audio/SimpleWaveWriter.cpp
+++ b/Engine/Source/Runtime/Core/Private/Audio/SimpleWaveWriter.cpp
@@ -1,5 +1,6 @@
 // Copyright Epic Games, Inc. All Rights Reserved.
 #include "Audio/SimpleWaveWriter.h"
+#include "Audio/SimpleWaveFormat.h"
 
 namespace Audio
 {
diff --git a/Engine/Source/Runtime/Core/Public/Audio/SimpleWaveReader.h b/Engine/Source/Runtime/Core/Public/Audio/SimpleWaveReader.h
new file mode 100644
--- /dev/null
+++ b/Engine/Source/Runtime/Core/Public/Audio/SimpleWaveReader.h
@@ -0,0 +1,49 @@
+﻿// Copyright Epic Games, Inc. All Rights Reserved.
+#pragma once
+
+#include "Containers/ArrayView.h"
+#include "HAL/Platform.h"
+#include "Templates/SharedPointer.h"
+
+class FArchive;
+
+namespace Audio
+{
+	/**
+	 * SimpleWaveReader class
+	 *
+	 * Only able to read wave files of a specific format
+	 * Specifically, it should be able to read any wave file written via the SimpleWaveWriter class
+	 */
+	class FSimpleWaveReader
+	{
+	public:
+		CORE_API FSimpleWaveReader(TUniquePtr<FArchive>&& InOutputStream);
+		CORE_API ~FSimpleWaveReader();
+
+		CORE_API void SeekToFrame(uint32 FrameIndex) const;
+
+		// returns true if the end of the file was reached
+		CORE_API bool Read(TArrayView<const float> OutBuffer, int64& OutNumSamplesRead) const;
+
+		CORE_API bool IsDataValid() const;
+
+		CORE_API uint32 GetSampleRate() const;
+
+		CORE_API uint16 GetNumChannels() const;
+
+		CORE_API uint32 GetNumSamples() const;
+		
+	private:
+	
+		bool ReadHeader(uint32& OutSampleRate, uint16& OutNumChannels, uint32& OutDataSize);
+		
+		TUniquePtr<FArchive> InputStream;
+		int64 DataStartPos = 0;
+
+		uint32 SampleRate = 0;
+		uint16 NumChannels = 0;
+		uint32 DataSize = 0;
+		bool bIsDataValid = false;
+	};
+}
\ No newline at end of file
diff --git a/Engine/Source/Runtime/Core/Public/Audio/SimpleWaveWriter.h b/Engine/Source/Runtime/Core/Public/Audio/SimpleWaveWriter.h
--- a/Engine/Source/Runtime/Core/Public/Audio/SimpleWaveWriter.h
+++ b/Engine/Source/Runtime/Core/Public/Audio/SimpleWaveWriter.h
@@ -6,19 +6,6 @@ namespace Audio
 {
 	class FSimpleWaveWriter
 	{
-		// Local definition so we don't depend on platform includes.
-		enum EFormatType { IEEE_FLOAT = 0x3 }; // WAVE_FORMAT_IEEE_FLOAT
-		struct FWaveFormatEx
-		{
-			uint16	FormatTag;
-			uint16	NumChannels;
-			uint32	NumSamplesPerSec;
-			uint32	AverageBytesPerSec;
-			uint16	BlockAlign;
-			uint16	NumBitsPerSample;
-			uint16	Size;
-		};
-
 	public:
 		CORE_API FSimpleWaveWriter(TUniquePtr<FArchive>&& InOutputStream, int32 InSampleRate, int32 InNumChannels, bool bInUpdateHeaderAfterEveryWrite);
 		CORE_API ~FSimpleWaveWriter();
