Fix static analysis error

check SampleRate instead of BitsPerSample

#rnx

diff --git a/Engine/Source/Runtime/Core/Private/Audio/SimpleWaveReader.cpp b/Engine/Source/Runtime/Core/Private/Audio/SimpleWaveReader.cpp
--- a/Engine/Source/Runtime/Core/Private/Audio/SimpleWaveReader.cpp
+++ b/Engine/Source/Runtime/Core/Private/Audio/SimpleWaveReader.cpp
@@ -150,9 +150,9 @@ namespace Audio
 			return false;
 		}
 		
-		if (Fmt.NumBitsPerSample == 0)
+		if (Fmt.NumSamplesPerSec == 0)
 		{
-			UE_LOG(LogAudioCore, Error, TEXT("[%hs]: Invalid Sample Rate %d"), __func__, Fmt.NumBitsPerSample);
+			UE_LOG(LogAudioCore, Error, TEXT("[%hs]: Invalid Sample Rate %d"), __func__, Fmt.NumSamplesPerSec);
 			return false;
 		}
 		
