[Insights] Added tracing for all LLM tag sets (assets and asset classes).

diff --git a/Engine/Source/Runtime/Core/Private/HAL/LowLevelMemTracker.cpp b/Engine/Source/Runtime/Core/Private/HAL/LowLevelMemTracker.cpp
--- a/Engine/Source/Runtime/Core/Private/HAL/LowLevelMemTracker.cpp
+++ b/Engine/Source/Runtime/Core/Private/HAL/LowLevelMemTracker.cpp
@@ -924,6 +924,7 @@ UE_TRACE_CHANNEL(MemTagChannel, "Memory overview", true)
 UE_TRACE_EVENT_BEGIN(LLM, TagsSpec, NoSync|Important)
 	UE_TRACE_EVENT_FIELD(const void*, TagId)
 	UE_TRACE_EVENT_FIELD(const void*, ParentId)
+	UE_TRACE_EVENT_FIELD(uint8, TagSetId)
 	UE_TRACE_EVENT_FIELD(UE::Trace::AnsiString, Name)
 UE_TRACE_EVENT_END()
 
@@ -932,6 +933,11 @@ UE_TRACE_EVENT_BEGIN(LLM, TrackerSpec, NoSync|Important)
 	UE_TRACE_EVENT_FIELD(UE::Trace::AnsiString, Name)
 UE_TRACE_EVENT_END()
 
+UE_TRACE_EVENT_BEGIN(LLM, TagSetSpec, NoSync|Important)
+	UE_TRACE_EVENT_FIELD(uint8, TagSetId)
+	UE_TRACE_EVENT_FIELD(UE::Trace::AnsiString, Name)
+UE_TRACE_EVENT_END()
+
 UE_TRACE_EVENT_BEGIN(LLM, TagValue)
 	UE_TRACE_EVENT_FIELD(uint8, TrackerId)
 	UE_TRACE_EVENT_FIELD(uint64, Cycle)
@@ -1322,9 +1328,10 @@ private:
 	static const void* GetTagId(const FTagData* TagData);
 	void SendTagDeclaration(const FTagData* TagData);
 
-	ELLMTracker				Tracker;
+	ELLMTracker Tracker;
 	TFastPointerLLMSet<const FTagData*> DeclaredTags;
-	bool					bTrackerSpecSent = false;
+	bool bTrackerSpecSent = false;
+	bool bTagSetSpecSent = false;
 };
 
 /** FLLMCsvWriter: class for writing out LLM stats to the Csv Profiler. */
@@ -6229,18 +6236,31 @@ void FLLMTraceWriter::Publish(FLowLevelMemTracker& LLMRef, const FTrackerTagSize
 	if (!bTrackerSpecSent)
 	{
 		bTrackerSpecSent = true;
-		static struct {
-			const ANSICHAR* Name;
-			uint32 Len;
-		} TrackerNames[] = {
-			{ "Platform", 8 },
-			{ "Default", 7 },
-		};
-		static_assert(UE_ARRAY_COUNT(TrackerNames) == int(ELLMTracker::Max), "");
-		uint32 NameLen = TrackerNames[(uint8)Tracker].Len;
-		UE_TRACE_LOG(LLM, TrackerSpec, MemTagChannel, NameLen * sizeof(ANSICHAR))
+		FAnsiStringView TrackerNames[] = { "Platform", "Default" };
+		static_assert(UE_ARRAY_COUNT(TrackerNames) == (int)ELLMTracker::Max, "");
+		FAnsiStringView TrackerName = TrackerNames[(uint8)Tracker];
+		UE_TRACE_LOG(LLM, TrackerSpec, MemTagChannel, TrackerName.Len() * sizeof(ANSICHAR))
 			<< TrackerSpec.TrackerId((uint8)Tracker)
-			<< TrackerSpec.Name(TrackerNames[(uint8)Tracker].Name, NameLen);
+			<< TrackerSpec.Name(TrackerName.GetData(), TrackerName.Len());
+	}
+
+	if (!bTagSetSpecSent && Tracker == ELLMTracker::Default)
+	{
+		bTagSetSpecSent = true;
+		FAnsiStringView TagSetNames[] = { "System", "Assets", "AssetClasses" };
+		static_assert((int)ELLMTagSet::Assets == 1, "");
+		static_assert((int)ELLMTagSet::AssetClasses == 2, "");
+		static_assert(UE_ARRAY_COUNT(TagSetNames) == (int)ELLMTagSet::Max, "");
+		for (uint8 TagSet = 0; TagSet < (uint8)ELLMTagSet::Max; ++TagSet)
+		{
+			if (LLMRef.IsTagSetActive((ELLMTagSet)TagSet))
+			{
+				FAnsiStringView TagSetName = TagSetNames[TagSet];
+				UE_TRACE_LOG(LLM, TagSetSpec, MemTagChannel, TagSetName.Len() * sizeof(ANSICHAR))
+					<< TagSetSpec.TagSetId(TagSet)
+					<< TagSetSpec.Name(TagSetName.GetData(), TagSetName.Len());
+			}
+		}
 	}
 
 	SendTagDeclaration(OverrideTrackedTotalTagData);
@@ -6248,11 +6268,11 @@ void FLLMTraceWriter::Publish(FLowLevelMemTracker& LLMRef, const FTrackerTagSize
 	for (const TPair<const FTagData*, FTrackerTagSizeData>& It : TagSizes)
 	{
 		const FTagData* TagData = It.Key;
-		if (!TagData->IsReportable())
+		if (!TagData->IsStatsReportable() || !LLMRef.IsTagSetActive(TagData->GetTagSet()))
 		{
 			continue;
 		}
-		if (OverrideUntaggedTagData != nullptr && TagData->GetName() == TagName_Untagged)
+		if (OverrideUntaggedTagData != nullptr && TagData->GetTagSet() == ELLMTagSet::None && TagData->GetName() == TagName_Untagged)
 		{
 			continue; // Handled by OverrideUntaggedTagData.
 		}
@@ -6291,11 +6311,11 @@ void FLLMTraceWriter::Publish(FLowLevelMemTracker& LLMRef, const FTrackerTagSize
 	for (const TPair<const FTagData*, FTrackerTagSizeData>& It : TagSizes)
 	{
 		const FTagData* TagData = It.Key;
-		if (!TagData->IsReportable())
+		if (!TagData->IsStatsReportable() || !LLMRef.IsTagSetActive(TagData->GetTagSet()))
 		{
 			continue;
 		}
-		if (OverrideUntaggedTagData && TagData->GetName() == TagName_Untagged)
+		if (OverrideUntaggedTagData && TagData->GetTagSet() == ELLMTagSet::None && TagData->GetName() == TagName_Untagged)
 		{
 			continue; // Handled by OverrideUntaggedTagData.
 		}
@@ -6333,6 +6353,7 @@ void FLLMTraceWriter::SendTagDeclaration(const FTagData* TagData)
 	UE_TRACE_LOG(LLM, TagsSpec, MemTagChannel, NameBuffer.Len() * sizeof(ANSICHAR))
 		<< TagsSpec.TagId(GetTagId(TagData))
 		<< TagsSpec.ParentId(GetTagId(Parent))
+		<< TagsSpec.TagSetId(static_cast<uint8>(TagData->GetTagSet()))
 		<< TagsSpec.Name(*NameBuffer, NameBuffer.Len());
 };
 
