Fix Android file GetSize after Truncate

diff --git a/Engine/Source/Runtime/Core/Private/Android/AndroidPlatformFile.cpp b/Engine/Source/Runtime/Core/Private/Android/AndroidPlatformFile.cpp
--- a/Engine/Source/Runtime/Core/Private/Android/AndroidPlatformFile.cpp
+++ b/Engine/Source/Runtime/Core/Private/Android/AndroidPlatformFile.cpp
@@ -451,6 +451,12 @@ public:
 
 		int Result = 0;
 		do { Result = __ftruncate(File->Handle, NewSize); } while (Result < 0 && errno == EINTR);
+		if (Result == 0)
+		{
+			const int CurrentPos = __lseek(File->Handle, 0, SEEK_CUR); 
+			Length = __lseek(File->Handle, 0, SEEK_END);
+			__lseek(File->Handle, CurrentPos, SEEK_SET);
+		}
 		return Result == 0;
 	}
 
