Moved GConfig init log at the end of the function. This Log is likely the first one happening in the engine and will trigger the creation of the log file. It must be done after GConfig has been allocated to ensure IFileHandle log wont be a FManagedStorageFileWriteHandle (which might cause dead lock on some platform when reporting a crash)

#rnx

diff --git a/Engine/Source/Runtime/Core/Private/Misc/ConfigCacheIni.cpp b/Engine/Source/Runtime/Core/Private/Misc/ConfigCacheIni.cpp
--- a/Engine/Source/Runtime/Core/Private/Misc/ConfigCacheIni.cpp
+++ b/Engine/Source/Runtime/Core/Private/Misc/ConfigCacheIni.cpp
@@ -5620,8 +5620,6 @@ bool FConfigCacheIni::CreateGConfigFromSaved(const TCHAR* Filename)
 		return false;
 	}
 
-	UE_LOG(LogInit, Display, TEXT("Loading binary GConfig from %lld bytes of data..."), Size);
-
 	// serialize right out of the preloaded data
 	FLargeMemoryReader MemoryReader((uint8*)PreloadedData, Size);
 	MemoryReader.SetIsPersistent(true);
@@ -5656,6 +5654,11 @@ bool FConfigCacheIni::CreateGConfigFromSaved(const TCHAR* Filename)
 		FCoreDelegates::TSConfigReadyForUse().Broadcast();
 	}
 
+	// This Log is likely the first one happening in the engine and will trigger the creation of the log file
+	// It must be done after GConfig has been allocated and ready to use to ensure IFileHandle log wont be a 
+	// FManagedStorageFileWriteHandle (which might cause dead lock on some platforms when reporting a crash)
+	UE_LOG(LogInit, Display, TEXT("Loaded binary GConfig from %lld bytes of data..."), Size);
+
 	FMemory::Free(PreloadedData);
 	return true;
 }
