Add missing TRACE_COUNTER_GET define when counters tracing is disabled

diff --git a/Engine/Source/Runtime/Core/Public/ProfilingDebugging/CountersTrace.h b/Engine/Source/Runtime/Core/Public/ProfilingDebugging/CountersTrace.h
--- a/Engine/Source/Runtime/Core/Public/ProfilingDebugging/CountersTrace.h
+++ b/Engine/Source/Runtime/Core/Public/ProfilingDebugging/CountersTrace.h
@@ -471,6 +471,7 @@ struct FCountersTrace
 #define TRACE_DECLARE_UNCHECKED_ATOMIC_MEMORY_COUNTER_EXTERN(CounterName)
 
 #define TRACE_COUNTER_SET(CounterName, Value)
+#define TRACE_COUNTER_GET(CounterName)
 #define TRACE_COUNTER_ADD(CounterName, Value)
 #define TRACE_COUNTER_SUBTRACT(CounterName, Value)
 #define TRACE_COUNTER_INCREMENT(CounterName)
