Make GetHighPerformanceEnableForStat transactionally-safe

* use FTransactionallySafeCriticalSection SynchronizationObject
* run MemoryCounter.Add(MemoryAllocated) atomic op on transaction commit -- safe because StatGroupEnableManager is a singleton
* avoid using EnableStat() and instead use a new constructor that initializes TStatIdData::Name

diff --git a/Engine/Source/Runtime/Core/Private/Stats/StatsSystemTypes.cpp b/Engine/Source/Runtime/Core/Private/Stats/StatsSystemTypes.cpp
--- a/Engine/Source/Runtime/Core/Private/Stats/StatsSystemTypes.cpp
+++ b/Engine/Source/Runtime/Core/Private/Stats/StatsSystemTypes.cpp
@@ -316,7 +316,7 @@ class FStatGroupEnableManager : public IStatGroupEnableManager
 	TMap<FName, FGroupEnable> HighPerformanceEnable;
 
 	/** Used to synchronize the access to the high performance stats groups. */
-	FCriticalSection SynchronizationObject;
+	FTransactionallySafeCriticalSection SynchronizationObject;
 
 	/** Pointer to the long name in the names block. */
 	TArray<TArray<TStatIdData>> StatIDs;
@@ -356,7 +356,7 @@ public:
 
 	virtual void SetHighPerformanceEnableForGroup(FName Group, bool Enable) override
 	{
-		FScopeLock ScopeLock(&SynchronizationObject);
+		UE::TScopeLock ScopeLock(SynchronizationObject);
 		FThreadStats::PrimaryDisableChangeTagLockAdd();
 		FGroupEnable* Found = HighPerformanceEnable.Find(Group);
 		if (Found)
@@ -382,7 +382,7 @@ public:
 
 	virtual void SetHighPerformanceEnableForAllGroups(bool Enable) override
 	{
-		FScopeLock ScopeLock(&SynchronizationObject);
+		UE::TScopeLock ScopeLock(SynchronizationObject);
 		FThreadStats::PrimaryDisableChangeTagLockAdd();
 		for (auto It = HighPerformanceEnable.CreateIterator(); It; ++It)
 		{
@@ -406,7 +406,7 @@ public:
 	}
 	virtual void ResetHighPerformanceEnableForAllGroups() override
 	{
-		FScopeLock ScopeLock(&SynchronizationObject);
+		UE::TScopeLock ScopeLock(SynchronizationObject);
 		FThreadStats::PrimaryDisableChangeTagLockAdd();
 		for (auto It = HighPerformanceEnable.CreateIterator(); It; ++It)
 		{
@@ -437,7 +437,7 @@ public:
 		TStatIdData* Result;
 		FString StatDescription;
 		{
-			FScopeLock ScopeLock(&SynchronizationObject);
+			UE::TScopeLock ScopeLock(SynchronizationObject);
 
 			FStatNameAndInfo LongName(StatShortName, InGroup, InCategory, InDescription, InStatType, bShouldClearEveryFrame, bCycleStat, bSortByName, MemoryRegion);
 
@@ -486,7 +486,15 @@ public:
 				TArray<TStatIdData>& NewBlock = StatIDs.AddDefaulted_GetRef();
 				NewBlock.Reserve(NUM_PER_BLOCK);
 			}
-			Result = &StatIDs.Last().AddDefaulted_GetRef();
+
+			if (Found->CurrentEnable)
+			{
+				Result = &StatIDs.Last().Emplace_GetRef(NameToMinimalName(Stat));
+			}
+			else
+			{
+				Result = &StatIDs.Last().AddDefaulted_GetRef();
+			}
 
 			StatDescription = InDescription ? InDescription : StatShortName.GetPlainNameString();
 
@@ -516,12 +524,11 @@ public:
 				MemoryAllocated += StatDescriptionAnsiLength * sizeof(ANSICHAR);
 			}
 
-			MemoryCounter.Add(MemoryAllocated);
-
-			if (Found->CurrentEnable)
+			UE_AUTORTFM_ONCOMMIT(this, MemoryAllocated)
 			{
-				EnableStat(Stat, Result);
-			}
+				// Safe to defer this because FStatGroupEnableManager is a singleton, see IStatGroupEnableManager::Get()
+				MemoryCounter.Add(MemoryAllocated);
+			};
 
 			if (bShouldClearEveryFrame)
 			{
@@ -602,7 +609,7 @@ public:
 
 	void StatGroupEnableManagerCommand(FString const& InCmd) override
 	{
-		FScopeLock ScopeLock(&SynchronizationObject);
+		UE::TScopeLock ScopeLock(SynchronizationObject);
 		const TCHAR* Cmd = *InCmd;
 		if (FParse::Command(&Cmd, TEXT("list")))
 		{
diff --git a/Engine/Source/Runtime/Core/Public/Stats/StatsSystemTypes.h b/Engine/Source/Runtime/Core/Public/Stats/StatsSystemTypes.h
--- a/Engine/Source/Runtime/Core/Public/Stats/StatsSystemTypes.h
+++ b/Engine/Source/Runtime/Core/Public/Stats/StatsSystemTypes.h
@@ -57,6 +57,12 @@ struct TStatIdData
 		return LocalName.IsNone();
 	}
 
+	TStatIdData() = default;
+
+	TStatIdData(FMinimalName InName) : Name(InName)
+	{
+	}
+
 	/** Name of the active stat; stored as a minimal name to minimize the data size */
 	TAtomic<FMinimalName> Name;
 
