[Scheduler]
- Fix potential deadlock when stopping workers if remaining threads attempted thread creation either by launching a task or entering oversubscription.
- Add a test case with the repro

[RN] Fix potential deadlock in scheduler when stopping workers if remaining threads attempted thread creation either by launching a task or entering oversubscription.

diff --git a/Engine/Source/Runtime/Core/Private/Async/Fundamental/Scheduler.cpp b/Engine/Source/Runtime/Core/Private/Async/Fundamental/Scheduler.cpp
--- a/Engine/Source/Runtime/Core/Private/Async/Fundamental/Scheduler.cpp
+++ b/Engine/Source/Runtime/Core/Private/Async/Fundamental/Scheduler.cpp
@@ -129,11 +129,12 @@ namespace LowLevelTasks
 
 			WorkerEvents.SetNum(MaxWorkers);
 			WorkerLocalQueues.Reserve(MaxWorkers);
-			WorkerThreads.Reserve(MaxWorkers);
+			WorkerThreads.SetNum(MaxWorkers);
 
 			auto CreateThread =
 				[this, IsForkable](Private::ELocalQueueType LocalQueueType, const TCHAR* ThreadGroup, const TCHAR* Prefix, std::atomic<int32>& CreationIndex, int32 NumWorkers, int32 NumMaxWorkers, EThreadPriority Priority, uint64 Affinity)
 				{
+					FScopeLock Lock(&WorkerCreationCS);
 					LLM_SCOPE_BYNAME(TEXT("EngineMisc/WorkerThreads"));
 
 					// Thread creation can end up waiting, we don't want to recursively oversubscribe if that happens.
@@ -156,7 +157,7 @@ namespace LowLevelTasks
 					UE::Trace::ThreadGroupBegin(ThreadGroup);
 					WorkerLocalQueues.Emplace(QueueRegistry, LocalQueueType);
 					WorkerEvents[WorkerId].bIsStandby = bIsStandbyWorker;
-					WorkerThreads.Add(
+					WorkerThreads[WorkerId] =
 						CreateWorker(
 							WorkerId,
 							*WorkerName,
@@ -165,22 +166,20 @@ namespace LowLevelTasks
 							&WorkerEvents[WorkerId],
 							&WorkerLocalQueues[WorkerId],
 							Priority,
-							Affinity)
-					);
+							Affinity).Release();
+
 					UE::Trace::ThreadGroupEnd();
 				};
 
 			TFunction<void()> ForegroundCreateThread =
 				[this, CreateThread, NumWorkers = NumForegroundWorkers, MaxWorkers = MaxForegroundWorkers]()
 				{
-					FScopeLock Lock(&WorkerThreadsCS);
 					CreateThread(Private::ELocalQueueType::EForeground, TEXT("Foreground Workers"), TEXT("Foreground"), ForegroundCreationIndex, NumWorkers, MaxWorkers, WorkerPriority, WorkerAffinity);
 				};
 
 			TFunction<void()> BackgroundCreateThread =
 				[this, CreateThread, NumWorkers = NumBackgroundWorkers, MaxWorkers = MaxBackgroundWorkers, ActualBackgroundPriority]()
 				{
-					FScopeLock Lock(&WorkerThreadsCS);
 					CreateThread(Private::ELocalQueueType::EBackground, TEXT("Background Workers"), TEXT("Background"), BackgroundCreationIndex, NumWorkers, MaxWorkers, ActualBackgroundPriority, BackgroundAffinity);
 				};
 
@@ -288,11 +287,14 @@ namespace LowLevelTasks
 			WaitingQueue[0].StartShutdown();
 			WaitingQueue[1].StartShutdown();
 
-			for (TUniquePtr<FThread>& Thread : WorkerThreads)
+			// We wait on threads to exit, once we're done with that
+			// it means no more threads can possibly get created.
+			for (std::atomic<FThread*>& ThreadEntry : WorkerThreads)
 			{
-				if (Thread.IsValid())
+				if (FThread* Thread = ThreadEntry.exchange(nullptr))
 				{
 					Thread->Join();
+					delete Thread;
 				}
 			}
 
diff --git a/Engine/Source/Runtime/Core/Public/Async/Fundamental/Scheduler.h b/Engine/Source/Runtime/Core/Public/Async/Fundamental/Scheduler.h
--- a/Engine/Source/Runtime/Core/Public/Async/Fundamental/Scheduler.h
+++ b/Engine/Source/Runtime/Core/Public/Async/Fundamental/Scheduler.h
@@ -232,8 +232,9 @@ namespace LowLevelTasks
 	private:
 		Private::FWaitingQueue                         WaitingQueue[2] = { { WorkerEvents, OversubscriptionLimitReachedEvent }, { WorkerEvents, OversubscriptionLimitReachedEvent } };
 		FSchedulerTls::FQueueRegistry                  QueueRegistry;
+		FCriticalSection                               WorkerCreationCS;
 		FCriticalSection                               WorkerThreadsCS;
-		TArray<TUniquePtr<FThread>>                    WorkerThreads;
+		TArray<std::atomic<FThread*>>                  WorkerThreads;
 		TAlignedArray<FSchedulerTls::FLocalQueueType>  WorkerLocalQueues;
 		TAlignedArray<Private::FWaitEvent>             WorkerEvents;
 		std::atomic_uint                               ActiveWorkers { 0 };
diff --git a/Engine/Source/Runtime/Core/Tests/Tasks/TasksTest.cpp b/Engine/Source/Runtime/Core/Tests/Tasks/TasksTest.cpp
--- a/Engine/Source/Runtime/Core/Tests/Tasks/TasksTest.cpp
+++ b/Engine/Source/Runtime/Core/Tests/Tasks/TasksTest.cpp
@@ -1736,6 +1736,30 @@ namespace UE { namespace TasksTests
 
 		LowLevelTasks::FScheduler::Get().RestartWorkers();
 	}
+
+	// StressTest RestartWorkers while having task performing oversubscription
+	TEST_CASE_NAMED(FTaskGraphOversubscribeDuringRestartWorkers, "System::Core::Async::Tasks::OversubscribeDuringRestartWorkers", "[.][ApplicationContextMask][EngineFilter]")
+	{
+		for (int32 Index = 0; Index < 10; ++Index)
+		{
+			FTask Task = Launch(TEXT("Task"),
+				[&]()
+				{
+					// Just trigger oversubscription
+					LowLevelTasks::FOversubscriptionScope _;
+				});
+
+			// For the repro to work we need to trigger oversubscription right between the time
+			// we acquire the critical section and before the waiting queues are shut down.
+			// This reproed very reliably before the fix.
+			LowLevelTasks::FScheduler::Get().RestartWorkers(0, 4);
+
+			Task.Wait();
+		}
+
+		// Restore original worker count
+		LowLevelTasks::FScheduler::Get().RestartWorkers();
+	}
 }}
 
 #endif // WITH_TESTS
