Removed overvalidating checkSlow in double linked list destructor. It doesn't own the memory and having to manually reset prior to destruction is not necessary.

diff --git a/Engine/Source/Runtime/Core/Public/Containers/IntrusiveDoubleLinkedList.h b/Engine/Source/Runtime/Core/Public/Containers/IntrusiveDoubleLinkedList.h
--- a/Engine/Source/Runtime/Core/Public/Containers/IntrusiveDoubleLinkedList.h
+++ b/Engine/Source/Runtime/Core/Public/Containers/IntrusiveDoubleLinkedList.h
@@ -172,11 +172,6 @@ public:
 	FORCEINLINE TIntrusiveDoubleLinkedList(const TIntrusiveDoubleLinkedList&) = delete;
 	FORCEINLINE TIntrusiveDoubleLinkedList& operator=(const TIntrusiveDoubleLinkedList&) = delete;
 
-	FORCEINLINE ~TIntrusiveDoubleLinkedList()
-	{
-		checkSlow(IsEmpty());
-	}
-
 	/**
 	 * Fast empty that clears this list *without* changing the links in any elements.
 	 * 
