[AutoRTFM] Remove TransactionallySafeRWScopeLock header.

The UE::TScopeLock classes can serve this same purpose.

diff --git a/Engine/Plugins/Experimental/ChaosModularVehicle/Source/ChaosModularVehicleEngine/Private/ChaosModularVehicle/ModularVehicleSimulationCU.cpp b/Engine/Plugins/Experimental/ChaosModularVehicle/Source/ChaosModularVehicleEngine/Private/ChaosModularVehicle/ModularVehicleSimulationCU.cpp
--- a/Engine/Plugins/Experimental/ChaosModularVehicle/Source/ChaosModularVehicleEngine/Private/ChaosModularVehicle/ModularVehicleSimulationCU.cpp
+++ b/Engine/Plugins/Experimental/ChaosModularVehicle/Source/ChaosModularVehicleEngine/Private/ChaosModularVehicle/ModularVehicleSimulationCU.cpp
@@ -87,7 +87,7 @@ void FModularVehicleSimulation::SimulateModuleTree(UWorld* InWorld, float DeltaS
 		//		, *Proxy->GetParticle_Internal()->W().ToString()));
 		//}
 
-		FTransactionallySafeReadScopeLock InputConfigLock(InputConfigurationLock);
+		UE::TReadScopeLock InputConfigLock(InputConfigurationLock);
 
 		FModuleInputContainer Container = InputData.PhysicsInputs.NetworkInputs.VehicleInputs.Container;
 		FInputInterface InputInterface(InputNameMap, Container);
diff --git a/Engine/Source/Programs/AutoRTFMTests/Private/UECoreTests.cpp b/Engine/Source/Programs/AutoRTFMTests/Private/UECoreTests.cpp
--- a/Engine/Source/Programs/AutoRTFMTests/Private/UECoreTests.cpp
+++ b/Engine/Source/Programs/AutoRTFMTests/Private/UECoreTests.cpp
@@ -20,7 +20,6 @@
 #include "MyAutoRTFMTestObject.h"
 #include "Containers/Queue.h"
 #include "Misc/ScopeRWLock.h"
-#include "Misc/TransactionallySafeRWScopeLock.h"
 #include "Misc/TransactionallySafeMutex.h"
 #include "Misc/ConfigCacheIni.h"
 #include "Misc/CoreDelegates.h"
@@ -1376,7 +1375,7 @@ TEST_CASE("UECore.TransactionallySafeReadScopeLock")
 
 		AutoRTFM::ETransactionResult Result = AutoRTFM::Transact([&]
 			{
-				FTransactionallySafeReadScopeLock Lock(RWLock);
+				UE::TReadScopeLock Lock(RWLock);
 				AutoRTFM::AbortTransaction();
 			});
 
@@ -1384,7 +1383,7 @@ TEST_CASE("UECore.TransactionallySafeReadScopeLock")
 
 		Result = AutoRTFM::Transact([&]
 			{
-				FTransactionallySafeReadScopeLock Lock(RWLock);
+				UE::TReadScopeLock Lock(RWLock);
 			});
 
 		REQUIRE(AutoRTFM::ETransactionResult::Committed == Result);
@@ -1395,7 +1394,7 @@ TEST_CASE("UECore.TransactionallySafeReadScopeLock")
 		AutoRTFM::ETransactionResult Result = AutoRTFM::Transact([&]
 			{
 				FTransactionallySafeRWLock RWLock;
-				FTransactionallySafeReadScopeLock Lock(RWLock);
+				UE::TReadScopeLock Lock(RWLock);
 				AutoRTFM::AbortTransaction();
 			});
 
@@ -1404,7 +1403,7 @@ TEST_CASE("UECore.TransactionallySafeReadScopeLock")
 		Result = AutoRTFM::Transact([&]
 			{
 				FTransactionallySafeRWLock RWLock;
-				FTransactionallySafeReadScopeLock Lock(RWLock);
+				UE::TReadScopeLock Lock(RWLock);
 			});
 
 		REQUIRE(AutoRTFM::ETransactionResult::Committed == Result);
@@ -1420,7 +1419,7 @@ TEST_CASE("UECore.TransactionallySafeReadScopeLock")
 
 				InnerResult = AutoRTFM::Transact([&]
 					{
-						FTransactionallySafeReadScopeLock Lock(RWLock);
+						UE::TReadScopeLock Lock(RWLock);
 						AutoRTFM::CascadingAbortTransaction();
 					});
 			});
@@ -1433,7 +1432,7 @@ TEST_CASE("UECore.TransactionallySafeReadScopeLock")
 
 				InnerResult = AutoRTFM::Transact([&]
 					{
-						FTransactionallySafeReadScopeLock Lock(RWLock);
+						UE::TReadScopeLock Lock(RWLock);
 						AutoRTFM::AbortTransaction();
 					});
 			});
@@ -1447,7 +1446,7 @@ TEST_CASE("UECore.TransactionallySafeReadScopeLock")
 
 				InnerResult = AutoRTFM::Transact([&]
 					{
-						FTransactionallySafeReadScopeLock Lock(RWLock);
+						UE::TReadScopeLock Lock(RWLock);
 					});
 
 				AutoRTFM::AbortTransaction();
@@ -1461,7 +1460,7 @@ TEST_CASE("UECore.TransactionallySafeReadScopeLock")
 
 				InnerResult = AutoRTFM::Transact([&]
 					{
-						FTransactionallySafeReadScopeLock Lock(RWLock);
+						UE::TReadScopeLock Lock(RWLock);
 					});
 			});
 
diff --git a/Engine/Source/Runtime/AssetRegistry/Private/AssetRegistryImpl.h b/Engine/Source/Runtime/AssetRegistry/Private/AssetRegistryImpl.h
--- a/Engine/Source/Runtime/AssetRegistry/Private/AssetRegistryImpl.h
+++ b/Engine/Source/Runtime/AssetRegistry/Private/AssetRegistryImpl.h
@@ -13,7 +13,6 @@
 #include "Math/NumericLimits.h"
 #include "Misc/Optional.h"
 #include "Misc/TransactionallySafeRWLock.h"
-#include "Misc/TransactionallySafeRWScopeLock.h"
 #include "Misc/ScopeRWLock.h"
 #include "ModuleDescriptor.h"
 #include "PackageDependencyData.h"
@@ -45,7 +44,7 @@ namespace UE::AssetRegistry::Impl { struct FScanPathContext; }
 namespace UE::AssetRegistry::Impl { struct FTickContext; }
 namespace UE::AssetRegistry::Impl { struct FClassInheritanceContext; }
 namespace UE::AssetRegistry { template <typename TScopeLockType> class TRWScopeLockWithPriority; }
-namespace UE::AssetRegistry { using FInterfaceReadScopeLock = TRWScopeLockWithPriority<FTransactionallySafeReadScopeLock>; }
+namespace UE::AssetRegistry { using FInterfaceReadScopeLock = TRWScopeLockWithPriority<UE::TReadScopeLock<FTransactionallySafeRWLock>>; }
 namespace UE::AssetRegistry { using FInterfaceWriteScopeLock = TRWScopeLockWithPriority<UE::TWriteScopeLock<FTransactionallySafeRWLock>>; }
 namespace UE::AssetRegistry { using FInterfaceRWScopeLock = class FRWScopeLockWithPriority; }
 
diff --git a/Engine/Source/Runtime/Core/Private/Misc/ConfigCacheIni.cpp b/Engine/Source/Runtime/Core/Private/Misc/ConfigCacheIni.cpp
--- a/Engine/Source/Runtime/Core/Private/Misc/ConfigCacheIni.cpp
+++ b/Engine/Source/Runtime/Core/Private/Misc/ConfigCacheIni.cpp
@@ -1085,7 +1085,7 @@ UE::ConfigAccessTracking::FFile* FConfigFile::GetFileAccess() const
 
 bool FConfigFile::operator==( const FConfigFile& Other ) const
 {
-	FTransactionallySafeReadScopeLock ScopeLock(ConfigFileMapLock);
+	UE::TReadScopeLock ScopeLock(ConfigFileMapLock);
 
 	if ( Pairs.Num() != Other.Pairs.Num() )
 		return 0;
diff --git a/Engine/Source/Runtime/Core/Private/Serialization/CustomVersion.cpp b/Engine/Source/Runtime/Core/Private/Serialization/CustomVersion.cpp
--- a/Engine/Source/Runtime/Core/Private/Serialization/CustomVersion.cpp
+++ b/Engine/Source/Runtime/Core/Private/Serialization/CustomVersion.cpp
@@ -13,7 +13,6 @@
 #include "Misc/AssertionMacros.h"
 #include "Misc/ScopeRWLock.h"
 #include "Misc/TransactionallySafeRWLock.h"
-#include "Misc/TransactionallySafeRWScopeLock.h"
 #include "Serialization/Archive.h"
 #include "Serialization/StructuredArchive.h"
 #include "Serialization/StructuredArchiveAdapters.h"
@@ -181,7 +180,7 @@ FCustomVersionContainer FCurrentCustomVersions::GetAll()
 	FStaticCustomVersionRegistry& Registry = FStaticCustomVersionRegistry::Get();
 
 	{
-		FTransactionallySafeReadScopeLock Scope(Registry.Lock);
+		UE::TReadScopeLock Scope(Registry.Lock);
 
 		if (Registry.Queue.Num() == 0)
 		{
@@ -198,7 +197,7 @@ TOptional<FCustomVersion> FCurrentCustomVersions::Get(const FGuid& Guid)
 {
 	FStaticCustomVersionRegistry& Registry = FStaticCustomVersionRegistry::Get();
 
-	FTransactionallySafeReadScopeLock Scope(Registry.Lock);
+	UE::TReadScopeLock Scope(Registry.Lock);
 	return Registry.Find(Guid);
 }
 
@@ -210,7 +209,7 @@ TArray<FCustomVersionDifference> FCurrentCustomVersions::Compare(const FCustomVe
 	{
 		FStaticCustomVersionRegistry& Registry = FStaticCustomVersionRegistry::Get();
 
-		FTransactionallySafeReadScopeLock Scope(Registry.Lock);
+		UE::TReadScopeLock Scope(Registry.Lock);
 
 		for (const FCustomVersion& CompareVersion : CompareVersions)
 		{
diff --git a/Engine/Source/Runtime/Core/Public/Misc/ConfigCacheIni.h b/Engine/Source/Runtime/Core/Public/Misc/ConfigCacheIni.h
--- a/Engine/Source/Runtime/Core/Public/Misc/ConfigCacheIni.h
+++ b/Engine/Source/Runtime/Core/Public/Misc/ConfigCacheIni.h
@@ -32,7 +32,7 @@
 #include "Misc/ConfigTypes.h"
 #include "Misc/Paths.h"
 #include "Misc/ScopeRWLock.h"
-#include "Misc/TransactionallySafeRWScopeLock.h"
+#include "Misc/TransactionallySafeRWLock.h"
 #include "Serialization/Archive.h"
 #include "Serialization/StructuredArchive.h"
 #include "Serialization/StructuredArchiveAdapters.h"
@@ -624,32 +624,32 @@ public:
 	UE_DEPRECATED(5.4, "Use FindSection, and/or use the new AddToSection, etc APIs to modify sections without retrieving the section. See top of ConfigCacheIni.h for more info.")
 	FORCEINLINE const FConfigSection* Find(const FString& SectionName) const
 	{
-		FTransactionallySafeReadScopeLock ScopeLock(ConfigFileMapLock);
+		UE::TReadScopeLock ScopeLock(ConfigFileMapLock);
 		return FConfigFileMap::Find(SectionName);
 	}
 
 	UE_DEPRECATED(5.4, "Use FindSection, and/or use the new AddToSection, etc APIs to modify sections without retrieving the section. See top of ConfigCacheIni.h for more info.")
 	FORCEINLINE FConfigSection* Find(const FString& SectionName)
 	{
-		FTransactionallySafeReadScopeLock ScopeLock(ConfigFileMapLock);
+		UE::TReadScopeLock ScopeLock(ConfigFileMapLock);
 		return FConfigFileMap::Find(SectionName);
 	}
 
 	FORCEINLINE const FConfigSection* FindSection(const FString& SectionName) const
 	{
-		FTransactionallySafeReadScopeLock ScopeLock(ConfigFileMapLock);
+		UE::TReadScopeLock ScopeLock(ConfigFileMapLock);
 		return FConfigFileMap::Find(SectionName);
 	}
 
 	FORCEINLINE int32 Num() const
 	{
-		FTransactionallySafeReadScopeLock ScopeLock(ConfigFileMapLock);
+		UE::TReadScopeLock ScopeLock(ConfigFileMapLock);
 		return FConfigFileMap::Num();
 	}
 
 	FORCEINLINE bool IsEmpty() const
 	{
-		FTransactionallySafeReadScopeLock ScopeLock(ConfigFileMapLock);
+		UE::TReadScopeLock ScopeLock(ConfigFileMapLock);
 		return FConfigFileMap::IsEmpty();
 	}
 
@@ -661,19 +661,19 @@ public:
 
 	FORCEINLINE bool Contains(const FString& SectionName) const
 	{
-		FTransactionallySafeReadScopeLock ScopeLock(ConfigFileMapLock);
+		UE::TReadScopeLock ScopeLock(ConfigFileMapLock);
 		return FConfigFileMap::Contains(SectionName);
 	}
 
 	FORCEINLINE int32 GetKeys(TArray<FString>& Keys) const
 	{
-		FTransactionallySafeReadScopeLock ScopeLock(ConfigFileMapLock);
+		UE::TReadScopeLock ScopeLock(ConfigFileMapLock);
 		return FConfigFileMap::GetKeys(Keys);
 	}
 
 	FORCEINLINE int32 GetKeys(TSet<FString>& Keys) const
 	{
-		FTransactionallySafeReadScopeLock ScopeLock(ConfigFileMapLock);
+		UE::TReadScopeLock ScopeLock(ConfigFileMapLock);
 		return FConfigFileMap::GetKeys(Keys);
 	}
 
@@ -729,26 +729,26 @@ public:
 	UE_DEPRECATED(5.4, "Use const ranged for iterators, (wrap your FConfigFile variable in AsConst to force the const iterator). See top of ConfigCacheIni.h for more info.")
 	FORCEINLINE TRangedForIterator begin()
 	{
-		FTransactionallySafeReadScopeLock ScopeLock(ConfigFileMapLock);
+		UE::TReadScopeLock ScopeLock(ConfigFileMapLock);
 		return TRangedForIterator(Pairs.begin());
 	}
 
 	FORCEINLINE TRangedForConstIterator begin() const
 	{
-		FTransactionallySafeReadScopeLock ScopeLock(ConfigFileMapLock); 
+		UE::TReadScopeLock ScopeLock(ConfigFileMapLock); 
 		return TRangedForConstIterator(Pairs.begin());
 	}
 
 	UE_DEPRECATED(5.4, "Use const ranged for iterators, (wrap your FConfigFile variable in AsConst to force the const iterator). See top of ConfigCacheIni.h for more info.")
 	FORCEINLINE TRangedForIterator end()
 	{
-		FTransactionallySafeReadScopeLock ScopeLock(ConfigFileMapLock); 
+		UE::TReadScopeLock ScopeLock(ConfigFileMapLock); 
 		return TRangedForIterator(Pairs.end());
 	}
 
 	FORCEINLINE TRangedForConstIterator end() const
 	{
-		FTransactionallySafeReadScopeLock ScopeLock(ConfigFileMapLock);
+		UE::TReadScopeLock ScopeLock(ConfigFileMapLock);
 		return TRangedForConstIterator(Pairs.end());
 	}
 	
@@ -825,7 +825,7 @@ private:
 
 	FORCEINLINE FConfigSection* FindInternal(const FString& SectionName)
 	{
-		FTransactionallySafeReadScopeLock ScopeLock(ConfigFileMapLock);
+		UE::TReadScopeLock ScopeLock(ConfigFileMapLock);
 		return FConfigFileMap::Find(SectionName);
 	};
 
diff --git a/Engine/Source/Runtime/Core/Public/Misc/TransactionallySafeRWScopeLock.h b/Engine/Source/Runtime/Core/Public/Misc/TransactionallySafeRWScopeLock.h
deleted file mode 100644
diff --git a/Engine/Source/Runtime/CoreUObject/Private/UObject/LinkerInstancingContext.cpp b/Engine/Source/Runtime/CoreUObject/Private/UObject/LinkerInstancingContext.cpp
--- a/Engine/Source/Runtime/CoreUObject/Private/UObject/LinkerInstancingContext.cpp
+++ b/Engine/Source/Runtime/CoreUObject/Private/UObject/LinkerInstancingContext.cpp
@@ -10,7 +10,6 @@
 #include "Misc/AutomationTest.h"
 #include "Misc/ScopeRWLock.h"
 #include "Misc/TransactionallySafeRWLock.h"
-#include "Misc/TransactionallySafeRWScopeLock.h"
 
 LLM_DEFINE_TAG(Loading_LinkerInstancingContext);
 
@@ -59,7 +58,7 @@ public:
 
 	bool FindPackageMapping(FName Original, FName& Instanced) const
 	{
-		FTransactionallySafeReadScopeLock ScopeLock(Lock);
+		UE::TReadScopeLock ScopeLock(Lock);
 		LLM_SCOPE_BYTAG(Loading_LinkerInstancingContext);
 		if (const FName* InstancedPtr = InstancedPackageMap.InstancedPackageMapping.Find(Original))
 		{
@@ -71,14 +70,14 @@ public:
 
 	void FixupSoftObjectPath(FSoftObjectPath& InOutSoftObjectPath) const
 	{
-		FTransactionallySafeReadScopeLock ScopeLock(Lock);
+		UE::TReadScopeLock ScopeLock(Lock);
 		LLM_SCOPE_BYTAG(Loading_LinkerInstancingContext);
 		InstancedPackageMap.FixupSoftObjectPath(InOutSoftObjectPath);
 	}
 
 	FSoftObjectPath RemapPath(const FSoftObjectPath& Path) const
 	{
-		FTransactionallySafeReadScopeLock ScopeLock(Lock);
+		UE::TReadScopeLock ScopeLock(Lock);
 		if (const FTopLevelAssetPath* Remapped = PathMapping.Find(Path.GetAssetPath()))
 		{
 			return FSoftObjectPath(*Remapped, Path.GetSubPathString());
@@ -88,7 +87,7 @@ public:
 
 	bool IsInstanced() const
 	{
-		FTransactionallySafeReadScopeLock ScopeLock(Lock);
+		UE::TReadScopeLock ScopeLock(Lock);
 		return InstancedPackageMap.IsInstanced() || InstancedPackageMapFunc || PathMapping.Num() > 0;
 	}
 
@@ -121,7 +120,7 @@ public:
 
 	bool HasTag(FName Tag) const
 	{
-		FTransactionallySafeReadScopeLock ScopeLock(Lock);
+		UE::TReadScopeLock ScopeLock(Lock);
 		return Tags.Contains(Tag);
 	}
 
@@ -154,7 +153,7 @@ public:
 
 		FName RemappedPackageName;
 		{
-			FTransactionallySafeReadScopeLock ScopeLock(Lock);
+			UE::TReadScopeLock ScopeLock(Lock);
 			RemappedPackageName = InstancedPackageMap.RemapPackage(PackageName);
 
 			if ((RemappedPackageName == PackageName) && InstancedPackageMapFunc)
diff --git a/Engine/Source/Runtime/CoreUObject/Public/UObject/UObjectAnnotation.h b/Engine/Source/Runtime/CoreUObject/Public/UObject/UObjectAnnotation.h
--- a/Engine/Source/Runtime/CoreUObject/Public/UObject/UObjectAnnotation.h
+++ b/Engine/Source/Runtime/CoreUObject/Public/UObject/UObjectAnnotation.h
@@ -13,7 +13,6 @@
 #include "Misc/ScopeRWLock.h"
 #include "Misc/TransactionallySafeCriticalSection.h"
 #include "Misc/TransactionallySafeRWLock.h"
-#include "Misc/TransactionallySafeRWScopeLock.h"
 
 /**
 * FUObjectAnnotationSparse is a helper class that is used to store sparse, slow, temporary, editor only, external 
@@ -867,7 +866,7 @@ public:
 
 		UE_AUTORTFM_OPEN
 		{
-			FTransactionallySafeReadScopeLock AnnotationArrayLock(AnnotationArrayCritical);
+			UE::TReadScopeLock AnnotationArrayLock(AnnotationArrayCritical);
 
 			const int32 ChunkIndex = Index / NumAnnotationsPerChunk;
 			if (ChunkIndex < Chunks.Num())
diff --git a/Engine/Source/Runtime/Engine/Classes/Components/ActorComponent.h b/Engine/Source/Runtime/Engine/Classes/Components/ActorComponent.h
--- a/Engine/Source/Runtime/Engine/Classes/Components/ActorComponent.h
+++ b/Engine/Source/Runtime/Engine/Classes/Components/ActorComponent.h
@@ -16,7 +16,8 @@
 #include "PropertyPairsMap.h"
 #include "ComponentInstanceDataCache.h"
 #include "Experimental/ConcurrentLinearAllocator.h"
-#include "Misc/TransactionallySafeRWScopeLock.h"
+#include "Misc/ScopeRWLock.h"
+#include "Misc/TransactionallySafeRWLock.h"
 #include "ActorComponent.generated.h"
 
 struct FTypedElementHandle;
diff --git a/Engine/Source/Runtime/Engine/Private/Animation/Skeleton.cpp b/Engine/Source/Runtime/Engine/Private/Animation/Skeleton.cpp
--- a/Engine/Source/Runtime/Engine/Private/Animation/Skeleton.cpp
+++ b/Engine/Source/Runtime/Engine/Private/Animation/Skeleton.cpp
@@ -773,7 +773,7 @@ const FSkeletonToMeshLinkup& USkeleton::FindOrAddMeshLinkupData(const USkinnedAs
 	const TUniquePtr<FSkeletonToMeshLinkup>* SkeletonToMeshLinkupPtr = nullptr;
 
 	{
-		FTransactionallySafeReadScopeLock Lock(SkinnedAssetLinkupCacheLock);
+		UE::TReadScopeLock Lock(SkinnedAssetLinkupCacheLock);
 		SkeletonToMeshLinkupPtr = SkinnedAssetLinkupCache.Find(InSkinnedAsset);
 	}
 
diff --git a/Engine/Source/Runtime/Engine/Private/VisualLogger/VisualLogger.cpp b/Engine/Source/Runtime/Engine/Private/VisualLogger/VisualLogger.cpp
--- a/Engine/Source/Runtime/Engine/Private/VisualLogger/VisualLogger.cpp
+++ b/Engine/Source/Runtime/Engine/Private/VisualLogger/VisualLogger.cpp
@@ -404,7 +404,7 @@ bool FVisualLogger::ExecuteOnLastEntryForObject(const UObject* Object, TFunction
 {
 	const UObject* LogOwner = nullptr;
 	{
-		FTransactionallySafeReadScopeLock Lock(RedirectRWLock);
+		UE::TReadScopeLock Lock(RedirectRWLock);
 		LogOwner = FindRedirectionInternal(Object);
 	}
 
@@ -426,7 +426,7 @@ FVisualLogEntry* FVisualLogger::GetEntryToWrite(const UObject* Object, const dou
 {
 	const UObject* LogOwner = nullptr;
 	{
-		FTransactionallySafeReadScopeLock Lock(RedirectRWLock);
+		UE::TReadScopeLock Lock(RedirectRWLock);
 		LogOwner = FindRedirectionInternal(Object);
 	}
 	if (LogOwner == nullptr)
@@ -582,7 +582,7 @@ FVisualLogEntry* FVisualLogger::GetEntryToWriteInternal(const UObject* Object, c
 			CurrentEntry->bIsLocationValid = true;
 		}
 
-		FTransactionallySafeReadScopeLock RedirectScopeLock(RedirectRWLock);
+		UE::TReadScopeLock RedirectScopeLock(RedirectRWLock);
 		FOwnerToChildrenRedirectionMap& RedirectionMap = GetRedirectionMap(LogOwner);
 		if (const IVisualLoggerDebugSnapshotInterface* DebugSnapshotInterface = Cast<const IVisualLoggerDebugSnapshotInterface>(LogOwner))
 		{
diff --git a/Engine/Source/Runtime/Engine/Public/VisualLogger/VisualLogger.h b/Engine/Source/Runtime/Engine/Public/VisualLogger/VisualLogger.h
--- a/Engine/Source/Runtime/Engine/Public/VisualLogger/VisualLogger.h
+++ b/Engine/Source/Runtime/Engine/Public/VisualLogger/VisualLogger.h
@@ -12,9 +12,8 @@
 #include "Traits/IsCharEncodingCompatibleWith.h"
 #include "UObject/ObjectKey.h"
 #include "Containers/Ticker.h"
-#include "Misc/TransactionallySafeRWLock.h"
-#include "Misc/TransactionallySafeRWScopeLock.h"
 #include "Misc/ScopeRWLock.h"
+#include "Misc/TransactionallySafeRWLock.h"
 
 #if ENABLE_VISUAL_LOG
 
@@ -744,7 +743,7 @@ public:
 	static UObject* FindRedirection(const UObject* Object)
 	{ 
 		FVisualLogger& Logger = FVisualLogger::Get();
-		FTransactionallySafeReadScopeLock Lock(Logger.RedirectRWLock);
+		UE::TReadScopeLock Lock(Logger.RedirectRWLock);
 		return Logger.FindRedirectionInternal(Object);
 	}
 
diff --git a/Engine/Source/Runtime/MovieScene/Private/MovieScenePlayer.cpp b/Engine/Source/Runtime/MovieScene/Private/MovieScenePlayer.cpp
--- a/Engine/Source/Runtime/MovieScene/Private/MovieScenePlayer.cpp
+++ b/Engine/Source/Runtime/MovieScene/Private/MovieScenePlayer.cpp
@@ -17,7 +17,6 @@
 #include "MovieSceneBindingReferences.h"
 #include "Misc/ScopeRWLock.h"
 #include "Misc/TransactionallySafeRWLock.h"
-#include "Misc/TransactionallySafeRWScopeLock.h"
 
 namespace UE
 {
@@ -79,14 +78,14 @@ PRAGMA_ENABLE_DEPRECATION_WARNINGS
 
 IMovieScenePlayer* IMovieScenePlayer::Get(uint16 InUniqueIndex)
 {
-	FTransactionallySafeReadScopeLock ScopeLock(UE::MovieScene::GGlobalPlayerRegistryLock);
+	UE::TReadScopeLock ScopeLock(UE::MovieScene::GGlobalPlayerRegistryLock);
 	check(UE::MovieScene::GGlobalPlayerRegistry.IsValidIndex(InUniqueIndex));
 	return UE::MovieScene::GGlobalPlayerRegistry[InUniqueIndex];
 }
 
 void IMovieScenePlayer::Get(TArray<IMovieScenePlayer*>& OutPlayers, bool bOnlyUnstoppedPlayers)
 {
-	FTransactionallySafeReadScopeLock ScopeLock(UE::MovieScene::GGlobalPlayerRegistryLock);
+	UE::TReadScopeLock ScopeLock(UE::MovieScene::GGlobalPlayerRegistryLock);
 	for (auto It = UE::MovieScene::GGlobalPlayerRegistry.CreateIterator(); It; ++It)
 	{
 		if (IMovieScenePlayer* Player = *It)
