Add set key funcs template to TMap::GetKeys(TSet)

diff --git a/Engine/Source/Runtime/Core/Public/Containers/Map.h b/Engine/Source/Runtime/Core/Public/Containers/Map.h
--- a/Engine/Source/Runtime/Core/Public/Containers/Map.h
+++ b/Engine/Source/Runtime/Core/Public/Containers/Map.h
@@ -358,7 +358,7 @@ public:
 	 * @param OutKeys Upon return, contains the set of unique keys in this map.
 	 * @return The number of unique keys in the map.
 	 */
-	template<typename Allocator> int32 GetKeys(TSet<KeyType, Allocator>& OutKeys) const
+	template<typename InSetKeyFuncs, typename InSetAllocator> int32 GetKeys(TSet<KeyType, InSetKeyFuncs, InSetAllocator>& OutKeys) const
 	{
 		OutKeys.Reset();
 		
