- MInor fixes
- some size_t -> int32 explicit typecasts
- some WIDECHAR -> int16 explicit typecasts
- improved a per-project SDK override exception message

diff --git a/Engine/Source/Programs/UnrealBuildTool/System/RulesAssembly.cs b/Engine/Source/Programs/UnrealBuildTool/System/RulesAssembly.cs
--- a/Engine/Source/Programs/UnrealBuildTool/System/RulesAssembly.cs
+++ b/Engine/Source/Programs/UnrealBuildTool/System/RulesAssembly.cs
@@ -810,7 +810,7 @@ namespace UnrealBuildTool
 					// if the target doesn't allow for an override at all, error
 					if (!Rules.AllowsPerProjectSDKVersion())
 					{
-						throw new CompilationResultException(CompilationResult.RulesError, KnownLogEvents.RulesAssembly, "Target {RulesName} is being built with a overridden {Platform} SDK version to '{SdkVersion}', but this target is not allowed - likely due to a modular build using a Shared BuildEnvironment",
+						throw new CompilationResultException(CompilationResult.RulesError, KnownLogEvents.RulesAssembly, "Target {RulesName} is being built with a overridden {Platform} SDK version to '{SdkVersion}', but this target is not allowed - likely due to a modular build using a Shared BuildEnvironment. If this target doesn't care about SDK versions, set 'bAreTargetSDKVersionsRelevantOverride = false' in your Target.cs file. Be aware this could cause confusion with multiple targets fighting over shared modules.",
 							Rules.Name, Platform, SDK.GetMainVersion());
 					}
 
diff --git a/Engine/Source/Runtime/Core/Public/Containers/StringView.h b/Engine/Source/Runtime/Core/Public/Containers/StringView.h
--- a/Engine/Source/Runtime/Core/Public/Containers/StringView.h
+++ b/Engine/Source/Runtime/Core/Public/Containers/StringView.h
@@ -518,25 +518,25 @@ inline FStringView operator "" _PrivateSV(const ANSICHAR* String, size_t Size)
 #else
 constexpr inline FStringView operator "" _PrivateSV(const TCHAR* String, size_t Size)
 {
-	return FStringView(String, Size);
+	return FStringView(String, (int32)Size);
 }
 #endif
 
 constexpr inline FAnsiStringView operator "" _PrivateASV(const ANSICHAR* String, size_t Size)
 {
-	return FAnsiStringView(String, Size);
+	return FAnsiStringView(String, (int32)Size);
 }
 
 constexpr inline FWideStringView operator "" _PrivateWSV(const WIDECHAR* String, size_t Size)
 {
-	return FWideStringView(String, Size);
+	return FWideStringView(String, (int32)Size);
 }
 
 /*constexpr*/ inline FUtf8StringView operator "" _PrivateU8SV(const ANSICHAR* String, size_t Size)
 {
 	// Would like this operator to be constexpr, but cannot be until after this operator can take a UTF8CHAR*
 	// rather than an ANSICHAR*, which won't be until we have C++20 char8_t string literals.
-	return FUtf8StringView(reinterpret_cast<const UTF8CHAR*>(String), Size);
+	return FUtf8StringView(reinterpret_cast<const UTF8CHAR*>(String), (int32)Size);
 }
 
 #if PLATFORM_TCHAR_IS_UTF8CHAR
diff --git a/Engine/Source/Runtime/Core/Public/GenericPlatform/GenericWidePlatformString.h b/Engine/Source/Runtime/Core/Public/GenericPlatform/GenericWidePlatformString.h
--- a/Engine/Source/Runtime/Core/Public/GenericPlatform/GenericWidePlatformString.h
+++ b/Engine/Source/Runtime/Core/Public/GenericPlatform/GenericWidePlatformString.h
@@ -230,12 +230,12 @@ PRAGMA_ENABLE_DEPRECATION_WARNINGS
 
 	static FORCEINLINE int32 Strlen( const ANSICHAR* String )
 	{
-		return strlen( String );
+		return (int32)strlen( String );
 	}
 
 	static FORCEINLINE int32 Strnlen( const ANSICHAR* String, SIZE_T StringSize )
 	{
-		return strnlen( String, StringSize );
+		return (int32)strnlen( String, StringSize );
 	}
 
 	static FORCEINLINE const ANSICHAR* Strstr( const ANSICHAR* String, const ANSICHAR* Find)
@@ -275,7 +275,7 @@ PRAGMA_ENABLE_DEPRECATION_WARNINGS
 
 	static FORCEINLINE int32 Strtoi( const ANSICHAR* Start, ANSICHAR** End, int32 Base )
 	{
-		return strtol( Start, End, Base );
+		return (int32)strtol( Start, End, Base );
 	}
 
 	static FORCEINLINE int64 Strtoi64( const ANSICHAR* Start, ANSICHAR** End, int32 Base )
diff --git a/Engine/Source/Runtime/Core/Public/Math/InterpCurve.h b/Engine/Source/Runtime/Core/Public/Math/InterpCurve.h
--- a/Engine/Source/Runtime/Core/Public/Math/InterpCurve.h
+++ b/Engine/Source/Runtime/Core/Public/Math/InterpCurve.h
@@ -729,7 +729,7 @@ float FInterpCurve<T>::FindNearestOnSegment(const T& PointInSpace, int32 PtIdx,
 
 		// Find the index of the best approximation
 		const float* MinValuePtr = Algo::MinElement(DistancesSq);
-		const int32 IndexOfTheSmallest = MinValuePtr - DistancesSq;
+		const int32 IndexOfTheSmallest = (int32)(MinValuePtr - DistancesSq);
 
 		float& Move = NextMovesT[IndexOfTheSmallest];
 
diff --git a/Engine/Source/Runtime/Core/Public/Misc/Crc.h b/Engine/Source/Runtime/Core/Public/Misc/Crc.h
--- a/Engine/Source/Runtime/Core/Public/Misc/Crc.h
+++ b/Engine/Source/Runtime/Core/Public/Misc/Crc.h
@@ -219,9 +219,9 @@ inline uint32 FCrc::Strihash_DEPRECATED(const WIDECHAR* Data)
 	while( *Data )
 	{
 		WIDECHAR Ch = TChar<WIDECHAR>::ToUpper(*Data++);
-		uint16  B  = Ch;
+		uint16  B  = (uint16)Ch;
 		Hash     = ((Hash >> 8) & 0x00FFFFFF) ^ CRCTable_DEPRECATED[(Hash ^ B) & 0x000000FF];
-		B        = Ch>>8;
+		B        = (uint16)Ch>>8;
 		Hash     = ((Hash >> 8) & 0x00FFFFFF) ^ CRCTable_DEPRECATED[(Hash ^ B) & 0x000000FF];
 	}
 	return Hash;
@@ -237,9 +237,9 @@ inline uint32 FCrc::Strihash_DEPRECATED(const int32 DataLen, const WIDECHAR* Dat
 	for (int32 Idx = 0; Idx < DataLen; ++Idx)
 	{
 		WIDECHAR Ch = TChar<WIDECHAR>::ToUpper(*Data++);
-		uint16  B  = Ch;
+		uint16  B  = (uint16)Ch;
 		Hash     = ((Hash >> 8) & 0x00FFFFFF) ^ CRCTable_DEPRECATED[(Hash ^ B) & 0x000000FF];
-		B        = Ch>>8;
+		B        = (uint16)Ch>>8;
 		Hash     = ((Hash >> 8) & 0x00FFFFFF) ^ CRCTable_DEPRECATED[(Hash ^ B) & 0x000000FF];
 	}
 	return Hash;
