CocoaThread - Make FCocoaRunLoopTask inherit SendMode from the thread instigator's WaitMode
This is to avoid getting stuck if a thread's WaitMode doesn't match the SendMode to a call to itself resulting in a deadlock.
Add new console command Mac.Tests.CocoaThread to stress test the CocoaThread call on other threads mecanism

#tests Local testing & QA Test

diff --git a/Engine/Source/Runtime/ApplicationCore/Private/Mac/CocoaThreadTesting.cpp b/Engine/Source/Runtime/ApplicationCore/Private/Mac/CocoaThreadTesting.cpp
new file mode 100644
--- /dev/null
+++ b/Engine/Source/Runtime/ApplicationCore/Private/Mac/CocoaThreadTesting.cpp
@@ -0,0 +1,43 @@
+// Copyright Epic Games, Inc. All Rights Reserved.
+
+#include "CoreMinimal.h"
+#include "HAL/IConsoleManager.h"
+#include "Mac/CocoaThread.h"
+
+#if (WITH_EDITOR || IS_PROGRAM) && !UE_BUILD_SHIPPING
+
+void RecursiveThreadCall(int Depth, bool bFromGameThread, int MaxDepth)
+{
+	if (Depth < MaxDepth)
+	{
+		if (bFromGameThread)
+		{
+			// Test other thread call
+			MainThreadCall(^{
+				RecursiveThreadCall(Depth + 1, false, MaxDepth);
+			}, UnrealNilEventMode);
+			
+			// Test current thread call
+			GameThreadCall(^{});
+		}
+		else
+		{
+			// Test other thread call
+			GameThreadCall(^{
+				RecursiveThreadCall(Depth + 1, true, MaxDepth);
+			});
+			
+			// Test current thread call
+			MainThreadCall(^{});
+		}
+	}
+}
+
+static void TestCocoaThread()
+{
+	RecursiveThreadCall(0, true, 100);
+}
+
+FAutoConsoleCommand TestCocoaCommand(TEXT("Mac.Tests.CocoaThread"), TEXT(""), FConsoleCommandDelegate::CreateStatic(&TestCocoaThread));
+
+#endif
diff --git a/Engine/Source/Runtime/Core/Private/Mac/CocoaThread.cpp b/Engine/Source/Runtime/Core/Private/Mac/CocoaThread.cpp
--- a/Engine/Source/Runtime/Core/Private/Mac/CocoaThread.cpp
+++ b/Engine/Source/Runtime/Core/Private/Mac/CocoaThread.cpp
@@ -403,6 +403,9 @@ FCocoaRunLoopSource* FCocoaRunLoopSource::GameRunLoopSource = nullptr;
 
 @end
 
+// Store the WaitMode of the thread that scheduled the execution block on the thread that is going to process it
+thread_local NSString* GThreadInstigatorWaitMode = nullptr;
+
 static void PerformBlockOnThread(FCocoaRunLoopSource& ThreadSource, dispatch_block_t Block, NSArray* SendModes, NSString* WaitMode, bool const bWait)
 {
 	dispatch_block_t CopiedBlock = Block_copy(Block);
@@ -410,9 +413,29 @@ static void PerformBlockOnThread(FCocoaRunLoopSource& ThreadSource, dispatch_blo
 	if ( bWait )
 	{
 		dispatch_semaphore_t Semaphore = dispatch_semaphore_create(0);
-		dispatch_block_t ExecuteBlock = Block_copy(^{ CopiedBlock(); dispatch_semaphore_signal(Semaphore); });
+		NSString* ThreadInstigatorWaitMode = [WaitMode retain];
+		dispatch_block_t ExecuteBlock = Block_copy(^{
+			NSString* SavedInstigatorWaitMode = GThreadInstigatorWaitMode;			
+			GThreadInstigatorWaitMode = ThreadInstigatorWaitMode;
+			CopiedBlock();
+			if (GThreadInstigatorWaitMode)
+			{
+				[GThreadInstigatorWaitMode release];
+			}
+			GThreadInstigatorWaitMode = SavedInstigatorWaitMode;
+			dispatch_semaphore_signal(Semaphore);
+		});
 		
-		ThreadSource.Schedule(ExecuteBlock, SendModes);
+		// If the thread instigator wait mode is not already contained in the SendModes, we add it to prevent potential deadlocks.
+		// The idea is that if thread A schedule a block on thread B and while executing that block thread B has to schedule a block on thread A, that block can be executed.
+		if (GThreadInstigatorWaitMode && ![SendModes containsObject:GThreadInstigatorWaitMode])
+		{
+			ThreadSource.Schedule(ExecuteBlock, [SendModes arrayByAddingObject:GThreadInstigatorWaitMode]);
+		}
+		else
+		{
+			ThreadSource.Schedule(ExecuteBlock, SendModes);
+		}
 		
 		do {
 			ThreadSource.Wake();
