StringBuilder: Added Add(CharType) to support use as the output of algorithms

diff --git a/Engine/Source/Runtime/Core/Public/Misc/StringBuilder.h b/Engine/Source/Runtime/Core/Public/Misc/StringBuilder.h
--- a/Engine/Source/Runtime/Core/Public/Misc/StringBuilder.h
+++ b/Engine/Source/Runtime/Core/Public/Misc/StringBuilder.h
@@ -232,9 +232,6 @@ public:
 		return *this;
 	}
 
-	UE_DEPRECATED(5.3, "Use Append instead of AppendAnsi.")
-	inline BuilderType& AppendAnsi(const FAnsiStringView String) { return Append(String); }
-
 	/** Replace characters at given position and length with substring */
 	void ReplaceAt(int32 Pos, int32 RemoveLen, ViewType Str)
 	{
@@ -287,6 +284,13 @@ public:
 		ReplaceAt(0, 0, Str);
 	}
 
+	/** Prefer AppendChar. Appends one character to the string. Supports this being the output for Algo::... */
+	template <typename AppendedCharType>
+	inline auto Add(AppendedCharType Char) -> decltype(AppendChar(Char), (void)0)
+	{
+		AppendChar(Char);
+	}
+
 	/**
 	 * Append every element of the range to the builder, separating the elements by the delimiter.
 	 *
