Fix TStringView constructor taking a character pointer to work in a constexpr context.

#rnx

diff --git a/Engine/Source/Runtime/Core/Public/Containers/StringView.h b/Engine/Source/Runtime/Core/Public/Containers/StringView.h
--- a/Engine/Source/Runtime/Core/Public/Containers/StringView.h
+++ b/Engine/Source/Runtime/Core/Public/Containers/StringView.h
@@ -27,6 +27,20 @@ namespace UE::Core::Private
 	{
 		return GetData(Forward<ArgTypes>(Args)...);
 	}
+
+	template<typename CharType>
+	constexpr int32 StringLength(const CharType* InString)
+	{
+		int32 Count = 0;
+		if (InString)
+		{
+			while (*InString++)
+			{
+				++Count;
+			}
+		}
+		return Count;
+	}
 } // UE::Core::Private
 
 /**
@@ -89,7 +103,7 @@ public:
 	/** Construct a view of the null-terminated string pointed to by InData. */
 	constexpr inline TStringView(const CharType* InData UE_LIFETIMEBOUND)
 		: DataPtr(InData)
-		, Size(InData ? TCString<CharType>::Strlen(InData) : 0)
+		, Size(UE::Core::Private::StringLength(InData))
 	{
 	}
 
@@ -110,9 +124,9 @@ public:
 			>::Value
 		)
 	>
-	constexpr inline TStringView(const OtherCharType* InData UE_LIFETIMEBOUND)
+	inline TStringView(const OtherCharType* InData UE_LIFETIMEBOUND)
 		: DataPtr((const CharType*)InData)
-		, Size(InData ? TCString<CharType>::Strlen((const CharType*)InData) : 0)
+		, Size(UE::Core::Private::StringLength((const CharType*)InData))
 	{
 	}
 
