Fix errors with latest compiler

#rnx

diff --git a/Engine/Source/Runtime/AssetRegistry/Private/AssetRegistry.cpp b/Engine/Source/Runtime/AssetRegistry/Private/AssetRegistry.cpp
--- a/Engine/Source/Runtime/AssetRegistry/Private/AssetRegistry.cpp
+++ b/Engine/Source/Runtime/AssetRegistry/Private/AssetRegistry.cpp
@@ -415,7 +415,12 @@ public:
 	{
 		//In the editor premade Asset Registry can be enabled by a command line argument so we need to wait until the task graph is ready 
 		//before we rely on UE::AssetRegistry::Premade::IsEnabled() to return the correct result
-		bool PremadeCanBeEnabled = WITH_EDITOR || UE::AssetRegistry::Premade::IsEnabled();
+		bool PremadeCanBeEnabled =
+#if WITH_EDITOR
+			true;
+#else
+			UE::AssetRegistry::Premade::IsEnabled();
+#endif
 
 		if (PremadeCanBeEnabled)
 		{
diff --git a/Engine/Source/Runtime/Core/Public/Misc/CoreDelegates.h b/Engine/Source/Runtime/Core/Public/Misc/CoreDelegates.h
--- a/Engine/Source/Runtime/Core/Public/Misc/CoreDelegates.h
+++ b/Engine/Source/Runtime/Core/Public/Misc/CoreDelegates.h
@@ -556,8 +556,10 @@ public:
 	/** Called when the verbosity of a log category is changed */
 	static CORE_API TMulticastDelegate<void(const FLogCategoryName& CategoryName, ELogVerbosity::Type OldVerbosity, ELogVerbosity::Type NewVerbosity)> OnLogVerbosityChanged;
 
+	PRAGMA_DISABLE_DEPRECATION_WARNINGS
 	UE_DEPRECATED(5.1, "Use FPackageStore::Mount() instead")
 	static CORE_API TDelegate<TSharedPtr<class IPackageStore>()> CreatePackageStore;
+	PRAGMA_ENABLE_DEPRECATION_WARNINGS
 
 	// Called immediately before the parent process will start responding to signals to fork
 	static CORE_API FSimpleMulticastDelegate OnParentBeginFork;
