[AutoRTFM] Check parent's hit set when merge write entries from nested transaction

Reduces peak write logs by about 1/4 for certain test runs.

diff --git a/Engine/Source/Runtime/Core/Private/AutoRTFM/HitSet.h b/Engine/Source/Runtime/Core/Private/AutoRTFM/HitSet.h
--- a/Engine/Source/Runtime/Core/Private/AutoRTFM/HitSet.h
+++ b/Engine/Source/Runtime/Core/Private/AutoRTFM/HitSet.h
@@ -12,6 +12,14 @@ namespace AutoRTFM
 	public:
 		using Key = FMemoryLocation;
 
+		// Maximum size in bytes for a hit-set write record.
+		// The cutoff here is arbitrarily any number less than UINT16_MAX, but its a
+		// weigh up what a good size is. Because the hitset doesn't detect when you
+		// are trying to write to a subregion of a previous hit (like memset something,
+		// then write to an individual element), we've got to balance the cost of
+		// recording meaningless hits, against the potential to hit again.
+		static constexpr size_t MaxSize = 16;
+
 	private:
 		// TODO: Revisit a good probe depth for the hashset.
 		static constexpr uint64_t LinearProbeDepth = 16;
diff --git a/Engine/Source/Runtime/Core/Private/AutoRTFM/Transaction.cpp b/Engine/Source/Runtime/Core/Private/AutoRTFM/Transaction.cpp
--- a/Engine/Source/Runtime/Core/Private/AutoRTFM/Transaction.cpp
+++ b/Engine/Source/Runtime/Core/Private/AutoRTFM/Transaction.cpp
@@ -230,14 +230,21 @@ void FTransaction::CommitNested()
 			continue;
 		}
 
-		Parent->WriteLog.Push(Write);
+		if (Write.GetSize() <= FHitSet::MaxSize)
+		{
+			FHitSet::Key HitSetEntry(Write.GetOriginal());
+			HitSetEntry.TopTag().Size = static_cast<uint16_t>(Write.GetSize());
+			HitSetEntry.TopTag().bNoMemoryValidation = Write.IgnoreForValidation();
 
-		FHitSet::Key HitSetEntry(Write.GetOriginal());
-		HitSetEntry.TopTag().Size = static_cast<uint16_t>(Write.GetSize());
-		HitSetEntry.TopTag().bNoMemoryValidation = Write.IgnoreForValidation();
+			if (Parent->HitSet.FindOrTryInsert(HitSetEntry))
+			{
+				continue; // Don't duplicate the write-log entry.
+			}
+		}
+
+		Parent->WriteLog.Push(Write);
+	}
 
-        Parent->HitSet.FindOrTryInsert(HitSetEntry);
-    }
 
     Parent->WriteLogBumpAllocator.Merge(MoveTemp(WriteLogBumpAllocator));
 
diff --git a/Engine/Source/Runtime/Core/Private/AutoRTFM/TransactionInlines.h b/Engine/Source/Runtime/Core/Private/AutoRTFM/TransactionInlines.h
--- a/Engine/Source/Runtime/Core/Private/AutoRTFM/TransactionInlines.h
+++ b/Engine/Source/Runtime/Core/Private/AutoRTFM/TransactionInlines.h
@@ -124,12 +124,7 @@ void FTransaction::RecordWrite(void* LogicalAddress, size_t Size, bool bNoMemory
         return;
     }
 
-	// The cutoff here is arbitrarily any number less than UINT16_MAX, but its a
-	// weigh up what a good size is. Because the hitset doesn't detect when you
-	// are trying to write to a subregion of a previous hit (like memset something,
-	// then write to an individual element), we've got to balance the cost of
-	// recording meaningless hits, against the potential to hit again.
-    if (Size <= 16)
+    if (Size <= FHitSet::MaxSize)
     {
         FMemoryLocation Key(LogicalAddress);
         FMemoryLocationTag& Tag = Key.TopTag();
