Csv profiler: add validaton to catch numbered FNames passed as stats (stats with an "_<number>" suffix) if UE_FNAME_OUTLINE_NUMBER is enabled, since this is not supported and leads to duplicate stats
Note: This will fatal error by default in Test builds, since checks/ensures are normally disabled in Test. The fatal error can be disabled by setting csv.NumberedFNamesStatsAreFatal to false.
#tests Tested locally and via preflight
[RN] Csv profiler: add validaton to catch numbered FNames passed as stats (stats with an "_<number>" suffix) if UE_FNAME_OUTLINE_NUMBER is enabled, since this is not supported and leads to duplicate stats. Set csv.NumberedFNamesStatsAreFatal to false to disable the fatal error.

diff --git a/Engine/Source/Runtime/Core/Private/ProfilingDebugging/CsvProfiler.cpp b/Engine/Source/Runtime/Core/Private/ProfilingDebugging/CsvProfiler.cpp
--- a/Engine/Source/Runtime/Core/Private/ProfilingDebugging/CsvProfiler.cpp
+++ b/Engine/Source/Runtime/Core/Private/ProfilingDebugging/CsvProfiler.cpp
@@ -182,6 +182,14 @@ TAutoConsoleVariable<int32> CVarMaxPerThreadStatDataSlackKB(
 	ECVF_Default
 );
 
+TAutoConsoleVariable<bool> CVarNumberedFNamesStatsAreFatal(
+	TEXT("csv.NumberedFNamesStatsAreFatal"),
+	true,
+	TEXT("Fatal error if numbered fname stats are encountered (ignored in shipping)"),
+	ECVF_Default
+);
+
+
 static bool GCsvUseProcessingThread = true;
 static int32 GCsvRepeatCount = 0;
 static int32 GCsvRepeatFrameCount = 0;
@@ -397,6 +405,18 @@ bool IsContinuousWriteEnabled(bool bGameThread)
 		CVarValue = bGameThread ? CVarCsvContinuousWrites.GetValueOnGameThread() : CVarCsvContinuousWrites.GetValueOnAnyThread();
 	}
 	return CVarValue > 0;
+} 
+
+
+static CSV_PROFILER_INLINE void ValidateFName(const FName& StatName)
+{
+	ensureMsgf(StatName.GetNumber() == 0, TEXT("Numbered FName stats (suffixed _<number>) are not supported. Stat name: '%s'"), *StatName.ToString());
+#if UE_FNAME_OUTLINE_NUMBER && !UE_BUILD_SHIPPING
+	if (StatName.GetNumber() > 0 && CVarNumberedFNamesStatsAreFatal.GetValueOnAnyThread())
+	{
+		UE_LOG(LogCsvProfiler, Fatal, TEXT("Numbered FName stats (suffixed _<number>) are not supported. Stat name: '%s'. Disable csv.NumberedFNamesStatsAreFatal to make this non-fatal"), *StatName.ToString());
+	}
+#endif
 }
 
 #if CSV_PROFILER_ALLOW_DEBUG_FEATURES
@@ -3757,6 +3777,7 @@ void FCsvProfiler::BeginStat(const FName& StatName, uint32 CategoryIndex)
 #if !CSV_PROFILER_MINIMAL
 	if (GCsvProfilerIsCapturing && GCsvCategoriesEnabled[CategoryIndex])
 	{
+		ValidateFName(StatName);
 #if CSV_PROFILER_SUPPORT_NAMED_EVENTS
 		if (UNLIKELY(GCsvProfilerNamedEventsTiming))
 		{
@@ -3789,6 +3810,7 @@ void FCsvProfiler::EndStat(const FName& StatName, uint32 CategoryIndex)
 #if !CSV_PROFILER_MINIMAL
 	if (GCsvProfilerIsCapturing && GCsvCategoriesEnabled[CategoryIndex])
 	{
+		ValidateFName(StatName);
 		FCsvProfilerThreadData::Get().AddTimestampEnd(StatName, CategoryIndex);
 #if CSV_PROFILER_SUPPORT_NAMED_EVENTS
 		if (UNLIKELY(GCsvProfilerNamedEventsTiming))
@@ -4103,6 +4125,7 @@ void FCsvProfiler::RecordCustomStat(const FName& StatName, uint32 CategoryIndex,
 #if !CSV_PROFILER_MINIMAL
 	if (GCsvProfilerIsCapturing && GCsvCategoriesEnabled[CategoryIndex])
 	{
+		ValidateFName(StatName);
 		FCsvProfilerThreadData::Get().AddCustomStat(StatName, CategoryIndex, Value, CustomStatOp);
 	}
 #endif
@@ -4132,6 +4155,7 @@ void FCsvProfiler::RecordCustomStat(const FName& StatName, uint32 CategoryIndex,
 #if !CSV_PROFILER_MINIMAL
 	if (GCsvProfilerIsCapturing && GCsvCategoriesEnabled[CategoryIndex])
 	{
+		ValidateFName(StatName);
 		FCsvProfilerThreadData::Get().AddCustomStat(StatName, CategoryIndex, Value, CustomStatOp);
 	}
 #endif
@@ -4568,4 +4592,4 @@ void CSVTest()
 
 #endif // CSV_PROFILER_ALLOW_DEBUG_FEATURES
 
-#endif // CSV_PROFILER
\ No newline at end of file
+#endif // CSV_PROFILER
