[Trivial] Add missing includes

#rnx

diff --git a/Engine/Source/Runtime/Core/Private/ProfilingDebugging/Microsoft/WindowsCallstackTrace.cpp b/Engine/Source/Runtime/Core/Private/ProfilingDebugging/Microsoft/WindowsCallstackTrace.cpp
--- a/Engine/Source/Runtime/Core/Private/ProfilingDebugging/Microsoft/WindowsCallstackTrace.cpp
+++ b/Engine/Source/Runtime/Core/Private/ProfilingDebugging/Microsoft/WindowsCallstackTrace.cpp
@@ -11,6 +11,7 @@
 #include "HAL/MemoryBase.h"
 #include "HAL/RunnableThread.h"
 #include "HAL/LowLevelMemTracker.h"
+#include "HAL/PlatformStackWalk.h"
 #include "ProfilingDebugging/CallstackTracePrivate.h"
 #include "ProfilingDebugging/MemoryTrace.h"
 
