Fixed Clang warnings in TNotNull.

diff --git a/Engine/Source/Runtime/Core/Public/Misc/NotNull.h b/Engine/Source/Runtime/Core/Public/Misc/NotNull.h
--- a/Engine/Source/Runtime/Core/Public/Misc/NotNull.h
+++ b/Engine/Source/Runtime/Core/Public/Misc/NotNull.h
@@ -38,6 +38,15 @@
 	#define UE_ENABLE_NOTNULL_WRAPPER (UE_BUILD_DEBUG || UE_BUILD_DEVELOPMENT)
 #endif
 
+#ifdef __clang__
+	// Disable clang warnings about returns_nonnull only being applicable to raw pointers
+	#define UE_NOTNULL_FUNCTION_NON_NULL_RETURN_START _Pragma("clang diagnostic push") _Pragma("clang diagnostic ignored \"-Wignored-attributes\"") FUNCTION_NON_NULL_RETURN_START _Pragma("clang diagnostic pop")
+	#define UE_NOTNULL_FUNCTION_NON_NULL_RETURN_END   _Pragma("clang diagnostic push") _Pragma("clang diagnostic ignored \"-Wignored-attributes\"") FUNCTION_NON_NULL_RETURN_END   _Pragma("clang diagnostic pop")
+#else
+	#define UE_NOTNULL_FUNCTION_NON_NULL_RETURN_START FUNCTION_NON_NULL_RETURN_START
+	#define UE_NOTNULL_FUNCTION_NON_NULL_RETURN_END   FUNCTION_NON_NULL_RETURN_END
+#endif
+
 #if UE_ENABLE_NOTNULL_WRAPPER
 
 	class FArchive;
@@ -145,17 +154,17 @@
 		operator bool() const = delete;
 
 		// Allow implicit conversion to the wrapped value, for compatibility with non-TNotNull APIs
-	FUNCTION_NON_NULL_RETURN_START
+	UE_NOTNULL_FUNCTION_NON_NULL_RETURN_START
 		operator T() const&
-	FUNCTION_NON_NULL_RETURN_END
+	UE_NOTNULL_FUNCTION_NON_NULL_RETURN_END
 		{
 			T Result = Val;
 			UE_ASSUME(Result);
 			return Result;
 		}
-	FUNCTION_NON_NULL_RETURN_START
+	UE_NOTNULL_FUNCTION_NON_NULL_RETURN_START
 		operator T() &&
-	FUNCTION_NON_NULL_RETURN_END
+	UE_NOTNULL_FUNCTION_NON_NULL_RETURN_END
 		{
 			T Result = (T&&)Val;
 			UE_ASSUME(Result);
@@ -163,9 +172,9 @@
 		}
 
 		// Allow dereferencing, if it's supported
-	FUNCTION_NON_NULL_RETURN_START
+	UE_NOTNULL_FUNCTION_NON_NULL_RETURN_START
 		const T& operator->() const
-	FUNCTION_NON_NULL_RETURN_END
+	UE_NOTNULL_FUNCTION_NON_NULL_RETURN_END
 		{
 			return Val;
 		}
@@ -261,6 +270,9 @@
 	}
 #endif // !PLATFORM_COMPILER_HAS_GENERATED_COMPARISON_OPERATORS
 
+#undef UE_NOTNULL_FUNCTION_NON_NULL_RETURN_END
+#undef UE_NOTNULL_FUNCTION_NON_NULL_RETURN_START
+
 #else // UE_ENABLE_NOTNULL_WRAPPER
 
 	template <typename T>
