Optimize FMatrix float/double coercion
Optimize FTransform ToMatrixInternal to avoid float coercion

diff --git a/Engine/Source/Runtime/Core/Public/Math/Matrix.h b/Engine/Source/Runtime/Core/Public/Math/Matrix.h
--- a/Engine/Source/Runtime/Core/Public/Math/Matrix.h
+++ b/Engine/Source/Runtime/Core/Public/Math/Matrix.h
@@ -410,11 +410,38 @@ public:
 	template<typename FArg UE_REQUIRES(!std::is_same_v<T, FArg>)>
 	explicit TMatrix(const TMatrix<FArg>& From)
 	{
-		// TODO: SIMD this?
+#if PLATFORM_ENABLE_VECTORINTRINSICS
+		const FArg* RESTRICT Src = &(From.M[0][0]);
+		T* RESTRICT Dest = &(M[0][0]);
+
+		// Load src
+		TVectorRegisterType<FArg> SrcRow0 = VectorLoad(&Src[0]);
+		TVectorRegisterType<FArg> SrcRow1 = VectorLoad(&Src[4]);
+		TVectorRegisterType<FArg> SrcRow2 = VectorLoad(&Src[8]);
+		TVectorRegisterType<FArg> SrcRow3 = VectorLoad(&Src[12]);
+
+		// Coerce and store
+		if constexpr (std::is_same_v<T, float>)
+		{
+			VectorStore(MakeVectorRegisterFloatFromDouble(SrcRow0), &Dest[0]);
+			VectorStore(MakeVectorRegisterFloatFromDouble(SrcRow1), &Dest[4]);
+			VectorStore(MakeVectorRegisterFloatFromDouble(SrcRow2), &Dest[8]);
+			VectorStore(MakeVectorRegisterFloatFromDouble(SrcRow3), &Dest[12]);
+		}
+		else
+		{
+			VectorStore(TVectorRegisterType<T>(SrcRow0), &Dest[0]);
+			VectorStore(TVectorRegisterType<T>(SrcRow1), &Dest[4]);
+			VectorStore(TVectorRegisterType<T>(SrcRow2), &Dest[8]);
+			VectorStore(TVectorRegisterType<T>(SrcRow3), &Dest[12]);
+		}
+#else
 		M[0][0] = (T)From.M[0][0]; M[0][1] = (T)From.M[0][1]; M[0][2] = (T)From.M[0][2]; M[0][3] = (T)From.M[0][3];
 		M[1][0] = (T)From.M[1][0]; M[1][1] = (T)From.M[1][1]; M[1][2] = (T)From.M[1][2]; M[1][3] = (T)From.M[1][3];
 		M[2][0] = (T)From.M[2][0]; M[2][1] = (T)From.M[2][1]; M[2][2] = (T)From.M[2][2]; M[2][3] = (T)From.M[2][3];
 		M[3][0] = (T)From.M[3][0]; M[3][1] = (T)From.M[3][1]; M[3][2] = (T)From.M[3][2]; M[3][3] = (T)From.M[3][3];
+#endif
+
 		DiagnosticCheckNaN();
 	}
 
diff --git a/Engine/Source/Runtime/Core/Public/Math/TransformVectorized.h b/Engine/Source/Runtime/Core/Public/Math/TransformVectorized.h
--- a/Engine/Source/Runtime/Core/Public/Math/TransformVectorized.h
+++ b/Engine/Source/Runtime/Core/Public/Math/TransformVectorized.h
@@ -1277,7 +1277,17 @@ private:
 #if !(UE_BUILD_SHIPPING || UE_BUILD_TEST) && WITH_EDITORONLY_DATA
 		// Make sure Rotation is normalized when we turn it into a matrix.
 		check( IsRotationNormalized() );
-#endif		
+#endif
+
+		TransformVectorRegister One;
+		if constexpr (std::is_same_v<T, double>)
+		{
+			One = VectorOneDouble();
+		}
+		else
+		{
+			One = VectorOne();
+		}
 
 		const TransformVectorRegister MyRotation = Rotation; // Load from persistent value, to avoid repeated loads.
 		const TransformVectorRegister RotationX2Y2Z2 = VectorAdd(MyRotation, MyRotation);	// x2, y2, z2
@@ -1290,7 +1300,7 @@ private:
 		const TransformVectorRegister yy2_xx2_xx2 = VectorSwizzle(RotationXX2YY2ZZ2, 1, 0, 0, 0);
 		const TransformVectorRegister zz2_zz2_yy2 = VectorSwizzle(RotationXX2YY2ZZ2, 2, 2, 1, 0);
 		const TransformVectorRegister DiagonalSum = VectorAdd(yy2_xx2_xx2, zz2_zz2_yy2);
-		const TransformVectorRegister Diagonals = VectorSubtract(VectorOne(), DiagonalSum);
+		const TransformVectorRegister Diagonals = VectorSubtract(One, DiagonalSum);
 		OutDiagonals = VectorMultiply(Diagonals, Scale3D);
 
 		// Grouping the non-diagonal elements in the rotation block by operations:
