Fixed build error when TNotNull is disabled.

diff --git a/Engine/Source/Runtime/Core/Private/Misc/NotNull.cpp b/Engine/Source/Runtime/Core/Private/Misc/NotNull.cpp
--- a/Engine/Source/Runtime/Core/Private/Misc/NotNull.cpp
+++ b/Engine/Source/Runtime/Core/Private/Misc/NotNull.cpp
@@ -5,7 +5,7 @@
 #include "CoreGlobals.h"
 #include "Serialization/Archive.h"
 
-#if DO_CHECK
+#if UE_ENABLE_NOTNULL_WRAPPER && DO_CHECK
 
 FORCENOINLINE void UE::Core::Private::ReportNotNullPtr()
 {
@@ -21,4 +21,4 @@ FORCENOINLINE void UE::Core::Private::CheckLoadingNotNullPtr(FArchive& Ar)
 	}
 }
 
-#endif // DO_CHECK
+#endif // UE_ENABLE_NOTNULL_WRAPPER && DO_CHECK
