- Handle setting SG group cvars in dynamic layers

diff --git a/Engine/Source/Runtime/Core/Private/Misc/ConfigCacheIni.cpp b/Engine/Source/Runtime/Core/Private/Misc/ConfigCacheIni.cpp
--- a/Engine/Source/Runtime/Core/Private/Misc/ConfigCacheIni.cpp
+++ b/Engine/Source/Runtime/Core/Private/Misc/ConfigCacheIni.cpp
@@ -6815,8 +6815,12 @@ class FIniExec : public FSelfRegisteringExec
 				FConfigBranch* Branch = GConfig->FindBranch(NAME_None, BranchName);
 				UE::DynamicConfig::PerformDynamicConfig("TrackedDyn", [Branch, Filename](FConfigModificationTracker* ChangeTracker)
 					{
+						// set which sections to track for cvars, with their priority
+						ChangeTracker->CVars.Add(TEXT("ConsoleVariables")).CVarPriority = (int)ECVF_SetByPluginLowPriority;
+						ChangeTracker->CVars.Add(TEXT("ConsoleVariables_HighPriority")).CVarPriority = (int)ECVF_SetByPluginHighPriority;
+
 						FDynamicLayerInfo Info { Filename, "TrackedDyn", (uint16)DynamicLayerPriority::Unknown };
-						Branch->AddDynamicLayersToHierarchy({ Info });
+						Branch->AddDynamicLayersToHierarchy({ Info }, ChangeTracker);
 					});
 			}
 		}
@@ -6826,6 +6830,7 @@ class FIniExec : public FSelfRegisteringExec
 			UE::DynamicConfig::PerformDynamicConfig("TrackedDyn", [](FConfigModificationTracker* ChangeTracker)
 				{
 					FConfigCacheIni::RemoveTagFromAllBranches("TrackedDyn", ChangeTracker);
+					IConsoleManager::Get().UnsetAllConsoleVariablesWithTag("TrackedDyn");
 				});
 		}
 
diff --git a/Engine/Source/Runtime/Core/Private/Misc/ConfigUtilities.cpp b/Engine/Source/Runtime/Core/Private/Misc/ConfigUtilities.cpp
--- a/Engine/Source/Runtime/Core/Private/Misc/ConfigUtilities.cpp
+++ b/Engine/Source/Runtime/Core/Private/Misc/ConfigUtilities.cpp
@@ -50,7 +50,37 @@ void PerformDynamicConfig(FName Tag, TFunction<void(class FConfigModificationTra
 					
 					for (const auto& CVarPair : Section)
 					{
-						UE::ConfigUtilities::OnSetCVarFromIniEntry(TEXT("DynamicLayer"), *CVarPair.Key.ToString(), *CVarPair.Value.GetValue(), Priority, false, false, TagPair.Key);
+						FString KeyString = CVarPair.Key.ToString();
+						if (KeyString.StartsWith(TEXT("sg.")))
+						{
+							UE::ConfigUtilities::ApplyCVarSettingsFromIni(*FString::Printf(TEXT("%s@%s"), *KeyString.Mid(3), *CVarPair.Value.GetValue()), *GScalabilityIni, Priority, false, TagPair.Key);
+						}
+						else
+						{
+							UE::ConfigUtilities::OnSetCVarFromIniEntry(TEXT("DynamicLayer"), *CVarPair.Key.ToString(), *CVarPair.Value.GetValue(), Priority, false, false, TagPair.Key);
+						}
+					}
+				}
+			}
+		}
+
+		// if we modify a scalability group section (ie EffectsQuality@2) that matches current value, re-set the group to apply the modified group
+		TSet<FString>* ScalabilitySections = ChangeTracker.ModifiedSectionsPerBranch.Find(*GScalabilityIni);
+		if (ScalabilitySections != nullptr)
+		{
+			for (const FString& SectionName : *ScalabilitySections)
+			{
+				FString Group, Level;
+				if (SectionName.Split(TEXT("@"), &Group, &Level))
+				{
+					FString SGName = FString::Printf(TEXT("sg.%s"), *Group);
+					IConsoleVariable* SGVar = IConsoleManager::Get().FindConsoleVariable(*SGName);
+					if (SGVar != nullptr)
+					{
+						if (SGVar->GetString() == Level)
+						{
+							SGVar->SetWithCurrentPriority(*Level);
+						}
 					}
 				}
 			}
@@ -401,7 +431,7 @@ void OnSetCVarFromIniEntry(const TCHAR *IniFile, const TCHAR *Key, const TCHAR*
 }
 
 
-void ApplyCVarSettingsFromIni(const TCHAR* InSectionName, const TCHAR* InIniFilename, uint32 SetBy, bool bAllowCheating)
+void ApplyCVarSettingsFromIni(const TCHAR* InSectionName, const TCHAR* InIniFilename, uint32 SetBy, bool bAllowCheating, FName Tag)
 {
 	FCoreDelegates::OnApplyCVarFromIni.Broadcast(InSectionName, InIniFilename, SetBy, bAllowCheating);
 
@@ -414,7 +444,7 @@ void ApplyCVarSettingsFromIni(const TCHAR* InSectionName, const TCHAR* InIniFile
 			const FString& KeyString = It.Key().GetPlainNameString(); 
 			const FString& ValueString = It.Value().GetValue();
 
-			OnSetCVarFromIniEntry(InIniFilename, *KeyString, *ValueString, SetBy, bAllowCheating, false);
+			OnSetCVarFromIniEntry(InIniFilename, *KeyString, *ValueString, SetBy, bAllowCheating, false, Tag);
 		}
 	}
 }
diff --git a/Engine/Source/Runtime/Core/Public/Misc/ConfigUtilities.h b/Engine/Source/Runtime/Core/Public/Misc/ConfigUtilities.h
--- a/Engine/Source/Runtime/Core/Public/Misc/ConfigUtilities.h
+++ b/Engine/Source/Runtime/Core/Public/Misc/ConfigUtilities.h
@@ -49,7 +49,7 @@ namespace UE::ConfigUtilities
 	 * @param InIniFilename - The ini filename
 	 * @param SetBy anything in ECVF_LastSetMask e.g. ECVF_SetByScalability
 	 */
-	CORE_API void ApplyCVarSettingsFromIni(const TCHAR* InSectionBaseName, const TCHAR* InIniFilename, uint32 SetBy, bool bAllowCheating=false);
+	CORE_API void ApplyCVarSettingsFromIni(const TCHAR* InSectionBaseName, const TCHAR* InIniFilename, uint32 SetBy, bool bAllowCheating=false, FName Tag=NAME_None);
 
 	/**
 	 * Helper function to operate a user defined function for each CVar key/value pair in the specified section in an ini file
