[RaceDetector] Let VirtualAlloc decide the shadow memory's base address.

This works around a Windows 10 bug where specifying a base address when reserving large regions can cause extreme system-wide performance degradation.

#rnx

diff --git a/Engine/Source/Runtime/Core/Private/Sanitizer/RaceDetectorWindows.cpp b/Engine/Source/Runtime/Core/Private/Sanitizer/RaceDetectorWindows.cpp
--- a/Engine/Source/Runtime/Core/Private/Sanitizer/RaceDetectorWindows.cpp
+++ b/Engine/Source/Runtime/Core/Private/Sanitizer/RaceDetectorWindows.cpp
@@ -848,11 +848,6 @@ namespace UE::Sanitizer::RaceDetector::Platform {
 
 	INSTRUMENTATION_FUNCTION_ATTRIBUTES void InitShadowMemory()
 	{
-		UPTRINT BiggestFreeRegionBase = 0;
-		UPTRINT BiggestFreeRegionSize = 0;
-
-		UPTRINT BaseAddress = 0;
-
 		SYSTEM_INFO si;
 		GetSystemInfo(&si);
 
@@ -860,22 +855,6 @@ namespace UE::Sanitizer::RaceDetector::Platform {
 		PageSizeBitShift = FPlatformMath::CeilLogTwo(PageSize);
 		PageSizeMask = ~((1 << PageSizeBitShift) - 1);
 
-		// Scan for the biggest free address space block we can find
-		MEMORY_BASIC_INFORMATION MemoryBasicInfo;
-		while (VirtualQuery((void*)BaseAddress, &MemoryBasicInfo, sizeof(MemoryBasicInfo)))
-		{
-			if (MemoryBasicInfo.State == MEM_FREE)
-			{
-				if (MemoryBasicInfo.RegionSize > BiggestFreeRegionSize)
-				{
-					BiggestFreeRegionBase = (UPTRINT)MemoryBasicInfo.BaseAddress;
-					BiggestFreeRegionSize = MemoryBasicInfo.RegionSize;
-				}
-			}
-
-			BaseAddress = (UPTRINT)MemoryBasicInfo.BaseAddress + MemoryBasicInfo.RegionSize;
-		}
-
 		// -----------
 		// Low Application Address Space
 		// -----------
@@ -885,10 +864,9 @@ namespace UE::Sanitizer::RaceDetector::Platform {
 		// -----------
 		// We need to split the total address space in a way that we always have
 		// 4 times the shadow address space as the total address space usable by the application.
-		UPTRINT TotalAddressSpace = BaseAddress;
+		UPTRINT TotalAddressSpace = (UPTRINT)si.lpMaximumApplicationAddress - (UPTRINT)si.lpMinimumApplicationAddress;
 
 		// Reserve 7/8 of the space for the shadow so that we leave the application with 1/8 of the total address space
-		ShadowBase = Align(BiggestFreeRegionBase, 1024 * 1024 * 1024);
 		ShadowSize = 7 * (TotalAddressSpace >> 3);
 
 		// Because page faults are super slow when attached with a debugger
@@ -897,21 +875,16 @@ namespace UE::Sanitizer::RaceDetector::Platform {
 		// don't need to call VirtualAlloc for that page.
 		// The required memory for this scheme is a maximum of 4GB support 128TB of address space.
 		ShadowBitmapSize = Align((ShadowSize / PageSize) >> 3, PageSize);
-		ShadowBitmapBase = ShadowBase;
-		ShadowBitmapEnd  = ShadowBitmapBase + ShadowBitmapSize;
 
 		// Because zeroing 4GB is very slow when unmapping shadow.
-		// We keep another 128KB of bits to know which page in the ShadowBitmap have been commited
+		// We keep another 128KB of bits to know which page in the ShadowBitmap have been committed
 		// so we just have to zero those page instead of uncommitting the whole 4GB and recommitting it.
-		DirtyShadowBitmapBase = ShadowBitmapEnd;
 		DirtyShadowBitmapSize = (ShadowBitmapSize / PageSize) >> 3;
-		DirtyShadowBitmapEnd  = DirtyShadowBitmapBase + DirtyShadowBitmapSize;
 
-		// Put the real shadow memory after the bitmaps
-		ShadowBase = DirtyShadowBitmapEnd;
-		ShadowEnd = ShadowBase + ShadowSize;
-
-		if (!VirtualAlloc((LPVOID)ShadowBitmapBase, ShadowSize + ShadowBitmapSize + DirtyShadowBitmapSize, MEM_RESERVE, PAGE_READWRITE))
+		// Let VirtualAlloc decide the best region to reserve. This works around a Windows 10 bug where specifying 
+		// a base address when reserving large regions can cause extreme system-wide performance degradation.
+		ShadowBitmapBase = (UPTRINT)VirtualAlloc(nullptr, ShadowSize + ShadowBitmapSize + DirtyShadowBitmapSize, MEM_RESERVE, PAGE_READWRITE);
+		if (!ShadowBitmapBase)
 		{
 			UE_LOG(LogRaceDetector, Fatal, TEXT("Failed to reserve shadow memory (err: %d)"), GetLastError());
 		}
@@ -920,6 +893,14 @@ namespace UE::Sanitizer::RaceDetector::Platform {
 		{
 			UE_LOG(LogRaceDetector, Fatal, TEXT("Failed to commit shadow memory bitmap (err: %d)"), GetLastError());
 		}
+
+		ShadowBitmapEnd = ShadowBitmapBase + ShadowBitmapSize;
+		DirtyShadowBitmapBase = ShadowBitmapEnd;
+		DirtyShadowBitmapEnd = DirtyShadowBitmapBase + DirtyShadowBitmapSize;
+
+		// Put the real shadow memory after the bitmaps.
+		ShadowBase = DirtyShadowBitmapEnd;
+		ShadowEnd = ShadowBase + ShadowSize;
 	}
 
 	INSTRUMENTATION_FUNCTION_ATTRIBUTES void DirtyShadowBitmapPage(UPTRINT ShadowBitmapAddress)
