Replaced usage of Algo::TransformAccumulate which wasn't using the return value, and will warn when [[nodiscard]] is added.

diff --git a/Engine/Source/Runtime/Core/Private/Memory/CompositeBuffer.cpp b/Engine/Source/Runtime/Core/Private/Memory/CompositeBuffer.cpp
--- a/Engine/Source/Runtime/Core/Private/Memory/CompositeBuffer.cpp
+++ b/Engine/Source/Runtime/Core/Private/Memory/CompositeBuffer.cpp
@@ -50,7 +50,11 @@ FSharedBuffer FCompositeBuffer::ToShared() const &
 		return Segments[0];
 	default:
 		FUniqueBuffer Buffer = FUniqueBuffer::Alloc(GetSize());
-		Algo::TransformAccumulate(Segments, &FSharedBuffer::GetView, Buffer.GetView(), &FMutableMemoryView::CopyFrom);
+		FMutableMemoryView BufferView = Buffer.GetView();
+		for (const FSharedBuffer& SegmentBuffer : Segments)
+		{
+			BufferView = BufferView.CopyFrom(SegmentBuffer.GetView());
+		}
 		return Buffer.MoveToShared();
 	}
 }
