Make -inifile imply -textconfig

diff --git a/Engine/Source/Runtime/Core/Private/Misc/ConfigCacheIni.cpp b/Engine/Source/Runtime/Core/Private/Misc/ConfigCacheIni.cpp
--- a/Engine/Source/Runtime/Core/Private/Misc/ConfigCacheIni.cpp
+++ b/Engine/Source/Runtime/Core/Private/Misc/ConfigCacheIni.cpp
@@ -5765,7 +5765,12 @@ void FConfigCacheIni::InitializeConfigSystem()
 	
 #if PLATFORM_SUPPORTS_BINARYCONFIG && PRELOAD_BINARY_CONFIG && !WITH_EDITOR && !WITH_SERVER_CODE
 	// attempt to load from staged binary config data
-	if (!FParse::Param(FCommandLine::Get(), TEXT("textconfig")) &&
+	const bool bCommandLineRequestsTextConfig = FParse::Param(FCommandLine::Get(), TEXT("textconfig")) // Explicitly requesting textconfig
+#if ALLOW_INI_OVERRIDE_FROM_COMMANDLINE
+		|| (FCString::Stristr(FCommandLine::Get(), CommandlineOverrideSpecifiers::IniFileOverrideIdentifier) != nullptr) // Implicit, requesting a file to be loaded as text from disk
+#endif
+		;
+	if (!bCommandLineRequestsTextConfig &&
 		FConfigCacheIni::CreateGConfigFromSaved(nullptr))
 	{
 		FConfigContext Context = FConfigContext::ForceReloadIntoGConfig();
