CocoaThread - Cleanup useless Block_copy & change the LIFO to a FIFO since it's being use as a FIFO in the first place.

#tests Local testing of order & crashes

diff --git a/Engine/Source/Runtime/Core/Private/Mac/CocoaThread.cpp b/Engine/Source/Runtime/Core/Private/Mac/CocoaThread.cpp
--- a/Engine/Source/Runtime/Core/Private/Mac/CocoaThread.cpp
+++ b/Engine/Source/Runtime/Core/Private/Mac/CocoaThread.cpp
@@ -106,13 +106,7 @@ public:
 	
 	void Process(CFStringRef Mode)
 	{
-		TArray<FCocoaRunLoopTask*> NewTasks;
-		Tasks.PopAll(NewTasks);
-		
-		for (uint32 i = NewTasks.Num(); i > 0; i--)
-		{
-			OutstandingTasks.Add(NewTasks[i - 1]);
-		}
+		Tasks.PopAll(OutstandingTasks);
 		
 		bool bDone = false;
 		while (!bDone)
@@ -233,7 +227,7 @@ private:
 	}
 	
 private:
-	TLockFreePointerListLIFO<FCocoaRunLoopTask> Tasks;
+	TLockFreePointerListFIFO<FCocoaRunLoopTask, 0> Tasks;
 	TArray<FCocoaRunLoopTask*> OutstandingTasks;
 	CFRunLoopRef TargetRunLoop;
 	CFMutableDictionaryRef SourceDictionary;
@@ -408,13 +402,14 @@ thread_local NSString* GThreadInstigatorWaitMode = nullptr;
 
 static void PerformBlockOnThread(FCocoaRunLoopSource& ThreadSource, dispatch_block_t Block, NSArray* SendModes, NSString* WaitMode, bool const bWait)
 {
-	dispatch_block_t CopiedBlock = Block_copy(Block);
 	
 	if ( bWait )
 	{
+		// We need to copy this block because it's going to be within the block in the task.
+		dispatch_block_t CopiedBlock = Block_copy(Block);
 		dispatch_semaphore_t Semaphore = dispatch_semaphore_create(0);
 		NSString* ThreadInstigatorWaitMode = [WaitMode retain];
-		dispatch_block_t ExecuteBlock = Block_copy(^{
+		dispatch_block_t ExecuteBlock = ^{
 			NSString* SavedInstigatorWaitMode = GThreadInstigatorWaitMode;			
 			GThreadInstigatorWaitMode = ThreadInstigatorWaitMode;
 			CopiedBlock();
@@ -424,7 +419,7 @@ static void PerformBlockOnThread(FCocoaRunLoopSource& ThreadSource, dispatch_blo
 			}
 			GThreadInstigatorWaitMode = SavedInstigatorWaitMode;
 			dispatch_semaphore_signal(Semaphore);
-		});
+		};
 		
 		// If the thread instigator wait mode is not already contained in the SendModes, we add it to prevent potential deadlocks.
 		// The idea is that if thread A schedule a block on thread B and while executing that block thread B has to schedule a block on thread A, that block can be executed.
@@ -442,7 +437,7 @@ static void PerformBlockOnThread(FCocoaRunLoopSource& ThreadSource, dispatch_blo
 			CFRunLoopRunInMode((CFStringRef)WaitMode, 0, true);
 		} while (dispatch_semaphore_wait(Semaphore, dispatch_time(0, 100000ull)));
 		
-		Block_release(ExecuteBlock);
+		Block_release(CopiedBlock);
 		dispatch_release(Semaphore);
 	}
 	else
@@ -450,8 +445,6 @@ static void PerformBlockOnThread(FCocoaRunLoopSource& ThreadSource, dispatch_blo
 		ThreadSource.Schedule(Block, SendModes);
 		ThreadSource.Wake();
 	}
-	
-	Block_release(CopiedBlock);
 }
 
 void MainThreadCall(dispatch_block_t Block, NSString* WaitMode, bool const bWait)
