Move ConstexprStringFnv32 hash function into Core and update it to support consistent hashing and 64 bit hashes.

To match other string hashing behavior in the engine, I've changed the algorithm to result in consistent hashes regardless of character width as this is more intuitive and desireable.

#rnx

diff --git a/Engine/Plugins/Experimental/AnimNext/Source/AnimNext/Public/ConstExprUID.h b/Engine/Plugins/Experimental/AnimNext/Source/AnimNext/Public/ConstExprUID.h
--- a/Engine/Plugins/Experimental/AnimNext/Source/AnimNext/Public/ConstExprUID.h
+++ b/Engine/Plugins/Experimental/AnimNext/Source/AnimNext/Public/ConstExprUID.h
@@ -3,55 +3,13 @@
 #pragma once
 
 #include "CoreMinimal.h"
+#include "Hash/Fnv.h"
 
 namespace UE::AnimNext
 {
 	// Type alias for a raw UID, not typesafe
 	using FFConstExprUIDRaw = uint32;
 
-	// Implements a constexpr usable version of FNVa32 for ANSICHAR string literals
-	constexpr uint32 ConstexprStringFnv32(const ANSICHAR* StringLiteral)
-	{
-		constexpr uint32 Offset = 0x811c9dc5;
-		constexpr uint32 Prime = 0x01000193;
-
-		const char* CharPtr = StringLiteral;
-
-		uint32 Fnv = Offset;
-		while (*CharPtr != 0)
-		{
-			Fnv ^= *CharPtr++;
-			Fnv *= Prime;
-		}
-
-		return Fnv;
-	}
-
-	// Implements a constexpr usable version of FNVa32 for WIDECHAR string literals
-	constexpr uint32 ConstexprStringFnv32(const WIDECHAR* StringLiteral)
-	{
-		constexpr uint32 Offset = 0x811c9dc5;
-		constexpr uint32 Prime = 0x01000193;
-
-		const WIDECHAR* CharPtr = StringLiteral;
-
-		uint32 Fnv = Offset;
-		while (*CharPtr != 0)
-		{
-			const uint32 HighByte = (*CharPtr >> 8) & 0xff;
-			Fnv ^= HighByte;
-			Fnv *= Prime;
-
-			const uint32 LowByte = *CharPtr & 0xff;
-			Fnv ^= LowByte;
-			Fnv *= Prime;
-
-			CharPtr++;
-		}
-
-		return Fnv;
-	}
-
 	/**
 	 * A constexpr safe UID
 	 *
@@ -84,7 +42,7 @@ namespace UE::AnimNext
 		// Constructs a UID
 		static constexpr FConstExprUID MakeFromString(const TCHAR* InName) noexcept
 		{
-			return FConstExprUID(ConstexprStringFnv32(InName), InName);
+			return FConstExprUID(UE::HashStringFNV1a32(MakeStringView(InName)), InName);
 		}
 
 #if !UE_BUILD_SHIPPING
diff --git a/Engine/Plugins/Experimental/AnimNextAnimGraph/Source/AnimNextAnimGraph/Public/EvaluationVM/EvaluationVM.h b/Engine/Plugins/Experimental/AnimNextAnimGraph/Source/AnimNextAnimGraph/Public/EvaluationVM/EvaluationVM.h
--- a/Engine/Plugins/Experimental/AnimNextAnimGraph/Source/AnimNextAnimGraph/Public/EvaluationVM/EvaluationVM.h
+++ b/Engine/Plugins/Experimental/AnimNextAnimGraph/Source/AnimNextAnimGraph/Public/EvaluationVM/EvaluationVM.h
@@ -8,6 +8,7 @@
 #include "Animation/AnimCurveFilter.h"
 #include "EvaluationVM/EvaluationFlags.h"
 #include "EvaluationVM/KeyframeState.h"
+#include "Hash/Fnv.h"
 #include "ReferencePose.h"
 #include "Misc/Crc.h"
 
@@ -18,7 +19,7 @@
 	template<> \
 	constexpr uint32 GetTypeID<Type>() \
 	{ \
-		return UE::AnimNext::ConstexprStringFnv32(#Type); \
+		return UE::HashStringFNV1a32(#Type); \
 	}
 
 namespace UE::AnimNext
diff --git a/Engine/Source/Runtime/Core/Public/Hash/Fnv.h b/Engine/Source/Runtime/Core/Public/Hash/Fnv.h
new file mode 100644
--- /dev/null
+++ b/Engine/Source/Runtime/Core/Public/Hash/Fnv.h
@@ -0,0 +1,81 @@
+// Copyright Epic Games, Inc. All Rights Reserved.
+#pragma once
+
+#include "Containers/StringView.h"
+#include "CoreTypes.h"
+#include "HAL/Platform.h"
+
+namespace UE
+{
+	/**
+	 * Implements a constexpr usable version of FNV1a for strings.
+	 * 
+	 * Note: This implementation is 'stable' such that the same string of different character widths will
+	 * hash to the same value (eg. HashStringFNV1a("Hello") == HashStringFNV1a(L"Hello")).
+	 * This behavior is appropriate here as we're operating on strings specifically rather than raw buffers
+	 * and we want consistent hashes between different platforms.
+	 */
+	template<typename HashType, typename CharType>
+	constexpr HashType HashStringFNV1a(TStringView<CharType> String)
+	{
+		static_assert(std::is_same_v<HashType, uint32> || std::is_same_v<HashType, uint64>, "HashType must be an unsigned 32 or 64 bit integer.");
+		static_assert(sizeof(CharType) <= 4, "FNV1a only works with characters up to 32 bits.");
+
+		// See https://en.wikipedia.org/wiki/Fowler%E2%80%93Noll%E2%80%93Vo_hash_function#FNV_hash_parameters
+		constexpr HashType Offset = sizeof(HashType) == 4 ? 0x811c9dc5 : 0xcbf29ce484222325;
+		constexpr HashType Prime = sizeof(HashType) == 4 ? 0x01000193 : 0x100000001b3;
+
+		HashType Fnv = Offset;
+		for (CharType Char : String)
+		{
+			// Operate on every character as if it's 4 bytes.
+			// Characters < 4 bytes will be padded out with zeros.
+			const uint32 Ch = static_cast<uint32>(Char);
+
+			Fnv ^= (Ch >> 24) & 0xff;
+			Fnv *= Prime;
+			Fnv ^= (Ch >> 16) & 0xff;
+			Fnv *= Prime;
+			Fnv ^= (Ch >> 8) & 0xff;
+			Fnv *= Prime;
+			Fnv ^= (Ch) & 0xff;
+			Fnv *= Prime;
+		}
+
+		return Fnv;
+	}
+
+	/* Version taking a string literal. Use this version to force hashing at compile time. */
+	template<typename HashType, typename CharType, int32 Len>
+	UE_CONSTEVAL HashType HashStringFNV1a(const CharType(&StringLiteral)[Len])
+	{
+		return HashStringFNV1a<HashType>(MakeStringView(StringLiteral, Len - 1));
+	}
+
+	/* Creates a 32-bit FNV1a hash for the given string. */
+	template<typename CharType>
+	constexpr uint32 HashStringFNV1a32(TStringView<CharType> String)
+	{
+		return HashStringFNV1a<uint32>(String);
+	}
+
+	template<typename CharType, int32 Len>
+	UE_CONSTEVAL uint32 HashStringFNV1a32(const CharType(&StringLiteral)[Len])
+	{
+		return HashStringFNV1a<uint32>(StringLiteral);
+	}
+
+	/* Creates a 64-bit FNV1a hash for the given string. */
+	template<typename CharType>
+	constexpr uint64 HashStringFNV1a64(TStringView<CharType> String)
+	{
+		return HashStringFNV1a<uint64>(String);
+	}
+
+	template<typename CharType, int32 Len>
+	UE_CONSTEVAL uint64 HashStringFNV1a64(const CharType(&StringLiteral)[Len])
+	{
+		return HashStringFNV1a<uint64>(StringLiteral);
+	}
+
+} // namespace UE
\ No newline at end of file
