Removed assert when using escapes (`) or argument modifiers (|) in localized structured logs since they actually format correclty in UE. It's only in external tools like Horde that they are not ATM but support can be added later.

diff --git a/Engine/Source/Runtime/Core/Private/Logging/StructuredLog.cpp b/Engine/Source/Runtime/Core/Private/Logging/StructuredLog.cpp
--- a/Engine/Source/Runtime/Core/Private/Logging/StructuredLog.cpp
+++ b/Engine/Source/Runtime/Core/Private/Logging/StructuredLog.cpp
@@ -385,12 +385,6 @@ FLogTemplate* FLogTemplate::CreateLocalized(const TCHAR* TextNamespace, const TC
 {
 	using namespace Logging::Private;
 
-	// Disallow escape sequences and argument modifiers.
-	// These are inconsistent with non-localized format strings and are disallowed for now.
-	constexpr FAsciiSet Invalid("`|");
-	checkf(FAsciiSet::HasNone(Format, Invalid),
-		TEXT("Log format does not currently allow escapes (`) or argument modifiers (|). [[%s]]"), Format);
-
 	FTextFormat TextFormat(FText::AsLocalizable_Advanced(TextNamespace, TextKey, Format));
 
 	const bool bFindFields = !!Fields;
