[AutoRTFM] Fold the FWriteLogBumpAllocator into FWriteLog

Reimplement FWriteLog to use blocks of memory, where the write log records (address, size, flags) begin at the start of the block and the write data at the end.
Successive writes move the record and data pointers towards the middle of the block.
The FWriteLog begins with an internal small block. Once a block is full, another block is heap allocated and maintained a linked list.

This approach reduces the baseline per-write size from 16 bytes to 8 bytes (excluding data payload), and runtime performance is fractionally faster.

diff --git a/Engine/Source/Programs/AutoRTFMTests/Private/WriteLogTests.cpp b/Engine/Source/Programs/AutoRTFMTests/Private/WriteLogTests.cpp
new file mode 100644
--- /dev/null
+++ b/Engine/Source/Programs/AutoRTFMTests/Private/WriteLogTests.cpp
@@ -0,0 +1,114 @@
+// Copyright Epic Games, Inc. All Rights Reserved.
+
+#include "AutoRTFM/WriteLog.h"
+#include "Catch2Includes.h"
+#include "HAL/Platform.h"
+#include <cstddef>
+#include <vector>
+
+TEST_CASE("FWriteLog")
+{
+	AutoRTFM::FWriteLog WriteLog;
+
+	auto NewBuffer = [](size_t Size) -> std::vector<std::byte>
+	{
+		std::vector<std::byte> Out(Size);
+		uint64 Value = 0x978c49230979854f;
+		for (size_t I = 0; I < Size; I++)
+		{
+			Out[I] = std::byte{Value & 0xff};
+			Value = (Value * 31) ^ (Value >> 4);
+		}
+		return Out;
+	};
+	
+	auto Check = [&](size_t EntrySize, size_t NumEntries)
+	{
+		// Populate the entries with pseudo random data
+		std::vector<std::byte> Data = NewBuffer(EntrySize * NumEntries);
+		std::vector<AutoRTFM::FWriteLogEntry> Entries;
+		for (size_t I = 0; I < NumEntries; I++)
+		{
+			AutoRTFM::FWriteLogEntry Entry;
+			Entry.Data = &Data[I * EntrySize];
+			Entry.Size = EntrySize;
+			Entry.LogicalAddress = reinterpret_cast<std::byte*>(static_cast<uintptr_t>(0x1234000 + I));
+			Entry.bNoMemoryValidation = (static_cast<int>(Entry.Data[0]) & 0x10) != 0;
+			Entries.push_back(Entry);
+			WriteLog.Push(Entry);
+		}
+
+		REQUIRE(!WriteLog.IsEmpty());
+		REQUIRE(WriteLog.Num() == NumEntries);
+		REQUIRE(WriteLog.TotalSize() == Data.size());
+
+		SECTION("Forwards iterator")
+		{
+			size_t Index = 0;
+			for (AutoRTFM::FWriteLogEntry Got : WriteLog)
+			{
+				const AutoRTFM::FWriteLogEntry& Expect = Entries[Index];
+				REQUIRE(Got.LogicalAddress == Expect.LogicalAddress);
+				REQUIRE(Got.Size == Expect.Size);
+				REQUIRE(memcmp(Got.Data, Expect.Data, Got.Size) == 0);
+				REQUIRE(Got.bNoMemoryValidation == Expect.bNoMemoryValidation);
+				Index++;
+			}
+		}
+
+		SECTION("Reverse iterator")
+		{
+			size_t Index = NumEntries - 1;
+			for (AutoRTFM::FWriteLog::ReverseIterator It = WriteLog.rbegin(); It != WriteLog.rend(); ++It)
+			{
+				AutoRTFM::FWriteLogEntry Got = *It;
+				const AutoRTFM::FWriteLogEntry& Expect = Entries[Index];
+				REQUIRE(Got.LogicalAddress == Expect.LogicalAddress);
+				REQUIRE(Got.Size == Expect.Size);
+				REQUIRE(memcmp(Got.Data, Expect.Data, Got.Size) == 0);
+				REQUIRE(Got.bNoMemoryValidation == Expect.bNoMemoryValidation);
+				Index--;
+			}
+		}
+
+		SECTION("Reset")
+		{
+			WriteLog.Reset();
+			REQUIRE(WriteLog.IsEmpty());
+			REQUIRE(WriteLog.Num() == 0u);
+			REQUIRE(WriteLog.TotalSize() == 0u);
+		}
+	};
+
+	SECTION("Empty")
+	{
+		REQUIRE(WriteLog.IsEmpty());
+		REQUIRE(WriteLog.Num() == 0u);
+		REQUIRE(WriteLog.TotalSize() == 0u);
+	}
+
+	SECTION("EntrySize: 1, NumEntries: 32")
+	{
+		Check(/* EntrySize */ 1, /* NumEntries */ 32);
+	}
+
+	SECTION("EntrySize: 32, NumEntries: 65536")
+	{
+		Check(/* EntrySize */ 32, /* NumEntries */ 65536);
+	}
+
+	SECTION("EntrySize: 1024, NumEntries: 32")
+	{
+		Check(/* EntrySize */ 1024, /* NumEntries */ 32);
+	}
+
+	SECTION("EntrySize: FWriteLogEntry::MaxSize-1, NumEntries: 32")
+	{
+		Check(/* EntrySize */ AutoRTFM::FWriteLogEntry::MaxSize-1, /* NumEntries */ 32);
+	}
+	
+	SECTION("EntrySize: FWriteLogEntry::MaxSize, NumEntries: 32")
+	{
+		Check(/* EntrySize */ AutoRTFM::FWriteLogEntry::MaxSize, /* NumEntries */ 32);
+	}
+}
diff --git a/Engine/Source/Runtime/Core/Private/AutoRTFM/HitSet.h b/Engine/Source/Runtime/Core/Private/AutoRTFM/HitSet.h
--- a/Engine/Source/Runtime/Core/Private/AutoRTFM/HitSet.h
+++ b/Engine/Source/Runtime/Core/Private/AutoRTFM/HitSet.h
@@ -2,16 +2,34 @@
 
 #pragma once
 
-#include "MemoryLocation.h"
 #include "Utils.h"
+#include <cstdint>
 
 namespace AutoRTFM
 {
+	struct FHitSetEntry
+	{
+		// The address of the write.
+		uintptr_t Address : 48;
+		// Size of the write.
+		uintptr_t Size : 15;
+		// If set, then the write shouldn't be considered by the memory validator.
+		uintptr_t bNoMemoryValidation : 1;
+		
+		UE_AUTORTFM_FORCEINLINE uintptr_t Payload() const
+		{
+			// Note: Clang optimizes this away to a reinterpret cast.
+			uintptr_t Result;
+			memcpy(&Result, this, sizeof(Result));
+			return Result;
+		}
+	};
+
+	static_assert(sizeof(FHitSetEntry) == sizeof(uintptr_t));
+	
 	class FHitSet final
 	{
 	public:
-		using Key = FMemoryLocation;
-
 		// Maximum size in bytes for a hit-set write record.
 		// The cutoff here is arbitrarily any number less than UINT16_MAX, but its a
 		// weigh up what a good size is. Because the hitset doesn't detect when you
@@ -30,9 +48,9 @@ namespace AutoRTFM
 		static constexpr uint64_t InitialCapacity = static_cast<uint64_t>(1) << LogInitialCapacity;
 
 		// TODO: Revisit a good max capacity for the hitset.
-		static constexpr uint64_t MaxCapacity = (128 * 1024 * 1024) / sizeof(Key);
-
+		static constexpr uint64_t MaxCapacity = (128 * 1024 * 1024) / sizeof(FHitSetEntry);
 	public:
+
 		~FHitSet()
 		{
 			Reset();
@@ -47,19 +65,19 @@ namespace AutoRTFM
 
 		// Insert something in the HitSet, returning whether the key already
 		// exists, it was inserted, or not.
-		UE_AUTORTFM_FORCEINLINE EInsertResult FindOrTryInsertNoResize(Key K)
+		UE_AUTORTFM_FORCEINLINE EInsertResult FindOrTryInsertNoResize(FHitSetEntry Entry)
 		{
-			return InsertNoResize(K.GetPayload());
+			return InsertNoResize(Entry.Payload());
 		}
 
 		// Insert something in the HitSet, returning true if we found the key.
 		// If false is returned, we could have inserted or not - depending on
 		// whether the HitSet has reached its max capacity.
-		UE_AUTORTFM_FORCEINLINE bool FindOrTryInsert(Key K)
+		UE_AUTORTFM_FORCEINLINE bool FindOrTryInsert(FHitSetEntry Entry)
 		{
 			while (true)
 			{
-				switch (InsertNoResize(K.GetPayload()))
+				switch (InsertNoResize(Entry.Payload()))
 				{
 				default:
 					AutoRTFM::InternalUnreachable();
diff --git a/Engine/Source/Runtime/Core/Private/AutoRTFM/MemoryLocation.h b/Engine/Source/Runtime/Core/Private/AutoRTFM/MemoryLocation.h
deleted file mode 100644
diff --git a/Engine/Source/Runtime/Core/Private/AutoRTFM/TaggedPtr.h b/Engine/Source/Runtime/Core/Private/AutoRTFM/TaggedPtr.h
deleted file mode 100644
diff --git a/Engine/Source/Runtime/Core/Private/AutoRTFM/Transaction.cpp b/Engine/Source/Runtime/Core/Private/AutoRTFM/Transaction.cpp
--- a/Engine/Source/Runtime/Core/Private/AutoRTFM/Transaction.cpp
+++ b/Engine/Source/Runtime/Core/Private/AutoRTFM/Transaction.cpp
@@ -11,6 +11,7 @@
 #include "Logging/LogVerbosity.h"
 #include "Misc/AssertionMacros.h"
 #include "TransactionInlines.h"
+#include "HitSet.h"
 #include "Utils.h"
 
 #include <cstddef>
@@ -57,7 +58,6 @@ void FTransaction::Reset()
 	HitSet.Reset();
 	NewMemoryTracker.Reset();
 	WriteLog.Reset();
-	WriteLogBumpAllocator.Reset();
 	CurrentMemoryValidationLevel = EMemoryValidationLevel::Disabled;
 
 	// Reset to the initial state.
@@ -101,7 +101,7 @@ void FTransaction::AbortWithoutThrowing(EAbortWithoutThrowingFlush State)
 
     // Call the destructors of all the OnCommit functors before undoing the transactional memory and
     // calling the OnAbort callbacks. This is important as the callback functions may have captured
-    // variables that are depending on the allocated memory. 
+    // variables that are depending on the allocated memory.
     CommitTasks.Reset();
 
     Undo();
@@ -173,32 +173,27 @@ void FTransaction::Undo()
 	int Num = WriteLog.Num();
 	for(auto Iter = WriteLog.rbegin(); Iter != WriteLog.rend(); ++Iter)
     {
-		FWriteLogEntry& Entry = *Iter;
-		void* const Original = Entry.GetOriginal();
-
+		FWriteLogEntry Entry = *Iter;
         // No write records should be within the transaction's stack range.
-        ensure(!IsOnStack(Original));
-
-		const size_t Size = Entry.GetSize();
-        void* const Copy = Entry.GetCopy();
+        ensure(!IsOnStack(Entry.LogicalAddress));
 
 		if (UE_LOG_ACTIVE(LogAutoRTFM, Verbose))
 		{
 			TStringBuilder<1024> Builder;
 
-			Builder.Appendf(TEXT("%4d [UNDO] %p %4llu : [ "), Num - VerboseCounter - 1, Original, Size);
+			Builder.Appendf(TEXT("%4d [UNDO] %p %4llu : [ "), Num - VerboseCounter - 1, Entry.LogicalAddress, Entry.Size);
 
-			unsigned char* Current = (unsigned char*)Original;
-			unsigned char* Old = (unsigned char*)Copy;
+			std::byte* Current = Entry.LogicalAddress;
+			std::byte* Old = Entry.Data;
 
-			for (size_t i = 0; i < Size; i++)
+			for (size_t i = 0; i < Entry.Size; i++)
 			{
 				Builder.Appendf(TEXT("%02X "), Current[i]);
 			}
 
 			Builder << TEXT("] -> [ ");
 
-			for (size_t i = 0; i < Size; i++)
+			for (size_t i = 0; i < Entry.Size; i++)
 			{
 				Builder.Appendf(TEXT("%02X "), Old[i]);
 			}
@@ -208,7 +203,7 @@ void FTransaction::Undo()
 			UE_LOG(LogAutoRTFM, Verbose, TEXT("%s"), Builder.ToString());
 		}
 
-        memcpy(Original, Copy, Size);
+        memcpy(Entry.LogicalAddress, Entry.Data, Entry.Size);
 
 		VerboseCounter++;
     }
@@ -223,18 +218,19 @@ void FTransaction::CommitNested()
 	// We need to pass our write log to our parent transaction, but with care!
 	// We need to discard any writes if the memory location is on the parent
 	// transaction's stack range.
-	for (FWriteLogEntry& Write : WriteLog)
+	for (FWriteLogEntry Write : WriteLog)
 	{
-		if (Parent->IsOnStack(Write.GetOriginal()))
+		if (Parent->IsOnStack(Write.LogicalAddress))
 		{
 			continue;
 		}
 
-		if (Write.GetSize() <= FHitSet::MaxSize)
+		if (Write.Size <= FHitSet::MaxSize)
 		{
-			FHitSet::Key HitSetEntry(Write.GetOriginal());
-			HitSetEntry.TopTag().Size = static_cast<uint16_t>(Write.GetSize());
-			HitSetEntry.TopTag().bNoMemoryValidation = Write.IgnoreForValidation();
+			FHitSetEntry HitSetEntry{};
+			HitSetEntry.Address = reinterpret_cast<uintptr_t>(Write.LogicalAddress);
+			HitSetEntry.Size = static_cast<uint16_t>(Write.Size);
+			HitSetEntry.bNoMemoryValidation = Write.bNoMemoryValidation;
 
 			if (Parent->HitSet.FindOrTryInsert(HitSetEntry))
 			{
@@ -243,10 +239,14 @@ void FTransaction::CommitNested()
 		}
 
 		Parent->WriteLog.Push(Write);
-	}
 
+		FHitSetEntry Entry{};
+		Entry.Address = reinterpret_cast<uintptr_t>(Write.LogicalAddress);
+		Entry.Size = static_cast<uint16_t>(Write.Size);
+		Entry.bNoMemoryValidation = Write.bNoMemoryValidation;
 
-    Parent->WriteLogBumpAllocator.Merge(MoveTemp(WriteLogBumpAllocator));
+		Parent->HitSet.FindOrTryInsert(Entry);
+	}
 
 	// For all the deferred calls to `PopOnCommitHandler` that we couldn't
 	// process (because our transaction nest didn't `PushOnCommitHandler`)
@@ -300,12 +300,12 @@ bool FTransaction::AttemptToCommitOuterNest()
 	UE_LOG(LogAutoRTFM, Verbose, TEXT("Running commit tasks..."));
 
     AbortTasks.Reset();
-	
+
 	CommitTasks.DeleteKeyedIf([&](const void* const Key) { return this->IsOnStack(Key); });
 
     CommitTasks.ForEachForward([] (TFunction<void()>& Task)
     { 
-        Task(); 
+        Task();
         Task.Reset();
     });
 
@@ -511,26 +511,23 @@ FTransaction::TWriteHash FTransaction::CalculateNestedWriteHash(size_t NumWriteE
 		}
 
 		const FWriteLogEntry& Entry = *Iter;
-		if (UNLIKELY(Entry.IgnoreForValidation()))
+		if (UNLIKELY(Entry.bNoMemoryValidation))
 		{
 			continue; // Next write
 		}
 
-		const uint8* Original = reinterpret_cast<const uint8*>(Entry.GetOriginal());
-		const size_t Size = Entry.GetSize();
-		if (Size > sizeof(FSmallValue))
+		if (Entry.Size > sizeof(FSmallValue))
 		{
-			Hasher.Update(Original, Size);
-			continue; // Next write
+			Hasher.Update(Entry.LogicalAddress, Entry.Size);
 		}
 
-		const uintptr_t Address = reinterpret_cast<uintptr_t>(Original);
-		switch (Size)
+		const uintptr_t Address = reinterpret_cast<uintptr_t>(Entry.LogicalAddress);
+		switch (Entry.Size)
 		{
 			case 8:
 				if (LIKELY((Address & 7) == 0))
 				{
-					uint64 Load = *reinterpret_cast<uint64 const*>(Original);
+					uint64 Load = *reinterpret_cast<uint64 const*>(Entry.LogicalAddress);
 					SmallValue = (SmallValue * 31) ^ static_cast<FSmallValue>(Load);
 				}
 				break;
@@ -538,7 +535,7 @@ FTransaction::TWriteHash FTransaction::CalculateNestedWriteHash(size_t NumWriteE
 			case 4:
 				if (LIKELY((Address & 3) == 0))
 				{
-					uint32 Load = *reinterpret_cast<uint32 const*>(Original);
+					uint32 Load = *reinterpret_cast<uint32 const*>(Entry.LogicalAddress);
 					SmallValue = (SmallValue * 31) ^ static_cast<FSmallValue>(Load);
 				}
 				break;
@@ -546,20 +543,20 @@ FTransaction::TWriteHash FTransaction::CalculateNestedWriteHash(size_t NumWriteE
 			case 2:
 				if (LIKELY((Address & 1) == 0))
 				{
-					uint16 Load = *reinterpret_cast<uint16 const*>(Original);
+					uint16 Load = *reinterpret_cast<uint16 const*>(Entry.LogicalAddress);
 					SmallValue = (SmallValue * 31) ^ static_cast<FSmallValue>(Load);
 				}
 				break;
 
 			case 1:
-				SmallValue = (SmallValue * 31) ^ static_cast<FSmallValue>(Original[0]);
+				SmallValue = (SmallValue * 31) ^ static_cast<FSmallValue>(Entry.LogicalAddress[0]);
 				break;
 
 			default:
 				// Unaligned / non-POT write.
-				for (size_t I = 0; I < Size; I++)
+				for (size_t I = 0; I < Entry.Size; I++)
 				{
-					SmallValue = (SmallValue * 31) ^ static_cast<FSmallValue>(Original[I]);
+					SmallValue = (SmallValue * 31) ^ static_cast<FSmallValue>(Entry.LogicalAddress[I]);
 				}
 				break;
 		}
diff --git a/Engine/Source/Runtime/Core/Private/AutoRTFM/Transaction.h b/Engine/Source/Runtime/Core/Private/AutoRTFM/Transaction.h
--- a/Engine/Source/Runtime/Core/Private/AutoRTFM/Transaction.h
+++ b/Engine/Source/Runtime/Core/Private/AutoRTFM/Transaction.h
@@ -12,7 +12,6 @@
 #include "Stats.h"
 #include "TaskArray.h"
 #include "WriteLog.h"
-#include "WriteLogBumpAllocator.h"
 
 namespace AutoRTFM
 {
@@ -119,7 +118,6 @@ public:
 
     // Record that a write is about to occur at the given LogicalAddress of Size bytes.
     void RecordWrite(void* LogicalAddress, size_t Size, bool bNoMemoryValidation = false);
-    void RecordWriteMaxPageSized(void* LogicalAddress, size_t Size, bool bNoMemoryValidation);
     template<unsigned SIZE> void RecordWrite(void* LogicalAddress);
 	template<unsigned SIZE> void RecordWriteInsertedSlow(void* LogicalAddress);
 	template<unsigned SIZE> void RecordWriteNotInsertedSlow(void* LogicalAddress);
@@ -219,7 +217,6 @@ private:
     FHitSet HitSet;
     FIntervalTree NewMemoryTracker;
     FWriteLog WriteLog;
-    FWriteLogBumpAllocator WriteLogBumpAllocator;
     TStatStorage<uint64_t> StatDepth = 1;
     TWriteHash RecordedWriteHash = 0;
     size_t NumWriteLogsHashed = 0;
diff --git a/Engine/Source/Runtime/Core/Private/AutoRTFM/TransactionInlines.h b/Engine/Source/Runtime/Core/Private/AutoRTFM/TransactionInlines.h
--- a/Engine/Source/Runtime/Core/Private/AutoRTFM/TransactionInlines.h
+++ b/Engine/Source/Runtime/Core/Private/AutoRTFM/TransactionInlines.h
@@ -3,8 +3,8 @@
 #pragma once
 
 #include "Context.h"
-#include "MemoryLocation.h"
 #include "ScopedGuard.h"
+#include "HitSet.h"
 #include "Transaction.h"
 
 #include "HAL/Platform.h"
@@ -33,7 +33,7 @@ namespace AutoRTFM
 {
 
 UE_AUTORTFM_FORCEINLINE bool FTransaction::IsOnStack(const void* LogicalAddress) const
-{ 
+{
     if (StackRange.Contains(LogicalAddress))
     {
         return true;
@@ -72,7 +72,7 @@ UE_AUTORTFM_FORCEINLINE bool FTransaction::ShouldRecordWrite(void* LogicalAddres
     // scope that encompasses the transaction. Instead a non-scoped transaction
     // is started with a call to StartTransaction() and ended with a call to
     // either AbortTransaction() or CommitTransaction(). Unlike a
-    // scoped transaction, there's no precise stack range for a non-scoped 
+    // scoped transaction, there's no precise stack range for a non-scoped
     // transaction, as the scope can freely grow or shrink between the calls to
     // [Start|Abort|Commit]Transaction() and any recorded writes. The only
     // guarantee we have is that a non-scoped transaction cannot shrink past the
@@ -88,57 +88,45 @@ UE_AUTORTFM_FORCEINLINE bool FTransaction::ShouldRecordWrite(void* LogicalAddres
     // ignored (they have the same constrained visibility as stack variables in
     // a scoped transaction).
 
-    // Hitting this assert? 
-    // Consider moving the variable being written to an inner scoped 
-    // transaction, or move the variable outside of the nearest parent 
+    // Hitting this assert?
+    // Consider moving the variable being written to an inner scoped
+    // transaction, or move the variable outside of the nearest parent
     // scoped-transaction.
-    ASSERT(bIsStackScoped || LogicalAddress < FContext::Get()->GetClosedStackAddress()); 
+    ASSERT(bIsStackScoped || LogicalAddress < FContext::Get()->GetClosedStackAddress());
 
     return false;
 }
 
-AUTORTFM_NO_ASAN UE_AUTORTFM_FORCEINLINE
-void FTransaction::RecordWriteMaxPageSized(void* LogicalAddress, size_t Size, bool bNoMemoryValidation)
-{
-    void* CopyAddress = WriteLogBumpAllocator.Allocate(Size);
-    memcpy(CopyAddress, LogicalAddress, Size);
-
-    FMemoryLocation OriginalAndFlags(LogicalAddress);
-    FMemoryLocationTag& Tag = OriginalAndFlags.TopTag();
-    Tag.bNoMemoryValidation = bNoMemoryValidation;
-    Tag.Size = static_cast<uint16_t>(Size);
-    WriteLog.Push(FWriteLogEntry(OriginalAndFlags, CopyAddress));
-}
 
 AUTORTFM_NO_ASAN UE_AUTORTFM_FORCEINLINE
 void FTransaction::RecordWrite(void* LogicalAddress, size_t Size, bool bNoMemoryValidation /* = false */)
 {
-    if (UNLIKELY(0 == Size))
-    {
-        return;
-    }
+	if (UNLIKELY(0 == Size))
+	{
+		return;
+	}
 
-    if (!ShouldRecordWrite(LogicalAddress))
-    {
-        Stats.Collect<EStatsKind::HitSetSkippedBecauseOfStackLocalMemory>();
-        return;
-    }
+	if (!ShouldRecordWrite(LogicalAddress))
+	{
+		Stats.Collect<EStatsKind::HitSetSkippedBecauseOfStackLocalMemory>();
+		return;
+	}
 
-    if (Size <= FHitSet::MaxSize)
-    {
-        FMemoryLocation Key(LogicalAddress);
-        FMemoryLocationTag& Tag = Key.TopTag();
-        Tag.bNoMemoryValidation = bNoMemoryValidation;
-        Tag.Size = static_cast<uint16_t>(Size);
-
-        if (HitSet.FindOrTryInsert(Key))
-        {
-            Stats.Collect<EStatsKind::HitSetHit>();
-            return;
-        }
-    
-        Stats.Collect<EStatsKind::HitSetMiss>();
-    }
+	if (Size <= FHitSet::MaxSize)
+	{
+		FHitSetEntry HitSetEntry{};
+		HitSetEntry.Address = reinterpret_cast<uintptr_t>(LogicalAddress);
+		HitSetEntry.Size = static_cast<uint16_t>(Size);
+		HitSetEntry.bNoMemoryValidation = bNoMemoryValidation;
+
+		if (HitSet.FindOrTryInsert(HitSetEntry))
+		{
+			Stats.Collect<EStatsKind::HitSetHit>();
+			return;
+		}
+
+		Stats.Collect<EStatsKind::HitSetMiss>();
+	}
 
 	if (NewMemoryTracker.Contains(LogicalAddress, Size))
 	{
@@ -148,17 +136,23 @@ void FTransaction::RecordWrite(void* LogicalAddress, size_t Size, bool bNoMemory
 
 	Stats.Collect<EStatsKind::NewMemoryTrackerMiss>();
 
-    uint8_t* const Address = reinterpret_cast<uint8_t*>(LogicalAddress);
-
-    size_t I = 0;
-
-    for (; (I + FWriteLogBumpAllocator::MaxSize) < Size; I += FWriteLogBumpAllocator::MaxSize)
-    {
-        RecordWriteMaxPageSized(Address + I, FWriteLogBumpAllocator::MaxSize, bNoMemoryValidation);
-    }
-
-    // Remainder at the end of the memcpy.
-    RecordWriteMaxPageSized(Address + I, Size - I, bNoMemoryValidation);
+	FWriteLogEntry WriteLogEntry{};
+	WriteLogEntry.LogicalAddress = static_cast<std::byte*>(LogicalAddress);
+	WriteLogEntry.Data = static_cast<std::byte*>(LogicalAddress);
+	WriteLogEntry.bNoMemoryValidation = bNoMemoryValidation;
+	if (UNLIKELY(Size > FWriteLogEntry::MaxSize))
+	{
+		WriteLogEntry.Size = FWriteLogEntry::MaxSize;
+		while (Size > FWriteLogEntry::MaxSize)
+		{
+			WriteLog.Push(WriteLogEntry);
+			WriteLogEntry.LogicalAddress += FWriteLogEntry::MaxSize;
+			WriteLogEntry.Data += FWriteLogEntry::MaxSize;
+			Size -= FWriteLogEntry::MaxSize;
+		}
+	}
+	WriteLogEntry.Size = Size;
+	WriteLog.Push(WriteLogEntry);
 }
 
 template<unsigned SIZE> AUTORTFM_NO_ASAN UE_AUTORTFM_FORCEINLINE void FTransaction::RecordWrite(void* LogicalAddress)
@@ -171,10 +165,11 @@ template<unsigned SIZE> AUTORTFM_NO_ASAN UE_AUTORTFM_FORCEINLINE void FTransacti
         return;
     }
 
-    FMemoryLocation Key(LogicalAddress);
-    Key.TopTag().Size = static_cast<uint16_t>(SIZE);
+    FHitSetEntry Entry{};
+    Entry.Address = reinterpret_cast<uintptr_t>(LogicalAddress);
+    Entry.Size = static_cast<uint16_t>(SIZE);
 
-	switch (HitSet.FindOrTryInsertNoResize(Key))
+	switch (HitSet.FindOrTryInsertNoResize(Entry))
 	{
 	case FHitSet::EInsertResult::Exists:
 		Stats.Collect<EStatsKind::HitSetHit>();
@@ -188,10 +183,11 @@ template<unsigned SIZE> AUTORTFM_NO_ASAN UE_AUTORTFM_FORCEINLINE void FTransacti
 
 template<unsigned SIZE> AUTORTFM_NO_ASAN FORCENOINLINE void FTransaction::RecordWriteNotInsertedSlow(void* LogicalAddress)
 {
-	FMemoryLocation Key(LogicalAddress);
-	Key.TopTag().Size = static_cast<uint16_t>(SIZE);
+    FHitSetEntry Entry{};
+    Entry.Address = reinterpret_cast<uintptr_t>(LogicalAddress);
+    Entry.Size = static_cast<uint16_t>(SIZE);
 
-	if (HitSet.FindOrTryInsert(Key))
+	if (HitSet.FindOrTryInsert(Entry))
 	{
 		Stats.Collect<EStatsKind::HitSetHit>();
 		return;
@@ -212,7 +208,12 @@ template<unsigned SIZE> AUTORTFM_NO_ASAN FORCENOINLINE void FTransaction::Record
 
 	Stats.Collect<EStatsKind::NewMemoryTrackerMiss>();
 
-	WriteLog.Push(FWriteLogEntry::CreateSmall<SIZE>(LogicalAddress));
+	FWriteLogEntry WriteLogEntry{};
+	WriteLogEntry.LogicalAddress = static_cast<std::byte*>(LogicalAddress);
+	WriteLogEntry.Data = static_cast<std::byte*>(LogicalAddress);
+	WriteLogEntry.Size = SIZE;
+
+	WriteLog.Push(WriteLogEntry);
 }
 
 UE_AUTORTFM_FORCEINLINE void FTransaction::DidAllocate(void* LogicalAddress, const size_t Size)
@@ -313,8 +314,8 @@ UE_AUTORTFM_FORCEINLINE void FTransaction::CollectStats() const
     Stats.Collect<EStatsKind::AverageWriteLogEntries>(WriteLog.Num());
     Stats.Collect<EStatsKind::MaximumWriteLogEntries>(WriteLog.Num());
 
-    Stats.Collect<EStatsKind::AverageWriteLogBytes>(WriteLogBumpAllocator.StatTotalSize);
-    Stats.Collect<EStatsKind::MaximumWriteLogBytes>(WriteLogBumpAllocator.StatTotalSize);
+    Stats.Collect<EStatsKind::AverageWriteLogBytes>(WriteLog.TotalSize());
+    Stats.Collect<EStatsKind::MaximumWriteLogBytes>(WriteLog.TotalSize());
 
     Stats.Collect<EStatsKind::AverageCommitTasks>(CommitTasks.Num());
     Stats.Collect<EStatsKind::MaximumCommitTasks>(CommitTasks.Num());
diff --git a/Engine/Source/Runtime/Core/Private/AutoRTFM/WriteLog.h b/Engine/Source/Runtime/Core/Private/AutoRTFM/WriteLog.h
--- a/Engine/Source/Runtime/Core/Private/AutoRTFM/WriteLog.h
+++ b/Engine/Source/Runtime/Core/Private/AutoRTFM/WriteLog.h
@@ -2,241 +2,351 @@
 
 #pragma once
 
-#include "MemoryLocation.h"
+#include "AutoRTFM/AutoRTFM.h"
+#include "HAL/Platform.h"
 #include "Utils.h"
+#include <cstddef>
+#include <cstdint>
+#include <cstring>
 #include <stdint.h>
-#include <type_traits>
 
 namespace AutoRTFM
 {
-	class FWriteLogEntry final
+	struct FWriteLogEntry final
 	{
-		FMemoryLocation OriginalAndFlags;
-		uintptr_t Copy;
+		// Number of bits used by the FWriteLog to represent a write's size.
+		static constexpr size_t SizeBits = 15;
 
-	public:
-		UE_AUTORTFM_FORCEINLINE FWriteLogEntry() = default;
-		UE_AUTORTFM_FORCEINLINE FWriteLogEntry(const FWriteLogEntry&) = default;
-		UE_AUTORTFM_FORCEINLINE FWriteLogEntry& operator=(const FWriteLogEntry&) = default;
+		// The maximum size for a single write log entry.
+		// Split into multiple entries if the write is too large.
+		static constexpr size_t MaxSize = (1u << SizeBits) - 1;
 
-		UE_AUTORTFM_FORCEINLINE FWriteLogEntry(void* Original, size_t Size, void* Copy) :
-			OriginalAndFlags(Original), Copy(reinterpret_cast<uintptr_t>(Copy))
-		{
-			OriginalAndFlags.TopTag().Size = static_cast<uint16_t>(Size);
-		}
+		// The address of the write.
+		std::byte* LogicalAddress = nullptr;
+		
+		// A pointer to the original data before the write occurred.
+		std::byte* Data = nullptr;
+		
+		// The size of the write in bytes. Must be smaller than MaxSize.
+		// If the write exceeds MaxSize, then the write must be split into
+		// multiple entries.
+		size_t Size = 0;
+		
+		// If true, then this write will not be considered by the AutoRTFM
+		// memory validator.
+		bool bNoMemoryValidation = false;
+	};
 
-		UE_AUTORTFM_FORCEINLINE FWriteLogEntry(FMemoryLocation OriginalAndFlags, void* Copy) :
-			OriginalAndFlags(OriginalAndFlags), Copy(reinterpret_cast<uintptr_t>(Copy))
+	// FWriteLog holds an ordered list of write records which can be iterated
+	// forwards and backwards.
+	class FWriteLog final
+	{
+		struct TRecord
 		{
-		}
-		
-		template<unsigned int SIZE> UE_AUTORTFM_FORCEINLINE static FWriteLogEntry CreateSmall(void* Original)
+			uintptr_t Address             : 48;
+			uintptr_t bNoMemoryValidation : 1;
+			uintptr_t Size                : 15;
+		};
+		static_assert(sizeof(uintptr_t) == 8, "assumption: a pointer is 8 bytes");
+		static_assert(sizeof(TRecord) == 8);
+
+		struct FBlock final
 		{
-			static_assert(SIZE <= 8);
+			// ┌────────┬────┬────┬────┬────┬────────────────┬────┬────┬────┬────┐
+			// │ FBlock │ R₀ │ R₁ │ R₂ │ R₃ │->            <-│ D₃ │ D₂ │ D₁ │ D₀ │
+			// └────────┴────┴────┴────┴────┴────────────────┴────┴────┴────┴────┘
+			//          ^                   ^                ^                   ^
+			//    Head().Record       Tail().Record    Tail().Data         Head().Data
+			// Where:
+			//   Rₙ = Record n, Dₙ = Data n
+
+			// Size of a heap-allocated block, excluding the FBlock struct header.
+			static constexpr size_t DefaultRecordsAndDataSize = 2048;
+
+			// Pointers to a write log record and its corresponding data.
+			struct FEntryPtrs
+			{
+				TRecord* Record = nullptr;
+				std::byte* Data = nullptr;
+			};
+
+			// Constructor
+			// RecordsAndDataSize is the total size of the block, excluding the FBlock struct header
+			// (the memory size for the write records and data)
+			FBlock(size_t RecordsAndDataSize) 
+				: RecordsAndDataSize(RecordsAndDataSize)
+			{ 
+				TailPtrs = Head();
+			}
 
-			FWriteLogEntry Entry;
+			// Allocate performs a heap allocation of a new block.
+			// RecordsAndDataSize is the total size of the block, excluding the FBlock struct header
+			// (the memory size for the write records and data)
+			static FBlock* Allocate(size_t RecordsAndDataSize)
+			{
+				std::byte* Memory = new std::byte[sizeof(FBlock) + RecordsAndDataSize];
+				return new (Memory) FBlock(RecordsAndDataSize);
+			}
 
-			Entry.OriginalAndFlags = Original;
+			// Free releases the heap-allocated memory for this block.
+			// Note: This block must have been allocated with a call to Allocate().
+			void Free()
+			{
+				delete [] reinterpret_cast<std::byte*>(this);
+			}
 
-			// We store the 8-byte or less write log entry in the write log itself since it is small.
-			uintptr_t Temp = 0;
-			memcpy(&Temp, Original, SIZE);
-			Entry.Copy = Temp;
+			// Returns the address of the first record.
+			UE_AUTORTFM_FORCEINLINE TRecord* FirstRecord() { return reinterpret_cast<TRecord*>(reinterpret_cast<std::byte*>(this) + sizeof(FBlock)); }
+			// Returns the address of the first record.
+			UE_AUTORTFM_FORCEINLINE TRecord const * FirstRecord() const { return reinterpret_cast<const TRecord*>(reinterpret_cast<std::byte const *>(this) + sizeof(FBlock)); }
 
-			FMemoryLocationTag& Tag = Entry.OriginalAndFlags.TopTag();
-			Tag.Size = SIZE;
-			Tag.bIsSmall = 1;
+			// Returns true if the block has no entries.
+			UE_AUTORTFM_FORCEINLINE bool IsEmpty()
+			{
+				return FirstRecord() == TailPtrs.Record;
+			}
 
-			ASSERT(Tag.Size <= FMemoryLocationTag::MaxSize);
+			// Attempts to add the entry into this block by copying the entry's data and creating a
+			// new record.
+			// Returns true if the entry was added, or false if the block does not have the capacity
+			// for the entry.
+			UE_AUTORTFM_FORCEINLINE bool Push(FWriteLogEntry Entry)
+			{
+				AssertIsValid(TailPtrs);
 
-			return Entry;
-		}
+				if (reinterpret_cast<std::byte*>(TailPtrs.Record + 1) > TailPtrs.Data - Entry.Size)
+				{
+					return false;
+				}
 
-		UE_AUTORTFM_FORCEINLINE void* GetOriginal() const
-		{
-			return OriginalAndFlags.Get();
-		}
+				TailPtrs.Record->Address = reinterpret_cast<uintptr_t>(Entry.LogicalAddress);
+				TailPtrs.Record->Size = Entry.Size;
+				TailPtrs.Record->bNoMemoryValidation = Entry.bNoMemoryValidation;
+				TailPtrs.Record++;
 
-		UE_AUTORTFM_FORCEINLINE void* GetCopy()
-		{
-			if (IsSmall())
-			{
-				return reinterpret_cast<void*>(&Copy);
+				TailPtrs.Data -= Entry.Size;
+				memcpy(TailPtrs.Data, Entry.Data, Entry.Size);
+				AssertIsValid(TailPtrs);
+				return true;
 			}
-			else
+
+			// Returns the record and data pointers to the first entry in the block.
+			UE_AUTORTFM_FORCEINLINE FEntryPtrs Head()
 			{
-				return reinterpret_cast<void*>(Copy);
+				FEntryPtrs Ptrs;
+				Ptrs.Record = FirstRecord();
+				Ptrs.Data = reinterpret_cast<std::byte*>(FirstRecord()) + RecordsAndDataSize;
+				return Ptrs;
 			}
-		}
 
-		UE_AUTORTFM_FORCEINLINE const void* GetCopy() const
-		{
-			if (IsSmall())
+			// Returns the record and data pointers to the last entry in the block.
+			UE_AUTORTFM_FORCEINLINE FEntryPtrs Tail()
 			{
-				return reinterpret_cast<const void*>(&Copy);
+				return TailPtrs;
 			}
-			else
+
+			// Asserts tha the pointers in Ptrs are within bounds for this block.
+			UE_AUTORTFM_FORCEINLINE void AssertIsValid(FEntryPtrs Ptrs) const
 			{
-				return reinterpret_cast<const void*>(Copy);
+#if UE_BUILD_DEBUG
+				ASSERT(FirstRecord() <= Ptrs.Record);
+				ASSERT(reinterpret_cast<const std::byte*>(Ptrs.Record) <= Ptrs.Data);
+				ASSERT(Ptrs.Data <= reinterpret_cast<const std::byte*>(FirstRecord()) + RecordsAndDataSize);
+#endif
 			}
-		}
 
-		UE_AUTORTFM_FORCEINLINE size_t GetSize() const
+			// The size in bytes of the block's records and data (total size - sizeof(FBlock)).
+			size_t const RecordsAndDataSize;
+			// The next block in the linked list.
+			FBlock* NextBlock = nullptr;
+			// The previous block in the linked list.
+			FBlock* PrevBlock = nullptr;
+			// The pointers to one-past the last entry in the block.
+			FEntryPtrs TailPtrs;
+
+		private:
+			~FBlock() = delete;
+		};
+
+	public:
+		// Constructor
+		FWriteLog()
 		{
-			return OriginalAndFlags.TopTag().Size;
+			new(HeadBlockMemory) FBlock(HeadBlockSize - sizeof(FBlock));
 		}
 
-		UE_AUTORTFM_FORCEINLINE bool IsSmall() const
+		// Destructor
+		~FWriteLog()
 		{
-			return OriginalAndFlags.TopTag().bIsSmall;
+			Reset();
 		}
 
-		UE_AUTORTFM_FORCEINLINE bool IgnoreForValidation() const
+		// Adds the write log entry to the log.
+		// The log will make a copy of the FWriteLogEntry's data.
+		void Push(FWriteLogEntry Entry)
 		{
-			return OriginalAndFlags.TopTag().bNoMemoryValidation;
+			ASSERT(Entry.Size <= FWriteLogEntry::MaxSize);
+			ASSERT((reinterpret_cast<uintptr_t>(Entry.LogicalAddress) & 0xffff0000'00000000) == 0);
+
+			if (UNLIKELY(!TailBlock->Push(Entry)))
+			{
+				const size_t RecordsAndDataSize = std::max(Entry.Size + sizeof(TRecord), FBlock::DefaultRecordsAndDataSize);
+				FBlock* NewBlock = FBlock::Allocate(RecordsAndDataSize);
+				NewBlock->PrevBlock = TailBlock;
+				TailBlock->NextBlock = NewBlock;
+				TailBlock = NewBlock;
+
+				const bool bAdded = NewBlock->Push(Entry);
+				ASSERT(bAdded);
+			}
+
+			NumEntries++;
+			TotalSizeBytes += Entry.Size;
 		}
-	};
 
-    class FWriteLog final
-    {
-        struct FWriteLogEntryBucket final
-        {
-            static constexpr uint32_t BucketSize = 128;
-
-			// Specify a constructor so that `Entries` below will be left uninitialized.
-			explicit FWriteLogEntryBucket() {}
-
-            FWriteLogEntry Entries[BucketSize];
-            size_t Size = 0;
-            FWriteLogEntryBucket* Next = nullptr;
-            FWriteLogEntryBucket* Prev = nullptr;
-        };
-
-    public:
-        void Push(FWriteLogEntry Entry)
-        {
-            if (nullptr == Start)
-            {
-                Start = new FWriteLogEntryBucket();
-                Current = Start;
-            }
-
-            if (FWriteLogEntryBucket::BucketSize == Current->Size)
-            {
-                Current->Next = new FWriteLogEntryBucket();
-				Current->Next->Prev = Current;
-                Current = Current->Next;
-            }
-
-            Current->Entries[Current->Size++] = Entry;
-
-            TotalSize++;
-            return;
-        }
-
-        ~FWriteLog()
-        {
-            Reset();
-        }
-
-		template<bool IS_FORWARD, bool IS_CONST>
+		// Iterator for enumerating the writes of the log.
+		template<bool IS_FORWARD>
 		struct TIterator final
 		{
-			using BucketType = std::conditional_t<IS_CONST, const FWriteLogEntryBucket, FWriteLogEntryBucket>;
-			using EntryType = std::conditional_t<IS_CONST, const FWriteLogEntry, FWriteLogEntry>;
+			TIterator() = default;
 
-			explicit TIterator(BucketType* Bucket) : Bucket(Bucket) {}
+			TIterator(FBlock* Block) : Block(Block)
+			{
+				Ptrs = IS_FORWARD ? Block->Head() : Block->Tail();
+				Progress();
+			}
 
-			EntryType& operator*() const
+			// Returns the entry at the current iterator's position.
+			FWriteLogEntry operator*() const
 			{
-				return IS_FORWARD ? Bucket->Entries[Offset] : Bucket->Entries[Bucket->Size - Offset - 1];
+				return Current;
 			}
 
+			// Progresses the iterator to the next entry
 			void operator++()
 			{
-				Offset++;
-
-				if (Offset == Bucket->Size)
-				{
-					Bucket = IS_FORWARD ? Bucket->Next : Bucket->Prev;
-					Offset = 0;
-				}
+				Progress();
 			}
 
+			// Inequality operator
 			bool operator!=(const TIterator& Other) const
 			{
-				return (Other.Bucket != Bucket) || (Other.Offset != Offset);
+				return (Other.Block != Block) || (Other.Ptrs.Record != Ptrs.Record);
 			}
 
 		private:
-			BucketType* Bucket;
-			size_t Offset = 0;
-		};
+			// Assigns the current write log entry to Current and progresses the iterator to the
+			// next entry in the log. If the iterator progresses beyond the start or end of the log
+			// then the iterator is reset.
+			void Progress()
+			{
+				Block->AssertIsValid(Ptrs);
 
-		using Iterator = TIterator</* IS_FORWARD */ true, /* IS_CONST */ false>;
-		using ReverseIterator = TIterator</* IS_FORWARD */ false, /* IS_CONST */ false>;
-		using ConstIterator = TIterator</* IS_FORWARD */ true, /* IS_CONST */ true>;
-		using ConstReverseIterator = TIterator</* IS_FORWARD */ false, /* IS_CONST */ true>;
+				if constexpr (IS_FORWARD)
+				{
+					if (Ptrs.Record == Block->TailPtrs.Record)
+					{
+						Block = Block->NextBlock;
+						if (!Block)
+						{
+							// Progressed beyond the end of the log.
+							Reset();
+							return;
+						}
+
+						Ptrs = Block->Head();
+					}
+
+					Ptrs.Data -= Ptrs.Record->Size;
+					SetCurrent();
+					Ptrs.Record++;
+				}
+				else
+				{
+					if (Ptrs.Record == Block->FirstRecord())
+					{
+						Block = Block->PrevBlock;
+						// Note: Only the head block is fixed size, and may not contain any entries.
+						if (!Block || Block->IsEmpty())
+						{
+							// Progressed beyond the start of the log.
+							Reset();
+							return;
+						}
+						Ptrs = Block->Tail();
+					}
+
+					Ptrs.Record--;
+					SetCurrent();
+					Ptrs.Data += Current.Size;
+				}
+				
+				Block->AssertIsValid(Ptrs);
+			}
 
-		Iterator begin()
-		{
-			return Iterator(Start);
-		}
+			// Assigns to Current the entry this iterator currently points to.
+			UE_AUTORTFM_FORCEINLINE void SetCurrent()
+			{
+				Current.LogicalAddress = reinterpret_cast<std::byte*>(Ptrs.Record->Address);
+				Current.Data = Ptrs.Data;
+				Current.Size = Ptrs.Record->Size;
+				Current.bNoMemoryValidation = Ptrs.Record->bNoMemoryValidation;
+			}
 
-		ConstIterator begin() const
-		{
-			return ConstIterator(Start);
-		}
+			// Resets the iterator (compares equal to the write log's end())
+			UE_AUTORTFM_FORCEINLINE void Reset()
+			{
+				Block = nullptr;
+				Ptrs = {};
+				Current = {};
+			}
 
-		ReverseIterator rbegin()
-		{
-			return ReverseIterator(Current);
-		}
+			FBlock* Block = nullptr;
+			FBlock::FEntryPtrs Ptrs;
+			FWriteLogEntry Current;
+		};
 
-		ConstReverseIterator rbegin() const
-		{
-			return ConstReverseIterator(Current);
-		}
+		using Iterator = TIterator</* IS_FORWARD */ true>;
+		using ReverseIterator = TIterator</* IS_FORWARD */ false>;
 
-		Iterator end()
-		{
-			return Iterator(nullptr);
-		}
+		Iterator begin() const { return Iterator(HeadBlock); }
+		ReverseIterator rbegin() const { return ReverseIterator(TailBlock); }
+		Iterator end() const { return Iterator(); }
+		ReverseIterator rend() const { return ReverseIterator(); }
 
-		ConstIterator end() const
+		// Resets the write log to its initial state, freeing any allocated memory.
+		void Reset()
 		{
-			return ConstIterator(nullptr);
+			// Skip HeadBlock, which is held as part of this structure.
+			FBlock* Block = HeadBlock->NextBlock; 
+			while (nullptr != Block)
+			{
+				FBlock* const Next = Block->NextBlock;
+				Block->Free();
+				Block = Next;
+			}
+			new(HeadBlockMemory) FBlock(HeadBlockSize - sizeof(FBlock));
+			HeadBlock = reinterpret_cast<FBlock*>(HeadBlockMemory);
+			TailBlock = reinterpret_cast<FBlock*>(HeadBlockMemory);
+			NumEntries = 0;
+			TotalSizeBytes = 0;
 		}
 
-		ReverseIterator rend()
-		{
-			return ReverseIterator(nullptr);
-		}
+		// Returns true if the log holds no entries.
+		UE_AUTORTFM_FORCEINLINE bool IsEmpty() const { return 0 == NumEntries; }
 
-		ConstReverseIterator rend() const
-		{
-			return ConstReverseIterator(nullptr);
-		}
+		// Return the number of entries in the log.
+		UE_AUTORTFM_FORCEINLINE size_t Num() const { return NumEntries; }
 
-        void Reset()
-        {
-            while (nullptr != Start)
-            {
-                FWriteLogEntryBucket* const Old = Start;
-                Start = Start->Next;
-                delete Old;
-            }
-
-            Current = nullptr;
-            TotalSize = 0;
-        }
-
-        UE_AUTORTFM_FORCEINLINE bool IsEmpty() const { return 0 == TotalSize; }
-		UE_AUTORTFM_FORCEINLINE size_t Num() const { return TotalSize; }
-    private:
-
-        FWriteLogEntryBucket* Start = nullptr;
-        FWriteLogEntryBucket* Current = nullptr;
-        size_t TotalSize = 0;
-    };
+		// Return the total size in bytes for all entries in the log.
+		UE_AUTORTFM_FORCEINLINE size_t TotalSize() const { return TotalSizeBytes; }
+
+	private:
+		static constexpr size_t HeadBlockSize = 256;
+		FBlock* HeadBlock = reinterpret_cast<FBlock*>(HeadBlockMemory);
+		FBlock* TailBlock = reinterpret_cast<FBlock*>(HeadBlockMemory);
+		size_t NumEntries = 0;
+		size_t TotalSizeBytes = 0;
+		alignas(alignof(FBlock)) std::byte HeadBlockMemory[HeadBlockSize];
+	};
 }
diff --git a/Engine/Source/Runtime/Core/Private/AutoRTFM/WriteLogBumpAllocator.h b/Engine/Source/Runtime/Core/Private/AutoRTFM/WriteLogBumpAllocator.h
deleted file mode 100644
