[AutoRTFM] Fix double transaction merge to parent's FHitStet

Cause by a bad P4 merge.

diff --git a/Engine/Source/Runtime/Core/Private/AutoRTFM/Transaction.cpp b/Engine/Source/Runtime/Core/Private/AutoRTFM/Transaction.cpp
--- a/Engine/Source/Runtime/Core/Private/AutoRTFM/Transaction.cpp
+++ b/Engine/Source/Runtime/Core/Private/AutoRTFM/Transaction.cpp
@@ -239,13 +239,6 @@ void FTransaction::CommitNested()
 		}
 
 		Parent->WriteLog.Push(Write);
-
-		FHitSetEntry Entry{};
-		Entry.Address = reinterpret_cast<uintptr_t>(Write.LogicalAddress);
-		Entry.Size = static_cast<uint16_t>(Write.Size);
-		Entry.bNoMemoryValidation = Write.bNoMemoryValidation;
-
-		Parent->HitSet.FindOrTryInsert(Entry);
 	}
 
 	// For all the deferred calls to `PopOnCommitHandler` that we couldn't
