Minor comment fix.

#rnx

diff --git a/Engine/Source/Runtime/Core/Public/Misc/NotNull.h b/Engine/Source/Runtime/Core/Public/Misc/NotNull.h
--- a/Engine/Source/Runtime/Core/Public/Misc/NotNull.h
+++ b/Engine/Source/Runtime/Core/Public/Misc/NotNull.h
@@ -26,7 +26,7 @@
  *   moved-from TNotNull variable can end up being null. However, users are never required to handle a variable
  *   in this state. If a variable is to be reused after it has been moved from (e.g. a data member of an object
  *   that isn't being destroyed) then code that made it null should assign a new non-null value to it before
- *   returning to user code or using to construct or assign to another TNotNull. Users of your variable should
+ *   returning to user code or using it to construct or assign to another TNotNull. Users of your variable should
  *   never be allowed to see a null value. Compilers and static analyzers are allowed to assume that the pointer
      is not null and optimize and analyze accordingly.
  *
