Fixed REQUIRE_CHECK for upcoming changes to add [[nodiscard]] to some container functions.  Copied from the same recent changes to /Engine/Source/Runtime/Core/Public/Tests/Assertions.h.

diff --git a/Engine/Source/Developer/LowLevelTestsRunner/Public/TestMacros/Assertions.h b/Engine/Source/Developer/LowLevelTestsRunner/Public/TestMacros/Assertions.h
--- a/Engine/Source/Developer/LowLevelTestsRunner/Public/TestMacros/Assertions.h
+++ b/Engine/Source/Developer/LowLevelTestsRunner/Public/TestMacros/Assertions.h
@@ -60,7 +60,7 @@
 		Catch::AssertionHandler catchAssertionHandler( macroName##_catch_sr, CATCH_INTERNAL_LINEINFO, checkExpr, resultDisposition ); \
 		INTERNAL_CATCH_TRY { \
 			FCheckScope scope; \
-			__VA_ARGS__; \
+			(void)(__VA_ARGS__); \
 			bool bEncounteredEnsure = scope.GetCount() > 0; \
 			if (doCheck && !bEncounteredEnsure) \
 				catchAssertionHandler.handleMessage( Catch::ResultWas::ExplicitFailure, "Expected failure of `check` not received"); \
@@ -73,7 +73,7 @@
 		Catch::AssertionHandler catchAssertionHandler( macroName##_catch_sr, CATCH_INTERNAL_LINEINFO, checkExpr, resultDisposition ); \
 		INTERNAL_CATCH_TRY { \
 			FCheckScope scope(msg); \
-			__VA_ARGS__; \
+			(void)(__VA_ARGS__); \
 			bool bEncounteredEnsure = scope.GetCount() > 0; \
 			if (doCheck && !bEncounteredEnsure) \
 				catchAssertionHandler.handleMessage( Catch::ResultWas::ExplicitFailure, ( Catch::MessageStream() << "Expected failure of `check` containing message '" << msg << "' not received" + ::Catch::StreamEndStop() ).m_stream.str() ); \
diff --git a/Engine/Source/Runtime/Core/Tests/Misc/AssertionMacrosTest.cpp b/Engine/Source/Runtime/Core/Tests/Misc/AssertionMacrosTest.cpp
--- a/Engine/Source/Runtime/Core/Tests/Misc/AssertionMacrosTest.cpp
+++ b/Engine/Source/Runtime/Core/Tests/Misc/AssertionMacrosTest.cpp
@@ -118,16 +118,16 @@ TEST_CASE("Core::Misc::TestEnsure", "[Core][Misc][AssertionMacros][Ensure]")
 
 TEST_CASE("Core::Misc::TestCheck", "[Core][Misc][AssertionMacros][Check]")
 {
-	REQUIRE_CHECK(check(1 == 2));
-	REQUIRE_CHECK_MSG("1 == 2", checkf(1 == 2, TEXT("Error Message")));
-	REQUIRE_CHECK_MSG("Error Message", checkf(1 == 2, TEXT("Error Message")));
+	REQUIRE_CHECK([]{ check(1 == 2); }());
+	REQUIRE_CHECK_MSG("1 == 2", []{ checkf(1 == 2, TEXT("Error Message")); }());
+	REQUIRE_CHECK_MSG("Error Message", []{ checkf(1 == 2, TEXT("Error Message")); }());
 }
 
 TEST_CASE("Core::Misc::TestCheck", "[Core][Misc][AssertionMacros][CheckSlow]")
 {
-	REQUIRE_CHECK_SLOW(checkSlow(1 == 2));
-	REQUIRE_CHECK_SLOW_MSG("1 == 2", checkfSlow(1 == 2, TEXT("Error Message")));
-	REQUIRE_CHECK_SLOW_MSG("Error Message", checkfSlow(1 == 2, TEXT("Error Message")));
+	REQUIRE_CHECK_SLOW([]{ checkSlow(1 == 2); }());
+	REQUIRE_CHECK_SLOW_MSG("1 == 2", []{ checkfSlow(1 == 2, TEXT("Error Message")); }());
+	REQUIRE_CHECK_SLOW_MSG("Error Message", []{ checkfSlow(1 == 2, TEXT("Error Message")); }());
 }
 
 #endif
