Added [[nodiscard]] to relevant algorithms.

diff --git a/Engine/Source/Runtime/Core/Public/Algo/Accumulate.h b/Engine/Source/Runtime/Core/Public/Algo/Accumulate.h
--- a/Engine/Source/Runtime/Core/Public/Algo/Accumulate.h
+++ b/Engine/Source/Runtime/Core/Public/Algo/Accumulate.h
@@ -10,13 +10,13 @@
 template<typename T = void>
 struct TPlus
 {
-	FORCEINLINE T operator()(const T& A, const T& B) { return A + B; }
+	[[nodiscard]] FORCEINLINE T operator()(const T& A, const T& B) { return A + B; }
 };
 template<>
 struct TPlus<void>
 {
 	template<typename U, typename V>
-	FORCEINLINE auto operator()(U&& A, V&& B) -> decltype(A + B) { return A + B; }
+	[[nodiscard]] FORCEINLINE auto operator()(U&& A, V&& B) -> decltype(A + B) { return A + B; }
 };
 
 
@@ -32,7 +32,7 @@ namespace Algo
 	* @return the result of summing all the elements of Input
 	*/
 	template <typename T, typename A, typename OpT>
-	FORCEINLINE T Accumulate(const A& Input, T Init, OpT Op)
+	[[nodiscard]] FORCEINLINE T Accumulate(const A& Input, T Init, OpT Op)
 	{
 		T Result = MoveTemp(Init);
 		for (const auto& InputElem : Input)
@@ -51,7 +51,7 @@ namespace Algo
 	 * @return the result of summing all the elements of Input
 	 */
 	template <typename T, typename A>
-	FORCEINLINE T Accumulate(const A& Input, T Init)
+	[[nodiscard]] FORCEINLINE T Accumulate(const A& Input, T Init)
 	{
 		return Accumulate(Input, MoveTemp(Init), TPlus<>());
 	}
@@ -67,7 +67,7 @@ namespace Algo
 	* @return the result of mapping and then summing all the elements of Input
 	*/
 	template <typename T, typename A, typename MapT, typename OpT>
-	FORCEINLINE T TransformAccumulate(const A& Input, MapT MapOp, T Init, OpT Op)
+	[[nodiscard]] FORCEINLINE T TransformAccumulate(const A& Input, MapT MapOp, T Init, OpT Op)
 	{
 		T Result = MoveTemp(Init);
 		for (const auto& InputElem : Input)
@@ -87,7 +87,7 @@ namespace Algo
 	* @return the result of mapping and then summing all the elements of Input
 	*/
 	template <typename T, typename A, typename MapT>
-	FORCEINLINE T TransformAccumulate(const A& Input, MapT MapOp, T Init)
+	[[nodiscard]] FORCEINLINE T TransformAccumulate(const A& Input, MapT MapOp, T Init)
 	{
 		return TransformAccumulate(Input, MapOp, MoveTemp(Init), TPlus<>());
 	}
diff --git a/Engine/Source/Runtime/Core/Public/Algo/AllOf.h b/Engine/Source/Runtime/Core/Public/Algo/AllOf.h
--- a/Engine/Source/Runtime/Core/Public/Algo/AllOf.h
+++ b/Engine/Source/Runtime/Core/Public/Algo/AllOf.h
@@ -16,7 +16,7 @@ namespace Algo
 	 * @return  true if all the elements are truthy, false otherwise.
 	 */
 	template <typename RangeType>
-	FORCEINLINE bool AllOf(const RangeType& Range)
+	[[nodiscard]] FORCEINLINE bool AllOf(const RangeType& Range)
 	{
 		for (const auto& Element : Range)
 		{
@@ -38,7 +38,7 @@ namespace Algo
 	 * @return  true if all the projections are truthy, false otherwise.
 	 */
 	template <typename RangeType, typename ProjectionType>
-	FORCEINLINE bool AllOf(const RangeType& Range, ProjectionType Projection)
+	[[nodiscard]] FORCEINLINE bool AllOf(const RangeType& Range, ProjectionType Projection)
 	{
 		for (const auto& Element : Range)
 		{
diff --git a/Engine/Source/Runtime/Core/Public/Algo/AnyOf.h b/Engine/Source/Runtime/Core/Public/Algo/AnyOf.h
--- a/Engine/Source/Runtime/Core/Public/Algo/AnyOf.h
+++ b/Engine/Source/Runtime/Core/Public/Algo/AnyOf.h
@@ -17,7 +17,7 @@ namespace Algo
 	 * @return  true if at least one element is truthy, false otherwise.
 	 */
 	template <typename RangeType>
-	FORCEINLINE bool AnyOf(const RangeType& Range)
+	[[nodiscard]] FORCEINLINE bool AnyOf(const RangeType& Range)
 	{
 		return !Algo::NoneOf(Range);
 	}
@@ -31,7 +31,7 @@ namespace Algo
 	 * @return  true if at least one projection is truthy, false otherwise.
 	 */
 	template <typename RangeType, typename ProjectionType>
-	FORCEINLINE bool AnyOf(const RangeType& Range, ProjectionType Projection)
+	[[nodiscard]] FORCEINLINE bool AnyOf(const RangeType& Range, ProjectionType Projection)
 	{
 		return !Algo::NoneOf(Range, MoveTemp(Projection));
 	}
@@ -46,7 +46,7 @@ namespace Algo
 	 * @return  true if at least one projection is truthy, false otherwise.
 	 */
 	template <typename RangeType, typename ProjectionType>
-	FORCEINLINE bool AnyOf(const RangeType& Range, ProjectionType Projection, ENoRef NoRef)
+	[[nodiscard]] FORCEINLINE bool AnyOf(const RangeType& Range, ProjectionType Projection, ENoRef NoRef)
 	{
 		return !Algo::NoneOf(Range, MoveTemp(Projection), NoRef);
 	}
diff --git a/Engine/Source/Runtime/Core/Public/Algo/BinarySearch.h b/Engine/Source/Runtime/Core/Public/Algo/BinarySearch.h
--- a/Engine/Source/Runtime/Core/Public/Algo/BinarySearch.h
+++ b/Engine/Source/Runtime/Core/Public/Algo/BinarySearch.h
@@ -89,12 +89,12 @@ namespace Algo
 	 * @returns Position of the first element >= Value, may be position after last element in range
 	 */
 	template <typename RangeType, typename ValueType, typename SortPredicateType>
-	FORCEINLINE auto LowerBound(const RangeType& Range, const ValueType& Value, SortPredicateType SortPredicate) -> decltype(GetNum(Range))
+	[[nodiscard]] FORCEINLINE auto LowerBound(const RangeType& Range, const ValueType& Value, SortPredicateType SortPredicate) -> decltype(GetNum(Range))
 	{
 		return AlgoImpl::LowerBoundInternal(GetData(Range), GetNum(Range), Value, FIdentityFunctor(), SortPredicate);
 	}
 	template <typename RangeType, typename ValueType>
-	FORCEINLINE auto LowerBound(const RangeType& Range, const ValueType& Value) -> decltype(GetNum(Range))
+	[[nodiscard]] FORCEINLINE auto LowerBound(const RangeType& Range, const ValueType& Value) -> decltype(GetNum(Range))
 	{
 		return AlgoImpl::LowerBoundInternal(GetData(Range), GetNum(Range), Value, FIdentityFunctor(), TLess<>());
 	}
@@ -110,12 +110,12 @@ namespace Algo
 	 * @returns Position of the first element >= Value, may be position after last element in range
 	 */
 	template <typename RangeType, typename ValueType, typename ProjectionType, typename SortPredicateType>
-	FORCEINLINE auto LowerBoundBy(const RangeType& Range, const ValueType& Value, ProjectionType Projection, SortPredicateType SortPredicate) -> decltype(GetNum(Range))
+	[[nodiscard]] FORCEINLINE auto LowerBoundBy(const RangeType& Range, const ValueType& Value, ProjectionType Projection, SortPredicateType SortPredicate) -> decltype(GetNum(Range))
 	{
 		return AlgoImpl::LowerBoundInternal(GetData(Range), GetNum(Range), Value, Projection, SortPredicate);
 	}
 	template <typename RangeType, typename ValueType, typename ProjectionType>
-	FORCEINLINE auto LowerBoundBy(const RangeType& Range, const ValueType& Value, ProjectionType Projection) -> decltype(GetNum(Range))
+	[[nodiscard]] FORCEINLINE auto LowerBoundBy(const RangeType& Range, const ValueType& Value, ProjectionType Projection) -> decltype(GetNum(Range))
 	{
 		return AlgoImpl::LowerBoundInternal(GetData(Range), GetNum(Range), Value, Projection, TLess<>());
 	}
@@ -130,12 +130,12 @@ namespace Algo
 	 * @returns Position of the first element > Value, may be past end of range
 	 */
 	template <typename RangeType, typename ValueType, typename SortPredicateType>
-	FORCEINLINE auto UpperBound(const RangeType& Range, const ValueType& Value, SortPredicateType SortPredicate) -> decltype(GetNum(Range))
+	[[nodiscard]] FORCEINLINE auto UpperBound(const RangeType& Range, const ValueType& Value, SortPredicateType SortPredicate) -> decltype(GetNum(Range))
 	{
 		return AlgoImpl::UpperBoundInternal(GetData(Range), GetNum(Range), Value, FIdentityFunctor(), SortPredicate);
 	}
 	template <typename RangeType, typename ValueType>
-	FORCEINLINE auto UpperBound(const RangeType& Range, const ValueType& Value) -> decltype(GetNum(Range))
+	[[nodiscard]] FORCEINLINE auto UpperBound(const RangeType& Range, const ValueType& Value) -> decltype(GetNum(Range))
 	{
 		return AlgoImpl::UpperBoundInternal(GetData(Range), GetNum(Range), Value, FIdentityFunctor(), TLess<>());
 	}
@@ -151,12 +151,12 @@ namespace Algo
 	 * @returns Position of the first element > Value, may be past end of range
 	 */
 	template <typename RangeType, typename ValueType, typename ProjectionType, typename SortPredicateType>
-	FORCEINLINE auto UpperBoundBy(const RangeType& Range, const ValueType& Value, ProjectionType Projection, SortPredicateType SortPredicate) -> decltype(GetNum(Range))
+	[[nodiscard]] FORCEINLINE auto UpperBoundBy(const RangeType& Range, const ValueType& Value, ProjectionType Projection, SortPredicateType SortPredicate) -> decltype(GetNum(Range))
 	{
 		return AlgoImpl::UpperBoundInternal(GetData(Range), GetNum(Range), Value, Projection, SortPredicate);
 	}
 	template <typename RangeType, typename ValueType, typename ProjectionType>
-	FORCEINLINE auto UpperBoundBy(const RangeType& Range, const ValueType& Value, ProjectionType Projection) -> decltype(GetNum(Range))
+	[[nodiscard]] FORCEINLINE auto UpperBoundBy(const RangeType& Range, const ValueType& Value, ProjectionType Projection) -> decltype(GetNum(Range))
 	{
 		return AlgoImpl::UpperBoundInternal(GetData(Range), GetNum(Range), Value, Projection, TLess<>());
 	}
@@ -170,7 +170,7 @@ namespace Algo
 	 * @return Index of found element, or INDEX_NONE
 	 */
 	template <typename RangeType, typename ValueType, typename SortPredicateType>
-	FORCEINLINE auto BinarySearch(const RangeType& Range, const ValueType& Value, SortPredicateType SortPredicate) -> decltype(GetNum(Range))
+	[[nodiscard]] FORCEINLINE auto BinarySearch(const RangeType& Range, const ValueType& Value, SortPredicateType SortPredicate) -> decltype(GetNum(Range))
 	{
 		auto CheckIndex = LowerBound(Range, Value, SortPredicate);
 		if (CheckIndex < GetNum(Range))
@@ -185,7 +185,7 @@ namespace Algo
 		return INDEX_NONE;
 	}
 	template <typename RangeType, typename ValueType>
-	FORCEINLINE auto BinarySearch(const RangeType& Range, const ValueType& Value)
+	[[nodiscard]] FORCEINLINE auto BinarySearch(const RangeType& Range, const ValueType& Value)
 	{
 		return BinarySearch(Range, Value, TLess<>());
 	}
@@ -200,7 +200,7 @@ namespace Algo
 	 * @return Index of found element, or INDEX_NONE
 	 */
 	template <typename RangeType, typename ValueType, typename ProjectionType, typename SortPredicateType>
-	FORCEINLINE auto BinarySearchBy(const RangeType& Range, const ValueType& Value, ProjectionType Projection, SortPredicateType SortPredicate) -> decltype(GetNum(Range))
+	[[nodiscard]] FORCEINLINE auto BinarySearchBy(const RangeType& Range, const ValueType& Value, ProjectionType Projection, SortPredicateType SortPredicate) -> decltype(GetNum(Range))
 	{
 		auto CheckIndex = LowerBoundBy(Range, Value, Projection, SortPredicate);
 		if (CheckIndex < GetNum(Range))
@@ -215,7 +215,7 @@ namespace Algo
 		return INDEX_NONE;
 	}
 	template <typename RangeType, typename ValueType, typename ProjectionType>
-	FORCEINLINE auto BinarySearchBy(const RangeType& Range, const ValueType& Value, ProjectionType Projection)
+	[[nodiscard]] FORCEINLINE auto BinarySearchBy(const RangeType& Range, const ValueType& Value, ProjectionType Projection)
 	{
 		return BinarySearchBy(Range, Value, Projection, TLess<>());
 	}
diff --git a/Engine/Source/Runtime/Core/Public/Algo/Compare.h b/Engine/Source/Runtime/Core/Public/Algo/Compare.h
--- a/Engine/Source/Runtime/Core/Public/Algo/Compare.h
+++ b/Engine/Source/Runtime/Core/Public/Algo/Compare.h
@@ -48,7 +48,7 @@ namespace Algo
 	 * @return Whether the containers are the same size and operator== returned true for every pair of elements.
 	 */
 	template <typename InAT, typename InBT>
-	constexpr bool Compare(const InAT& InputA, const InBT& InputB)
+	[[nodiscard]] constexpr bool Compare(const InAT& InputA, const InBT& InputB)
 	{
 		return Private::Compare(InputA, InputB, FIdentityFunctor(), TEqualTo<>());
 	}
@@ -63,7 +63,7 @@ namespace Algo
 	 * @return Whether the containers are the same size and the predicate returned true for every pair of elements.
 	 */
 	template <typename InAT, typename InBT, typename PredicateT>
-	constexpr bool Compare(const InAT& InputA, const InBT& InputB, PredicateT Predicate)
+	[[nodiscard]] constexpr bool Compare(const InAT& InputA, const InBT& InputB, PredicateT Predicate)
 	{
 		return Private::Compare(InputA, InputB, FIdentityFunctor(), MoveTemp(Predicate));
 	}
@@ -78,7 +78,7 @@ namespace Algo
 	 * @return Whether the containers are the same size and operator== returned true for every pair of elements.
 	 */
 	template <typename InAT, typename InBT, typename ProjectionT>
-	constexpr bool CompareBy(const InAT& InputA, const InBT& InputB, ProjectionT Projection)
+	[[nodiscard]] constexpr bool CompareBy(const InAT& InputA, const InBT& InputB, ProjectionT Projection)
 	{
 		return Private::Compare(InputA, InputB, MoveTemp(Projection), TEqualTo<>());
 	}
@@ -94,7 +94,7 @@ namespace Algo
 	 * @return Whether the containers are the same size and the predicate returned true for every pair of elements.
 	 */
 	template <typename InAT, typename InBT, typename ProjectionT, typename PredicateT>
-	constexpr bool CompareBy(const InAT& InputA, const InBT& InputB, ProjectionT Projection, PredicateT Predicate)
+	[[nodiscard]] constexpr bool CompareBy(const InAT& InputA, const InBT& InputB, ProjectionT Projection, PredicateT Predicate)
 	{
 		return Private::Compare(InputA, InputB, MoveTemp(Projection), MoveTemp(Predicate));
 	}
@@ -122,7 +122,7 @@ namespace Algo
 	 * @return -1 if A < B, 0 if A == B, 1 if A > B
 	 */
 	template <typename MapType, typename KeyLessThanType, typename ValueLessThanType>
-	int CompareMap(const MapType& A, const MapType& B, KeyLessThanType KeyLessThan, ValueLessThanType ValueLessThan)
+	[[nodiscard]] int CompareMap(const MapType& A, const MapType& B, KeyLessThanType KeyLessThan, ValueLessThanType ValueLessThan)
 	{
 		using KeyType = typename MapType::KeyType;
 		using ValueType = typename MapType::ValueType;
@@ -191,13 +191,13 @@ namespace Algo
 	}
 
 	template <typename MapType>
-	int CompareMap(const MapType& A, const MapType& B)
+	[[nodiscard]] int CompareMap(const MapType& A, const MapType& B)
 	{
 		return CompareMap(A, B, TLess<>(), TLess<>());
 	}
 
 	template <typename MapType, typename KeyLessThanType>
-	int CompareMap(const MapType& A, const MapType& B, KeyLessThanType KeyLessThan)
+	[[nodiscard]] int CompareMap(const MapType& A, const MapType& B, KeyLessThanType KeyLessThan)
 	{
 		return CompareMap(A, B, MoveTempIfPossible(KeyLessThan), TLess<>());
 	}
@@ -221,7 +221,7 @@ namespace Algo
 	 * @return -1 if A < B, 0 if A == B, 1 if A > B
 	 */
 	template <typename SetType, typename KeyLessThanType>
-	int CompareSet(const SetType& A, const SetType& B, KeyLessThanType KeyLessThan)
+	[[nodiscard]] int CompareSet(const SetType& A, const SetType& B, KeyLessThanType KeyLessThan)
 	{
 		using KeyType = typename SetType::ElementType;
 		if (A.Num() != B.Num())
@@ -273,7 +273,7 @@ namespace Algo
 	}
 
 	template <typename SetType>
-	int CompareSet(const SetType& A, const SetType& B)
+	[[nodiscard]] int CompareSet(const SetType& A, const SetType& B)
 	{
 		return CompareSet(A, B, TLess<>());
 	}
diff --git a/Engine/Source/Runtime/Core/Public/Algo/Count.h b/Engine/Source/Runtime/Core/Public/Algo/Count.h
--- a/Engine/Source/Runtime/Core/Public/Algo/Count.h
+++ b/Engine/Source/Runtime/Core/Public/Algo/Count.h
@@ -13,7 +13,7 @@ namespace Algo
 	 * @param  InValue  Value to compare against
 	 */
 	template <typename InT, typename ValueT>
-	FORCEINLINE SIZE_T Count(const InT& Input, const ValueT& InValue)
+	[[nodiscard]] FORCEINLINE SIZE_T Count(const InT& Input, const ValueT& InValue)
 	{
 		SIZE_T Result = 0;
 		for (const auto& Value : Input)
@@ -33,7 +33,7 @@ namespace Algo
 	 * @param  Predicate  Condition which returns true for elements that should be counted and false for elements that should be skipped
 	 */
 	template <typename InT, typename PredicateT>
-	FORCEINLINE SIZE_T CountIf(const InT& Input, PredicateT Predicate)
+	[[nodiscard]] FORCEINLINE SIZE_T CountIf(const InT& Input, PredicateT Predicate)
 	{
 		SIZE_T Result = 0;
 		for (const auto& Value : Input)
@@ -54,7 +54,7 @@ namespace Algo
 	 * @param  Projection  The projection to apply to the element
 	 */
 	template <typename InT, typename ValueT, typename ProjectionT>
-	FORCEINLINE SIZE_T CountBy(const InT& Input, const ValueT& InValue, ProjectionT Proj)
+	[[nodiscard]] FORCEINLINE SIZE_T CountBy(const InT& Input, const ValueT& InValue, ProjectionT Proj)
 	{
 		SIZE_T Result = 0;
 		for (const auto& Value : Input)
diff --git a/Engine/Source/Runtime/Core/Public/Algo/Find.h b/Engine/Source/Runtime/Core/Public/Algo/Find.h
--- a/Engine/Source/Runtime/Core/Public/Algo/Find.h
+++ b/Engine/Source/Runtime/Core/Public/Algo/Find.h
@@ -49,7 +49,7 @@ namespace Algo
 	 * @return A pointer to the first element found, or nullptr if none was found.
 	 */
 	template <typename RangeType, typename ValueType>
-	FORCEINLINE auto Find(RangeType&& Range, const ValueType& Value)
+	[[nodiscard]] FORCEINLINE auto Find(RangeType&& Range, const ValueType& Value)
 		-> decltype(AlgoImpl::FindBy(Forward<RangeType>(Range), Value, FIdentityFunctor()))
 	{
 		return AlgoImpl::FindBy(Forward<RangeType>(Range), Value, FIdentityFunctor());
@@ -65,7 +65,7 @@ namespace Algo
 	 * @return A pointer to the first element found, or nullptr if none was found.
 	 */
 	template <typename RangeType, typename ValueType, typename ProjectionType>
-	FORCEINLINE auto FindBy(RangeType&& Range, const ValueType& Value, ProjectionType Proj)
+	[[nodiscard]] FORCEINLINE auto FindBy(RangeType&& Range, const ValueType& Value, ProjectionType Proj)
 		-> decltype(AlgoImpl::FindBy(Forward<RangeType>(Range), Value, MoveTemp(Proj)))
 	{
 		return AlgoImpl::FindBy(Forward<RangeType>(Range), Value, MoveTemp(Proj));
@@ -80,7 +80,7 @@ namespace Algo
 	 * @return A pointer to the first element found, or nullptr if none was found.
 	 */
 	template <typename RangeType, typename PredicateType>
-	FORCEINLINE auto FindByPredicate(RangeType&& Range, PredicateType Pred)
+	[[nodiscard]] FORCEINLINE auto FindByPredicate(RangeType&& Range, PredicateType Pred)
 		-> decltype(AlgoImpl::FindByPredicate(Forward<RangeType>(Range), MoveTemp(Pred)))
 	{
 		return AlgoImpl::FindByPredicate(Forward<RangeType>(Range), MoveTemp(Pred));
diff --git a/Engine/Source/Runtime/Core/Public/Algo/FindLast.h b/Engine/Source/Runtime/Core/Public/Algo/FindLast.h
--- a/Engine/Source/Runtime/Core/Public/Algo/FindLast.h
+++ b/Engine/Source/Runtime/Core/Public/Algo/FindLast.h
@@ -48,7 +48,7 @@ namespace Algo
 	 * @return A pointer to the last element found, or nullptr if none was found.
 	 */
 	template <typename RangeType, typename ValueType>
-	FORCEINLINE auto FindLast(RangeType&& Range, const ValueType& Value)
+	[[nodiscard]] FORCEINLINE auto FindLast(RangeType&& Range, const ValueType& Value)
 		-> decltype(AlgoImpl::FindLastBy(GetData(Range), GetNum(Range), Value, FIdentityFunctor()))
 	{
 		return AlgoImpl::FindLastBy(GetData(Range), GetNum(Range), Value, FIdentityFunctor());
@@ -64,7 +64,7 @@ namespace Algo
 	 * @return A pointer to the last element found, or nullptr if none was found.
 	 */
 	template <typename RangeType, typename ValueType, typename ProjectionType>
-	FORCEINLINE auto FindLastBy(RangeType&& Range, const ValueType& Value, ProjectionType Proj)
+	[[nodiscard]] FORCEINLINE auto FindLastBy(RangeType&& Range, const ValueType& Value, ProjectionType Proj)
 		-> decltype(AlgoImpl::FindLastBy(GetData(Range), GetNum(Range), Value, MoveTemp(Proj)))
 	{
 		return AlgoImpl::FindLastBy(GetData(Range), GetNum(Range), Value, MoveTemp(Proj));
@@ -79,7 +79,7 @@ namespace Algo
 	 * @return A pointer to the last element found, or nullptr if none was found.
 	 */
 	template <typename RangeType, typename PredicateType>
-	FORCEINLINE auto FindLastByPredicate(RangeType&& Range, PredicateType Pred)
+	[[nodiscard]] FORCEINLINE auto FindLastByPredicate(RangeType&& Range, PredicateType Pred)
 		-> decltype(AlgoImpl::FindLastByPredicate(GetData(Range), GetNum(Range), MoveTemp(Pred)))
 	{
 		return AlgoImpl::FindLastByPredicate(GetData(Range), GetNum(Range), MoveTemp(Pred));
diff --git a/Engine/Source/Runtime/Core/Public/Algo/FindSequence.h b/Engine/Source/Runtime/Core/Public/Algo/FindSequence.h
--- a/Engine/Source/Runtime/Core/Public/Algo/FindSequence.h
+++ b/Engine/Source/Runtime/Core/Public/Algo/FindSequence.h
@@ -43,7 +43,7 @@ namespace Algo
 	* @return A pointer to the first occurrence of the "What" sequence in "Where" sequence, or nullptr if not found.
 	*/
 	template<typename RangeWhereType, typename RangeWhatType>
-	FORCEINLINE auto FindSequence(const RangeWhereType& Where, const RangeWhatType& What)
+	[[nodiscard]] FORCEINLINE auto FindSequence(const RangeWhereType& Where, const RangeWhatType& What)
 		-> decltype( AlgoImpl::FindSequence( GetData(Where), GetData(Where) + GetNum(Where), GetData(What), GetData(What) + GetNum(What)) )
 	{
 		if (GetNum(What) > GetNum(Where))
diff --git a/Engine/Source/Runtime/Core/Public/Algo/FindSortedStringCaseInsensitive.h b/Engine/Source/Runtime/Core/Public/Algo/FindSortedStringCaseInsensitive.h
--- a/Engine/Source/Runtime/Core/Public/Algo/FindSortedStringCaseInsensitive.h
+++ b/Engine/Source/Runtime/Core/Public/Algo/FindSortedStringCaseInsensitive.h
@@ -15,7 +15,7 @@ namespace Algo
 	 *
 	 * @return The index of the found string in the array, or -1 if the string was not found.
 	 */
-	CORE_API int32 FindSortedStringCaseInsensitive(const TCHAR* Str, const TCHAR* const* SortedArray, int32 ArrayCount);
+	[[nodiscard]] CORE_API int32 FindSortedStringCaseInsensitive(const TCHAR* Str, const TCHAR* const* SortedArray, int32 ArrayCount);
 
 	/**
 	 * Finds a string in an array of sorted strings, by case-insensitive search, by using binary subdivision of the array.
@@ -26,7 +26,7 @@ namespace Algo
 	 * @return The index of the found string in the array, or -1 if the string was not found.
 	 */
 	template <int32 ArraySize>
-	FORCEINLINE int32 FindSortedStringCaseInsensitive(const TCHAR* Str, const TCHAR* const (&SortedArray)[ArraySize])
+	[[nodiscard]] FORCEINLINE int32 FindSortedStringCaseInsensitive(const TCHAR* Str, const TCHAR* const (&SortedArray)[ArraySize])
 	{
 		return FindSortedStringCaseInsensitive(Str, SortedArray, ArraySize);
 	}
diff --git a/Engine/Source/Runtime/Core/Public/Algo/Includes.h b/Engine/Source/Runtime/Core/Public/Algo/Includes.h
--- a/Engine/Source/Runtime/Core/Public/Algo/Includes.h
+++ b/Engine/Source/Runtime/Core/Public/Algo/Includes.h
@@ -52,7 +52,7 @@ namespace Algo
 	 * @return True if RangeB is a subsequence of RangeA, false otherwise.
 	 */
 	template <typename RangeTypeA, typename RangeTypeB>
-	constexpr bool Includes(const RangeTypeA& RangeA, const RangeTypeB& RangeB)
+	[[nodiscard]] constexpr bool Includes(const RangeTypeA& RangeA, const RangeTypeB& RangeB)
 	{
 		return Private::Includes(GetData(RangeA), GetNum(RangeA), GetData(RangeB), GetNum(RangeB), FIdentityFunctor(), TLess<>());
 	}
@@ -68,7 +68,7 @@ namespace Algo
 	 * @return True if RangeB is a subsequence of RangeA according to the provided predicate, false otherwise.
 	 */
 	template <typename RangeTypeA, typename RangeTypeB, typename SortPredicateType>
-	constexpr bool Includes(const RangeTypeA& RangeA, const RangeTypeB& RangeB, SortPredicateType SortPredicate)
+	[[nodiscard]] constexpr bool Includes(const RangeTypeA& RangeA, const RangeTypeB& RangeB, SortPredicateType SortPredicate)
 	{
 		return Private::Includes(GetData(RangeA), GetNum(RangeA), GetData(RangeB), GetNum(RangeB), FIdentityFunctor(), MoveTemp(SortPredicate));
 	}
@@ -84,7 +84,7 @@ namespace Algo
 	 * @return True if RangeB is a subsequence of RangeA, based on the comparison of projected elements, false otherwise.
 	 */
 	template <typename RangeTypeA, typename RangeTypeB, typename ProjectionType>
-	constexpr bool IncludesBy(const RangeTypeA& RangeA, const RangeTypeB& RangeB, ProjectionType Projection)
+	[[nodiscard]] constexpr bool IncludesBy(const RangeTypeA& RangeA, const RangeTypeB& RangeB, ProjectionType Projection)
 	{
 		return Private::Includes(GetData(RangeA), GetNum(RangeA), GetData(RangeB), GetNum(RangeB), MoveTemp(Projection), TLess<>());
 	}
@@ -101,7 +101,7 @@ namespace Algo
 	 * @return True if RangeB is a subsequence of RangeA according to the comparison of projected elements using the provided predicate, false otherwise.
 	 */
 	template <typename RangeTypeA, typename RangeTypeB, typename ProjectionType, typename SortPredicateType>
-	constexpr bool IncludesBy(const RangeTypeA& RangeA, const RangeTypeB& RangeB, ProjectionType Projection, SortPredicateType SortPredicate)
+	[[nodiscard]] constexpr bool IncludesBy(const RangeTypeA& RangeA, const RangeTypeB& RangeB, ProjectionType Projection, SortPredicateType SortPredicate)
 	{
 		return Private::Includes(GetData(RangeA), GetNum(RangeA), GetData(RangeB), GetNum(RangeB), MoveTemp(Projection), MoveTemp(SortPredicate));
 	}
diff --git a/Engine/Source/Runtime/Core/Public/Algo/IndexOf.h b/Engine/Source/Runtime/Core/Public/Algo/IndexOf.h
--- a/Engine/Source/Runtime/Core/Public/Algo/IndexOf.h
+++ b/Engine/Source/Runtime/Core/Public/Algo/IndexOf.h
@@ -56,7 +56,7 @@ namespace Algo
 	 * @return The index of the first element found, or (SizeType)-1 if none was found (where SizeType = decltype(GetNum(Range)))
 	 */
 	template <typename RangeType, typename ValueType>
-	FORCEINLINE auto IndexOf(const RangeType& Range, const ValueType& Value)
+	[[nodiscard]] FORCEINLINE auto IndexOf(const RangeType& Range, const ValueType& Value)
 	{
 		static_assert(TIsContiguousContainer<RangeType>::Value, "Must only be used with contiguous containers");
 		return AlgoImpl::IndexOfBy(Range, Value, FIdentityFunctor());
@@ -72,7 +72,7 @@ namespace Algo
 	 * @return The index of the first element found, or (SizeType)-1 if none was found (where SizeType = decltype(GetNum(Range)))
 	 */
 	template <typename RangeType, typename ValueType, typename ProjectionType>
-	FORCEINLINE auto IndexOfBy(const RangeType& Range, const ValueType& Value, ProjectionType Proj)
+	[[nodiscard]] FORCEINLINE auto IndexOfBy(const RangeType& Range, const ValueType& Value, ProjectionType Proj)
 	{
 		static_assert(TIsContiguousContainer<RangeType>::Value, "Must only be used with contiguous containers");
 		return AlgoImpl::IndexOfBy(Range, Value, MoveTemp(Proj));
@@ -87,7 +87,7 @@ namespace Algo
 	 * @return The index of the first element found, or (SizeType)-1 if none was found (where SizeType = decltype(GetNum(Range)))
 	 */
 	template <typename RangeType, typename PredicateType>
-	FORCEINLINE auto IndexOfByPredicate(const RangeType& Range, PredicateType Pred)
+	[[nodiscard]] FORCEINLINE auto IndexOfByPredicate(const RangeType& Range, PredicateType Pred)
 	{
 		static_assert(TIsContiguousContainer<RangeType>::Value, "Must only be used with contiguous containers");
 		return AlgoImpl::IndexOfByPredicate(Range, MoveTemp(Pred));
diff --git a/Engine/Source/Runtime/Core/Public/Algo/IsHeap.h b/Engine/Source/Runtime/Core/Public/Algo/IsHeap.h
--- a/Engine/Source/Runtime/Core/Public/Algo/IsHeap.h
+++ b/Engine/Source/Runtime/Core/Public/Algo/IsHeap.h
@@ -47,7 +47,7 @@ namespace Algo
 	 * @return	returns	true if the range is a min-heap
 	 */
 	template <typename RangeType>
-	FORCEINLINE bool IsHeap(const RangeType& Range)
+	[[nodiscard]] FORCEINLINE bool IsHeap(const RangeType& Range)
 	{
 		return AlgoImpl::IsHeapInternal(GetData(Range), GetNum(Range), FIdentityFunctor(), TLess<>());
 	}
@@ -61,7 +61,7 @@ namespace Algo
 	 * @return	returns		true if the range is a min-heap
 	 */
 	template <typename RangeType, typename PredicateType>
-	FORCEINLINE bool IsHeap(const RangeType& Range, PredicateType Predicate)
+	[[nodiscard]] FORCEINLINE bool IsHeap(const RangeType& Range, PredicateType Predicate)
 	{
 		return AlgoImpl::IsHeapInternal(GetData(Range), GetNum(Range), FIdentityFunctor(), MoveTemp(Predicate));
 	}
@@ -75,7 +75,7 @@ namespace Algo
 	 * @return	returns		true if the range is a min-heap
 	 */
 	template <typename RangeType, typename ProjectionType>
-	FORCEINLINE bool IsHeapBy(const RangeType& Range, ProjectionType Projection)
+	[[nodiscard]] FORCEINLINE bool IsHeapBy(const RangeType& Range, ProjectionType Projection)
 	{
 		return AlgoImpl::IsHeapInternal(GetData(Range), GetNum(Range), MoveTemp(Projection), TLess<>());
 	}
@@ -90,7 +90,7 @@ namespace Algo
 	 * @return	returns		true if the range is a min-heap
 	 */
 	template <typename RangeType, typename ProjectionType, typename PredicateType>
-	FORCEINLINE bool IsHeapBy(const RangeType& Range, ProjectionType Projection, PredicateType Predicate)
+	[[nodiscard]] FORCEINLINE bool IsHeapBy(const RangeType& Range, ProjectionType Projection, PredicateType Predicate)
 	{
 		return AlgoImpl::IsHeapInternal(GetData(Range), GetNum(Range), MoveTemp(Projection), MoveTemp(Predicate));
 	}
diff --git a/Engine/Source/Runtime/Core/Public/Algo/IsSorted.h b/Engine/Source/Runtime/Core/Public/Algo/IsSorted.h
--- a/Engine/Source/Runtime/Core/Public/Algo/IsSorted.h
+++ b/Engine/Source/Runtime/Core/Public/Algo/IsSorted.h
@@ -63,7 +63,7 @@ namespace Algo
 	 * @return true if the range is sorted, false otherwise.
 	 */
 	template <typename RangeType>
-	FORCEINLINE bool IsSorted(const RangeType& Range)
+	[[nodiscard]] FORCEINLINE bool IsSorted(const RangeType& Range)
 	{
 		return AlgoImpl::IsSortedBy(GetData(Range), GetNum(Range), FIdentityFunctor(), TLess<>());
 	}
@@ -77,7 +77,7 @@ namespace Algo
 	 * @return true if the range is sorted, false otherwise.
 	 */
 	template <typename RangeType, typename PredType>
-	FORCEINLINE bool IsSorted(const RangeType& Range, PredType Pred)
+	[[nodiscard]] FORCEINLINE bool IsSorted(const RangeType& Range, PredType Pred)
 	{
 		return AlgoImpl::IsSortedBy(GetData(Range), GetNum(Range), FIdentityFunctor(), MoveTemp(Pred));
 	}
@@ -90,7 +90,7 @@ namespace Algo
 	 * @return true if the range is sorted, false otherwise.
 	 */
 	template <typename RangeType, typename ProjectionType>
-	FORCEINLINE bool IsSortedBy(const RangeType& Range, ProjectionType Projection)
+	[[nodiscard]] FORCEINLINE bool IsSortedBy(const RangeType& Range, ProjectionType Projection)
 	{
 		return AlgoImpl::IsSortedBy(GetData(Range), GetNum(Range), MoveTemp(Projection), TLess<>());
 	}
@@ -104,7 +104,7 @@ namespace Algo
 	 * @return true if the range is sorted, false otherwise.
 	 */
 	template <typename RangeType, typename ProjectionType, typename PredType>
-	FORCEINLINE bool IsSortedBy(const RangeType& Range, ProjectionType Projection, PredType Pred)
+	[[nodiscard]] FORCEINLINE bool IsSortedBy(const RangeType& Range, ProjectionType Projection, PredType Pred)
 	{
 		return AlgoImpl::IsSortedBy(GetData(Range), GetNum(Range), MoveTemp(Projection), MoveTemp(Pred));
 	}
diff --git a/Engine/Source/Runtime/Core/Public/Algo/LevenshteinDistance.h b/Engine/Source/Runtime/Core/Public/Algo/LevenshteinDistance.h
--- a/Engine/Source/Runtime/Core/Public/Algo/LevenshteinDistance.h
+++ b/Engine/Source/Runtime/Core/Public/Algo/LevenshteinDistance.h
@@ -20,7 +20,7 @@ namespace Algo
 	* @return				The number of operation to transform RangeA to RangeB
 	*/
 	template <typename RangeAType, typename RangeBType>
-	int32 LevenshteinDistance(const RangeAType& RangeA, const RangeBType& RangeB)
+	[[nodiscard]] int32 LevenshteinDistance(const RangeAType& RangeA, const RangeBType& RangeB)
 	{
 		const int32 LenA = GetNum(RangeA);
 		const int32 LenB = GetNum(RangeB);
diff --git a/Engine/Source/Runtime/Core/Public/Algo/MaxElement.h b/Engine/Source/Runtime/Core/Public/Algo/MaxElement.h
--- a/Engine/Source/Runtime/Core/Public/Algo/MaxElement.h
+++ b/Engine/Source/Runtime/Core/Public/Algo/MaxElement.h
@@ -37,7 +37,7 @@ namespace Algo
 	 * @return A pointer to the maximum element, or nullptr if the range was empty.
 	 */
 	template <typename RangeType>
-	FORCEINLINE auto MaxElement(RangeType&& Range)
+	[[nodiscard]] FORCEINLINE auto MaxElement(RangeType&& Range)
 		-> decltype(AlgoImpl::MaxElementBy(Range, FIdentityFunctor(), TLess<>()))
 	{
 		return AlgoImpl::MaxElementBy(Range, FIdentityFunctor(), TLess<>());
@@ -52,7 +52,7 @@ namespace Algo
 	 * @return A pointer to the maximum element, or nullptr if the range was empty.
 	 */
 	template <typename RangeType, typename ComparatorType>
-	FORCEINLINE auto MaxElement(RangeType&& Range, ComparatorType Comp)
+	[[nodiscard]] FORCEINLINE auto MaxElement(RangeType&& Range, ComparatorType Comp)
 		-> decltype(AlgoImpl::MaxElementBy(Range, FIdentityFunctor(), MoveTemp(Comp)))
 	{
 		return AlgoImpl::MaxElementBy(Range, FIdentityFunctor(), MoveTemp(Comp));
@@ -67,7 +67,7 @@ namespace Algo
 	 * @return A pointer to the maximum element, or nullptr if the range was empty.
 	 */
 	template <typename RangeType, typename ProjectionType>
-	FORCEINLINE auto MaxElementBy(RangeType&& Range, ProjectionType Proj)
+	[[nodiscard]] FORCEINLINE auto MaxElementBy(RangeType&& Range, ProjectionType Proj)
 		-> decltype(AlgoImpl::MaxElementBy(Range, MoveTemp(Proj), TLess<>()))
 	{
 		return AlgoImpl::MaxElementBy(Range, MoveTemp(Proj), TLess<>());
@@ -83,7 +83,7 @@ namespace Algo
 	 * @return A pointer to the maximum element, or nullptr if the range was empty.
 	 */
 	template <typename RangeType, typename ProjectionType, typename ComparatorType>
-	FORCEINLINE auto MaxElementBy(RangeType&& Range, ProjectionType Proj, ComparatorType Comp)
+	[[nodiscard]] FORCEINLINE auto MaxElementBy(RangeType&& Range, ProjectionType Proj, ComparatorType Comp)
 		-> decltype(AlgoImpl::MaxElementBy(Range, MoveTemp(Proj), MoveTemp(Comp)))
 	{
 		return AlgoImpl::MaxElementBy(Range, MoveTemp(Proj), MoveTemp(Comp));
diff --git a/Engine/Source/Runtime/Core/Public/Algo/MinElement.h b/Engine/Source/Runtime/Core/Public/Algo/MinElement.h
--- a/Engine/Source/Runtime/Core/Public/Algo/MinElement.h
+++ b/Engine/Source/Runtime/Core/Public/Algo/MinElement.h
@@ -37,7 +37,7 @@ namespace Algo
 	 * @return A pointer to the minimum element, or nullptr if the range was empty.
 	 */
 	template <typename RangeType>
-	FORCEINLINE auto MinElement(RangeType&& Range)
+	[[nodiscard]] FORCEINLINE auto MinElement(RangeType&& Range)
 		-> decltype(AlgoImpl::MinElementBy(Range, FIdentityFunctor(), TLess<>()))
 	{
 		return AlgoImpl::MinElementBy(Range, FIdentityFunctor(), TLess<>());
@@ -52,7 +52,7 @@ namespace Algo
 	 * @return A pointer to the minimum element, or nullptr if the range was empty.
 	 */
 	template <typename RangeType, typename ComparatorType>
-	FORCEINLINE auto MinElement(RangeType&& Range, ComparatorType Comp)
+	[[nodiscard]] FORCEINLINE auto MinElement(RangeType&& Range, ComparatorType Comp)
 		-> decltype(AlgoImpl::MinElementBy(Range, FIdentityFunctor(), MoveTemp(Comp)))
 	{
 		return AlgoImpl::MinElementBy(Range, FIdentityFunctor(), MoveTemp(Comp));
@@ -67,7 +67,7 @@ namespace Algo
 	 * @return A pointer to the minimum element, or nullptr if the range was empty.
 	 */
 	template <typename RangeType, typename ProjectionType>
-	FORCEINLINE auto MinElementBy(RangeType&& Range, ProjectionType Proj)
+	[[nodiscard]] FORCEINLINE auto MinElementBy(RangeType&& Range, ProjectionType Proj)
 		-> decltype(AlgoImpl::MinElementBy(Range, MoveTemp(Proj), TLess<>()))
 	{
 		return AlgoImpl::MinElementBy(Range, MoveTemp(Proj), TLess<>());
@@ -83,7 +83,7 @@ namespace Algo
 	 * @return A pointer to the minimum element, or nullptr if the range was empty.
 	 */
 	template <typename RangeType, typename ProjectionType, typename ComparatorType>
-	FORCEINLINE auto MinElementBy(RangeType&& Range, ProjectionType Proj, ComparatorType Comp)
+	[[nodiscard]] FORCEINLINE auto MinElementBy(RangeType&& Range, ProjectionType Proj, ComparatorType Comp)
 		-> decltype(AlgoImpl::MinElementBy(Range, MoveTemp(Proj), MoveTemp(Comp)))
 	{
 		return AlgoImpl::MinElementBy(Range, MoveTemp(Proj), MoveTemp(Comp));
diff --git a/Engine/Source/Runtime/Core/Public/Algo/Mismatch.h b/Engine/Source/Runtime/Core/Public/Algo/Mismatch.h
--- a/Engine/Source/Runtime/Core/Public/Algo/Mismatch.h
+++ b/Engine/Source/Runtime/Core/Public/Algo/Mismatch.h
@@ -55,7 +55,7 @@ namespace Algo
 	 * @return The index at which the containers differ.
 	 */
 	template <typename InputAType, typename InputBType>
-	constexpr auto Mismatch(const InputAType& InputA, const InputBType& InputB) -> std::common_type_t<decltype(GetNum(InputA)), decltype(GetNum(InputB))>
+	[[nodiscard]] constexpr auto Mismatch(const InputAType& InputA, const InputBType& InputB) -> std::common_type_t<decltype(GetNum(InputA)), decltype(GetNum(InputB))>
 	{
 		return Private::Mismatch(InputA, InputB, FIdentityFunctor(), TEqualTo<>());
 	}
@@ -70,7 +70,7 @@ namespace Algo
 	 * @return The index at which the containers differ.
 	 */
 	template <typename InputAType, typename InputBType, typename PredicateType>
-	constexpr auto Mismatch(const InputAType& InputA, const InputBType& InputB, PredicateType Predicate) -> std::common_type_t<decltype(GetNum(InputA)), decltype(GetNum(InputB))>
+	[[nodiscard]] constexpr auto Mismatch(const InputAType& InputA, const InputBType& InputB, PredicateType Predicate) -> std::common_type_t<decltype(GetNum(InputA)), decltype(GetNum(InputB))>
 	{
 		return Private::Mismatch(InputA, InputB, FIdentityFunctor(), MoveTemp(Predicate));
 	}
@@ -85,7 +85,7 @@ namespace Algo
 	 * @return The index at which the containers differ.
 	 */
 	template <typename InputAType, typename InputBType, typename ProjectionType>
-	constexpr auto MismatchBy(const InputAType& InputA, const InputBType& InputB, ProjectionType Projection) -> std::common_type_t<decltype(GetNum(InputA)), decltype(GetNum(InputB))>
+	[[nodiscard]] constexpr auto MismatchBy(const InputAType& InputA, const InputBType& InputB, ProjectionType Projection) -> std::common_type_t<decltype(GetNum(InputA)), decltype(GetNum(InputB))>
 	{
 		return Private::Mismatch(InputA, InputB, MoveTemp(Projection), TEqualTo<>());
 	}
@@ -101,7 +101,7 @@ namespace Algo
 	 * @return The index at which the containers differ.
 	 */
 	template <typename InputAType, typename InputBType, typename ProjectionType, typename PredicateType>
-	constexpr auto MismatchBy(const InputAType& InputA, const InputBType& InputB, ProjectionType Projection, PredicateType Predicate) -> std::common_type_t<decltype(GetNum(InputA)), decltype(GetNum(InputB))>
+	[[nodiscard]] constexpr auto MismatchBy(const InputAType& InputA, const InputBType& InputB, ProjectionType Projection, PredicateType Predicate) -> std::common_type_t<decltype(GetNum(InputA)), decltype(GetNum(InputB))>
 	{
 		return Private::Mismatch(InputA, InputB, MoveTemp(Projection), MoveTemp(Predicate));
 	}
diff --git a/Engine/Source/Runtime/Core/Public/Algo/NoneOf.h b/Engine/Source/Runtime/Core/Public/Algo/NoneOf.h
--- a/Engine/Source/Runtime/Core/Public/Algo/NoneOf.h
+++ b/Engine/Source/Runtime/Core/Public/Algo/NoneOf.h
@@ -17,7 +17,7 @@ namespace Algo
 	 * @return  true if no element is truthy, false otherwise.
 	 */
 	template <typename RangeType>
-	FORCEINLINE bool NoneOf(const RangeType& Range)
+	[[nodiscard]] FORCEINLINE bool NoneOf(const RangeType& Range)
 	{
 		for (const auto& Element : Range)
 		{
@@ -39,7 +39,7 @@ namespace Algo
 	 * @return  true if none of the projections are truthy, false otherwise.
 	 */
 	template <typename RangeType, typename ProjectionType>
-	FORCEINLINE bool NoneOf(const RangeType& Range, ProjectionType Projection)
+	[[nodiscard]] FORCEINLINE bool NoneOf(const RangeType& Range, ProjectionType Projection)
 	{
 		for (const auto& Element : Range)
 		{
@@ -62,7 +62,7 @@ namespace Algo
 	 * @return  true if none of the projections are truthy, false otherwise.
 	 */
 	template <typename RangeType, typename ProjectionType>
-	FORCEINLINE bool NoneOf(const RangeType& Range, ProjectionType Projection, ENoRef NoRef)
+	[[nodiscard]] FORCEINLINE bool NoneOf(const RangeType& Range, ProjectionType Projection, ENoRef NoRef)
 	{
 		for (const auto Element : Range)
 		{
diff --git a/Engine/Source/Runtime/Core/Public/Algo/SelectRandomWeighted.h b/Engine/Source/Runtime/Core/Public/Algo/SelectRandomWeighted.h
--- a/Engine/Source/Runtime/Core/Public/Algo/SelectRandomWeighted.h
+++ b/Engine/Source/Runtime/Core/Public/Algo/SelectRandomWeighted.h
@@ -12,7 +12,7 @@
 namespace AlgoImpl
 {
 	template <typename RangeType, typename ProjectionType>
-	typename TRangePointerType<typename TRemoveReference<RangeType>::Type>::Type SelectRandomWeightedBy(RangeType&& Range, ProjectionType Proj)
+	[[nodiscard]] typename TRangePointerType<typename TRemoveReference<RangeType>::Type>::Type SelectRandomWeightedBy(RangeType&& Range, ProjectionType Proj)
 	{
 		using WeightType = std::decay_t<TInvokeResult_T<ProjectionType, TElementType_T<RangeType>>>;
 
diff --git a/Engine/Source/Runtime/Core/Public/Algo/TiedTupleOutput.h b/Engine/Source/Runtime/Core/Public/Algo/TiedTupleOutput.h
--- a/Engine/Source/Runtime/Core/Public/Algo/TiedTupleOutput.h
+++ b/Engine/Source/Runtime/Core/Public/Algo/TiedTupleOutput.h
@@ -1,4 +1,4 @@
-﻿// Copyright Epic Games, Inc. All Rights Reserved.
+// Copyright Epic Games, Inc. All Rights Reserved.
 
 #pragma once
 
@@ -49,7 +49,7 @@ namespace Algo
 	 *	@return An object with an Add function that will call Add on the tied elements
 	 */
 	template <typename... OutputTypes>
-	Private::TTiedTupleAdder<OutputTypes...> TieTupleAdd(OutputTypes&... Outputs)
+	[[nodiscard]] Private::TTiedTupleAdder<OutputTypes...> TieTupleAdd(OutputTypes&... Outputs)
 	{
 		return { { Outputs... } };
 	}
diff --git a/Engine/Source/Runtime/Core/Public/Algo/Unique.h b/Engine/Source/Runtime/Core/Public/Algo/Unique.h
--- a/Engine/Source/Runtime/Core/Public/Algo/Unique.h
+++ b/Engine/Source/Runtime/Core/Public/Algo/Unique.h
@@ -56,25 +56,25 @@ namespace Algo
 	 * See https://en.cppreference.com/w/cpp/algorithm/unique
 	 */
 	template<typename RangeType>
-	auto Unique(RangeType&& Range) -> decltype(AlgoImpl::Unique(GetData(Range), GetNum(Range), FIdentityFunctor(), TEqualTo<>{}))
+	[[nodiscard]] auto Unique(RangeType&& Range) -> decltype(AlgoImpl::Unique(GetData(Range), GetNum(Range), FIdentityFunctor(), TEqualTo<>{}))
 	{
 		return AlgoImpl::Unique(GetData(Range), GetNum(Range), FIdentityFunctor(), TEqualTo<>{});
 	}
 
 	template<typename RangeType, typename BinaryPredicate>
-	auto Unique(RangeType&& Range, BinaryPredicate Predicate) -> decltype(AlgoImpl::Unique(GetData(Range), GetNum(Range), FIdentityFunctor(), MoveTemp(Predicate)))
+	[[nodiscard]] auto Unique(RangeType&& Range, BinaryPredicate Predicate) -> decltype(AlgoImpl::Unique(GetData(Range), GetNum(Range), FIdentityFunctor(), MoveTemp(Predicate)))
 	{
 		return AlgoImpl::Unique(GetData(Range), GetNum(Range), FIdentityFunctor(), MoveTemp(Predicate));
 	}
 
 	template<typename RangeType, typename ProjectionType>
-	auto UniqueBy(RangeType&& Range, ProjectionType Proj) -> decltype(AlgoImpl::Unique(GetData(Range), GetNum(Range), MoveTemp(Proj), TEqualTo<>{}))
+	[[nodiscard]] auto UniqueBy(RangeType&& Range, ProjectionType Proj) -> decltype(AlgoImpl::Unique(GetData(Range), GetNum(Range), MoveTemp(Proj), TEqualTo<>{}))
 	{
 		return AlgoImpl::Unique(GetData(Range), GetNum(Range), MoveTemp(Proj), TEqualTo<>{});
 	}
 
 	template<typename RangeType, typename ProjectionType, typename BinaryPredicate>
-	auto UniqueBy(RangeType&& Range, ProjectionType Proj, BinaryPredicate Predicate) -> decltype(AlgoImpl::Unique(GetData(Range), GetNum(Range), MoveTemp(Proj), MoveTemp(Predicate)))
+	[[nodiscard]] auto UniqueBy(RangeType&& Range, ProjectionType Proj, BinaryPredicate Predicate) -> decltype(AlgoImpl::Unique(GetData(Range), GetNum(Range), MoveTemp(Proj), MoveTemp(Predicate)))
 	{
 		return AlgoImpl::Unique(GetData(Range), GetNum(Range), MoveTemp(Proj), MoveTemp(Predicate));
 	}
