Added [[nodiscard]] to relevant container functions.

diff --git a/Engine/Source/Runtime/Core/Public/Containers/Array.h b/Engine/Source/Runtime/Core/Public/Containers/Array.h
--- a/Engine/Source/Runtime/Core/Public/Containers/Array.h
+++ b/Engine/Source/Runtime/Core/Public/Containers/Array.h
@@ -62,7 +62,7 @@ template <typename ContainerType, typename ElementType, typename SizeType>
 class TIndexedContainerIterator
 {
 public:
-	TIndexedContainerIterator(ContainerType& InContainer, SizeType StartIndex = 0)
+	[[nodiscard]] TIndexedContainerIterator(ContainerType& InContainer, SizeType StartIndex = 0)
 		: Container(InContainer)
 		, Index    (StartIndex)
 	{
@@ -101,7 +101,7 @@ public:
 		return *this;
 	}
 
-	TIndexedContainerIterator operator+(SizeType Offset) const
+	[[nodiscard]] TIndexedContainerIterator operator+(SizeType Offset) const
 	{
 		TIndexedContainerIterator Tmp(*this);
 		return Tmp += Offset;
@@ -112,30 +112,30 @@ public:
 		return *this += -Offset;
 	}
 
-	TIndexedContainerIterator operator-(SizeType Offset) const
+	[[nodiscard]] TIndexedContainerIterator operator-(SizeType Offset) const
 	{
 		TIndexedContainerIterator Tmp(*this);
 		return Tmp -= Offset;
 	}
 
-	FORCEINLINE ElementType& operator* () const
+	[[nodiscard]] FORCEINLINE ElementType& operator* () const
 	{
 		return Container[ Index ];
 	}
 
-	FORCEINLINE ElementType* operator->() const
+	[[nodiscard]] FORCEINLINE ElementType* operator->() const
 	{
 		return &Container[ Index ];
 	}
 
 	/** conversion to "bool" returning true if the iterator has not reached the last element. */
-	FORCEINLINE explicit operator bool() const
+	[[nodiscard]] FORCEINLINE explicit operator bool() const
 	{
 		return Container.IsValidIndex(Index);
 	}
 
 	/** Returns an index to the current element. */
-	SizeType GetIndex() const
+	[[nodiscard]] SizeType GetIndex() const
 	{
 		return Index;
 	}
@@ -170,8 +170,8 @@ public:
 		Index--;
 	}
 
-	FORCEINLINE bool operator==(const TIndexedContainerIterator& Rhs) const { return &Container == &Rhs.Container && Index == Rhs.Index; }
-	FORCEINLINE bool operator!=(const TIndexedContainerIterator& Rhs) const { return &Container != &Rhs.Container || Index != Rhs.Index; }
+	[[nodiscard]] FORCEINLINE bool operator==(const TIndexedContainerIterator& Rhs) const { return &Container == &Rhs.Container && Index == Rhs.Index; }
+	[[nodiscard]] FORCEINLINE bool operator!=(const TIndexedContainerIterator& Rhs) const { return &Container != &Rhs.Container || Index != Rhs.Index; }
 
 private:
 
@@ -182,7 +182,7 @@ private:
 
 /** operator + */
 template <typename ContainerType, typename ElementType, typename SizeType>
-FORCEINLINE TIndexedContainerIterator<ContainerType, ElementType, SizeType> operator+(SizeType Offset, TIndexedContainerIterator<ContainerType, ElementType, SizeType> RHS)
+[[nodiscard]] FORCEINLINE TIndexedContainerIterator<ContainerType, ElementType, SizeType> operator+(SizeType Offset, TIndexedContainerIterator<ContainerType, ElementType, SizeType> RHS)
 {
 	return RHS + Offset;
 }
@@ -201,14 +201,14 @@ FORCEINLINE TIndexedContainerIterator<ContainerType, ElementType, SizeType> oper
 		//
 		// We do add an operator-- to help FString implementation
 
-		explicit TCheckedPointerIterator(const SizeType& InNum, ElementType* InPtr)
+		[[nodiscard]] explicit TCheckedPointerIterator(const SizeType& InNum, ElementType* InPtr)
 			: Ptr       (InPtr)
 			, CurrentNum(InNum)
 			, InitialNum(InNum)
 		{
 		}
 
-		FORCEINLINE ElementType& operator*() const
+		[[nodiscard]] FORCEINLINE ElementType& operator*() const
 		{
 			if constexpr (bReverse)
 			{
@@ -246,7 +246,7 @@ FORCEINLINE TIndexedContainerIterator<ContainerType, ElementType, SizeType> oper
 			return *this;
 		}
 
-		FORCEINLINE bool operator!=(const TCheckedPointerIterator& Rhs) const
+		[[nodiscard]] FORCEINLINE bool operator!=(const TCheckedPointerIterator& Rhs) const
 		{
 			// We only need to do the check in this operator, because no other operator will be
 			// called until after this one returns.
@@ -268,12 +268,12 @@ FORCEINLINE TIndexedContainerIterator<ContainerType, ElementType, SizeType> oper
 template <typename ElementType, typename IteratorType>
 struct TDereferencingIterator
 {
-	explicit TDereferencingIterator(IteratorType InIter)
+	[[nodiscard]] explicit TDereferencingIterator(IteratorType InIter)
 		: Iter(InIter)
 	{
 	}
 
-	FORCEINLINE ElementType& operator*() const
+	[[nodiscard]] FORCEINLINE ElementType& operator*() const
 	{
 		return *(ElementType*)*Iter;
 	}
@@ -284,7 +284,7 @@ struct TDereferencingIterator
 		return *this;
 	}
 
-	FORCEINLINE bool operator!=(const TDereferencingIterator& Rhs) const
+	[[nodiscard]] FORCEINLINE bool operator!=(const TDereferencingIterator& Rhs) const
 	{
 		return Iter != Rhs.Iter;
 	}
@@ -298,13 +298,13 @@ namespace UE::Core::Private
 	// Simply forwards to an unqualified GetData(), but can be called from within a container or view
 	// where GetData() is already a member and so hides any others.
 	template <typename T>
-	FORCEINLINE decltype(auto) GetDataHelper(T&& Arg)
+	[[nodiscard]] FORCEINLINE decltype(auto) GetDataHelper(T&& Arg)
 	{
 		return GetData(Forward<T>(Arg));
 	}
 
 	template <typename FromArrayType, typename ToArrayType>
-	constexpr bool CanMoveTArrayPointersBetweenArrayTypes()
+	[[nodiscard]] constexpr bool CanMoveTArrayPointersBetweenArrayTypes()
 	{
 		using FromAllocatorType          = typename FromArrayType::AllocatorType;
 		using ToAllocatorType            = typename ToArrayType::AllocatorType;
@@ -363,7 +363,7 @@ namespace UE::Core::Private
 	//
 	// When C++20 is guaranteed, concept checks can be used instead.
 	template <typename AllocatorType>
-	constexpr uint32 GetAllocatorFlags()
+	[[nodiscard]] constexpr uint32 GetAllocatorFlags()
 	{
 		uint32 Result = 0;
 		if constexpr (TAllocatorTraits<AllocatorType>::SupportsElementAlignment)
@@ -595,10 +595,11 @@ public:
 	/**
 	 * Constructor, initializes element number counters.
 	 */
-	FORCEINLINE TArray()
+	[[nodiscard]] FORCEINLINE TArray()
 		: ArrayNum(0)
 		, ArrayMax(AllocatorInstance.GetInitialCapacity())
-	{}
+	{
+	}
 
 	/**
 	 * Constructor from a raw array of elements.
@@ -607,7 +608,7 @@ public:
 	 * @param Count The number of elements to copy from Ptr.
 	 * @see Append
 	 */
-	FORCEINLINE TArray(const ElementType* Ptr, SizeType Count)
+	[[nodiscard]] FORCEINLINE TArray(const ElementType* Ptr, SizeType Count)
 	{
 		if (Count < 0)
 		{
@@ -621,12 +622,12 @@ public:
 	}
 
 	template <typename OtherElementType, typename OtherSizeType>
-	explicit TArray(const TArrayView<OtherElementType, OtherSizeType>& Other);
+	[[nodiscard]] explicit TArray(const TArrayView<OtherElementType, OtherSizeType>& Other);
 
 	/**
 	 * Initializer list constructor
 	 */
-	TArray(std::initializer_list<InElementType> InitList)
+	[[nodiscard]] TArray(std::initializer_list<InElementType> InitList)
 	{
 		// This is not strictly legal, as std::initializer_list's iterators are not guaranteed to be pointers, but
 		// this appears to be the case on all of our implementations.  Also, if it's not true on a new implementation,
@@ -644,7 +645,7 @@ public:
 		typename OtherAllocator
 		UE_REQUIRES(UE::Core::Private::TArrayElementsAreCompatible_V<ElementType, const OtherElementType&>)
 	>
-	FORCEINLINE explicit TArray(const TArray<OtherElementType, OtherAllocator>& Other)
+	[[nodiscard]] FORCEINLINE explicit TArray(const TArray<OtherElementType, OtherAllocator>& Other)
 	{
 		CopyToEmpty(Other.GetData(), Other.Num(), 0);
 	}
@@ -654,7 +655,7 @@ public:
 	 *
 	 * @param Other The source array to copy.
 	 */
-	FORCEINLINE TArray(const TArray& Other)
+	[[nodiscard]] FORCEINLINE TArray(const TArray& Other)
 	{
 		CopyToEmpty(Other.GetData(), Other.Num(), 0);
 	}
@@ -666,7 +667,7 @@ public:
 	 * @param ExtraSlack Tells how much extra memory should be preallocated
 	 *                   at the end of the array in the number of elements.
 	 */
-	FORCEINLINE TArray(const TArray& Other, SizeType ExtraSlack)
+	[[nodiscard]] FORCEINLINE TArray(const TArray& Other, SizeType ExtraSlack)
 	{
 		CopyToEmptyWithSlack(Other.GetData(), Other.Num(), 0, ExtraSlack);
 	}
@@ -832,7 +833,7 @@ public:
 	 *
 	 * @param Other Array to move from.
 	 */
-	FORCEINLINE TArray(TArray&& Other)
+	[[nodiscard]] FORCEINLINE TArray(TArray&& Other)
 	{
 		MoveOrCopy(*this, Other, 0);
 	}
@@ -847,7 +848,7 @@ public:
 		typename OtherAllocator
 		UE_REQUIRES(UE::Core::Private::TArrayElementsAreCompatible_V<ElementType, OtherElementType&&>)
 	>
-	FORCEINLINE explicit TArray(TArray<OtherElementType, OtherAllocator>&& Other)
+	[[nodiscard]] FORCEINLINE explicit TArray(TArray<OtherElementType, OtherAllocator>&& Other)
 	{
 		MoveOrCopy(*this, Other, 0);
 	}
@@ -863,7 +864,7 @@ public:
 		typename OtherElementType
 		UE_REQUIRES(UE::Core::Private::TArrayElementsAreCompatible_V<ElementType, OtherElementType&&>)
 	>
-	TArray(TArray<OtherElementType, AllocatorType>&& Other, SizeType ExtraSlack)
+	[[nodiscard]] TArray(TArray<OtherElementType, AllocatorType>&& Other, SizeType ExtraSlack)
 	{
 		MoveOrCopyWithSlack(*this, Other, 0, ExtraSlack);
 	}
@@ -898,13 +899,13 @@ public:
 	constexpr static bool bHasIntrusiveUnsetOptionalState = true;
 	using IntrusiveUnsetOptionalStateType = TArray;
 
-	explicit TArray(FIntrusiveUnsetOptionalState Tag)
+	[[nodiscard]] explicit TArray(FIntrusiveUnsetOptionalState Tag)
 		: ArrayNum(0)
 		, ArrayMax(-1)
 	{
 		// Use ArrayMax == -1 as our intrusive state so that the destructor still works without change, as it doesn't use ArrayMax.
 	}
-	bool operator==(FIntrusiveUnsetOptionalState Tag) const
+	[[nodiscard]] bool operator==(FIntrusiveUnsetOptionalState Tag) const
 	{
 		return ArrayMax == -1;
 	}
@@ -917,7 +918,7 @@ public:
 	 *
 	 * @returns Pointer to first array entry or nullptr if ArrayMax == 0.
 	 */
-	FORCEINLINE ElementType* GetData() UE_LIFETIMEBOUND
+	[[nodiscard]] FORCEINLINE ElementType* GetData() UE_LIFETIMEBOUND
 	{
 		return (ElementType*)AllocatorInstance.GetAllocation();
 	}
@@ -927,7 +928,7 @@ public:
 	 *
 	 * @returns Pointer to first array entry or nullptr if ArrayMax == 0.
 	 */
-	FORCEINLINE const ElementType* GetData() const UE_LIFETIMEBOUND
+	[[nodiscard]] FORCEINLINE const ElementType* GetData() const UE_LIFETIMEBOUND
 	{
 		return (const ElementType*)AllocatorInstance.GetAllocation();
 	}
@@ -937,7 +938,7 @@ public:
 	 *
 	 * @returns Size in bytes of array type.
 	 */
-	FORCEINLINE static constexpr uint32 GetTypeSize()
+	[[nodiscard]] FORCEINLINE static constexpr uint32 GetTypeSize()
 	{
 		return sizeof(ElementType);
 	}
@@ -949,7 +950,7 @@ public:
 	 *
 	 * @returns Number of bytes allocated by this container.
 	 */
-	FORCEINLINE SIZE_T GetAllocatedSize(void) const
+	[[nodiscard]] FORCEINLINE SIZE_T GetAllocatedSize(void) const
 	{
 		return AllocatorInstance.GetAllocatedSize(ArrayMax, sizeof(ElementType));
 	}
@@ -959,7 +960,7 @@ public:
 	 *
 	 * @see Num, Shrink
 	 */
-	FORCEINLINE SizeType GetSlack() const
+	[[nodiscard]] FORCEINLINE SizeType GetSlack() const
 	{
 		return ArrayMax - ArrayNum;
 	}
@@ -1012,7 +1013,7 @@ public:
 	 * @param Index Index to test.
 	 * @returns True if index is valid. False otherwise.
 	 */
-	FORCEINLINE bool IsValidIndex(SizeType Index) const
+	[[nodiscard]] FORCEINLINE bool IsValidIndex(SizeType Index) const
 	{
 		return Index >= 0 && Index < ArrayNum;
 	}
@@ -1023,7 +1024,7 @@ public:
 	 * @returns True if the array is empty.
 	 * @see Num
 	 */
-	bool IsEmpty() const
+	[[nodiscard]] bool IsEmpty() const
 	{
 		return ArrayNum == 0;
 	}
@@ -1034,13 +1035,13 @@ public:
 	 * @returns Number of elements in array.
 	 * @see GetSlack
 	 */
-	FORCEINLINE SizeType Num() const
+	[[nodiscard]] FORCEINLINE SizeType Num() const
 	{
 		return ArrayNum;
 	}
 
 	/** @returns Number of bytes used, excluding slack */
-	FORCEINLINE SIZE_T NumBytes() const
+	[[nodiscard]] FORCEINLINE SIZE_T NumBytes() const
 	{
 		return static_cast<SIZE_T>(ArrayNum) * sizeof(ElementType);
 	}
@@ -1051,7 +1052,7 @@ public:
 	 * @returns Maximum number of elements in array.
 	 * @see GetSlack
 	 */
-	FORCEINLINE SizeType Max() const
+	[[nodiscard]] FORCEINLINE SizeType Max() const
 	{
 		return ArrayMax;
 	}
@@ -1061,7 +1062,7 @@ public:
 	 *
 	 * @returns Reference to indexed element.
 	 */
-	FORCEINLINE ElementType& operator[](SizeType Index) UE_LIFETIMEBOUND
+	[[nodiscard]] FORCEINLINE ElementType& operator[](SizeType Index) UE_LIFETIMEBOUND
 	{
 		RangeCheck(Index);
 		return GetData()[Index];
@@ -1074,7 +1075,7 @@ public:
 	 *
 	 * @returns Reference to indexed element.
 	 */
-	FORCEINLINE const ElementType& operator[](SizeType Index) const UE_LIFETIMEBOUND
+	[[nodiscard]] FORCEINLINE const ElementType& operator[](SizeType Index) const UE_LIFETIMEBOUND
 	{
 		RangeCheck(Index);
 		return GetData()[Index];
@@ -1138,7 +1139,7 @@ public:
 	 * @returns Reference to the top element.
 	 * @see Pop, Push
 	 */
-	FORCEINLINE ElementType& Top() UE_LIFETIMEBOUND
+	[[nodiscard]] FORCEINLINE ElementType& Top() UE_LIFETIMEBOUND
 	{
 		return Last();
 	}
@@ -1151,7 +1152,7 @@ public:
 	 * @returns Reference to the top element.
 	 * @see Pop, Push
 	 */
-	FORCEINLINE const ElementType& Top() const UE_LIFETIMEBOUND
+	[[nodiscard]] FORCEINLINE const ElementType& Top() const UE_LIFETIMEBOUND
 	{
 		return Last();
 	}
@@ -1162,7 +1163,7 @@ public:
 	 * @param IndexFromTheEnd (Optional) Index from the end of array (default = 0).
 	 * @returns Reference to n-th last element from the array.
 	 */
-	FORCEINLINE ElementType& Last(SizeType IndexFromTheEnd = 0) UE_LIFETIMEBOUND
+	[[nodiscard]] FORCEINLINE ElementType& Last(SizeType IndexFromTheEnd = 0) UE_LIFETIMEBOUND
 	{
 		RangeCheck(ArrayNum - IndexFromTheEnd - 1);
 		return GetData()[ArrayNum - IndexFromTheEnd - 1];
@@ -1176,7 +1177,7 @@ public:
 	 * @param IndexFromTheEnd (Optional) Index from the end of array (default = 0).
 	 * @returns Reference to n-th last element from the array.
 	 */
-	FORCEINLINE const ElementType& Last(SizeType IndexFromTheEnd = 0) const UE_LIFETIMEBOUND
+	[[nodiscard]] FORCEINLINE const ElementType& Last(SizeType IndexFromTheEnd = 0) const UE_LIFETIMEBOUND
 	{
 		RangeCheck(ArrayNum - IndexFromTheEnd - 1);
 		return GetData()[ArrayNum - IndexFromTheEnd - 1];
@@ -1224,7 +1225,7 @@ public:
 	 * @returns Index of the found element. INDEX_NONE otherwise.
 	 * @see FindLast, FindLastByPredicate
 	 */
-	SizeType Find(const ElementType& Item) const
+	[[nodiscard]] SizeType Find(const ElementType& Item) const
 	{
 		const ElementType* RESTRICT Start = GetData();
 		for (const ElementType* RESTRICT Data = Start, *RESTRICT DataEnd = Data + ArrayNum; Data != DataEnd; ++Data)
@@ -1257,7 +1258,7 @@ public:
 	 * @param Item Item to look for.
 	 * @returns Index of the found element. INDEX_NONE otherwise.
 	 */
-	SizeType FindLast(const ElementType& Item) const
+	[[nodiscard]] SizeType FindLast(const ElementType& Item) const
 	{
 		for (const ElementType* RESTRICT Start = GetData(), *RESTRICT Data = Start + ArrayNum; Data != Start; )
 		{
@@ -1278,7 +1279,7 @@ public:
 	 * @returns Index of the found element. INDEX_NONE otherwise.
 	 */
 	template <typename Predicate>
-	SizeType FindLastByPredicate(Predicate Pred, SizeType Count) const
+	[[nodiscard]] SizeType FindLastByPredicate(Predicate Pred, SizeType Count) const
 	{
 		check(Count >= 0 && Count <= this->Num());
 		for (const ElementType* RESTRICT Start = GetData(), *RESTRICT Data = Start + Count; Data != Start; )
@@ -1299,7 +1300,7 @@ public:
 	 * @returns Index of the found element. INDEX_NONE otherwise.
 	 */
 	template <typename Predicate>
-	FORCEINLINE SizeType FindLastByPredicate(Predicate Pred) const
+	[[nodiscard]] FORCEINLINE SizeType FindLastByPredicate(Predicate Pred) const
 	{
 		return FindLastByPredicate(Pred, ArrayNum);
 	}
@@ -1312,7 +1313,7 @@ public:
 	 * @returns Index to the first matching element, or INDEX_NONE if none is found.
 	 */
 	template <typename KeyType>
-	SizeType IndexOfByKey(const KeyType& Key) const
+	[[nodiscard]] SizeType IndexOfByKey(const KeyType& Key) const
 	{
 		const ElementType* RESTRICT Start = GetData();
 		for (const ElementType* RESTRICT Data = Start, *RESTRICT DataEnd = Start + ArrayNum; Data != DataEnd; ++Data)
@@ -1332,7 +1333,7 @@ public:
 	 * @returns Index to the first matching element, or INDEX_NONE if none is found.
 	 */
 	template <typename Predicate>
-	SizeType IndexOfByPredicate(Predicate Pred) const
+	[[nodiscard]] SizeType IndexOfByPredicate(Predicate Pred) const
 	{
 		const ElementType* RESTRICT Start = GetData();
 		for (const ElementType* RESTRICT Data = Start, *RESTRICT DataEnd = Start + ArrayNum; Data != DataEnd; ++Data)
@@ -1354,7 +1355,7 @@ public:
 	 * @see Find
 	 */
 	template <typename KeyType>
-	FORCEINLINE const ElementType* FindByKey(const KeyType& Key) const
+	[[nodiscard]] FORCEINLINE const ElementType* FindByKey(const KeyType& Key) const
 	{
 		return const_cast<TArray*>(this)->FindByKey(Key);
 	}
@@ -1368,7 +1369,7 @@ public:
 	 * @see Find
 	 */
 	template <typename KeyType>
-	ElementType* FindByKey(const KeyType& Key)
+	[[nodiscard]] ElementType* FindByKey(const KeyType& Key)
 	{
 		for (ElementType* RESTRICT Data = GetData(), *RESTRICT DataEnd = Data + ArrayNum; Data != DataEnd; ++Data)
 		{
@@ -1389,7 +1390,7 @@ public:
 	 * @see FilterByPredicate, ContainsByPredicate
 	 */
 	template <typename Predicate>
-	FORCEINLINE const ElementType* FindByPredicate(Predicate Pred) const
+	[[nodiscard]] FORCEINLINE const ElementType* FindByPredicate(Predicate Pred) const
 	{
 		return const_cast<TArray*>(this)->FindByPredicate(Pred);
 	}
@@ -1401,7 +1402,7 @@ public:
 	 * @see FilterByPredicate, ContainsByPredicate
 	 */
 	template <typename Predicate>
-	ElementType* FindByPredicate(Predicate Pred)
+	[[nodiscard]] ElementType* FindByPredicate(Predicate Pred)
 	{
 		for (ElementType* RESTRICT Data = GetData(), *RESTRICT DataEnd = Data + ArrayNum; Data != DataEnd; ++Data)
 		{
@@ -1423,7 +1424,7 @@ public:
 	 * @see FindByPredicate, ContainsByPredicate
 	 */
 	template <typename Predicate>
-	TArray<ElementType> FilterByPredicate(Predicate Pred) const
+	[[nodiscard]] TArray<ElementType> FilterByPredicate(Predicate Pred) const
 	{
 		TArray<ElementType> FilterResults;
 		for (const ElementType* RESTRICT Data = GetData(), *RESTRICT DataEnd = Data + ArrayNum; Data != DataEnd; ++Data)
@@ -1443,7 +1444,7 @@ public:
 	 * @see ContainsByPredicate, FilterByPredicate, FindByPredicate
 	 */
 	template <typename ComparisonType>
-	bool Contains(const ComparisonType& Item) const
+	[[nodiscard]] bool Contains(const ComparisonType& Item) const
 	{
 		for (const ElementType* RESTRICT Data = GetData(), *RESTRICT DataEnd = Data + ArrayNum; Data != DataEnd; ++Data)
 		{
@@ -1463,7 +1464,7 @@ public:
 	 * @see Contains, Find
 	 */
 	template <typename Predicate>
-	FORCEINLINE bool ContainsByPredicate(Predicate Pred) const
+	[[nodiscard]] FORCEINLINE bool ContainsByPredicate(Predicate Pred) const
 	{
 		return FindByPredicate(Pred) != nullptr;
 	}
@@ -1474,7 +1475,7 @@ public:
 	 * @param OtherArray Array to compare.
 	 * @returns True if this array is the same as OtherArray. False otherwise.
 	 */
-	bool operator==(const TArray& OtherArray) const
+	[[nodiscard]] bool operator==(const TArray& OtherArray) const
 	{
 		SizeType Count = Num();
 
@@ -1488,7 +1489,7 @@ public:
 	 * @returns True if this array is NOT the same as OtherArray. False otherwise.
 	 */
 #if !PLATFORM_COMPILER_HAS_GENERATED_COMPARISON_OPERATORS
-	FORCEINLINE bool operator!=(const TArray& OtherArray) const
+	[[nodiscard]] FORCEINLINE bool operator!=(const TArray& OtherArray) const
 	{
 		return !(*this == OtherArray);
 	}
@@ -1892,7 +1893,7 @@ public:
 	 * @return A reference to the newly-inserted element.
 	 * @see Add, Remove
 	 */
-	ElementType& Insert_GetRef(ElementType&& Item, SizeType Index) UE_LIFETIMEBOUND
+	[[nodiscard]] ElementType& Insert_GetRef(ElementType&& Item, SizeType Index) UE_LIFETIMEBOUND
 	{
 		CheckAddress(&Item);
 
@@ -1912,7 +1913,7 @@ public:
 	 * @return A reference to the newly-inserted element.
 	 * @see Add, Remove
 	 */
-	ElementType& Insert_GetRef(const ElementType& Item, SizeType Index) UE_LIFETIMEBOUND
+	[[nodiscard]] ElementType& Insert_GetRef(const ElementType& Item, SizeType Index) UE_LIFETIMEBOUND
 	{
 		CheckAddress(&Item);
 
@@ -2492,7 +2493,7 @@ public:
 	 * @return A reference to the newly-inserted element.
 	 */
 	template <typename... ArgsType>
-	FORCEINLINE ElementType& Emplace_GetRef(ArgsType&&... Args) UE_LIFETIMEBOUND
+	[[nodiscard]] FORCEINLINE ElementType& Emplace_GetRef(ArgsType&&... Args) UE_LIFETIMEBOUND
 	{
 		// If this fails to compile when trying to call Emplace with a non-public constructor,
 		// do not make TArray a friend.
@@ -2549,7 +2550,7 @@ public:
 	 * @return A reference to the newly-inserted element.
 	 */
 	template <typename... ArgsType>
-	FORCEINLINE ElementType& EmplaceAt_GetRef(SizeType Index, ArgsType&&... Args) UE_LIFETIMEBOUND
+	[[nodiscard]] FORCEINLINE ElementType& EmplaceAt_GetRef(SizeType Index, ArgsType&&... Args) UE_LIFETIMEBOUND
 	{
 		InsertUninitializedImpl(Index, 1);
 		ElementType* Ptr = GetData() + Index;
@@ -2594,7 +2595,7 @@ public:
 	 * @return A reference to the newly-inserted element.
 	 * @see AddDefaulted_GetRef, AddUnique_GetRef, AddZeroed_GetRef, Insert_GetRef
 	 */
-	FORCEINLINE ElementType& Add_GetRef(ElementType&& Item) UE_LIFETIMEBOUND
+	[[nodiscard]] FORCEINLINE ElementType& Add_GetRef(ElementType&& Item) UE_LIFETIMEBOUND
 	{
 		CheckAddress(&Item);
 		return Emplace_GetRef(MoveTempIfPossible(Item));
@@ -2607,7 +2608,7 @@ public:
 	 * @return A reference to the newly-inserted element.
 	 * @see AddDefaulted_GetRef, AddUnique_GetRef, AddZeroed_GetRef, Insert_GetRef
 	 */
-	FORCEINLINE ElementType& Add_GetRef(const ElementType& Item) UE_LIFETIMEBOUND
+	[[nodiscard]] FORCEINLINE ElementType& Add_GetRef(const ElementType& Item) UE_LIFETIMEBOUND
 	{
 		CheckAddress(&Item);
 		return Emplace_GetRef(Item);
@@ -2649,7 +2650,7 @@ public:
 	 * @return A reference to the newly-inserted element.
 	 * @see Add_GetRef, AddDefaulted_GetRef, AddUnique_GetRef, Insert_GetRef
 	 */
-	ElementType& AddZeroed_GetRef() UE_LIFETIMEBOUND
+	[[nodiscard]] ElementType& AddZeroed_GetRef() UE_LIFETIMEBOUND
 	{
 		const SizeType Index = AddUninitialized();
 		ElementType* Ptr = GetData() + Index;
@@ -2685,7 +2686,7 @@ public:
 	 * @return A reference to the newly-inserted element.
 	 * @see Add_GetRef, AddZeroed_GetRef, AddUnique_GetRef, Insert_GetRef
 	 */
-	ElementType& AddDefaulted_GetRef() UE_LIFETIMEBOUND
+	[[nodiscard]] ElementType& AddDefaulted_GetRef() UE_LIFETIMEBOUND
 	{
 		const SizeType Index = AddUninitialized();
 		ElementType* Ptr = GetData() + Index;
@@ -2699,7 +2700,7 @@ public:
 		typename AliasElementType = ElementType
 		UE_REQUIRES(TIsContainerElementTypeReinterpretable_V<AliasElementType>)
 	>
-	operator TArray<typename TContainerElementTypeCompatibility<AliasElementType>::ReinterpretType, AllocatorType>& ()
+	[[nodiscard]] operator TArray<typename TContainerElementTypeCompatibility<AliasElementType>::ReinterpretType, AllocatorType>& ()
 	{
 		using ElementCompat = TContainerElementTypeCompatibility<ElementType>;
 		ElementCompat::ReinterpretRangeContiguous(begin(), end(), Num());
@@ -2712,7 +2713,7 @@ public:
 		typename AliasElementType = ElementType
 		UE_REQUIRES(TIsContainerElementTypeReinterpretable_V<AliasElementType>)
 	>
-	operator const TArray<typename TContainerElementTypeCompatibility<AliasElementType>::ReinterpretType, AllocatorType>& () const
+	[[nodiscard]] operator const TArray<typename TContainerElementTypeCompatibility<AliasElementType>::ReinterpretType, AllocatorType>& () const
 	{
 		using ElementCompat = TContainerElementTypeCompatibility<ElementType>;
 		ElementCompat::ReinterpretRangeContiguous(begin(), end(), Num());
@@ -3222,7 +3223,7 @@ public:
 	 *
 	 * @returns The iterator.
 	 */
-	TIterator CreateIterator()
+	[[nodiscard]] TIterator CreateIterator()
 	{
 		return TIterator(*this);
 	}
@@ -3232,7 +3233,7 @@ public:
 	 *
 	 * @returns The const iterator.
 	 */
-	TConstIterator CreateConstIterator() const
+	[[nodiscard]] TConstIterator CreateConstIterator() const
 	{
 		return TConstIterator(*this);
 	}
@@ -3256,23 +3257,23 @@ public:
 	 * STL-like iterators to enable range-based for loop support.
 	 */
 	#if TARRAY_RANGED_FOR_CHECKS
-		FORCEINLINE RangedForIteratorType             begin ()       { return RangedForIteratorType            (ArrayNum, GetData()); }
-		FORCEINLINE RangedForConstIteratorType        begin () const { return RangedForConstIteratorType       (ArrayNum, GetData()); }
-		FORCEINLINE RangedForIteratorType             end   ()       { return RangedForIteratorType            (ArrayNum, GetData() + Num()); }
-		FORCEINLINE RangedForConstIteratorType        end   () const { return RangedForConstIteratorType       (ArrayNum, GetData() + Num()); }
-		FORCEINLINE RangedForReverseIteratorType      rbegin()       { return RangedForReverseIteratorType     (ArrayNum, GetData() + Num()); }
-		FORCEINLINE RangedForConstReverseIteratorType rbegin() const { return RangedForConstReverseIteratorType(ArrayNum, GetData() + Num()); }
-		FORCEINLINE RangedForReverseIteratorType      rend  ()       { return RangedForReverseIteratorType     (ArrayNum, GetData()); }
-		FORCEINLINE RangedForConstReverseIteratorType rend  () const { return RangedForConstReverseIteratorType(ArrayNum, GetData()); }
+		[[nodiscard]] FORCEINLINE RangedForIteratorType             begin ()       { return RangedForIteratorType            (ArrayNum, GetData()); }
+		[[nodiscard]] FORCEINLINE RangedForConstIteratorType        begin () const { return RangedForConstIteratorType       (ArrayNum, GetData()); }
+		[[nodiscard]] FORCEINLINE RangedForIteratorType             end   ()       { return RangedForIteratorType            (ArrayNum, GetData() + Num()); }
+		[[nodiscard]] FORCEINLINE RangedForConstIteratorType        end   () const { return RangedForConstIteratorType       (ArrayNum, GetData() + Num()); }
+		[[nodiscard]] FORCEINLINE RangedForReverseIteratorType      rbegin()       { return RangedForReverseIteratorType     (ArrayNum, GetData() + Num()); }
+		[[nodiscard]] FORCEINLINE RangedForConstReverseIteratorType rbegin() const { return RangedForConstReverseIteratorType(ArrayNum, GetData() + Num()); }
+		[[nodiscard]] FORCEINLINE RangedForReverseIteratorType      rend  ()       { return RangedForReverseIteratorType     (ArrayNum, GetData()); }
+		[[nodiscard]] FORCEINLINE RangedForConstReverseIteratorType rend  () const { return RangedForConstReverseIteratorType(ArrayNum, GetData()); }
 	#else
-		FORCEINLINE RangedForIteratorType             begin ()       { return                                   GetData(); }
-		FORCEINLINE RangedForConstIteratorType        begin () const { return                                   GetData(); }
-		FORCEINLINE RangedForIteratorType             end   ()       { return                                   GetData() + Num(); }
-		FORCEINLINE RangedForConstIteratorType        end   () const { return                                   GetData() + Num(); }
-		FORCEINLINE RangedForReverseIteratorType      rbegin()       { return RangedForReverseIteratorType     (GetData() + Num()); }
-		FORCEINLINE RangedForConstReverseIteratorType rbegin() const { return RangedForConstReverseIteratorType(GetData() + Num()); }
-		FORCEINLINE RangedForReverseIteratorType      rend  ()       { return RangedForReverseIteratorType     (GetData()); }
-		FORCEINLINE RangedForConstReverseIteratorType rend  () const { return RangedForConstReverseIteratorType(GetData()); }
+		[[nodiscard]] FORCEINLINE RangedForIteratorType             begin ()       { return                                   GetData(); }
+		[[nodiscard]] FORCEINLINE RangedForConstIteratorType        begin () const { return                                   GetData(); }
+		[[nodiscard]] FORCEINLINE RangedForIteratorType             end   ()       { return                                   GetData() + Num(); }
+		[[nodiscard]] FORCEINLINE RangedForConstIteratorType        end   () const { return                                   GetData() + Num(); }
+		[[nodiscard]] FORCEINLINE RangedForReverseIteratorType      rbegin()       { return RangedForReverseIteratorType     (GetData() + Num()); }
+		[[nodiscard]] FORCEINLINE RangedForConstReverseIteratorType rbegin() const { return RangedForConstReverseIteratorType(GetData() + Num()); }
+		[[nodiscard]] FORCEINLINE RangedForReverseIteratorType      rend  ()       { return RangedForReverseIteratorType     (GetData()); }
+		[[nodiscard]] FORCEINLINE RangedForConstReverseIteratorType rend  () const { return RangedForConstReverseIteratorType(GetData()); }
 	#endif
 
 public:
@@ -3349,7 +3350,7 @@ private:
 	 * @param Index Position to get.
 	 * @returns Reference to the element at given position.
 	 */
-	FORCENOINLINE const ElementType& DebugGet(SizeType Index) const
+	[[nodiscard]] FORCENOINLINE const ElementType& DebugGet(SizeType Index) const
 	{
 		return GetData()[Index];
 	}
@@ -3712,7 +3713,7 @@ public:
 	 *
 	 * @returns The reference to the top element from the heap.
 	 */
-	const ElementType& HeapTop() const UE_LIFETIMEBOUND
+	[[nodiscard]] const ElementType& HeapTop() const UE_LIFETIMEBOUND
 	{
 		return (*this)[0];
 	}
@@ -3722,7 +3723,7 @@ public:
 	 *
 	 * @returns The reference to the top element from the heap.
 	 */
-	ElementType& HeapTop() UE_LIFETIMEBOUND
+	[[nodiscard]] ElementType& HeapTop() UE_LIFETIMEBOUND
 	{
 		return (*this)[0];
 	}
@@ -3805,8 +3806,14 @@ public:
 		HeapSort(TLess<ElementType>());
 	}
 
-	const ElementAllocatorType& GetAllocatorInstance() const { return AllocatorInstance; }
-	ElementAllocatorType& GetAllocatorInstance() { return AllocatorInstance; }
+	[[nodiscard]] const ElementAllocatorType& GetAllocatorInstance() const
+	{
+		return AllocatorInstance;
+	}
+	[[nodiscard]] ElementAllocatorType& GetAllocatorInstance()
+	{
+		return AllocatorInstance;
+	}
 
 	friend struct TArrayPrivateFriend;
 };
@@ -3821,20 +3828,20 @@ namespace Freeze
 	}
 
 	template<typename T, typename AllocatorType>
-	uint32 IntrinsicUnfrozenCopy(const FMemoryUnfreezeContent& Context, const TArray<T, AllocatorType>& Object, void* OutDst)
+	[[nodiscard]] uint32 IntrinsicUnfrozenCopy(const FMemoryUnfreezeContent& Context, const TArray<T, AllocatorType>& Object, void* OutDst)
 	{
 		Object.CopyUnfrozen(Context, OutDst);
 		return sizeof(Object);
 	}
 
 	template<typename T, typename AllocatorType>
-	uint32 IntrinsicAppendHash(const TArray<T, AllocatorType>* DummyObject, const FTypeLayoutDesc& TypeDesc, const FPlatformTypeLayoutParameters& LayoutParams, FSHA1& Hasher)
+	[[nodiscard]] uint32 IntrinsicAppendHash(const TArray<T, AllocatorType>* DummyObject, const FTypeLayoutDesc& TypeDesc, const FPlatformTypeLayoutParameters& LayoutParams, FSHA1& Hasher)
 	{
 		return AppendHashForNameAndSize(TypeDesc.Name, sizeof(TArray<T, AllocatorType>), Hasher);
 	}
 
 	template<typename T, typename AllocatorType>
-	uint32 IntrinsicGetTargetAlignment(const TArray<T, AllocatorType>* DummyObject, const FTypeLayoutDesc& TypeDesc, const FPlatformTypeLayoutParameters& LayoutParams)
+	[[nodiscard]] uint32 IntrinsicGetTargetAlignment(const TArray<T, AllocatorType>* DummyObject, const FTypeLayoutDesc& TypeDesc, const FPlatformTypeLayoutParameters& LayoutParams)
 	{
 		// Assume alignment of array is drive by pointer
 		return FMath::Min(8u, LayoutParams.MaxFieldAlignment);
@@ -4007,7 +4014,7 @@ FArchive& operator<<(FArchive& Ar, TArray<ElementType, AllocatorType>& A)
 
 /** Returns a unique hash by combining those of each array element. */
 template<typename InElementType, typename InAllocatorType>
-uint32 GetTypeHash(const TArray<InElementType, InAllocatorType>& A)
+[[nodiscard]] uint32 GetTypeHash(const TArray<InElementType, InAllocatorType>& A)
 {
 	uint32 Hash = 0;
 	for (const InElementType& V : A)
diff --git a/Engine/Source/Runtime/Core/Public/Containers/ArrayView.h b/Engine/Source/Runtime/Core/Public/Containers/ArrayView.h
--- a/Engine/Source/Runtime/Core/Public/Containers/ArrayView.h
+++ b/Engine/Source/Runtime/Core/Public/Containers/ArrayView.h
@@ -142,14 +142,14 @@ public:
 	static_assert(std::is_signed_v<SizeType>, "TArrayView only supports signed index types");
 
 	// Defaulted object behavior - we want compiler-generated functions rather than going through the generic range constructor.
-	TArrayView(const TArrayView&) = default;
+	[[nodiscard]] TArrayView(const TArrayView&) = default;
 	TArrayView& operator=(const TArrayView&) = default;
 	~TArrayView() = default;
 
 	/**
 	 * Constructor.
 	 */
-	TArrayView()
+	[[nodiscard]] TArrayView()
 		: DataPtr(nullptr)
 		, ArrayNum(0)
 	{
@@ -183,7 +183,7 @@ public:
 			!std::is_same_v<CVUnqualifiedOtherRangeType, TArrayView>
 		)
 	>
-	FORCEINLINE TArrayView(OtherRangeType&& Other)
+	[[nodiscard]] FORCEINLINE TArrayView(OtherRangeType&& Other)
 		: DataPtr(std::conditional_t<
 						TIsCompatibleRangeType<OtherRangeType>::Value,
 						TIsCompatibleRangeType<OtherRangeType>,
@@ -215,7 +215,7 @@ public:
 			!TIsTArrayView_V<CVUnqualifiedOtherRangeType>
 		)
 	>
-	FORCEINLINE TArrayView(OtherRangeType&& Other UE_LIFETIMEBOUND)
+	[[nodiscard]] FORCEINLINE TArrayView(OtherRangeType&& Other UE_LIFETIMEBOUND)
 		: DataPtr(std::conditional_t<
 			TIsCompatibleRangeType<OtherRangeType>::Value,
 			TIsCompatibleRangeType<OtherRangeType>,
@@ -245,7 +245,7 @@ public:
 		typename OtherElementType
 		UE_REQUIRES(UE::Core::ArrayView::Private::TIsCompatibleElementType_V<OtherElementType, ElementType>)
 	>
-	FORCEINLINE TArrayView(OtherElementType* InData UE_LIFETIMEBOUND, SizeType InCount)
+	[[nodiscard]] FORCEINLINE TArrayView(OtherElementType* InData UE_LIFETIMEBOUND, SizeType InCount)
 		: DataPtr(InData)
 		, ArrayNum(InCount)
 	{
@@ -257,7 +257,7 @@ public:
 	 *
 	 * The caller is responsible for ensuring that the view does not outlive the initializer list.
 	 */
-	FORCEINLINE TArrayView(std::initializer_list<ElementType> List UE_LIFETIMEBOUND)
+	[[nodiscard]] FORCEINLINE TArrayView(std::initializer_list<ElementType> List UE_LIFETIMEBOUND)
 		: DataPtr(UE::Core::ArrayView::Private::GetDataHelper(List))
 		, ArrayNum(GetNum(List))
 	{
@@ -270,12 +270,12 @@ public:
 	constexpr static bool bHasIntrusiveUnsetOptionalState = true;
 	using IntrusiveUnsetOptionalStateType = TArrayView;
 
-	explicit TArrayView(FIntrusiveUnsetOptionalState)
+	[[nodiscard]] explicit TArrayView(FIntrusiveUnsetOptionalState)
 		: DataPtr(nullptr)
 		, ArrayNum(-1)
 	{
 	}
-	bool operator==(FIntrusiveUnsetOptionalState) const
+	[[nodiscard]] bool operator==(FIntrusiveUnsetOptionalState) const
 	{
 		return ArrayNum == -1;
 	}
@@ -290,7 +290,7 @@ public:
 	 *
 	 * @returns Pointer to first array entry.
 	 */
-	FORCEINLINE ElementType* GetData() const
+	[[nodiscard]] FORCEINLINE ElementType* GetData() const
 	{
 		return DataPtr;
 	}
@@ -300,7 +300,7 @@ public:
 	 *
 	 * @returns Size in bytes of array type.
 	 */
-	FORCEINLINE static constexpr size_t GetTypeSize()
+	[[nodiscard]] FORCEINLINE static constexpr size_t GetTypeSize()
 	{
 		return sizeof(ElementType);
 	}
@@ -308,7 +308,7 @@ public:
 	/** 
 	 * Helper function returning the alignment of the inner type.
 	 */
-	FORCEINLINE static constexpr size_t GetTypeAlignment()
+	[[nodiscard]] FORCEINLINE static constexpr size_t GetTypeAlignment()
 	{
 		return alignof(ElementType);
 	}
@@ -354,7 +354,7 @@ public:
 	 *
 	 * @returns True if index is valid. False otherwise.
 	 */
-	FORCEINLINE bool IsValidIndex(SizeType Index) const
+	[[nodiscard]] FORCEINLINE bool IsValidIndex(SizeType Index) const
 	{
 		return (Index >= 0) && (Index < ArrayNum);
 	}
@@ -365,7 +365,7 @@ public:
 	 * @returns True if the array view is empty.
 	 * @see Num
 	 */
-	bool IsEmpty() const
+	[[nodiscard]] bool IsEmpty() const
 	{
 		return ArrayNum == 0;
 	}
@@ -375,13 +375,13 @@ public:
 	 *
 	 * @returns Number of elements in array.
 	 */
-	FORCEINLINE SizeType Num() const
+	[[nodiscard]] FORCEINLINE SizeType Num() const
 	{
 		return ArrayNum;
 	}
 
 	/** @returns Number of bytes used, excluding slack */
-	FORCEINLINE SIZE_T NumBytes() const
+	[[nodiscard]] FORCEINLINE SIZE_T NumBytes() const
 	{
 		return static_cast<SIZE_T>(ArrayNum) * sizeof(ElementType);
 	}
@@ -391,7 +391,7 @@ public:
 	 *
 	 * @returns Reference to indexed element.
 	 */
-	FORCEINLINE ElementType& operator[](SizeType Index) const
+	[[nodiscard]] FORCEINLINE ElementType& operator[](SizeType Index) const
 	{
 		RangeCheck(Index);
 		return GetData()[Index];
@@ -405,7 +405,7 @@ public:
 	 *
 	 * @returns Reference to n-th last element from the array.
 	 */
-	FORCEINLINE ElementType& Last(SizeType IndexFromTheEnd = 0) const
+	[[nodiscard]] FORCEINLINE ElementType& Last(SizeType IndexFromTheEnd = 0) const
 	{
 		RangeCheck(ArrayNum - IndexFromTheEnd - 1);
 		return GetData()[ArrayNum - IndexFromTheEnd - 1];
@@ -512,7 +512,7 @@ public:
 	 *
 	 * @returns True if found. False otherwise.
 	 */
-	FORCEINLINE bool Find(const ElementType& Item, SizeType& Index) const
+	[[nodiscard]] FORCEINLINE bool Find(const ElementType& Item, SizeType& Index) const
 	{
 		Index = this->Find(Item);
 		return Index != static_cast<SizeType>(INDEX_NONE);
@@ -525,7 +525,7 @@ public:
 	 *
 	 * @returns Index of the found element. INDEX_NONE otherwise.
 	 */
-	SizeType Find(const ElementType& Item) const
+	[[nodiscard]] SizeType Find(const ElementType& Item) const
 	{
 		const ElementType* RESTRICT Start = GetData();
 		for (const ElementType* RESTRICT Data = Start, *RESTRICT DataEnd = Data + ArrayNum; Data != DataEnd; ++Data)
@@ -546,7 +546,7 @@ public:
 	 *
 	 * @returns True if found. False otherwise.
 	 */
-	FORCEINLINE bool FindLast(const ElementType& Item, SizeType& Index) const
+	[[nodiscard]] FORCEINLINE bool FindLast(const ElementType& Item, SizeType& Index) const
 	{
 		Index = this->FindLast(Item);
 		return Index != static_cast<SizeType>(INDEX_NONE);
@@ -561,7 +561,7 @@ public:
 	 * @returns Index of the found element. INDEX_NONE otherwise.
 	 */
 	template <typename Predicate>
-	SizeType FindLastByPredicate(Predicate Pred, SizeType StartIndex) const
+	[[nodiscard]] SizeType FindLastByPredicate(Predicate Pred, SizeType StartIndex) const
 	{
 		check(StartIndex >= 0 && StartIndex <= this->Num());
 		for (const ElementType* RESTRICT Start = GetData(), *RESTRICT Data = Start + StartIndex; Data != Start; )
@@ -583,7 +583,7 @@ public:
 	* @returns Index of the found element. INDEX_NONE otherwise.
 	*/
 	template <typename Predicate>
-	FORCEINLINE SizeType FindLastByPredicate(Predicate Pred) const
+	[[nodiscard]] FORCEINLINE SizeType FindLastByPredicate(Predicate Pred) const
 	{
 		return FindLastByPredicate(Pred, ArrayNum);
 	}
@@ -598,7 +598,7 @@ public:
 	 *          found.
 	 */
 	template <typename KeyType>
-	SizeType IndexOfByKey(const KeyType& Key) const
+	[[nodiscard]] SizeType IndexOfByKey(const KeyType& Key) const
 	{
 		const ElementType* RESTRICT Start = GetData();
 		for (const ElementType* RESTRICT Data = Start, *RESTRICT DataEnd = Start + ArrayNum; Data != DataEnd; ++Data)
@@ -620,7 +620,7 @@ public:
 	 *          found.
 	 */
 	template <typename Predicate>
-	SizeType IndexOfByPredicate(Predicate Pred) const
+	[[nodiscard]] SizeType IndexOfByPredicate(Predicate Pred) const
 	{
 		const ElementType* RESTRICT Start = GetData();
 		for (const ElementType* RESTRICT Data = Start, *RESTRICT DataEnd = Start + ArrayNum; Data != DataEnd; ++Data)
@@ -642,7 +642,7 @@ public:
 	 * @returns Pointer to the first matching element, or nullptr if none is found.
 	 */
 	template <typename KeyType>
-	ElementType* FindByKey(const KeyType& Key) const
+	[[nodiscard]] ElementType* FindByKey(const KeyType& Key) const
 	{
 		for (ElementType* RESTRICT Data = GetData(), *RESTRICT DataEnd = Data + ArrayNum; Data != DataEnd; ++Data)
 		{
@@ -664,7 +664,7 @@ public:
 	 *         true, or nullptr if none is found.
 	 */
 	template <typename Predicate>
-	ElementType* FindByPredicate(Predicate Pred) const
+	[[nodiscard]] ElementType* FindByPredicate(Predicate Pred) const
 	{
 		for (ElementType* RESTRICT Data = GetData(), *RESTRICT DataEnd = Data + ArrayNum; Data != DataEnd; ++Data)
 		{
@@ -686,7 +686,7 @@ public:
 	 *          the subset of elements for which the functor returns true.
 	 */
 	template <typename Predicate>
-	TArray<std::remove_const_t<ElementType>> FilterByPredicate(Predicate Pred) const
+	[[nodiscard]] TArray<std::remove_const_t<ElementType>> FilterByPredicate(Predicate Pred) const
 	{
 		TArray<std::remove_const_t<ElementType>> FilterResults;
 		for (const ElementType* RESTRICT Data = GetData(), *RESTRICT DataEnd = Data + ArrayNum; Data != DataEnd; ++Data)
@@ -705,7 +705,7 @@ public:
 	 * @returns	True if found. False otherwise.
 	 */
 	template <typename ComparisonType>
-	bool Contains(const ComparisonType& Item) const
+	[[nodiscard]] bool Contains(const ComparisonType& Item) const
 	{
 		for (const ElementType* RESTRICT Data = GetData(), *RESTRICT DataEnd = Data + ArrayNum; Data != DataEnd; ++Data)
 		{
@@ -725,7 +725,7 @@ public:
 	 * @returns	True if found. False otherwise.
 	 */
 	template <typename Predicate>
-	FORCEINLINE bool ContainsByPredicate(Predicate Pred) const
+	[[nodiscard]] FORCEINLINE bool ContainsByPredicate(Predicate Pred) const
 	{
 		return FindByPredicate(Pred) != nullptr;
 	}
@@ -735,10 +735,10 @@ public:
 	 * DO NOT USE DIRECTLY
 	 * STL-like iterators to enable range-based for loop support.
 	 */
-	FORCEINLINE ElementType*                         begin () const { return GetData(); }
-	FORCEINLINE ElementType*                         end   () const { return GetData() + Num(); }
-	FORCEINLINE TReversePointerIterator<ElementType> rbegin() const { return TReversePointerIterator<ElementType>(GetData() + Num()); }
-	FORCEINLINE TReversePointerIterator<ElementType> rend  () const { return TReversePointerIterator<ElementType>(GetData()); }
+	[[nodiscard]] FORCEINLINE ElementType*                         begin () const { return GetData(); }
+	[[nodiscard]] FORCEINLINE ElementType*                         end   () const { return GetData() + Num(); }
+	[[nodiscard]] FORCEINLINE TReversePointerIterator<ElementType> rbegin() const { return TReversePointerIterator<ElementType>(GetData() + Num()); }
+	[[nodiscard]] FORCEINLINE TReversePointerIterator<ElementType> rend  () const { return TReversePointerIterator<ElementType>(GetData()); }
 
 public:
 	/**
@@ -826,7 +826,7 @@ template <
 	typename CVUnqualifiedOtherRangeType = std::remove_cv_t<std::remove_reference_t<OtherRangeType>>
 	UE_REQUIRES(TIsContiguousContainer<CVUnqualifiedOtherRangeType>::Value && TIsTArrayView_V<CVUnqualifiedOtherRangeType>)
 >
-auto MakeArrayView(OtherRangeType&& Other)
+[[nodiscard]] auto MakeArrayView(OtherRangeType&& Other)
 {
 	return TArrayView<std::remove_pointer_t<decltype(GetData(DeclVal<OtherRangeType&>()))>>(Forward<OtherRangeType>(Other));
 }
@@ -835,19 +835,19 @@ template <
 	typename CVUnqualifiedOtherRangeType = std::remove_cv_t<std::remove_reference_t<OtherRangeType>>
 	UE_REQUIRES(TIsContiguousContainer<CVUnqualifiedOtherRangeType>::Value && !TIsTArrayView_V<CVUnqualifiedOtherRangeType>)
 >
-auto MakeArrayView(OtherRangeType&& Other UE_LIFETIMEBOUND)
+[[nodiscard]] auto MakeArrayView(OtherRangeType&& Other UE_LIFETIMEBOUND)
 {
 	return TArrayView<std::remove_pointer_t<decltype(GetData(DeclVal<OtherRangeType&>()))>>(Forward<OtherRangeType>(Other));
 }
 
 template<typename ElementType>
-TArrayView<ElementType> MakeArrayView(ElementType* Pointer UE_LIFETIMEBOUND, int32 Size)
+[[nodiscard]] TArrayView<ElementType> MakeArrayView(ElementType* Pointer UE_LIFETIMEBOUND, int32 Size)
 {
 	return TArrayView<ElementType>(Pointer, Size);
 }
 
 template <typename ElementType>
-TArrayView<const ElementType> MakeArrayView(std::initializer_list<ElementType> List UE_LIFETIMEBOUND)
+[[nodiscard]] TArrayView<const ElementType> MakeArrayView(std::initializer_list<ElementType> List UE_LIFETIMEBOUND)
 {
 	return TArrayView<const ElementType>(List.begin(), List.size());
 }
@@ -857,7 +857,7 @@ template <
 	typename CVUnqualifiedOtherRangeType = std::remove_cv_t<std::remove_reference_t<OtherRangeType>>
 	UE_REQUIRES(TIsContiguousContainer<CVUnqualifiedOtherRangeType>::Value && TIsTArrayView_V<CVUnqualifiedOtherRangeType>)
 >
-auto MakeConstArrayView(OtherRangeType&& Other)
+[[nodiscard]] auto MakeConstArrayView(OtherRangeType&& Other)
 {
 	return TArrayView<const std::remove_pointer_t<decltype(GetData(DeclVal<OtherRangeType&>()))>>(Forward<OtherRangeType>(Other));
 }
@@ -866,19 +866,19 @@ template <
 	typename CVUnqualifiedOtherRangeType = std::remove_cv_t<std::remove_reference_t<OtherRangeType>>
 	UE_REQUIRES(TIsContiguousContainer<CVUnqualifiedOtherRangeType>::Value && !TIsTArrayView_V<CVUnqualifiedOtherRangeType>)
 >
-auto MakeConstArrayView(OtherRangeType&& Other UE_LIFETIMEBOUND)
+[[nodiscard]] auto MakeConstArrayView(OtherRangeType&& Other UE_LIFETIMEBOUND)
 {
 	return TArrayView<const std::remove_pointer_t<decltype(GetData(DeclVal<OtherRangeType&>()))>>(Forward<OtherRangeType>(Other));
 }
 
 template<typename ElementType>
-TArrayView<const ElementType> MakeConstArrayView(const ElementType* Pointer UE_LIFETIMEBOUND, int32 Size)
+[[nodiscard]] TArrayView<const ElementType> MakeConstArrayView(const ElementType* Pointer UE_LIFETIMEBOUND, int32 Size)
 {
 	return TArrayView<const ElementType>(Pointer, Size);
 }
 
 template <typename ElementType>
-TArrayView<const ElementType> MakeConstArrayView(std::initializer_list<ElementType> List UE_LIFETIMEBOUND)
+[[nodiscard]] TArrayView<const ElementType> MakeConstArrayView(std::initializer_list<ElementType> List UE_LIFETIMEBOUND)
 {
 	return TArrayView<const ElementType>(List.begin(), List.size());
 }
@@ -904,7 +904,7 @@ template <
 	typename ElementType
 	UE_REQUIRES(std::is_convertible_v<decltype(GetData(std::declval<RangeType&>())), const ElementType*>)
 >
-bool operator==(RangeType&& Lhs, TArrayView<ElementType> Rhs)
+[[nodiscard]] bool operator==(RangeType&& Lhs, TArrayView<ElementType> Rhs)
 {
 	auto Count = GetNum(Lhs);
 	return Count == Rhs.Num() && CompareItems(GetData(Lhs), Rhs.GetData(), Count);
@@ -915,7 +915,7 @@ template <
 	typename ElementType
 	UE_REQUIRES(std::is_convertible_v<decltype(GetData(std::declval<RangeType&>())), const ElementType*>)
 >
-bool operator==(TArrayView<ElementType> Lhs, RangeType&& Rhs)
+[[nodiscard]] bool operator==(TArrayView<ElementType> Lhs, RangeType&& Rhs)
 {
 	return (Rhs == Lhs);
 }
@@ -932,7 +932,7 @@ template <
 	typename ElementType
 	UE_REQUIRES(std::is_convertible_v<decltype(GetData(std::declval<RangeType&>())), const ElementType*>)
 >
-bool operator!=(RangeType&& Lhs, TArrayView<ElementType> Rhs)
+[[nodiscard]] bool operator!=(RangeType&& Lhs, TArrayView<ElementType> Rhs)
 {
 	return !(Lhs == Rhs);
 }
@@ -942,7 +942,7 @@ template <
 	typename ElementType
 	UE_REQUIRES(std::is_convertible_v<decltype(GetData(std::declval<RangeType&>())), const ElementType*>)
 >
-bool operator!=(TArrayView<ElementType> Lhs, RangeType&& Rhs)
+[[nodiscard]] bool operator!=(TArrayView<ElementType> Lhs, RangeType&& Rhs)
 {
 	return !(Rhs == Lhs);
 }
@@ -950,7 +950,7 @@ bool operator!=(TArrayView<ElementType> Lhs, RangeType&& Rhs)
 
 template<typename InElementType, typename InAllocatorType>
 template<typename OtherElementType, typename OtherSizeType>
-FORCEINLINE TArray<InElementType, InAllocatorType>::TArray(const TArrayView<OtherElementType, OtherSizeType>& Other)
+[[nodiscard]] FORCEINLINE TArray<InElementType, InAllocatorType>::TArray(const TArrayView<OtherElementType, OtherSizeType>& Other)
 {
 	CopyToEmpty(Other.GetData(), Other.Num(), 0);
 }
diff --git a/Engine/Source/Runtime/Core/Public/Containers/BasicArray.h b/Engine/Source/Runtime/Core/Public/Containers/BasicArray.h
--- a/Engine/Source/Runtime/Core/Public/Containers/BasicArray.h
+++ b/Engine/Source/Runtime/Core/Public/Containers/BasicArray.h
@@ -21,7 +21,7 @@ template <typename T>
 class TBasicArray
 {
 public:
-	TBasicArray()
+	[[nodiscard]] TBasicArray()
 		: Data(nullptr)
 	{
 	}
@@ -30,7 +30,7 @@ public:
 	TBasicArray(const TBasicArray&) = delete;
 	TBasicArray& operator=(const TBasicArray&) = delete;
 
-	TBasicArray(TBasicArray&& Other)
+	[[nodiscard]] TBasicArray(TBasicArray&& Other)
 		: Data(Other.Data)
 	{
 		Other.Data = nullptr;
@@ -88,34 +88,34 @@ public:
 		}
 	}
 
-	bool IsEmpty() const
+	[[nodiscard]] bool IsEmpty() const
 	{
 		return Num() == 0;
 	}
 
-	int32 Num() const
+	[[nodiscard]] int32 Num() const
 	{
 		FData* LocalData = Data;
 		return LocalData ? LocalData->Num : 0;
 	}
 
-	T* GetData()
+	[[nodiscard]] T* GetData()
 	{
 		FData* LocalData = Data;
 		return LocalData ? &LocalData->Data[0] : nullptr;
 	}
 
-	FORCEINLINE const T* GetData() const
+	[[nodiscard]] FORCEINLINE const T* GetData() const
 	{
 		return const_cast<TBasicArray*>(this)->GetData();
 	}
 
-	FORCEINLINE T& operator[](int32 Index)
+	[[nodiscard]] FORCEINLINE T& operator[](int32 Index)
 	{
 		return GetData()[Index];
 	}
 
-	FORCEINLINE const T& operator[](int32 Index) const
+	[[nodiscard]] FORCEINLINE const T& operator[](int32 Index) const
 	{
 		return GetData()[Index];
 	}
@@ -171,10 +171,10 @@ private:
 
 	FData* Data;
 
-	friend       T* begin(      TBasicArray& Arr) { return Arr.GetData(); }
-	friend const T* begin(const TBasicArray& Arr) { return Arr.GetData(); }
-	friend       T* end  (      TBasicArray& Arr) { return Arr.GetData() + Arr.Num(); }
-	friend const T* end  (const TBasicArray& Arr) { return Arr.GetData() + Arr.Num(); }
+	[[nodiscard]] friend       T* begin(      TBasicArray& Arr) { return Arr.GetData(); }
+	[[nodiscard]] friend const T* begin(const TBasicArray& Arr) { return Arr.GetData(); }
+	[[nodiscard]] friend       T* end  (      TBasicArray& Arr) { return Arr.GetData() + Arr.Num(); }
+	[[nodiscard]] friend const T* end  (const TBasicArray& Arr) { return Arr.GetData() + Arr.Num(); }
 };
 
 template <typename T>
diff --git a/Engine/Source/Runtime/Core/Public/Containers/BinaryHeap.h b/Engine/Source/Runtime/Core/Public/Containers/BinaryHeap.h
--- a/Engine/Source/Runtime/Core/Public/Containers/BinaryHeap.h
+++ b/Engine/Source/Runtime/Core/Public/Containers/BinaryHeap.h
@@ -21,8 +21,8 @@ class FBinaryHeap
 public:
 	static_assert(std::is_unsigned_v<IndexType>, "FBinaryHeap only supports unsigned index types");
 
-				FBinaryHeap();
-				FBinaryHeap( uint32 InHeapSize, uint32 InIndexSize );
+				[[nodiscard]] FBinaryHeap();
+				[[nodiscard]] FBinaryHeap( uint32 InHeapSize, uint32 InIndexSize );
 				~FBinaryHeap();
 
 				FBinaryHeap(const FBinaryHeap&) = delete;
@@ -33,16 +33,16 @@ public:
 	void		Free();
 	void		Resize( uint32 NewHeapSize, uint32 NewIndexSize );
 
-	bool		IsEmpty() const			{ return HeapNum == 0; }
-	uint32		Num() const				{ return HeapNum; }
-	uint32		GetHeapSize() const		{ return HeapSize; }
-	uint32		GetIndexSize() const	{ return IndexSize; }
+	[[nodiscard]] bool		IsEmpty() const			{ return HeapNum == 0; }
+	[[nodiscard]] uint32		Num() const				{ return HeapNum; }
+	[[nodiscard]] uint32		GetHeapSize() const		{ return HeapSize; }
+	[[nodiscard]] uint32		GetIndexSize() const	{ return IndexSize; }
 
-	bool		IsPresent( IndexType Index ) const;
-	KeyType		GetKey( IndexType Index ) const;
-	IndexType	Peek( IndexType Index ) const;
+	[[nodiscard]] bool		IsPresent( IndexType Index ) const;
+	[[nodiscard]] KeyType		GetKey( IndexType Index ) const;
+	[[nodiscard]] IndexType	Peek( IndexType Index ) const;
 
-	IndexType	Top() const;
+	[[nodiscard]] IndexType	Top() const;
 	void		Pop();
 
 	void		Add( KeyType Key, IndexType Index );
@@ -72,13 +72,13 @@ protected:
 };
 
 template< typename KeyType, typename IndexType >
-FORCEINLINE FBinaryHeap< KeyType, IndexType >::FBinaryHeap()
+[[nodiscard]] FORCEINLINE FBinaryHeap< KeyType, IndexType >::FBinaryHeap()
 {
 	ResetInternal();
 }
 
 template< typename KeyType, typename IndexType >
-FORCEINLINE FBinaryHeap< KeyType, IndexType >::FBinaryHeap( uint32 InHeapSize, uint32 InIndexSize )
+[[nodiscard]] FORCEINLINE FBinaryHeap< KeyType, IndexType >::FBinaryHeap( uint32 InHeapSize, uint32 InIndexSize )
 	: HeapNum(0)
 	, HeapSize( InHeapSize )
 	, IndexSize( InIndexSize )
@@ -91,7 +91,7 @@ FORCEINLINE FBinaryHeap< KeyType, IndexType >::FBinaryHeap( uint32 InHeapSize, u
 }
 
 template< typename KeyType, typename IndexType >
-FORCEINLINE FBinaryHeap< KeyType, IndexType >::FBinaryHeap(FBinaryHeap< KeyType, IndexType >&& Other)
+[[nodiscard]] FORCEINLINE FBinaryHeap< KeyType, IndexType >::FBinaryHeap(FBinaryHeap< KeyType, IndexType >&& Other)
 {
 	HeapNum = Other.HeapNum;
 	HeapSize = Other.HeapSize;
@@ -215,7 +215,7 @@ FORCEINLINE void FBinaryHeap< KeyType, IndexType >::Resize( uint32 NewHeapSize,
 }
 
 template< typename KeyType, typename IndexType >
-FORCEINLINE bool FBinaryHeap< KeyType, IndexType >::IsPresent( IndexType Index ) const
+[[nodiscard]] FORCEINLINE bool FBinaryHeap< KeyType, IndexType >::IsPresent( IndexType Index ) const
 {
 	if (Index >= IndexSize)
 	{
@@ -225,21 +225,21 @@ FORCEINLINE bool FBinaryHeap< KeyType, IndexType >::IsPresent( IndexType Index )
 }
 
 template< typename KeyType, typename IndexType >
-FORCEINLINE KeyType FBinaryHeap< KeyType, IndexType >::GetKey( IndexType Index ) const
+[[nodiscard]] FORCEINLINE KeyType FBinaryHeap< KeyType, IndexType >::GetKey( IndexType Index ) const
 {
 	checkSlow( IsPresent( Index ) );
 	return Keys[ Index ];
 }
 
 template< typename KeyType, typename IndexType >
-FORCEINLINE IndexType FBinaryHeap< KeyType, IndexType >::Peek(IndexType Index) const
+[[nodiscard]] FORCEINLINE IndexType FBinaryHeap< KeyType, IndexType >::Peek(IndexType Index) const
 {
 	checkSlow(Index < HeapNum);
 	return Heap[Index];
 }
 
 template< typename KeyType, typename IndexType >
-FORCEINLINE IndexType FBinaryHeap< KeyType, IndexType >::Top() const
+[[nodiscard]] FORCEINLINE IndexType FBinaryHeap< KeyType, IndexType >::Top() const
 {
 	checkSlow( Heap );
 	checkSlow( HeapNum > 0 );
diff --git a/Engine/Source/Runtime/Core/Public/Containers/BitArray.h b/Engine/Source/Runtime/Core/Public/Containers/BitArray.h
--- a/Engine/Source/Runtime/Core/Public/Containers/BitArray.h
+++ b/Engine/Source/Runtime/Core/Public/Containers/BitArray.h
@@ -194,7 +194,7 @@ public:
 	,	Mask(InMask)
 	{}
 
-	FORCEINLINE operator bool() const
+	[[nodiscard]] FORCEINLINE operator bool() const
 	{
 		 return (Data & Mask) != 0;
 	}
@@ -265,7 +265,7 @@ public:
 	,	Mask(InMask)
 	{}
 
-	FORCEINLINE operator bool() const
+	[[nodiscard]] FORCEINLINE operator bool() const
 	{
 		 return (Data & Mask) != 0;
 	}
@@ -289,12 +289,12 @@ public:
 	int32  WordIndex;
 	uint32 Mask;
 
-	FORCEINLINE bool operator==(FRelativeBitReference Other) const
+	[[nodiscard]] FORCEINLINE bool operator==(FRelativeBitReference Other) const
 	{
 		return (WordIndex == Other.WordIndex) & (Mask == Other.Mask);
 	}
 
-	FORCEINLINE bool operator!=(FRelativeBitReference Other) const
+	[[nodiscard]] FORCEINLINE bool operator!=(FRelativeBitReference Other) const
 	{
 		return !(Other == *this);
 	}
@@ -370,7 +370,7 @@ public:
 	template<typename,typename,bool>
 	friend class TConstDualSetBitIterator;
 
-	TBitArray()
+	[[nodiscard]] TBitArray()
 	:	NumBits(0)
 	,	MaxBits(AllocatorInstance.GetInitialCapacity() * NumBitsPerDWORD)
 	{
@@ -382,7 +382,7 @@ public:
 	 * @param Value - The value to initial the bits to.
 	 * @param InNumBits - The initial number of bits in the array.
 	 */
-	FORCEINLINE explicit TBitArray(bool bValue, int32 InNumBits)
+	[[nodiscard]] FORCEINLINE explicit TBitArray(bool bValue, int32 InNumBits)
 	:	MaxBits(AllocatorInstance.GetInitialCapacity() * NumBitsPerDWORD)
 	{
 		Init(bValue, InNumBits);
@@ -391,7 +391,7 @@ public:
 	/**
 	 * Move constructor.
 	 */
-	FORCEINLINE TBitArray(TBitArray&& Other)
+	[[nodiscard]] FORCEINLINE TBitArray(TBitArray&& Other)
 	{
 		this->Move(*this, Other);
 	}
@@ -399,7 +399,7 @@ public:
 	/**
 	 * Copy constructor.
 	 */
-	FORCEINLINE TBitArray(const TBitArray& Copy)
+	[[nodiscard]] FORCEINLINE TBitArray(const TBitArray& Copy)
 	:	NumBits(0)
 	,	MaxBits(0)
 	{
@@ -407,7 +407,7 @@ public:
 	}
 
 	template<typename OtherAllocator>
-	FORCEINLINE explicit TBitArray(const TBitArray<OtherAllocator> & Copy)
+	[[nodiscard]] FORCEINLINE explicit TBitArray(const TBitArray<OtherAllocator> & Copy)
 		: NumBits(0)
 		, MaxBits(0)
 	{
@@ -450,7 +450,7 @@ public:
 		return *this;
 	}
 
-	FORCEINLINE bool operator==(const TBitArray<Allocator>& Other) const
+	[[nodiscard]] FORCEINLINE bool operator==(const TBitArray<Allocator>& Other) const
 	{
 		if (Num() != Other.Num())
 		{
@@ -460,7 +460,7 @@ public:
 		return FMemory::Memcmp(GetData(), Other.GetData(), GetNumWords() * sizeof(uint32)) == 0;
 	}
 
-	FORCEINLINE bool operator<(const TBitArray<Allocator>& Other) const
+	[[nodiscard]] FORCEINLINE bool operator<(const TBitArray<Allocator>& Other) const
 	{
 		//sort by length
 		if (Num() != Other.Num())
@@ -484,24 +484,24 @@ public:
 	}
 
 #if !PLATFORM_COMPILER_HAS_GENERATED_COMPARISON_OPERATORS
-	FORCEINLINE bool operator!=(const TBitArray<Allocator>& Other) const
+	[[nodiscard]] FORCEINLINE bool operator!=(const TBitArray<Allocator>& Other) const
 	{
 		return !(*this == Other);
 	}
 #endif
 
 private:
-	FORCEINLINE uint32 GetNumWords() const
+	[[nodiscard]] FORCEINLINE uint32 GetNumWords() const
 	{
 		return FBitSet::CalculateNumWords(NumBits);
 	}
 
-	FORCEINLINE uint32 GetMaxWords() const
+	[[nodiscard]] FORCEINLINE uint32 GetMaxWords() const
 	{
 		return FBitSet::CalculateNumWords(MaxBits);
 	}
 
-	FORCEINLINE uint32 GetLastWordMask(int32 EndIndexExclusive) const
+	[[nodiscard]] FORCEINLINE uint32 GetLastWordMask(int32 EndIndexExclusive) const
 	{
 		const uint32 UnusedBits = (FBitSet::BitsPerWord - static_cast<uint32>(EndIndexExclusive) % FBitSet::BitsPerWord) % FBitSet::BitsPerWord;
 		return ~0u >> UnusedBits;
@@ -1052,7 +1052,7 @@ public:
 	 * Helper function to return the amount of memory allocated by this container 
 	 * @return number of bytes allocated by this container
 	 */
-	uint32 GetAllocatedSize( void ) const
+	[[nodiscard]] uint32 GetAllocatedSize( void ) const
 	{
 		return FBitSet::CalculateNumWords(MaxBits) * sizeof(uint32);
 	}
@@ -1074,7 +1074,7 @@ public:
 	 *
 	 * @return The index of the first occurrence of the specified value (true/false), or INDEX_NONE if not found.
 	 */
-	int32 Find(bool bValue) const
+	[[nodiscard]] int32 Find(bool bValue) const
 	{
 		return FindFromImpl(bValue, 0);
 	}
@@ -1167,7 +1167,7 @@ public:
 	 * @return The index of the first occurrence of the specified value (true/false) from StartIndex, or INDEX_NONE if not found.
 	 */
 	template <typename IndexType>
-	FORCEINLINE int32 FindFrom(bool bValue, IndexType StartIndex) const
+	[[nodiscard]] FORCEINLINE int32 FindFrom(bool bValue, IndexType StartIndex) const
 	{
 		static_assert(!std::is_same_v<IndexType, bool>, "TBitArray::FindFrom: unexpected bool passed as the StartIndex argument");
 		checkSlow((StartIndex >= 0) && (StartIndex <= NumBits));
@@ -1182,7 +1182,7 @@ public:
 	 *
 	 * @return The index of the last occurrence of the specified value (true/false), or INDEX_NONE if not found.
 	 */
-	int32 FindLast(bool bValue) const 
+	[[nodiscard]] int32 FindLast(bool bValue) const 
 	{
 		return FindLastFromImpl(bValue, /*EndIndexExclusive = */NumBits);
 	}
@@ -1197,7 +1197,7 @@ public:
 	 * @return The index of the last occurrence of the specified value (true/false) from EndIndexInclusive, or INDEX_NONE if not found.
 	 */
 	template <typename IndexType>
-	int32 FindLastFrom(bool bValue, IndexType EndIndexInclusive) const
+	[[nodiscard]] int32 FindLastFrom(bool bValue, IndexType EndIndexInclusive) const
 	{
 		static_assert(!std::is_same_v<IndexType, bool>, "TBitArray::FindLastFrom: unexpected bool passed as the EndIndexInclusive argument");
 		// This range check here ([-1, NumBits)) is because the end index is inclusive and because this is the reverse logic of the FindFrom range ([0, NumBits]) 
@@ -1214,7 +1214,7 @@ public:
 	 *
 	 * @return true if the array contains the specified value, false otherwise.
 	 */
-	FORCEINLINE bool Contains(bool bValue) const
+	[[nodiscard]] FORCEINLINE bool Contains(bool bValue) const
 	{
 		return Find(bValue) != INDEX_NONE;
 	}
@@ -1259,7 +1259,7 @@ public:
 	 * Return the bitwise AND of two bit arrays. The resulting bit array will be sized according to InFlags.
 	 */
 	template <typename AllocatorA, typename AllocatorB>
-	static TBitArray BitwiseAND(const TBitArray<AllocatorA>& A, const TBitArray<AllocatorB>& B, EBitwiseOperatorFlags InFlags)
+	[[nodiscard]] static TBitArray BitwiseAND(const TBitArray<AllocatorA>& A, const TBitArray<AllocatorB>& B, EBitwiseOperatorFlags InFlags)
 	{
 		TBitArray Result;
 		UE::Core::Private::BitwiseBinaryOperatorImpl(A, B, Result, InFlags, [](uint32 InA, uint32 InB) { return InA & InB; });
@@ -1290,7 +1290,7 @@ public:
 	 * Return the bitwise OR of two bit arrays. The resulting bit array will be sized according to InFlags.
 	 */
 	template <typename AllocatorA, typename AllocatorB>
-	static TBitArray BitwiseOR(const TBitArray<AllocatorA>& A, const TBitArray<AllocatorB>& B, EBitwiseOperatorFlags InFlags)
+	[[nodiscard]] static TBitArray BitwiseOR(const TBitArray<AllocatorA>& A, const TBitArray<AllocatorB>& B, EBitwiseOperatorFlags InFlags)
 	{
 		if constexpr (std::is_same_v<AllocatorA, AllocatorB>)
 		{
@@ -1316,7 +1316,7 @@ public:
 	 * Return the bitwise XOR of two bit arrays. The resulting bit array will be sized according to InFlags.
 	 */
 	template <typename AllocatorA, typename AllocatorB>
-	static TBitArray BitwiseXOR(const TBitArray<AllocatorA>& A, const TBitArray<AllocatorB>& B, EBitwiseOperatorFlags InFlags)
+	[[nodiscard]] static TBitArray BitwiseXOR(const TBitArray<AllocatorA>& A, const TBitArray<AllocatorB>& B, EBitwiseOperatorFlags InFlags)
 	{
 		TBitArray Result;
 		UE::Core::Private::BitwiseBinaryOperatorImpl(A, B, Result, InFlags, [](uint32 InA, uint32 InB) { return InA ^ InB; });
@@ -1347,7 +1347,7 @@ public:
 	/**
 	 * Count the number of set bits in this array  FromIndex <= bit < ToIndex
 	 */
-	int32 CountSetBits(int32 FromIndex = 0, int32 ToIndex = INDEX_NONE) const
+	[[nodiscard]] int32 CountSetBits(int32 FromIndex = 0, int32 ToIndex = INDEX_NONE) const
 	{
 		if (ToIndex == INDEX_NONE)
 		{
@@ -1374,7 +1374,7 @@ public:
 	 * @return true if this array matches Other, including any missing bits, false otherwise
 	 */
 	template <typename OtherAllocator>
-	bool CompareSetBits(const TBitArray<OtherAllocator>& Other, const bool bMissingBitValue) const
+	[[nodiscard]] bool CompareSetBits(const TBitArray<OtherAllocator>& Other, const bool bMissingBitValue) const
 	{
 		const uint32 MissingBitsFill = bMissingBitValue ? ~0u : 0;
 
@@ -1419,7 +1419,7 @@ public:
 	}
 
 	// Accessors.
-	FORCEINLINE bool IsValidIndex(int32 InIndex) const
+	[[nodiscard]] FORCEINLINE bool IsValidIndex(int32 InIndex) const
 	{
 		return InIndex >= 0 && InIndex < NumBits;
 	}
@@ -1430,14 +1430,14 @@ public:
 	 * @returns True if the bit array is empty.
 	 * @see Num
 	 */
-	bool IsEmpty() const 
+	[[nodiscard]] bool IsEmpty() const 
 	{ 
 		return NumBits == 0; 
 	}
 
-	FORCEINLINE int32 Num() const { return NumBits; }
-	FORCEINLINE int32 Max() const { return MaxBits; }
-	FORCEINLINE FBitReference operator[](int32 Index)
+	[[nodiscard]] FORCEINLINE int32 Num() const { return NumBits; }
+	[[nodiscard]] FORCEINLINE int32 Max() const { return MaxBits; }
+	[[nodiscard]] FORCEINLINE FBitReference operator[](int32 Index)
 	{
 		check(Index>=0 && Index<NumBits);
 		return FBitReference(
@@ -1445,7 +1445,7 @@ public:
 			1 << (Index & (NumBitsPerDWORD - 1))
 			);
 	}
-	FORCEINLINE const FConstBitReference operator[](int32 Index) const
+	[[nodiscard]] FORCEINLINE const FConstBitReference operator[](int32 Index) const
 	{
 		check(Index>=0 && Index<NumBits);
 		return FConstBitReference(
@@ -1453,7 +1453,7 @@ public:
 			1 << (Index & (NumBitsPerDWORD - 1))
 			);
 	}
-	FORCEINLINE FBitReference AccessCorrespondingBit(const FRelativeBitReference& RelativeReference)
+	[[nodiscard]] FORCEINLINE FBitReference AccessCorrespondingBit(const FRelativeBitReference& RelativeReference)
 	{
 		checkSlow(RelativeReference.Mask);
 		checkSlow(RelativeReference.WordIndex >= 0);
@@ -1463,7 +1463,7 @@ public:
 			RelativeReference.Mask
 			);
 	}
-	FORCEINLINE const FConstBitReference AccessCorrespondingBit(const FRelativeBitReference& RelativeReference) const
+	[[nodiscard]] FORCEINLINE const FConstBitReference AccessCorrespondingBit(const FRelativeBitReference& RelativeReference) const
 	{
 		checkSlow(RelativeReference.Mask);
 		checkSlow(RelativeReference.WordIndex >= 0);
@@ -1478,7 +1478,7 @@ public:
 	class FIterator : public FRelativeBitReference
 	{
 	public:
-		FORCEINLINE explicit FIterator(TBitArray<Allocator>& InArray UE_LIFETIMEBOUND,int32 StartIndex = 0)
+		[[nodiscard]] FORCEINLINE explicit FIterator(TBitArray<Allocator>& InArray UE_LIFETIMEBOUND,int32 StartIndex = 0)
 		:	FRelativeBitReference(StartIndex)
 		,	Array(InArray)
 		,	Index(StartIndex)
@@ -1497,25 +1497,25 @@ public:
 			return *this;
 		}
 
-		FORCEINLINE FBitReference operator*() const
+		[[nodiscard]] FORCEINLINE FBitReference operator*() const
 		{
 			return GetValue();
 		}
 
 		/** conversion to "bool" returning true if the iterator is valid. */
-		FORCEINLINE explicit operator bool() const
+		[[nodiscard]] FORCEINLINE explicit operator bool() const
 		{ 
 			return Index < Array.Num(); 
 		}
 
 		/** inverse of the "bool" operator */
-		FORCEINLINE bool operator !() const 
+		[[nodiscard]] FORCEINLINE bool operator !() const
 		{
 			return !(bool)*this;
 		}
 
-		FORCEINLINE FBitReference GetValue() const { return FBitReference(Array.GetData()[this->WordIndex],this->Mask); }
-		FORCEINLINE int32 GetIndex() const { return Index; }
+		[[nodiscard]] FORCEINLINE FBitReference GetValue() const { return FBitReference(Array.GetData()[this->WordIndex],this->Mask); }
+		[[nodiscard]] FORCEINLINE int32 GetIndex() const { return Index; }
 	private:
 		TBitArray<Allocator>& Array;
 		int32 Index;
@@ -1525,14 +1525,14 @@ public:
 	class FReverseIterator : public FRelativeBitReference
 	{
 	public:
-		FORCEINLINE explicit FReverseIterator(TBitArray<Allocator>& InArray UE_LIFETIMEBOUND)
+		[[nodiscard]] FORCEINLINE explicit FReverseIterator(TBitArray<Allocator>& InArray UE_LIFETIMEBOUND)
 			: FRelativeBitReference(InArray.Num() - 1)
 			, Array(InArray)
 			, Index(InArray.Num() - 1)
 		{
 		}
 
-		FORCEINLINE explicit FReverseIterator(TBitArray<Allocator>& InArray UE_LIFETIMEBOUND, int32 StartIndex)
+		[[nodiscard]] FORCEINLINE explicit FReverseIterator(TBitArray<Allocator>& InArray UE_LIFETIMEBOUND, int32 StartIndex)
 			: FRelativeBitReference(-1)
 			, Array(InArray)
 			, Index(-1)
@@ -1552,25 +1552,25 @@ public:
 			return *this;
 		}
 
-		FORCEINLINE FBitReference operator*() const
+		[[nodiscard]] FORCEINLINE FBitReference operator*() const
 		{
 			return GetValue();
 		}
 
 		/** conversion to "bool" returning true if the iterator is valid. */
-		FORCEINLINE explicit operator bool() const
+		[[nodiscard]] FORCEINLINE explicit operator bool() const
 		{
 			return Index >= 0;
 		}
 
 		/** inverse of the "bool" operator */
-		FORCEINLINE bool operator !() const
+		[[nodiscard]] FORCEINLINE bool operator !() const
 		{
 			return !(bool)*this;
 		}
 
-		FORCEINLINE FBitReference GetValue() const { return FBitReference(Array.GetData()[this->WordIndex], this->Mask); }
-		FORCEINLINE int32 GetIndex() const { return Index; }
+		[[nodiscard]] FORCEINLINE FBitReference GetValue() const { return FBitReference(Array.GetData()[this->WordIndex], this->Mask); }
+		[[nodiscard]] FORCEINLINE int32 GetIndex() const { return Index; }
 	private:
 		TBitArray<Allocator>& Array;
 		int32 Index;
@@ -1580,7 +1580,7 @@ public:
 	class FConstIterator : public FRelativeBitReference
 	{
 	public:
-		FORCEINLINE explicit FConstIterator(const TBitArray<Allocator>& InArray UE_LIFETIMEBOUND,int32 StartIndex = 0)
+		[[nodiscard]] FORCEINLINE explicit FConstIterator(const TBitArray<Allocator>& InArray UE_LIFETIMEBOUND,int32 StartIndex = 0)
 		:	FRelativeBitReference(StartIndex)
 		,	Array(InArray)
 		,	Index(StartIndex)
@@ -1599,24 +1599,24 @@ public:
 			return *this;
 		}
 
-		FORCEINLINE FConstBitReference operator*() const
+		[[nodiscard]] FORCEINLINE FConstBitReference operator*() const
 		{
 			return GetValue();
 		}
 
 		/** conversion to "bool" returning true if the iterator is valid. */
-		FORCEINLINE explicit operator bool() const
-		{ 
+		[[nodiscard]] FORCEINLINE explicit operator bool() const
+		{
 			return Index < Array.Num(); 
 		}
 		/** inverse of the "bool" operator */
-		FORCEINLINE bool operator !() const 
+		[[nodiscard]] FORCEINLINE bool operator !() const 
 		{
 			return !(bool)*this;
 		}
 
-		FORCEINLINE FConstBitReference GetValue() const { return FConstBitReference(Array.GetData()[this->WordIndex],this->Mask); }
-		FORCEINLINE int32 GetIndex() const { return Index; }
+		[[nodiscard]] FORCEINLINE FConstBitReference GetValue() const { return FConstBitReference(Array.GetData()[this->WordIndex],this->Mask); }
+		[[nodiscard]] FORCEINLINE int32 GetIndex() const { return Index; }
 	private:
 		const TBitArray<Allocator>& Array;
 		int32 Index;
@@ -1626,13 +1626,13 @@ public:
 	class FConstReverseIterator : public FRelativeBitReference
 	{
 	public:
-		FORCEINLINE explicit FConstReverseIterator(const TBitArray<Allocator>& InArray UE_LIFETIMEBOUND)
+		[[nodiscard]] FORCEINLINE explicit FConstReverseIterator(const TBitArray<Allocator>& InArray UE_LIFETIMEBOUND)
 			:	FRelativeBitReference(InArray.Num() - 1)
 			,	Array(InArray)
 			,	Index(InArray.Num() - 1)
 		{
 		}
-		FORCEINLINE explicit FConstReverseIterator(const TBitArray<Allocator>& InArray UE_LIFETIMEBOUND, int32 StartIndex)
+		[[nodiscard]] FORCEINLINE explicit FConstReverseIterator(const TBitArray<Allocator>& InArray UE_LIFETIMEBOUND, int32 StartIndex)
 			: FRelativeBitReference(-1)
 			, Array(InArray)
 			, Index(-1)
@@ -1657,29 +1657,29 @@ public:
 		}
 
 		/** conversion to "bool" returning true if the iterator is valid. */
-		FORCEINLINE explicit operator bool() const
-		{ 
+		[[nodiscard]] FORCEINLINE explicit operator bool() const
+		{
 			return Index >= 0; 
 		}
 		/** inverse of the "bool" operator */
-		FORCEINLINE bool operator !() const 
+		[[nodiscard]] FORCEINLINE bool operator !() const 
 		{
 			return !(bool)*this;
 		}
 
-		FORCEINLINE FConstBitReference GetValue() const { return FConstBitReference(Array.GetData()[this->WordIndex],this->Mask); }
-		FORCEINLINE int32 GetIndex() const { return Index; }
+		[[nodiscard]] FORCEINLINE FConstBitReference GetValue() const { return FConstBitReference(Array.GetData()[this->WordIndex],this->Mask); }
+		[[nodiscard]] FORCEINLINE int32 GetIndex() const { return Index; }
 	private:
 		const TBitArray<Allocator>& Array;
 		int32 Index;
 	};
 
-	FORCEINLINE const uint32* GetData() const
+	[[nodiscard]] FORCEINLINE const uint32* GetData() const
 	{
 		return (uint32*)AllocatorInstance.GetAllocation();
 	}
 
-	FORCEINLINE uint32* GetData()
+	[[nodiscard]] FORCEINLINE uint32* GetData()
 	{
 		return (uint32*)AllocatorInstance.GetAllocation();
 	}
@@ -1688,17 +1688,17 @@ private:
 	template<typename WordType>
 	struct TWordIteratorBase
 	{
-		explicit operator bool() const
+		[[nodiscard]] explicit operator bool() const
 		{
 			return CurrentIndex < NumWords;
 		}
 
-		int32 GetIndex() const
+		[[nodiscard]] int32 GetIndex() const
 		{
 			return CurrentIndex;
 		}
 
-		uint32 GetWord() const
+		[[nodiscard]] uint32 GetWord() const
 		{
 			checkSlow(CurrentIndex < NumWords);
 
@@ -1736,7 +1736,7 @@ private:
 
 	protected:
 
-		explicit TWordIteratorBase(WordType* InData, int32 InStartBitIndex, int32 InEndBitIndex)
+		[[nodiscard]] explicit TWordIteratorBase(WordType* InData, int32 InStartBitIndex, int32 InEndBitIndex)
 			: Data(InData)
 			, CurrentIndex(InStartBitIndex / NumBitsPerDWORD)
 			, NumWords(FMath::DivideAndRoundUp(InEndBitIndex, NumBitsPerDWORD))
@@ -1770,11 +1770,12 @@ private:
 public:
 	struct FConstWordIterator : TWordIteratorBase<const uint32>
 	{
-		explicit FConstWordIterator(const TBitArray<Allocator>& InArray UE_LIFETIMEBOUND)
+		[[nodiscard]] explicit FConstWordIterator(const TBitArray<Allocator>& InArray UE_LIFETIMEBOUND)
 			: TWordIteratorBase<const uint32>(InArray.GetData(), 0, InArray.Num())
-		{}
+		{
+		}
 
-		explicit FConstWordIterator(const TBitArray<Allocator>& InArray UE_LIFETIMEBOUND, int32 InStartBitIndex, int32 InEndBitIndex)
+		[[nodiscard]] explicit FConstWordIterator(const TBitArray<Allocator>& InArray UE_LIFETIMEBOUND, int32 InStartBitIndex, int32 InEndBitIndex)
 			: TWordIteratorBase<const uint32>(InArray.GetData(), InStartBitIndex, InEndBitIndex)
 		{
 			checkSlow(InStartBitIndex <= InEndBitIndex && InStartBitIndex <= InArray.Num() && InEndBitIndex <= InArray.Num());
@@ -1785,9 +1786,10 @@ public:
 
 	struct FWordIterator : TWordIteratorBase<uint32>
 	{
-		explicit FWordIterator(TBitArray<Allocator>& InArray UE_LIFETIMEBOUND)
+		[[nodiscard]] explicit FWordIterator(TBitArray<Allocator>& InArray UE_LIFETIMEBOUND)
 			: TWordIteratorBase<uint32>(InArray.GetData(), 0, InArray.Num())
-		{}
+		{
+		}
 
 		void SetWord(uint32 InWord)
 		{
@@ -1806,15 +1808,15 @@ public:
 
 	/** Enables range-based for loops, DO NOT USE DIRECTLY. */
 
-	FORCEINLINE FIterator		begin()			{ return FIterator(*this); }
-	FORCEINLINE FConstIterator	begin() const	{ return FConstIterator(*this); }
-	FORCEINLINE FIterator		end()			{ return FIterator(*this, NumBits); }
-	FORCEINLINE FConstIterator	end() const		{ return FConstIterator(*this, NumBits); }
+	[[nodiscard]] FORCEINLINE FIterator      begin()       { return FIterator(*this); }
+	[[nodiscard]] FORCEINLINE FConstIterator begin() const { return FConstIterator(*this); }
+	[[nodiscard]] FORCEINLINE FIterator      end()         { return FIterator(*this, NumBits); }
+	[[nodiscard]] FORCEINLINE FConstIterator end()   const { return FConstIterator(*this, NumBits); }
 
-	FORCEINLINE FReverseIterator		rbegin() { return FReverseIterator(*this); }
-	FORCEINLINE FConstReverseIterator	rbegin() const { return FConstReverseIterator(*this); }
-	FORCEINLINE FReverseIterator		rend() { return FReverseIterator(*this, 0); }
-	FORCEINLINE FConstReverseIterator	rend() const { return FConstReverseIterator(*this, 0); }
+	[[nodiscard]] FORCEINLINE FReverseIterator      rbegin()       { return FReverseIterator(*this); }
+	[[nodiscard]] FORCEINLINE FConstReverseIterator rbegin() const { return FConstReverseIterator(*this); }
+	[[nodiscard]] FORCEINLINE FReverseIterator      rend()         { return FReverseIterator(*this, 0); }
+	[[nodiscard]] FORCEINLINE FConstReverseIterator rend()   const { return FConstReverseIterator(*this, 0); }
 
 private:
 	AllocatorType AllocatorInstance;
@@ -1895,7 +1897,7 @@ namespace Freeze
 DECLARE_TEMPLATE_INTRINSIC_TYPE_LAYOUT(template<typename Allocator>, TBitArray<Allocator>);
 
 template<typename Allocator>
-FORCEINLINE uint32 GetTypeHash(const TBitArray<Allocator>& BitArray)
+[[nodiscard]] FORCEINLINE uint32 GetTypeHash(const TBitArray<Allocator>& BitArray)
 {
 	uint32 NumWords = FBitSet::CalculateNumWords(BitArray.Num());
 	uint32 Hash = NumWords;
@@ -1914,7 +1916,7 @@ class TConstSetBitIterator : public FRelativeBitReference
 {
 public:
 
-	explicit TConstSetBitIterator(const TBitArray<Allocator>& InArray UE_LIFETIMEBOUND)
+	[[nodiscard]] explicit TConstSetBitIterator(const TBitArray<Allocator>& InArray UE_LIFETIMEBOUND)
 		: FRelativeBitReference(0)
 		, Array                (InArray)
 		, UnvisitedBitMask     (~0U)
@@ -1927,7 +1929,7 @@ public:
 		}
 	}
 
-	explicit TConstSetBitIterator(const TBitArray<Allocator>& InArray UE_LIFETIMEBOUND, int32 StartIndex)
+	[[nodiscard]] explicit TConstSetBitIterator(const TBitArray<Allocator>& InArray UE_LIFETIMEBOUND, int32 StartIndex)
 		: FRelativeBitReference(StartIndex)
 		, Array                (InArray)
 		, UnvisitedBitMask     ((~0U) << (StartIndex & (NumBitsPerDWORD - 1)))
@@ -1953,32 +1955,32 @@ public:
 		return *this;
 	}
 
-	FORCEINLINE bool operator==(const TConstSetBitIterator& Rhs) const
+	[[nodiscard]] FORCEINLINE bool operator==(const TConstSetBitIterator& Rhs) const
 	{
 		// We only need to compare the bit index and the array... all the rest of the state is unobservable.
 		return CurrentBitIndex == Rhs.CurrentBitIndex && &Array == &Rhs.Array;
 	}
 
 #if !PLATFORM_COMPILER_HAS_GENERATED_COMPARISON_OPERATORS
-	FORCEINLINE bool operator!=(const TConstSetBitIterator& Rhs) const
+	[[nodiscard]] FORCEINLINE bool operator!=(const TConstSetBitIterator& Rhs) const
 	{ 
 		return !(*this == Rhs);
 	}
 #endif
 
 	/** conversion to "bool" returning true if the iterator is valid. */
-	FORCEINLINE explicit operator bool() const
+	[[nodiscard]] FORCEINLINE explicit operator bool() const
 	{ 
 		return CurrentBitIndex < Array.Num(); 
 	}
 	/** inverse of the "bool" operator */
-	FORCEINLINE bool operator !() const 
+	[[nodiscard]] FORCEINLINE bool operator !() const
 	{
 		return !(bool)*this;
 	}
 
 	/** Index accessor. */
-	FORCEINLINE int32 GetIndex() const
+	[[nodiscard]] FORCEINLINE int32 GetIndex() const
 	{
 		return CurrentBitIndex;
 	}
@@ -2043,7 +2045,7 @@ class TConstDualSetBitIterator : public FRelativeBitReference
 public:
 
 	/** Constructor. */
-	FORCEINLINE TConstDualSetBitIterator(
+	[[nodiscard]] FORCEINLINE TConstDualSetBitIterator(
 		const TBitArray<Allocator>& InArrayA,
 		const TBitArray<OtherAllocator>& InArrayB,
 		int32 StartIndex = 0
@@ -2076,18 +2078,18 @@ public:
 	}
 
 	/** conversion to "bool" returning true if the iterator is valid. */
-	FORCEINLINE explicit operator bool() const
+	[[nodiscard]] FORCEINLINE explicit operator bool() const
 	{ 
 		return CurrentBitIndex < ArrayA.Num(); 
 	}
 	/** inverse of the "bool" operator */
-	FORCEINLINE bool operator !() const 
+	[[nodiscard]] FORCEINLINE bool operator !() const
 	{
 		return !(bool)*this;
 	}
 
 	/** Index accessor. */
-	FORCEINLINE int32 GetIndex() const
+	[[nodiscard]] FORCEINLINE int32 GetIndex() const
 	{
 		return CurrentBitIndex;
 	}
@@ -2192,24 +2194,24 @@ public:
 	 * @param Value - The value to initial the bits to.
 	 * @param InNumBits - The initial number of bits in the array.
 	 */
-	TScriptBitArray()
+	[[nodiscard]] TScriptBitArray()
 		: NumBits(0)
 		, MaxBits(0)
 	{
 	}
 
-	bool IsValidIndex(int32 Index) const
+	[[nodiscard]] bool IsValidIndex(int32 Index) const
 	{
 		return Index >= 0 && Index < NumBits;
 	}
 
-	FBitReference operator[](int32 Index)
+	[[nodiscard]] FBitReference operator[](int32 Index)
 	{
 		check(IsValidIndex(Index));
 		return FBitReference(GetData()[Index / NumBitsPerDWORD], 1 << (Index & (NumBitsPerDWORD - 1)));
 	}
 
-	FConstBitReference operator[](int32 Index) const
+	[[nodiscard]] FConstBitReference operator[](int32 Index) const
 	{
 		check(IsValidIndex(Index));
 		return FConstBitReference(GetData()[Index / NumBitsPerDWORD], 1 << (Index & (NumBitsPerDWORD - 1)));
@@ -2282,12 +2284,12 @@ private:
 		static_assert(STRUCT_OFFSET(ScriptType, MaxBits)           == STRUCT_OFFSET(RealType, MaxBits),           "TScriptBitArray's MaxBits member offset does not match TBitArray's");
 	}
 
-	FORCEINLINE uint32* GetData()
+	[[nodiscard]] FORCEINLINE uint32* GetData()
 	{
 		return (uint32*)AllocatorInstance.GetAllocation();
 	}
 
-	FORCEINLINE const uint32* GetData() const
+	[[nodiscard]] FORCEINLINE const uint32* GetData() const
 	{
 		return (const uint32*)AllocatorInstance.GetAllocation();
 	}
@@ -2330,7 +2332,7 @@ private:
 public:
 	// These should really be private, because they shouldn't be called, but there's a bunch of code
 	// that needs to be fixed first.
-	TScriptBitArray(const TScriptBitArray&) { check(false); }
+	[[nodiscard]] TScriptBitArray(const TScriptBitArray&) { check(false); }
 	void operator=(const TScriptBitArray&) { check(false); }
 };
 
diff --git a/Engine/Source/Runtime/Core/Public/Containers/ChunkedArray.h b/Engine/Source/Runtime/Core/Public/Containers/ChunkedArray.h
--- a/Engine/Source/Runtime/Core/Public/Containers/ChunkedArray.h
+++ b/Engine/Source/Runtime/Core/Public/Containers/ChunkedArray.h
@@ -18,7 +18,7 @@ namespace UE4ChunkedArray_Private
 		uint32		Count = 0;
 		uint32		ElementIndex = 0;
 
-		ElementType& operator*() const
+		[[nodiscard]] ElementType& operator*() const
 		{
 			return (*Chunk)->Elements[ElementIndex];
 		}
@@ -35,7 +35,7 @@ namespace UE4ChunkedArray_Private
 			++Count;
 		}
 
-		bool operator!=(const TChunkedArrayIterator& Rhs) const
+		[[nodiscard]] bool operator!=(const TChunkedArrayIterator& Rhs) const
 		{
 			return Count < Rhs.Count;
 		}
@@ -59,7 +59,7 @@ class TChunkedArray
 public:
 
 	/** Initialization constructor. */
-	TChunkedArray(int32 InNumElements = 0):
+	[[nodiscard]] TChunkedArray(int32 InNumElements = 0):
 		NumElements(InNumElements)
 	{
 		// Compute the number of chunks needed.
@@ -83,7 +83,7 @@ private:
 	}
 
 public:
-	TChunkedArray(TChunkedArray&& Other)
+	[[nodiscard]] TChunkedArray(TChunkedArray&& Other)
 	{
 		this->Move(*this, Other);
 	}
@@ -98,7 +98,7 @@ public:
 		return *this;
 	}
 
-	TChunkedArray(const TChunkedArray&) = default;
+	[[nodiscard]] TChunkedArray(const TChunkedArray&) = default;
 	TChunkedArray& operator=(const TChunkedArray&) = default;
 
 	//////////////////////////////////////////////////////
@@ -107,11 +107,11 @@ public:
 	constexpr static bool bHasIntrusiveUnsetOptionalState = true;
 	using IntrusiveUnsetOptionalStateType = TChunkedArray;
 
-	explicit TChunkedArray(FIntrusiveUnsetOptionalState)
+	[[nodiscard]] explicit TChunkedArray(FIntrusiveUnsetOptionalState)
 		: NumElements(-1)
 	{
 	}
-	bool operator==(FIntrusiveUnsetOptionalState) const
+	[[nodiscard]] bool operator==(FIntrusiveUnsetOptionalState) const
 	{
 		return NumElements == -1;
 	}
@@ -120,25 +120,25 @@ public:
 	////////////////////////////////////////////////////
 
 	// Accessors.
-	ElementType& operator()(int32 ElementIndex)
+	[[nodiscard]] ElementType& operator()(int32 ElementIndex)
 	{
 		const uint32 ChunkIndex = ElementIndex / NumElementsPerChunk;
 		const uint32 ChunkElementIndex = ElementIndex % NumElementsPerChunk;
 		return Chunks[ChunkIndex].Elements[ChunkElementIndex];
 	}
-	const ElementType& operator()(int32 ElementIndex) const
+	[[nodiscard]] const ElementType& operator()(int32 ElementIndex) const
 	{
 		const int32 ChunkIndex = ElementIndex / NumElementsPerChunk;
 		const int32 ChunkElementIndex = ElementIndex % NumElementsPerChunk;
 		return Chunks[ChunkIndex].Elements[ChunkElementIndex];
 	}
-	ElementType& operator[](int32 ElementIndex)
+	[[nodiscard]] ElementType& operator[](int32 ElementIndex)
 	{
 		const uint32 ChunkIndex = ElementIndex / NumElementsPerChunk;
 		const uint32 ChunkElementIndex = ElementIndex % NumElementsPerChunk;
 		return Chunks[ChunkIndex].Elements[ChunkElementIndex];
 	}
-	const ElementType& operator[](int32 ElementIndex) const
+	[[nodiscard]] const ElementType& operator[](int32 ElementIndex) const
 	{
 		const int32 ChunkIndex = ElementIndex / NumElementsPerChunk;
 		const int32 ChunkElementIndex = ElementIndex % NumElementsPerChunk;
@@ -151,17 +151,17 @@ public:
 	 * @returns True if the chunked array is empty.
 	 * @see Num
 	 */
-	bool IsEmpty() const
+	[[nodiscard]] bool IsEmpty() const
 	{
 		return NumElements == 0;
 	}
 
-	int32 Num() const 
+	[[nodiscard]] int32 Num() const
 	{ 
 		return NumElements; 
 	}
 
-	SIZE_T GetAllocatedSize( void ) const
+	[[nodiscard]] SIZE_T GetAllocatedSize( void ) const
 	{
 		return Chunks.GetAllocatedSize();
 	}
@@ -174,7 +174,7 @@ public:
 	*
 	* @returns True if index is valid. False otherwise.
 	*/
-	FORCEINLINE bool IsValidIndex(int32 Index) const
+	[[nodiscard]] FORCEINLINE bool IsValidIndex(int32 Index) const
 	{
 		return Index >= 0 && Index < NumElements;
 	}
@@ -307,22 +307,22 @@ private:
 	typedef UE4ChunkedArray_Private::TChunkedArrayIterator<const FChunk, const ElementType, NumElementsPerChunk> FConstIterType;
 
 public:
-	FIterType begin()
+	[[nodiscard]] FIterType begin()
 	{
 		return FIterType{Chunks.GetData()};
 	}
 
-	FConstIterType begin() const
+	[[nodiscard]] FConstIterType begin() const
 	{
 		return FConstIterType{Chunks.GetData()};
 	}
 
-	FIterType end()
+	[[nodiscard]] FIterType end()
 	{
 		return FIterType{nullptr, uint32(NumElements)};
 	}
 
-	FConstIterType end() const
+	[[nodiscard]] FConstIterType end() const
 	{
 		return FConstIterType{nullptr, uint32(NumElements)};
 	}
diff --git a/Engine/Source/Runtime/Core/Public/Containers/CircularBuffer.h b/Engine/Source/Runtime/Core/Public/Containers/CircularBuffer.h
--- a/Engine/Source/Runtime/Core/Public/Containers/CircularBuffer.h
+++ b/Engine/Source/Runtime/Core/Public/Containers/CircularBuffer.h
@@ -24,7 +24,7 @@ public:
 	 *
 	 * @param Capacity The number of elements that the buffer can store (will be rounded up to the next power of 2).
 	 */
-	explicit TCircularBuffer(uint32 Capacity)
+	[[nodiscard]] explicit TCircularBuffer(uint32 Capacity)
 	{
 		checkSlow(Capacity > 0);
 		checkSlow(Capacity <= 0xffffffffU);
@@ -39,7 +39,7 @@ public:
 	 * @param Capacity The number of elements that the buffer can store (will be rounded up to the next power of 2).
 	 * @param InitialValue The initial value for the buffer's elements.
 	 */
-	TCircularBuffer(uint32 Capacity, const ElementType& InitialValue)
+	[[nodiscard]] TCircularBuffer(uint32 Capacity, const ElementType& InitialValue)
 	{
 		checkSlow(Capacity <= 0xffffffffU);
 
@@ -54,7 +54,7 @@ public:
 	 *
 	 * @param Index The index of the element to return.
 	 */
-	FORCEINLINE ElementType& operator[](uint32 Index)
+	[[nodiscard]] FORCEINLINE ElementType& operator[](uint32 Index)
 	{
 		return Elements[Index & IndexMask];
 	}
@@ -64,7 +64,7 @@ public:
 	 *
 	 * @param Index The index of the element to return.
 	 */
-	FORCEINLINE const ElementType& operator[](uint32 Index) const
+	[[nodiscard]] FORCEINLINE const ElementType& operator[](uint32 Index) const
 	{
 		return Elements[Index & IndexMask];
 	}
@@ -76,7 +76,7 @@ public:
 	 *
 	 * @return Buffer capacity.
 	 */
-	FORCEINLINE uint32 Capacity() const
+	[[nodiscard]] FORCEINLINE uint32 Capacity() const
 	{
 		return Elements.Num();
 	}
@@ -87,7 +87,7 @@ public:
 	 * @param CurrentIndex The current index.
 	 * @return The next index.
 	 */
-	FORCEINLINE uint32 GetNextIndex(uint32 CurrentIndex) const
+	[[nodiscard]] FORCEINLINE uint32 GetNextIndex(uint32 CurrentIndex) const
 	{
 		return ((CurrentIndex + 1) & IndexMask);
 	}
@@ -98,7 +98,7 @@ public:
 	 * @param CurrentIndex The current index.
 	 * @return The previous index.
 	 */
-	FORCEINLINE uint32 GetPreviousIndex(uint32 CurrentIndex) const
+	[[nodiscard]] FORCEINLINE uint32 GetPreviousIndex(uint32 CurrentIndex) const
 	{
 		return ((CurrentIndex - 1) & IndexMask);
 	}
diff --git a/Engine/Source/Runtime/Core/Public/Containers/ClosableMpscQueue.h b/Engine/Source/Runtime/Core/Public/Containers/ClosableMpscQueue.h
--- a/Engine/Source/Runtime/Core/Public/Containers/ClosableMpscQueue.h
+++ b/Engine/Source/Runtime/Core/Public/Containers/ClosableMpscQueue.h
@@ -18,7 +18,7 @@ class TClosableMpscQueue final
 public:
 	UE_NONCOPYABLE(TClosableMpscQueue);
 
-	TClosableMpscQueue() = default;
+	[[nodiscard]] TClosableMpscQueue() = default;
 
 	~TClosableMpscQueue()
 	{
@@ -88,7 +88,7 @@ public:
 		return Head_Local != nullptr;
 	}
 
-	bool IsClosed() const
+	[[nodiscard]] bool IsClosed() const
 	{
 		return Head.load(std::memory_order_relaxed) == nullptr;
 	}
diff --git a/Engine/Source/Runtime/Core/Public/Containers/ConsumeAllMpmcQueue.h b/Engine/Source/Runtime/Core/Public/Containers/ConsumeAllMpmcQueue.h
--- a/Engine/Source/Runtime/Core/Public/Containers/ConsumeAllMpmcQueue.h
+++ b/Engine/Source/Runtime/Core/Public/Containers/ConsumeAllMpmcQueue.h
@@ -36,7 +36,7 @@ namespace UE
 	public:
 		UE_NONCOPYABLE(TConsumeAllMpmcQueue);
 
-		TConsumeAllMpmcQueue() = default;
+		[[nodiscard]] TConsumeAllMpmcQueue() = default;
 
 		~TConsumeAllMpmcQueue()
 		{
@@ -83,7 +83,7 @@ namespace UE
 		}
 
 		// the result can be relied upon only in special cases (e.g. debug checks), as the state can change concurrently. use with caution 
-		bool IsEmpty() const
+		[[nodiscard]] bool IsEmpty() const
 		{
 			return Head.load(std::memory_order_relaxed) == nullptr;
 		}
diff --git a/Engine/Source/Runtime/Core/Public/Containers/DepletableMpmcQueue.h b/Engine/Source/Runtime/Core/Public/Containers/DepletableMpmcQueue.h
--- a/Engine/Source/Runtime/Core/Public/Containers/DepletableMpmcQueue.h
+++ b/Engine/Source/Runtime/Core/Public/Containers/DepletableMpmcQueue.h
@@ -37,7 +37,7 @@ namespace UE
 	public:
 		UE_NONCOPYABLE(TDepletableMpmcQueue);
 
-		TDepletableMpmcQueue() = default;
+		[[nodiscard]] TDepletableMpmcQueue() = default;
 
 		~TDepletableMpmcQueue()
 		{
@@ -57,11 +57,11 @@ namespace UE
 		template <typename... ArgTypes>
 		void Enqueue(ArgTypes&&... Args)
 		{
-			EnqueueAndReturnWasEmpty(Forward<ArgTypes>(Args)...);
+			(void)EnqueueAndReturnWasEmpty(Forward<ArgTypes>(Args)...);
 		}
 
 		template <typename... ArgTypes>
-		bool EnqueueAndReturnWasEmpty(ArgTypes&&... Args)
+		[[nodiscard]] bool EnqueueAndReturnWasEmpty(ArgTypes&&... Args)
 		{
 			FNode* New = ::new(AllocatorType::Malloc(sizeof(FNode), alignof(FNode))) FNode;
 			::new ((void*)&New->Value) T(Forward<ArgTypes>(Args)...);
diff --git a/Engine/Source/Runtime/Core/Public/Containers/Deque.h b/Engine/Source/Runtime/Core/Public/Containers/Deque.h
--- a/Engine/Source/Runtime/Core/Public/Containers/Deque.h
+++ b/Engine/Source/Runtime/Core/Public/Containers/Deque.h
@@ -51,7 +51,7 @@ protected:
 	{
 	}
 
-	FORCEINLINE ElementType& Dereference() const
+	[[nodiscard]] FORCEINLINE ElementType& Dereference() const
 	{
 		return *(Data + WrapAround(Offset, Range));
 	}
@@ -62,7 +62,7 @@ protected:
 		Offset++;
 	}
 
-	FORCEINLINE bool Equals(const TIteratorBase& Other) const
+	[[nodiscard]] FORCEINLINE bool Equals(const TIteratorBase& Other) const
 	{
 		return Data + Offset == Other.Data + Other.Offset;
 	}
@@ -103,21 +103,21 @@ public:
 	using ConstIteratorType = UE::Deque::Private::TIterator<const ElementType, SizeType>;
 	using IteratorType = UE::Deque::Private::TIterator<ElementType, SizeType>;
 
-	TDeque() : Capacity(Storage.GetInitialCapacity())
+	[[nodiscard]] TDeque() : Capacity(Storage.GetInitialCapacity())
 	{
 	}
 
-	TDeque(TDeque&& Other)
+	[[nodiscard]] TDeque(TDeque&& Other)
 	{
 		MoveUnchecked(MoveTemp(Other));
 	}
 
-	TDeque(const TDeque& Other)
+	[[nodiscard]] TDeque(const TDeque& Other)
 	{
 		CopyUnchecked(Other);
 	}
 
-	TDeque(std::initializer_list<ElementType> InList)
+	[[nodiscard]] TDeque(std::initializer_list<ElementType> InList)
 	{
 		CopyUnchecked(InList);
 	}
@@ -154,53 +154,53 @@ public:
 		return *this;
 	}
 
-	FORCEINLINE const ElementType& operator[](SizeType Index) const
+	[[nodiscard]] FORCEINLINE const ElementType& operator[](SizeType Index) const
 	{
 		CheckValidIndex(Index);
 		return GetData()[UE::Deque::Private::WrapAround(Head + Index, Capacity)];
 	}
 
-	FORCEINLINE ElementType& operator[](SizeType Index)
+	[[nodiscard]] FORCEINLINE ElementType& operator[](SizeType Index)
 	{
 		CheckValidIndex(Index);
 		return GetData()[UE::Deque::Private::WrapAround(Head + Index, Capacity)];
 	}
 
-	FORCEINLINE const ElementType& Last() const
+	[[nodiscard]] FORCEINLINE const ElementType& Last() const
 	{
 		CheckValidIndex(0);
 		return GetData()[UE::Deque::Private::WrapAround(Tail + Capacity - 1, Capacity)];
 	}
 
-	FORCEINLINE ElementType& Last()
+	[[nodiscard]] FORCEINLINE ElementType& Last()
 	{
 		CheckValidIndex(0);
 		return GetData()[UE::Deque::Private::WrapAround(Tail + Capacity - 1, Capacity)];
 	}
 
-	FORCEINLINE const ElementType& First() const
+	[[nodiscard]] FORCEINLINE const ElementType& First() const
 	{
 		CheckValidIndex(0);
 		return GetData()[Head];
 	}
 
-	FORCEINLINE ElementType& First()
+	[[nodiscard]] FORCEINLINE ElementType& First()
 	{
 		CheckValidIndex(0);
 		return GetData()[Head];
 	}
 
-	FORCEINLINE bool IsEmpty() const
+	[[nodiscard]] FORCEINLINE bool IsEmpty() const
 	{
 		return !Count;
 	}
 
-	FORCEINLINE SizeType Max() const
+	[[nodiscard]] FORCEINLINE SizeType Max() const
 	{
 		return Capacity;
 	}
 
-	FORCEINLINE SizeType Num() const
+	[[nodiscard]] FORCEINLINE SizeType Num() const
 	{
 		return Count;
 	}
@@ -211,7 +211,7 @@ public:
 	 *
 	 * @returns Number of bytes allocated by this container.
 	 */
-	FORCEINLINE SIZE_T GetAllocatedSize() const
+	[[nodiscard]] FORCEINLINE SIZE_T GetAllocatedSize() const
 	{
 		return Storage.GetAllocatedSize(Capacity, sizeof(ElementType));
 	}
@@ -301,7 +301,7 @@ public:
 	 * Removes the element at the back of the queue if existent after copying it to the parameter
 	 * out value.
 	 */
-	bool TryPopLast(ElementType& OutValue)
+	[[nodiscard]] bool TryPopLast(ElementType& OutValue)
 	{
 		if (IsEmpty())
 		{
@@ -319,7 +319,7 @@ public:
 	 * Removes the element at the front of the queue if existent after copying it to the parameter
 	 * out value.
 	 */
-	bool TryPopFirst(ElementType& OutValue)
+	[[nodiscard]] bool TryPopFirst(ElementType& OutValue)
 	{
 		if (IsEmpty())
 		{
@@ -379,19 +379,19 @@ public:
 	/*
 	 * STL IteratorType model compliance methods.
 	 */
-	FORCEINLINE ConstIteratorType begin() const
+	[[nodiscard]] FORCEINLINE ConstIteratorType begin() const
 	{
 		return GetIterator(Head);
 	}
-	FORCEINLINE IteratorType begin()
+	[[nodiscard]] FORCEINLINE IteratorType begin()
 	{
 		return GetIterator(Head);
 	}
-	FORCEINLINE ConstIteratorType end() const
+	[[nodiscard]] FORCEINLINE ConstIteratorType end() const
 	{
 		return GetIterator(Head + Count);
 	}
-	FORCEINLINE IteratorType end()
+	[[nodiscard]] FORCEINLINE IteratorType end()
 	{
 		return GetIterator(Head + Count);
 	}
@@ -403,22 +403,22 @@ private:
 	SizeType Head = 0;
 	SizeType Tail = 0;
 
-	FORCEINLINE const ElementType* GetData() const
+	[[nodiscard]] FORCEINLINE const ElementType* GetData() const
 	{
 		return (const ElementType*)Storage.GetAllocation();
 	}
 
-	FORCEINLINE ElementType* GetData()
+	[[nodiscard]] FORCEINLINE ElementType* GetData()
 	{
 		return (ElementType*)Storage.GetAllocation();
 	}
 
-	FORCEINLINE ConstIteratorType GetIterator(SizeType HeadOffset) const
+	[[nodiscard]] FORCEINLINE ConstIteratorType GetIterator(SizeType HeadOffset) const
 	{
 		return ConstIteratorType(InPlace, GetData(), Max(), HeadOffset);
 	}
 
-	FORCEINLINE IteratorType GetIterator(SizeType HeadOffset)
+	[[nodiscard]] FORCEINLINE IteratorType GetIterator(SizeType HeadOffset)
 	{
 		return IteratorType(InPlace, GetData(), Max(), HeadOffset);
 	}
@@ -591,7 +591,7 @@ private:
 	// TDeque comparison operators
 	//---------------------------------------------------------------------------------------------------------------------
 public:
-	inline bool operator==(const TDeque& Right) const
+	[[nodiscard]] inline bool operator==(const TDeque& Right) const
 	{
 		if (Num() == Right.Num())
 		{
@@ -611,7 +611,7 @@ public:
 	}
 
 #if !PLATFORM_COMPILER_HAS_GENERATED_COMPARISON_OPERATORS
-	inline bool operator!=(const TDeque& Right) const
+	[[nodiscard]] inline bool operator!=(const TDeque& Right) const
 	{
 		return !(*this == Right);
 	}
diff --git a/Engine/Source/Runtime/Core/Public/Containers/DiscardableKeyValueCache.h b/Engine/Source/Runtime/Core/Public/Containers/DiscardableKeyValueCache.h
--- a/Engine/Source/Runtime/Core/Public/Containers/DiscardableKeyValueCache.h
+++ b/Engine/Source/Runtime/Core/Public/Containers/DiscardableKeyValueCache.h
@@ -34,21 +34,32 @@ public:
 
 	typedef TMap<KeyType, ValueType> TypedMap;
 
-	TDiscardableKeyValueCache() :
+	[[nodiscard]] TDiscardableKeyValueCache() :
 		CurrentMap(&Map1)
-		, BackfillMap(&Map2) {}
+		, BackfillMap(&Map2)
+	{
+	}
 
 	/* Access to the internal locking object */
-	FRWLock&  RWLock() { return LockObject; }
+	[[nodiscard]] FRWLock&  RWLock()
+	{
+		return LockObject;
+	}
 
 	/* Reference to the current map */
-	TypedMap& Current() { return *CurrentMap; }
+	[[nodiscard]] TypedMap& Current()
+	{
+		return *CurrentMap;
+	}
 
 	/* Reference to the current map */
-	TypedMap& Backfill() { return *BackfillMap; }
+	[[nodiscard]] TypedMap& Backfill()
+	{
+		return *BackfillMap;
+	}
 
 	/* Returns the total number of items in the cache*/
-	int32 Num()
+	[[nodiscard]] int32 Num()
 	{
 		uint32 LockFlags = ApplyLock(0, LockFlags::ReadLock);
 
@@ -62,7 +73,7 @@ public:
 	/**
 	*  Returns true and sets OutType to the value with the associated key if it exists.
 	*/
-	bool  Find(const KeyType& Key, ValueType& OutType)
+	[[nodiscard]] bool Find(const KeyType& Key, ValueType& OutType)
 	{
 		uint32 LockFlags = ApplyLock(0, LockFlags::ReadLock);
 
@@ -78,7 +89,7 @@ public:
 	 *	InFlags represents the currently locked state of the object, OutFlags the state after the
 	 *	find operation has completed. Caller should be sure to unlock this object with OutFlags
 	 */
-	bool Find(const KeyType& Key, ValueType& OutType, uint32 InCurrentLockFlags, uint32& OutLockFlags)
+	[[nodiscard]] bool Find(const KeyType& Key, ValueType& OutType, uint32 InCurrentLockFlags, uint32& OutLockFlags)
 	{
 		return InternalFindWhileLocked(Key, OutType, InCurrentLockFlags, OutLockFlags);
 	}
@@ -239,7 +250,7 @@ protected:
 	*	Checks for the entry in our current map, and if not found the backfill. If the entry is in the backfill it is moved
 	*	to the current map so it will not be discarded when DiscardUnusedEntries is called
 	*/
-	bool  InternalFindWhileLocked(const KeyType& Key, ValueType& OutType, uint32 InCurrentLockFlags, uint32& OutFlags)
+	[[nodiscard]] bool  InternalFindWhileLocked(const KeyType& Key, ValueType& OutType, uint32 InCurrentLockFlags, uint32& OutFlags)
 	{
 		// by default we'll do exactly what was asked...
 
diff --git a/Engine/Source/Runtime/Core/Public/Containers/DynamicRHIResourceArray.h b/Engine/Source/Runtime/Core/Public/Containers/DynamicRHIResourceArray.h
--- a/Engine/Source/Runtime/Core/Public/Containers/DynamicRHIResourceArray.h
+++ b/Engine/Source/Runtime/Core/Public/Containers/DynamicRHIResourceArray.h
@@ -37,14 +37,14 @@ public:
 	/** 
 	* Constructor 
 	*/
-	TResourceArray(bool InNeedsCPUAccess=false)
+	[[nodiscard]] explicit TResourceArray(bool InNeedsCPUAccess=false)
 		:	Super()
 		,	bNeedsCPUAccess(InNeedsCPUAccess)
 	{
 	}
 
-	TResourceArray(TResourceArray&&) = default;
-	TResourceArray(const TResourceArray&) = default;
+	[[nodiscard]] TResourceArray(TResourceArray&&) = default;
+	[[nodiscard]] TResourceArray(const TResourceArray&) = default;
 	TResourceArray& operator=(TResourceArray&&) = default;
 	TResourceArray& operator=(const TResourceArray&) = default;
 
@@ -56,7 +56,7 @@ public:
 	/**
 	* @return A pointer to the resource data.
 	*/
-	virtual const void* GetResourceData() const override
+	[[nodiscard]] virtual const void* GetResourceData() const override
 	{ 
 		return &(*this)[0]; 
 	}
@@ -64,7 +64,7 @@ public:
 	/**
 	* @return size of resource data allocation
 	*/
-	virtual uint32 GetResourceDataSize() const override
+	[[nodiscard]] virtual uint32 GetResourceDataSize() const override
 	{
 		if (this->Num() > UINT32_MAX / sizeof(ElementType))
 		{
@@ -90,7 +90,7 @@ public:
 	/**
 	* @return true if the resource array is static and shouldn't be modified
 	*/
-	virtual bool IsStatic() const override
+	[[nodiscard]] virtual bool IsStatic() const override
 	{
 		return false;
 	}
@@ -98,7 +98,7 @@ public:
 	/**
 	* @return true if the resource keeps a copy of its resource data after the RHI resource has been created
 	*/
-	virtual bool GetAllowCPUAccess() const override
+	[[nodiscard]] virtual bool GetAllowCPUAccess() const override
 	{
 		return bNeedsCPUAccess;
 	}
