Added [[nodiscard]] to relevant container functions.

diff --git a/Engine/Source/Runtime/Core/Public/Containers/EnumAsByte.h b/Engine/Source/Runtime/Core/Public/Containers/EnumAsByte.h
--- a/Engine/Source/Runtime/Core/Public/Containers/EnumAsByte.h
+++ b/Engine/Source/Runtime/Core/Public/Containers/EnumAsByte.h
@@ -25,8 +25,8 @@ class TEnumAsByte
 public:
 	using EnumType = InEnumType;
 
-	TEnumAsByte() = default;
-	TEnumAsByte(const TEnumAsByte&) = default;
+	[[nodiscard]] TEnumAsByte() = default;
+	[[nodiscard]] TEnumAsByte(const TEnumAsByte&) = default;
 	TEnumAsByte& operator=(const TEnumAsByte&) = default;
 
 	/**
@@ -34,27 +34,30 @@ public:
 	 *
 	 * @param InValue value to construct with.
 	 */
-	FORCEINLINE TEnumAsByte( EnumType InValue )
+	[[nodiscard]] FORCEINLINE TEnumAsByte( EnumType InValue )
 		: Value(static_cast<uint8>(InValue))
-	{ }
+	{
+	}
 
 	/**
 	 * Constructor, initialize to the int32 value.
 	 *
 	 * @param InValue value to construct with.
 	 */
-	explicit FORCEINLINE TEnumAsByte( int32 InValue )
+	[[nodiscard]] explicit FORCEINLINE TEnumAsByte( int32 InValue )
 		: Value(static_cast<uint8>(InValue))
-	{ }
+	{
+	}
 
 	/**
 	 * Constructor, initialize to the int32 value.
 	 *
 	 * @param InValue value to construct with.
 	 */
-	explicit FORCEINLINE TEnumAsByte( uint8 InValue )
+	[[nodiscard]] explicit FORCEINLINE TEnumAsByte( uint8 InValue )
 		: Value(InValue)
-	{ }
+	{
+	}
 
 public:
 	/**
@@ -63,7 +66,7 @@ public:
 	 * @param InValue The value to compare with.
 	 * @return true if the two values are equal, false otherwise.
 	 */
-	bool operator==( EnumType InValue ) const
+	[[nodiscard]] bool operator==( EnumType InValue ) const
 	{
 		return static_cast<EnumType>(Value) == InValue;
 	}
@@ -74,13 +77,13 @@ public:
 	 * @param InValue The value to compare with.
 	 * @return true if the two values are equal, false otherwise.
 	 */
-	bool operator==(TEnumAsByte InValue) const
+	[[nodiscard]] bool operator==(TEnumAsByte InValue) const
 	{
 		return Value == InValue.Value;
 	}
 
 	/** Implicit conversion to EnumType. */
-	operator EnumType() const
+	[[nodiscard]] operator EnumType() const
 	{
 		return (EnumType)Value;
 	}
@@ -92,7 +95,7 @@ public:
 	 *
 	 * @return The enumeration value.
 	 */
-	EnumType GetValue() const
+	[[nodiscard]] EnumType GetValue() const
 	{
 		return (EnumType)Value;
 	}
@@ -102,7 +105,7 @@ public:
 	 *
 	 * @return The enumeration value.
 	 */
-	uint8 GetIntValue() const
+	[[nodiscard]] uint8 GetIntValue() const
 	{
 		return Value;
 	}
@@ -114,7 +117,7 @@ private:
 };
 
 template<class T>
-FORCEINLINE uint32 GetTypeHash(const TEnumAsByte<T>& Enum)
+[[nodiscard]] FORCEINLINE uint32 GetTypeHash(const TEnumAsByte<T>& Enum)
 {
 	return GetTypeHash((uint8)Enum.GetValue());
 }
diff --git a/Engine/Source/Runtime/Core/Public/Containers/IndirectArray.h b/Engine/Source/Runtime/Core/Public/Containers/IndirectArray.h
--- a/Engine/Source/Runtime/Core/Public/Containers/IndirectArray.h
+++ b/Engine/Source/Runtime/Core/Public/Containers/IndirectArray.h
@@ -23,14 +23,15 @@ public:
 	typedef TArray<void*, Allocator> InternalArrayType;
 
 	/** Default constructors. */
-	TIndirectArray() = default;
-	TIndirectArray(TIndirectArray&&) = default;
+	[[nodiscard]] TIndirectArray() = default;
+	[[nodiscard]] TIndirectArray(TIndirectArray&&) = default;
+
 	/**
 	 * Copy constructor.
 	 *
 	 * @param Other Other array to copy from.
 	 */
-	TIndirectArray(const TIndirectArray& Other)
+	[[nodiscard]] TIndirectArray(const TIndirectArray& Other)
 	{
 		for (auto& Item : Other)
 		{
@@ -80,7 +81,7 @@ public:
 	 * @returns True if the array is empty.
 	 * @see Num
 	 */
-	bool IsEmpty() const
+	[[nodiscard]] bool IsEmpty() const
 	{
 		return Array.IsEmpty();
 	}
@@ -90,7 +91,7 @@ public:
 	 *
 	 * @returns Number of elements in array.
 	 */
-	FORCEINLINE int32 Num() const
+	[[nodiscard]] FORCEINLINE int32 Num() const
 	{
 		return Array.Num();
 	}
@@ -100,7 +101,7 @@ public:
 	 *
 	 * @returns Pointer to first array entry or nullptr if this->ArrayMax == 0.
 	 */
-	FORCEINLINE T** GetData()
+	[[nodiscard]] FORCEINLINE T** GetData()
 	{
 		return (T**)Array.GetData();
 	}
@@ -110,7 +111,7 @@ public:
 	 *
 	 * @returns Pointer to first array entry or nullptr if this->ArrayMax == 0.
 	 */
-	FORCEINLINE const T** GetData() const
+	[[nodiscard]] FORCEINLINE const T** GetData() const
 	{
 		return (const T**)Array.GetData();
 	}
@@ -120,7 +121,7 @@ public:
 	 *
 	 * @returns Size in bytes of array type.
 	 */
-	static constexpr uint32 GetTypeSize()
+	[[nodiscard]] static constexpr uint32 GetTypeSize()
 	{
 		return sizeof(T*);
 	}
@@ -131,7 +132,7 @@ public:
 	 * @param Index Position of element to return.
 	 * @returns Reference to element in array at given position.
 	 */
-	FORCEINLINE T& operator[](int32 Index)
+	[[nodiscard]] FORCEINLINE T& operator[](int32 Index)
 	{
 		return *(T*)Array[Index];
 	}
@@ -144,7 +145,7 @@ public:
 	 * @param Index Position of element to return.
 	 * @returns Reference to element in array at given position.
 	 */
-	FORCEINLINE const T& operator[](int32 Index) const
+	[[nodiscard]] FORCEINLINE const T& operator[](int32 Index) const
 	{
 		return *(T*)Array[Index];
 	}
@@ -155,7 +156,7 @@ public:
 	 * @param IndexFromTheEnd (Optional) Index from the end of array (default = 0).
 	 * @returns Reference to n-th last element from the array.
 	 */
-	FORCEINLINE ElementType& Last(int32 IndexFromTheEnd = 0)
+	[[nodiscard]] FORCEINLINE ElementType& Last(int32 IndexFromTheEnd = 0)
 	{
 		return *(T*)Array.Last(IndexFromTheEnd);
 	}
@@ -168,7 +169,7 @@ public:
 	 * @param IndexFromTheEnd (Optional) Index from the end of array (default = 0).
 	 * @returns Reference to n-th last element from the array.
 	 */
-	FORCEINLINE const ElementType& Last(int32 IndexFromTheEnd = 0) const
+	[[nodiscard]] FORCEINLINE const ElementType& Last(int32 IndexFromTheEnd = 0) const
 	{
 		return *(T*)Array.Last(IndexFromTheEnd);
 	}
@@ -404,7 +405,7 @@ public:
 	 * @param Index Index to test.
 	 * @returns True if index is valid. False otherwise.
 	 */
-	FORCEINLINE bool IsValidIndex(int32 Index) const
+	[[nodiscard]] FORCEINLINE bool IsValidIndex(int32 Index) const
 	{
 		return Array.IsValidIndex(Index);
 	}
@@ -414,7 +415,7 @@ public:
 	 *
 	 * @returns Number of bytes allocated by this container.
 	 */
-	SIZE_T GetAllocatedSize() const
+	[[nodiscard]] SIZE_T GetAllocatedSize() const
 	{
 		return Array.Max() * sizeof(T*) + Array.Num() * sizeof(T);
 	}
@@ -428,7 +429,7 @@ public:
 	 *
 	 * @returns The iterator.
 	 */
-	TIterator CreateIterator()
+	[[nodiscard]] TIterator CreateIterator()
 	{
 		return TIterator(*this);
 	}
@@ -438,7 +439,7 @@ public:
 	 *
 	 * @returns The const iterator.
 	 */
-	TConstIterator CreateConstIterator() const
+	[[nodiscard]] TConstIterator CreateConstIterator() const
 	{
 		return TConstIterator(*this);
 	}
@@ -463,10 +464,10 @@ public:
 	 * DO NOT USE DIRECTLY
 	 * STL-like iterators to enable range-based for loop support.
 	 */
-	FORCEINLINE TDereferencingIterator<      ElementType, typename InternalArrayType::RangedForIteratorType     > begin()       { return TDereferencingIterator<      ElementType, typename InternalArrayType::RangedForIteratorType     >(Array.begin()); }
-	FORCEINLINE TDereferencingIterator<const ElementType, typename InternalArrayType::RangedForConstIteratorType> begin() const { return TDereferencingIterator<const ElementType, typename InternalArrayType::RangedForConstIteratorType>(Array.begin()); }
-	FORCEINLINE TDereferencingIterator<      ElementType, typename InternalArrayType::RangedForIteratorType     > end  ()       { return TDereferencingIterator<      ElementType, typename InternalArrayType::RangedForIteratorType     >(Array.end()); }
-	FORCEINLINE TDereferencingIterator<const ElementType, typename InternalArrayType::RangedForConstIteratorType> end  () const { return TDereferencingIterator<const ElementType, typename InternalArrayType::RangedForConstIteratorType>(Array.end()); }
+	[[nodiscard]] FORCEINLINE TDereferencingIterator<      ElementType, typename InternalArrayType::RangedForIteratorType     > begin()       { return TDereferencingIterator<      ElementType, typename InternalArrayType::RangedForIteratorType     >(Array.begin()); }
+	[[nodiscard]] FORCEINLINE TDereferencingIterator<const ElementType, typename InternalArrayType::RangedForConstIteratorType> begin() const { return TDereferencingIterator<const ElementType, typename InternalArrayType::RangedForConstIteratorType>(Array.begin()); }
+	[[nodiscard]] FORCEINLINE TDereferencingIterator<      ElementType, typename InternalArrayType::RangedForIteratorType     > end  ()       { return TDereferencingIterator<      ElementType, typename InternalArrayType::RangedForIteratorType     >(Array.end()); }
+	[[nodiscard]] FORCEINLINE TDereferencingIterator<const ElementType, typename InternalArrayType::RangedForConstIteratorType> end  () const { return TDereferencingIterator<const ElementType, typename InternalArrayType::RangedForConstIteratorType>(Array.end()); }
 
 private:
 
diff --git a/Engine/Source/Runtime/Core/Public/Containers/IntrusiveDoubleLinkedList.h b/Engine/Source/Runtime/Core/Public/Containers/IntrusiveDoubleLinkedList.h
--- a/Engine/Source/Runtime/Core/Public/Containers/IntrusiveDoubleLinkedList.h
+++ b/Engine/Source/Runtime/Core/Public/Containers/IntrusiveDoubleLinkedList.h
@@ -27,7 +27,7 @@ public:
 	using NodeType = TIntrusiveDoubleLinkedListNode<InElementType, ContainerType>;
 	using ElementType = InElementType;
 
-	TIntrusiveDoubleLinkedListNode()
+	[[nodiscard]] TIntrusiveDoubleLinkedListNode()
 		: Next(GetThisElement())
 		, Prev(GetThisElement())
 	{}
@@ -37,9 +37,18 @@ public:
 		Next = Prev = GetThisElement();
 	}
 
-	FORCEINLINE bool         IsInList() const { return Next != GetThisElement(); }
-	FORCEINLINE ElementType* GetNext() const  { return Next; }
-	FORCEINLINE ElementType* GetPrev() const  { return Prev; }
+	[[nodiscard]] FORCEINLINE bool IsInList() const
+	{
+		return Next != GetThisElement();
+	}
+	[[nodiscard]] FORCEINLINE ElementType* GetNext() const
+	{
+		return Next;
+	}
+	[[nodiscard]] FORCEINLINE ElementType* GetPrev() const
+	{
+		return Prev;
+	}
 
 	/**
 	 * Removes this element from the list in constant time.
@@ -79,8 +88,8 @@ protected:
 	friend class TIntrusiveDoubleLinkedListIterator<TIntrusiveDoubleLinkedListNode>;
 	friend class TIntrusiveDoubleLinkedList<ElementType, ContainerType>;
 
-	FORCEINLINE ElementType*       GetThisElement()       { return static_cast<ElementType*>(this); }
-	FORCEINLINE const ElementType* GetThisElement() const { return static_cast<const ElementType*>(this); }
+	[[nodiscard]] FORCEINLINE ElementType*       GetThisElement()       { return static_cast<ElementType*>(this); }
+	[[nodiscard]] FORCEINLINE const ElementType* GetThisElement() const { return static_cast<const ElementType*>(this); }
 
 	ElementType* Next;
 	ElementType* Prev;
@@ -96,9 +105,10 @@ class TIntrusiveDoubleLinkedListIterator
 public:
 	using ElementType = typename NodeType::ElementType;
 
-	FORCEINLINE explicit TIntrusiveDoubleLinkedListIterator(ElementType* Node)
+	[[nodiscard]] FORCEINLINE explicit TIntrusiveDoubleLinkedListIterator(ElementType* Node)
 		: CurrentNode(Node)
-	{}
+	{
+	}
 
 	FORCEINLINE TIntrusiveDoubleLinkedListIterator& operator++()
 	{
@@ -129,26 +139,26 @@ public:
 	}
 
 	// Accessors.
-	FORCEINLINE ElementType& operator->() const
+	[[nodiscard]] FORCEINLINE ElementType& operator->() const
 	{
 		checkSlow(CurrentNode);
 		return *CurrentNode;
 	}
 
-	FORCEINLINE ElementType& operator*() const
+	[[nodiscard]] FORCEINLINE ElementType& operator*() const
 	{
 		checkSlow(CurrentNode);
 		return *CurrentNode;
 	}
 
-	FORCEINLINE ElementType* GetNode() const
+	[[nodiscard]] FORCEINLINE ElementType* GetNode() const
 	{
 		checkSlow(CurrentNode);
 		return CurrentNode;
 	}
 
-	FORCEINLINE bool operator==(const TIntrusiveDoubleLinkedListIterator& Other) const { return CurrentNode == Other.CurrentNode; }
-	FORCEINLINE bool operator!=(const TIntrusiveDoubleLinkedListIterator& Other) const { return CurrentNode != Other.CurrentNode; }
+	[[nodiscard]] FORCEINLINE bool operator==(const TIntrusiveDoubleLinkedListIterator& Other) const { return CurrentNode == Other.CurrentNode; }
+	[[nodiscard]] FORCEINLINE bool operator!=(const TIntrusiveDoubleLinkedListIterator& Other) const { return CurrentNode != Other.CurrentNode; }
 
 private:
 	ElementType* CurrentNode;
@@ -168,7 +178,10 @@ public:
 	using ElementType = InElementType;
 	using NodeType = TIntrusiveDoubleLinkedListNode<ElementType, ContainerType>;
 
-	FORCEINLINE TIntrusiveDoubleLinkedList() {}
+	[[nodiscard]] FORCEINLINE TIntrusiveDoubleLinkedList()
+	{
+	}
+
 	FORCEINLINE TIntrusiveDoubleLinkedList(const TIntrusiveDoubleLinkedList&) = delete;
 	FORCEINLINE TIntrusiveDoubleLinkedList& operator=(const TIntrusiveDoubleLinkedList&) = delete;
 
@@ -177,12 +190,21 @@ public:
 	 * 
 	 * @see IsEmpty, IsFilled
 	 */
-	FORCEINLINE void Reset() { Sentinel.Reset(); }
+	FORCEINLINE void Reset()
+	{
+		Sentinel.Reset();
+	}
 
 	// Accessors.
 
-	FORCEINLINE bool IsEmpty() const  { return Sentinel.Next == GetSentinel(); }
-	FORCEINLINE bool IsFilled() const { return Sentinel.Next != GetSentinel(); }
+	[[nodiscard]] FORCEINLINE bool IsEmpty() const
+	{
+		return Sentinel.Next == GetSentinel();
+	}
+	[[nodiscard]] FORCEINLINE bool IsFilled() const
+	{
+		return Sentinel.Next != GetSentinel();
+	}
 
 	// Adding/Removing methods
 
@@ -220,17 +242,17 @@ public:
 		}
 	}
 
-	FORCEINLINE ElementType* GetHead()
+	[[nodiscard]] FORCEINLINE ElementType* GetHead()
 	{
 		return IsFilled() ? Sentinel.Next : nullptr;
 	}
 
-	FORCEINLINE ElementType* GetTail()
+	[[nodiscard]] FORCEINLINE ElementType* GetTail()
 	{
 		return IsFilled() ? Sentinel.Prev : nullptr;
 	}
 
-	FORCEINLINE ElementType* PopHead()
+	[[nodiscard]] FORCEINLINE ElementType* PopHead()
 	{
 		if (IsEmpty())
 		{
@@ -242,7 +264,7 @@ public:
 		return Head;
 	}
 
-	FORCEINLINE ElementType* PopTail()
+	[[nodiscard]] FORCEINLINE ElementType* PopTail()
 	{
 		if (IsEmpty())
 		{
@@ -272,15 +294,15 @@ public:
 	using TIterator = TIntrusiveDoubleLinkedListIterator<NodeType>;
 	using TConstIterator = TIntrusiveDoubleLinkedListIterator<const NodeType>;
 
-	FORCEINLINE TIterator      begin()       { return TIterator(Sentinel.Next); }
-	FORCEINLINE TConstIterator begin() const { return TConstIterator(Sentinel.Next); }
-	FORCEINLINE TIterator      end()         { return TIterator(GetSentinel()); }
-	FORCEINLINE TConstIterator end() const   { return TConstIterator(GetSentinel()); }
+	[[nodiscard]] FORCEINLINE TIterator      begin()       { return TIterator(Sentinel.Next); }
+	[[nodiscard]] FORCEINLINE TConstIterator begin() const { return TConstIterator(Sentinel.Next); }
+	[[nodiscard]] FORCEINLINE TIterator      end()         { return TIterator(GetSentinel()); }
+	[[nodiscard]] FORCEINLINE TConstIterator end() const   { return TConstIterator(GetSentinel()); }
 
 private:
 
-	FORCEINLINE ElementType*       GetSentinel()       { return static_cast<ElementType*>(&Sentinel); }       //-V717
-	FORCEINLINE const ElementType* GetSentinel() const { return static_cast<const ElementType*>(&Sentinel); } //-V717
+	[[nodiscard]] FORCEINLINE ElementType*       GetSentinel()       { return static_cast<ElementType*>(&Sentinel); }       //-V717
+	[[nodiscard]] FORCEINLINE const ElementType* GetSentinel() const { return static_cast<const ElementType*>(&Sentinel); } //-V717
 
 	NodeType Sentinel;
 
diff --git a/Engine/Source/Runtime/Core/Public/Containers/IteratorAdapter.h b/Engine/Source/Runtime/Core/Public/Containers/IteratorAdapter.h
--- a/Engine/Source/Runtime/Core/Public/Containers/IteratorAdapter.h
+++ b/Engine/Source/Runtime/Core/Public/Containers/IteratorAdapter.h
@@ -34,22 +34,22 @@ public:
 	using ElementType = typename BaseType::ElementType;
 	using SizeType = typename BaseType::SizeType;
 
-	TIteratorAdapter() = default;
+	[[nodiscard]] TIteratorAdapter() = default;
 
 	/**
 	 * Perfect forwarding constructor to the iterator class.
 	 */
 	template < typename... Args>
-	TIteratorAdapter(EInPlace, Args&&... InArgs) : BaseType(Forward<Args>(InArgs)...)
+	[[nodiscard]] TIteratorAdapter(EInPlace, Args&&... InArgs) : BaseType(Forward<Args>(InArgs)...)
 	{
 	}
 
-	ElementType& operator*() const
+	[[nodiscard]] ElementType& operator*() const
 	{
 		return this->Dereference();
 	}
 
-	ElementType* operator->() const
+	[[nodiscard]] ElementType* operator->() const
 	{
 		return &this->Dereference();
 	}
@@ -114,12 +114,12 @@ protected:
 	friend class TIteratorAdapter;
 
 public:
-	bool operator==(const TIteratorAdapter& Right) const
+	[[nodiscard]] bool operator==(const TIteratorAdapter& Right) const
 	{
 		return Equals(Right);
 	}
 
-	bool operator!=(const TIteratorAdapter& Right) const
+	[[nodiscard]] bool operator!=(const TIteratorAdapter& Right) const
 	{
 		return !Equals(Right);
 	}
diff --git a/Engine/Source/Runtime/Core/Public/Containers/LazyPrintf.h b/Engine/Source/Runtime/Core/Public/Containers/LazyPrintf.h
--- a/Engine/Source/Runtime/Core/Public/Containers/LazyPrintf.h
+++ b/Engine/Source/Runtime/Core/Public/Containers/LazyPrintf.h
@@ -18,7 +18,7 @@ public:
 		CurrentState.Empty(50 * 1024);
 	}
 
-	FString GetResultString()
+	[[nodiscard]] FString GetResultString()
 	{
 		// internal error more %s than %s in MaterialTemplate.usf
 		check(!ProcessUntilPercentS());
diff --git a/Engine/Source/Runtime/Core/Public/Containers/LinkedListBuilder.h b/Engine/Source/Runtime/Core/Public/Containers/LinkedListBuilder.h
--- a/Engine/Source/Runtime/Core/Public/Containers/LinkedListBuilder.h
+++ b/Engine/Source/Runtime/Core/Public/Containers/LinkedListBuilder.h
@@ -10,7 +10,7 @@ struct TLinkedListBuilderNextLink
 {
 	using ElementType = InElementType;
 
-	FORCEINLINE static ElementType** GetNextPtr(ElementType& Element)
+	[[nodiscard]] FORCEINLINE static ElementType** GetNextPtr(ElementType& Element)
 	{
 		return &(Element.Next);
 	}
@@ -21,7 +21,7 @@ struct TLinkedListBuilderNextLinkMemberVar
 {
 	using ElementType = InElementType;
 
-	FORCEINLINE static ElementType** GetNextPtr(ElementType& Element)
+	[[nodiscard]] FORCEINLINE static ElementType** GetNextPtr(ElementType& Element)
 	{
 		return &(Element.*NextLink);
 	}
@@ -51,7 +51,7 @@ private:
 	}
 public:
 
-	explicit TLinkedListBuilder(ElementType** ListStartPtr) :
+	[[nodiscard]] explicit TLinkedListBuilder(ElementType** ListStartPtr) :
 		StartPtr(ListStartPtr),
 		EndPtr(ListStartPtr)
 	{
@@ -127,23 +127,23 @@ public:
 		return false;
 	}
 
-	FORCEINLINE ElementType* GetNext(ElementType& Element) const
+	[[nodiscard]] FORCEINLINE ElementType* GetNext(ElementType& Element) const
 	{
 		return GetNextRef(Element);
 	}
 
-	FORCEINLINE ElementType* GetListStart() const
+	[[nodiscard]] FORCEINLINE ElementType* GetListStart() const
 	{
 		return *StartPtr;
 	}
 
-	FORCEINLINE ElementType* GetListEnd() const
+	[[nodiscard]] FORCEINLINE ElementType* GetListEnd() const
 	{
 		return *EndPtr;
 	}
 
 private:
-	FORCEINLINE ElementType*& GetNextRef(ElementType& Element) const
+	[[nodiscard]] FORCEINLINE ElementType*& GetNextRef(ElementType& Element) const
 	{
 		return *LinkAccessor::GetNextPtr(Element);
 	}
diff --git a/Engine/Source/Runtime/Core/Public/Containers/List.h b/Engine/Source/Runtime/Core/Public/Containers/List.h
--- a/Engine/Source/Runtime/Core/Public/Containers/List.h
+++ b/Engine/Source/Runtime/Core/Public/Containers/List.h
@@ -9,9 +9,10 @@ template <class ContainerType>
 class TLinkedListIteratorBase
 {
 public:
-	explicit TLinkedListIteratorBase(ContainerType* FirstLink)
+	[[nodiscard]] explicit TLinkedListIteratorBase(ContainerType* FirstLink)
 		: CurrentLink(FirstLink)
-	{ }
+	{
+	}
 
 	/**
 	 * Advances the iterator to the next element.
@@ -36,13 +37,13 @@ public:
 	}
 
 	/** conversion to "bool" returning true if the iterator is valid. */
-	FORCEINLINE explicit operator bool() const
+	[[nodiscard]] FORCEINLINE explicit operator bool() const
 	{ 
 		return CurrentLink != nullptr;
 	}
 
-	FORCEINLINE bool operator==(const TLinkedListIteratorBase& Rhs) const { return CurrentLink == Rhs.CurrentLink; }
-	FORCEINLINE bool operator!=(const TLinkedListIteratorBase& Rhs) const { return CurrentLink != Rhs.CurrentLink; }
+	[[nodiscard]] FORCEINLINE bool operator==(const TLinkedListIteratorBase& Rhs) const { return CurrentLink == Rhs.CurrentLink; }
+	[[nodiscard]] FORCEINLINE bool operator!=(const TLinkedListIteratorBase& Rhs) const { return CurrentLink != Rhs.CurrentLink; }
 protected:
 
 	ContainerType* CurrentLink;
@@ -54,7 +55,7 @@ class TLinkedListIterator : public TLinkedListIteratorBase<ContainerType>
 	typedef TLinkedListIteratorBase<ContainerType> Super;
 
 public:
-	explicit TLinkedListIterator(ContainerType* FirstLink)
+	[[nodiscard]] explicit TLinkedListIterator(ContainerType* FirstLink)
 		: Super(FirstLink)
 	{
 	}
@@ -79,7 +80,7 @@ class TIntrusiveLinkedListIterator : public TLinkedListIteratorBase<ElementType>
 	typedef TLinkedListIteratorBase<ElementType> Super;
 
 public:
-	explicit TIntrusiveLinkedListIterator(ElementType* FirstLink)
+	[[nodiscard]] explicit TIntrusiveLinkedListIterator(ElementType* FirstLink)
 		: Super(FirstLink)
 	{
 	}
@@ -116,7 +117,7 @@ public:
 	/**
 	 * Default constructor (empty list)
 	 */
-	TLinkedListBase()
+	[[nodiscard]] TLinkedListBase()
 		: NextLink(NULL)
 		, PrevLink(NULL)
 	{
@@ -239,22 +240,22 @@ public:
 	 *
 	 * @return true if currently linked, false otherwise
 	 */
-	FORCEINLINE bool IsLinked()
+	[[nodiscard]] FORCEINLINE bool IsLinked()
 	{
 		return PrevLink != nullptr;
 	}
 
-	FORCEINLINE ContainerType** GetPrevLink() const
+	[[nodiscard]] FORCEINLINE ContainerType** GetPrevLink() const
 	{
 		return PrevLink;
 	}
 
-	FORCEINLINE ContainerType* GetNextLink() const
+	[[nodiscard]] FORCEINLINE ContainerType* GetNextLink() const
 	{
 		return NextLink;
 	}
 
-	FORCEINLINE ContainerType* Next()
+	[[nodiscard]] FORCEINLINE ContainerType* Next()
 	{
 		return NextLink;
 	}
@@ -267,10 +268,10 @@ private:
 	ContainerType** PrevLink;
 
 
-	FORCEINLINE friend TIterator      begin(      ContainerType& List) { return TIterator     (&List); }
-	FORCEINLINE friend TConstIterator begin(const ContainerType& List) { return TConstIterator(const_cast<ContainerType*>(&List)); }
-	FORCEINLINE friend TIterator      end  (      ContainerType& List) { return TIterator     (nullptr); }
-	FORCEINLINE friend TConstIterator end  (const ContainerType& List) { return TConstIterator(nullptr); }
+	[[nodiscard]] FORCEINLINE friend TIterator      begin(      ContainerType& List) { return TIterator     (&List); }
+	[[nodiscard]] FORCEINLINE friend TConstIterator begin(const ContainerType& List) { return TConstIterator(const_cast<ContainerType*>(&List)); }
+	[[nodiscard]] FORCEINLINE friend TIterator      end  (      ContainerType& List) { return TIterator     (nullptr); }
+	[[nodiscard]] FORCEINLINE friend TConstIterator end  (const ContainerType& List) { return TConstIterator(nullptr); }
 };
 
 /**
@@ -285,7 +286,7 @@ class TLinkedList : public TLinkedListBase<TLinkedList<ElementType>, ElementType
 
 public:
 	/** Default constructor (empty list). */
-	TLinkedList()
+	[[nodiscard]] TLinkedList()
 		: Super()
 	{
 	}
@@ -295,7 +296,7 @@ public:
 	 *
 	 * @param InElement The element to add to the list.
 	 */
-	explicit TLinkedList(const ElementType& InElement)
+	[[nodiscard]] explicit TLinkedList(const ElementType& InElement)
 		: Super()
 		, Element(InElement)
 	{
@@ -303,19 +304,19 @@ public:
 
 
 	// Accessors.
-	FORCEINLINE ElementType* operator->()
+	[[nodiscard]] FORCEINLINE ElementType* operator->()
 	{
 		return &Element;
 	}
-	FORCEINLINE const ElementType* operator->() const
+	[[nodiscard]] FORCEINLINE const ElementType* operator->() const
 	{
 		return &Element;
 	}
-	FORCEINLINE ElementType& operator*()
+	[[nodiscard]] FORCEINLINE ElementType& operator*()
 	{
 		return Element;
 	}
-	FORCEINLINE const ElementType& operator*() const
+	[[nodiscard]] FORCEINLINE const ElementType& operator*() const
 	{
 		return Element;
 	}
@@ -349,7 +350,7 @@ class TIntrusiveLinkedList : public TLinkedListBase<ElementType, ElementType, TI
 
 public:
 	/** Default constructor (empty list). */
-	TIntrusiveLinkedList()
+	[[nodiscard]] TIntrusiveLinkedList()
 		: Super()
 	{
 	}
@@ -361,12 +362,13 @@ class TDoubleLinkedListIterator
 {
 public:
 
-	explicit TDoubleLinkedListIterator(NodeType* StartingNode)
+	[[nodiscard]] explicit TDoubleLinkedListIterator(NodeType* StartingNode)
 		: CurrentNode(StartingNode)
-	{ }
+	{
+	}
 
 	/** conversion to "bool" returning true if the iterator is valid. */
-	FORCEINLINE explicit operator bool() const
+	[[nodiscard]] FORCEINLINE explicit operator bool() const
 	{ 
 		return CurrentNode != nullptr; 
 	}
@@ -400,26 +402,26 @@ public:
 	}
 
 	// Accessors.
-	ElementType& operator->() const
+	[[nodiscard]] ElementType& operator->() const
 	{
 		checkSlow(CurrentNode);
 		return CurrentNode->GetValue();
 	}
 
-	ElementType& operator*() const
+	[[nodiscard]] ElementType& operator*() const
 	{
 		checkSlow(CurrentNode);
 		return CurrentNode->GetValue();
 	}
 
-	NodeType* GetNode() const
+	[[nodiscard]] NodeType* GetNode() const
 	{
 		checkSlow(CurrentNode);
 		return CurrentNode;
 	}
 
-	bool operator==(const TDoubleLinkedListIterator& Rhs) const { return CurrentNode == Rhs.CurrentNode; }
-	bool operator!=(const TDoubleLinkedListIterator& Rhs) const { return CurrentNode != Rhs.CurrentNode; }
+	[[nodiscard]] bool operator==(const TDoubleLinkedListIterator& Rhs) const { return CurrentNode == Rhs.CurrentNode; }
+	[[nodiscard]] bool operator!=(const TDoubleLinkedListIterator& Rhs) const { return CurrentNode != Rhs.CurrentNode; }
 
 private:
 	NodeType* CurrentNode;
@@ -441,36 +443,37 @@ public:
 		friend class TDoubleLinkedList;
 
 		/** Constructor */
-		TDoubleLinkedListNode( const ElementType& InValue )
+		[[nodiscard]] TDoubleLinkedListNode( const ElementType& InValue )
 			: Value(InValue), NextNode(nullptr), PrevNode(nullptr)
-		{ }
+		{
+		}
 
-		const ElementType& GetValue() const
+		[[nodiscard]] const ElementType& GetValue() const
 		{
 			return Value;
 		}
 
-		ElementType& GetValue()
+		[[nodiscard]] ElementType& GetValue()
 		{
 			return Value;
 		}
 
-		TDoubleLinkedListNode* GetNextNode()
+		[[nodiscard]] TDoubleLinkedListNode* GetNextNode()
 		{
 			return NextNode;
 		}
 
-		const TDoubleLinkedListNode* GetNextNode() const
+		[[nodiscard]] const TDoubleLinkedListNode* GetNextNode() const
 		{
 			return NextNode;
 		}
 
-		TDoubleLinkedListNode* GetPrevNode()
+		[[nodiscard]] TDoubleLinkedListNode* GetPrevNode()
 		{
 			return PrevNode;
 		}
 
-		const TDoubleLinkedListNode* GetPrevNode() const
+		[[nodiscard]] const TDoubleLinkedListNode* GetPrevNode() const
 		{
 			return PrevNode;
 		}
@@ -488,11 +491,12 @@ public:
 	typedef TDoubleLinkedListIterator<TDoubleLinkedListNode, const ElementType> TConstIterator;
 
 	/** Constructors. */
-	TDoubleLinkedList()
+	[[nodiscard]] TDoubleLinkedList()
 		: HeadNode(nullptr)
 		, TailNode(nullptr)
 		, ListSize(0)
-	{ }
+	{
+	}
 
 	/** Destructor */
 	virtual ~TDoubleLinkedList()
@@ -698,7 +702,7 @@ public:
 	 * @return Pointer to the first node.
 	 * @see GetTail
 	 */
-	TDoubleLinkedListNode* GetHead() const
+	[[nodiscard]] TDoubleLinkedListNode* GetHead() const
 	{
 		return HeadNode;
 	}
@@ -709,7 +713,7 @@ public:
 	 * @return Pointer to the last node.
 	 * @see GetHead
 	 */
-	TDoubleLinkedListNode* GetTail() const
+	[[nodiscard]] TDoubleLinkedListNode* GetTail() const
 	{
 		return TailNode;
 	}
@@ -720,7 +724,7 @@ public:
 	 * @param	InElement	the value to find
 	 * @return	a pointer to the node that contains the value specified, or nullptr of the value couldn't be found
 	 */
-	TDoubleLinkedListNode* FindNode( const ElementType& InElement )
+	[[nodiscard]] TDoubleLinkedListNode* FindNode( const ElementType& InElement )
 	{
 		TDoubleLinkedListNode* Node = HeadNode;
 		while ( Node != nullptr )
@@ -736,7 +740,7 @@ public:
 		return Node;
 	}
 
-	bool Contains( const ElementType& InElement )
+	[[nodiscard]] bool Contains( const ElementType& InElement )
 	{
 		return (FindNode(InElement) != nullptr);
 	}
@@ -747,7 +751,7 @@ public:
 	 * @returns True if the list is empty.
 	 * @see Num
 	 */
-	bool IsEmpty() const
+	[[nodiscard]] bool IsEmpty() const
 	{
 		return ListSize == 0;
 	}
@@ -757,7 +761,7 @@ public:
 	 *
 	 * @return Item count.
 	 */
-	int32 Num() const
+	[[nodiscard]] int32 Num() const
 	{
 		return ListSize;
 	}
@@ -783,10 +787,10 @@ private:
 	TDoubleLinkedList(const TDoubleLinkedList&);
 	TDoubleLinkedList& operator=(const TDoubleLinkedList&);
 
-	friend TIterator      begin(      TDoubleLinkedList& List) { return TIterator     (List.GetHead()); }
-	friend TConstIterator begin(const TDoubleLinkedList& List) { return TConstIterator(List.GetHead()); }
-	friend TIterator      end  (      TDoubleLinkedList& List) { return TIterator     (nullptr); }
-	friend TConstIterator end  (const TDoubleLinkedList& List) { return TConstIterator(nullptr); }
+	[[nodiscard]] friend TIterator      begin(      TDoubleLinkedList& List) { return TIterator     (List.GetHead()); }
+	[[nodiscard]] friend TConstIterator begin(const TDoubleLinkedList& List) { return TConstIterator(List.GetHead()); }
+	[[nodiscard]] friend TIterator      end  (      TDoubleLinkedList& List) { return TIterator     (nullptr); }
+	[[nodiscard]] friend TConstIterator end  (const TDoubleLinkedList& List) { return TConstIterator(nullptr); }
 };
 
 
@@ -806,7 +810,7 @@ public:
 
 	// Constructor.
 
-	TList(const ElementType &InElement, TList<ElementType>* InNext = nullptr)
+	[[nodiscard]] TList(const ElementType &InElement, TList<ElementType>* InNext = nullptr)
 	{
 		Element = InElement;
 		Next = InNext;
diff --git a/Engine/Source/Runtime/Core/Public/Containers/LockFreeList.h b/Engine/Source/Runtime/Core/Public/Containers/LockFreeList.h
--- a/Engine/Source/Runtime/Core/Public/Containers/LockFreeList.h
+++ b/Engine/Source/Runtime/Core/Public/Containers/LockFreeList.h
@@ -64,7 +64,7 @@ class TLockFreeAllocOnceIndexedAllocator
 	};
 public:
 
-	TLockFreeAllocOnceIndexedAllocator()
+	[[nodiscard]] TLockFreeAllocOnceIndexedAllocator()
 	{
 		NextIndex.Increment(); // skip the null ptr
 		for (uint32 Index = 0; Index < MaxBlocks; Index++)
@@ -86,7 +86,7 @@ public:
 		}
 		return FirstItem;
 	}
-	FORCEINLINE T* GetItem(uint32 Index)
+	[[nodiscard]] FORCEINLINE T* GetItem(uint32 Index)
 	{
 		if (!Index)
 		{
@@ -99,7 +99,7 @@ public:
 	}
 
 private:
-	void* GetRawItem(uint32 Index)
+	[[nodiscard]] void* GetRawItem(uint32 Index)
 	{
 		uint32 BlockIndex = Index / ItemsPerPage;
 		uint32 SubIndex = Index % ItemsPerPage;
@@ -149,7 +149,7 @@ struct FIndexedPointer
 		Ptrs.store(uint64(Ptr) | (CounterAndState << MAX_LOCK_FREE_LINKS_AS_BITS), std::memory_order_relaxed);
 	}
 
-	FORCEINLINE uint32 GetPtr() const
+	[[nodiscard]] FORCEINLINE uint32 GetPtr() const
 	{
 		return uint32(Ptrs.load(std::memory_order_relaxed) & (MAX_LOCK_FREE_LINKS - 1));
 	}
@@ -159,7 +159,7 @@ struct FIndexedPointer
 		SetAll(To, GetCounterAndState());
 	}
 
-	FORCEINLINE uint64 GetCounterAndState() const
+	[[nodiscard]] FORCEINLINE uint64 GetCounterAndState() const
 	{
 		return (Ptrs.load(std::memory_order_relaxed) >> MAX_LOCK_FREE_LINKS_AS_BITS);
 	}
@@ -180,7 +180,7 @@ struct FIndexedPointer
 	}
 
 	template<uint64 TABAInc>
-	FORCEINLINE uint64 GetState() const
+	[[nodiscard]] FORCEINLINE uint64 GetState() const
 	{
 		return GetCounterAndState() & (TABAInc - 1);
 	}
@@ -206,11 +206,11 @@ struct FIndexedPointer
 		return Ptrs.compare_exchange_strong(Expected, Exchange.Ptrs.load(std::memory_order_relaxed), std::memory_order_acq_rel, std::memory_order_relaxed);
 	}
 
-	FORCEINLINE bool operator==(const FIndexedPointer& Other) const
+	[[nodiscard]] FORCEINLINE bool operator==(const FIndexedPointer& Other) const
 	{
 		return Ptrs.load(std::memory_order_relaxed) == Other.Ptrs.load(std::memory_order_relaxed);
 	}
-	FORCEINLINE bool operator!=(const FIndexedPointer& Other) const
+	[[nodiscard]] FORCEINLINE bool operator!=(const FIndexedPointer& Other) const
 	{
 		return Ptrs.load(std::memory_order_relaxed) != Other.Ptrs.load(std::memory_order_relaxed);
 	}
@@ -239,15 +239,15 @@ struct FLockFreeLinkPolicy
 	typedef uint32 TLinkPtr;
 	typedef TLockFreeAllocOnceIndexedAllocator<FIndexedLockFreeLink, MAX_LOCK_FREE_LINKS, 16384> TAllocator;
 
-	static FORCEINLINE FIndexedLockFreeLink* DerefLink(uint32 Ptr)
+	[[nodiscard]] static FORCEINLINE FIndexedLockFreeLink* DerefLink(uint32 Ptr)
 	{
 		return LinkAllocator.GetItem(Ptr);
 	}
-	static FORCEINLINE FIndexedLockFreeLink* IndexToLink(uint32 Index)
+	[[nodiscard]] static FORCEINLINE FIndexedLockFreeLink* IndexToLink(uint32 Index)
 	{
 		return LinkAllocator.GetItem(Index);
 	}
-	static FORCEINLINE uint32 IndexToPtr(uint32 Index)
+	[[nodiscard]] static FORCEINLINE uint32 IndexToPtr(uint32 Index)
 	{
 		return Index;
 	}
@@ -267,7 +267,7 @@ class FLockFreePointerListLIFORoot
 	typedef FLockFreeLinkPolicy::TLinkPtr TLinkPtr;
 
 public:
-	FORCEINLINE FLockFreePointerListLIFORoot()
+	[[nodiscard]] FORCEINLINE FLockFreePointerListLIFORoot()
 	{
 		// We want to make sure we have quite a lot of extra counter values to avoid the ABA problem. This could probably be relaxed, but eventually it will be dangerous. 
 		// The question is "how many queue operations can a thread starve for".
@@ -394,12 +394,12 @@ public:
 		return Item;
 	}
 
-	FORCEINLINE bool IsEmpty() const
+	[[nodiscard]] FORCEINLINE bool IsEmpty() const
 	{
 		return !Head.GetPtr();
 	}
 
-	FORCEINLINE uint64 GetState() const
+	[[nodiscard]] FORCEINLINE uint64 GetState() const
 	{
 		TDoublePtr LocalHead;
 		LocalHead.AtomicRead(Head);
@@ -420,7 +420,7 @@ class FLockFreePointerListLIFOBase
 	typedef FLockFreeLinkPolicy::TLinkPtr TLinkPtr;
 
 public:
-	FLockFreePointerListLIFOBase() = default;
+	[[nodiscard]] FLockFreePointerListLIFOBase() = default;
 
 	~FLockFreePointerListLIFOBase()
 	{
@@ -470,7 +470,7 @@ public:
 	}
 
 
-	T* Pop()
+	[[nodiscard]] T* Pop()
 	{
 		TLinkPtr Item = RootList.Pop();
 		T* Result = nullptr;
@@ -524,12 +524,12 @@ public:
 		}
 	}
 
-	FORCEINLINE bool IsEmpty() const
+	[[nodiscard]] FORCEINLINE bool IsEmpty() const
 	{
 		return RootList.IsEmpty();
 	}
 
-	FORCEINLINE uint64 GetState() const
+	[[nodiscard]] FORCEINLINE uint64 GetState() const
 	{
 		return RootList.GetState();
 	}
@@ -549,7 +549,7 @@ class FLockFreePointerFIFOBase
 	typedef FLockFreeLinkPolicy::TLinkPtr TLinkPtr;
 public:
 
-	FORCEINLINE FLockFreePointerFIFOBase()
+	[[nodiscard]] FORCEINLINE FLockFreePointerFIFOBase()
 	{
 		// We want to make sure we have quite a lot of extra counter values to avoid the ABA problem. This could probably be relaxed, but eventually it will be dangerous. 
 		// The question is "how many queue operations can a thread starve for".
@@ -614,7 +614,7 @@ public:
 		}
 	}
 
-	T* Pop()
+	[[nodiscard]] T* Pop()
 	{
 		T* Result = nullptr;
 		TDoublePtr LocalHead;
@@ -669,7 +669,7 @@ public:
 	}
 
 
-	FORCEINLINE bool IsEmpty() const
+	[[nodiscard]] FORCEINLINE bool IsEmpty() const
 	{
 		TDoublePtr LocalHead;
 		LocalHead.AtomicRead(Head);
@@ -693,7 +693,7 @@ class FStallingTaskQueue
 	typedef FLockFreeLinkPolicy::TLink TLink;
 	typedef FLockFreeLinkPolicy::TLinkPtr TLinkPtr;
 public:
-	FStallingTaskQueue()
+	[[nodiscard]] FStallingTaskQueue()
 	{
 		MasterState.Init();
 	}
@@ -750,7 +750,7 @@ public:
 		return ThreadToWake;
 	}
 
-	T* Pop(int32 MyThread, bool bAllowStall)
+	[[nodiscard]] T* Pop(int32 MyThread, bool bAllowStall)
 	{
 		check(MyThread >= 0 && MyThread < FLockFreeLinkPolicy::MAX_BITS_IN_TLinkPtr);
 
@@ -797,7 +797,7 @@ public:
 
 private:
 
-	static int32 FindThreadToWake(TLinkPtr Ptr)
+	[[nodiscard]] static int32 FindThreadToWake(TLinkPtr Ptr)
 	{
 		int32 Result = -1;
 		UPTRINT Test = UPTRINT(Ptr);
@@ -813,17 +813,17 @@ private:
 		return Result;
 	}
 
-	static TLinkPtr TurnOffBit(TLinkPtr Ptr, int32 BitToTurnOff)
+	[[nodiscard]] static TLinkPtr TurnOffBit(TLinkPtr Ptr, int32 BitToTurnOff)
 	{
 		return (TLinkPtr)(UPTRINT(Ptr) & ~(UPTRINT(1) << BitToTurnOff));
 	}
 
-	static TLinkPtr TurnOnBit(TLinkPtr Ptr, int32 BitToTurnOn)
+	[[nodiscard]] static TLinkPtr TurnOnBit(TLinkPtr Ptr, int32 BitToTurnOn)
 	{
 		return (TLinkPtr)(UPTRINT(Ptr) | (UPTRINT(1) << BitToTurnOn));
 	}
 
-	static bool TestBit(TLinkPtr Ptr, int32 BitToTest)
+	[[nodiscard]] static bool TestBit(TLinkPtr Ptr, int32 BitToTest)
 	{
 		return !!(UPTRINT(Ptr) & (UPTRINT(1) << BitToTest));
 	}
@@ -855,7 +855,7 @@ public:
 	*	Pop an item from the list or return NULL if the list is empty.
 	*	@return The popped item, if any.
 	*/
-	T* Pop()
+	[[nodiscard]] T* Pop()
 	{
 		return FLockFreePointerListLIFOBase<T, TPaddingForCacheContention>::Pop();
 	}
@@ -887,7 +887,7 @@ public:
 	*	CAUTION: This methods safety depends on external assumptions. For example, if another thread could add to the list at any time, the return value is no better than a best guess.
 	*	As typically used, the list is not being access concurrently when this is called.
 	*/
-	FORCEINLINE bool IsEmpty() const
+	[[nodiscard]] FORCEINLINE bool IsEmpty() const
 	{
 		return FLockFreePointerListLIFOBase<T, TPaddingForCacheContention>::IsEmpty();
 	}
@@ -924,7 +924,7 @@ public:
 	*	Pop an item from the list or return NULL if the list is empty.
 	*	@return The popped item, if any.
 	*/
-	T* Pop()
+	[[nodiscard]] T* Pop()
 	{
 		return FLockFreePointerFIFOBase<T, TPaddingForCacheContention>::Pop();
 	}
@@ -947,7 +947,7 @@ public:
 	*	CAUTION: This methods safety depends on external assumptions. For example, if another thread could add to the list at any time, the return value is no better than a best guess.
 	*	As typically used, the list is not being access concurrently when this is called.
 	*/
-	FORCEINLINE bool IsEmpty() const
+	[[nodiscard]] FORCEINLINE bool IsEmpty() const
 	{
 		return FLockFreePointerFIFOBase<T, TPaddingForCacheContention>::IsEmpty();
 	}
@@ -999,7 +999,7 @@ public:
 	*
 	*	@return true if the list is closed.
 	*/
-	bool IsClosed() const
+	[[nodiscard]] bool IsClosed() const
 	{
 		return !!(FLockFreePointerListLIFOBase<T, TPaddingForCacheContention, 2>::GetState() & 1);
 	}
diff --git a/Engine/Source/Runtime/Core/Public/Containers/LruCache.h b/Engine/Source/Runtime/Core/Public/Containers/LruCache.h
--- a/Engine/Source/Runtime/Core/Public/Containers/LruCache.h
+++ b/Engine/Source/Runtime/Core/Public/Containers/LruCache.h
@@ -15,13 +15,13 @@
 template<typename KeyType>
 struct DefaultKeyComparer
 {
-	static FORCEINLINE bool Matches(KeyType A, KeyType B)
+	[[nodiscard]] static FORCEINLINE bool Matches(KeyType A, KeyType B)
 	{
 		return (A == B);
 	}
 
 	/** Calculates a hash index for a key. */
-	static FORCEINLINE uint32 GetKeyHash(KeyType Key)
+	[[nodiscard]] static FORCEINLINE uint32 GetKeyHash(KeyType Key)
 	{
 		return GetTypeHash(Key);
 	}
@@ -59,23 +59,25 @@ class TLruCache
 		 * @param InKey The entry's key.
 		 * @param InValue The entry's value.
 		 */
-		FCacheEntry(const KeyType& InKey, const ValueType& InValue)
+		[[nodiscard]] FCacheEntry(const KeyType& InKey, const ValueType& InValue)
 			: Key(InKey)
 			, LessRecent(nullptr)
 			, MoreRecent(nullptr)
 			, Value(InValue)
-		{ }
+		{
+		}
 
 		/**
 		 * Create a new instance with a default key value.
 		 *
 		 * @param InKey The entry's key.
 		 */
-		FCacheEntry(const KeyType& InKey)
+		[[nodiscard]] FCacheEntry(const KeyType& InKey)
 			: Key(InKey)
 			, LessRecent(nullptr)
 			, MoreRecent(nullptr)
-		{ }
+		{
+		}
 
 
 		/** Add this entry before the given one. */
@@ -110,17 +112,17 @@ class TLruCache
 	/** Lookup set key functions. */
 	struct FKeyFuncs : public BaseKeyFuncs<FCacheEntry*, KeyType>
 	{
-		FORCEINLINE static const KeyType& GetSetKey(const FCacheEntry* Entry)
+		[[nodiscard]] FORCEINLINE static const KeyType& GetSetKey(const FCacheEntry* Entry)
 		{
 			return Entry->Key;
 		}
 
-		FORCEINLINE static bool Matches(KeyType A, KeyType B)
+		[[nodiscard]] FORCEINLINE static bool Matches(KeyType A, KeyType B)
 		{
 			return KeyComp::Matches(A, B);
 		}
 
-		FORCEINLINE static uint32 GetKeyHash(KeyType Key)
+		[[nodiscard]] FORCEINLINE static uint32 GetKeyHash(KeyType Key)
 		{
 			return KeyComp::GetKeyHash(Key);
 		}
@@ -129,18 +131,19 @@ class TLruCache
 public:
 
 	/** Default constructor (empty cache that cannot hold any values). */
-	TLruCache()
+	[[nodiscard]] TLruCache()
 		: LeastRecent(nullptr)
 		, MostRecent(nullptr)
 		, MaxNumElements(0)
-	{ }
+	{
+	}
 
 	/**
 	 * Create and initialize a new instance.
 	 *
 	 * @param InMaxNumElements The maximum number of elements this cache can hold.
 	 */
-	TLruCache(int32 InMaxNumElements)
+	[[nodiscard]] TLruCache(int32 InMaxNumElements)
 		: LeastRecent(nullptr)
 		, MostRecent(nullptr)
 		, MaxNumElements(InMaxNumElements)
@@ -260,7 +263,7 @@ public:
 	 * @return true if the entry is in the cache, false otherwise.
 	 * @see Add, ContainsByPredicate, Empty, FilterByPredicate, Find, GetKeys, Remove
 	 */
-	FORCEINLINE bool Contains(const KeyType& Key) const
+	[[nodiscard]] FORCEINLINE bool Contains(const KeyType& Key) const
 	{
 		return LookupSet.Contains(Key);
 	}
@@ -273,7 +276,7 @@ public:
 	 * @see Contains, FilterByPredicate, FindByPredicate, RemoveByPredicate
 	 */
 	template<typename Predicate>
-	FORCEINLINE bool ContainsByPredicate(Predicate Pred) const
+	[[nodiscard]] FORCEINLINE bool ContainsByPredicate(Predicate Pred) const
 	{
 		for (const FCacheEntry* Entry : LookupSet)
 		{
@@ -314,7 +317,7 @@ public:
 	 * @see ContainsByPredicate, FindByPredicate, Find, RemoveByPredicate
 	 */
 	template<typename Predicate>
-	TArray<ValueType> FilterByPredicate(Predicate Pred) const
+	[[nodiscard]] TArray<ValueType> FilterByPredicate(Predicate Pred) const
 	{
 		TArray<ValueType> Result;
 
@@ -336,7 +339,7 @@ public:
 	 * @return Pointer to the value, or nullptr if not found.
 	 * @see Add, Contains, Empty, FindAndTouch, GetKeys, Remove
 	 */
-	FORCEINLINE const ValueType* Find(const KeyType& Key) const
+	[[nodiscard]] FORCEINLINE const ValueType* Find(const KeyType& Key) const
 	{
 		FCacheEntry*const * EntryPtr = LookupSet.Find(Key);
 
@@ -355,7 +358,7 @@ public:
 	 * @return Pointer to the value, or nullptr if not found.
 	 * @see Add, Contains, Empty, FindAndTouch, GetKeys, Remove
 	 */
-	FORCEINLINE ValueType* Find(const KeyType& Key)
+	[[nodiscard]] FORCEINLINE ValueType* Find(const KeyType& Key)
 	{
 		FCacheEntry* const* EntryPtr = LookupSet.Find(Key);
 
@@ -373,7 +376,7 @@ public:
 	 * @param Key The key of the entry to get.
 	 * @return Reference to the value, or triggers an assertion if the key does not exist.
 	 */
-	FORCEINLINE const ValueType& FindChecked(const KeyType& Key) const
+	[[nodiscard]] FORCEINLINE const ValueType& FindChecked(const KeyType& Key) const
 	{
 		FCacheEntry*const * EntryPtr = LookupSet.Find(Key);
 
@@ -388,7 +391,7 @@ public:
 	 * @param Key The key of the entry to get.
 	 * @return Copy of the value, or the default value for the ValueType if the key does not exist.
 	 */
-	FORCEINLINE ValueType FindRef(const KeyType& Key) const
+	[[nodiscard]] FORCEINLINE ValueType FindRef(const KeyType& Key) const
 	{
 		FCacheEntry*const * EntryPtr = LookupSet.Find(Key);
 
@@ -407,7 +410,7 @@ public:
 	 * @return Pointer to the value, or nullptr if not found.
 	 * @see Add, Contains, Empty, Find, GetKeys, Remove
 	 */
-	const ValueType* FindAndTouch(const KeyType& Key)
+	[[nodiscard]] const ValueType* FindAndTouch(const KeyType& Key)
 	{
 		FCacheEntry** EntryPtr = LookupSet.Find(Key);
 
@@ -427,7 +430,7 @@ public:
 	 * @param Key The key of the entry to get.
 	 * @return Pointer to the value, or triggers an assertion if the key does not exist.
 	 */
-	const ValueType& FindAndTouchChecked(const KeyType& Key)
+	[[nodiscard]] const ValueType& FindAndTouchChecked(const KeyType& Key)
 	{
 		FCacheEntry** EntryPtr = LookupSet.Find(Key);
 
@@ -444,7 +447,7 @@ public:
 	 * @param Key The key of the entry to get.
 	 * @return Copy of the value, or the default value for the ValueType if the key does not exist.
 	 */
-	ValueType FindAndTouchRef(const KeyType& Key)
+	[[nodiscard]] ValueType FindAndTouchRef(const KeyType& Key)
 	{
 		FCacheEntry** EntryPtr = LookupSet.Find(Key);
 
@@ -466,7 +469,7 @@ public:
 	 * @see ContainsByPredicate, FilterByPredicate, RemoveByPredicate
 	 */
 	template<typename Predicate>
-	const ValueType* FindByPredicate(Predicate Pred) const
+	[[nodiscard]] const ValueType* FindByPredicate(Predicate Pred) const
 	{
 		for (const FCacheEntry* Entry : LookupSet)
 		{
@@ -499,7 +502,7 @@ public:
 	 * @return Maximum number of entries.
 	 * @see Empty, Num
 	 */
-	FORCEINLINE int32 Max() const
+	[[nodiscard]] FORCEINLINE int32 Max() const
 	{
 		return MaxNumElements;
 	}
@@ -510,7 +513,7 @@ public:
 	 * @returns True if the cache is empty.
 	 * @see Num
 	 */
-	bool IsEmpty() const
+	[[nodiscard]] bool IsEmpty() const
 	{
 		return LookupSet.IsEmpty();
 	}
@@ -521,7 +524,7 @@ public:
 	 * @return Number of entries.
 	 * @see Empty, Max
 	 */
-	FORCEINLINE int32 Num() const
+	[[nodiscard]] FORCEINLINE int32 Num() const
 	{
 		return LookupSet.Num();
 	}
@@ -591,13 +594,15 @@ public:
 	{
 	public:
 
-		FORCEINLINE TBaseIterator()
+		[[nodiscard]] FORCEINLINE TBaseIterator()
 			: CurrentEntry(nullptr)
-		{ }
+		{
+		}
 
-		FORCEINLINE TBaseIterator(const TLruCache& Cache)
+		[[nodiscard]] FORCEINLINE TBaseIterator(const TLruCache& Cache)
 			: CurrentEntry(Cache.MostRecent)
-		{ }
+		{
+		}
 
 	public:
 
@@ -607,47 +612,47 @@ public:
 			return *this;
 		}
 
-		FORCEINLINE bool operator==(const TBaseIterator& Rhs) const
+		[[nodiscard]] FORCEINLINE bool operator==(const TBaseIterator& Rhs) const
 		{
 			return CurrentEntry == Rhs.CurrentEntry;
 		}
 
-		FORCEINLINE bool operator!=(const TBaseIterator& Rhs) const
+		[[nodiscard]] FORCEINLINE bool operator!=(const TBaseIterator& Rhs) const
 		{
 			return CurrentEntry != Rhs.CurrentEntry;
 		}
 
-		ValueType& operator->() const
+		[[nodiscard]] ValueType& operator->() const
 		{
 			check(CurrentEntry != nullptr);
 			return CurrentEntry->Value;
 		}
 
-		ValueType& operator*() const
+		[[nodiscard]] ValueType& operator*() const
 		{
 			check(CurrentEntry != nullptr);
 			return CurrentEntry->Value;
 		}
 
-		FORCEINLINE explicit operator bool() const
+		[[nodiscard]] FORCEINLINE explicit operator bool() const
 		{
 			return (CurrentEntry != nullptr);
 		}
 
-		FORCEINLINE bool operator!() const
+		[[nodiscard]] FORCEINLINE bool operator!() const
 		{
 			return !(bool)*this;
 		}
 
 	public:
 
-		FORCEINLINE KeyType& Key() const
+		[[nodiscard]] FORCEINLINE KeyType& Key() const
 		{
 			check(CurrentEntry != nullptr);
 			return CurrentEntry->Key;
 		}
 
-		FORCEINLINE ValueType& Value() const
+		[[nodiscard]] FORCEINLINE ValueType& Value() const
 		{
 			check(CurrentEntry != nullptr);
 			return CurrentEntry->Value;
@@ -655,7 +660,7 @@ public:
 
 	protected:
 
-		FCacheEntry* GetCurrentEntry()
+		[[nodiscard]] FCacheEntry* GetCurrentEntry()
 		{
 			return CurrentEntry;
 		}
@@ -680,13 +685,15 @@ public:
 	{
 	public:
 
-		FORCEINLINE TConstIterator()
+		[[nodiscard]] FORCEINLINE TConstIterator()
 			: TBaseIterator<true>()
-		{ }
+		{
+		}
 
-		FORCEINLINE TConstIterator(const TLruCache& Cache)
+		[[nodiscard]] FORCEINLINE TConstIterator(const TLruCache& Cache)
 			: TBaseIterator<true>(Cache)
-		{ }
+		{
+		}
 	};
 
 
@@ -698,15 +705,17 @@ public:
 	{
 	public:
 
-		FORCEINLINE TIterator()
+		[[nodiscard]] FORCEINLINE TIterator()
 			: TBaseIterator<false>()
 			, Cache(nullptr)
-		{ }
+		{
+		}
 
-		FORCEINLINE TIterator(TLruCache& InCache)
+		[[nodiscard]] FORCEINLINE TIterator(TLruCache& InCache)
 			: TBaseIterator<false>(InCache)
 			, Cache(&InCache)
-		{ }
+		{
+		}
 
 		/** Removes the current element from the cache and increments the iterator. */
 		FORCEINLINE void RemoveCurrentAndIncrement()
@@ -780,10 +789,10 @@ protected:
 
 public:
 
-	TIterator begin() { return TIterator(*this); }
-	TConstIterator begin() const { return TConstIterator(*this); }
-	TIterator end() { return TIterator(); }
-	TConstIterator end() const { return TConstIterator(); }
+	[[nodiscard]] TIterator begin() { return TIterator(*this); }
+	[[nodiscard]] TConstIterator begin() const { return TConstIterator(*this); }
+	[[nodiscard]] TIterator end() { return TIterator(); }
+	[[nodiscard]] TConstIterator end() const { return TConstIterator(); }
 
 private:
 
diff --git a/Engine/Source/Runtime/Core/Public/Containers/MRUArray.h b/Engine/Source/Runtime/Core/Public/Containers/MRUArray.h
--- a/Engine/Source/Runtime/Core/Public/Containers/MRUArray.h
+++ b/Engine/Source/Runtime/Core/Public/Containers/MRUArray.h
@@ -29,14 +29,14 @@ public:
 	/**
 	 * Constructor.
 	 */
-	TMRUArray()
+	[[nodiscard]] TMRUArray()
 		: Super()
 	{
 		MaxItems = 0;
 	}
 
-	TMRUArray(TMRUArray&&) = default;
-	TMRUArray(const TMRUArray&) = default;
+	[[nodiscard]] TMRUArray(TMRUArray&&) = default;
+	[[nodiscard]] TMRUArray(const TMRUArray&) = default;
 	TMRUArray& operator=(TMRUArray&&) = default;
 	TMRUArray& operator=(const TMRUArray&) = default;
 
diff --git a/Engine/Source/Runtime/Core/Public/Containers/Map.h b/Engine/Source/Runtime/Core/Public/Containers/Map.h
--- a/Engine/Source/Runtime/Core/Public/Containers/Map.h
+++ b/Engine/Source/Runtime/Core/Public/Containers/Map.h
@@ -31,7 +31,7 @@ public:
 	std::conditional_t<std::is_rvalue_reference_v<ValueInitType>, ValueInitType&, ValueInitType> Value;
 
 	/** Initialization constructor. */
-	FORCEINLINE TPairInitializer(KeyInitType InKey, ValueInitType InValue)
+	[[nodiscard]] FORCEINLINE TPairInitializer(KeyInitType InKey, ValueInitType InValue)
 		: Key(InKey)
 		, Value(InValue)
 	{
@@ -39,14 +39,14 @@ public:
 
 	/** Implicit conversion to pair initializer. */
 	template <typename KeyType, typename ValueType>
-	FORCEINLINE TPairInitializer(const TPair<KeyType, ValueType>& Pair)
+	[[nodiscard]] FORCEINLINE TPairInitializer(const TPair<KeyType, ValueType>& Pair)
 		: Key(Pair.Key)
 		, Value(Pair.Value)
 	{
 	}
 
 	template <typename KeyType, typename ValueType>
-	operator TPair<KeyType, ValueType>() const
+	[[nodiscard]] operator TPair<KeyType, ValueType>() const
 	{
 		return TPair<KeyType, ValueType>(StaticCast<KeyInitType>(Key), StaticCast<ValueInitType>(Value));
 	}
@@ -61,13 +61,13 @@ public:
 	std::conditional_t<std::is_rvalue_reference_v<KeyInitType>, KeyInitType&, KeyInitType> Key;
 
 	/** Initialization constructor. */
-	FORCEINLINE explicit TKeyInitializer(KeyInitType InKey)
+	[[nodiscard]] FORCEINLINE explicit TKeyInitializer(KeyInitType InKey)
 		: Key(InKey)
 	{
 	}
 
 	template <typename KeyType, typename ValueType>
-	operator TPair<KeyType, ValueType>() const
+	[[nodiscard]] operator TPair<KeyType, ValueType>() const
 	{
 		return TPair<KeyType, ValueType>(StaticCast<KeyInitType>(Key), ValueType());
 	}
@@ -80,29 +80,29 @@ struct TDefaultMapKeyFuncs : BaseKeyFuncs<TPair<KeyType, ValueType>, KeyType, bI
 	typedef typename TTypeTraits<KeyType>::ConstPointerType KeyInitType;
 	typedef const TPairInitializer<typename TTypeTraits<KeyType>::ConstInitType, typename TTypeTraits<ValueType>::ConstInitType>& ElementInitType;
 
-	static FORCEINLINE KeyInitType GetSetKey(ElementInitType Element)
+	[[nodiscard]] static FORCEINLINE KeyInitType GetSetKey(ElementInitType Element)
 	{
 		return Element.Key;
 	}
 
-	static FORCEINLINE bool Matches(KeyInitType A, KeyInitType B)
+	[[nodiscard]] static FORCEINLINE bool Matches(KeyInitType A, KeyInitType B)
 	{
 		return A == B;
 	}
 
 	template<typename ComparableKey>
-	static FORCEINLINE bool Matches(KeyInitType A, ComparableKey B)
+	[[nodiscard]] static FORCEINLINE bool Matches(KeyInitType A, ComparableKey B)
 	{
 		return A == B;
 	}
 
-	static FORCEINLINE uint32 GetKeyHash(KeyInitType Key)
+	[[nodiscard]] static FORCEINLINE uint32 GetKeyHash(KeyInitType Key)
 	{
 		return GetTypeHash(Key);
 	}
 
 	template<typename ComparableKey>
-	static FORCEINLINE uint32 GetKeyHash(ComparableKey Key)
+	[[nodiscard]] static FORCEINLINE uint32 GetKeyHash(ComparableKey Key)
 	{
 		return GetTypeHash(Key);
 	}
@@ -143,23 +143,25 @@ public:
 	typedef TPair<KeyType, ValueType> ElementType;
 
 protected:
-	TMapBase() = default;
-	TMapBase(TMapBase&&) = default;
-	TMapBase(const TMapBase&) = default;
+	[[nodiscard]] TMapBase() = default;
+	[[nodiscard]] TMapBase(TMapBase&&) = default;
+	[[nodiscard]] TMapBase(const TMapBase&) = default;
 	TMapBase& operator=(TMapBase&&) = default;
 	TMapBase& operator=(const TMapBase&) = default;
 
 	/** Constructor for moving elements from a TMap with a different SetAllocator */
 	template<typename OtherSetAllocator>
-	TMapBase(TMapBase<KeyType, ValueType, OtherSetAllocator, KeyFuncs>&& Other)
+	[[nodiscard]] TMapBase(TMapBase<KeyType, ValueType, OtherSetAllocator, KeyFuncs>&& Other)
 		: Pairs(MoveTemp(Other.Pairs))
-	{ }
+	{
+	}
 
 	/** Constructor for copying elements from a TMap with a different SetAllocator */
 	template<typename OtherSetAllocator>
-	TMapBase(const TMapBase<KeyType, ValueType, OtherSetAllocator, KeyFuncs>& Other)
+	[[nodiscard]] TMapBase(const TMapBase<KeyType, ValueType, OtherSetAllocator, KeyFuncs>& Other)
 		: Pairs(Other.Pairs)
-	{ }
+	{
+	}
 
 	/////////////////////////////////////////////////
 	// Start - intrusive TOptional<TMapBase> state //
@@ -167,11 +169,11 @@ protected:
 	constexpr static bool bHasIntrusiveUnsetOptionalState = true;
 	using IntrusiveUnsetOptionalStateType = TMapBase;
 
-	explicit TMapBase(FIntrusiveUnsetOptionalState Tag)
+	[[nodiscard]] explicit TMapBase(FIntrusiveUnsetOptionalState Tag)
 		: Pairs(Tag)
 	{
 	}
-	bool operator==(FIntrusiveUnsetOptionalState Tag) const
+	[[nodiscard]] bool operator==(FIntrusiveUnsetOptionalState Tag) const
 	{
 		return Pairs == Tag;
 	}
@@ -205,7 +207,7 @@ public:
 	 * @param Other The other map to compare against
 	 * @returns True if both this and Other contain the same keys with values that compare ==
 	 */
-	bool OrderIndependentCompareEqual(const TMapBase& Other) const
+	[[nodiscard]] bool OrderIndependentCompareEqual(const TMapBase& Other) const
 	{
 		// first check counts (they should be the same obviously)
 		if (Num() != Other.Num())
@@ -280,13 +282,13 @@ public:
 	 * @returns True if the map is empty.
 	 * @see Num
 	 */
-	bool IsEmpty() const
+	[[nodiscard]] bool IsEmpty() const
 	{
 		return Pairs.IsEmpty();
 	}
 
 	/** @return The number of elements in the map. */
-	FORCEINLINE int32 Num() const
+	[[nodiscard]] FORCEINLINE int32 Num() const
 	{
 		return Pairs.Num();
 	}
@@ -383,7 +385,7 @@ public:
 	 * @return Number of bytes allocated by this container.
 	 * @see CountBytes
 	 */
-	FORCEINLINE SIZE_T GetAllocatedSize() const
+	[[nodiscard]] FORCEINLINE SIZE_T GetAllocatedSize() const
 	{
 		return Pairs.GetAllocatedSize();
 	}
@@ -515,7 +517,7 @@ public:
 	 *     or nullptr if the value isn't contained in this map. The pointer
 	 *     is only valid until the next change to any key in the map.
 	 */
-	const KeyType* FindKey(ValueInitType Value) const
+	[[nodiscard]] const KeyType* FindKey(ValueInitType Value) const
 	{
 		for (typename ElementSetType::TConstIterator PairIt(Pairs); PairIt; ++PairIt)
 		{
@@ -535,7 +537,7 @@ public:
 	 *          the subset of elements for which the functor returns true.
 	 */
 	template <typename Predicate>
-	TMap<KeyType, ValueType, SetAllocator, KeyFuncs> FilterByPredicate(Predicate Pred) const
+	[[nodiscard]] TMap<KeyType, ValueType, SetAllocator, KeyFuncs> FilterByPredicate(Predicate Pred) const
 	{
 		TMap<KeyType, ValueType, SetAllocator, KeyFuncs> FilterResults;
 		FilterResults.Reserve(Pairs.Num());
@@ -556,7 +558,7 @@ public:
 	 * @return A pointer to the value associated with the specified key, or nullptr if the key isn't contained in this map.  The pointer
 	 *			is only valid until the next change to any key in the map.
 	 */
-	FORCEINLINE ValueType* Find(KeyConstPointerType Key)
+	[[nodiscard]] FORCEINLINE ValueType* Find(KeyConstPointerType Key)
 	{
 		if (auto* Pair = Pairs.Find(Key))
 		{
@@ -565,14 +567,14 @@ public:
 
 		return nullptr;
 	}
-	FORCEINLINE const ValueType* Find(KeyConstPointerType Key) const
+	[[nodiscard]] FORCEINLINE const ValueType* Find(KeyConstPointerType Key) const
 	{
 		return const_cast<TMapBase*>(this)->Find(Key);
 	}
 
 	/** See Find() and class documentation section on ByHash() functions */
 	template<typename ComparableKey>
-	FORCEINLINE ValueType* FindByHash(uint32 KeyHash, const ComparableKey& Key)
+	[[nodiscard]] FORCEINLINE ValueType* FindByHash(uint32 KeyHash, const ComparableKey& Key)
 	{
 		if (auto* Pair = Pairs.FindByHash(KeyHash, Key))
 		{
@@ -582,26 +584,26 @@ public:
 		return nullptr;
 	}
 	template<typename ComparableKey>
-	FORCEINLINE const ValueType* FindByHash(uint32 KeyHash, const ComparableKey& Key) const
+	[[nodiscard]] FORCEINLINE const ValueType* FindByHash(uint32 KeyHash, const ComparableKey& Key) const
 	{
 		return const_cast<TMapBase*>(this)->FindByHash(KeyHash, Key);
 	}
 
 	template<typename ComparableKey>
-	FORCEINLINE ValueType& FindByHashChecked(uint32 KeyHash, const ComparableKey& Key)
+	[[nodiscard]] FORCEINLINE ValueType& FindByHashChecked(uint32 KeyHash, const ComparableKey& Key)
 	{
 		auto* Pair = Pairs.FindByHash(KeyHash, Key);
 		check(Pair != nullptr);
 		return Pair->Value;
 	}
 	template<typename ComparableKey>
-	FORCEINLINE const ValueType& FindByHashChecked(uint32 KeyHash, const ComparableKey& Key) const
+	[[nodiscard]] FORCEINLINE const ValueType& FindByHashChecked(uint32 KeyHash, const ComparableKey& Key) const
 	{
 		return const_cast<TMapBase*>(this)->FindByHashChecked(KeyHash, Key);
 	}
 
 private:
-	FORCEINLINE static uint32 HashKey(const KeyType& Key)
+	[[nodiscard]] FORCEINLINE static uint32 HashKey(const KeyType& Key)
 	{
 		return KeyFuncs::GetKeyHash(Key);
 	}
@@ -614,7 +616,7 @@ private:
 	 * @return A reference to the value associated with the specified key.
 	 */
 	template <typename InitKeyType>
-	ValueType& FindOrAddImpl(uint32 KeyHash, InitKeyType&& Key)
+	[[nodiscard]] ValueType& FindOrAddImpl(uint32 KeyHash, InitKeyType&& Key)
 	{
 		if (auto* Pair = Pairs.FindByHash(KeyHash, Key))
 		{
@@ -633,7 +635,7 @@ private:
 	 * @return A reference to the value associated with the specified key.
 	 */
 	template <typename InitKeyType, typename InitValueType>
-	ValueType& FindOrAddImpl(uint32 KeyHash, InitKeyType&& Key, InitValueType&& Value)
+	[[nodiscard]] ValueType& FindOrAddImpl(uint32 KeyHash, InitKeyType&& Key, InitValueType&& Value)
 	{
 		if (auto* Pair = Pairs.FindByHash(KeyHash, Key))
 		{
@@ -684,7 +686,7 @@ public:
 	 * @param Key The key to search for.
 	 * @return The value associated with the specified key, or triggers an assertion if the key does not exist.
 	 */
-	FORCEINLINE const ValueType& FindChecked(KeyConstPointerType Key) const
+	[[nodiscard]] FORCEINLINE const ValueType& FindChecked(KeyConstPointerType Key) const
 	{
 		const auto* Pair = Pairs.Find(Key);
 		check(Pair != nullptr);
@@ -697,7 +699,7 @@ public:
 	 * @param Key The key to search for.
 	 * @return The value associated with the specified key, or triggers an assertion if the key does not exist.
 	 */
-	FORCEINLINE ValueType& FindChecked(KeyConstPointerType Key)
+	[[nodiscard]] FORCEINLINE ValueType& FindChecked(KeyConstPointerType Key)
 	{
 		auto* Pair = Pairs.Find(Key);
 		check(Pair != nullptr);
@@ -710,7 +712,7 @@ public:
 	 * @param Key The key to search for.
 	 * @return The value associated with the specified key, or the default value for the ValueType if the key isn't contained in this map.
 	 */
-	FORCEINLINE ValueType FindRef(KeyConstPointerType Key) const
+	[[nodiscard]] FORCEINLINE ValueType FindRef(KeyConstPointerType Key) const
 	{
 		if (const auto* Pair = Pairs.Find(Key))
 		{
@@ -727,7 +729,7 @@ public:
 	 * @param DefaultValue The fallback value if the key is not found.
 	 * @return The value associated with the specified key, or DefaultValue if the key isn't contained in this map.
 	 */
-	FORCEINLINE ValueType FindRef(KeyConstPointerType Key, ValueType DefaultValue) const
+	[[nodiscard]] FORCEINLINE ValueType FindRef(KeyConstPointerType Key, ValueType DefaultValue) const
 	{
 		if (const auto* Pair = Pairs.Find(Key))
 		{
@@ -742,13 +744,13 @@ public:
 	 * Callers should not depend on particular patterns in the behaviour of this function.
 	 * @return A pointer to an arbitrary pair, or nullptr if the container is empty.
 	 */
-	ElementType* FindArbitraryElement()
+	[[nodiscard]] ElementType* FindArbitraryElement()
 	{
 		// The goal of this function is to be fast, and so the implementation may be improved at any time even if it gives different results.
 
 		return Pairs.FindArbitraryElement();
 	}
-	const ElementType* FindArbitraryElement() const
+	[[nodiscard]] const ElementType* FindArbitraryElement() const
 	{
 		return const_cast<TMapBase*>(this)->FindArbitraryElement();
 	}
@@ -759,20 +761,20 @@ public:
 	 * @param Key The key to check for.
 	 * @return true if the map contains the key.
 	 */
-	FORCEINLINE bool Contains(KeyConstPointerType Key) const
+	[[nodiscard]] FORCEINLINE bool Contains(KeyConstPointerType Key) const
 	{
 		return Pairs.Contains(Key);
 	}
 
 	/** See Contains() and class documentation section on ByHash() functions */
 	template<typename ComparableKey>
-	FORCEINLINE bool ContainsByHash(uint32 KeyHash, const ComparableKey& Key) const
+	[[nodiscard]] FORCEINLINE bool ContainsByHash(uint32 KeyHash, const ComparableKey& Key) const
 	{
 		return Pairs.ContainsByHash(KeyHash, Key);
 	}
 
 	/** Copy the key/value pairs in this map into an array. */
-	TArray<ElementType> Array() const
+	[[nodiscard]] TArray<ElementType> Array() const
 	{
 		return Pairs.Array();
 	}
@@ -835,7 +837,7 @@ protected:
 		typedef std::conditional_t<bConst, const typename ElementSetType::ElementType, typename ElementSetType::ElementType> PairType;
 
 	public:
-		FORCEINLINE TBaseIterator(const PairItType& InElementIt)
+		[[nodiscard]] FORCEINLINE TBaseIterator(const PairItType& InElementIt)
 			: PairIt(InElementIt)
 		{
 		}
@@ -847,29 +849,47 @@ protected:
 		}
 
 		/** conversion to "bool" returning true if the iterator is valid. */
-		FORCEINLINE explicit operator bool() const
+		[[nodiscard]] FORCEINLINE explicit operator bool() const
 		{ 
 			return !!PairIt; 
 		}
 		/** inverse of the "bool" operator */
-		FORCEINLINE bool operator !() const 
+		[[nodiscard]] FORCEINLINE bool operator !() const
 		{
 			return !(bool)*this;
 		}
 
-		FORCEINLINE bool operator==(const TBaseIterator& Rhs) const { return PairIt == Rhs.PairIt; }
-		FORCEINLINE bool operator!=(const TBaseIterator& Rhs) const { return PairIt != Rhs.PairIt; }
+		[[nodiscard]] FORCEINLINE bool operator==(const TBaseIterator& Rhs) const
+		{
+			return PairIt == Rhs.PairIt;
+		}
+		[[nodiscard]] FORCEINLINE bool operator!=(const TBaseIterator& Rhs) const
+		{
+			return PairIt != Rhs.PairIt;
+		}
 
-		FORCEINLINE ItKeyType&   Key()   const { return PairIt->Key; }
-		FORCEINLINE ItValueType& Value() const { return PairIt->Value; }
+		[[nodiscard]] FORCEINLINE ItKeyType& Key() const
+		{
+			return PairIt->Key;
+		}
+		[[nodiscard]] FORCEINLINE ItValueType& Value() const
+		{
+			return PairIt->Value;
+		}
 
 		[[nodiscard]] FORCEINLINE FSetElementId GetId() const
 		{
 			return PairIt.GetId();
 		}
 
-		FORCEINLINE PairType& operator* () const { return  *PairIt; }
-		FORCEINLINE PairType* operator->() const { return &*PairIt; }
+		[[nodiscard]] FORCEINLINE PairType& operator*() const
+		{
+			return *PairIt;
+		}
+		[[nodiscard]] FORCEINLINE PairType* operator->() const
+		{
+			return &*PairIt;
+		}
 
 	protected:
 		PairItType PairIt;
@@ -886,7 +906,7 @@ protected:
 
 	public:
 		/** Initialization constructor. */
-		FORCEINLINE TBaseKeyIterator(const SetItType& InSetIt)
+		[[nodiscard]] FORCEINLINE TBaseKeyIterator(const SetItType& InSetIt)
 			: SetIt(InSetIt)
 		{
 		}
@@ -898,12 +918,12 @@ protected:
 		}
 
 		/** conversion to "bool" returning true if the iterator is valid. */
-		FORCEINLINE explicit operator bool() const
-		{ 
+		[[nodiscard]] FORCEINLINE explicit operator bool() const
+		{
 			return !!SetIt; 
 		}
 		/** inverse of the "bool" operator */
-		FORCEINLINE bool operator !() const 
+		[[nodiscard]] FORCEINLINE bool operator !() const 
 		{
 			return !(bool)*this;
 		}
@@ -912,11 +932,23 @@ protected:
 		{
 			return SetIt.GetId();
 		}
-		FORCEINLINE ItKeyType&   Key() const { return SetIt->Key; }
-		FORCEINLINE ItValueType& Value() const { return SetIt->Value; }
+		[[nodiscard]] FORCEINLINE ItKeyType& Key() const
+		{
+			return SetIt->Key;
+		}
+		[[nodiscard]] FORCEINLINE ItValueType& Value() const
+		{
+			return SetIt->Value;
+		}
 
-		FORCEINLINE decltype(auto) operator* () const { return SetIt.operator* (); }
-		FORCEINLINE decltype(auto) operator->() const { return SetIt.operator->(); }
+		[[nodiscard]] FORCEINLINE decltype(auto) operator*() const
+		{
+			return SetIt.operator*();
+		}
+		[[nodiscard]] FORCEINLINE decltype(auto) operator->() const
+		{
+			return SetIt.operator->();
+		}
 
 	protected:
 		SetItType SetIt;
@@ -947,7 +979,7 @@ public:
 	public:
 
 		/** Initialization constructor. */
-		FORCEINLINE TIterator(TMapBase& InMap, bool bInRequiresRehashOnRemoval = false)
+		[[nodiscard]] FORCEINLINE TIterator(TMapBase& InMap, bool bInRequiresRehashOnRemoval = false)
 			: TBaseIterator<false>(InMap.Pairs.CreateIterator())
 			, Map(InMap)
 			, bElementsHaveBeenRemoved(false)
@@ -984,7 +1016,7 @@ public:
 	class TConstIterator : public TBaseIterator<true>
 	{
 	public:
-		FORCEINLINE TConstIterator(const TMapBase& InMap)
+		[[nodiscard]] FORCEINLINE TConstIterator(const TMapBase& InMap)
 			: TBaseIterator<true>(InMap.Pairs.CreateConstIterator())
 		{
 		}
@@ -1003,7 +1035,7 @@ public:
 	public:
 		using KeyArgumentType = typename IteratorType::KeyArgumentType;
 
-		FORCEINLINE TConstKeyIterator(const TMapBase& InMap, KeyArgumentType InKey)
+		[[nodiscard]] FORCEINLINE TConstKeyIterator(const TMapBase& InMap, KeyArgumentType InKey)
 			: Super(IteratorType(InMap.Pairs, InKey))
 		{
 		}
@@ -1019,7 +1051,7 @@ public:
 	public:
 		using KeyArgumentType = typename IteratorType::KeyArgumentType;
 
-		FORCEINLINE TKeyIterator(TMapBase& InMap, KeyArgumentType InKey)
+		[[nodiscard]] FORCEINLINE TKeyIterator(TMapBase& InMap, KeyArgumentType InKey)
 			: Super(IteratorType(InMap.Pairs, InKey))
 		{
 		}
@@ -1032,25 +1064,25 @@ public:
 	};
 
 	/** Creates an iterator over all the pairs in this map */
-	FORCEINLINE TIterator CreateIterator()
+	[[nodiscard]] FORCEINLINE TIterator CreateIterator()
 	{
 		return TIterator(*this);
 	}
 
 	/** Creates a const iterator over all the pairs in this map */
-	FORCEINLINE TConstIterator CreateConstIterator() const
+	[[nodiscard]] FORCEINLINE TConstIterator CreateConstIterator() const
 	{
 		return TConstIterator(*this);
 	}
 
 	/** Creates an iterator over the values associated with a specified key in a map */
-	FORCEINLINE TKeyIterator CreateKeyIterator(typename TKeyIterator::KeyArgumentType InKey)
+	[[nodiscard]] FORCEINLINE TKeyIterator CreateKeyIterator(typename TKeyIterator::KeyArgumentType InKey)
 	{
 		return TKeyIterator(*this, InKey);
 	}
 
 	/** Creates a const iterator over the values associated with a specified key in a map */
-	FORCEINLINE TConstKeyIterator CreateConstKeyIterator(typename TConstKeyIterator::KeyArgumentType InKey) const
+	[[nodiscard]] FORCEINLINE TConstKeyIterator CreateConstKeyIterator(typename TConstKeyIterator::KeyArgumentType InKey) const
 	{
 		return TConstKeyIterator(*this, InKey);
 	}
@@ -1062,10 +1094,10 @@ public:
 	 * DO NOT USE DIRECTLY
 	 * STL-like iterators to enable range-based for loop support.
 	 */
-	FORCEINLINE TRangedForIterator      begin() { return TRangedForIterator(Pairs.begin()); }
-	FORCEINLINE TRangedForConstIterator begin() const { return TRangedForConstIterator(Pairs.begin()); }
-	FORCEINLINE TRangedForIterator      end() { return TRangedForIterator(Pairs.end()); }
-	FORCEINLINE TRangedForConstIterator end() const { return TRangedForConstIterator(Pairs.end()); }
+	[[nodiscard]] FORCEINLINE TRangedForIterator      begin() { return TRangedForIterator(Pairs.begin()); }
+	[[nodiscard]] FORCEINLINE TRangedForConstIterator begin() const { return TRangedForConstIterator(Pairs.begin()); }
+	[[nodiscard]] FORCEINLINE TRangedForIterator      end() { return TRangedForIterator(Pairs.end()); }
+	[[nodiscard]] FORCEINLINE TRangedForConstIterator end() const { return TRangedForConstIterator(Pairs.end()); }
 
 	// Maps are deliberately prevented from being hashed or compared, because this would hide potentially major performance problems behind default operations.
 	friend uint32 GetTypeHash(const TMapBase& Map) = delete;
@@ -1080,22 +1112,22 @@ class TSortableMapBase : public TMapBase<KeyType, ValueType, SetAllocator, KeyFu
 protected:
 	typedef TMapBase<KeyType, ValueType, SetAllocator, KeyFuncs> Super;
 
-	TSortableMapBase() = default;
-	TSortableMapBase(TSortableMapBase&&) = default;
-	TSortableMapBase(const TSortableMapBase&) = default;
+	[[nodiscard]] TSortableMapBase() = default;
+	[[nodiscard]] TSortableMapBase(TSortableMapBase&&) = default;
+	[[nodiscard]] TSortableMapBase(const TSortableMapBase&) = default;
 	TSortableMapBase& operator=(TSortableMapBase&&) = default;
 	TSortableMapBase& operator=(const TSortableMapBase&) = default;
 
 	/** Constructor for moving elements from a TMap with a different SetAllocator */
 	template<typename OtherSetAllocator>
-	TSortableMapBase(TSortableMapBase<KeyType, ValueType, OtherSetAllocator, KeyFuncs>&& Other)
+	[[nodiscard]] TSortableMapBase(TSortableMapBase<KeyType, ValueType, OtherSetAllocator, KeyFuncs>&& Other)
 		: Super(MoveTemp(Other))
 	{
 	}
 
 	/** Constructor for copying elements from a TMap with a different SetAllocator */
 	template<typename OtherSetAllocator>
-	TSortableMapBase(const TSortableMapBase<KeyType, ValueType, OtherSetAllocator, KeyFuncs>& Other)
+	[[nodiscard]] TSortableMapBase(const TSortableMapBase<KeyType, ValueType, OtherSetAllocator, KeyFuncs>& Other)
 		: Super(Other)
 	{
 	}
@@ -1122,11 +1154,11 @@ protected:
 	constexpr static bool bHasIntrusiveUnsetOptionalState = true;
 	using IntrusiveUnsetOptionalStateType = TSortableMapBase;
 
-	explicit TSortableMapBase(FIntrusiveUnsetOptionalState Tag)
+	[[nodiscard]] explicit TSortableMapBase(FIntrusiveUnsetOptionalState Tag)
 		: Super(Tag)
 	{
 	}
-	bool operator==(FIntrusiveUnsetOptionalState Tag) const
+	[[nodiscard]] bool operator==(FIntrusiveUnsetOptionalState Tag) const
 	{
 		return Super::operator==(Tag);
 	}
@@ -1195,11 +1227,12 @@ private:
 
 	public:
 
-		FORCEINLINE FKeyComparisonClass(const PREDICATE_CLASS& InPredicate)
+		[[nodiscard]] FORCEINLINE FKeyComparisonClass(const PREDICATE_CLASS& InPredicate)
 			: Predicate(InPredicate)
-		{}
+		{
+		}
 
-		FORCEINLINE bool operator()(const typename Super::ElementType& A, const typename Super::ElementType& B) const
+		[[nodiscard]] FORCEINLINE bool operator()(const typename Super::ElementType& A, const typename Super::ElementType& B) const
 		{
 			return Predicate(A.Key, B.Key);
 		}
@@ -1213,11 +1246,12 @@ private:
 
 	public:
 
-		FORCEINLINE FValueComparisonClass(const PREDICATE_CLASS& InPredicate)
+		[[nodiscard]] FORCEINLINE FValueComparisonClass(const PREDICATE_CLASS& InPredicate)
 			: Predicate(InPredicate)
-		{}
+		{
+		}
 
-		FORCEINLINE bool operator()(const typename Super::ElementType& A, const typename Super::ElementType& B) const
+		[[nodiscard]] FORCEINLINE bool operator()(const typename Super::ElementType& A, const typename Super::ElementType& B) const
 		{
 			return Predicate(A.Value, B.Value);
 		}
@@ -1246,28 +1280,28 @@ public:
 	typedef typename Super::KeyInitType KeyInitType;
 	typedef typename Super::KeyConstPointerType KeyConstPointerType;
 
-	TMap() = default;
-	TMap(TMap&&) = default;
-	TMap(const TMap&) = default;
+	[[nodiscard]] TMap() = default;
+	[[nodiscard]] TMap(TMap&&) = default;
+	[[nodiscard]] TMap(const TMap&) = default;
 	TMap& operator=(TMap&&) = default;
 	TMap& operator=(const TMap&) = default;
 
 	/** Constructor for moving elements from a TMap with a different SetAllocator */
 	template<typename OtherSetAllocator>
-	TMap(TMap<KeyType, ValueType, OtherSetAllocator, KeyFuncs>&& Other)
+	[[nodiscard]] TMap(TMap<KeyType, ValueType, OtherSetAllocator, KeyFuncs>&& Other)
 		: Super(MoveTemp(Other))
 	{
 	}
 
 	/** Constructor for copying elements from a TMap with a different SetAllocator */
 	template<typename OtherSetAllocator>
-	TMap(const TMap<KeyType, ValueType, OtherSetAllocator, KeyFuncs>& Other)
+	[[nodiscard]] TMap(const TMap<KeyType, ValueType, OtherSetAllocator, KeyFuncs>& Other)
 		: Super(Other)
 	{
 	}
 
 	/** Constructor which gets its elements from a native initializer list */
-	TMap(std::initializer_list<TPairInitializer<const KeyType&, const ValueType&>> InitList)
+	[[nodiscard]] TMap(std::initializer_list<TPairInitializer<const KeyType&, const ValueType&>> InitList)
 	{
 		this->Reserve((int32)InitList.size());
 		for (const TPairInitializer<const KeyType&, const ValueType&>& Element : InitList)
@@ -1282,11 +1316,11 @@ public:
 	constexpr static bool bHasIntrusiveUnsetOptionalState = true;
 	using IntrusiveUnsetOptionalStateType = TMap;
 
-	explicit TMap(FIntrusiveUnsetOptionalState Tag)
+	[[nodiscard]] explicit TMap(FIntrusiveUnsetOptionalState Tag)
 		: Super(Tag)
 	{
 	}
-	bool operator==(FIntrusiveUnsetOptionalState Tag) const
+	[[nodiscard]] bool operator==(FIntrusiveUnsetOptionalState Tag) const
 	{
 		return Super::operator==(Tag);
 	}
@@ -1408,8 +1442,14 @@ public:
 		}
 	}
 
-	FORCEINLINE       ValueType& operator[](KeyConstPointerType Key) { return this->FindChecked(Key); }
-	FORCEINLINE const ValueType& operator[](KeyConstPointerType Key) const { return this->FindChecked(Key); }
+	[[nodiscard]] FORCEINLINE ValueType& operator[](KeyConstPointerType Key)
+	{
+		return this->FindChecked(Key);
+	}
+	[[nodiscard]] FORCEINLINE const ValueType& operator[](KeyConstPointerType Key) const
+	{
+		return this->FindChecked(Key);
+	}
 };
 
 namespace Freeze
@@ -1449,28 +1489,28 @@ public:
 	typedef typename Super::KeyInitType KeyInitType;
 	typedef typename Super::ValueInitType ValueInitType;
 
-	TMultiMap() = default;
-	TMultiMap(TMultiMap&&) = default;
-	TMultiMap(const TMultiMap&) = default;
+	[[nodiscard]] TMultiMap() = default;
+	[[nodiscard]] TMultiMap(TMultiMap&&) = default;
+	[[nodiscard]] TMultiMap(const TMultiMap&) = default;
 	TMultiMap& operator=(TMultiMap&&) = default;
 	TMultiMap& operator=(const TMultiMap&) = default;
 
 	/** Constructor for moving elements from a TMap with a different SetAllocator */
 	template<typename OtherSetAllocator>
-	TMultiMap(TMultiMap<KeyType, ValueType, OtherSetAllocator, KeyFuncs>&& Other)
+	[[nodiscard]] TMultiMap(TMultiMap<KeyType, ValueType, OtherSetAllocator, KeyFuncs>&& Other)
 		: Super(MoveTemp(Other))
 	{
 	}
 
 	/** Constructor for copying elements from a TMap with a different SetAllocator */
 	template<typename OtherSetAllocator>
-	TMultiMap(const TMultiMap<KeyType, ValueType, OtherSetAllocator, KeyFuncs>& Other)
+	[[nodiscard]] TMultiMap(const TMultiMap<KeyType, ValueType, OtherSetAllocator, KeyFuncs>& Other)
 		: Super(Other)
 	{
 	}
 
 	/** Constructor which gets its elements from a native initializer list */
-	TMultiMap(std::initializer_list<TPairInitializer<const KeyType&, const ValueType&>> InitList)
+	[[nodiscard]] TMultiMap(std::initializer_list<TPairInitializer<const KeyType&, const ValueType&>> InitList)
 	{
 		this->Reserve((int32)InitList.size());
 		for (const TPairInitializer<const KeyType&, const ValueType&>& Element : InitList)
@@ -1485,11 +1525,11 @@ public:
 	constexpr static bool bHasIntrusiveUnsetOptionalState = true;
 	using IntrusiveUnsetOptionalStateType = TMultiMap;
 
-	explicit TMultiMap(FIntrusiveUnsetOptionalState Tag)
+	[[nodiscard]] explicit TMultiMap(FIntrusiveUnsetOptionalState Tag)
 		: Super(Tag)
 	{
 	}
-	bool operator==(FIntrusiveUnsetOptionalState Tag) const
+	[[nodiscard]] bool operator==(FIntrusiveUnsetOptionalState Tag) const
 	{
 		return Super::operator==(Tag);
 	}
@@ -1682,7 +1722,7 @@ public:
 	 * @return If the map contains a matching association, a pointer to the value in the map is returned.  Otherwise nullptr is returned.
 	 *			The pointer is only valid as long as the map isn't changed.
 	 */
-	FORCEINLINE const ValueType* FindPair(KeyInitType Key, ValueInitType Value) const
+	[[nodiscard]] FORCEINLINE const ValueType* FindPair(KeyInitType Key, ValueInitType Value) const
 	{
 		return const_cast<TMultiMap*>(this)->FindPair(Key, Value);
 	}
@@ -1695,7 +1735,7 @@ public:
 	 * @return If the map contains a matching association, a pointer to the value in the map is returned.  Otherwise nullptr is returned.
 	 *			The pointer is only valid as long as the map isn't changed.
 	 */
-	ValueType* FindPair(KeyInitType Key, ValueInitType Value)
+	[[nodiscard]] ValueType* FindPair(KeyInitType Key, ValueInitType Value)
 	{
 		// Iterate over pairs with a matching key.
 		for (typename Super::ElementSetType::TKeyIterator It(Super::Pairs, Key); It; ++It)
@@ -1711,7 +1751,7 @@ public:
 	}
 
 	/** Returns the number of values within this map associated with the specified key */
-	int32 Num(KeyInitType Key) const
+	[[nodiscard]] int32 Num(KeyInitType Key) const
 	{
 		// Iterate over pairs with a matching key.
 		int32 NumMatchingPairs = 0;
@@ -1723,7 +1763,7 @@ public:
 	}
 
 	// Since we implement an overloaded Num() function in TMultiMap, we need to reimplement TMapBase::Num to make it visible.
-	FORCEINLINE int32 Num() const
+	[[nodiscard]] FORCEINLINE int32 Num() const
 	{
 		return Super::Num();
 	}
@@ -1806,7 +1846,7 @@ class TScriptMap
 	using DerivedType = std::conditional_t<std::is_void_v<InDerivedType>, TScriptMap, InDerivedType>;
 
 public:
-	static FScriptMapLayout GetScriptLayout(int32 KeySize, int32 KeyAlignment, int32 ValueSize, int32 ValueAlignment)
+	[[nodiscard]] static FScriptMapLayout GetScriptLayout(int32 KeySize, int32 KeyAlignment, int32 ValueSize, int32 ValueAlignment)
 	{
 		FScriptMapLayout Result;
 
@@ -1821,7 +1861,7 @@ public:
 		return Result;
 	}
 
-	TScriptMap()
+	[[nodiscard]] TScriptMap()
 	{
 	}
 
@@ -1831,11 +1871,11 @@ public:
 	constexpr static bool bHasIntrusiveUnsetOptionalState = true;
 	using IntrusiveUnsetOptionalStateType = TScriptMap;
 
-	explicit TScriptMap(FIntrusiveUnsetOptionalState Tag)
+	[[nodiscard]] explicit TScriptMap(FIntrusiveUnsetOptionalState Tag)
 		: Pairs(Tag)
 	{
 	}
-	bool operator==(FIntrusiveUnsetOptionalState Tag) const
+	[[nodiscard]] bool operator==(FIntrusiveUnsetOptionalState Tag) const
 	{
 		return Pairs == Tag;
 	}
@@ -1843,38 +1883,38 @@ public:
 	// End - intrusive TOptional<TScriptMap> state //
 	/////////////////////////////////////////////////
 
-	bool IsValidIndex(int32 Index) const
+	[[nodiscard]] bool IsValidIndex(int32 Index) const
 	{
 		return Pairs.IsValidIndex(Index);
 	}
 
-	bool IsEmpty() const
+	[[nodiscard]] bool IsEmpty() const
 	{
 		return Pairs.IsEmpty();
 	}
 
-	int32 Num() const
+	[[nodiscard]] int32 Num() const
 	{
 		return Pairs.Num();
 	}
 
-	int32 NumUnchecked() const
+	[[nodiscard]] int32 NumUnchecked() const
 	{
 		return Pairs.NumUnchecked();
 	}
 
 	/** @return The max valid index of the elements in the sparse storage. */
-	int32 GetMaxIndex() const
+	[[nodiscard]] int32 GetMaxIndex() const
 	{
 		return Pairs.GetMaxIndex();
 	}
 
-	void* GetData(int32 Index, const FScriptMapLayout& Layout)
+	[[nodiscard]] void* GetData(int32 Index, const FScriptMapLayout& Layout)
 	{
 		return Pairs.GetData(Index, Layout.SetLayout);
 	}
 
-	const void* GetData(int32 Index, const FScriptMapLayout& Layout) const
+	[[nodiscard]] const void* GetData(int32 Index, const FScriptMapLayout& Layout) const
 	{
 		return Pairs.GetData(Index, Layout.SetLayout);
 	}
@@ -1915,7 +1955,7 @@ public:
 	}
 	
 	/** Finds the associated key, value from hash of Key, rather than linearly searching */
-	int32 FindPairIndex(const void* Key, const FScriptMapLayout& MapLayout, TFunctionRef<uint32(const void*)> GetKeyHash, TFunctionRef<bool(const void*, const void*)> KeyEqualityFn) const
+	[[nodiscard]] int32 FindPairIndex(const void* Key, const FScriptMapLayout& MapLayout, TFunctionRef<uint32(const void*)> GetKeyHash, TFunctionRef<bool(const void*, const void*)> KeyEqualityFn) const
 	{
 		if (Pairs.Num())
 		{
@@ -1937,7 +1977,7 @@ public:
 	}
 
 	/** Finds the associated value from hash of Key, rather than linearly searching */
-	uint8* FindValue(const void* Key, const FScriptMapLayout& MapLayout, TFunctionRef<uint32(const void*)> GetKeyHash, TFunctionRef<bool(const void*, const void*)> KeyEqualityFn)
+	[[nodiscard]] uint8* FindValue(const void* Key, const FScriptMapLayout& MapLayout, TFunctionRef<uint32(const void*)> GetKeyHash, TFunctionRef<bool(const void*, const void*)> KeyEqualityFn)
 	{
 		int32 FoundIndex = FindPairIndex(Key, MapLayout, GetKeyHash, KeyEqualityFn);
 		if (FoundIndex != INDEX_NONE)
@@ -2117,7 +2157,7 @@ struct TMapPrivateFriend
 	}
 
 	template <typename KeyType, typename ValueType, typename SetAllocator, typename KeyFuncs>
-	static bool LegacyCompareEqual(const TMapBase<KeyType, ValueType, SetAllocator, KeyFuncs>& A, const TMapBase<KeyType, ValueType, SetAllocator, KeyFuncs>& B)
+	[[nodiscard]] static bool LegacyCompareEqual(const TMapBase<KeyType, ValueType, SetAllocator, KeyFuncs>& A, const TMapBase<KeyType, ValueType, SetAllocator, KeyFuncs>& B)
 	{
 		return TSetPrivateFriend::LegacyCompareEqual(A.Pairs, B.Pairs);
 	}
@@ -2139,12 +2179,12 @@ FORCEINLINE void operator<<(FStructuredArchive::FSlot Slot, TMapBase<KeyType, Va
 
 // Legacy comparison operators.  Note that these also test whether the map's key-value pairs were added in the same order!
 template <typename KeyType, typename ValueType, typename SetAllocator, typename KeyFuncs>
-bool LegacyCompareEqual(const TMapBase<KeyType, ValueType, SetAllocator, KeyFuncs>& A, const TMapBase<KeyType, ValueType, SetAllocator, KeyFuncs>& B)
+[[nodiscard]] bool LegacyCompareEqual(const TMapBase<KeyType, ValueType, SetAllocator, KeyFuncs>& A, const TMapBase<KeyType, ValueType, SetAllocator, KeyFuncs>& B)
 {
 	return TMapPrivateFriend::LegacyCompareEqual(A, B);
 }
 template <typename KeyType, typename ValueType, typename SetAllocator, typename KeyFuncs>
-bool LegacyCompareNotEqual(const TMapBase<KeyType, ValueType, SetAllocator, KeyFuncs>& A, const TMapBase<KeyType, ValueType, SetAllocator, KeyFuncs>& B)
+[[nodiscard]] bool LegacyCompareNotEqual(const TMapBase<KeyType, ValueType, SetAllocator, KeyFuncs>& A, const TMapBase<KeyType, ValueType, SetAllocator, KeyFuncs>& B)
 {
 	return !TMapPrivateFriend::LegacyCompareEqual(A, B);
 }
diff --git a/Engine/Source/Runtime/Core/Public/Containers/MapBuilder.h b/Engine/Source/Runtime/Core/Public/Containers/MapBuilder.h
--- a/Engine/Source/Runtime/Core/Public/Containers/MapBuilder.h
+++ b/Engine/Source/Runtime/Core/Public/Containers/MapBuilder.h
@@ -18,16 +18,20 @@ template<typename KeyType, typename ValueType, typename SetAllocator = FDefaultS
 public:
 
 	/** Default constructor. */
-	TMapBuilder() { }
+	[[nodiscard]] TMapBuilder()
+	{
+	}
 
 	/**
 	 * Creates and initializes a new map builder from another map.
 	 *
 	 * @param InMap The map to copy.
 	 */
-	template<typename OtherAllocator> TMapBuilder( const TMap<KeyType, ValueType, OtherAllocator>& InMap )
+	template<typename OtherAllocator>
+	[[nodiscard]] TMapBuilder( const TMap<KeyType, ValueType, OtherAllocator>& InMap )
 		: Map(InMap)
-	{ }
+	{
+	}
 
 public:
 
@@ -65,7 +69,7 @@ public:
 	 *
 	 * @return A new map.
 	 */
-	TMap<KeyType, ValueType, SetAllocator> Build()
+	[[nodiscard]] TMap<KeyType, ValueType, SetAllocator> Build()
 	{
 		return Map;
 	}
@@ -75,7 +79,7 @@ public:
 	 *
 	 * @return A new map.
 	 */
-	operator TMap<KeyType, ValueType, SetAllocator>()
+	[[nodiscard]] operator TMap<KeyType, ValueType, SetAllocator>()
 	{
 		return Build();
 	}
diff --git a/Engine/Source/Runtime/Core/Public/Containers/MpscQueue.h b/Engine/Source/Runtime/Core/Public/Containers/MpscQueue.h
--- a/Engine/Source/Runtime/Core/Public/Containers/MpscQueue.h
+++ b/Engine/Source/Runtime/Core/Public/Containers/MpscQueue.h
@@ -21,7 +21,7 @@ public:
 
 	UE_NONCOPYABLE(TMpscQueue);
 
-	TMpscQueue()
+	[[nodiscard]] TMpscQueue()
 	{
 		FNode* Sentinel = ::new(AllocatorType::Malloc(sizeof(FNode), alignof(FNode))) FNode;
 		Head.store(Sentinel, std::memory_order_relaxed);
@@ -89,7 +89,7 @@ public:
 	// as there can be only one consumer, a consumer can safely "peek" the tail of the queue.
 	// returns a pointer to the tail if the queue is not empty, nullptr otherwise
 	// there's no overload with TOptional as it doesn't support references
-	ElementType* Peek() const
+	[[nodiscard]] ElementType* Peek() const
 	{
 		FNode* Next = Tail->Next.load(std::memory_order_acquire);
 
@@ -101,7 +101,7 @@ public:
 		return (ElementType*)&Next->Value;
 	}
 
-	bool IsEmpty() const
+	[[nodiscard]] bool IsEmpty() const
 	{
 		return Tail->Next.load(std::memory_order_acquire) == nullptr;
 	}
diff --git a/Engine/Source/Runtime/Core/Public/Containers/PagedArray.h b/Engine/Source/Runtime/Core/Public/Containers/PagedArray.h
--- a/Engine/Source/Runtime/Core/Public/Containers/PagedArray.h
+++ b/Engine/Source/Runtime/Core/Public/Containers/PagedArray.h
@@ -53,12 +53,12 @@ protected:
 	/**
 	 * Internal parameter constructor (used exclusively by the container).
 	 */
-	TIteratorBase(InPageType* Data, SizeType Offset)
+	[[nodiscard]] TIteratorBase(InPageType* Data, SizeType Offset)
 		: Data(Data), PageIndex(Offset / PageTraits::Capacity), PageOffset(Offset % PageTraits::Capacity)
 	{
 	}
 
-	FORCEINLINE ElementType& Dereference() const
+	[[nodiscard]] FORCEINLINE ElementType& Dereference() const
 	{
 		return Data[PageIndex][PageOffset];
 	}
@@ -72,7 +72,7 @@ protected:
 		}
 	}
 
-	FORCEINLINE bool Equals(const TIteratorBase& Other) const
+	[[nodiscard]] FORCEINLINE bool Equals(const TIteratorBase& Other) const
 	{
 		return Data == Other.Data && PageIndex == Other.PageIndex && PageOffset == Other.PageOffset;
 	}
@@ -173,17 +173,17 @@ public:
 		typename AllocatorType::ForAnyElementType>;
 
 private:
-	static constexpr SizeType GetPageIndex(SizeType Index)
+	[[nodiscard]] static constexpr SizeType GetPageIndex(SizeType Index)
 	{
 		return Index / PageTraits::Capacity;
 	}
 
-	static constexpr SizeType GetPageOffset(SizeType Index)
+	[[nodiscard]] static constexpr SizeType GetPageOffset(SizeType Index)
 	{
 		return Index % PageTraits::Capacity;
 	}
 
-	static constexpr SizeType NumRequiredPages(SizeType Count)
+	[[nodiscard]] static constexpr SizeType NumRequiredPages(SizeType Count)
 	{
 		const SizeType Div = Count / PageTraits::Capacity;
 		return Count % PageTraits::Capacity == 0 ? Div : Div + 1;
@@ -193,26 +193,26 @@ public:
 	using ConstIteratorType = UE::Core::PagedArray::Private::TIterator<const ElementType, const PageType, PageTraits>;
 	using IteratorType = UE::Core::PagedArray::Private::TIterator<ElementType, PageType, PageTraits>;
 
-	static constexpr SizeType MaxPerPage()
+	[[nodiscard]] static constexpr SizeType MaxPerPage()
 	{
 		return PageTraits::Capacity;
 	}
 
-	TPagedArray() = default;
+	[[nodiscard]] TPagedArray() = default;
 
-	TPagedArray(TPagedArray&& Other) : Pages(MoveTemp(Other.Pages)), Count(Other.Count)
+	[[nodiscard]] TPagedArray(TPagedArray&& Other) : Pages(MoveTemp(Other.Pages)), Count(Other.Count)
 	{
 		Other.Count = 0;
 	}
 
-	TPagedArray(const TPagedArray& Other) = default;
+	[[nodiscard]] TPagedArray(const TPagedArray& Other) = default;
 
-	TPagedArray(std::initializer_list<ElementType> InList)
+	[[nodiscard]] TPagedArray(std::initializer_list<ElementType> InList)
 	{
 		CopyToEmpty(InList.begin(), static_cast<SizeType>(InList.size()));
 	}
 
-	TPagedArray(const ElementType* InSource, SizeType InSize)
+	[[nodiscard]] TPagedArray(const ElementType* InSource, SizeType InSize)
 	{
 		check(InSource || !InSize);
 		CopyToEmpty(InSource, InSize);
@@ -239,13 +239,13 @@ public:
 		return *this;
 	}
 
-	FORCEINLINE const ElementType& operator[](SizeType Index) const
+	[[nodiscard]] FORCEINLINE const ElementType& operator[](SizeType Index) const
 	{
 		CheckValidIndex(Index);
 		return Pages[GetPageIndex(Index)][GetPageOffset(Index)];
 	}
 
-	FORCEINLINE ElementType& operator[](SizeType Index)
+	[[nodiscard]] FORCEINLINE ElementType& operator[](SizeType Index)
 	{
 		CheckValidIndex(Index);
 		return Pages[GetPageIndex(Index)][GetPageOffset(Index)];
@@ -671,25 +671,25 @@ public:
 	}
 
 	/** STL IteratorType model compliance methods.*/
-	FORCEINLINE ConstIteratorType begin() const
+	[[nodiscard]] FORCEINLINE ConstIteratorType begin() const
 	{
 		return GetIterator(0);
 	}
-	FORCEINLINE IteratorType begin()
+	[[nodiscard]] FORCEINLINE IteratorType begin()
 	{
 		return GetIterator(0);
 	}
-	FORCEINLINE ConstIteratorType end() const
+	[[nodiscard]] FORCEINLINE ConstIteratorType end() const
 	{
 		return GetIterator(Count);
 	}
-	FORCEINLINE IteratorType end()
+	[[nodiscard]] FORCEINLINE IteratorType end()
 	{
 		return GetIterator(Count);
 	}
 
 	// Friend operators
-	bool operator==(const TPagedArray& Right) const
+	[[nodiscard]] bool operator==(const TPagedArray& Right) const
 	{
 		if (Num() != Right.Num())
 		{
@@ -706,7 +706,7 @@ public:
 		return true;
 	}
 
-	FORCEINLINE bool operator!=(const TPagedArray& Right) const
+	[[nodiscard]] FORCEINLINE bool operator!=(const TPagedArray& Right) const
 	{
 		return !(*this == Right);
 	}
@@ -715,12 +715,12 @@ private:
 	TArray<PageType> Pages;
 	SizeType Count = 0;
 
-	FORCEINLINE ConstIteratorType GetIterator(SizeType Offset) const
+	[[nodiscard]] FORCEINLINE ConstIteratorType GetIterator(SizeType Offset) const
 	{
 		return ConstIteratorType(InPlace, Pages.GetData(), Offset);
 	}
 
-	FORCEINLINE IteratorType GetIterator(SizeType Offset)
+	[[nodiscard]] FORCEINLINE IteratorType GetIterator(SizeType Offset)
 	{
 		return IteratorType(InPlace, Pages.GetData(), Offset);
 	}
diff --git a/Engine/Source/Runtime/Core/Public/Containers/Queue.h b/Engine/Source/Runtime/Core/Public/Containers/Queue.h
--- a/Engine/Source/Runtime/Core/Public/Containers/Queue.h
+++ b/Engine/Source/Runtime/Core/Public/Containers/Queue.h
@@ -47,7 +47,7 @@ public:
 	using FElementType = T;
 
 	/** Default constructor. */
-	TQueue()
+	[[nodiscard]] TQueue()
 	{
 		Head = Tail = new TNode();
 	}
@@ -128,7 +128,7 @@ public:
 
 		if constexpr (Mode == EQueueMode::Mpsc)
 		{
-            OldHead = (TNode*)FPlatformAtomics::InterlockedExchangePtr((void**)&Head, NewNode);
+			OldHead = (TNode*)FPlatformAtomics::InterlockedExchangePtr((void**)&Head, NewNode);
 			TSAN_BEFORE(&OldHead->NextNode);
 			FPlatformAtomics::InterlockedExchangePtr((void**)&OldHead->NextNode, NewNode);
 		}
@@ -143,7 +143,7 @@ public:
 				FPlatformMisc::MemoryBarrier();
 			}
 
-            OldHead->NextNode = NewNode;
+			OldHead->NextNode = NewNode;
 		}
 
 		return true;
@@ -170,9 +170,9 @@ public:
 
 		if constexpr (Mode == EQueueMode::Mpsc)
 		{
-            OldHead = (TNode*)FPlatformAtomics::InterlockedExchangePtr((void**)&Head, NewNode);
+			OldHead = (TNode*)FPlatformAtomics::InterlockedExchangePtr((void**)&Head, NewNode);
 			TSAN_BEFORE(&OldHead->NextNode);
-            FPlatformAtomics::InterlockedExchangePtr((void**)&OldHead->NextNode, NewNode);
+			FPlatformAtomics::InterlockedExchangePtr((void**)&OldHead->NextNode, NewNode);
 		}
 		else
 		{
@@ -198,7 +198,7 @@ public:
 	 * @note To be called only from consumer thread.
 	 * @see Dequeue, Empty, Peek, Pop
 	 */
-	bool IsEmpty() const
+	[[nodiscard]] bool IsEmpty() const
 	{
 		return (Tail->NextNode == nullptr);
 	}
@@ -231,7 +231,7 @@ public:
 	 *
 	 * @return Pointer to the item, or nullptr if queue is empty
 	 */
-	FElementType* Peek()
+	[[nodiscard]] FElementType* Peek()
 	{
 		if (Tail->NextNode == nullptr)
 		{
@@ -241,7 +241,7 @@ public:
 		return &Tail->NextNode->Item;
 	}
 
-	FORCEINLINE const FElementType* Peek() const
+	[[nodiscard]] FORCEINLINE const FElementType* Peek() const
 	{
 		return const_cast<TQueue*>(this)->Peek();
 	}
@@ -287,21 +287,24 @@ private:
 		FElementType Item;
 
 		/** Default constructor. */
-		TNode()
+		[[nodiscard]] TNode()
 			: NextNode(nullptr)
-		{ }
+		{
+		}
 
 		/** Creates and initializes a new node. */
-		explicit TNode(const FElementType& InItem)
+		[[nodiscard]] explicit TNode(const FElementType& InItem)
 			: NextNode(nullptr)
 			, Item(InItem)
-		{ }
+		{
+		}
 
 		/** Creates and initializes a new node. */
-		explicit TNode(FElementType&& InItem)
+		[[nodiscard]] explicit TNode(FElementType&& InItem)
 			: NextNode(nullptr)
 			, Item(MoveTemp(InItem))
-		{ }
+		{
+		}
 	};
 
 	/** Holds a pointer to the head of the list. */
diff --git a/Engine/Source/Runtime/Core/Public/Containers/ResourceArray.h b/Engine/Source/Runtime/Core/Public/Containers/ResourceArray.h
--- a/Engine/Source/Runtime/Core/Public/Containers/ResourceArray.h
+++ b/Engine/Source/Runtime/Core/Public/Containers/ResourceArray.h
@@ -11,10 +11,10 @@ struct FResourceArrayUploadInterface
 	virtual ~FResourceArrayUploadInterface() {}
 
 	/** Returns a pointer to the resource data. */
-	virtual const void* GetResourceData() const = 0;
+	[[nodiscard]] virtual const void* GetResourceData() const = 0;
 
 	/** Returns size of resource data allocation */
-	virtual uint32 GetResourceDataSize() const = 0;
+	[[nodiscard]] virtual uint32 GetResourceDataSize() const = 0;
 
 	/** Called on non-UMA systems after the RHI has copied the resource data, and no longer needs the CPU's copy. */
 	virtual void Discard() = 0;
@@ -28,32 +28,32 @@ struct FResourceArrayUploadArrayView : public FResourceArrayUploadInterface
 
 	FResourceArrayUploadArrayView() = delete;
 
-	FResourceArrayUploadArrayView(const void* InData, uint32 InSizeInBytes)
+	[[nodiscard]] explicit FResourceArrayUploadArrayView(const void* InData, uint32 InSizeInBytes)
 		: Data(InData)
 		, SizeInBytes(InSizeInBytes)
 	{
 	}
 
 	template<typename ElementType>
-	FResourceArrayUploadArrayView(TConstArrayView<ElementType> View)
+	[[nodiscard]] explicit FResourceArrayUploadArrayView(TConstArrayView<ElementType> View)
 		: Data(View.GetData())
 		, SizeInBytes(View.Num() * View.GetTypeSize())
 	{
 	}
 
 	template<typename ElementType, typename AllocatorType>
-	FResourceArrayUploadArrayView(const TArray<ElementType, AllocatorType>& InArray)
+	[[nodiscard]] explicit FResourceArrayUploadArrayView(const TArray<ElementType, AllocatorType>& InArray)
 		: FResourceArrayUploadArrayView(MakeArrayView(InArray))
 	{
 	}
 
 	// FResourceArrayUploadInterface
-	virtual const void* GetResourceData() const final
+	[[nodiscard]] virtual const void* GetResourceData() const final
 	{
 		return Data;
 	}
 
-	virtual uint32 GetResourceDataSize() const final
+	[[nodiscard]] virtual uint32 GetResourceDataSize() const final
 	{
 		return SizeInBytes;
 	}
@@ -75,12 +75,12 @@ public:
 	/**
 	 * @return true if the resource array is static and shouldn't be modified
 	 */
-	virtual bool IsStatic() const = 0;
+	[[nodiscard]] virtual bool IsStatic() const = 0;
 
 	/**
 	 * @return true if the resource keeps a copy of its resource data after the RHI resource has been created
 	 */
-	virtual bool GetAllowCPUAccess() const = 0;
+	[[nodiscard]] virtual bool GetAllowCPUAccess() const = 0;
 
 	/** 
 	 * Sets whether the resource array will be accessed by CPU. 
@@ -101,12 +101,12 @@ public:
 	/** 
 	 * @return ptr to the resource memory which has been preallocated
 	 */
-	virtual const void* GetResourceBulkData() const = 0;
+	[[nodiscard]] virtual const void* GetResourceBulkData() const = 0;
 
 	/** 
 	 * @return size of resource memory
 	 */
-	virtual uint32 GetResourceBulkDataSize() const = 0;
+	[[nodiscard]] virtual uint32 GetResourceBulkDataSize() const = 0;
 
 	/**
 	 * Free memory after it has been used to initialize RHI resource 
@@ -123,7 +123,7 @@ public:
 	/**
 	 * @return the type of bulk data for special handling
 	 */
-	virtual EBulkDataType GetResourceType() const
+	[[nodiscard]] virtual EBulkDataType GetResourceType() const
 	{
 		return EBulkDataType::Default;
 	}
@@ -141,32 +141,32 @@ public:
 	 * @param MipIdx index for mip to retrieve
 	 * @return ptr to the offset in bulk memory for the given mip
 	 */
-	virtual void* GetMipData(int32 MipIdx) = 0;
+	[[nodiscard]] virtual void* GetMipData(int32 MipIdx) = 0;
 
 	/**
 	 * @return total number of mips stored in this resource
 	 */
-	virtual int32	GetNumMips() = 0;
+	[[nodiscard]] virtual int32	GetNumMips() = 0;
 
 	/** 
 	 * @return width of texture stored in this resource
 	 */
-	virtual int32 GetSizeX() = 0;
+	[[nodiscard]] virtual int32 GetSizeX() = 0;
 
 	/** 
 	 * @return height of texture stored in this resource
 	 */
-	virtual int32 GetSizeY() = 0;
+	[[nodiscard]] virtual int32 GetSizeY() = 0;
 
 	/**
 	 * @return Whether the resource memory is properly allocated or not.
 	 */
-	virtual bool IsValid() = 0;
+	[[nodiscard]] virtual bool IsValid() = 0;
 
 	/**
 	 * @return Whether the resource memory has an async allocation request and it's been completed.
 	 */
-	virtual bool HasAsyncAllocationCompleted() const = 0;
+	[[nodiscard]] virtual bool HasAsyncAllocationCompleted() const = 0;
 
 	/** Blocks the calling thread until the allocation has been completed. */
 	virtual void FinishAsyncAllocation() = 0;
