Added [[nodiscard]] to relevant container functions.

diff --git a/Engine/Source/Runtime/Core/Public/Containers/RingBuffer.h b/Engine/Source/Runtime/Core/Public/Containers/RingBuffer.h
--- a/Engine/Source/Runtime/Core/Public/Containers/RingBuffer.h
+++ b/Engine/Source/Runtime/Core/Public/Containers/RingBuffer.h
@@ -15,7 +15,7 @@ template< typename ContainerType, typename ElementType, typename SizeType>
 class TRingBufferIterator
 {
 public:
-	TRingBufferIterator(ContainerType& InContainer, SizeType StartIndex = 0)
+	[[nodiscard]] TRingBufferIterator(ContainerType& InContainer, SizeType StartIndex = 0)
 		: Container(InContainer)
 		, Index(StartIndex)
 	{
@@ -71,24 +71,24 @@ public:
 		return Tmp -= Offset;
 	}
 
-	FORCEINLINE ElementType& operator* () const
+	[[nodiscard]] FORCEINLINE ElementType& operator*() const
 	{
 		return Container[Index];
 	}
 
-	FORCEINLINE ElementType* operator->() const
+	[[nodiscard]] FORCEINLINE ElementType* operator->() const
 	{
 		return &Container[Index];
 	}
 
 	/** conversion to "bool" returning true if the iterator has not reached the last element. */
-	FORCEINLINE explicit operator bool() const
+	[[nodiscard]] FORCEINLINE explicit operator bool() const
 	{
 		return Container.IsValidIndex(Index);
 	}
 
 	/** Returns an index to the current element. */
-	SizeType GetIndex() const
+	[[nodiscard]] SizeType GetIndex() const
 	{
 		return Index;
 	}
@@ -112,8 +112,8 @@ public:
 		--Index;
 	}
 
-	FORCEINLINE bool operator==(const TRingBufferIterator& Rhs) const { return &Container == &Rhs.Container && Index == Rhs.Index; }
-	FORCEINLINE bool operator!=(const TRingBufferIterator& Rhs) const { return &Container != &Rhs.Container || Index != Rhs.Index; }
+	[[nodiscard]] FORCEINLINE bool operator==(const TRingBufferIterator& Rhs) const { return &Container == &Rhs.Container && Index == Rhs.Index; }
+	[[nodiscard]] FORCEINLINE bool operator!=(const TRingBufferIterator& Rhs) const { return &Container != &Rhs.Container || Index != Rhs.Index; }
 
 private:
 
@@ -161,7 +161,7 @@ private:
 public:
 
 	/** Construct Empty Queue with capacity 0. */
-	TRingBuffer()
+	[[nodiscard]] TRingBuffer()
 		: IndexMask(static_cast<StorageModuloType>(-1))
 		, Front(0u)
 		, AfterBack(0u)
@@ -169,14 +169,14 @@ public:
 	}
 
 	/** Construct Empty Queue with the given initial requested capacity. */
-	explicit TRingBuffer(SizeType InitialCapacity)
+	[[nodiscard]] explicit TRingBuffer(SizeType InitialCapacity)
 		: TRingBuffer()
 	{
 		Reserve(InitialCapacity);
 	}
 
 	/** Construct a Queue with initial state (from front to back) equal to the given list. */
-	TRingBuffer(std::initializer_list<ElementType> InitList)
+	[[nodiscard]] TRingBuffer(std::initializer_list<ElementType> InitList)
 		: TRingBuffer(static_cast<SizeType>(InitList.size()))
 	{
 		for (const ElementType& Element : InitList)
@@ -185,7 +185,7 @@ public:
 		}
 	}
 
-	TRingBuffer(TRingBuffer&& Other)
+	[[nodiscard]] TRingBuffer(TRingBuffer&& Other)
 		: TRingBuffer()
 	{
 		*this = MoveTemp(Other);
@@ -205,7 +205,7 @@ public:
 		return *this;
 	}
 
-	TRingBuffer(const TRingBuffer& Other)
+	[[nodiscard]] TRingBuffer(const TRingBuffer& Other)
 		:TRingBuffer()
 	{
 		*this = Other;
@@ -227,19 +227,19 @@ public:
 	}
 
 	/** Returns true if the RingBuffer is empty. */
-	bool IsEmpty() const
+	[[nodiscard]] bool IsEmpty() const
 	{
 		return AfterBack == Front;
 	}
 
 	/** Gets the number of elements in the RingBuffer. */
-	IndexType Num() const
+	[[nodiscard]] IndexType Num() const
 	{
 		return AfterBack - Front;
 	}
 
 	/** Current allocated Capacity, note this will always be a power of two, or the special case 0. */
-	IndexType Max() const
+	[[nodiscard]] IndexType Max() const
 	{
 		return static_cast<SizeType>(IndexMask + 1);
 	}
@@ -288,7 +288,7 @@ public:
 	}
 
 	/** Add a new element after the back pointer of the RingBuffer, resizing if necessary.  The new element is move constructed from the argument. Returns a reference to the added element. */
-	ElementType& Add_GetRef(ElementType&& Element)
+	[[nodiscard]] ElementType& Add_GetRef(ElementType&& Element)
 	{
 		ElementType& Result = AddUninitialized_GetRef();
 		::new ((void*)&Result) ElementType(MoveTempIfPossible(Element));
@@ -305,7 +305,7 @@ public:
 	}
 
 	/** Add a new element after the back pointer of the RingBuffer, resizing if necessary.  The new element is copy constructed from the argument. Returns a reference to the added element. */
-	ElementType& Add_GetRef(const ElementType& Element)
+	[[nodiscard]] ElementType& Add_GetRef(const ElementType& Element)
 	{
 		IndexType ResultIndex = AddUninitialized();
 		ElementType& Result = GetAtIndexNoCheck(ResultIndex);
@@ -325,7 +325,7 @@ public:
 
 	/** Add a new element after the back pointer of the RingBuffer, resizing if necessary.  The new element is constructed from the given arguments. Returns a reference to the added element. */
 	template <typename... ArgsType>
-	ElementType& Emplace_GetRef(ArgsType&&... Args)
+	[[nodiscard]] ElementType& Emplace_GetRef(ArgsType&&... Args)
 	{
 		ElementType& Result = AddUninitialized_GetRef();
 		::new ((void*)&Result) ElementType(Forward<ArgsType>(Args)...);
@@ -345,7 +345,7 @@ public:
 	}
 
 	/** Add a new element after the back pointer of the RingBuffer, resizing if necessary.  The constructor is not called on the new element and its values in memory are arbitrary. Returns a reference to the added element. */
-	ElementType& AddUninitialized_GetRef()
+	[[nodiscard]] ElementType& AddUninitialized_GetRef()
 	{
 		return GetAtIndexNoCheck(AddUninitialized());
 	}
@@ -400,7 +400,7 @@ public:
 	}
 
 	/** Add a new element before the front pointer of the RingBuffer, resizing if necessary.  The new element is move constructed from the argument. Returns a reference to the added element. */
-	ElementType& AddFront_GetRef(ElementType&& Element)
+	[[nodiscard]] ElementType& AddFront_GetRef(ElementType&& Element)
 	{
 		ElementType& Result = AddFrontUninitialized_GetRef();
 		::new ((void*)&Result) ElementType(MoveTempIfPossible(Element));
@@ -417,7 +417,7 @@ public:
 	}
 
 	/** Add a new element before the front pointer of the RingBuffer, resizing if necessary.  The new element is copy constructed from the argument. Returns a reference to the added element. */
-	ElementType& AddFront_GetRef(const ElementType& Element)
+	[[nodiscard]] ElementType& AddFront_GetRef(const ElementType& Element)
 	{
 		IndexType ResultIndex = AddFrontUninitialized();
 		ElementType& Result = GetAtIndexNoCheck(ResultIndex);
@@ -437,7 +437,7 @@ public:
 
 	/** Add a new element before the front pointer of the RingBuffer, resizing if necessary.  The new element is constructed from the given arguments. Returns a reference to the added element. */
 	template <typename... ArgsType>
-	ElementType& EmplaceFront_GetRef(ArgsType&&... Args)
+	[[nodiscard]] ElementType& EmplaceFront_GetRef(ArgsType&&... Args)
 	{
 		ElementType& Result = AddFrontUninitialized_GetRef();
 		::new ((void*)&Result) ElementType(Forward<ArgsType>(Args)...);
@@ -457,31 +457,31 @@ public:
 	}
 
 	/** Add a new element before the front pointer of the RingBuffer, resizing if necessary.  The constructor is not called on the new element and its values in memory are arbitrary. Returns a reference to the added element. */
-	ElementType& AddFrontUninitialized_GetRef()
+	[[nodiscard]] ElementType& AddFrontUninitialized_GetRef()
 	{
 		return GetAtIndexNoCheck(AddFrontUninitialized());
 	}
 
 	/** Return a reference to the element at the front pointer of the RingBuffer.  Invalid to call on an empty RingBuffer. */
-	ElementType& First()
+	[[nodiscard]] ElementType& First()
 	{
 		return (*this)[0];
 	}
 
 	/** Return a const reference to the element at the front pointer of the RingBuffer.  Invalid to call on an empty RingBuffer. */
-	const ElementType& First() const
+	[[nodiscard]] const ElementType& First() const
 	{
 		return (*this)[0];
 	}
 
 	/** Return a reference to the element at the back pointer of the RingBuffer.  Invalid to call on an empty RingBuffer. */
-	ElementType& Last()
+	[[nodiscard]] ElementType& Last()
 	{
 		return (*this)[Num() - 1];
 	}
 
 	/** Return a const reference to the element at the back pointer of the RingBuffer.  Invalid to call on an empty RingBuffer. */
-	const ElementType& Last() const
+	[[nodiscard]] const ElementType& Last() const
 	{
 		return (*this)[Num() - 1];
 	}
@@ -562,63 +562,63 @@ public:
 	}
 
 	/** begin iterator for ranged-for */
-	TIterator begin()
+	[[nodiscard]] TIterator begin()
 	{
 		return TIterator(*this);
 	}
 
 	/** begin iterator for const ranged-for */
-	TConstIterator begin() const
+	[[nodiscard]] TConstIterator begin() const
 	{
 		return TConstIterator(*this);
 	}
 
 	/** end iterator for ranged-for */
-	TIterator end()
+	[[nodiscard]] TIterator end()
 	{
 		return TIterator(*this, Num());
 	}
 
 	/** end iterator for const ranged-for */
-	TConstIterator end() const
+	[[nodiscard]] TConstIterator end() const
 	{
 		return TConstIterator(*this, Num());
 	}
 
 	/** Tests if index is valid, i.e. greater than or equal to zero, and less than the number of elements in the array. */
-	bool IsValidIndex(IndexType Index) const
+	[[nodiscard]] bool IsValidIndex(IndexType Index) const
 	{
 		return Index >= 0 && Index < Num();
 	}
 
 	/** Return a writable reference to the value at the given Index.  Invalid to call with an index less than 0 or greater than or equal to the number of elements in the RingBuffer. */
-	ElementType& operator[](IndexType Index)
+	[[nodiscard]] ElementType& operator[](IndexType Index)
 	{
 		RangeCheck(Index);
 		return GetAtIndexNoCheck(Index);
 	}
 
 	/** Return a const reference to the value at the given Index.  Invalid to call with an index less than 0 or greater than or equal to the number of elements in the RingBuffer. */
-	const ElementType& operator[](IndexType Index) const
+	[[nodiscard]] const ElementType& operator[](IndexType Index) const
 	{
 		RangeCheck(Index);
 		return GetAtIndexNoCheck(Index);
 	}
 
 	/** Unsafely return a writable reference to the value at the given Index.  Invalid to call (and does not warn) with an index less than 0 or greater than or equal to the number of elements in the RingBuffer. */
-	ElementType& GetAtIndexNoCheck(IndexType Index)
+	[[nodiscard]] ElementType& GetAtIndexNoCheck(IndexType Index)
 	{
 		return GetStorage()[(Front + Index) & IndexMask];
 	}
 
 	/** Unsafely return a const reference to the value at the given Index.  Invalid to call (and does not warn) with an index less than 0 or greater than or equal to the number of elements in the RingBuffer. */
-	const ElementType& GetAtIndexNoCheck(IndexType Index) const
+	[[nodiscard]] const ElementType& GetAtIndexNoCheck(IndexType Index) const
 	{
 		return GetStorage()[(Front + Index) & IndexMask];
 	}
 
 	/** Given a pointer to an Element anywhere in memory, return the index of the element in the RingBuffer, or INDEX_NONE if it is not present. */
-	IndexType ConvertPointerToIndex(const ElementType* Ptr) const
+	[[nodiscard]] IndexType ConvertPointerToIndex(const ElementType* Ptr) const
 	{
 		const ElementType* const Data = GetStorage();
 		const ElementType* const DataEnd = Data + Max();
@@ -739,7 +739,7 @@ public:
 	}
 
 	template <typename OtherAllocator>
-	bool operator==(const TRingBuffer<ElementType, OtherAllocator>& Other) const
+	[[nodiscard]] bool operator==(const TRingBuffer<ElementType, OtherAllocator>& Other) const
 	{
 		const IndexType LocalNum = Num();
 		if (Other.Num() != LocalNum)
@@ -757,7 +757,7 @@ public:
 	}
 
 	template <typename OtherAllocator>
-	bool operator!=(const TRingBuffer<ElementType, OtherAllocator>& Other) const
+	[[nodiscard]] bool operator!=(const TRingBuffer<ElementType, OtherAllocator>& Other) const
 	{
 		return !(*this == Other);
 	}
@@ -786,13 +786,13 @@ public:
 	 *
 	 * @returns Number of bytes allocated by this container.
 	 */
-	SIZE_T GetAllocatedSize(void) const
+	[[nodiscard]] SIZE_T GetAllocatedSize(void) const
 	{
 		return Max() * sizeof(ElementType);
 	}
 
 private:
-	static constexpr bool NeedsDestructElements()
+	[[nodiscard]] static constexpr bool NeedsDestructElements()
 	{
 		return !std::is_trivially_destructible_v<T>;
 	}
@@ -976,13 +976,13 @@ private:
 	}
 
 	/** Return a pointer to the underlying storage of the RingBuffer */
-	FORCEINLINE const ElementType* GetStorage() const
+	[[nodiscard]] FORCEINLINE const ElementType* GetStorage() const
 	{
 		return const_cast<TRingBuffer*>(this)->GetStorage();
 	}
 
 	/** Return a const pointer to the underlying storage of the RingBuffer */
-	ElementType* GetStorage()
+	[[nodiscard]] ElementType* GetStorage()
 	{
 		return (ElementType*)AllocatorInstance.GetAllocation();
 	}
diff --git a/Engine/Source/Runtime/Core/Public/Containers/ScriptArray.h b/Engine/Source/Runtime/Core/Public/Containers/ScriptArray.h
--- a/Engine/Source/Runtime/Core/Public/Containers/ScriptArray.h
+++ b/Engine/Source/Runtime/Core/Public/Containers/ScriptArray.h
@@ -21,29 +21,29 @@ class TScriptArray
 {
 public:
 
-	FORCEINLINE void* GetData()
+	[[nodiscard]] FORCEINLINE void* GetData()
 	{
 		return this->GetAllocation();
 	}
-	FORCEINLINE const void* GetData() const
+	[[nodiscard]] FORCEINLINE const void* GetData() const
 	{
 		return this->GetAllocation();
 	}
-	FORCEINLINE bool IsValidIndex(int32 i) const
+	[[nodiscard]] FORCEINLINE bool IsValidIndex(int32 i) const
 	{
 		return i>=0 && i<ArrayNum;
 	}
-	bool IsEmpty() const
+	[[nodiscard]] bool IsEmpty() const
 	{
 		return ArrayNum == 0;
 	}
-	FORCEINLINE int32 Num() const
+	[[nodiscard]] FORCEINLINE int32 Num() const
 	{
 		checkSlow(ArrayNum>=0);
 		checkSlow(ArrayMax>=ArrayNum);
 		return ArrayNum;
 	}
-	FORCEINLINE int32 NumUnchecked() const
+	[[nodiscard]] FORCEINLINE int32 NumUnchecked() const
 	{
 		return ArrayNum;
 	}
@@ -168,8 +168,8 @@ public:
 			);
 	}
 	TScriptArray()
-	:   ArrayNum( 0 )
-	,	ArrayMax( 0 )
+		: ArrayNum( 0 )
+		, ArrayMax( 0 )
 	{
 	}
 	void CountBytes( FArchive& Ar, int32 NumBytesPerElement  ) const
@@ -183,7 +183,7 @@ public:
 	/**
 	 * Returns the amount of slack in this array in elements.
 	 */
-	FORCEINLINE int32 GetSlack() const
+	[[nodiscard]] FORCEINLINE int32 GetSlack() const
 	{
 		return ArrayMax - ArrayNum;
 	}
@@ -225,17 +225,16 @@ public:
 	{
 		Remove(Index, Count, NumBytesPerElement, AlignmentOfElement, bAllowShrinking ? EAllowShrinking::Yes : EAllowShrinking::No);
 	}
-	SIZE_T GetAllocatedSize(int32 NumBytesPerElement) const
+	[[nodiscard]] SIZE_T GetAllocatedSize(int32 NumBytesPerElement) const
 	{
 		return ((const typename AllocatorType::ForAnyElementType*)this)->GetAllocatedSize(ArrayMax, NumBytesPerElement);
 	}
 
 protected:
 
-	TScriptArray( int32 InNum, int32 NumBytesPerElement, uint32 AlignmentOfElement )
-	:   ArrayNum( 0 )
-	,	ArrayMax( InNum )
-
+	[[nodiscard]] TScriptArray( int32 InNum, int32 NumBytesPerElement, uint32 AlignmentOfElement )
+		: ArrayNum( 0 )
+		, ArrayMax( InNum )
 	{
 		if (ArrayMax)
 		{
@@ -245,8 +244,8 @@ protected:
 
 		SlackTrackerNumChanged();
 	}
-	int32	  ArrayNum;
-	int32	  ArrayMax;
+	int32 ArrayNum;
+	int32 ArrayMax;
 
 	FORCENOINLINE void ResizeInit(int32 NumBytesPerElement, uint32 AlignmentOfElement)
 	{
@@ -298,13 +297,13 @@ public:
 	constexpr static bool bHasIntrusiveUnsetOptionalState = true;
 	using IntrusiveUnsetOptionalStateType = TScriptArray;
 
-	explicit TScriptArray(FIntrusiveUnsetOptionalState Tag)
+	[[nodiscard]] explicit TScriptArray(FIntrusiveUnsetOptionalState Tag)
 		: ArrayNum(0)
 		, ArrayMax(-1)
 	{
 		// Use ArrayMax == -1 as our intrusive state so that the destructor still works without change, as it doesn't use ArrayMax.
 	}
-	bool operator==(FIntrusiveUnsetOptionalState Tag) const
+	[[nodiscard]] bool operator==(FIntrusiveUnsetOptionalState Tag) const
 	{
 		return ArrayMax == -1;
 	}
@@ -325,7 +324,7 @@ class FScriptArray : public TScriptArray<FHeapAllocator>
 	using Super = TScriptArray<FHeapAllocator>;
 
 public:
-	FScriptArray() = default;
+	[[nodiscard]] FScriptArray() = default;
 
 	void MoveAssign(FScriptArray& Other, int32 NumBytesPerElement, uint32 AlignmentOfElement)
 	{
@@ -333,7 +332,7 @@ public:
 	}
 
 protected:
-	FScriptArray(int32 InNum, int32 NumBytesPerElement, uint32 AlignmentOfElement)
+	[[nodiscard]] FScriptArray(int32 InNum, int32 NumBytesPerElement, uint32 AlignmentOfElement)
 		: TScriptArray<FHeapAllocator>(InNum, NumBytesPerElement, AlignmentOfElement)
 	{
 	}
@@ -343,7 +342,7 @@ public:
 	// Start - intrusive TOptional<FScriptArray> state //
 	/////////////////////////////////////////////////////
 	using IntrusiveUnsetOptionalStateType = FScriptArray;
-	explicit FScriptArray(FIntrusiveUnsetOptionalState Tag)
+	[[nodiscard]] explicit FScriptArray(FIntrusiveUnsetOptionalState Tag)
 		: TScriptArray(Tag)
 	{
 	}
@@ -353,7 +352,7 @@ public:
 
 	// These should really be private, because they shouldn't be called, but there's a bunch of code
 	// that needs to be fixed first.
-	FScriptArray(const FScriptArray&) { check(false); }
+	[[nodiscard]] FScriptArray(const FScriptArray&) { check(false); }
 	void operator=(const FScriptArray&) { check(false); }
 };
 
diff --git a/Engine/Source/Runtime/Core/Public/Containers/Set.h b/Engine/Source/Runtime/Core/Public/Containers/Set.h
--- a/Engine/Source/Runtime/Core/Public/Containers/Set.h
+++ b/Engine/Source/Runtime/Core/Public/Containers/Set.h
@@ -49,7 +49,7 @@ struct DefaultKeyFuncs : BaseKeyFuncs<ElementType,ElementType,bInAllowDuplicateK
 	/**
 	 * @return The key used to index the given element.
 	 */
-	static FORCEINLINE KeyInitType GetSetKey(ElementInitType Element)
+	[[nodiscard]] static FORCEINLINE KeyInitType GetSetKey(ElementInitType Element)
 	{
 		return Element;
 	}
@@ -57,7 +57,7 @@ struct DefaultKeyFuncs : BaseKeyFuncs<ElementType,ElementType,bInAllowDuplicateK
 	/**
 	 * @return True if the keys match.
 	 */
-	static FORCEINLINE bool Matches(KeyInitType A, KeyInitType B)
+	[[nodiscard]] static FORCEINLINE bool Matches(KeyInitType A, KeyInitType B)
 	{
 		return A == B;
 	}
@@ -66,20 +66,20 @@ struct DefaultKeyFuncs : BaseKeyFuncs<ElementType,ElementType,bInAllowDuplicateK
 	 * @return True if the keys match.
 	 */
 	template<typename ComparableKey>
-	static FORCEINLINE bool Matches(KeyInitType A, ComparableKey B)
+	[[nodiscard]] static FORCEINLINE bool Matches(KeyInitType A, ComparableKey B)
 	{
 		return A == B;
 	}
 
 	/** Calculates a hash index for a key. */
-	static FORCEINLINE uint32 GetKeyHash(KeyInitType Key)
+	[[nodiscard]] static FORCEINLINE uint32 GetKeyHash(KeyInitType Key)
 	{
 		return GetTypeHash(Key);
 	}
 
 	/** Calculates a hash index for a key. */
 	template<typename ComparableKey>
-	static FORCEINLINE uint32 GetKeyHash(ComparableKey Key)
+	[[nodiscard]] static FORCEINLINE uint32 GetKeyHash(ComparableKey Key)
 	{
 		return GetTypeHash(Key);
 	}
@@ -115,27 +115,27 @@ public:
 	FORCEINLINE FSetElementId() = default;
 
 	/** @return a boolean value representing whether the id is NULL. */
-	FORCEINLINE bool IsValidId() const
+	[[nodiscard]] FORCEINLINE bool IsValidId() const
 	{
 		return Index != INDEX_NONE;
 	}
 
 	/** Comparison operators. */
-	FORCEINLINE friend bool operator==(const FSetElementId& A,const FSetElementId& B)
+	[[nodiscard]] FORCEINLINE friend bool operator==(const FSetElementId& A,const FSetElementId& B)
 	{
 		return A.Index == B.Index;
 	}
-	FORCEINLINE friend bool operator!=(const FSetElementId& A, const FSetElementId& B)
+	[[nodiscard]] FORCEINLINE friend bool operator!=(const FSetElementId& A, const FSetElementId& B)
 	{
 		return A.Index != B.Index;
 	}
 
-	FORCEINLINE int32 AsInteger() const
+	[[nodiscard]] FORCEINLINE int32 AsInteger() const
 	{
 		return Index;
 	}
 
-	FORCEINLINE static FSetElementId FromInteger(int32 Integer)
+	[[nodiscard]] FORCEINLINE static FSetElementId FromInteger(int32 Integer)
 	{
 		return FSetElementId(Integer);
 	}
@@ -145,7 +145,7 @@ private:
 	int32 Index = INDEX_NONE;
 
 	/** Initialization constructor. */
-	FORCEINLINE explicit FSetElementId(int32 InIndex)
+	[[nodiscard]] FORCEINLINE explicit FSetElementId(int32 InIndex)
 		: Index(InIndex)
 	{
 	}
@@ -160,20 +160,20 @@ class TSetElementBase
 public:
 	typedef InElementType ElementType;
 
-	TSetElementBase() = default;
+	[[nodiscard]] TSetElementBase() = default;
 
 	/** Initialization constructor. */
 	template <
 		typename InitType
 		UE_REQUIRES(!std::is_same_v<TSetElementBase, std::decay_t<InitType>>)
 	>
-	explicit FORCEINLINE TSetElementBase(InitType&& InValue)
+	[[nodiscard]] explicit FORCEINLINE TSetElementBase(InitType&& InValue)
 		: Value(Forward<InitType>(InValue))
 	{
 	}
 
-	TSetElementBase(TSetElementBase&&) = default;
-	TSetElementBase(const TSetElementBase&) = default;
+	[[nodiscard]] TSetElementBase(TSetElementBase&&) = default;
+	[[nodiscard]] TSetElementBase(const TSetElementBase&) = default;
 	TSetElementBase& operator=(TSetElementBase&&) = default;
 	TSetElementBase& operator=(const TSetElementBase&) = default;
 
@@ -194,20 +194,20 @@ class TSetElementBase<InElementType, true>
 public:
 	typedef InElementType ElementType;
 
-	TSetElementBase() = default;
+	[[nodiscard]] TSetElementBase() = default;
 
 	/** Initialization constructor. */
 	template <
 		typename InitType
 		UE_REQUIRES(!std::is_same_v<TSetElementBase, std::decay_t<InitType>>)
 	>
-	explicit FORCEINLINE TSetElementBase(InitType&& InValue)
+	[[nodiscard]] explicit FORCEINLINE TSetElementBase(InitType&& InValue)
 		: Value(Forward<InitType>(InValue))
 	{
 	}
 
-	TSetElementBase(TSetElementBase&&) = default;
-	TSetElementBase(const TSetElementBase&) = default;
+	[[nodiscard]] TSetElementBase(TSetElementBase&&) = default;
+	[[nodiscard]] TSetElementBase(const TSetElementBase&) = default;
 	TSetElementBase& operator=(TSetElementBase&&) = default;
 	TSetElementBase& operator=(const TSetElementBase&) = default;
 
@@ -228,29 +228,29 @@ class TSetElement : public TSetElementBase<InElementType, THasTypeLayout<InEleme
 	using Super = TSetElementBase<InElementType, THasTypeLayout<InElementType>::Value>;
 public:
 	/** Default constructor. */
-	TSetElement() = default;
+	[[nodiscard]] TSetElement() = default;
 
 	/** Initialization constructor. */
 	template <
 		typename InitType
 		UE_REQUIRES(!std::is_same_v<TSetElement, std::decay_t<InitType>>)
 	>
-	explicit FORCEINLINE TSetElement(InitType&& InValue)
+	[[nodiscard]] explicit FORCEINLINE TSetElement(InitType&& InValue)
 		: Super(Forward<InitType>(InValue))
 	{
 	}
 
-	TSetElement(TSetElement&&) = default;
-	TSetElement(const TSetElement&) = default;
+	[[nodiscard]] TSetElement(TSetElement&&) = default;
+	[[nodiscard]] TSetElement(const TSetElement&) = default;
 	TSetElement& operator=(TSetElement&&) = default;
 	TSetElement& operator=(const TSetElement&) = default;
 
 	// Comparison operators
-	FORCEINLINE bool operator==(const TSetElement& Other) const
+	[[nodiscard]] FORCEINLINE bool operator==(const TSetElement& Other) const
 	{
 		return this->Value == Other.Value;
 	}
-	FORCEINLINE bool operator!=(const TSetElement& Other) const
+	[[nodiscard]] FORCEINLINE bool operator!=(const TSetElement& Other) const
 	{
 		return this->Value != Other.Value;
 	}
@@ -305,22 +305,22 @@ private:
 
 public:
 	/** Initialization constructor. */
-	FORCEINLINE TSet()
+	[[nodiscard]] FORCEINLINE TSet()
 	{
 	}
 
 	/** Copy constructor. */
-	FORCEINLINE TSet(const TSet& Copy)
+	[[nodiscard]] FORCEINLINE TSet(const TSet& Copy)
 	{
 		*this = Copy;
 	}
 
-	FORCEINLINE explicit TSet(TArrayView<const ElementType> InArrayView)
+	[[nodiscard]] FORCEINLINE explicit TSet(TArrayView<const ElementType> InArrayView)
 	{
 		Append(InArrayView);
 	}
 
-	FORCEINLINE explicit TSet(TArray<ElementType>&& InArray)
+	[[nodiscard]] FORCEINLINE explicit TSet(TArray<ElementType>&& InArray)
 	{
 		Append(MoveTemp(InArray));
 	}
@@ -337,11 +337,11 @@ public:
 	constexpr static bool bHasIntrusiveUnsetOptionalState = true;
 	using IntrusiveUnsetOptionalStateType = TSet;
 
-	explicit TSet(FIntrusiveUnsetOptionalState Tag)
+	[[nodiscard]] explicit TSet(FIntrusiveUnsetOptionalState Tag)
 		: Elements(Tag)
 	{
 	}
-	bool operator==(FIntrusiveUnsetOptionalState Tag) const
+	[[nodiscard]] bool operator==(FIntrusiveUnsetOptionalState Tag) const
 	{
 		return Elements == Tag;
 	}
@@ -380,14 +380,14 @@ private:
 
 public:
 	/** Initializer list constructor. */
-	TSet(std::initializer_list<ElementType> InitList)
+	[[nodiscard]] TSet(std::initializer_list<ElementType> InitList)
 		: HashSize(0)
 	{
 		Append(InitList);
 	}
 
 	/** Move constructor. */
-	TSet(TSet&& Other)
+	[[nodiscard]] TSet(TSet&& Other)
 		: HashSize(0)
 	{
 		this->Move(*this, Other);
@@ -406,7 +406,7 @@ public:
 
 	/** Constructor for moving elements from a TSet with a different SetAllocator */
 	template<typename OtherAllocator>
-	TSet(TSet<ElementType, KeyFuncs, OtherAllocator>&& Other)
+	[[nodiscard]] TSet(TSet<ElementType, KeyFuncs, OtherAllocator>&& Other)
 		: HashSize(0)
 	{
 		Append(MoveTemp(Other));
@@ -414,7 +414,7 @@ public:
 
 	/** Constructor for copying elements from a TSet with a different SetAllocator */
 	template<typename OtherAllocator>
-	TSet(const TSet<ElementType, KeyFuncs, OtherAllocator>& Other)
+	[[nodiscard]] TSet(const TSet<ElementType, KeyFuncs, OtherAllocator>& Other)
 		: HashSize(0)
 	{
 		Append(Other);
@@ -473,17 +473,17 @@ public:
 	}
 
 	/** Efficiently empties out the set but preserves all allocations and capacities */
-    void Reset()
-    {
+	void Reset()
+	{
 		if (Num() == 0)
 		{
 			return;
 		}
-    	
+
 		// Reset the elements array.
 		UnhashElements();
 		Elements.Reset();
-    }
+	}
 
 	/** Shrinks the set's element storage to avoid slack. */
 	FORCEINLINE void Shrink()
@@ -552,7 +552,7 @@ public:
 	 * Only returns the size of allocations made directly by the container, not the elements themselves.
 	 * @return number of bytes allocated by this container
 	 */
-	FORCEINLINE SIZE_T GetAllocatedSize( void ) const
+	[[nodiscard]] FORCEINLINE SIZE_T GetAllocatedSize( void ) const
 	{
 		return Elements.GetAllocatedSize() + Hash.GetAllocatedSize(HashSize, sizeof(FSetElementId));
 	}
@@ -570,18 +570,18 @@ public:
 	 * @returns True if the set is empty.
 	 * @see Num
 	 */
-	bool IsEmpty() const
+	[[nodiscard]] bool IsEmpty() const
 	{
 		return Elements.IsEmpty();
 	}
 
 	/** @return the number of elements. */
-	FORCEINLINE int32 Num() const
+	[[nodiscard]] FORCEINLINE int32 Num() const
 	{
 		return Elements.Num();
 	}
 
-	FORCEINLINE int32 GetMaxIndex() const
+	[[nodiscard]] FORCEINLINE int32 GetMaxIndex() const
 	{
 		return Elements.GetMaxIndex();
 	}
@@ -591,7 +591,7 @@ public:
 	 * @param Id - The element id to check.
 	 * @return true if the element identifier refers to a valid element in this set.
 	 */
-	FORCEINLINE bool IsValidId(FSetElementId Id) const
+	[[nodiscard]] FORCEINLINE bool IsValidId(FSetElementId Id) const
 	{
 		SizeType Index = Id.Index;
 		return Index != INDEX_NONE &&
@@ -601,13 +601,13 @@ public:
 	}
 
 	/** Accesses the identified element's value. Element must be valid (see @IsValidId). */
-	FORCEINLINE ElementType& operator[](FSetElementId Id)
+	[[nodiscard]] FORCEINLINE ElementType& operator[](FSetElementId Id)
 	{
 		return Elements[Id.Index].Value;
 	}
 
 	/** Accesses the identified element's value. Element must be valid (see @IsValidId). */
-	FORCEINLINE const ElementType& operator[](FSetElementId Id) const
+	[[nodiscard]] FORCEINLINE const ElementType& operator[](FSetElementId Id) const
 	{
 		return Elements[Id.Index].Value;
 	}
@@ -887,7 +887,7 @@ private:
 	 * @return The element id that matches the key and hash or an invalid element id
 	 */
 	template <typename ComparableKey>
-	SizeType FindIndexByHash(uint32 KeyHash, const ComparableKey& Key) const
+	[[nodiscard]] SizeType FindIndexByHash(uint32 KeyHash, const ComparableKey& Key) const
 	{
 		if (Elements.Num() == 0)
 		{
@@ -919,14 +919,14 @@ public:
 	 * Callers should not depend on particular patterns in the behaviour of this function.
 	 * @return A pointer to an arbitrary element, or nullptr if the container is empty.
 	 */
-	ElementType* FindArbitraryElement()
+	[[nodiscard]] ElementType* FindArbitraryElement()
 	{
 		// The goal of this function is to be fast, and so the implementation may be improved at any time even if it gives different results.
 
 		int32 Result = Elements.FindArbitraryElementIndex();
 		return (Result != INDEX_NONE) ? &Elements[Result].Value : nullptr;
 	}
-	const ElementType* FindArbitraryElement() const
+	[[nodiscard]] const ElementType* FindArbitraryElement() const
 	{
 		return const_cast<TSet*>(this)->FindArbitraryElement();
 	}
@@ -947,7 +947,7 @@ public:
 	 * @return The element id that matches the key and hash or an invalid element id
 	 */
 	template<typename ComparableKey>
-	FSetElementId FindIdByHash(uint32 KeyHash, const ComparableKey& Key) const
+	[[nodiscard]] FSetElementId FindIdByHash(uint32 KeyHash, const ComparableKey& Key) const
 	{
 		PRAGMA_DISABLE_DEPRECATION_WARNINGS
 		// Disable deprecations warnings to stop warnings being thrown by our check macro.
@@ -962,7 +962,7 @@ public:
 	 * @param Key - The key to search for.
 	 * @return A pointer to an element with the given key.  If no element in the set has the given key, this will return NULL.
 	 */
-	FORCEINLINE ElementType* Find(KeyInitType Key)
+	[[nodiscard]] FORCEINLINE ElementType* Find(KeyInitType Key)
 	{
 		SizeType ElementIndex = FindIndexByHash(KeyFuncs::GetKeyHash(Key), Key);
 		if (ElementIndex != INDEX_NONE)
@@ -980,7 +980,7 @@ public:
 	 * @param Key - The key to search for.
 	 * @return A const pointer to an element with the given key.  If no element in the set has the given key, this will return NULL.
 	 */
-	FORCEINLINE const ElementType* Find(KeyInitType Key) const
+	[[nodiscard]] FORCEINLINE const ElementType* Find(KeyInitType Key) const
 	{
 		return const_cast<TSet*>(this)->Find(Key);
 	}
@@ -991,7 +991,7 @@ public:
 	 * @return A pointer to the contained element or nullptr.
 	 */
 	template<typename ComparableKey>
-	ElementType* FindByHash(uint32 KeyHash, const ComparableKey& Key)
+	[[nodiscard]] ElementType* FindByHash(uint32 KeyHash, const ComparableKey& Key)
 	{
 		SizeType ElementIndex = FindIndexByHash(KeyHash, Key);
 		if (ElementIndex != INDEX_NONE)
@@ -1005,7 +1005,7 @@ public:
 	}
 
 	template<typename ComparableKey>
-	const ElementType* FindByHash(uint32 KeyHash, const ComparableKey& Key) const
+	[[nodiscard]] const ElementType* FindByHash(uint32 KeyHash, const ComparableKey& Key) const
 	{
 		return const_cast<TSet*>(this)->FindByHash(KeyHash, Key);
 	}
@@ -1087,7 +1087,7 @@ public:
 	 * @param Key - The key to check for.
 	 * @return true if the set contains an element with the given key.
 	 */
-	FORCEINLINE bool Contains(KeyInitType Key) const
+	[[nodiscard]] FORCEINLINE bool Contains(KeyInitType Key) const
 	{
 		return FindIndexByHash(KeyFuncs::GetKeyHash(Key), Key) != INDEX_NONE;
 	}
@@ -1098,7 +1098,7 @@ public:
 	 * @see	Class documentation section on ByHash() functions
 	 */
 	template<typename ComparableKey>
-	FORCEINLINE bool ContainsByHash(uint32 KeyHash, const ComparableKey& Key) const
+	[[nodiscard]] FORCEINLINE bool ContainsByHash(uint32 KeyHash, const ComparableKey& Key) const
 	{
 		PRAGMA_DISABLE_DEPRECATION_WARNINGS
 		// Disable deprecations warnings to stop warnings being thrown by our check macro.
@@ -1166,7 +1166,7 @@ public:
 		}
 	}
 
-	bool VerifyHashElementsKey(KeyInitType Key)
+	[[nodiscard]] bool VerifyHashElementsKey(KeyInitType Key) const
 	{
 		bool bResult=true;
 		if (Elements.Num())
@@ -1212,7 +1212,7 @@ public:
 	}
 
 	/** @return the intersection of two sets. (A AND B)*/
-	TSet Intersect(const TSet& OtherSet) const
+	[[nodiscard]] TSet Intersect(const TSet& OtherSet) const
 	{
 		const bool bOtherSmaller = (Num() > OtherSet.Num());
 		const TSet& A = (bOtherSmaller ? OtherSet : *this);
@@ -1232,7 +1232,7 @@ public:
 	}
 
 	/** @return the union of two sets. (A OR B)*/
-	TSet Union(const TSet& OtherSet) const
+	[[nodiscard]] TSet Union(const TSet& OtherSet) const
 	{
 		TSet Result;
 		Result.Reserve(Num() + OtherSet.Num()); // Worst case is 2 totally unique Sets
@@ -1249,7 +1249,7 @@ public:
 	}
 
 	/** @return the complement of two sets. (A not in B where A is this and B is Other)*/
-	TSet Difference(const TSet& OtherSet) const
+	[[nodiscard]] TSet Difference(const TSet& OtherSet) const
 	{
 		TSet Result;
 		Result.Reserve(Num()); // Worst case is no elements of this are in Other
@@ -1271,7 +1271,7 @@ public:
 	 * 
 	 * @return True if the other set is entirely included in this set, false if it is not
 	 */
-	bool Includes(const TSet<ElementType,KeyFuncs,Allocator>& OtherSet) const
+	[[nodiscard]] bool Includes(const TSet<ElementType,KeyFuncs,Allocator>& OtherSet) const
 	{
 		bool bIncludesSet = true;
 		if (OtherSet.Num() <= Num())
@@ -1294,7 +1294,7 @@ public:
 	}
 
 	/** @return a TArray of the elements */
-	TArray<ElementType> Array() const
+	[[nodiscard]] TArray<ElementType> Array() const
 	{
 		TArray<ElementType> Result;
 		Result.Reserve(Num());
@@ -1406,11 +1406,12 @@ private:
 		TDereferenceWrapper< ElementType, PREDICATE_CLASS > Predicate;
 
 	public:
-		FORCEINLINE FElementCompareClass( const PREDICATE_CLASS& InPredicate )
+		[[nodiscard]] FORCEINLINE FElementCompareClass( const PREDICATE_CLASS& InPredicate )
 			: Predicate( InPredicate )
-		{}
+		{
+		}
 
-		FORCEINLINE bool operator()( const SetElementType& A,const SetElementType& B ) const
+		[[nodiscard]] FORCEINLINE bool operator()( const SetElementType& A,const SetElementType& B ) const
 		{
 			return Predicate( A.Value, B.Value );
 		}
@@ -1465,7 +1466,7 @@ public:
 
 private:
 
-	FORCEINLINE FSetElementId& GetTypedHash(int32 HashIndex) const
+	[[nodiscard]] FORCEINLINE FSetElementId& GetTypedHash(int32 HashIndex) const
 	{
 		return ((FSetElementId*)Hash.GetAllocation())[HashIndex & (HashSize - 1)];
 	}
@@ -1518,7 +1519,7 @@ private:
 	 * @param AllowShrinking - If the hash is allowed to shrink.
 	 * @return true if the set should berehashed.
 	 */
-	FORCEINLINE bool ShouldRehash(int32 NumHashedElements, int32 DesiredHashSize, EAllowShrinking AllowShrinking) const
+	[[nodiscard]] FORCEINLINE bool ShouldRehash(int32 NumHashedElements, int32 DesiredHashSize, EAllowShrinking AllowShrinking) const
 	{
 		// If the hash hasn't been created yet, or is smaller than the desired hash size, rehash.
 		// If shrinking is allowed and the hash is bigger than the desired hash size, rehash.
@@ -1587,7 +1588,7 @@ private:
 			std::conditional_t<bRangedFor, typename ElementArrayType::TRangedForIterator,      typename ElementArrayType::TIterator     >
 		> ElementItType;
 
-		FORCEINLINE TBaseIterator(const ElementItType& InElementIt)
+		[[nodiscard]] FORCEINLINE TBaseIterator(const ElementItType& InElementIt)
 			: ElementIt(InElementIt)
 		{
 		}
@@ -1600,32 +1601,32 @@ private:
 		}
 
 		/** conversion to "bool" returning true if the iterator is valid. */
-		FORCEINLINE explicit operator bool() const
+		[[nodiscard]] FORCEINLINE explicit operator bool() const
 		{ 
 			return !!ElementIt; 
 		}
 		/** inverse of the "bool" operator */
-		FORCEINLINE bool operator !() const 
+		[[nodiscard]] FORCEINLINE bool operator!() const 
 		{
 			return !(bool)*this;
 		}
 
 		// Accessors.
-		FORCEINLINE FSetElementId GetId() const
+		[[nodiscard]] FORCEINLINE FSetElementId GetId() const
 		{
 			return FSetElementId(ElementIt.GetIndex());
 		}
-		FORCEINLINE ItElementType* operator->() const
+		[[nodiscard]] FORCEINLINE ItElementType* operator->() const
 		{
 			return &ElementIt->Value;
 		}
-		FORCEINLINE ItElementType& operator*() const
+		[[nodiscard]] FORCEINLINE ItElementType& operator*() const
 		{
 			return ElementIt->Value;
 		}
 
-		FORCEINLINE bool operator==(const TBaseIterator& Rhs) const { return ElementIt == Rhs.ElementIt; }
-		FORCEINLINE bool operator!=(const TBaseIterator& Rhs) const { return ElementIt != Rhs.ElementIt; }
+		[[nodiscard]] FORCEINLINE bool operator!=(const TBaseIterator& Rhs) const { return ElementIt != Rhs.ElementIt; }
+		[[nodiscard]] FORCEINLINE bool operator==(const TBaseIterator& Rhs) const { return ElementIt == Rhs.ElementIt; }
 
 		ElementItType ElementIt;
 	};
@@ -1648,7 +1649,7 @@ private:
 			>;
 
 		/** Initialization constructor. */
-		FORCEINLINE TBaseKeyIterator(SetType& InSet, KeyArgumentType InKey)
+		[[nodiscard]] FORCEINLINE TBaseKeyIterator(SetType& InSet, KeyArgumentType InKey)
 			: Set  (InSet)
 			, Key  (InKey) //-V1041
 			, Index(INDEX_NONE)
@@ -1685,12 +1686,12 @@ private:
 		}
 
 		/** conversion to "bool" returning true if the iterator is valid. */
-		FORCEINLINE explicit operator bool() const
+		[[nodiscard]] FORCEINLINE explicit operator bool() const
 		{ 
 			return Index != INDEX_NONE;
 		}
 		/** inverse of the "bool" operator */
-		FORCEINLINE bool operator !() const 
+		[[nodiscard]] FORCEINLINE bool operator!() const
 		{
 			return !(bool)*this;
 		}
@@ -1700,11 +1701,11 @@ private:
 		{
 			return FSetElementId(Index);
 		}
-		FORCEINLINE ItElementType* operator->() const
+		[[nodiscard]] FORCEINLINE ItElementType* operator->() const
 		{
 			return &Set.Elements[Index].Value;
 		}
-		FORCEINLINE ItElementType& operator*() const
+		[[nodiscard]] FORCEINLINE ItElementType& operator*() const
 		{
 			return Set.Elements[Index].Value;
 		}
@@ -1724,7 +1725,7 @@ public:
 		friend class TSet;
 
 	public:
-		FORCEINLINE TConstIterator(const TSet& InSet)
+		[[nodiscard]] FORCEINLINE TConstIterator(const TSet& InSet)
 			: TBaseIterator<true>(InSet.Elements.begin())
 		{
 		}
@@ -1736,7 +1737,7 @@ public:
 		friend class TSet;
 
 	public:
-		FORCEINLINE TIterator(TSet& InSet)
+		[[nodiscard]] FORCEINLINE TIterator(TSet& InSet)
 			: TBaseIterator<false>(InSet.Elements.begin())
 			, Set                 (InSet)
 		{
@@ -1764,7 +1765,7 @@ public:
 	public:
 		using KeyArgumentType = typename Super::KeyArgumentType;
 
-		FORCEINLINE TConstKeyIterator(const TSet& InSet, KeyArgumentType InKey)
+		[[nodiscard]] FORCEINLINE TConstKeyIterator(const TSet& InSet, KeyArgumentType InKey)
 			: Super(InSet, InKey)
 		{
 		}
@@ -1779,7 +1780,7 @@ public:
 	public:
 		using KeyArgumentType = typename Super::KeyArgumentType;
 
-		FORCEINLINE TKeyIterator(TSet& InSet, KeyArgumentType InKey)
+		[[nodiscard]] FORCEINLINE TKeyIterator(TSet& InSet, KeyArgumentType InKey)
 			: Super(InSet, InKey)
 		{
 		}
@@ -1793,13 +1794,13 @@ public:
 	};
 
 	/** Creates an iterator for the contents of this set */
-	FORCEINLINE TIterator CreateIterator()
+	[[nodiscard]] FORCEINLINE TIterator CreateIterator()
 	{
 		return TIterator(*this);
 	}
 
 	/** Creates a const iterator for the contents of this set */
-	FORCEINLINE TConstIterator CreateConstIterator() const
+	[[nodiscard]] FORCEINLINE TConstIterator CreateConstIterator() const
 	{
 		return TConstIterator(*this);
 	}
@@ -1811,10 +1812,10 @@ public:
 	 * DO NOT USE DIRECTLY
 	 * STL-like iterators to enable range-based for loop support.
 	 */
-	FORCEINLINE TRangedForIterator      begin()       { return TRangedForIterator     (Elements.begin()); }
-	FORCEINLINE TRangedForConstIterator begin() const { return TRangedForConstIterator(Elements.begin()); }
-	FORCEINLINE TRangedForIterator      end()         { return TRangedForIterator     (Elements.end());   }
-	FORCEINLINE TRangedForConstIterator end() const   { return TRangedForConstIterator(Elements.end());   }
+	[[nodiscard]] FORCEINLINE TRangedForIterator      begin()       { return TRangedForIterator     (Elements.begin()); }
+	[[nodiscard]] FORCEINLINE TRangedForConstIterator begin() const { return TRangedForConstIterator(Elements.begin()); }
+	[[nodiscard]] FORCEINLINE TRangedForIterator      end()         { return TRangedForIterator     (Elements.end());   }
+	[[nodiscard]] FORCEINLINE TRangedForConstIterator end() const   { return TRangedForConstIterator(Elements.end());   }
 
 	// Maps are deliberately prevented from being hashed or compared, because this would hide potentially major performance problems behind default operations.
 	friend uint32 GetTypeHash(const TSet& Set) = delete;
@@ -1868,7 +1869,7 @@ class TScriptSet
 	using DerivedType = std::conditional_t<std::is_void_v<InDerivedType>, TScriptSet, InDerivedType>;
 
 public:
-	static FScriptSetLayout GetScriptLayout(int32 ElementSize, int32 ElementAlignment)
+	[[nodiscard]] static FScriptSetLayout GetScriptLayout(int32 ElementSize, int32 ElementAlignment)
 	{
 		FScriptSetLayout Result;
 
@@ -1885,7 +1886,7 @@ public:
 		return Result;
 	}
 
-	TScriptSet()
+	[[nodiscard]] TScriptSet()
 		: HashSize(0)
 	{
 	}
@@ -1896,11 +1897,11 @@ public:
 	constexpr static bool bHasIntrusiveUnsetOptionalState = true;
 	using IntrusiveUnsetOptionalStateType = TScriptSet;
 
-	explicit TScriptSet(FIntrusiveUnsetOptionalState Tag)
+	[[nodiscard]] explicit TScriptSet(FIntrusiveUnsetOptionalState Tag)
 		: Elements(Tag)
 	{
 	}
-	bool operator==(FIntrusiveUnsetOptionalState Tag) const
+	[[nodiscard]] bool operator==(FIntrusiveUnsetOptionalState Tag) const
 	{
 		return Elements == Tag;
 	}
@@ -1908,37 +1909,37 @@ public:
 	// End - intrusive TOptional<TScriptSet> state //
 	/////////////////////////////////////////////////
 
-	bool IsValidIndex(int32 Index) const
+	[[nodiscard]] bool IsValidIndex(int32 Index) const
 	{
 		return Elements.IsValidIndex(Index);
 	}
 
-	bool IsEmpty() const
+	[[nodiscard]] bool IsEmpty() const
 	{
 		return Elements.IsEmpty();
 	}
 
-	int32 Num() const
+	[[nodiscard]] int32 Num() const
 	{
 		return Elements.Num();
 	}
 
-	int32 NumUnchecked() const
+	[[nodiscard]] int32 NumUnchecked() const
 	{
 		return Elements.NumUnchecked();
 	}
 
-	int32 GetMaxIndex() const
+	[[nodiscard]] int32 GetMaxIndex() const
 	{
 		return Elements.GetMaxIndex();
 	}
 
-	void* GetData(int32 Index, const FScriptSetLayout& Layout)
+	[[nodiscard]] void* GetData(int32 Index, const FScriptSetLayout& Layout)
 	{
 		return Elements.GetData(Index, Layout.SparseArrayLayout);
 	}
 
-	const void* GetData(int32 Index, const FScriptSetLayout& Layout) const
+	[[nodiscard]] const void* GetData(int32 Index, const FScriptSetLayout& Layout) const
 	{
 		return Elements.GetData(Index, Layout.SparseArrayLayout);
 	}
@@ -2053,7 +2054,7 @@ public:
 	}
 
 private:
-	int32 FindIndexImpl(const void* Element, const FScriptSetLayout& Layout, uint32 KeyHash, TFunctionRef<bool (const void*, const void*)> EqualityFn) const
+	[[nodiscard]] int32 FindIndexImpl(const void* Element, const FScriptSetLayout& Layout, uint32 KeyHash, TFunctionRef<bool (const void*, const void*)> EqualityFn) const
 	{
 		const int32  HashIndex = KeyHash & (HashSize - 1);
 
@@ -2073,7 +2074,7 @@ private:
 	}
 
 public:
-	int32 FindIndex(const void* Element, const FScriptSetLayout& Layout, TFunctionRef<uint32 (const void*)> GetKeyHash, TFunctionRef<bool (const void*, const void*)> EqualityFn) const
+	[[nodiscard]] int32 FindIndex(const void* Element, const FScriptSetLayout& Layout, TFunctionRef<uint32 (const void*)> GetKeyHash, TFunctionRef<bool (const void*, const void*)> EqualityFn) const
 	{
 		if (Elements.Num())
 		{
@@ -2084,7 +2085,7 @@ public:
 	}
 
 
-	int32 FindIndexByHash(const void* Element, const FScriptSetLayout& Layout, uint32 KeyHash, TFunctionRef<bool (const void*, const void*)> EqualityFn) const
+	[[nodiscard]] int32 FindIndexByHash(const void* Element, const FScriptSetLayout& Layout, uint32 KeyHash, TFunctionRef<bool (const void*, const void*)> EqualityFn) const
 	{
 		if (Elements.Num())
 		{
@@ -2164,17 +2165,17 @@ private:
 	HashType         Hash;
 	int32            HashSize;
 
-	FORCEINLINE FSetElementId& GetTypedHash(int32 HashIndex) const
+	[[nodiscard]] FORCEINLINE FSetElementId& GetTypedHash(int32 HashIndex) const
 	{
 		return ((FSetElementId*)Hash.GetAllocation())[HashIndex & (HashSize - 1)];
 	}
 
-	static FSetElementId& GetHashNextIdRef(const void* Element, const FScriptSetLayout& Layout)
+	[[nodiscard]] static FSetElementId& GetHashNextIdRef(const void* Element, const FScriptSetLayout& Layout)
 	{
 		return *(FSetElementId*)((uint8*)Element + Layout.HashNextIdOffset);
 	}
 
-	static int32& GetHashIndexRef(const void* Element, const FScriptSetLayout& Layout)
+	[[nodiscard]] static int32& GetHashIndexRef(const void* Element, const FScriptSetLayout& Layout)
 	{
 		return *(int32*)((uint8*)Element + Layout.HashIndexOffset);
 	}
@@ -2280,7 +2281,7 @@ struct TSetPrivateFriend
 
 	// Legacy comparison operators.  Note that these also test whether the set's elements were added in the same order!
 	template<typename ElementType, typename KeyFuncs,typename Allocator>
-	static bool LegacyCompareEqual(const TSet<ElementType, KeyFuncs, Allocator>& A, const TSet<ElementType, KeyFuncs, Allocator>& B)
+	[[nodiscard]] static bool LegacyCompareEqual(const TSet<ElementType, KeyFuncs, Allocator>& A, const TSet<ElementType, KeyFuncs, Allocator>& B)
 	{
 		return A.Elements == B.Elements;
 	}
@@ -2316,12 +2317,12 @@ void operator<<(FStructuredArchive::FSlot& Ar, TSet<ElementType, KeyFuncs, Alloc
 
 // Legacy comparison operators.  Note that these also test whether the set's elements were added in the same order!
 template<typename ElementType, typename KeyFuncs,typename Allocator>
-bool LegacyCompareEqual(const TSet<ElementType, KeyFuncs, Allocator>& A,const TSet<ElementType, KeyFuncs, Allocator>& B)
+[[nodiscard]] bool LegacyCompareEqual(const TSet<ElementType, KeyFuncs, Allocator>& A,const TSet<ElementType, KeyFuncs, Allocator>& B)
 {
 	return TSetPrivateFriend::LegacyCompareEqual(A, B);
 }
 template<typename ElementType, typename KeyFuncs,typename Allocator>
-bool LegacyCompareNotEqual(const TSet<ElementType, KeyFuncs, Allocator>& A,const TSet<ElementType, KeyFuncs, Allocator>& B)
+[[nodiscard]] bool LegacyCompareNotEqual(const TSet<ElementType, KeyFuncs, Allocator>& A,const TSet<ElementType, KeyFuncs, Allocator>& B)
 {
 	return !TSetPrivateFriend::LegacyCompareEqual(A, B);
 }
diff --git a/Engine/Source/Runtime/Core/Public/Containers/SharedString.h b/Engine/Source/Runtime/Core/Public/Containers/SharedString.h
--- a/Engine/Source/Runtime/Core/Public/Containers/SharedString.h
+++ b/Engine/Source/Runtime/Core/Public/Containers/SharedString.h
@@ -30,14 +30,14 @@ class TSharedString
 public:
 	using ElementType = CharType;
 
-	TSharedString() = default;
-	TSharedString(TSharedString&& String);
-	TSharedString(const TSharedString& String);
+	[[nodiscard]] TSharedString() = default;
+	[[nodiscard]] TSharedString(TSharedString&& String);
+	[[nodiscard]] TSharedString(const TSharedString& String);
 	TSharedString& operator=(TSharedString&& String);
 	TSharedString& operator=(const TSharedString& String);
 
 	/** Allocates a copy of the string and constructs this as a reference to it. */
-	TSharedString(TStringView<CharType> String);
+	[[nodiscard]] TSharedString(TStringView<CharType> String);
 
 	/** Allocates a copy of the string and assigns this as a reference to it. */
 	TSharedString& operator=(TStringView<CharType> String);
@@ -76,27 +76,27 @@ public:
 	static const TSharedString Empty;
 
 private:
-	friend constexpr inline const CharType* GetData(const TSharedString String)
+	[[nodiscard]] friend constexpr inline const CharType* GetData(const TSharedString String)
 	{
 		return *String;
 	}
 
-	friend constexpr inline auto GetNum(const TSharedString String)
+	[[nodiscard]] friend constexpr inline auto GetNum(const TSharedString String)
 	{
 		return String.Len();
 	}
 
-	friend inline uint32 GetTypeHash(const TSharedString& String)
+	[[nodiscard]] friend inline uint32 GetTypeHash(const TSharedString& String)
 	{
 		return GetTypeHash(MakeStringView(String));
 	}
 
-	friend inline bool operator==(const TSharedString& Lhs, const TSharedString& Rhs)
+	[[nodiscard]] friend inline bool operator==(const TSharedString& Lhs, const TSharedString& Rhs)
 	{
 		return MakeStringView(Lhs).Equals(MakeStringView(Rhs), ESearchCase::IgnoreCase);
 	}
 
-	friend inline bool operator<(const TSharedString& Lhs, const TSharedString& Rhs)
+	[[nodiscard]] friend inline bool operator<(const TSharedString& Lhs, const TSharedString& Rhs)
 	{
 		return MakeStringView(Lhs).Compare(MakeStringView(Rhs), ESearchCase::IgnoreCase) < 0;
 	}
@@ -120,7 +120,7 @@ inline TSharedString<CharType>::TSharedString(TSharedString&& String)
 }
 
 template <typename CharType>
-TSharedString<CharType>::TSharedString(const TSharedString& String)
+[[nodiscard]] TSharedString<CharType>::TSharedString(const TSharedString& String)
 	: Chars(String.Chars)
 {
 	AddRef(Chars);
@@ -150,7 +150,7 @@ TSharedString<CharType>& TSharedString<CharType>::operator=(const TSharedString&
 }
 
 template <typename CharType>
-TSharedString<CharType>::TSharedString(const TStringView<CharType> String)
+[[nodiscard]] TSharedString<CharType>::TSharedString(const TStringView<CharType> String)
 {
 	if (const int32 Length = String.Len())
 	{
diff --git a/Engine/Source/Runtime/Core/Public/Containers/SortedMap.h b/Engine/Source/Runtime/Core/Public/Containers/SortedMap.h
--- a/Engine/Source/Runtime/Core/Public/Containers/SortedMap.h
+++ b/Engine/Source/Runtime/Core/Public/Containers/SortedMap.h
@@ -28,29 +28,29 @@ public:
 	typedef typename TTypeTraits<KeyType  >::ConstInitType    KeyInitType;
 	typedef typename TTypeTraits<ValueType>::ConstInitType    ValueInitType;
 	typedef TPair<KeyType, ValueType> ElementType;
-	
-	TSortedMap() = default;
-	TSortedMap(TSortedMap&&) = default;
-	TSortedMap(const TSortedMap&) = default;
+
+	[[nodiscard]] TSortedMap() = default;
+	[[nodiscard]] TSortedMap(TSortedMap&&) = default;
+	[[nodiscard]] TSortedMap(const TSortedMap&) = default;
 	TSortedMap& operator=(TSortedMap&&) = default;
 	TSortedMap& operator=(const TSortedMap&) = default;
 
 	/** Constructor for moving elements from a TSortedMap with a different ArrayAllocator. */
 	template<typename OtherArrayAllocator>
-	TSortedMap(TSortedMap<KeyType, ValueType, OtherArrayAllocator, SortPredicate>&& Other)
+	[[nodiscard]] TSortedMap(TSortedMap<KeyType, ValueType, OtherArrayAllocator, SortPredicate>&& Other)
 		: Pairs(MoveTemp(Other.Pairs))
 	{
 	}
 
 	/** Constructor for copying elements from a TSortedMap with a different ArrayAllocator. */
 	template<typename OtherArrayAllocator>
-	TSortedMap(const TSortedMap<KeyType, ValueType, OtherArrayAllocator, SortPredicate>& Other)
+	[[nodiscard]] TSortedMap(const TSortedMap<KeyType, ValueType, OtherArrayAllocator, SortPredicate>& Other)
 		: Pairs(Other.Pairs)
 	{
 	}
 
 	/** Constructor which gets its elements from a native initializer list */
-	TSortedMap(std::initializer_list<TPairInitializer<const KeyType&, const ValueType&>> InitList)
+	[[nodiscard]] TSortedMap(std::initializer_list<TPairInitializer<const KeyType&, const ValueType&>> InitList)
 	{
 		this->Reserve((int32)InitList.size());
 		for (const TPairInitializer<const KeyType&, const ValueType&>& Element : InitList)
@@ -65,11 +65,11 @@ public:
 	constexpr static bool bHasIntrusiveUnsetOptionalState = true;
 	using IntrusiveUnsetOptionalStateType = TSortedMap;
 
-	explicit TSortedMap(FIntrusiveUnsetOptionalState Tag)
+	[[nodiscard]] explicit TSortedMap(FIntrusiveUnsetOptionalState Tag)
 		: Pairs(Tag)
 	{
 	}
-	bool operator==(FIntrusiveUnsetOptionalState Tag) const
+	[[nodiscard]] bool operator==(FIntrusiveUnsetOptionalState Tag) const
 	{
 		return Pairs == Tag;
 	}
@@ -150,13 +150,13 @@ public:
 	 * @returns True if the map is empty.
 	 * @see Num
 	 */
-	bool IsEmpty() const
+	[[nodiscard]] bool IsEmpty() const
 	{
 		return Pairs.IsEmpty();
 	}
 
 	/** @return The number of elements in the map. */
-	FORCEINLINE int32 Num() const
+	[[nodiscard]] FORCEINLINE int32 Num() const
 	{
 		return Pairs.Num();
 	}
@@ -167,7 +167,7 @@ public:
 	 *
 	 * @return number of bytes allocated by this container.
 	 */
-	FORCEINLINE SIZE_T GetAllocatedSize() const
+	[[nodiscard]] FORCEINLINE SIZE_T GetAllocatedSize() const
 	{
 		return Pairs.GetAllocatedSize();
 	}
@@ -258,7 +258,7 @@ public:
 	 * @param Value The value to search for.
 	 * @return A pointer to the key associated with the specified value, or nullptr if the value isn't contained in this map (only valid until the next change to any key in the map).
 	 */
-	const KeyType* FindKey(ValueInitType Value) const
+	[[nodiscard]] const KeyType* FindKey(ValueInitType Value) const
 	{
 		for(typename ElementArrayType::TConstIterator PairIt(Pairs);PairIt;++PairIt)
 		{
@@ -276,7 +276,7 @@ public:
 	 * @param Key The key to search for.
 	 * @return A pointer to the value associated with the specified key, or nullptr if the key isn't contained in this map.  The pointer (only valid until the next change to any key in the map).
 	 */
-	FORCEINLINE ValueType* Find(KeyConstPointerType Key)
+	[[nodiscard]] FORCEINLINE ValueType* Find(KeyConstPointerType Key)
 	{
 		int32 FoundIndex = FindIndex(Key);
 
@@ -288,7 +288,7 @@ public:
 		return nullptr;
 	}
 
-	FORCEINLINE const ValueType* Find(KeyConstPointerType Key) const
+	[[nodiscard]] FORCEINLINE const ValueType* Find(KeyConstPointerType Key) const
 	{
 		return const_cast<TSortedMap*>(this)->Find(Key);
 	}
@@ -309,14 +309,14 @@ public:
 	 * @param Key The key to search for.
 	 * @return The value associated with the specified key, or triggers an assertion if the key does not exist.
 	 */
-	FORCEINLINE ValueType& FindChecked(KeyConstPointerType Key)
+	[[nodiscard]] FORCEINLINE ValueType& FindChecked(KeyConstPointerType Key)
 	{
 		ValueType* Value = Find(Key);
 		check(Value != nullptr);
 		return *Value;
 	}
 
-	FORCEINLINE const ValueType& FindChecked(KeyConstPointerType Key) const
+	[[nodiscard]] FORCEINLINE const ValueType& FindChecked(KeyConstPointerType Key) const
 	{
 		const ValueType* Value = Find(Key);
 		check(Value != nullptr);
@@ -329,7 +329,7 @@ public:
 	 * @param Key The key to search for.
 	 * @return The value associated with the specified key, or the default value for the ValueType if the key isn't contained in this map.
 	 */
-	FORCEINLINE ValueType FindRef(KeyConstPointerType Key) const
+	[[nodiscard]] FORCEINLINE ValueType FindRef(KeyConstPointerType Key) const
 	{
 		if (const ValueType* Value = Find(Key))
 		{
@@ -346,7 +346,7 @@ public:
 	 * @param DefaultValue The fallback value if the key is not found.
 	 * @return The value associated with the specified key, or DefaultValue if the key isn't contained in this map.
 	 */
-	FORCEINLINE ValueType FindRef(KeyConstPointerType Key, ValueType DefaultValue) const
+	[[nodiscard]] FORCEINLINE ValueType FindRef(KeyConstPointerType Key, ValueType DefaultValue) const
 	{
 		if (const ValueType* Value = Find(Key))
 		{
@@ -361,13 +361,13 @@ public:
 	 * Callers should not depend on particular patterns in the behaviour of this function.
 	 * @return A pointer to an arbitrary pair, or nullptr if the container is empty.
 	 */
-	ElementType* FindArbitraryElement()
+	[[nodiscard]] ElementType* FindArbitraryElement()
 	{
 		// The goal of this function is to be fast, and so the implementation may be improved at any time even if it gives different results.
 
 		return Pairs.IsEmpty() ? nullptr : &Pairs.Last();
 	}
-	const ElementType* FindArbitraryElement() const
+	[[nodiscard]] const ElementType* FindArbitraryElement() const
 	{
 		return const_cast<TSortedMap*>(this)->FindArbitraryElement();
 	}
@@ -378,7 +378,7 @@ public:
 	 * @param Key The key to check for.
 	 * @return true if the map contains the key.
 	 */
-	FORCEINLINE bool Contains(KeyConstPointerType Key) const
+	[[nodiscard]] FORCEINLINE bool Contains(KeyConstPointerType Key) const
 	{
 		if (Find(Key))
 		{
@@ -506,8 +506,14 @@ public:
 		}
 	}
 
-	FORCEINLINE       ValueType& operator[](KeyConstPointerType Key)       { return this->FindChecked(Key); }
-	FORCEINLINE const ValueType& operator[](KeyConstPointerType Key) const { return this->FindChecked(Key); }
+	[[nodiscard]] FORCEINLINE ValueType& operator[](KeyConstPointerType Key)
+	{
+		return this->FindChecked(Key);
+	}
+	[[nodiscard]] FORCEINLINE const ValueType& operator[](KeyConstPointerType Key) const
+	{
+		return this->FindChecked(Key);
+	}
 
 	// Interface functions to match TMap/TSet
 
@@ -544,7 +550,7 @@ private:
 
 	/** Implementation of find and add */
 	template <typename ArgType>
-	FORCEINLINE ValueType& FindOrAddImpl(ArgType&& Key)
+	[[nodiscard]] FORCEINLINE ValueType& FindOrAddImpl(ArgType&& Key)
 	{
 		if (ValueType* Value = Find(Key))
 		{
@@ -555,7 +561,7 @@ private:
 	}
 
 	/** Find index of key */
-	FORCEINLINE int32 FindIndex(KeyConstPointerType Key) const
+	[[nodiscard]] FORCEINLINE int32 FindIndex(KeyConstPointerType Key) const
 	{
 		return Algo::BinarySearchBy(Pairs, Key, FKeyForward(), SortPredicate());
 	}
@@ -607,7 +613,7 @@ private:
 		typedef std::conditional_t<bConst,const typename ElementArrayType::ElementType, typename ElementArrayType::ElementType> PairType;
 
 	protected:
-		FORCEINLINE TBaseIterator(const PairItType& InElementIt)
+		[[nodiscard]] FORCEINLINE TBaseIterator(const PairItType& InElementIt)
 			: PairIt(InElementIt)
 		{
 		}
@@ -620,24 +626,42 @@ private:
 		}
 
 		/** conversion to "bool" returning true if the iterator is valid. */
-		FORCEINLINE explicit operator bool() const
-		{ 
+		[[nodiscard]] FORCEINLINE explicit operator bool() const
+		{
 			return !!PairIt; 
 		}
 
-		FORCEINLINE bool operator==(const TBaseIterator& Rhs) const { return PairIt == Rhs.PairIt; }
-		FORCEINLINE bool operator!=(const TBaseIterator& Rhs) const { return PairIt != Rhs.PairIt; }
+		[[nodiscard]] FORCEINLINE bool operator==(const TBaseIterator& Rhs) const
+		{
+			return PairIt == Rhs.PairIt;
+		}
+		[[nodiscard]] FORCEINLINE bool operator!=(const TBaseIterator& Rhs) const
+		{
+			return PairIt != Rhs.PairIt;
+		}
 
-		FORCEINLINE ItKeyType&   Key()   const { return PairIt->Key; }
-		FORCEINLINE ItValueType& Value() const { return PairIt->Value; }
+		FORCEINLINE ItKeyType& Key() const
+		{
+			return PairIt->Key;
+		}
+		FORCEINLINE ItValueType& Value() const
+		{
+			return PairIt->Value;
+		}
 
 		[[nodiscard]] FORCEINLINE FSetElementId GetId() const
 		{
 			return FSetElementId::FromInteger(PairIt.GetIndex());
 		}
 
-		FORCEINLINE PairType& operator* () const { return  *PairIt; }
-		FORCEINLINE PairType* operator->() const { return &*PairIt; }
+		FORCEINLINE PairType& operator*() const
+		{
+			return  *PairIt;
+		}
+		FORCEINLINE PairType* operator->() const
+		{
+			return &*PairIt;
+		}
 
 	protected:
 		PairItType PairIt;
@@ -658,7 +682,7 @@ private:
 		typedef std::conditional_t<bConst, const typename ElementArrayType::ElementType, typename ElementArrayType::ElementType> PairType;
 
 	protected:
-		FORCEINLINE TBaseReverseIterator(PairType* InData, SizeType InNum)
+		[[nodiscard]] FORCEINLINE TBaseReverseIterator(PairType* InData, SizeType InNum)
 			: Data(InData), Index(InNum-1)
 		{
 		}
@@ -672,24 +696,42 @@ private:
 		}
 
 		/** conversion to "bool" returning true if the iterator is valid. */
-		FORCEINLINE explicit operator bool() const
+		[[nodiscard]] FORCEINLINE explicit operator bool() const
 		{
 			return Index != static_cast<SizeType>(-1);
 		}
 
-		FORCEINLINE bool operator==(const TBaseReverseIterator& Rhs) const { return Index == Rhs.Index; }
-		FORCEINLINE bool operator!=(const TBaseReverseIterator& Rhs) const { return Index != Rhs.Index; }
+		[[nodiscard]] FORCEINLINE bool operator==(const TBaseReverseIterator& Rhs) const
+		{
+			return Index == Rhs.Index;
+		}
+		[[nodiscard]] FORCEINLINE bool operator!=(const TBaseReverseIterator& Rhs) const
+		{
+			return Index != Rhs.Index;
+		}
 
-		FORCEINLINE ItKeyType& Key()   const { return Data[Index].Key; }
-		FORCEINLINE ItValueType& Value() const { return Data[Index].Value; }
+		FORCEINLINE ItKeyType& Key() const
+		{
+			return Data[Index].Key;
+		}
+		FORCEINLINE ItValueType& Value() const
+		{
+			return Data[Index].Value;
+		}
 
 		[[nodiscard]] FORCEINLINE FSetElementId GetId() const
 		{
 			return FSetElementId::FromInteger(Index);
 		}
 
-		FORCEINLINE PairType& operator* () const { return  Data[Index]; }
-		FORCEINLINE PairType* operator->() const { return &Data[Index]; }
+		FORCEINLINE PairType& operator*() const
+		{
+			return  Data[Index];
+		}
+		FORCEINLINE PairType* operator->() const
+		{
+			return &Data[Index];
+		}
 
 	protected:
 		PairType* Data;
@@ -706,12 +748,12 @@ public:
 	{
 	public:
 
-		FORCEINLINE TIterator(TSortedMap& InMap)
+		[[nodiscard]] FORCEINLINE TIterator(TSortedMap& InMap)
 			: TBaseIterator<false>(InMap.Pairs.CreateIterator())
 		{
 		}
 
-		FORCEINLINE TIterator(const typename TBaseIterator<false>::PairItType& InPairIt)
+		[[nodiscard]] FORCEINLINE TIterator(const typename TBaseIterator<false>::PairItType& InPairIt)
 			: TBaseIterator<false>(InPairIt)
 		{
 		}
@@ -727,12 +769,12 @@ public:
 	class TConstIterator : public TBaseIterator<true>
 	{
 	public:
-		FORCEINLINE TConstIterator(const TSortedMap& InMap)
+		[[nodiscard]] FORCEINLINE TConstIterator(const TSortedMap& InMap)
 			: TBaseIterator<true>(InMap.Pairs.CreateConstIterator())
 		{
 		}
 
-		FORCEINLINE TConstIterator(const typename TBaseIterator<true>::PairItType& InPairIt)
+		[[nodiscard]] FORCEINLINE TConstIterator(const typename TBaseIterator<true>::PairItType& InPairIt)
 			: TBaseIterator<true>(InPairIt)
 		{
 		}
@@ -742,7 +784,7 @@ public:
 	class TReverseIterator : public TBaseReverseIterator<false>
 	{
 	public:
-		FORCEINLINE TReverseIterator(TSortedMap& InMap)
+		[[nodiscard]] FORCEINLINE TReverseIterator(TSortedMap& InMap)
 			: TBaseReverseIterator<false>(InMap.Pairs.GetData(), InMap.Pairs.Num())
 		{
 		}
@@ -754,7 +796,7 @@ public:
 	class TConstReverseIterator : public TBaseReverseIterator<true>
 	{
 	public:
-		FORCEINLINE TConstReverseIterator(const TSortedMap& InMap)
+		[[nodiscard]] FORCEINLINE TConstReverseIterator(const TSortedMap& InMap)
 			: TBaseReverseIterator<true>(InMap.Pairs.GetData(), InMap.Pairs.Num())
 		{
 		}
@@ -768,7 +810,7 @@ public:
 		using Super = TBaseIterator<true>;
 
 	public:
-		FORCEINLINE TConstKeyIterator(const TSortedMap& InMap, KeyInitType InKey)
+		[[nodiscard]] FORCEINLINE TConstKeyIterator(const TSortedMap& InMap, KeyInitType InKey)
 			: Super(InMap.Pairs.CreateConstIterator())
 		{
 			int32 NewIndex = InMap.FindIndex(InKey);
@@ -796,7 +838,7 @@ public:
 		using Super = TBaseIterator<false>;
 
 	public:
-		FORCEINLINE TKeyIterator(TSortedMap& InMap, KeyInitType InKey)
+		[[nodiscard]] FORCEINLINE TKeyIterator(TSortedMap& InMap, KeyInitType InKey)
 			: Super(InMap.Pairs.CreateIterator())
 		{
 			int32 NewIndex = InMap.FindIndex(InKey);
@@ -826,25 +868,25 @@ public:
 	};
 
 	/** Creates an iterator over all the pairs in this map */
-	FORCEINLINE TIterator CreateIterator()
+	[[nodiscard]] FORCEINLINE TIterator CreateIterator()
 	{
 		return TIterator(*this);
 	}
 
 	/** Creates a const iterator over all the pairs in this map */
-	FORCEINLINE TConstIterator CreateConstIterator() const
+	[[nodiscard]] FORCEINLINE TConstIterator CreateConstIterator() const
 	{
 		return TConstIterator(*this);
 	}
 
 	/** Creates an iterator over the values associated with a specified key in a map. This will be at most one value because keys must be unique */
-	FORCEINLINE TKeyIterator CreateKeyIterator(KeyInitType InKey)
+	[[nodiscard]] FORCEINLINE TKeyIterator CreateKeyIterator(KeyInitType InKey)
 	{
 		return TKeyIterator(*this, InKey);
 	}
 
 	/** Creates a const iterator over the values associated with a specified key in a map. This will be at most one value because keys must be unique */
-	FORCEINLINE TConstKeyIterator CreateConstKeyIterator(KeyInitType InKey) const
+	[[nodiscard]] FORCEINLINE TConstKeyIterator CreateConstKeyIterator(KeyInitType InKey) const
 	{
 		return TConstKeyIterator(*this, InKey);
 	}
@@ -858,10 +900,10 @@ public:
 	 * DO NOT USE DIRECTLY
 	 * STL-like iterators to enable range-based for loop support.
 	 */
-	FORCEINLINE RangedForIteratorType		begin()	      { return Pairs.begin(); }
-	FORCEINLINE RangedForConstIteratorType	begin() const { return Pairs.begin(); }
-	FORCEINLINE RangedForIteratorType		end()         { return Pairs.end(); }
-	FORCEINLINE RangedForConstIteratorType	end() const   { return Pairs.end(); }
+	[[nodiscard]] FORCEINLINE RangedForIteratorType		begin()	      { return Pairs.begin(); }
+	[[nodiscard]] FORCEINLINE RangedForConstIteratorType	begin() const { return Pairs.begin(); }
+	[[nodiscard]] FORCEINLINE RangedForIteratorType		end()         { return Pairs.end(); }
+	[[nodiscard]] FORCEINLINE RangedForConstIteratorType	end() const   { return Pairs.end(); }
 
 	friend struct TSortedMapPrivateFriend;
 };
diff --git a/Engine/Source/Runtime/Core/Public/Containers/SparseArray.h b/Engine/Source/Runtime/Core/Public/Containers/SparseArray.h
--- a/Engine/Source/Runtime/Core/Public/Containers/SparseArray.h
+++ b/Engine/Source/Runtime/Core/Public/Containers/SparseArray.h
@@ -91,11 +91,11 @@ public:
 	constexpr static bool bHasIntrusiveUnsetOptionalState = true;
 	using IntrusiveUnsetOptionalStateType = TSparseArray;
 
-	explicit TSparseArray(FIntrusiveUnsetOptionalState Tag)
+	[[nodiscard]] explicit TSparseArray(FIntrusiveUnsetOptionalState Tag)
 		: Data(Tag)
 	{
 	}
-	bool operator==(FIntrusiveUnsetOptionalState Tag) const
+	[[nodiscard]] bool operator==(FIntrusiveUnsetOptionalState Tag) const
 	{
 		return Data == Tag;
 	}
@@ -686,8 +686,8 @@ public:
 	 * @returns Index to the first matching element, or INDEX_NONE if none is found.
 	 */
 	template <typename Predicate>
-	int32 IndexOfByPredicate(Predicate Pred) const
-	{	
+	[[nodiscard]] int32 IndexOfByPredicate(Predicate Pred) const
+	{
 		for (TConstIterator It = CreateConstIterator(); It; ++It)
 		{
 			if (::Invoke(Pred, *It))
@@ -704,7 +704,7 @@ public:
 	 * Callers should not depend on particular patterns in the behaviour of this function.
 	 * @return The index of any element in the container, or INDEX_NONE if it is empty.
 	 */
-	int32 FindArbitraryElementIndex() const
+	[[nodiscard]] int32 FindArbitraryElementIndex() const
 	{
 		// The goal of this function is to be fast, and so the implementation may be improved at any time even if it gives different results.
 
@@ -721,7 +721,7 @@ public:
 	 * Only returns the size of allocations made directly by the container, not the elements themselves.
 	 * @return number of bytes allocated by this container
 	 */
-	SIZE_T GetAllocatedSize( void ) const
+	[[nodiscard]] SIZE_T GetAllocatedSize( void ) const
 	{
 		return Data.GetAllocatedSize() + AllocationFlags.GetAllocatedSize();
 	}
@@ -733,7 +733,7 @@ public:
 		AllocationFlags.CountBytes(Ar);
 	}
 
-	bool IsCompact() const
+	[[nodiscard]] bool IsCompact() const
 	{
 		return NumFreeIndices == 0;
 	}
@@ -742,7 +742,7 @@ public:
 	 * Equality comparison operator.
 	 * Checks that both arrays have the same elements and element indices; that means that unallocated elements are signifigant!
 	 */
-	bool operator==(const TSparseArray& B) const
+	[[nodiscard]] bool operator==(const TSparseArray& B) const
 	{
 		if(GetMaxIndex() != B.GetMaxIndex())
 		{
@@ -773,27 +773,28 @@ public:
 	 * Inequality comparison operator.
 	 * Checks that both arrays have the same elements and element indices; that means that unallocated elements are signifigant!
 	 */
-	bool operator!=(const TSparseArray& B) const
+	[[nodiscard]] bool operator!=(const TSparseArray& B) const
 	{
 		return !(*this == B);
 	}
 
 	/** Default constructor. */
-	TSparseArray()
-	:	FirstFreeIndex(-1)
-	,	NumFreeIndices(0)
-	{}
+	[[nodiscard]] TSparseArray()
+		: FirstFreeIndex(-1)
+		, NumFreeIndices(0)
+	{
+	}
 
 	/** Move constructor. */
-	TSparseArray(TSparseArray&& InCopy)
+	[[nodiscard]] TSparseArray(TSparseArray&& InCopy)
 	{
 		this->Move(*this, InCopy);
 	}
 
 	/** Copy constructor. */
-	TSparseArray(const TSparseArray& InCopy)
-	:	FirstFreeIndex(-1)
-	,	NumFreeIndices(0)
+	[[nodiscard]] TSparseArray(const TSparseArray& InCopy)
+		: FirstFreeIndex(-1)
+		, NumFreeIndices(0)
 	{
 		*this = InCopy;
 	}
@@ -882,33 +883,45 @@ private:
 
 public:
 	// Accessors.
-	ElementType& operator[](int32 Index)
+	[[nodiscard]] ElementType& operator[](int32 Index)
 	{
 		checkSlow(Index >= 0 && Index < Data.Num() && Index < AllocationFlags.Num());
 		//checkSlow(AllocationFlags[Index]); // Disabled to improve loading times -BZ
 		return *(ElementType*)&((FElementOrFreeListLink*)Data.GetData())[Index].ElementData;
 	}
-	const ElementType& operator[](int32 Index) const
+	[[nodiscard]] const ElementType& operator[](int32 Index) const
 	{
 		checkSlow(Index >= 0 && Index < Data.Num() && Index < AllocationFlags.Num());
 		//checkSlow(AllocationFlags[Index]); // Disabled to improve loading times -BZ
 		return *(ElementType*)&((FElementOrFreeListLink*)Data.GetData())[Index].ElementData;
 	}
-	int32 PointerToIndex(const ElementType* Ptr) const
+	[[nodiscard]] int32 PointerToIndex(const ElementType* Ptr) const
 	{
 		checkSlow(Data.Num());
 		int32 Index = (int32)((FElementOrFreeListLink*)Ptr - (FElementOrFreeListLink*)Data.GetData());
 		checkSlow(Index >= 0 && Index < Data.Num() && Index < AllocationFlags.Num() && AllocationFlags[Index]);
 		return Index;
 	}
-	bool IsValidIndex(int32 Index) const
+	[[nodiscard]] bool IsValidIndex(int32 Index) const
 	{
 		return AllocationFlags.IsValidIndex(Index) && AllocationFlags[Index];
 	}
-	bool IsAllocated(int32 Index) const { return AllocationFlags[Index]; }
-	int32 GetMaxIndex() const { return Data.Num(); }
-	bool IsEmpty() const { return Data.Num() == NumFreeIndices; }
-	int32 Num() const { return Data.Num() - NumFreeIndices; }
+	[[nodiscard]] bool IsAllocated(int32 Index) const
+	{
+		return AllocationFlags[Index];
+	}
+	[[nodiscard]] int32 GetMaxIndex() const
+	{
+		return Data.Num();
+	}
+	[[nodiscard]] bool IsEmpty() const
+	{
+		return Data.Num() == NumFreeIndices;
+	}
+	[[nodiscard]] int32 Num() const
+	{
+		return Data.Num() - NumFreeIndices;
+	}
 
 	/**
 	 * Checks that the specified address is not part of an element within the container.  Used for implementations
@@ -935,7 +948,7 @@ private:
 		typedef std::conditional_t<bConst,const ElementType,ElementType> ItElementType;
 
 	public:
-		explicit TBaseIterator(ArrayType& InArray, const BitArrayItType& InBitArrayIt)
+		[[nodiscard]] explicit TBaseIterator(ArrayType& InArray, const BitArrayItType& InBitArrayIt)
 			: Array     (InArray)
 			, BitArrayIt(InBitArrayIt)
 		{
@@ -948,26 +961,44 @@ private:
 			return *this;
 		}
 
-		FORCEINLINE int32 GetIndex() const { return BitArrayIt.GetIndex(); }
+		[[nodiscard]] FORCEINLINE int32 GetIndex() const
+		{
+			return BitArrayIt.GetIndex();
+		}
 
-		FORCEINLINE bool operator==(const TBaseIterator& Rhs) const { return BitArrayIt == Rhs.BitArrayIt && &Array == &Rhs.Array; }
-		FORCEINLINE bool operator!=(const TBaseIterator& Rhs) const { return BitArrayIt != Rhs.BitArrayIt || &Array != &Rhs.Array; }
+		[[nodiscard]] FORCEINLINE bool operator==(const TBaseIterator& Rhs) const
+		{
+			return BitArrayIt == Rhs.BitArrayIt && &Array == &Rhs.Array;
+		}
+		[[nodiscard]] FORCEINLINE bool operator!=(const TBaseIterator& Rhs) const
+		{
+			return BitArrayIt != Rhs.BitArrayIt || &Array != &Rhs.Array;
+		}
 
 		/** conversion to "bool" returning true if the iterator is valid. */
-		FORCEINLINE explicit operator bool() const
-		{ 
-			return !!BitArrayIt; 
+		[[nodiscard]] FORCEINLINE explicit operator bool() const
+		{
+			return !!BitArrayIt;
 		}
 
 		/** inverse of the "bool" operator */
-		FORCEINLINE bool operator !() const 
+		[[nodiscard]] FORCEINLINE bool operator !() const 
 		{
 			return !(bool)*this;
 		}
 
-		FORCEINLINE ItElementType& operator*() const { return Array[GetIndex()]; }
-		FORCEINLINE ItElementType* operator->() const { return &Array[GetIndex()]; }
-		FORCEINLINE const FRelativeBitReference& GetRelativeBitReference() const { return BitArrayIt; }
+		[[nodiscard]] FORCEINLINE ItElementType& operator*() const
+		{
+			return Array[GetIndex()];
+		}
+		[[nodiscard]] FORCEINLINE ItElementType* operator->() const
+		{
+			return &Array[GetIndex()];
+		}
+		[[nodiscard]] FORCEINLINE const FRelativeBitReference& GetRelativeBitReference() const
+		{
+			return BitArrayIt;
+		}
 
 	protected:
 		ArrayType&     Array;
@@ -980,12 +1011,12 @@ public:
 	class TIterator : public TBaseIterator<false>
 	{
 	public:
-		TIterator(TSparseArray& InArray)
+		[[nodiscard]] TIterator(TSparseArray& InArray)
 			: TBaseIterator<false>(InArray, TConstSetBitIterator<typename Allocator::BitArrayAllocator>(InArray.AllocationFlags))
 		{
 		}
 
-		TIterator(TSparseArray& InArray, const typename TBaseIterator<false>::BitArrayItType& InBitArrayIt)
+		[[nodiscard]] TIterator(TSparseArray& InArray, const typename TBaseIterator<false>::BitArrayItType& InBitArrayIt)
 			: TBaseIterator<false>(InArray, InBitArrayIt)
 		{
 		}
@@ -1001,12 +1032,12 @@ public:
 	class TConstIterator : public TBaseIterator<true>
 	{
 	public:
-		TConstIterator(const TSparseArray& InArray)
+		[[nodiscard]] TConstIterator(const TSparseArray& InArray)
 			: TBaseIterator<true>(InArray, TConstSetBitIterator<typename Allocator::BitArrayAllocator>(InArray.AllocationFlags))
 		{
 		}
 
-		TConstIterator(const TSparseArray& InArray, const typename TBaseIterator<true>::BitArrayItType& InBitArrayIt)
+		[[nodiscard]] TConstIterator(const TSparseArray& InArray, const typename TBaseIterator<true>::BitArrayItType& InBitArrayIt)
 			: TBaseIterator<true>(InArray, InBitArrayIt)
 		{
 		}
@@ -1016,13 +1047,13 @@ public:
 		class TRangedForIterator : public TIterator
 		{
 		public:
-			TRangedForIterator(TSparseArray& InArray, const typename TBaseIterator<false>::BitArrayItType& InBitArrayIt)
+			[[nodiscard]] TRangedForIterator(TSparseArray& InArray, const typename TBaseIterator<false>::BitArrayItType& InBitArrayIt)
 				: TIterator (InArray, InBitArrayIt)
 				, InitialNum(InArray.Num())
 			{
 			}
 
-			FORCEINLINE bool operator!=(const TRangedForIterator& Rhs) const
+			[[nodiscard]] FORCEINLINE bool operator!=(const TRangedForIterator& Rhs) const
 			{
 				// We only need to do the check in this operator, because no other operator will be
 				// called until after this one returns.
@@ -1039,13 +1070,13 @@ public:
 		class TRangedForConstIterator : public TConstIterator
 		{
 		public:
-			TRangedForConstIterator(const TSparseArray& InArray, const typename TBaseIterator<true>::BitArrayItType& InBitArrayIt)
+			[[nodiscard]] TRangedForConstIterator(const TSparseArray& InArray, const typename TBaseIterator<true>::BitArrayItType& InBitArrayIt)
 				: TConstIterator(InArray, InBitArrayIt)
 				, InitialNum    (InArray.Num())
 			{
 			}
 
-			FORCEINLINE bool operator!=(const TRangedForConstIterator& Rhs) const
+			[[nodiscard]] FORCEINLINE bool operator!=(const TRangedForConstIterator& Rhs) const
 			{
 				// We only need to do the check in this operator, because no other operator will be
 				// called until after this one returns.
@@ -1064,13 +1095,13 @@ public:
 	#endif
 
 	/** Creates an iterator for the contents of this array */
-	TIterator CreateIterator()
+	[[nodiscard]] TIterator CreateIterator()
 	{
 		return TIterator(*this);
 	}
 
 	/** Creates a const iterator for the contents of this array */
-	TConstIterator CreateConstIterator() const
+	[[nodiscard]] TConstIterator CreateConstIterator() const
 	{
 		return TConstIterator(*this);
 	}
@@ -1080,10 +1111,10 @@ public:
 	 * DO NOT USE DIRECTLY
 	 * STL-like iterators to enable range-based for loop support.
 	 */
-	FORCEINLINE TRangedForIterator      begin()       { return TRangedForIterator     (*this, TConstSetBitIterator<typename Allocator::BitArrayAllocator>(AllocationFlags)); }
-	FORCEINLINE TRangedForConstIterator begin() const { return TRangedForConstIterator(*this, TConstSetBitIterator<typename Allocator::BitArrayAllocator>(AllocationFlags)); }
-	FORCEINLINE TRangedForIterator      end  ()       { return TRangedForIterator     (*this, TConstSetBitIterator<typename Allocator::BitArrayAllocator>(AllocationFlags, AllocationFlags.Num())); }
-	FORCEINLINE TRangedForConstIterator end  () const { return TRangedForConstIterator(*this, TConstSetBitIterator<typename Allocator::BitArrayAllocator>(AllocationFlags, AllocationFlags.Num())); }
+	[[nodiscard]] FORCEINLINE TRangedForIterator      begin()       { return TRangedForIterator     (*this, TConstSetBitIterator<typename Allocator::BitArrayAllocator>(AllocationFlags)); }
+	[[nodiscard]] FORCEINLINE TRangedForConstIterator begin() const { return TRangedForConstIterator(*this, TConstSetBitIterator<typename Allocator::BitArrayAllocator>(AllocationFlags)); }
+	[[nodiscard]] FORCEINLINE TRangedForIterator      end  ()       { return TRangedForIterator     (*this, TConstSetBitIterator<typename Allocator::BitArrayAllocator>(AllocationFlags, AllocationFlags.Num())); }
+	[[nodiscard]] FORCEINLINE TRangedForConstIterator end  () const { return TRangedForConstIterator(*this, TConstSetBitIterator<typename Allocator::BitArrayAllocator>(AllocationFlags, AllocationFlags.Num())); }
 
 public:
 	/** An iterator which only iterates over the elements of the array which correspond to set bits in a separate bit array. */
@@ -1091,32 +1122,45 @@ public:
 	class TConstSubsetIterator
 	{
 	public:
-		TConstSubsetIterator( const TSparseArray& InArray, const TBitArray<SubsetAllocator>& InBitArray ):
+		[[nodiscard]] TConstSubsetIterator( const TSparseArray& InArray, const TBitArray<SubsetAllocator>& InBitArray ):
 			Array(InArray),
 			BitArrayIt(InArray.AllocationFlags,InBitArray)
-		{}
+		{
+		}
 		FORCEINLINE TConstSubsetIterator& operator++()
 		{
 			// Iterate to the next element which is both allocated and has its bit set in the other bit array.
 			++BitArrayIt;
 			return *this;
 		}
-		FORCEINLINE int32 GetIndex() const { return BitArrayIt.GetIndex(); }
+		[[nodiscard]] FORCEINLINE int32 GetIndex() const
+		{
+			return BitArrayIt.GetIndex();
+		}
 		
 		/** conversion to "bool" returning true if the iterator is valid. */
-		FORCEINLINE explicit operator bool() const
+		[[nodiscard]] FORCEINLINE explicit operator bool() const
 		{ 
 			return !!BitArrayIt; 
 		}
 		/** inverse of the "bool" operator */
-		FORCEINLINE bool operator !() const 
+		[[nodiscard]] FORCEINLINE bool operator !() const 
 		{
 			return !(bool)*this;
 		}
 
-		FORCEINLINE const ElementType& operator*() const { return Array[GetIndex()]; }
-		FORCEINLINE const ElementType* operator->() const { return &Array[GetIndex()]; }
-		FORCEINLINE const FRelativeBitReference& GetRelativeBitReference() const { return BitArrayIt; }
+		[[nodiscard]] FORCEINLINE const ElementType& operator*() const
+		{
+			return Array[GetIndex()];
+		}
+		[[nodiscard]] FORCEINLINE const ElementType* operator->() const
+		{
+			return &Array[GetIndex()];
+		}
+		[[nodiscard]] FORCEINLINE const FRelativeBitReference& GetRelativeBitReference() const
+		{
+			return BitArrayIt;
+		}
 	private:
 		const TSparseArray& Array;
 		TConstDualSetBitIterator<typename Allocator::BitArrayAllocator,SubsetAllocator> BitArrayIt;
@@ -1159,11 +1203,12 @@ private:
 		const PREDICATE_CLASS& Predicate;
 
 	public:
-		FElementCompareClass( const PREDICATE_CLASS& InPredicate )
+		[[nodiscard]] FElementCompareClass( const PREDICATE_CLASS& InPredicate )
 			: Predicate( InPredicate )
-		{}
+		{
+		}
 
-		bool operator()( const FElementOrFreeListLink& A,const FElementOrFreeListLink& B ) const
+		[[nodiscard]] bool operator()( const FElementOrFreeListLink& A,const FElementOrFreeListLink& B ) const
 		{
 			return Predicate(*(ElementType*)&A.ElementData,*(ElementType*)&B.ElementData);
 		}
@@ -1327,7 +1372,7 @@ class TScriptSparseArray
 	using DerivedType = std::conditional_t<std::is_void_v<InDerivedType>, TScriptSparseArray, InDerivedType>;
 
 public:
-	static FScriptSparseArrayLayout GetScriptLayout(int32 ElementSize, int32 ElementAlignment)
+	[[nodiscard]] static FScriptSparseArrayLayout GetScriptLayout(int32 ElementSize, int32 ElementAlignment)
 	{
 		FScriptSparseArrayLayout Result;
 		Result.Alignment     = FMath::Max(ElementAlignment, (int32)alignof(FFreeListLink));
@@ -1336,7 +1381,7 @@ public:
 		return Result;
 	}
 
-	TScriptSparseArray()
+	[[nodiscard]] TScriptSparseArray()
 		: FirstFreeIndex(-1)
 		, NumFreeIndices(0)
 	{
@@ -1348,11 +1393,11 @@ public:
 	constexpr static bool bHasIntrusiveUnsetOptionalState = true;
 	using IntrusiveUnsetOptionalStateType = TScriptSparseArray;
 
-	explicit TScriptSparseArray(FIntrusiveUnsetOptionalState Tag)
+	[[nodiscard]] explicit TScriptSparseArray(FIntrusiveUnsetOptionalState Tag)
 		: Data(Tag)
 	{
 	}
-	bool operator==(FIntrusiveUnsetOptionalState Tag) const
+	[[nodiscard]] bool operator==(FIntrusiveUnsetOptionalState Tag) const
 	{
 		return Data == Tag;
 	}
@@ -1360,47 +1405,47 @@ public:
 	// End - intrusive TOptional<TSparseArray> state //
 	///////////////////////////////////////////////////
 
-	bool IsValidIndex(int32 Index) const
+	[[nodiscard]] bool IsValidIndex(int32 Index) const
 	{
 		return AllocationFlags.IsValidIndex(Index) && AllocationFlags[Index];
 	}
 
-	bool IsAllocated(int32 Index) const
+	[[nodiscard]] bool IsAllocated(int32 Index) const
 	{
 		return AllocationFlags[Index];
 	}
 
-	bool IsEmpty() const
+	[[nodiscard]] bool IsEmpty() const
 	{
 		return Data.Num() == NumFreeIndices;
 	}
 
-	int32 Num() const
+	[[nodiscard]] int32 Num() const
 	{
 		return Data.Num() - NumFreeIndices;
 	}
 
-	int32 NumUnchecked() const
+	[[nodiscard]] int32 NumUnchecked() const
 	{
 		return Data.NumUnchecked() - NumFreeIndices;
 	}
 
-	int32 GetMaxIndex() const
+	[[nodiscard]] int32 GetMaxIndex() const
 	{
 		return Data.Num();
 	}
 
-	bool IsCompact() const
+	[[nodiscard]] bool IsCompact() const
 	{
 		return NumFreeIndices == 0;
 	}
 
-	void* GetData(int32 Index, const FScriptSparseArrayLayout& Layout)
+	[[nodiscard]] void* GetData(int32 Index, const FScriptSparseArrayLayout& Layout)
 	{
 		return (uint8*)Data.GetData() + Layout.Size * Index;
 	}
 
-	const void* GetData(int32 Index, const FScriptSparseArrayLayout& Layout) const
+	[[nodiscard]] const void* GetData(int32 Index, const FScriptSparseArrayLayout& Layout) const
 	{
 		return (const uint8*)Data.GetData() + Layout.Size * Index;
 	}
diff --git a/Engine/Source/Runtime/Core/Public/Containers/SpscQueue.h b/Engine/Source/Runtime/Core/Public/Containers/SpscQueue.h
--- a/Engine/Source/Runtime/Core/Public/Containers/SpscQueue.h
+++ b/Engine/Source/Runtime/Core/Public/Containers/SpscQueue.h
@@ -21,7 +21,7 @@ public:
 
 	UE_NONCOPYABLE(TSpscQueue);
 
-	TSpscQueue()
+	[[nodiscard]] TSpscQueue()
 	{
 		FNode* Node = ::new(AllocatorType::Malloc(sizeof(FNode), alignof(FNode))) FNode;
 		Tail.store(Node, std::memory_order_relaxed);
@@ -63,7 +63,7 @@ public:
 		Head = Node;
 	}
 
-	// returns empty TOptional if queue is empty 
+	// returns empty TOptional if queue is empty
 	TOptional<ElementType> Dequeue()
 	{
 		FNode* LocalTail = Tail.load(std::memory_order_relaxed);
@@ -93,7 +93,7 @@ public:
 		return false;
 	}
 
-	bool IsEmpty() const
+	[[nodiscard]] bool IsEmpty() const
 	{
 		FNode* LocalTail = Tail.load(std::memory_order_relaxed);
 		FNode* LocalTailNext = LocalTail->Next.load(std::memory_order_acquire);
@@ -103,7 +103,7 @@ public:
 	// as there can be only one consumer, a consumer can safely "peek" the tail of the queue.
 	// returns a pointer to the tail if the queue is not empty, nullptr otherwise
 	// there's no overload with TOptional as it doesn't support references
-	ElementType* Peek() const
+	[[nodiscard]] ElementType* Peek() const
 	{
 		FNode* LocalTail = Tail.load(std::memory_order_relaxed);
 		FNode* LocalTailNext = LocalTail->Next.load(std::memory_order_acquire);
diff --git a/Engine/Source/Runtime/Core/Public/Containers/StackTracker.h b/Engine/Source/Runtime/Core/Public/Containers/StackTracker.h
--- a/Engine/Source/Runtime/Core/Public/Containers/StackTracker.h
+++ b/Engine/Source/Runtime/Core/Public/Containers/StackTracker.h
@@ -41,7 +41,7 @@ public:
 	static void DefaultDeleteUserDataFn(void* UserData);
 
 	/** Constructor, initializing all member variables */
-	CORE_API FStackTracker(StackTrackerUpdateFn InUpdateFn = nullptr, StackTrackerReportFn InReportFn = nullptr, StackTrackerDeleteUserDataFn InDeleteUserDataFn = nullptr, bool bInIsEnabled = false);
+	[[nodiscard]] CORE_API FStackTracker(StackTrackerUpdateFn InUpdateFn = nullptr, StackTrackerReportFn InReportFn = nullptr, StackTrackerDeleteUserDataFn InDeleteUserDataFn = nullptr, bool bInIsEnabled = false);
 
 	/**
 	 * Captures the current stack and updates stack tracking information.
diff --git a/Engine/Source/Runtime/Core/Public/Containers/StaticArray.h b/Engine/Source/Runtime/Core/Public/Containers/StaticArray.h
--- a/Engine/Source/Runtime/Core/Public/Containers/StaticArray.h
+++ b/Engine/Source/Runtime/Core/Public/Containers/StaticArray.h
@@ -25,11 +25,11 @@ class alignas(Alignment) TStaticArray
 public:
 	using ElementType = InElementType;
 
-	TStaticArray() = default;
+	[[nodiscard]] TStaticArray() = default;
 
 	// Constructs each element with Args
 	template <typename... ArgTypes>
-	explicit TStaticArray(EInPlace, ArgTypes&&... Args)
+	[[nodiscard]] explicit TStaticArray(EInPlace, ArgTypes&&... Args)
 		: Storage(InPlace, TMakeIntegerSequence<uint32, NumElements>(), Forward<ArgTypes>(Args)...)
 	{
 	}
@@ -39,31 +39,31 @@ public:
 		typename... ArgTypes
 		UE_REQUIRES((sizeof...(ArgTypes) > 0 && sizeof...(ArgTypes) <= NumElements) && UE::Core::Private::TCanBeConvertedToFromAll_V<InElementType, ArgTypes...>)
 	>
-	TStaticArray(ArgTypes&&... Args)
+	[[nodiscard]] TStaticArray(ArgTypes&&... Args)
 		: Storage(PerElement, Forward<ArgTypes>(Args)...)
 	{
 	}
 
-	TStaticArray(TStaticArray&& Other) = default;
-	TStaticArray(const TStaticArray& Other) = default;
+	[[nodiscard]] TStaticArray(TStaticArray&& Other) = default;
+	[[nodiscard]] TStaticArray(const TStaticArray& Other) = default;
 	TStaticArray& operator=(TStaticArray&& Other) = default;
 	TStaticArray& operator=(const TStaticArray& Other) = default;
 
 	// Accessors.
-	FORCEINLINE_DEBUGGABLE InElementType& operator[](uint32 Index)
+	[[nodiscard]] FORCEINLINE_DEBUGGABLE InElementType& operator[](uint32 Index)
 	{
 		checkSlow(Index < NumElements);
 		return Storage.Elements[Index].Element;
 	}
 
-	FORCEINLINE_DEBUGGABLE const InElementType& operator[](uint32 Index) const
+	[[nodiscard]] FORCEINLINE_DEBUGGABLE const InElementType& operator[](uint32 Index) const
 	{
 		checkSlow(Index < NumElements);
 		return Storage.Elements[Index].Element;
 	}
 
 	// Comparisons.
-	friend bool operator==(const TStaticArray& A,const TStaticArray& B)
+	[[nodiscard]] friend bool operator==(const TStaticArray& A,const TStaticArray& B)
 	{
 		for(uint32 ElementIndex = 0;ElementIndex < NumElements;++ElementIndex)
 		{
@@ -75,7 +75,7 @@ public:
 		return true;
 	}
 
-	bool operator!=(const TStaticArray& B) const
+	[[nodiscard]] bool operator!=(const TStaticArray& B) const
 	{
 		for(uint32 ElementIndex = 0;ElementIndex < NumElements;++ElementIndex)
 		{
@@ -93,27 +93,37 @@ public:
 	 * @returns True if the array is empty.
 	 * @see Num
 	 */
-	bool IsEmpty() const
+	[[nodiscard]] bool IsEmpty() const
 	{
 		return NumElements == 0;
 	}
 
 	/** The number of elements in the array. */
-	FORCEINLINE_DEBUGGABLE int32 Num() const { return NumElements; }
+	[[nodiscard]] FORCEINLINE_DEBUGGABLE int32 Num() const
+	{
+		return NumElements;
+	}
 
 	/** A pointer to the first element of the array */
-	FORCEINLINE_DEBUGGABLE       InElementType* GetData()       { static_assert((alignof(ElementType) % Alignment) == 0, "GetData() cannot be called on a TStaticArray with non-standard alignment"); return &Storage.Elements[0].Element; }
-	FORCEINLINE_DEBUGGABLE const InElementType* GetData() const { static_assert((alignof(ElementType) % Alignment) == 0, "GetData() cannot be called on a TStaticArray with non-standard alignment"); return &Storage.Elements[0].Element; }
+	[[nodiscard]] FORCEINLINE_DEBUGGABLE InElementType* GetData()
+	{
+		static_assert((alignof(ElementType) % Alignment) == 0, "GetData() cannot be called on a TStaticArray with non-standard alignment");
+		return &Storage.Elements[0].Element;
+	}
+	[[nodiscard]] FORCEINLINE_DEBUGGABLE const InElementType* GetData() const
+	{
+		return const_cast<TStaticArray*>(this)->GetData();
+	}
 
 private:
 
 	struct alignas(Alignment) TArrayStorageElementAligned
 	{
-		TArrayStorageElementAligned() = default;
+		[[nodiscard]] TArrayStorageElementAligned() = default;
 
 		// Index is used to achieve pack expansion in TArrayStorage's first constructor, but is unused here
 		template <typename... ArgTypes>
-		explicit TArrayStorageElementAligned(EInPlace, uint32 /*Index*/, ArgTypes&&... Args)
+		[[nodiscard]] explicit TArrayStorageElementAligned(EInPlace, uint32 /*Index*/, ArgTypes&&... Args)
 			: Element(Forward<ArgTypes>(Args)...)
 		{
 		}
@@ -123,10 +133,10 @@ private:
 
 	struct TArrayStorage
 	{
-		TArrayStorage() = default;
+		[[nodiscard]] TArrayStorage() = default;
 
 		template<uint32... Indices, typename... ArgTypes>
-		explicit TArrayStorage(EInPlace, TIntegerSequence<uint32, Indices...>, ArgTypes&&... Args)
+		[[nodiscard]] explicit TArrayStorage(EInPlace, TIntegerSequence<uint32, Indices...>, ArgTypes&&... Args)
 			: Elements{ TArrayStorageElementAligned(InPlace, Indices, Args...)... }
 		{
 			// The arguments are deliberately not forwarded arguments here, because we're initializing multiple elements
@@ -139,7 +149,7 @@ private:
 		}
 
 		template<typename... ArgTypes>
-		explicit TArrayStorage(EPerElement, ArgTypes&&... Args)
+		[[nodiscard]] explicit TArrayStorage(EPerElement, ArgTypes&&... Args)
 			: Elements{ TArrayStorageElementAligned(InPlace, 0 /* dummy index */, Forward<ArgTypes>(Args))... }
 		{
 		}
@@ -155,11 +165,12 @@ public:
 	template <typename StorageElementType, bool bReverse = false>
 	struct FRangedForIterator
 	{
-		explicit FRangedForIterator(StorageElementType* InPtr)
+		[[nodiscard]] explicit FRangedForIterator(StorageElementType* InPtr)
 			: Ptr(InPtr)
-		{}
+		{
+		}
 
-		auto& operator*() const
+		[[nodiscard]] auto& operator*() const
 		{
 			if constexpr (bReverse)
 			{
@@ -184,7 +195,7 @@ public:
 			return *this;
 		}
 
-		bool operator!=(const FRangedForIterator& B) const
+		[[nodiscard]] bool operator!=(const FRangedForIterator& B) const
 		{
 			return Ptr != B.Ptr;
 		}
@@ -199,19 +210,19 @@ public:
 	using RangedForConstReverseIteratorType = FRangedForIterator<const TArrayStorageElementAligned, true>;
 
 	/** STL-like iterators to enable range-based for loop support. */
-	FORCEINLINE RangedForIteratorType				begin()			{ return RangedForIteratorType(Storage.Elements); }
-	FORCEINLINE RangedForConstIteratorType			begin() const	{ return RangedForConstIteratorType(Storage.Elements); }
-	FORCEINLINE RangedForIteratorType				end()			{ return RangedForIteratorType(Storage.Elements + NumElements); }
-	FORCEINLINE RangedForConstIteratorType			end() const		{ return RangedForConstIteratorType(Storage.Elements + NumElements); }
-	FORCEINLINE RangedForReverseIteratorType		rbegin()		{ return RangedForReverseIteratorType(Storage.Elements + NumElements); }
-	FORCEINLINE RangedForConstReverseIteratorType	rbegin() const	{ return RangedForConstReverseIteratorType(Storage.Elements + NumElements); }
-	FORCEINLINE RangedForReverseIteratorType		rend()			{ return RangedForReverseIteratorType(Storage.Elements); }
-	FORCEINLINE RangedForConstReverseIteratorType	rend() const	{ return RangedForConstReverseIteratorType(Storage.Elements); }
+	[[nodiscard]] FORCEINLINE RangedForIteratorType             begin()        { return RangedForIteratorType(Storage.Elements); }
+	[[nodiscard]] FORCEINLINE RangedForConstIteratorType        begin()  const { return RangedForConstIteratorType(Storage.Elements); }
+	[[nodiscard]] FORCEINLINE RangedForIteratorType             end()          { return RangedForIteratorType(Storage.Elements + NumElements); }
+	[[nodiscard]] FORCEINLINE RangedForConstIteratorType        end()    const { return RangedForConstIteratorType(Storage.Elements + NumElements); }
+	[[nodiscard]] FORCEINLINE RangedForReverseIteratorType      rbegin()       { return RangedForReverseIteratorType(Storage.Elements + NumElements); }
+	[[nodiscard]] FORCEINLINE RangedForConstReverseIteratorType rbegin() const { return RangedForConstReverseIteratorType(Storage.Elements + NumElements); }
+	[[nodiscard]] FORCEINLINE RangedForReverseIteratorType      rend()         { return RangedForReverseIteratorType(Storage.Elements); }
+	[[nodiscard]] FORCEINLINE RangedForConstReverseIteratorType rend()   const { return RangedForConstReverseIteratorType(Storage.Elements); }
 };
 
 /** Creates a static array filled with the specified value. */
 template <typename InElementType, uint32 NumElements>
-TStaticArray<InElementType,NumElements> MakeUniformStaticArray(typename TCallTraits<InElementType>::ParamType InValue)
+[[nodiscard]] TStaticArray<InElementType,NumElements> MakeUniformStaticArray(typename TCallTraits<InElementType>::ParamType InValue)
 {
 	TStaticArray<InElementType,NumElements> Result;
 	for(uint32 ElementIndex = 0;ElementIndex < NumElements;++ElementIndex)
@@ -240,7 +251,7 @@ FArchive& operator<<(FArchive& Ar,TStaticArray<ElementType, NumElements, Alignme
 
 /** Hash function. */
 template <typename ElementType, uint32 NumElements, uint32 Alignment>
-uint32 GetTypeHash(const TStaticArray<ElementType, NumElements, Alignment>& Array)
+[[nodiscard]] uint32 GetTypeHash(const TStaticArray<ElementType, NumElements, Alignment>& Array)
 {
 	uint32 Hash = 0;
 	for (const ElementType& Element : Array)
