Renamed Templates/NotNull.h to Templates/GSLNotNull.h to avoid confusion with Misc/NotNull.h

diff --git a/Engine/Plugins/Runtime/StateTree/Source/StateTreeModule/Public/Debugger/StateTreeDebug.h b/Engine/Plugins/Runtime/StateTree/Source/StateTreeModule/Public/Debugger/StateTreeDebug.h
--- a/Engine/Plugins/Runtime/StateTree/Source/StateTreeModule/Public/Debugger/StateTreeDebug.h
+++ b/Engine/Plugins/Runtime/StateTree/Source/StateTreeModule/Public/Debugger/StateTreeDebug.h
@@ -3,7 +3,7 @@
 #pragma once
 
 #include "StateTreeTypes.h"
-#include "Templates/NotNull.h"
+#include "Templates/GSLNotNull.h"
 
 class UStateTree;
 struct FStateTreeExecutionContext;
diff --git a/Engine/Source/Runtime/Core/Public/Templates/GSLNotNull.h b/Engine/Source/Runtime/Core/Public/Templates/GSLNotNull.h
new file mode 100644
--- /dev/null
+++ b/Engine/Source/Runtime/Core/Public/Templates/GSLNotNull.h
@@ -0,0 +1,55 @@
+// Copyright Epic Games, Inc. All Rights Reserved.
+
+#pragma once
+
+#include "CoreTypes.h"
+#include "Misc/OptionalFwd.h"
+#include "Templates/UnrealTemplate.h"
+#include <gsl/pointers>
+
+class FArchive;
+
+/** Utility function to turn an `TOptional<gsl::not_null<T>>` back into a nullable T. */
+template<typename ObjectType>
+[[nodiscard]] inline ObjectType* GetRawPointerOrNull(const TOptional<gsl::not_null<ObjectType*>>& Optional)
+{
+	return Optional.IsSet() ? static_cast<ObjectType*>(Optional.GetValue()) : nullptr;
+}
+
+/** Utility function to turn an `TOptional<gsl::strict_not_null<T>>` back into a nullable T. */
+template<typename ObjectType>
+[[nodiscard]] inline ObjectType* GetRawPointerOrNull(const TOptional<gsl::strict_not_null<ObjectType*>>& Optional)
+{
+	return Optional.IsSet() ? static_cast<ObjectType*>(Optional.GetValue()) : nullptr;
+}
+
+
+/** Utility function to serialize a `gsl::not_null<T>. */
+template<typename ObjectType>
+FArchive& operator<<(FArchive& Ar, gsl::not_null<ObjectType>& NotNull)
+{
+	Ar << *NotNull;
+	return Ar;
+}
+
+/** Utility function to serialize a `gsl::strict_not_null<T>`. */
+template<typename ObjectType>
+FArchive& operator<<(FArchive& Ar, gsl::strict_not_null<ObjectType>& NotNull)
+{
+	Ar << *NotNull;
+	return Ar;
+}
+
+/** Utility function to hash a `gsl::not_null<T>`. */
+template<typename ObjectType>
+[[nodiscard]] inline auto GetTypeHash(const gsl::not_null<ObjectType>& NotNull) -> decltype(GetTypeHash(*NotNull))
+{
+	return GetTypeHash(*NotNull);
+}
+
+/** Utility function to hash a `gsl::strict_not_null<T>`. */
+template<typename ObjectType>
+[[nodiscard]] inline auto GetTypeHash(const gsl::strict_not_null<ObjectType>& NotNull) -> decltype(GetTypeHash(*NotNull))
+{
+	return GetTypeHash(*NotNull);
+}
diff --git a/Engine/Source/Runtime/Core/Public/Templates/NotNull.h b/Engine/Source/Runtime/Core/Public/Templates/NotNull.h
--- a/Engine/Source/Runtime/Core/Public/Templates/NotNull.h
+++ b/Engine/Source/Runtime/Core/Public/Templates/NotNull.h
@@ -2,54 +2,10 @@
 
 #pragma once
 
-#include "CoreTypes.h"
-#include "Misc/OptionalFwd.h"
-#include "Templates/UnrealTemplate.h"
-#include <gsl/pointers>
-
-class FArchive;
-
-/** Utility function to turn an `TOptional<gsl::not_null<T>>` back into a nullable T. */
-template<typename ObjectType>
-[[nodiscard]] inline ObjectType* GetRawPointerOrNull(const TOptional<gsl::not_null<ObjectType*>>& Optional)
-{
-	return Optional.IsSet() ? static_cast<ObjectType*>(Optional.GetValue()) : nullptr;
-}
-
-/** Utility function to turn an `TOptional<gsl::strict_not_null<T>>` back into a nullable T. */
-template<typename ObjectType>
-[[nodiscard]] inline ObjectType* GetRawPointerOrNull(const TOptional<gsl::strict_not_null<ObjectType*>>& Optional)
-{
-	return Optional.IsSet() ? static_cast<ObjectType*>(Optional.GetValue()) : nullptr;
-}
+// HEADER_UNIT_SKIP - Deprecated
 
+#include "CoreTypes.h"
 
-/** Utility function to serialize a `gsl::not_null<T>. */
-template<typename ObjectType>
-FArchive& operator<<(FArchive& Ar, gsl::not_null<ObjectType>& NotNull)
-{
-	Ar << *NotNull;
-	return Ar;
-}
-
-/** Utility function to serialize a `gsl::strict_not_null<T>`. */
-template<typename ObjectType>
-FArchive& operator<<(FArchive& Ar, gsl::strict_not_null<ObjectType>& NotNull)
-{
-	Ar << *NotNull;
-	return Ar;
-}
-
-/** Utility function to hash a `gsl::not_null<T>`. */
-template<typename ObjectType>
-[[nodiscard]] inline auto GetTypeHash(const gsl::not_null<ObjectType>& NotNull) -> decltype(GetTypeHash(*NotNull))
-{
-	return GetTypeHash(*NotNull);
-}
+UE_DEPRECATED_HEADER(5.6, "Templates/NotNull.h has been renamed to Templates/GSLNotNull.h to avoid confusion with Misc/NotNull.h which contains the TNotNull type")
 
-/** Utility function to hash a `gsl::strict_not_null<T>`. */
-template<typename ObjectType>
-[[nodiscard]] inline auto GetTypeHash(const gsl::strict_not_null<ObjectType>& NotNull) -> decltype(GetTypeHash(*NotNull))
-{
-	return GetTypeHash(*NotNull);
-}
+#include "Templates/GSLNotNull.h"
diff --git a/Engine/Source/Runtime/Core/Tests/Templates/NotNullTest.cpp b/Engine/Source/Runtime/Core/Tests/Templates/NotNullTest.cpp
--- a/Engine/Source/Runtime/Core/Tests/Templates/NotNullTest.cpp
+++ b/Engine/Source/Runtime/Core/Tests/Templates/NotNullTest.cpp
@@ -3,8 +3,8 @@
 #if WITH_TESTS
 
 #include "Misc/Optional.h"
+#include "Templates/GSLNotNull.h"
 #include "Templates/NonNullPointer.h"
-#include "Templates/NotNull.h"
 #include "Templates/SharedPointer.h"
 #include "Templates/Tuple.h"
 #include "Templates/UniquePtr.h"
