[core] speculative FormatStringSan intermittent compiler bug workaround

diff --git a/Engine/Source/Runtime/Core/Public/String/FormatStringSan.h b/Engine/Source/Runtime/Core/Public/String/FormatStringSan.h
--- a/Engine/Source/Runtime/Core/Public/String/FormatStringSan.h
+++ b/Engine/Source/Runtime/Core/Public/String/FormatStringSan.h
@@ -151,12 +151,12 @@ namespace UE::Core::Private
 		{
 			static constexpr FResult Check(bool bInsideFormatSpec, int CurArgPos, const TCHAR* Fmt)
 			{
-				auto CharIsDigit = [](const TCHAR C)
+				auto CharIsDigit = [](const TCHAR C) -> bool
 				{
 					return C >= TEXT('0') && C <= TEXT('9');
 				};
 
-				auto SkipInteger = [&CharIsDigit](const TCHAR* P)
+				auto SkipInteger = [&CharIsDigit](const TCHAR* P) -> const TCHAR*
 				{
 					while (CharIsDigit(*P))
 					{
@@ -165,7 +165,7 @@ namespace UE::Core::Private
 					return P;
 				};
 
-				auto CharIsIntegerFormatSpecifier = [](const TCHAR C)
+				auto CharIsIntegerFormatSpecifier = [](const TCHAR C) -> bool
 				{
 					switch (C)
 					{
