Stats - Add StatIgnoreList  which allows ignoring certain stats/groups on particular threads.

Currently only supported for lightweight stats.

Which stats/groups are ignored can be configured in the engine ini to allow per-project/per-platform control.

The main purpose of this change is to allow suprressing stats that are very noisy on particular threads but useful on others, so marking them verbose is not an option.
For example, on the AsyncLoading thread we can have thousands of STATGROUP_Object events in a short amount of time which can overwhelm certain profilers and don't provide a ton of value. So to avoid this but still have them on the game thread where they are useful context, we can mark them as ignored on that thread.

There is 0 overhead from this at runtime when named events are disabled (it's slightly faster than before due to moving other existing checks to compile time), and a pretty minor overhead when named events are enabled.

This CL isn't feature isn't fully complete as I still need to add support for non-lightweight stats, but that can be done in a follow up CL.

How it works:
1. On startup we load any ignored stats/groups from the engine config. The names of these stats/groups are hashed and store into FStatIgnoreList. After this initialization the ignore list is read-only.
2. The various SCOPE_CYCLE_COUNTER macros now call an internal UE_INTERNAL_GET_STAT_IF_ENABLED which hashes the stat name at compile time to be compared against the hashed ignored stats (if the other cheaper compile time/runtime checks pass).
3. UE_INTERNAL_GET_STAT_IF_ENABLED will return an FVerifiedStat which is just a wrapper around a TStatId, but this let's us take an optimized path in the FScopeCycleCounter to skip any checks that were already done. If the stat is not enabled (aka ignored), then an empty FVerifiedStat is returned which will fail the IsValidStat() check in the scope cycle counter.
4. Another improvement of this system is we no longer do several of the checks twice, once for the FScopeCycleCounter and another for the insights trace macro. Both can just check if the StatId is valid to determine if it should be emitted or not.

Other changes:
I've also added a FStaticScopeCycleCounter which makes uses of BeginNamedEventStatic for platforms that support it for extra performance.

#rnx
#tests local replay runs with tracing on/off

diff --git a/Engine/Source/Runtime/Core/Private/Stats/StatIgnoreList.cpp b/Engine/Source/Runtime/Core/Private/Stats/StatIgnoreList.cpp
new file mode 100644
--- /dev/null
+++ b/Engine/Source/Runtime/Core/Private/Stats/StatIgnoreList.cpp
@@ -0,0 +1,145 @@
+// Copyright Epic Games, Inc. All Rights Reserved.
+
+#include "Stats/StatIgnoreList.h"
+
+#include "CoreGlobals.h"
+#include "Containers/Map.h"
+#include "HAL/IConsoleManager.h"
+#include "Misc/ConfigCacheIni.h"
+#include "Misc/Optional.h"
+#include "String/ParseTokens.h"
+#include "String/Split.h"
+
+#if UE_STATS_ALLOW_PER_THREAD_IGNORELIST
+namespace UE::Stats::Private
+{
+	// The text representation of the thread names are what we match in the config.
+	// This let's us avoid exposing ETaskTag in any public facing API and restrict which threads
+	// stats can be ignored on.
+	static const TMap<FStringView, ETaskTag> NamedThreadMap
+	{
+		{ TEXTVIEW("GameThread"), ETaskTag::EGameThread },
+		{ TEXTVIEW("SlateThread"), ETaskTag::ESlateThread },
+		{ TEXTVIEW("RenderingThread"), ETaskTag::ERenderingThread },
+		{ TEXTVIEW("RhiThread"), ETaskTag::ERhiThread },
+		{ TEXTVIEW("AsyncLoadingThread"), ETaskTag::EAsyncLoadingThread }
+	};
+
+	struct FStatIgnoreList
+	{
+		FORCEINLINE static FStatIgnoreList& Get()
+		{
+			static FStatIgnoreList Instance;
+			return Instance;
+		}
+
+		void Initialize(TMap<uint32, ETaskTag>&& InIgnoreMap)
+		{
+			check(!bInitialized);
+			IgnoreMap = MoveTemp(InIgnoreMap);
+			bInitialized.store(true, std::memory_order_release);
+		}
+
+		inline bool IsStatOrGroupIgnored(uint32 StatNameHash, uint32 GroupNameHash)
+		{
+			if (bInitialized.load(std::memory_order_acquire))
+			{
+				ETaskTag CurrentThread = ETaskTag::ENone;
+				if (const ETaskTag* IgnoredThreads = IgnoreMap.Find(StatNameHash))
+				{
+					// Note: FTaskTagScope::GetCurrentTag() is reading from a thread_local which can be quite slow
+					// so we only do it if we know this stat is ignored on some thread.
+					CurrentThread = FTaskTagScope::GetCurrentTag();
+					return ((*IgnoredThreads & CurrentThread) == CurrentThread);
+				}
+
+				if (GroupNameHash != 0)
+				{
+					if (const ETaskTag* IgnoredThreads = IgnoreMap.Find(GroupNameHash))
+					{
+						CurrentThread = CurrentThread == ETaskTag::ENone ? FTaskTagScope::GetCurrentTag() : CurrentThread;
+						return ((*IgnoredThreads & CurrentThread) == CurrentThread);
+					}
+				}
+			}
+			return false;
+		}
+
+	private:
+		TMap<uint32, ETaskTag> IgnoreMap;
+		std::atomic_bool bInitialized = false;
+	};
+
+} // namespace UE::Stats::Private
+
+namespace UE::Stats
+{
+	void InitializeIgnoreList()
+	{
+		if (!GConfig)
+		{
+			return;
+		}
+
+		static constexpr const TCHAR* SectionName = TEXT("Stats.PerThreadIgnoreList");
+
+		bool bIgnoreListEnabled = false;
+		GConfig->GetBool(SectionName, TEXT("IgnoreListEnabled"), bIgnoreListEnabled, GEngineIni);
+		if (!bIgnoreListEnabled)
+		{
+			return;
+		}
+
+		const auto ParseIgnoreList = [](const TCHAR* ListName, TMap<uint32, ETaskTag>& OutIgnoreMap)
+		{
+			TArray<FString> IgnoredEntries;
+			if (GConfig->GetArray(SectionName, ListName, IgnoredEntries, GEngineIni) == 0)
+			{
+				return;
+			}
+
+			TMap<FStringView, ETaskTag> NameToIgnoredThreadsMap;
+
+			for (FStringView Entry : IgnoredEntries)
+			{
+				FStringView Name, ThreadsString;
+				if (UE::String::SplitFirst(Entry, TEXTVIEW(":"), Name, ThreadsString))
+				{
+					ETaskTag& IgnoredThreads = NameToIgnoredThreadsMap.FindOrAdd(Name);
+
+					UE::String::ParseTokens(ThreadsString, TEXTVIEW("|"), [&IgnoredThreads](FStringView ThreadName)
+					{
+						if (const ETaskTag* Tag = Private::NamedThreadMap.Find(ThreadName.TrimStartAndEnd()))
+						{
+							IgnoredThreads |= *Tag;
+						}
+					});
+				}
+			}
+
+			for (const TPair<FStringView, ETaskTag>& Pair : NameToIgnoredThreadsMap)
+			{
+				if (Pair.Value != ETaskTag::ENone)
+				{
+					const uint32 Hash = UE::HashStringFNV1a32(Pair.Key);
+					OutIgnoreMap.Add(Hash, Pair.Value);
+				}
+			}
+		};
+
+		TMap<uint32, ETaskTag> IgnoreMap;
+		// Note: we could combine these into a single list but this makes it easier to change 
+		// how we store the data if we want later.
+		ParseIgnoreList(TEXT("IgnoredStats"), IgnoreMap);
+		ParseIgnoreList(TEXT("IgnoredGroups"), IgnoreMap);
+
+		Private::FStatIgnoreList::Get().Initialize(MoveTemp(IgnoreMap));
+	}
+
+	bool IsStatOrGroupIgnoredOnCurrentThread(uint32 StatNameHash, uint32 GroupNameHash)
+	{
+		return Private::FStatIgnoreList::Get().IsStatOrGroupIgnored(StatNameHash, GroupNameHash);
+	}
+} // namespace UE::Stats
+
+#endif // UE_STATS_ALLOW_PER_THREAD_IGNORELIST
diff --git a/Engine/Source/Runtime/Core/Private/Stats/StatsSystem.cpp b/Engine/Source/Runtime/Core/Private/Stats/StatsSystem.cpp
--- a/Engine/Source/Runtime/Core/Private/Stats/StatsSystem.cpp
+++ b/Engine/Source/Runtime/Core/Private/Stats/StatsSystem.cpp
@@ -4,12 +4,16 @@
 
 #include "Async/TaskGraphInterfaces.h"
 #include "Containers/Ticker.h"
+#include "CoreGlobals.h"
 #include "HAL/LowLevelMemTracker.h"
 #include "Misc/CommandLine.h"
+#include "Misc/ConfigCacheIni.h"
+#include "Misc/CoreDelegates.h"
 #include "ProfilingDebugging/CpuProfilerTrace.h"
 #include "Stats/Stats.h"
 #include "Stats/StatsData.h"
 #include "Stats/StatsSystemTypes.h"
+#include "Stats/StatIgnoreList.h"
 
 // Declared in GlobalStats.ini
 DEFINE_STAT(STAT_FrameTime);
@@ -20,6 +24,30 @@ namespace UE::Stats
 {
 	TAtomic<int32> FStats::GameThreadStatsFrame(1);
 
+	namespace Private
+	{
+		static void OnInit()
+		{
+#if UE_STATS_ALLOW_PER_THREAD_IGNORELIST
+			InitializeIgnoreList();
+#endif
+		}
+	}
+
+	void FStats::Init()
+	{
+#if UE_STATS_ALLOW_PER_THREAD_IGNORELIST
+		if (GConfig != nullptr && GConfig->IsReadyForUse())
+		{
+			InitializeIgnoreList();
+		}
+		else
+		{
+			FCoreDelegates::OnInit.AddStatic(&Private::OnInit);
+		}
+#endif // UE_STATS_ALLOW_PER_THREAD_IGNORELIST
+	}
+
 	void FStats::AdvanceFrame(bool bDiscardCallstack, const FOnAdvanceRenderingThreadStats& AdvanceRenderingThreadStatsDelegate /*= FOnAdvanceRenderingThreadStats()*/)
 	{
 #if STATS
diff --git a/Engine/Source/Runtime/Core/Public/Stats/LightweightStats.h b/Engine/Source/Runtime/Core/Public/Stats/LightweightStats.h
--- a/Engine/Source/Runtime/Core/Public/Stats/LightweightStats.h
+++ b/Engine/Source/Runtime/Core/Public/Stats/LightweightStats.h
@@ -6,15 +6,15 @@
 #include "Misc/Build.h"
 #include "StatsCommon.h"
 
-#ifndef UE_USE_LIGHTWEIGHT_STATS
-#define UE_USE_LIGHTWEIGHT_STATS (!STATS && ENABLE_STATNAMEDEVENTS)
-#endif
-
+// UE_USE_LIGHTWEIGHT_STATS is defined in StatsCommon.h
 #if UE_USE_LIGHTWEIGHT_STATS
 
 #include "AutoRTFM/AutoRTFM.h"
+#include "Hash/Fnv.h"
 #include "Math/Color.h"
+#include "Misc/NotNull.h"
 #include "Stats/HitchTrackingStatScope.h"
+#include "Stats/StatIgnoreList.h"
 
 struct TStatId
 {
@@ -30,6 +30,11 @@ struct TStatId
 	{
 	}
 
+	FORCEINLINE constexpr const PROFILER_CHAR* GetName() const
+	{
+		return StatString;
+	}
+
 	FORCEINLINE constexpr bool IsValidStat() const
 	{
 		return StatString != nullptr;
@@ -48,90 +53,119 @@ struct TStatId
 
 namespace UE::Stats::Private
 {
-	// We cannot add members to TStatId as one is stored per uobject which results in a large memory overhead,
-	// so instead we utilze a proxy object internally that can implictly convert back into a TStatId when needed.
-	struct TStatIdProxy
+	struct FCheckedStat : TStatId
 	{
-		TStatId StatId;
-		const PROFILER_CHAR* Group = nullptr;
-		EStatFlags Flags = EStatFlags::None;
+		using TStatId::TStatId;
+	};
 
-		constexpr TStatIdProxy() = default;
+	template<typename DerivedScopeCounterType, typename StatType = TStatId>
+	class TScopeCycleCounterBase
+	{
+	public:
+		FORCEINLINE explicit TScopeCycleCounterBase(StatType InStatId)
+#if USE_LIGHTWEIGHT_STATS_FOR_HITCH_DETECTION && USE_HITCH_DETECTION
+			: StatScope(InStatId.GetName())
+#endif
+		{
+			if (DerivedScopeCounterType::CanEmitStat(InStatId))
+			{
+				UE_AUTORTFM_OPEN_NO_VALIDATION
+				{
+					DerivedScopeCounterType::BeginNamedEvent(InStatId.GetName());
+					bPop = true;
+				};
+				AutoRTFM::PushOnAbortHandler(this, []() { FPlatformMisc::EndNamedEvent(); });
+			}
+		}
 
-		// Allow implicit conversions from nullptr to support conditional scopes
-		constexpr TStatIdProxy(TYPE_OF_NULLPTR)
+		FORCEINLINE ~TScopeCycleCounterBase()
 		{
+			AutoRTFM::PopOnAbortHandler(this);
+
+			// Memory validation disabled as this is hot code.
+			UE_AUTORTFM_OPEN_NO_VALIDATION
+			{
+				if (bPop)
+				{
+					FPlatformMisc::EndNamedEvent();
+				}
+			};
 		}
 
-		constexpr explicit TStatIdProxy(const PROFILER_CHAR* InString, const PROFILER_CHAR* InGroup = nullptr, EStatFlags InFlags = EStatFlags::None)
-			: StatId(InString)
-			, Group(InGroup)
-			, Flags(InFlags)
+		FORCEINLINE static bool CanEmitStat(const TStatId& Stat)
 		{
+			return Stat.IsValidStat() && GCycleStatsShouldEmitNamedEvents > 0;
 		}
 
-		// Allow implicit conversions to TStatId to keep macros like GET_STATID compatible
-		// with code that takes a TStatId.
-		constexpr operator TStatId() const
+	private:
+#if USE_LIGHTWEIGHT_STATS_FOR_HITCH_DETECTION && USE_HITCH_DETECTION
+		UE::Stats::FHitchTrackingStatScope StatScope;
+#endif
+		bool bPop = false;
+	};
+
+	/**
+	 * Scope counter for stat names with static storage.
+	 */
+	template<typename StatIdType>
+	class FScopeCycleCounterStatic : TScopeCycleCounterBase<FScopeCycleCounterStatic<StatIdType>, StatIdType>
+	{
+		using Base = TScopeCycleCounterBase<FScopeCycleCounterStatic<StatIdType>, StatIdType>;
+	public:
+		using Base::CanEmitStat;
+
+		FORCEINLINE constexpr explicit FScopeCycleCounterStatic(StatIdType InStatId)
+			: Base(InStatId)
 		{
-			return StatId;
+		}
+
+		FORCEINLINE static constexpr bool CanEmitStat(const FCheckedStat& Stat)
+		{
+			// If we were given a FCheckedStat it means it's gone through GetIfStatEnabled so
+			// GCycleStatsShouldEmitNamedEvents has already been checked.
+			return Stat.IsValidStat();
+		}
+
+		FORCEINLINE static void BeginNamedEvent(TNotNull<const PROFILER_CHAR* const> StatName)
+		{
+#if PLATFORM_IMPLEMENTS_BeginNamedEventStatic
+			FPlatformMisc::BeginNamedEventStatic(FColor(0), StatName);
+#else
+			FPlatformMisc::BeginNamedEvent(FColor(0), StatName);
+#endif
 		}
 	};
+
 } // namespace UE::Stats::Private
 
-class FScopeCycleCounter
+// ----------------------------------------------------------------
+// FScopeCycleCounter
+// ----------------------------------------------------------------
+
+// Scope counter for stat names with dynamic storage.
+class FScopeCycleCounter : private UE::Stats::Private::TScopeCycleCounterBase<FScopeCycleCounter>
 {
+	using Base = UE::Stats::Private::TScopeCycleCounterBase<FScopeCycleCounter>;
+
 public:
-	FORCEINLINE FScopeCycleCounter(TStatId InStatId, EStatFlags StatFlags, bool bAlways = false)
-		:
-#if USE_LIGHTWEIGHT_STATS_FOR_HITCH_DETECTION && USE_HITCH_DETECTION
-		StatScope(InStatId.StatString),
-#endif
+	using Base::CanEmitStat;
 
-		bPop(false)
+	// NOTE: this signature must match the other FScopeCycleCounter implementations
+	FORCEINLINE FScopeCycleCounter(TStatId InStatId, EStatFlags, bool bAlways = false)
+		: Base(InStatId)
 	{
-		if (GCycleStatsShouldEmitNamedEvents &&
-			InStatId.IsValidStat() &&
-			(GShouldEmitVerboseNamedEvents || !EnumHasAnyFlags(StatFlags, EStatFlags::Verbose)))
-		{
-			// Memory validation disabled as this is hot code.
-			UE_AUTORTFM_OPEN_NO_VALIDATION
-			{
-				bPop = true;
-				FPlatformMisc::BeginNamedEvent(FColor(0), InStatId.StatString);
-			};
-			AutoRTFM::PushOnAbortHandler(this, []() { FPlatformMisc::EndNamedEvent(); });
-		}
 	}
 
+	// NOTE: this signature must match the other FScopeCycleCounter implementations
 	FORCEINLINE FScopeCycleCounter(TStatId InStatId, bool bAlways = false)
 		: FScopeCycleCounter(InStatId, EStatFlags::None, bAlways)
 	{
 	}
 
-	FORCEINLINE explicit FScopeCycleCounter(const UE::Stats::Private::TStatIdProxy& InStatIdProxy)
-		: FScopeCycleCounter(InStatIdProxy.StatId, InStatIdProxy.Flags, false)
+	FORCEINLINE static void BeginNamedEvent(TNotNull<const PROFILER_CHAR* const> StatName)
 	{
+		FPlatformMisc::BeginNamedEvent(FColor(0), StatName);
 	}
-
-	FORCEINLINE ~FScopeCycleCounter()
-	{
-		AutoRTFM::PopOnAbortHandler(this);
-
-		// Memory validation disabled as this is hot code.
-		UE_AUTORTFM_OPEN_NO_VALIDATION
-		{
-			if (bPop)
-			{
-				FPlatformMisc::EndNamedEvent();
-			}
-		};
-	}
-private:
-#if USE_LIGHTWEIGHT_STATS_FOR_HITCH_DETECTION && USE_HITCH_DETECTION
-	UE::Stats::FHitchTrackingStatScope StatScope;
-#endif
-	bool bPop;
 };
 
 // ----------------------------------------------------------------
@@ -139,10 +173,34 @@ private:
 // ----------------------------------------------------------------
 namespace UE::Stats::Private
 {
+	// ----------------------------------------------------------------
+	// Stat/Group Getter Helpers
+	// ----------------------------------------------------------------
+	
+	// Fallback group data for groups that don't have a custom group struct defined.
+	struct FDefaultGroupData
+	{
+		static constexpr const PROFILER_CHAR* GetName()
+		{
+			return nullptr;
+		}
+		static constexpr uint32 GetNameHash()
+		{
+			return 0;
+		}
+		// Note: matches spelling of implementation when STATS is 1.
+		static constexpr bool IsCompileTimeEnable()
+		{
+			return true;
+		}
+	};
+
 	// Fallback stat data for StatId's that don't have a custom stat struct defined.
 	struct FDefaultStatData
 	{
-		static constexpr const PROFILER_CHAR* GetGroup()
+		using TGroup = FDefaultGroupData;
+
+		static constexpr const PROFILER_CHAR* GetName()
 		{
 			return nullptr;
 		}
@@ -189,10 +247,76 @@ namespace UE::Stats::Private
 	template<typename AllowDefaultStatToken>
 	constexpr bool AllowDefaultForStat = TAllowDefaultStatHelper<AllowDefaultStatToken>::Value;
 
-	template<typename StatDataType>
-	constexpr TStatIdProxy MakeStatId(const PROFILER_CHAR* StatName)
+	template<typename GroupDataType, typename = void>
+	struct TGroupDataHelper
+	{
+		using Type = FDefaultGroupData;
+	};
+	template<typename GroupDataType>
+	struct TGroupDataHelper<GroupDataType, std::void_t<decltype(sizeof(GroupDataType))>>
+	{
+		using Type = GroupDataType;
+	};
+
+	// Helper defined as either the custom group data struct if it's defined, or the default (FDefaultGroupData).
+	template<typename GroupDataType>
+	using TGroupDataType = typename TGroupDataHelper<GroupDataType>::Type;
+
+	// Tests if this stat can be emitted
+	template<typename StatDataType = FDefaultStatData, typename GroupDataType = typename StatDataType::TGroup>
+	inline bool IsStatEnabled(uint32 StatNameHash, uint32 GroupNameHash = 0)
+	{
+		if constexpr (!GroupDataType::IsCompileTimeEnable())
+		{
+			return false;
+		}
+
+		if constexpr (EnumHasAnyFlags(StatDataType::GetFlags(), EStatFlags::Verbose))
+		{
+			if (!GShouldEmitVerboseNamedEvents)
+			{
+				return false;
+			}
+		}
+
+#if UE_STATS_ALLOW_PER_THREAD_IGNORELIST
+		const uint32 GroupHash = GroupNameHash != 0 ? GroupNameHash : GroupDataType::GetNameHash();
+		if (UE::Stats::IsStatOrGroupIgnoredOnCurrentThread(StatNameHash, GroupHash))
+		{
+			return false;
+		}
+#endif // UE_STATS_ALLOW_PER_THREAD_IGNORELIST
+		return true;
+	}
+
+	/**
+	 * Returns a valid TStatId if we can emit it.
+	 * 
+	 * StatNameHash and GroupNameHash must be calculated at compile time to avoid a very large overhead.
+	 * GroupNameHash can be 0 if the stat has no group or the group is default.
+	 * Both of these values will be 0 if UE_STATS_ALLOW_PER_THREAD_IGNORELIST is 0.
+	 */
+	template<typename StatDataType = FDefaultStatData, typename GroupDataType = typename StatDataType::TGroup>
+	FORCEINLINE FCheckedStat GetStatIfEnabled(TNotNull<const PROFILER_CHAR*> StatName, uint32 StatNameHash, uint32 GroupNameHash = 0)
+	{
+		// We're testing if named events are enabled here as it's cheaper than some of the checks in IsStatEnabled.
+		if (GCycleStatsShouldEmitNamedEvents > 0 &&
+			IsStatEnabled<StatDataType, GroupDataType>(StatNameHash, GroupNameHash))
+		{
+			return FCheckedStat{ StatName };
+		}
+		return FCheckedStat{};
+	}
+
+	/* Returns a valid TStatId if the condition is true and we can emit it. */
+	template<typename StatDataType = FDefaultStatData, typename GroupDataType = typename StatDataType::TGroup>
+	FORCEINLINE FCheckedStat GetStatIfEnabled(TNotNull<const PROFILER_CHAR*> StatName, bool bCondition, uint32 StatNameHash, uint32 GroupNameHash = 0)
 	{
-		return TStatIdProxy(StatName, StatDataType::GetGroup(), StatDataType::GetFlags());
+		if (bCondition)
+		{
+			return GetStatIfEnabled<StatDataType, GroupDataType>(StatName, StatNameHash, GroupNameHash);
+		}
+		return FCheckedStat{};
 	}
 } // namespace UE::Stats::Private
 
@@ -200,6 +324,9 @@ namespace UE::Stats::Private
 // Helper Macros
 // ----------------------------------------------------------------
 
+#define UE_INTERNAL_GET_STATGROUP_TYPE(GroupId)\
+	UE::Stats::Private::TGroupDataType<struct FStatGroup_##GroupId>
+
 // Specifies that a custom FStat struct doesn't need to be defined for this stat and we can fallback to the default.
 // This is for internal use only currently.
 #define UE_INTERNAL_ALLOW_DEFAULT_STAT(Stat)\
@@ -221,6 +348,20 @@ namespace UE::Stats::Private
 #define UE_IS_STAT_VERBOSE_ALLOW_DEFAULT(Stat)\
 	EnumHasAnyFlags(UE_INTERNAL_GET_STATDATA_ALLOW_DEFAULT(Stat)::GetFlags(), EStatFlags::Verbose)
 
+// Getters
+#define UE_INTERNAL_GET_STAT_IF_ENABLED(Stat)\
+	UE::Stats::Private::GetStatIfEnabled<UE_INTERNAL_GET_STATDATA(Stat)>(ANSI_TO_PROFILING(#Stat), UE_STATS_HASH_NAME(Stat))
+
+#define UE_INTERNAL_GET_STAT_IF_ENABLED_COND(Stat, Cond)\
+	UE::Stats::Private::GetStatIfEnabled<UE_INTERNAL_GET_STATDATA(Stat)>(ANSI_TO_PROFILING(#Stat), (Cond), UE_STATS_HASH_NAME(Stat))
+
+#define UE_INTERNAL_GET_QUICK_STAT_IF_ENABLED(Stat)\
+	UE::Stats::Private::GetStatIfEnabled<UE_INTERNAL_GET_STATDATA_ALLOW_DEFAULT(Stat)>(ANSI_TO_PROFILING(#Stat), UE_STATS_HASH_NAME(Stat))
+
+#define UE_INTERNAL_GET_QUICK_STAT_WITH_GROUP_IF_ENABLED(Stat, Group)\
+	UE::Stats::Private::GetStatIfEnabled<UE_INTERNAL_GET_STATDATA_ALLOW_DEFAULT(Stat), UE_INTERNAL_GET_STATGROUP_TYPE(Group)>(ANSI_TO_PROFILING(#Stat), UE_STATS_HASH_NAME(Stat), UE_STATS_HASH_NAME(Group))
+
+
 #elif !STATS // UE_USE_LIGHTWEIGHT_STATS
 
 struct TStatId {};
diff --git a/Engine/Source/Runtime/Core/Public/Stats/StatIgnoreList.h b/Engine/Source/Runtime/Core/Public/Stats/StatIgnoreList.h
new file mode 100644
--- /dev/null
+++ b/Engine/Source/Runtime/Core/Public/Stats/StatIgnoreList.h
@@ -0,0 +1,63 @@
+// Copyright Epic Games, Inc. All Rights Reserved.
+#pragma once
+
+#include "CoreTypes.h"
+#include "Hash/Fnv.h"
+#include "Stats/StatsCommon.h"
+
+#ifndef UE_STATS_ALLOW_PER_THREAD_IGNORELIST
+// TODO: support in regular stats as well.
+#define UE_STATS_ALLOW_PER_THREAD_IGNORELIST UE_USE_LIGHTWEIGHT_STATS
+#endif
+
+/**
+ * This is an internal API that allows individual stats or groups to be ignored on certain threads such that they will not be emitted to profilers.
+ * 
+ * The primary use case for this is to silence super verbose stats when they're only a problem on certain threads but still desired on others.
+ * Marking stats as Verbose or removing them entirely is still preferred before utilizing this system as there is a slight cost that 
+ * scales with the number of ignored stats (only while named events are enabled).
+ * 
+ * The per-thread ignore lists can be configured per platform as they are read from the *Engine.ini file. See InitializeIgnoreList for details.
+ */
+#if UE_STATS_ALLOW_PER_THREAD_IGNORELIST
+namespace UE::Stats
+{
+	/**
+	 * Loads the ignored stats/groups from the engine config and initializes the ignore list.
+	 * IsStatIgnored/IsStatOrGroupIgnored will always return false until this initialization has occured.
+	 * 
+	 * Example config setup:
+	 * [Stats.PerThreadIgnoreList]
+	 * IgnoreListEnabled=true ; enable the system
+	 * IgnoredStats=STAT_MyStat:GameThread|RenderingThread ; ignore STAT_MyStat on the GameThread and RenderingThread.
+	 * +IgnoredStats=STAT_MyOtherStat:RhiThread
+	 * IgnoredGroups=STATGROUP_Object:AsyncLoadingThread ; Ignore this entire group on the AsyncLoadingThread.
+	 * +IgnoredGroups=STATGROUP_ObjectVerbose:AsyncLoadingThread
+	 * 
+	 * See NamedThreadMap in StatIgnoreList.cpp for the full list of valid threads.
+	 */
+	void InitializeIgnoreList();
+
+	/**
+	 * Returns true if the stat is ignored on this thread.
+	 * 
+	 * NOTE: This API is only meant to be called internally by the various STAT macros.
+	 * 
+	 * The stat is identified by the FNV1a hash of its name rather than the raw string for sake of performance.
+	 * You can use UE::HashStringFNV1a32 to generate this hash, or use UE_STATS_HASH_NAME to calculate it at compile time.
+	 * 
+	 * @param StatNameHash: The 32-bit FNV1a hash of the stat name.
+	 * @param GroupNameHash: The 32-bit FNV1a hash of the group name this stat is part of, or 0 to skip checking the group.
+	 */
+	bool IsStatOrGroupIgnoredOnCurrentThread(uint32 StatNameHash, uint32 GroupNameHash = 0);
+
+} // namespace UE::Stats
+#endif // UE_STATS_ALLOW_PER_THREAD_IGNORELIST
+
+#if UE_STATS_ALLOW_PER_THREAD_IGNORELIST
+	// Hashes the name at compile time.
+	#define UE_STATS_HASH_NAME(Name)\
+		UE::HashStringFNV1a<uint32>(TEXT(#Name))
+#else
+	#define UE_STATS_HASH_NAME(Name) 0
+#endif
diff --git a/Engine/Source/Runtime/Core/Public/Stats/Stats.h b/Engine/Source/Runtime/Core/Public/Stats/Stats.h
--- a/Engine/Source/Runtime/Core/Public/Stats/Stats.h
+++ b/Engine/Source/Runtime/Core/Public/Stats/Stats.h
@@ -500,6 +500,39 @@ struct FStat_##StatName\
 // profiling namedevents if enabled.
 // ----------------------------------------------------------------------------
 
+// ----------------------------------------------------------------
+// Group Declarations
+// ----------------------------------------------------------------
+
+#define DECLARE_STAT_GROUP(GroupId, CompiledIn) \
+	struct FStatGroup_##GroupId\
+	{\
+		static constexpr const PROFILER_CHAR* GetName()\
+		{\
+			return ANSI_TO_PROFILING(#GroupId);\
+		}\
+		static constexpr uint32 GetNameHash()\
+		{\
+			return UE_STATS_HASH_NAME(GroupId);\
+		}\
+		static constexpr bool IsCompileTimeEnable()\
+		{\
+			return CompiledIn;\
+		}\
+	};
+
+#define DECLARE_STATS_GROUP(GroupDesc, GroupId, GroupCat) \
+	DECLARE_STAT_GROUP(GroupId, true);
+
+#define DECLARE_STATS_GROUP_SORTBYNAME(GroupDesc, GroupId, GroupCat) \
+	DECLARE_STAT_GROUP(GroupId, true);
+
+#define DECLARE_STATS_GROUP_VERBOSE(GroupDesc, GroupId, GroupCat) \
+	DECLARE_STAT_GROUP(GroupId, true);
+
+#define DECLARE_STATS_GROUP_MAYBE_COMPILED_OUT(GroupDesc, GroupId, GroupCat, CompileIn) \
+	DECLARE_STAT_GROUP(GroupId, static_cast<bool>(CompileIn));
+
 // ----------------------------------------------------------------
 // Stat Declarations
 // ----------------------------------------------------------------
@@ -507,9 +540,10 @@ struct FStat_##StatName\
 #define DECLARE_STAT(StatName, GroupName, StatFlags)\
 	struct FStat_##StatName\
 	{\
-		static constexpr const PROFILER_CHAR* GetGroup()\
+		using TGroup = UE_INTERNAL_GET_STATGROUP_TYPE(GroupName);\
+		static constexpr const PROFILER_CHAR* GetName()\
 		{\
-			return ANSI_TO_PROFILING(#GroupName);\
+			return ANSI_TO_PROFILING(#StatName);\
 		}\
 		static constexpr EStatFlags GetFlags()\
 		{\
@@ -567,36 +601,38 @@ struct FStat_##StatName\
 
 // Note: Since these are defining the stat inline, no custom stat data will exist.
 #define DECLARE_SCOPE_CYCLE_COUNTER(CounterName,Stat,GroupId) \
-	FScopeCycleCounter StatNamedEventsScope_##Stat(UE::Stats::Private::TStatIdProxy(ANSI_TO_PROFILING(#Stat), ANSI_TO_PROFILING(#GroupId))); \
-	SCOPE_CYCLE_COUNTER_TO_TRACE(CounterName, Stat, true);
+	const UE::Stats::Private::FCheckedStat Stat_##Stat = UE_INTERNAL_GET_QUICK_STAT_WITH_GROUP_IF_ENABLED(Stat, GroupId);\
+	UE::Stats::Private::FScopeCycleCounterStatic StatNamedEventsScope_##Stat(Stat_##Stat); \
+	SCOPE_CYCLE_COUNTER_TO_TRACE(CounterName, Stat, Stat_##Stat.IsValidStat());
 
 #define QUICK_SCOPE_CYCLE_COUNTER(Stat) \
-	FScopeCycleCounter StatNamedEventsScope_##Stat(UE::Stats::Private::TStatIdProxy(ANSI_TO_PROFILING(#Stat))); \
-	SCOPE_CYCLE_COUNTER_TO_TRACE(#Stat, Stat, true);
+	const UE::Stats::Private::FCheckedStat Stat_##Stat = UE_INTERNAL_GET_QUICK_STAT_IF_ENABLED(Stat);\
+	UE::Stats::Private::FScopeCycleCounterStatic StatNamedEventsScope_##Stat(Stat_##Stat); \
+	SCOPE_CYCLE_COUNTER_TO_TRACE(#Stat, Stat, Stat_##Stat.IsValidStat());
 
 // Note: FScopeCycleCounter is what emits the trace events to external profilers
 // while SCOPE_CYCLE_COUNTER_TO_TRACE handles emitting it to insights.
 // FScopeCycleCounter handles verbosity checks itself.
 #define SCOPE_CYCLE_COUNTER(Stat) \
-	FScopeCycleCounter StatNamedEventsScope_##Stat(GET_STATID(Stat)); \
-	SCOPE_CYCLE_COUNTER_TO_TRACE(#Stat, Stat, (GShouldEmitVerboseNamedEvents || !UE_IS_STAT_VERBOSE(Stat)));
+	const UE::Stats::Private::FCheckedStat Stat_##Stat = UE_INTERNAL_GET_STAT_IF_ENABLED(Stat);\
+	UE::Stats::Private::FScopeCycleCounterStatic StatNamedEventsScope_##Stat(Stat_##Stat); \
+	SCOPE_CYCLE_COUNTER_TO_TRACE(#Stat, Stat, Stat_##Stat.IsValidStat());
 
-// Note: This is often used with non declared stats so we always allow default stat data.
+// NOTE: This macro bypasses the standard enable checks as we don't know the stat's type/group.
 #define SCOPE_CYCLE_COUNTER_STATID(StatId) \
 	FScopeCycleCounter StatNamedEventsScope_STATID(StatId); \
-	TRACE_CPUPROFILER_EVENT_SCOPE_TEXT_CONDITIONAL(StatId.StatString, StatId.StatString != nullptr && GCycleStatsShouldEmitNamedEvents > 0 && (GShouldEmitVerboseNamedEvents || !UE_IS_STAT_VERBOSE_ALLOW_DEFAULT(Stat)));
+	TRACE_CPUPROFILER_EVENT_SCOPE_TEXT_CONDITIONAL(StatId.StatString, StatId.IsValidStat() && GCycleStatsShouldEmitNamedEvents > 0);
 
 #define CONDITIONAL_SCOPE_CYCLE_COUNTER(Stat,bCondition) \
-	FScopeCycleCounter StatNamedEventsScope_##Stat(bCondition ? GET_STATID(Stat) : nullptr); \
-	SCOPE_CYCLE_COUNTER_TO_TRACE(#Stat, Stat, bCondition && (GShouldEmitVerboseNamedEvents || !UE_IS_STAT_VERBOSE(Stat)));
+	const UE::Stats::Private::FCheckedStat Stat_##Stat = UE_INTERNAL_GET_STAT_IF_ENABLED_COND(Stat, static_cast<bool>(bCondition));\
+	UE::Stats::Private::FScopeCycleCounterStatic StatNamedEventsScope_##Stat(Stat_##Stat); \
+	SCOPE_CYCLE_COUNTER_TO_TRACE(#Stat, Stat, Stat_##Stat.IsValidStat());
 
 #define RETURN_QUICK_DECLARE_CYCLE_STAT(StatId,GroupId) return QUICK_USE_CYCLE_STAT(StatId, GroupId);
 
-#define QUICK_USE_CYCLE_STAT(StatId,GroupId) UE::Stats::Private::TStatIdProxy(ANSI_TO_PROFILING(#StatId), ANSI_TO_PROFILING(#GroupId))
-
-// Creates a TStatIdProxy on the fly, optionally populating the group/flags from the custom stat data struct if it exists.
-#define GET_STATID(Stat) (UE::Stats::Private::MakeStatId<UE_INTERNAL_GET_STATDATA(Stat)>(ANSI_TO_PROFILING(#Stat)))
+#define QUICK_USE_CYCLE_STAT(StatId,GroupId) UE_INTERNAL_GET_QUICK_STAT_WITH_GROUP_IF_ENABLED(StatId, GroupId)
 
+#define GET_STATID(Stat) UE_INTERNAL_GET_STAT_IF_ENABLED(Stat)
 
 #elif USE_LIGHTWEIGHT_STATS_FOR_HITCH_DETECTION && USE_HITCH_DETECTION
 // STATS == 0
@@ -672,12 +708,11 @@ struct FStat_##StatName\
 #define DECLARE_PTR_STAT_EXTERN(CounterName,StatId,GroupId, API)
 #define DECLARE_MEMORY_STAT_EXTERN(CounterName,StatId,GroupId, API)
 #define DECLARE_MEMORY_STAT_POOL_EXTERN(CounterName,StatId,GroupId,Pool, API)
-#endif // !UE_USE_LIGHTWEIGHT_STATS
-
 #define DECLARE_STATS_GROUP(GroupDesc,GroupId,GroupCat)
 #define DECLARE_STATS_GROUP_VERBOSE(GroupDesc,GroupId,GroupCat)
 #define DECLARE_STATS_GROUP_SORTBYNAME(GroupDesc,GroupId,GroupCat)
 #define DECLARE_STATS_GROUP_MAYBE_COMPILED_OUT(GroupDesc,GroupId,GroupCat,CompileIn)
+#endif // !UE_USE_LIGHTWEIGHT_STATS
 
 #define SET_CYCLE_COUNTER(Stat,Cycles)
 #define INC_DWORD_STAT(StatId)
diff --git a/Engine/Source/Runtime/Core/Public/Stats/StatsCommon.h b/Engine/Source/Runtime/Core/Public/Stats/StatsCommon.h
--- a/Engine/Source/Runtime/Core/Public/Stats/StatsCommon.h
+++ b/Engine/Source/Runtime/Core/Public/Stats/StatsCommon.h
@@ -3,6 +3,7 @@
 #pragma once
 
 #include "CoreTypes.h"
+#include "Misc/Build.h"
 #include "Misc/EnumClassFlags.h"
 
 #define FORCEINLINE_STATS FORCEINLINE
@@ -13,6 +14,10 @@
 #error "STATS must be defined as either zero or one."
 #endif
 
+#ifndef UE_USE_LIGHTWEIGHT_STATS
+#define UE_USE_LIGHTWEIGHT_STATS (!STATS && ENABLE_STATNAMEDEVENTS)
+#endif
+
 #if	PLATFORM_USES_ANSI_STRING_FOR_EXTERNAL_PROFILING
 typedef ANSICHAR PROFILER_CHAR;
 #else
diff --git a/Engine/Source/Runtime/Core/Public/Stats/StatsSystem.h b/Engine/Source/Runtime/Core/Public/Stats/StatsSystem.h
--- a/Engine/Source/Runtime/Core/Public/Stats/StatsSystem.h
+++ b/Engine/Source/Runtime/Core/Public/Stats/StatsSystem.h
@@ -14,6 +14,9 @@ namespace UE::Stats
 		/** Delegate to fire every time we need to advance the stats for the rendering thread. */
 		DECLARE_DELEGATE_ThreeParams(FOnAdvanceRenderingThreadStats, bool /*bDiscardCallstack*/, int64 /*StatsFrame*/, int32 /*PrimaryDisableChangeTagStartFrame*/);
 
+		/* Initializes any relevant stat systems, including those that exist when STATS is 0. */
+		static CORE_API void Init();
+
 		/** Advances stats for the current frame. */
 		static CORE_API void AdvanceFrame(bool bDiscardCallstack, const FOnAdvanceRenderingThreadStats& AdvanceRenderingThreadStatsDelegate = FOnAdvanceRenderingThreadStats());
 
diff --git a/Engine/Source/Runtime/Launch/Private/LaunchEngineLoop.cpp b/Engine/Source/Runtime/Launch/Private/LaunchEngineLoop.cpp
--- a/Engine/Source/Runtime/Launch/Private/LaunchEngineLoop.cpp
+++ b/Engine/Source/Runtime/Launch/Private/LaunchEngineLoop.cpp
@@ -2814,6 +2814,7 @@ int32 FEngineLoop::PreInitPreStartupScreen(const TCHAR* CmdLine)
 	}
 #endif
 
+	UE::Stats::FStats::Init();
 #if STATS
 	FThreadStats::StartThread();
 #endif
