[core] another FormatStringSan speculative fix to avoid an intermittent compiler bug

diff --git a/Engine/Source/Runtime/Core/Public/String/FormatStringSan.h b/Engine/Source/Runtime/Core/Public/String/FormatStringSan.h
--- a/Engine/Source/Runtime/Core/Public/String/FormatStringSan.h
+++ b/Engine/Source/Runtime/Core/Public/String/FormatStringSan.h
@@ -149,33 +149,33 @@ namespace UE::Core::Private
 		template <typename Arg, typename... Args>
 		struct TCheckFormatString<TFmtArgTypes<Arg, Args...>>
 		{
-			static constexpr FResult Check(bool bInsideFormatSpec, int CurArgPos, const TCHAR* Fmt)
+			static constexpr bool CharIsDigit(const TCHAR C)
 			{
-				auto CharIsDigit = [](const TCHAR C) -> bool
-				{
-					return C >= TEXT('0') && C <= TEXT('9');
-				};
+				return C >= TEXT('0') && C <= TEXT('9');
+			}
 
-				auto SkipInteger = [&CharIsDigit](const TCHAR* P) -> const TCHAR*
+			static constexpr const TCHAR* SkipInteger(const TCHAR* P)
+			{
+				while (CharIsDigit(*P))
 				{
-					while (CharIsDigit(*P))
-					{
-						++P;
-					}
-					return P;
-				};
+					++P;
+				}
+				return P;
+			}
 
-				auto CharIsIntegerFormatSpecifier = [](const TCHAR C) -> bool
+			static constexpr bool CharIsIntegerFormatSpecifier(const TCHAR C)
+			{
+				switch (C)
 				{
-					switch (C)
-					{
-					case TEXT('i'): case TEXT('d'): case TEXT('u'): case TEXT('X'): case TEXT('x'):
-						return true;
-					default:
-						return false;
-					}
-				};
+				case TEXT('i'): case TEXT('d'): case TEXT('u'): case TEXT('X'): case TEXT('x'):
+					return true;
+				default:
+					return false;
+				}
+			}
 
+			static constexpr FResult Check(bool bInsideFormatSpec, int CurArgPos, const TCHAR* Fmt)
+			{
 				auto HandleDynamicLengthSpecifier = [CurArgPos](const TCHAR* P) constexpr -> FResult
 				{
 					if constexpr (!(std::is_integral_v<Arg> || bIsIntegralEnum<Arg>))
