Get AutoRTFMTests into nightly ASan testing, and fix the issues it raised.

Two main ones - had to backout my previous change 38206600 because of ASan shadow stack fun, and we were incorrectly using a vmovapd on an unaligned address.

#rnx

diff --git a/Engine/Source/Programs/AutoRTFMTests/Private/Abort.cpp b/Engine/Source/Programs/AutoRTFMTests/Private/Abort.cpp
--- a/Engine/Source/Programs/AutoRTFMTests/Private/Abort.cpp
+++ b/Engine/Source/Programs/AutoRTFMTests/Private/Abort.cpp
@@ -684,67 +684,6 @@ TEST_CASE("Abort.PushOnAbortHandler_Order")
 			REQUIRE(42 == Value);
 		}
 	}
-
-	SECTION("Handler Ran Because Key Is Stack Local To Current Aborting Transaction")
-	{
-		bool bHit = false;
-		AutoRTFM::Testing::Abort([&]
-		{
-			int Foo = 42;
-			AutoRTFM::PushOnAbortHandler(&Foo, [&] { bHit = true; });
-			AutoRTFM::AbortTransaction();
-		});
-		REQUIRE(bHit);
-	}
-
-	SECTION("Handler Ran Because Key Is Stack Local In Aborting Child")
-	{
-		bool bHit = false;
-		AutoRTFM::Testing::Commit([&]
-		{
-			AutoRTFM::Testing::Abort([&]
-			{
-				int Foo = 42;
-				AutoRTFM::PushOnAbortHandler(&Foo, [&] { bHit = true; });	
-				AutoRTFM::AbortTransaction();
-			});
-		});
-		REQUIRE(bHit);
-	}
-
-	SECTION("Handler Ran Because Key Is Stack Local In Aborting Parent")
-	{
-		// Need to disable retries because the inner commit would retry and run the on-abort handler.
-		AutoRTFMTestUtils::FScopedRetry _(AutoRTFM::ForTheRuntime::EAutoRTFMRetryTransactionState::NoRetry);
-
-		bool bHit = false;
-		AutoRTFM::Testing::Abort([&]
-		{
-			int Foo = 42;
-			AutoRTFM::Testing::Commit([&]
-			{
-				AutoRTFM::PushOnAbortHandler(&Foo, [&] { bHit = true; });	
-			});
-			AutoRTFM::AbortTransaction();
-		});
-		REQUIRE(bHit);
-	}
-
-	SECTION("Handler Ignored Because Key Is Stack Local In Committing Child")
-	{
-		// Need to disable retries because the inner commit would retry and run the on-abort handler.
-		AutoRTFMTestUtils::FScopedRetry _(AutoRTFM::ForTheRuntime::EAutoRTFMRetryTransactionState::NoRetry);
-
-		AutoRTFM::Testing::Abort([&]
-		{
-			AutoRTFM::Testing::Commit([&]
-			{
-				int Foo = 42;
-				AutoRTFM::PushOnAbortHandler(&Foo, [] { FAIL("Unreachable!"); });	
-			});
-			AutoRTFM::AbortTransaction();
-		});
-	}
 }
 
 TEST_CASE("Abort.OnAbortTiming")
diff --git a/Engine/Source/Programs/AutoRTFMTests/Private/Commit.cpp b/Engine/Source/Programs/AutoRTFMTests/Private/Commit.cpp
--- a/Engine/Source/Programs/AutoRTFMTests/Private/Commit.cpp
+++ b/Engine/Source/Programs/AutoRTFMTests/Private/Commit.cpp
@@ -1,7 +1,6 @@
 // Copyright Epic Games, Inc. All Rights Reserved.
 
 #include "Catch2Includes.h"
-#include "Utils.h"
 #include <AutoRTFM/AutoRTFMTesting.h>
 #include <map>
 #include <vector>
@@ -425,39 +424,6 @@ TEST_CASE("Commit.PushOnCommitHandler_Order")
 			REQUIRE(42 == Value);
 		}
 	}
-
-	SECTION("Handler Ignored Because Key Is Stack Local")
-	{
-		AutoRTFM::Testing::Commit([&]
-		{
-			int Foo = 42;
-			AutoRTFM::PushOnCommitHandler(&Foo, [] { FAIL("Unreachable!"); });
-		});
-	}
-
-	SECTION("Handler Ignored Because Key Is Stack Local In Committing Child")
-	{
-		AutoRTFM::Testing::Commit([&]
-		{
-			AutoRTFM::Testing::Commit([&]
-			{
-				int Foo = 42;
-				AutoRTFM::PushOnCommitHandler(&Foo, [] { FAIL("Unreachable!"); });	
-			});
-		});
-	}
-
-	SECTION("Handler Ignored Because Key Is Stack Local In Parent")
-	{
-		AutoRTFM::Testing::Commit([&]
-		{
-			int Foo = 42;
-			AutoRTFM::Testing::Commit([&]
-			{
-				AutoRTFM::PushOnCommitHandler(&Foo, [] { FAIL("Unreachable!"); });	
-			});
-		});
-	}
 }
 
 // Test functions that may have LLVM byval() attributed parameters
diff --git a/Engine/Source/Programs/AutoRTFMTests/Private/Tests.cpp b/Engine/Source/Programs/AutoRTFMTests/Private/Tests.cpp
--- a/Engine/Source/Programs/AutoRTFMTests/Private/Tests.cpp
+++ b/Engine/Source/Programs/AutoRTFMTests/Private/Tests.cpp
@@ -373,7 +373,7 @@ TEST_CASE("Tests.AVXMaskedStore")
 	{
 		AutoRTFM::Open([&]()
 		{
-			__m256d Lhs = _mm256_load_pd(VectorPtr);
+			__m256d Lhs = _mm256_loadu_pd(VectorPtr);
 			__m256 CmpResult = _mm256_castpd_ps(_mm256_cmp_pd(Lhs, Rhs, 0));
 			int Result = _mm256_movemask_ps(CmpResult);
 			REQUIRE(Result == 0xFF);
diff --git a/Engine/Source/Programs/AutoRTFMTests/Private/Utils.h b/Engine/Source/Programs/AutoRTFMTests/Private/Utils.h
--- a/Engine/Source/Programs/AutoRTFMTests/Private/Utils.h
+++ b/Engine/Source/Programs/AutoRTFMTests/Private/Utils.h
@@ -10,9 +10,9 @@ namespace AutoRTFMTestUtils
 {
 
 // Temporarily changes the AutoRTFM retry mode for the lifetime of the FScopedRetry object.
-struct FScopedRetry final
+struct FScopedRetry
 {
-	explicit FScopedRetry(AutoRTFM::ForTheRuntime::EAutoRTFMRetryTransactionState NewRetry)
+	FScopedRetry(AutoRTFM::ForTheRuntime::EAutoRTFMRetryTransactionState NewRetry)
 		: OldRetry(AutoRTFM::ForTheRuntime::GetRetryTransaction())
 	{
 		AutoRTFM::ForTheRuntime::SetRetryTransaction(NewRetry);
@@ -23,9 +23,6 @@ struct FScopedRetry final
 		AutoRTFM::ForTheRuntime::SetRetryTransaction(OldRetry);
 	}
 
-	UE_NONCOPYABLE(FScopedRetry)
-
-private:
 	AutoRTFM::ForTheRuntime::EAutoRTFMRetryTransactionState OldRetry;
 };
 
diff --git a/Engine/Source/Runtime/Core/Private/AutoRTFM/Context.cpp b/Engine/Source/Runtime/Core/Private/AutoRTFM/Context.cpp
--- a/Engine/Source/Runtime/Core/Private/AutoRTFM/Context.cpp
+++ b/Engine/Source/Runtime/Core/Private/AutoRTFM/Context.cpp
@@ -455,7 +455,7 @@ ETransactionResult FContext::Transact(void (*InstrumentedFunction)(void*), void*
 				{
 					// We skip trying to commit this time, and instead re-run the transaction.
 					Status = EContextStatus::AbortedByFailedLockAcquisition;
-					CurrentTransaction->AbortWithoutThrowing(EAbortWithoutThrowingFlush::Flush);
+					CurrentTransaction->AbortWithoutThrowing();
 					ClearTransactionStatus();
 
 					// We've tried to run at least once if we get here!
diff --git a/Engine/Source/Runtime/Core/Private/AutoRTFM/TaskArray.h b/Engine/Source/Runtime/Core/Private/AutoRTFM/TaskArray.h
--- a/Engine/Source/Runtime/Core/Private/AutoRTFM/TaskArray.h
+++ b/Engine/Source/Runtime/Core/Private/AutoRTFM/TaskArray.h
@@ -250,18 +250,6 @@ public:
 		return true;
 	}
 
-	template<typename FunctorType>
-	void DeleteKeyedIf(const FunctorType& Func)
-	{
-		for (auto& Pair : BitMasksPerKey)
-		{
-			if (Func(Pair.Key))
-			{
-				Pair.Value.Bits.SetRange(Pair.Value.InitialIndex, Pair.Value.Num(), true);
-			}
-		}
-	}
-
     void AddAll(TTaskArray&& Other)
     {
         Canonicalize();
diff --git a/Engine/Source/Runtime/Core/Private/AutoRTFM/Transaction.cpp b/Engine/Source/Runtime/Core/Private/AutoRTFM/Transaction.cpp
--- a/Engine/Source/Runtime/Core/Private/AutoRTFM/Transaction.cpp
+++ b/Engine/Source/Runtime/Core/Private/AutoRTFM/Transaction.cpp
@@ -85,7 +85,7 @@ bool FTransaction::IsFresh() const
 		&& DeferredPopAllOnAbortHandlers.IsEmpty();
 }
 
-void FTransaction::AbortWithoutThrowing(EAbortWithoutThrowingFlush State)
+void FTransaction::AbortWithoutThrowing()
 {
 	UE_LOG(LogAutoRTFM, Verbose, TEXT("Aborting '%hs'!"), GetContextStatusName(Context->GetStatus()));
 
@@ -105,13 +105,6 @@ void FTransaction::AbortWithoutThrowing(EAbortWithoutThrowingFlush State)
     CommitTasks.Reset();
 
     Undo();
-
-	// Only used by the aborteroonie so it is definitely not in the hot path.
-	if (UNLIKELY(EAbortWithoutThrowingFlush::Flush == State))
-	{
-		AbortTasks.DeleteKeyedIf([&](const void* const Key) { return this->IsOnStack(Key); });
-	}
-
 	AbortTasks.ForEachBackward([&](TFunction<void()>& Task)
     { 
         // Call and then reset each of the tasks in reverse order.
@@ -275,9 +268,6 @@ void FTransaction::CommitNested()
 	}
 	DeferredPopAllOnAbortHandlers.Reset();
 
-	CommitTasks.DeleteKeyedIf([&](const void* const Key) { return this->IsOnStack(Key); });
-	AbortTasks.DeleteKeyedIf([&](const void* const Key) { return this->IsOnStack(Key); });
-
     Parent->CommitTasks.AddAll(MoveTemp(CommitTasks));
     Parent->AbortTasks.AddAll(MoveTemp(AbortTasks));
 
@@ -294,7 +284,6 @@ bool FTransaction::AttemptToCommitOuterNest()
 
     AbortTasks.Reset();
 
-	CommitTasks.DeleteKeyedIf([&](const void* const Key) { return this->IsOnStack(Key); });
 
     CommitTasks.ForEachForward([] (TFunction<void()>& Task)
     { 
diff --git a/Engine/Source/Runtime/Core/Private/AutoRTFM/Transaction.h b/Engine/Source/Runtime/Core/Private/AutoRTFM/Transaction.h
--- a/Engine/Source/Runtime/Core/Private/AutoRTFM/Transaction.h
+++ b/Engine/Source/Runtime/Core/Private/AutoRTFM/Transaction.h
@@ -17,18 +17,6 @@ namespace AutoRTFM
 {
 class FContext;
 
-enum class EAbortWithoutThrowingFlush : bool
-{
-	// The default - since aborts happen eagerly when we call `AbortTransaction`
-	// we want to run on-abort handlers that are still on the stack.
-	NoFlush = false,
-
-	// Only used by the Aborteroonie - we call `AbortWithoutThrowing` after the
-	// transactions stack has already disappeared, and so it is thus unsafe to
-	// call potentially stack local on-abort handlers.
-	Flush = true,
-};
-
 class FTransaction final
 {
 public:
@@ -112,8 +100,7 @@ public:
 	void PopAllDeferUntilAbortHandlers(const void* Key);
 
     [[noreturn]] void AbortAndThrow();
-
-    void AbortWithoutThrowing(EAbortWithoutThrowingFlush State = EAbortWithoutThrowingFlush::NoFlush);
+    void AbortWithoutThrowing();
     bool AttemptToCommit();
 
     // Record that a write is about to occur at the given LogicalAddress of Size bytes.
diff --git a/Engine/Source/Runtime/Core/Public/AutoRTFM/AutoRTFM.h b/Engine/Source/Runtime/Core/Public/AutoRTFM/AutoRTFM.h
--- a/Engine/Source/Runtime/Core/Public/AutoRTFM/AutoRTFM.h
+++ b/Engine/Source/Runtime/Core/Public/AutoRTFM/AutoRTFM.h
@@ -126,7 +126,7 @@ typedef enum
 // This must match AutoRTFM::EContextStatus.
 typedef enum
 {
-	autortfm_status_idle,
+	autortfm_status_idle = 0,
 	autortfm_status_ontrack,
 	autortfm_status_aborted_by_failed_lock_aquisition,
 	autortfm_status_aborted_by_language,
@@ -894,9 +894,7 @@ template<typename FunctorType> static UE_AUTORTFM_FORCEINLINE void OnAbort(Funct
 // Register a handler for transaction commit. Takes a key parameter so that
 // the handler can be unregistered (see `PopOnCommitHandler`). This is useful
 // for scoped mutations that need an abort handler present unless execution
-// reaches the end of the relevant scope. If `Key` was stack local to a
-// committing transaction then the runtime will not run the handler as the data
-// backing the handler (normally a `this` object) has already been destroyed.
+// reaches the end of the relevant scope.
 template<typename FunctorType> static UE_AUTORTFM_FORCEINLINE void PushOnCommitHandler(const void* Key, FunctorType&& Work)
 {
 	if (autortfm_is_closed())
@@ -908,9 +906,7 @@ template<typename FunctorType> static UE_AUTORTFM_FORCEINLINE void PushOnCommitH
 // Register a handler for transaction commit. Takes a key parameter so that
 // the handler can be unregistered (see `PopOnCommitHandler`). This is useful
 // for scoped mutations that need an abort handler present unless execution
-// reaches the end of the relevant scope. If `Key` was stack local to a
-// committing transaction then the runtime will not run the handler as the data
-// backing the handler (normally a `this` object) has already been destroyed.
+// reaches the end of the relevant scope.
 template<typename FunctorType> static UE_AUTORTFM_FORCEINLINE void PushOnCommitHandler(const void*, FunctorType&&) {}
 #endif
 
@@ -946,10 +942,7 @@ static UE_AUTORTFM_FORCEINLINE void PopAllOnCommitHandlers(const void*) {}
 // Register a handler for transaction abort. Takes a key parameter so that
 // the handler can be unregistered (see `PopOnAbortHandler`). This is useful
 // for scoped mutations that need an abort handler present unless execution
-// reaches the end of the relevant scope. If `Key` was stack-local to a
-// nested child transaction then the runtime will not run the handler as the
-// data backing the handler (normally a `this` object) has already been
-// destroyed.
+// reaches the end of the relevant scope.
 template<typename FunctorType> static UE_AUTORTFM_FORCEINLINE void PushOnAbortHandler(const void* Key, FunctorType&& Work)
 {
 	if (autortfm_is_closed())
@@ -961,10 +954,7 @@ template<typename FunctorType> static UE_AUTORTFM_FORCEINLINE void PushOnAbortHa
 // Register a handler for transaction abort. Takes a key parameter so that
 // the handler can be unregistered (see `PopOnAbortHandler`). This is useful
 // for scoped mutations that need an abort handler present unless execution
-// reaches the end of the relevant scope. If `Key` was stack-local to a
-// nested child transaction then the runtime will not run the handler as the
-// data backing the handler (normally a `this` object) has already been
-// destroyed.
+// reaches the end of the relevant scope.
 template<typename FunctorType> static UE_AUTORTFM_FORCEINLINE void PushOnAbortHandler(const void* Key, FunctorType&&) {}
 #endif
 
@@ -1023,7 +1013,11 @@ static UE_AUTORTFM_FORCEINLINE void RecordOpenWrite(void* Ptr, size_t Size)
 	autortfm_record_open_write(Ptr, Size);
 }
 
-// Informs the runtime that a block of memory is about to be overwritten.
+// Informs the runtime that a block of memory is about to be overwritten in the open.
+// During a transaction, this allows the runtime to copy the data in preparation for
+// a possible abort. Normally, tracking memory overwrites should be automatically
+// handled by AutoRTFM, but manual overwrite tracking may be required for third-party
+// libraries or outside compilers (such as ISPC).
 template<typename Type> static UE_AUTORTFM_FORCEINLINE void RecordOpenWrite(Type* Ptr)
 {
 	autortfm_record_open_write(Ptr, sizeof(Type));
