[TaskGraph]
- Fix leak in UE::Tasks::Any when using it with a collection that has a single item

[RN] Fix leak in UE::Tasks::Any when using it with a collection that has a single item

diff --git a/Engine/Source/Runtime/Core/Public/Tasks/Task.h b/Engine/Source/Runtime/Core/Public/Tasks/Task.h
--- a/Engine/Source/Runtime/Core/Public/Tasks/Task.h
+++ b/Engine/Source/Runtime/Core/Public/Tasks/Task.h
@@ -559,8 +559,9 @@ namespace UE::Tasks
 					{	// the first completed task
 						Event.Trigger();
 					}
-					else if (UNLIKELY(PrevRefCount == 1))
-					{	// the last competed task
+
+					if (UNLIKELY(PrevRefCount == 1))
+					{	// the last completed task
 						delete SharedData;
 					}
 				},
