Csv Profiler: Add a simple benchmark so we can estimate instrumentation cost of timing and custom stats. This runs on the pre-init frame before a csv profile starts, and takes ~8-12ms. Enable with the  csv.benchmark cvar (defaults to on in non-shipping)
Timings are stored in metadata values CsvBench_TimingStat_ns and CsvBench_CustomStat_ns.

Also:
- Fix potential memory corruption with stat series set to invisible. Didn't affect the old code, but setting bVisible to true for stats other than taskgraph aggregate stats would crash
- Expose Csv profiler update thread time as a cvar (csv.processingThread.TimeBetweenUpdates) instead of hard-coding

[RN] Csv Profiler: Add a simple benchmark so we can estimate instrumentation cost of timing and custom stats

diff --git a/Engine/Source/Runtime/Core/Private/ProfilingDebugging/CsvProfiler.cpp b/Engine/Source/Runtime/Core/Private/ProfilingDebugging/CsvProfiler.cpp
--- a/Engine/Source/Runtime/Core/Private/ProfilingDebugging/CsvProfiler.cpp
+++ b/Engine/Source/Runtime/Core/Private/ProfilingDebugging/CsvProfiler.cpp
@@ -74,6 +74,7 @@ CSV_DEFINE_CATEGORY_MODULE(CORE_API, FileIO, true);
 
 // Other categories
 CSV_DEFINE_CATEGORY(CsvProfiler, true);
+CSV_DEFINE_CATEGORY(CsvBench, true);
 
 #if CSV_PROFILER_ALLOW_DEBUG_FEATURES
 	CSV_DEFINE_CATEGORY(CsvTest, true);
@@ -119,6 +120,25 @@ TAutoConsoleVariable<int32> CVarCsvForceExit(
 	ECVF_Default
 );
 
+TAutoConsoleVariable<bool> CVarCsvBenchmark(
+	TEXT("csv.Benchmark"),
+#if UE_BUILD_SHIPPING
+	false,
+#else
+	true,
+#endif
+	TEXT("If emabled, do a quick benchmark test on the frame before the CSV profiler starts"),
+	ECVF_Default
+);
+
+TAutoConsoleVariable<int32> CVarCsvBenchmarkIterationCount(
+	TEXT("csv.Benchmark.IterationCount"),
+	10000,
+	TEXT("Number of iterations of each CsvBenchmark test"),
+	ECVF_Default
+);
+
+
 TAutoConsoleVariable<int32> CVarCsvTargetFrameRateOverride(
 	TEXT("csv.TargetFrameRateOverride"),
 	0,
@@ -174,6 +194,15 @@ TAutoConsoleVariable<int32> CVarCsvPauseProcessingThread(
 	ECVF_Default
 );
 
+TAutoConsoleVariable<float> CVarCsvProcessingThreadTimeBetweenUpdates(
+	TEXT("csv.ProcessingThread.TimeBetweenUpdates"),
+	50.0f,
+	TEXT("Specifies the minimum time between CSV processing thread updates.\r\n")
+	TEXT("Note: This is the time between the start of updates. If processing takes longer than this, the next update will commence immediately."),
+	ECVF_Default
+);
+
+
 TAutoConsoleVariable<int32> CVarMaxPerThreadStatDataSlackKB(
 	TEXT("csv.MaxPerThreadStatDataSlackKB"),
 	64,
@@ -620,6 +649,77 @@ static FCsvABTest GCsvABTest;
 
 #endif // CSV_PROFILER_ALLOW_DEBUG_FEATURES
 
+struct FCsvBenchmarkResult
+{
+	double TimingStatTimeNs = 0.0;
+	double CustomStatTimeNs = 0.0;
+};
+
+class FCsvBenchmark
+{
+	struct FTimer
+	{
+		FTimer()
+		{
+			StartCycles = FPlatformTime::Cycles64();
+		}
+		double GetTimeElapsed()
+		{
+			uint64 CyclesElapsed = FPlatformTime::Cycles64() - StartCycles;
+			return FPlatformTime::ToSeconds64(CyclesElapsed);
+		}
+		uint64 StartCycles = 0;
+	};
+	bool bIsRunning = false;
+
+public:
+	FCsvBenchmarkResult Run(int32 IterationCount)
+	{
+		CSV_SCOPED_TIMING_STAT_EXCLUSIVE(CsvBench);
+		bIsRunning = true;
+		FCsvBenchmarkResult Result;
+		FTimer MainTimer;
+
+		FTimer TimingStatTimer;
+		for (int i=0; i<IterationCount; i++)
+		{
+			CSV_SCOPED_TIMING_STAT(CsvBench, TimingStat1);
+			CSV_SCOPED_TIMING_STAT(CsvBench, TimingStat2);
+			CSV_SCOPED_TIMING_STAT(CsvBench, TimingStat3);
+			CSV_SCOPED_TIMING_STAT(CsvBench, TimingStat4);
+		}
+		Result.TimingStatTimeNs = TimingStatTimer.GetTimeElapsed() * 1000000000.0 / double(IterationCount * 4);
+
+		FTimer CustomStatTimer;
+		for (int i = 0; i < IterationCount; i++)
+		{
+			CSV_CUSTOM_STAT(CsvBench, CustomStat, 1.0f, ECsvCustomStatOp::Accumulate);
+			CSV_CUSTOM_STAT(CsvBench, CustomStat, 1.0f, ECsvCustomStatOp::Accumulate);
+			CSV_CUSTOM_STAT(CsvBench, CustomStat, 1.0f, ECsvCustomStatOp::Accumulate);
+			CSV_CUSTOM_STAT(CsvBench, CustomStat, 1.0f, ECsvCustomStatOp::Accumulate);
+		}
+		Result.CustomStatTimeNs = CustomStatTimer.GetTimeElapsed() * 1000000000.0 / double(IterationCount * 4);
+
+		double BenchTotalTimeMs = MainTimer.GetTimeElapsed() * 1000.0;
+		UE_LOG(LogCsvProfiler, Log, TEXT("Csv profiler benchmark completed (%d iterations). Time taken: %.2fms"), IterationCount, BenchTotalTimeMs );
+		UE_LOG(LogCsvProfiler, Log, TEXT("Timing stat ns: %.2fns"), Result.TimingStatTimeNs);
+		UE_LOG(LogCsvProfiler, Log, TEXT("Custom stat ns: %.2fns"), Result.CustomStatTimeNs);
+
+		CSV_METADATA(TEXT("CsvBench_TimingStat_ns"), *LexToString(Result.TimingStatTimeNs));
+		CSV_METADATA(TEXT("CsvBench_CustomStat_ns"), *LexToString(Result.CustomStatTimeNs));
+		CSV_METADATA(TEXT("CsvBench_Duration_ms"), *LexToString(BenchTotalTimeMs));
+		bIsRunning = false;
+		return Result;
+	}
+
+	bool IsRunning() const
+	{
+		return bIsRunning;
+	}
+};
+static FCsvBenchmark GCsvBenchmark;
+
+
 class FCsvCategoryData
 {
 public:
@@ -2065,6 +2165,13 @@ FCsvStatSeries::FCsvStatSeries(EType InSeriesType, const FCsvStatID& InStatID, F
 
 	// If we have a linked stat series, don't write directly to the CSV
 	bool bVisible = LinkedAggregateStatSeries == nullptr || (CSV_DEBUG_EMIT_SEPARATE_THREAD_STATS_WHEN_TASK_AGGREGATION_ENABLED==1);
+
+	// Internal CsvBench stats should be hidden
+	if (StatCategoryIndex == CSV_CATEGORY_INDEX(CsvBench))
+	{
+		bVisible = false;
+	}
+
 	Writer->AddSeries(this, bVisible);
 }
 
@@ -2090,9 +2197,8 @@ void FCsvStatSeries::FlushIfDirty()
 		{
 			LinkedAggregateStatSeries->AccumulateLinkedSeriesValue(CurrentWriteFrameNumber, Value);
 		}
-#if !CSV_DEBUG_EMIT_SEPARATE_THREAD_STATS_WHEN_TASK_AGGREGATION_ENABLED
-		else
-#endif
+		// If this series is hidden then don't write
+		if (ColumnIndex != -1)
 		{
 			Writer->PushValue(this, CurrentWriteFrameNumber, Value);
 		}
@@ -2835,7 +2941,6 @@ public:
 
 	virtual uint32 Run() override
 	{
-		const float TimeBetweenUpdatesMS = 50.0f;
 		GCsvProcessingThreadId = FPlatformTLS::GetCurrentThreadId();
 		GGameThreadIsCsvProcessingThread = false;
 
@@ -2845,9 +2950,12 @@ public:
 
 		while (StopCounter.GetValue() == 0)
 		{
-			if (CVarCsvPauseProcessingThread.GetValueOnAnyThread())
+			const float TimeBetweenUpdatesMS = CVarCsvProcessingThreadTimeBetweenUpdates.GetValueOnAnyThread();
+
+			// Pause the processing thread if requested by cvar or if the benchmark is running (since it can interfere with timings under extreme load)
+			if (CVarCsvPauseProcessingThread.GetValueOnAnyThread() || GCsvBenchmark.IsRunning())
 			{
-				FPlatformProcess::Sleep(TimeBetweenUpdatesMS / 1000.0f);
+				FPlatformProcess::Sleep(5.0f / 1000.0f);
 				continue;
 			}
 			float ElapsedMS = CsvProfiler.ProcessStatData();
@@ -3404,9 +3512,14 @@ void FCsvProfiler::BeginCaptureInternal(const FCsvCaptureCommand& CurrentCommand
 	TRACE_CSV_PROFILER_BEGIN_CAPTURE(*Filename, RenderThreadId, RHIThreadId, GDefaultWaitStatName.StatName, GCsvStatCounts);
 	GCsvProfilerIsCapturing = true;
 	CaptureStartTime = FPlatformTime::Seconds();
-}
-
 
+	// Run the Csv benchmark if enabled
+	if (CVarCsvBenchmark.GetValueOnGameThread() || FParse::Param(FCommandLine::Get(), TEXT("csvBench")))
+	{
+		// This runs on the frame before the CSV kicks off, so doesn't affect timings
+		GCsvBenchmark.Run(CVarCsvBenchmarkIterationCount.GetValueOnGameThread());
+	}
+}
 
 
 void FCsvProfiler::EndFrame()
