Define LIKELY and UNLIKELY for all Clang/GCC platforms (not just Linux). It seems to be slightly faster in general

diff --git a/Engine/Source/Runtime/Core/Public/HAL/Platform.h b/Engine/Source/Runtime/Core/Public/HAL/Platform.h
--- a/Engine/Source/Runtime/Core/Public/HAL/Platform.h
+++ b/Engine/Source/Runtime/Core/Public/HAL/Platform.h
@@ -780,7 +780,7 @@
 
 /** Branch prediction hints */
 #ifndef LIKELY						/* Hints compiler that expression is likely to be true, much softer than UE_ASSUME - allows (penalized by worse performance) expression to be false */
-	#if ( defined(__clang__) || defined(__GNUC__) ) && (PLATFORM_UNIX)	// effect of these on non-Linux platform has not been analyzed as of 2016-03-21
+	#if ( defined(__clang__) || defined(__GNUC__) ) 
 		#define LIKELY(x)			__builtin_expect(!!(x), 1)
 	#else
 		// the additional "!!" is added to silence "warning: equality comparison with exteraneous parenthese" messages on android
@@ -789,7 +789,7 @@
 #endif
 
 #ifndef UNLIKELY					/* Hints compiler that expression is unlikely to be true, allows (penalized by worse performance) expression to be true */
-	#if ( defined(__clang__) || defined(__GNUC__) ) && (PLATFORM_UNIX)	// effect of these on non-Linux platform has not been analyzed as of 2016-03-21
+	#if ( defined(__clang__) || defined(__GNUC__) ) 
 		#define UNLIKELY(x)			__builtin_expect(!!(x), 0)
 	#else
 		// the additional "!!" is added to silence "warning: equality comparison with exteraneous parenthese" messages on android
