Made FMath::Min, FMath::Max, HashCombine and HashCombineFast variadic.

diff --git a/Engine/Source/Runtime/Core/Public/GenericPlatform/GenericPlatformMath.h b/Engine/Source/Runtime/Core/Public/GenericPlatform/GenericPlatformMath.h
--- a/Engine/Source/Runtime/Core/Public/GenericPlatform/GenericPlatformMath.h
+++ b/Engine/Source/Runtime/Core/Public/GenericPlatform/GenericPlatformMath.h
@@ -937,18 +937,48 @@ struct FGenericPlatformMath
 	}
 
 	/** Returns higher value in a generic way */
-	template< class T > 
-	static constexpr FORCEINLINE T Max( const T A, const T B )
+	template <typename T> 
+	static constexpr FORCEINLINE T Max(T A, T B)
 	{
+		// Even though this should be covered by the variadic, we still need this overload because of the many instances of
+		// FMath::Max<T>(a, b) with an explicit template parameter, which needs to continue to be supported.
+
 		return (B < A) ? A : B;
 	}
+	template <
+		typename T, typename... OtherTypes
+		UE_REQUIRES((std::is_same_v<T, OtherTypes> && ...))
+	>
+	static constexpr FORCEINLINE T Max(T A, OtherTypes... Others)
+	{
+		if constexpr (sizeof...(OtherTypes) > 0)
+		{
+			A = FGenericPlatformMath::Max(A, FGenericPlatformMath::Max(Others...));
+		}
+		return A;
+	}
 
 	/** Returns lower value in a generic way */
-	template< class T > 
-	static constexpr FORCEINLINE T Min( const T A, const T B )
+	template <typename T> 
+	static constexpr FORCEINLINE T Min(T A, T B)
 	{
+		// Even though this should be covered by the variadic, we still need this overload because of the many instances of
+		// FMath::Min<T>(a, b) with an explicit template parameter, which needs to continue to be supported.
+
 		return (A < B) ? A : B;
 	}
+	template <
+		typename T, typename... OtherTypes
+		UE_REQUIRES((std::is_same_v<T, OtherTypes> && ...))
+	>
+	static constexpr FORCEINLINE T Min(T A, OtherTypes... Others)
+	{
+		if constexpr (sizeof...(OtherTypes) > 0)
+		{
+			A = FGenericPlatformMath::Min(A, FGenericPlatformMath::Min(Others...));
+		}
+		return A;
+	}
 
 	// Allow mixing of float types to promote to highest precision type
 	MIX_FLOATS_2_ARGS(Max);
diff --git a/Engine/Source/Runtime/Core/Public/Templates/TypeHash.h b/Engine/Source/Runtime/Core/Public/Templates/TypeHash.h
--- a/Engine/Source/Runtime/Core/Public/Templates/TypeHash.h
+++ b/Engine/Source/Runtime/Core/Public/Templates/TypeHash.h
@@ -3,7 +3,7 @@
 #pragma once
 
 #include "CoreTypes.h"
-#include "Templates/IsEnum.h"
+#include "Templates/Requires.h"
 #include "Misc/Crc.h"
 
 #include <stdint.h>
@@ -33,7 +33,7 @@ namespace UE
  * This function cannot change for backward compatibility reasons.
  * You may want to choose HashCombineFast for a better in-memory hash combining function.
  */
-[[nodiscard]] inline uint32 HashCombine(uint32 A, uint32 C)
+[[nodiscard]] inline constexpr uint32 HashCombine(uint32 A, uint32 C)
 {
 	uint32 B = 0x9e3779b9;
 	A += B;
@@ -50,6 +50,18 @@ namespace UE
 
 	return C;
 }
+template <
+	typename... Types
+	UE_REQUIRES((std::is_convertible_v<Types, uint32> && ...))
+>
+[[nodiscard]] FORCEINLINE constexpr uint32 HashCombine(uint32 A, Types... Vals)
+{
+	if constexpr (sizeof...(Types) > 0)
+	{
+		A = HashCombine(A, HashCombine((uint32)Vals...));
+	}
+	return A;
+}
 
 /**
  * Combines two hash values to get a third.
@@ -59,10 +71,22 @@ namespace UE
  *           combined hash values which don't leave the running process,
  *           e.g. GetTypeHash() overloads.
  */
-[[nodiscard]] inline uint32 HashCombineFast(uint32 A, uint32 B)
+[[nodiscard]] inline constexpr uint32 HashCombineFast(uint32 A, uint32 B)
 {
 	return A ^ (B + 0x9e3779b9 + (A << 6) + (A >> 2));
 }
+template <
+	typename... Types
+	UE_REQUIRES((std::is_convertible_v<Types, uint32> && ...))
+>
+[[nodiscard]] FORCEINLINE constexpr uint32 HashCombineFast(uint32 A, Types... Vals)
+{
+	if constexpr (sizeof...(Types) > 0)
+	{
+		A = HashCombineFast(A, HashCombineFast((uint32)Vals...));
+	}
+	return A;
+}
 
 [[nodiscard]] inline uint32 PointerHash(const void* Key)
 {
@@ -189,3 +213,7 @@ template <typename T>
 // Use this when inside type that has GetTypeHash() (no in-parameters) implemented. It makes GetTypeHash dispatch in global namespace
 template <typename T>
 [[nodiscard]] FORCEINLINE uint32 GetTypeHashHelper(const T& V) { return static_cast<uint32>(GetTypeHash(V)); }
+
+#if UE_ENABLE_INCLUDE_ORDER_DEPRECATED_IN_5_6
+#include "Templates/IsEnum.h"
+#endif
