Updated MallocBinned2\3 with extended CSV stats (disabled by default) to get info on bins fragmentation, amount of memory allocated per bin and total memory allocated for small pool and large allocations.

diff --git a/Engine/Source/Runtime/Core/Private/HAL/MallocBinned2.cpp b/Engine/Source/Runtime/Core/Private/HAL/MallocBinned2.cpp
--- a/Engine/Source/Runtime/Core/Private/HAL/MallocBinned2.cpp
+++ b/Engine/Source/Runtime/Core/Private/HAL/MallocBinned2.cpp
@@ -5,7 +5,6 @@
 #include "GenericPlatform/GenericPlatformProcess.h"
 #include "Stats/Stats.h"
 #include "HAL/IConsoleManager.h"
-#include "HAL/MemoryMisc.h"
 #include "HAL/PlatformMisc.h"
 #include "Misc/App.h"
 #include "Misc/CoreDelegates.h"
@@ -34,11 +33,6 @@ PRAGMA_DISABLE_UNSAFE_TYPECAST_WARNINGS
 namespace
 {
 #if UE_MB2_ALLOCATOR_STATS
-	std::atomic<int64> AllocatedOSSmallPoolMemory(0);
-
-	std::atomic<int64> AllocatedLargePoolMemory(0); // memory requests to the OS which don't fit in the small pool
-	std::atomic<int64> AllocatedLargePoolMemoryWAlignment(0); // when we allocate at OS level we need to align to a size
-
 	int64 Binned2HashMemory = 0;
 #endif
 
@@ -281,6 +275,8 @@ struct FMallocBinned2::Private
 				Free->PoolIndex = InPoolIndex;
 				NodePool->FirstFreeBlock = Free;
 
+				UE_MBC_UPDATE_STATS(--Table.TotalUsedBins);
+
 				// Free this pool.
 				check(NodePool->Taken >= 1);
 				if (--NodePool->Taken == 0)
@@ -294,7 +290,16 @@ struct FMallocBinned2::Private
 						FScopeLock InnerLock(&Allocator.ExternalAllocMutex);
 						Allocator.CachedOSPageAllocator.Free(BasePtrOfNode, GMB2PageSize);
 					}
+
 #if UE_MB2_ALLOCATOR_STATS
+					int NumBinsInBlock = GMB2PageSize / InBinSize;
+					if (NumBinsInBlock * InBinSize + sizeof(FFreeBlock) > GMB2PageSize)
+					{
+						NumBinsInBlock--;
+					}
+
+					Table.TotalAllocatedBins -= NumBinsInBlock;
+					Table.TotalAllocatedMem -= GMB2PageSize;
 					AllocatedOSSmallPoolMemory.fetch_sub((int64)GMB2PageSize, std::memory_order_relaxed);
 #endif
 				}
@@ -350,7 +355,7 @@ void FMallocBinned2::FPoolList::LinkToFront(FPoolInfo* Pool)
 	Pool->Link(Front);
 }
 
-FMallocBinned2::FPoolInfo& FMallocBinned2::FPoolList::PushNewPoolToFront(FMallocBinned2& Allocator, uint32 InBinSize, uint32 InPoolIndex)
+FMallocBinned2::FPoolInfo& FMallocBinned2::FPoolList::PushNewPoolToFront(FMallocBinned2& Allocator, FPoolTable& Table, uint32 InPoolIndex)
 {
 	LLM_PLATFORM_SCOPE(ELLMTag::FMalloc);
 
@@ -365,14 +370,16 @@ FMallocBinned2::FPoolInfo& FMallocBinned2::FPoolList::PushNewPoolToFront(FMalloc
 	}
 
 #if !UE_USE_VERYLARGEPAGEALLOCATOR || !UE_MB2_BOOKKEEPING_AT_THE_END_OF_LARGEBLOCK
-	FFreeBlock* Free = new (FreePtr) FFreeBlock(GMB2PageSize, InBinSize, InPoolIndex, Allocator.CurrentCanary);
+	FFreeBlock* Free = new (FreePtr) FFreeBlock(GMB2PageSize, Table.BinSize, InPoolIndex, Allocator.CurrentCanary);
 	check(IsAligned(Free, GMB2PageSize));
 #else
 	FFreeBlock* FreeBlockPtr = GetPoolHeaderFromPointer(FreePtr);
-	FFreeBlock* Free = new (FreeBlockPtr) FFreeBlock(GMB2PageSize, InBinSize, InPoolIndex, Allocator.CurrentCanary);
+	FFreeBlock* Free = new (FreeBlockPtr) FFreeBlock(GMB2PageSize, Table.BinSize, InPoolIndex, Allocator.CurrentCanary);
 #endif
 
 #if UE_MB2_ALLOCATOR_STATS
+	Table.TotalAllocatedBins += Free->NumFreeBins;
+	Table.TotalAllocatedMem += GMB2PageSize;
 	AllocatedOSSmallPoolMemory.fetch_add((int64)GMB2PageSize, std::memory_order_relaxed);
 #endif
 
@@ -572,11 +579,12 @@ void* FMallocBinned2::MallocExternalSmall(SIZE_T Size, uint32 Alignment)
 	}
 	else
 	{
-		Pool = &Table.ActivePools.PushNewPoolToFront(*this, Table.BinSize, PoolIndex);
+		Pool = &Table.ActivePools.PushNewPoolToFront(*this, Table, PoolIndex);
 	}
 
 	void* Result = Pool->AllocateBin();
 #if UE_MB2_ALLOCATOR_STATS
+	++Table.TotalUsedBins;
 	AllocatedSmallPoolMemory.fetch_add(Table.BinSize, std::memory_order_relaxed);
 #endif
 
@@ -592,6 +600,7 @@ void* FMallocBinned2::MallocExternalSmall(SIZE_T Size, uint32 Alignment)
 					break;
 				}
 				Result = Pool->AllocateBin();
+				UE_MBC_UPDATE_STATS(++Table.TotalUsedBins);
 			}
 		}
 	}
@@ -742,9 +751,7 @@ void* FMallocBinned2::ReallocExternal(void* Ptr, SIZE_T NewSize, uint32 Alignmen
 	checkf(Alignment <= GMB2PageSize, TEXT("Invalid Realloc alignment: '%u' is greater than the page size '%u'"), Alignment, GMB2PageSize);
 	checkf(IsSupportedSize(NewSize), TEXT("Invalid Realloc size: '%" SIZE_T_FMT "'"), NewSize);
 
-#if UE_MB2_ALLOCATOR_STATS
-	AllocatedLargePoolMemory.fetch_add((int64)NewSize - (int64)PoolOSRequestedBytes, std::memory_order_relaxed);
-#endif
+	UE_MBC_UPDATE_STATS(AllocatedLargePoolMemory.fetch_add((int64)NewSize - (int64)PoolOSRequestedBytes, std::memory_order_relaxed));
 
 	Pool->SetOSAllocationSizes(NewSize, PoolOsBytes);
 
@@ -778,9 +785,7 @@ void FMallocBinned2::FreeExternal(void* Ptr)
 			BundlesToRecycle = Lists->RecycleFullBundle(PoolIndex, Private::GGlobalRecycler);
 			const bool bPushed = Lists->Free(Ptr, PoolIndex, BinSize);
 			check(bPushed);
-#if UE_MB2_ALLOCATOR_STATS
-			Lists->AllocatedMemory -= BinSize;
-#endif
+			UE_MBC_UPDATE_STATS(Lists->AllocatedMemory -= BinSize);
 		}
 		else
 		{
@@ -901,29 +906,6 @@ void FMallocBinned2::CanaryFail(const FFreeBlock* Block) const
 #endif
 }
 
-void FMallocBinned2::GetAllocatorStats( FGenericMemoryStats& OutStats )
-{
-#if UE_MB2_ALLOCATOR_STATS
-	const int64  TotalAllocatedSmallPoolMemory           = GetTotalAllocatedSmallPoolMemory();
-	const int64  LocalAllocatedOSSmallPoolMemory         = AllocatedOSSmallPoolMemory.load(std::memory_order_relaxed);
-	const int64  LocalAllocatedLargePoolMemory           = AllocatedLargePoolMemory.load(std::memory_order_relaxed);
-	const int64  LocalAllocatedLargePoolMemoryWAlignment = AllocatedLargePoolMemoryWAlignment.load(std::memory_order_relaxed);
-
-	OutStats.Add(TEXT("AllocatedSmallPoolMemory"), TotalAllocatedSmallPoolMemory);
-	OutStats.Add(TEXT("AllocatedOSSmallPoolMemory"), LocalAllocatedOSSmallPoolMemory);
-	OutStats.Add(TEXT("AllocatedLargePoolMemory"), LocalAllocatedLargePoolMemory);
-	OutStats.Add(TEXT("AllocatedLargePoolMemoryWAlignment"), LocalAllocatedLargePoolMemoryWAlignment);
-
-	const uint64 TotalAllocated = TotalAllocatedSmallPoolMemory + LocalAllocatedLargePoolMemory;
-	const uint64 TotalOSAllocated = LocalAllocatedOSSmallPoolMemory + LocalAllocatedLargePoolMemoryWAlignment + CachedOSPageAllocator.GetCachedFreeTotal();
-
-	OutStats.Add(TEXT("TotalAllocated"), TotalAllocated);
-	OutStats.Add(TEXT("TotalOSAllocated"), TotalOSAllocated);
-#endif
-
-	FMalloc::GetAllocatorStats(OutStats);
-}
-
 void FMallocBinned2::DumpAllocatorStats(class FOutputDevice& Ar)
 {
 #if UE_MB2_ALLOCATOR_STATS
@@ -975,6 +957,7 @@ void FMallocBinned2::UpdateStats()
 	CSV_CUSTOM_STAT(FMemory, AllocatorCachedSlackMB, (int32)(CachedOSPageAllocator.GetCachedFreeTotal() / (1024 * 1024)), ECsvCustomStatOp::Set);
 	CSV_CUSTOM_STAT(FMemory, AllocatorImmediatelyFreeableCachedSlackMB, (int32)(CachedOSPageAllocator.GetCachedImmediatelyFreeable() / (1024 * 1024)), ECsvCustomStatOp::Set);
 
+	UpdateStatsCommon(*this);
 	CachedOSPageAllocator.UpdateStats();
 	FScopedVirtualMallocTimer::UpdateStats();
 }
diff --git a/Engine/Source/Runtime/Core/Private/HAL/MallocBinned3.cpp b/Engine/Source/Runtime/Core/Private/HAL/MallocBinned3.cpp
--- a/Engine/Source/Runtime/Core/Private/HAL/MallocBinned3.cpp
+++ b/Engine/Source/Runtime/Core/Private/HAL/MallocBinned3.cpp
@@ -10,7 +10,6 @@ PRAGMA_DISABLE_UNSAFE_TYPECAST_WARNINGS
 #include "GenericPlatform/GenericPlatformProcess.h"
 #include "Stats/Stats.h"
 #include "HAL/IConsoleManager.h"
-#include "HAL/MemoryMisc.h"
 #include "HAL/PlatformMisc.h"
 #include "HAL/MallocBinnedCommonUtils.h"
 
@@ -30,11 +29,6 @@ PRAGMA_DISABLE_UNSAFE_TYPECAST_WARNINGS
 #endif
 
 #if UE_MB3_ALLOCATOR_STATS
-	int64 Binned3AllocatedOSSmallPoolMemory = 0;
-
-	int64 Binned3AllocatedLargePoolMemory = 0; // memory requests to the OS which don't fit in the small pool
-	int64 Binned3AllocatedLargePoolMemoryWAlignment = 0; // when we allocate at OS level we need to align to a size
-
 	TAtomic<int64> Binned3Commits;
 	TAtomic<int64> Binned3Decommits;
 
@@ -295,9 +289,8 @@ struct FMallocBinned3::Private
 		if (!InfoBlock)
 		{
 			InfoBlock = (FPoolInfoSmall*)Allocator.AllocateMetaDataMemory(Allocator.OsAllocationGranularity);
-#if UE_MB3_ALLOCATOR_STATS
-			PoolInfoMemory += Allocator.OsAllocationGranularity;
-#endif
+			UE_MBC_UPDATE_STATS(PoolInfoMemory += Allocator.OsAllocationGranularity);
+
 			DefaultConstructItems<FPoolInfoSmall>((void*)InfoBlock, InfosPerPage);
 		}
 
@@ -356,6 +349,8 @@ struct FMallocBinned3::Private
 				NodePool->NoFirstFreeIndex = 0;
 				check(uint32(NodePool->FirstFreeIndex) == BinIndexWithinBlock);
 
+				UE_MBC_UPDATE_STATS(--Table.TotalUsedBins);
+
 				// Free this pool.
 				check(NodePool->Taken >= 1);
 				if (--NodePool->Taken == 0)
@@ -363,16 +358,18 @@ struct FMallocBinned3::Private
 					NodePool->SetCanary(FPoolInfoSmall::ECanary::Unassigned, true, false);
 					Table.BlocksAllocatedBits.FreeBit(OutBlockIndex);
 
-					const uint64 AllocSize = Allocator.SmallPoolTables[InPoolIndex].NumMemoryPagesPerBlock * Allocator.OsAllocationGranularity;
-
 					if (!bWasExhaused)
 					{
 						Table.BlocksExhaustedBits.AllocBit(OutBlockIndex);
 					}
 
-					Allocator.Decommit(InPoolIndex, BaseBlockPtr, AllocSize);
+					const uint64 BlockSize = Allocator.SmallPoolTables[InPoolIndex].BlockSize;
+					Allocator.Decommit(InPoolIndex, BaseBlockPtr, BlockSize);
+
 #if UE_MB3_ALLOCATOR_STATS
-					Binned3AllocatedOSSmallPoolMemory -= AllocSize;
+					Table.TotalAllocatedBins -= BlockSize / InBinSize;
+					Table.TotalAllocatedMem -= BlockSize;
+					AllocatedOSSmallPoolMemory.fetch_sub(BlockSize, std::memory_order_relaxed);
 #endif
 				}
 				else if (bWasExhaused)
@@ -397,26 +394,26 @@ void FMallocBinned3::FreeBundles(FBundleNode* Bundles, uint32 PoolIndex)
 
 FMallocBinned3::FPoolInfoSmall* FMallocBinned3::PushNewPoolToFront(FMallocBinned3::FPoolTable& Table, uint32 InBinSize, uint32 InPoolIndex, uint32& OutBlockIndex)
 {
-	const uint32 BlockSize = OsAllocationGranularity * Table.NumMemoryPagesPerBlock;
-
 	// Allocate memory.
 	const uint32 BlockIndex = Table.BlocksAllocatedBits.AllocBit();
 	if (BlockIndex == MAX_uint32)
 	{
 		return nullptr;
 	}
-	uint8* FreePtr = BlockPointerFromIndecies(InPoolIndex, BlockIndex, BlockSize);
+	uint8* FreePtr = BlockPointerFromIndecies(InPoolIndex, BlockIndex, Table.BlockSize);
 
 	LLM_PLATFORM_SCOPE(ELLMTag::FMalloc);
-	Commit(InPoolIndex, FreePtr, BlockSize);
-	const uint64 EndOffset = UPTRINT(FreePtr + BlockSize) - UPTRINT(PoolBasePtr(InPoolIndex));
+	Commit(InPoolIndex, FreePtr, Table.BlockSize);
+	const uint64 EndOffset = UPTRINT(FreePtr + Table.BlockSize) - UPTRINT(PoolBasePtr(InPoolIndex));
 	if (EndOffset > Table.UnusedAreaOffsetLow)
 	{
 		Table.UnusedAreaOffsetLow = EndOffset;
 	}
-	FFreeBlock* Free = new ((void*)FreePtr) FFreeBlock(BlockSize, InBinSize, InPoolIndex);
+	FFreeBlock* Free = new ((void*)FreePtr) FFreeBlock(Table.BlockSize, InBinSize, InPoolIndex);
 #if UE_MB3_ALLOCATOR_STATS
-	Binned3AllocatedOSSmallPoolMemory += (int64)BlockSize;
+	AllocatedOSSmallPoolMemory.fetch_add((int64)Table.BlockSize, std::memory_order_relaxed);
+	Table.TotalAllocatedBins += Free->NumFreeBins;
+	Table.TotalAllocatedMem += Table.BlockSize;
 #endif
 	check(IsAligned(Free, OsAllocationGranularity));
 	// Create pool
@@ -553,27 +550,19 @@ FMallocBinned3::FMallocBinned3(const FPlatformMemory::FPlatformVirtualMemoryBloc
 
 		SmallPoolTables[Index].BinSize = SizeTable[Index].BinSize;
 		SmallPoolTables[Index].NumMemoryPagesPerBlock = SizeTable[Index].NumMemoryPagesPerBlock;
-
-		SmallPoolTables[Index].UnusedAreaOffsetLow = 0;
-		SmallPoolTables[Index].NumEverUsedBlocks = 0;
-#if UE_M3_ALLOCATOR_PER_BIN_STATS
-		SmallPoolTables[Index].TotalRequestedAllocSize.store(0, std::memory_order_relaxed);
-		SmallPoolTables[Index].TotalAllocCount.store(0, std::memory_order_relaxed);
-		SmallPoolTables[Index].TotalFreeCount.store(0, std::memory_order_relaxed);
-#endif
+		SmallPoolTables[Index].BlockSize = SizeTable[Index].NumMemoryPagesPerBlock * OsAllocationGranularity;
 
 		SmallBinSizesShifted[Index + 1] = (SizeTable[Index].BinSize >> UE_MBC_MIN_SMALL_POOL_ALIGNMENT_SHIFT);
 
-		const int64 TotalNumberOfBlocks = UE_MB3_MAX_MEMORY_PER_POOL_SIZE / (SizeTable[Index].NumMemoryPagesPerBlock * OsAllocationGranularity);
+		const int64 TotalNumberOfBlocks = UE_MB3_MAX_MEMORY_PER_POOL_SIZE / SmallPoolTables[Index].BlockSize;
 		const uint32 Size = Align(sizeof(FPoolInfoSmall**) * (TotalNumberOfBlocks + SmallPoolInfosPerPlatformPage - 1) / SmallPoolInfosPerPlatformPage, PLATFORM_CACHE_LINE_SIZE);
 		RequiredMetaMem += Size;
-#if UE_MB3_ALLOCATOR_STATS
-		PoolInfoMemory += Size;
-#endif
 
 		const int64 AllocationSize = Align(FBitTree::GetMemoryRequirements(TotalNumberOfBlocks), PLATFORM_CACHE_LINE_SIZE);
 		RequiredMetaMem += AllocationSize * 2;
+
 #if UE_MB3_ALLOCATOR_STATS
+		PoolInfoMemory += Size;
 		Binned3FreeBitsMemory += AllocationSize * 2;
 #endif
 	}
@@ -585,7 +574,7 @@ FMallocBinned3::FMallocBinned3(const FPlatformMemory::FPlatformVirtualMemoryBloc
 
 	for (uint32 Index = 0; Index < UE_MB3_SMALL_POOL_COUNT; ++Index)
 	{
-		const int64 TotalNumberOfBlocks = UE_MB3_MAX_MEMORY_PER_POOL_SIZE / (SizeTable[Index].NumMemoryPagesPerBlock * OsAllocationGranularity);
+		const int64 TotalNumberOfBlocks = UE_MB3_MAX_MEMORY_PER_POOL_SIZE / SmallPoolTables[Index].BlockSize;
 		const uint32 Size = Align(sizeof(FPoolInfoSmall**) * (TotalNumberOfBlocks + SmallPoolInfosPerPlatformPage - 1) / SmallPoolInfosPerPlatformPage, PLATFORM_CACHE_LINE_SIZE);
 
 		SmallPoolTables[Index].PoolInfos = (FPoolInfoSmall**)MetaMem;
@@ -627,9 +616,7 @@ FMallocBinned3::~FMallocBinned3()
 
 void FMallocBinned3::Commit(uint32 InPoolIndex, void *Ptr, SIZE_T Size)
 {
-#if UE_MB3_ALLOCATOR_STATS
-	Binned3Commits++;
-#endif
+	UE_MBC_UPDATE_STATS(Binned3Commits++);
 
 #if !BINNED3_USE_SEPARATE_VM_PER_POOL
 	Binned3BaseVMBlock.CommitByPtr(Ptr, Size);
@@ -643,10 +630,7 @@ void FMallocBinned3::Commit(uint32 InPoolIndex, void *Ptr, SIZE_T Size)
 void FMallocBinned3::Decommit(uint32 InPoolIndex, void *Ptr, SIZE_T Size)
 {
 	LLM_IF_ENABLED(FLowLevelMemTracker::Get().OnLowLevelFree(ELLMTracker::Platform, Ptr));
-
-#if UE_MB3_ALLOCATOR_STATS
-	Binned3Decommits++;
-#endif
+	UE_MBC_UPDATE_STATS(Binned3Decommits++);
 
 #if !BINNED3_USE_SEPARATE_VM_PER_POOL
 	Binned3BaseVMBlock.DecommitByPtr(Ptr, Size);
@@ -748,13 +732,13 @@ void* FMallocBinned3::MallocExternal(SIZE_T Size, uint32 Alignment)
 			}
 		}
 
-		const uint32 BlockSize = OsAllocationGranularity * Table.NumMemoryPagesPerBlock;
-		uint8* BlockPtr = BlockPointerFromIndecies(PoolIndex, BlockIndex, BlockSize);
+		uint8* BlockPtr = BlockPointerFromIndecies(PoolIndex, BlockIndex, Table.BlockSize);
 
 		void* Result = Pool->AllocateBin(BlockPtr, Table.BinSize);
 #if UE_MB3_ALLOCATOR_STATS
+		++Table.TotalUsedBins;
 		Table.HeadEndAlloc(Size);
-		AllocatedSmallPoolMemory.fetch_add(PoolIndexToBinSize(PoolIndex), std::memory_order_relaxed);
+		AllocatedSmallPoolMemory.fetch_add(Table.BinSize, std::memory_order_relaxed);
 #endif
 		if (GMallocBinnedAllocExtra)
 		{
@@ -768,6 +752,7 @@ void* FMallocBinned3::MallocExternal(SIZE_T Size, uint32 Alignment)
 						break;
 					}
 					Result = Pool->AllocateBin(BlockPtr, Table.BinSize);
+					UE_MBC_UPDATE_STATS(++Table.TotalUsedBins);
 				}
 			}
 		}
@@ -829,8 +814,8 @@ void* FMallocBinned3::MallocExternal(SIZE_T Size, uint32 Alignment)
 #endif
 
 #if UE_MB3_ALLOCATOR_STATS
-	Binned3AllocatedLargePoolMemory += Size;
-	Binned3AllocatedLargePoolMemoryWAlignment += AlignedSize;
+	AllocatedLargePoolMemory.fetch_add(Size, std::memory_order_relaxed);
+	AllocatedLargePoolMemoryWAlignment.fetch_add(AlignedSize, std::memory_order_relaxed);
 #endif
 
 	// Create pool.
@@ -910,10 +895,8 @@ void* FMallocBinned3::ReallocExternal(void* Ptr, SIZE_T NewSize, uint32 Alignmen
 		return Result;
 	}
 
-#if UE_MB3_ALLOCATOR_STATS
-	Binned3AllocatedLargePoolMemory += ((int64)NewSize) - ((int64)PoolOSRequestedBytes);
+	UE_MBC_UPDATE_STATS(AllocatedLargePoolMemory.fetch_add((int64)NewSize - (int64)PoolOSRequestedBytes, std::memory_order_relaxed));
 	// don't need to change the Binned3AllocatedLargePoolMemoryWAlignment because we didn't reallocate so it's the same size
-#endif
 	
 	Pool->SetOSAllocationSize(NewSize);
 	ExternalAllocMutex.Unlock();
@@ -980,8 +963,8 @@ void FMallocBinned3::FreeExternal(void* Ptr, uint64 PoolIndex)
 			VMPages = Pool->GetOsVMPages();
 
 #if UE_MB3_ALLOCATOR_STATS
-			Binned3AllocatedLargePoolMemory -= PoolOSRequestedBytes;
-			Binned3AllocatedLargePoolMemoryWAlignment -= PoolOsBytes;
+			AllocatedLargePoolMemory.fetch_sub((int64)PoolOSRequestedBytes, std::memory_order_relaxed);
+			AllocatedLargePoolMemoryWAlignment.fetch_sub((int64)PoolOsBytes, std::memory_order_relaxed);
 #endif
 
 			checkf(PoolOSRequestedBytes <= PoolOsBytes, TEXT("FMallocBinned3::FreeExternal %u %u"), int32(PoolOSRequestedBytes), int32(PoolOsBytes));
@@ -1039,33 +1022,6 @@ void FMallocBinned3::FlushCurrentThreadCacheInternal(bool bNewEpochOnly)
 	FMallocBinnedCommonUtils::FlushCurrentThreadCache(*this, bNewEpochOnly);
 }
 
-void FMallocBinned3::GetAllocatorStats(FGenericMemoryStats& OutStats)
-{
-#if UE_MB3_ALLOCATOR_STATS
-
-	const int64 TotalAllocatedSmallPoolMemory = GetTotalAllocatedSmallPoolMemory();
-
-	OutStats.Add(TEXT("Binned3AllocatedSmallPoolMemory"), TotalAllocatedSmallPoolMemory);
-	OutStats.Add(TEXT("Binned3AllocatedOSSmallPoolMemory"), Binned3AllocatedOSSmallPoolMemory);
-	OutStats.Add(TEXT("Binned3AllocatedLargePoolMemory"), Binned3AllocatedLargePoolMemory);
-	OutStats.Add(TEXT("Binned3AllocatedLargePoolMemoryWAlignment"), Binned3AllocatedLargePoolMemoryWAlignment);
-
-	const uint64 CachedMem =  
-#if UE_MB3_USE_CACHED_PAGE_ALLOCATOR_FOR_LARGE_ALLOCS
-		GetCachedOSPageAllocator().GetCachedFreeTotal();
-#else
-		0;
-#endif
-
-	const uint64 TotalAllocated = TotalAllocatedSmallPoolMemory + Binned3AllocatedLargePoolMemory;
-	const uint64 TotalOSAllocated = Binned3AllocatedOSSmallPoolMemory + Binned3AllocatedLargePoolMemoryWAlignment + CachedMem;
-
-	OutStats.Add(TEXT("TotalAllocated"), TotalAllocated);
-	OutStats.Add(TEXT("TotalOSAllocated"), TotalOSAllocated);
-#endif
-	FMalloc::GetAllocatorStats(OutStats);
-}
-
 #if UE_MB3_ALLOCATOR_STATS && BINNED3_USE_SEPARATE_VM_PER_POOL
 void FMallocBinned3::RecordPoolSearch(uint32 Tests) const
 {
@@ -1084,9 +1040,9 @@ void FMallocBinned3::DumpAllocatorStats(class FOutputDevice& Ar)
 	Ar.Logf(TEXT("UE_MB3_MAX_SMALL_POOL_SIZE = %d"), int32(UE_MB3_MAX_SMALL_POOL_SIZE));
 	Ar.Logf(TEXT("UE_MB3_MAX_MEMORY_PER_POOL_SIZE = %llu"), uint64(UE_MB3_MAX_MEMORY_PER_POOL_SIZE));
 	Ar.Logf(TEXT("Small Pool Allocations: %fmb  (including bin size padding)"), ((double)TotalAllocatedSmallPoolMemory) / (1024.0f * 1024.0f));
-	Ar.Logf(TEXT("Small Pool OS Allocated: %fmb"), ((double)Binned3AllocatedOSSmallPoolMemory) / (1024.0f * 1024.0f));
-	Ar.Logf(TEXT("Large Pool Requested Allocations: %fmb"), ((double)Binned3AllocatedLargePoolMemory) / (1024.0f * 1024.0f));
-	Ar.Logf(TEXT("Large Pool OS Allocated: %fmb"), ((double)Binned3AllocatedLargePoolMemoryWAlignment) / (1024.0f * 1024.0f));
+	Ar.Logf(TEXT("Small Pool OS Allocated: %fmb"), ((double)AllocatedOSSmallPoolMemory.load(std::memory_order_relaxed)) / (1024.0f * 1024.0f));
+	Ar.Logf(TEXT("Large Pool Requested Allocations: %fmb"), ((double)AllocatedLargePoolMemory.load(std::memory_order_relaxed)) / (1024.0f * 1024.0f));
+	Ar.Logf(TEXT("Large Pool OS Allocated: %fmb"), ((double)AllocatedLargePoolMemoryWAlignment.load(std::memory_order_relaxed)) / (1024.0f * 1024.0f));
 	Ar.Logf(TEXT("PoolInfo: %fmb"), ((double)PoolInfoMemory) / (1024.0f * 1024.0f));
 	Ar.Logf(TEXT("Hash: %fmb"), ((double)HashMemory) / (1024.0f * 1024.0f));
 	Ar.Logf(TEXT("Free Bits: %fmb"), ((double)Binned3FreeBitsMemory) / (1024.0f * 1024.0f));
@@ -1107,7 +1063,7 @@ void FMallocBinned3::DumpAllocatorStats(class FOutputDevice& Ar)
 #endif
 	Ar.Logf(TEXT("Total allocated from OS: %fmb"), 
 		((double)
-			Binned3AllocatedOSSmallPoolMemory + Binned3AllocatedLargePoolMemoryWAlignment + PoolInfoMemory + HashMemory + Binned3FreeBitsMemory + TLSMemory.load(std::memory_order_relaxed)
+			AllocatedOSSmallPoolMemory.load(std::memory_order_relaxed) + AllocatedLargePoolMemoryWAlignment.load(std::memory_order_relaxed) + PoolInfoMemory + HashMemory + Binned3FreeBitsMemory + TLSMemory.load(std::memory_order_relaxed)
 			) / (1024.0f * 1024.0f));
 
 
@@ -1116,6 +1072,14 @@ void FMallocBinned3::DumpAllocatorStats(class FOutputDevice& Ar)
 	Ar.Logf(TEXT("MemoryRangeFree    %d calls %6.3fs    %6.3fus / call"), MemoryRangeFreeTotalCount.Load(), float(MemoryRangeFreeTotalTime.Load()), float(MemoryRangeFreeTotalTime.Load()) * 1000000.0f / float(MemoryRangeFreeTotalCount.Load()));
 #endif
 
+	for (int32 i = 0; i < UE_MB3_SMALL_POOL_COUNT; i++)
+	{
+		const float Fragmentation = 1.0f - (float)SmallPoolTables[i].TotalUsedBins / (float)SmallPoolTables[i].TotalAllocatedBins;
+		const float TotalMem = (float)SmallPoolTables[i].TotalAllocatedMem / 1024.0f / 1024.0f;
+		Ar.Logf(TEXT("Bin %6d Fragrmentation %d %%, Wasted Mem %.2f MB, Total Allocated Mem %.2f MB"), 
+			PoolIndexToBinSize(i), int(Fragmentation * 100.0f), TotalMem * Fragmentation, TotalMem);
+	}
+
 #if UE_M3_ALLOCATOR_PER_BIN_STATS
 	for (int32 PoolIndex = 0; PoolIndex < UE_MB3_SMALL_POOL_COUNT; PoolIndex++)
 	{
@@ -1123,7 +1087,7 @@ void FMallocBinned3::DumpAllocatorStats(class FOutputDevice& Ar)
 		const uint32 CommittedBlocks = SmallPoolTables[PoolIndex].BlocksAllocatedBits.CountOnes(SmallPoolTables[PoolIndex].NumEverUsedBlocks);
 		const uint32 PartialBlocks = SmallPoolTables[PoolIndex].NumEverUsedBlocks - SmallPoolTables[PoolIndex].BlocksExhaustedBits.CountOnes(SmallPoolTables[PoolIndex].NumEverUsedBlocks);
 		const uint32 FullBlocks = CommittedBlocks - PartialBlocks;
-		const int64 ComittedVM = VM - (SmallPoolTables[PoolIndex].NumEverUsedBlocks - CommittedBlocks) * SmallPoolTables[PoolIndex].NumMemoryPagesPerBlock * OsAllocationGranularity;
+		const int64 ComittedVM = VM - (SmallPoolTables[PoolIndex].NumEverUsedBlocks - CommittedBlocks) * SmallPoolTables[PoolIndex].BlockSize;
 
 		const int64 AveSize = SmallPoolTables[PoolIndex].TotalAllocCount.load(std::memory_order_relaxed) ? SmallPoolTables[PoolIndex].TotalRequestedAllocSize.load(std::memory_order_relaxed) / SmallPoolTables[PoolIndex].TotalAllocCount.load(std::memory_order_relaxed) : 0;
 		const int64 EstPadWaste = ((SmallPoolTables[PoolIndex].TotalAllocCount.load(std::memory_order_relaxed) - SmallPoolTables[PoolIndex].TotalFreeCount.load(std::memory_order_relaxed)) * (PoolIndexToBinSize(PoolIndex) - AveSize));
diff --git a/Engine/Source/Runtime/Core/Private/HAL/MallocBinnedCommon.cpp b/Engine/Source/Runtime/Core/Private/HAL/MallocBinnedCommon.cpp
--- a/Engine/Source/Runtime/Core/Private/HAL/MallocBinnedCommon.cpp
+++ b/Engine/Source/Runtime/Core/Private/HAL/MallocBinnedCommon.cpp
@@ -5,12 +5,16 @@
 #include "Containers/ArrayView.h"
 #include "Misc/AssertionMacros.h"
 #include "Math/NumericLimits.h"
+#include "HAL/MemoryMisc.h"
 #include "HAL/IConsoleManager.h"
 #include "Templates/AlignmentTemplates.h"
 #include "Templates/UnrealTemplate.h"
+#include "ProfilingDebugging/CsvProfiler.h"
 
 PRAGMA_DISABLE_UNSAFE_TYPECAST_WARNINGS
 
+CSV_DEFINE_CATEGORY(MallocBinned, false);
+
 // Bin sizes are based around getting the maximum amount of allocations per block, with as little alignment waste as possible.
 // Bin sizes should be close to even divisors of the system page size, and well distributed.
 // They must be 16-byte aligned as well.
@@ -532,7 +536,7 @@ int32 GMallocBinnedFlushRegisteredThreadCachesOnOneThread = 1;
 static FAutoConsoleVariableRef GMallocBinnedFlushRegisteredThreadCachesOnOneThreadCVar(
 	TEXT("MallocBinned.FlushRegisteredThreadCachesOnOneThread"),
 	GMallocBinnedFlushRegisteredThreadCachesOnOneThread,
-	TEXT("Wether or not to attempt to flush registered thread caches on one thread (enabled by default)."));
+	TEXT("Whether or not to attempt to flush registered thread caches on one thread (enabled by default)."));
 
 #if UE_MBC_ALLOW_RUNTIME_TWEAKING
 
@@ -578,10 +582,42 @@ uint32 FMallocBinnedCommonBase::BinnedTlsSlot = FPlatformTLS::InvalidTlsSlot;
 	std::atomic<int64> FMallocBinnedCommonBase::TLSMemory(0);
 	std::atomic<int64> FMallocBinnedCommonBase::ConsolidatedMemory(0);
 
-	std::atomic<int64> FMallocBinnedCommonBase::AllocatedSmallPoolMemory(0); // memory that's requested to be allocated by the game
+	std::atomic<int64> FMallocBinnedCommonBase::AllocatedSmallPoolMemory(0);
+	std::atomic<int64> FMallocBinnedCommonBase::AllocatedOSSmallPoolMemory(0);
+
+	std::atomic<int64> FMallocBinnedCommonBase::AllocatedLargePoolMemory(0);
+	std::atomic<int64> FMallocBinnedCommonBase::AllocatedLargePoolMemoryWAlignment(0);
 
 	int64 FMallocBinnedCommonBase::PoolInfoMemory = 0;
 	int64 FMallocBinnedCommonBase::HashMemory = 0;
+
+	int32 GMallocBinnedEnableCSVStats = 0;
+	static FAutoConsoleVariableRef GMallocBinnedEnableCSVStatsCVar(
+		TEXT("MallocBinned.EnableCSVStats"),
+		GMallocBinnedEnableCSVStats,
+		TEXT("Whether or not to enable extended CSV stats with fragmentation stats (disabled by default)."));
+#endif
+
+void FMallocBinnedCommonBase::GetAllocatorStatsInternal(FGenericMemoryStats& OutStats, int64 TotalAllocatedSmallPoolMemory)
+{
+#if UE_MBC_ALLOCATOR_STATS
+	const int64 LocalAllocatedOSSmallPoolMemory = AllocatedOSSmallPoolMemory.load(std::memory_order_relaxed);
+	const int64 LocalAllocatedLargePoolMemory = AllocatedLargePoolMemory.load(std::memory_order_relaxed);
+	const int64 LocalAllocatedLargePoolMemoryWAlignment = AllocatedLargePoolMemoryWAlignment.load(std::memory_order_relaxed);
+
+	OutStats.Add(TEXT("AllocatedSmallPoolMemory"), TotalAllocatedSmallPoolMemory);
+	OutStats.Add(TEXT("AllocatedOSSmallPoolMemory"), LocalAllocatedOSSmallPoolMemory);
+	OutStats.Add(TEXT("AllocatedLargePoolMemory"), LocalAllocatedLargePoolMemory);
+	OutStats.Add(TEXT("AllocatedLargePoolMemoryWAlignment"), LocalAllocatedLargePoolMemoryWAlignment);
+
+	const uint64 TotalAllocated = TotalAllocatedSmallPoolMemory + LocalAllocatedLargePoolMemory;
+	const uint64 TotalOSAllocated = LocalAllocatedOSSmallPoolMemory + LocalAllocatedLargePoolMemoryWAlignment + GetTotalFreeCachedMemorySize();
+
+	OutStats.Add(TEXT("TotalAllocated"), TotalAllocated);
+	OutStats.Add(TEXT("TotalOSAllocated"), TotalOSAllocated);
 #endif
 
+	FMalloc::GetAllocatorStats(OutStats);
+}
+
 PRAGMA_RESTORE_UNSAFE_TYPECAST_WARNINGS
diff --git a/Engine/Source/Runtime/Core/Public/HAL/MallocBinned2.h b/Engine/Source/Runtime/Core/Public/HAL/MallocBinned2.h
--- a/Engine/Source/Runtime/Core/Public/HAL/MallocBinned2.h
+++ b/Engine/Source/Runtime/Core/Public/HAL/MallocBinned2.h
@@ -15,7 +15,6 @@
 #include "Misc/Fork.h"
 #include "Templates/Atomic.h"
 
-struct FGenericMemoryStats;
 
 #define UE_MB2_MAX_CACHED_OS_FREES (64)
 #if PLATFORM_64BITS
@@ -81,6 +80,8 @@ class FMallocBinned2 : public TMallocBinnedCommon<FMallocBinned2, UE_MB2_SMALL_P
 	struct FFreeBlock;
 
 public:
+	struct FPoolTable;
+
 	struct FPoolInfo
 	{
 		enum class ECanary : uint16
@@ -187,7 +188,7 @@ private:
 
 		void LinkToFront(FPoolInfo* Pool);
 
-		FPoolInfo& PushNewPoolToFront(FMallocBinned2& Allocator, uint32 InBytes, uint32 InPoolIndex);
+		FPoolInfo& PushNewPoolToFront(FMallocBinned2& Allocator, FPoolTable& Table, uint32 InPoolIndex);
 
 		void ValidateActivePools() const;
 		void ValidateExhaustedPools() const;
@@ -196,12 +197,19 @@ private:
 		FPoolInfo* Front = nullptr;
 	};
 
+public:
 	/** Pool table. */
 	struct FPoolTable
 	{
 		FPoolList ActivePools;
 		FPoolList ExhaustedPools;
-		uint32    BinSize = 0;
+		uint32 BinSize = 0;
+
+#if UE_MB2_ALLOCATOR_STATS
+		uint32 TotalUsedBins = 0;			// Used to calculated load factor, i.e.: 
+		uint32 TotalAllocatedBins = 0;		// used bins divided by total bins number in all allocated blocks
+		uint32 TotalAllocatedMem = 0;
+#endif
 
 		FCriticalSection Mutex;
 
@@ -211,6 +219,7 @@ private:
 	// Pool tables for different pool sizes
 	FPoolTable SmallPoolTables[UE_MB2_SMALL_POOL_COUNT];
 
+private:
 #if BINNED2_FORK_SUPPORT
 	EBlockCanary CurrentCanary = EBlockCanary::PreFork; // The value of the canary for pages we have allocated this side of the fork 
 	EBlockCanary OldCanary = EBlockCanary::PreFork;		// If we have forked, the value canary of old pages we should avoid touching 
@@ -417,9 +426,7 @@ public:
 						{
 							const bool bDidPush = Lists->Free(Ptr, PoolIndex, BinSize);
 							checkSlow(bDidPush);
-#if UE_MB2_ALLOCATOR_STATS
-							Lists->AllocatedMemory -= BinSize;
-#endif
+							UE_MBC_UPDATE_STATS(Lists->AllocatedMemory -= BinSize);
 						}
 
 						return Result;
@@ -469,9 +476,7 @@ public:
 				// If canary is invalid we will assert in FreeExternal. Otherwise it's the pre-fork canary and we will turn this free into a no-op.
 				if (BasePtr->CanaryAndForkState == CurrentCanary && Lists->Free(Ptr, BasePtr->PoolIndex, BinSize))
 				{
-#if UE_MB2_ALLOCATOR_STATS
-					Lists->AllocatedMemory -= BinSize;
-#endif
+					UE_MBC_UPDATE_STATS(Lists->AllocatedMemory -= BinSize);
 					return;
 				}
 			}
@@ -541,7 +546,6 @@ public:
 		}
 	}
 
-	CORE_API virtual void GetAllocatorStats( FGenericMemoryStats& out_Stats ) override;
 	/** Dumps current allocator stats to the log. */
 	CORE_API virtual void DumpAllocatorStats(class FOutputDevice& Ar) override;
 	
diff --git a/Engine/Source/Runtime/Core/Public/HAL/MallocBinned3.h b/Engine/Source/Runtime/Core/Public/HAL/MallocBinned3.h
--- a/Engine/Source/Runtime/Core/Public/HAL/MallocBinned3.h
+++ b/Engine/Source/Runtime/Core/Public/HAL/MallocBinned3.h
@@ -17,7 +17,6 @@
 #include "Templates/AlignmentTemplates.h"
 #include "Templates/Atomic.h"
 
-struct FGenericMemoryStats;
 
 #define UE_MB3_USE_CACHED_PAGE_ALLOCATOR_FOR_LARGE_ALLOCS (0)
 
@@ -123,7 +122,6 @@ public:
 		uint32 CommitSize;						// Number of bytes committed by the OS
 	};
 
-private:
 	// Forward declarations
 	struct FPoolInfoSmall;
 	struct FPoolTable;
@@ -133,24 +131,30 @@ private:
 	struct FPoolTable
 	{
 		uint32 BinSize;						// Bin size, i.e. 16, 32, 64 etc bytes. Pretty redundant and can be computed
-		//uint32 BlockSize;					// Size of a block. A block contains N bins in a way to minimize tail memory waste
-		uint32 NumMemoryPagesPerBlock;		// Num memory pages needed to allocate one block. 
+		uint32 BlockSize;					// Size of a block. A block contains N bins in a way to minimize tail memory waste
+		uint32 NumMemoryPagesPerBlock;		// Num memory pages needed to allocate one block.
+		uint32 NumEverUsedBlocks = 0;
 
 		FBitTree BlocksAllocatedBits;		// A bit per per block, one indicates a committed block
 		FBitTree BlocksExhaustedBits;		// One here means that the block is completely full
 
-		uint32 NumEverUsedBlocks;
 		FPoolInfoSmall** PoolInfos;			// Book keeping info about every allocated area for this bin's pool
 
-		uint64 UnusedAreaOffsetLow;			// High watermark for allocated vm memory for this pool
+		uint64 UnusedAreaOffsetLow = 0;		// High watermark for allocated vm memory for this pool
+
+#if UE_MB3_ALLOCATOR_STATS
+		uint32 TotalUsedBins = 0;			// Used to calculated load factor, i.e.: 
+		uint32 TotalAllocatedBins = 0;		// used bins divided by total bins number in all allocated blocks
+		uint32 TotalAllocatedMem = 0;
+#endif
 
 		FCriticalSection Mutex;
 
 #if UE_M3_ALLOCATOR_PER_BIN_STATS
 		// these are "head end" stats, above the TLS cache
-		std::atomic<int64> TotalRequestedAllocSize;
-		std::atomic<int64> TotalAllocCount;
-		std::atomic<int64> TotalFreeCount;
+		std::atomic<int64> TotalRequestedAllocSize = 0;
+		std::atomic<int64> TotalAllocCount = 0;
+		std::atomic<int64> TotalFreeCount = 0;
 
 		FORCEINLINE void HeadEndAlloc(SIZE_T Size)
 		{
@@ -175,6 +179,7 @@ private:
 	// Pool tables for different pool sizes
 	FPoolTable SmallPoolTables[UE_MB3_SMALL_POOL_COUNT];
 
+private:
 	uint32 SmallPoolInfosPerPlatformPage;
 
 #if !BINNED3_USE_SEPARATE_VM_PER_POOL
@@ -437,15 +442,17 @@ public:
 	CORE_API virtual bool ValidateHeap() override;
 	CORE_API virtual void Trim(bool bTrimThreadCaches) override;
 	CORE_API virtual const TCHAR* GetDescriptiveName() override;
+	/** Dumps current allocator stats to the log. */
+	CORE_API virtual void DumpAllocatorStats(class FOutputDevice& Ar) override;
+	CORE_API virtual void UpdateStats() override
+	{
+		UpdateStatsCommon(*this);
+	}
 	// End FMalloc interface.
 
 	CORE_API void* MallocExternal(SIZE_T Size, uint32 Alignment);
 	CORE_API void* ReallocExternal(void* Ptr, SIZE_T NewSize, uint32 Alignment);
 	CORE_API void FreeExternal(void *Ptr, uint64 PoolIndex);
-
-	CORE_API virtual void GetAllocatorStats( FGenericMemoryStats& out_Stats ) override;
-	/** Dumps current allocator stats to the log. */
-	CORE_API virtual void DumpAllocatorStats(class FOutputDevice& Ar) override;
 	
 	// +1 enables PoolIndexToBinSize(~0u / -1) dummy access that helps avoid PoolIndex == 0 branching in Realloc(),
 	// see ((!PoolIndex) | (NewSize > PoolIndexToBinSize(static_cast<uint32>(PoolIndex - 1)))
diff --git a/Engine/Source/Runtime/Core/Public/HAL/MallocBinnedCommon.h b/Engine/Source/Runtime/Core/Public/HAL/MallocBinnedCommon.h
--- a/Engine/Source/Runtime/Core/Public/HAL/MallocBinnedCommon.h
+++ b/Engine/Source/Runtime/Core/Public/HAL/MallocBinnedCommon.h
@@ -12,6 +12,7 @@
 #include "Misc/ScopeLock.h"
 #include "Templates/AlignmentTemplates.h"
 #include "Logging/LogMacros.h"
+#include "ProfilingDebugging/CsvProfiler.h"
 
 
 // A project can define it's own UE_MBC_MAX_LISTED_SMALL_POOL_SIZE and UE_MBC_NUM_LISTED_SMALL_POOLS to reduce runtime memory usage
@@ -84,9 +85,20 @@
 #	define UE_MBC_ALLOCATOR_STATS (!UE_BUILD_SHIPPING || WITH_EDITOR)
 #endif
 
+#if UE_MBC_ALLOCATOR_STATS
+#	define UE_MBC_UPDATE_STATS(x) x
+	extern CORE_API int32 GMallocBinnedEnableCSVStats;
+#else
+#	define UE_MBC_UPDATE_STATS(x)
+#endif
+
+CSV_DECLARE_CATEGORY_EXTERN(MallocBinned);
+
 extern CORE_API float GMallocBinnedFlushThreadCacheMaxWaitTime;
 extern CORE_API int32 GMallocBinnedFlushRegisteredThreadCachesOnOneThread;
 
+struct FGenericMemoryStats;
+
 
 class FBitTree
 {
@@ -381,7 +393,10 @@ protected:
 #if UE_MBC_ALLOCATOR_STATS
 	static std::atomic<int64> TLSMemory;
 	static std::atomic<int64> ConsolidatedMemory;
-	static std::atomic<int64> AllocatedSmallPoolMemory;
+	static std::atomic<int64> AllocatedSmallPoolMemory;				// requested small pool memory allocations
+	static std::atomic<int64> AllocatedOSSmallPoolMemory;			// total small pool memory allocated by the os, always larger than AllocatedSmallPoolMemory
+	static std::atomic<int64> AllocatedLargePoolMemory;				// memory requests to the OS which don't fit in the small pool
+	static std::atomic<int64> AllocatedLargePoolMemoryWAlignment;	// when we allocate at OS level we need to align to a size
 
 	static int64 PoolInfoMemory;
 	static int64 HashMemory;
@@ -394,6 +409,8 @@ protected:
 		// this is expected not to return
 		FPlatformMemory::OnOutOfMemory(Size, Alignment);
 	}
+
+	void GetAllocatorStatsInternal(FGenericMemoryStats& OutStats, int64 TotalAllocatedSmallPoolMemory);
 };
 
 template <class AllocType, int NumSmallPools, int MaxSmallPoolSize>
@@ -437,6 +454,12 @@ class TMallocBinnedCommon : public FMallocBinnedCommonBase
 		}
 	};
 
+public:
+	CORE_API virtual void GetAllocatorStats(FGenericMemoryStats& OutStats) override
+	{
+		GetAllocatorStatsInternal(OutStats, GetTotalAllocatedSmallPoolMemory());
+	}
+
 protected:
 	struct FPerThreadFreeBlockLists
 	{
@@ -464,9 +487,8 @@ protected:
 			{
 				const int64 TLSSize = Align(sizeof(FPerThreadFreeBlockLists), AllocType::OsAllocationGranularity);
 				ThreadSingleton = new (AllocType::AllocateMetaDataMemory(TLSSize)) FPerThreadFreeBlockLists();
-#if UE_MBC_ALLOCATOR_STATS
-				TLSMemory.fetch_add(TLSSize, std::memory_order_relaxed);
-#endif
+				UE_MBC_UPDATE_STATS(TLSMemory.fetch_add(TLSSize, std::memory_order_relaxed));
+
 				verify(ThreadSingleton);
 				ThreadSingleton->bLockedByOwnerThread = true;
 				ThreadSingleton->Lock();
@@ -502,9 +524,8 @@ protected:
 			if (ThreadSingleton)
 			{
 				const int64 TLSSize = Align(sizeof(FPerThreadFreeBlockLists), AllocType::OsAllocationGranularity);
-#if UE_MBC_ALLOCATOR_STATS
-				TLSMemory.fetch_sub(TLSSize, std::memory_order_relaxed);
-#endif
+				UE_MBC_UPDATE_STATS(TLSMemory.fetch_sub(TLSSize, std::memory_order_relaxed));
+
 				UnregisterThreadFreeBlockLists(ThreadSingleton);
 				ThreadSingleton->bLockedByOwnerThread = false;
 				ThreadSingleton->Unlock();
@@ -605,9 +626,8 @@ protected:
 					const uint64 PoolArraySize = NumPools * sizeof(PoolInfo);
 
 					void* Result = Allocator.AllocateMetaDataMemory(PoolArraySize);
-#if UE_MBC_ALLOCATOR_STATS
-					PoolInfoMemory += PoolArraySize;
-#endif
+					UE_MBC_UPDATE_STATS(PoolInfoMemory += PoolArraySize);
+
 					if (!Result)
 					{
 						FScopeUnlock TempUnlock(&Allocator.ExternalAllocMutex);
@@ -648,9 +668,7 @@ protected:
 			if (!Allocator.HashBucketFreeList)
 			{
 				Allocator.HashBucketFreeList = (FPoolHashBucket*)Allocator.AllocateMetaDataMemory(AllocType::OsAllocationGranularity);
-#if UE_MBC_ALLOCATOR_STATS
-				HashMemory += AllocType::OsAllocationGranularity;
-#endif
+				UE_MBC_UPDATE_STATS(HashMemory += AllocType::OsAllocationGranularity);
 
 				for (UPTRINT i = 0, n = AllocType::OsAllocationGranularity / sizeof(FPoolHashBucket); i < n; ++i)
 				{
@@ -897,14 +915,60 @@ protected:
 	}
 #endif
 
+	void UpdateStatsCommon(const AllocType& Alloc)
+	{
+#if UE_MBC_ALLOCATOR_STATS && CSV_PROFILER_STATS 
+		if (!GMallocBinnedEnableCSVStats && !FCsvProfiler::Get()->IsCategoryEnabled(CSV_CATEGORY_INDEX(MallocBinned)))
+		{
+			return;
+		}
+
+		FCsvProfiler::Get()->EnableCategoryByIndex(CSV_CATEGORY_INDEX(MallocBinned), true);
+
+		static bool bFirstTime = true;
+		static FName FragmentationsStats[NumSmallPools];
+		static FName WasteStats[NumSmallPools];
+		static FName TotalMemStats[NumSmallPools];
+
+		for (int32 i = 0; i < NumSmallPools; i++)
+		{
+			if (bFirstTime)
+			{
+				TCHAR Name[64];
+				FCString::Sprintf(Name, TEXT("FragmentationBin%d"), Alloc.PoolIndexToBinSize(i));
+				FragmentationsStats[i] = FName(Name);
+
+				FCString::Sprintf(Name, TEXT("WasteBin%d"), Alloc.PoolIndexToBinSize(i));
+				WasteStats[i] = FName(Name);
+
+				FCString::Sprintf(Name, TEXT("TotalMemBin%d"), Alloc.PoolIndexToBinSize(i));
+				TotalMemStats[i] = FName(Name);
+			}
+
+			const float Fragmentation = 1.0f - (float)Alloc.SmallPoolTables[i].TotalUsedBins / (float)Alloc.SmallPoolTables[i].TotalAllocatedBins;
+			FCsvProfiler::RecordCustomStat(FragmentationsStats[i], CSV_CATEGORY_INDEX(MallocBinned), int(Fragmentation * 100.0f), ECsvCustomStatOp::Set);
+
+			const float TotalMem = (float)Alloc.SmallPoolTables[i].TotalAllocatedMem / 1024.0f / 1024.0f;
+			FCsvProfiler::RecordCustomStat(TotalMemStats[i], CSV_CATEGORY_INDEX(MallocBinned), TotalMem, ECsvCustomStatOp::Set);
+
+			FCsvProfiler::RecordCustomStat(WasteStats[i], CSV_CATEGORY_INDEX(MallocBinned), TotalMem * Fragmentation, ECsvCustomStatOp::Set);
+		}
+
+		CSV_CUSTOM_STAT(MallocBinned, RequestedSmallPoolMemoryMB,		(float)GetTotalAllocatedSmallPoolMemory() / (1024.0f * 1024.0f), ECsvCustomStatOp::Set);
+		CSV_CUSTOM_STAT(MallocBinned, TotalAllocatedSmallPoolMemoryMB,   (float)AllocatedOSSmallPoolMemory.load(std::memory_order_relaxed) / (1024.0f * 1024.0f), ECsvCustomStatOp::Set);
+		CSV_CUSTOM_STAT(MallocBinned, RequestedLargeAllocsMemoryMB,		(float)AllocatedLargePoolMemory.load(std::memory_order_relaxed) / (1024.0f * 1024.0f), ECsvCustomStatOp::Set);
+		CSV_CUSTOM_STAT(MallocBinned, TotalAllocatedLargeAllocsMemoryMB, (float)AllocatedLargePoolMemoryWAlignment.load(std::memory_order_relaxed) / (1024.0f * 1024.0f), ECsvCustomStatOp::Set);
+
+		bFirstTime = false;
+#endif
+	}
+
 	void AllocateHashBuckets()
 	{
 		const uint64 MaxHashBuckets = PtrToPoolMapping.GetMaxHashBuckets();
 		const uint64 HashAllocSize = Align(MaxHashBuckets * sizeof(FPoolHashBucket), OsAllocationGranularity);
 		HashBuckets = (FPoolHashBucket*)AllocType::AllocateMetaDataMemory(HashAllocSize);
-#if UE_MBC_ALLOCATOR_STATS
-		HashMemory += HashAllocSize;
-#endif
+		UE_MBC_UPDATE_STATS(HashMemory += HashAllocSize);
 		verify(HashBuckets);
 
 		DefaultConstructItems<FPoolHashBucket>(HashBuckets, MaxHashBuckets);
@@ -926,8 +990,6 @@ private:
 		//NOALLOC_SCOPE_CYCLE_COUNTER(STAT_FMallocBinned_UnregisterThreadFreeBlockLists);
 		FScopeLock Lock(&GetFreeBlockListsRegistrationMutex());
 		GetRegisteredFreeBlockLists().Remove(FreeBlockLists);
-#if UE_MBC_ALLOCATOR_STATS
-		ConsolidatedMemory.fetch_add(FreeBlockLists->AllocatedMemory, std::memory_order_relaxed);
-#endif
+		UE_MBC_UPDATE_STATS(ConsolidatedMemory.fetch_add(FreeBlockLists->AllocatedMemory, std::memory_order_relaxed));
 	}
 };
\ No newline at end of file
