Resolved compilation errors when UE_MBC_ALLOCATOR_STATS is not defined - in shipping or Editor builds

diff --git a/Engine/Source/Runtime/Core/Private/HAL/MallocBinnedCommon.cpp b/Engine/Source/Runtime/Core/Private/HAL/MallocBinnedCommon.cpp
--- a/Engine/Source/Runtime/Core/Private/HAL/MallocBinnedCommon.cpp
+++ b/Engine/Source/Runtime/Core/Private/HAL/MallocBinnedCommon.cpp
@@ -596,11 +596,9 @@ uint32 FMallocBinnedCommonBase::BinnedTlsSlot = FPlatformTLS::InvalidTlsSlot;
 		TEXT("MallocBinned.EnableCSVStats"),
 		GMallocBinnedEnableCSVStats,
 		TEXT("Whether or not to enable extended CSV stats with fragmentation stats (disabled by default)."));
-#endif
 
 void FMallocBinnedCommonBase::GetAllocatorStatsInternal(FGenericMemoryStats& OutStats, int64 TotalAllocatedSmallPoolMemory)
 {
-#if UE_MBC_ALLOCATOR_STATS
 	const int64 LocalAllocatedOSSmallPoolMemory = AllocatedOSSmallPoolMemory.load(std::memory_order_relaxed);
 	const int64 LocalAllocatedLargePoolMemory = AllocatedLargePoolMemory.load(std::memory_order_relaxed);
 	const int64 LocalAllocatedLargePoolMemoryWAlignment = AllocatedLargePoolMemoryWAlignment.load(std::memory_order_relaxed);
@@ -615,9 +613,9 @@ void FMallocBinnedCommonBase::GetAllocatorStatsInternal(FGenericMemoryStats& Out
 
 	OutStats.Add(TEXT("TotalAllocated"), TotalAllocated);
 	OutStats.Add(TEXT("TotalOSAllocated"), TotalOSAllocated);
-#endif
 
 	FMalloc::GetAllocatorStats(OutStats);
 }
+#endif	//~UE_MBC_ALLOCATOR_STATS
 
 PRAGMA_RESTORE_UNSAFE_TYPECAST_WARNINGS
diff --git a/Engine/Source/Runtime/Core/Public/HAL/MallocBinnedCommon.h b/Engine/Source/Runtime/Core/Public/HAL/MallocBinnedCommon.h
--- a/Engine/Source/Runtime/Core/Public/HAL/MallocBinnedCommon.h
+++ b/Engine/Source/Runtime/Core/Public/HAL/MallocBinnedCommon.h
@@ -390,6 +390,8 @@ protected:
 
 	static CORE_API uint32 BinnedTlsSlot;
 
+	std::atomic<uint64> MemoryTrimEpoch{ 0 };
+
 #if UE_MBC_ALLOCATOR_STATS
 	static std::atomic<int64> TLSMemory;
 	static std::atomic<int64> ConsolidatedMemory;
@@ -400,17 +402,15 @@ protected:
 
 	static int64 PoolInfoMemory;
 	static int64 HashMemory;
-#endif
 
-	std::atomic<uint64> MemoryTrimEpoch{ 0 };
+	void GetAllocatorStatsInternal(FGenericMemoryStats& OutStats, int64 TotalAllocatedSmallPoolMemory);
+#endif
 
 	static CA_NO_RETURN void OutOfMemory(uint64 Size, uint32 Alignment = 0)
 	{
 		// this is expected not to return
 		FPlatformMemory::OnOutOfMemory(Size, Alignment);
 	}
-
-	void GetAllocatorStatsInternal(FGenericMemoryStats& OutStats, int64 TotalAllocatedSmallPoolMemory);
 };
 
 template <class AllocType, int NumSmallPools, int MaxSmallPoolSize>
@@ -457,7 +457,9 @@ class TMallocBinnedCommon : public FMallocBinnedCommonBase
 public:
 	CORE_API virtual void GetAllocatorStats(FGenericMemoryStats& OutStats) override
 	{
+#if UE_MBC_ALLOCATOR_STATS
 		GetAllocatorStatsInternal(OutStats, GetTotalAllocatedSmallPoolMemory());
+#endif
 	}
 
 protected:
