FPersistentStorageCategory
- Add TotalOptionalSize to category stats
- Removed unused bIncludeOptional parameters

diff --git a/Engine/Source/Runtime/Core/Public/HAL/IPlatformFileManagedStorageWrapper.h b/Engine/Source/Runtime/Core/Public/HAL/IPlatformFileManagedStorageWrapper.h
--- a/Engine/Source/Runtime/Core/Public/HAL/IPlatformFileManagedStorageWrapper.h
+++ b/Engine/Source/Runtime/Core/Public/HAL/IPlatformFileManagedStorageWrapper.h
@@ -161,16 +161,9 @@ public:
 		return CategoryName;
 	}
 
-	int64 GetCategoryQuota( bool bIncludeOptional = true ) const
+	int64 GetCategoryQuota() const
 	{
-		if (bIncludeOptional || StorageQuota < 0)
-		{
-			return StorageQuota;
-		}
-		else
-		{
-			return StorageQuota - OptionalStorageQuota;
-		}
+		return StorageQuota;
 	}
 
 	int64 GetCategoryOptionalQuota() const
@@ -183,17 +176,10 @@ public:
 		return UsedQuota;
 	}
 
-	int64 GetAvailableSize(bool bIncludeOptionalStorage = true) const
+	int64 GetAvailableSize() const
 	{
 		int64 ActualStorageQuota = StorageQuota;
-		if (ActualStorageQuota >= 0)
-		{
-			if (!bIncludeOptionalStorage)
-			{
-				ActualStorageQuota -= OptionalStorageQuota;
-			}
-		}
-		else
+		if (ActualStorageQuota < 0)
 		{
 			ActualStorageQuota = MAX_int64;
 		}
@@ -275,6 +261,7 @@ public:
 		FString CategoryName;
 		int64 UsedSize = 0;
 		int64 TotalSize = -1;
+		int64 TotalOptionalSize = 0;
 		TMap<FString, int64> FileSizes;
 		TArray<FString> Directories;
 	};
@@ -283,7 +270,7 @@ public:
 	CategoryStat GetStat() const
 	{
 		FReadScopeLock ScopeLock(FileSizesLock);
-		return CategoryStat{ CategoryName, UsedQuota, StorageQuota, FileSizes, Directories };
+		return CategoryStat{ CategoryName, UsedQuota, StorageQuota, OptionalStorageQuota, FileSizes, Directories };
 	}
 
 	FManagedStorageFileLockRegistry& GetLockRegistry() { return FileLockRegistry; }
