Added extra info to assert in FPolyglotTextData constructor

#rnx

diff --git a/Engine/Source/Runtime/Core/Private/Internationalization/PolyglotTextData.cpp b/Engine/Source/Runtime/Core/Private/Internationalization/PolyglotTextData.cpp
--- a/Engine/Source/Runtime/Core/Private/Internationalization/PolyglotTextData.cpp
+++ b/Engine/Source/Runtime/Core/Private/Internationalization/PolyglotTextData.cpp
@@ -17,8 +17,8 @@ FPolyglotTextData::FPolyglotTextData(const ELocalizedTextSourceCategory& InCateg
 	, Key(InKey)
 	, NativeString(InNativeString)
 {
-	checkf(!Key.IsEmpty(), TEXT("Polyglot data cannot have an empty key!"));
-	checkf(!NativeString.IsEmpty(), TEXT("Polyglot data cannot have an empty native string!"));
+	checkf(!Key.IsEmpty(), TEXT("Polyglot data cannot have an empty key! (Namespace: '%s', Key: '%s', Native String: '%s')"), *Namespace, *Key, *NativeString);
+	checkf(!NativeString.IsEmpty(), TEXT("Polyglot data cannot have an empty native string! (Namespace: '%s', Key: '%s', Native String: '%s')"), *Namespace, *Key, *NativeString);
 }
 
 bool FPolyglotTextData::IsValid(FText* OutFailureReason) const
