Make FTransactionallySafeMutex 1 byte again.

diff --git a/Engine/Source/Programs/AutoRTFMTests/Private/UECoreTests.cpp b/Engine/Source/Programs/AutoRTFMTests/Private/UECoreTests.cpp
--- a/Engine/Source/Programs/AutoRTFMTests/Private/UECoreTests.cpp
+++ b/Engine/Source/Programs/AutoRTFMTests/Private/UECoreTests.cpp
@@ -28,6 +28,8 @@
 
 DEFINE_LOG_CATEGORY_STATIC(LogAutoRTFM_UECoreTests, Log, All)
 
+#include <thread>
+
 TEST_CASE("UECore.FDelegateHandle")
 {
 	FDelegateHandle Handle;
@@ -1747,21 +1749,17 @@ TEST_CASE("UECore.TransactionallySafeMutex")
 	SECTION("IsLocked Inside Transaction")
 	{
 		FTransactionallySafeMutex Mutex;
-		bool BeforeLock, DuringLock, AfterLock;
+		REQUIRE(!Mutex.IsLocked());
 
-		AutoRTFM::ETransactionResult Result = AutoRTFM::Transact([&]
-			{
-				BeforeLock = Mutex.IsLocked();
-				Mutex.Lock();
-				DuringLock = Mutex.IsLocked();
-				Mutex.Unlock();
-				AfterLock = Mutex.IsLocked();
-			});
+		AutoRTFM::Testing::Commit([&]
+		{
+			REQUIRE(!Mutex.IsLocked());
+			Mutex.Lock();
+			REQUIRE(Mutex.IsLocked());
+			Mutex.Unlock();
+			REQUIRE(!Mutex.IsLocked());
+		});
 
-		REQUIRE(AutoRTFM::ETransactionResult::Committed == Result);
-		REQUIRE(!BeforeLock);
-		REQUIRE(DuringLock);
-		REQUIRE(AfterLock);  // A TransactionallySafeMutex remains locked until the transaction ends.
 		REQUIRE(!Mutex.IsLocked());
 	}
 
@@ -1883,6 +1881,210 @@ TEST_CASE("UECore.TransactionallySafeMutex")
 
 		REQUIRE(42 == Lambda());
 	}
+
+	SECTION("Delete Heap Allocated Mutex")
+	{
+		FTransactionallySafeMutex* const Mutex = new FTransactionallySafeMutex();
+
+		AutoRTFM::Testing::Abort([&]
+		{
+			Mutex->Lock();
+			delete Mutex;
+			AutoRTFM::AbortTransaction();
+		});
+
+		REQUIRE(!Mutex->IsLocked());
+
+		AutoRTFM::Testing::Commit([&]
+		{
+			Mutex->Lock();
+			delete Mutex;
+		});
+	}
+
+	SECTION("Lock Within, Unlock Outside")
+	{
+		FTransactionallySafeMutex Mutex;
+
+		AutoRTFM::Testing::Commit([&]
+		{
+			Mutex.Lock();
+		});
+
+		REQUIRE(Mutex.IsLocked());
+		Mutex.Unlock();
+		REQUIRE(!Mutex.IsLocked());
+	}
+
+	SECTION("Lock Outside, Unlock Within")
+	{
+		FTransactionallySafeMutex Mutex = FTransactionallySafeMutex(UE::FAcquireLock());
+
+		AutoRTFM::Testing::Commit([&]
+		{
+			REQUIRE(Mutex.IsLocked());
+			Mutex.Unlock();
+		});
+
+		REQUIRE(!Mutex.IsLocked());
+	}
+
+	SECTION("Lock, Unlock, Lock")
+	{
+		FTransactionallySafeMutex Mutex;
+
+		AutoRTFM::Testing::Commit([&]
+		{
+			Mutex.Lock();
+			Mutex.Unlock();
+			Mutex.Lock();
+		});
+
+		REQUIRE(Mutex.IsLocked());
+	}
+
+	SECTION("Lock Outside, Unlock Within")
+	{
+		FTransactionallySafeMutex Mutex = FTransactionallySafeMutex(UE::FAcquireLock());
+
+		AutoRTFM::Testing::Commit([&]
+		{
+			REQUIRE(Mutex.IsLocked());
+			Mutex.Unlock();
+		});
+
+		REQUIRE(!Mutex.IsLocked());
+	}
+
+	SECTION("Lock Outside, Unlock & Lock Within")
+	{
+		FTransactionallySafeMutex Mutex = FTransactionallySafeMutex(UE::FAcquireLock());
+
+		AutoRTFM::Testing::Commit([&]
+		{
+			REQUIRE(Mutex.IsLocked());
+			Mutex.Unlock();
+			Mutex.Lock();
+		});
+
+		REQUIRE(Mutex.IsLocked());
+	}
+
+	SECTION("Commit(Lock, Commit(Unlock))")
+	{
+		FTransactionallySafeMutex Mutex;
+
+		AutoRTFM::Testing::Commit([&]
+		{
+			Mutex.Lock();
+
+			AutoRTFM::Testing::Commit([&]
+			{
+				Mutex.Unlock();
+			});
+
+			REQUIRE(!Mutex.IsLocked());
+		});
+
+		REQUIRE(!Mutex.IsLocked());
+	}
+
+	SECTION("Abort(Lock, Commit(Unlock, Lock))")
+	{
+		FTransactionallySafeMutex Mutex;
+
+		AutoRTFM::Testing::Abort([&]
+		{
+			Mutex.Lock();
+
+			AutoRTFM::Testing::Commit([&]
+			{
+				Mutex.Unlock();
+				Mutex.Lock();
+			});
+
+			REQUIRE(Mutex.IsLocked());
+
+			AutoRTFM::AbortTransaction();
+		});
+
+		REQUIRE(!Mutex.IsLocked());
+	}
+
+	SECTION("Contended Lock")
+	{
+		struct RAIIeroo final
+		{
+			RAIIeroo()
+			{
+				State = AutoRTFM::ForTheRuntime::GetRetryTransaction();
+				AutoRTFM::ForTheRuntime::SetRetryTransaction(AutoRTFM::ForTheRuntime::EAutoRTFMRetryTransactionState::NoRetry);
+			}
+
+			~RAIIeroo()
+			{
+				AutoRTFM::ForTheRuntime::SetRetryTransaction(State);
+			}
+
+		private:
+			AutoRTFM::ForTheRuntime::EAutoRTFMRetryTransactionState State;
+		} _;
+
+		SECTION("Contender Parks Transaction")
+		{
+			FTransactionallySafeMutex Mutex;
+			std::atomic_uint Orderer = 0;
+
+			std::thread Contender([&]
+			{
+				REQUIRE(0 == Orderer);
+				Mutex.Lock();
+				Orderer += 1; // unblock main thread
+				while (1 == Orderer) {} // wait on main thread
+				std::this_thread::sleep_for(std::chrono::milliseconds(100));
+				Mutex.Unlock();
+			});
+
+			while (0 == Orderer) {} // wait on contender
+
+			AutoRTFM::Testing::Commit([&]
+			{
+				REQUIRE(Mutex.IsLocked());
+				AutoRTFM::Open([&] { Orderer += 1; }); // unblock contender
+				Mutex.Lock();
+			});
+
+			REQUIRE(Mutex.IsLocked());
+			Mutex.Unlock();
+			Contender.join();
+		}
+
+		SECTION("Transaction Parks Contender")
+		{
+			FTransactionallySafeMutex Mutex;
+			std::atomic_uint Orderer = 0;
+
+			std::thread Contender([&]
+			{
+				while (0 == Orderer) {} // wait on main thread
+				Mutex.Lock();
+			});
+
+			AutoRTFM::Testing::Commit([&]
+			{
+				Mutex.Lock();
+				AutoRTFM::Open([&]
+				{
+					Orderer += 1; // unblock contender
+					std::this_thread::sleep_for(std::chrono::milliseconds(100));
+				});
+				Mutex.Unlock();
+			});
+
+			Contender.join();
+			Mutex.Unlock();
+		}
+	}
 }
 
 TEST_CASE("UECore.FTextFormatPatternDefinition")
diff --git a/Engine/Source/Runtime/Core/Public/Misc/TransactionallySafeMutex.h b/Engine/Source/Runtime/Core/Public/Misc/TransactionallySafeMutex.h
--- a/Engine/Source/Runtime/Core/Public/Misc/TransactionallySafeMutex.h
+++ b/Engine/Source/Runtime/Core/Public/Misc/TransactionallySafeMutex.h
@@ -2,166 +2,255 @@
 
 #pragma once
 
-#include "Async/Mutex.h"
+#include "Async/ExternalMutex.h"
 #include "AutoRTFM/AutoRTFM.h"
-#include "Templates/SharedPointer.h"
+#include <atomic>
 
-#if UE_AUTORTFM
-
-// this was initially copied from wrapping critical section in a transactionally safe way
+namespace UE
+{
 
 // A transactionally safe mutex that works in the following novel ways:
 // - In the open (non-transactional):
 //   - Take the lock like before. Simple!
 //   - Free the lock like before too.
 // - In the closed (transactional):
-//   - During locking we query `TransactionalLockCount`:
-//     - 0 means we haven't taken the lock within our transaction nest and need to acquire the lock.
-//     - Otherwise we already have the lock (and are preventing non-transactional code seeing any
-//       modifications we've made while holding the lock), so just bump `TransactionalLockCount`.
-//     - We also register an on-abort handler to release the lock should we abort (but we need to
-//       query `TransactionalLockCount` even there because we could be aborting an inner transaction
-//       and the parent transaction still wants to have the lock held!).
-//   - During unlocking we defer doing the unlock until the transaction commits.
+//   - We flip the internal transactional state between two bits `IsLockedWithinTransactionFlag` and
+// 	   `IsUnlockedWithinTransactionFlag`. These let us keep track of what the last action taken with
+// 	   the mutex was within the transaction.
 //
 // Thus with this approach we will hold this lock for the *entirety* of the transactional nest should
 // we take the lock during the transaction, thus preventing non-transactional code from seeing any
 // modifications we should make.
-struct FTransactionallySafeMutexDefinition final
+class FTransactionallySafeMutex final
 {
-	// Always open because the constructor arguments will create the underlying critical section.
-	UE_AUTORTFM_ALWAYS_OPEN
-	FTransactionallySafeMutexDefinition() : State(MakeShared<FState>())
+public:
+	constexpr FTransactionallySafeMutex() = default;
+
+	// Construct in a locked state. Avoids an expensive compare-and-swap at creation time.
+	inline explicit FTransactionallySafeMutex(UE::FAcquireLock Tag)
 	{
-		if (AutoRTFM::IsTransactional())
-		{
-			const AutoRTFM::EContextStatus Status = AutoRTFM::Close([this]
-				{
-					AutoRTFM::PushOnAbortHandler(this, [this]
-						{
-							this->~FTransactionallySafeMutexDefinition();
-						});
-				});
+		FExternalMutex(State, Tag);
+	}
 
-			ensure(AutoRTFM::EContextStatus::OnTrack == Status);
+	inline ~FTransactionallySafeMutex()
+	{
+		if (AutoRTFM::IsOnCurrentTransactionStack(this))
+		{
+			// We'll have either 0 (we've destructed a mutex but never locked it), or 1 of each
+			// on-abort and on-commit handlers.
+			AutoRTFM::PopAllOnAbortHandlers(this);
+			AutoRTFM::PopAllOnCommitHandlers(this);
 		}
 	}
 
-	~FTransactionallySafeMutexDefinition()
+	UE_NONCOPYABLE(FTransactionallySafeMutex);
+
+	inline bool IsLocked() const
 	{
 		if (AutoRTFM::IsTransactional())
 		{
-			const AutoRTFM::EContextStatus Status = AutoRTFM::Close([this]
+			return AutoRTFM::Open([this]
+			{
+				if (FExternalMutex(State).IsLocked())
 				{
-					AutoRTFM::PopOnAbortHandler(this);
-
-					// We explicitly copy the state here for the case that `this` was stack
-					// allocated and has already died before the on-commit is hit.
-					AutoRTFM::OnCommit([State = this->State]
-						{
-							ensure(0 == State->TransactionalLockCount);
-						});
-				});
+					return !IsTransactionallyUnlocked();
+				}
 
-			ensure(AutoRTFM::EContextStatus::OnTrack == Status);
+				return false;
+			});
+		}
+		else
+		{
+			return FExternalMutex(State).IsLocked();
 		}
-
-		// As the State was constructed in the open, it must be released in the open.
-		AutoRTFM::Open([&] { State = nullptr; });
 	}
-
-	void Lock()
+	
+	inline bool TryLock()
 	{
-		if (AutoRTFM::IsTransactional() || AutoRTFM::IsCommittingOrAborting())
+		if (AutoRTFM::IsTransactional())
 		{
-			// Memory validation disabled as this is hot code.
-			UE_AUTORTFM_OPEN_NO_VALIDATION
+			return AutoRTFM::Open([this]
+			{
+				if (!DoesTransactionHoldLock())
 				{
-					// The transactional system which can increment TransactionalLockCount
-					// is always single-threaded, thus this is safe to check without atomicity.
-					if (0 == State->TransactionalLockCount)
+					if (!FExternalMutex(State).TryLock())
 					{
-						State->Mutex.Lock();
+						return false;
 					}
 
-					State->TransactionalLockCount += 1;
-				};
+					// If we get here, it means we've taken the lock for the *first* time within the
+					// transaction. We need to thus register the on-abort and on-commit handlers to
+					// correctly undo the transactionality of our locking.
+					const AutoRTFM::EContextStatus Status = AutoRTFM::Close([this]
+					{
+						DoOnAbortHandlerInClosed();
+						DoOnCommitHandlerInClosed();
+					});
 
-			// We explicitly copy the state here for the case that `this` was stack
-			// allocated and has already died before the on-abort is hit.
-			AutoRTFM::OnAbort([State = this->State]
-				{
-					ensure(0 != State->TransactionalLockCount);
+					check(Status == AutoRTFM::EContextStatus::OnTrack);
+				}
 
-					State->TransactionalLockCount -= 1;
+				RecordPreviousStateAndSetTransactionalLockMode(ETransactionalLockMode::TransactionallyLocked);
 
-					if (0 == State->TransactionalLockCount)
-					{
-						State->Mutex.Unlock();
-					}
-				});
+				return true;
+			});
 		}
 		else
 		{
-			State->Mutex.Lock();
-			ensure(0 == State->TransactionalLockCount);
+			const bool bResult = FExternalMutex(State).TryLock();
+			ensure(!bResult || !DoesTransactionHoldLock());
+			return bResult;
 		}
 	}
 
-	void Unlock()
+	inline void Lock()
 	{
-		if (AutoRTFM::IsTransactional() || AutoRTFM::IsCommittingOrAborting())
+		if (AutoRTFM::IsTransactional())
 		{
-			// We explicitly copy the state here for the case that `this` was stack
-			// allocated and has already died before the on-commit is hit.
-			AutoRTFM::OnCommit([State = this->State]
+			AutoRTFM::Open([this]
+			{
+				if (!DoesTransactionHoldLock())
 				{
-					ensure(0 != State->TransactionalLockCount);
-
-					State->TransactionalLockCount -= 1;
+					FExternalMutex(State).Lock();
 
-					if (0 == State->TransactionalLockCount)
+					const AutoRTFM::EContextStatus Status = AutoRTFM::Close([this]
 					{
-						State->Mutex.Unlock();
-					}
-				});
+						DoOnAbortHandlerInClosed();
+						DoOnCommitHandlerInClosed();
+					});
+
+					check(Status == AutoRTFM::EContextStatus::OnTrack);
+				}
+
+				RecordPreviousStateAndSetTransactionalLockMode(ETransactionalLockMode::TransactionallyLocked);
+			});
 		}
 		else
 		{
-			ensure(0 == State->TransactionalLockCount);
-			State->Mutex.Unlock();
+			FExternalMutex(State).Lock();
+			ensure(!DoesTransactionHoldLock());
 		}
 	}
 
-	/** 
-	 * This method may give surprising results and should be used with caution!
-	 * 
-	 * - You cannot safely use this function to determine whether Lock() will block. 
-	 *   You may be in race with another thread which is also trying to lock the mutex.
-	 * 
-	 * - If a Mutex is locked by the AutoRTFM thread, it will not fully release the mutex 
-	 *   until the transaction fully succeeds or rolls back. Because of this, IsLocked may 
-	 *   return true even after the mutex has been fully unlocked.
-	 */
-	bool IsLocked() const
+	inline void Unlock()
 	{
-		return State->Mutex.IsLocked();
+		if (AutoRTFM::IsTransactional())
+		{
+			AutoRTFM::Open([this]
+			{
+				if (!DoesTransactionHoldLock())
+				{
+					const AutoRTFM::EContextStatus Status = AutoRTFM::Close([this]
+					{
+						DoOnCommitHandlerInClosed();
+					});
+
+					check(Status == AutoRTFM::EContextStatus::OnTrack);
+				}
+
+				RecordPreviousStateAndSetTransactionalLockMode(ETransactionalLockMode::TransactionallyUnlocked);
+			});
+		}
+		else
+		{
+			ensure(!DoesTransactionHoldLock());
+			FExternalMutex(State).Unlock();
+		}
 	}
 
 private:
-	UE_NONCOPYABLE(FTransactionallySafeMutexDefinition)
+	// Need `mutable` because `FExternalMutex` only works non-const!
+	mutable std::atomic<uint8> State = 0;
+
+	// Since we are going to be writing non-atomically to `State` in the write-log, prior to us
+	// doing an atomic write to actually undo the lock, we need to be sure that the atomic will
+	// always be lock-free and take no bits beyond the integer backing type!
+	static_assert(decltype(State)::is_always_lock_free && (sizeof(decltype(State)) == sizeof(uint8)));
+	
+	// These are the flags that `FExternalMutex` already reserves for the lock!
+	static constexpr uint8 IsLockedFlag = 1 << 0;
+	static constexpr uint8 HasWaitingThreadsFlag = 1 << 1;
+
+	// Need 2 extra bits ontop of the 2 that `FExternalMutex` uses, to record what
+	// the last transactional action was within the transaction.
+	static constexpr uint8 IsLockedWithinTransactionFlag = 1 << 2;
+	static constexpr uint8 IsUnlockedWithinTransactionFlag = 1 << 3;
+
+	inline bool IsTransactionallyUnlocked() const
+	{
+		const uint8 LoadedState = State.load(std::memory_order_consume);
+		return (LoadedState & IsUnlockedWithinTransactionFlag);
+	}
+
+	inline bool DoesTransactionHoldLock() const
+	{
+		const uint8 LoadedState = State.load(std::memory_order_consume);
+		const bool bTransactionallyLocked = (LoadedState & (IsLockedWithinTransactionFlag | IsUnlockedWithinTransactionFlag));
+		ensure(!bTransactionallyLocked || (LoadedState & IsLockedFlag));
+		return bTransactionallyLocked;
+	}
+
+	inline void DoOnAbortHandlerInClosed()
+	{
+		check(AutoRTFM::IsClosed());
+
+		AutoRTFM::PushOnAbortHandler(this, [this]
+		{
+			State.fetch_and(~(IsLockedWithinTransactionFlag | IsUnlockedWithinTransactionFlag), std::memory_order_release);
+			FExternalMutex(State).Unlock();
+		});
+	}
 
-	struct FState final
+	inline void DoOnCommitHandlerInClosed()
 	{
-		UE::FMutex Mutex;
-		uint32 TransactionalLockCount = 0;
+		check(AutoRTFM::IsClosed());
+
+		AutoRTFM::PushOnCommitHandler(this, [this]
+		{
+			// If we commit, we won't have undone any of our transactional writes to the
+			// lock. So just double check that we still hold the lock here (famous last
+			// words, but this is a temporary check just while we deploy the new mutex
+			// more widely. Once we've got a few runs around the block we can remove!).
+			ensure(DoesTransactionHoldLock());
+
+			const bool bIsTransactionallyUnlocked = IsTransactionallyUnlocked();
+
+			State.fetch_and(~(IsLockedWithinTransactionFlag | IsUnlockedWithinTransactionFlag), std::memory_order_release);
+
+			if (bIsTransactionallyUnlocked)
+			{
+				FExternalMutex(State).Unlock();
+			}
+		});
+	}
+
+	enum class ETransactionalLockMode : bool
+	{
+		TransactionallyLocked = false,
+		TransactionallyUnlocked = true,
 	};
 
-	TSharedPtr<FState> State;
+	inline void RecordPreviousStateAndSetTransactionalLockMode(ETransactionalLockMode LockMode)
+	{
+		// Record the state into the write log. If an abort happens we will non-atomically
+		// overwrite `State`, removing the transactional flags. This will be followed-up
+		// with an atomic write in the on-abort handler to actually do the unlock though,
+		// thus ensuring data integrity.
+		AutoRTFM::RecordOpenWriteNoMemoryValidation(&State);
+
+		if (LockMode == ETransactionalLockMode::TransactionallyLocked)
+		{
+			State |= IsLockedWithinTransactionFlag;
+			State &= ~IsUnlockedWithinTransactionFlag;
+		}
+		else
+		{
+			State |= IsUnlockedWithinTransactionFlag;
+			State &= ~IsLockedWithinTransactionFlag;
+		}
+	}
 };
 
-using FTransactionallySafeMutex = FTransactionallySafeMutexDefinition;
+} // UE
 
-#else
-using FTransactionallySafeMutex = UE::FMutex;
-#endif // UE_AUTORTFM
+using FTransactionallySafeMutex = UE::FTransactionallySafeMutex;
