Fix for warning:

error : implicit conversion from 'unsigned long long' to 'const unsigned int' changes value from 14695981039346656037 to 2216829733 [-Werror,-Wconstant-conversion]

diff --git a/Engine/Source/Runtime/Core/Public/Hash/Fnv.h b/Engine/Source/Runtime/Core/Public/Hash/Fnv.h
--- a/Engine/Source/Runtime/Core/Public/Hash/Fnv.h
+++ b/Engine/Source/Runtime/Core/Public/Hash/Fnv.h
@@ -22,8 +22,28 @@ namespace UE
 		static_assert(sizeof(CharType) <= 4, "FNV1a only works with characters up to 32 bits.");
 
 		// See https://en.wikipedia.org/wiki/Fowler%E2%80%93Noll%E2%80%93Vo_hash_function#FNV_hash_parameters
-		constexpr HashType Offset = sizeof(HashType) == 4 ? 0x811c9dc5 : 0xcbf29ce484222325;
-		constexpr HashType Prime = sizeof(HashType) == 4 ? 0x01000193 : 0x100000001b3;
+		constexpr HashType Offset = []() -> HashType
+		{
+			if constexpr (sizeof(HashType) == 4)
+			{
+				return 0x811c9dc5;
+			}
+			else
+			{
+				return 0xcbf29ce484222325;
+			}
+		}();
+		constexpr HashType Prime = []() -> HashType
+		{
+			if constexpr (sizeof(HashType) == 4)
+			{
+				return 0x01000193;
+			}
+			else
+			{
+				return 0x100000001b3;
+			}
+		}();
 
 		HashType Fnv = Offset;
 		for (CharType Char : String)
