Added TMap::Remove(Id), TMap::FindId(Key) and TMap::FindIdByHash.

diff --git a/Engine/Source/Runtime/Core/Public/Containers/Map.h b/Engine/Source/Runtime/Core/Public/Containers/Map.h
--- a/Engine/Source/Runtime/Core/Public/Containers/Map.h
+++ b/Engine/Source/Runtime/Core/Public/Containers/Map.h
@@ -507,6 +507,15 @@ public:
 		return NumRemovedPairs;
 	}
 
+	/**
+	 * Removes the element at the specified index. The caller has to ensure that the index is valid.
+	 * @param Id The index of the element to remove.
+	 */
+	FORCEINLINE void Remove(FSetElementId Id)
+	{
+		Pairs.Remove(Id);
+	}
+
 	/**
 	 * Find the key associated with the specified value.
 	 *
@@ -602,6 +611,24 @@ public:
 		return const_cast<TMapBase*>(this)->FindByHashChecked(KeyHash, Key);
 	}
 
+	/**
+	 * Finds the index of the first element that is assigned to the specified key.
+	 * The returned index is only valid until the map is changed.
+	 *
+	 * @param Key The key to search for.
+	 */
+	FORCEINLINE FSetElementId FindId(KeyInitType Key) const
+	{
+		return Pairs.FindId(Key);
+	}
+
+	/** See FindId() and class documentation section on ByHash() functions */
+	template<typename ComparableKey>
+	FORCEINLINE FSetElementId FindIdByHash(uint32 KeyHash, const ComparableKey& Key) const
+	{
+		return Pairs.FindIdByHash(KeyHash, Key);
+	}
+
 private:
 	[[nodiscard]] FORCEINLINE static uint32 HashKey(const KeyType& Key)
 	{
