FPaths::FindCorrectCase: Fix typo that was causing input argument to be ignored and return value to always be the corrected case of FPaths::GetProjectFilePath. This fixes the symptom in zenserver cooks that the wrong project id was sometimes used.
#rnx

diff --git a/Engine/Source/Runtime/Core/Private/Misc/Paths.cpp b/Engine/Source/Runtime/Core/Private/Misc/Paths.cpp
--- a/Engine/Source/Runtime/Core/Private/Misc/Paths.cpp
+++ b/Engine/Source/Runtime/Core/Private/Misc/Paths.cpp
@@ -880,10 +880,9 @@ FString FPaths::FindCorrectCase(const FString& Path)
 	// for workflows that use a junction at their root (eg: p4 gets confused about paths and operations fail).
 	// There is a way to get a case-accurate path on Windows without resolving directory junctions, but it is slow.
 	// We can use it here for this one-off situation without causing all uses of GetFilenameOnDisk to be slower.
-	FString ProjectFilePath = FPaths::GetProjectFilePath();
 	TStringBuilder<MAX_PATH> Builder;
 	FPathViews::IterateComponents(
-		ProjectFilePath,
+		Path,
 		[&Builder](FStringView CurrentPathComponent)
 		{
 			if (Builder.Len() != 0)
