Fix Apple Memory Probe non-unity compilation

diff --git a/Engine/Source/Runtime/Core/Private/ProfilingDebugging/Apple/AppleMemoryProbe.cpp b/Engine/Source/Runtime/Core/Private/ProfilingDebugging/Apple/AppleMemoryProbe.cpp
--- a/Engine/Source/Runtime/Core/Private/ProfilingDebugging/Apple/AppleMemoryProbe.cpp
+++ b/Engine/Source/Runtime/Core/Private/ProfilingDebugging/Apple/AppleMemoryProbe.cpp
@@ -1,6 +1,10 @@
 ﻿// Copyright Epic Games, Inc. All Rights Reserved.
 
 #include "ProfilingDebugging/MemoryTrace.h"
+#include "HAL/IConsoleManager.h"
+#include "HAL/Thread.h"
+#include "HAL/PlatformProcess.h"
+#include "HAL/PlatformAffinity.h"
 
 #if UE_MEMORY_TRACE_ENABLED
 
