Detect well-of-despair situations with the CSV profiler (where stats are coming in faster than we can process them, leading to a feedback loop of progressively slower updates). When this happens, we stall the gamethread to avoid crashing OOM.

This is controlled with the following cvars:
-csv.ProcessingThread.GtStallUpdateTimeThresholdMs : Specifies the max time a processing thread update can take before we consider stalling the game thread. Set to 0 to disable stalling
-csv.ProcessingThread.GtStallUpdateEscalationThreshold: Number of progressively slower updates before we stall the GT. Updates faster than UpdateTimeThresholdMs are ignored
#tests Tested via replayrun to make sure we're not currently hitting this

diff --git a/Engine/Source/Runtime/Core/Private/ProfilingDebugging/CsvProfiler.cpp b/Engine/Source/Runtime/Core/Private/ProfilingDebugging/CsvProfiler.cpp
--- a/Engine/Source/Runtime/Core/Private/ProfilingDebugging/CsvProfiler.cpp
+++ b/Engine/Source/Runtime/Core/Private/ProfilingDebugging/CsvProfiler.cpp
@@ -202,6 +202,21 @@ TAutoConsoleVariable<float> CVarCsvProcessingThreadTimeBetweenUpdates(
 	ECVF_Default
 );
 
+TAutoConsoleVariable<float> CVarCsvProcessingThreadGtStallUpdateTimeThresholdMs(
+	TEXT("csv.ProcessingThread.GtStallUpdateTimeThresholdMs"),
+	1000.0f,
+	TEXT("Specifies the max time a processing thread update can take before we consider stalling the game thread.\n")
+	TEXT("Set to 0 to disable stalling"),
+	ECVF_Default
+);
+
+TAutoConsoleVariable<int32> CVarCsvProcessingThreadGtStallUpdateEscalationThreshold(
+	TEXT("csv.ProcessingThread.GtStallUpdateEscalationThreshold"),
+	3,
+	TEXT("Number of progressively slower updates before we stall the GT. Updates faster than UpdateTimeThresholdMs are ignored"),
+	ECVF_Default
+);
+
 
 TAutoConsoleVariable<int32> CVarMaxPerThreadStatDataSlackKB(
 	TEXT("csv.MaxPerThreadStatDataSlackKB"),
@@ -2900,22 +2915,22 @@ uint64 FCsvStreamWriter::GetAllocatedSize() const
 class FCsvProfilerProcessingThread : public FRunnable
 {
 	FThreadSafeCounter StopCounter;
+	EThreadPriority Priority;
 
 public:
 	FCsvProfilerProcessingThread(FCsvProfiler& InCsvProfiler)
 		: CsvProfiler(InCsvProfiler)
 	{
+		Priority = TPri_Lowest;
+		uint64 AffinityMask = FPlatformAffinity::GetTaskGraphBackgroundTaskMask(); 
 #if CSV_PROFILER_ALLOW_DEBUG_FEATURES
 		if (FParse::Param(FCommandLine::Get(), TEXT("csvProfilerHighPriority")))
 		{
-			Thread = FForkProcessHelper::CreateForkableThread(this, TEXT("CSVProfiler"), 0, TPri_Highest, FPlatformAffinity::GetTaskGraphThreadMask());
+			Priority = TPri_Highest;
+			AffinityMask = FPlatformAffinity::GetTaskGraphThreadMask();
 		}
-		else
 #endif
-		{
-			Thread = FForkProcessHelper::CreateForkableThread(this, TEXT("CSVProfiler"), 0, TPri_Lowest, FPlatformAffinity::GetTaskGraphBackgroundTaskMask());
-		}
-	
+		Thread = FForkProcessHelper::CreateForkableThread(this, TEXT("CSVProfiler"), 0, Priority, AffinityMask);
 	}
 
 	virtual ~FCsvProfilerProcessingThread()
@@ -2948,6 +2963,13 @@ public:
 
 		LLM_SCOPE(ELLMTag::CsvProfiler);
 
+		double TotalProcessingTime = 0.0;
+		uint64 TotalStatEntriesProcessed = 0;
+		float PrevElapsedMs = 0.0f;
+		uint32 PrevNumStatEntriesProcessed = 0;
+		uint32 StallCounter = 0;
+		int32 SlowUpdateEscalationCount = 0;
+
 		while (StopCounter.GetValue() == 0)
 		{
 			const float TimeBetweenUpdatesMS = CVarCsvProcessingThreadTimeBetweenUpdates.GetValueOnAnyThread();
@@ -2958,7 +2980,11 @@ public:
 				FPlatformProcess::Sleep(5.0f / 1000.0f);
 				continue;
 			}
-			float ElapsedMS = CsvProfiler.ProcessStatData();
+			uint32 NumStatEntriesProcessed = 0;
+			float ElapsedMS = CsvProfiler.ProcessStatData(&NumStatEntriesProcessed);
+
+			TotalStatEntriesProcessed += (uint64)NumStatEntriesProcessed;
+			TotalProcessingTime += (double)ElapsedMS/1000.0;
 
 			if (GCsvProfilerIsWritingFile)
 			{
@@ -2966,9 +2992,57 @@ public:
 				CsvProfiler.FileWriteBlockingEvent->Trigger();
 			}
 
+			// If we've stalled the GT then check if we can resume
+			if (StallCounter > 0)
+			{
+				StallCounter--;
+				if (StallCounter == 0 || ElapsedMS < TimeBetweenUpdatesMS)
+				{
+					FPlatformProcess::SetThreadPriority(Priority);
+					StallGameThreadCs.Unlock();
+					StallCounter = 0;
+				}
+			}
+			else
+			{
+				// Detect well-of-despair (stats coming in faster than they can be processed)
+				const float StallThresholdMs = CVarCsvProcessingThreadGtStallUpdateTimeThresholdMs.GetValueOnAnyThread();
+
+				// If this thread is slower than the last and processed more stats, increase SlowUpdateEscalationCount
+				if (StallThresholdMs > 0.0f && PrevElapsedMs > StallThresholdMs && ElapsedMS > PrevElapsedMs && NumStatEntriesProcessed > PrevNumStatEntriesProcessed)
+				{
+					SlowUpdateEscalationCount++;
+					if (SlowUpdateEscalationCount >= CVarCsvProcessingThreadGtStallUpdateEscalationThreshold.GetValueOnAnyThread())
+					{
+						// Stall the GT for up to 2 updates to allow us to catch up
+						StallGameThreadCs.Lock();
+						StallCounter = 2;
+						double StatEntriesProcessedPerSecond = (double)TotalStatEntriesProcessed / TotalProcessingTime;
+						UE_LOG(LogCsvProfiler, Warning, TEXT("Stats coming in faster than we can process them! GT stalled until we can catch up!"));
+						UE_LOG(LogCsvProfiler, Warning, TEXT("Avg processing rate: %.0f stat entries per second (timestamps+custom stats)"), StatEntriesProcessedPerSecond);
+						UE_LOG(LogCsvProfiler, Warning, TEXT("Check CsvProfiler/* stats to see current rates. Run with -csvStatCounts to report per-stat counts"));
+						// Temporarily boost thread priority since the GT is waiting
+						FPlatformProcess::SetThreadPriority(TPri_AboveNormal);
+						SlowUpdateEscalationCount = 0;
+					}
+				}
+				else
+				{
+					SlowUpdateEscalationCount = 0;
+				}
+			}
+			
+			PrevNumStatEntriesProcessed = NumStatEntriesProcessed;
+			PrevElapsedMs = ElapsedMS;
+
 			float SleepTimeSeconds = FMath::Max(TimeBetweenUpdatesMS - ElapsedMS, 0.0f) / 1000.0f;
 			FPlatformProcess::Sleep(SleepTimeSeconds);
 		}
+		
+		if (StallCounter>0)
+		{
+			StallGameThreadCs.Unlock();
+		}
 
 		FMemory::ClearAndDisableTLSCachesOnCurrentThread();
 
@@ -2982,9 +3056,22 @@ public:
 
 	virtual void Exit() override { }
 
+	void StallGameThreadIfNeeded()
+	{
+		check(IsInGameThread());
+		if (StallGameThreadCs.TryLock())
+		{
+			StallGameThreadCs.Unlock();
+			return;
+		}
+		CSV_SCOPED_TIMING_STAT_EXCLUSIVE(CsvProfiler_Stall);
+		FScopeLock Lock(&StallGameThreadCs);
+	}
+
 private:
 	FRunnableThread* Thread;
 	FCsvProfiler& CsvProfiler;
+	FCriticalSection StallGameThreadCs;
 };
 
 void FCsvProfilerThreadDataProcessor::Process(FCsvProcessThreadDataStats& OutStats, int32& OutMinFrameNumberProcessed, TSharedPtr<FCsvThreadGroupStatProcessor> TaskWorkerThreadGroupStatProcessor)
@@ -3531,6 +3618,12 @@ void FCsvProfiler::EndFrame()
 	check(IsInGameThread());
 	if (GCsvProfilerIsCapturing)
 	{
+		// If the GT is stalled, wait
+		if (ProcessingThread)
+		{
+			ProcessingThread->StallGameThreadIfNeeded();
+		}		
+
 		OnCSVProfileEndFrameDelegate.Broadcast();
 
 		GCsvPersistentCustomStats.RecordStats();
@@ -4532,7 +4625,7 @@ bool FCsvProfiler::IsCapturing_Renderthread() const
 	return GCsvProfilerIsCapturingRT;
 }
 
-float FCsvProfiler::ProcessStatData()
+float FCsvProfiler::ProcessStatData(uint32* OutNumStatEntriesProcessed)
 {
 	check(IsInCsvProcessingThread());
 
@@ -4553,6 +4646,11 @@ float FCsvProfiler::ProcessStatData()
 		CSV_CUSTOM_STAT(CsvProfiler, NumCustomStatsProcessed, (int32)Stats.CustomStatCount, ECsvCustomStatOp::Accumulate);
 		CSV_CUSTOM_STAT(CsvProfiler, NumEventsProcessed, (int32)Stats.EventCount, ECsvCustomStatOp::Accumulate);
 		CSV_CUSTOM_STAT(CsvProfiler, ProcessCSVStats, ElapsedMS, ECsvCustomStatOp::Accumulate);
+
+		if (OutNumStatEntriesProcessed)
+		{
+			*OutNumStatEntriesProcessed = Stats.CustomStatCount + Stats.TimestampCount;
+		}		
 	}
 	return ElapsedMS;
 }
diff --git a/Engine/Source/Runtime/Core/Public/ProfilingDebugging/CsvProfiler.h b/Engine/Source/Runtime/Core/Public/ProfilingDebugging/CsvProfiler.h
--- a/Engine/Source/Runtime/Core/Public/ProfilingDebugging/CsvProfiler.h
+++ b/Engine/Source/Runtime/Core/Public/ProfilingDebugging/CsvProfiler.h
@@ -510,7 +510,7 @@ private:
 
 	void FinalizeCsvFile();
 
-	float ProcessStatData();
+	float ProcessStatData(uint32* OutNumStatEntriesProcessed = nullptr);
 
 	void BeginCaptureInternal(const FCsvCaptureCommand& CurrentCommand);
 
