- Added support for loading saved ini files when using binary config (mainly to fix loading the console history), cleaned up the code in the process!
- Fixed a check that was keeping binary config from working in Game config (was only working in Client, mistakenly)

diff --git a/Engine/Source/Runtime/Core/Private/Misc/ConfigCacheIni.cpp b/Engine/Source/Runtime/Core/Private/Misc/ConfigCacheIni.cpp
--- a/Engine/Source/Runtime/Core/Private/Misc/ConfigCacheIni.cpp
+++ b/Engine/Source/Runtime/Core/Private/Misc/ConfigCacheIni.cpp
@@ -5633,18 +5633,15 @@ bool FConfigCacheIni::CreateGConfigFromSaved(const TCHAR* Filename)
 
 	// fix up some things that weren't saved out
 	GConfig->Type = EConfigCacheType::DiskBacked;
+
+	// read in any needed generated/saved ini files (using the logic in FConfigContext to determine if it's needed to be loaded)
+	// and apply commandline overrides
+	FConfigContext Context = FConfigContext::FixupBranchAfterBinaryConfig();
 	for (FConfigBranch& Branch : GConfig->KnownFiles.Branches)
 	{
-		Branch.IniPath = FConfigCacheIni::GetDestIniFilename(*Branch.IniName.ToString(), nullptr, *FPaths::GeneratedConfigDir());
-
-#if ALLOW_INI_OVERRIDE_FROM_COMMANDLINE
-		// process any commandline overrides
-		FConfigFile::OverrideFromCommandline(&Branch.CommandLineOverrides, Branch.IniName.ToString());
-		// and push it into the current values
-		Branch.InMemoryFile.ApplyFile(&Branch.CommandLineOverrides);
-#endif
+		Context.Load(*Branch.IniName.ToString(), /*out*/ Branch.IniPath);
 	}
-	
+
 	MemoryReader << ExtraData.Data;
 
 	// now let the delegates pull their data out, after GConfig is set up
@@ -5764,7 +5761,7 @@ void FConfigCacheIni::InitializeConfigSystem()
 	
 	InitializeConfigRemap();
 	
-#if PLATFORM_SUPPORTS_BINARYCONFIG && PRELOAD_BINARY_CONFIG && !WITH_EDITOR && !WITH_SERVER_CODE
+#if PLATFORM_SUPPORTS_BINARYCONFIG && PRELOAD_BINARY_CONFIG && !WITH_EDITOR && WITH_CLIENT_CODE
 	// attempt to load from staged binary config data
 	const bool bCommandLineRequestsTextConfig = FParse::Param(FCommandLine::Get(), TEXT("textconfig")) // Explicitly requesting textconfig
 #if ALLOW_INI_OVERRIDE_FROM_COMMANDLINE
@@ -5774,11 +5771,7 @@ void FConfigCacheIni::InitializeConfigSystem()
 	if (!bCommandLineRequestsTextConfig &&
 		FConfigCacheIni::CreateGConfigFromSaved(nullptr))
 	{
-		FConfigContext Context = FConfigContext::ForceReloadIntoGConfig();
-		// Force reload GameUserSettings because they may be saved to disk on consoles/similar platforms
-		// So the safest thing to do is to re-read the file after binary configs load.
-		Context.Load(TEXT("GameUserSettings"), GGameUserSettingsIni);
-
+		FConfigContext Context = FConfigContext::ReadIntoGConfig();
 		LoadRemainingConfigFiles(Context);
 	}
 	else
diff --git a/Engine/Source/Runtime/Core/Private/Misc/ConfigContext.cpp b/Engine/Source/Runtime/Core/Private/Misc/ConfigContext.cpp
--- a/Engine/Source/Runtime/Core/Private/Misc/ConfigContext.cpp
+++ b/Engine/Source/Runtime/Core/Private/Misc/ConfigContext.cpp
@@ -561,14 +561,24 @@ bool FConfigContext::PerformLoad()
 		return true;
 	}
 	
-	// generate the whole standard ini hierarchy
-	AddStaticLayersToHierarchy(nullptr);
+	if (!bIsFixingUpAfterBinaryConfig)
+	{
+		// generate the whole standard ini hierarchy
+		AddStaticLayersToHierarchy(nullptr);
+	}
 
 	// now generate and make sure it's up to date (using IniName as a Base for an ini filename)
 	// @todo This bNeedsWrite afaict is always true even if it loaded a completely valid generated/final .ini, and the write below will
 	// just write out the exact same thing it read in!
 	bool bGeneratedFile = GenerateDestIniFile();
 
+	// we are done here!
+	if (bIsFixingUpAfterBinaryConfig)
+	{
+		return true;
+	}
+
+
 	FinalFile.Name = FName(*BaseIniName);
 	FinalFile.PlatformName = Platform;
 	FinalFile.bHasPlatformName = true;
@@ -1074,17 +1084,20 @@ bool FConfigContext::LoadIniFileHierarchy()
  */
 bool FConfigContext::GenerateDestIniFile()
 {
-	// reset the file to empty
-	Branch->InMemoryFile.Cleanup();
-	Branch->CombinedStaticLayers.Cleanup();
-	Branch->SavedLayer.Empty();
-	Branch->CommandLineOverrides.Empty();
-	Branch->StaticLayers.Empty();
-	Branch->DynamicLayers.Empty();
-	Branch->FinalCombinedLayers.Empty();
-
-	// read the static files into the branch
-	LoadIniFileHierarchy();
+	if (!bIsFixingUpAfterBinaryConfig)
+	{
+		// reset the file to empty
+		Branch->InMemoryFile.Cleanup();
+		Branch->CombinedStaticLayers.Cleanup();
+		Branch->SavedLayer.Empty();
+		Branch->CommandLineOverrides.Empty();
+		Branch->StaticLayers.Empty();
+		Branch->DynamicLayers.Empty();
+		Branch->FinalCombinedLayers.Empty();
+
+		// read the static files into the branch
+		LoadIniFileHierarchy();
+	}
 	
 #if !IS_PROGRAM
 	// Don't try to load any generated files from disk in cooked builds. We will always use the re-generated INIs.
diff --git a/Engine/Source/Runtime/Core/Public/Misc/ConfigContext.h b/Engine/Source/Runtime/Core/Public/Misc/ConfigContext.h
--- a/Engine/Source/Runtime/Core/Public/Misc/ConfigContext.h
+++ b/Engine/Source/Runtime/Core/Public/Misc/ConfigContext.h
@@ -151,6 +151,20 @@ public:
 		return Context;
 	}
 
+	/**
+	 * Create a context to read only the saved/generated file (eg. <Project>/Saved/Windows/Input.ini) and commandline options after loading
+	 * a binary config, that was of course made without saved files and commandline options
+	 */
+	static FConfigContext FixupBranchAfterBinaryConfig()
+	{
+		FConfigContext Context(GConfig, true, FString());
+		Context.bIsFixingUpAfterBinaryConfig = true;
+		Context.bForceReload = true;
+		return Context;
+	}
+
+
+
 	/**
 	 * Call to make before attempting parallel config init
 	 */
@@ -235,6 +249,7 @@ public:
 	// GameFeaturePlugins have WindowsFooGame.ini, not WindowsGame.ini
 	bool bIncludeTagNameInBranchName = false;
 	bool bIsMakingBinaryConfig = false;
+	bool bIsFixingUpAfterBinaryConfig = false;
 
 	DynamicLayerPriority PluginModificationPriority;
 
