[Core] Added get-with-default overload to TVariant

diff --git a/Engine/Source/Runtime/Core/Public/Misc/TVariant.h b/Engine/Source/Runtime/Core/Public/Misc/TVariant.h
--- a/Engine/Source/Runtime/Core/Public/Misc/TVariant.h
+++ b/Engine/Source/Runtime/Core/Public/Misc/TVariant.h
@@ -148,6 +148,13 @@ public:
 		return const_cast<TVariant*>(this)->template Get<U>();
 	}
 
+	/** Get a reference to the held value if set, otherwise the DefaultValue */
+	template <typename U>
+	const U& Get(const TIdentity_T<U>& DefaultValue UE_LIFETIMEBOUND) const UE_LIFETIMEBOUND
+	{
+		return IsType<U>() ? Get<U>() : DefaultValue;
+	}
+
 	/** Get a pointer to the held value if the held type is the same as the one specified */
 	template <typename U>
 	U* TryGet()
diff --git a/Engine/Source/Runtime/Core/Tests/Misc/TVariantTest.cpp b/Engine/Source/Runtime/Core/Tests/Misc/TVariantTest.cpp
--- a/Engine/Source/Runtime/Core/Tests/Misc/TVariantTest.cpp
+++ b/Engine/Source/Runtime/Core/Tests/Misc/TVariantTest.cpp
@@ -150,6 +150,38 @@ TEST_CASE_NAMED(FTVariantConstGetTest, "System::Core::Misc::TVariant::ConstGet",
 
 }
 
+TEST_CASE_NAMED(FTVariantGetDefaultTest, "System::Core::Misc::TVariant::GetDefault", "[ApplicationContextMask][EngineFilter]")
+{
+	// TVariant returns provided default when requested (valid) type isn't the one actually stored
+	{
+		const TVariant<int, float, double> ConstVar{ TInPlaceType<float>(), 20.0f }; // Set float
+		check(!ConstVar.IsType<int>()); // Not int
+
+		static constexpr int DefaultInt = 55;
+		check(ConstVar.Get<int>(DefaultInt) == DefaultInt); // Get int, should return default value
+	}
+
+	// TVariant returns provided default when requested (valid) type isn't the one actually stored
+	{
+		const TVariant<int, float, double> ConstVar{ TInPlaceType<float>(), 20.0f }; // Set float
+		check(!ConstVar.IsType<int>()); // Not int
+
+		static constexpr int DefaultInt = 55;
+		check(ConstVar.Get<int>(DefaultInt) == DefaultInt); // Get int, should return default value
+	}
+
+	// TVariant returns stored value when requested type is the one actually stored
+	{
+		static constexpr int StoredInt = 15;
+		const TVariant<int, float, double> ConstVar{ TInPlaceType<int>(), StoredInt }; // Set int
+		check(ConstVar.IsType<int>()); // Is int
+
+		static constexpr int DefaultInt = 55;
+		check(ConstVar.Get<int>(DefaultInt) == StoredInt); // Get int, should return stored value, not the provided default value
+	}
+
+}
+
 TEST_CASE_NAMED(FTVariantTryGetTest, "System::Core::Misc::TVariant::TryGet", "[ApplicationContextMask][EngineFilter]")
 {
 	// TVariant supports conditional get
