Add some missing non unity includes

#rnx

diff --git a/Engine/Source/Runtime/Core/Tests/Templates/MemoryOpsTest.cpp b/Engine/Source/Runtime/Core/Tests/Templates/MemoryOpsTest.cpp
--- a/Engine/Source/Runtime/Core/Tests/Templates/MemoryOpsTest.cpp
+++ b/Engine/Source/Runtime/Core/Tests/Templates/MemoryOpsTest.cpp
@@ -2,7 +2,10 @@
 
 #if WITH_TESTS
 
+#include "Containers/ArrayView.h"
+#include "Containers/Set.h"
 #include "Containers/UnrealString.h"
+
 #include "Tests/TestHarnessAdapter.h"
 
 namespace
