Reduced template instantiations in weak and SP lambda bindings, as the type of the user object is unimportant in these cases.

diff --git a/Engine/Source/Runtime/Core/Public/Delegates/DelegateBase.h b/Engine/Source/Runtime/Core/Public/Delegates/DelegateBase.h
--- a/Engine/Source/Runtime/Core/Public/Delegates/DelegateBase.h
+++ b/Engine/Source/Runtime/Core/Public/Delegates/DelegateBase.h
@@ -144,7 +144,7 @@ private:
 	template <bool, class, ESPMode, typename, typename, typename...>
 	friend class TBaseSPMethodDelegateInstance;
 
-	template <typename, ESPMode, typename, typename, typename, typename...>
+	template <ESPMode, typename, typename, typename, typename...>
 	friend class TBaseSPLambdaDelegateInstance;
 
 	template <bool, class, typename, typename, typename...>
@@ -159,7 +159,7 @@ private:
 	template <typename, typename, typename, typename...>
 	friend class TBaseFunctorDelegateInstance;
 
-	template <typename, typename, typename, typename, typename...>
+	template <typename, typename, typename, typename...>
 	friend class TWeakBaseFunctorDelegateInstance;
 
 	friend class FWriteLockedDelegateAllocation;
diff --git a/Engine/Source/Runtime/Core/Public/Delegates/DelegateInstancesImpl.h b/Engine/Source/Runtime/Core/Public/Delegates/DelegateInstancesImpl.h
--- a/Engine/Source/Runtime/Core/Public/Delegates/DelegateInstancesImpl.h
+++ b/Engine/Source/Runtime/Core/Public/Delegates/DelegateInstancesImpl.h
@@ -324,8 +324,8 @@ protected:
 };
 
 
-template <typename UserClass, ESPMode SPMode, typename RetValType, typename... ParamTypes, typename UserPolicy, typename FunctorType, typename... VarTypes>
-	class TBaseSPLambdaDelegateInstance<UserClass, SPMode, RetValType(ParamTypes...), UserPolicy, FunctorType, VarTypes...> : public TCommonDelegateInstanceState<RetValType(ParamTypes...), UserPolicy, VarTypes...>
+template <ESPMode SPMode, typename RetValType, typename... ParamTypes, typename UserPolicy, typename FunctorType, typename... VarTypes>
+class TBaseSPLambdaDelegateInstance<SPMode, RetValType(ParamTypes...), UserPolicy, FunctorType, VarTypes...> : public TCommonDelegateInstanceState<RetValType(ParamTypes...), UserPolicy, VarTypes...>
 {
 private:
 	static_assert(std::is_same_v<FunctorType, typename TRemoveReference<FunctorType>::Type>, "FunctorType cannot be a reference");
@@ -333,9 +333,11 @@ private:
 	using Super = TCommonDelegateInstanceState<RetValType(ParamTypes...), UserPolicy, VarTypes...>;
 	using DelegateBaseType = typename UserPolicy::FDelegateExtras;
 
+	using UserClass = const void;
+
 public:
 	template <typename InFunctorType, typename... InVarTypes>
-	explicit TBaseSPLambdaDelegateInstance(const TSharedPtr<UserClass, SPMode>& InContextObject, InFunctorType&& InFunctor, InVarTypes&&... Vars)
+	explicit TBaseSPLambdaDelegateInstance(const TWeakPtr<UserClass, SPMode>& InContextObject, InFunctorType&& InFunctor, InVarTypes&&... Vars)
 		: Super(Forward<InVarTypes>(Vars)...)
 		, ContextObject(InContextObject)
 		, Functor(Forward<InFunctorType>(InFunctor))
@@ -901,8 +903,8 @@ private:
 };
 
 
-template <typename UserClass, typename RetValType, typename... ParamTypes, typename UserPolicy, typename FunctorType, typename... VarTypes>
-class TWeakBaseFunctorDelegateInstance<UserClass, RetValType(ParamTypes...), UserPolicy, FunctorType, VarTypes...> : public TCommonDelegateInstanceState<RetValType(ParamTypes...), UserPolicy, VarTypes...>
+template <typename RetValType, typename... ParamTypes, typename UserPolicy, typename FunctorType, typename... VarTypes>
+class TWeakBaseFunctorDelegateInstance<RetValType(ParamTypes...), UserPolicy, FunctorType, VarTypes...> : public TCommonDelegateInstanceState<RetValType(ParamTypes...), UserPolicy, VarTypes...>
 {
 private:
 	static_assert(std::is_same_v<FunctorType, typename TRemoveReference<FunctorType>::Type>, "FunctorType cannot be a reference");
@@ -910,6 +912,10 @@ private:
 	using Super            = TCommonDelegateInstanceState<RetValType(ParamTypes...), UserPolicy, VarTypes...>;
 	using DelegateBaseType = typename UserPolicy::FDelegateExtras;
 
+	// A dependency on the template parameter is required here to stop early instantiation of types which are
+	// defined in CoreUObject.
+	using UserClass = std::conditional_t<sizeof...(ParamTypes) == 0, const UObject, const UObject>;
+
 public:
 	template <typename InFunctorType, typename... InVarTypes>
 	explicit TWeakBaseFunctorDelegateInstance(UserClass* InContextObject, InFunctorType&& InFunctor, InVarTypes&&... Vars)
diff --git a/Engine/Source/Runtime/Core/Public/Delegates/DelegateInstancesImplFwd.h b/Engine/Source/Runtime/Core/Public/Delegates/DelegateInstancesImplFwd.h
--- a/Engine/Source/Runtime/Core/Public/Delegates/DelegateInstancesImplFwd.h
+++ b/Engine/Source/Runtime/Core/Public/Delegates/DelegateInstancesImplFwd.h
@@ -25,7 +25,7 @@ class TBaseSPMethodDelegateInstance;
 /**
  * Implements a delegate binding for shared pointer functors, e.g. lambdas.
  */
-template <typename UserClass, ESPMode SPMode, typename FuncType, typename UserPolicy, typename FunctorType, typename... VarTypes>
+template <ESPMode SPMode, typename FuncType, typename UserPolicy, typename FunctorType, typename... VarTypes>
 class TBaseSPLambdaDelegateInstance;
 
 /**
@@ -55,5 +55,5 @@ class TBaseFunctorDelegateInstance;
 /**
  * Implements a weak object delegate binding for C++ functors, e.g. lambdas.
  */
-template <typename UserClass, typename FuncType, typename UserPolicy, typename FunctorType, typename... VarTypes>
+template <typename FuncType, typename UserPolicy, typename FunctorType, typename... VarTypes>
 class TWeakBaseFunctorDelegateInstance;
diff --git a/Engine/Source/Runtime/Core/Public/Delegates/DelegateSignatureImpl.inl b/Engine/Source/Runtime/Core/Public/Delegates/DelegateSignatureImpl.inl
--- a/Engine/Source/Runtime/Core/Public/Delegates/DelegateSignatureImpl.inl
+++ b/Engine/Source/Runtime/Core/Public/Delegates/DelegateSignatureImpl.inl
@@ -135,22 +135,22 @@ public:
 	template<typename UserClass, ESPMode Mode, typename FunctorType, typename... VarTypes>
 	inline void BindSPLambda(const TSharedRef<UserClass, Mode>& InUserObjectRef, FunctorType&& InFunctor, VarTypes&&... Vars)
 	{
-		new (TWriteLockedDelegateAllocation{*this}) TBaseSPLambdaDelegateInstance<const UserClass, Mode, FuncType, UserPolicy, typename TRemoveReference<FunctorType>::Type, std::decay_t<VarTypes>...>(InUserObjectRef, Forward<FunctorType>(InFunctor), Forward<VarTypes>(Vars)...);
+		new (TWriteLockedDelegateAllocation{*this}) TBaseSPLambdaDelegateInstance<Mode, FuncType, UserPolicy, typename TRemoveReference<FunctorType>::Type, std::decay_t<VarTypes>...>(TWeakPtr<const void>(InUserObjectRef), Forward<FunctorType>(InFunctor), Forward<VarTypes>(Vars)...);
 	}
 	template <typename UserClass, typename FunctorType, typename... VarTypes>
 	inline void BindSPLambda(const UserClass* InUserObject, FunctorType&& InFunctor, VarTypes&&... Vars)
 	{
-		new (TWriteLockedDelegateAllocation{*this}) TBaseSPLambdaDelegateInstance<const UserClass, decltype(InUserObject->AsShared())::Mode, FuncType, UserPolicy, typename TRemoveReference<FunctorType>::Type, std::decay_t<VarTypes>...>(StaticCastSharedRef<const UserClass>(InUserObject->AsShared()), Forward<FunctorType>(InFunctor), Forward<VarTypes>(Vars)...);
+		new (TWriteLockedDelegateAllocation{*this}) TBaseSPLambdaDelegateInstance<decltype(InUserObject->AsShared())::Mode, FuncType, UserPolicy, typename TRemoveReference<FunctorType>::Type, std::decay_t<VarTypes>...>(TWeakPtr<const void>(InUserObject->AsWeak()), Forward<FunctorType>(InFunctor), Forward<VarTypes>(Vars)...);
 	}
 
 	/**
 	 * Static: Binds a weak object C++ lambda delegate
 	 * technically this works for any functor types, but lambdas are the primary use case
 	 */
-	template<typename UserClass, typename FunctorType, typename... VarTypes>
-	inline void BindWeakLambda(UserClass* InUserObject, FunctorType&& InFunctor, VarTypes&&... Vars)
+	template<typename FunctorType, typename... VarTypes>
+	inline void BindWeakLambda(const UObject* InUserObject, FunctorType&& InFunctor, VarTypes&&... Vars)
 	{
-		new (TWriteLockedDelegateAllocation{*this}) TWeakBaseFunctorDelegateInstance<UserClass, FuncType, UserPolicy, typename TRemoveReference<FunctorType>::Type, std::decay_t<VarTypes>...>(InUserObject, Forward<FunctorType>(InFunctor), Forward<VarTypes>(Vars)...);
+		new (TWriteLockedDelegateAllocation{*this}) TWeakBaseFunctorDelegateInstance<FuncType, UserPolicy, typename TRemoveReference<FunctorType>::Type, std::decay_t<VarTypes>...>(InUserObject, Forward<FunctorType>(InFunctor), Forward<VarTypes>(Vars)...);
 	}
 
 	/**
@@ -389,14 +389,14 @@ public:
 	[[nodiscard]] inline static TDelegate<RetValType(ParamTypes...), UserPolicy> CreateSPLambda(const TSharedRef<UserClass, Mode>& InUserObjectRef, FunctorType&& InFunctor, VarTypes&&... Vars)
 	{
 		TDelegate<RetValType(ParamTypes...), UserPolicy> Result;
-		new (TWriteLockedDelegateAllocation{Result}) TBaseSPLambdaDelegateInstance<UserClass, Mode, FuncType, UserPolicy, typename TRemoveReference<FunctorType>::Type, std::decay_t<VarTypes>...>(InUserObjectRef, Forward<FunctorType>(InFunctor), Forward<VarTypes>(Vars)...);
+		new (TWriteLockedDelegateAllocation{Result}) TBaseSPLambdaDelegateInstance<Mode, FuncType, UserPolicy, typename TRemoveReference<FunctorType>::Type, std::decay_t<VarTypes>...>(InUserObjectRef, Forward<FunctorType>(InFunctor), Forward<VarTypes>(Vars)...);
 		return Result;
 	}
 	template <typename UserClass, typename FunctorType, typename... VarTypes>
 	[[nodiscard]] inline static TDelegate<RetValType(ParamTypes...), UserPolicy> CreateSPLambda(UserClass* InUserObject, FunctorType&& InFunctor, VarTypes&&... Vars)
 	{
 		TDelegate<RetValType(ParamTypes...), UserPolicy> Result;
-		new (TWriteLockedDelegateAllocation{Result}) TBaseSPLambdaDelegateInstance<UserClass, decltype(InUserObject->AsShared())::Mode, FuncType, UserPolicy, typename TRemoveReference<FunctorType>::Type, std::decay_t<VarTypes>...>(StaticCastSharedRef<UserClass>(InUserObject->AsShared()), Forward<FunctorType>(InFunctor), Forward<VarTypes>(Vars)...);
+		new (TWriteLockedDelegateAllocation{Result}) TBaseSPLambdaDelegateInstance<decltype(InUserObject->AsWeak())::Mode, FuncType, UserPolicy, typename TRemoveReference<FunctorType>::Type, std::decay_t<VarTypes>...>(TWeakPtr<const void>(InUserObject->AsWeak()), Forward<FunctorType>(InFunctor), Forward<VarTypes>(Vars)...);
 		return Result;
 	}
 
@@ -404,11 +404,11 @@ public:
 	 * Static: Creates a weak object C++ lambda delegate
 	 * technically this works for any functor types, but lambdas are the primary use case
 	 */
-	template<typename UserClass, typename FunctorType, typename... VarTypes>
-	[[nodiscard]] inline static TDelegate<RetValType(ParamTypes...), UserPolicy> CreateWeakLambda(UserClass* InUserObject, FunctorType&& InFunctor, VarTypes&&... Vars)
+	template<typename FunctorType, typename... VarTypes>
+	[[nodiscard]] inline static TDelegate<RetValType(ParamTypes...), UserPolicy> CreateWeakLambda(const UObject* InUserObject, FunctorType&& InFunctor, VarTypes&&... Vars)
 	{
 		TDelegate<RetValType(ParamTypes...), UserPolicy> Result;
-		new (TWriteLockedDelegateAllocation{Result}) TWeakBaseFunctorDelegateInstance<UserClass, FuncType, UserPolicy, typename TRemoveReference<FunctorType>::Type, std::decay_t<VarTypes>...>(InUserObject, Forward<FunctorType>(InFunctor), Forward<VarTypes>(Vars)...);
+		new (TWriteLockedDelegateAllocation{Result}) TWeakBaseFunctorDelegateInstance<FuncType, UserPolicy, typename TRemoveReference<FunctorType>::Type, std::decay_t<VarTypes>...>(InUserObject, Forward<FunctorType>(InFunctor), Forward<VarTypes>(Vars)...);
 		return Result;
 	}
 
