[core] intermittent msvc compiler bug workaround for FormatStringSan

diff --git a/Engine/Source/Runtime/Core/Public/String/FormatStringSan.h b/Engine/Source/Runtime/Core/Public/String/FormatStringSan.h
--- a/Engine/Source/Runtime/Core/Public/String/FormatStringSan.h
+++ b/Engine/Source/Runtime/Core/Public/String/FormatStringSan.h
@@ -174,20 +174,20 @@ namespace UE::Core::Private
 				}
 			}
 
-			static constexpr FResult Check(bool bInsideFormatSpec, int CurArgPos, const TCHAR* Fmt)
+			static constexpr FResult HandleDynamicLengthSpecifier(int CurArgPos, const TCHAR* P)
 			{
-				auto HandleDynamicLengthSpecifier = [CurArgPos](const TCHAR* P) constexpr -> FResult
+				if constexpr (!(std::is_integral_v<Arg> || bIsIntegralEnum<Arg>))
 				{
-					if constexpr (!(std::is_integral_v<Arg> || bIsIntegralEnum<Arg>))
-					{
-						return {StatusDynamicLengthSpecNeedsIntegerArg, CurArgPos};
-					}
-					else
-					{
-						return TCheckFormatString<TFmtArgTypes<Args...>>::Check(true, CurArgPos + 1, P + 1);
-					}
-				};
+					return {StatusDynamicLengthSpecNeedsIntegerArg, CurArgPos};
+				}
+				else
+				{
+					return TCheckFormatString<TFmtArgTypes<Args...>>::Check(true, CurArgPos + 1, P + 1);
+				}
+			};
 
+			static constexpr FResult Check(bool bInsideFormatSpec, int CurArgPos, const TCHAR* Fmt)
+			{
 				const TCHAR* P = Fmt;
 				if (!P[0] && !bInsideFormatSpec)
 				{
@@ -229,7 +229,7 @@ namespace UE::Core::Private
 				}
 				if (*P == TEXT('*'))
 				{
-					return HandleDynamicLengthSpecifier(P);
+					return HandleDynamicLengthSpecifier(CurArgPos, P);
 				}
 				else if (CharIsDigit(*P))
 				{
@@ -242,7 +242,7 @@ namespace UE::Core::Private
 				}
 				if (*P == TEXT('*'))
 				{
-					return HandleDynamicLengthSpecifier(P);
+					return HandleDynamicLengthSpecifier(CurArgPos, P);
 				}
 				else if (CharIsDigit(*P))
 				{
