[HWRT] Add missing UsesRayTracing overload to TTargetPlatformSettingsBase that checks TPlatformProperties::SupportsRayTracing() otherwise we end up including HWRT data on platforms that do not support it.

#rnx

diff --git a/Engine/Source/Developer/TargetPlatform/Public/Common/TargetPlatformSettingsBase.h b/Engine/Source/Developer/TargetPlatform/Public/Common/TargetPlatformSettingsBase.h
--- a/Engine/Source/Developer/TargetPlatform/Public/Common/TargetPlatformSettingsBase.h
+++ b/Engine/Source/Developer/TargetPlatform/Public/Common/TargetPlatformSettingsBase.h
@@ -84,6 +84,15 @@ public:
 		return IniPlatformNameValue;
 	}
 
+	virtual bool UsesRayTracing() const override
+	{
+		if (TPlatformProperties::SupportsRayTracing())
+		{
+			return FTargetPlatformSettingsBase::UsesRayTracing();
+		}
+		return false;
+	}
+
 	virtual bool SupportsFeature(ETargetPlatformFeatures Feature) const override
 	{
 		switch (Feature)
diff --git a/Engine/Source/Runtime/Core/Private/UObject/FortniteMainBranchObjectVersion.cpp b/Engine/Source/Runtime/Core/Private/UObject/FortniteMainBranchObjectVersion.cpp
--- a/Engine/Source/Runtime/Core/Private/UObject/FortniteMainBranchObjectVersion.cpp
+++ b/Engine/Source/Runtime/Core/Private/UObject/FortniteMainBranchObjectVersion.cpp
@@ -13,7 +13,7 @@ TMap<FGuid, FGuid> FFortniteMainBranchObjectVersion::GetSystemGuids()
 	SystemGuids.Add(DevGuids.NIAGARASHADERMAP_DERIVEDDATA_VER, FGuid("C417A50271B3427D98785388266E3B66"));
 	SystemGuids.Add(DevGuids.Niagara_LatestScriptCompileVersion, FGuid("393081597FA7234C92F00869DE262D73"));
 	SystemGuids.Add(DevGuids.SkeletalMeshDerivedDataVersion, FGuid("0829C9BAB43B4B25A63747FDA68E7197"));
-	SystemGuids.Add(DevGuids.STATICMESH_DERIVEDDATA_VER, FGuid("68C7D98BE48042E0A35351325741F70E"));
+	SystemGuids.Add(DevGuids.STATICMESH_DERIVEDDATA_VER, FGuid("656E338C78C34641A5D3BFAE303997AF"));
 	SystemGuids.Add(DevGuids.MaterialTranslationDDCVersion, FGuid("8CF065BF76124292A1A0546D532BEBDF"));
 	return SystemGuids;
 }
