- Fixed Hotfixes and SafeUnload (we don't unload string-based dynamic layers anymore)

diff --git a/Engine/Source/Runtime/Core/Private/Misc/ConfigCacheIni.cpp b/Engine/Source/Runtime/Core/Private/Misc/ConfigCacheIni.cpp
--- a/Engine/Source/Runtime/Core/Private/Misc/ConfigCacheIni.cpp
+++ b/Engine/Source/Runtime/Core/Private/Misc/ConfigCacheIni.cpp
@@ -3139,7 +3139,7 @@ bool FConfigBranch::AddDynamicLayersToHierarchy(const TArray<FString>& Filenames
 	return AddDynamicLayersToHierarchy(Layers, ModificationTracker);
 }
 
-bool FConfigBranch::AddDynamicLayersToHierarchy(const TArray<FDynamicLayerInfo>& Layers, FConfigModificationTracker* ModificationTracker, TSet<FString>* GlobalConfigFileCache, TSet<FString>* PluginConfigFileCache)
+bool FConfigBranch::AddDynamicLayersToHierarchy(const TArray<FDynamicLayerInfo>& Layers, FConfigModificationTracker* ModificationTracker, TSet<FString>* GlobalConfigFileCache, TSet<FString>* PluginConfigFileCache, bool bForceFullDynamicLayerUpdate)
 {
 	static bool bDumpIniLoadInfo = FParse::Param(FCommandLine::Get(), TEXT("dumpiniloads"));
 
@@ -3148,6 +3148,10 @@ bool FConfigBranch::AddDynamicLayersToHierarchy(const TArray<FDynamicLayerInfo>&
 
 	// calculate a patch so we don't lose in-memory changes
 	FConfigCommandStream Patch;
+	if (bForceFullDynamicLayerUpdate)
+	{
+		Patch = CalculateDiff(FinalCombinedLayers, InMemoryFile);
+	}
 
 	TArray<FConfigCommandStream*> AddedLayers;
 	for (const FDynamicLayerInfo& Layer : Layers)
@@ -3164,7 +3168,7 @@ bool FConfigBranch::AddDynamicLayersToHierarchy(const TArray<FDynamicLayerInfo>&
 
 		UE_CLOG(bDumpIniLoadInfo, LogConfig, Display, TEXT("   Found %s!"), *Layer.Filename);
 
-		if (AddedLayers.Num() == 0)
+		if (AddedLayers.Num() == 0 && !bForceFullDynamicLayerUpdate)
 		{
 			Patch = CalculateDiff(FinalCombinedLayers, InMemoryFile);
 			UE_LOG(LogConfig, Verbose, TEXT("  .. calculating diff on first file"));
@@ -3243,10 +3247,10 @@ bool FConfigBranch::AddDynamicLayersToHierarchy(const TArray<FDynamicLayerInfo>&
 		}
 	}
 	
-	if (AddedLayers.Num() > 0)
+	if (AddedLayers.Num() > 0 || bForceFullDynamicLayerUpdate)
 	{
 		// if all were added at the end (or there's no replay), we can just apply them without rewinding
-		if (!bInsertedBeforeEnd)
+		if (!bInsertedBeforeEnd && !bForceFullDynamicLayerUpdate)
 		{
 			for (FConfigCommandStream* NewLayer : AddedLayers)
 			{
@@ -3317,6 +3321,8 @@ bool FConfigBranch::AddDynamicLayerStringToHierarchy(const FString& Filename, co
 		}
 	}
 
+	// we can't SafeUnload a string-based layer because we'll never be able to load it again
+	DynamicLayer->bNeverSafeUnload = true;
 	DynamicLayer->Tag = Tag;
 	FillFileFromBuffer(DynamicLayer, Contents, true, Filename);
 
@@ -3466,14 +3472,17 @@ void FConfigBranch::SafeUnload()
 	StaticLayers.Empty();
 	for (DynamicLayerList::TIterator Node(DynamicLayers.GetHead()); Node; ++Node)
 	{
-		Node->Empty();
+		if (!Node->bNeverSafeUnload)
+		{
+			Node->Empty();
+		}
 	}
 }
 
 void FConfigBranch::SafeReload()
 {
 	double StartTime = FPlatformTime::Seconds();
-	
+
 	// read static layers back in from disk
 	// @todo make sure we only Unload from GConfig
 	FConfigContext Context = FConfigContext::ReadIntoConfigSystem(GConfig, Platform.ToString());
@@ -3482,20 +3491,31 @@ void FConfigBranch::SafeReload()
 	Context.Load(*IniName.ToString());
 
 	// read dynamic layers back in from disk
-	FConfigBranch::DynamicLayerList EmptiedDynamicLayers;
-	while (!DynamicLayers.IsEmpty())
+	TArray<FDynamicLayerInfo> ReloadInfos;
+
+	FConfigBranch::DynamicLayerList::TDoubleLinkedListNode* CurrentNode = DynamicLayers.GetHead(), * NextNode;
+	while (CurrentNode != nullptr)
 	{
-		FConfigBranch::DynamicLayerList::TDoubleLinkedListNode* HeadNode = DynamicLayers.GetHead();
-		DynamicLayers.RemoveNode(HeadNode, false);
-		EmptiedDynamicLayers.AddTail(HeadNode);
+		NextNode = CurrentNode->GetNextNode();
+
+		// any never unload layers we just leave in the stream
+		FConfigCommandStream& S = *CurrentNode->GetValue();
+		if (!S.bNeverSafeUnload)
+		{
+			ReloadInfos.Add({ S.Filename, S.Tag, S.Priority });
+			DynamicLayers.RemoveNode(CurrentNode, false);
+		}
+		CurrentNode = NextNode;
 	}
-	TArray<FDynamicLayerInfo> ReloadInfos;
-	for (FConfigBranch::DynamicLayerList::TIterator Node(EmptiedDynamicLayers.GetHead()); Node; ++Node)
+
+	// even if we have no ReloadInfos, we may have leftover some string-based (bNeverSafeUnload) layers in the list, and we
+	// will need to perform a full dynamic layer fixup of the final InMemoryFile in the branch - and we force a full fat 
+	// update in either case to be 100% safe
+	if (ReloadInfos.Num() > 0 || !DynamicLayers.IsEmpty())
 	{
-		FConfigCommandStream& S = *Node.GetNode()->GetValue();
-		ReloadInfos.Add( { S.Filename, S.Tag, S.Priority } );
+		const bool bForceFullDynamicLayerUpdate = true;
+		AddDynamicLayersToHierarchy(ReloadInfos, nullptr, nullptr, nullptr, bForceFullDynamicLayerUpdate);
 	}
-	AddDynamicLayersToHierarchy(ReloadInfos);
 
 	UE_LOG(LogConfig, Log, TEXT("Branch '%s' had been unloaded. Reloading on-demand took %.2fms"), *IniName.ToString(), (FPlatformTime::Seconds() - StartTime) * 1000.0f);
 }
@@ -6891,17 +6911,21 @@ class FIniExec : public FSelfRegisteringExec
 				
 				if (Branch != nullptr)
 				{
-					UE::DynamicConfig::PerformDynamicConfig("HotFixText", [Branch, FileName](FConfigModificationTracker* ChangeTracker)
+					FString Contents;
+					if (FFileHelper::LoadFileToString(Contents, FileName))
 					{
-						Branch->AddDynamicLayersToHierarchy( { { FileName , "HotfixTest", (uint16)DynamicLayerPriority::Hotfix } }, ChangeTracker);
-					});
+						UE::DynamicConfig::PerformDynamicConfig("HotfixTest", [Branch, &FileName, &Contents](FConfigModificationTracker* ChangeTracker)
+							{
+								Branch->AddDynamicLayerStringToHierarchy(FileName, Contents, "HotfixTest", DynamicLayerPriority::Hotfix, ChangeTracker);
+							});
+					}
 				}
 			}
 		}
 		
 		if (FParse::Command(&Cmd, TEXT("RemoveHotFixes")))
 		{
-			UE::DynamicConfig::PerformDynamicConfig("HotFixText", [](FConfigModificationTracker* ChangeTracker)
+			UE::DynamicConfig::PerformDynamicConfig("HotfixTest", [](FConfigModificationTracker* ChangeTracker)
 			{
 				FConfigCacheIni::RemoveTagFromAllBranches("HotfixTest", ChangeTracker);
 			});
diff --git a/Engine/Source/Runtime/Core/Public/Misc/ConfigCacheIni.h b/Engine/Source/Runtime/Core/Public/Misc/ConfigCacheIni.h
--- a/Engine/Source/Runtime/Core/Public/Misc/ConfigCacheIni.h
+++ b/Engine/Source/Runtime/Core/Public/Misc/ConfigCacheIni.h
@@ -554,6 +554,9 @@ public:
 	// used to determine if existing settings should be removed from a FConfigFile before applying this stream - used for compatibility with
 	// how Saved config files are stored (replace the static layers values fully if at least one key exists)
 	uint8 bIsSavedConfigFile = false;
+
+	// we can't SafeUnload layers that come from a string, because we can't reload (Hotfixes, in particular)
+	uint8 bNeverSafeUnload = false;
 	
 	uint16 Priority;
 	FString Filename;
@@ -1118,7 +1121,7 @@ public:
 	 * Can return the modified sections if the callers wants to reloadconfig on classes
 	 */
 	CORE_API bool AddDynamicLayerToHierarchy(const FString& Filename, FConfigModificationTracker* ModificationTracker=nullptr, TSet<FString>* GlobalConfigFileCache=nullptr, TSet<FString>* PluginConfigFileCache=nullptr);
-	CORE_API bool AddDynamicLayersToHierarchy(const TArray<FDynamicLayerInfo>& Layers, FConfigModificationTracker* ModificationTracker=nullptr, TSet<FString>* GlobalConfigFileCache=nullptr, TSet<FString>* PluginConfigFileCache=nullptr);
+	CORE_API bool AddDynamicLayersToHierarchy(const TArray<FDynamicLayerInfo>& Layers, FConfigModificationTracker* ModificationTracker=nullptr, TSet<FString>* GlobalConfigFileCache=nullptr, TSet<FString>* PluginConfigFileCache=nullptr, bool bForceFullDynamicLayerUpdate=false);
 
 		UE_DEPRECATED(5.6, "Use AddDynamicLayersToHierarchy that takes a FDynamicLayerInfo list")
 	CORE_API bool AddDynamicLayersToHierarchy(const TArray<FString>& Filenames, FName Tag=NAME_None, DynamicLayerPriority Priority=DynamicLayerPriority::Unknown, FConfigModificationTracker* ModificationTracker =nullptr);
