Make UNIX_PLATFORM_FILE_SPEEDUP_FILE_OPERATIONS overrideable, and use it for some test program targets. This substantially reduces the startup overhead of these test programs when running on WSL with its relatively slow file I/O.

diff --git a/Engine/Source/Programs/AutoRTFMEngineTests/AutoRTFMEngineTests.Target.cs b/Engine/Source/Programs/AutoRTFMEngineTests/AutoRTFMEngineTests.Target.cs
--- a/Engine/Source/Programs/AutoRTFMEngineTests/AutoRTFMEngineTests.Target.cs
+++ b/Engine/Source/Programs/AutoRTFMEngineTests/AutoRTFMEngineTests.Target.cs
@@ -45,6 +45,8 @@ public class AutoRTFMEngineTestsTarget : TargetRules
 		// Load time profiling brings object construction to a crawl.
 		GlobalDefinitions.Add("LOADTIMEPROFILERTRACE_ENABLED=0");
 
+		GlobalDefinitions.Add("UNIX_PLATFORM_FILE_SPEEDUP_FILE_OPERATIONS=1");
+
 		if (!bGenerateProjectFiles)
 		{
 			bUseAutoRTFMCompiler = true;
diff --git a/Engine/Source/Programs/AutoRTFMTests/AutoRTFMTests.Target.cs b/Engine/Source/Programs/AutoRTFMTests/AutoRTFMTests.Target.cs
--- a/Engine/Source/Programs/AutoRTFMTests/AutoRTFMTests.Target.cs
+++ b/Engine/Source/Programs/AutoRTFMTests/AutoRTFMTests.Target.cs
@@ -64,6 +64,8 @@ public class AutoRTFMTestsTarget : TargetRules
 		GlobalDefinitions.Add("MALLOC_LEAKDETECTION=1");
 		GlobalDefinitions.Add("PLATFORM_USES_FIXED_GMalloc_CLASS=0");
 
+		GlobalDefinitions.Add("UNIX_PLATFORM_FILE_SPEEDUP_FILE_OPERATIONS=1");
+
 		// For testing insights, this sets `UE_TRACE_ENABLED`.
 		bEnableTrace = true;
 	}
diff --git a/Engine/Source/Runtime/Core/Private/Unix/UnixPlatformFile.cpp b/Engine/Source/Runtime/Core/Private/Unix/UnixPlatformFile.cpp
--- a/Engine/Source/Runtime/Core/Private/Unix/UnixPlatformFile.cpp
+++ b/Engine/Source/Runtime/Core/Private/Unix/UnixPlatformFile.cpp
@@ -16,7 +16,9 @@
 
 DEFINE_LOG_CATEGORY_STATIC(LogUnixPlatformFile, Log, All);
 
-#define UNIX_PLATFORM_FILE_SPEEDUP_FILE_OPERATIONS	((!WITH_EDITOR && !IS_PROGRAM) || !PLATFORM_LINUX) 
+#ifndef UNIX_PLATFORM_FILE_SPEEDUP_FILE_OPERATIONS
+#define UNIX_PLATFORM_FILE_SPEEDUP_FILE_OPERATIONS ((!WITH_EDITOR && !IS_PROGRAM) || !PLATFORM_LINUX) 
+#endif
 
 // make an FTimeSpan object that represents the "epoch" for time_t (from a stat struct)
 const FDateTime UnixEpoch(1970, 1, 1);
