Add [[nodiscard]] to FMaths functions

diff --git a/Engine/Source/Runtime/Core/Public/GenericPlatform/GenericPlatformMath.h b/Engine/Source/Runtime/Core/Public/GenericPlatform/GenericPlatformMath.h
--- a/Engine/Source/Runtime/Core/Public/GenericPlatform/GenericPlatformMath.h
+++ b/Engine/Source/Runtime/Core/Public/GenericPlatform/GenericPlatformMath.h
@@ -31,7 +31,7 @@ struct FGenericPlatformMath
 {
 	// load half (F16) to float
 	//https://gist.github.com/rygorous/2156668
-	static constexpr FORCEINLINE float LoadHalf(const uint16* Ptr)
+	[[nodiscard]] static constexpr FORCEINLINE float LoadHalf(const uint16* Ptr)
 	{
 		uint16 FP16 = *Ptr;
 		constexpr uint32 shifted_exp = 0x7c00 << 13;			// exponent mask after shift
@@ -155,7 +155,7 @@ struct FGenericPlatformMath
 	 * @param F The float to bit cast to an unsigned integer.
 	*  @return A bitwise copy of the float in a 32-bit unsigned integer value.
 	 */
-	static inline uint32 AsUInt(float F)
+	[[nodiscard]] static inline uint32 AsUInt(float F)
 	{
 		uint32 U{};
 		static_assert(sizeof(F) == sizeof(U), "The float and uint sizes must be equal");
@@ -168,11 +168,11 @@ struct FGenericPlatformMath
 	 * @param D The double to bit cast to an unsigned integer.
 	*  @return A bitwise copy of the double in a 64-bit unsigned integer value.
 	 */
-	static inline uint64 AsUInt(double F)
+	[[nodiscard]] static inline uint64 AsUInt(double D)
 	{
 		uint64 U{};
-		static_assert(sizeof(F) == sizeof(U), "The float and uint sizes must be equal");
-		::memcpy(&U, &F, sizeof(F));
+		static_assert(sizeof(D) == sizeof(U), "The float and uint sizes must be equal");
+		::memcpy(&U, &D, sizeof(D));
 		return U;
 	}
 
@@ -181,7 +181,7 @@ struct FGenericPlatformMath
 	 * @param U The 32-bit unsigned int to bit cast to a 32-bit float.
 	*  @return A bitwise copy of the 32-bit float in a 32-bit unsigned integer value.
 	 */
-	static inline float AsFloat(uint32 U)
+	[[nodiscard]] static inline float AsFloat(uint32 U)
 	{
 		float F{};
 		static_assert(sizeof(F) == sizeof(U), "The float and uint32 sizes must be equal");
@@ -194,7 +194,7 @@ struct FGenericPlatformMath
 	 * @param U The 64-bit unsigned int to bit cast to a 64-bit float.
 	*  @return A bitwise copy of the 64-bit float in a 64-bit unsigned integer value.
 	 */
-	static inline double AsFloat(uint64 U)
+	[[nodiscard]] static inline double AsFloat(uint64 U)
 	{
 		double F{};
 		static_assert(sizeof(F) == sizeof(U), "The double and uint64 sizes must be equal");
@@ -208,28 +208,28 @@ struct FGenericPlatformMath
 	 * @param F		Floating point value to convert
 	 * @return		Truncated integer.
 	 */
-	static constexpr FORCEINLINE int32 TruncToInt32(float F)
+	[[nodiscard]] static constexpr FORCEINLINE int32 TruncToInt32(float F)
 	{
 		return (int32)F;
 	}
-	static constexpr FORCEINLINE int32 TruncToInt32(double F)
+	[[nodiscard]] static constexpr FORCEINLINE int32 TruncToInt32(double F)
 	{
 		return (int32)F;
 	}
-	static constexpr FORCEINLINE int64 TruncToInt64(double F)
+	[[nodiscard]] static constexpr FORCEINLINE int64 TruncToInt64(double F)
 	{
 		return (int64)F;
 	}
 
-	static constexpr FORCEINLINE int32 TruncToInt(float F) { return TruncToInt32(F); }
-	static constexpr FORCEINLINE int64 TruncToInt(double F) { return TruncToInt64(F); }
+	[[nodiscard]] static constexpr FORCEINLINE int32 TruncToInt(float F) { return TruncToInt32(F); }
+	[[nodiscard]] static constexpr FORCEINLINE int64 TruncToInt(double F) { return TruncToInt64(F); }
 
 	/**
 	 * Converts a float to an integer value with truncation towards zero.
 	 * @param F		Floating point value to convert
 	 * @return		Truncated integer value.
 	 */
-	static FORCEINLINE float TruncToFloat(float F)
+	[[nodiscard]] static FORCEINLINE float TruncToFloat(float F)
 	{
 		return truncf(F);
 	}
@@ -239,12 +239,12 @@ struct FGenericPlatformMath
 	 * @param F		Floating point value to convert
 	 * @return		Truncated integer value.
 	 */
-	static FORCEINLINE double TruncToDouble(double F)
+	[[nodiscard]] static FORCEINLINE double TruncToDouble(double F)
 	{
 		return trunc(F);
 	}
 
-	static FORCEINLINE double TruncToFloat(double F)
+	[[nodiscard]] static FORCEINLINE double TruncToFloat(double F)
 	{
 		return TruncToDouble(F);
 	}
@@ -254,27 +254,27 @@ struct FGenericPlatformMath
 	 * @param F		Floating point value to convert
 	 * @return		An integer less or equal to 'F'.
 	 */
-	static FORCEINLINE int32 FloorToInt32(float F)
+	[[nodiscard]] static FORCEINLINE int32 FloorToInt32(float F)
 	{
 		int32 I = TruncToInt32(F);
 		I -= ((float)I > F);
 		return I;
 	}
-	static FORCEINLINE int32 FloorToInt32(double F)
+	[[nodiscard]] static FORCEINLINE int32 FloorToInt32(double F)
 	{
 		int32 I = TruncToInt32(F);
 		I -= ((double)I > F);
 		return I;
 	}
-	static FORCEINLINE int64 FloorToInt64(double F)
+	[[nodiscard]] static FORCEINLINE int64 FloorToInt64(double F)
 	{
 		int64 I = TruncToInt64(F);
 		I -= ((double)I > F);
 		return I;
 	}
 
-	static FORCEINLINE int32 FloorToInt(float F) { return FloorToInt32(F); }
-	static FORCEINLINE int64 FloorToInt(double F) { return FloorToInt64(F); }
+	[[nodiscard]] static FORCEINLINE int32 FloorToInt(float F) { return FloorToInt32(F); }
+	[[nodiscard]] static FORCEINLINE int64 FloorToInt(double F) { return FloorToInt64(F); }
 	
 	
 	/**
@@ -282,7 +282,7 @@ struct FGenericPlatformMath
 	* @param F		Floating point value to convert
 	* @return		An integer less or equal to 'F'.
 	*/
-	static FORCEINLINE float FloorToFloat(float F)
+	[[nodiscard]] static FORCEINLINE float FloorToFloat(float F)
 	{
 		return floorf(F);
 	}
@@ -292,12 +292,12 @@ struct FGenericPlatformMath
 	* @param F		Floating point value to convert
 	* @return		The nearest integer value to 'F'.
 	*/
-	static FORCEINLINE double FloorToDouble(double F)
+	[[nodiscard]] static FORCEINLINE double FloorToDouble(double F)
 	{
 		return floor(F);
 	}
 
-	static FORCEINLINE double FloorToFloat(double F)
+	[[nodiscard]] static FORCEINLINE double FloorToFloat(double F)
 	{
 		return FloorToDouble(F);
 	}
@@ -307,28 +307,28 @@ struct FGenericPlatformMath
 	 * @param F		Floating point value to convert
 	 * @return		The nearest integer to 'F'.
 	 */
-	static FORCEINLINE int32 RoundToInt32(float F)
+	[[nodiscard]] static FORCEINLINE int32 RoundToInt32(float F)
 	{
 		return FloorToInt32(F + 0.5f);
 	}
-	static FORCEINLINE int32 RoundToInt32(double F)
+	[[nodiscard]] static FORCEINLINE int32 RoundToInt32(double F)
 	{
 		return FloorToInt32(F + 0.5);
 	}
-	static FORCEINLINE int64 RoundToInt64(double F)
+	[[nodiscard]] static FORCEINLINE int64 RoundToInt64(double F)
 	{
 		return FloorToInt64(F + 0.5);
 	}
 	
-	static FORCEINLINE int32 RoundToInt(float F) { return RoundToInt32(F); }
-	static FORCEINLINE int64 RoundToInt(double F) { return RoundToInt64(F); }
+	[[nodiscard]] static FORCEINLINE int32 RoundToInt(float F) { return RoundToInt32(F); }
+	[[nodiscard]] static FORCEINLINE int64 RoundToInt(double F) { return RoundToInt64(F); }
 
 	/**
 	* Converts a float to the nearest integer. Rounds up when the fraction is .5
 	* @param F		Floating point value to convert
 	* @return		The nearest integer to 'F'.
 	*/
-	static FORCEINLINE float RoundToFloat(float F)
+	[[nodiscard]] static FORCEINLINE float RoundToFloat(float F)
 	{
 		return FloorToFloat(F + 0.5f);
 	}
@@ -338,12 +338,12 @@ struct FGenericPlatformMath
 	* @param F		Floating point value to convert
 	* @return		The nearest integer to 'F'.
 	*/
-	static FORCEINLINE double RoundToDouble(double F)
+	[[nodiscard]] static FORCEINLINE double RoundToDouble(double F)
 	{
 		return FloorToDouble(F + 0.5);
 	}
 
-	static FORCEINLINE double RoundToFloat(double F)
+	[[nodiscard]] static FORCEINLINE double RoundToFloat(double F)
 	{
 		return RoundToDouble(F);
 	}
@@ -353,34 +353,34 @@ struct FGenericPlatformMath
 	* @param F		Floating point value to convert
 	* @return		An integer greater or equal to 'F'.
 	*/
-	static FORCEINLINE int32 CeilToInt32(float F)
+	[[nodiscard]] static FORCEINLINE int32 CeilToInt32(float F)
 	{
 		int32 I = TruncToInt32(F);
 		I += ((float)I < F);
 		return I;
 	}
-	static FORCEINLINE int32 CeilToInt32(double F)
+	[[nodiscard]] static FORCEINLINE int32 CeilToInt32(double F)
 	{
 		int32 I = TruncToInt32(F);
 		I += ((double)I < F);
 		return I;
 	}
-	static FORCEINLINE int64 CeilToInt64(double F)
+	[[nodiscard]] static FORCEINLINE int64 CeilToInt64(double F)
 	{
 		int64 I = TruncToInt64(F);
 		I += ((double)I < F);
 		return I;
 	}
 
-	static FORCEINLINE int32 CeilToInt(float F) { return CeilToInt32(F); }
-	static FORCEINLINE int64 CeilToInt(double F) { return CeilToInt64(F); }
+	[[nodiscard]] static FORCEINLINE int32 CeilToInt(float F) { return CeilToInt32(F); }
+	[[nodiscard]] static FORCEINLINE int64 CeilToInt(double F) { return CeilToInt64(F); }
 
 	/**
 	* Converts a float to the nearest greater or equal integer.
 	* @param F		Floating point value to convert
 	* @return		An integer greater or equal to 'F'.
 	*/
-	static FORCEINLINE float CeilToFloat(float F)
+	[[nodiscard]] static FORCEINLINE float CeilToFloat(float F)
 	{
 		return ceilf(F);
 	}
@@ -390,12 +390,12 @@ struct FGenericPlatformMath
 	* @param F		Floating point value to convert
 	* @return		An integer greater or equal to 'F'.
 	*/
-	static FORCEINLINE double CeilToDouble(double F)
+	[[nodiscard]] static FORCEINLINE double CeilToDouble(double F)
 	{
 		return ceil(F);
 	}
 
-	static FORCEINLINE double CeilToFloat(double F)
+	[[nodiscard]] static FORCEINLINE double CeilToFloat(double F)
 	{
 		return CeilToDouble(F);
 	}
@@ -406,7 +406,7 @@ struct FGenericPlatformMath
 	 * @param F		Double precision floating point value to convert
 	 * @return		The 64-bit integer closest to 'F', with ties going to the nearest even number
 	 */
-	static int64 RoundToNearestTiesToEven(double F)
+	[[nodiscard]] static int64 RoundToNearestTiesToEven(double F)
 	{
 #if PLATFORM_HAS_FENV_H == 0
 		ensureAlwaysMsgf(FLT_ROUNDS == 1, TEXT("Platform does not support FE_TONEAREST for double to int64."));
@@ -431,12 +431,12 @@ struct FGenericPlatformMath
 	* @param Value	Floating point value to convert
 	* @return		A float between >=0 and < 1 for nonnegative input. A float between >= -1 and < 0 for negative input.
 	*/
-	static FORCEINLINE float Fractional(float Value)
+	[[nodiscard]] static FORCEINLINE float Fractional(float Value)
 	{
 		return Value - TruncToFloat(Value);
 	}
 
-	static FORCEINLINE double Fractional(double Value)
+	[[nodiscard]] static FORCEINLINE double Fractional(double Value)
 	{
 		return Value - TruncToDouble(Value);
 	}
@@ -446,13 +446,13 @@ struct FGenericPlatformMath
 	* @param Value	Floating point value to convert
 	* @return		A float between >=0 and < 1.
 	*/
-	static FORCEINLINE float Frac(float Value)
+	[[nodiscard]] static FORCEINLINE float Frac(float Value)
 	{
 		return Value - FloorToFloat(Value);
 	}
 	
 
-	static FORCEINLINE double Frac(double Value)
+	[[nodiscard]] static FORCEINLINE double Frac(double Value)
 	{
 		return Value - FloorToDouble(Value);
 	}
@@ -463,7 +463,7 @@ struct FGenericPlatformMath
 	* @param OutIntPart Floating point value that receives the integral part of the number.
 	* @return			The fractional part of the number.
 	*/
-	static FORCEINLINE float Modf(const float InValue, float* OutIntPart)
+	[[nodiscard]] static FORCEINLINE float Modf(const float InValue, float* OutIntPart)
 	{
 		return modff(InValue, OutIntPart);
 	}
@@ -474,30 +474,30 @@ struct FGenericPlatformMath
 	* @param OutIntPart Floating point value that receives the integral part of the number.
 	* @return			The fractional part of the number.
 	*/
-	static FORCEINLINE double Modf(const double InValue, double* OutIntPart)
+	[[nodiscard]] static FORCEINLINE double Modf(const double InValue, double* OutIntPart)
 	{
 		return modf(InValue, OutIntPart);
 	}
 
 	// Returns e^Value
-	static FORCEINLINE float Exp( float Value ) { return expf(Value); }
-	static FORCEINLINE double Exp(double Value) { return exp(Value); }
+	[[nodiscard]] static FORCEINLINE float Exp( float Value ) { return expf(Value); }
+	[[nodiscard]] static FORCEINLINE double Exp(double Value) { return exp(Value); }
 
 	// Returns 2^Value
-	static FORCEINLINE float Exp2( float Value ) { return powf(2.f, Value); /*exp2f(Value);*/ }
-	static FORCEINLINE double Exp2(double Value) { return pow(2.0, Value); /*exp2(Value);*/ }
+	[[nodiscard]] static FORCEINLINE float Exp2( float Value ) { return powf(2.f, Value); /*exp2f(Value);*/ }
+	[[nodiscard]] static FORCEINLINE double Exp2(double Value) { return pow(2.0, Value); /*exp2(Value);*/ }
 
-	static FORCEINLINE float Loge( float Value ) {	return logf(Value); }
-	static FORCEINLINE double Loge(double Value) { return log(Value); }
+	[[nodiscard]] static FORCEINLINE float Loge( float Value ) {	return logf(Value); }
+	[[nodiscard]] static FORCEINLINE double Loge(double Value) { return log(Value); }
 
-	static FORCEINLINE float LogX( float Base, float Value ) { return Loge(Value) / Loge(Base); }
-	static FORCEINLINE double LogX(double Base, double Value) { return Loge(Value) / Loge(Base); }
+	[[nodiscard]] static FORCEINLINE float LogX( float Base, float Value ) { return Loge(Value) / Loge(Base); }
+	[[nodiscard]] static FORCEINLINE double LogX(double Base, double Value) { return Loge(Value) / Loge(Base); }
 	RESOLVE_FLOAT_AMBIGUITY_2_ARGS(LogX);
 
 	// 1.0 / Loge(2) = 1.4426950f
-	static FORCEINLINE float Log2( float Value ) { return Loge(Value) * 1.4426950f; }	
+	[[nodiscard]] static FORCEINLINE float Log2( float Value ) { return Loge(Value) * 1.4426950f; }	
 	// 1.0 / Loge(2) = 1.442695040888963387
-	static FORCEINLINE double Log2(double Value) { return Loge(Value) * 1.442695040888963387; }
+	[[nodiscard]] static FORCEINLINE double Log2(double Value) { return Loge(Value) * 1.442695040888963387; }
 
 	/**
 	 * Returns the floating-point remainder of X / Y
@@ -509,121 +509,121 @@ struct FGenericPlatformMath
 	 * which is not consistent with the intent nor with the vectorized version.
 	 */
 
-	static CORE_API FORCENOINLINE float Fmod(float X, float Y);
-	static CORE_API FORCENOINLINE double Fmod(double X, double Y);
+	[[nodiscard]] static CORE_API FORCENOINLINE float Fmod(float X, float Y);
+	[[nodiscard]] static CORE_API FORCENOINLINE double Fmod(double X, double Y);
 	RESOLVE_FLOAT_AMBIGUITY_2_ARGS(Fmod);
 
-	static FORCEINLINE float Sin( float Value ) { return sinf(Value); }
-	static FORCEINLINE double Sin( double Value ) { return sin(Value); }
+	[[nodiscard]] static FORCEINLINE float Sin( float Value ) { return sinf(Value); }
+	[[nodiscard]] static FORCEINLINE double Sin( double Value ) { return sin(Value); }
 
-	static FORCEINLINE float Asin( float Value ) { return asinf( (Value<-1.f) ? -1.f : ((Value<1.f) ? Value : 1.f) ); }
-	static FORCEINLINE double Asin( double Value ) { return asin( (Value<-1.0) ? -1.0 : ((Value<1.0) ? Value : 1.0) ); }
+	[[nodiscard]] static FORCEINLINE float Asin( float Value ) { return asinf( (Value<-1.f) ? -1.f : ((Value<1.f) ? Value : 1.f) ); }
+	[[nodiscard]] static FORCEINLINE double Asin( double Value ) { return asin( (Value<-1.0) ? -1.0 : ((Value<1.0) ? Value : 1.0) ); }
 
-	static FORCEINLINE float Sinh(float Value) { return sinhf(Value); }
-	static FORCEINLINE double Sinh(double Value) { return sinh(Value); }
+	[[nodiscard]] static FORCEINLINE float Sinh(float Value) { return sinhf(Value); }
+	[[nodiscard]] static FORCEINLINE double Sinh(double Value) { return sinh(Value); }
 
-	static FORCEINLINE float Cos( float Value ) { return cosf(Value); }
-	static FORCEINLINE double Cos( double Value ) { return cos(Value); }
+	[[nodiscard]] static FORCEINLINE float Cos( float Value ) { return cosf(Value); }
+	[[nodiscard]] static FORCEINLINE double Cos( double Value ) { return cos(Value); }
 
-	static FORCEINLINE float Acos( float Value ) { return acosf( (Value<-1.f) ? -1.f : ((Value<1.f) ? Value : 1.f) ); }
-	static FORCEINLINE double Acos( double Value ) { return acos( (Value<-1.0) ? -1.0 : ((Value<1.0) ? Value : 1.0) ); }
+	[[nodiscard]] static FORCEINLINE float Acos( float Value ) { return acosf( (Value<-1.f) ? -1.f : ((Value<1.f) ? Value : 1.f) ); }
+	[[nodiscard]] static FORCEINLINE double Acos( double Value ) { return acos( (Value<-1.0) ? -1.0 : ((Value<1.0) ? Value : 1.0) ); }
 
-	static FORCEINLINE float Cosh(float Value) { return coshf(Value); }
-	static FORCEINLINE double Cosh(double Value) { return cosh(Value); }
+	[[nodiscard]] static FORCEINLINE float Cosh(float Value) { return coshf(Value); }
+	[[nodiscard]] static FORCEINLINE double Cosh(double Value) { return cosh(Value); }
 
-	static FORCEINLINE float Tan( float Value ) { return tanf(Value); }
-	static FORCEINLINE double Tan( double Value ) { return tan(Value); }
+	[[nodiscard]] static FORCEINLINE float Tan( float Value ) { return tanf(Value); }
+	[[nodiscard]] static FORCEINLINE double Tan( double Value ) { return tan(Value); }
 
-	static FORCEINLINE float Atan( float Value ) { return atanf(Value); }
-	static FORCEINLINE double Atan( double Value ) { return atan(Value); }
+	[[nodiscard]] static FORCEINLINE float Atan( float Value ) { return atanf(Value); }
+	[[nodiscard]] static FORCEINLINE double Atan( double Value ) { return atan(Value); }
 
-	static FORCEINLINE float Tanh(float Value) { return tanhf(Value); }
-	static FORCEINLINE double Tanh(double Value) { return tanh(Value); }
+	[[nodiscard]] static FORCEINLINE float Tanh(float Value) { return tanhf(Value); }
+	[[nodiscard]] static FORCEINLINE double Tanh(double Value) { return tanh(Value); }
 
-	static CORE_API float Atan2( float Y, float X );
-	static CORE_API double Atan2( double Y, double X );
+	[[nodiscard]] static CORE_API float Atan2( float Y, float X );
+	[[nodiscard]] static CORE_API double Atan2( double Y, double X );
 	RESOLVE_FLOAT_AMBIGUITY_2_ARGS(Atan2);
 
-	static FORCEINLINE float Sqrt( float Value ) { return sqrtf(Value); }
-	static FORCEINLINE double Sqrt( double Value ) { return sqrt(Value); }
+	[[nodiscard]] static FORCEINLINE float Sqrt( float Value ) { return sqrtf(Value); }
+	[[nodiscard]] static FORCEINLINE double Sqrt( double Value ) { return sqrt(Value); }
 
-	static FORCEINLINE float Pow( float A, float B ) { return powf(A,B); }
-	static FORCEINLINE double Pow( double A, double B ) { return pow(A,B); }
+	[[nodiscard]] static FORCEINLINE float Pow( float A, float B ) { return powf(A,B); }
+	[[nodiscard]] static FORCEINLINE double Pow( double A, double B ) { return pow(A,B); }
 	RESOLVE_FLOAT_AMBIGUITY_2_ARGS(Pow);
 
 	/** Computes a fully accurate inverse square root */
-	static FORCEINLINE float InvSqrt( float F ) { return 1.0f / sqrtf( F ); }
-	static FORCEINLINE double InvSqrt( double F ) { return 1.0 / sqrt( F ); }
+	[[nodiscard]] static FORCEINLINE float InvSqrt( float F ) { return 1.0f / sqrtf( F ); }
+	[[nodiscard]] static FORCEINLINE double InvSqrt( double F ) { return 1.0 / sqrt( F ); }
 
 	/** Computes a faster but less accurate inverse square root */
-	static FORCEINLINE float InvSqrtEst( float F ) { return InvSqrt( F ); }
-	static FORCEINLINE double InvSqrtEst( double F ) { return InvSqrt( F ); }
+	[[nodiscard]] static FORCEINLINE float InvSqrtEst( float F ) { return InvSqrt( F ); }
+	[[nodiscard]] static FORCEINLINE double InvSqrtEst( double F ) { return InvSqrt( F ); }
 
 	/** Return true if value is NaN (not a number). */
-	static FORCEINLINE bool IsNaN( float A ) 
+	[[nodiscard]] static FORCEINLINE bool IsNaN( float A ) 
 	{
 		return (BitCast<uint32>(A) & 0x7FFFFFFFU) > 0x7F800000U;
 	}
-	static FORCEINLINE bool IsNaN(double A)
+	[[nodiscard]] static FORCEINLINE bool IsNaN(double A)
 	{
 		return (BitCast<uint64>(A) & 0x7FFFFFFFFFFFFFFFULL) > 0x7FF0000000000000ULL;
 	}
 
 	/** Return true if value is finite (not NaN and not Infinity). */
-	static FORCEINLINE bool IsFinite( float A )
+	[[nodiscard]] static FORCEINLINE bool IsFinite( float A )
 	{
 		return (BitCast<uint32>(A) & 0x7F800000U) != 0x7F800000U;
 	}
-	static FORCEINLINE bool IsFinite(double A)
+	[[nodiscard]] static FORCEINLINE bool IsFinite(double A)
 	{
 		return (BitCast<uint64>(A) & 0x7FF0000000000000ULL) != 0x7FF0000000000000ULL;
 	}
 
-	static FORCEINLINE bool IsNegativeOrNegativeZero(float A)
+	[[nodiscard]] static FORCEINLINE bool IsNegativeOrNegativeZero(float A)
 	{
 		return BitCast<uint32>(A) >= (uint32)0x80000000; // Detects sign bit.
 	}
 
-	static FORCEINLINE bool IsNegativeOrNegativeZero(double A)
+	[[nodiscard]] static FORCEINLINE bool IsNegativeOrNegativeZero(double A)
 	{
 		return BitCast<uint64>(A) >= (uint64)0x8000000000000000; // Detects sign bit.
 	}
 
 	UE_DEPRECATED(5.1, "IsNegativeFloat has been deprecated in favor of IsNegativeOrNegativeZero or simply A < 0.")
-	static FORCEINLINE bool IsNegativeFloat(float A)
+	[[nodiscard]] static FORCEINLINE bool IsNegativeFloat(float A)
 	{
 		return IsNegativeOrNegativeZero(A);
 	}
 
 	UE_DEPRECATED(5.1, "IsNegativeDouble has been deprecated in favor of IsNegativeOrNegativeZero or simply A < 0.")
-	static FORCEINLINE bool IsNegativeDouble(double A)
+	[[nodiscard]] static FORCEINLINE bool IsNegativeDouble(double A)
 	{
 		return IsNegativeOrNegativeZero(A);
 	}
 
 	UE_DEPRECATED(5.1, "IsNegative has been deprecated in favor of IsNegativeOrNegativeZero or simply A < 0.")
-	static FORCEINLINE bool IsNegative(float A)
+	[[nodiscard]] static FORCEINLINE bool IsNegative(float A)
 	{
 		return IsNegativeOrNegativeZero(A);
 	}
 
 	UE_DEPRECATED(5.1, "IsNegative has been deprecated in favor of IsNegativeOrNegativeZero or simply A < 0.")
-	static FORCEINLINE bool IsNegative(double A)
+	[[nodiscard]] static FORCEINLINE bool IsNegative(double A)
 	{
 		return IsNegativeOrNegativeZero(A);
 	}
 
 	/** Returns a random integer between 0 and RAND_MAX, inclusive */
-	static FORCEINLINE int32 Rand() { return rand(); }
+	[[nodiscard]] static FORCEINLINE int32 Rand() { return rand(); }
 
 	/** Returns a random integer between 0 and MAX_int32, inclusive. RAND_MAX may only be 15 bits, so compose from multiple calls. */
-	static FORCEINLINE int32 Rand32() { return ((rand() & 0x7fff) << 16) | ((rand() & 0x7fff) << 1) | (rand() & 0x1); }
+	[[nodiscard]] static FORCEINLINE int32 Rand32() { return ((rand() & 0x7fff) << 16) | ((rand() & 0x7fff) << 1) | (rand() & 0x1); }
 
 	/** Seeds global random number functions Rand() and FRand() */
 	static FORCEINLINE void RandInit(int32 Seed) { srand( Seed ); }
 
 	/** Returns a random float between 0 and 1, inclusive. */
-	static FORCEINLINE float FRand() 
+	[[nodiscard]] static FORCEINLINE float FRand() 
 	{ 
 		// FP32 mantissa can only represent 24 bits before losing precision
 		constexpr int32 RandMax = 0x00ffffff < RAND_MAX ? 0x00ffffff : RAND_MAX;
@@ -634,10 +634,10 @@ struct FGenericPlatformMath
 	static CORE_API void SRandInit( int32 Seed );
 
 	/** Returns the current seed for SRand(). */
-	static CORE_API int32 GetRandSeed();
+	[[nodiscard]] static CORE_API int32 GetRandSeed();
 
 	/** Returns a seeded random float in the range [0,1), using the seed from SRandInit(). */
-	static CORE_API float SRand();
+	[[nodiscard]] static CORE_API float SRand();
 
 	/**
 	 * Computes the base 2 logarithm for an integer value.
@@ -646,7 +646,7 @@ struct FGenericPlatformMath
 	 * @param Value		The value to compute the log of
 	 * @return			Log2 of Value. 0 if Value is 0.
 	 */	
-	static constexpr FORCEINLINE uint32 FloorLog2(uint32 Value) 
+	[[nodiscard]] static constexpr FORCEINLINE uint32 FloorLog2(uint32 Value) 
 	{
 		uint32 pos = 0;
 		if (Value >= 1<<16) { Value >>= 16; pos += 16; }
@@ -658,7 +658,7 @@ struct FGenericPlatformMath
 	}
 
 	/** FloorLog2 but the caller guarantees that Value is not 0. */
-	static constexpr FORCEINLINE uint32 FloorLog2NonZero(uint32 Value)
+	[[nodiscard]] static constexpr FORCEINLINE uint32 FloorLog2NonZero(uint32 Value)
 	{
 		return FloorLog2(Value);
 	}
@@ -670,7 +670,7 @@ struct FGenericPlatformMath
 	 * @param Value		The value to compute the log of
 	 * @return			Log2 of Value. 0 if Value is 0.
 	 */	
-	static constexpr FORCEINLINE uint64 FloorLog2_64(uint64 Value) 
+	[[nodiscard]] static constexpr FORCEINLINE uint64 FloorLog2_64(uint64 Value) 
 	{
 		uint64 pos = 0;
 		if (Value >= 1ull<<32) { Value >>= 32; pos += 32; }
@@ -683,7 +683,7 @@ struct FGenericPlatformMath
 	}
 
 	/** FloorLog2_64 but the caller guarantees that Value is not 0. */
-	static constexpr FORCEINLINE uint64 FloorLog2NonZero_64(uint64 Value)
+	[[nodiscard]] static constexpr FORCEINLINE uint64 FloorLog2NonZero_64(uint64 Value)
 	{
 		return FloorLog2_64(Value);
 	}
@@ -695,7 +695,7 @@ struct FGenericPlatformMath
 	 *
 	 * @return the number of zeros before the first "on" bit
 	 */
-	static constexpr FORCEINLINE uint8 CountLeadingZeros8(uint8 Value)
+	[[nodiscard]] static constexpr FORCEINLINE uint8 CountLeadingZeros8(uint8 Value)
 	{
 		if (Value == 0) return 8;
 		return uint8(7 - FloorLog2(uint32(Value)));
@@ -708,7 +708,7 @@ struct FGenericPlatformMath
 	 *
 	 * @return the number of zeros before the first "on" bit
 	 */
-	static constexpr FORCEINLINE uint32 CountLeadingZeros(uint32 Value)
+	[[nodiscard]] static constexpr FORCEINLINE uint32 CountLeadingZeros(uint32 Value)
 	{
 		if (Value == 0) return 32;
 		return 31 - FloorLog2(Value);
@@ -721,7 +721,7 @@ struct FGenericPlatformMath
 	 *
 	 * @return the number of zeros before the first "on" bit
 	 */
-	static constexpr FORCEINLINE uint64 CountLeadingZeros64(uint64 Value)
+	[[nodiscard]] static constexpr FORCEINLINE uint64 CountLeadingZeros64(uint64 Value)
 	{
 		if (Value == 0) return 64;
 		return 63 - FloorLog2_64(Value);
@@ -734,7 +734,7 @@ struct FGenericPlatformMath
 	 *
 	 * @return the number of zeros after the last "on" bit
 	 */
-	static constexpr FORCEINLINE uint32 CountTrailingZeros(uint32 Value)
+	[[nodiscard]] static constexpr FORCEINLINE uint32 CountTrailingZeros(uint32 Value)
 	{
 		if (Value == 0)
 		{
@@ -756,7 +756,7 @@ struct FGenericPlatformMath
 	 *
 	 * @return the number of zeros after the last "on" bit
 	 */
-	static constexpr FORCEINLINE uint64 CountTrailingZeros64(uint64 Value)
+	[[nodiscard]] static constexpr FORCEINLINE uint64 CountTrailingZeros64(uint64 Value)
 	{
 		if (Value == 0)
 		{
@@ -775,14 +775,14 @@ struct FGenericPlatformMath
 	 * Returns smallest N such that (1<<N)>=Arg.
 	 * Note: CeilLogTwo(0)=0 
 	 */
-	static constexpr FORCEINLINE uint32 CeilLogTwo( uint32 Arg )
+	[[nodiscard]] static constexpr FORCEINLINE uint32 CeilLogTwo( uint32 Arg )
 	{
 		// if Arg is 0, change it to 1 so that we return 0
 		Arg = Arg ? Arg : 1;
 		return 32 - CountLeadingZeros(Arg - 1);
 	}
 
-	static constexpr FORCEINLINE uint64 CeilLogTwo64( uint64 Arg )
+	[[nodiscard]] static constexpr FORCEINLINE uint64 CeilLogTwo64( uint64 Arg )
 	{
 		// if Arg is 0, change it to 1 so that we return 0
 		Arg = Arg ? Arg : 1;
@@ -793,7 +793,7 @@ struct FGenericPlatformMath
 	 * Returns the smallest N such that (1<<N)>=Arg. This is a less efficient version of CeilLogTwo, but written in a
 	 * way that can be evaluated at compile-time.
 	 */
-	static constexpr FORCEINLINE uint8 ConstExprCeilLogTwo(SIZE_T Arg)
+	[[nodiscard]] static constexpr FORCEINLINE uint8 ConstExprCeilLogTwo(SIZE_T Arg)
 	{
 		if (Arg <= 1)
 		{
@@ -808,18 +808,18 @@ struct FGenericPlatformMath
 	}
 
 	/** @return Rounds the given number up to the next highest power of two. */
-	static constexpr FORCEINLINE uint32 RoundUpToPowerOfTwo(uint32 Arg)
+	[[nodiscard]] static constexpr FORCEINLINE uint32 RoundUpToPowerOfTwo(uint32 Arg)
 	{
 		return 1 << CeilLogTwo(Arg);
 	}
 
-	static constexpr FORCEINLINE uint64 RoundUpToPowerOfTwo64(uint64 V)
+	[[nodiscard]] static constexpr FORCEINLINE uint64 RoundUpToPowerOfTwo64(uint64 V)
 	{
 		return uint64(1) << CeilLogTwo64(V);
 	}
 
 	/** Spreads bits to every other. */
-	static constexpr FORCEINLINE uint32 MortonCode2( uint32 x )
+	[[nodiscard]] static constexpr FORCEINLINE uint32 MortonCode2( uint32 x )
 	{
 		x &= 0x0000ffff;
 		x = (x ^ (x << 8)) & 0x00ff00ff;
@@ -829,7 +829,7 @@ struct FGenericPlatformMath
 		return x;
 	}
 
-	static constexpr FORCEINLINE uint64 MortonCode2_64( uint64 x )
+	[[nodiscard]] static constexpr FORCEINLINE uint64 MortonCode2_64( uint64 x )
 	{
 		x &= 0x00000000ffffffff;
 		x = (x ^ (x << 16)) & 0x0000ffff0000ffff;
@@ -841,7 +841,7 @@ struct FGenericPlatformMath
 	}
 
 	/** Reverses MortonCode2. Compacts every other bit to the right. */
-	static constexpr FORCEINLINE uint32 ReverseMortonCode2( uint32 x )
+	[[nodiscard]] static constexpr FORCEINLINE uint32 ReverseMortonCode2( uint32 x )
 	{
 		x &= 0x55555555;
 		x = (x ^ (x >> 1)) & 0x33333333;
@@ -851,7 +851,7 @@ struct FGenericPlatformMath
 		return x;
 	}
 
-	static constexpr FORCEINLINE uint64 ReverseMortonCode2_64( uint64 x )
+	[[nodiscard]] static constexpr FORCEINLINE uint64 ReverseMortonCode2_64( uint64 x )
 	{
 		x &= 0x5555555555555555;
 		x = (x ^ (x >> 1)) & 0x3333333333333333;
@@ -863,7 +863,7 @@ struct FGenericPlatformMath
 	}
 
 	/** Spreads bits to every 3rd. */
-	static constexpr FORCEINLINE uint32 MortonCode3( uint32 x )
+	[[nodiscard]] static constexpr FORCEINLINE uint32 MortonCode3( uint32 x )
 	{
 		x &= 0x000003ff;
 		x = (x ^ (x << 16)) & 0xff0000ff;
@@ -874,7 +874,7 @@ struct FGenericPlatformMath
 	}
 
 	/** Reverses MortonCode3. Compacts every 3rd bit to the right. */
-	static constexpr FORCEINLINE uint32 ReverseMortonCode3( uint32 x )
+	[[nodiscard]] static constexpr FORCEINLINE uint32 ReverseMortonCode3( uint32 x )
 	{
 		x &= 0x09249249;
 		x = (x ^ (x >>  2)) & 0x030c30c3;
@@ -898,7 +898,7 @@ struct FGenericPlatformMath
 	 *
 	 * @return	ValueGEZero if Comparand >= 0, ValueLTZero otherwise
 	 */
-	static constexpr FORCEINLINE float FloatSelect( float Comparand, float ValueGEZero, float ValueLTZero )
+	[[nodiscard]] static constexpr FORCEINLINE float FloatSelect( float Comparand, float ValueGEZero, float ValueLTZero )
 	{
 		return Comparand >= 0.f ? ValueGEZero : ValueLTZero;
 	}
@@ -917,28 +917,28 @@ struct FGenericPlatformMath
 	 *
 	 * @return	ValueGEZero if Comparand >= 0, ValueLTZero otherwise
 	 */
-	static constexpr FORCEINLINE double FloatSelect( double Comparand, double ValueGEZero, double ValueLTZero )
+	[[nodiscard]] static constexpr FORCEINLINE double FloatSelect( double Comparand, double ValueGEZero, double ValueLTZero )
 	{
 		return Comparand >= 0.0 ? ValueGEZero : ValueLTZero;
 	}
 
 	/** Computes absolute value in a generic way */
 	template< class T > 
-	static constexpr FORCEINLINE T Abs( const T A )
+	[[nodiscard]] static constexpr FORCEINLINE T Abs( const T A )
 	{
 		return (A < (T)0) ? -A : A;
 	}
 
 	/** Returns 1, 0, or -1 depending on relation of T to 0 */
 	template< class T > 
-	static constexpr FORCEINLINE T Sign( const T A )
+	[[nodiscard]] static constexpr FORCEINLINE T Sign( const T A )
 	{
 		return (A > (T)0) ? (T)1 : ((A < (T)0) ? (T)-1 : (T)0);
 	}
 
 	/** Returns higher value in a generic way */
 	template <typename T> 
-	static constexpr FORCEINLINE T Max(T A, T B)
+	[[nodiscard]] static constexpr FORCEINLINE T Max(T A, T B)
 	{
 		// Even though this should be covered by the variadic, we still need this overload because of the many instances of
 		// FMath::Max<T>(a, b) with an explicit template parameter, which needs to continue to be supported.
@@ -949,7 +949,7 @@ struct FGenericPlatformMath
 		typename T, typename... OtherTypes
 		UE_REQUIRES((std::is_same_v<T, OtherTypes> && ...))
 	>
-	static constexpr FORCEINLINE T Max(T A, OtherTypes... Others)
+	[[nodiscard]] static constexpr FORCEINLINE T Max(T A, OtherTypes... Others)
 	{
 		if constexpr (sizeof...(OtherTypes) > 0)
 		{
@@ -960,7 +960,7 @@ struct FGenericPlatformMath
 
 	/** Returns lower value in a generic way */
 	template <typename T> 
-	static constexpr FORCEINLINE T Min(T A, T B)
+	[[nodiscard]] static constexpr FORCEINLINE T Min(T A, T B)
 	{
 		// Even though this should be covered by the variadic, we still need this overload because of the many instances of
 		// FMath::Min<T>(a, b) with an explicit template parameter, which needs to continue to be supported.
@@ -971,7 +971,7 @@ struct FGenericPlatformMath
 		typename T, typename... OtherTypes
 		UE_REQUIRES((std::is_same_v<T, OtherTypes> && ...))
 	>
-	static constexpr FORCEINLINE T Min(T A, OtherTypes... Others)
+	[[nodiscard]] static constexpr FORCEINLINE T Min(T A, OtherTypes... Others)
 	{
 		if constexpr (sizeof...(OtherTypes) > 0)
 		{
@@ -990,12 +990,37 @@ struct FGenericPlatformMath
 
 	/**
 	* Min of Array
-	* @param	Array of templated type
-	* @param	Optional pointer for returning the index of the minimum element, if multiple minimum elements the first index is returned
+	* @param	Values Array of templated type
 	* @return	The min value found in the array or default value if the array was empty
 	*/
 	template< class T >
-	static FORCEINLINE T Min(const TArray<T>& Values, int32* MinIndex = NULL)
+	[[nodiscard]] static FORCEINLINE T Min(const TArray<T>& Values)
+	{
+		if (Values.Num() == 0)
+		{
+			return T();
+		}
+
+		T CurMin = Values[0];
+		for (int32 v = 1; v < Values.Num(); ++v)
+		{
+			const T Value = Values[v];
+			if (Value < CurMin)
+			{
+				CurMin = Value;
+			}
+		}
+		return CurMin;
+	}
+	
+	/**
+	* Min of Array
+	* @param	Values Array of templated type
+	* @param	MinIndex Optional pointer for returning the index of the minimum element, if multiple minimum elements the first index is returned
+	* @return	The min value found in the array or default value if the array was empty
+	*/
+	template< class T >
+	static FORCEINLINE T Min(const TArray<T>& Values, int32* MinIndex)
 	{
 		if (Values.Num() == 0)
 		{
@@ -1023,16 +1048,42 @@ struct FGenericPlatformMath
 			*MinIndex = CurMinIndex;
 		}
 		return CurMin;
+	}	
+	
+	/**
+* Max of Array
+* @param	Values Array of templated type
+* @return	The max value found in the array or default value if the array was empty
+*/
+	template< class T >
+	[[nodiscard]] static FORCEINLINE T Max(const TArray<T>& Values)
+	{
+		if (Values.Num() == 0)
+		{
+			return T();
+		}
+
+		T CurMax = Values[0];
+		for (int32 v = 1; v < Values.Num(); ++v)
+		{
+			const T Value = Values[v];
+			if (CurMax < Value)
+			{
+				CurMax = Value;
+			}
+		}
+
+		return CurMax;
 	}
 
 	/**
 	* Max of Array
-	* @param	Array of templated type
-	* @param	Optional pointer for returning the index of the maximum element, if multiple maximum elements the first index is returned
+	* @param	Values Array of templated type
+	* @param	MaxIndex Optional pointer for returning the index of the maximum element, if multiple maximum elements the first index is returned
 	* @return	The max value found in the array or default value if the array was empty
 	*/
 	template< class T >
-	static FORCEINLINE T Max(const TArray<T>& Values, int32* MaxIndex = NULL)
+	static FORCEINLINE T Max(const TArray<T>& Values, int32* MaxIndex)
 	{
 		if (Values.Num() == 0)
 		{
@@ -1062,7 +1113,7 @@ struct FGenericPlatformMath
 		return CurMax;
 	}
 
-	static constexpr FORCEINLINE int32 CountBits(uint64 Bits)
+	[[nodiscard]] static constexpr FORCEINLINE int32 CountBits(uint64 Bits)
 	{
 		// https://en.wikipedia.org/wiki/Hamming_weight
 		Bits -= (Bits >> 1) & 0x5555555555555555ull;
@@ -1085,7 +1136,7 @@ struct FGenericPlatformMath
 		typename IntType
 		UE_REQUIRES(std::is_integral_v<IntType>)
 	>
-	static FORCEINLINE bool AddAndCheckForOverflow(IntType A, IntType B, IntType& OutResult)
+	[[nodiscard]] static FORCEINLINE bool AddAndCheckForOverflow(IntType A, IntType B, IntType& OutResult)
 	{
 		// This follows Hacker's Delight, Chapter 2-12
 		// Signed->unsigned conversion and unsigned addition have defined behavior always
@@ -1135,7 +1186,7 @@ struct FGenericPlatformMath
 		typename IntType
 		UE_REQUIRES(std::is_integral_v<IntType>)
 	>
-	static FORCEINLINE bool SubtractAndCheckForOverflow(IntType A, IntType B, IntType& OutResult)
+	[[nodiscard]] static FORCEINLINE bool SubtractAndCheckForOverflow(IntType A, IntType B, IntType& OutResult)
 	{
 		// This follows Hacker's Delight, Chapter 2-12
 		// Signed->unsigned conversion and unsigned addition have defined behavior always
@@ -1179,7 +1230,7 @@ struct FGenericPlatformMath
 		typename IntType
 		UE_REQUIRES(std::is_integral_v<IntType>)
 	>
-	static FORCEINLINE bool MultiplyAndCheckForOverflow(IntType A, IntType B, IntType& OutResult)
+	[[nodiscard]] static FORCEINLINE bool MultiplyAndCheckForOverflow(IntType A, IntType B, IntType& OutResult)
 	{
 		// Handle the case where the second factor is 0 specially (why will become clear in a minute).
 		if (B == 0)
@@ -1253,12 +1304,12 @@ private:
 
 /** Float specialization */
 template<>
-FORCEINLINE float FGenericPlatformMath::Abs( const float A )
+[[nodiscard]] FORCEINLINE float FGenericPlatformMath::Abs( const float A )
 {
 	return fabsf( A );
 }
 template<>
-FORCEINLINE double FGenericPlatformMath::Abs(const double A)
+[[nodiscard]] FORCEINLINE double FGenericPlatformMath::Abs(const double A)
 {
 	return fabs(A);
 }
\ No newline at end of file
