OodleData: Downgrade Error to Display verbosity in diagnostic logs that are dependent on the compressed stream

diff --git a/Engine/Source/Runtime/Core/Private/Compression/OodleDataCompression.cpp b/Engine/Source/Runtime/Core/Private/Compression/OodleDataCompression.cpp
--- a/Engine/Source/Runtime/Core/Private/Compression/OodleDataCompression.cpp
+++ b/Engine/Source/Runtime/Core/Private/Compression/OodleDataCompression.cpp
@@ -348,12 +348,12 @@ struct OodleScratchBuffers
 
 		if ( InCompBufSize <= 0 )
 		{
-			UE_LOG(OodleDataCompression, Error, TEXT("OodleDecode: no compressed buffer size to decode!"));
+			UE_LOG(OodleDataCompression, Display, TEXT("OodleDecode: no compressed buffer size to decode!"));
 			return OODLELZ_FAILED;
 		}
 		else if ( static_cast<const OO_U8 *>(InCompBuf)[0] == 0 )
 		{
-			UE_LOG(OodleDataCompression, Error, TEXT("OodleDecode: compressed buffer starts with zero byte; invalid or corrupt compressed stream."));
+			UE_LOG(OodleDataCompression, Display, TEXT("OodleDecode: compressed buffer starts with zero byte; invalid or corrupt compressed stream."));
 			return OODLELZ_FAILED;
 		}
 					
@@ -394,14 +394,14 @@ struct OodleScratchBuffers
 		OodleLZ_Compressor CurCompressor = OodleLZ_GetChunkCompressor(InCompBuf, InCompBufSize, NULL);
 		if( CurCompressor == OodleLZ_Compressor_Invalid )
 		{
-			UE_LOG(OodleDataCompression, Error, TEXT("OodleDecode : no Oodle compressor found!; likely an invalid or corrupt compressed stream."));
+			UE_LOG(OodleDataCompression, Display, TEXT("OodleDecode : no Oodle compressor found!; likely an invalid or corrupt compressed stream."));
 			return OODLELZ_FAILED;
 		}
 
 		OO_SINTa DecoderMemorySize = OodleLZDecoder_MemorySizeNeeded(CurCompressor, InRawLen);
 		if ( DecoderMemorySize <= 0 )
 		{
-			UE_LOG(OodleDataCompression, Error, TEXT("OodleDecode : MemorySizeNeeded invalid; likely an invalid or corrupt compressed stream."));
+			UE_LOG(OodleDataCompression, Display, TEXT("OodleDecode : MemorySizeNeeded invalid; likely an invalid or corrupt compressed stream."));
 			return OODLELZ_FAILED;
 		}
 
@@ -799,7 +799,7 @@ bool DecompressParallel(
 
 	if ( ! OodleLZ_FillSeekTable(SeekTable,OodleLZSeekTable_Flags_None,ChunkLen,OutUncompressedData,UncompressedSize,InCompressedData,CompressedSize) )
 	{
-		UE_LOG(OodleDataCompression,Error,TEXT("OodleLZ_FillSeekTable failed\n"));
+		UE_LOG(OodleDataCompression,Display,TEXT("OodleLZ_FillSeekTable failed\n"));
 		return false;
 	}
 
