[AutoRTFM] Allow IoDispatcher::DoesChunkExist to work inside of a transaction.

This is used to check for things like whether or not a package is loaded, and is reachable from any thread.

diff --git a/Engine/Source/Runtime/Core/Private/IO/IoDispatcher.cpp b/Engine/Source/Runtime/Core/Private/IO/IoDispatcher.cpp
--- a/Engine/Source/Runtime/Core/Private/IO/IoDispatcher.cpp
+++ b/Engine/Source/Runtime/Core/Private/IO/IoDispatcher.cpp
@@ -470,7 +470,7 @@ public:
 		}
 	}
 
-	bool DoesChunkExist(const FIoChunkId& ChunkId) const
+	UE_AUTORTFM_ALWAYS_OPEN bool DoesChunkExist(const FIoChunkId& ChunkId) const
 	{
 		FReadScopeLock _(BackendsLock);
 		for (const FBackendAndPriority& Backend : Backends)
@@ -483,7 +483,7 @@ public:
 		return false;
 	}
 
-	bool DoesChunkExist(const FIoChunkId& ChunkId, const FIoOffsetAndLength& ChunkRange) const
+	UE_AUTORTFM_ALWAYS_OPEN bool DoesChunkExist(const FIoChunkId& ChunkId, const FIoOffsetAndLength& ChunkRange) const
 	{
 		FReadScopeLock _(BackendsLock);
 		for (const FBackendAndPriority& Backend : Backends)
diff --git a/Engine/Source/Runtime/CoreUObject/Private/Serialization/TransactionallySafeAsyncLoading.cpp b/Engine/Source/Runtime/CoreUObject/Private/Serialization/TransactionallySafeAsyncLoading.cpp
--- a/Engine/Source/Runtime/CoreUObject/Private/Serialization/TransactionallySafeAsyncLoading.cpp
+++ b/Engine/Source/Runtime/CoreUObject/Private/Serialization/TransactionallySafeAsyncLoading.cpp
@@ -35,7 +35,7 @@ public:
 
 	bool ShouldAlwaysLoadPackageAsync(const FPackagePath& PackagePath) override
 	{
-		return WrappedPackageLoader->ShouldAlwaysLoadPackageAsync(PackagePath);
+		return AutoRTFM::IsClosed() || WrappedPackageLoader->ShouldAlwaysLoadPackageAsync(PackagePath);
 	}
 
 	int32 LoadPackage(
