PR #12545: Fix incorrect paired comments for DO_CHECK and DO_ENSURE in AssertionMacros.h

#rnx

diff --git a/Engine/Source/Runtime/Core/Public/Misc/AssertionMacros.h b/Engine/Source/Runtime/Core/Public/Misc/AssertionMacros.h
--- a/Engine/Source/Runtime/Core/Public/Misc/AssertionMacros.h
+++ b/Engine/Source/Runtime/Core/Public/Misc/AssertionMacros.h
@@ -309,7 +309,7 @@ RetType UE_COLD UE_DEBUG_SECTION DispatchCheckVerify (InnerType&& Inner, ArgType
 	#define unimplemented()		check(!"Unimplemented function called")
 #endif
 
-#else
+#else // DO_CHECK
 	#define checkCode(...)
 	#define check(expr)					{ CA_ASSUME(expr); }
 	#define checkf(expr, format, ...)	{ CA_ASSUME(expr); }
@@ -319,7 +319,7 @@ RetType UE_COLD UE_DEBUG_SECTION DispatchCheckVerify (InnerType&& Inner, ArgType
 	#define verify(expr)				{ if(UNLIKELY(!(expr))){ CA_ASSUME(false); } }
 	#define verifyf(expr, format, ...)	{ if(UNLIKELY(!(expr))){ CA_ASSUME(false); } }
 	#define unimplemented()				{ CA_ASSUME(false); }
-#endif
+#endif // DO_CHECK
 
 //
 // Check for development only.
@@ -430,14 +430,14 @@ RetType UE_COLD UE_DEBUG_SECTION DispatchCheckVerify (InnerType&& Inner, ArgType
 	#define ensureAlwaysMsgf( InExpression, InFormat, ... ) UE_ENSURE_IMPL2(&, true,  InExpression, InFormat, ##__VA_ARGS__)
 
 
-#else	// DO_ENSURE
+#else // DO_ENSURE && !USING_CODE_ANALYSIS
 
 	#define ensure(           InExpression                ) (LIKELY(!!(InExpression)))
 	#define ensureMsgf(       InExpression, InFormat, ... ) (LIKELY(!!(InExpression)))
 	#define ensureAlways(     InExpression                ) (LIKELY(!!(InExpression)))
 #define ensureAlwaysMsgf( InExpression, InFormat, ... ) (LIKELY(!!(InExpression)))
 
-#endif	// DO_CHECK
+#endif	// DO_ENSURE && !USING_CODE_ANALYSIS
 
 namespace UEAsserts_Private
 {
