Added constraints to TVariant's constructor and Emplace function to ensure constructibility.

diff --git a/Engine/Source/Runtime/Core/Public/Misc/TVariant.h b/Engine/Source/Runtime/Core/Public/Misc/TVariant.h
--- a/Engine/Source/Runtime/Core/Public/Misc/TVariant.h
+++ b/Engine/Source/Runtime/Core/Public/Misc/TVariant.h
@@ -58,13 +58,17 @@ public:
 	}
 
 	/** Perform in-place construction of a type into the variant */
-	template <typename U, typename... TArgs>
-	explicit TVariant(TInPlaceType<U>&&, TArgs&&... Args)
+	template <
+		typename U,
+		typename... ArgTypes
+		UE_REQUIRES(std::is_constructible_v<U, ArgTypes...>)
+	>
+	explicit TVariant(TInPlaceType<U>&&, ArgTypes&&... Args)
 	{
 		constexpr SIZE_T Index = UE::Core::Private::TParameterPackTypeIndex<U, T, Ts...>::Value;
 		static_assert(Index != (SIZE_T)-1, "The TVariant is not declared to hold the type being constructed");
 
-		::new((void*)&UE::Core::Private::CastToStorage(*this).Storage) U(Forward<TArgs>(Args)...);
+		::new((void*)&UE::Core::Private::CastToStorage(*this).Storage) U(Forward<ArgTypes>(Args)...);
 		TypeIndex = (uint8)Index;
 	}
 
@@ -189,14 +193,18 @@ public:
 	}
 
 	/** Set a specifically-typed value into the variant using in-place construction */
-	template <typename U, typename... TArgs>
-	void Emplace(TArgs&&... Args)
+	template <
+		typename U,
+		typename... ArgTypes
+		UE_REQUIRES(std::is_constructible_v<U, ArgTypes...>)
+	>
+	void Emplace(ArgTypes&&... Args)
 	{
 		constexpr SIZE_T Index = UE::Core::Private::TParameterPackTypeIndex<U, T, Ts...>::Value;
 		static_assert(Index != (SIZE_T)-1, "The TVariant is not declared to hold the type passed to Emplace<>");
 
 		UE::Core::Private::TDestructorLookup<T, Ts...>::Destruct(TypeIndex, &UE::Core::Private::CastToStorage(*this).Storage);
-		::new((void*)&UE::Core::Private::CastToStorage(*this).Storage) U(Forward<TArgs>(Args)...);
+		::new((void*)&UE::Core::Private::CastToStorage(*this).Storage) U(Forward<ArgTypes>(Args)...);
 		TypeIndex = (uint8)Index;
 	}
 
