[AutoRTFM] Improve memory validation throttling algorithm and add new switch

Add `--AutoRTFMMemoryValidationThrottlingEnabled` switch to disable throttling.

There were several issues with the old implementation:
* Only profile the hashing of the inner most transaction. Many of the parent transactions had a null open return address which skewed the probability biases.
* There were several ways that an open hash probability could get stuck at 0%.
* Because the threshold adjustment was only updated every N'th hash-second, this could get stuck with no hashing at all (or at least for long periods where probabilties were all low)

Include the throttling factor in the log output.

diff --git a/Engine/Source/Programs/AutoRTFMTests/Private/AutoRTFMTests.cpp b/Engine/Source/Programs/AutoRTFMTests/Private/AutoRTFMTests.cpp
--- a/Engine/Source/Programs/AutoRTFMTests/Private/AutoRTFMTests.cpp
+++ b/Engine/Source/Programs/AutoRTFMTests/Private/AutoRTFMTests.cpp
@@ -101,8 +101,9 @@ int RunTests(int ArgC, const char* ArgV[])
 	// AUTORTFM_SCOPED_SET_ENSURE_ON_ABORT_BY_LANGUAGE(false)
 	AutoRTFM::ForTheRuntime::SetEnsureOnAbortByLanguage(true);
 
-	// Unexpected memory validation errors should be full assertions.
+	// Unexpected memory validation errors should be full assertions, without validation throttling.
 	AutoRTFM::ForTheRuntime::SetMemoryValidationLevel(AutoRTFM::EMemoryValidationLevel::Error);
+	AutoRTFM::ForTheRuntime::SetMemoryValidationThrottlingEnabled(false);
 
 	int Result = Session.run();
 
diff --git a/Engine/Source/Programs/AutoRTFMTests/Private/OpenHashThrottlerTests.cpp b/Engine/Source/Programs/AutoRTFMTests/Private/OpenHashThrottlerTests.cpp
--- a/Engine/Source/Programs/AutoRTFMTests/Private/OpenHashThrottlerTests.cpp
+++ b/Engine/Source/Programs/AutoRTFMTests/Private/OpenHashThrottlerTests.cpp
@@ -9,8 +9,8 @@ TEST_CASE("OpenHashThrottler")
 	using FSeconds = AutoRTFM::FOpenHashThrottler::FSeconds;
 	
 	AutoRTFM::FOpenHashThrottler Throttler{
-		/* LogInterval            */ 1e10, // Not testing logging
-		/* AdjustThrottleInterval */ 0.1,  // Adjust probabilities every 100ms spent hashing.
+		/* LogInterval            */ 1e10, // Not testing logging.
+		/* AdjustThrottleInterval */ 0.1,  // Adjust probabilities every 100ms.
 		/* TargetFractionHashing  */ 0.1   // At most we want to spent 10% of the time hashing.
 	};
 
@@ -25,9 +25,17 @@ TEST_CASE("OpenHashThrottler")
 	auto Hash = [&](FSeconds DurationNotHashing, FSeconds DurationHashing, const void* OpenReturnAddress)
 	{
 		Time += DurationNotHashing;
-		Throttler.Update(Time, Time + DurationHashing, OpenReturnAddress, WriteLog);
+		Throttler.OnHash(Time, Time + DurationHashing, OpenReturnAddress, WriteLog);
 		Time += DurationHashing;
 	};
+	
+	FSeconds LastUpdate = 0;
+	auto Update = [&](FSeconds Duration = 0)
+	{
+		Time += Duration;
+		Throttler.Update(Time - LastUpdate);
+		LastUpdate = Time;
+	};
 
 	SECTION("Initial state")
 	{
@@ -41,6 +49,7 @@ TEST_CASE("OpenHashThrottler")
 		Hash(/* DurationNotHashing */ 0, /* DurationHashing */ 0.01, OpenReturnAddressA);
 		Hash(/* DurationNotHashing */ 0, /* DurationHashing */ 0.01, OpenReturnAddressB);
 		Hash(/* DurationNotHashing */ 0, /* DurationHashing */ 0.01, OpenReturnAddressC);
+		Update();
 		REQUIRE(Throttler.HashProbabilityFor(OpenReturnAddressA) == 1.0);
 		REQUIRE(Throttler.HashProbabilityFor(OpenReturnAddressB) == 1.0);
 		REQUIRE(Throttler.HashProbabilityFor(OpenReturnAddressC) == 1.0);
@@ -54,6 +63,7 @@ TEST_CASE("OpenHashThrottler")
 			Hash(/* DurationNotHashing */ 0.099, /* DurationHashing */ 0.001, OpenReturnAddressA);
 			Hash(/* DurationNotHashing */ 0.099, /* DurationHashing */ 0.001, OpenReturnAddressB);
 			Hash(/* DurationNotHashing */ 0.099, /* DurationHashing */ 0.001, OpenReturnAddressC);
+			Update();
 		}
 		REQUIRE(Throttler.HashProbabilityFor(OpenReturnAddressA) == 1.0);
 		REQUIRE(Throttler.HashProbabilityFor(OpenReturnAddressB) == 1.0);
@@ -68,6 +78,7 @@ TEST_CASE("OpenHashThrottler")
 			Hash(/* DurationNotHashing */ 0.08, /* DurationHashing */ 0.02, OpenReturnAddressA);
 			Hash(/* DurationNotHashing */ 0.08, /* DurationHashing */ 0.02, OpenReturnAddressB);
 			Hash(/* DurationNotHashing */ 0.08, /* DurationHashing */ 0.02, OpenReturnAddressC);
+			Update();
 		}
 		REQUIRE(Throttler.HashProbabilityFor(OpenReturnAddressA) < 1.0);
 		REQUIRE(Throttler.HashProbabilityFor(OpenReturnAddressA) > 0.0);
@@ -86,6 +97,7 @@ TEST_CASE("OpenHashThrottler")
 			Hash(/* DurationNotHashing */ 0.05, /* DurationHashing */ 0.05, OpenReturnAddressA);
 			Hash(/* DurationNotHashing */ 0.05, /* DurationHashing */ 0.03, OpenReturnAddressB);
 			Hash(/* DurationNotHashing */ 0.05, /* DurationHashing */ 0.01, OpenReturnAddressC);
+			Update();
 		}
 		REQUIRE(Throttler.HashProbabilityFor(OpenReturnAddressA) < 1.0);
 		REQUIRE(Throttler.HashProbabilityFor(OpenReturnAddressA) > 0.0);
@@ -103,15 +115,16 @@ TEST_CASE("OpenHashThrottler")
 		{
 			// Hash count ratios:
 			// OpenReturnAddressA:OpenReturnAddressB:OpenReturnAddressC - 5:3:1 
-			Hash(/* DurationNotHashing */ 0.05, /* DurationHashing */ 0.02, OpenReturnAddressA);
-			Hash(/* DurationNotHashing */ 0.05, /* DurationHashing */ 0.02, OpenReturnAddressB);
-			Hash(/* DurationNotHashing */ 0.05, /* DurationHashing */ 0.02, OpenReturnAddressA);
-			Hash(/* DurationNotHashing */ 0.05, /* DurationHashing */ 0.02, OpenReturnAddressC);
-			Hash(/* DurationNotHashing */ 0.05, /* DurationHashing */ 0.02, OpenReturnAddressA);
-			Hash(/* DurationNotHashing */ 0.05, /* DurationHashing */ 0.02, OpenReturnAddressB);
-			Hash(/* DurationNotHashing */ 0.05, /* DurationHashing */ 0.02, OpenReturnAddressA);
-			Hash(/* DurationNotHashing */ 0.05, /* DurationHashing */ 0.02, OpenReturnAddressB);
-			Hash(/* DurationNotHashing */ 0.05, /* DurationHashing */ 0.02, OpenReturnAddressA);
+			Hash(/* DurationNotHashing */ 0.05, /* DurationHashing */ 0.05, OpenReturnAddressA);
+			Hash(/* DurationNotHashing */ 0.05, /* DurationHashing */ 0.05, OpenReturnAddressB);
+			Hash(/* DurationNotHashing */ 0.05, /* DurationHashing */ 0.05, OpenReturnAddressA);
+			Hash(/* DurationNotHashing */ 0.05, /* DurationHashing */ 0.05, OpenReturnAddressC);
+			Hash(/* DurationNotHashing */ 0.05, /* DurationHashing */ 0.05, OpenReturnAddressA);
+			Hash(/* DurationNotHashing */ 0.05, /* DurationHashing */ 0.05, OpenReturnAddressB);
+			Hash(/* DurationNotHashing */ 0.05, /* DurationHashing */ 0.05, OpenReturnAddressA);
+			Hash(/* DurationNotHashing */ 0.05, /* DurationHashing */ 0.05, OpenReturnAddressB);
+			Hash(/* DurationNotHashing */ 0.05, /* DurationHashing */ 0.05, OpenReturnAddressA);
+			Update();
 		}
 		REQUIRE(Throttler.HashProbabilityFor(OpenReturnAddressA) < 1.0);
 		REQUIRE(Throttler.HashProbabilityFor(OpenReturnAddressA) > 0.0);
@@ -129,8 +142,10 @@ TEST_CASE("OpenHashThrottler")
 		for (int I = 0; I < 10; I++)
 		{
 			Hash(/* DurationNotHashing */ 0.05, /* DurationHashing */ 0.02, OpenReturnAddressA);
+			Update();
 		}
 
+		Update();
 		const double ProbabilityA = Throttler.HashProbabilityFor(OpenReturnAddressA);
 		REQUIRE(ProbabilityA < 1.0);
 
@@ -139,6 +154,7 @@ TEST_CASE("OpenHashThrottler")
 		{
 			Hash(/* DurationNotHashing */ 0.05, /* DurationHashing */ 0.01, OpenReturnAddressB);
 			Hash(/* DurationNotHashing */ 0.05, /* DurationHashing */ 0.01, OpenReturnAddressC);
+			Update();
 		}
 		
 		REQUIRE(Throttler.HashProbabilityFor(OpenReturnAddressB) < 1.0);
@@ -151,19 +167,22 @@ TEST_CASE("OpenHashThrottler")
 	{
 		for (int I = 0; I < 100; I++)
 		{
-			Hash(/* DurationNotHashing */ 0.05, /* DurationHashing */ 0.02, OpenReturnAddressA);
-			Hash(/* DurationNotHashing */ 0.05, /* DurationHashing */ 0.02, OpenReturnAddressB);
-			Hash(/* DurationNotHashing */ 0.05, /* DurationHashing */ 0.02, OpenReturnAddressC);
+			Hash(/* DurationNotHashing */ 0.05, /* DurationHashing */ 0.2, OpenReturnAddressA);
+			Hash(/* DurationNotHashing */ 0.05, /* DurationHashing */ 0.2, OpenReturnAddressB);
+			Hash(/* DurationNotHashing */ 0.05, /* DurationHashing */ 0.2, OpenReturnAddressC);
+			Update(1.0);
 		}
+
 		const double ProbabilityA = Throttler.HashProbabilityFor(OpenReturnAddressA);
 		const double ProbabilityB = Throttler.HashProbabilityFor(OpenReturnAddressB);
 		const double ProbabilityC = Throttler.HashProbabilityFor(OpenReturnAddressC);
 
 		for (int I = 0; I < 100; I++)
 		{
-			Hash(/* DurationNotHashing */ 0.05, /* DurationHashing */ 0.01, OpenReturnAddressA);
-			Hash(/* DurationNotHashing */ 0.05, /* DurationHashing */ 0.01, OpenReturnAddressB);
-			Hash(/* DurationNotHashing */ 0.05, /* DurationHashing */ 0.01, OpenReturnAddressC);
+			Hash(/* DurationNotHashing */ 0.05, /* DurationHashing */ 0.1, OpenReturnAddressA);
+			Hash(/* DurationNotHashing */ 0.05, /* DurationHashing */ 0.1, OpenReturnAddressB);
+			Hash(/* DurationNotHashing */ 0.05, /* DurationHashing */ 0.1, OpenReturnAddressC);
+			Update(1.0);
 		}
 
 		REQUIRE(Throttler.HashProbabilityFor(OpenReturnAddressA) > ProbabilityA);
@@ -175,22 +194,25 @@ TEST_CASE("OpenHashThrottler")
 	{
 		for (int I = 0; I < 100; I++)
 		{
-			Hash(/* DurationNotHashing */ 0.01, /* DurationHashing */ 0.01, OpenReturnAddressA);
-			Hash(/* DurationNotHashing */ 0.01, /* DurationHashing */ 0.01, OpenReturnAddressB);
-			Hash(/* DurationNotHashing */ 0.01, /* DurationHashing */ 0.01, OpenReturnAddressC);
-			Hash(/* DurationNotHashing */ 0.01, /* DurationHashing */ 0.01, OpenReturnAddressA);
-			Hash(/* DurationNotHashing */ 0.01, /* DurationHashing */ 0.01, OpenReturnAddressB);
-			Hash(/* DurationNotHashing */ 0.01, /* DurationHashing */ 0.01, OpenReturnAddressC);
+			Hash(/* DurationNotHashing */ 0.01, /* DurationHashing */ 0.1, OpenReturnAddressA);
+			Hash(/* DurationNotHashing */ 0.01, /* DurationHashing */ 0.1, OpenReturnAddressB);
+			Hash(/* DurationNotHashing */ 0.01, /* DurationHashing */ 0.1, OpenReturnAddressC);
+			Hash(/* DurationNotHashing */ 0.01, /* DurationHashing */ 0.1, OpenReturnAddressA);
+			Hash(/* DurationNotHashing */ 0.01, /* DurationHashing */ 0.1, OpenReturnAddressB);
+			Hash(/* DurationNotHashing */ 0.01, /* DurationHashing */ 0.1, OpenReturnAddressC);
+			Update(1.0);
 		}
+
 		const double ProbabilityA = Throttler.HashProbabilityFor(OpenReturnAddressA);
 		const double ProbabilityB = Throttler.HashProbabilityFor(OpenReturnAddressB);
 		const double ProbabilityC = Throttler.HashProbabilityFor(OpenReturnAddressC);
 
 		for (int I = 0; I < 100; I++)
 		{
-			Hash(/* DurationNotHashing */ 0.02, /* DurationHashing */ 0.01, OpenReturnAddressA);
-			Hash(/* DurationNotHashing */ 0.02, /* DurationHashing */ 0.01, OpenReturnAddressB);
-			Hash(/* DurationNotHashing */ 0.02, /* DurationHashing */ 0.01, OpenReturnAddressC);
+			Hash(/* DurationNotHashing */ 0.02, /* DurationHashing */ 0.1, OpenReturnAddressA);
+			Hash(/* DurationNotHashing */ 0.02, /* DurationHashing */ 0.1, OpenReturnAddressB);
+			Hash(/* DurationNotHashing */ 0.02, /* DurationHashing */ 0.1, OpenReturnAddressC);
+			Update(1.0);
 		}
 
 		REQUIRE(Throttler.HashProbabilityFor(OpenReturnAddressA) > ProbabilityA);
diff --git a/Engine/Source/Runtime/Core/Private/AutoRTFM/API.cpp b/Engine/Source/Runtime/Core/Private/AutoRTFM/API.cpp
--- a/Engine/Source/Runtime/Core/Private/AutoRTFM/API.cpp
+++ b/Engine/Source/Runtime/Core/Private/AutoRTFM/API.cpp
@@ -28,6 +28,8 @@ namespace
 	static constexpr AutoRTFM::EMemoryValidationLevel GDefaultMemoryValidationLevel = AutoRTFM::EMemoryValidationLevel::Disabled;
 
 	int GMemoryValidationLevel = static_cast<int>(GDefaultMemoryValidationLevel);
+	
+	bool GMemoryValidationThrottlingEnabled = true;
 
 	bool GAutoRTFMAssertOnError = true;
 
@@ -69,6 +71,13 @@ namespace
 			return FGenericCrashContext::SetGameData(TEXT("AutoRTFMMemoryValidationLevel"), TEXT("Enabled"));
 		}
 	}
+	
+	void OnAutoRTFMMemoryValidationThrottlingChanged()
+	{
+		const TCHAR* Text = GMemoryValidationThrottlingEnabled ? TEXT("true") : TEXT("false");
+		UE_LOG(LogCore, Log, TEXT("AutoRTFM Memory Validation: %s"), Text);
+		return FGenericCrashContext::SetGameData(TEXT("AutoRTFMMemoryValidationThrottlingEnabled"), Text);
+	}
 }
 
 #if UE_AUTORTFM
@@ -96,6 +105,14 @@ static FAutoConsoleVariableRef CVarAutoRTFMMemoryValidationLevel(
 	ECVF_Default
 );
 
+static FAutoConsoleVariableRef CVarAutoRTFMMemoryValidationThrottlingEnabled(
+	TEXT("AutoRTFMMemoryValidationThrottlingEnabled"),
+	GMemoryValidationThrottlingEnabled,
+	TEXT("Automatically skips memory validation on opens if validation is taking an excessive amount of time"),
+	FConsoleVariableDelegate::CreateLambda([](IConsoleVariable*) { OnAutoRTFMMemoryValidationThrottlingChanged(); }),
+	ECVF_Default
+);
+
 static FAutoConsoleVariableRef CVarAutoRTFMAssertOnError(
 	TEXT("AutoRTFMAssertOnError"),
 	GAutoRTFMAssertOnError,
@@ -108,6 +125,7 @@ static FDelayedAutoRegisterHelper DelayedAutoRegister(EDelayedRegisterRunPhase::
 		OnIsAutoRTFMEnabledChanged();
 		OnAutoRTFMRetryTransactionsChanged();
 		OnAutoRTFMMemoryValidationLevelChanged();
+		OnAutoRTFMMemoryValidationThrottlingChanged();
 	});
 #endif
 
@@ -263,6 +281,22 @@ namespace AutoRTFM
 #endif
 		}
 
+		bool GetMemoryValidationThrottlingEnabled()
+		{
+#if UE_AUTORTFM
+			return GMemoryValidationThrottlingEnabled;
+#else
+			return false;
+#endif
+		}
+
+		void SetMemoryValidationThrottlingEnabled(bool bEnabled)
+		{
+#if UE_AUTORTFM
+			GMemoryValidationThrottlingEnabled = bEnabled;
+#endif
+		}
+	
 		UE_AUTORTFM_ALWAYS_OPEN void DebugBreakIfMemoryValidationFails()
 		{
 #if UE_AUTORTFM
diff --git a/Engine/Source/Runtime/Core/Private/AutoRTFM/Context.h b/Engine/Source/Runtime/Core/Private/AutoRTFM/Context.h
--- a/Engine/Source/Runtime/Core/Private/AutoRTFM/Context.h
+++ b/Engine/Source/Runtime/Core/Private/AutoRTFM/Context.h
@@ -86,7 +86,7 @@ public:
 	// See FTransaction::ShouldRecordWrite()
 	inline const void* GetClosedStackAddress() const { return ClosedStackAddress; }
 
-	inline FOpenHashThrottler& GetTransactionHashProfiler() { return OpenHashThrottler; }
+	inline FOpenHashThrottler& GetOpenHashThrottler() { return OpenHashThrottler; }
 
     void DumpState() const;
 
@@ -165,9 +165,9 @@ private:
 	uint32 StackLocalInitializerDepth = 0;
 	TFunction<void()> CascadingRetryTransactionWork;
 	FOpenHashThrottler OpenHashThrottler{
-		/* LogInterval            */ 5,   // Log every 5 seconds spent hashing.
-		/* AdjustThrottleInterval */ 0.1, // Adjust probabilities every 100ms spent hashing.
-		/* TargetFractionHashing  */ 0.1  // At most we want to spent 10% of the time hashing.
+		/* LogInterval            */ 10,   // Log every 10 seconds.
+		/* AdjustThrottleInterval */ 0.25, // Adjust throttling probabilities every 250ms.
+		/* TargetFractionHashing  */ 0.1   // At most we want to spent 10% of the time hashing.
 	};
 };
 
diff --git a/Engine/Source/Runtime/Core/Private/AutoRTFM/OpenHashThrottler.cpp b/Engine/Source/Runtime/Core/Private/AutoRTFM/OpenHashThrottler.cpp
--- a/Engine/Source/Runtime/Core/Private/AutoRTFM/OpenHashThrottler.cpp
+++ b/Engine/Source/Runtime/Core/Private/AutoRTFM/OpenHashThrottler.cpp
@@ -20,6 +20,12 @@ FOpenHashThrottler::FHashScope::FHashScope(FOpenHashThrottler& Throttler, const
 {
 }
 
+FOpenHashThrottler::FHashScope::~FHashScope()
+{
+	const FSeconds EndTime = FPlatformTime::Seconds();
+	Throttler.OnHash(StartTime, EndTime, OpenReturnAddress, WriteLog);
+}
+
 FOpenHashThrottler::FOpenHashThrottler(FSeconds LogInterval, FSeconds AdjustThrottleInterval, FSeconds TargetFractionHashing)
 	: LogInterval(LogInterval)
 	, AdjustThrottleInterval(AdjustThrottleInterval)
@@ -27,22 +33,25 @@ FOpenHashThrottler::FOpenHashThrottler(FSeconds LogInterval, FSeconds AdjustThro
 {
 }
 
-FOpenHashThrottler::FHashScope::~FHashScope()
+void FOpenHashThrottler::OnHash(FSeconds StartTime, FSeconds EndTime, const void* OpenReturnAddress, const FWriteLog& WriteLog)
 {
-	const FSeconds EndTime = FPlatformTime::Seconds();
-	Throttler.Update(StartTime, EndTime, OpenReturnAddress, WriteLog);
-}
+	const FSeconds HashDuration = EndTime - StartTime;
+	FOpenAddressInfo& OpenInfo = ThrottlingData.Opens.FindOrAdd(OpenReturnAddress);
+	OpenInfo.TimeSpentHashing += HashDuration;
+	OpenInfo.bActive = true;
+	ThrottlingData.TotalTimeSpentHashing += HashDuration;
 
-void FOpenHashThrottler::Update(FSeconds StartTime, FSeconds EndTime, const void* OpenReturnAddress, const FWriteLog& WriteLog)
-{
-	UpdateThrottlingData(StartTime, EndTime, OpenReturnAddress);
-	UpdateLogStats(StartTime, EndTime, OpenReturnAddress, WriteLog);
+	LogStats.TimeSpentHashing += HashDuration;
+	LogStats.NumHashCalls++;
+	LogStats.NumBytesHashed += WriteLog.TotalSize();
+	LogStats.NumWriteRecords += WriteLog.Num();
 }
 
 double FOpenHashThrottler::HashProbabilityFor(const void* OpenReturnAddress)
 {
 	if (FOpenAddressInfo* Info = ThrottlingData.Opens.Find(OpenReturnAddress))
 	{
+		Info->bActive = true;
 		return Info->Probability;
 	}
 
@@ -52,93 +61,143 @@ double FOpenHashThrottler::HashProbabilityFor(const void* OpenReturnAddress)
 
 bool FOpenHashThrottler::ShouldHashFor(const void* OpenReturnAddress)
 {
-	return FMath::FRand() < HashProbabilityFor(OpenReturnAddress);
+	if (FMath::FRand() < HashProbabilityFor(OpenReturnAddress))
+	{
+		LogStats.NumShouldHashForTrue++;
+		return true;
+	}
+	LogStats.NumShouldHashForFalse++;
+	return false;
 }
 
-void FOpenHashThrottler::UpdateThrottlingData(FSeconds StartTime, FSeconds EndTime, const void* OpenReturnAddress)
+void FOpenHashThrottler::Update(FSeconds DeltaTime /* = 0 */)
 {
-	const FSeconds HashDuration = EndTime - StartTime;
-	ThrottlingData.Opens.FindOrAdd(OpenReturnAddress).TimeSpentHashing += HashDuration;
-	ThrottlingData.TotalTimeSpentHashing += HashDuration;
-	if (ThrottlingData.TotalTimeSpentHashing < AdjustThrottleInterval)
+	const FSeconds Now = FPlatformTime::Seconds();
+	if (DeltaTime == 0)
 	{
-		return;
+		DeltaTime = (LastUpdateTimestamp > 0) ? (Now - LastUpdateTimestamp) : 0;
 	}
 
-	// Count the number of opens that performed hashing this update.
+	ThrottlingData.TimeSinceLastUpdate += DeltaTime;
+	if (ThrottlingData.TimeSinceLastUpdate >= AdjustThrottleInterval)
+	{
+		UpdateThrottlingData();
+		ThrottlingData.TimeSinceLastUpdate = 0;
+	}
+
+	LogStats.TimeSinceLastReset += DeltaTime;
+	if (LogStats.TimeSinceLastReset >= LogInterval)
+	{
+		UpdateLogStats();
+		LogStats.TimeSinceLastReset = 0;
+	}
+
+	LastUpdateTimestamp = Now;
+}
+
+void FOpenHashThrottler::UpdateThrottlingData()
+{
+	// Count the number of opens that were active and those that performed hashing this update.
+	size_t NumActiveOpens = 0;
 	size_t NumOpensThatHashed = 0;
 	for (auto& It : ThrottlingData.Opens)
 	{
+		if (It.Value.bActive)
+		{
+			NumActiveOpens++;
+		}
 		if (It.Value.TimeSpentHashing > 0)
 		{
 			NumOpensThatHashed++;
 		}
 	}
-	ensure(NumOpensThatHashed > 0); // Otherwise how is ThrottlingData.TotalTimeSpentHashing non zero?
 
-	// Time in seconds since the last update of the throttling data.
-	const FSeconds TimeSinceLastUpdate = EndTime - ThrottlingData.Timestamp;
-	// Fraction of the time since the last update spent hashing.
-	const double TotalFractionHashing = ThrottlingData.TotalTimeSpentHashing / TimeSinceLastUpdate;
+	// Active is a superset of those that hashed.
+	ensure(NumActiveOpens >= NumOpensThatHashed);
+
+	if (NumActiveOpens == 0)
+	{
+		return; // No opens were queried or used this update.
+	}
+
 	// Probability multiplier to reach the target time spent hashing.
-	const double TotalGain = TargetFractionHashing / TotalFractionHashing;
+	const double TotalGain = TargetFractionHashing * ThrottlingData.TimeSinceLastUpdate / ThrottlingData.TotalTimeSpentHashing;
 
-	// Calculate the average amount of time spent hashing for all opens this update.
-	const double AverageOpenTimeHashing = ThrottlingData.TotalTimeSpentHashing / static_cast<double>(NumOpensThatHashed);
+	// Average time spent hashing per open
+	const double AverageTimeSpentHashingPerOpen =
+		ThrottlingData.TotalTimeSpentHashing / static_cast<double>(FMath::Max<size_t>(NumOpensThatHashed, 1));
 
 	// Apply the probability multiplier and normalize the time spent in each open.
 	for (auto& It : ThrottlingData.Opens)
 	{
-		if (It.Value.TimeSpentHashing > 0)
+		if (!It.Value.bActive)
+		{
+			continue;
+		}
+
+		// Calculate the target probability to normalize the time spent in each open, and
+		// to aim for the target total fractional time spent in hashing (TargetFractionHashing).
+
+		// Something non-zero, so we can scale probabilities back up with multiplications.
+		const double MinProbability = 1e-6;
+
+		// The time spent hashing this open this update
+		const FSeconds TimeSpentHashing = It.Value.TimeSpentHashing;
+		
+		// Target probability starts with the current probability
+		double TargetProbability = FMath::Max(It.Value.Probability, MinProbability);
+
+		// Adjust for the relative time spent hashing this open compared to the others.
+		if (TimeSpentHashing > 0)
 		{
-			// Calculate the target probability to normalize the time spent in each open, and
-			// to aim for the target total fractional time spent in hashing (TargetFractionHashing).
-			const double TargetProbability = FMath::Clamp(TotalGain * AverageOpenTimeHashing / It.Value.TimeSpentHashing, 0, 1);
-			// Interpolate towards the target probability at 50% per second.
-			const double InterpolationFactor = FMath::Pow(0.5, TimeSinceLastUpdate);
-			It.Value.Probability = FMath::Lerp(TargetProbability, It.Value.Probability, InterpolationFactor);
-			// Reset the TimeSpentHashing for this open.
-			It.Value.TimeSpentHashing = 0;
+			TargetProbability *= AverageTimeSpentHashingPerOpen / TimeSpentHashing;
 		}
+
+		// If the open was hashed this update, or probabilities are being raised
+		// then adjust by the total gain for all opens.
+		if (TimeSpentHashing > 0 || TotalGain > 1)
+		{
+			TargetProbability *= TotalGain;
+		}
+
+		// Finally clamp between 0..1
+		It.Value.Probability = FMath::Clamp(TargetProbability, 0, 1);
+
+		// Reset the TimeSpentHashing and bActive state for this open.
+		It.Value.TimeSpentHashing = 0;
+		It.Value.bActive = false;
 	}
 
 	ThrottlingData.TotalTimeSpentHashing = 0;
-	ThrottlingData.Timestamp = EndTime;
 }
 
-void FOpenHashThrottler::UpdateLogStats(FSeconds StartTime, FSeconds EndTime, const void* OpenReturnAddress, const FWriteLog& WriteLog)
+void FOpenHashThrottler::UpdateLogStats()
 {
-	const FSeconds HashDuration = EndTime - StartTime;
-	LogStats.TimeSpentHashing += HashDuration;
-	LogStats.NumHashCalls++;
-	LogStats.NumBytesHashed += WriteLog.TotalSize();
-	LogStats.NumWriteRecords += WriteLog.Num();
-	if (LogStats.TimeSpentHashing < LogInterval)
-	{
-		return;
-	}
-
-	const double TimeSinceLastLog = EndTime - LogStats.Timestamp;
+	const double TimeSinceLastLog = LogStats.TimeSinceLastReset;
 	UE_LOG(LogAutoRTFM, Display,
 		TEXT("Transaction Hash Statistics\n")
 		TEXT("-----------------------------\n")
-		/* a */ TEXT("%f / %f seconds spent hashing (%.3f%%)\n")
-		/* b */ TEXT(" * %zu hash calls (avr: %.3f calls/second)\n")
-		/* c */ TEXT(" * %zu bytes hashed (avr: %.3f bytes/call, %.3f bytes/second)\n")
-		/* d */ TEXT(" * %zu write records (avr: %.3f records/call)\n"),
+		/* a */ TEXT("%f / %f seconds spent hashing (%.1f%%, target: %.1f%%)\n")
+		/* b */ TEXT(" * %zu hash calls (avr: %.1f calls/second)\n")
+		/* c */ TEXT(" * %zu / %zu open validations skipped (%.1f%%)\n")
+		/* d */ TEXT(" * %zu bytes hashed (avr: %.1f bytes/call, %.1f bytes/second)\n")
+		/* e */ TEXT(" * %zu write records (avr: %.1f records/call)\n"),
 		/* a.0 */ LogStats.TimeSpentHashing,
 		/* a.1 */ TimeSinceLastLog,
 		/* a.2 */ 100 * LogStats.TimeSpentHashing / TimeSinceLastLog,
+		/* a.3 */ 100 * TargetFractionHashing,
 		/* b.0 */ LogStats.NumHashCalls,
 		/* b.1 */ static_cast<double>(LogStats.NumHashCalls) / TimeSinceLastLog,
-		/* c.0 */ LogStats.NumBytesHashed,
-		/* c.1 */ static_cast<double>(LogStats.NumBytesHashed) / static_cast<double>(LogStats.NumHashCalls),
-		/* c.2 */ static_cast<double>(LogStats.NumBytesHashed) / LogStats.TimeSpentHashing,
-		/* d.0 */ LogStats.NumWriteRecords,
-		/* d.1 */ static_cast<double>(LogStats.NumWriteRecords) / static_cast<double>(LogStats.NumHashCalls));
-	
+		/* c.0 */ LogStats.NumShouldHashForFalse,
+		/* c.1 */ LogStats.NumShouldHashForTrue + LogStats.NumShouldHashForFalse,
+		/* c.2 */ 100 * static_cast<double>(LogStats.NumShouldHashForFalse) / static_cast<double>(LogStats.NumShouldHashForTrue + LogStats.NumShouldHashForFalse),
+		/* d.0 */ LogStats.NumBytesHashed,
+		/* d.1 */ static_cast<double>(LogStats.NumBytesHashed) / static_cast<double>(LogStats.NumHashCalls),
+		/* d.2 */ static_cast<double>(LogStats.NumBytesHashed) / LogStats.TimeSpentHashing,
+		/* e.0 */ LogStats.NumWriteRecords,
+		/* e.1 */ static_cast<double>(LogStats.NumWriteRecords) / static_cast<double>(LogStats.NumHashCalls));
+		
 	LogStats = FLogStats{};
-	LogStats.Timestamp = EndTime;
 }
 
 } // namespace AutoRTFM
diff --git a/Engine/Source/Runtime/Core/Private/AutoRTFM/OpenHashThrottler.h b/Engine/Source/Runtime/Core/Private/AutoRTFM/OpenHashThrottler.h
--- a/Engine/Source/Runtime/Core/Private/AutoRTFM/OpenHashThrottler.h
+++ b/Engine/Source/Runtime/Core/Private/AutoRTFM/OpenHashThrottler.h
@@ -63,11 +63,15 @@ public:
 	// * HashStart is the wall-clock time after hashing ended
 	// * OpenReturnAddress is the return address for the call to AutoRTFM::Open(). Used to identify the open.
 	// * WriteLog is the write log that is being hashed. Used only for statistics logging.
-	void Update(FSeconds HashStart, FSeconds HashEnd, const void* OpenReturnAddress, const FWriteLog& WriteLog);
+	void OnHash(FSeconds HashStart, FSeconds HashEnd, const void* OpenReturnAddress, const FWriteLog& WriteLog);
+
+	// Periodically adjusts the probabilities for hashing opens, and prints statistics.
+	// * DeltaTime the time since the last call to Update(), or 0 to calculate using an internal clock.
+	void Update(FSeconds DeltaTime = 0);
 
 private:
-	void UpdateThrottlingData(FSeconds StartTime, FSeconds EndTime, const void* OpenReturnAddress);
-	void UpdateLogStats(FSeconds StartTime, FSeconds EndTime, const void* OpenReturnAddress, const FWriteLog& WriteLog);
+	void UpdateThrottlingData();
+	void UpdateLogStats();
 
 	struct FOpenAddressInfo
 	{
@@ -75,13 +79,15 @@ private:
 		FSeconds TimeSpentHashing = 0;
 		// Throttled probability to hash this open.
 		double Probability = 1;
+		// True if the open address was queried or hashed since last throttling update.
+		bool bActive = true;
 	};
 	
 	// Data used for throttling hashing
 	struct FThrottlingData
 	{
-		// Timestamp since this was last reset.
-		FSeconds Timestamp = 0;
+		// Timestamp since this was last updated.
+		FSeconds TimeSinceLastUpdate = 0;
 		// Total time spent hashing for all opens since last reset.
 		FSeconds TotalTimeSpentHashing = 0;
 		// A map of open return address to open info.
@@ -92,7 +98,7 @@ private:
 	struct FLogStats
 	{
 		// Timestamp since this was last reset.
-		FSeconds Timestamp = 0;
+		FSeconds TimeSinceLastReset = 0;
 		// Total time spent hashing since last reset.
 		FSeconds TimeSpentHashing = 0;
 		// Number of hash calls since last reset.
@@ -101,6 +107,10 @@ private:
 		size_t NumBytesHashed = 0;
 		// Number of write records hashed since last reset.
 		size_t NumWriteRecords = 0;
+		// Number of times ShouldHashFor() return true.
+		size_t NumShouldHashForTrue = 0;
+		// Number of times ShouldHashFor() return false.
+		size_t NumShouldHashForFalse = 0;
 	};
 
 	const FSeconds LogInterval;
@@ -109,6 +119,7 @@ private:
 
 	FThrottlingData ThrottlingData;
 	FLogStats LogStats;
+	FSeconds LastUpdateTimestamp = 0;
 };
 
 } // namespace AutoRTFM
diff --git a/Engine/Source/Runtime/Core/Private/AutoRTFM/Transaction.cpp b/Engine/Source/Runtime/Core/Private/AutoRTFM/Transaction.cpp
--- a/Engine/Source/Runtime/Core/Private/AutoRTFM/Transaction.cpp
+++ b/Engine/Source/Runtime/Core/Private/AutoRTFM/Transaction.cpp
@@ -1,8 +1,8 @@
 // Copyright Epic Games, Inc. All Rights Reserved.
 
 #if (defined(__AUTORTFM) && __AUTORTFM)
-#include "Transaction.h"
 
+#include "Transaction.h"
 #include "CallNestInlines.h"
 #include "GlobalData.h"
 #include "HAL/Platform.h"
@@ -299,12 +299,15 @@ void FTransaction::SetOpenActive(EMemoryValidationLevel NewMemoryValidationLevel
 	CurrentMemoryValidationLevel = NewMemoryValidationLevel;
 	CurrentOpenReturnAddress = ReturnAddress;
 
-	if (CurrentMemoryValidationLevel != EMemoryValidationLevel::Disabled)
+	if (CurrentMemoryValidationLevel != EMemoryValidationLevel::Disabled &&
+		AutoRTFM::ForTheRuntime::GetMemoryValidationThrottlingEnabled())
 	{
-		if (!Context->GetTransactionHashProfiler().ShouldHashFor(ReturnAddress))
+		FOpenHashThrottler& Throttler = Context->GetOpenHashThrottler();
+		if (!Throttler.ShouldHashFor(ReturnAddress))
 		{
 			CurrentMemoryValidationLevel = EMemoryValidationLevel::Disabled;
 		}
+		Throttler.Update();
 	}
 
 	if (CurrentMemoryValidationLevel == EMemoryValidationLevel::Disabled)
@@ -440,7 +443,7 @@ void FTransaction::DebugBreakIfMemoryValidationFails()
 	if (CurrentMemoryValidationLevel != EMemoryValidationLevel::Disabled)
 	{
 		FWriteHash OldHash = RecordedWriteHash;
-		FWriteHash NewHash = CalculateNestedWriteHash(NumWriteLogsHashed);
+		FWriteHash NewHash = CalculateNestedWriteHash();
 		if (OldHash != NewHash)
 		{
 			UE_LOG(LogAutoRTFM, Warning, TEXT("DebugBreakIfInvalidMemoryHash() detected a change in hash"));
@@ -452,13 +455,15 @@ void FTransaction::DebugBreakIfMemoryValidationFails()
 void FTransaction::RecordWriteHash()
 {
 	NumWriteLogsHashed = WriteLog.Num();
-	RecordedWriteHash = CalculateNestedWriteHash(NumWriteLogsHashed);
+	RecordedWriteHash = CalculateNestedWriteHash();
+	Context->GetOpenHashThrottler().Update();
 }
 
 void FTransaction::ValidateWriteHash() const
 {
 	FWriteHash OldHash = RecordedWriteHash;
-	FWriteHash NewHash = CalculateNestedWriteHash(NumWriteLogsHashed);
+	FWriteHash NewHash = CalculateNestedWriteHash();
+	Context->GetOpenHashThrottler().Update();
 
 	static constexpr TCHAR Message[] =
 		TEXT("Memory modified in a transaction was also modified in an call to AutoRTFM::Open(). ")
@@ -483,17 +488,21 @@ void FTransaction::ValidateWriteHash() const
 	}
 }
 
-FTransaction::FWriteHash FTransaction::CalculateNestedWriteHash(size_t NumWriteEntries) const
+FTransaction::FWriteHash FTransaction::CalculateNestedWriteHash() const
+{
+	FOpenHashThrottler::FHashScope Profile(Context->GetOpenHashThrottler(), CurrentOpenReturnAddress, WriteLog);
+	return CalculateNestedWriteHashWithLimit(NumWriteLogsHashed);
+}
+
+FTransaction::FWriteHash FTransaction::CalculateNestedWriteHashWithLimit(size_t NumWriteEntries) const
 {
 	FXxHash64Builder Hasher;
 	if (nullptr != Parent)
 	{
-		FWriteHash ParentHash = Parent->CalculateNestedWriteHash(Parent->WriteLog.Num());
+		FWriteHash ParentHash = Parent->CalculateNestedWriteHashWithLimit(Parent->WriteLog.Num());
 		Hasher.Update(&ParentHash, sizeof(ParentHash));
 	}
 
-	FOpenHashThrottler::FHashScope Profile(Context->GetTransactionHashProfiler(), CurrentOpenReturnAddress, WriteLog);
-
 	size_t WriteIndex = 0;
 
 	using FSmallValue = uint64;
diff --git a/Engine/Source/Runtime/Core/Private/AutoRTFM/Transaction.h b/Engine/Source/Runtime/Core/Private/AutoRTFM/Transaction.h
--- a/Engine/Source/Runtime/Core/Private/AutoRTFM/Transaction.h
+++ b/Engine/Source/Runtime/Core/Private/AutoRTFM/Transaction.h
@@ -163,10 +163,14 @@ private:
     // See CalculateNestedWriteHash().
     void ValidateWriteHash() const;
 
+    // Returns a hash of all the memory written during this transaction and all parent transactions.
+    // Can be used to ensure that no transaction written memory is modified during an open.
+    FWriteHash CalculateNestedWriteHash() const;
+    
     // Returns a hash of all the memory written during this transaction and all parent transactions.
     // Can be used to ensure that no transaction written memory is modified during an open.
     // NumWriteEntries is the number of write entries to hash before stopping.
-    FWriteHash CalculateNestedWriteHash(size_t NumWriteEntries) const;
+    FWriteHash CalculateNestedWriteHashWithLimit(size_t NumWriteEntries) const;
 
     void CollectStats() const;
 
diff --git a/Engine/Source/Runtime/Core/Public/AutoRTFM/AutoRTFM.h b/Engine/Source/Runtime/Core/Public/AutoRTFM/AutoRTFM.h
--- a/Engine/Source/Runtime/Core/Public/AutoRTFM/AutoRTFM.h
+++ b/Engine/Source/Runtime/Core/Public/AutoRTFM/AutoRTFM.h
@@ -1128,6 +1128,12 @@ namespace ForTheRuntime
 	// Sets the memory validation level. See IsWriteValidationEnabled().
 	UE_AUTORTFM_API void SetMemoryValidationLevel(EMemoryValidationLevel Level);
 
+	// Returns true if the memory validation throttling is enabled.
+	UE_AUTORTFM_API bool GetMemoryValidationThrottlingEnabled();
+
+	// Sets the memory validation throttling mode.
+	UE_AUTORTFM_API void SetMemoryValidationThrottlingEnabled(bool bEnabled);
+
 	// A debug helper that will break to the debugger if the hash of the memory
 	// write locations no longer matches the hash recorded when the transaction
 	// was opened. Useful for isolating where the open write happened.
