Update the in-memory text ID of a FText when it gets re-keyed during package save

This ensures the in-memory state matches what was written to disk, and keeps the ID deterministic after the first save (eg, saving after resetting the package localization ID)

diff --git a/Engine/Source/Runtime/Core/Private/Internationalization/TextHistory.cpp b/Engine/Source/Runtime/Core/Private/Internationalization/TextHistory.cpp
--- a/Engine/Source/Runtime/Core/Private/Internationalization/TextHistory.cpp
+++ b/Engine/Source/Runtime/Core/Private/Internationalization/TextHistory.cpp
@@ -873,6 +873,14 @@ void FTextHistory_Base::Serialize(FStructuredArchive::FRecord Record)
 			{
 				Key = FGuid::NewGuid().ToString();
 			}
+
+			// If the ID changed, and this is a persistent archive with a linker (meaning we're saving a package to disk), then apply the updated ID to the in-memory state
+			if (GIsEditor && BaseArchive.IsPersistent() && BaseArchive.GetLinker() && (TextId.GetNamespace() != Namespace || TextId.GetKey() != Key))
+			{
+				TextId = FTextId(Namespace, Key);
+				LocalizedString.Reset();
+				MarkDisplayStringOutOfDate();
+			}
 #endif // WITH_EDITOR
 		}
 
