Regenerate crash guid on Linux post-fork so that we don't get a ton of crash duplicates in the backend

diff --git a/Engine/Source/Runtime/Core/Private/GenericPlatform/GenericPlatformCrashContext.cpp b/Engine/Source/Runtime/Core/Private/GenericPlatform/GenericPlatformCrashContext.cpp
--- a/Engine/Source/Runtime/Core/Private/GenericPlatform/GenericPlatformCrashContext.cpp
+++ b/Engine/Source/Runtime/Core/Private/GenericPlatform/GenericPlatformCrashContext.cpp
@@ -317,6 +317,17 @@ namespace NCached
 	}
 }
 
+
+FString FGenericCrashContext::GenerateNewCrashGuidBase()
+{
+	const FGuid Guid = FGuid::NewGuid();
+	const FString IniPlatformName(FPlatformProperties::IniPlatformName());
+	const FString GuidRoot = FString::Printf(TEXT("%s%s-%s"), CrashGUIDRootPrefix, *IniPlatformName, *Guid.ToString(EGuidFormats::Digits));
+	NCached::Set(NCached::Session.CrashGUIDRoot, *GuidRoot);
+
+	return GuidRoot;
+}
+
 void FGenericCrashContext::Initialize()
 {
 #if !NOINITCRASHREPORTER
@@ -410,12 +421,8 @@ void FGenericCrashContext::Initialize()
 	}
 
 	// Create a unique base guid for bug report ids
-	const FGuid Guid = FGuid::NewGuid();
-	const FString IniPlatformName(FPlatformProperties::IniPlatformName());
-	const FString CrashGUIDRoot = FString::Printf(TEXT("%s%s-%s"), CrashGUIDRootPrefix, *IniPlatformName, *Guid.ToString(EGuidFormats::Digits));
-	NCached::Set(NCached::Session.CrashGUIDRoot, *CrashGUIDRoot);
-	UE_LOG(LogInit, Log, TEXT("Session CrashGUID >====================================================\n         Session CrashGUID >   %s\n         Session CrashGUID >===================================================="), *CrashGUIDRoot);
-
+	const FString GuidRoot = GenerateNewCrashGuidBase();
+	UE_LOG(LogInit, Log, TEXT("Session CrashGUID >====================================================\n         Session CrashGUID >   %s\n         Session CrashGUID >===================================================="), *GuidRoot);
 	if (FInternationalization::IsAvailable())
 	{
 		FInternationalization& Internationalization = FInternationalization::Get();
@@ -489,6 +496,10 @@ void FGenericCrashContext::Initialize()
 			UE_LOG(LogCrashContext, VeryVerbose, TEXT("Updating forked child Session ProcessID: %u -> %u"), NCached::Session.ProcessId, FPlatformProcess::GetCurrentProcessId());
 
 			NCached::Session.ProcessId = FPlatformProcess::GetCurrentProcessId();
+
+			// Generate a new unique CrashGuidBase for this child
+			const FString GuidRoot = GenerateNewCrashGuidBase();
+			UE_LOG(LogInit, Log, TEXT("Session CrashGUID >====================================================\n                                       Session CrashGUID >   %s\n                                       Session CrashGUID >===================================================="), *GuidRoot);			
 			SerializeTempCrashContextToFile();
 		}
 	});
diff --git a/Engine/Source/Runtime/Core/Public/GenericPlatform/GenericPlatformCrashContext.h b/Engine/Source/Runtime/Core/Public/GenericPlatform/GenericPlatformCrashContext.h
--- a/Engine/Source/Runtime/Core/Public/GenericPlatform/GenericPlatformCrashContext.h
+++ b/Engine/Source/Runtime/Core/Public/GenericPlatform/GenericPlatformCrashContext.h
@@ -793,6 +793,9 @@ protected:
 	/** Get arbitrary game data from the crash context */
 	CORE_API static const FString* GetGameData(const FString& Key);
 
+	/** Generate a new crash guid base string. */
+	CORE_API static FString GenerateNewCrashGuidBase();
+
 private:
 
 	/** Serializes the session context section of the crash context to a temporary file. */
