[AutoRTFM] Swap direction of records and data in FWriteLog

This will make it easy to fold a write into the previous write log entry when back-to-back, which is apparently common in UObject constructors (not done in this CL).
Benchmarks also suggest that this is ~1% faster, although it's not obviously clear why.

Also expand test coverage for the write log.

diff --git a/Engine/Source/Programs/AutoRTFMTests/Private/WriteLogTests.cpp b/Engine/Source/Programs/AutoRTFMTests/Private/WriteLogTests.cpp
--- a/Engine/Source/Programs/AutoRTFMTests/Private/WriteLogTests.cpp
+++ b/Engine/Source/Programs/AutoRTFMTests/Private/WriteLogTests.cpp
@@ -3,36 +3,106 @@
 #include "AutoRTFM/WriteLog.h"
 #include "Catch2Includes.h"
 #include "HAL/Platform.h"
+#include <algorithm>
 #include <cstddef>
+#include <numeric>
 #include <vector>
 
 TEST_CASE("FWriteLog")
 {
+	std::mt19937 Rand(0x1234);
 	AutoRTFM::FWriteLog WriteLog;
 
-	auto NewBuffer = [](size_t Size) -> std::vector<std::byte>
+	// Returns a vector with the given size filled with random bytes
+	auto RandomBuffer = [&Rand](size_t Size) -> std::vector<std::byte>
 	{
 		std::vector<std::byte> Out(Size);
-		uint64 Value = 0x978c49230979854f;
 		for (size_t I = 0; I < Size; I++)
 		{
-			Out[I] = std::byte{Value & 0xff};
-			Value = (Value * 31) ^ (Value >> 4);
+			Out[I] = static_cast<std::byte>(Rand() & 0xff);
 		}
 		return Out;
 	};
-	
-	auto Check = [&](size_t EntrySize, size_t NumEntries)
+
+	// Returns a std::vector<size_t> with NumEntries elements all equal to
+	// EntrySize
+	auto FixedSize = [](size_t EntrySize, size_t NumEntries)
+	{
+		return std::vector<size_t>(NumEntries, EntrySize);
+	};
+
+	// Returns a std::vector<size_t> with NumEntries elements with random values
+	// between 1..AutoRTFM::FWriteLogEntry::MaxSize
+	auto RandomSize = [&Rand](size_t NumEntries)
+	{
+		std::vector<size_t> EntrySizes(NumEntries);
+		for (size_t& EntrySize : EntrySizes)
+		{
+			EntrySize = std::max<size_t>(1, Rand() & AutoRTFM::FWriteLogEntry::MaxSize);
+		}
+		return EntrySizes;
+	};
+
+	// The order in which to add data to the write log.
+	enum class EEntryOrder
 	{
-		// Populate the entries with pseudo random data
-		std::vector<std::byte> Data = NewBuffer(EntrySize * NumEntries);
+		Forwards,  // Sequentially forward
+		Backwards, // Sequentially backwards
+		Random,    // Random order
+	};
+
+	auto CheckWithOrder = [&](std::vector<size_t> EntrySizes, EEntryOrder EntryOrder)
+	{
+		// Total number of write entries to test
+		const size_t NumEntries = EntrySizes.size();
+
+		// Total number of bytes to add to the write log
+		const size_t TotalBufferSize = std::reduce(EntrySizes.begin(), EntrySizes.end());
+
+		// Buffer of random data
+		std::vector<std::byte> Data = RandomBuffer(TotalBufferSize);
+
+		// An interval in Data
+		struct FDataSpan
+		{
+			size_t Offset;
+			size_t Size;
+		};
+
+		// Produce a list of data spans on Data. These will be added to the write log in order.
+		std::vector<FDataSpan> DataSpans(NumEntries);
+		{
+			size_t Offset = 0;
+			for (size_t I = 0; I < NumEntries; I++)
+			{
+				FDataSpan& DataSpan = DataSpans[I];
+				DataSpan.Offset = Offset;
+				DataSpan.Size = EntrySizes[I];
+				Offset += EntrySizes[I];
+			}
+			
+			// Shuffle the spans
+			switch (EntryOrder)
+			{
+				case EEntryOrder::Forwards:
+					break; // Already forwards
+				case EEntryOrder::Backwards:
+					std::reverse(std::begin(DataSpans), std::end(DataSpans));
+					break;
+				case EEntryOrder::Random:
+					std::shuffle(std::begin(DataSpans), std::end(DataSpans), Rand);
+					break;
+			}
+		}
+		
+		// Populate the write log and Entries with the expected write log entries
 		std::vector<AutoRTFM::FWriteLogEntry> Entries;
-		for (size_t I = 0; I < NumEntries; I++)
+		for (FDataSpan& DataSpan : DataSpans)
 		{
 			AutoRTFM::FWriteLogEntry Entry;
-			Entry.Data = &Data[I * EntrySize];
-			Entry.Size = EntrySize;
-			Entry.LogicalAddress = reinterpret_cast<std::byte*>(static_cast<uintptr_t>(0x1234000 + I));
+			Entry.Data = &Data[DataSpan.Offset];
+			Entry.Size = DataSpan.Size;
+			Entry.LogicalAddress = reinterpret_cast<std::byte*>(static_cast<uintptr_t>(0x1234000 + DataSpan.Offset));
 			Entry.bNoMemoryValidation = (static_cast<int>(Entry.Data[0]) & 0x10) != 0;
 			Entries.push_back(Entry);
 			WriteLog.Push(Entry);
@@ -54,6 +124,7 @@ TEST_CASE("FWriteLog")
 				REQUIRE(Got.bNoMemoryValidation == Expect.bNoMemoryValidation);
 				Index++;
 			}
+			REQUIRE(Index == Entries.size());
 		}
 
 		SECTION("Reverse iterator")
@@ -69,6 +140,7 @@ TEST_CASE("FWriteLog")
 				REQUIRE(Got.bNoMemoryValidation == Expect.bNoMemoryValidation);
 				Index--;
 			}
+			REQUIRE(Index + 1 == 0);
 		}
 
 		SECTION("Reset")
@@ -79,6 +151,22 @@ TEST_CASE("FWriteLog")
 			REQUIRE(WriteLog.TotalSize() == 0u);
 		}
 	};
+	
+	auto Check = [&](std::vector<size_t> EntrySizes)
+	{
+		SECTION("Forwards")
+		{
+			CheckWithOrder(std::move(EntrySizes), EEntryOrder::Forwards);
+		}
+		SECTION("Backwards")
+		{
+			CheckWithOrder(std::move(EntrySizes), EEntryOrder::Backwards);
+		}
+		SECTION("Random")
+		{
+			CheckWithOrder(std::move(EntrySizes), EEntryOrder::Random);
+		}
+	};
 
 	SECTION("Empty")
 	{
@@ -89,26 +177,31 @@ TEST_CASE("FWriteLog")
 
 	SECTION("EntrySize: 1, NumEntries: 32")
 	{
-		Check(/* EntrySize */ 1, /* NumEntries */ 32);
+		Check(FixedSize(/* EntrySize */ 1, /* NumEntries */ 32));
 	}
 
 	SECTION("EntrySize: 32, NumEntries: 65536")
 	{
-		Check(/* EntrySize */ 32, /* NumEntries */ 65536);
+		Check(FixedSize(/* EntrySize */ 32, /* NumEntries */ 65536));
 	}
 
 	SECTION("EntrySize: 1024, NumEntries: 32")
 	{
-		Check(/* EntrySize */ 1024, /* NumEntries */ 32);
+		Check(FixedSize(/* EntrySize */ 1024, /* NumEntries */ 32));
 	}
 
 	SECTION("EntrySize: FWriteLogEntry::MaxSize-1, NumEntries: 32")
 	{
-		Check(/* EntrySize */ AutoRTFM::FWriteLogEntry::MaxSize-1, /* NumEntries */ 32);
+		Check(FixedSize(/* EntrySize */ AutoRTFM::FWriteLogEntry::MaxSize-1, /* NumEntries */ 32));
 	}
-	
+
 	SECTION("EntrySize: FWriteLogEntry::MaxSize, NumEntries: 32")
 	{
-		Check(/* EntrySize */ AutoRTFM::FWriteLogEntry::MaxSize, /* NumEntries */ 32);
+		Check(FixedSize(/* EntrySize */ AutoRTFM::FWriteLogEntry::MaxSize, /* NumEntries */ 32));
+	}
+
+	SECTION("EntrySize: random, NumEntries: 32")
+	{
+		Check(RandomSize(/* NumEntries */ 32));
 	}
 }
diff --git a/Engine/Source/Runtime/Core/Private/AutoRTFM/WriteLog.h b/Engine/Source/Runtime/Core/Private/AutoRTFM/WriteLog.h
--- a/Engine/Source/Runtime/Core/Private/AutoRTFM/WriteLog.h
+++ b/Engine/Source/Runtime/Core/Private/AutoRTFM/WriteLog.h
@@ -4,6 +4,8 @@
 
 #include "AutoRTFM/AutoRTFM.h"
 #include "HAL/Platform.h"
+#include "Misc/CoreMiscDefines.h"
+#include "Templates/AlignmentTemplates.h"
 #include "Utils.h"
 #include <cstddef>
 #include <cstdint>
@@ -41,51 +43,52 @@ namespace AutoRTFM
 	// forwards and backwards.
 	class FWriteLog final
 	{
-		struct TRecord
+		struct FRecord
 		{
 			uintptr_t Address             : 48;
 			uintptr_t bNoMemoryValidation : 1;
 			uintptr_t Size                : 15;
 		};
 		static_assert(sizeof(uintptr_t) == 8, "assumption: a pointer is 8 bytes");
-		static_assert(sizeof(TRecord) == 8);
+		static_assert(sizeof(FRecord) == 8);
 
 		struct FBlock final
 		{
 			// ┌────────┬────┬────┬────┬────┬────────────────┬────┬────┬────┬────┐
-			// │ FBlock │ R₀ │ R₁ │ R₂ │ R₃ │->            <-│ D₃ │ D₂ │ D₁ │ D₀ │
+			// │ FBlock │ D₀ │ D₁ │ D₂ │ D₃ │->            <-│ R₃ │ R₂ │ R₁ │ R₀ │
 			// └────────┴────┴────┴────┴────┴────────────────┴────┴────┴────┴────┘
-			//          ^                   ^                ^                   ^
-			//    Head().Record       Tail().Record    Tail().Data         Head().Data
+			//          ^                   ^                ^              ^
+			//      DataStart()          DataEnd         LastRecord    FirstRecord
 			// Where:
-			//   Rₙ = Record n, Dₙ = Data n
+			//   Dₙ = Data n, Rₙ = Record n
 
-			// Size of a heap-allocated block, excluding the FBlock struct header.
-			static constexpr size_t DefaultRecordsAndDataSize = 2048;
-
-			// Pointers to a write log record and its corresponding data.
-			struct FEntryPtrs
-			{
-				TRecord* Record = nullptr;
-				std::byte* Data = nullptr;
-			};
+			// Size of a heap-allocated block, including the FBlock struct header.
+			static constexpr size_t DefaultSize = 2048;
 
 			// Constructor
-			// RecordsAndDataSize is the total size of the block, excluding the FBlock struct header
-			// (the memory size for the write records and data)
-			FBlock(size_t RecordsAndDataSize) 
-				: RecordsAndDataSize(RecordsAndDataSize)
+			// TotalSize is the total size of the allocated memory for the block including
+			// the FBlock header.
+			explicit FBlock(size_t TotalSize)
 			{ 
-				TailPtrs = Head();
+				ensure((TotalSize & (alignof(FRecord) - 1)) == 0);
+				std::byte* End = reinterpret_cast<std::byte*>(this) + TotalSize;
+				DataEnd = DataStart();
+				// Note: The initial empty state has LastRecord pointing one
+				// FRecord beyond the immutable FirstRecord.
+				LastRecord = reinterpret_cast<FRecord*>(End);
+				FirstRecord = LastRecord - 1;
 			}
 
 			// Allocate performs a heap allocation of a new block.
-			// RecordsAndDataSize is the total size of the block, excluding the FBlock struct header
-			// (the memory size for the write records and data)
-			static FBlock* Allocate(size_t RecordsAndDataSize)
+			// TotalSize is the total size of the allocated memory for the block including
+			// the FBlock header.
+			static FBlock* Allocate(size_t TotalSize)
 			{
-				std::byte* Memory = new std::byte[sizeof(FBlock) + RecordsAndDataSize];
-				return new (Memory) FBlock(RecordsAndDataSize);
+				ASSERT(TotalSize > (sizeof(FBlock) + sizeof(FRecord)));
+				std::byte* Memory = new std::byte[TotalSize];
+				// Disable false-positive warning C6386: Buffer overrun while writing to 'Memory'
+				CA_SUPPRESS(6386)
+				return new (Memory) FBlock(TotalSize);
 			}
 
 			// Free releases the heap-allocated memory for this block.
@@ -95,15 +98,22 @@ namespace AutoRTFM
 				delete [] reinterpret_cast<std::byte*>(this);
 			}
 
-			// Returns the address of the first record.
-			UE_AUTORTFM_FORCEINLINE TRecord* FirstRecord() { return reinterpret_cast<TRecord*>(reinterpret_cast<std::byte*>(this) + sizeof(FBlock)); }
-			// Returns the address of the first record.
-			UE_AUTORTFM_FORCEINLINE TRecord const * FirstRecord() const { return reinterpret_cast<const TRecord*>(reinterpret_cast<std::byte const *>(this) + sizeof(FBlock)); }
+			// Returns a pointer to the data for the first entry
+			std::byte* DataStart()
+			{
+				return reinterpret_cast<std::byte*>(this) + sizeof(FBlock);
+			}
+
+			// Returns a pointer to the data for the last entry
+			std::byte* LastData()
+			{
+				return DataEnd - LastRecord->Size;
+			}
 
-			// Returns true if the block has no entries.
-			UE_AUTORTFM_FORCEINLINE bool IsEmpty()
+			// Returns true if the block holds no entries.
+			bool IsEmpty() const
 			{
-				return FirstRecord() == TailPtrs.Record;
+				return LastRecord > FirstRecord;
 			}
 
 			// Attempts to add the entry into this block by copying the entry's data and creating a
@@ -112,58 +122,38 @@ namespace AutoRTFM
 			// for the entry.
 			UE_AUTORTFM_FORCEINLINE bool Push(FWriteLogEntry Entry)
 			{
-				AssertIsValid(TailPtrs);
-
-				if (reinterpret_cast<std::byte*>(TailPtrs.Record + 1) > TailPtrs.Data - Entry.Size)
+				if (DataEnd + Entry.Size > reinterpret_cast<std::byte*>(LastRecord - 1))
 				{
+					// Entry's data + new record does not fit in the block's remaining space.
 					return false;
 				}
 
-				TailPtrs.Record->Address = reinterpret_cast<uintptr_t>(Entry.LogicalAddress);
-				TailPtrs.Record->Size = Entry.Size;
-				TailPtrs.Record->bNoMemoryValidation = Entry.bNoMemoryValidation;
-				TailPtrs.Record++;
+				LastRecord--;
+				LastRecord->Address = reinterpret_cast<uintptr_t>(Entry.LogicalAddress);
+				LastRecord->Size = Entry.Size;
+				LastRecord->bNoMemoryValidation = Entry.bNoMemoryValidation;
 
-				TailPtrs.Data -= Entry.Size;
-				memcpy(TailPtrs.Data, Entry.Data, Entry.Size);
-				AssertIsValid(TailPtrs);
-				return true;
-			}
+				memcpy(DataEnd, Entry.Data, Entry.Size);
 
-			// Returns the record and data pointers to the first entry in the block.
-			UE_AUTORTFM_FORCEINLINE FEntryPtrs Head()
-			{
-				FEntryPtrs Ptrs;
-				Ptrs.Record = FirstRecord();
-				Ptrs.Data = reinterpret_cast<std::byte*>(FirstRecord()) + RecordsAndDataSize;
-				return Ptrs;
-			}
+				DataEnd += Entry.Size;
 
-			// Returns the record and data pointers to the last entry in the block.
-			UE_AUTORTFM_FORCEINLINE FEntryPtrs Tail()
-			{
-				return TailPtrs;
-			}
-
-			// Asserts tha the pointers in Ptrs are within bounds for this block.
-			UE_AUTORTFM_FORCEINLINE void AssertIsValid(FEntryPtrs Ptrs) const
-			{
 #if UE_BUILD_DEBUG
-				ASSERT(FirstRecord() <= Ptrs.Record);
-				ASSERT(reinterpret_cast<const std::byte*>(Ptrs.Record) <= Ptrs.Data);
-				ASSERT(Ptrs.Data <= reinterpret_cast<const std::byte*>(FirstRecord()) + RecordsAndDataSize);
+				ASSERT(DataEnd <= reinterpret_cast<std::byte*>(LastRecord));
 #endif
+
+				return true;
 			}
 
-			// The size in bytes of the block's records and data (total size - sizeof(FBlock)).
-			size_t const RecordsAndDataSize;
 			// The next block in the linked list.
 			FBlock* NextBlock = nullptr;
 			// The previous block in the linked list.
 			FBlock* PrevBlock = nullptr;
-			// The pointers to one-past the last entry in the block.
-			FEntryPtrs TailPtrs;
-
+			// The pointer to the first entry's record
+			FRecord* FirstRecord = nullptr;
+			// The pointer to the last entry's record
+			FRecord* LastRecord = nullptr;
+			// One byte beyond the end of the last entry's data
+			std::byte* DataEnd = nullptr;
 		private:
 			~FBlock() = delete;
 		};
@@ -172,7 +162,7 @@ namespace AutoRTFM
 		// Constructor
 		FWriteLog()
 		{
-			new(HeadBlockMemory) FBlock(HeadBlockSize - sizeof(FBlock));
+			new(HeadBlockMemory) FBlock(HeadBlockSize);
 		}
 
 		// Destructor
@@ -190,8 +180,8 @@ namespace AutoRTFM
 
 			if (UNLIKELY(!TailBlock->Push(Entry)))
 			{
-				const size_t RecordsAndDataSize = std::max(Entry.Size + sizeof(TRecord), FBlock::DefaultRecordsAndDataSize);
-				FBlock* NewBlock = FBlock::Allocate(RecordsAndDataSize);
+				const size_t RequiredSize = Align(sizeof(FBlock) + Entry.Size, alignof(FRecord)) + sizeof(FRecord);
+				FBlock* NewBlock = FBlock::Allocate(std::max(RequiredSize, FBlock::DefaultSize));
 				NewBlock->PrevBlock = TailBlock;
 				TailBlock->NextBlock = NewBlock;
 				TailBlock = NewBlock;
@@ -210,109 +200,109 @@ namespace AutoRTFM
 		{
 			TIterator() = default;
 
-			TIterator(FBlock* Block) : Block(Block)
+			TIterator(FBlock* StartBlock) : Block(StartBlock)
 			{
-				Ptrs = IS_FORWARD ? Block->Head() : Block->Tail();
-				Progress();
+				if constexpr (IS_FORWARD)
+				{
+					if (Block->IsEmpty())
+					{
+						// First block is fixed size and may be empty if the
+						// first write is larger than its fixed size.
+						Block = Block->NextBlock;
+					}
+				}
+
+				Data = IS_FORWARD ? Block->DataStart() : Block->LastData();
+				Record = IS_FORWARD ? Block->FirstRecord : Block->LastRecord;
 			}
 
 			// Returns the entry at the current iterator's position.
 			FWriteLogEntry operator*() const
 			{
-				return Current;
+				FWriteLogEntry Entry;
+				Entry.LogicalAddress = reinterpret_cast<std::byte*>(Record->Address);
+				Entry.Data = Data;
+				Entry.Size = Record->Size;
+				Entry.bNoMemoryValidation = Record->bNoMemoryValidation;
+				return Entry;
 			}
 
 			// Progresses the iterator to the next entry
 			void operator++()
 			{
-				Progress();
-			}
-
-			// Inequality operator
-			bool operator!=(const TIterator& Other) const
-			{
-				return (Other.Block != Block) || (Other.Ptrs.Record != Ptrs.Record);
-			}
-
-		private:
-			// Assigns the current write log entry to Current and progresses the iterator to the
-			// next entry in the log. If the iterator progresses beyond the start or end of the log
-			// then the iterator is reset.
-			void Progress()
-			{
-				Block->AssertIsValid(Ptrs);
-
 				if constexpr (IS_FORWARD)
 				{
-					if (Ptrs.Record == Block->TailPtrs.Record)
+					if (Record == Block->LastRecord)
 					{
 						Block = Block->NextBlock;
 						if (!Block)
 						{
-							// Progressed beyond the end of the log.
 							Reset();
 							return;
 						}
-
-						Ptrs = Block->Head();
+						Data = Block->DataStart();
+						Record = Block->FirstRecord;
+					}
+					else
+					{
+						Data += Record->Size;
+						Record--;
 					}
-
-					Ptrs.Data -= Ptrs.Record->Size;
-					SetCurrent();
-					Ptrs.Record++;
 				}
 				else
 				{
-					if (Ptrs.Record == Block->FirstRecord())
+					if (Record == Block->FirstRecord)
 					{
 						Block = Block->PrevBlock;
-						// Note: Only the head block is fixed size, and may not contain any entries.
 						if (!Block || Block->IsEmpty())
 						{
-							// Progressed beyond the start of the log.
 							Reset();
 							return;
 						}
-						Ptrs = Block->Tail();
+						Data = Block->LastData();
+						Record = Block->LastRecord;
+					}
+					else
+					{
+						Record++;
+						Data -= Record->Size;
 					}
-
-					Ptrs.Record--;
-					SetCurrent();
-					Ptrs.Data += Current.Size;
 				}
-				
-				Block->AssertIsValid(Ptrs);
 			}
 
-			// Assigns to Current the entry this iterator currently points to.
-			UE_AUTORTFM_FORCEINLINE void SetCurrent()
+			// Inequality operator
+			bool operator!=(const TIterator& Other) const
 			{
-				Current.LogicalAddress = reinterpret_cast<std::byte*>(Ptrs.Record->Address);
-				Current.Data = Ptrs.Data;
-				Current.Size = Ptrs.Record->Size;
-				Current.bNoMemoryValidation = Ptrs.Record->bNoMemoryValidation;
+				return (Other.Block != Block) || (Other.Record != Record);
 			}
 
+		private:
 			// Resets the iterator (compares equal to the write log's end())
 			UE_AUTORTFM_FORCEINLINE void Reset()
 			{
 				Block = nullptr;
-				Ptrs = {};
-				Current = {};
+				Data = nullptr;
+				Record = nullptr;
 			}
 
 			FBlock* Block = nullptr;
-			FBlock::FEntryPtrs Ptrs;
-			FWriteLogEntry Current;
+			std::byte* Data = nullptr;
+			FRecord* Record = nullptr;
 		};
 
 		using Iterator = TIterator</* IS_FORWARD */ true>;
 		using ReverseIterator = TIterator</* IS_FORWARD */ false>;
 
-		Iterator begin() const { return Iterator(HeadBlock); }
-		ReverseIterator rbegin() const { return ReverseIterator(TailBlock); }
-		Iterator end() const { return Iterator(); }
-		ReverseIterator rend() const { return ReverseIterator(); }
+		Iterator begin() const
+		{
+			return (NumEntries > 0) ? Iterator(HeadBlock) : Iterator{};
+		}
+		ReverseIterator rbegin() const
+		{
+			return (NumEntries > 0) ? ReverseIterator(TailBlock) : ReverseIterator{};
+		}
+		Iterator end() const { return Iterator{}; }
+		ReverseIterator rend() const { return ReverseIterator{}; }
 
 		// Resets the write log to its initial state, freeing any allocated memory.
 		void Reset()
@@ -342,7 +332,14 @@ namespace AutoRTFM
 		UE_AUTORTFM_FORCEINLINE size_t TotalSize() const { return TotalSizeBytes; }
 
 	private:
+		template<size_t SIZE>
+		static constexpr bool IsAlignedForTRecord = (SIZE & (alignof(FRecord) - 1)) == 0;
+
+		static constexpr size_t DefaultBlockSize = sizeof(FBlock) + FBlock::DefaultSize;
 		static constexpr size_t HeadBlockSize = 256;
+		static_assert(IsAlignedForTRecord<HeadBlockSize>);
+		static_assert(IsAlignedForTRecord<DefaultBlockSize>);
+		
 		FBlock* HeadBlock = reinterpret_cast<FBlock*>(HeadBlockMemory);
 		FBlock* TailBlock = reinterpret_cast<FBlock*>(HeadBlockMemory);
 		size_t NumEntries = 0;
