Substrate - fixed compilation error

#tests pc editor dx12
#rnx

diff --git a/Engine/Source/Runtime/Core/Private/UObject/UE5MainStreamObjectVersion.cpp b/Engine/Source/Runtime/Core/Private/UObject/UE5MainStreamObjectVersion.cpp
--- a/Engine/Source/Runtime/Core/Private/UObject/UE5MainStreamObjectVersion.cpp
+++ b/Engine/Source/Runtime/Core/Private/UObject/UE5MainStreamObjectVersion.cpp
@@ -20,7 +20,7 @@ TMap<FGuid, FGuid> FUE5MainStreamObjectVersion::GetSystemGuids()
 	SystemGuids.Add(DevGuids.POSESEARCHDB_DERIVEDDATA_VER, FGuid("4E595C2AC5E947D6BA9ABC874353E5BC"));
 	SystemGuids.Add(DevGuids.GROOM_BINDING_DERIVED_DATA_VERSION, FGuid("156678D4F9084D7CAE98FADC6AB93573"));
 	SystemGuids.Add(DevGuids.GROOM_DERIVED_DATA_VERSION, FGuid("00B8207248CC44B88F7DEC8F328A5D8B"));
-	SystemGuids.Add(DevGuids.MaterialTranslationDDCVersion, FGuid("2EBF326928DF4469BEA514CDBC1081AE"));
+	SystemGuids.Add(DevGuids.MaterialTranslationDDCVersion, FGuid("5E6E21D1AC364AACA1921EE46BF10B4C"));
 
 	return SystemGuids;
 }
diff --git a/Engine/Source/Runtime/Engine/Private/Materials/HLSLMaterialTranslator.cpp b/Engine/Source/Runtime/Engine/Private/Materials/HLSLMaterialTranslator.cpp
--- a/Engine/Source/Runtime/Engine/Private/Materials/HLSLMaterialTranslator.cpp
+++ b/Engine/Source/Runtime/Engine/Private/Materials/HLSLMaterialTranslator.cpp
@@ -14401,7 +14401,7 @@ int32 FHLSLMaterialTranslator::SubstrateVerticalLayeringParameterBlending(int32
 	}
 
 	return AddCodeChunk(
-		MCT_Substrate, TEXT("SubstrateVerticalLayeringParameterBlending(%s.InlinedBSDF, %s.InlinedBSDF, %s, %s, %s) /* Thickness = %s */"),
+		MCT_Substrate, TEXT("SubstrateVerticalLayeringParameterBlending(%s, %s, %s, %s, %s) /* Thickness = %s */"),
 		*GetParameterCode(Top),
 		*GetParameterCode(Base),
 		*SharedLocalBasisIndexMacro,
