Fix Stack Overflow by [REDACTED_PERSON] of Vectorized Transform Methods in Debug Builds

diff --git a/Engine/Source/Runtime/Core/Public/Math/TransformVectorized.h b/Engine/Source/Runtime/Core/Public/Math/TransformVectorized.h
--- a/Engine/Source/Runtime/Core/Public/Math/TransformVectorized.h
+++ b/Engine/Source/Runtime/Core/Public/Math/TransformVectorized.h
@@ -86,7 +86,7 @@ public:
 	CORE_API static const TTransform<T> Identity;
 
 #if ENABLE_NAN_DIAGNOSTIC
-	FORCEINLINE void DiagnosticCheckNaN_Scale3D() const
+	FORCEINLINE_DEBUGGABLE void DiagnosticCheckNaN_Scale3D() const
 	{
 		if (VectorContainsNaNOrInfinite(Scale3D))
 		{
@@ -95,7 +95,7 @@ public:
 		}
 	}
 
-	FORCEINLINE void DiagnosticCheckNaN_Translate() const
+	FORCEINLINE_DEBUGGABLE void DiagnosticCheckNaN_Translate() const
 	{
 		if (VectorContainsNaNOrInfinite(Translation))
 		{
@@ -104,7 +104,7 @@ public:
 		}
 	}
 
-	FORCEINLINE void DiagnosticCheckNaN_Rotate() const
+	FORCEINLINE_DEBUGGABLE void DiagnosticCheckNaN_Rotate() const
 	{
 		if (VectorContainsNaNOrInfinite(Rotation))
 		{
@@ -113,14 +113,14 @@ public:
 		}
 	}
 
-	FORCEINLINE void DiagnosticCheckNaN_All() const
+	FORCEINLINE_DEBUGGABLE void DiagnosticCheckNaN_All() const
 	{
 		DiagnosticCheckNaN_Scale3D();
 		DiagnosticCheckNaN_Rotate();
 		DiagnosticCheckNaN_Translate();
 	}
 
-	FORCEINLINE void DiagnosticCheck_IsValid() const
+	FORCEINLINE_DEBUGGABLE void DiagnosticCheck_IsValid() const
 	{
 		DiagnosticCheckNaN_All();
 		if (!IsValid())
@@ -130,17 +130,17 @@ public:
 		
 	}
 #else
-	FORCEINLINE void DiagnosticCheckNaN_Translate() const {}
-	FORCEINLINE void DiagnosticCheckNaN_Rotate() const {}
-	FORCEINLINE void DiagnosticCheckNaN_Scale3D() const {}
-	FORCEINLINE void DiagnosticCheckNaN_All() const {}
-	FORCEINLINE void DiagnosticCheck_IsValid() const {}
+	FORCEINLINE_DEBUGGABLE void DiagnosticCheckNaN_Translate() const {}
+	FORCEINLINE_DEBUGGABLE void DiagnosticCheckNaN_Rotate() const {}
+	FORCEINLINE_DEBUGGABLE void DiagnosticCheckNaN_Scale3D() const {}
+	FORCEINLINE_DEBUGGABLE void DiagnosticCheckNaN_All() const {}
+	FORCEINLINE_DEBUGGABLE void DiagnosticCheck_IsValid() const {}
 #endif
 
 	/**
 	 * Constructor with initialization to the identity transform.
 	 */
-	FORCEINLINE TTransform()
+	FORCEINLINE_DEBUGGABLE TTransform()
 	{
 		// Rotation = {0,0,0,1)
 		Rotation = GlobalVectorConstants::Float0001;
@@ -155,7 +155,7 @@ public:
 	 *
 	 * @param InTranslation The value to use for the translation component
 	 */
-	FORCEINLINE explicit TTransform(const TVector<T>& InTranslation) 
+	FORCEINLINE_DEBUGGABLE explicit TTransform(const TVector<T>& InTranslation) 
 	{
 		// Rotation = {0,0,0,1) quaternion identity
 		Rotation = GlobalVectorConstants::Float0001;
@@ -172,7 +172,7 @@ public:
 	 *
 	 * @param InRotation The value to use for rotation component
 	 */
-	FORCEINLINE explicit TTransform(const TQuat<T>& InRotation) 
+	FORCEINLINE_DEBUGGABLE explicit TTransform(const TQuat<T>& InRotation) 
 	{
 		// Rotation = InRotation
 		Rotation =  VectorLoadAligned( &InRotation );
@@ -189,7 +189,7 @@ public:
 	 *
 	 * @param InRotation The value to use for rotation component  (after being converted to a quaternion)
 	 */
-	FORCEINLINE explicit TTransform(const TRotator<T>& InRotation) 
+	FORCEINLINE_DEBUGGABLE explicit TTransform(const TRotator<T>& InRotation) 
 	{
 		TQuat<T> InQuatRotation = TQuat<T>::MakeFromRotator(InRotation);
 		// Rotation = InRotation
@@ -209,7 +209,7 @@ public:
 	 * @param InTranslation The value to use for the translation component
 	 * @param InScale3D The value to use for the scale component
 	 */
-	FORCEINLINE TTransform(const TQuat<T>& InRotation, const TVector<T>& InTranslation, const TVector<T>& InScale3D = TVector<T>(1.f, 1.f, 1.f))
+	FORCEINLINE_DEBUGGABLE TTransform(const TQuat<T>& InRotation, const TVector<T>& InTranslation, const TVector<T>& InScale3D = TVector<T>(1.f, 1.f, 1.f))
 	{
 		// Rotation = InRotation
 		Rotation =  VectorLoadAligned( &InRotation );
@@ -228,7 +228,7 @@ public:
 	 * @param InTranslation The value to use for the translation component
 	 * @param InScale3D The value to use for the scale component
 	 */
-	FORCEINLINE TTransform(const TransformVectorRegister& InRotation, const TransformVectorRegister& InTranslation, const TransformVectorRegister& InScale3D)
+	FORCEINLINE_DEBUGGABLE TTransform(const TransformVectorRegister& InRotation, const TransformVectorRegister& InTranslation, const TransformVectorRegister& InScale3D)
 		: Rotation(InRotation),
 		Translation(InTranslation),
 		Scale3D(InScale3D)
@@ -243,7 +243,7 @@ public:
 	 * @param InTranslation The value to use for the translation component
 	 * @param InScale3D The value to use for the scale component
 	 */
-	FORCEINLINE TTransform(const TRotator<T>& InRotation, const TVector<T>& InTranslation, const TVector<T>& InScale3D = TVector<T>(1.f, 1.f, 1.f))
+	FORCEINLINE_DEBUGGABLE TTransform(const TRotator<T>& InRotation, const TVector<T>& InTranslation, const TVector<T>& InScale3D = TVector<T>(1.f, 1.f, 1.f))
 	{
 		TQuat<T> InQuatRotation = TQuat<T>::MakeFromRotator(InRotation);
 		// Rotation = InRotation
@@ -259,7 +259,7 @@ public:
 	/**
 	 * Constructor with leaving uninitialized memory
 	 */
-	FORCEINLINE explicit TTransform(ENoInit) 
+	FORCEINLINE_DEBUGGABLE explicit TTransform(ENoInit) 
 	{
 		// Note: This can be used to track down initialization issues with bone transform arrays; but it will
 		// cause issues with transient fields such as RootMotionDelta that get initialized to 0 by default
@@ -275,14 +275,14 @@ public:
 	/**
 	 * Constructor for converting a Matrix (including scale) into a TTransform.
 	 */
-	FORCEINLINE explicit TTransform(const TMatrix<T>& InMatrix)
+	FORCEINLINE_DEBUGGABLE explicit TTransform(const TMatrix<T>& InMatrix)
 	{
 		SetFromMatrix(InMatrix);
 		DiagnosticCheckNaN_All();
 	}
 
 	/** Constructor that takes basis axes and translation */
-	FORCEINLINE TTransform(const TVector<T>& InX, const TVector<T>& InY, const TVector<T>& InZ, const TVector<T>& InTranslation)
+	FORCEINLINE_DEBUGGABLE TTransform(const TVector<T>& InX, const TVector<T>& InY, const TVector<T>& InZ, const TVector<T>& InTranslation)
 	{
 		SetFromMatrix(TMatrix<T>(InX, InY, InZ, InTranslation));
 		DiagnosticCheckNaN_All();
@@ -305,7 +305,7 @@ public:
 	CORE_API bool InitFromString( const FString& InSourceString );
 
 
-	FORCEINLINE TMatrix<T> ToMatrixWithScale() const
+	FORCEINLINE_DEBUGGABLE TMatrix<T> ToMatrixWithScale() const
 	{
 
 		TMatrix<T> OutMatrix;
@@ -356,7 +356,7 @@ public:
 	/**
 	* Convert this Transform to matrix with scaling and compute the inverse of that.
 	*/
-	FORCEINLINE TMatrix<T> ToInverseMatrixWithScale() const
+	FORCEINLINE_DEBUGGABLE TMatrix<T> ToInverseMatrixWithScale() const
 	{
 		// todo: optimize
 		return ToMatrixWithScale().Inverse();
@@ -365,7 +365,7 @@ public:
 	/**
 	* Convert this Transform to inverse.
 	*/
-	FORCEINLINE TTransform Inverse() const
+	FORCEINLINE_DEBUGGABLE TTransform Inverse() const
 	{
 		// Replacement of Inverse of TMatrix<T>
 		if (VectorAnyGreaterThan(VectorAbs(Scale3D), GlobalVectorConstants::SmallNumber))
@@ -381,7 +381,7 @@ public:
 	/**
 	* Convert this Transform to a transformation matrix, ignoring its scaling
 	*/
-	FORCEINLINE TMatrix<T> ToMatrixNoScale() const
+	FORCEINLINE_DEBUGGABLE TMatrix<T> ToMatrixNoScale() const
 	{
 		TMatrix<T> OutMatrix;
 		TransformVectorRegister DiagonalsXYZ;
@@ -429,7 +429,7 @@ public:
 	}
 
 	/** Set this transform to the weighted blend of the supplied two transforms. */
-	FORCEINLINE void Blend(const TTransform<T>& Atom1, const TTransform<T>& Atom2, float Alpha)
+	FORCEINLINE_DEBUGGABLE void Blend(const TTransform<T>& Atom1, const TTransform<T>& Atom2, float Alpha)
 	{
 #if !(UE_BUILD_SHIPPING || UE_BUILD_TEST) && WITH_EDITORONLY_DATA
 		// Check that all bone atoms coming from animation are normalized
@@ -465,7 +465,7 @@ public:
 	}
 
 	/** Set this Transform to the weighted blend of it and the supplied Transform. */
-	FORCEINLINE void BlendWith(const TTransform<T>& OtherAtom, float Alpha)
+	FORCEINLINE_DEBUGGABLE void BlendWith(const TTransform<T>& OtherAtom, float Alpha)
 	{
 #if !(UE_BUILD_SHIPPING || UE_BUILD_TEST) && WITH_EDITORONLY_DATA
 		// Check that all bone atoms coming from animation are normalized
@@ -504,12 +504,12 @@ public:
 	 * Use only within blends!!
 	 * Rotation part is NOT normalized!!
 	 */
-	FORCEINLINE TTransform<T> operator+(const TTransform<T>& Atom) const
+	FORCEINLINE_DEBUGGABLE TTransform<T> operator+(const TTransform<T>& Atom) const
 	{
 		return TTransform<T>( VectorAdd(Rotation, Atom.Rotation), VectorAdd(Translation, Atom.Translation ), VectorAdd( Scale3D, Atom.Scale3D ) );
 	}
 
-	FORCEINLINE TTransform<T>& operator+=(const TTransform<T>& Atom)
+	FORCEINLINE_DEBUGGABLE TTransform<T>& operator+=(const TTransform<T>& Atom)
 	{ 
 		Translation = VectorAdd(Translation, Atom.Translation);
 		Rotation = VectorAdd(Rotation, Atom.Rotation);
@@ -518,13 +518,13 @@ public:
 		return *this;
 	}
 
-	FORCEINLINE TTransform<T> operator*(const ScalarRegister& MultScalar) const
+	FORCEINLINE_DEBUGGABLE TTransform<T> operator*(const ScalarRegister& MultScalar) const
 	{		
 		const TransformVectorRegister Mult = MultScalar.Value;
 		return TTransform<T>( VectorMultiply(Rotation, Mult), VectorMultiply(Translation, Mult), VectorMultiply(Scale3D, Mult) );
 	}
 
-	FORCEINLINE TTransform<T>& operator*=(const ScalarRegister& MultScalar)
+	FORCEINLINE_DEBUGGABLE TTransform<T>& operator*=(const ScalarRegister& MultScalar)
 	{			
 		const TransformVectorRegister Mult = MultScalar.Value;
 		Translation= VectorMultiply(Translation, Mult);
@@ -534,17 +534,17 @@ public:
 		return *this;
 	}
 
-	FORCEINLINE TTransform<T>	operator*(const TTransform<T>& Other) const;
-	FORCEINLINE void			operator*=(const TTransform<T>& Other);
-	FORCEINLINE TTransform<T>	operator*(const TQuat<T>& Other) const;
-	FORCEINLINE void			operator*=(const TQuat<T>& Other);
+	FORCEINLINE_DEBUGGABLE TTransform<T>	operator*(const TTransform<T>& Other) const;
+	FORCEINLINE_DEBUGGABLE void			operator*=(const TTransform<T>& Other);
+	FORCEINLINE_DEBUGGABLE TTransform<T>	operator*(const TQuat<T>& Other) const;
+	FORCEINLINE_DEBUGGABLE void			operator*=(const TQuat<T>& Other);
 
-	FORCEINLINE static bool AnyHasNegativeScale(const TVector<T>& InScale3D, const TVector<T>& InOtherScale3D);
-	FORCEINLINE void ScaleTranslation(const TVector<T>& InScale3D);
-	FORCEINLINE void ScaleTranslation(const T& Scale);
-	FORCEINLINE void RemoveScaling(T Tolerance=UE_SMALL_NUMBER);
-	FORCEINLINE float GetMaximumAxisScale() const;
-	FORCEINLINE float GetMinimumAxisScale() const;
+	FORCEINLINE_DEBUGGABLE static bool AnyHasNegativeScale(const TVector<T>& InScale3D, const TVector<T>& InOtherScale3D);
+	FORCEINLINE_DEBUGGABLE void ScaleTranslation(const TVector<T>& InScale3D);
+	FORCEINLINE_DEBUGGABLE void ScaleTranslation(const T& Scale);
+	FORCEINLINE_DEBUGGABLE void RemoveScaling(T Tolerance=UE_SMALL_NUMBER);
+	FORCEINLINE_DEBUGGABLE float GetMaximumAxisScale() const;
+	FORCEINLINE_DEBUGGABLE float GetMinimumAxisScale() const;
 	// Inverse does not work well with VQS format(in particular non-uniform), so removing it, but made two below functions to be used instead. 
 
 	/*******************************************************************************************
@@ -562,51 +562,51 @@ public:
 	 */
 	CORE_API void		SetToRelativeTransform(const TTransform<T>& ParentTransform);
 
-	FORCEINLINE TVector4<T>	TransformFVector4(const TVector4<T>& V) const;
-	FORCEINLINE TVector4<T>	TransformFVector4NoScale(const TVector4<T>& V) const;
-	FORCEINLINE TVector<T>		TransformPosition(const TVector<T>& V) const;
-	FORCEINLINE TVector<T>		TransformPositionNoScale(const TVector<T>& V) const;
+	FORCEINLINE_DEBUGGABLE TVector4<T>	TransformFVector4(const TVector4<T>& V) const;
+	FORCEINLINE_DEBUGGABLE TVector4<T>	TransformFVector4NoScale(const TVector4<T>& V) const;
+	FORCEINLINE_DEBUGGABLE TVector<T>		TransformPosition(const TVector<T>& V) const;
+	FORCEINLINE_DEBUGGABLE TVector<T>		TransformPositionNoScale(const TVector<T>& V) const;
 
 
 	/** Inverts the transform and then transforms V - correctly handles scaling in this transform. */
-	FORCEINLINE TVector<T>		InverseTransformPosition(const TVector<T> &V) const;
-	FORCEINLINE TVector<T>		InverseTransformPositionNoScale(const TVector<T> &V) const;
-	FORCEINLINE TVector<T>		TransformVector(const TVector<T>& V) const;
-	FORCEINLINE TVector<T>		TransformVectorNoScale(const TVector<T>& V) const;
+	FORCEINLINE_DEBUGGABLE TVector<T>		InverseTransformPosition(const TVector<T> &V) const;
+	FORCEINLINE_DEBUGGABLE TVector<T>		InverseTransformPositionNoScale(const TVector<T> &V) const;
+	FORCEINLINE_DEBUGGABLE TVector<T>		TransformVector(const TVector<T>& V) const;
+	FORCEINLINE_DEBUGGABLE TVector<T>		TransformVectorNoScale(const TVector<T>& V) const;
 
 	/** 
 	 *	Transform a direction vector by the inverse of this matrix - will not take into account translation part.
 	 *	If you want to transform a surface normal (or plane) and correctly account for non-uniform scaling you should use TransformByUsingAdjointT with adjoint of matrix inverse.
 	 */
-	FORCEINLINE TVector<T> InverseTransformVector(const TVector<T> &V) const;
-	FORCEINLINE TVector<T> InverseTransformVectorNoScale(const TVector<T> &V) const;
+	FORCEINLINE_DEBUGGABLE TVector<T> InverseTransformVector(const TVector<T> &V) const;
+	FORCEINLINE_DEBUGGABLE TVector<T> InverseTransformVectorNoScale(const TVector<T> &V) const;
 
 	/**
 	* Transform a rotation.
 	* For example if this is a LocalToWorld transform, TransformRotation(Q) would transform Q from local to world space.
 	*/
-	FORCEINLINE TQuat<T> TransformRotation(const TQuat<T>& Q) const;
+	FORCEINLINE_DEBUGGABLE TQuat<T> TransformRotation(const TQuat<T>& Q) const;
 
 	/**
 	* Inverse transform a rotation.
 	* For example if this is a LocalToWorld transform, InverseTransformRotation(Q) would transform Q from world to local space.
 	*/
-	FORCEINLINE TQuat<T> InverseTransformRotation(const TQuat<T>& Q) const;
+	FORCEINLINE_DEBUGGABLE TQuat<T> InverseTransformRotation(const TQuat<T>& Q) const;
 
-	FORCEINLINE TTransform<T>	GetScaled(T Scale) const;
-	FORCEINLINE TTransform<T>	GetScaled(TVector<T> Scale) const;
-	FORCEINLINE TVector<T>		GetScaledAxis(EAxis::Type InAxis) const;
-	FORCEINLINE TVector<T>		GetUnitAxis(EAxis::Type InAxis) const;
-	FORCEINLINE void		Mirror(EAxis::Type MirrorAxis, EAxis::Type FlipAxis);
-	FORCEINLINE static TVector<T>	GetSafeScaleReciprocal(const TVector<T>& InScale, T Tolerance=UE_SMALL_NUMBER);
+	FORCEINLINE_DEBUGGABLE TTransform<T>	GetScaled(T Scale) const;
+	FORCEINLINE_DEBUGGABLE TTransform<T>	GetScaled(TVector<T> Scale) const;
+	FORCEINLINE_DEBUGGABLE TVector<T>		GetScaledAxis(EAxis::Type InAxis) const;
+	FORCEINLINE_DEBUGGABLE TVector<T>		GetUnitAxis(EAxis::Type InAxis) const;
+	FORCEINLINE_DEBUGGABLE void		Mirror(EAxis::Type MirrorAxis, EAxis::Type FlipAxis);
+	FORCEINLINE_DEBUGGABLE static TVector<T>	GetSafeScaleReciprocal(const TVector<T>& InScale, T Tolerance=UE_SMALL_NUMBER);
 
 
-	FORCEINLINE TVector<T> GetLocation() const
+	FORCEINLINE_DEBUGGABLE TVector<T> GetLocation() const
 	{
 		return GetTranslation();
 	}
 
-	FORCEINLINE TRotator<T> Rotator() const
+	FORCEINLINE_DEBUGGABLE TRotator<T> Rotator() const
 	{
 		TQuat<T> OutRotation;
 		VectorStoreAligned(Rotation, &OutRotation);
@@ -614,7 +614,7 @@ public:
 	}
 
 	/** Calculate the  */
-	FORCEINLINE T GetDeterminant() const
+	FORCEINLINE_DEBUGGABLE T GetDeterminant() const
 	{
 		//#todo - vectorized version of this
 		UE::Math::TVector4<T> OutScale3D;
@@ -623,7 +623,7 @@ public:
 	}
 
 	/** Set the translation of this transformation */
-	FORCEINLINE void SetLocation(const TVector<T>& Origin)
+	FORCEINLINE_DEBUGGABLE void SetLocation(const TVector<T>& Origin)
 	{		
 		Translation = VectorLoadFloat3_W0(&Origin);
 		DiagnosticCheckNaN_Translate();
@@ -703,12 +703,12 @@ public:
 
 private:
 
-	FORCEINLINE static bool Private_AnyHasNegativeScale(const TransformVectorRegister& InScale3D, const TransformVectorRegister& InOtherScale3D)
+	FORCEINLINE_DEBUGGABLE static bool Private_AnyHasNegativeScale(const TransformVectorRegister& InScale3D, const TransformVectorRegister& InOtherScale3D)
 	{
 		return !!VectorAnyLesserThan(VectorMin(InScale3D, InOtherScale3D), (TransformVectorRegister)GlobalVectorConstants::FloatZero);
 	}
 
-	FORCEINLINE bool Private_RotationEquals( const TransformVectorRegister& InRotation, const FReal ToleranceScalar = UE_KINDA_SMALL_NUMBER) const
+	FORCEINLINE_DEBUGGABLE bool Private_RotationEquals( const TransformVectorRegister& InRotation, const FReal ToleranceScalar = UE_KINDA_SMALL_NUMBER) const
 	{			
 		const TransformVectorRegister MyRotation = Rotation; // Load from persistent value, to avoid repeated loads.
 		const TransformVectorRegister Tolerance = VectorLoadFloat1(&ToleranceScalar);
@@ -719,7 +719,7 @@ private:
 		return !VectorAnyGreaterThan(RotationSub, Tolerance) || !VectorAnyGreaterThan(RotationAdd, Tolerance);
 	}
 
-	FORCEINLINE bool Private_TranslationEquals( const TransformVectorRegister& InTranslation, const FReal ToleranceScalar = UE_KINDA_SMALL_NUMBER) const
+	FORCEINLINE_DEBUGGABLE bool Private_TranslationEquals( const TransformVectorRegister& InTranslation, const FReal ToleranceScalar = UE_KINDA_SMALL_NUMBER) const
 	{			
 		const TransformVectorRegister Tolerance = VectorLoadFloat1(&ToleranceScalar);
 		// !( (FMath::Abs(X-V.X) > Tolerance) || (FMath::Abs(Y-V.Y) > Tolerance) || (FMath::Abs(Z-V.Z) > Tolerance) )
@@ -727,7 +727,7 @@ private:
 		return !VectorAnyGreaterThan(TranslationDiff, Tolerance);
 	}
 
-	FORCEINLINE bool Private_Scale3DEquals( const TransformVectorRegister& InScale3D, const FReal ToleranceScalar = UE_KINDA_SMALL_NUMBER) const
+	FORCEINLINE_DEBUGGABLE bool Private_Scale3DEquals( const TransformVectorRegister& InScale3D, const FReal ToleranceScalar = UE_KINDA_SMALL_NUMBER) const
 	{
 		const TransformVectorRegister Tolerance = VectorLoadFloat1(&ToleranceScalar);
 		// !( (FMath::Abs(X-V.X) > Tolerance) || (FMath::Abs(Y-V.Y) > Tolerance) || (FMath::Abs(Z-V.Z) > Tolerance) )
@@ -738,68 +738,68 @@ private:
 public:
 
 	// Test if A's rotation equals B's rotation, within a tolerance. Preferred over "A.GetRotation().Equals(B.GetRotation())" because it is faster on some platforms.
-	FORCEINLINE static bool AreRotationsEqual(const TTransform<T>& A, const TTransform<T>& B, FReal Tolerance=UE_KINDA_SMALL_NUMBER)
+	FORCEINLINE_DEBUGGABLE static bool AreRotationsEqual(const TTransform<T>& A, const TTransform<T>& B, FReal Tolerance=UE_KINDA_SMALL_NUMBER)
 	{
 		return A.Private_RotationEquals(B.Rotation, Tolerance);
 	}
 
 	// Test if A's translation equals B's translation, within a tolerance. Preferred over "A.GetTranslation().Equals(B.GetTranslation())" because it avoids TransformVectorRegister->TVector<T> conversion.
-	FORCEINLINE static bool AreTranslationsEqual(const TTransform<T>& A, const TTransform<T>& B, FReal Tolerance=UE_KINDA_SMALL_NUMBER)
+	FORCEINLINE_DEBUGGABLE static bool AreTranslationsEqual(const TTransform<T>& A, const TTransform<T>& B, FReal Tolerance=UE_KINDA_SMALL_NUMBER)
 	{
 		return A.Private_TranslationEquals(B.Translation, Tolerance);
 	}
 
 	// Test if A's scale equals B's scale, within a tolerance. Preferred over "A.GetScale3D().Equals(B.GetScale3D())" because it avoids TransformVectorRegister->TVector<T> conversion.
-	FORCEINLINE static bool AreScale3DsEqual(const TTransform<T>& A, const TTransform<T>& B, FReal Tolerance=UE_KINDA_SMALL_NUMBER)
+	FORCEINLINE_DEBUGGABLE static bool AreScale3DsEqual(const TTransform<T>& A, const TTransform<T>& B, FReal Tolerance=UE_KINDA_SMALL_NUMBER)
 	{
 		return A.Private_Scale3DEquals(B.Scale3D, Tolerance);
 	}
 
 
 	// Test if this Transform's rotation equals another's rotation, within a tolerance. Preferred over "GetRotation().Equals(Other.GetRotation())" because it is faster on some platforms.
-	FORCEINLINE bool RotationEquals(const TTransform<T>& Other, FReal Tolerance=UE_KINDA_SMALL_NUMBER) const
+	FORCEINLINE_DEBUGGABLE bool RotationEquals(const TTransform<T>& Other, FReal Tolerance=UE_KINDA_SMALL_NUMBER) const
 	{
 		return AreRotationsEqual(*this, Other, Tolerance);
 	}
 
 	// Test if this Transform's translation equals another's translation, within a tolerance. Preferred over "GetTranslation().Equals(Other.GetTranslation())" because it avoids TransformVectorRegister->TVector<T> conversion.
-	FORCEINLINE bool TranslationEquals(const TTransform<T>& Other, FReal Tolerance=UE_KINDA_SMALL_NUMBER) const
+	FORCEINLINE_DEBUGGABLE bool TranslationEquals(const TTransform<T>& Other, FReal Tolerance=UE_KINDA_SMALL_NUMBER) const
 	{
 		return AreTranslationsEqual(*this, Other, Tolerance);
 	}
 
 	// Test if this Transform's scale equals another's scale, within a tolerance. Preferred over "GetScale3D().Equals(Other.GetScale3D())" because it avoids TransformVectorRegister->TVector<T> conversion.
-	FORCEINLINE bool Scale3DEquals(const TTransform<T>& Other, FReal Tolerance=UE_KINDA_SMALL_NUMBER) const
+	FORCEINLINE_DEBUGGABLE bool Scale3DEquals(const TTransform<T>& Other, FReal Tolerance=UE_KINDA_SMALL_NUMBER) const
 	{
 		return AreScale3DsEqual(*this, Other, Tolerance);
 	}
 
 	// Test if all components of the transforms are equal, within a tolerance.
-	FORCEINLINE bool Equals(const TTransform<T>& Other, FReal Tolerance=UE_KINDA_SMALL_NUMBER) const
+	FORCEINLINE_DEBUGGABLE bool Equals(const TTransform<T>& Other, FReal Tolerance=UE_KINDA_SMALL_NUMBER) const
 	{
 		return Private_TranslationEquals(Other.Translation, Tolerance) && Private_RotationEquals(Other.Rotation, Tolerance) && Private_Scale3DEquals(Other.Scale3D, Tolerance);
 	}
 
 	// Test if all components of the transform property are equal.
-	FORCEINLINE bool Identical(const TTransform<T>* Other, uint32 PortFlags) const
+	FORCEINLINE_DEBUGGABLE bool Identical(const TTransform<T>* Other, uint32 PortFlags) const
 	{
 		return Equals(*Other, 0.f);
 	}
 
 	// Test if rotation and translation components of the transforms are equal, within a tolerance.
-	FORCEINLINE bool EqualsNoScale(const TTransform<T>& Other, FReal Tolerance=UE_KINDA_SMALL_NUMBER) const
+	FORCEINLINE_DEBUGGABLE bool EqualsNoScale(const TTransform<T>& Other, FReal Tolerance=UE_KINDA_SMALL_NUMBER) const
 	{
 		return Private_TranslationEquals(Other.Translation, Tolerance) && Private_RotationEquals(Other.Rotation, Tolerance);
 	}
 
-	FORCEINLINE static void Multiply(TTransform<T>* OutTransform, const TTransform<T>* A, const TTransform<T>* B);
+	FORCEINLINE_DEBUGGABLE static void Multiply(TTransform<T>* OutTransform, const TTransform<T>* A, const TTransform<T>* B);
 	/**
 	 * Sets the components
 	 * @param InRotation The new value for the Rotation component
 	 * @param InTranslation The new value for the Translation component
 	 * @param InScale3D The new value for the Scale3D component
 	 */
-	FORCEINLINE void SetComponents(const TQuat<T>& InRotation, const TVector<T>& InTranslation, const TVector<T>& InScale3D) 
+	FORCEINLINE_DEBUGGABLE void SetComponents(const TQuat<T>& InRotation, const TVector<T>& InTranslation, const TVector<T>& InScale3D) 
 	{
 		Rotation = VectorLoadAligned(&InRotation);
 		Translation = VectorLoadFloat3_W0(&InTranslation);
@@ -814,7 +814,7 @@ public:
 	 *   Translation = (0,0,0)
 	 *   Scale3D = (1,1,1)
 	 */
-	FORCEINLINE void SetIdentity()
+	FORCEINLINE_DEBUGGABLE void SetIdentity()
 	{
 		// Rotation = {0,0,0,1)
 		Rotation = GlobalVectorConstants::Float0001; //VectorSet_W1( VectorZero() );
@@ -830,7 +830,7 @@ public:
 	 *   Translation = (0,0,0)
 	 *   Scale3D = (0,0,0)
 	 */
-	FORCEINLINE void SetIdentityZeroScale()
+	FORCEINLINE_DEBUGGABLE void SetIdentityZeroScale()
 	{
 		// Rotation = {0,0,0,1)
 		Rotation = GlobalVectorConstants::Float0001; //VectorSet_W1( VectorZero() );
@@ -844,7 +844,7 @@ public:
 	 * Scales the Scale3D component by a new factor
 	 * @param Scale3DMultiplier The value to multiply Scale3D with
 	 */
-	FORCEINLINE void MultiplyScale3D(const TVector<T>& Scale3DMultiplier)
+	FORCEINLINE_DEBUGGABLE void MultiplyScale3D(const TVector<T>& Scale3DMultiplier)
 	{
 		Scale3D = VectorMultiply(Scale3D, VectorLoadFloat3_W0(&Scale3DMultiplier));
 		DiagnosticCheckNaN_Scale3D();
@@ -854,14 +854,14 @@ public:
 	 * Sets the translation component
 	 * @param NewTranslation The new value for the translation component
 	 */
-	FORCEINLINE void SetTranslation(const TVector<T>& NewTranslation)
+	FORCEINLINE_DEBUGGABLE void SetTranslation(const TVector<T>& NewTranslation)
 	{
 		Translation = VectorLoadFloat3_W0(&NewTranslation);
 		DiagnosticCheckNaN_Translate();
 	}
 
 	/** Copy translation from another TTransform<T>. */
-	FORCEINLINE void CopyTranslation(const TTransform<T>& Other)
+	FORCEINLINE_DEBUGGABLE void CopyTranslation(const TTransform<T>& Other)
 	{
 		Translation = Other.Translation;
 	}
@@ -870,7 +870,7 @@ public:
 	 * Concatenates another rotation to this transformation 
 	 * @param DeltaRotation The rotation to concatenate in the following fashion: Rotation = Rotation * DeltaRotation
 	 */
-	FORCEINLINE void ConcatenateRotation(const TQuat<T>& DeltaRotation)
+	FORCEINLINE_DEBUGGABLE void ConcatenateRotation(const TQuat<T>& DeltaRotation)
 	{		
 		Rotation = VectorQuaternionMultiply2(Rotation, VectorLoadAligned(&DeltaRotation));		
 		DiagnosticCheckNaN_Rotate();
@@ -880,7 +880,7 @@ public:
 	 * Adjusts the translation component of this transformation 
 	 * @param DeltaTranslation The translation to add in the following fashion: Translation += DeltaTranslation
 	 */
-	FORCEINLINE void AddToTranslation(const TVector<T>& DeltaTranslation)
+	FORCEINLINE_DEBUGGABLE void AddToTranslation(const TVector<T>& DeltaTranslation)
 	{		
 		Translation = VectorAdd(Translation, VectorLoadFloat3_W0(&DeltaTranslation));
 		DiagnosticCheckNaN_Translate();
@@ -890,7 +890,7 @@ public:
 	 * Add the translations from two FTransforms and return the result.
 	 * @return A.Translation + B.Translation
 	 */
-	FORCEINLINE static TVector<T> AddTranslations(const TTransform<T>& A, const TTransform<T>& B)
+	FORCEINLINE_DEBUGGABLE static TVector<T> AddTranslations(const TTransform<T>& A, const TTransform<T>& B)
 	{
 		TVector<T> Result;
 		VectorStoreFloat3(VectorAdd(A.Translation, B.Translation), &Result);
@@ -901,7 +901,7 @@ public:
 	 * Subtract translations from two FTransforms and return the difference.
 	 * @return A.Translation - B.Translation.
 	 */
-	FORCEINLINE static TVector<T> SubtractTranslations(const TTransform<T>& A, const TTransform<T>& B)
+	FORCEINLINE_DEBUGGABLE static TVector<T> SubtractTranslations(const TTransform<T>& A, const TTransform<T>& B)
 	{
 		TVector<T> Result;
 		VectorStoreFloat3(VectorSubtract(A.Translation, B.Translation), &Result);
@@ -912,14 +912,14 @@ public:
 	 * Sets the rotation component
 	 * @param NewRotation The new value for the rotation component
 	 */
-	FORCEINLINE void SetRotation(const TQuat<T>& NewRotation)
+	FORCEINLINE_DEBUGGABLE void SetRotation(const TQuat<T>& NewRotation)
 	{
 		Rotation = VectorLoadAligned(&NewRotation);
 		DiagnosticCheckNaN_Rotate();
 	}
 
 	/** Copy rotation from another TTransform<T>. */
-	FORCEINLINE void CopyRotation(const TTransform<T>& Other)
+	FORCEINLINE_DEBUGGABLE void CopyRotation(const TTransform<T>& Other)
 	{
 		Rotation = Other.Rotation;
 	}
@@ -928,14 +928,14 @@ public:
 	 * Sets the Scale3D component
 	 * @param NewScale3D The new value for the Scale3D component
 	 */
-	FORCEINLINE void SetScale3D(const TVector<T>& NewScale3D)
+	FORCEINLINE_DEBUGGABLE void SetScale3D(const TVector<T>& NewScale3D)
 	{
 		Scale3D = VectorLoadFloat3_W0(&NewScale3D);
 		DiagnosticCheckNaN_Scale3D();
 	}
 
 	/** Copy scale from another TTransform<T>. */
-	FORCEINLINE void CopyScale3D(const TTransform<T>& Other)
+	FORCEINLINE_DEBUGGABLE void CopyScale3D(const TTransform<T>& Other)
 	{
 		Scale3D = Other.Scale3D;
 	}
@@ -945,7 +945,7 @@ public:
 	 * @param NewTranslation The new value for the translation component
 	 * @param NewScale3D The new value for the Scale3D component
 	 */
-	FORCEINLINE void SetTranslationAndScale3D(const TVector<T>& NewTranslation, const TVector<T>& NewScale3D)
+	FORCEINLINE_DEBUGGABLE void SetTranslationAndScale3D(const TVector<T>& NewTranslation, const TVector<T>& NewScale3D)
 	{
 		Translation = VectorLoadFloat3_W0(&NewTranslation);
 		Scale3D = VectorLoadFloat3_W0(&NewScale3D);
@@ -977,7 +977,7 @@ public:
 	*
 	* @param SourceAtom The other transform to accumulate into this one
 	*/
-	FORCEINLINE void Accumulate(const TTransform<T>& SourceAtom)
+	FORCEINLINE_DEBUGGABLE void Accumulate(const TTransform<T>& SourceAtom)
 	{
 		// Rotation = SourceAtom.Rotation * Rotation;
 		Rotation = VectorQuaternionMultiply2(SourceAtom.Rotation, Rotation);
@@ -1005,7 +1005,7 @@ public:
 	* @param Atom The other transform to accumulate into this one
 	* @param BlendWeight The weight to multiply Atom by before it is accumulated.
 	*/
-	FORCEINLINE void Accumulate(const TTransform<T>& Atom, const ScalarRegister& BlendWeightScalar)
+	FORCEINLINE_DEBUGGABLE void Accumulate(const TTransform<T>& Atom, const ScalarRegister& BlendWeightScalar)
 	{
 		// SourceAtom = Atom * BlendWeight;
 		const TransformVectorRegister BlendWeight(BlendWeightScalar.Value);
@@ -1032,7 +1032,7 @@ public:
 	 * @param DeltaAtom The other transform to accumulate into this one
 	 * @param Weight The weight to multiply DeltaAtom by before it is accumulated.
 	 */
-	FORCEINLINE void AccumulateWithShortestRotation(const TTransform<T>& DeltaAtom, const ScalarRegister& BlendWeightScalar)
+	FORCEINLINE_DEBUGGABLE void AccumulateWithShortestRotation(const TTransform<T>& DeltaAtom, const ScalarRegister& BlendWeightScalar)
 	{
 		const TransformVectorRegister BlendWeight(BlendWeightScalar.Value);
 		const TransformVectorRegister BlendedRotation = VectorMultiply(DeltaAtom.Rotation, BlendWeight);
@@ -1063,7 +1063,7 @@ public:
 	* @param Atom The other transform to accumulate into this one
 	* @param BlendWeight The weight to multiply Atom by before it is accumulated.
 	*/
-	FORCEINLINE void AccumulateWithAdditiveScale(const TTransform<T>& Atom, const ScalarRegister& BlendWeightScalar)
+	FORCEINLINE_DEBUGGABLE void AccumulateWithAdditiveScale(const TTransform<T>& Atom, const ScalarRegister& BlendWeightScalar)
 	{
 		const TransformVectorRegister BlendWeight(BlendWeightScalar.Value);
 		const TransformVectorRegister DefaultScale = GlobalVectorConstants::Float1110;
@@ -1093,7 +1093,7 @@ public:
 	 * @param SourceAtom2 The ending point source atom (used 100% if Alpha is 1)
 	 * @param Alpha The blending weight between SourceAtom1 and SourceAtom2
 	 */
-	FORCEINLINE void LerpTranslationScale3D(const TTransform<T>& SourceAtom1, const TTransform<T>& SourceAtom2, const ScalarRegister& AlphaScalar)
+	FORCEINLINE_DEBUGGABLE void LerpTranslationScale3D(const TTransform<T>& SourceAtom1, const TTransform<T>& SourceAtom2, const ScalarRegister& AlphaScalar)
 	{
 		const TransformVectorRegister Alpha(AlphaScalar.Value);
 		Translation	= FMath::Lerp<TransformVectorRegister>(SourceAtom1.Translation, SourceAtom2.Translation, Alpha);
@@ -1106,7 +1106,7 @@ public:
 	/**
 	 * Normalize the rotation component of this transformation
 	 */
-	FORCEINLINE void NormalizeRotation()
+	FORCEINLINE_DEBUGGABLE void NormalizeRotation()
 	{
 		Rotation = VectorNormalizeQuaternion(Rotation);
 		DiagnosticCheckNaN_Rotate();
@@ -1117,7 +1117,7 @@ public:
 	 *
 	 * @return true if the rotation component is normalized, and false otherwise.
 	 */
-	FORCEINLINE bool IsRotationNormalized() const
+	FORCEINLINE_DEBUGGABLE bool IsRotationNormalized() const
 	{		
 		const TransformVectorRegister TestValue = VectorAbs(VectorSubtract(VectorOne(), VectorDot4(Rotation, Rotation)));
 		return !VectorAnyGreaterThan(TestValue, GlobalVectorConstants::ThreshQuatNormalized);
@@ -1135,7 +1135,7 @@ public:
 	 * @param SourceAtom The target transformation (used when BlendWeight = 1)
 	 * @param Alpha The blend weight between Identity and SourceAtom
 	 */
-	FORCEINLINE static void BlendFromIdentityAndAccumulate(TTransform<T>& FinalAtom, const TTransform<T>& SourceAtom, const ScalarRegister& BlendWeightScalar)
+	FORCEINLINE_DEBUGGABLE static void BlendFromIdentityAndAccumulate(TTransform<T>& FinalAtom, const TTransform<T>& SourceAtom, const ScalarRegister& BlendWeightScalar)
 	{
 		const TransformVectorRegister BlendWeight(BlendWeightScalar.Value);
 		const TransformVectorRegister Const0001 = GlobalVectorConstants::Float0001;
@@ -1183,7 +1183,7 @@ public:
 	 *
 	 * @return The rotation component
 	 */
-	FORCEINLINE TQuat<T> GetRotation() const
+	FORCEINLINE_DEBUGGABLE TQuat<T> GetRotation() const
 	{
 		DiagnosticCheckNaN_Rotate();
 		TQuat<T> OutRotation;
@@ -1196,7 +1196,7 @@ public:
 	 *
 	 * @return The translation component
 	 */
-	FORCEINLINE TVector<T> GetTranslation() const
+	FORCEINLINE_DEBUGGABLE TVector<T> GetTranslation() const
 	{
 		DiagnosticCheckNaN_Translate();
 		TVector<T> OutTranslation;
@@ -1209,7 +1209,7 @@ public:
 	 *
 	 * @return The Scale3D component
 	 */
-	FORCEINLINE TVector<T> GetScale3D() const
+	FORCEINLINE_DEBUGGABLE TVector<T> GetScale3D() const
 	{
 		DiagnosticCheckNaN_Scale3D();
 		TVector<T> OutScale3D;
@@ -1222,7 +1222,7 @@ public:
 	 *
 	 * @param SrcBA The transform to copy rotation and Scale3D from
 	 */
-	FORCEINLINE void CopyRotationPart(const TTransform<T>& SrcBA)
+	FORCEINLINE_DEBUGGABLE void CopyRotationPart(const TTransform<T>& SrcBA)
 	{
 		Rotation = SrcBA.Rotation;
 		Scale3D = SrcBA.Scale3D;
@@ -1236,7 +1236,7 @@ public:
 	 *
 	 * @param SrcBA The transform to copy translation and Scale3D from
 	 */
-	FORCEINLINE void CopyTranslationAndScale3D(const TTransform<T>& SrcBA)
+	FORCEINLINE_DEBUGGABLE void CopyTranslationAndScale3D(const TTransform<T>& SrcBA)
 	{
 		Translation = SrcBA.Translation;
 		Scale3D = SrcBA.Scale3D;
@@ -1272,7 +1272,7 @@ public:
 	}
 
 private:
-	FORCEINLINE void ToMatrixInternal( TransformVectorRegister& OutDiagonals, TransformVectorRegister& OutAdds, TransformVectorRegister& OutSubtracts ) const
+	FORCEINLINE_DEBUGGABLE void ToMatrixInternal( TransformVectorRegister& OutDiagonals, TransformVectorRegister& OutAdds, TransformVectorRegister& OutSubtracts ) const
 	{
 #if !(UE_BUILD_SHIPPING || UE_BUILD_TEST) && WITH_EDITORONLY_DATA
 		// Make sure Rotation is normalized when we turn it into a matrix.
@@ -1328,7 +1328,7 @@ private:
 		OutSubtracts = VectorMultiply(Subtracts , Scale3DYZXW);
 	}
 
-	FORCEINLINE void ToMatrixInternalNoScale( TransformVectorRegister& OutDiagonals, TransformVectorRegister& OutAdds, TransformVectorRegister& OutSubtracts ) const
+	FORCEINLINE_DEBUGGABLE void ToMatrixInternalNoScale( TransformVectorRegister& OutDiagonals, TransformVectorRegister& OutAdds, TransformVectorRegister& OutSubtracts ) const
 	{
 #if !(UE_BUILD_SHIPPING || UE_BUILD_TEST) && WITH_EDITORONLY_DATA
 		// Make sure Rotation is normalized when we turn it into a matrix.
@@ -1376,7 +1376,7 @@ private:
 	 * also returning BIG_NUMBER causes sequential NaN issues by multiplying 
 	 * so we hardcode as 0
 	 */
-	static FORCEINLINE TransformVectorRegister		GetSafeScaleReciprocal(const TransformVectorRegister& InScale, const ScalarRegister& Tolerance = ScalarRegister(GlobalVectorConstants::SmallNumber))
+	static FORCEINLINE_DEBUGGABLE TransformVectorRegister		GetSafeScaleReciprocal(const TransformVectorRegister& InScale, const ScalarRegister& Tolerance = ScalarRegister(GlobalVectorConstants::SmallNumber))
 	{		
 		// SafeReciprocalScale.X = (InScale.X == 0) ? 0.f : 1/InScale.X; // same for YZW
 
@@ -1395,7 +1395,7 @@ private:
 	}
 
 	/** Returns Inverse Transform of this TTransform<T> **/
-	FORCEINLINE TTransform<T> InverseFast() const
+	FORCEINLINE_DEBUGGABLE TTransform<T> InverseFast() const
 	{
 		// Inverse QST (A) = QST (~A)
 		// Since A*~A = Identity, 
@@ -1438,7 +1438,7 @@ private:
 	* @param  A Transform A.
 	* @param  B Transform B.
 	*/
-	FORCEINLINE static void MultiplyUsingMatrixWithScale(TTransform<T>* OutTransform, const TTransform<T>* A, const TTransform<T>* B);
+	FORCEINLINE_DEBUGGABLE static void MultiplyUsingMatrixWithScale(TTransform<T>* OutTransform, const TTransform<T>* A, const TTransform<T>* B);
 	/**
 	* Create a new transform from multiplications of given to matrices (AMatrix*BMatrix) using desired scale
 	* This is used by MultiplyUsingMatrixWithScale and GetRelativeTransformUsingMatrixWithScale
@@ -1449,7 +1449,7 @@ private:
 	* @param	DesiredScale - there is no check on if the magnitude is correct here. It assumes that is correct. 
 	* @param	OutTransform the constructed transform 
 	*/
-	FORCEINLINE static void ConstructTransformFromMatrixWithDesiredScale(const TMatrix<T>& AMatrix, const TMatrix<T>& BMatrix, const TransformVectorRegister& DesiredScale, TTransform<T>& OutTransform);
+	FORCEINLINE_DEBUGGABLE static void ConstructTransformFromMatrixWithDesiredScale(const TMatrix<T>& AMatrix, const TMatrix<T>& BMatrix, const TransformVectorRegister& DesiredScale, TTransform<T>& OutTransform);
 	/**
 	* Create a new transform: OutTransform = Base * Relative(-1) using the matrix while keeping the scale that's given by Base and Relative
 	* Please note that this operation is a lot more expensive than normal GetRelativeTrnasform
@@ -1485,7 +1485,7 @@ template<> CORE_API const FTransform3d FTransform3d::Identity;
 
 
 template<typename T>
-FORCEINLINE bool TTransform<T>::AnyHasNegativeScale(const TVector<T>& InScale3D, const TVector<T>& InOtherScale3D)
+FORCEINLINE_DEBUGGABLE bool TTransform<T>::AnyHasNegativeScale(const TVector<T>& InScale3D, const TVector<T>& InOtherScale3D)
 {
 	TransformVectorRegister VectorInScale3D = VectorLoadFloat3_W0(&InScale3D);
 	TransformVectorRegister VectorInOtherScale3D = VectorLoadFloat3_W0(&InOtherScale3D);
@@ -1495,7 +1495,7 @@ FORCEINLINE bool TTransform<T>::AnyHasNegativeScale(const TVector<T>& InScale3D,
 
 /** Scale the translation part of the Transform by the supplied vector. */
 template<typename T>
-FORCEINLINE void TTransform<T>::ScaleTranslation(const TVector<T>& InScale3D)
+FORCEINLINE_DEBUGGABLE void TTransform<T>::ScaleTranslation(const TVector<T>& InScale3D)
 {
 	TransformVectorRegister VectorInScale3D = VectorLoadFloat3_W0(&InScale3D);
 	Translation = VectorMultiply( Translation, VectorInScale3D );
@@ -1504,14 +1504,14 @@ FORCEINLINE void TTransform<T>::ScaleTranslation(const TVector<T>& InScale3D)
 
 /** Scale the translation part of the Transform by the supplied value. */
 template<typename T>
-FORCEINLINE void TTransform<T>::ScaleTranslation(const T& InScale)
+FORCEINLINE_DEBUGGABLE void TTransform<T>::ScaleTranslation(const T& InScale)
 {
 	ScaleTranslation( TVector<T>(InScale) );
 }
 
 // this function is from matrix, and all it does is to normalize rotation portion
 template<typename T>
-FORCEINLINE void TTransform<T>::RemoveScaling(T Tolerance/*=UE_SMALL_NUMBER*/)
+FORCEINLINE_DEBUGGABLE void TTransform<T>::RemoveScaling(T Tolerance/*=UE_SMALL_NUMBER*/)
 {
 	Scale3D = GlobalVectorConstants::Float1110;
 	NormalizeRotation();	
@@ -1521,13 +1521,13 @@ FORCEINLINE void TTransform<T>::RemoveScaling(T Tolerance/*=UE_SMALL_NUMBER*/)
 }
 
 template<typename T>
-FORCEINLINE void TTransform<T>::MultiplyUsingMatrixWithScale(TTransform<T>* OutTransform, const TTransform<T>* A, const TTransform<T>* B)
+FORCEINLINE_DEBUGGABLE void TTransform<T>::MultiplyUsingMatrixWithScale(TTransform<T>* OutTransform, const TTransform<T>* A, const TTransform<T>* B)
 {
 	ConstructTransformFromMatrixWithDesiredScale(A->ToMatrixWithScale(), B->ToMatrixWithScale(), VectorMultiply(A->Scale3D, B->Scale3D), *OutTransform);
 }
 
 template<typename T>
-FORCEINLINE void TTransform<T>::ConstructTransformFromMatrixWithDesiredScale(const UE::Math::TMatrix<T>& AMatrix, const UE::Math::TMatrix<T>& BMatrix, const TransformVectorRegister& DesiredScale, TTransform<T>& OutTransform)
+FORCEINLINE_DEBUGGABLE void TTransform<T>::ConstructTransformFromMatrixWithDesiredScale(const UE::Math::TMatrix<T>& AMatrix, const UE::Math::TMatrix<T>& BMatrix, const TransformVectorRegister& DesiredScale, TTransform<T>& OutTransform)
 {
 	// the goal of using M is to get the correct orientation
 	// but for translation, we still need scale
@@ -1560,7 +1560,7 @@ FORCEINLINE void TTransform<T>::ConstructTransformFromMatrixWithDesiredScale(con
 
 /** Returns Multiplied Transform of 2 FTransforms **/
 template<typename T>
-FORCEINLINE void TTransform<T>::Multiply(TTransform<T>* OutTransform, const TTransform<T>* A, const TTransform<T>* B)
+FORCEINLINE_DEBUGGABLE void TTransform<T>::Multiply(TTransform<T>* OutTransform, const TTransform<T>* A, const TTransform<T>* B)
 {
 	A->DiagnosticCheckNaN_All();
 	B->DiagnosticCheckNaN_All();
@@ -1615,7 +1615,7 @@ FORCEINLINE void TTransform<T>::Multiply(TTransform<T>* OutTransform, const TTra
  * Apply Scale to this transform
  */
 template<typename T>
-FORCEINLINE TTransform<T> TTransform<T>::GetScaled(T InScale) const
+FORCEINLINE_DEBUGGABLE TTransform<T> TTransform<T>::GetScaled(T InScale) const
 {
 	TTransform<T> A(*this);
 	
@@ -1631,7 +1631,7 @@ FORCEINLINE TTransform<T> TTransform<T>::GetScaled(T InScale) const
  * Apply Scale to this transform
  */
 template<typename T>
-FORCEINLINE TTransform<T> TTransform<T>::GetScaled(TVector<T> InScale) const
+FORCEINLINE_DEBUGGABLE TTransform<T> TTransform<T>::GetScaled(TVector<T> InScale) const
 {
 	TTransform<T> A(*this);
 
@@ -1644,7 +1644,7 @@ FORCEINLINE TTransform<T> TTransform<T>::GetScaled(TVector<T> InScale) const
 }
 
 template<typename T>
-FORCEINLINE TVector4<T> TTransform<T>::TransformFVector4NoScale(const TVector4<T>& V) const
+FORCEINLINE_DEBUGGABLE TVector4<T> TTransform<T>::TransformFVector4NoScale(const TVector4<T>& V) const
 {
 	DiagnosticCheckNaN_All();
 
@@ -1671,7 +1671,7 @@ FORCEINLINE TVector4<T> TTransform<T>::TransformFVector4NoScale(const TVector4<T
 }
 
 template<typename T>
-FORCEINLINE TVector4<T> TTransform<T>::TransformFVector4(const TVector4<T>& V) const
+FORCEINLINE_DEBUGGABLE TVector4<T> TTransform<T>::TransformFVector4(const TVector4<T>& V) const
 {
 	DiagnosticCheckNaN_All();
 
@@ -1700,7 +1700,7 @@ FORCEINLINE TVector4<T> TTransform<T>::TransformFVector4(const TVector4<T>& V) c
 
 
 template<typename T>
-FORCEINLINE TVector<T> TTransform<T>::TransformPosition(const TVector<T>& V) const
+FORCEINLINE_DEBUGGABLE TVector<T> TTransform<T>::TransformPosition(const TVector<T>& V) const
 {
 	DiagnosticCheckNaN_All();
 
@@ -1721,7 +1721,7 @@ FORCEINLINE TVector<T> TTransform<T>::TransformPosition(const TVector<T>& V) con
 }
 
 template<typename T>
-FORCEINLINE TVector<T> TTransform<T>::TransformPositionNoScale(const TVector<T>& V) const
+FORCEINLINE_DEBUGGABLE TVector<T> TTransform<T>::TransformPositionNoScale(const TVector<T>& V) const
 {
 	DiagnosticCheckNaN_All();
 
@@ -1741,7 +1741,7 @@ FORCEINLINE TVector<T> TTransform<T>::TransformPositionNoScale(const TVector<T>&
 }
 
 template<typename T>
-FORCEINLINE TVector<T> TTransform<T>::TransformVector(const TVector<T>& V) const
+FORCEINLINE_DEBUGGABLE TVector<T> TTransform<T>::TransformVector(const TVector<T>& V) const
 {
 	DiagnosticCheckNaN_All();
 
@@ -1757,7 +1757,7 @@ FORCEINLINE TVector<T> TTransform<T>::TransformVector(const TVector<T>& V) const
 }
 
 template<typename T>
-FORCEINLINE TVector<T> TTransform<T>::TransformVectorNoScale(const TVector<T>& V) const
+FORCEINLINE_DEBUGGABLE TVector<T> TTransform<T>::TransformVectorNoScale(const TVector<T>& V) const
 {
 	DiagnosticCheckNaN_All();
 
@@ -1773,7 +1773,7 @@ FORCEINLINE TVector<T> TTransform<T>::TransformVectorNoScale(const TVector<T>& V
 
 // do backward operation when inverse, translation -> rotation -> scale
 template<typename T>
-FORCEINLINE TVector<T> TTransform<T>::InverseTransformPosition(const TVector<T> &V) const
+FORCEINLINE_DEBUGGABLE TVector<T> TTransform<T>::InverseTransformPosition(const TVector<T> &V) const
 {
 	DiagnosticCheckNaN_All();
 
@@ -1798,7 +1798,7 @@ FORCEINLINE TVector<T> TTransform<T>::InverseTransformPosition(const TVector<T>
 
 // do backward operation when inverse, translation -> rotation
 template<typename T>
-FORCEINLINE TVector<T> TTransform<T>::InverseTransformPositionNoScale(const TVector<T> &V) const
+FORCEINLINE_DEBUGGABLE TVector<T> TTransform<T>::InverseTransformPositionNoScale(const TVector<T> &V) const
 {
 	DiagnosticCheckNaN_All();
 
@@ -1818,7 +1818,7 @@ FORCEINLINE TVector<T> TTransform<T>::InverseTransformPositionNoScale(const TVec
 
 // do backward operation when inverse, translation -> rotation -> scale
 template<typename T>
-FORCEINLINE TVector<T> TTransform<T>::InverseTransformVector(const TVector<T> &V) const
+FORCEINLINE_DEBUGGABLE TVector<T> TTransform<T>::InverseTransformVector(const TVector<T> &V) const
 {
 	DiagnosticCheckNaN_All();
 
@@ -1840,7 +1840,7 @@ FORCEINLINE TVector<T> TTransform<T>::InverseTransformVector(const TVector<T> &V
 
 // do backward operation when inverse, translation -> rotation
 template<typename T>
-FORCEINLINE TVector<T> TTransform<T>::InverseTransformVectorNoScale(const TVector<T> &V) const
+FORCEINLINE_DEBUGGABLE TVector<T> TTransform<T>::InverseTransformVectorNoScale(const TVector<T> &V) const
 {
 	DiagnosticCheckNaN_All();
 
@@ -1855,19 +1855,19 @@ FORCEINLINE TVector<T> TTransform<T>::InverseTransformVectorNoScale(const TVecto
 }
 
 template<typename T>
-FORCEINLINE TQuat<T> TTransform<T>::TransformRotation(const TQuat<T>& Q) const
+FORCEINLINE_DEBUGGABLE TQuat<T> TTransform<T>::TransformRotation(const TQuat<T>& Q) const
 {
 	return GetRotation() * Q;
 }
 
 template<typename T>
-FORCEINLINE TQuat<T> TTransform<T>::InverseTransformRotation(const TQuat<T>& Q) const
+FORCEINLINE_DEBUGGABLE TQuat<T> TTransform<T>::InverseTransformRotation(const TQuat<T>& Q) const
 {
 	return GetRotation().Inverse() * Q;
 }
 
 template<typename T>
-FORCEINLINE TTransform<T> TTransform<T>::operator*(const TTransform<T>& Other) const
+FORCEINLINE_DEBUGGABLE TTransform<T> TTransform<T>::operator*(const TTransform<T>& Other) const
 {
 	TTransform<T> Output;
 	Multiply(&Output, this, &Other);
@@ -1875,13 +1875,13 @@ FORCEINLINE TTransform<T> TTransform<T>::operator*(const TTransform<T>& Other) c
 }
 
 template<typename T>
-FORCEINLINE void TTransform<T>::operator*=(const TTransform<T>& Other)
+FORCEINLINE_DEBUGGABLE void TTransform<T>::operator*=(const TTransform<T>& Other)
 {
 	Multiply(this, this, &Other);
 }
 
 template<typename T>
-FORCEINLINE TTransform<T> TTransform<T>::operator*(const TQuat<T>& Other) const
+FORCEINLINE_DEBUGGABLE TTransform<T> TTransform<T>::operator*(const TQuat<T>& Other) const
 {
 	TTransform<T> Output, OtherTransform(Other, TVector<T>::ZeroVector, TVector<T>::OneVector);
 	Multiply(&Output, this, &OtherTransform);
@@ -1889,7 +1889,7 @@ FORCEINLINE TTransform<T> TTransform<T>::operator*(const TQuat<T>& Other) const
 }
 
 template<typename T>
-FORCEINLINE void TTransform<T>::operator*=(const TQuat<T>& Other)
+FORCEINLINE_DEBUGGABLE void TTransform<T>::operator*=(const TQuat<T>& Other)
 {
 	TTransform<T> OtherTransform(Other, TVector<T>::ZeroVector, TVector<T>::OneVector);
 	Multiply(this, this, &OtherTransform);
@@ -1897,7 +1897,7 @@ FORCEINLINE void TTransform<T>::operator*=(const TQuat<T>& Other)
 
 // x = 0, y = 1, z = 2
 template<typename T>
-FORCEINLINE TVector<T> TTransform<T>::GetScaledAxis( EAxis::Type InAxis ) const
+FORCEINLINE_DEBUGGABLE TVector<T> TTransform<T>::GetScaledAxis( EAxis::Type InAxis ) const
 {
 	if ( InAxis == EAxis::X )
 	{
@@ -1913,7 +1913,7 @@ FORCEINLINE TVector<T> TTransform<T>::GetScaledAxis( EAxis::Type InAxis ) const
 
 // x = 0, y = 1, z = 2
 template<typename T>
-FORCEINLINE TVector<T> TTransform<T>::GetUnitAxis( EAxis::Type InAxis ) const
+FORCEINLINE_DEBUGGABLE TVector<T> TTransform<T>::GetUnitAxis( EAxis::Type InAxis ) const
 {
 	if ( InAxis == EAxis::X )
 	{
@@ -1928,7 +1928,7 @@ FORCEINLINE TVector<T> TTransform<T>::GetUnitAxis( EAxis::Type InAxis ) const
 }
 
 template<typename T>
-FORCEINLINE void TTransform<T>::Mirror(EAxis::Type MirrorAxis, EAxis::Type FlipAxis)
+FORCEINLINE_DEBUGGABLE void TTransform<T>::Mirror(EAxis::Type MirrorAxis, EAxis::Type FlipAxis)
 {
 	// We do convert to Matrix for mirroring. 
 	UE::Math::TMatrix<T> M = ToMatrixWithScale();
@@ -1991,7 +1991,7 @@ inline float TTransform<T>::GetMinimumAxisScale() const
  * so we hardcode as 0
  */
 template<typename T>
-FORCEINLINE TVector<T> TTransform<T>::GetSafeScaleReciprocal(const TVector<T>& InScale, T Tolerance)
+FORCEINLINE_DEBUGGABLE TVector<T> TTransform<T>::GetSafeScaleReciprocal(const TVector<T>& InScale, T Tolerance)
 {
 	TVector<T> SafeReciprocalScale;
 	if (FMath::Abs(InScale.X) <= Tolerance)
