[AutoRTFM] Fold new write log entries into the last entry, if possible.

For consecutive write log entries (address starts at the end of the last write), grow the last record instead of appending a new record.

Reduces the number of write log records for a AActor (re)construction from 173 to 72.
This is a 808 byte reduction, and a 20% overall memory saving when also considering payload data.

diff --git a/Engine/Source/Programs/AutoRTFMTests/Private/WriteLogTests.cpp b/Engine/Source/Programs/AutoRTFMTests/Private/WriteLogTests.cpp
--- a/Engine/Source/Programs/AutoRTFMTests/Private/WriteLogTests.cpp
+++ b/Engine/Source/Programs/AutoRTFMTests/Private/WriteLogTests.cpp
@@ -109,38 +109,50 @@ TEST_CASE("FWriteLog")
 		}
 
 		REQUIRE(!WriteLog.IsEmpty());
-		REQUIRE(WriteLog.Num() == NumEntries);
+		REQUIRE(WriteLog.Num() <= NumEntries); // WriteLog may append back-to-back entries.
 		REQUIRE(WriteLog.TotalSize() == Data.size());
 
 		SECTION("Forwards iterator")
 		{
-			size_t Index = 0;
+			size_t EntryIndex = 0;
 			for (AutoRTFM::FWriteLogEntry Got : WriteLog)
 			{
-				const AutoRTFM::FWriteLogEntry& Expect = Entries[Index];
-				REQUIRE(Got.LogicalAddress == Expect.LogicalAddress);
-				REQUIRE(Got.Size == Expect.Size);
-				REQUIRE(memcmp(Got.Data, Expect.Data, Got.Size) == 0);
-				REQUIRE(Got.bNoMemoryValidation == Expect.bNoMemoryValidation);
-				Index++;
+				size_t Offset = 0;
+				while (Offset < Got.Size)
+				{
+					const AutoRTFM::FWriteLogEntry& Expect = Entries[EntryIndex];
+					REQUIRE(Got.LogicalAddress + Offset == Expect.LogicalAddress);
+					REQUIRE(Got.bNoMemoryValidation == Expect.bNoMemoryValidation);
+					REQUIRE(Expect.Size + Offset <= Got.Size);
+					REQUIRE(memcmp(Got.Data + Offset, Expect.Data, Expect.Size) == 0);
+					EntryIndex++;
+					Offset += Expect.Size;
+				}
+				REQUIRE(Offset == Got.Size);
 			}
-			REQUIRE(Index == Entries.size());
+			REQUIRE(EntryIndex == Entries.size());
 		}
 
 		SECTION("Reverse iterator")
 		{
-			size_t Index = NumEntries - 1;
+			size_t EntryIndex = Entries.size() - 1;
 			for (AutoRTFM::FWriteLog::ReverseIterator It = WriteLog.rbegin(); It != WriteLog.rend(); ++It)
 			{
-				AutoRTFM::FWriteLogEntry Got = *It;
-				const AutoRTFM::FWriteLogEntry& Expect = Entries[Index];
-				REQUIRE(Got.LogicalAddress == Expect.LogicalAddress);
-				REQUIRE(Got.Size == Expect.Size);
-				REQUIRE(memcmp(Got.Data, Expect.Data, Got.Size) == 0);
-				REQUIRE(Got.bNoMemoryValidation == Expect.bNoMemoryValidation);
-				Index--;
+				const AutoRTFM::FWriteLogEntry Got = *It;
+				size_t Offset = Got.Size;
+				while (Offset > 0)
+				{
+					const AutoRTFM::FWriteLogEntry& Expect = Entries[EntryIndex];
+					Offset -= Expect.Size;
+					REQUIRE(Offset < Got.Size);
+					REQUIRE(Got.LogicalAddress + Offset == Expect.LogicalAddress);
+					REQUIRE(Got.bNoMemoryValidation == Expect.bNoMemoryValidation);
+					REQUIRE(Expect.Size + Offset <= Got.Size);
+					REQUIRE(memcmp(Got.Data + Offset, Expect.Data, Expect.Size) == 0);
+					EntryIndex--;
+				}
 			}
-			REQUIRE(Index + 1 == 0);
+			REQUIRE(EntryIndex + 1 == 0);
 		}
 
 		SECTION("Reset")
diff --git a/Engine/Source/Runtime/Core/Private/AutoRTFM/WriteLog.h b/Engine/Source/Runtime/Core/Private/AutoRTFM/WriteLog.h
--- a/Engine/Source/Runtime/Core/Private/AutoRTFM/WriteLog.h
+++ b/Engine/Source/Runtime/Core/Private/AutoRTFM/WriteLog.h
@@ -116,22 +116,51 @@ namespace AutoRTFM
 				return LastRecord > FirstRecord;
 			}
 
+			// The result enumerator of Push()
+			enum class EPushResult
+			{
+				// The block does not have enough capacity to fit the entry.
+				Full,
+				// The block added the entry as a new write.
+				Added,
+				// The block folded the result into the end of the last write.
+				Folded,
+			};
+
 			// Attempts to add the entry into this block by copying the entry's data and creating a
 			// new record.
 			// Returns true if the entry was added, or false if the block does not have the capacity
 			// for the entry.
-			UE_AUTORTFM_FORCEINLINE bool Push(FWriteLogEntry Entry)
+			UE_AUTORTFM_FORCEINLINE EPushResult Push(FWriteLogEntry Entry)
 			{
-				if (DataEnd + Entry.Size > reinterpret_cast<std::byte*>(LastRecord - 1))
+				EPushResult Result = EPushResult::Full;
+
+				if (!IsEmpty() &&
+					reinterpret_cast<uintptr_t>(Entry.LogicalAddress) == LastRecord->Address + LastRecord->Size &&
+					LastRecord->bNoMemoryValidation == Entry.bNoMemoryValidation)
 				{
-					// Entry's data + new record does not fit in the block's remaining space.
-					return false;
-				}
+					if (DataEnd + Entry.Size > reinterpret_cast<std::byte*>(LastRecord))
+					{
+						// Entry's data does not fit in the block's remaining space.
+						return EPushResult::Full;
+					}
 
-				LastRecord--;
-				LastRecord->Address = reinterpret_cast<uintptr_t>(Entry.LogicalAddress);
-				LastRecord->Size = Entry.Size;
-				LastRecord->bNoMemoryValidation = Entry.bNoMemoryValidation;
+					LastRecord->Size += Entry.Size;
+					Result = EPushResult::Folded;
+				}
+				else
+				{
+					if (DataEnd + Entry.Size > reinterpret_cast<std::byte*>(LastRecord - 1))
+					{
+						// Entry's data + new record does not fit in the block's remaining space.
+						return EPushResult::Full;
+					}
+					LastRecord--;
+					LastRecord->Address = reinterpret_cast<uintptr_t>(Entry.LogicalAddress);
+					LastRecord->Size = Entry.Size;
+					LastRecord->bNoMemoryValidation = Entry.bNoMemoryValidation;
+					Result = EPushResult::Added;
+				}
 
 				memcpy(DataEnd, Entry.Data, Entry.Size);
 
@@ -141,7 +170,7 @@ namespace AutoRTFM
 				ASSERT(DataEnd <= reinterpret_cast<std::byte*>(LastRecord));
 #endif
 
-				return true;
+				return Result;
 			}
 
 			// The next block in the linked list.
@@ -178,7 +207,12 @@ namespace AutoRTFM
 			ASSERT(Entry.Size <= FWriteLogEntry::MaxSize);
 			ASSERT((reinterpret_cast<uintptr_t>(Entry.LogicalAddress) & 0xffff0000'00000000) == 0);
 
-			if (UNLIKELY(!TailBlock->Push(Entry)))
+			FBlock::EPushResult PushResult = TailBlock->Push(Entry);
+			if (PushResult == FBlock::EPushResult::Added)
+			{
+				NumEntries++;
+			}
+			else if (UNLIKELY(PushResult == FBlock::EPushResult::Full))
 			{
 				const size_t RequiredSize = Align(sizeof(FBlock) + Entry.Size, alignof(FRecord)) + sizeof(FRecord);
 				FBlock* NewBlock = FBlock::Allocate(std::max(RequiredSize, FBlock::DefaultSize));
@@ -186,11 +220,11 @@ namespace AutoRTFM
 				TailBlock->NextBlock = NewBlock;
 				TailBlock = NewBlock;
 
-				const bool bAdded = NewBlock->Push(Entry);
-				ASSERT(bAdded);
+				PushResult = NewBlock->Push(Entry);
+				ASSERT(PushResult == FBlock::EPushResult::Added);
+				NumEntries++;
 			}
 
-			NumEntries++;
 			TotalSizeBytes += Entry.Size;
 		}
 
