[AutoRTFM] Add AutoRTFM::FWriteLog visualizations to Unreal.natvis

Greatly helps debugging AutoRTFM memory validation failures.

diff --git a/Engine/Source/Runtime/Core/Private/AutoRTFM/WriteLog.h b/Engine/Source/Runtime/Core/Private/AutoRTFM/WriteLog.h
--- a/Engine/Source/Runtime/Core/Private/AutoRTFM/WriteLog.h
+++ b/Engine/Source/Runtime/Core/Private/AutoRTFM/WriteLog.h
@@ -41,6 +41,7 @@ namespace AutoRTFM
 
 	// FWriteLog holds an ordered list of write records which can be iterated
 	// forwards and backwards.
+	// Ensure changes to this class are kept in sync with Unreal.natvis.
 	class FWriteLog final
 	{
 		struct FRecord
@@ -52,6 +53,7 @@ namespace AutoRTFM
 		static_assert(sizeof(uintptr_t) == 8, "assumption: a pointer is 8 bytes");
 		static_assert(sizeof(FRecord) == 8);
 
+		// Ensure changes to this structure are kept in sync with Unreal.natvis.
 		struct FBlock final
 		{
 			// ┌────────┬────┬────┬────┬────┬────────────────┬────┬────┬────┬────┐
