Added TCopyQualifiersFromTo_T and TLosesQualifiersFromTo_V.

diff --git a/Engine/Plugins/Experimental/Avalanche/Source/AvalancheMedia/Private/Playable/AvaPlayableRemoteControl.cpp b/Engine/Plugins/Experimental/Avalanche/Source/AvalancheMedia/Private/Playable/AvaPlayableRemoteControl.cpp
--- a/Engine/Plugins/Experimental/Avalanche/Source/AvalancheMedia/Private/Playable/AvaPlayableRemoteControl.cpp
+++ b/Engine/Plugins/Experimental/Avalanche/Source/AvalancheMedia/Private/Playable/AvaPlayableRemoteControl.cpp
@@ -25,12 +25,12 @@ namespace UE::AvaPlayableRemoteControl::Private
 	 * @return the entity casted to the desired type if it was derived from it. nullptr otherwise
 	 */
 	template<typename To, typename From>
-	TSharedPtr<typename TCopyQualifiersFromTo<From, To>::Type> CastEntity(const TSharedPtr<From>& InEntity)
+	TSharedPtr<TCopyQualifiersFromTo_T<From, To>> CastEntity(const TSharedPtr<From>& InEntity)
 	{
 		const UScriptStruct* EntityStruct = InEntity.IsValid() ? InEntity->GetStruct() : nullptr;
 		if (EntityStruct && EntityStruct->IsChildOf(To::StaticStruct()))
 		{
-			return StaticCastSharedPtr<typename TCopyQualifiersFromTo<From, To>::Type>(InEntity);
+			return StaticCastSharedPtr<TCopyQualifiersFromTo_T<From, To>>(InEntity);
 		}
 		return nullptr;
 	}
diff --git a/Engine/Source/Runtime/Core/Public/Containers/Array.h b/Engine/Source/Runtime/Core/Public/Containers/Array.h
--- a/Engine/Source/Runtime/Core/Public/Containers/Array.h
+++ b/Engine/Source/Runtime/Core/Public/Containers/Array.h
@@ -317,7 +317,7 @@ namespace UE::Core::Private
 		if constexpr (std::is_same_v<FromAllocatorType, ToAllocatorType> || TCanMoveBetweenAllocators<FromAllocatorType, ToAllocatorType>::Value)
 		{
 			return
-				!TLosesQualifiersFromTo<FromElementType, ToElementType>::Value &&
+				!TLosesQualifiersFromTo_V<FromElementType, ToElementType> &&
 				(
 					std::is_same_v         <const ToElementType, const FromElementType> ||               // The element type of the container must be the same, or...
 					TIsBitwiseConstructible<UnqualifiedToElementType, UnqualifiedFromElementType>::Value // ... the element type of the source container must be bitwise constructible from the element type in the destination container
diff --git a/Engine/Source/Runtime/Core/Public/Templates/CopyQualifiersAndRefsFromTo.h b/Engine/Source/Runtime/Core/Public/Templates/CopyQualifiersAndRefsFromTo.h
--- a/Engine/Source/Runtime/Core/Public/Templates/CopyQualifiersAndRefsFromTo.h
+++ b/Engine/Source/Runtime/Core/Public/Templates/CopyQualifiersAndRefsFromTo.h
@@ -10,15 +10,15 @@
  * TCopyQualifiersFromTo_T<const T1, T2> == const T2
  * TCopyQualifiersFromTo_T<T1&, const T2> == const T2&
  */
-template <typename From, typename To> struct TCopyQualifiersAndRefsFromTo               { using Type = typename TCopyQualifiersFromTo<From, To>::Type;   };
-template <typename From, typename To> struct TCopyQualifiersAndRefsFromTo<From,   To& > { using Type = typename TCopyQualifiersFromTo<From, To>::Type&;  };
-template <typename From, typename To> struct TCopyQualifiersAndRefsFromTo<From,   To&&> { using Type = typename TCopyQualifiersFromTo<From, To>::Type&&; };
-template <typename From, typename To> struct TCopyQualifiersAndRefsFromTo<From&,  To  > { using Type = typename TCopyQualifiersFromTo<From, To>::Type&;  };
-template <typename From, typename To> struct TCopyQualifiersAndRefsFromTo<From&,  To& > { using Type = typename TCopyQualifiersFromTo<From, To>::Type&;  };
-template <typename From, typename To> struct TCopyQualifiersAndRefsFromTo<From&,  To&&> { using Type = typename TCopyQualifiersFromTo<From, To>::Type&;  };
-template <typename From, typename To> struct TCopyQualifiersAndRefsFromTo<From&&, To  > { using Type = typename TCopyQualifiersFromTo<From, To>::Type&&; };
-template <typename From, typename To> struct TCopyQualifiersAndRefsFromTo<From&&, To& > { using Type = typename TCopyQualifiersFromTo<From, To>::Type&;  };
-template <typename From, typename To> struct TCopyQualifiersAndRefsFromTo<From&&, To&&> { using Type = typename TCopyQualifiersFromTo<From, To>::Type&&; };
+template <typename From, typename To> struct TCopyQualifiersAndRefsFromTo               { using Type = TCopyQualifiersFromTo_T<From, To>;   };
+template <typename From, typename To> struct TCopyQualifiersAndRefsFromTo<From,   To& > { using Type = TCopyQualifiersFromTo_T<From, To>&;  };
+template <typename From, typename To> struct TCopyQualifiersAndRefsFromTo<From,   To&&> { using Type = TCopyQualifiersFromTo_T<From, To>&&; };
+template <typename From, typename To> struct TCopyQualifiersAndRefsFromTo<From&,  To  > { using Type = TCopyQualifiersFromTo_T<From, To>&;  };
+template <typename From, typename To> struct TCopyQualifiersAndRefsFromTo<From&,  To& > { using Type = TCopyQualifiersFromTo_T<From, To>&;  };
+template <typename From, typename To> struct TCopyQualifiersAndRefsFromTo<From&,  To&&> { using Type = TCopyQualifiersFromTo_T<From, To>&;  };
+template <typename From, typename To> struct TCopyQualifiersAndRefsFromTo<From&&, To  > { using Type = TCopyQualifiersFromTo_T<From, To>&&; };
+template <typename From, typename To> struct TCopyQualifiersAndRefsFromTo<From&&, To& > { using Type = TCopyQualifiersFromTo_T<From, To>&;  };
+template <typename From, typename To> struct TCopyQualifiersAndRefsFromTo<From&&, To&&> { using Type = TCopyQualifiersFromTo_T<From, To>&&; };
 
 template <typename From, typename To>
 using TCopyQualifiersAndRefsFromTo_T = typename TCopyQualifiersAndRefsFromTo<From, To>::Type;
diff --git a/Engine/Source/Runtime/Core/Public/Templates/CopyQualifiersFromTo.h b/Engine/Source/Runtime/Core/Public/Templates/CopyQualifiersFromTo.h
--- a/Engine/Source/Runtime/Core/Public/Templates/CopyQualifiersFromTo.h
+++ b/Engine/Source/Runtime/Core/Public/Templates/CopyQualifiersFromTo.h
@@ -5,10 +5,13 @@
 /**
  * Copies the cv-qualifiers from one type to another, e.g.:
  *
- * TCopyQualifiersFromTo<const    T1,       T2>::Type == const T2
- * TCopyQualifiersFromTo<volatile T1, const T2>::Type == const volatile T2
+ * TCopyQualifiersFromTo_T<const    T1,       T2> == const T2
+ * TCopyQualifiersFromTo_T<volatile T1, const T2> == const volatile T2
  */
 template <typename From, typename To> struct TCopyQualifiersFromTo                          { typedef                To Type; };
 template <typename From, typename To> struct TCopyQualifiersFromTo<const          From, To> { typedef const          To Type; };
 template <typename From, typename To> struct TCopyQualifiersFromTo<      volatile From, To> { typedef       volatile To Type; };
 template <typename From, typename To> struct TCopyQualifiersFromTo<const volatile From, To> { typedef const volatile To Type; };
+
+template <typename From, typename To>
+using TCopyQualifiersFromTo_T = typename TCopyQualifiersFromTo<From, To>::Type;
diff --git a/Engine/Source/Runtime/Core/Public/Templates/LosesQualifiersFromTo.h b/Engine/Source/Runtime/Core/Public/Templates/LosesQualifiersFromTo.h
--- a/Engine/Source/Runtime/Core/Public/Templates/LosesQualifiersFromTo.h
+++ b/Engine/Source/Runtime/Core/Public/Templates/LosesQualifiersFromTo.h
@@ -8,11 +8,14 @@
 /**
  * Tests if qualifiers are lost between one type and another, e.g.:
  *
- * TLosesQualifiersFromTo<const    T1,                T2>::Value == true
- * TLosesQualifiersFromTo<volatile T1, const volatile T2>::Value == false
+ * TLosesQualifiersFromTo_V<const    T1,                T2> == true
+ * TLosesQualifiersFromTo_V<volatile T1, const volatile T2> == false
  */
 template <typename From, typename To>
 struct TLosesQualifiersFromTo
 {
-	enum { Value = !std::is_same_v<typename TCopyQualifiersFromTo<From, To>::Type, To> };
+	enum { Value = !std::is_same_v<TCopyQualifiersFromTo_T<From, To>, To> };
 };
+
+template <typename From, typename To>
+inline constexpr bool TLosesQualifiersFromTo_V = TLosesQualifiersFromTo<From, To>::Value;
diff --git a/Engine/Source/Runtime/Core/Public/Templates/PointerIsConvertibleFromTo.h b/Engine/Source/Runtime/Core/Public/Templates/PointerIsConvertibleFromTo.h
--- a/Engine/Source/Runtime/Core/Public/Templates/PointerIsConvertibleFromTo.h
+++ b/Engine/Source/Runtime/Core/Public/Templates/PointerIsConvertibleFromTo.h
@@ -27,21 +27,21 @@ namespace UE::Core::Private::PointerIsConvertibleFromTo
 	struct TImpl<From, To, NoCVFrom, NoCVFrom>
 	{
 		// cv T* to cv T* conversions are always allowed as long as no CVs are lost
-		enum { Value = !TLosesQualifiersFromTo<From, To>::Value };
+		enum { Value = !TLosesQualifiersFromTo_V<From, To> };
 	};
 
 	template <typename From, typename To, typename NoCVFrom>
 	struct TImpl<From, To, NoCVFrom, void>
 	{
 		// cv T* to cv void* conversions are always allowed as long as no CVs are lost
-		enum { Value = !TLosesQualifiersFromTo<From, To>::Value };
+		enum { Value = !TLosesQualifiersFromTo_V<From, To> };
 	};
 
 	template <typename From, typename To>
 	struct TImpl<From, To, void, void>
 	{
 		// cv void* to cv void* conversions are always allowed as long as no CVs are lost
-		enum { Value = !TLosesQualifiersFromTo<From, To>::Value };
+		enum { Value = !TLosesQualifiersFromTo_V<From, To> };
 	};
 
 	template <typename From, typename To, typename NoCVTo>
diff --git a/Engine/Source/Runtime/Core/Public/UObject/WeakObjectPtrTemplates.h b/Engine/Source/Runtime/Core/Public/UObject/WeakObjectPtrTemplates.h
--- a/Engine/Source/Runtime/Core/Public/UObject/WeakObjectPtrTemplates.h
+++ b/Engine/Source/Runtime/Core/Public/UObject/WeakObjectPtrTemplates.h
@@ -91,7 +91,7 @@ public:
 	 */
 	template <
 		typename U
-		UE_REQUIRES(!TLosesQualifiersFromTo<U, T>::Value)
+		UE_REQUIRES(!TLosesQualifiersFromTo_V<U, T>)
 	>
 	FORCEINLINE TWeakObjectPtr& operator=(U* Object)
 	{
diff --git a/Engine/Source/Runtime/CoreUObject/Public/Templates/Casts.h b/Engine/Source/Runtime/CoreUObject/Public/Templates/Casts.h
--- a/Engine/Source/Runtime/CoreUObject/Public/Templates/Casts.h
+++ b/Engine/Source/Runtime/CoreUObject/Public/Templates/Casts.h
@@ -292,7 +292,7 @@ template< class T, class U > FORCEINLINE T* CastChecked( const TWeakObjectPtr<U>
 
 // object ptr versions
 template <typename To, typename From>
-FORCEINLINE typename TCopyQualifiersFromTo<From, To>::Type* Cast(const TObjectPtr<From>& InSrc)
+FORCEINLINE TCopyQualifiersFromTo_T<From, To>* Cast(const TObjectPtr<From>& InSrc)
 {
 	static_assert(sizeof(To) > 0 && sizeof(From) > 0, "Attempting to cast between incomplete types");
 
@@ -349,7 +349,7 @@ FORCEINLINE typename TCopyQualifiersFromTo<From, To>::Type* Cast(const TObjectPt
 }
 
 template <typename To, typename From>
-FORCEINLINE typename TCopyQualifiersFromTo<From, To>::Type* ExactCast(const TObjectPtr<From>& Src)
+FORCEINLINE TCopyQualifiersFromTo_T<From, To>* ExactCast(const TObjectPtr<From>& Src)
 {
 	static_assert(sizeof(To) > 0, "Attempting to cast to an incomplete type");
 
@@ -363,7 +363,7 @@ FORCEINLINE typename TCopyQualifiersFromTo<From, To>::Type* ExactCast(const TObj
 }
 
 template <typename To, typename From>
-FORCEINLINE typename TCopyQualifiersFromTo<From, To>::Type* CastChecked(const TObjectPtr<From>& Src, ECastCheckedType::Type CheckType = ECastCheckedType::NullChecked)
+FORCEINLINE TCopyQualifiersFromTo_T<From, To>* CastChecked(const TObjectPtr<From>& Src, ECastCheckedType::Type CheckType = ECastCheckedType::NullChecked)
 {
 	static_assert(sizeof(From) > 0 && sizeof(To) > 0, "Attempting to cast between incomplete types");
 
@@ -498,7 +498,7 @@ namespace UE::CoreUObject::Private
 		else
 		{
 			// Casting away const/volatile
-			static_assert(!TLosesQualifiersFromTo<From, ToValueType>::Value, "Conversion loses qualifiers");
+			static_assert(!TLosesQualifiersFromTo_V<From, ToValueType>, "Conversion loses qualifiers");
 
 			if constexpr (std::is_void_v<ToValueType>)
 			{
@@ -526,7 +526,7 @@ namespace UE::CoreUObject::Private
 		else
 		{
 			// Casting away const/volatile
-			static_assert(!TLosesQualifiersFromTo<FromValueType, ToValueType>::Value, "Conversion loses qualifiers");
+			static_assert(!TLosesQualifiersFromTo_V<FromValueType, ToValueType>, "Conversion loses qualifiers");
 
 			// T&& can only be cast to U&&
 			// http://en.cppreference.com/w/cpp/language/dynamic_cast
diff --git a/Engine/Source/Runtime/CoreUObject/Public/UObject/ObjectPtr.h b/Engine/Source/Runtime/CoreUObject/Public/UObject/ObjectPtr.h
--- a/Engine/Source/Runtime/CoreUObject/Public/UObject/ObjectPtr.h
+++ b/Engine/Source/Runtime/CoreUObject/Public/UObject/ObjectPtr.h
@@ -1036,7 +1036,7 @@ ToRawPtrTArrayUnsafe(ArrayType&& Array)
 	using ArrayElementType         = typename ArrayTypeNoRef::ElementType;
 	using ArrayAllocatorType       = typename ArrayTypeNoRef::AllocatorType;
 	using RawPointerType           = typename ObjectPtr_Private::TRawPointerType<ArrayElementType>::Type;
-	using QualifiedRawPointerType  = typename TCopyQualifiersFromTo<ArrayElementType, RawPointerType>::Type;
+	using QualifiedRawPointerType  = TCopyQualifiersFromTo_T<ArrayElementType, RawPointerType>;
 	using NewArrayType             = TArray<QualifiedRawPointerType, ArrayAllocatorType>;
 	using RefQualifiedNewArrayType = typename TCopyQualifiersAndRefsFromTo<ArrayType, NewArrayType>::Type;
 
diff --git a/Engine/Source/Runtime/CoreUObject/Public/UObject/ScriptInterface.h b/Engine/Source/Runtime/CoreUObject/Public/UObject/ScriptInterface.h
--- a/Engine/Source/Runtime/CoreUObject/Public/UObject/ScriptInterface.h
+++ b/Engine/Source/Runtime/CoreUObject/Public/UObject/ScriptInterface.h
@@ -139,7 +139,7 @@ class TScriptInterface : public FScriptInterface
 {
 public:
 	using InterfaceType = InInterfaceType;
-	using UObjectType = typename TCopyQualifiersFromTo<InterfaceType, UObject>::Type;
+	using UObjectType = TCopyQualifiersFromTo_T<InterfaceType, UObject>;
 
 	/**
 	 * Default constructor
diff --git a/Engine/Source/Runtime/CoreUObject/Public/UObject/WeakInterfacePtr.h b/Engine/Source/Runtime/CoreUObject/Public/UObject/WeakInterfacePtr.h
--- a/Engine/Source/Runtime/CoreUObject/Public/UObject/WeakInterfacePtr.h
+++ b/Engine/Source/Runtime/CoreUObject/Public/UObject/WeakInterfacePtr.h
@@ -17,7 +17,7 @@ template<class T>
 struct TWeakInterfacePtr
 {
 	using ElementType = T;
-	using UObjectType = typename TCopyQualifiersFromTo<T, UObject>::Type;
+	using UObjectType = TCopyQualifiersFromTo_T<T, UObject>;
 
 	FORCEINLINE TWeakInterfacePtr() = default;
 	FORCEINLINE TWeakInterfacePtr(const TWeakInterfacePtr& Other) = default;
@@ -32,11 +32,11 @@ struct TWeakInterfacePtr
 	 */
 	template<
 		typename U
-		UE_REQUIRES(std::is_convertible_v<U, typename TCopyQualifiersFromTo<U, UObject>::Type*>)
+		UE_REQUIRES(std::is_convertible_v<U, TCopyQualifiersFromTo_T<U, UObject>*>)
 	>
 	TWeakInterfacePtr(U&& Object)
 	{
-		InterfaceInstance = Cast<T>(ImplicitConv<typename TCopyQualifiersFromTo<U, UObject>::Type*>(Object));
+		InterfaceInstance = Cast<T>(ImplicitConv<TCopyQualifiersFromTo_T<U, UObject>*>(Object));
 		if (InterfaceInstance != nullptr)
 		{
 			ObjectInstance = Object;
diff --git a/Engine/Source/Runtime/Engine/Private/RepLayout.cpp b/Engine/Source/Runtime/Engine/Private/RepLayout.cpp
--- a/Engine/Source/Runtime/Engine/Private/RepLayout.cpp
+++ b/Engine/Source/Runtime/Engine/Private/RepLayout.cpp
@@ -132,10 +132,10 @@ namespace UEPushModelPrivate
 namespace UE_RepLayout_Private
 {
 	template<typename OutputType, typename CommandType, typename BufferType>
-	static typename TCopyQualifiersFromTo<BufferType, OutputType>::Type*
+	static TCopyQualifiersFromTo_T<BufferType, OutputType>*
 	GetTypedProperty(const BufferType& Buffer, const CommandType& Cmd)
 	{
-		using ConstOrNotOutputType = typename TCopyQualifiersFromTo<BufferType, OutputType>::Type;
+		using ConstOrNotOutputType = TCopyQualifiersFromTo_T<BufferType, OutputType>;
 
 		using BaseBufferType = typename TRemovePointer<typename TDecay<BufferType>::Type>::Type;
 
@@ -147,10 +147,10 @@ namespace UE_RepLayout_Private
 	}
 
 	template<typename OutputType, typename CommandType, enum ERepDataBufferType BufferDataType, typename BufferUnderlyingType>
-	static typename TCopyQualifiersFromTo<typename TRepDataBufferBase<BufferDataType, BufferUnderlyingType>::ConstOrNotVoid, OutputType>::Type*
+	static TCopyQualifiersFromTo_T<typename TRepDataBufferBase<BufferDataType, BufferUnderlyingType>::ConstOrNotVoid, OutputType>*
 	GetTypedProperty(const TRepDataBufferBase<BufferDataType, BufferUnderlyingType>& Buffer, const CommandType& Cmd)
 	{
-		using ConstOrNotOutputType = typename TCopyQualifiersFromTo<typename TRepDataBufferBase<BufferDataType, BufferUnderlyingType>::ConstOrNotVoid, OutputType>::Type;
+		using ConstOrNotOutputType = TCopyQualifiersFromTo_T<typename TRepDataBufferBase<BufferDataType, BufferUnderlyingType>::ConstOrNotVoid, OutputType>;
 
 		static_assert(!TIsPointer<OutputType>::Value, "GetTypedProperty invalid OutputType! Don't specify output as a pointer.");
 
diff --git a/Engine/Source/Runtime/Engine/Public/Net/RepLayout.h b/Engine/Source/Runtime/Engine/Public/Net/RepLayout.h
--- a/Engine/Source/Runtime/Engine/Public/Net/RepLayout.h
+++ b/Engine/Source/Runtime/Engine/Public/Net/RepLayout.h
@@ -73,7 +73,7 @@ namespace UE_RepLayout_Private
 	struct TRepDataBufferBase
 	{
 		static constexpr ERepDataBufferType Type = DataType;
-		using ConstOrNotVoid = typename TCopyQualifiersFromTo<ConstOrNotType, void>::Type;
+		using ConstOrNotVoid = TCopyQualifiersFromTo_T<ConstOrNotType, void>;
 
 	public:
 
diff --git a/Engine/Source/Runtime/MeshDescription/Public/AttributeArrayContainer.h b/Engine/Source/Runtime/MeshDescription/Public/AttributeArrayContainer.h
--- a/Engine/Source/Runtime/MeshDescription/Public/AttributeArrayContainer.h
+++ b/Engine/Source/Runtime/MeshDescription/Public/AttributeArrayContainer.h
@@ -392,7 +392,7 @@ template <typename AttributeType>
 class TArrayAttribute
 {
 	template <typename T> friend class TArrayAttribute;
-	using ArrayType = typename TCopyQualifiersFromTo<AttributeType, TAttributeArrayContainer<std::remove_cv_t<AttributeType>>>::Type;
+	using ArrayType = TCopyQualifiersFromTo_T<AttributeType, TAttributeArrayContainer<std::remove_cv_t<AttributeType>>>;
 
 public:
 	explicit TArrayAttribute(ArrayType& InArray, int32 InIndex)
diff --git a/Engine/Source/Runtime/MeshDescription/Public/MeshAttributeArray.h b/Engine/Source/Runtime/MeshDescription/Public/MeshAttributeArray.h
--- a/Engine/Source/Runtime/MeshDescription/Public/MeshAttributeArray.h
+++ b/Engine/Source/Runtime/MeshDescription/Public/MeshAttributeArray.h
@@ -836,8 +836,8 @@ class TMeshAttributesRef
 	template <typename T, typename U> friend class TMeshAttributesRef;
 
 public:
-	using BaseArrayType = typename TCopyQualifiersFromTo<AttributeType, FMeshAttributeArraySetBase>::Type;
-	using ArrayType = typename TCopyQualifiersFromTo<AttributeType, TMeshAttributeArraySet<std::remove_cv_t<AttributeType>>>::Type;
+	using BaseArrayType = TCopyQualifiersFromTo_T<AttributeType, FMeshAttributeArraySetBase>;
+	using ArrayType = TCopyQualifiersFromTo_T<AttributeType, TMeshAttributeArraySet<std::remove_cv_t<AttributeType>>>;
 
 	/** Constructor taking a pointer to a FMeshAttributeArraySetBase */
 	explicit TMeshAttributesRef(BaseArrayType* InArrayPtr = nullptr, uint32 InExtent = 1)
@@ -1040,9 +1040,9 @@ class TMeshAttributesRef<ElementIDType, TArrayView<AttributeType>>
 	template <typename T, typename U> friend class TMeshAttributesRef;
 
 public:
-	using BaseArrayType = typename TCopyQualifiersFromTo<AttributeType, FMeshAttributeArraySetBase>::Type;
-	using BoundedArrayType = typename TCopyQualifiersFromTo<AttributeType, TMeshAttributeArraySet<std::remove_cv_t<AttributeType>>>::Type;
-	using UnboundedArrayType = typename TCopyQualifiersFromTo<AttributeType, TMeshUnboundedAttributeArraySet<std::remove_cv_t<AttributeType>>>::Type;
+	using BaseArrayType = TCopyQualifiersFromTo_T<AttributeType, FMeshAttributeArraySetBase>;
+	using BoundedArrayType = TCopyQualifiersFromTo_T<AttributeType, TMeshAttributeArraySet<std::remove_cv_t<AttributeType>>>;
+	using UnboundedArrayType = TCopyQualifiersFromTo_T<AttributeType, TMeshUnboundedAttributeArraySet<std::remove_cv_t<AttributeType>>>;
 
 	/** Constructor taking a pointer to a TMeshAttributeArraySet */
 	explicit TMeshAttributesRef(BaseArrayType* InArrayPtr = nullptr, uint32 InExtent = 1)
@@ -1316,8 +1316,8 @@ class TMeshAttributesRef<ElementIDType, TArrayAttribute<AttributeType>>
 	template <typename T, typename U> friend class TMeshAttributesRef;
 
 public:
-	using BaseArrayType = typename TCopyQualifiersFromTo<AttributeType, FMeshAttributeArraySetBase>::Type;
-	using ArrayType = typename TCopyQualifiersFromTo<AttributeType, TMeshUnboundedAttributeArraySet<std::remove_cv_t<AttributeType>>>::Type;
+	using BaseArrayType = TCopyQualifiersFromTo_T<AttributeType, FMeshAttributeArraySetBase>;
+	using ArrayType = TCopyQualifiersFromTo_T<AttributeType, TMeshUnboundedAttributeArraySet<std::remove_cv_t<AttributeType>>>;
 
 	/** Constructor taking a pointer to a TMeshUnboundedAttributeArraySet */
 	explicit TMeshAttributesRef(BaseArrayType* InArrayPtr = nullptr, uint32 InExtent = 0)
