Fix ServerExec not working with AutoRTFM

#rnx

diff --git a/Engine/Source/Runtime/Core/Private/Misc/CoreMisc.cpp b/Engine/Source/Runtime/Core/Private/Misc/CoreMisc.cpp
--- a/Engine/Source/Runtime/Core/Private/Misc/CoreMisc.cpp
+++ b/Engine/Source/Runtime/Core/Private/Misc/CoreMisc.cpp
@@ -44,11 +44,11 @@ DEFINE_LOG_CATEGORY(LogCore);
 using FSelfRegisteredExecArray = TArray<FSelfRegisteringExec*, TInlineAllocator<8>>;
 
 // Lazy because pthread implementation doesn't like static initialization.
-FCriticalSection* GetExecRegistryLock()
+FTransactionallySafeCriticalSection* GetExecRegistryLock()
 {
 	// Note: Using FCriticalSection still allows calls to addition or removal
 	// to/from Execs and FSelfRegisteringExec::StaticExec on the same thread
-	static FCriticalSection ExecRegistryLock;
+	static FTransactionallySafeCriticalSection ExecRegistryLock;
 	return &ExecRegistryLock;
 }
 
@@ -61,20 +61,20 @@ FSelfRegisteredExecArray& GetExecRegistry()
 /** Constructor, registering this instance. */
 FSelfRegisteringExec::FSelfRegisteringExec()
 {
-	FScopeLock ScopeLock(GetExecRegistryLock());
+	UE::TScopeLock ScopeLock(*GetExecRegistryLock());
 	GetExecRegistry().Add( this );
 }
 
 /** Destructor, unregistering this instance. */
 FSelfRegisteringExec::~FSelfRegisteringExec()
 {
-	FScopeLock ScopeLock(GetExecRegistryLock());
+	UE::TScopeLock ScopeLock(*GetExecRegistryLock());
 	verify(GetExecRegistry().Remove( this ) == 1 );
 }
 
 bool FSelfRegisteringExec::StaticExec( UWorld* InWorld, const TCHAR* Cmd, FOutputDevice& Ar )
 {
-	FScopeLock ScopeLock(GetExecRegistryLock());
+	UE::TScopeLock ScopeLock(*GetExecRegistryLock());
 	for (FSelfRegisteringExec* Exe : GetExecRegistry())
 	{
 		if (Exe->Exec( InWorld, Cmd,Ar ))
