[LRUCache] Add method GetLeastRecentKey()

#rnx

diff --git a/Engine/Source/Runtime/Core/Public/Containers/LruCache.h b/Engine/Source/Runtime/Core/Public/Containers/LruCache.h
--- a/Engine/Source/Runtime/Core/Public/Containers/LruCache.h
+++ b/Engine/Source/Runtime/Core/Public/Containers/LruCache.h
@@ -582,6 +582,17 @@ public:
 		return LeastRecentElement;
 	}
 
+	/**
+	 * Return the least recent element key from the cache.
+	 *
+	 * @return Copy of least recent key.
+	 */
+	[[nodiscard]] FORCEINLINE KeyType GetLeastRecentKey() const
+	{
+		check(LeastRecent);
+		return LeastRecent->Key;
+	}
+
 public:
 
 	/**
