Move the transactionallysafemutex to Async as per core teams request.

diff --git a/Engine/Plugins/FX/Niagara/Source/Niagara/Public/NiagaraTypeRegistry.h b/Engine/Plugins/FX/Niagara/Source/Niagara/Public/NiagaraTypeRegistry.h
--- a/Engine/Plugins/FX/Niagara/Source/Niagara/Public/NiagaraTypeRegistry.h
+++ b/Engine/Plugins/FX/Niagara/Source/Niagara/Public/NiagaraTypeRegistry.h
@@ -3,9 +3,9 @@
 #pragma once
 
 #include "CoreMinimal.h"
+#include "Async/TransactionallySafeMutex.h"
 #include "Misc/ScopeLock.h"
 #include "Misc/SecureHash.h"
-#include "Misc/TransactionallySafeMutex.h"
 #include "UObject/GCObject.h"
 #include "UObject/UnrealType.h"
 #include "Containers/Queue.h"
@@ -118,7 +118,7 @@ private:
 
 	// The spinlock is used for accessing the RegisteredTypes array in performance critical areas (resolve TypeDefinitionHandles being the
 	// most common).
-	mutable FTransactionallySafeMutex RegisteredTypesLock; 
+	mutable UE::FTransactionallySafeMutex RegisteredTypesLock; 
 	FRegisteredTypesArray RegisteredTypes;
 
 	// Covers manipulation of all of the containers populated during registration (except for RegisteredTypes which is handled
diff --git a/Engine/Source/Programs/AutoRTFMTests/Private/UECoreTests.cpp b/Engine/Source/Programs/AutoRTFMTests/Private/UECoreTests.cpp
--- a/Engine/Source/Programs/AutoRTFMTests/Private/UECoreTests.cpp
+++ b/Engine/Source/Programs/AutoRTFMTests/Private/UECoreTests.cpp
@@ -2,6 +2,7 @@
 
 #include "AutoRTFM/AutoRTFM.h"
 #include "Catch2Includes.h"
+#include "Async/TransactionallySafeMutex.h"
 #include "AssetRegistry/AssetDataTagMap.h"
 #include "AutoRTFM/AutoRTFMTesting.h"
 #include "Delegates/IDelegateInstance.h"
@@ -20,7 +21,6 @@
 #include "MyAutoRTFMTestObject.h"
 #include "Containers/Queue.h"
 #include "Misc/ScopeRWLock.h"
-#include "Misc/TransactionallySafeMutex.h"
 #include "Misc/ConfigCacheIni.h"
 #include "Misc/CoreDelegates.h"
 #include "StructUtils/PropertyBag.h"
@@ -1697,7 +1697,7 @@ TEST_CASE("UECore.TransactionallySafeMutex")
 {
 	SECTION("Outside Transaction")
 	{
-		FTransactionallySafeMutex Mutex;
+		UE::FTransactionallySafeMutex Mutex;
 
 		AutoRTFM::ETransactionResult Result = AutoRTFM::Transact([&]
 			{
@@ -1717,7 +1717,7 @@ TEST_CASE("UECore.TransactionallySafeMutex")
 
 	SECTION("IsLocked Outside Transaction")
 	{
-		FTransactionallySafeMutex Mutex;
+		UE::FTransactionallySafeMutex Mutex;
 
 		REQUIRE(!Mutex.IsLocked());
 		Mutex.Lock();
@@ -1730,7 +1730,7 @@ TEST_CASE("UECore.TransactionallySafeMutex")
 	{
 		AutoRTFM::ETransactionResult Result = AutoRTFM::Transact([&]
 			{
-				FTransactionallySafeMutex Mutex;
+				UE::FTransactionallySafeMutex Mutex;
 				UE::TScopeLock Lock(Mutex);
 				AutoRTFM::AbortTransaction();
 			});
@@ -1739,7 +1739,7 @@ TEST_CASE("UECore.TransactionallySafeMutex")
 
 		Result = AutoRTFM::Transact([&]
 			{
-				FTransactionallySafeMutex Mutex;
+				UE::FTransactionallySafeMutex Mutex;
 				UE::TScopeLock Lock(Mutex);
 			});
 
@@ -1748,7 +1748,7 @@ TEST_CASE("UECore.TransactionallySafeMutex")
 
 	SECTION("IsLocked Inside Transaction")
 	{
-		FTransactionallySafeMutex Mutex;
+		UE::FTransactionallySafeMutex Mutex;
 		REQUIRE(!Mutex.IsLocked());
 
 		AutoRTFM::Testing::Commit([&]
@@ -1769,7 +1769,7 @@ TEST_CASE("UECore.TransactionallySafeMutex")
 
 		AutoRTFM::ETransactionResult Result = AutoRTFM::Transact([&]
 			{
-				FTransactionallySafeMutex Mutex;
+				UE::FTransactionallySafeMutex Mutex;
 
 				InnerResult = AutoRTFM::Transact([&]
 					{
@@ -1782,7 +1782,7 @@ TEST_CASE("UECore.TransactionallySafeMutex")
 
 		Result = AutoRTFM::Transact([&]
 			{
-				FTransactionallySafeMutex Mutex;
+				UE::FTransactionallySafeMutex Mutex;
 
 				InnerResult = AutoRTFM::Transact([&]
 					{
@@ -1796,7 +1796,7 @@ TEST_CASE("UECore.TransactionallySafeMutex")
 
 		Result = AutoRTFM::Transact([&]
 			{
-				FTransactionallySafeMutex Mutex;
+				UE::FTransactionallySafeMutex Mutex;
 
 				InnerResult = AutoRTFM::Transact([&]
 					{
@@ -1810,7 +1810,7 @@ TEST_CASE("UECore.TransactionallySafeMutex")
 
 		Result = AutoRTFM::Transact([&]
 			{
-				FTransactionallySafeMutex Mutex;
+				UE::FTransactionallySafeMutex Mutex;
 
 				InnerResult = AutoRTFM::Transact([&]
 					{
@@ -1826,7 +1826,7 @@ TEST_CASE("UECore.TransactionallySafeMutex")
 	{
 		struct MyStruct final
 		{
-			FTransactionallySafeMutex Mutex;
+			UE::FTransactionallySafeMutex Mutex;
 		};
 
 		auto Lambda = []()
@@ -1856,7 +1856,7 @@ TEST_CASE("UECore.TransactionallySafeMutex")
 	{
 		struct MyStruct final
 		{
-			FTransactionallySafeMutex Mutex;
+			UE::FTransactionallySafeMutex Mutex;
 		};
 
 		auto Lambda = []()
@@ -1884,7 +1884,7 @@ TEST_CASE("UECore.TransactionallySafeMutex")
 
 	SECTION("Delete Heap Allocated Mutex")
 	{
-		FTransactionallySafeMutex* const Mutex = new FTransactionallySafeMutex();
+		UE::FTransactionallySafeMutex* const Mutex = new UE::FTransactionallySafeMutex();
 
 		AutoRTFM::Testing::Abort([&]
 		{
@@ -1904,7 +1904,7 @@ TEST_CASE("UECore.TransactionallySafeMutex")
 
 	SECTION("Lock Within, Unlock Outside")
 	{
-		FTransactionallySafeMutex Mutex;
+		UE::FTransactionallySafeMutex Mutex;
 
 		AutoRTFM::Testing::Commit([&]
 		{
@@ -1918,7 +1918,7 @@ TEST_CASE("UECore.TransactionallySafeMutex")
 
 	SECTION("Lock Outside, Unlock Within")
 	{
-		FTransactionallySafeMutex Mutex = FTransactionallySafeMutex(UE::FAcquireLock());
+		UE::FTransactionallySafeMutex Mutex = UE::FTransactionallySafeMutex(UE::FAcquireLock());
 
 		AutoRTFM::Testing::Commit([&]
 		{
@@ -1931,7 +1931,7 @@ TEST_CASE("UECore.TransactionallySafeMutex")
 
 	SECTION("Lock, Unlock, Lock")
 	{
-		FTransactionallySafeMutex Mutex;
+		UE::FTransactionallySafeMutex Mutex;
 
 		AutoRTFM::Testing::Commit([&]
 		{
@@ -1945,7 +1945,7 @@ TEST_CASE("UECore.TransactionallySafeMutex")
 
 	SECTION("Lock Outside, Unlock Within")
 	{
-		FTransactionallySafeMutex Mutex = FTransactionallySafeMutex(UE::FAcquireLock());
+		UE::FTransactionallySafeMutex Mutex = UE::FTransactionallySafeMutex(UE::FAcquireLock());
 
 		AutoRTFM::Testing::Commit([&]
 		{
@@ -1958,7 +1958,7 @@ TEST_CASE("UECore.TransactionallySafeMutex")
 
 	SECTION("Lock Outside, Unlock & Lock Within")
 	{
-		FTransactionallySafeMutex Mutex = FTransactionallySafeMutex(UE::FAcquireLock());
+		UE::FTransactionallySafeMutex Mutex = UE::FTransactionallySafeMutex(UE::FAcquireLock());
 
 		AutoRTFM::Testing::Commit([&]
 		{
@@ -1972,7 +1972,7 @@ TEST_CASE("UECore.TransactionallySafeMutex")
 
 	SECTION("Commit(Lock, Commit(Unlock))")
 	{
-		FTransactionallySafeMutex Mutex;
+		UE::FTransactionallySafeMutex Mutex;
 
 		AutoRTFM::Testing::Commit([&]
 		{
@@ -1991,7 +1991,7 @@ TEST_CASE("UECore.TransactionallySafeMutex")
 
 	SECTION("Abort(Lock, Commit(Unlock, Lock))")
 	{
-		FTransactionallySafeMutex Mutex;
+		UE::FTransactionallySafeMutex Mutex;
 
 		AutoRTFM::Testing::Abort([&]
 		{
@@ -2032,7 +2032,7 @@ TEST_CASE("UECore.TransactionallySafeMutex")
 
 		SECTION("Contender Parks Transaction")
 		{
-			FTransactionallySafeMutex Mutex;
+			UE::FTransactionallySafeMutex Mutex;
 			std::atomic_uint Orderer = 0;
 
 			std::thread Contender([&]
@@ -2061,7 +2061,7 @@ TEST_CASE("UECore.TransactionallySafeMutex")
 
 		SECTION("Transaction Parks Contender")
 		{
-			FTransactionallySafeMutex Mutex;
+			UE::FTransactionallySafeMutex Mutex;
 			std::atomic_uint Orderer = 0;
 
 			std::thread Contender([&]
diff --git a/Engine/Source/Runtime/Core/Private/Internationalization/TextFormatter.cpp b/Engine/Source/Runtime/Core/Private/Internationalization/TextFormatter.cpp
--- a/Engine/Source/Runtime/Core/Private/Internationalization/TextFormatter.cpp
+++ b/Engine/Source/Runtime/Core/Private/Internationalization/TextFormatter.cpp
@@ -28,8 +28,7 @@
 #include "Templates/ValueOrError.h"
 #include "Trace/Detail/Channel.h"
 #include "AutoRTFM/AutoRTFM.h"
-#include "Misc/TransactionallySafeMutex.h"
-#include "Async/Mutex.h"
+#include "Async/TransactionallySafeMutex.h"
 
 #define LOCTEXT_NAMESPACE "TextFormatter"
 
@@ -527,7 +526,7 @@ private:
 	/**
 	 * Mutex protecting the compiled data from being modified concurrently.
 	 */
-	mutable FTransactionallySafeMutex CompiledDataMutex;
+	mutable UE::FTransactionallySafeMutex CompiledDataMutex;
 
 	/**
 	 * Copy of the string that was last compiled.
diff --git a/Engine/Source/Runtime/Core/Private/Internationalization/TextHistory.h b/Engine/Source/Runtime/Core/Private/Internationalization/TextHistory.h
--- a/Engine/Source/Runtime/Core/Private/Internationalization/TextHistory.h
+++ b/Engine/Source/Runtime/Core/Private/Internationalization/TextHistory.h
@@ -1,6 +1,7 @@
 // Copyright Epic Games, Inc. All Rights Reserved.
 #pragma once
 
+#include "Async/TransactionallySafeMutex.h"
 #include "Containers/Array.h"
 #include "CoreTypes.h"
 #include "HAL/CriticalSection.h"
@@ -16,7 +17,6 @@
 #include "Templates/SharedPointerInternals.h"
 #include "Templates/UnrealTemplate.h"
 #include "UObject/NameTypes.h"
-#include "Misc/TransactionallySafeMutex.h"
 
 class ITextGenerator;
 struct FDecimalNumberFormattingRules;
@@ -241,7 +241,7 @@ private:
 	/** Local revision index of this history, rebuilds when it is out of sync with the FTextLocalizationManager */
 	uint16 LocalRevision = 0;
 	/** Mutex to prevent races when multiple threads are updating out of date state. */
-	mutable FTransactionallySafeMutex Mutex = FTransactionallySafeMutex();
+	mutable UE::FTransactionallySafeMutex Mutex = UE::FTransactionallySafeMutex();
 };
 
 /** A potentially localized piece of source text (may have a TextId). */
diff --git a/Engine/Source/Runtime/Core/Public/Async/TransactionallySafeMutex.h b/Engine/Source/Runtime/Core/Public/Async/TransactionallySafeMutex.h
new file mode 100644
--- /dev/null
+++ b/Engine/Source/Runtime/Core/Public/Async/TransactionallySafeMutex.h
@@ -0,0 +1,256 @@
+// Copyright Epic Games, Inc. All Rights Reserved.
+
+#pragma once
+
+#include "Async/ExternalMutex.h"
+#include "AutoRTFM/AutoRTFM.h"
+#include <atomic>
+
+namespace UE
+{
+
+// A transactionally safe mutex that works in the following novel ways:
+// - In the open (non-transactional):
+//   - Take the lock like before. Simple!
+//   - Free the lock like before too.
+// - In the closed (transactional):
+//   - We flip the internal transactional state between two bits `IsLockedWithinTransactionFlag` and
+// 	   `IsUnlockedWithinTransactionFlag`. These let us keep track of what the last action taken with
+// 	   the mutex was within the transaction.
+//
+// Thus with this approach we will hold this lock for the *entirety* of the transactional nest should
+// we take the lock during the transaction, thus preventing non-transactional code from seeing any
+// modifications we should make.
+class FTransactionallySafeMutex final
+{
+public:
+	constexpr FTransactionallySafeMutex() = default;
+
+	// Construct in a locked state. Avoids an expensive compare-and-swap at creation time.
+	inline explicit FTransactionallySafeMutex(UE::FAcquireLock Tag)
+	{
+		FExternalMutex(State, Tag);
+	}
+
+	inline ~FTransactionallySafeMutex()
+	{
+		if (AutoRTFM::IsOnCurrentTransactionStack(this))
+		{
+			// We'll have either 0 (we've destructed a mutex but never locked it), or 1 of each
+			// on-abort and on-commit handlers.
+			AutoRTFM::PopAllOnAbortHandlers(this);
+			AutoRTFM::PopAllOnCommitHandlers(this);
+		}
+	}
+
+	UE_NONCOPYABLE(FTransactionallySafeMutex);
+
+	inline bool IsLocked() const
+	{
+		if (AutoRTFM::IsTransactional())
+		{
+			return AutoRTFM::Open([this]
+			{
+				if (FExternalMutex(State).IsLocked())
+				{
+					return !IsTransactionallyUnlocked();
+				}
+
+				return false;
+			});
+		}
+		else
+		{
+			return FExternalMutex(State).IsLocked();
+		}
+	}
+	
+	inline bool TryLock()
+	{
+		if (AutoRTFM::IsTransactional())
+		{
+			return AutoRTFM::Open([this]
+			{
+				if (!DoesTransactionHoldLock())
+				{
+					if (!FExternalMutex(State).TryLock())
+					{
+						return false;
+					}
+
+					// If we get here, it means we've taken the lock for the *first* time within the
+					// transaction. We need to thus register the on-abort and on-commit handlers to
+					// correctly undo the transactionality of our locking.
+					const AutoRTFM::EContextStatus Status = AutoRTFM::Close([this]
+					{
+						DoOnAbortHandlerInClosed();
+						DoOnCommitHandlerInClosed();
+					});
+
+					check(Status == AutoRTFM::EContextStatus::OnTrack);
+				}
+
+				RecordPreviousStateAndSetTransactionalLockMode(ETransactionalLockMode::TransactionallyLocked);
+
+				return true;
+			});
+		}
+		else
+		{
+			const bool bResult = FExternalMutex(State).TryLock();
+			ensure(!bResult || !DoesTransactionHoldLock());
+			return bResult;
+		}
+	}
+
+	inline void Lock()
+	{
+		if (AutoRTFM::IsTransactional())
+		{
+			AutoRTFM::Open([this]
+			{
+				if (!DoesTransactionHoldLock())
+				{
+					FExternalMutex(State).Lock();
+
+					const AutoRTFM::EContextStatus Status = AutoRTFM::Close([this]
+					{
+						DoOnAbortHandlerInClosed();
+						DoOnCommitHandlerInClosed();
+					});
+
+					check(Status == AutoRTFM::EContextStatus::OnTrack);
+				}
+
+				RecordPreviousStateAndSetTransactionalLockMode(ETransactionalLockMode::TransactionallyLocked);
+			});
+		}
+		else
+		{
+			FExternalMutex(State).Lock();
+			ensure(!DoesTransactionHoldLock());
+		}
+	}
+
+	inline void Unlock()
+	{
+		if (AutoRTFM::IsTransactional())
+		{
+			AutoRTFM::Open([this]
+			{
+				if (!DoesTransactionHoldLock())
+				{
+					const AutoRTFM::EContextStatus Status = AutoRTFM::Close([this]
+					{
+						DoOnCommitHandlerInClosed();
+					});
+
+					check(Status == AutoRTFM::EContextStatus::OnTrack);
+				}
+
+				RecordPreviousStateAndSetTransactionalLockMode(ETransactionalLockMode::TransactionallyUnlocked);
+			});
+		}
+		else
+		{
+			ensure(!DoesTransactionHoldLock());
+			FExternalMutex(State).Unlock();
+		}
+	}
+
+private:
+	// Need `mutable` because `FExternalMutex` only works non-const!
+	mutable std::atomic<uint8> State = 0;
+
+	// Since we are going to be writing non-atomically to `State` in the write-log, prior to us
+	// doing an atomic write to actually undo the lock, we need to be sure that the atomic will
+	// always be lock-free and take no bits beyond the integer backing type!
+	static_assert(decltype(State)::is_always_lock_free && (sizeof(decltype(State)) == sizeof(uint8)));
+	
+	// These are the flags that `FExternalMutex` already reserves for the lock!
+	static constexpr uint8 IsLockedFlag = 1 << 0;
+	static constexpr uint8 HasWaitingThreadsFlag = 1 << 1;
+
+	// Need 2 extra bits ontop of the 2 that `FExternalMutex` uses, to record what
+	// the last transactional action was within the transaction.
+	static constexpr uint8 IsLockedWithinTransactionFlag = 1 << 2;
+	static constexpr uint8 IsUnlockedWithinTransactionFlag = 1 << 3;
+
+	inline bool IsTransactionallyUnlocked() const
+	{
+		const uint8 LoadedState = State.load(std::memory_order_consume);
+		return (LoadedState & IsUnlockedWithinTransactionFlag);
+	}
+
+	inline bool DoesTransactionHoldLock() const
+	{
+		const uint8 LoadedState = State.load(std::memory_order_consume);
+		const bool bTransactionallyLocked = (LoadedState & (IsLockedWithinTransactionFlag | IsUnlockedWithinTransactionFlag));
+		ensure(!bTransactionallyLocked || (LoadedState & IsLockedFlag));
+		return bTransactionallyLocked;
+	}
+
+	inline void DoOnAbortHandlerInClosed()
+	{
+		check(AutoRTFM::IsClosed());
+
+		AutoRTFM::PushOnAbortHandler(this, [this]
+		{
+			State.fetch_and(~(IsLockedWithinTransactionFlag | IsUnlockedWithinTransactionFlag), std::memory_order_release);
+			FExternalMutex(State).Unlock();
+		});
+	}
+
+	inline void DoOnCommitHandlerInClosed()
+	{
+		check(AutoRTFM::IsClosed());
+
+		AutoRTFM::PushOnCommitHandler(this, [this]
+		{
+			// If we commit, we won't have undone any of our transactional writes to the
+			// lock. So just double check that we still hold the lock here (famous last
+			// words, but this is a temporary check just while we deploy the new mutex
+			// more widely. Once we've got a few runs around the block we can remove!).
+			ensure(DoesTransactionHoldLock());
+
+			const bool bIsTransactionallyUnlocked = IsTransactionallyUnlocked();
+
+			State.fetch_and(~(IsLockedWithinTransactionFlag | IsUnlockedWithinTransactionFlag), std::memory_order_release);
+
+			if (bIsTransactionallyUnlocked)
+			{
+				FExternalMutex(State).Unlock();
+			}
+		});
+	}
+
+	enum class ETransactionalLockMode : bool
+	{
+		TransactionallyLocked = false,
+		TransactionallyUnlocked = true,
+	};
+
+	inline void RecordPreviousStateAndSetTransactionalLockMode(ETransactionalLockMode LockMode)
+	{
+		// Record the state into the write log. If an abort happens we will non-atomically
+		// overwrite `State`, removing the transactional flags. This will be followed-up
+		// with an atomic write in the on-abort handler to actually do the unlock though,
+		// thus ensuring data integrity.
+		AutoRTFM::RecordOpenWriteNoMemoryValidation(&State);
+
+		if (LockMode == ETransactionalLockMode::TransactionallyLocked)
+		{
+			State |= IsLockedWithinTransactionFlag;
+			State &= ~IsUnlockedWithinTransactionFlag;
+		}
+		else
+		{
+			State |= IsUnlockedWithinTransactionFlag;
+			State &= ~IsLockedWithinTransactionFlag;
+		}
+	}
+};
+
+} // UE
+
+using FTransactionallySafeMutex = UE::FTransactionallySafeMutex;
diff --git a/Engine/Source/Runtime/Core/Public/Misc/TransactionallySafeMutex.h b/Engine/Source/Runtime/Core/Public/Misc/TransactionallySafeMutex.h
deleted file mode 100644
diff --git a/Engine/Source/Runtime/MassEntity/Public/MassEntityManagerStorage.h b/Engine/Source/Runtime/MassEntity/Public/MassEntityManagerStorage.h
--- a/Engine/Source/Runtime/MassEntity/Public/MassEntityManagerStorage.h
+++ b/Engine/Source/Runtime/MassEntity/Public/MassEntityManagerStorage.h
@@ -1,9 +1,8 @@
 // Copyright Epic Games, Inc. All Rights Reserved.
 #pragma once
 
-#include "Async/Mutex.h"
+#include "Async/TransactionallySafeMutex.h"
 #include "Containers/ChunkedArray.h"
-#include "Misc/TransactionallySafeMutex.h"
 #include "Misc/TVariant.h"
 #include "Templates/SharedPointer.h"
 #include "MassProcessingTypes.h"
@@ -253,12 +252,12 @@ namespace UE::Mass
 		uint32 MaxEntitiesPerPage = 0;
 		uint32 PageCount = 0;
 		/** ALWAYS acquire FreeListMutex before this one */
-		FTransactionallySafeMutex PageAllocateMutex;
+		UE::FTransactionallySafeMutex PageAllocateMutex;
 		/** Pointer to array of pages */
 		FEntityData** EntityPages = nullptr;
 
 		TArray<int32> EntityFreeIndexList;
-		FTransactionallySafeMutex FreeListMutex;
+		UE::FTransactionallySafeMutex FreeListMutex;
 	};
 
 } // namespace UE::Mass
diff --git a/Engine/Source/Runtime/MovieScene/Public/Evaluation/MovieSceneAnimTypeID.h b/Engine/Source/Runtime/MovieScene/Public/Evaluation/MovieSceneAnimTypeID.h
--- a/Engine/Source/Runtime/MovieScene/Public/Evaluation/MovieSceneAnimTypeID.h
+++ b/Engine/Source/Runtime/MovieScene/Public/Evaluation/MovieSceneAnimTypeID.h
@@ -2,11 +2,11 @@
 
 #pragma once
 
+#include "Async/TransactionallySafeMutex.h"
 #include "Containers/Array.h"
 #include "CoreMinimal.h"
 #include "CoreTypes.h"
 #include "Misc/ScopeLock.h"
-#include "Misc/TransactionallySafeMutex.h"
 #include "Templates/TypeHash.h"
 
 class FString;
@@ -131,7 +131,7 @@ private:
 		}
 	};
 
-	FTransactionallySafeMutex Mutex;
+	UE::FTransactionallySafeMutex Mutex;
 
 	/** Array of existing type identifiers */
 	TArray<FMovieSceneAnimTypeID> TypeIDs;
