CocoaThread - Add CVAR enabled by default that Ignore the WaitMode with MainThreadCalls which are causing deadlocks.
This is a tentative fix and a cleanup will be done once we verify that there is no downside to this change.

#tests Local Testing & testing by peers

diff --git a/Engine/Source/Runtime/Core/Private/Mac/CocoaThread.cpp b/Engine/Source/Runtime/Core/Private/Mac/CocoaThread.cpp
--- a/Engine/Source/Runtime/Core/Private/Mac/CocoaThread.cpp
+++ b/Engine/Source/Runtime/Core/Private/Mac/CocoaThread.cpp
@@ -2,6 +2,7 @@
 
 #include "Mac/CocoaThread.h"
 #include "Containers/LockFreeList.h"
+#include "HAL/IConsoleManager.h"
 #include "Templates/RefCounting.h"
 #include "Misc/OutputDeviceRedirector.h"
 
@@ -24,6 +25,14 @@ NSString* UnrealIMEEventMode = @"UnrealIMEEventMode";
 static FCocoaGameThread* GCocoaGameThread = nil;
 static uint32 GMainThreadId = 0;
 
+static TAutoConsoleVariable<bool> CVarMacCocoaThreadIgnoreWaitMode(
+	TEXT("Mac.IgnoreCocoaThreadWaitMode"),
+	1,
+	TEXT("Wether or not CocoaThread should ignore the specified WaitMode when processing blocks.\n")
+    TEXT("Changing the value to 0 might cause deadlocks because it will prevent the thread that calls MainThreadCall from processing certain types of blocks.\n")
+	TEXT("Default is 1."),
+	ECVF_ReadOnly);
+
 class FCocoaRunLoopSource;
 @interface FCocoaRunLoopSourceInfo : NSObject
 {
@@ -451,6 +460,10 @@ void MainThreadCall(dispatch_block_t Block, bool const bWait, NSString* WaitMode
 {
 	if ( [NSThread mainThread] != [NSThread currentThread] )
 	{
+		if (CVarMacCocoaThreadIgnoreWaitMode.GetValueOnAnyThread())
+		{
+			WaitMode = NSDefaultRunLoopMode;
+		}
 		FCocoaRunLoopSource& MainRunLoopSource = FCocoaRunLoopSource::GetMainRunLoopSource();
 		PerformBlockOnThread(MainRunLoopSource, Block, @[ NSDefaultRunLoopMode, NSModalPanelRunLoopMode, NSEventTrackingRunLoopMode ], WaitMode, bWait);
 	}
