update API used to get memory used on IOS.
This is from the 2022 WWDC presention Profile and optimize your game's memory

diff --git a/Engine/Source/Runtime/Core/Private/Apple/ApplePlatformMemory.cpp b/Engine/Source/Runtime/Core/Private/Apple/ApplePlatformMemory.cpp
--- a/Engine/Source/Runtime/Core/Private/Apple/ApplePlatformMemory.cpp
+++ b/Engine/Source/Runtime/Core/Private/Apple/ApplePlatformMemory.cpp
@@ -37,6 +37,12 @@
 #include "HAL/LowLevelMemTracker.h"
 #include "Apple/AppleLLM.h"
 
+#if PLATFORM_IOS
+#include <sys/resource.h> // On iOS, iPadOS and tvOS.
+extern "C" int proc_pid_rusage(int pid, int flavor, rusage_info_t *buffer)  
+	__OSX_AVAILABLE_STARTING(__MAC_10_9, __IPHONE_7_0);
+#endif
+
 NS_ASSUME_NONNULL_BEGIN
 
 // Enable to use XCode Instuments hooks to propagate UE allocs.
@@ -433,12 +439,23 @@ FPlatformMemoryStats FApplePlatformMemory::GetStats()
 	// considers "free" pages a waste of resources.
 	MemoryStats.AvailableVirtual = (Stats.free_count + Stats.inactive_count) * MemoryConstants.PageSize;
 	
+#if PLATFORM_IOS
+	rusage_info_current rusage_payload;
+		 
+	int ret = proc_pid_rusage(getpid(),
+							  RUSAGE_INFO_CURRENT, // I.e., new RUSAGE_INFO_V6 this year.
+							  (rusage_info_t *)&rusage_payload);
+	NSCAssert(ret == 0, @"Could not get rusage: %i.", errno); // Look up in `man errno`.
+			
+	MemoryStats.UsedPhysical = rusage_payload.ri_phys_footprint;
+#else
 	// Just get memory information for the process and report the working set instead
 	mach_task_basic_info_data_t TaskInfo;
 	mach_msg_type_number_t TaskInfoCount = MACH_TASK_BASIC_INFO_COUNT;
 	task_info( mach_task_self(), MACH_TASK_BASIC_INFO, (task_info_t)&TaskInfo, &TaskInfoCount );
 
     MemoryStats.UsedPhysical = TaskInfo.resident_size;
+#endif
 	if(MemoryStats.UsedPhysical > MemoryStats.PeakUsedPhysical)
 	{
 		MemoryStats.PeakUsedPhysical = MemoryStats.UsedPhysical;
