Adding EngineData and GameData to the temp crash context to minimize data loss when the normal crash handling flow is bypassed for things like abnormal terminations.

diff --git a/Engine/Source/Runtime/Core/Private/GenericPlatform/GenericPlatformCrashContext.cpp b/Engine/Source/Runtime/Core/Private/GenericPlatform/GenericPlatformCrashContext.cpp
--- a/Engine/Source/Runtime/Core/Private/GenericPlatform/GenericPlatformCrashContext.cpp
+++ b/Engine/Source/Runtime/Core/Private/GenericPlatform/GenericPlatformCrashContext.cpp
@@ -856,6 +856,22 @@ void FGenericCrashContext::SerializeTempCrashContextToFile()
 
 	SerializeSessionContext(SessionBuffer);
 	SerializeUserSettings(SessionBuffer);
+	
+	// Add the engine data
+	BeginSection(SessionBuffer, EngineDataTag);
+	for (const TPair<FString, FString>& Pair : NCached::EngineData)
+	{
+		AddCrashPropertyInternal(SessionBuffer, *Pair.Key, *Pair.Value);
+	}
+	EndSection(SessionBuffer, EngineDataTag);
+
+	// Add the game data
+	BeginSection(SessionBuffer, GameDataTag);
+	for (const TPair<FString, FString>& Pair : NCached::GameData)
+	{
+		AddCrashPropertyInternal(SessionBuffer, *Pair.Key, *Pair.Value);
+	}
+	EndSection(SessionBuffer, GameDataTag);
 
 	AddFooter(SessionBuffer);
 
