[Core]
* Fixed so callstacks converts virtual paths to local

diff --git a/Engine/Source/Runtime/Core/Private/Apple/ApplePlatformSymbolication.cpp b/Engine/Source/Runtime/Core/Private/Apple/ApplePlatformSymbolication.cpp
--- a/Engine/Source/Runtime/Core/Private/Apple/ApplePlatformSymbolication.cpp
+++ b/Engine/Source/Runtime/Core/Private/Apple/ApplePlatformSymbolication.cpp
@@ -458,7 +458,7 @@ bool FApplePlatformSymbolication::SymbolInfoForAddress(uint64 ProgramCounter, FP
 			if(!CSIsNull(Symbol))
 			{
 				out_SymbolInfo.LineNumber = CSSourceInfoGetLineNumber(Symbol);
-				FCStringAnsi::Strcpy(out_SymbolInfo.Filename, UE_ARRAY_COUNT(out_SymbolInfo.Filename), CSSourceInfoGetPath(Symbol));
+				FGenericPlatformStackWalk::CopyVirtualPathToLocal(out_SymbolInfo.Filename, UE_ARRAY_COUNT(out_SymbolInfo.Filename), CSSourceInfoGetPath(Symbol));
 				FCStringAnsi::Strcpy(out_SymbolInfo.FunctionName, UE_ARRAY_COUNT(out_SymbolInfo.FunctionName), CSSymbolGetName(CSSourceInfoGetSymbol(Symbol)));
 				CSRange CodeRange = CSSourceInfoGetRange(Symbol);
 				out_SymbolInfo.SymbolDisplacement = (ProgramCounter - CodeRange.Location);
diff --git a/Engine/Source/Runtime/Core/Private/GenericPlatform/GenericPlatformStackWalk.cpp b/Engine/Source/Runtime/Core/Private/GenericPlatform/GenericPlatformStackWalk.cpp
--- a/Engine/Source/Runtime/Core/Private/GenericPlatform/GenericPlatformStackWalk.cpp
+++ b/Engine/Source/Runtime/Core/Private/GenericPlatform/GenericPlatformStackWalk.cpp
@@ -1,6 +1,7 @@
 // Copyright Epic Games, Inc. All Rights Reserved.
 
 #include "GenericPlatform/GenericPlatformStackWalk.h"
+#include "BuildSettings.h"
 #include "HAL/PlatformStackWalk.h"
 #include "Misc/ConfigCacheIni.h"
 #include "GenericPlatform/GenericPlatformCrashContext.h"
@@ -274,3 +275,35 @@ TMap<FName, FString> FGenericPlatformStackWalk::GetSymbolMetaData()
 {
 	return TMap<FName, FString>();
 }
+
+void FGenericPlatformStackWalk::CopyVirtualPathToLocal(char* Dest, int32 DestCapacity, const char* Source)
+{
+	const char* VfsPathsIt = BuildSettings::GetVfsPaths();
+	while (true)
+	{
+		const char* VirtualPath = VfsPathsIt;
+		const char* VirtualPathEnd = FCStringAnsi::Strchr(VirtualPath, ';');
+		if (!VirtualPathEnd)
+		{
+			break;
+		}
+
+		const char* LocalPath = VirtualPathEnd + 1;
+		const char* LocalPathEnd = FCStringAnsi::Strchr(LocalPath, ';');
+
+		int32 VirtualPathLen = int32(VirtualPathEnd - VirtualPath);
+
+		if (FCStringAnsi::Strncmp(Source, VirtualPath, VirtualPathLen) == 0)
+		{
+			int32 LocalPathLen = int32(LocalPathEnd - LocalPath);
+
+			FCStringAnsi::Strncpy(Dest, LocalPath, FMath::Min(DestCapacity, LocalPathLen+1));
+			FCStringAnsi::Strncat(Dest, Source + VirtualPathLen, DestCapacity);
+			return;
+		}
+
+		VfsPathsIt = LocalPathEnd + 1;
+	}
+
+	FCStringAnsi::Strncpy(Dest, Source, DestCapacity);
+}
diff --git a/Engine/Source/Runtime/Core/Private/Unix/UnixPlatformStackWalk.cpp b/Engine/Source/Runtime/Core/Private/Unix/UnixPlatformStackWalk.cpp
--- a/Engine/Source/Runtime/Core/Private/Unix/UnixPlatformStackWalk.cpp
+++ b/Engine/Source/Runtime/Core/Private/Unix/UnixPlatformStackWalk.cpp
@@ -464,7 +464,10 @@ namespace
 							}
 						}
 
-						Reader.ReadOffsetIntoMemory(out_SymbolInfo.Filename, sizeof(out_SymbolInfo.Filename), Current.FileRelativeOffset);
+						ANSICHAR TempFileName[sizeof(out_SymbolInfo.Filename)];
+						Reader.ReadOffsetIntoMemory(TempFileName, sizeof(TempFileName), Current.FileRelativeOffset);
+						FUnixPlatformStackWalk::CopyVirtualPathToLocal(out_SymbolInfo.Filename, sizeof(out_SymbolInfo.Filename), TempFileName);
+
 						Reader.ReadOffsetIntoMemory(out_SymbolInfo.FunctionName, sizeof(out_SymbolInfo.FunctionName), Current.SymbolRelativeOffset);
 						out_SymbolInfo.LineNumber = Current.LineNumber;
 
diff --git a/Engine/Source/Runtime/Core/Private/Windows/WindowsPlatformStackWalk.cpp b/Engine/Source/Runtime/Core/Private/Windows/WindowsPlatformStackWalk.cpp
--- a/Engine/Source/Runtime/Core/Private/Windows/WindowsPlatformStackWalk.cpp
+++ b/Engine/Source/Runtime/Core/Private/Windows/WindowsPlatformStackWalk.cpp
@@ -552,7 +552,7 @@ UE_AUTORTFM_ALWAYS_OPEN void FWindowsPlatformStackWalk::EnumerateSymbolInfosForP
 	ImageHelpLine.SizeOfStruct = sizeof(ImageHelpLine);
 	if (SymGetLineFromAddr64(ProcessHandle, ProgramCounter, (::DWORD*)&SymbolInfo.SymbolDisplacement, &ImageHelpLine))
 	{
-		FCStringAnsi::Strncpy(SymbolInfo.Filename, ImageHelpLine.FileName, FProgramCounterSymbolInfo::MAX_NAME_LENGTH);
+		CopyVirtualPathToLocal(SymbolInfo.Filename, FProgramCounterSymbolInfo::MAX_NAME_LENGTH, ImageHelpLine.FileName);
 		SymbolInfo.LineNumber = ImageHelpLine.LineNumber;
 	}
 
diff --git a/Engine/Source/Runtime/Core/Public/GenericPlatform/GenericPlatformStackWalk.h b/Engine/Source/Runtime/Core/Public/GenericPlatform/GenericPlatformStackWalk.h
--- a/Engine/Source/Runtime/Core/Public/GenericPlatform/GenericPlatformStackWalk.h
+++ b/Engine/Source/Runtime/Core/Public/GenericPlatform/GenericPlatformStackWalk.h
@@ -323,6 +323,9 @@ struct FGenericPlatformStackWalk
 	 */
 	static CORE_API TMap<FName, FString> GetSymbolMetaData();
 
+	/** Converts virtual path to local and write local path to Dest. If no virtual path if sound, Source will be written */
+	static CORE_API void CopyVirtualPathToLocal(char* Dest, int32 DestCapacity, const char* Source);
+
 protected:
 
 	/** Returns true if non-monolithic builds should produce full callstacks in the log (and load all debug symbols) */
