[AutoRTFM] Fix the compiler lit tests after the signature change to autortfm_open

diff --git a/Engine/Source/Runtime/Core/Public/AutoRTFM/AutoRTFM.h b/Engine/Source/Runtime/Core/Public/AutoRTFM/AutoRTFM.h
--- a/Engine/Source/Runtime/Core/Public/AutoRTFM/AutoRTFM.h
+++ b/Engine/Source/Runtime/Core/Public/AutoRTFM/AutoRTFM.h
@@ -68,6 +68,7 @@
 #define PRAGMA_DISABLE_UNREACHABLE_CODE_WARNINGS
 #define PRAGMA_RESTORE_UNREACHABLE_CODE_WARNINGS
 #define LIKELY(EXPR) __builtin_expect(!!(EXPR), 1)
+#define PLATFORM_RETURN_ADDRESS() __builtin_return_address(0)
 #define UE_AUTORTFM_TCHAR char
 template <typename FuncType> class TFunction;  // Should be replaced with `std::function`
 #else
