Expand ConfigCacheIni values on demand to save memory

diff --git a/Engine/Source/Runtime/Core/Private/Misc/ConfigCacheIni.cpp b/Engine/Source/Runtime/Core/Private/Misc/ConfigCacheIni.cpp
--- a/Engine/Source/Runtime/Core/Private/Misc/ConfigCacheIni.cpp
+++ b/Engine/Source/Runtime/Core/Private/Misc/ConfigCacheIni.cpp
@@ -226,13 +226,10 @@ FString FConfigValue::ExpandValue(const FString& InCollapsedValue)
 	return OutExpandedValue;
 }
 
-void FConfigValue::ExpandValueInternal()
+bool FConfigValue::NeedsToExpandValue()
 {
 	const TCHAR* Dummy;
-	if (FindNextExpansion(*SavedValue, Dummy))
-	{
-		ExpandValue(SavedValue, /* out */ ExpandedValue);
-	}
+	return FindNextExpansion(*SavedValue, Dummy) != nullptr;
 }
 
 bool FConfigValue::CollapseValue(const FString& InExpandedValue, FString& OutCollapsedValue)
diff --git a/Engine/Source/Runtime/Core/Public/Misc/ConfigCacheIni.h b/Engine/Source/Runtime/Core/Public/Misc/ConfigCacheIni.h
--- a/Engine/Source/Runtime/Core/Public/Misc/ConfigCacheIni.h
+++ b/Engine/Source/Runtime/Core/Public/Misc/ConfigCacheIni.h
@@ -175,7 +175,7 @@ public:
 #endif
 	{
 		SavedValueHash = FTextLocalizationResource::HashString(SavedValue);
-		ExpandValueInternal();
+		bExpandOnDemand = NeedsToExpandValue();
 	}
 
 	FConfigValue(const FString& InValue, EValueType Type=EValueType::Combined)
@@ -191,7 +191,7 @@ public:
 #endif
 	{
 		SavedValueHash = FTextLocalizationResource::HashString(SavedValue);
-		ExpandValueInternal();
+		bExpandOnDemand = NeedsToExpandValue();
 	}
 
 	FConfigValue(FString&& InValue, EValueType Type=EValueType::Combined)
@@ -207,14 +207,14 @@ public:
 #endif
 	{
 		SavedValueHash = FTextLocalizationResource::HashString(SavedValue);
-		ExpandValueInternal();
+		bExpandOnDemand = NeedsToExpandValue();
 	}
 
 	FConfigValue(const FConfigValue& InConfigValue)
 		: ValueType(InConfigValue.ValueType)
 		, SavedValue(InConfigValue.SavedValue)
-		, ExpandedValue(InConfigValue.ExpandedValue)
 		, SavedValueHash(InConfigValue.SavedValueHash)
+		, bExpandOnDemand(InConfigValue.bExpandOnDemand)
 #if UE_WITH_CONFIG_TRACKING
 		, SectionAccess(InConfigValue.SectionAccess)
 		, ValueName(InConfigValue.ValueName)
@@ -226,8 +226,8 @@ public:
 	FConfigValue(FConfigValue&& InConfigValue)
 		: ValueType(InConfigValue.ValueType)
 		, SavedValue(MoveTemp(InConfigValue.SavedValue))
-		, ExpandedValue(MoveTemp(InConfigValue.ExpandedValue))
 		, SavedValueHash(InConfigValue.SavedValueHash)
+		, bExpandOnDemand(InConfigValue.bExpandOnDemand)
 #if UE_WITH_CONFIG_TRACKING
 		, SectionAccess(MoveTemp(InConfigValue.SectionAccess))
 		, ValueName(InConfigValue.ValueName)
@@ -240,8 +240,8 @@ public:
 	{
 		ValueType = RHS.ValueType;
 		SavedValue = MoveTemp(RHS.SavedValue);
-		ExpandedValue = MoveTemp(RHS.ExpandedValue);
 		SavedValueHash = RHS.SavedValueHash;
+		bExpandOnDemand = RHS.bExpandOnDemand;
 #if UE_WITH_CONFIG_TRACKING
 		SectionAccess = MoveTemp(RHS.SectionAccess);
 		ValueName = RHS.ValueName;
@@ -254,8 +254,8 @@ public:
 	{
 		ValueType = RHS.ValueType;
 		SavedValue = RHS.SavedValue;
-		ExpandedValue = RHS.ExpandedValue;
 		SavedValueHash = RHS.SavedValueHash;
+		bExpandOnDemand = RHS.bExpandOnDemand;
 #if UE_WITH_CONFIG_TRACKING
 		SectionAccess = RHS.SectionAccess;
 		ValueName = RHS.ValueName;
@@ -278,18 +278,17 @@ public:
 	{
 		SavedValue = MoveTemp(RHS);
 		SavedValueHash = FTextLocalizationResource::HashString(SavedValue);
-		ExpandedValue.Empty();
-		ExpandValueInternal();
+		bExpandOnDemand = NeedsToExpandValue();
 		return *this;
 	}
 
 	// Returns the ini setting with any macros expanded out
-	const FString& GetValue() const
+	FORCEINLINE FString GetValue() const
 	{
 #if UE_WITH_CONFIG_TRACKING
 		UE::ConfigAccessTracking::Private::OnConfigValueRead(SectionAccess, ValueName, *this);
 #endif
-		return (ExpandedValue.Len() > 0 ? ExpandedValue : SavedValue);
+		return bExpandOnDemand ? ExpandValue(SavedValue) : SavedValue;
 	}
 
 	// Returns the original ini setting without macro expansion
@@ -332,7 +331,7 @@ public:
 
 		if (Slot.GetUnderlyingArchive().IsLoading())
 		{
-			ConfigValue.ExpandValueInternal();
+			ConfigValue.bExpandOnDemand = ConfigValue.NeedsToExpandValue();
 		}
 	}
 
@@ -382,9 +381,9 @@ public:
 #endif
 
 	/** Gets the expanded value (GetValue) without marking it as having been accessed for e.g. writing out to a ConfigFile to disk */
-	const FString& GetValueForWriting() const
+	FORCEINLINE FString GetValueForWriting() const
 	{
-		return (ExpandedValue.Len() > 0 ? ExpandedValue : SavedValue);
+		return bExpandOnDemand ? ExpandValue(SavedValue) : SavedValue;
 	};
 
 	/** Gets the SavedValue without marking it as having been accessed for e.g. writing out to a ConfigFile to disk */
@@ -399,11 +398,11 @@ public:
 
 private:
 	/** Internal version of ExpandValue that expands SavedValue into ExpandedValue, or produces an empty ExpandedValue if no expansion occurred. */
-	CORE_API void ExpandValueInternal();
+	CORE_API bool NeedsToExpandValue();
 
 	FString SavedValue;
-	FString ExpandedValue;
 	uint32 SavedValueHash;
+	bool bExpandOnDemand = false;
 #if UE_WITH_CONFIG_TRACKING
 	TRefCountPtr<UE::ConfigAccessTracking::FSection> SectionAccess;
 	FMinimalName ValueName;
