CocoaThread - Fix ensure when MainThreadCall is use when GameThread is being terminated
Since WaitMode is only use when bWait = true, we can safely ignore the new CVar CVarMacCocoaThreadIgnoreWaitMode and therefore avoid calling
GetValueOnAnyThread which would cause the ensure because GGameThreadId is set to GMainThreadID at this point.

#tests Local testing, this issue was triggered when exiting the editor on Mac (100% repro)

diff --git a/Engine/Source/Runtime/Core/Private/Mac/CocoaThread.cpp b/Engine/Source/Runtime/Core/Private/Mac/CocoaThread.cpp
--- a/Engine/Source/Runtime/Core/Private/Mac/CocoaThread.cpp
+++ b/Engine/Source/Runtime/Core/Private/Mac/CocoaThread.cpp
@@ -460,7 +460,7 @@ void MainThreadCall(dispatch_block_t Block, bool const bWait, NSString* WaitMode
 {
 	if ( [NSThread mainThread] != [NSThread currentThread] )
 	{
-		if (CVarMacCocoaThreadIgnoreWaitMode.GetValueOnAnyThread())
+		if (bWait && CVarMacCocoaThreadIgnoreWaitMode.GetValueOnAnyThread())
 		{
 			WaitMode = NSDefaultRunLoopMode;
 		}
