Color functions are  [[nodiscard]]

#rnx

diff --git a/Engine/Plugins/Enterprise/MDLImporter/Source/MDLImporter/Private/MDLMaterialSelector.cpp b/Engine/Plugins/Enterprise/MDLImporter/Source/MDLImporter/Private/MDLMaterialSelector.cpp
--- a/Engine/Plugins/Enterprise/MDLImporter/Source/MDLImporter/Private/MDLMaterialSelector.cpp
+++ b/Engine/Plugins/Enterprise/MDLImporter/Source/MDLImporter/Private/MDLMaterialSelector.cpp
@@ -67,7 +67,7 @@ namespace MDLImporterImpl
 		{
 			EmissionStrength = MdlMaterial.EmissionStrength.Value;
 
-			EmissionColor.Desaturate(1.f);
+			EmissionColor = EmissionColor.Desaturate(1.f);
 			return EmissionStrength > 0.1f && EmissionColor.R > 0.1f;
 		}
 
diff --git a/Engine/Plugins/Experimental/Enterprise/DatasmithRuntime/Source/Private/MaterialImportUtils.cpp b/Engine/Plugins/Experimental/Enterprise/DatasmithRuntime/Source/Private/MaterialImportUtils.cpp
--- a/Engine/Plugins/Experimental/Enterprise/DatasmithRuntime/Source/Private/MaterialImportUtils.cpp
+++ b/Engine/Plugins/Experimental/Enterprise/DatasmithRuntime/Source/Private/MaterialImportUtils.cpp
@@ -737,7 +737,7 @@ namespace DatasmithRuntime
 			}
 			else if (FCString::Stricmp(Expression, TEXT("Desaturation")) == 0)
 			{
-				OutValue.Numeric.Desaturate(1.0);
+				OutValue.Numeric = OutValue.Numeric.Desaturate(1.0);
 			}
 			else if (FCString::Stricmp(Expression, TEXT("Noise")) == 0)
 			{
diff --git a/Engine/Plugins/Experimental/Fracture/Source/FractureEngine/Private/FractureEngineFracturing.cpp b/Engine/Plugins/Experimental/Fracture/Source/FractureEngine/Private/FractureEngineFracturing.cpp
--- a/Engine/Plugins/Experimental/Fracture/Source/FractureEngine/Private/FractureEngineFracturing.cpp
+++ b/Engine/Plugins/Experimental/Fracture/Source/FractureEngine/Private/FractureEngineFracturing.cpp
@@ -1525,9 +1525,9 @@ void FFractureEngineFracturing::SetBoneColorByLevel(FManagedArrayCollection& InC
 			int32 ColorIndex = Bone + 1; // parent can be -1 for root, range [-1..n]
 			BoneColor = RandomColors[ColorIndex % RandomColors.Num()];
 
-			BoneColor.LinearRGBToHSV();
+			BoneColor = BoneColor.LinearRGBToHSV();
 			BoneColor.B *= .5;
-			BoneColor.HSVToLinearRGB();
+			BoneColor = BoneColor.HSVToLinearRGB();
 		}
 		else
 		{
diff --git a/Engine/Plugins/Interchange/Runtime/Source/Import/Private/Usd/InterchangeUsdTranslator.cpp b/Engine/Plugins/Interchange/Runtime/Source/Import/Private/Usd/InterchangeUsdTranslator.cpp
--- a/Engine/Plugins/Interchange/Runtime/Source/Import/Private/Usd/InterchangeUsdTranslator.cpp
+++ b/Engine/Plugins/Interchange/Runtime/Source/Import/Private/Usd/InterchangeUsdTranslator.cpp
@@ -1092,7 +1092,7 @@ namespace UE::InterchangeUsdTranslator::Private
 		FLinearColor Color = UsdUtils::GetAttributeValue<FLinearColor>(Prim, ColorToken);
 
 		const bool bSRGB = true;
-		Color.ToFColor(bSRGB);
+		Color = Color.ToFColor(bSRGB);
 
 		static const FString TemperatureToken = UsdToUnreal::ConvertToken(pxr::UsdLuxTokens->inputsColorTemperature);
 		static const FString UseTemperatureToken = UsdToUnreal::ConvertToken(pxr::UsdLuxTokens->inputsEnableColorTemperature);
diff --git a/Engine/Source/Runtime/Core/Public/Math/Color.h b/Engine/Source/Runtime/Core/Public/Math/Color.h
--- a/Engine/Source/Runtime/Core/Public/Math/Color.h
+++ b/Engine/Source/Runtime/Core/Public/Math/Color.h
@@ -114,13 +114,13 @@ struct FLinearColor
 	}
 
 	// Conversions.
-	CORE_API FColor ToRGBE() const;
+	[[nodiscard]] CORE_API FColor ToRGBE() const;
 
 	/**
 	 * Converts an FColor coming from an observed sRGB output, into a linear color.
 	 * @param Color The sRGB color that needs to be converted into linear space.
 	 */
-	FORCEINLINE static FLinearColor FromSRGBColor(const FColor& Color)
+	[[nodiscard]] FORCEINLINE static FLinearColor FromSRGBColor(const FColor& Color)
 	{
 		return FLinearColor(Color);
 	}
@@ -129,25 +129,25 @@ struct FLinearColor
 	 * Converts an FColor coming from an observed Pow(1/2.2) output, into a linear color.
 	 * @param Color The Pow(1/2.2) color that needs to be converted into linear space.
 	 */
-	CORE_API static FLinearColor FromPow22Color(const FColor& Color);
+	[[nodiscard]] CORE_API static FLinearColor FromPow22Color(const FColor& Color);
 
 	// Operators.
 
-	FORCEINLINE float& Component(int32 Index)
+	[[nodiscard]] FORCEINLINE float& Component(int32 Index)
 	{
 	PRAGMA_DISABLE_DEPRECATION_WARNINGS
 		return RGBA[Index];
 	PRAGMA_ENABLE_DEPRECATION_WARNINGS
 	}
 
-	FORCEINLINE const float& Component(int32 Index) const
+	[[nodiscard]] FORCEINLINE const float& Component(int32 Index) const
 	{
 	PRAGMA_DISABLE_DEPRECATION_WARNINGS
 		return RGBA[Index];
 	PRAGMA_ENABLE_DEPRECATION_WARNINGS
 	}
 
-	FORCEINLINE FLinearColor operator+(const FLinearColor& ColorB) const
+	[[nodiscard]] FORCEINLINE FLinearColor operator+(const FLinearColor& ColorB) const
 	{
 		return FLinearColor(
 			this->R + ColorB.R,
@@ -165,7 +165,7 @@ struct FLinearColor
 		return *this;
 	}
 
-	FORCEINLINE FLinearColor operator-(const FLinearColor& ColorB) const
+	[[nodiscard]] FORCEINLINE FLinearColor operator-(const FLinearColor& ColorB) const
 	{
 		return FLinearColor(
 			this->R - ColorB.R,
@@ -183,7 +183,7 @@ struct FLinearColor
 		return *this;
 	}
 
-	FORCEINLINE FLinearColor operator*(const FLinearColor& ColorB) const
+	[[nodiscard]] FORCEINLINE FLinearColor operator*(const FLinearColor& ColorB) const
 	{
 		return FLinearColor(
 			this->R * ColorB.R,
@@ -201,7 +201,7 @@ struct FLinearColor
 		return *this;
 	}
 
-	FORCEINLINE FLinearColor operator*(float Scalar) const
+	[[nodiscard]] FORCEINLINE FLinearColor operator*(float Scalar) const
 	{
 		return FLinearColor(
 			this->R * Scalar,
@@ -220,7 +220,7 @@ struct FLinearColor
 		return *this;
 	}
 
-	FORCEINLINE FLinearColor operator/(const FLinearColor& ColorB) const
+	[[nodiscard]] FORCEINLINE FLinearColor operator/(const FLinearColor& ColorB) const
 	{
 		return FLinearColor(
 			this->R / ColorB.R,
@@ -238,7 +238,7 @@ struct FLinearColor
 		return *this;
 	}
 
-	FORCEINLINE FLinearColor operator/(float Scalar) const
+	[[nodiscard]] FORCEINLINE FLinearColor operator/(float Scalar) const
 	{
 		const float	InvScalar = 1.0f / Scalar;
 		return FLinearColor(
@@ -259,7 +259,7 @@ struct FLinearColor
 	}
 
 	// clamped in 0..1 range
-	FORCEINLINE FLinearColor GetClamped(float InMin = 0.0f, float InMax = 1.0f) const
+	[[nodiscard]] FORCEINLINE FLinearColor GetClamped(float InMin = 0.0f, float InMax = 1.0f) const
 	{
 		FLinearColor Ret;
 
@@ -272,22 +272,22 @@ struct FLinearColor
 	}
 
 	/** Comparison operators */
-	FORCEINLINE bool operator==(const FLinearColor& ColorB) const
+	[[nodiscard]] FORCEINLINE bool operator==(const FLinearColor& ColorB) const
 	{
 		return this->R == ColorB.R && this->G == ColorB.G && this->B == ColorB.B && this->A == ColorB.A;
 	}
-	FORCEINLINE bool operator!=(const FLinearColor& Other) const
+	[[nodiscard]] FORCEINLINE bool operator!=(const FLinearColor& Other) const
 	{
 		return this->R != Other.R || this->G != Other.G || this->B != Other.B || this->A != Other.A;
 	}
 
 	// Error-tolerant comparison.
-	FORCEINLINE bool Equals(const FLinearColor& ColorB, float Tolerance=UE_KINDA_SMALL_NUMBER) const
+	[[nodiscard]] FORCEINLINE bool Equals(const FLinearColor& ColorB, float Tolerance=UE_KINDA_SMALL_NUMBER) const
 	{
 		return FMath::Abs(this->R - ColorB.R) < Tolerance && FMath::Abs(this->G - ColorB.G) < Tolerance && FMath::Abs(this->B - ColorB.B) < Tolerance && FMath::Abs(this->A - ColorB.A) < Tolerance;
 	}
 
-	FLinearColor CopyWithNewOpacity(float NewOpacicty) const
+	[[nodiscard]] FLinearColor CopyWithNewOpacity(float NewOpacicty) const
 	{
 		FLinearColor NewCopy = *this;
 		NewCopy.A = NewOpacicty;
@@ -295,32 +295,32 @@ struct FLinearColor
 	}
 
 	UE_DEPRECATED(4.22, "FGetHSV doesn't perform a valid HSV conversion, use MakeFromHSV8 instead")
-	static CORE_API FLinearColor FGetHSV(uint8 H,uint8 S,uint8 V);
+	[[nodiscard]] static CORE_API FLinearColor FGetHSV(uint8 H,uint8 S,uint8 V);
 
 	/**
 	 * Converts byte hue-saturation-brightness to floating point red-green-blue.
 	 */
-	static CORE_API FLinearColor MakeFromHSV8(uint8 H, uint8 S, uint8 V);
+	[[nodiscard]] static CORE_API FLinearColor MakeFromHSV8(uint8 H, uint8 S, uint8 V);
 
 	/**
 	* Makes a random but quite nice color.
 	*/
-	static CORE_API FLinearColor MakeRandomColor();
+	[[nodiscard]] static CORE_API FLinearColor MakeRandomColor();
 
 	/**
 	* Converts temperature in Kelvins of a black body radiator to RGB chromaticity.
 	*/
-	static CORE_API FLinearColor MakeFromColorTemperature( float Temp );
+	[[nodiscard]] static CORE_API FLinearColor MakeFromColorTemperature( float Temp );
 
 	/**
 	* Makes a random color based on a seed.
 	*/
-	static CORE_API FLinearColor MakeRandomSeededColor(int32 Seed);
+	[[nodiscard]] static CORE_API FLinearColor MakeRandomSeededColor(int32 Seed);
 
 	/**
 	 * Euclidean distance between two points.
 	 */
-	static inline float Dist( const FLinearColor &V1, const FLinearColor &V2 )
+	[[nodiscard]] static inline float Dist( const FLinearColor &V1, const FLinearColor &V2 )
 	{
 		return FMath::Sqrt( FMath::Square(V2.R-V1.R) + FMath::Square(V2.G-V1.G) + FMath::Square(V2.B-V1.B) + FMath::Square(V2.A-V1.A) );
 	}
@@ -333,13 +333,13 @@ struct FLinearColor
 	 * @param	OutPoints		Receives the output samples.
 	 * @return					Path length.
 	 */
-	static CORE_API float EvaluateBezier(const FLinearColor* ControlPoints, int32 NumPoints, TArray<FLinearColor>& OutPoints);
+	[[nodiscard]] static CORE_API float EvaluateBezier(const FLinearColor* ControlPoints, int32 NumPoints, TArray<FLinearColor>& OutPoints);
 
 	/** Converts a linear space RGB color to an HSV color */
-	CORE_API FLinearColor LinearRGBToHSV() const;
+	[[nodiscard]] CORE_API FLinearColor LinearRGBToHSV() const;
 
 	/** Converts an HSV color to a linear space RGB color */
-	CORE_API FLinearColor HSVToLinearRGB() const;
+	[[nodiscard]] CORE_API FLinearColor HSVToLinearRGB() const;
 
 	/**
 	 * Linearly interpolates between two colors by the specified progress amount.  The interpolation is performed in HSV color space
@@ -351,34 +351,34 @@ struct FLinearColor
 	 * @param	Progress	Scalar interpolation amount (usually between 0.0 and 1.0 inclusive)
 	 * @return	The interpolated color in linear RGB space along with the interpolated alpha value
 	 */
-	static CORE_API FLinearColor LerpUsingHSV( const FLinearColor& From, const FLinearColor& To, const float Progress );
+	[[nodiscard]] static CORE_API FLinearColor LerpUsingHSV( const FLinearColor& From, const FLinearColor& To, const float Progress );
 
 	/** Quantizes the linear color with rounding and returns the result as a FColor.  This bypasses the SRGB conversion. 
 	* QuantizeRound can be dequantized back to linear with FColor::ReinterpretAsLinear (just /255.f)
 	* this matches the GPU UNORM<->float conversion spec and should be preferred
 	*/
-	FORCEINLINE FColor QuantizeRound() const;
+	[[nodiscard]] FORCEINLINE FColor QuantizeRound() const;
 	
 	/** Quantizes the linear color and returns the result as a FColor.  This bypasses the SRGB conversion.
 	* Uses floor quantization, which does not match the GPU standard conversion.
 	* Restoration to float should be done with a +0.5 bias to restore to centered buckets.
 	* Do NOT use this for graphics or textures or images, use QuantizeRound instead.
 	*/
-	FORCEINLINE FColor QuantizeFloor() const;
+	[[nodiscard]] FORCEINLINE FColor QuantizeFloor() const;
 
 	/** backwards compatible Quantize function name, does QuantizeFloor.
 	* @todo deprecate me
 	*/
 	UE_DEPRECATED(5.0, "Most callers of Quantize should have been calling QuantizeRound; to match old behavior use QuantizeFloor")
-	FORCEINLINE FColor Quantize() const;
+	[[nodiscard]] FORCEINLINE FColor Quantize() const;
 
 	/** Quantizes the linear color and returns the result as a FColor with optional sRGB conversion. 
 	* Clamps in [0,1] range before conversion.
 	* ToFColor(false) is QuantizeRound
 	*/
-	CORE_API FColor ToFColorSRGB() const;
+	[[nodiscard]] CORE_API FColor ToFColorSRGB() const;
 	
-	FORCEINLINE FColor ToFColor(const bool bSRGB) const;
+	[[nodiscard]] FORCEINLINE FColor ToFColor(const bool bSRGB) const;
 	
 	/**
 	 * Returns a desaturated color, with 0 meaning no desaturation and 1 == full desaturation
@@ -386,23 +386,23 @@ struct FLinearColor
 	 * @param	Desaturation	Desaturation factor in range [0..1]
 	 * @return	Desaturated color
 	 */
-	CORE_API FLinearColor Desaturate( float Desaturation ) const;
+	[[nodiscard]] CORE_API FLinearColor Desaturate( float Desaturation ) const;
 
 	/** Computes the perceptually weighted luminance value of a color. */
-	CORE_API float GetLuminance() const;
+	[[nodiscard]] CORE_API float GetLuminance() const;
 	
 	/**
 	 * Returns the maximum value in this color structure
 	 *
 	 * @return The maximum color channel value
 	 */
-	FORCEINLINE float GetMax() const
+	[[nodiscard]] FORCEINLINE float GetMax() const
 	{
 		return FMath::Max( FMath::Max( FMath::Max( R, G ), B ), A );
 	}
 
 	/** useful to detect if a light contribution needs to be rendered */
-	bool IsAlmostBlack() const
+	[[nodiscard]] bool IsAlmostBlack() const
 	{
 		return FMath::Square(R) < UE_DELTA && FMath::Square(G) < UE_DELTA && FMath::Square(B) < UE_DELTA;
 	}
@@ -412,12 +412,12 @@ struct FLinearColor
 	 *
 	 * @return The minimum color channel value
 	 */
-	FORCEINLINE float GetMin() const
+	[[nodiscard]] FORCEINLINE float GetMin() const
 	{
 		return FMath::Min( FMath::Min( FMath::Min( R, G ), B ), A );
 	}
 
-	CORE_API FString ToString() const;
+	[[nodiscard]] CORE_API FString ToString() const;
 
 	/**
 	 * Initialize this Color based on an FString. The String is expected to contain R=, G=, B=, A=.
@@ -435,7 +435,7 @@ struct FLinearColor
 	 * @param InValue The input value.
 	 * @return InValue clamped to [0,1]. NaNs map to 0.
 	 */
-	static FORCEINLINE float Clamp01NansTo0(float InValue)
+	[[nodiscard]] static FORCEINLINE float Clamp01NansTo0(float InValue)
 	{
 		// Write this explicitly instead of using FMath::Clamp because we're particular
 		// about what happens with NaNs here.
@@ -451,7 +451,7 @@ struct FLinearColor
 	 * @param Value			0-1 brightness of resulting color
 	 * @param HueRotation	0-360 amount to rotate around the hue wheel
 	 * @return color with semi-random hue */
-	static FLinearColor IntToDistinctColor(
+	[[nodiscard]] static FLinearColor IntToDistinctColor(
 		const int32 Seed,
 		const float Saturation=1.f,
 		const float Value=1.f,
@@ -471,7 +471,7 @@ struct FLinearColor
 	static CORE_API const FLinearColor Blue;
 	static CORE_API const FLinearColor Yellow;
 
-	friend FORCEINLINE uint32 GetTypeHash( const FLinearColor& LinearColor )
+	[[nodiscard]] friend FORCEINLINE uint32 GetTypeHash( const FLinearColor& LinearColor )
 	{
 		// Note: this assumes there's no padding in FLinearColor that could contain uncompared data.
 		return FCrc::MemCrc_DEPRECATED(&LinearColor, sizeof(FLinearColor));
@@ -550,12 +550,12 @@ public:
 	}
 
 	// Operators.
-	FORCEINLINE bool operator==( const FColor &C ) const
+	[[nodiscard]] FORCEINLINE bool operator==( const FColor &C ) const
 	{
 		return DWColor() == C.DWColor();
 	}
 
-	FORCEINLINE bool operator!=( const FColor& C ) const
+	[[nodiscard]] FORCEINLINE bool operator!=( const FColor& C ) const
 	{
 		return DWColor() != C.DWColor();
 	}
@@ -568,7 +568,7 @@ public:
 		A = (uint8) FMath::Min((int32) A + (int32) C.A,255);
 	}
 
-	CORE_API FLinearColor FromRGBE() const;
+	[[nodiscard]] CORE_API FLinearColor FromRGBE() const;
 
 	/**
 	 * Creates a color value from the given hexadecimal string.
@@ -579,58 +579,58 @@ public:
 	 * @return The corresponding color value.
 	 * @see ToHex
 	 */
-	static CORE_API FColor FromHex( const FString& HexString );
+	[[nodiscard]] static CORE_API FColor FromHex( const FString& HexString );
 
 	/**
 	 * Makes a random but quite nice color.
 	 */
-	static CORE_API FColor MakeRandomColor();
+	[[nodiscard]] static CORE_API FColor MakeRandomColor();
 
 	/**
 	 * Makes a color red->green with the passed in scalar (e.g. 0 is red, 1 is green)
 	 */
-	static CORE_API FColor MakeRedToGreenColorFromScalar(float Scalar);
+	[[nodiscard]] static CORE_API FColor MakeRedToGreenColorFromScalar(float Scalar);
 
 	/**
 	* Converts temperature in Kelvins of a black body radiator to RGB chromaticity.
 	*/
-	static CORE_API FColor MakeFromColorTemperature( float Temp );
+	[[nodiscard]] static CORE_API FColor MakeFromColorTemperature( float Temp );
 
 	/**
 	* Makes a random color based on a seed.
 	*/
-	static CORE_API FColor MakeRandomSeededColor(int32 Seed);
+	[[nodiscard]] static CORE_API FColor MakeRandomSeededColor(int32 Seed);
 
 	/**
 	* Conversions to/from GPU UNorm floats, U8, U16
 	* matches convention of FColor FLinearColor::QuantizeRound
 	*/
 
-	static uint8 QuantizeUNormFloatTo8( float UnitFloat )
+	[[nodiscard]] static uint8 QuantizeUNormFloatTo8( float UnitFloat )
 	{
 		return (uint8)( 0.5f + FLinearColor::Clamp01NansTo0(UnitFloat) * 255.f );
 	}
 	
-	static uint16 QuantizeUNormFloatTo16( float UnitFloat )
+	[[nodiscard]] static uint16 QuantizeUNormFloatTo16( float UnitFloat )
 	{
 		return (uint16)( 0.5f + FLinearColor::Clamp01NansTo0(UnitFloat) * 65535.f );
 	}
 
-	static float DequantizeUNorm8ToFloat( int Value8 )
+	[[nodiscard]] static float DequantizeUNorm8ToFloat( int Value8 )
 	{
 		check( Value8 >= 0 && Value8 <= 255 );
 
 		return (float)Value8 * (1.f/255.f);
 	}
 	
-	static float DequantizeUNorm16ToFloat( int Value16 )
+	[[nodiscard]] static float DequantizeUNorm16ToFloat( int Value16 )
 	{
 		check( Value16 >= 0 && Value16 <= 65535 );
 
 		return (float)Value16 * (1.f/65535.f);
 	}
 
-	static uint8 Requantize10to8( int Value10 )
+	[[nodiscard]] static uint8 Requantize10to8( int Value10 )
 	{
 		check( Value10 >= 0 && Value10 <= 1023 );
 
@@ -644,7 +644,7 @@ public:
 		return (uint8)Value8;
 	}
 	
-	static uint8 Requantize16to8(int Value16)
+	[[nodiscard]] static uint8 Requantize16to8(int Value16)
 	{
 		check( Value16 >= 0 && Value16 <= 65535 );
 
@@ -660,7 +660,7 @@ public:
 	/**
 	* Return 8-bit color Quantized from 10-bit RGB , 2-bit A
 	*/
-	static FColor MakeRequantizeFrom1010102( int R, int G, int B, int A )
+	[[nodiscard]] static FColor MakeRequantizeFrom1010102( int R, int G, int B, int A )
 	{
 		check( A >= 0 && A <= 3 );
 
@@ -678,7 +678,7 @@ public:
 	 *	@return a new FColor based of this color with the new alpha value.
 	 *	Usage: const FColor& MyColor = FColorList::Green.WithAlpha(128);
 	 */
-	FColor WithAlpha( uint8 Alpha ) const
+	[[nodiscard]] FColor WithAlpha( uint8 Alpha ) const
 	{
 		return FColor( R, G, B, Alpha );
 	}
@@ -689,7 +689,7 @@ public:
 	 * This matches the GPU spec conversion for U8<->float
 	 * @return The linear color representation.
 	 */
-	FORCEINLINE FLinearColor ReinterpretAsLinear() const
+	[[nodiscard]] FORCEINLINE FLinearColor ReinterpretAsLinear() const
 	{
 		constexpr float inv255 = 1.f / 255.f;
 		return FLinearColor(R * inv255, G * inv255, B * inv255, A * inv255);
@@ -703,7 +703,7 @@ public:
 	 * @return Hexadecimal string.
 	 * @see FromHex, ToString
 	 */
-	CORE_API FString ToHex() const;
+	[[nodiscard]] CORE_API FString ToHex() const;
 
 	/**
 	 * Converts this color value to a string.
@@ -711,7 +711,7 @@ public:
 	 * @return The string representation.
 	 * @see ToHex
 	 */
-	CORE_API FString ToString() const;
+	[[nodiscard]] CORE_API FString ToString() const;
 
 	/**
 	 * Initialize this Color based on an FString. The String is expected to contain R=, G=, B=, A=.
@@ -725,7 +725,7 @@ public:
 	/**
 	 * Gets the color in a packed uint32 format packed in the order ARGB.
 	 */
-	FORCEINLINE uint32 ToPackedARGB() const
+	[[nodiscard]] FORCEINLINE uint32 ToPackedARGB() const
 	{
 		return ( A << 24 ) | ( R << 16 ) | ( G << 8 ) | ( B << 0 );
 	}
@@ -733,7 +733,7 @@ public:
 	/**
 	 * Gets the color in a packed uint32 format packed in the order ABGR.
 	 */
-	FORCEINLINE uint32 ToPackedABGR() const
+	[[nodiscard]] FORCEINLINE uint32 ToPackedABGR() const
 	{
 		return ( A << 24 ) | ( B << 16 ) | ( G << 8 ) | ( R << 0 );
 	}
@@ -741,7 +741,7 @@ public:
 	/**
 	 * Gets the color in a packed uint32 format packed in the order RGBA.
 	 */
-	FORCEINLINE uint32 ToPackedRGBA() const
+	[[nodiscard]] FORCEINLINE uint32 ToPackedRGBA() const
 	{
 		return ( R << 24 ) | ( G << 16 ) | ( B << 8 ) | ( A << 0 );
 	}
@@ -749,7 +749,7 @@ public:
 	/**
 	 * Gets the color in a packed uint32 format packed in the order BGRA.
 	 */
-	FORCEINLINE uint32 ToPackedBGRA() const
+	[[nodiscard]] FORCEINLINE uint32 ToPackedBGRA() const
 	{
 		return ( B << 24 ) | ( G << 16 ) | ( R << 8 ) | ( A << 0 );
 	}
@@ -770,7 +770,7 @@ public:
 	static CORE_API const FColor Silver;
 	static CORE_API const FColor Emerald;
 
-	friend FORCEINLINE uint32 GetTypeHash( const FColor& Color )
+	[[nodiscard]] friend FORCEINLINE uint32 GetTypeHash( const FColor& Color )
 	{
 		return Color.DWColor();
 	}
diff --git a/Engine/Source/Runtime/Experimental/GeometryCollectionEngine/Private/GeometryCollection/GeometryCollectionComponent.cpp b/Engine/Source/Runtime/Experimental/GeometryCollectionEngine/Private/GeometryCollection/GeometryCollectionComponent.cpp
--- a/Engine/Source/Runtime/Experimental/GeometryCollectionEngine/Private/GeometryCollection/GeometryCollectionComponent.cpp
+++ b/Engine/Source/Runtime/Experimental/GeometryCollectionEngine/Private/GeometryCollection/GeometryCollectionComponent.cpp
@@ -5697,9 +5697,9 @@ void FScopedColorEdit::UpdateBoneColors()
 						int32 ColorIndex = Bone + 1; // parent can be -1 for root, range [-1..n]
 						BoneColor = RandomColors[ColorIndex % RandomColors.Num()];
 
-						BoneColor.LinearRGBToHSV();
+						BoneColor = BoneColor.LinearRGBToHSV();
 						BoneColor.B *= .5;
-						BoneColor.HSVToLinearRGB();
+						BoneColor = BoneColor.HSVToLinearRGB();
 					}
 					else
 					{
