CVars: Allow FAutoConsoleVariableRef to hold a FName directly

This is useful in rendering code for cvars that ultimately get used as an FName such as for debug visualization modes. The cvar is a string cvar internally, but also keeps a FName representation available so that the string to name conversion cost only happens when setting the variable, not every time it is queried.

diff --git a/Engine/Source/Runtime/Core/Private/HAL/ConsoleManager.cpp b/Engine/Source/Runtime/Core/Private/HAL/ConsoleManager.cpp
--- a/Engine/Source/Runtime/Core/Private/HAL/ConsoleManager.cpp
+++ b/Engine/Source/Runtime/Core/Private/HAL/ConsoleManager.cpp
@@ -91,6 +91,7 @@ public:
 	virtual IConsoleVariable* RegisterConsoleVariableRef(const TCHAR* Name, int32& RefValue, const TCHAR* Help, uint32 Flags) override;
 	virtual IConsoleVariable* RegisterConsoleVariableRef(const TCHAR* Name, float& RefValue, const TCHAR* Help, uint32 Flags) override;
 	virtual IConsoleVariable* RegisterConsoleVariableRef(const TCHAR* Name, FString& RefValue, const TCHAR* Help, uint32 Flags) override;
+	virtual IConsoleVariable* RegisterConsoleVariableRef(const TCHAR* Name, FName& RefValue, const TCHAR* Help, uint32 Flags) override;
 	virtual IConsoleVariable* RegisterConsoleVariableBitRef(const TCHAR* CVarName, const TCHAR* FlagName, uint32 BitNumber, uint8* Force0MaskPtr, uint8* Force1MaskPtr, const TCHAR* Help, uint32 Flags) override;
 
 
@@ -1821,6 +1822,86 @@ private: // ----------------------------------------------------
 	}
 };
 
+// name version
+
+class FConsoleVariableNameRef : public FConsoleVariableExtendedData<FString>
+{
+public:
+	FConsoleVariableNameRef(FName& InRefValue, const TCHAR* Help, EConsoleVariableFlags Flags)
+	: FConsoleVariableExtendedData<FString>(FString(), Help, Flags)
+	, RefValue(InRefValue)
+	, MainValue(InRefValue)
+	{
+	}
+
+	// interface IConsoleVariable -----------------------------------
+
+	virtual void Release()
+	{
+		delete this;
+	}
+	virtual void Set(const TCHAR* InValue, EConsoleVariableFlags SetBy, FName Tag)
+	{
+		TrackHistory(InValue, SetBy, Tag);
+		
+		if (CanChange(SetBy))
+		{
+			MainValue = FName(InValue);
+			OnChanged(SetBy);
+		}
+	}
+	virtual bool GetBool() const
+	{
+		ensureMsgf(false, TEXT("Getting the bool value of a FName based CVar is not supported"));
+		return false;
+	}
+	virtual int32 GetInt() const
+	{
+		ensureMsgf(false, TEXT("Getting the integer value of a FName based CVar is not supported"));
+		return 0;
+	}
+	virtual float GetFloat() const
+	{
+		ensureMsgf(false, TEXT("Getting the float value of a FName based CVar is not supported"));
+		return 0.0f;
+	}
+	virtual FString GetString() const
+	{
+		return MainValue.ToString();
+	}
+	virtual bool IsVariableString() const override
+	{
+		return true;
+	}
+
+private: // ----------------------------------------------------
+
+	// reference the the value (should not be changed from outside), if ECVF_RenderThreadSafe this is the render thread version, otherwise same as MainValue
+	FName& RefValue;
+	// main thread version 
+	FName MainValue;
+
+	const FName& Value() const
+	{
+		uint32 Index = GetShadowIndex();
+		checkSlow(Index < 2);
+		return (Index == 0) ? MainValue : RefValue;
+	}
+
+	virtual void SetInternal(const FString& TypedValue, EConsoleVariableFlags SetBy) override
+	{
+		MainValue = FName(TypedValue);
+	}
+	
+	void OnChanged(EConsoleVariableFlags SetBy, bool bForce=false)
+	{
+		// propagate from main thread to render thread or to reference
+		OnCVarChange(RefValue, MainValue, Flags, SetBy);
+		FConsoleVariableBase::OnChanged(SetBy, bForce);
+	}
+};
+
+
 class FConsoleVariableBitRef : public FConsoleVariableExtendedData<int>
 {
 public:
@@ -2316,6 +2397,13 @@ IConsoleVariable* FConsoleManager::RegisterConsoleVariableRef(const TCHAR* Name,
 	return AddConsoleObject(Name, new FConsoleVariableStringRef(RefValue, Help, (EConsoleVariableFlags)Flags))->AsVariable();
 }
 
+IConsoleVariable* FConsoleManager::RegisterConsoleVariableRef(const TCHAR* Name, FName& RefValue, const TCHAR* Help, uint32 Flags)
+{
+	LLM_SCOPE_BYNAME(TEXT("EngineMisc/ConsoleCommands"));
+	return AddConsoleObject(Name, new FConsoleVariableNameRef(RefValue, Help, (EConsoleVariableFlags)Flags))->AsVariable();
+}
+
+
 IConsoleCommand* FConsoleManager::RegisterConsoleCommand(const TCHAR* Name, const TCHAR* Help, const FConsoleCommandDelegate& Command, uint32 Flags)
 {
 	LLM_SCOPE_BYNAME(TEXT("EngineMisc/ConsoleCommands"));
diff --git a/Engine/Source/Runtime/Core/Public/HAL/IConsoleManager.h b/Engine/Source/Runtime/Core/Public/HAL/IConsoleManager.h
--- a/Engine/Source/Runtime/Core/Public/HAL/IConsoleManager.h
+++ b/Engine/Source/Runtime/Core/Public/HAL/IConsoleManager.h
@@ -720,6 +720,7 @@ struct IConsoleThreadPropagation
 	virtual void OnCVarChange(float& Dest, float NewValue) = 0;
 	virtual void OnCVarChange(bool& Dest, bool NewValue) = 0;
 	virtual void OnCVarChange(FString& Dest, const FString& NewValue) = 0;
+	virtual void OnCVarChange(FName& Dest, const FName& NewValue) = 0;
 };
 
 /**
@@ -927,6 +928,13 @@ struct IConsoleManager
 	* @param Flags bitmask combined from EConsoleVariableFlags
 	*/
 	virtual IConsoleVariable* RegisterConsoleVariableRef(const TCHAR* Name, FString& RefValue, const TCHAR* Help, uint32 Flags = ECVF_Default) = 0;
+	/**
+	* Create a reference to a name console variable
+	* @param Name must not be 0
+	* @param Help must not be 0
+	* @param Flags bitmask combined from EConsoleVariableFlags
+	*/
+	virtual IConsoleVariable* RegisterConsoleVariableRef(const TCHAR* Name, FName& RefValue, const TCHAR* Help, uint32 Flags = ECVF_Default) = 0;
 	/**
 	 * Create a reference to a show flag variable
 	 * @param CVarName must not be 0, e.g. "Show.PostProcessing"
@@ -1501,6 +1509,17 @@ public:
 	{
 	}
 
+	/**
+	 * Create a reference to a FName console variable
+	 * @param Name must not be 0
+	 * @param Help must not be 0
+	 * @param Flags bitmask combined from EConsoleVariableFlags
+	 */
+	FAutoConsoleVariableRef(const TCHAR* Name, FName& RefValue, const TCHAR* Help, uint32 Flags = ECVF_Default)
+		: FAutoConsoleObject(IConsoleManager::Get().RegisterConsoleVariableRef(Name, RefValue, Help, Flags))
+	{
+	}
+
 	/**
 	 * Create a reference to a int console variable
 	 * @param Name must not be 0
diff --git a/Engine/Source/Runtime/RenderCore/Private/RenderingThread.cpp b/Engine/Source/Runtime/RenderCore/Private/RenderingThread.cpp
--- a/Engine/Source/Runtime/RenderCore/Private/RenderingThread.cpp
+++ b/Engine/Source/Runtime/RenderCore/Private/RenderingThread.cpp
@@ -584,6 +584,16 @@ struct FConsoleRenderThreadPropagation : public IConsoleThreadPropagation
 			});
 	}
 
+	virtual void OnCVarChange(FName& Dest, const FName& NewValue)
+	{
+		FName* DestPtr = &Dest;
+		ENQUEUE_RENDER_COMMAND(OnCVarChange3)(
+			[DestPtr, NewValue](FRHICommandListImmediate& RHICmdList)
+			{
+				*DestPtr = NewValue;
+			});
+	}
+
 	static FConsoleRenderThreadPropagation& GetSingleton()
 	{
 		static FConsoleRenderThreadPropagation This;
