[AutoRTFM] Add atoi() as a pass-through function.

This is used by GenericPlatform\MicrosoftPlatformString.h to implement UE's Atoi method, and is relied upon by the JSON serializer. This was causing VerseTestScriptCmd to fail with AutoRTFM in Debug mode.

I'm not entirely sure why this was working in Development builds. The atoi call happened inside a `verify()` expression, but I thought Development would still run these.

diff --git a/Engine/Source/Runtime/Core/Private/AutoRTFM/StdLibSupport.cpp b/Engine/Source/Runtime/Core/Private/AutoRTFM/StdLibSupport.cpp
--- a/Engine/Source/Runtime/Core/Private/AutoRTFM/StdLibSupport.cpp
+++ b/Engine/Source/Runtime/Core/Private/AutoRTFM/StdLibSupport.cpp
@@ -239,6 +239,7 @@ UE_AUTORTFM_REGISTER_SELF_OVERLOADED_FUNCTION(const char* (const char*, int), st
 UE_AUTORTFM_REGISTER_SELF_OVERLOADED_FUNCTION(const char* (const char*, const char*), strstr);
 UE_AUTORTFM_REGISTER_SELF_FUNCTION(strlen);
 UE_AUTORTFM_REGISTER_SELF_FUNCTION(atof);
+UE_AUTORTFM_REGISTER_SELF_FUNCTION(atoi);
 
 UE_AUTORTFM_REGISTER_SELF_OVERLOADED_FUNCTION(const wchar_t* (const wchar_t*, wchar_t), wcschr);
 UE_AUTORTFM_REGISTER_SELF_OVERLOADED_FUNCTION(wchar_t* (wchar_t*, wchar_t), wcschr);
