Follow-up to document the code better (I missed a swarm comment doh!).

#rnx

diff --git a/Engine/Source/Runtime/Core/Public/Async/TransactionallySafeMutex.h b/Engine/Source/Runtime/Core/Public/Async/TransactionallySafeMutex.h
--- a/Engine/Source/Runtime/Core/Public/Async/TransactionallySafeMutex.h
+++ b/Engine/Source/Runtime/Core/Public/Async/TransactionallySafeMutex.h
@@ -59,7 +59,7 @@ public:
 				return false;
 			});
 		}
-		else if (AutoRTFM::IsCommittingOrAborting() && DoesTransactionHoldLock())
+		else if (HasLockTransactionallyHeldDuringCommittingOrAborting())
 		{
 			return !IsTransactionallyUnlocked();
 		}
@@ -99,7 +99,7 @@ public:
 				return true;
 			});
 		}
-		else if (AutoRTFM::IsCommittingOrAborting() && DoesTransactionHoldLock())
+		else if (HasLockTransactionallyHeldDuringCommittingOrAborting())
 		{
 			SetTransactionalLockMode(ETransactionalLockMode::TransactionallyLocked);
 			return true;
@@ -134,7 +134,7 @@ public:
 				RecordPreviousStateAndSetTransactionalLockMode(ETransactionalLockMode::TransactionallyLocked);
 			});
 		}
-		else if (AutoRTFM::IsCommittingOrAborting() && DoesTransactionHoldLock())
+		else if (HasLockTransactionallyHeldDuringCommittingOrAborting())
 		{
 			SetTransactionalLockMode(ETransactionalLockMode::TransactionallyLocked);
 		}
@@ -164,7 +164,7 @@ public:
 				RecordPreviousStateAndSetTransactionalLockMode(ETransactionalLockMode::TransactionallyUnlocked);
 			});
 		}
-		else if (AutoRTFM::IsCommittingOrAborting() && DoesTransactionHoldLock())
+		else if (HasLockTransactionallyHeldDuringCommittingOrAborting())
 		{
 			SetTransactionalLockMode(ETransactionalLockMode::TransactionallyUnlocked);
 		}
@@ -265,6 +265,16 @@ private:
 
 		SetTransactionalLockMode(LockMode);
 	}
+
+	inline bool HasLockTransactionallyHeldDuringCommittingOrAborting() const
+	{
+		// When we are running an on-commit / on-abort, we *could* not have run
+		// the on-commit / on-abort handlers for the mutex yet, and so if we get
+		// here we've got a mutex that is still participating in a transactionally
+		// held lock state, and thus we need to access the data as if we are on
+		// the transactional codepath.
+		return AutoRTFM::IsCommittingOrAborting() && DoesTransactionHoldLock();
+	}
 };
 
 } // UE
