Add [[Nodiscard]] on FName and FText
Add Preload to FText

#rnx

diff --git a/Engine/Plugins/Animation/ControlRig/Source/ControlRigEditor/Private/EditMode/ControlRigEditModeToolkit.cpp b/Engine/Plugins/Animation/ControlRig/Source/ControlRigEditor/Private/EditMode/ControlRigEditModeToolkit.cpp
--- a/Engine/Plugins/Animation/ControlRig/Source/ControlRigEditor/Private/EditMode/ControlRigEditModeToolkit.cpp
+++ b/Engine/Plugins/Animation/ControlRig/Source/ControlRigEditor/Private/EditMode/ControlRigEditModeToolkit.cpp
@@ -92,7 +92,7 @@ FText FControlRigEditModeToolkit::GetToolPaletteDisplayName(FName PaletteName) c
 {
 	if (PaletteName == AnimationName)
 	{
-		FText::FromName(AnimationName);
+		return FText::FromName(AnimationName);
 	}
 	return FText();
 }
diff --git a/Engine/Plugins/FX/Niagara/Source/NiagaraEditor/Private/NiagaraMessages.cpp b/Engine/Plugins/FX/Niagara/Source/NiagaraEditor/Private/NiagaraMessages.cpp
--- a/Engine/Plugins/FX/Niagara/Source/NiagaraEditor/Private/NiagaraMessages.cpp
+++ b/Engine/Plugins/FX/Niagara/Source/NiagaraEditor/Private/NiagaraMessages.cpp
@@ -278,7 +278,7 @@ TSharedRef<const INiagaraMessage> FNiagaraMessageJobCompileEvent::GenerateNiagar
 	//The originating script weak ptr is no longer valid, send an error message asking for recompile.
 	FText MessageText = LOCTEXT("CompileEventMessageJobFail", "Cached info for compile event is out of date, recompile to get full info. Event: {0}");
 	//Add in the message of the compile event for visibility on which compile events do not have info.
-	FText::Format(MessageText, FText::FromString(CompileEvent.Message));
+	MessageText = FText::Format(MessageText, FText::FromString(CompileEvent.Message));
 	return MakeShared<FNiagaraMessageText>(MessageText, EMessageSeverity::Error, FNiagaraMessageTopics::CompilerTopicName);
 }
 
diff --git a/Engine/Plugins/FX/Niagara/Source/NiagaraEditorWidgets/Private/NiagaraStackCommandContext.cpp b/Engine/Plugins/FX/Niagara/Source/NiagaraEditorWidgets/Private/NiagaraStackCommandContext.cpp
--- a/Engine/Plugins/FX/Niagara/Source/NiagaraEditorWidgets/Private/NiagaraStackCommandContext.cpp
+++ b/Engine/Plugins/FX/Niagara/Source/NiagaraEditorWidgets/Private/NiagaraStackCommandContext.cpp
@@ -235,7 +235,7 @@ FText FNiagaraStackCommandContext::GetCanRenameSelectedEntriesToolTip() const
 {
 	if (SelectedEntries.Num() > 1)
 	{
-		LOCTEXT("CantRenameMultipleToolTip", "Can't rename multiple items.");
+		return LOCTEXT("CantRenameMultipleToolTip", "Can't rename multiple items.");
 	}
 
 	if (!SelectedEntries[0]->SupportsRename())
diff --git a/Engine/Plugins/Runtime/LevelStreamingPersistence/Source/LevelStreamingPersistence/Private/LevelStreamingPersistenceManager.cpp b/Engine/Plugins/Runtime/LevelStreamingPersistence/Source/LevelStreamingPersistence/Private/LevelStreamingPersistenceManager.cpp
--- a/Engine/Plugins/Runtime/LevelStreamingPersistence/Source/LevelStreamingPersistence/Private/LevelStreamingPersistenceManager.cpp
+++ b/Engine/Plugins/Runtime/LevelStreamingPersistence/Source/LevelStreamingPersistence/Private/LevelStreamingPersistenceManager.cpp
@@ -1182,7 +1182,7 @@ namespace LSPMConsoleCommandHelper
 	template<> bool GetValueFromString(const FString& InPropertyValue, double& OutResult) { LexFromString(OutResult, *InPropertyValue); return true; }
 	template<> bool GetValueFromString(const FString& InPropertyValue, FName& OutResult) { OutResult = FName(InPropertyValue); return true; }
 	template<> bool GetValueFromString(const FString& InPropertyValue, FString& OutResult) { OutResult = InPropertyValue; return true; }
-	template<> bool GetValueFromString(const FString& InPropertyValue, FText& OutResult) { OutResult.FromString(InPropertyValue); return true; }
+	template<> bool GetValueFromString(const FString& InPropertyValue, FText& OutResult) { OutResult = FText::FromString(InPropertyValue); return true; }
 	template<> bool GetValueFromString(const FString& InPropertyValue, FVector& OutResult) { return OutResult.InitFromString(InPropertyValue); }
 	template<> bool GetValueFromString(const FString& InPropertyValue, FRotator& OutResult) { return OutResult.InitFromString(InPropertyValue); }
 	template<> bool GetValueFromString(const FString& InPropertyValue, FTransform& OutResult) { return OutResult.InitFromString(InPropertyValue); }
diff --git a/Engine/Plugins/TraceUtilities/Source/EditorTraceUtilities/Private/InsightsStatusBar/SInsightsStatusBar.cpp b/Engine/Plugins/TraceUtilities/Source/EditorTraceUtilities/Private/InsightsStatusBar/SInsightsStatusBar.cpp
--- a/Engine/Plugins/TraceUtilities/Source/EditorTraceUtilities/Private/InsightsStatusBar/SInsightsStatusBar.cpp
+++ b/Engine/Plugins/TraceUtilities/Source/EditorTraceUtilities/Private/InsightsStatusBar/SInsightsStatusBar.cpp
@@ -1096,7 +1096,7 @@ FText SInsightsStatusBarWidget::GetRegionSwitchDescText() const
 {
 	if (!ToggleRegion_CanExecute())
 	{
-		LOCTEXT("TraceRegionDisabledDesc", "Regions can only be traced when the Region channel in enabled.");
+		return LOCTEXT("TraceRegionDisabledDesc", "Regions can only be traced when the Region channel in enabled.");
 	}
 
 	if(RegionIsActive())
diff --git a/Engine/Plugins/VirtualProduction/DMX/DMXEngine/Source/DMXEditor/Private/DMXEditorUtils.cpp b/Engine/Plugins/VirtualProduction/DMX/DMXEngine/Source/DMXEditor/Private/DMXEditorUtils.cpp
--- a/Engine/Plugins/VirtualProduction/DMX/DMXEngine/Source/DMXEditor/Private/DMXEditorUtils.cpp
+++ b/Engine/Plugins/VirtualProduction/DMX/DMXEngine/Source/DMXEditor/Private/DMXEditorUtils.cpp
@@ -410,7 +410,7 @@ bool FDMXEditorUtils::DoesLibraryHaveUniverseConflicts(UDMXLibrary* Library, FTe
 						OutInputPortConflictMessage = LOCTEXT("LibraryInputPortUniverseConflictMessageStart", "Libraries use the same Input Port: ");
 					}
 					
-					FText::Format(LOCTEXT("LibraryInputPortUniverseConflictMessage", "{0} {1}"), OutInputPortConflictMessage, FText::FromString(OtherDMXLibrary->GetName()));
+					OutInputPortConflictMessage = FText::Format(LOCTEXT("LibraryInputPortUniverseConflictMessage", "{0} {1}"), OutInputPortConflictMessage, FText::FromString(OtherDMXLibrary->GetName()));
 				}
 			}
 		}
@@ -433,7 +433,7 @@ bool FDMXEditorUtils::DoesLibraryHaveUniverseConflicts(UDMXLibrary* Library, FTe
 						OutOutputPortConflictMessage = LOCTEXT("LibraryOutputPortUniverseConflictMessageStart", "Libraries that use the same Output Port: ");
 					}
 
-					FText::Format(LOCTEXT("LibraryOutputPortUniverseConflictMessage", "{0} {1}"), OutOutputPortConflictMessage, FText::FromString(OtherDMXLibrary->GetName()));
+					OutInputPortConflictMessage = FText::Format(LOCTEXT("LibraryOutputPortUniverseConflictMessage", "{0} {1}"), OutOutputPortConflictMessage, FText::FromString(OtherDMXLibrary->GetName()));
 				}
 			}
 		}
diff --git a/Engine/Source/Runtime/Core/Private/Internationalization/Text.cpp b/Engine/Source/Runtime/Core/Private/Internationalization/Text.cpp
--- a/Engine/Source/Runtime/Core/Private/Internationalization/Text.cpp
+++ b/Engine/Source/Runtime/Core/Private/Internationalization/Text.cpp
@@ -1241,6 +1241,11 @@ bool FText::ShouldGatherForLocalization() const
 	return !((Flags & ETextFlag::CultureInvariant) || (Flags & ETextFlag::Transient)) && !IsFromStringTable() && !SourceString.IsEmpty() && !IsAllWhitespace(SourceString);
 }
 
+void FText::Preload( FName InTableId, const FTextKey& InKey )
+{
+	FTextHistory_StringTableEntry(InTableId, InKey, EStringTableLoadingPolicy::FindOrLoad);
+}
+
 void FText::GetHistoricFormatData(TArray<FHistoricTextFormatData>& OutHistoricFormatData) const
 {
 	TextData->GetTextHistory().GetHistoricFormatData(*this, OutHistoricFormatData);
diff --git a/Engine/Source/Runtime/Core/Public/Internationalization/Text.h b/Engine/Source/Runtime/Core/Public/Internationalization/Text.h
--- a/Engine/Source/Runtime/Core/Public/Internationalization/Text.h
+++ b/Engine/Source/Runtime/Core/Public/Internationalization/Text.h
@@ -401,30 +401,30 @@ public:
 	/**
 	 * Generate an FText that represents the passed number in the current culture
 	 */
-	static CORE_API FText AsNumber(float Val,	const FNumberFormattingOptions* const Options = NULL, const FCulturePtr& TargetCulture = NULL);
-	static CORE_API FText AsNumber(double Val,	const FNumberFormattingOptions* const Options = NULL, const FCulturePtr& TargetCulture = NULL);
-	static CORE_API FText AsNumber(int8 Val,		const FNumberFormattingOptions* const Options = NULL, const FCulturePtr& TargetCulture = NULL);
-	static CORE_API FText AsNumber(int16 Val,	const FNumberFormattingOptions* const Options = NULL, const FCulturePtr& TargetCulture = NULL);
-	static CORE_API FText AsNumber(int32 Val,	const FNumberFormattingOptions* const Options = NULL, const FCulturePtr& TargetCulture = NULL);
-	static CORE_API FText AsNumber(int64 Val,	const FNumberFormattingOptions* const Options = NULL, const FCulturePtr& TargetCulture = NULL);
-	static CORE_API FText AsNumber(uint8 Val,	const FNumberFormattingOptions* const Options = NULL, const FCulturePtr& TargetCulture = NULL);
-	static CORE_API FText AsNumber(uint16 Val,	const FNumberFormattingOptions* const Options = NULL, const FCulturePtr& TargetCulture = NULL);
-	static CORE_API FText AsNumber(uint32 Val,	const FNumberFormattingOptions* const Options = NULL, const FCulturePtr& TargetCulture = NULL);
-	static CORE_API FText AsNumber(uint64 Val,	const FNumberFormattingOptions* const Options = NULL, const FCulturePtr& TargetCulture = NULL);
-	static CORE_API FText AsNumber(long Val,		const FNumberFormattingOptions* const Options = NULL, const FCulturePtr& TargetCulture = NULL);
-
-
-	static CORE_API FText AsCurrency(float Val,  const FString& CurrencyCode = FString(), const FNumberFormattingOptions* const Options = NULL, const FCulturePtr& TargetCulture = NULL);
-	static CORE_API FText AsCurrency(double Val, const FString& CurrencyCode = FString(), const FNumberFormattingOptions* const Options = NULL, const FCulturePtr& TargetCulture = NULL);
-	static CORE_API FText AsCurrency(int8 Val,   const FString& CurrencyCode = FString(), const FNumberFormattingOptions* const Options = NULL, const FCulturePtr& TargetCulture = NULL);
-	static CORE_API FText AsCurrency(int16 Val,  const FString& CurrencyCode = FString(), const FNumberFormattingOptions* const Options = NULL, const FCulturePtr& TargetCulture = NULL);
-	static CORE_API FText AsCurrency(int32 Val,  const FString& CurrencyCode = FString(), const FNumberFormattingOptions* const Options = NULL, const FCulturePtr& TargetCulture = NULL);
-	static CORE_API FText AsCurrency(int64 Val,  const FString& CurrencyCode = FString(), const FNumberFormattingOptions* const Options = NULL, const FCulturePtr& TargetCulture = NULL);
-	static CORE_API FText AsCurrency(uint8 Val,  const FString& CurrencyCode = FString(), const FNumberFormattingOptions* const Options = NULL, const FCulturePtr& TargetCulture = NULL);
-	static CORE_API FText AsCurrency(uint16 Val, const FString& CurrencyCode = FString(), const FNumberFormattingOptions* const Options = NULL, const FCulturePtr& TargetCulture = NULL);
-	static CORE_API FText AsCurrency(uint32 Val, const FString& CurrencyCode = FString(), const FNumberFormattingOptions* const Options = NULL, const FCulturePtr& TargetCulture = NULL);
-	static CORE_API FText AsCurrency(uint64 Val, const FString& CurrencyCode = FString(), const FNumberFormattingOptions* const Options = NULL, const FCulturePtr& TargetCulture = NULL);
-	static CORE_API FText AsCurrency(long Val,   const FString& CurrencyCode = FString(), const FNumberFormattingOptions* const Options = NULL, const FCulturePtr& TargetCulture = NULL);
+	[[nodiscard]] static CORE_API FText AsNumber(float Val,	const FNumberFormattingOptions* const Options = NULL, const FCulturePtr& TargetCulture = NULL);
+	[[nodiscard]] static CORE_API FText AsNumber(double Val,	const FNumberFormattingOptions* const Options = NULL, const FCulturePtr& TargetCulture = NULL);
+	[[nodiscard]] static CORE_API FText AsNumber(int8 Val,		const FNumberFormattingOptions* const Options = NULL, const FCulturePtr& TargetCulture = NULL);
+	[[nodiscard]] static CORE_API FText AsNumber(int16 Val,	const FNumberFormattingOptions* const Options = NULL, const FCulturePtr& TargetCulture = NULL);
+	[[nodiscard]] static CORE_API FText AsNumber(int32 Val,	const FNumberFormattingOptions* const Options = NULL, const FCulturePtr& TargetCulture = NULL);
+	[[nodiscard]] static CORE_API FText AsNumber(int64 Val,	const FNumberFormattingOptions* const Options = NULL, const FCulturePtr& TargetCulture = NULL);
+	[[nodiscard]] static CORE_API FText AsNumber(uint8 Val,	const FNumberFormattingOptions* const Options = NULL, const FCulturePtr& TargetCulture = NULL);
+	[[nodiscard]] static CORE_API FText AsNumber(uint16 Val,	const FNumberFormattingOptions* const Options = NULL, const FCulturePtr& TargetCulture = NULL);
+	[[nodiscard]] static CORE_API FText AsNumber(uint32 Val,	const FNumberFormattingOptions* const Options = NULL, const FCulturePtr& TargetCulture = NULL);
+	[[nodiscard]] static CORE_API FText AsNumber(uint64 Val,	const FNumberFormattingOptions* const Options = NULL, const FCulturePtr& TargetCulture = NULL);
+	[[nodiscard]] static CORE_API FText AsNumber(long Val,		const FNumberFormattingOptions* const Options = NULL, const FCulturePtr& TargetCulture = NULL);
+
+
+	[[nodiscard]] static CORE_API FText AsCurrency(float Val,  const FString& CurrencyCode = FString(), const FNumberFormattingOptions* const Options = NULL, const FCulturePtr& TargetCulture = NULL);
+	[[nodiscard]] static CORE_API FText AsCurrency(double Val, const FString& CurrencyCode = FString(), const FNumberFormattingOptions* const Options = NULL, const FCulturePtr& TargetCulture = NULL);
+	[[nodiscard]] static CORE_API FText AsCurrency(int8 Val,   const FString& CurrencyCode = FString(), const FNumberFormattingOptions* const Options = NULL, const FCulturePtr& TargetCulture = NULL);
+	[[nodiscard]] static CORE_API FText AsCurrency(int16 Val,  const FString& CurrencyCode = FString(), const FNumberFormattingOptions* const Options = NULL, const FCulturePtr& TargetCulture = NULL);
+	[[nodiscard]] static CORE_API FText AsCurrency(int32 Val,  const FString& CurrencyCode = FString(), const FNumberFormattingOptions* const Options = NULL, const FCulturePtr& TargetCulture = NULL);
+	[[nodiscard]] static CORE_API FText AsCurrency(int64 Val,  const FString& CurrencyCode = FString(), const FNumberFormattingOptions* const Options = NULL, const FCulturePtr& TargetCulture = NULL);
+	[[nodiscard]] static CORE_API FText AsCurrency(uint8 Val,  const FString& CurrencyCode = FString(), const FNumberFormattingOptions* const Options = NULL, const FCulturePtr& TargetCulture = NULL);
+	[[nodiscard]] static CORE_API FText AsCurrency(uint16 Val, const FString& CurrencyCode = FString(), const FNumberFormattingOptions* const Options = NULL, const FCulturePtr& TargetCulture = NULL);
+	[[nodiscard]] static CORE_API FText AsCurrency(uint32 Val, const FString& CurrencyCode = FString(), const FNumberFormattingOptions* const Options = NULL, const FCulturePtr& TargetCulture = NULL);
+	[[nodiscard]] static CORE_API FText AsCurrency(uint64 Val, const FString& CurrencyCode = FString(), const FNumberFormattingOptions* const Options = NULL, const FCulturePtr& TargetCulture = NULL);
+	[[nodiscard]] static CORE_API FText AsCurrency(long Val,   const FString& CurrencyCode = FString(), const FNumberFormattingOptions* const Options = NULL, const FCulturePtr& TargetCulture = NULL);
 
 	/**
 	 * Generate an FText that represents the passed number as currency in the current culture.
@@ -432,72 +432,72 @@ public:
 	 * Keep in mind the CurrencyCode is completely independent of the culture it's displayed in (and they do not imply one another).
 	 * For example: FText::AsCurrencyBase(650, TEXT("EUR")); would return an FText of "<EUR>6.50" in most English cultures (en_US/en_UK) and "6,50<EUR>" in Spanish (es_ES) (where <EUR> is U+20AC)
 	 */
-	static CORE_API FText AsCurrencyBase(int64 BaseVal, const FString& CurrencyCode, const FCulturePtr& TargetCulture = NULL, int32 ForceDecimalPlaces = -1);
+	[[nodiscard]] static CORE_API FText AsCurrencyBase(int64 BaseVal, const FString& CurrencyCode, const FCulturePtr& TargetCulture = NULL, int32 ForceDecimalPlaces = -1);
 
 	/**
 	 * Generate an FText that represents the passed number as a percentage in the current culture
 	 */
-	static CORE_API FText AsPercent(float Val,	const FNumberFormattingOptions* const Options = NULL, const FCulturePtr& TargetCulture = NULL);
-	static CORE_API FText AsPercent(double Val,	const FNumberFormattingOptions* const Options = NULL, const FCulturePtr& TargetCulture = NULL);
+	[[nodiscard]] static CORE_API FText AsPercent(float Val,	const FNumberFormattingOptions* const Options = NULL, const FCulturePtr& TargetCulture = NULL);
+	[[nodiscard]] static CORE_API FText AsPercent(double Val,	const FNumberFormattingOptions* const Options = NULL, const FCulturePtr& TargetCulture = NULL);
 
 	/**
 	 * Generate an FText that represents the passed number as a date and/or time in the current culture
 	 * @note The overload using a custom pattern uses strftime-like syntax (see FDateTime::ToFormattedString)
 	 */
-	static CORE_API FText AsDate(const FDateTime& DateTime, const EDateTimeStyle::Type DateStyle = EDateTimeStyle::Default, const FString& TimeZone = FString(), const FCulturePtr& TargetCulture = NULL);
-	static CORE_API FText AsDateTime(const FDateTime& DateTime, const EDateTimeStyle::Type DateStyle = EDateTimeStyle::Default, const EDateTimeStyle::Type TimeStyle = EDateTimeStyle::Default, const FString& TimeZone = FString(), const FCulturePtr& TargetCulture = NULL);
-	static CORE_API FText AsDateTime(const FDateTime& DateTime, const FString& CustomPattern, const FString& TimeZone = FString(), const FCulturePtr& TargetCulture = NULL);
-	static CORE_API FText AsTime(const FDateTime& DateTime, const EDateTimeStyle::Type TimeStyle = EDateTimeStyle::Default, const FString& TimeZone = FString(), const FCulturePtr& TargetCulture = NULL);
-	static CORE_API FText AsTimespan(const FTimespan& Timespan, const FCulturePtr& TargetCulture = NULL);
-
+	[[nodiscard]] static CORE_API FText AsDate(const FDateTime& DateTime, const EDateTimeStyle::Type DateStyle = EDateTimeStyle::Default, const FString& TimeZone = FString(), const FCulturePtr& TargetCulture = NULL);
+	[[nodiscard]] static CORE_API FText AsDateTime(const FDateTime& DateTime, const EDateTimeStyle::Type DateStyle = EDateTimeStyle::Default, const EDateTimeStyle::Type TimeStyle = EDateTimeStyle::Default, const FString& TimeZone = FString(), const FCulturePtr& TargetCulture = NULL);
+	[[nodiscard]] static CORE_API FText AsDateTime(const FDateTime& DateTime, const FString& CustomPattern, const FString& TimeZone = FString(), const FCulturePtr& TargetCulture = NULL);
+	[[nodiscard]] static CORE_API FText AsTime(const FDateTime& DateTime, const EDateTimeStyle::Type TimeStyle = EDateTimeStyle::Default, const FString& TimeZone = FString(), const FCulturePtr& TargetCulture = NULL);
+	[[nodiscard]] static CORE_API FText AsTimespan(const FTimespan& Timespan, const FCulturePtr& TargetCulture = NULL);
+[[nodiscard]] 
 	/**
 	 * Gets the time zone string that represents a non-specific, zero offset, culture invariant time zone.
 	 */
-	static CORE_API FString GetInvariantTimeZone();
+	[[nodiscard]] static CORE_API FString GetInvariantTimeZone();
 
 	/**
 	 * Generate an FText that represents the passed number as a memory size in the current culture
 	 */
-	static CORE_API FText AsMemory(uint64 NumBytes, const FNumberFormattingOptions* const Options = NULL, const FCulturePtr& TargetCulture = NULL, EMemoryUnitStandard UnitStandard = EMemoryUnitStandard::IEC);
+	[[nodiscard]] static CORE_API FText AsMemory(uint64 NumBytes, const FNumberFormattingOptions* const Options = NULL, const FCulturePtr& TargetCulture = NULL, EMemoryUnitStandard UnitStandard = EMemoryUnitStandard::IEC);
 
 	/**
 	 * Generate an FText that represents the passed number as a memory size in the current culture
 	 */
-	static CORE_API FText AsMemory(uint64 NumBytes, EMemoryUnitStandard UnitStandard);
+	[[nodiscard]] static CORE_API FText AsMemory(uint64 NumBytes, EMemoryUnitStandard UnitStandard);
 
 	/**
 	 * Attempts to create an FText instance from a string table ID and key (this is the same as the LOCTABLE macro, except this can also work with non-literal string values).
 	 * @return The found text, or a dummy FText if not found.
 	 */
-	static CORE_API FText FromStringTable(const FName InTableId, const FTextKey& InKey, const EStringTableLoadingPolicy InLoadingPolicy = EStringTableLoadingPolicy::FindOrLoad);
+	[[nodiscard]] static CORE_API FText FromStringTable(const FName InTableId, const FTextKey& InKey, const EStringTableLoadingPolicy InLoadingPolicy = EStringTableLoadingPolicy::FindOrLoad);
 
 	/**
 	 * Generate an FText representing the pass name
 	 */
-	static CORE_API FText FromName( const FName& Val);
+	[[nodiscard]] static CORE_API FText FromName( const FName& Val);
 	
 	/**
 	 * Generate an FText representing the passed in string
 	 */
-	static CORE_API FText FromString( const FString& String );
-	static CORE_API FText FromString( FString&& String );
+	[[nodiscard]] static CORE_API FText FromString( const FString& String );
+	[[nodiscard]] static CORE_API FText FromString( FString&& String );
 
 	/**
 	 * Generate a FText representing the passed string view
 	 */
-	static CORE_API FText FromStringView(FStringView InString);
+	[[nodiscard]] static CORE_API FText FromStringView(FStringView InString);
 
 	/**
 	 * Generate a culture invariant FText representing the passed in string
 	 */
-	static CORE_API FText AsCultureInvariant( const TCHAR* String );
-	static CORE_API FText AsCultureInvariant( FStringView String );
-	static CORE_API FText AsCultureInvariant( FString&& String );
+	[[nodiscard]] static CORE_API FText AsCultureInvariant( const TCHAR* String );
+	[[nodiscard]] static CORE_API FText AsCultureInvariant( FStringView String );
+	[[nodiscard]] static CORE_API FText AsCultureInvariant( FString&& String );
 
 	/**
 	 * Generate a culture invariant FText representing the passed in FText
 	 */
-	static CORE_API FText AsCultureInvariant( FText Text );
+	[[nodiscard]] static CORE_API FText AsCultureInvariant( FText Text );
 
 	/**
 	 * === !! This is an ADVANCED function. USE WITH CAUTION !! ===
@@ -510,9 +510,9 @@ public:
 	 * When using it to dynamically reference an EXISTING FText (such as known LOCTEXT in code) it can be thought of as an alternative to FText::FindTextInLiveTable_Advanced that also handles untranslated text.
 	 * Note: Direct dynamic references to FText are EXTREMELY FRAGILE, and you may want to use a string table instead!
 	 */
-	static CORE_API FText AsLocalizable_Advanced(const FTextKey& Namespace, const FTextKey& Key, const TCHAR* String);
-	static CORE_API FText AsLocalizable_Advanced(const FTextKey& Namespace, const FTextKey& Key, FStringView String);
-	static CORE_API FText AsLocalizable_Advanced(const FTextKey& Namespace, const FTextKey& Key, FString&& String);
+	[[nodiscard]] static CORE_API FText AsLocalizable_Advanced(const FTextKey& Namespace, const FTextKey& Key, const TCHAR* String);
+	[[nodiscard]] static CORE_API FText AsLocalizable_Advanced(const FTextKey& Namespace, const FTextKey& Key, FStringView String);
+	[[nodiscard]] static CORE_API FText AsLocalizable_Advanced(const FTextKey& Namespace, const FTextKey& Key, FString&& String);
 
 	/**
 	 * === !! This is an ADVANCED function. USE WITH CAUTION !! ===
@@ -535,18 +535,18 @@ public:
 		return FText::FindTextInLiveTable_Advanced(Namespace, Key, OutText, SourceString);
 	}
 
-	CORE_API const FString& ToString() const;
+	[[nodiscard]] CORE_API const FString& ToString() const;
 
 	/** Deep build of the source string for this FText, climbing the history hierarchy */
-	CORE_API FString BuildSourceString() const;
+	[[nodiscard]] CORE_API FString BuildSourceString() const;
 
-	CORE_API bool IsNumeric() const;
+	[[nodiscard]] CORE_API bool IsNumeric() const;
 
-	CORE_API int32 CompareTo( const FText& Other, const ETextComparisonLevel::Type ComparisonLevel = ETextComparisonLevel::Default ) const;
-	CORE_API int32 CompareToCaseIgnored( const FText& Other ) const;
+	[[nodiscard]] CORE_API int32 CompareTo( const FText& Other, const ETextComparisonLevel::Type ComparisonLevel = ETextComparisonLevel::Default ) const;
+	[[nodiscard]] CORE_API int32 CompareToCaseIgnored( const FText& Other ) const;
 
-	CORE_API bool EqualTo( const FText& Other, const ETextComparisonLevel::Type ComparisonLevel = ETextComparisonLevel::Default ) const;
-	CORE_API bool EqualToCaseIgnored( const FText& Other ) const;
+	[[nodiscard]] CORE_API bool EqualTo( const FText& Other, const ETextComparisonLevel::Type ComparisonLevel = ETextComparisonLevel::Default ) const;
+	[[nodiscard]] CORE_API bool EqualToCaseIgnored( const FText& Other ) const;
 
 	/**
 	 * Check to see if this FText is identical to the other FText
@@ -556,7 +556,7 @@ public:
 	 *
 	 * @note If you actually want to perform a full lexical comparison, then you need to use EqualTo instead.
 	 */
-	CORE_API bool IdenticalTo( const FText& Other, const ETextIdenticalModeFlags CompareModeFlags = ETextIdenticalModeFlags::None ) const;
+	[[nodiscard]] CORE_API bool IdenticalTo( const FText& Other, const ETextIdenticalModeFlags CompareModeFlags = ETextIdenticalModeFlags::None ) const;
 
 	class FSortPredicate
 	{
@@ -572,41 +572,41 @@ public:
 #endif
 	};
 
-	CORE_API bool IsEmpty() const;
+	[[nodiscard]] CORE_API bool IsEmpty() const;
 
-	CORE_API bool IsEmptyOrWhitespace() const;
+	[[nodiscard]] CORE_API bool IsEmptyOrWhitespace() const;
 
 	/**
 	 * Transforms the text to lowercase in a culture correct way.
 	 * @note The returned instance is linked to the original and will be rebuilt if the active culture is changed.
 	 */
-	CORE_API FText ToLower() const;
+	[[nodiscard]] CORE_API FText ToLower() const;
 
 	/**
 	 * Transforms the text to uppercase in a culture correct way.
 	 * @note The returned instance is linked to the original and will be rebuilt if the active culture is changed.
 	 */
-	CORE_API FText ToUpper() const;
+	[[nodiscard]] CORE_API FText ToUpper() const;
 
 	/**
 	 * Removes any whitespace characters from the start of the text.
 	 */
-	static CORE_API FText TrimPreceding( const FText& );
+	[[nodiscard]] static CORE_API FText TrimPreceding( const FText& );
 
 	/**
 	 * Removes any whitespace characters from the end of the text.
 	 */
-	static CORE_API FText TrimTrailing( const FText& );
+	[[nodiscard]] static CORE_API FText TrimTrailing( const FText& );
 
 	/**
 	 * Removes any whitespace characters from the start and end of the text.
 	 */
-	static CORE_API FText TrimPrecedingAndTrailing( const FText& );
+	[[nodiscard]] static CORE_API FText TrimPrecedingAndTrailing( const FText& );
 
 	/**
 	 * Check to see if the given character is considered whitespace by the current culture
 	 */
-	static CORE_API bool IsWhitespace( const TCHAR Char );
+	[[nodiscard]] static CORE_API bool IsWhitespace( const TCHAR Char );
 
 	static CORE_API void GetFormatPatternParameters(const FTextFormat& Fmt, TArray<FString>& ParameterNames);
 
@@ -620,8 +620,8 @@ public:
 	 * @param InArguments The map of key->value pairs to inject into the format pattern
 	 * @return The formatted FText
 	 */
-	static CORE_API FText Format(FTextFormat Fmt, const FFormatNamedArguments& InArguments);
-	static CORE_API FText Format(FTextFormat Fmt, FFormatNamedArguments&& InArguments);
+	[[nodiscard]] static CORE_API FText Format(FTextFormat Fmt, const FFormatNamedArguments& InArguments);
+	[[nodiscard]] static CORE_API FText Format(FTextFormat Fmt, FFormatNamedArguments&& InArguments);
 
 	/**
 	 * Format the given list values as ordered arguments within the given format pattern
@@ -633,8 +633,8 @@ public:
 	 * @param InArguments The list of values to inject into the format pattern
 	 * @return The formatted FText
 	 */
-	static CORE_API FText Format(FTextFormat Fmt, const FFormatOrderedArguments& InArguments);
-	static CORE_API FText Format(FTextFormat Fmt, FFormatOrderedArguments&& InArguments);
+	[[nodiscard]] static CORE_API FText Format(FTextFormat Fmt, const FFormatOrderedArguments& InArguments);
+	[[nodiscard]] static CORE_API FText Format(FTextFormat Fmt, FFormatOrderedArguments&& InArguments);
 
 	/**
 	 * Format the given list of variadic values as ordered arguments within the given format pattern
@@ -649,7 +649,7 @@ public:
 	 * @return The formatted FText
 	 */
 	template <typename... ArgTypes>
-	static FORCEINLINE FText Format(FTextFormat Fmt, ArgTypes... Args)
+	[[nodiscard]] static FORCEINLINE FText Format(FTextFormat Fmt, ArgTypes... Args)
 	{
 		static_assert((TIsConstructible<FFormatArgumentValue, ArgTypes>::Value && ...), "Invalid argument type passed to FText::Format");
 		static_assert(sizeof...(Args) > 0, "FText::Format expects at least one non-format argument"); // we do this to ensure that people don't call Format for no good reason
@@ -674,7 +674,7 @@ public:
 	 * @return The formatted FText
 	 */
 	template < typename... TArguments >
-	static FText FormatNamed( FTextFormat Fmt, TArguments&&... Args );
+	[[nodiscard]] static FText FormatNamed( FTextFormat Fmt, TArguments&&... Args );
 
 	/**
 	 * Format the given list of variadic values as ordered arguments within the given format pattern
@@ -689,7 +689,7 @@ public:
 	 * @return The formatted FText
 	 */
 	template < typename... TArguments >
-	static FText FormatOrdered( FTextFormat Fmt, TArguments&&... Args );
+	[[nodiscard]] static FText FormatOrdered( FTextFormat Fmt, TArguments&&... Args );
 
 	/**
 	 * Join an arbitrary list of formattable values together, separated by the given delimiter
@@ -699,8 +699,8 @@ public:
 	 * @param Args An array of formattable values to join together
 	 * @return The joined FText
 	 */
-	static CORE_API FText Join(const FText& Delimiter, const FFormatOrderedArguments& Args);
-	static CORE_API FText Join(const FText& Delimiter, const TArray<FText>& Args);
+	[[nodiscard]] static CORE_API FText Join(const FText& Delimiter, const FFormatOrderedArguments& Args);
+	[[nodiscard]] static CORE_API FText Join(const FText& Delimiter, const TArray<FText>& Args);
 
 	/**
 	 * Join an arbitrary list of formattable items together, separated by the given delimiter
@@ -711,7 +711,7 @@ public:
 	 * @return The joined FText
 	 */
 	template <typename... ArgTypes>
-	static FORCEINLINE FText Join(const FText& Delimiter, ArgTypes... Args)
+	[[nodiscard]] static FORCEINLINE FText Join(const FText& Delimiter, ArgTypes... Args)
 	{
 		static_assert((TIsConstructible<FFormatArgumentValue, ArgTypes>::Value && ...), "Invalid argument type passed to FText::Join");
 		static_assert(sizeof...(Args) > 0, "FText::Join expects at least one non-format argument"); // we do this to ensure that people don't call Join for no good reason
@@ -725,7 +725,7 @@ public:
 	 *
 	 * @param TextGenerator the text generator object that will generate the text
 	 */
-	static CORE_API FText FromTextGenerator( const TSharedRef<ITextGenerator>& TextGenerator );
+	[[nodiscard]] static CORE_API FText FromTextGenerator( const TSharedRef<ITextGenerator>& TextGenerator );
 
 	DECLARE_DELEGATE_RetVal_OneParam( TSharedRef<ITextGenerator>, FCreateTextGeneratorDelegate, FStructuredArchive::FRecord );
 	/**
@@ -733,7 +733,7 @@ public:
 	 *
 	 * @param TypeID the name under which to look up the factory function
 	 */
-	static CORE_API FCreateTextGeneratorDelegate FindRegisteredTextGenerator( FName TypeID );
+	[[nodiscard]] static CORE_API FCreateTextGeneratorDelegate FindRegisteredTextGenerator( FName TypeID );
 
 	/**
 	 * Registers a factory function to be used with serialization of text generators within FText.
@@ -791,18 +791,20 @@ public:
 		UnregisterTextGenerator( T::TypeID );
 	}
 
-	CORE_API bool IsTransient() const;
-	CORE_API bool IsCultureInvariant() const;
-	CORE_API bool IsInitializedFromString() const;
-	CORE_API bool IsFromStringTable() const;
+	[[nodiscard]] CORE_API bool IsTransient() const;
+	[[nodiscard]] CORE_API bool IsCultureInvariant() const;
+	[[nodiscard]] CORE_API bool IsInitializedFromString() const;
+	[[nodiscard]] CORE_API bool IsFromStringTable() const;
 
-	CORE_API bool ShouldGatherForLocalization() const;
+	[[nodiscard]] CORE_API bool ShouldGatherForLocalization() const;
+	
+	static CORE_API void Preload( FName InTableId, const FTextKey& InKey );
 
 #if WITH_EDITORONLY_DATA
 	/**
 	 * Constructs a new FText with the SourceString of the specified text but with the specified namespace and key
 	 */
-	static CORE_API FText ChangeKey( const FTextKey& Namespace, const FTextKey& Key, const FText& Text );
+	[[nodiscard]] static CORE_API FText ChangeKey( const FTextKey& Namespace, const FTextKey& Key, const FText& Text );
 #endif
 
 	/////////////////////////////////////////////////////
diff --git a/Engine/Source/Runtime/Core/Public/UObject/NameTypes.h b/Engine/Source/Runtime/Core/Public/UObject/NameTypes.h
--- a/Engine/Source/Runtime/Core/Public/UObject/NameTypes.h
+++ b/Engine/Source/Runtime/Core/Public/UObject/NameTypes.h
@@ -588,16 +588,16 @@ public:
 	using IntrusiveUnsetOptionalStateType = FName;
 
 #if UE_FNAME_OUTLINE_NUMBER
-	CORE_API FNameEntryId GetComparisonIndex() const;
-	CORE_API FNameEntryId GetDisplayIndex() const;
+	[[nodiscard]] CORE_API FNameEntryId GetComparisonIndex() const;
+	[[nodiscard]] CORE_API FNameEntryId GetDisplayIndex() const;
 #else // UE_FNAME_OUTLINE_NUMBER
-	FORCEINLINE FNameEntryId GetComparisonIndex() const
+	[[nodiscard]] FORCEINLINE FNameEntryId GetComparisonIndex() const
 	{
 		checkName(IsWithinBounds(ComparisonIndex));
 		return ComparisonIndex;
 	}
 
-	FORCEINLINE FNameEntryId GetDisplayIndex() const
+	[[nodiscard]] FORCEINLINE FNameEntryId GetDisplayIndex() const
 	{
 		const FNameEntryId Index = GetDisplayIndexFast();
 		checkName(IsWithinBounds(Index));
@@ -606,10 +606,10 @@ public:
 #endif //UE_FNAME_OUTLINE_NUMBER
 
 #if UE_FNAME_OUTLINE_NUMBER
-	CORE_API int32 GetNumber() const;
+	[[nodiscard]] CORE_API int32 GetNumber() const;
 	CORE_API void SetNumber(const int32 NewNumber);
 #else //UE_FNAME_OUTLINE_NUMBER
-	FORCEINLINE int32 GetNumber() const
+	[[nodiscard]] FORCEINLINE int32 GetNumber() const
 	{
 		return Number;
 	}
@@ -621,7 +621,7 @@ public:
 #endif //UE_FNAME_OUTLINE_NUMBER
 	
 	/** Get name without number part as a dynamically allocated string */
-	CORE_API FString GetPlainNameString() const;
+	[[nodiscard]] CORE_API FString GetPlainNameString() const;
 
 	/** Convert name without number part into TCHAR buffer and returns string length. Doesn't allocate. */
 	CORE_API uint32 GetPlainNameString(TCHAR(&OutName)[NAME_SIZE]) const;
@@ -632,15 +632,15 @@ public:
 	/** Copy wide name without number part. Must *only* be used for wide FNames. Doesn't allocate. */
 	CORE_API void GetPlainWIDEString(WIDECHAR(&WideName)[NAME_SIZE]) const;
 
-	CORE_API const FNameEntry* GetComparisonNameEntry() const;
-	CORE_API const FNameEntry* GetDisplayNameEntry() const;
+	[[nodiscard]] CORE_API const FNameEntry* GetComparisonNameEntry() const;
+	[[nodiscard]] CORE_API const FNameEntry* GetDisplayNameEntry() const;
 
 	/**
 	 * Converts an FName to a readable format
 	 *
 	 * @return String representation of the name
 	 */
-	CORE_API FString ToString() const;
+	[[nodiscard]] CORE_API FString ToString() const;
 
 	/**
 	 * Converts an FName to a readable format, in place
@@ -660,7 +660,7 @@ public:
 	/**
 	 * Get the number of characters, excluding null-terminator, that ToString() would yield
 	 */
-	CORE_API uint32 GetStringLength() const;
+	[[nodiscard]] CORE_API uint32 GetStringLength() const;
 
 	/**
 	 * Buffer size required for any null-terminated FName string, i.e. [name] '_' [digits] '\0'
@@ -709,38 +709,38 @@ public:
 	/**
 	 * Check to see if this FName matches the other FName, potentially also checking for any case variations
 	 */
-	FORCEINLINE bool IsEqual(const FName& Other, const ENameCase CompareMethod = ENameCase::IgnoreCase, const bool bCompareNumber = true) const;
+	[[nodiscard]] FORCEINLINE bool IsEqual(const FName& Other, const ENameCase CompareMethod = ENameCase::IgnoreCase, const bool bCompareNumber = true) const;
 
-	FORCEINLINE bool operator==(FName Other) const
+	[[nodiscard]] FORCEINLINE bool operator==(FName Other) const
 	{
 		return ToUnstableInt() == Other.ToUnstableInt();
 	}
 
-	FORCEINLINE bool operator!=(FName Other) const
+	[[nodiscard]] FORCEINLINE bool operator!=(FName Other) const
 	{
 		return !(*this == Other);
 	}
 
 	/** Special comparison operator for TOptional<FName>::IsSet */
-	inline bool operator==(FIntrusiveUnsetOptionalState I) const
+	[[nodiscard]] inline bool operator==(FIntrusiveUnsetOptionalState I) const
 	{
 		return ComparisonIndex == I;
 	}
 
 	/** Fast non-alphabetical order that is only stable during this process' lifetime. */
-	FORCEINLINE bool FastLess(const FName& Other) const
+	[[nodiscard]] FORCEINLINE bool FastLess(const FName& Other) const
 	{
 		return CompareIndexes(Other) < 0;
 	}
 
 	/** Slow alphabetical order that is stable / deterministic over process runs. */
-	FORCEINLINE bool LexicalLess(const FName& Other) const
+	[[nodiscard]] FORCEINLINE bool LexicalLess(const FName& Other) const
 	{
 		return Compare(Other) < 0;
 	}
 
 	/** True for FName(), FName(NAME_None) and FName("None") */
-	FORCEINLINE bool IsNone() const
+	[[nodiscard]] FORCEINLINE bool IsNone() const
 	{
 #if PLATFORM_64BITS && !WITH_CASE_PRESERVING_NAME
 		return ToUnstableInt() == 0;
@@ -755,10 +755,10 @@ public:
 	 * All FNames are valid except for stomped memory, dangling pointers, etc.
 	 * Should only be used to investigate such bugs and not in production code.
 	 */
-	bool IsValid() const { return IsWithinBounds(ComparisonIndex); }
+	[[nodiscard]] bool IsValid() const { return IsWithinBounds(ComparisonIndex); }
 
 	/** Paranoid sanity check, same as IsValid() */
-	bool IsValidIndexFast() const { return IsValid(); }
+	[[nodiscard]] bool IsValidIndexFast() const { return IsValid(); }
 
 
 	/**
@@ -832,7 +832,7 @@ public:
 	* @return the new FName based upon the calling FName, but with any Whitespace characters potentially problematic for
 	 * showing in a log or on screen omitted.
 	 */
-	CORE_API static FString SanitizeWhitespace(const FString& FNameString);
+	[[nodiscard]] CORE_API static FString SanitizeWhitespace(const FString& FNameString);
 	
 	/**
 	 * Compares name to passed in one. Sort is alphabetical ascending.
@@ -840,7 +840,7 @@ public:
 	 * @param	Other	Name to compare this against
 	 * @return	< 0 is this < Other, 0 if this == Other, > 0 if this > Other
 	 */
-	CORE_API int32 Compare( const FName& Other ) const;
+	[[nodiscard]] CORE_API int32 Compare( const FName& Other ) const;
 
 	/**
 	 * Fast non-alphabetical order that is only stable during this process' lifetime.
@@ -848,7 +848,7 @@ public:
 	 * @param	Other	Name to compare this against
 	 * @return	< 0 is this < Other, 0 if this == Other, > 0 if this > Other
 	 */
-	FORCEINLINE int32 CompareIndexes(const FName& Other) const
+	[[nodiscard]] FORCEINLINE int32 CompareIndexes(const FName& Other) const
 	{
 		if (int32 ComparisonDiff = ComparisonIndex.CompareFast(Other.ComparisonIndex))
 		{
@@ -902,15 +902,15 @@ public:
 	}
 
 #if WITH_CASE_PRESERVING_NAME
-	CORE_API static FNameEntryId GetComparisonIdFromDisplayId(FNameEntryId DisplayId);
+	[[nodiscard]] CORE_API static FNameEntryId GetComparisonIdFromDisplayId(FNameEntryId DisplayId);
 #else
-	static FNameEntryId GetComparisonIdFromDisplayId(FNameEntryId DisplayId) { return DisplayId; }
+	[[nodiscard]] static FNameEntryId GetComparisonIdFromDisplayId(FNameEntryId DisplayId) { return DisplayId; }
 #endif
 
 	/**
 	 * Only call this if you *really* know what you're doing
 	 */
-	static FName CreateFromDisplayId(FNameEntryId DisplayId, int32 Number)
+	[[nodiscard]] static FName CreateFromDisplayId(FNameEntryId DisplayId, int32 Number)
 	{
 #if UE_FNAME_OUTLINE_NUMBER
 		checkSlow(ResolveEntry(DisplayId)->IsNumbered() == false); // This id should be unnumbered i.e. returned from GetDisplayIndex on an FName.
@@ -1043,49 +1043,49 @@ public:
 	CORE_API FName(const FNameEntrySerialized& LoadedEntry);
 
 	CORE_API static void DisplayHash( class FOutputDevice& Ar );
-	CORE_API static FString SafeString(FNameEntryId InDisplayIndex, int32 InstanceNumber = NAME_NO_NUMBER_INTERNAL);
+	[[nodiscard]] CORE_API static FString SafeString(FNameEntryId InDisplayIndex, int32 InstanceNumber = NAME_NO_NUMBER_INTERNAL);
 
 	CORE_API static void Reserve(uint32 NumBytes, uint32 NumNames);
 
 	/**
 	 * @return Size of all name entries.
 	 */
-	CORE_API static int32 GetNameEntryMemorySize();
+	[[nodiscard]] CORE_API static int32 GetNameEntryMemorySize();
 
 	/**
 	 * @return Estimated remaining size the name entry table is willing to allocate.
 	 */
-	CORE_API static int32 GetNameEntryMemoryEstimatedAvailable();
+	[[nodiscard]] CORE_API static int32 GetNameEntryMemoryEstimatedAvailable();
 
 	/**
 	* @return Size of Name Table object as a whole
 	*/
-	CORE_API static int32 GetNameTableMemorySize();
+	[[nodiscard]] CORE_API static int32 GetNameTableMemorySize();
 
 	/**
 	 * @return number of ansi names in name table
 	 */
-	CORE_API static int32 GetNumAnsiNames();
+	[[nodiscard]] CORE_API static int32 GetNumAnsiNames();
 
 	/**
 	 * @return number of wide names in name table
 	 */
-	CORE_API static int32 GetNumWideNames();
+	[[nodiscard]] CORE_API static int32 GetNumWideNames();
 
 #if UE_FNAME_OUTLINE_NUMBER
 	/**
 	 * @return number of numbered names in name table
 	 */
-	CORE_API static int32 GetNumNumberedNames();
+	[[nodiscard]] CORE_API static int32 GetNumNumberedNames();
 #endif
 
-	CORE_API static TArray<const FNameEntry*> DebugDump();
+	[[nodiscard]] CORE_API static TArray<const FNameEntry*> DebugDump();
 
-	CORE_API static FNameEntry const* GetEntry(EName Ename);
-	CORE_API static FNameEntry const* GetEntry(FNameEntryId Id);
+	[[nodiscard]] CORE_API static FNameEntry const* GetEntry(EName Ename);
+	[[nodiscard]] CORE_API static FNameEntry const* GetEntry(FNameEntryId Id);
 
 #if UE_TRACE_ENABLED
-	CORE_API static UE::Trace::FEventRef32 TraceName(const FName& Name);
+	[[nodiscard]] CORE_API static UE::Trace::FEventRef32 TraceName(const FName& Name);
 	CORE_API static void TraceNamesOnConnection();
 #endif
 
@@ -1103,7 +1103,7 @@ public:
 	 *
 	 * @return	the sanitized version of the display name
 	 */
-	CORE_API static FString NameToDisplayString( const FString& InDisplayName, const bool bIsBool );
+	[[nodiscard]] CORE_API static FString NameToDisplayString( const FString& InDisplayName, const bool bIsBool );
 
 	/**
 	 * Add/remove an exemption to the formatting applied by NameToDisplayString.
@@ -1114,7 +1114,7 @@ public:
 	CORE_API static void RemoveNameToDisplayStringExemption(const FString& InExemption);
 
 	/** Get the EName that this FName represents or nullptr */
-	CORE_API const EName* ToEName() const;
+	[[nodiscard]] CORE_API const EName* ToEName() const;
 
 	/** 
 		Tear down system and free all allocated memory 
@@ -1125,12 +1125,12 @@ public:
 
 	/** Returns an integer that compares equal in the same way FNames do, only usable within the current process */
 #if UE_FNAME_OUTLINE_NUMBER
-	FORCEINLINE uint64 ToUnstableInt() const
+	[[nodiscard]] FORCEINLINE uint64 ToUnstableInt() const
 	{
 		return ComparisonIndex.ToUnstableInt();
 	}
 #else
-	FORCEINLINE uint64 ToUnstableInt() const
+	[[nodiscard]] FORCEINLINE uint64 ToUnstableInt() const
 	{
 		static_assert(STRUCT_OFFSET(FName, ComparisonIndex) == 0);
 		static_assert(STRUCT_OFFSET(FName, Number) == 4);
@@ -1178,13 +1178,13 @@ private:
 	}
 
 	// Accessor for unmodified comparison index when UE_FNAME_OUTLINE_NUMBER is set
-	FORCEINLINE FNameEntryId GetComparisonIndexInternal() const
+	[[nodiscard]] FORCEINLINE FNameEntryId GetComparisonIndexInternal() const
 	{
 		return ComparisonIndex;
 	}
 
 	// Accessor for unmodified display index when UE_FNAME_OUTLINE_NUMBER is set
-	FORCEINLINE FNameEntryId GetDisplayIndexInternal() const
+	[[nodiscard]] FORCEINLINE FNameEntryId GetDisplayIndexInternal() const
 	{
 #if WITH_CASE_PRESERVING_NAME
 		return DisplayIndex;
@@ -1194,18 +1194,18 @@ private:
 	}
 
 	// Resolve the entry directly referred to by LookupId
-	CORE_API static const FNameEntry* ResolveEntry(FNameEntryId LookupId);
+	[[nodiscard]] CORE_API static const FNameEntry* ResolveEntry(FNameEntryId LookupId);
 	// Recursively resolve through the entry referred to by LookupId to reach the allocated string entry, in the case of UE_FNAME_OUTLINE_NUMBER=1
-	static const FNameEntry* ResolveEntryRecursive(FNameEntryId LookupId);
+	[[nodiscard]] static const FNameEntry* ResolveEntryRecursive(FNameEntryId LookupId);
 
-	CORE_API static bool IsWithinBounds(FNameEntryId Id);
+	[[nodiscard]] CORE_API static bool IsWithinBounds(FNameEntryId Id);
 
 	// These FNameEntryIds are passed in from user code so they must be non-numbered if Number != NAME_NO_NUMBER_INTERNAL
 #if UE_FNAME_OUTLINE_NUMBER
-	CORE_API static FName CreateNumberedName(FNameEntryId ComparisonId, FNameEntryId DisplayId, int32 Number);
+	[[nodiscard]] CORE_API static FName CreateNumberedName(FNameEntryId ComparisonId, FNameEntryId DisplayId, int32 Number);
 #endif
 
-	FORCEINLINE static FName CreateNumberedNameIfNecessary(FNameEntryId ComparisonId, FNameEntryId DisplayId, int32 Number)
+	[[nodiscard]] FORCEINLINE static FName CreateNumberedNameIfNecessary(FNameEntryId ComparisonId, FNameEntryId DisplayId, int32 Number)
 	{
 #if UE_FNAME_OUTLINE_NUMBER
 		if (Number != NAME_NO_NUMBER_INTERNAL)
@@ -1226,19 +1226,19 @@ private:
 		return Out;
 	}
 
-	FORCEINLINE static FName CreateNumberedNameIfNecessary(FNameEntryId ComparisonId, int32 Number)
+	[[nodiscard]] FORCEINLINE static FName CreateNumberedNameIfNecessary(FNameEntryId ComparisonId, int32 Number)
 	{
 		return CreateNumberedNameIfNecessary(ComparisonId, ComparisonId, Number);
 	}
 	
 	
-	static bool Equals(FName A, FName B) = delete;
-	CORE_API static bool Equals(FName A, FAnsiStringView B);
-	CORE_API static bool Equals(FName A, FWideStringView B);
-	CORE_API static bool Equals(FName A, const ANSICHAR* B);
-	CORE_API static bool Equals(FName A, const WIDECHAR* B);
+	[[nodiscard]] static bool Equals(FName A, FName B) = delete;
+	[[nodiscard]] CORE_API static bool Equals(FName A, FAnsiStringView B);
+	[[nodiscard]] CORE_API static bool Equals(FName A, FWideStringView B);
+	[[nodiscard]] CORE_API static bool Equals(FName A, const ANSICHAR* B);
+	[[nodiscard]] CORE_API static bool Equals(FName A, const WIDECHAR* B);
 
-	FORCEINLINE static bool Equals(FName A, EName B)
+	[[nodiscard]] FORCEINLINE static bool Equals(FName A, EName B)
 	{
 		// With UE_FNAME_OUTLINE_NUMBER 1, FName == FName(EName) is
 		// faster than extracting index and number for direct EName comparison.
@@ -1247,62 +1247,62 @@ private:
 	}
 
 #if UE_FNAME_OUTLINE_NUMBER
-	FORCEINLINE static bool Equals(FName A, FMinimalName B)
+	[[nodiscard]] FORCEINLINE static bool Equals(FName A, FMinimalName B)
 	{
 		return A.GetComparisonIndexInternal() == B.Index;
 	}
-	FORCEINLINE static bool Equals(FName A, FScriptName B)
+	[[nodiscard]] FORCEINLINE static bool Equals(FName A, FScriptName B)
 	{
 		return A.GetComparisonIndexInternal() == B.ComparisonIndex;
 	}
-	FORCEINLINE static bool Equals(FName A, FMemoryImageName B)
+	[[nodiscard]] FORCEINLINE static bool Equals(FName A, FMemoryImageName B)
 	{
 		return A.GetComparisonIndexInternal() == B.ComparisonIndex;
 	}
-	friend FORCEINLINE uint32 GetTypeHash(FName Name)
+	[[nodiscard]] friend FORCEINLINE uint32 GetTypeHash(FName Name)
 	{
 		return GetTypeHash(Name.GetComparisonIndexInternal());
 	}
 #else
-	FORCEINLINE static bool Equals(FName A, FMinimalName B)
+	[[nodiscard]] FORCEINLINE static bool Equals(FName A, FMinimalName B)
 	{
 		return A.GetComparisonIndex() == B.Index && A.GetNumber() == B.Number;
 	}
-	FORCEINLINE static bool Equals(FName A, FScriptName B)
+	[[nodiscard]] FORCEINLINE static bool Equals(FName A, FScriptName B)
 	{
 		return A.GetComparisonIndex() == B.ComparisonIndex && A.GetNumber() == B.Number;
 	}
-	FORCEINLINE static bool Equals(FName A, FMemoryImageName B)
+	[[nodiscard]] FORCEINLINE static bool Equals(FName A, FMemoryImageName B)
 	{
 		return A.GetComparisonIndex() == B.ComparisonIndex && A.GetNumber() == B.Number;
 	}
-	friend FORCEINLINE uint32 GetTypeHash(FName Name)
+	[[nodiscard]] friend FORCEINLINE uint32 GetTypeHash(FName Name)
 	{
 		return GetTypeHash(Name.GetComparisonIndex()) + Name.GetNumber();
 	}
 #endif
 
     template <typename T>
-    friend FORCEINLINE auto operator==(FName N, T O) -> decltype(FName::Equals(N, O))
+    [[nodiscard]] friend FORCEINLINE auto operator==(FName N, T O) -> decltype(FName::Equals(N, O))
     {
     	return FName::Equals(N, O);
     }
 
     template <typename T>
-    friend FORCEINLINE auto operator==(T O, FName N) -> decltype(FName::Equals(N, O))
+    [[nodiscard]] friend FORCEINLINE auto operator==(T O, FName N) -> decltype(FName::Equals(N, O))
     {
     	return FName::Equals(N, O);
     }
 
 #if !PLATFORM_COMPILER_HAS_GENERATED_COMPARISON_OPERATORS
     template <typename T>
-    friend FORCEINLINE auto operator!=(FName N, T O) -> decltype(FName::Equals(N, O))
+    [[nodiscard]] friend FORCEINLINE auto operator!=(FName N, T O) -> decltype(FName::Equals(N, O))
     {
     	return !FName::Equals(N, O);
     }
 
     template <typename T>
-    friend FORCEINLINE auto operator!=(T O, FName N) -> decltype(FName::Equals(N, O))
+    [[nodiscard]] friend FORCEINLINE auto operator!=(T O, FName N) -> decltype(FName::Equals(N, O))
     {
     	return !FName::Equals(N, O);
     }
diff --git a/Engine/Source/Runtime/Core/Tests/Internationalization/TextTest.cpp b/Engine/Source/Runtime/Core/Tests/Internationalization/TextTest.cpp
--- a/Engine/Source/Runtime/Core/Tests/Internationalization/TextTest.cpp
+++ b/Engine/Source/Runtime/Core/Tests/Internationalization/TextTest.cpp
@@ -572,7 +572,7 @@ public:
 						I18N.SetCurrentCulture("fr-CA");
 
 						// Need the FText to be rebuilt in fr-CA.
-						FormattedTestLayer2.ToString();
+						(void)FormattedTestLayer2.ToString();
 
 
 						if (AsNumberTest1.CompareTo(FText::AsNumber(5.5421)) != 0)
@@ -653,7 +653,7 @@ public:
 					}
 
 					// Rebuild in LEET so that when we build the source string the DisplayString is still in LEET. 
-					FormattedTestLayer2.ToString();
+					(void)FormattedTestLayer2.ToString();
 
 					{
 						I18N.RestoreCultureState(OriginalCultureState);
diff --git a/Engine/Source/Runtime/Landscape/Private/LandscapeEdgeFixup.cpp b/Engine/Source/Runtime/Landscape/Private/LandscapeEdgeFixup.cpp
--- a/Engine/Source/Runtime/Landscape/Private/LandscapeEdgeFixup.cpp
+++ b/Engine/Source/Runtime/Landscape/Private/LandscapeEdgeFixup.cpp
@@ -60,9 +60,6 @@ namespace UE::Landscape
 
 	const FString& GetDirectionString(EDirectionIndex Index)
 	{
-		FName temp(TEXT("Hello"));
-		temp.ToString();
-
 		static const FString DirectionStrings[static_cast<int32>(EDirectionIndex::Count)] =
 		{
 			FString("B"),			// Bottom
diff --git a/Engine/Source/Runtime/Online/HTTP/Public/Interfaces/IHttpResponse.h b/Engine/Source/Runtime/Online/HTTP/Public/Interfaces/IHttpResponse.h
--- a/Engine/Source/Runtime/Online/HTTP/Public/Interfaces/IHttpResponse.h
+++ b/Engine/Source/Runtime/Online/HTTP/Public/Interfaces/IHttpResponse.h
@@ -157,7 +157,7 @@ namespace EHttpResponseCodes
 		case EHttpResponseCodes::BadGateway: return LOCTEXT("HttpResponseCode_502", "Bad Gateway");
 		case EHttpResponseCodes::ServiceUnavail: return LOCTEXT("HttpResponseCode_503", "Service Unavailable");
 		case EHttpResponseCodes::GatewayTimeout: return LOCTEXT("HttpResponseCode_504", "Gateway Timeout");
-		case EHttpResponseCodes::VersionNotSup: LOCTEXT("HttpResponseCode_505", "HTTP Version Not Supported");
+		case EHttpResponseCodes::VersionNotSup: return LOCTEXT("HttpResponseCode_505", "HTTP Version Not Supported");
 
 		case EHttpResponseCodes::Unknown:
 		default:
