Fix orphan [[nodiscard]]
Rename Preload to PreloadStringTable
#rnx

diff --git a/Engine/Source/Runtime/Core/Private/Internationalization/Text.cpp b/Engine/Source/Runtime/Core/Private/Internationalization/Text.cpp
--- a/Engine/Source/Runtime/Core/Private/Internationalization/Text.cpp
+++ b/Engine/Source/Runtime/Core/Private/Internationalization/Text.cpp
@@ -1241,7 +1241,7 @@ bool FText::ShouldGatherForLocalization() const
 	return !((Flags & ETextFlag::CultureInvariant) || (Flags & ETextFlag::Transient)) && !IsFromStringTable() && !SourceString.IsEmpty() && !IsAllWhitespace(SourceString);
 }
 
-void FText::Preload( FName InTableId, const FTextKey& InKey )
+void FText::PreloadStringTable( FName InTableId, const FTextKey& InKey )
 {
 	FTextHistory_StringTableEntry(InTableId, InKey, EStringTableLoadingPolicy::FindOrLoad);
 }
diff --git a/Engine/Source/Runtime/Core/Public/Internationalization/Text.h b/Engine/Source/Runtime/Core/Public/Internationalization/Text.h
--- a/Engine/Source/Runtime/Core/Public/Internationalization/Text.h
+++ b/Engine/Source/Runtime/Core/Public/Internationalization/Text.h
@@ -449,7 +449,7 @@ public:
 	[[nodiscard]] static CORE_API FText AsDateTime(const FDateTime& DateTime, const FString& CustomPattern, const FString& TimeZone = FString(), const FCulturePtr& TargetCulture = NULL);
 	[[nodiscard]] static CORE_API FText AsTime(const FDateTime& DateTime, const EDateTimeStyle::Type TimeStyle = EDateTimeStyle::Default, const FString& TimeZone = FString(), const FCulturePtr& TargetCulture = NULL);
 	[[nodiscard]] static CORE_API FText AsTimespan(const FTimespan& Timespan, const FCulturePtr& TargetCulture = NULL);
-[[nodiscard]] 
+
 	/**
 	 * Gets the time zone string that represents a non-specific, zero offset, culture invariant time zone.
 	 */
@@ -798,7 +798,7 @@ public:
 
 	[[nodiscard]] CORE_API bool ShouldGatherForLocalization() const;
 	
-	static CORE_API void Preload( FName InTableId, const FTextKey& InKey );
+	static CORE_API void PreloadStringTable( FName InTableId, const FTextKey& InKey );
 
 #if WITH_EDITORONLY_DATA
 	/**
