Fix handling of non-ascii path segments in IterateAncestors and add tests for such.

#tests Added unit tests

diff --git a/Engine/Source/Runtime/Core/Private/Misc/PathViews.cpp b/Engine/Source/Runtime/Core/Private/Misc/PathViews.cpp
--- a/Engine/Source/Runtime/Core/Private/Misc/PathViews.cpp
+++ b/Engine/Source/Runtime/Core/Private/Misc/PathViews.cpp
@@ -250,7 +250,7 @@ void FPathViews::IterateAncestors(FStringView InPath, TFunctionRef<bool(FStringV
 	do
 	{
 		// Strip any trailing path separators
-		InPath = FAsciiSet::TrimSuffixWithout(InPath, NonSeparators);
+		InPath = FAsciiSet::TrimSuffixWith(InPath, Separators);
 
 		// If we've iterated into the root, yield that as the last element
 		// Need to use pointer checks rather than length here because we may have stripped extra separators from the root
diff --git a/Engine/Source/Runtime/Core/Tests/Misc/PathViewsTest.cpp b/Engine/Source/Runtime/Core/Tests/Misc/PathViewsTest.cpp
--- a/Engine/Source/Runtime/Core/Tests/Misc/PathViewsTest.cpp
+++ b/Engine/Source/Runtime/Core/Tests/Misc/PathViewsTest.cpp
@@ -653,6 +653,14 @@ TEST_CASE_NAMED(FPathViewsAncestorTest, "System::Core::Misc::PathViews::IterateA
 	RunAncestorTest(TEXTVIEW("\\\\ShareName/Path/Sub/File.txt"), TEXTVIEW("\\\\ShareName/Path/Sub/File.txt"), TEXTVIEW("\\\\ShareName/Path/Sub"), TEXTVIEW("\\\\ShareName/Path"), TEXTVIEW("\\\\ShareName"));
 	RunAncestorTest(TEXTVIEW("\\\\ShareName/Path/Sub/File.tar.gz"), TEXTVIEW("\\\\ShareName/Path/Sub/File.tar.gz"), TEXTVIEW("\\\\ShareName/Path/Sub"), TEXTVIEW("\\\\ShareName/Path"), TEXTVIEW("\\\\ShareName"));
 
+	// Non-ascii path characters
+	RunAncestorTest(TEXTVIEW("/MountPoint/Sub/\u30D5\u30A1\u30A4\u30EB"), TEXTVIEW("/MountPoint/Sub/\u30D5\u30A1\u30A4\u30EB"), TEXTVIEW("/MountPoint/Sub"), TEXTVIEW("/MountPoint"), TEXTVIEW("/"));
+	RunAncestorTest(TEXTVIEW("/MountPoint/Sub/\u30D5\u30A1\u30A4\u30EB.txt"), TEXTVIEW("/MountPoint/Sub/\u30D5\u30A1\u30A4\u30EB.txt"), TEXTVIEW("/MountPoint/Sub"), TEXTVIEW("/MountPoint"), TEXTVIEW("/"));
+	RunAncestorTest(TEXTVIEW("/MountPoint/Sub/\u30D5\u30A1\u30A4\u30EB.tar.gz"), TEXTVIEW("/MountPoint/Sub/\u30D5\u30A1\u30A4\u30EB.tar.gz"), TEXTVIEW("/MountPoint/Sub"), TEXTVIEW("/MountPoint"), TEXTVIEW("/"));
+	RunAncestorTest(TEXTVIEW("/MountPoint/\u30B5\u30D6/File"), TEXTVIEW("/MountPoint/\u30B5\u30D6/File"), TEXTVIEW("/MountPoint/\u30B5\u30D6"), TEXTVIEW("/MountPoint"), TEXTVIEW("/"));
+	RunAncestorTest(TEXTVIEW("/\u6839/Sub/File"), TEXTVIEW("/\u6839/Sub/File"), TEXTVIEW("/\u6839/Sub"), TEXTVIEW("/\u6839"), TEXTVIEW("/"));
+	RunAncestorTest(TEXTVIEW("/\u6839/\u30B5\u30D6/\u30D5\u30A1\u30A4\u30EB"), TEXTVIEW("/\u6839/\u30B5\u30D6/\u30D5\u30A1\u30A4\u30EB"), TEXTVIEW("/\u6839/\u30B5\u30D6"), TEXTVIEW("/\u6839"), TEXTVIEW("/"));
+
 	// Duplicated path separators
 	RunAncestorTest(TEXTVIEW("//MountPoint"), TEXTVIEW("//MountPoint"));     // Double leadings slashes treated like a windows share
 	RunAncestorTest(TEXTVIEW("//\\MountPoint"), TEXTVIEW("//\\MountPoint")); // SplitVolumeSpecifier treats this as a root
