PR #12559: PlatformFile: Make mapped file region checks thread-safe

diff --git a/Engine/Source/Runtime/Core/Private/Apple/ApplePlatformFile.cpp b/Engine/Source/Runtime/Core/Private/Apple/ApplePlatformFile.cpp
--- a/Engine/Source/Runtime/Core/Private/Apple/ApplePlatformFile.cpp
+++ b/Engine/Source/Runtime/Core/Private/Apple/ApplePlatformFile.cpp
@@ -598,7 +598,7 @@ public:
 	{
 		LLM_PLATFORM_SCOPE(ELLMTag::PlatformMMIO);
 		
-		int32 OldNumOutstandingRegions = NumOutstandingRegions.fetch_sub(1, std::memory_order_acq_rel);
+		int32 OldNumOutstandingRegions = NumOutstandingRegions.fetch_sub(1, std::memory_order_relaxed);
 		check(OldNumOutstandingRegions > 0);	
 		
 		LLM_IF_ENABLED(FLowLevelMemTracker::Get().OnLowLevelFree(ELLMTracker::Platform, (void*)Region->AlignedPtr));
diff --git a/Engine/Source/Runtime/Core/Private/Unix/UnixPlatformFile.cpp b/Engine/Source/Runtime/Core/Private/Unix/UnixPlatformFile.cpp
--- a/Engine/Source/Runtime/Core/Private/Unix/UnixPlatformFile.cpp
+++ b/Engine/Source/Runtime/Core/Private/Unix/UnixPlatformFile.cpp
@@ -806,7 +806,17 @@ public:
 
 	virtual ~FUnixMappedFileHandle() override
 	{
-		check(!NumOutstandingRegions); // can't delete the file before you delete all outstanding regions
+		// can't delete the file before you delete all outstanding regions
+		UE_CLOG(
+				NumOutstandingRegions.load(std::memory_order_relaxed),
+				LogHAL,
+#if UE_BUILD_SHIPPING
+				Error,
+#else
+				Fatal,
+#endif
+				TEXT("Cleaning mapped file with alive mapped regions: %s"),
+				*Filename);
 		close(FileHandle);
 	}
 
@@ -844,15 +854,15 @@ public:
 		// create a mapping for this range
 		const uint8* MapPtr = AlignedMapPtr + Offset - AlignedOffset;
 		FUnixMappedFileRegion* Result = new FUnixMappedFileRegion(MapPtr, AlignedMapPtr, BytesToMap, AlignedSize, Filename, Offset, this);
-		NumOutstandingRegions++;
+		NumOutstandingRegions.fetch_add(1, std::memory_order_relaxed);
 		return Result;
 	}
 
 	void UnMap(const FUnixMappedFileRegion* Region)
 	{
 		LLM_PLATFORM_SCOPE(ELLMTag::PlatformMMIO);
-		check(NumOutstandingRegions > 0);
-		NumOutstandingRegions--;
+		int32 OldNumOutstandingRegions = NumOutstandingRegions.fetch_sub(1, std::memory_order_relaxed);
+		check(OldNumOutstandingRegions > 0);
 
 		LLM_IF_ENABLED(FLowLevelMemTracker::Get().OnLowLevelFree(ELLMTracker::Platform, Region->AlignedPtr));
 		const int Res = munmap(const_cast<uint8*>(Region->AlignedPtr), Region->AlignedSize);
@@ -862,7 +872,7 @@ public:
 private:
 	const uint8* MappedPtr;
 	FString Filename;
-	int32 NumOutstandingRegions;
+	std::atomic<int32> NumOutstandingRegions;
 	int FileHandle;
 };
 
diff --git a/Engine/Source/Runtime/Core/Private/Windows/WindowsPlatformFile.cpp b/Engine/Source/Runtime/Core/Private/Windows/WindowsPlatformFile.cpp
--- a/Engine/Source/Runtime/Core/Private/Windows/WindowsPlatformFile.cpp
+++ b/Engine/Source/Runtime/Core/Private/Windows/WindowsPlatformFile.cpp
@@ -967,7 +967,7 @@ class FMappedFileHandleWindows final : public IMappedFileHandle
 	HANDLE Handle;
 	HANDLE MappingHandle;
 	FString DebugFilename;
-	int32 NumOutstandingRegions;
+	std::atomic<int32> NumOutstandingRegions;
 public:
 	FMappedFileHandleWindows(HANDLE InHandle, HANDLE InMappingHandle, int64 Size, const TCHAR* InDebugFilename)
 		: IMappedFileHandle(Size)
@@ -981,7 +981,18 @@ public:
 	}
 	~FMappedFileHandleWindows()
 	{
-		check(!NumOutstandingRegions); // can't delete the file before you delete all outstanding regions
+		// can't delete the file before you delete all outstanding regions
+		UE_CLOG(
+				NumOutstandingRegions.load(std::memory_order_relaxed),
+				LogHAL,
+#if UE_BUILD_SHIPPING
+				Error,
+#else
+				Fatal,
+#endif
+				TEXT("Cleaning mapped file with alive mapped regions: %s"),
+				*Filename);
+
 		CloseHandle(MappingHandle);
 		TRACE_PLATFORMFILE_BEGIN_CLOSE(Handle);
 		BOOL CloseResult = CloseHandle(Handle);
@@ -1024,14 +1035,14 @@ public:
 		}
 		const uint8* MapPtr = AlignedMapPtr + Offset - AlignedOffset;
 		FMappedFileRegionWindows* Result = new FMappedFileRegionWindows(MapPtr, AlignedMapPtr, BytesToMap, AlignedSize, DebugFilename, Offset, this);
-		NumOutstandingRegions++;
+		NumOutstandingRegions.fetch_add(1, std::memory_order_relaxed);
 		return Result;
 	}
 
 	void UnMap(FMappedFileRegionWindows* Region)
 	{
-		check(NumOutstandingRegions > 0);
-		NumOutstandingRegions--;
+		int32 OldNumOutstandingRegions = NumOutstandingRegions.fetch_sub(1, std::memory_order_relaxed);
+		check(OldNumOutstandingRegions > 0);
 	}
 };
 
