Win: Fixing compile issue with memmapfile and log messages.

diff --git a/Engine/Source/Runtime/Core/Private/Windows/WindowsPlatformFile.cpp b/Engine/Source/Runtime/Core/Private/Windows/WindowsPlatformFile.cpp
--- a/Engine/Source/Runtime/Core/Private/Windows/WindowsPlatformFile.cpp
+++ b/Engine/Source/Runtime/Core/Private/Windows/WindowsPlatformFile.cpp
@@ -982,16 +982,11 @@ public:
 	~FMappedFileHandleWindows()
 	{
 		// can't delete the file before you delete all outstanding regions
-		UE_CLOG(
-				NumOutstandingRegions.load(std::memory_order_relaxed),
-				LogHAL,
 #if UE_BUILD_SHIPPING
-				Error,
+		UE_CLOG(NumOutstandingRegions.load(std::memory_order_relaxed), LogHAL, Error, TEXT("Cleaning mapped file with alive mapped regions: %s"), *DebugFilename);
 #else
-				Fatal,
+		UE_CLOG(NumOutstandingRegions.load(std::memory_order_relaxed), LogHAL, Fatal, TEXT("Cleaning mapped file with alive mapped regions: %s"), *DebugFilename);
 #endif
-				TEXT("Cleaning mapped file with alive mapped regions: %s"),
-				*Filename);
 
 		CloseHandle(MappingHandle);
 		TRACE_PLATFORMFILE_BEGIN_CLOSE(Handle);
