Fix apple memory probe not finding mincore declaration

diff --git a/Engine/Source/Runtime/Core/Private/ProfilingDebugging/Apple/AppleMemoryProbe.cpp b/Engine/Source/Runtime/Core/Private/ProfilingDebugging/Apple/AppleMemoryProbe.cpp
--- a/Engine/Source/Runtime/Core/Private/ProfilingDebugging/Apple/AppleMemoryProbe.cpp
+++ b/Engine/Source/Runtime/Core/Private/ProfilingDebugging/Apple/AppleMemoryProbe.cpp
@@ -10,6 +10,8 @@
 
 #include <Containers/Set.h>
 #include <mach/mach.h>
+#include <sys/types.h>
+#include <sys/mman.h>
 
 extern bool MemoryTrace_IsActive();
 
