Added support for 64-bit radix sorting.

diff --git a/Engine/Source/Runtime/Core/Public/Templates/Sorting.h b/Engine/Source/Runtime/Core/Public/Templates/Sorting.h
--- a/Engine/Source/Runtime/Core/Public/Templates/Sorting.h
+++ b/Engine/Source/Runtime/Core/Public/Templates/Sorting.h
@@ -541,4 +541,169 @@ template< typename CountType >
 void RadixSort32( float* RESTRICT Dst, float* RESTRICT Src, CountType Num )
 {
 	RadixSort32( Dst, Src, Num, FRadixSortKeyFloat() );
+}
+
+enum class ERadixSortBufferState
+{
+	IsInitialized,
+	IsUninitialized
+};
+
+/**
+ * Very fast 64bit radix sort.
+ * SortKeyClass defines operator() that takes ValueType and returns a uint32. Sorting based on key.
+ * No comparisons. Is stable.
+ * The default takes up 40k of stack space. Use a smaller CountType for smaller histograms and less stack space.
+ * Buffer needs to be able to hold at least Num elements.
+ */
+template< ERadixSortBufferState BufferState, typename ValueType, typename CountType, class SortKeyClass >
+void RadixSort64(ValueType* RESTRICT Array, ValueType* RESTRICT Buffer, CountType Num, const SortKeyClass& SortKey)
+{
+	CountType Histograms[(1024 * 2) + (2048 * 4)];
+	CountType* RESTRICT Histogram0 = Histograms + 0;
+	CountType* RESTRICT Histogram1 = Histogram0 + 1024;
+	CountType* RESTRICT Histogram2 = Histogram1 + 1024;
+	CountType* RESTRICT Histogram3 = Histogram2 + 2048;
+	CountType* RESTRICT Histogram4 = Histogram3 + 2048;
+	CountType* RESTRICT Histogram5 = Histogram4 + 2048;
+	FMemory::Memzero(Histograms, sizeof(Histograms));
+
+	{
+		// Parallel histogram generation pass
+		ValueType* RESTRICT s = Array;
+		for (CountType i = 0; i < Num; i++)
+		{
+			uint64 Key = SortKey( *s++ );
+			Histogram0[(Key >> 0) & 1023]++;
+			Histogram1[(Key >> 10) & 1023]++;
+			Histogram2[(Key >> 20) & 2047]++;
+			Histogram3[(Key >> 31) & 2047]++;
+			Histogram4[(Key >> 42) & 2047]++;
+			Histogram5[(Key >> 53) & 2047]++;
+		}
+	}
+	{
+		// Prefix sum
+		// Set each histogram entry to the sum of entries preceding it
+		CountType Sum0 = 0;
+		CountType Sum1 = 0;
+		CountType Sum2 = 0;
+		CountType Sum3 = 0;
+		CountType Sum4 = 0;
+		CountType Sum5 = 0;
+		CountType i = 0;
+		for (; i < 1024; i++)
+		{
+			CountType t0 = Histogram0[i] + Sum0; Histogram0[i] = Sum0; Sum0 = t0;
+			CountType t1 = Histogram1[i] + Sum1; Histogram1[i] = Sum1; Sum1 = t1;
+			CountType t2 = Histogram2[i] + Sum2; Histogram2[i] = Sum2; Sum2 = t2;
+			CountType t3 = Histogram3[i] + Sum3; Histogram3[i] = Sum3; Sum3 = t3;
+			CountType t4 = Histogram4[i] + Sum4; Histogram4[i] = Sum4; Sum4 = t4;
+			CountType t5 = Histogram5[i] + Sum5; Histogram5[i] = Sum5; Sum5 = t5;
+		}
+		for (; i < 2048; i++)
+		{
+			CountType t2 = Histogram2[i] + Sum2; Histogram2[i] = Sum2; Sum2 = t2;
+			CountType t3 = Histogram3[i] + Sum3; Histogram3[i] = Sum3; Sum3 = t3;
+			CountType t4 = Histogram4[i] + Sum4; Histogram4[i] = Sum4; Sum4 = t4;
+			CountType t5 = Histogram5[i] + Sum5; Histogram5[i] = Sum5; Sum5 = t5;
+		}
+	}
+	{
+		// Sort pass 1
+		ValueType* RESTRICT Source = Array;
+		ValueType* RESTRICT Destination = Buffer;
+		for (CountType i = 0; i < Num; i++)
+		{
+			uint64 Value = SortKey(*Source);
+			SIZE_T Index = Histogram0[((Value >> 0) & 1023)]++;
+			if constexpr (BufferState == ERadixSortBufferState::IsInitialized)
+			{
+				Destination[Index] = MoveTemp(*Source++);
+			}
+			else
+			{
+				new(Destination + Index) ValueType(MoveTemp(*Source++));
+			}
+		}
+	}
+	{
+		// Sort pass 2
+		ValueType* RESTRICT Source = Buffer;
+		ValueType* RESTRICT Destination = Array;
+		for (CountType i = 0; i < Num; i++)
+		{
+			uint64 Value = SortKey(*Source);
+			Destination[Histogram1[((Value >> 10) & 1023)]++] = MoveTemp(*Source++);
+		}
+	}
+	{
+		// Sort pass 3
+		ValueType* RESTRICT Source = Array;
+		ValueType* RESTRICT Destination = Buffer;
+		for (CountType i = 0; i < Num; i++)
+		{
+			uint64 Value = SortKey(*Source);
+			Destination[Histogram2[((Value >> 20) & 2047)]++] = MoveTemp(*Source++);
+		}
+	}
+	{
+		// Sort pass 4
+		ValueType* RESTRICT Source = Buffer;
+		ValueType* RESTRICT Destination = Array;
+		for (CountType i = 0; i < Num; i++)
+		{
+			uint64 Value = SortKey(*Source);
+			Destination[Histogram3[((Value >> 31) & 2047)]++] = MoveTemp(*Source++);
+		}
+	}
+	{
+		// Sort pass 5
+		ValueType* RESTRICT Source = Array;
+		ValueType* RESTRICT Destination = Buffer;
+		for (CountType i = 0; i < Num; i++)
+		{
+			uint64 Value = SortKey(*Source);
+			Destination[Histogram4[((Value >> 42) & 2047)]++] = MoveTemp(*Source++);
+		}
+	}
+	{
+		// Sort pass 6
+		ValueType* RESTRICT Source = Buffer;
+		ValueType* RESTRICT Destination = Array;
+		for (CountType i = 0; i < Num; i++)
+		{
+			uint64 Value = SortKey(*Source);
+			Destination[Histogram5[((Value >> 53) & 2047)]++] = MoveTemp(*Source++);
+		}
+	}
+}
+
+template< typename T >
+struct TRadixSortKeyCastUint64
+{
+	FORCEINLINE uint64 operator()(const T& Value) const
+	{
+		return (uint64)Value;
+	}
+};
+
+template< ERadixSortBufferState BufferState, typename ValueType, typename CountType >
+void RadixSort64(ValueType* RESTRICT Array, ValueType* RESTRICT Buffer, CountType Num)
+{
+	RadixSort64<BufferState>(Array, Buffer, Num, TRadixSortKeyCastUint64< ValueType >());
+}
+
+template< typename ValueType, typename CountType, class SortKeyClass >
+void RadixSort64(ValueType* RESTRICT Array, CountType Num, const SortKeyClass& SortKey)
+{
+	TArray<ValueType> Buffer;
+	Buffer.AddUninitialized(Num);
+	RadixSort64<ERadixSortBufferState::IsUninitialized>(Array, Buffer.GetData(), Num, SortKey);
+}
+
+template< typename ValueType, typename CountType >
+void RadixSort64(ValueType* RESTRICT Array, CountType Num)
+{
+	RadixSort64(Array, Num, TRadixSortKeyCastUint64< ValueType >());
 }
\ No newline at end of file
diff --git a/Engine/Source/Runtime/Core/Tests/Templates/SortingTests.cpp b/Engine/Source/Runtime/Core/Tests/Templates/SortingTests.cpp
new file mode 100644
--- /dev/null
+++ b/Engine/Source/Runtime/Core/Tests/Templates/SortingTests.cpp
@@ -0,0 +1,136 @@
+// Copyright Epic Games, Inc. All Rights Reserved.
+
+#if WITH_TESTS
+
+#include "Containers/Array.h"
+#include "Templates/Sorting.h"
+#include "Tests/TestHarnessAdapter.h"
+
+TEST_CASE_NAMED(FTemplatesSortingTest, "System::Core::Templates::Sorting", "[Core][Templates][SmokeFilter]")
+{
+	struct FRadixSort64Test
+	{
+		uint64 Key;
+		FString Value;
+		
+		bool operator==(const FRadixSort64Test&) const = default;
+		bool operator!=(const FRadixSort64Test&) const = default;
+	};
+	struct SortKey
+	{
+		FORCEINLINE uint64 operator()(const FRadixSort64Test& Value) const
+		{
+			return Value.Key;
+		}
+	};
+
+	SECTION("RadixSort64 low bits")
+	{
+		TArray<uint64> Array{ 3, 1, 2 };
+		RadixSort64(Array.GetData(), Array.Num());
+
+		TArray<uint64> Compare{ 1, 2, 3 };
+		CHECK_MESSAGE(TEXT("List incorrectly sorted."), Array == Compare);
+	}
+
+	SECTION("RadixSort64 high bits")
+	{
+		TArray<uint64> Array
+		{
+			0xaabbccdd00112233,
+				0xeeff998800112233,
+				0x9988776600112233
+		};
+		RadixSort64(Array.GetData(), Array.Num());
+
+		TArray<uint64> Compare
+		{
+			0x9988776600112233,
+			0xaabbccdd00112233,
+			0xeeff998800112233
+		};
+		CHECK_MESSAGE(TEXT("List incorrectly sorted."), Array == Compare);
+	}
+
+	SECTION("RadixSort64 all")
+	{
+		TArray<uint64> Array
+		{
+			0xaabbccdd00112233,
+			0xeeff998800887766,
+			0x9988776600443322
+		};
+		RadixSort64(Array.GetData(), Array.Num());
+
+		TArray<uint64> Compare
+		{
+			0x9988776600443322,
+			0xaabbccdd00112233,
+			0xeeff998800887766
+		};
+		CHECK_MESSAGE(TEXT("List incorrectly sorted."), Array == Compare);
+	}
+
+	SECTION("RadixSort64 all with buffer")
+	{
+		TArray<uint64> Array
+		{
+			0xaabbccdd00112233,
+			0xeeff998800887766,
+			0x9988776600443322
+		};
+		TArray<uint64> Buffer;
+		Buffer.AddDefaulted(Array.Num());
+		RadixSort64<ERadixSortBufferState::IsInitialized>(Array.GetData(), Buffer.GetData(), Array.Num());
+
+		TArray<uint64> Compare
+		{
+			0x9988776600443322,
+			0xaabbccdd00112233,
+			0xeeff998800887766
+		};
+		CHECK_MESSAGE(TEXT("List incorrectly sorted."), Array == Compare);
+	}
+
+	SECTION("RadixSort64 custom sort key")
+	{
+		TArray<FRadixSort64Test> Array
+		{
+			FRadixSort64Test{ .Key = 2, .Value = TEXT("B") },
+			FRadixSort64Test{ .Key = 1, .Value = TEXT("A") },
+			FRadixSort64Test{ .Key = 3, .Value = TEXT("C") },
+		};
+		RadixSort64(Array.GetData(), Array.Num(), SortKey());
+
+		TArray<FRadixSort64Test> Compare
+		{
+			FRadixSort64Test{.Key = 1, .Value = TEXT("A") },
+			FRadixSort64Test{.Key = 2, .Value = TEXT("B") },
+			FRadixSort64Test{.Key = 3, .Value = TEXT("C") },
+		};
+		CHECK_MESSAGE(TEXT("List incorrectly sorted."), Array == Compare);
+	}
+
+	SECTION("RadixSort64 custom sort key and provided buffer.")
+	{
+		TArray<FRadixSort64Test> Array
+		{
+			FRadixSort64Test{.Key = 2, .Value = TEXT("B") },
+			FRadixSort64Test{.Key = 1, .Value = TEXT("A") },
+			FRadixSort64Test{.Key = 3, .Value = TEXT("C") },
+		};
+		TArray<FRadixSort64Test> Buffer;
+		Buffer.AddDefaulted(Array.Num());
+		RadixSort64<ERadixSortBufferState::IsInitialized>(Array.GetData(), Buffer.GetData(), Array.Num(), SortKey());
+
+		TArray<FRadixSort64Test> Compare
+		{
+			FRadixSort64Test{.Key = 1, .Value = TEXT("A") },
+			FRadixSort64Test{.Key = 2, .Value = TEXT("B") },
+			FRadixSort64Test{.Key = 3, .Value = TEXT("C") },
+		};
+		CHECK_MESSAGE(TEXT("List incorrectly sorted."), Array == Compare);
+	}
+}
+
+#endif //WITH_TESTS
