[AutoRTFM] In a Shipping build, eliminate unnecessary null checks.

These checks only exist so that the crash dump from a null pointer access is slightly easier to read. In a shipping build, we should prioritize running fast.

diff --git a/Engine/Source/Runtime/Core/Private/AutoRTFM/CompilerInterface.cpp b/Engine/Source/Runtime/Core/Private/AutoRTFM/CompilerInterface.cpp
--- a/Engine/Source/Runtime/Core/Private/AutoRTFM/CompilerInterface.cpp
+++ b/Engine/Source/Runtime/Core/Private/AutoRTFM/CompilerInterface.cpp
@@ -29,80 +29,67 @@ void CheckAlignment(FContext* Context, void* Ptr, size_t AlignmentMask)
     }
 }
 
+#if !UE_BUILD_SHIPPING
+
+// Check for writes to null in development code, so that the inevitable crash will occur
+// in the caller's code rather than in the AutoRTFM runtime.
+#define UE_AUTORTFM_HANDLE_NULL_WRITE(Ptr)  \
+    do                                      \
+    {                                       \
+        if (UNLIKELY(Ptr == nullptr))       \
+        {                                   \
+            return;                         \
+        }                                   \
+    }                                       \
+    while (0)
+
+#else
+
+// In shipping code, we don't want to spend any cycles on a redundant check.
+// We do want the compiler to optimize as if the pointer is non-null, though.
+#define UE_AUTORTFM_HANDLE_NULL_WRITE(Ptr)  \
+    UE_ASSUME(Ptr != nullptr)
+
+#endif
+
 extern "C" UE_AUTORTFM_API AUTORTFM_NO_ASAN void autortfm_record_write(void* Ptr, size_t Size)
 {
-	// check for writes to null here so we end up crashing in the user
-	// code rather than in the autortfm runtime.
-	if (UNLIKELY(nullptr == Ptr))
-	{
-		return;
-	}
-
+	UE_AUTORTFM_HANDLE_NULL_WRITE(Ptr);
 	FContext* Context = FContext::Get();
 	Context->RecordWrite(Ptr, Size);
 }
 
 extern "C" UE_AUTORTFM_API AUTORTFM_NO_ASAN void autortfm_record_write_1(void* Ptr)
 {
-	// check for writes to null here so we end up crashing in the user
-	// code rather than in the autortfm runtime.
-	if (UNLIKELY(nullptr == Ptr))
-	{
-		return;
-	}
-
+	UE_AUTORTFM_HANDLE_NULL_WRITE(Ptr);
 	FContext* Context = FContext::Get();
 	Context->RecordWrite<1>(Ptr);
 }
 
 extern "C" UE_AUTORTFM_API AUTORTFM_NO_ASAN void autortfm_record_write_2(void* Ptr)
 {
-	// check for writes to null here so we end up crashing in the user
-	// code rather than in the autortfm runtime.
-	if (UNLIKELY(nullptr == Ptr))
-	{
-		return;
-	}
-
+	UE_AUTORTFM_HANDLE_NULL_WRITE(Ptr);
 	FContext* Context = FContext::Get();
 	Context->RecordWrite<2>(Ptr);
 }
 
 extern "C" UE_AUTORTFM_API AUTORTFM_NO_ASAN void autortfm_record_write_4(void* Ptr)
 {
-	// check for writes to null here so we end up crashing in the user
-	// code rather than in the autortfm runtime.
-	if (UNLIKELY(nullptr == Ptr))
-	{
-		return;
-	}
-
+	UE_AUTORTFM_HANDLE_NULL_WRITE(Ptr);
 	FContext* Context = FContext::Get();
 	Context->RecordWrite<4>(Ptr);
 }
 
 extern "C" UE_AUTORTFM_API AUTORTFM_NO_ASAN void autortfm_record_write_8(void* Ptr)
 {
-	// check for writes to null here so we end up crashing in the user
-	// code rather than in the autortfm runtime.
-	if (UNLIKELY(nullptr == Ptr))
-	{
-		return;
-	}
-
+	UE_AUTORTFM_HANDLE_NULL_WRITE(Ptr);
 	FContext* Context = FContext::Get();
 	Context->RecordWrite<8>(Ptr);
 }
 
 extern "C" UE_AUTORTFM_API AUTORTFM_NO_ASAN void autortfm_record_masked_write(void* Ptr, uintptr_t Mask, int MaskWidthBits, int ValueSizeBytes)
 {
-	// check for writes to null here so we end up crashing in the user
-	// code rather than in the autortfm runtime.
-	if (UNLIKELY(nullptr == Ptr))
-	{
-		return;
-	}
-
+	UE_AUTORTFM_HANDLE_NULL_WRITE(Ptr);
 	FContext* Context = FContext::Get();
 
 	char* IncrementablePtr = static_cast<char*>(Ptr);
