[AutoRTFM] Add a check to ~FTransactionallySafeRWLockDefinition() that the State is non-null

We're seeing crashes in the OnCommit handler where State is null. Try and catch this bad state earlier.

diff --git a/Engine/Source/Runtime/Core/Public/Misc/TransactionallySafeRWLock.h b/Engine/Source/Runtime/Core/Public/Misc/TransactionallySafeRWLock.h
--- a/Engine/Source/Runtime/Core/Public/Misc/TransactionallySafeRWLock.h
+++ b/Engine/Source/Runtime/Core/Public/Misc/TransactionallySafeRWLock.h
@@ -58,6 +58,8 @@ struct FTransactionallySafeRWLockDefinition
 
 	~FTransactionallySafeRWLockDefinition()
 	{
+		ensure(State.IsValid());
+
 		if (AutoRTFM::IsTransactional())
 		{
 			const AutoRTFM::EContextStatus Status = AutoRTFM::Close([this]
