[AutoRTFM] Tweak memory validator throttling and reenable validation in presubmit tests.

diff --git a/Engine/Source/Runtime/Core/Private/AutoRTFM/Context.h b/Engine/Source/Runtime/Core/Private/AutoRTFM/Context.h
--- a/Engine/Source/Runtime/Core/Private/AutoRTFM/Context.h
+++ b/Engine/Source/Runtime/Core/Private/AutoRTFM/Context.h
@@ -166,7 +166,7 @@ private:
 	TFunction<void()> CascadingRetryTransactionWork;
 	FOpenHashThrottler OpenHashThrottler{
 		/* LogInterval            */ 10,   // Log every 10 seconds.
-		/* AdjustThrottleInterval */ 0.25, // Adjust throttling probabilities every 250ms.
+		/* AdjustThrottleInterval */ 0.5,  // Adjust throttling probabilities every 500ms.
 		/* TargetFractionHashing  */ 0.1   // At most we want to spent 10% of the time hashing.
 	};
 };
diff --git a/Engine/Source/Runtime/Core/Private/AutoRTFM/OpenHashThrottler.cpp b/Engine/Source/Runtime/Core/Private/AutoRTFM/OpenHashThrottler.cpp
--- a/Engine/Source/Runtime/Core/Private/AutoRTFM/OpenHashThrottler.cpp
+++ b/Engine/Source/Runtime/Core/Private/AutoRTFM/OpenHashThrottler.cpp
@@ -8,6 +8,7 @@
 #include "Misc/AssertionMacros.h"
 #include "HAL/PlatformTime.h"
 #include "WriteLog.h"
+#include <limits>
 
 namespace AutoRTFM
 {
@@ -61,7 +62,12 @@ double FOpenHashThrottler::HashProbabilityFor(const void* OpenReturnAddress)
 
 bool FOpenHashThrottler::ShouldHashFor(const void* OpenReturnAddress)
 {
-	if (FMath::FRand() < HashProbabilityFor(OpenReturnAddress))
+	// Multiply the results of two calls to FMath::FRand() to get decent fractional precision.
+	const double RandSqr = static_cast<double>(FMath::FRand()) * static_cast<double>(FMath::FRand());
+	// Negate the probability as we want to bias towards not-hashing, and if
+	// either call to FRand() returns zero then RandSqr will be zero.
+	const double SkipProbability = 1.0 - HashProbabilityFor(OpenReturnAddress);
+	if (RandSqr > SkipProbability * SkipProbability)
 	{
 		LogStats.NumShouldHashForTrue++;
 		return true;
@@ -161,7 +167,20 @@ void FOpenHashThrottler::UpdateThrottlingData()
 		}
 
 		// Finally clamp between 0..1
-		It.Value.Probability = FMath::Clamp(TargetProbability, 0, 1);
+		TargetProbability = FMath::Clamp(TargetProbability, 0, 1);
+		
+		if (TargetProbability < It.Value.Probability)
+		{
+			// Probability is being reduced.
+			// Apply target probability immediately to ensure the application doesn't stall.
+			It.Value.Probability = TargetProbability;
+		}
+		else
+		{
+			// Probability is being increased.
+			// Interpolate to the new target probability at 15% per second.
+			It.Value.Probability = FMath::Lerp(TargetProbability, It.Value.Probability, FMath::Pow(0.85, ThrottlingData.TimeSinceLastUpdate));
+		}
 
 		// Reset the TimeSpentHashing and bActive state for this open.
 		It.Value.TimeSpentHashing = 0;
