Added Projection(), a helper companion/replacement for Invoke.

diff --git a/Engine/Source/Runtime/Core/Public/Algo/Impl/BinaryHeap.h b/Engine/Source/Runtime/Core/Public/Algo/Impl/BinaryHeap.h
--- a/Engine/Source/Runtime/Core/Public/Algo/Impl/BinaryHeap.h
+++ b/Engine/Source/Runtime/Core/Public/Algo/Impl/BinaryHeap.h
@@ -2,8 +2,7 @@
 
 #pragma once
 
-//#include "Templates/Decay.h"
-#include "Templates/Invoke.h"
+#include "Templates/Projection.h"
 #include "Templates/ReversePredicate.h"
 
 #include <type_traits>
@@ -52,12 +51,13 @@ namespace AlgoImpl
 	 * @param	Heap		Pointer to the first element of a binary heap.
 	 * @param	Index		Node index.
 	 * @param	Count		Size of the heap.
-	 * @param	Projection	The projection to apply to the elements.
+	 * @param	InProj		The projection to apply to the elements.
 	 * @param	Predicate	A binary predicate object used to specify if one element should precede another.
 	 */
 	template <typename RangeValueType, typename IndexType, typename ProjectionType, typename PredicateType>
-	FORCEINLINE void HeapSiftDown(RangeValueType* Heap, IndexType Index, const IndexType Count, const ProjectionType& Projection, const PredicateType& Predicate)
+	FORCEINLINE void HeapSiftDown(RangeValueType* Heap, IndexType Index, const IndexType Count, const ProjectionType& InProj, const PredicateType& Predicate)
 	{
+		auto&& Proj = Projection(InProj);
 		while (!HeapIsLeaf(Index, Count))
 		{
 			const IndexType LeftChildIndex = HeapGetLeftChildIndex(Index);
@@ -66,10 +66,10 @@ namespace AlgoImpl
 			IndexType MinChildIndex = LeftChildIndex;
 			if (RightChildIndex < Count)
 			{
-				MinChildIndex = Predicate( Invoke(Projection, Heap[LeftChildIndex]), Invoke(Projection, Heap[RightChildIndex]) ) ? LeftChildIndex : RightChildIndex;
+				MinChildIndex = Predicate( Proj(Heap[LeftChildIndex]), Proj(Heap[RightChildIndex]) ) ? LeftChildIndex : RightChildIndex;
 			}
 
-			if (!Predicate( Invoke(Projection, Heap[MinChildIndex]), Invoke(Projection, Heap[Index]) ))
+			if (!Predicate( Proj(Heap[MinChildIndex]), Proj(Heap[Index]) ))
 			{
 				break;
 			}
@@ -85,18 +85,19 @@ namespace AlgoImpl
 	 * @param	Heap		Pointer to the first element of a binary heap.
 	 * @param	RootIndex	How far to go up?
 	 * @param	NodeIndex	Node index.
-	 * @param	Projection	The projection to apply to the elements.
+	 * @param	InProj		The projection to apply to the elements.
 	 * @param	Predicate	A binary predicate object used to specify if one element should precede another.
 	 *
 	 * @return	The new index of the node that was at NodeIndex
 	 */
 	template <class RangeValueType, typename IndexType, typename ProjectionType, class PredicateType>
-	FORCEINLINE IndexType HeapSiftUp(RangeValueType* Heap, IndexType RootIndex, IndexType NodeIndex, const ProjectionType& Projection, const PredicateType& Predicate)
+	FORCEINLINE IndexType HeapSiftUp(RangeValueType* Heap, IndexType RootIndex, IndexType NodeIndex, const ProjectionType& InProj, const PredicateType& Predicate)
 	{
+		auto&& Proj = Projection(InProj);
 		while (NodeIndex > RootIndex)
 		{
 			IndexType ParentIndex = HeapGetParentIndex(NodeIndex);
-			if (!Predicate( Invoke(Projection, Heap[NodeIndex]), Invoke(Projection, Heap[ParentIndex]) ))
+			if (!Predicate( Proj(Heap[NodeIndex]), Proj(Heap[ParentIndex]) ))
 			{
 				break;
 			}
@@ -114,11 +115,11 @@ namespace AlgoImpl
 	 *
 	 * @param	First		pointer to the first element to heapify
 	 * @param	Num			the number of items to heapify
-	 * @param	Projection	The projection to apply to the elements.
+	 * @param	Proj		The projection to apply to the elements.
 	 * @param	Predicate	A binary predicate object used to specify if one element should precede another.
 	 */
 	template <typename RangeValueType, typename IndexType, typename ProjectionType, typename PredicateType>
-	FORCEINLINE void HeapifyInternal(RangeValueType* First, IndexType Num, ProjectionType Projection, PredicateType Predicate)
+	FORCEINLINE void HeapifyInternal(RangeValueType* First, IndexType Num, ProjectionType Proj, PredicateType Predicate)
 	{
 		if constexpr (std::is_signed_v<IndexType>)
 		{
@@ -133,7 +134,7 @@ namespace AlgoImpl
 		IndexType Index = HeapGetParentIndex(Num - 1);
 		for (;;)
 		{
-			HeapSiftDown(First, Index, Num, Projection, Predicate);
+			HeapSiftDown(First, Index, Num, Proj, Predicate);
 			if (Index == 0)
 			{
 				return;
@@ -148,10 +149,11 @@ namespace AlgoImpl
 	 *
 	 * @param	First		pointer to the first element to sort
 	 * @param	Num			the number of elements to sort
-	 * @param	Predicate	predicate class
+	 * @param	Proj		The projection to apply to the elements.
+	 * @param	Predicate	A binary predicate object used to specify if one element should precede another.
 	 */
 	template <typename RangeValueType, typename IndexType, typename ProjectionType, class PredicateType>
-	void HeapSortInternal(RangeValueType* First, IndexType Num, ProjectionType Projection, PredicateType Predicate)
+	void HeapSortInternal(RangeValueType* First, IndexType Num, ProjectionType Proj, PredicateType Predicate)
 	{
 		if constexpr (std::is_signed_v<IndexType>)
 		{
@@ -164,13 +166,17 @@ namespace AlgoImpl
 		}
 
 		TReversePredicate< PredicateType > ReversePredicateWrapper(Predicate); // Reverse the predicate to build a max-heap instead of a min-heap
-		HeapifyInternal(First, Num, Projection, ReversePredicateWrapper);
+		HeapifyInternal(First, Num, Proj, ReversePredicateWrapper);
 
 		for (IndexType Index = Num - 1; Index > 0; Index--)
 		{
 			Swap(First[0], First[Index]);
 
-			HeapSiftDown(First, (IndexType)0, Index, Projection, ReversePredicateWrapper);
+			HeapSiftDown(First, (IndexType)0, Index, Proj, ReversePredicateWrapper);
 		}
 	}
 }
+
+#if UE_ENABLE_INCLUDE_ORDER_DEPRECATED_IN_5_6
+#include "Templates/Invoke.h"
+#endif
diff --git a/Engine/Source/Runtime/Core/Public/Algo/IntroSort.h b/Engine/Source/Runtime/Core/Public/Algo/IntroSort.h
--- a/Engine/Source/Runtime/Core/Public/Algo/IntroSort.h
+++ b/Engine/Source/Runtime/Core/Public/Algo/IntroSort.h
@@ -5,7 +5,7 @@
 #include "Algo/Impl/BinaryHeap.h"
 #include "Math/UnrealMathUtility.h"
 #include "Templates/IdentityFunctor.h"
-#include "Templates/Invoke.h"
+#include "Templates/Projection.h"
 #include "Templates/Less.h"
 #include "Templates/UnrealTemplate.h" // For GetData, GetNum
 
@@ -19,11 +19,11 @@ namespace AlgoImpl
 	 *
 	 * @param First			pointer to the first element to sort
 	 * @param Num			the number of items to sort
-	 * @param Projection	The projection to sort by when applied to the element.
+	 * @param Proj			The projection to sort by when applied to the element.
 	 * @param Predicate		predicate class
 	 */
 	template <typename T, typename IndexType, typename ProjectionType, typename PredicateType> 
-	void IntroSortInternal(T* First, IndexType Num, ProjectionType Projection, PredicateType Predicate)
+	void IntroSortInternal(T* First, IndexType Num, ProjectionType Proj, PredicateType Predicate)
 	{
 		struct FStack
 		{
@@ -48,7 +48,7 @@ namespace AlgoImpl
 			if ( Current.MaxDepth == 0 )
 			{
 				// We're too deep into quick sort, switch to heap sort
-				HeapSortInternal( Current.Min, Count, Projection, Predicate );
+				HeapSortInternal( Current.Min, Count, Proj, Predicate );
 				continue;
 			}
 
@@ -60,7 +60,7 @@ namespace AlgoImpl
 					T *Max, *Item;
 					for( Max=Current.Min, Item=Current.Min+1; Item<=Current.Max; Item++ )
 					{
-						if( Invoke( Predicate, Invoke( Projection, *Max ), Invoke( Projection, *Item ) ) )
+						if( Predicate( Proj( *Max ), Proj( *Item ) ) )
 						{
 							Max = Item;
 						}
@@ -78,8 +78,8 @@ namespace AlgoImpl
 				Inner.Max = Current.Max+1;
 				for( ; ; )
 				{
-					while( ++Inner.Min<=Current.Max && !Invoke( Predicate, Invoke( Projection, *Current.Min ), Invoke( Projection, *Inner.Min ) ) );
-					while( --Inner.Max> Current.Min && !Invoke( Predicate, Invoke( Projection, *Inner.Max ), Invoke( Projection, *Current.Min ) ) );
+					while( ++Inner.Min<=Current.Max && !Predicate( Proj( *Current.Min ), Proj( *Inner.Min ) ) );
+					while( --Inner.Max> Current.Min && !Predicate( Proj( *Inner.Max ), Proj( *Current.Min ) ) );
 					if( Inner.Min>Inner.Max )
 					{
 						break;
@@ -155,24 +155,28 @@ namespace Algo
 	 * Sort a range of elements by a projection using the projection's operator<. The sort is unstable.
 	 *
 	 * @param Range			The range to sort.
-	 * @param Projection	The projection to sort by when applied to the element.
+	 * @param Proj			The projection to sort by when applied to the element.
 	 */
 	template <typename RangeType, typename ProjectionType>
-	FORCEINLINE void IntroSortBy(RangeType&& Range, ProjectionType Projection)
+	FORCEINLINE void IntroSortBy(RangeType&& Range, ProjectionType Proj)
 	{
-		AlgoImpl::IntroSortInternal(GetData(Range), GetNum(Range), MoveTemp(Projection), TLess<>());
+		AlgoImpl::IntroSortInternal(GetData(Range), GetNum(Range), Projection(MoveTemp(Proj)), TLess<>());
 	}
 
 	/**
 	 * Sort a range of elements by a projection using a user-defined predicate class. The sort is unstable.
 	 *
 	 * @param Range			The range to sort.
-	 * @param Projection	The projection to sort by when applied to the element.
+	 * @param Proj			The projection to sort by when applied to the element.
 	 * @param Predicate		A binary predicate object, applied to the projection, used to specify if one element should precede another.
 	 */
 	template <typename RangeType, typename ProjectionType, typename PredicateType>
-	FORCEINLINE void IntroSortBy(RangeType&& Range, ProjectionType Projection, PredicateType Predicate)
+	FORCEINLINE void IntroSortBy(RangeType&& Range, ProjectionType Proj, PredicateType Predicate)
 	{
-		AlgoImpl::IntroSortInternal(GetData(Range), GetNum(Range), MoveTemp(Projection), MoveTemp(Predicate));
+		AlgoImpl::IntroSortInternal(GetData(Range), GetNum(Range), Projection(MoveTemp(Proj)), Projection(MoveTemp(Predicate)));
 	}
 }
+
+#if UE_ENABLE_INCLUDE_ORDER_DEPRECATED_IN_5_6
+#include "Templates/Invoke.h"
+#endif
diff --git a/Engine/Source/Runtime/Core/Public/Templates/Projection.h b/Engine/Source/Runtime/Core/Public/Templates/Projection.h
new file mode 100644
--- /dev/null
+++ b/Engine/Source/Runtime/Core/Public/Templates/Projection.h
@@ -0,0 +1,132 @@
+// Copyright Epic Games, Inc. All Rights Reserved.
+
+#pragma once
+
+#include <type_traits>
+
+namespace UE::Core::Private
+{
+	template <typename InvocableType>
+	struct TProjectionMemberFunction;
+
+	template <typename ClassType, typename FunctionType>
+	struct TProjectionMemberFunction<FunctionType ClassType::*>
+	{
+		FunctionType ClassType::* MemberFunctionPtr;
+
+		template <typename Arg0Type, typename... ArgTypes>
+		constexpr decltype(auto) operator()(Arg0Type&& Arg0, ArgTypes&&... Args) const
+		{
+			using DecayedArg0Type = std::decay_t<Arg0Type>;
+			if constexpr (std::is_base_of_v<ClassType, DecayedArg0Type>)
+			{
+				return (Forward<Arg0Type>(Arg0).*this->MemberFunctionPtr)(Forward<ArgTypes>(Args)...);
+			}
+			else
+			{
+				return ((*Forward<Arg0Type>(Arg0)).*this->MemberFunctionPtr)(Forward<ArgTypes>(Args)...);
+			}
+		}
+	};
+
+	template <typename InvocableType>
+	struct TProjectionMemberData;
+
+	template <typename ClassType, typename MemberType>
+	struct TProjectionMemberData<MemberType ClassType::*>
+	{
+		MemberType ClassType::* DataMemberPtr;
+
+		template <typename Arg0Type>
+		constexpr decltype(auto) operator()(Arg0Type&& Arg0) const
+		{
+			using DecayedArg0Type = std::decay_t<Arg0Type>;
+			if constexpr (std::is_base_of_v<ClassType, DecayedArg0Type>)
+			{
+				return Forward<Arg0Type>(Arg0).*this->DataMemberPtr;
+			}
+			else
+			{
+				return (*Forward<Arg0Type>(Arg0)).*this->DataMemberPtr;
+			}
+		}
+	};
+
+	template <typename Class, typename MemberType>
+	inline constexpr bool TIsMemberPointerToFunction(MemberType Class::*)
+	{
+		return std::is_function_v<MemberType>;
+	}
+}
+
+/**
+ * Projection() is a related function to Invoke(), in that it can be used to invoke an object with a set of arguments.
+ * However, it works by transforming something invocable into something callable, i.e. can be called with a normal parenthesized set of arguments.
+ *
+ * // These are equivalent:
+ * Projection(I)(Args...)
+ * Invoke(I, Args...).
+ *
+ * ... with the same advantages:
+ *
+ * // Can accept member function pointers
+ * Invoke(&FObjType::Member, Obj);            // equivalent to Obj.Member
+ * Invoke(&FObjType::Func, Obj, Args...);     // equivalent to Obj.Func(Args...)
+ * Projection(&FObjType::Member)(Obj);        // equivalent to Obj.Member
+ * Projection(&FObjType::Func)(Obj, Args...); // equivalent to Obj.Func(Args...)
+ *
+ * // Can operate on pointers, including smart pointers
+ * Invoke(&FObjType::Member, ObjPtr);            // equivalent to ObjPtr->Member
+ * Invoke(&FObjType::Func, ObjPtr, Args...);     // equivalent to ObjPtr->Func(Args...)
+ * Projection(&FObjType::Member)(ObjPtr);        // equivalent to ObjPtr->Member
+ * Projection(&FObjType::Func)(ObjPtr, Args...); // equivalent to ObjPtr->Func(Args...)
+ *
+ * However, Projection() has some additional advantages:
+ *
+ * - Projection(I) returns I unchanged if it is already callable, meaning no redundant stepping in and out of many Invoke() calls in the debugger.
+ * - Projection(...) is variadic and can transform a sequence of invocables into a callable that invokes them one by one:
+ *
+ * Projection(A, B, C)(Args...)
+ * Invoke(C, Invoke(B, Invoke(A, Args...)))
+ *
+ * This allows users to pass a sequence of projections to an algorithm that takes a single projection:
+ *
+ * struct FInner
+ * {
+ *     FString Name;
+ * };
+ * struct FOuter
+ * {
+ *     FInner Inner;
+ * };
+ *
+ * // Sort array of outers by the names of inner, 
+ * Algo::SortBy(ArrayOfOuters, Projection(&FOuter::Inner, &FInner::Name));
+ */
+template <typename Invocable0Type, typename... InvocableTypes>
+[[nodiscard]] constexpr auto Projection(Invocable0Type&& Invocable0, InvocableTypes&&... Invocables)
+{
+	if constexpr (sizeof...(InvocableTypes) == 0)
+	{
+		using DecayedInvocable0Type = std::decay_t<Invocable0Type>;
+		if constexpr (!std::is_member_pointer_v<Invocable0Type>)
+		{
+			return Forward<Invocable0Type&&>(Invocable0);
+		}
+		else if constexpr (UE::Core::Private::TIsMemberPointerToFunction(DecayedInvocable0Type{}))
+		{
+			return UE::Core::Private::TProjectionMemberFunction<DecayedInvocable0Type>{ Forward<Invocable0Type>(Invocable0) };
+		}
+		else
+		{
+			return UE::Core::Private::TProjectionMemberData<DecayedInvocable0Type>{ Forward<Invocable0Type>(Invocable0) };
+		}
+	}
+	else
+	{
+		return [Callable0 = Projection(Forward<Invocable0Type>(Invocable0)), CallableRest = Projection(Forward<InvocableTypes>(Invocables)...)](auto&&... Args) -> decltype(auto)
+		{
+			return CallableRest(Callable0(Forward<decltype(Args)>(Args)...));
+		};
+	}
+}
diff --git a/Engine/Source/Runtime/MovieScene/Private/MovieSceneCommonHelpers.cpp b/Engine/Source/Runtime/MovieScene/Private/MovieSceneCommonHelpers.cpp
--- a/Engine/Source/Runtime/MovieScene/Private/MovieSceneCommonHelpers.cpp
+++ b/Engine/Source/Runtime/MovieScene/Private/MovieSceneCommonHelpers.cpp
@@ -84,9 +84,9 @@ UMovieSceneSection* MovieSceneHelpers::FindNearestSectionAtTime( TArrayView<UMov
 
 	if (NonOverlappingSections.Num())
 	{
-		Algo::SortBy(NonOverlappingSections, [](const UMovieSceneSection* A) { return A->GetRange().GetUpperBound(); }, SortUpperBounds);
+		Algo::SortBy(NonOverlappingSections, Projection(&UMovieSceneSection::GetRange, &TRange<FFrameNumber>::GetUpperBound), SortUpperBounds);
 
-		const int32 PreviousIndex = Algo::UpperBoundBy(NonOverlappingSections, TRangeBound<FFrameNumber>(Time), [](const UMovieSceneSection* A){ return A ? A->GetRange().GetUpperBound() : FFrameNumber(0); }, SortUpperBounds)-1;
+		const int32 PreviousIndex = Algo::UpperBoundBy(NonOverlappingSections, TRangeBound<FFrameNumber>(Time), Projection(&UMovieSceneSection::GetRange, &TRange<FFrameNumber>::GetUpperBound), SortUpperBounds)-1;
 		if (NonOverlappingSections.IsValidIndex(PreviousIndex))
 		{
 			return NonOverlappingSections[PreviousIndex];
