[EventPool]
- Improve detection of use-after-free of FRecyclableEvent to detect events that sometimes got triggered after they were returned to the pool and caused havok with other users of the pool expecting events to be in an unnotified state at creation time

diff --git a/Engine/Source/Runtime/Core/Private/Misc/EventPool.h b/Engine/Source/Runtime/Core/Private/Misc/EventPool.h
--- a/Engine/Source/Runtime/Core/Private/Misc/EventPool.h
+++ b/Engine/Source/Runtime/Core/Private/Misc/EventPool.h
@@ -101,8 +101,12 @@ public:
 		check(Event->IsManualReset() == (PoolType == EEventMode::ManualReset));
 
 		FSafeRecyclableEvent* SafeEvent = (FSafeRecyclableEvent*)Event;
-		ReturnRawEvent(SafeEvent->InnerEvent);
+		FEvent* InnerEvent = SafeEvent->InnerEvent;
+		// Make sure the safeevent can't be used anymore before returning the inner event to the pool
+		// This will help trap use-after-free of events that can end up triggering events while they're in the pool
+		// and causing issues for other pool users.
 		delete SafeEvent;
+		ReturnRawEvent(InnerEvent);
 	}
 
 	void EmptyPool()
