Allow setting a console variable in a transaction that required temporarily change the locale.

diff --git a/Engine/Source/Programs/AutoRTFMTests/Private/UECoreTests.cpp b/Engine/Source/Programs/AutoRTFMTests/Private/UECoreTests.cpp
--- a/Engine/Source/Programs/AutoRTFMTests/Private/UECoreTests.cpp
+++ b/Engine/Source/Programs/AutoRTFMTests/Private/UECoreTests.cpp
@@ -6,6 +6,7 @@
 #include "AssetRegistry/AssetDataTagMap.h"
 #include "AutoRTFM/AutoRTFMTesting.h"
 #include "Delegates/IDelegateInstance.h"
+#include "HAL/IConsoleManager.h"
 #include "HAL/MallocLeakDetection.h"
 #include "HAL/ThreadSingleton.h"
 #include "Internationalization/TextCache.h"
@@ -3098,3 +3099,17 @@ TEST_CASE("UECore.PackageName")
 		REQUIRE(FPackageName::EPackageLocationFilter::None == FPackageName::DoesPackageExistEx(Path, FPackageName::EPackageLocationFilter::IoDispatcher));
 	});
 }
+
+TEST_CASE("UECore.ConsoleManager")
+{
+	IConsoleManager& Manager = IConsoleManager::Get();
+	float Thing = 42.0f;
+	IConsoleVariable* const Variable = Manager.RegisterConsoleVariableRef(TEXT("WOWWEE"), Thing, TEXT("Halp!"));
+
+	AutoRTFM::Testing::Commit([&]
+	{
+		Variable->Set(13.0f);
+	});
+
+	REQUIRE(Thing == 13.0f);
+}
\ No newline at end of file
diff --git a/Engine/Source/Runtime/Core/Private/HAL/ConsoleManager.cpp b/Engine/Source/Runtime/Core/Private/HAL/ConsoleManager.cpp
--- a/Engine/Source/Runtime/Core/Private/HAL/ConsoleManager.cpp
+++ b/Engine/Source/Runtime/Core/Private/HAL/ConsoleManager.cpp
@@ -5,6 +5,7 @@ ConsoleManager.cpp: console command handling
 =============================================================================*/
 
 #include "HAL/ConsoleManager.h"
+#include "AutoRTFM/AutoRTFM.h"
 #include "Containers/Array.h"
 #include "Containers/Map.h"
 #include "Containers/UnrealString.h"
@@ -221,17 +222,46 @@ namespace UE::ConsoleManager::Private
 		// Set a custom locale within this object scope
 		FConsoleManagerLocaleScope()
 		{
-			if (const char* saved = std::setlocale(LC_NUMERIC, nullptr))
+			if (AutoRTFM::IsClosed())
 			{
-				SavedLocale = TArray<char>(saved, TCString<char>::Strlen(saved) + 1);
+				// We only support locale setting if the scope is on the stack. This
+				// is just a defensive block against any new usecase potentially doing
+				// a heap allocation wrongly!
+				ensure(AutoRTFM::IsOnCurrentTransactionStack(this));
+
+				// If we abort while `this` is currently alive, we need to run the
+				// destructor so that any locale we had temporarily disturbed would be
+				// restored correctly.
+				AutoRTFM::PushOnAbortHandler(this, [this]
+				{
+					this->~FConsoleManagerLocaleScope();
+				});
 			}
-			std::setlocale(LC_NUMERIC, "C");
+
+			AutoRTFM::Open([&]
+			{
+				if (const char* saved = std::setlocale(LC_NUMERIC, nullptr))
+				{
+					SavedLocale = TArray<char>(saved, TCString<char>::Strlen(saved) + 1);
+				}
+				std::setlocale(LC_NUMERIC, "C");
+			});
 		}
 
 		// restore the captured Locale
 		~FConsoleManagerLocaleScope()
 		{
-			std::setlocale(LC_NUMERIC, SavedLocale.GetData());
+			AutoRTFM::Open([&]
+			{
+				std::setlocale(LC_NUMERIC, SavedLocale.GetData());
+			});
+
+			if (AutoRTFM::IsClosed())
+			{
+				// We just restored the locale to the previous value above, so we pop here
+				// to ensure we don't try and do a double restore incorrectly.
+				AutoRTFM::PopOnAbortHandler(this);
+			}
 		}
 
 	private:
