[TaskGraph]
- Fix use-after-free in UE::Tasks::Wait along with a rare but severe issue where the function would exit before all tasks were finished.

diff --git a/Engine/Source/Runtime/Core/Public/Tasks/Task.h b/Engine/Source/Runtime/Core/Public/Tasks/Task.h
--- a/Engine/Source/Runtime/Core/Public/Tasks/Task.h
+++ b/Engine/Source/Runtime/Core/Public/Tasks/Task.h
@@ -363,37 +363,6 @@ namespace UE::Tasks
 		Task.Wait();
 	}
 
-	// wait for multiple tasks, with timeout
-	// @param TaskCollectionType - an iterable collection of `TTask<T>`, e.g. `TArray<FTask>`
-	template<typename TaskCollectionType>
-	bool Wait(const TaskCollectionType& Tasks, FTimespan InTimeout/* = FTimespan::MaxValue()*/)
-	{
-		TaskTrace::FWaitingScope WaitingScope(Private::GetTraceIds(Tasks));
-		TRACE_CPUPROFILER_EVENT_SCOPE(Tasks::Wait);
-
-		FTimeout Timeout{ InTimeout };
-
-		// ignore the result as we still have to make sure the task is completed upon returning from this function call
-		Private::TryRetractAndExecute(Tasks, Timeout);
-
-		// the event must be alive for the task and this function lifetime, we don't know which one will be finished first as waiting can time out
-		// before the waiting task is completed
-		FSharedEventRef CompletionEvent;
-		auto WaitingTaskBody = [CompletionEvent] { CompletionEvent->Trigger(); };
-		using FWaitingTask = Private::TExecutableTask<decltype(WaitingTaskBody)>;
-
-		TRefCountPtr<FWaitingTask> WaitingTask{ FWaitingTask::Create(TEXT("Waiting Task"), MoveTemp(WaitingTaskBody), ETaskPriority::Default /* doesn't matter */, EExtendedTaskPriority::Inline, ETaskFlags::None), /*bAddRef=*/ false};
-		WaitingTask->AddPrerequisites(Tasks);
-
-		if (WaitingTask->TryLaunch(sizeof(WaitingTask)))
-		{	// was executed inline
-			check(WaitingTask->IsCompleted());
-			return true;
-		}
-
-		return CompletionEvent->Wait(Timeout.GetRemainingRoundedUpMilliseconds());
-	}
-
 	// wait for multiple tasks while executing other tasks
 	template<typename TaskCollectionType>
 	UE_DEPRECATED(5.5, "Use Wait instead.")
@@ -465,6 +434,23 @@ namespace UE::Tasks
 		return Tasks;
 	}
 
+	// wait for multiple tasks, with timeout
+	// @param TaskCollectionType - an iterable collection of `TTask<T>`, e.g. `TArray<FTask>`
+	template<typename TaskCollectionType>
+	bool Wait(const TaskCollectionType& Tasks, FTimespan InTimeout/* = FTimespan::MaxValue()*/)
+	{
+		TaskTrace::FWaitingScope WaitingScope(Private::GetTraceIds(Tasks));
+		TRACE_CPUPROFILER_EVENT_SCOPE(Tasks::Wait);
+
+		return UE::Tasks::Launch(
+				TEXT("Waiting Task"),
+				[](){},
+				UE::Tasks::Prerequisites(Tasks),
+				ETaskPriority::Default /* doesn't matter */,
+				EExtendedTaskPriority::Inline
+			).Wait(InTimeout);
+	}
+
 	/////////////////////////////////////////////////////////////
 	// "any task" support. these functions allocate excessively (per input task plus more).
 	// can be reduced to a single alloc if this is a perf issue
