Add static NumComponents variable and IsValidIndex function to TVector2, TVector3, and TVector4. Clean up operator[] and Component functions on all vector types

#tests Editor - Ran all System::Core::Math tests locally
[RN] minor Core

diff --git a/Engine/Source/Runtime/Core/Public/Math/Vector.h b/Engine/Source/Runtime/Core/Public/Math/Vector.h
--- a/Engine/Source/Runtime/Core/Public/Math/Vector.h
+++ b/Engine/Source/Runtime/Core/Public/Math/Vector.h
@@ -72,6 +72,9 @@ public:
 		T XYZ[3];
 	};
 
+	/** The number of components this vector type has. */
+	CORE_API static constexpr int32 NumComponents = 3;
+
 	/** A zero vector (0,0,0) */
 	CORE_API static const TVector<T> ZeroVector;
 	
@@ -480,6 +483,13 @@ public:
     */
     T Component(int32 Index) const;
 
+	/**
+	 * Tests if index is valid, i.e. greater than or equal to zero, and less than the number of components in the vector.
+	 *
+	 * @param Index Index to test.
+	 * @returns True if index is valid. False otherwise.
+	 */
+	bool IsValidIndex(int32 Index) const;
 
     /** Get a specific component of the vector, given a specific axis by enum */
     T GetComponentForAxis(EAxis::Type Axis) const;
@@ -1623,19 +1633,13 @@ FORCEINLINE TVector<T> TVector<T>::operator/=(const TVector<T>& V)
 template<typename T>
 FORCEINLINE T& TVector<T>::operator[](int32 Index)
 {
-    checkSlow(Index >= 0 && Index < 3);
-PRAGMA_DISABLE_DEPRECATION_WARNINGS
-    return XYZ[Index];
-PRAGMA_ENABLE_DEPRECATION_WARNINGS
+	return Component(Index);
 }
 
 template<typename T>
 FORCEINLINE T TVector<T>::operator[](int32 Index)const
 {
-    checkSlow(Index >= 0 && Index < 3);
-PRAGMA_DISABLE_DEPRECATION_WARNINGS
-    return XYZ[Index];
-PRAGMA_ENABLE_DEPRECATION_WARNINGS
+	return Component(Index);
 }
 
 template<typename T>
@@ -1928,6 +1932,7 @@ FORCEINLINE void TVector<T>::AddBounded(const TVector<T>& V, T Radius)
 template<typename T>
 FORCEINLINE T& TVector<T>::Component(int32 Index)
 {
+	check(IsValidIndex(Index));
 PRAGMA_DISABLE_DEPRECATION_WARNINGS
     return XYZ[Index];
 PRAGMA_ENABLE_DEPRECATION_WARNINGS
@@ -1936,11 +1941,18 @@ PRAGMA_ENABLE_DEPRECATION_WARNINGS
 template<typename T>
 FORCEINLINE T TVector<T>::Component(int32 Index) const
 {
+	check(IsValidIndex(Index));
 PRAGMA_DISABLE_DEPRECATION_WARNINGS
     return XYZ[Index];
 PRAGMA_ENABLE_DEPRECATION_WARNINGS
 }
 
+template<typename T>
+bool TVector<T>::IsValidIndex(int32 Index) const
+{
+	return (Index >= 0) && (Index < NumComponents);
+}
+
 template<typename T>
 FORCEINLINE T TVector<T>::GetComponentForAxis(EAxis::Type Axis) const
 {
diff --git a/Engine/Source/Runtime/Core/Public/Math/Vector2D.h b/Engine/Source/Runtime/Core/Public/Math/Vector2D.h
--- a/Engine/Source/Runtime/Core/Public/Math/Vector2D.h
+++ b/Engine/Source/Runtime/Core/Public/Math/Vector2D.h
@@ -55,6 +55,9 @@ public:
 		T XY[2];
 	};
 
+	/** The number of components this vector type has. */
+	CORE_API static constexpr int32 NumComponents = 2;
+
 	/** Global 2D zero vector constant (0,0) */
 	CORE_API static const TVector2<T> ZeroVector;
 
@@ -406,6 +409,14 @@ public:
 	*/
 	T Component(int32 Index) const;
 
+	/**
+	 * Tests if index is valid, i.e. greater than or equal to zero, and less than the number of components in the vector.
+	 *
+	 * @param Index Index to test.
+	 * @returns True if index is valid. False otherwise.
+	 */
+	bool IsValidIndex(int32 Index) const;
+
 public:
 
 	/**
@@ -1063,15 +1074,13 @@ FORCEINLINE TVector2<T> TVector2<T>::operator/=(const TVector2<T>& V)
 template<typename T>
 FORCEINLINE T& TVector2<T>::operator[](int32 Index)
 {
-	check(Index>=0 && Index<2);
-	return ((Index == 0) ? X : Y);
+	return Component(Index);
 }
 
 template<typename T>
 FORCEINLINE T TVector2<T>::operator[](int32 Index) const
 {
-	check(Index>=0 && Index<2);
-	return ((Index == 0) ? X : Y);
+	return Component(Index);
 }
 
 template<typename T>
@@ -1218,6 +1227,7 @@ FORCEINLINE bool TVector2<T>::IsZero() const
 template<typename T>
 FORCEINLINE T& TVector2<T>::Component(int32 Index)
 {
+	check(IsValidIndex(Index));
 PRAGMA_DISABLE_DEPRECATION_WARNINGS
 	return XY[Index];
 PRAGMA_ENABLE_DEPRECATION_WARNINGS
@@ -1226,11 +1236,18 @@ PRAGMA_ENABLE_DEPRECATION_WARNINGS
 template<typename T>
 FORCEINLINE T TVector2<T>::Component(int32 Index) const
 {
+	check(IsValidIndex(Index));
 PRAGMA_DISABLE_DEPRECATION_WARNINGS
 	return XY[Index];
 PRAGMA_ENABLE_DEPRECATION_WARNINGS
 }
 
+template<typename T>
+bool TVector2<T>::IsValidIndex(int32 Index) const
+{
+	return (Index >= 0) && (Index < NumComponents);
+}
+
 template<typename T>
 FORCEINLINE FIntPoint TVector2<T>::IntPoint() const
 {
diff --git a/Engine/Source/Runtime/Core/Public/Math/Vector4.h b/Engine/Source/Runtime/Core/Public/Math/Vector4.h
--- a/Engine/Source/Runtime/Core/Public/Math/Vector4.h
+++ b/Engine/Source/Runtime/Core/Public/Math/Vector4.h
@@ -55,6 +55,9 @@ public:
 		T XYZW[4];
 	};
 
+	/** The number of components this vector type has. */
+	CORE_API static constexpr int32 NumComponents = 4;
+
 public:
 
 	/**
@@ -394,6 +397,14 @@ public:
 	*/
 	const T& Component(int32 Index) const;
 
+	/**
+	 * Tests if index is valid, i.e. greater than or equal to zero, and less than the number of components in the vector.
+	 *
+	 * @param Index Index to test.
+	 * @returns True if index is valid. False otherwise.
+	 */
+	bool IsValidIndex(int32 Index) const;
+
 	/**
 	 * Error tolerant comparison.
 	 *
@@ -667,18 +678,14 @@ FORCEINLINE TVector4<T>::TVector4(const TIntVector4<IntType>& InVector)
 template<typename T>
 FORCEINLINE T& TVector4<T>::operator[](int32 ComponentIndex)
 {
-PRAGMA_DISABLE_DEPRECATION_WARNINGS
-	return XYZW[ ComponentIndex ];
-PRAGMA_ENABLE_DEPRECATION_WARNINGS
+	return Component(ComponentIndex);
 }
 
 
 template<typename T>
 FORCEINLINE T TVector4<T>::operator[](int32 ComponentIndex) const
 {
-PRAGMA_DISABLE_DEPRECATION_WARNINGS
-	return XYZW[ ComponentIndex ];
-PRAGMA_ENABLE_DEPRECATION_WARNINGS
+	return Component(ComponentIndex);
 }
 
 
@@ -755,6 +762,7 @@ FORCEINLINE TVector4<T> TVector4<T>::operator^(const TVector4<T>& V) const
 template<typename T>
 FORCEINLINE T& TVector4<T>::Component(int32 Index)
 {
+	check(IsValidIndex(Index));
 PRAGMA_DISABLE_DEPRECATION_WARNINGS
 	return XYZW[Index];
 PRAGMA_ENABLE_DEPRECATION_WARNINGS
@@ -763,11 +771,18 @@ PRAGMA_ENABLE_DEPRECATION_WARNINGS
 template<typename T>
 FORCEINLINE const T& TVector4<T>::Component(int32 Index) const
 {
+	check(IsValidIndex(Index));
 PRAGMA_DISABLE_DEPRECATION_WARNINGS
 	return XYZW[Index];
 PRAGMA_ENABLE_DEPRECATION_WARNINGS
 }
 
+template<typename T>
+bool TVector4<T>::IsValidIndex(int32 Index) const
+{
+	return (Index >= 0) && (Index < NumComponents);
+}
+
 template<typename T>
 FORCEINLINE bool TVector4<T>::operator==(const TVector4<T>& V) const
 {
