Fixed Projection's use of Forward<> when UnrealTemplate.h isn't included.

diff --git a/Engine/Source/Runtime/Core/Public/Templates/Projection.h b/Engine/Source/Runtime/Core/Public/Templates/Projection.h
--- a/Engine/Source/Runtime/Core/Public/Templates/Projection.h
+++ b/Engine/Source/Runtime/Core/Public/Templates/Projection.h
@@ -20,11 +20,11 @@ namespace UE::Core::Private
 			using DecayedArg0Type = std::decay_t<Arg0Type>;
 			if constexpr (std::is_base_of_v<ClassType, DecayedArg0Type>)
 			{
-				return (Forward<Arg0Type>(Arg0).*this->MemberFunctionPtr)(Forward<ArgTypes>(Args)...);
+				return (((Arg0Type&&)Arg0).*this->MemberFunctionPtr)((ArgTypes&&)Args...);
 			}
 			else
 			{
-				return ((*Forward<Arg0Type>(Arg0)).*this->MemberFunctionPtr)(Forward<ArgTypes>(Args)...);
+				return ((*(Arg0Type&&)Arg0).*this->MemberFunctionPtr)((ArgTypes&&)Args...);
 			}
 		}
 	};
@@ -43,11 +43,11 @@ namespace UE::Core::Private
 			using DecayedArg0Type = std::decay_t<Arg0Type>;
 			if constexpr (std::is_base_of_v<ClassType, DecayedArg0Type>)
 			{
-				return Forward<Arg0Type>(Arg0).*this->DataMemberPtr;
+				return ((Arg0Type&&)Arg0).*this->DataMemberPtr;
 			}
 			else
 			{
-				return (*Forward<Arg0Type>(Arg0)).*this->DataMemberPtr;
+				return (*(Arg0Type&&)Arg0).*this->DataMemberPtr;
 			}
 		}
 	};
@@ -111,22 +111,22 @@ template <typename Invocable0Type, typename... InvocableTypes>
 		using DecayedInvocable0Type = std::decay_t<Invocable0Type>;
 		if constexpr (!std::is_member_pointer_v<Invocable0Type>)
 		{
-			return Forward<Invocable0Type&&>(Invocable0);
+			return (Invocable0Type&&)Invocable0;
 		}
 		else if constexpr (UE::Core::Private::TIsMemberPointerToFunction(DecayedInvocable0Type{}))
 		{
-			return UE::Core::Private::TProjectionMemberFunction<DecayedInvocable0Type>{ Forward<Invocable0Type>(Invocable0) };
+			return UE::Core::Private::TProjectionMemberFunction<DecayedInvocable0Type>{ (Invocable0Type&&)Invocable0 };
 		}
 		else
 		{
-			return UE::Core::Private::TProjectionMemberData<DecayedInvocable0Type>{ Forward<Invocable0Type>(Invocable0) };
+			return UE::Core::Private::TProjectionMemberData<DecayedInvocable0Type>{ (Invocable0Type&&)Invocable0 };
 		}
 	}
 	else
 	{
-		return [Callable0 = Projection(Forward<Invocable0Type>(Invocable0)), CallableRest = Projection(Forward<InvocableTypes>(Invocables)...)](auto&&... Args) -> decltype(auto)
+		return [Callable0 = Projection((Invocable0Type&&)Invocable0), CallableRest = Projection((InvocableTypes&&)Invocables...)](auto&&... Args) -> decltype(auto)
 		{
-			return CallableRest(Callable0(Forward<decltype(Args)>(Args)...));
+			return CallableRest(Callable0((decltype(Args)&&)Args...));
 		};
 	}
 }
