[FormatStringSan] improve UTF8<->TCHAR* mismatch error message

diff --git a/Engine/Source/Runtime/Core/Public/String/FormatStringSanErrors.h b/Engine/Source/Runtime/Core/Public/String/FormatStringSanErrors.h
--- a/Engine/Source/Runtime/Core/Public/String/FormatStringSanErrors.h
+++ b/Engine/Source/Runtime/Core/Public/String/FormatStringSanErrors.h
@@ -43,8 +43,8 @@ UE_FORMAT_STRING_ERROR_BEGIN
   CASE(StatusI64NeedsIntegerArg, "'%I64[ ]' expects integral arg (eg. `char`, `int`, `long`, etc.)");
   CASE(StatusDynamicLengthSpecNeedsIntegerArg,
        "dynamic field width specifier '*' expects integral arg (eg. `char`, `int`, `long`, etc.)");
-  CASE(StatusUTF8NeedsStringCastAndS, "Pass UTF8 strings to StringCast<TCHAR>(...) and use %s.");
-  CASE(StatusUTF8NeedsStringCast, "'%s' expects `TCHAR*`; pass UTF8 strings to StringCast<TCHAR>(...) first.");
+  CASE(StatusUTF8NeedsStringCastAndS, "Pass UTF8 strings to `StringCast<TCHAR>(MyUTF8String).Get()` and use %s.");
+  CASE(StatusUTF8NeedsStringCast, "'%s' expects `TCHAR*`; pass UTF8 strings to `StringCast<TCHAR>(MyUTF8String).Get()` first.");
 UE_FORMAT_STRING_ERROR_END
 
 #undef CASE
