Simplify the LLMScope constructor by folding three constructors into one.

diff --git a/Engine/Source/Runtime/Core/Public/HAL/LowLevelMemTracker.h b/Engine/Source/Runtime/Core/Public/HAL/LowLevelMemTracker.h
--- a/Engine/Source/Runtime/Core/Public/HAL/LowLevelMemTracker.h
+++ b/Engine/Source/Runtime/Core/Public/HAL/LowLevelMemTracker.h
@@ -888,34 +888,13 @@ public:
 class FLLMScope
 {
 public:
-	FLLMScope(FName TagName, bool bIsStatTag, ELLMTagSet InTagSet, ELLMTracker InTracker, bool bOverride = true)
+	template<typename TagType>
+	FLLMScope(TagType Tag, bool bIsStatTag, ELLMTagSet InTagSet, ELLMTracker InTracker, bool bOverride = true)
 	{
 		if (FLowLevelMemTracker::IsEnabled())
 		{
 			// We run the init in the open, because we want to track LLM even in transactions.
-			UE_AUTORTFM_OPEN { Init(TagName, bIsStatTag, InTagSet, InTracker, bOverride); };
-			// But remember that if we abort while we hold the scope, we need to destroy the scope.
-			AutoRTFM::PushOnAbortHandler(this, [this] { if (bEnabled) { Destruct(); } });
-		}
-	}
-
-	FLLMScope(ELLMTag TagEnum, bool bIsStatTag, ELLMTagSet InTagSet, ELLMTracker InTracker, bool bOverride = true)
-	{
-		if (FLowLevelMemTracker::IsEnabled())
-		{
-			// We run the init in the open, because we want to track LLM even in transactions.
-			UE_AUTORTFM_OPEN { Init(TagEnum, bIsStatTag, InTagSet, InTracker, bOverride); };
-			// But remember that if we abort while we hold the scope, we need to destroy the scope.
-			AutoRTFM::PushOnAbortHandler(this, [this] { if (bEnabled) { Destruct(); } });
-		}
-	}
-
-	FLLMScope(const UE::LLMPrivate::FTagData* TagData, bool bIsStatTag, ELLMTagSet Set, ELLMTracker Tracker, bool bOverride = true)
-	{
-		if (FLowLevelMemTracker::IsEnabled())
-		{
-			// We run the init in the open, because we want to track LLM even in transactions.
-			UE_AUTORTFM_OPEN { Init(TagData, bIsStatTag, Set, Tracker, bOverride); };
+			UE_AUTORTFM_OPEN { Init(Tag, bIsStatTag, InTagSet, InTracker, bOverride); };
 			// But remember that if we abort while we hold the scope, we need to destroy the scope.
 			AutoRTFM::PushOnAbortHandler(this, [this] { if (bEnabled) { Destruct(); } });
 		}
