Clean out API cruft in the AutoRTFM header.

Removes unused functions `autortfm_abort_if_transactional`, `autortfm_abort_if_closed`, `autortfm_check_consistency_assuming_no_races`, `autortfm_check_abi`, `AbortIfTransactional`, `AbortIfClosed`, `CheckConsistencyAssumingNoRaces`, and `WriteMemory` from the exposed API as they are unused / underused.

diff --git a/Engine/Source/Programs/AutoRTFMTests/Private/API.cpp b/Engine/Source/Programs/AutoRTFMTests/Private/API.cpp
--- a/Engine/Source/Programs/AutoRTFMTests/Private/API.cpp
+++ b/Engine/Source/Programs/AutoRTFMTests/Private/API.cpp
@@ -85,68 +85,6 @@ TEST_CASE("API.autortfm_abort_transaction")
     REQUIRE(true == AfterNest);
 }
 
-TEST_CASE("API.autortfm_abort_if_transactional")
-{
-    // Calling this outwith any transaction won't abort the program.
-    autortfm_abort_if_transactional();
-
-    bool BeforeNest = false;
-    bool InNest = false;
-    bool AfterNest = false;
-    
-    AutoRTFM::Commit([&]
-    {
-        BeforeNest = true;
-
-        AutoRTFM::ETransactionResult Result = AutoRTFM::Transact([&]
-        {
-            // Because we are aborting this won't actually occur!
-            InNest = true;
-
-            autortfm_abort_if_transactional();
-        });
-
-        REQUIRE(AutoRTFM::ETransactionResult::AbortedByRequest == Result);
-
-        AfterNest = true;
-    });
-
-    REQUIRE(true == BeforeNest);
-    REQUIRE(false == InNest);
-    REQUIRE(true == AfterNest);
-}
-
-TEST_CASE("API.autortfm_abort_if_closed")
-{
-    // Calling this outwith any transaction won't abort the program.
-    autortfm_abort_if_closed();
-
-    bool BeforeNest = false;
-    bool InNest = false;
-    bool AfterNest = false;
-
-    REQUIRE(AutoRTFM::ETransactionResult::AbortedByRequest == AutoRTFM::Transact([&]
-    {
-        BeforeNest = true;
-
-        AutoRTFM::Open([&]
-        {
-            InNest = true;
-
-            // This won't abort because we aren't closed!
-            autortfm_abort_if_closed();
-        });
-
-        AfterNest = true;
-
-        autortfm_abort_if_closed();
-    }));
-
-    REQUIRE(false == BeforeNest);
-    REQUIRE(true == InNest);
-    REQUIRE(false == AfterNest);
-}
-
 TEST_CASE("API.autortfm_open")
 {
     int Answer = 6 * 9;
@@ -384,14 +322,6 @@ TEST_CASE("API.autortfm_did_allocate")
     }
 }
 
-TEST_CASE("API.autortfm_check_consistency_assuming_no_races")
-{
-    AutoRTFM::Commit([&]
-    {
-        autortfm_check_consistency_assuming_no_races();
-    });
-}
-
 TEST_CASE("API.ETransactionResult")
 {
     int Answer = 6 * 9;
@@ -410,14 +340,6 @@ TEST_CASE("API.ETransactionResult")
     }));
 
     REQUIRE(42 == Answer);
-
-    REQUIRE(AutoRTFM::ETransactionResult::AbortedByRequest == AutoRTFM::Transact([&]
-    {
-        Answer = 13;
-        AutoRTFM::AbortIfTransactional();
-    }));
-    
-    REQUIRE(42 == Answer);
 }
 
 TEST_CASE("API.IsTransactional")
@@ -635,68 +557,6 @@ TEST_CASE("API.Abort")
     REQUIRE(true == AfterNest);
 }
 
-TEST_CASE("API.AbortIfTransactional")
-{
-    // Calling this outwith any transaction won't abort the program.
-    AutoRTFM::AbortIfTransactional();
-
-    bool BeforeNest = false;
-    bool InNest = false;
-    bool AfterNest = false;
-    
-    AutoRTFM::Commit([&]
-    {
-        BeforeNest = true;
-
-        AutoRTFM::ETransactionResult Result = AutoRTFM::Transact([&]
-        {
-            // Because we are aborting this won't actually occur!
-            InNest = true;
-
-            AutoRTFM::AbortIfTransactional();
-        });
-
-        REQUIRE(AutoRTFM::ETransactionResult::AbortedByRequest == Result);
-
-        AfterNest = true;
-    });
-
-    REQUIRE(true == BeforeNest);
-    REQUIRE(false == InNest);
-    REQUIRE(true == AfterNest);
-}
-
-TEST_CASE("API.AbortIfClosed")
-{
-    // Calling this outwith any transaction won't abort the program.
-    AutoRTFM::AbortIfClosed();
-
-    bool BeforeNest = false;
-    bool InNest = false;
-    bool AfterNest = false;
-
-    REQUIRE(AutoRTFM::ETransactionResult::AbortedByRequest == AutoRTFM::Transact([&]
-    {
-        BeforeNest = true;
-
-        AutoRTFM::Open([&]
-        {
-            InNest = true;
-
-            // This won't abort because we aren't closed!
-            AutoRTFM::AbortIfClosed();
-        });
-
-        AfterNest = true;
-
-        AutoRTFM::AbortIfClosed();
-    }));
-
-    REQUIRE(false == BeforeNest);
-    REQUIRE(true == InNest);
-    REQUIRE(false == AfterNest);
-}
-
 TEST_CASE("API.Open")
 {
     int Answer = 6 * 9;
@@ -1101,14 +961,6 @@ TEST_CASE("API.DidAllocate")
     }
 }
 
-TEST_CASE("API.CheckConsistencyAssumingNoRaces")
-{
-    AutoRTFM::Commit([&]
-    {
-        AutoRTFM::ForTheRuntime::CheckConsistencyAssumingNoRaces();
-    });
-}
-
 TEST_CASE("API.IsOnCurrentTransactionStack")
 {
 	{
diff --git a/Engine/Source/Programs/AutoRTFMTests/Private/Commit.cpp b/Engine/Source/Programs/AutoRTFMTests/Private/Commit.cpp
--- a/Engine/Source/Programs/AutoRTFMTests/Private/Commit.cpp
+++ b/Engine/Source/Programs/AutoRTFMTests/Private/Commit.cpp
@@ -82,7 +82,6 @@ TEST_CASE("Commit.Large")
                 m[i + j].push_back(i + j);
             }
         }
-        AutoRTFM::ForTheRuntime::CheckConsistencyAssumingNoRaces();
     });
 }
 
diff --git a/Engine/Source/Programs/AutoRTFMTests/Private/FakeVM.cpp b/Engine/Source/Programs/AutoRTFMTests/Private/FakeVM.cpp
--- a/Engine/Source/Programs/AutoRTFMTests/Private/FakeVM.cpp
+++ b/Engine/Source/Programs/AutoRTFMTests/Private/FakeVM.cpp
@@ -113,7 +113,8 @@ AutoRTFM::ETransactionResult FakeVM(const FOpcode* Opcodes, const int Length, in
 				case EOpcodeType::WriteToValue:
 				{
 					const int Arg = Value - Opcodes[i].Value;
-					AutoRTFM::ForTheRuntime::WriteMemory(&Value, &Arg);
+					AutoRTFM::RecordOpenWrite(&Value);
+					Value = Arg;
 					break;
 				}
 				}
diff --git a/Engine/Source/Programs/AutoRTFMTests/Private/OpenAPI.cpp b/Engine/Source/Programs/AutoRTFMTests/Private/OpenAPI.cpp
--- a/Engine/Source/Programs/AutoRTFMTests/Private/OpenAPI.cpp
+++ b/Engine/Source/Programs/AutoRTFMTests/Private/OpenAPI.cpp
@@ -25,7 +25,7 @@ TEST_CASE("OpenAPI.StartAbortAndStartAgain")
 		{
 			AutoRTFM::ForTheRuntime::StartTransaction();
 			AutoRTFM::RecordOpenWrite(&valueB);
-			valueB=10;
+			valueB = 10;
 			AutoRTFM::AbortTransaction();
 
 			AutoRTFM::ForTheRuntime::ClearTransactionStatus();
@@ -535,7 +535,8 @@ TEST_CASE("OpenAPI.StackWriteCommitInTheOpen3_Illegal")
 		AutoRTFM::Open([&]()
 		{
 			AutoRTFM::ForTheRuntime::StartTransaction();
-			AutoRTFM::ForTheRuntime::WriteMemory(&value, 10);
+			AutoRTFM::RecordOpenWrite(&value);
+			value = 10;
 			AutoRTFM::ForTheRuntime::CommitTransaction();
 			REQUIRE(value == 10);
 		});
@@ -553,8 +554,10 @@ TEST_CASE("OpenAPI.WriteMemory1")
 		AutoRTFM::Open([&]()
 		{
 			AutoRTFM::ForTheRuntime::StartTransaction();
+			
+			AutoRTFM::RecordOpenWrite(&value1);
+			value1 = sourceValue;
 
-			AutoRTFM::ForTheRuntime::WriteMemory(&value1, &sourceValue);
 			AutoRTFM::ForTheRuntime::CommitTransaction();
 		});
 		REQUIRE(value1 == 10);
@@ -570,7 +573,10 @@ TEST_CASE("OpenAPI.StackWriteAbortInTheOpen2")
 		AutoRTFM::Open([&]()
 		{
 			AutoRTFM::ForTheRuntime::StartTransaction();
-			AutoRTFM::ForTheRuntime::WriteMemory(&value, 10); // Illegal to write to value because it's in the inner-most closed-nest
+
+			// Illegal to write to value because it's in the inner-most closed-nest
+			AutoRTFM::RecordOpenWrite(&value);
+			value = 10;
 
 			AutoRTFM::AbortTransaction();
 		});
@@ -604,7 +610,8 @@ TEST_CASE("OpenAPI.WriteTrivialStructure")
 		{
 			AutoRTFM::ForTheRuntime::StartTransaction();
 
-			AutoRTFM::ForTheRuntime::WriteMemory(&data, &data2);
+			AutoRTFM::RecordOpenWrite(&data);
+			data = data2;
 			REQUIRE(data.A == 9);
 			REQUIRE(data.B == 8.0);
 			REQUIRE(data.C == 7.0f);
@@ -649,8 +656,9 @@ TEST_CASE("OpenAPI.WriteTrivialStructure2")
 		AutoRTFM::Open([&]()
 			{
 				AutoRTFM::ForTheRuntime::StartTransaction();
-
-				AutoRTFM::ForTheRuntime::WriteMemory(&data, &data2);
+				
+				AutoRTFM::RecordOpenWrite(&data);
+				data = data2;
 				REQUIRE(data.A == 9);
 				REQUIRE(data.B == 8.0);
 				REQUIRE(data.C == 7.0f);
@@ -660,8 +668,9 @@ TEST_CASE("OpenAPI.WriteTrivialStructure2")
 				REQUIRE(data.E[2] == 777);
 				REQUIRE(data.E[3] == 666);
 				REQUIRE(data.E[4] == 555);
-					
-				AutoRTFM::ForTheRuntime::WriteMemory(&data, &data3);
+
+				AutoRTFM::RecordOpenWrite(&data);
+				data = data3;
 				REQUIRE(data.A == 19);
 				REQUIRE(data.B == 28.0);
 				REQUIRE(data.C == 37.0f);
diff --git a/Engine/Source/Runtime/Core/Private/AutoRTFM/API.cpp b/Engine/Source/Runtime/Core/Private/AutoRTFM/API.cpp
--- a/Engine/Source/Runtime/Core/Private/AutoRTFM/API.cpp
+++ b/Engine/Source/Runtime/Core/Private/AutoRTFM/API.cpp
@@ -588,10 +588,6 @@ AUTORTFM_CORE_API void RTFM_autortfm_did_free(void* Ptr)
 	AutoRTFM::InternalUnreachable();
 }
 
-AUTORTFM_CORE_API void RTFM_autortfm_check_consistency_assuming_no_races()
-{
-}
-
 }  // anonymous namespace
 
 // The AutoRTFM public API.
@@ -1013,19 +1009,12 @@ extern "C" void autortfm_did_free(void* Ptr)
 	}
 }
 
-extern "C" void autortfm_check_consistency_assuming_no_races()
-{
-	if (autortfm_is_closed())
-	{
-		return RTFM_autortfm_check_consistency_assuming_no_races();
-	}
-
-    if (FContext::Get()->IsTransactional())
-    {
-        AutoRTFM::InternalUnreachable();
-    }
-}
-
+// If running with AutoRTFM enabled, then perform an ABI check between the
+// AutoRTFM compiler and the AutoRTFM runtime, to ensure that memory is being
+// laid out in an identical manner between the AutoRTFM runtime and the AutoRTFM
+// compiler pass. Should not be called manually by the user, a call to this will
+// be injected by the compiler into a global constructor in the AutoRTFM compiled
+// code.
 extern "C" UE_AUTORTFM_NOAUTORTFM void autortfm_check_abi(void* const Ptr, const size_t Size)
 {
     struct FConstants final
diff --git a/Engine/Source/Runtime/Core/Public/AutoRTFM/AutoRTFM.h b/Engine/Source/Runtime/Core/Public/AutoRTFM/AutoRTFM.h
--- a/Engine/Source/Runtime/Core/Public/AutoRTFM/AutoRTFM.h
+++ b/Engine/Source/Runtime/Core/Public/AutoRTFM/AutoRTFM.h
@@ -235,30 +235,6 @@ UE_AUTORTFM_API autortfm_status autortfm_get_context_status();
 static UE_AUTORTFM_FORCEINLINE autortfm_status autortfm_get_context_status() { return autortfm_status::autortfm_status_idle; }
 #endif
 
-#if UE_AUTORTFM_ENABLED
-static UE_AUTORTFM_FORCEINLINE void autortfm_abort_if_transactional(void)
-{
-	if (autortfm_is_transactional())
-	{
-		autortfm_abort_transaction();
-	}
-}
-#else
-static UE_AUTORTFM_FORCEINLINE void autortfm_abort_if_transactional(void) { }
-#endif
-
-#if UE_AUTORTFM_ENABLED
-static UE_AUTORTFM_FORCEINLINE void autortfm_abort_if_closed(void)
-{
-	if (autortfm_is_closed())
-	{
-		autortfm_abort_transaction();
-	}
-}
-#else
-static UE_AUTORTFM_FORCEINLINE void autortfm_abort_if_closed(void) { }
-#endif
-
 #if UE_AUTORTFM_ENABLED
 UE_AUTORTFM_API void autortfm_open(void (*work)(void* arg), void* arg, const void* return_address);
 #else
@@ -380,28 +356,6 @@ static UE_AUTORTFM_FORCEINLINE void autortfm_did_free(void* ptr)
 }
 #endif
 
-#if UE_AUTORTFM_ENABLED
-UE_AUTORTFM_API void autortfm_check_consistency_assuming_no_races(void);
-#else
-static UE_AUTORTFM_FORCEINLINE void autortfm_check_consistency_assuming_no_races(void) { }
-#endif
-
-// If running with AutoRTFM enabled, then perform an ABI check between the
-// AutoRTFM compiler and the AutoRTFM runtime, to ensure that memory is being
-// laid out in an identical manner between the AutoRTFM runtime and the AutoRTFM
-// compiler pass. Should not be called manually by the user, a call to this will
-// be injected by the compiler into a global constructor in the AutoRTFM compiled
-// code.
-#if UE_AUTORTFM_ENABLED
-UE_AUTORTFM_API void autortfm_check_abi(void* ptr, size_t size);
-#else
-static UE_AUTORTFM_FORCEINLINE void autortfm_check_abi(void* ptr, size_t size)
-{
-	UE_AUTORTFM_UNUSED(ptr);
-	UE_AUTORTFM_UNUSED(size);
-}
-#endif
-
 #if UE_AUTORTFM_ENABLED
 UE_AUTORTFM_API void autortfm_unreachable(const UE_AUTORTFM_TCHAR* const Message);
 #else
@@ -625,18 +579,6 @@ static UE_AUTORTFM_FORCEINLINE ETransactionResult CascadingAbortTransaction()
 	return static_cast<ETransactionResult>(autortfm_cascading_abort_transaction());
 }
 
-// Abort if running in a transaction.
-static UE_AUTORTFM_FORCEINLINE void AbortIfTransactional()
-{
-    autortfm_abort_if_transactional();
-}
-
-// Abort if running in closed code.
-static UE_AUTORTFM_FORCEINLINE void AbortIfClosed()
-{
-    autortfm_abort_if_closed();
-}
-
 namespace Detail
 {
 template<typename T, typename = void>
@@ -1206,51 +1148,6 @@ namespace ForTheRuntime
 	// Reserved for future.
 	template<typename Type> static UE_AUTORTFM_FORCEINLINE void RecordOpenRead(Type*) {}
 
-	// WriteMemory first records the memory span as written (see RecordOpenWrite) and then copies the specified value into it.
-	static UE_AUTORTFM_FORCEINLINE void WriteMemory(void* DestPtr, void const* SrcPtr, size_t Size)
-	{
-		RecordOpenWrite(DestPtr, Size);
-		UE_AUTORTFM_MEMCPY(DestPtr, SrcPtr, Size);
-	}
-
-	// WriteMemory first records the memory span as written (see RecordOpenWrite) and then copies the specified value into it.
-	template<typename Type> static UE_AUTORTFM_FORCEINLINE void WriteMemory(Type* DestPtr, Type const* SrcPtr)
-	{
-		if constexpr (std::is_trivially_copyable<Type>::value)
-		{
-			RecordOpenWrite(DestPtr, sizeof(Type));
-			*DestPtr = *SrcPtr;
-		}
-		else
-		{
-			WriteMemory(DestPtr, SrcPtr, sizeof(Type));
-		}
-	}
-
-	// WriteMemory first records the memory span as written (see RecordOpenWrite) and then copies the specified value into it.
-	template<typename Type> static UE_AUTORTFM_FORCEINLINE void WriteMemory(Type* DestPtr, Type const SrcValue)
-	{
-		if constexpr (std::is_trivially_copyable<Type>::value)
-		{
-			RecordOpenWrite(DestPtr, sizeof(Type));
-			*DestPtr = SrcValue;
-		}
-		else
-		{
-			WriteMemory(DestPtr, &SrcValue, sizeof(Type));
-		}
-	}
-
-	// If running in a transaction, then perform a consistency check of the
-	// transaction's read-write set. If possible, this compares the read-write set's
-	// expected values with the actual values in global memory. Does nothing when
-	// called outside of a transaction. May do nothing if debugging features aren't
-	// enabled in the autortfm runtime.
-	static UE_AUTORTFM_FORCEINLINE void CheckConsistencyAssumingNoRaces()
-	{
-		autortfm_check_consistency_assuming_no_races();
-	}
-
 } // namespace ForTheRuntime
 
 } // namespace AutoRTFM
@@ -1328,7 +1225,6 @@ namespace AutoRTFM::Private
 // Runs a block of code in the open, non-transactionally. Anything performed in the open will not be undone if a transaction fails.
 // Calls should be written like this: UE_AUTORTFM_OPEN { ... code ... };
 #define UE_AUTORTFM_OPEN  UE_AUTORTFM_OPEN_IMPL
-#define UE_AUTORTFM_OPEN2 UE_AUTORTFM_OPEN_IMPL
 
 #define UE_AUTORTFM_OPEN_NO_VALIDATION UE_AUTORTFM_OPEN_NO_VALIDATION_IMPL
 
@@ -1337,19 +1233,16 @@ namespace AutoRTFM::Private
 // The macro arguments are the capture specification for the lambda.
 // Calls should be written like this: UE_AUTORTFM_ONABORT(=) { ... code ... };
 #define UE_AUTORTFM_ONABORT(...)  UE_AUTORTFM_ONABORT_IMPL(__VA_ARGS__)
-#define UE_AUTORTFM_ONABORT2(...) UE_AUTORTFM_ONABORT_IMPL(__VA_ARGS__)
 
 // Runs a block of code if a transaction commits successfully.
 // In non-transactional code paths the block of code will be executed immediately.
 // The macro arguments are the capture specification for the lambda.
 // Calls should be written like this: UE_AUTORTFM_ONCOMMIT(=) { ... code ... };
 #define UE_AUTORTFM_ONCOMMIT(...)  UE_AUTORTFM_ONCOMMIT_IMPL(__VA_ARGS__)
-#define UE_AUTORTFM_ONCOMMIT2(...) UE_AUTORTFM_ONCOMMIT_IMPL(__VA_ARGS__)
 
 // Runs a block of code in the closed, transactionally, within a new transaction.
 // Calls should be written like this: UE_AUTORTFM_TRANSACT { ... code ... };
 #define UE_AUTORTFM_TRANSACT  UE_AUTORTFM_TRANSACT_IMPL
-#define UE_AUTORTFM_TRANSACT2 UE_AUTORTFM_TRANSACT_IMPL
 
 #define UE_AUTORTFM_CONCAT_IMPL(A, B) A ## B
 #define UE_AUTORTFM_CONCAT(A, B) UE_AUTORTFM_CONCAT_IMPL(A, B)
