[IoStore] Implemented MMIO for utoc files.
Allows platform OS to unload memory pages occupied for that files and only keep what's actively used.
Can be disabled with CVar "IoStore.AllowMemoryMappedUtoc".

#tests RR, playtests

diff --git a/Engine/Source/Runtime/Core/Internal/IO/IoStore.h b/Engine/Source/Runtime/Core/Internal/IO/IoStore.h
--- a/Engine/Source/Runtime/Core/Internal/IO/IoStore.h
+++ b/Engine/Source/Runtime/Core/Internal/IO/IoStore.h
@@ -2,7 +2,9 @@
 
 #pragma once
 
+#include "Async/MappedFileHandle.h"
 #include "Containers/Array.h"
+#include "Experimental/ConcurrentLinearAllocator.h"
 #include "HAL/Platform.h"
 #include "HAL/UnrealMemory.h"
 #include "IO/IoContainerId.h"
@@ -12,6 +14,7 @@
 #include "Misc/EnumClassFlags.h"
 #include "Misc/Guid.h"
 #include "Misc/SecureHash.h"
+#include "Misc/TVariant.h"
 #include "UObject/NameTypes.h"
 
 CORE_API DECLARE_LOG_CATEGORY_EXTERN(LogIoStore, Log, All);
@@ -172,6 +175,102 @@ enum class EIoStoreTocReadOptions
 };
 ENUM_CLASS_FLAGS(EIoStoreTocReadOptions);
 
+
+struct FIoStoreTocResourceStorage
+{
+	FIoStoreTocResourceStorage() = default;
+	CORE_API explicit FIoStoreTocResourceStorage(const TCHAR* TocFilePath);
+
+	bool IsLoaded() const { return !Data.IsType<FEmptyVariantState>(); }
+	CORE_API uint64 GetAllocatedSize() const;
+
+	template<typename T>
+	TConstArrayView<T> ChopView(int32 Num)
+	{
+		const int32 RequestedSize = Num * static_cast<int32>(sizeof(T));
+		TConstArrayView<uint8> Bytes = ChopBytes(RequestedSize);
+		check(RequestedSize == Bytes.Num());
+		return TConstArrayView<T>(reinterpret_cast<const T*>(Bytes.GetData()), Num);
+	}
+
+	template<typename T>
+	TArray<T> ChopArray(int32 Num)
+	{
+		TArray<T> Result(ChopView<T>(Num));
+		if (!Result.IsEmpty())
+		{
+			ReleaseOwnershipOfLastBlock();
+		}
+		return Result;
+	}
+
+	CORE_API void FinalizeRead();
+
+private:
+	CORE_API void ReleaseOwnershipOfLastBlock();
+
+	struct FIOStoreAllocationTag
+	{
+		static constexpr uint32 BlockSize = 64 * 1024;
+		static constexpr bool AllowOversizedBlocks = true;
+		static constexpr bool RequiresAccurateSize = true;
+		static constexpr bool InlineBlockAllocation = false;
+		static constexpr const char* TagName = "IOStoreLinear";
+
+		using Allocator = TBlockAllocationLockFreeCache<BlockSize, FOsAllocator>;
+	};
+	using FDataAllocator = TConcurrentLinearArrayAllocator<FIOStoreAllocationTag>;
+	using FLocalAllocator = TInlineAllocator<10>;
+	struct FMappedFile
+	{
+		TUniquePtr<IMappedFileHandle> MappedFile;
+		TArray<TUniquePtr<IMappedFileRegion>, FLocalAllocator> MappedRegions;
+		TArray<TArray<uint8, FDataAllocator>, FLocalAllocator> OwnedRegions;
+		int32 Cursor = 0;
+		bool bLastReadBlockWasOwned = false;
+	};
+	struct FReadBlocks
+	{
+		TUniquePtr<IFileHandle> File;
+		TArray<TArray<uint8, FDataAllocator>, FLocalAllocator> Blocks;
+	};
+
+	static CORE_API TConstArrayView<uint8> ChopBytes(FMappedFile& File, int32 Size);
+	static CORE_API TConstArrayView<uint8> ChopBytes(FReadBlocks& File, int32 Size);
+
+	TConstArrayView<uint8> ChopBytes(FEmptyVariantState&, int32 Size)
+	{
+		return {};
+	}
+
+	TConstArrayView<uint8> ChopBytes(int32 Size)
+	{
+		return Visit([this, Size](auto&& Value) { return ChopBytes(Value, Size); }, Data);
+	}
+
+	TVariant<FEmptyVariantState, FMappedFile, FReadBlocks> Data;
+};
+
+struct FIoStoreTocResourceView
+{
+	FIoStoreTocHeader Header;
+	TConstArrayView<FIoChunkId> ChunkIds;
+	TConstArrayView<FIoOffsetAndLength> ChunkOffsetLengths;
+	TConstArrayView<int32> ChunkPerfectHashSeeds;
+	TConstArrayView<int32> ChunkIndicesWithoutPerfectHash;
+	TConstArrayView<FIoStoreTocCompressedBlockEntry> CompressionBlocks;
+
+	TArray<FName> CompressionMethods;
+	FSHAHash SignatureHash;
+	TConstArrayView<FSHAHash> ChunkBlockSignatures;
+	TConstArrayView<uint8> DirectoryIndexBuffer;
+	TArray<FIoStoreTocEntryMeta> ChunkMetas;
+
+	[[nodiscard]] CORE_API static FIoStoreTocResourceView BuildViewFromResource(const struct FIoStoreTocResource& Resource);
+	[[nodiscard]] CORE_API static FIoStatus Read(const TCHAR* TocFilePath, EIoStoreTocReadOptions ReadOptions, FIoStoreTocResourceView& OutTocResource, FIoStoreTocResourceStorage& OutTocResourceStorage);
+	[[nodiscard]] CORE_API FIoStoreTocChunkInfo GetTocChunkInfo(int32 TocEntryIndex) const;
+};
+
 /**
  * Container TOC data.
  */
@@ -194,18 +293,18 @@ struct FIoStoreTocResource
 	TArray<FName> CompressionMethods;
 
 	FSHAHash SignatureHash;
-	
+
 	TArray<FSHAHash> ChunkBlockSignatures;
 
 	TArray<uint8> DirectoryIndexBuffer; 
-	
+
 	TArray<FIoStoreTocEntryMeta> ChunkMetas;
 
 	[[nodiscard]] CORE_API FIoStoreTocChunkInfo GetTocChunkInfo(int32 TocEntryIndex) const;
 
+	[[nodiscard]] CORE_API static FIoStoreTocResource BuildResourceFromMappedView(const FIoStoreTocResourceView& View);
 	[[nodiscard]] CORE_API static FIoStatus Read(const TCHAR* TocFilePath, EIoStoreTocReadOptions ReadOptions, FIoStoreTocResource& OutTocResource);
-
 	[[nodiscard]] CORE_API static TIoStatusOr<uint64> Write(const TCHAR* TocFilePath, FIoStoreTocResource& TocResource, uint32 CompressionBlockSize, uint64 MaxPartitionSize, const FIoContainerSettings& ContainerSettings);
 
-	CORE_API static uint64 HashChunkIdWithSeed(int32 Seed, const FIoChunkId& ChunkId);
+	[[nodiscard]] CORE_API static uint64 HashChunkIdWithSeed(int32 Seed, const FIoChunkId& ChunkId);
 };
diff --git a/Engine/Source/Runtime/Core/Private/IO/IoDirectoryIndex.cpp b/Engine/Source/Runtime/Core/Private/IO/IoDirectoryIndex.cpp
--- a/Engine/Source/Runtime/Core/Private/IO/IoDirectoryIndex.cpp
+++ b/Engine/Source/Runtime/Core/Private/IO/IoDirectoryIndex.cpp
@@ -319,7 +319,7 @@ uint32 FIoDirectoryIndexWriter::AddFile(const FStringView& FileName, uint32 Dire
 class FIoDirectoryIndexReaderImpl
 {
 public:
-	FIoStatus Initialize(TArray<uint8>& InBuffer, FAES::FAESKey InDecryptionKey)
+	FIoStatus Initialize(TConstArrayView<uint8> InBuffer, FAES::FAESKey InDecryptionKey)
 	{
 		if (InBuffer.Num() == 0)
 		{
@@ -328,11 +328,16 @@ public:
 
 		if (InDecryptionKey.IsValid())
 		{
-			FAES::DecryptData(InBuffer.GetData(), InBuffer.Num(), InDecryptionKey);
+			TArray<uint8, TInlineAllocator<2048>> Bytes(InBuffer); // temp buffer for dencryption
+			FAES::DecryptData(Bytes.GetData(), Bytes.Num(), InDecryptionKey);
+			FMemoryReaderView Ar(MakeArrayView<const uint8>(Bytes.GetData(), Bytes.Num()));
+			Ar << DirectoryIndex;
+		}
+		else
+		{
+			FMemoryReaderView Ar(MakeArrayView<const uint8>(InBuffer.GetData(), InBuffer.Num()));
+			Ar << DirectoryIndex;
 		}
-
-		FMemoryReaderView Ar(MakeArrayView<const uint8>(InBuffer.GetData(), InBuffer.Num()));
-		Ar << DirectoryIndex;
 
 		return FIoStatus::Ok;
 	}
@@ -466,7 +471,7 @@ FIoDirectoryIndexReader::~FIoDirectoryIndexReader()
 	delete Impl;
 }
 
-FIoStatus FIoDirectoryIndexReader::Initialize(TArray<uint8>& InBuffer, FAES::FAESKey InDecryptionKey)
+FIoStatus FIoDirectoryIndexReader::Initialize(TConstArrayView<uint8> InBuffer, FAES::FAESKey InDecryptionKey)
 {
 	return Impl->Initialize(InBuffer, InDecryptionKey);
 }
diff --git a/Engine/Source/Runtime/Core/Private/IO/IoStore.cpp b/Engine/Source/Runtime/Core/Private/IO/IoStore.cpp
--- a/Engine/Source/Runtime/Core/Private/IO/IoStore.cpp
+++ b/Engine/Source/Runtime/Core/Private/IO/IoStore.cpp
@@ -2,6 +2,7 @@
 
 #include "IO/IoStore.h"
 
+#include "Async/MappedFileHandle.h"
 #include "Containers/Map.h"
 #include "Features/IModularFeatures.h"
 #include "HAL/FileManager.h"
@@ -158,7 +159,7 @@ public:
 
 	[[nodiscard]] FIoStatus Read(const TCHAR* TocFilePath, const TMap<FGuid, FAES::FAESKey>& DecryptionKeys)
 	{
-		FIoStatus TocStatus = FIoStoreTocResource::Read(TocFilePath, EIoStoreTocReadOptions::ReadAll, Toc);
+		FIoStatus TocStatus = FIoStoreTocResourceView::Read(TocFilePath, EIoStoreTocReadOptions::ReadAll, Toc, TocStorage);
 		if (!TocStatus.IsOk())
 		{
 			return TocStatus;
@@ -202,12 +203,8 @@ public:
 		return TocStatus;
 	}
 
-	FIoStoreTocResource& GetTocResource()
-	{
-		return Toc;
-	}
 
-	const FIoStoreTocResource& GetTocResource() const
+	const FIoStoreTocResourceView& GetTocResource() const
 	{
 		return Toc;
 	}
@@ -260,14 +257,15 @@ private:
 		IndexToFileName.Add(TocEntryIndex, FString(Filename));
 	}
 
-	FIoStoreTocResource Toc;
+	FIoStoreTocResourceView Toc;
+	FIoStoreTocResourceStorage TocStorage;
 	FIoDirectoryIndexReader DirectoryIndexReader;
 	FAES::FAESKey DecryptionKey;
 	TMap<FIoChunkId, int32> ChunkIdToIndex;
 	TMap<int32, FString> IndexToFileName;
 };
 
-FIoStoreTocChunkInfo FIoStoreTocResource::GetTocChunkInfo(int32 TocEntryIndex) const
+FIoStoreTocChunkInfo FIoStoreTocResourceView::GetTocChunkInfo(int32 TocEntryIndex) const
 {
 	const FIoStoreTocEntryMeta& Meta = ChunkMetas[TocEntryIndex];
 	const FIoOffsetAndLength& OffsetLength = ChunkOffsetLengths[TocEntryIndex];
@@ -411,7 +409,7 @@ public:
 			return TocStatus;
 		}
 
-		FIoStoreTocResource& TocResource = TocReader.GetTocResource();
+		const FIoStoreTocResourceView& TocResource = TocReader.GetTocResource();
 
 		IPlatformFile& Ipf = FPlatformFileManager::Get().GetPlatformFile();
 		ContainerFileAccessors.Reserve(TocResource.Header.PartitionCount);
@@ -467,7 +465,7 @@ public:
 
 	void EnumerateChunks(TFunction<bool(FIoStoreTocChunkInfo&&)>&& Callback) const
 	{
-		const FIoStoreTocResource& TocResource = TocReader.GetTocResource();
+		const FIoStoreTocResourceView& TocResource = TocReader.GetTocResource();
 
 		for (int32 ChunkIndex = 0; ChunkIndex < TocResource.ChunkIds.Num(); ++ChunkIndex)
 		{
@@ -494,7 +492,7 @@ public:
 
 	TIoStatusOr<FIoStoreTocChunkInfo> GetChunkInfo(const uint32 TocEntryIndex) const
 	{
-		const FIoStoreTocResource& TocResource = TocReader.GetTocResource();
+		const FIoStoreTocResourceView& TocResource = TocReader.GetTocResource();
 
 		if (TocEntryIndex < uint32(TocResource.ChunkIds.Num()))
 		{
@@ -534,7 +532,7 @@ public:
 		const uint64 RequestedOffset = Options.GetOffset();
 		const uint64 ResolvedOffset = OffsetAndLength->GetOffset() + RequestedOffset;
 		const uint64 ResolvedSize = RequestedOffset <= OffsetAndLength->GetLength() ? FMath::Min(Options.GetSize(), OffsetAndLength->GetLength() - RequestedOffset) : 0;
-		const FIoStoreTocResource& TocResource = TocReader.GetTocResource();
+		const FIoStoreTocResourceView& TocResource = TocReader.GetTocResource();
 		const uint64 CompressionBlockSize = TocResource.Header.CompressionBlockSize;
 		const int32 FirstBlockIndex = int32(ResolvedOffset / CompressionBlockSize);
 		const int32 LastBlockIndex = int32((Align(ResolvedOffset + ResolvedSize, CompressionBlockSize) - 1) / CompressionBlockSize);
@@ -578,7 +576,7 @@ public:
 					{
 						uint8* CompressedSource = State->CompressedBuffer.GetData() + CompressedSourceOffset;
 						uint8* UncompressedDestination = State->UncompressedBuffer->Data() + UncompressedDestinationOffset;
-						const FIoStoreTocResource& TocResource = TocReader.GetTocResource();
+						const FIoStoreTocResourceView& TocResource = TocReader.GetTocResource();
 						const FIoStoreTocCompressedBlockEntry& CompressionBlock = TocResource.CompressionBlocks[BlockIndex];
 						const uint32 RawSize = Align(CompressionBlock.GetCompressedSize(), FAES::AESBlockSize);
 						const uint32 UncompressedSize = CompressionBlock.GetUncompressedSize();
@@ -673,7 +671,7 @@ public:
 			ResolvedSize = FMath::Min(Options.GetSize(), OffsetAndLength->GetLength() - RequestedOffset);
 		}
 
-		const FIoStoreTocResource& TocResource = TocReader.GetTocResource();
+		const FIoStoreTocResourceView& TocResource = TocReader.GetTocResource();
 		const uint64 CompressionBlockSize = TocResource.Header.CompressionBlockSize;
 		FIoBuffer UncompressedBuffer(ResolvedSize);
 		if (ResolvedSize == 0)
@@ -819,7 +817,7 @@ public:
 		}
 
 		// Find what compressed blocks this read straddles.
-		const FIoStoreTocResource& TocResource = TocReader.GetTocResource();
+		const FIoStoreTocResourceView& TocResource = TocReader.GetTocResource();
 		const uint64 CompressionBlockSize = TocResource.Header.CompressionBlockSize;
 		int32 FirstBlockIndex = int32(ResolvedOffset / CompressionBlockSize);
 		int32 LastBlockIndex = int32((Align(ResolvedOffset + ResolvedSize, CompressionBlockSize) - 1) / CompressionBlockSize);
@@ -900,7 +898,7 @@ public:
 
 	bool TocChunkContainsBlockIndex(const int32 TocEntryIndex, const int32 BlockIndex) const
 	{
-		const FIoStoreTocResource& TocResource = TocReader.GetTocResource();
+		const FIoStoreTocResourceView& TocResource = TocReader.GetTocResource();
 		const FIoOffsetAndLength& OffsetLength = TocResource.ChunkOffsetLengths[TocEntryIndex];
 
 		const uint64 CompressionBlockSize = TocResource.Header.CompressionBlockSize;
@@ -929,7 +927,7 @@ public:
 		}
 
 		// Find what compressed blocks this chunk straddles.
-		const FIoStoreTocResource& TocResource = TocReader.GetTocResource();
+		const FIoStoreTocResourceView& TocResource = TocReader.GetTocResource();
 		const uint64 CompressionBlockSize = TocResource.Header.CompressionBlockSize;
 		int32 FirstBlockIndex = int32(OffsetAndLength->GetOffset() / CompressionBlockSize);
 		int32 LastBlockIndex = int32((Align(OffsetAndLength->GetOffset() + OffsetAndLength->GetLength(), CompressionBlockSize) - 1) / CompressionBlockSize);
@@ -954,7 +952,7 @@ public:
 
 	void EnumerateCompressedBlocks(TFunction<bool(const FIoStoreTocCompressedBlockInfo&)>&& Callback) const
 	{
-		const FIoStoreTocResource& TocResource = TocReader.GetTocResource();
+		const FIoStoreTocResourceView& TocResource = TocReader.GetTocResource();
 
 		for (int32 BlockIndex = 0; BlockIndex < TocResource.CompressionBlocks.Num(); ++BlockIndex)
 		{
@@ -1103,21 +1101,180 @@ void FIoStoreReader::GetContainerFilePaths(TArray<FString>& OutPaths)
 	Impl->GetContainerFilePaths(OutPaths);
 }
 
-FIoStatus FIoStoreTocResource::Read(const TCHAR* TocFilePath, EIoStoreTocReadOptions ReadOptions, FIoStoreTocResource& OutTocResource)
+static TAutoConsoleVariable<bool> CVarIoStoreAllowMemoryMappedUtoc(
+	TEXT("IoStore.AllowMemoryMappedUtoc"),
+	true,
+	TEXT("Allow IoStore to memory map utoc containers instead of loading thier content as a whole into memory.")
+	);
+
+FIoStoreTocResourceStorage::FIoStoreTocResourceStorage(const TCHAR* TocFilePath)
+{
+	IPlatformFile& Ipf = FPlatformFileManager::Get().GetPlatformFile();
+
+	// try memory mapping first
+	if (FPlatformProperties::SupportsMemoryMappedFiles() && CVarIoStoreAllowMemoryMappedUtoc.GetValueOnAnyThread())
+	{
+		FOpenMappedResult Result = Ipf.OpenMappedEx(TocFilePath, IPlatformFile::EOpenReadFlags::None);
+		if (Result.HasValue())
+		{
+			FMappedFile Mapped;
+			Mapped.MappedFile = Result.StealValue();
+			Data.Emplace<FMappedFile>(MoveTemp(Mapped));
+		}
+	}
+	// then try to open file directly and own all read blocks.
+	if (Data.IsType<FEmptyVariantState>())
+	{
+		FFileOpenResult Result = Ipf.OpenRead(TocFilePath, IPlatformFile::EOpenReadFlags::None);
+		if (Result.HasValue())
+		{
+			FReadBlocks Blocks;
+			Blocks.File = Result.StealValue();
+			Data.Emplace<FReadBlocks>(MoveTemp(Blocks));
+		}
+	}
+}
+
+uint64 FIoStoreTocResourceStorage::GetAllocatedSize() const
+{
+	auto Visitor = [](auto&& V) -> int64
+	{
+		using ValueType = std::decay_t<decltype(V)>;
+		if constexpr (std::is_same_v<ValueType, FReadBlocks>)
+		{
+			return V.Blocks.GetAllocatedSize();
+		}
+		else if constexpr (std::is_same_v<ValueType, FMappedFile>)
+		{
+			return sizeof(*V.MappedFile) + V.MappedRegions.GetAllocatedSize() + V.OwnedRegions.GetAllocatedSize();
+		}
+		else
+		{
+			return 0ll;
+		}
+	};
+	return sizeof(this) + Visit(Visitor, Data);
+}
+
+TConstArrayView<uint8> FIoStoreTocResourceStorage::ChopBytes(FReadBlocks& Blocks, int32 Size)
+{
+	if (Size <= 0)
+	{
+		return TConstArrayView<uint8>();
+	}
+	auto& Range = Blocks.Blocks.AddDefaulted_GetRef();
+	Range.SetNumUninitialized(Size);
+	if (Blocks.File->Read(Range.GetData(), Size))
+	{
+		return MakeConstArrayView(Range);
+	}
+	Blocks.Blocks.Pop();
+	Blocks.File->SeekFromEnd();
+	return TConstArrayView<uint8>();
+}
+
+TConstArrayView<uint8> FIoStoreTocResourceStorage::ChopBytes(FMappedFile& MappedFile, int32 Size)
+{
+	if (Size <= 0)
+	{
+		return TConstArrayView<uint8>();
+	}
+	TUniquePtr<IMappedFileRegion> Region(MappedFile.MappedFile->MapRegion(MappedFile.Cursor, Size));
+	const int32 ActualSize = static_cast<int32>(Region->GetMappedSize());
+	ensureAlways(ActualSize == Size);
+	// optimization for small regions
+	if (ActualSize < FPlatformMemory::GetConstants().OsAllocationGranularity)
+	{
+		auto& Range = MappedFile.OwnedRegions.Emplace_GetRef(Region->GetMappedPtr(), ActualSize);
+		Region.Reset(); // intentionally drop mapped Region to let the OS to unload the page
+		MappedFile.Cursor += ActualSize;
+		MappedFile.bLastReadBlockWasOwned = true;
+		return MakeConstArrayView(Range);
+	}
+	else
+	{
+		auto& Range = MappedFile.MappedRegions.Add_GetRef(MoveTemp(Region));
+		MappedFile.Cursor += ActualSize;
+		MappedFile.bLastReadBlockWasOwned = false;
+		return TConstArrayView<uint8>(Range->GetMappedPtr(), ActualSize);
+	}
+}
+
+void FIoStoreTocResourceStorage::ReleaseOwnershipOfLastBlock()
+{
+	auto Visitor = [](auto&& V)
+	{
+		using ValueType = std::decay_t<decltype(V)>;
+		auto NoShrink = EAllowShrinking::No;
+		if constexpr (std::is_same_v<ValueType, FReadBlocks>)
+		{
+			V.Blocks.Pop(NoShrink);
+		}
+		else if constexpr (std::is_same_v<ValueType, FMappedFile>)
+		{
+			if (V.bLastReadBlockWasOwned)
+			{
+				V.OwnedRegions.Pop(NoShrink);
+			}
+			else
+			{
+				V.MappedRegions.Pop(NoShrink);
+			}
+		}
+	};
+	Visit(Visitor, Data);
+}
+
+void FIoStoreTocResourceStorage::FinalizeRead()
+{
+	auto Visitor = [](auto&& V)
+	{
+		using ValueType = std::decay_t<decltype(V)>;
+		if constexpr (std::is_same_v<ValueType, FReadBlocks>)
+		{
+			V.File.Reset();
+		}
+	};
+	Visit(Visitor, Data);
+}
+
+FIoStoreTocResourceView FIoStoreTocResourceView::BuildViewFromResource(const FIoStoreTocResource& Resource)
+{
+	FIoStoreTocResourceView View;
+	View.Header = Resource.Header;
+	View.ChunkIds = MakeConstArrayView(Resource.ChunkIds);
+	View.ChunkOffsetLengths = MakeConstArrayView(Resource.ChunkOffsetLengths);
+	View.ChunkPerfectHashSeeds = MakeConstArrayView(Resource.ChunkPerfectHashSeeds);
+	View.ChunkIndicesWithoutPerfectHash = MakeConstArrayView(Resource.ChunkIndicesWithoutPerfectHash);
+	View.CompressionBlocks = MakeConstArrayView(Resource.CompressionBlocks);
+
+	View.CompressionMethods = Resource.CompressionMethods;
+	View.SignatureHash = Resource.SignatureHash;
+	View.ChunkBlockSignatures = MakeConstArrayView(Resource.ChunkBlockSignatures);
+	View.DirectoryIndexBuffer = MakeConstArrayView(Resource.DirectoryIndexBuffer);
+	View.ChunkMetas = Resource.ChunkMetas;
+
+	return View;
+}
+
+FIoStatus FIoStoreTocResourceView::Read(const TCHAR* TocFilePath, EIoStoreTocReadOptions ReadOptions, FIoStoreTocResourceView& OutTocResource, FIoStoreTocResourceStorage& OutTocStorage)
 {
 	check(TocFilePath != nullptr);
 
-	IPlatformFile& Ipf = FPlatformFileManager::Get().GetPlatformFile();
-	TUniquePtr<IFileHandle>	TocFileHandle(Ipf.OpenRead(TocFilePath, /* allowwrite */ false));
+	FIoStoreTocResourceStorage TocStorage(TocFilePath);
 
-	if (!TocFileHandle)
+	if (!TocStorage.IsLoaded())
 	{
 		return FIoStatusBuilder(EIoErrorCode::FileOpenFailed) << TEXT("Failed to open IoStore TOC file '") << TocFilePath << TEXT("'");
 	}
 
 	// Header
 	FIoStoreTocHeader& Header = OutTocResource.Header;
-	if (!TocFileHandle->Read(reinterpret_cast<uint8*>(&Header), sizeof(FIoStoreTocHeader)))
+	if (TArray<FIoStoreTocHeader> HeaderView = TocStorage.ChopArray<FIoStoreTocHeader>(1); !HeaderView.IsEmpty())
+	{
+		FPlatformMemory::Memcpy(&Header, HeaderView.GetData(), HeaderView.NumBytes());
+	}
+	else
 	{
 		return FIoStatusBuilder(EIoErrorCode::CorruptToc) << TEXT("Failed to read IoStore TOC file '") << TocFilePath << TEXT("'");
 	}
@@ -1147,38 +1304,11 @@ FIoStatus FIoStoreTocResource::Read(const TCHAR* TocFilePath, EIoStoreTocReadOpt
 		return FIoStatusBuilder(EIoErrorCode::CorruptToc) << TEXT("Too new TOC header version while reading '") << TocFilePath << TEXT("'");
 	}
 
-	const uint64 TotalTocSize = TocFileHandle->Size() - sizeof(FIoStoreTocHeader);
-	const uint64 TocMetaSize = Header.TocEntryCount * sizeof(FIoStoreTocEntryMeta);
-
-	const uint64 DefaultTocSize = TotalTocSize - (Header.DirectoryIndexSize + TocMetaSize);
-	uint64 TocSize = DefaultTocSize;
-
-	if (EnumHasAnyFlags(ReadOptions, EIoStoreTocReadOptions::ReadTocMeta))
-	{
-		TocSize = TotalTocSize; // Meta data is at the end of the TOC file
-	}
-	else if (EnumHasAnyFlags(ReadOptions, EIoStoreTocReadOptions::ReadDirectoryIndex))
-	{
-		TocSize = DefaultTocSize + Header.DirectoryIndexSize;
-	}
-
-	TUniquePtr<uint8[]> TocBuffer = MakeUnique<uint8[]>(TocSize);
-
-	if (TocSize != 0 && !TocFileHandle->Read(TocBuffer.Get(), TocSize))
-	{
-		return FIoStatusBuilder(EIoErrorCode::CorruptToc) << TEXT("Failed to read IoStore TOC file '") << TocFilePath << TEXT("'");
-	}
-
 	// Chunk IDs
-	const uint8* DataPtr = TocBuffer.Get();
-	const FIoChunkId* ChunkIds = reinterpret_cast<const FIoChunkId*>(DataPtr);
-	OutTocResource.ChunkIds = MakeArrayView<FIoChunkId const>(ChunkIds, Header.TocEntryCount);
-	DataPtr += Header.TocEntryCount * sizeof(FIoChunkId);
+	OutTocResource.ChunkIds = TocStorage.ChopView<FIoChunkId>(Header.TocEntryCount);
 
 	// Chunk offsets
-	const FIoOffsetAndLength* ChunkOffsetLengths = reinterpret_cast<const FIoOffsetAndLength*>(DataPtr);
-	OutTocResource.ChunkOffsetLengths = MakeArrayView<FIoOffsetAndLength const>(ChunkOffsetLengths, Header.TocEntryCount);
-	DataPtr += Header.TocEntryCount * sizeof(FIoOffsetAndLength);
+	OutTocResource.ChunkOffsetLengths = TocStorage.ChopView<FIoOffsetAndLength>(Header.TocEntryCount);
 
 	// Chunk perfect hash map
 	uint32 PerfectHashSeedsCount = 0;
@@ -1194,38 +1324,30 @@ FIoStatus FIoStoreTocResource::Read(const TCHAR* TocFilePath, EIoStoreTocReadOpt
 	}
 	if (PerfectHashSeedsCount)
 	{
-		const int32* ChunkPerfectHashSeeds = reinterpret_cast<const int32*>(DataPtr);
-		OutTocResource.ChunkPerfectHashSeeds = MakeArrayView<int32 const>(ChunkPerfectHashSeeds, PerfectHashSeedsCount);
-		DataPtr += PerfectHashSeedsCount * sizeof(int32);
+		OutTocResource.ChunkPerfectHashSeeds = TocStorage.ChopView<int32>(PerfectHashSeedsCount);
 	}
 	if (ChunksWithoutPerfectHashCount)
 	{
-		const int32* ChunkIndicesWithoutPerfectHash = reinterpret_cast<const int32*>(DataPtr);
-		OutTocResource.ChunkIndicesWithoutPerfectHash = MakeArrayView<int32 const>(ChunkIndicesWithoutPerfectHash, ChunksWithoutPerfectHashCount);
-		DataPtr += ChunksWithoutPerfectHashCount * sizeof(int32);
+		OutTocResource.ChunkIndicesWithoutPerfectHash = TocStorage.ChopView<int32>(ChunksWithoutPerfectHashCount);
 	}
 
 	// Compression blocks
-	const FIoStoreTocCompressedBlockEntry* CompressionBlocks = reinterpret_cast<const FIoStoreTocCompressedBlockEntry*>(DataPtr);
-	OutTocResource.CompressionBlocks = MakeArrayView<FIoStoreTocCompressedBlockEntry const>(CompressionBlocks, Header.TocCompressedBlockEntryCount);
-	DataPtr += Header.TocCompressedBlockEntryCount * sizeof(FIoStoreTocCompressedBlockEntry);
+	OutTocResource.CompressionBlocks = TocStorage.ChopView<FIoStoreTocCompressedBlockEntry>(Header.TocCompressedBlockEntryCount);
 
 	// Compression methods
-	OutTocResource.CompressionMethods.Reserve(Header.CompressionMethodNameCount + 1);
-	OutTocResource.CompressionMethods.Add(NAME_None);
-
-	const ANSICHAR* AnsiCompressionMethodNames = reinterpret_cast<const ANSICHAR*>(DataPtr);
-	for (uint32 CompressonNameIndex = 0; CompressonNameIndex < Header.CompressionMethodNameCount; CompressonNameIndex++)
 	{
-		const ANSICHAR* AnsiCompressionMethodName = AnsiCompressionMethodNames + CompressonNameIndex * Header.CompressionMethodNameLength;
-		OutTocResource.CompressionMethods.Add(FName(AnsiCompressionMethodName));
+		OutTocResource.CompressionMethods.Reserve(Header.CompressionMethodNameCount + 1);
+		OutTocResource.CompressionMethods.Add(NAME_None);
+		TArray<ANSICHAR> AnsiCompressionMethodNamesBlock = TocStorage.ChopArray<ANSICHAR>(Header.CompressionMethodNameCount * Header.CompressionMethodNameLength);
+		const ANSICHAR* AnsiCompressionMethodNames = reinterpret_cast<const ANSICHAR*>(AnsiCompressionMethodNamesBlock.GetData());
+		for (uint32 CompressonNameIndex = 0; CompressonNameIndex < Header.CompressionMethodNameCount; CompressonNameIndex++)
+		{
+			const ANSICHAR* AnsiCompressionMethodName = AnsiCompressionMethodNames + CompressonNameIndex * Header.CompressionMethodNameLength;
+			OutTocResource.CompressionMethods.Add(FName(AnsiCompressionMethodName));
+		}
 	}
-	DataPtr += Header.CompressionMethodNameCount * Header.CompressionMethodNameLength;
 
 	// Chunk block signatures
-	const uint8* SignatureBuffer = reinterpret_cast<const uint8*>(DataPtr);
-	const uint8* DirectoryIndexBuffer = SignatureBuffer;
-
 	const bool bIsSigned = EnumHasAnyFlags(Header.ContainerFlags, EIoContainerFlags::Signed);
 	if (IsSigningEnabled() || bIsSigned)
 	{
@@ -1234,23 +1356,27 @@ FIoStatus FIoStoreTocResource::Read(const TCHAR* TocFilePath, EIoStoreTocReadOpt
 			return FIoStatus(EIoErrorCode::SignatureError, TEXT("Missing signature"));
 		}
 
-		const int32* HashSize = reinterpret_cast<const int32*>(SignatureBuffer);
-		TArrayView<const uint8> TocSignature = MakeArrayView<const uint8>(reinterpret_cast<const uint8*>(HashSize + 1), *HashSize);
-		TArrayView<const uint8> BlockSignature = MakeArrayView<const uint8>(TocSignature.GetData() + *HashSize, *HashSize);
+		const int32 HashSize = TocStorage.ChopArray<int32>(1)[0];
 
-		TArrayView<const uint8> BothSignatures = MakeArrayView<const uint8>(TocSignature.GetData(), *HashSize * 2);
+		TArray<uint8> BothSignatures;
+		BothSignatures.Reserve(2 * HashSize);
+		BothSignatures.Append(TocStorage.ChopArray<uint8>(HashSize));
+		BothSignatures.Append(TocStorage.ChopArray<uint8>(HashSize));
 		FSHA1::HashBuffer(BothSignatures.GetData(), BothSignatures.Num(), OutTocResource.SignatureHash.Hash);
 
-		TArrayView<const FSHAHash> ChunkBlockSignatures = MakeArrayView<const FSHAHash>(reinterpret_cast<const FSHAHash*>(BlockSignature.GetData() + *HashSize), Header.TocCompressedBlockEntryCount);
-
-		// Adjust address to meta data
-		DirectoryIndexBuffer = reinterpret_cast<const uint8*>(ChunkBlockSignatures.GetData() + ChunkBlockSignatures.Num());
+		TArrayView<uint8> TocSignature = MakeArrayView(BothSignatures.GetData(), HashSize);
+		TArrayView<uint8> BlockSignature = MakeArrayView(BothSignatures.GetData() + HashSize, HashSize);
 
-		OutTocResource.ChunkBlockSignatures = ChunkBlockSignatures;
+		OutTocResource.ChunkBlockSignatures = TocStorage.ChopView<FSHAHash>(Header.TocCompressedBlockEntryCount);
 
 		if (IsSigningEnabled())
 		{
-			FIoStatus SignatureStatus = ValidateContainerSignature(GetPublicSigningKey(), Header, OutTocResource.ChunkBlockSignatures, TocSignature, BlockSignature);
+			FIoStatus SignatureStatus = ValidateContainerSignature(
+				GetPublicSigningKey(),
+				Header,
+				OutTocResource.ChunkBlockSignatures,
+				TocSignature,
+				BlockSignature);
 			if (!SignatureStatus.IsOk())
 			{
 				return SignatureStatus;
@@ -1263,18 +1389,18 @@ FIoStatus FIoStoreTocResource::Read(const TCHAR* TocFilePath, EIoStoreTocReadOpt
 		EnumHasAnyFlags(Header.ContainerFlags, EIoContainerFlags::Indexed) &&
 		Header.DirectoryIndexSize > 0)
 	{
-		OutTocResource.DirectoryIndexBuffer = MakeArrayView<const uint8>(DirectoryIndexBuffer, Header.DirectoryIndexSize);
+		OutTocResource.DirectoryIndexBuffer = TocStorage.ChopView<uint8>(Header.DirectoryIndexSize);
 	}
 
 	// Meta
 	if (EnumHasAnyFlags(ReadOptions, EIoStoreTocReadOptions::ReadTocMeta))
 	{
-		const uint8* TocMeta = (uint8*)DirectoryIndexBuffer + Header.DirectoryIndexSize;
-
 		if (Header.Version >= static_cast<uint8>(EIoStoreTocVersion::ReplaceIoChunkHashWithIoHash))
 		{
-			const FIoStoreTocEntryMeta* ChunkMetas = reinterpret_cast<const FIoStoreTocEntryMeta*>(TocMeta);
-			OutTocResource.ChunkMetas = MakeArrayView<FIoStoreTocEntryMeta const>(ChunkMetas, Header.TocEntryCount);
+			// TODO adunchev: for now just chop into temp array and move ownership,
+			// but later maybe deprecate "Version < ReplaceIoChunkHashWithIoHash"
+			// and do a proper direct Chop into a view?
+			OutTocResource.ChunkMetas = TocStorage.ChopArray<FIoStoreTocEntryMeta>(Header.TocEntryCount);
 		}
 		else
 		{
@@ -1283,8 +1409,7 @@ FIoStatus FIoStoreTocResource::Read(const TCHAR* TocFilePath, EIoStoreTocReadOpt
 				uint8 ChunkHash[32];
 				FIoStoreTocEntryMetaFlags Flags;
 			};
-			const FIoStoreTocEntryMetaOld* ChunkMetas = reinterpret_cast<const FIoStoreTocEntryMetaOld*>(TocMeta);
-			TConstArrayView<FIoStoreTocEntryMetaOld> OldChunkMetas = MakeArrayView<FIoStoreTocEntryMetaOld const>(ChunkMetas, Header.TocEntryCount);
+			TArray<FIoStoreTocEntryMetaOld> OldChunkMetas = TocStorage.ChopArray<FIoStoreTocEntryMetaOld>(Header.TocEntryCount);
 			OutTocResource.ChunkMetas.Reserve(OldChunkMetas.Num());
 			for (const FIoStoreTocEntryMetaOld& OldChunkMeta : OldChunkMetas)
 			{
@@ -1301,9 +1426,51 @@ FIoStatus FIoStoreTocResource::Read(const TCHAR* TocFilePath, EIoStoreTocReadOpt
 		Header.PartitionSize = MAX_uint64;
 	}
 
+	TocStorage.FinalizeRead();
+	OutTocStorage = MoveTemp(TocStorage);
+
 	return FIoStatus::Ok;
 }
 
+FIoStoreTocChunkInfo FIoStoreTocResource::GetTocChunkInfo(int32 TocEntryIndex) const
+{
+	FIoStoreTocResourceView View = FIoStoreTocResourceView::BuildViewFromResource(*this);
+	return View.GetTocChunkInfo(TocEntryIndex);
+}
+
+FIoStoreTocResource FIoStoreTocResource::BuildResourceFromMappedView(const FIoStoreTocResourceView& View)
+{
+	FIoStoreTocResource Resource;
+	Resource.Header = View.Header;
+	Resource.ChunkIds = View.ChunkIds;
+	Resource.ChunkOffsetLengths = View.ChunkOffsetLengths;
+	Resource.ChunkPerfectHashSeeds = View.ChunkPerfectHashSeeds;
+	Resource.ChunkIndicesWithoutPerfectHash = View.ChunkIndicesWithoutPerfectHash;
+	Resource.CompressionBlocks = View.CompressionBlocks;
+
+	Resource.CompressionMethods = View.CompressionMethods;
+	Resource.SignatureHash = View.SignatureHash;
+	Resource.ChunkBlockSignatures = View.ChunkBlockSignatures;
+	Resource.DirectoryIndexBuffer = View.DirectoryIndexBuffer;
+	Resource.ChunkMetas = View.ChunkMetas;
+
+	return Resource;
+}
+
+FIoStatus FIoStoreTocResource::Read(const TCHAR* TocFilePath, EIoStoreTocReadOptions ReadOptions, FIoStoreTocResource& OutTocResource)
+{
+	FIoStoreTocResourceView View;
+	FIoStoreTocResourceStorage Storage;
+	FIoStatus Status = FIoStoreTocResourceView::Read(TocFilePath, ReadOptions, View, Storage);
+
+	if (Status.IsOk())
+	{
+		OutTocResource = BuildResourceFromMappedView(View);
+	}
+
+	return Status;
+}
+
 TIoStatusOr<uint64> FIoStoreTocResource::Write(
 	const TCHAR* TocFilePath,
 	FIoStoreTocResource& TocResource,
diff --git a/Engine/Source/Runtime/Core/Public/IO/IoDispatcher.h b/Engine/Source/Runtime/Core/Public/IO/IoDispatcher.h
--- a/Engine/Source/Runtime/Core/Public/IO/IoDispatcher.h
+++ b/Engine/Source/Runtime/Core/Public/IO/IoDispatcher.h
@@ -455,7 +455,7 @@ class FIoDirectoryIndexReader
 public:
 	CORE_API FIoDirectoryIndexReader();
 	CORE_API ~FIoDirectoryIndexReader();
-	CORE_API FIoStatus Initialize(TArray<uint8>& InBuffer, FAES::FAESKey InDecryptionKey);
+	CORE_API FIoStatus Initialize(TConstArrayView<uint8> InBuffer, FAES::FAESKey InDecryptionKey);
 
 	CORE_API const FString& GetMountPoint() const;
 	CORE_API FIoDirectoryIndexHandle GetChildDirectory(FIoDirectoryIndexHandle Directory) const;
diff --git a/Engine/Source/Runtime/PakFile/Internal/IoDispatcherFileBackendTypes.h b/Engine/Source/Runtime/PakFile/Internal/IoDispatcherFileBackendTypes.h
--- a/Engine/Source/Runtime/PakFile/Internal/IoDispatcherFileBackendTypes.h
+++ b/Engine/Source/Runtime/PakFile/Internal/IoDispatcherFileBackendTypes.h
@@ -40,27 +40,6 @@ struct FFileIoStoreContainerFilePartition
 	std::atomic<int32> StartedReadRequestsCount = 0;
 };
 
-struct FFileIoStoreBlockSignatureTable
-{
-	TArray<FSHAHash> Hashes;
-
-	void AddRef()
-	{
-		++RefCount;
-	}
-
-	void Release()
-	{
-		if (--RefCount == 0)
-		{
-			delete this;
-		}
-	}
-
-private:
-	std::atomic<int64> RefCount = 0;
-};
-
 struct FFileIoStoreContainerFile
 {
 	FFileIoStoreContainerFile() = default;
@@ -73,12 +52,12 @@ struct FFileIoStoreContainerFile
 	uint64 PartitionSize = 0;
 	uint64 CompressionBlockSize = 0;
 	TArray<FName> CompressionMethods;
-	TArray<FIoStoreTocCompressedBlockEntry> CompressionBlocks;
+	TConstArrayView<FIoStoreTocCompressedBlockEntry> CompressionBlocks;
 	FString FilePath;
 	FGuid EncryptionKeyGuid;
 	FAES::FAESKey EncryptionKey;
 	EIoContainerFlags ContainerFlags;
-	TRefCountPtr<FFileIoStoreBlockSignatureTable> BlockSignatureTable;
+	TConstArrayView<FSHAHash> BlockSignatureTable;
 	TArray<FFileIoStoreContainerFilePartition> Partitions;
 	uint32 ContainerInstanceId = 0;
 
@@ -144,7 +123,7 @@ struct FFileIoStoreCompressedBlock
 	FFileIoStoreCompressionContext* CompressionContext = nullptr;
 	uint8* CompressedDataBuffer = nullptr;
 	FAES::FAESKey EncryptionKey;
-	TRefCountPtr<FFileIoStoreBlockSignatureTable> BlockSignatureTable;
+	TConstArrayView<FSHAHash> BlockSignatureTable;
 	const FSHAHash* SignatureHash = nullptr;
 	bool bFailed = false;
 	bool bCancelled = false;
diff --git a/Engine/Source/Runtime/PakFile/Private/IPlatformFilePak.cpp b/Engine/Source/Runtime/PakFile/Private/IPlatformFilePak.cpp
--- a/Engine/Source/Runtime/PakFile/Private/IPlatformFilePak.cpp
+++ b/Engine/Source/Runtime/PakFile/Private/IPlatformFilePak.cpp
@@ -100,8 +100,9 @@ static bool CheckIoStoreContainerBlockSignatures(const TCHAR* InContainerPath)
 	UE_LOG(LogPakFile, Display, TEXT("Checking container file \"%s\"..."), InContainerPath);
 	double StartTime = FPlatformTime::Seconds();
 
-	FIoStoreTocResource TocResource;
-	FIoStatus Status = FIoStoreTocResource::Read(InContainerPath, EIoStoreTocReadOptions::Default, TocResource);
+	FIoStoreTocResourceView TocResource;
+	FIoStoreTocResourceStorage TocResourceStorage;
+	FIoStatus Status = FIoStoreTocResourceView::Read(InContainerPath, EIoStoreTocReadOptions::Default, TocResource, TocResourceStorage);
 	if (!Status.IsOk())
 	{
 		UE_LOG(LogPakFile, Error, TEXT("Failed reading toc file \"%s\"."), InContainerPath);
diff --git a/Engine/Source/Runtime/PakFile/Private/IoDispatcherFileBackend.cpp b/Engine/Source/Runtime/PakFile/Private/IoDispatcherFileBackend.cpp
--- a/Engine/Source/Runtime/PakFile/Private/IoDispatcherFileBackend.cpp
+++ b/Engine/Source/Runtime/PakFile/Private/IoDispatcherFileBackend.cpp
@@ -726,12 +726,7 @@ FFileIoStoreReader::~FFileIoStoreReader()
 
 uint64 FFileIoStoreReader::GetTocAllocatedSize() const
 {
-	return TocImperfectHashMapFallback.GetAllocatedSize() +
-		PerfectHashMap.TocOffsetAndLengths.GetAllocatedSize() +
-		PerfectHashMap.TocChunkIds.GetAllocatedSize() +
-		PerfectHashMap.TocChunkHashSeeds.GetAllocatedSize() +
-		ContainerFile.CompressionBlocks.GetAllocatedSize() +
-		(ContainerFile.BlockSignatureTable.IsValid() ? ContainerFile.BlockSignatureTable->Hashes.GetAllocatedSize() : 0);
+	return DataContainer.GetAllocatedSize() + TocImperfectHashMapFallback.GetAllocatedSize();
 }
 
 FIoStatus FFileIoStoreReader::Initialize(const TCHAR* InTocFilePath, int32 InOrder)
@@ -750,9 +745,9 @@ FIoStatus FFileIoStoreReader::Initialize(const TCHAR* InTocFilePath, int32 InOrd
 
 	UE_LOG(LogIoDispatcher, Display, TEXT("Reading toc: %s"), InTocFilePath);
 
-	TUniquePtr<FIoStoreTocResource> TocResourcePtr = MakeUnique<FIoStoreTocResource>();
-	FIoStoreTocResource& TocResource = *TocResourcePtr;
-	FIoStatus Status = FIoStoreTocResource::Read(InTocFilePath, EIoStoreTocReadOptions::Default, TocResource);
+	FIoStoreTocResourceView TocResource;
+	FIoStoreTocResourceStorage TocStorage;
+	FIoStatus Status = FIoStoreTocResourceView::Read(InTocFilePath, EIoStoreTocReadOptions::Default, TocResource, TocStorage);
 	if (!Status.IsOk())
 	{
 		return Status;
@@ -805,15 +800,16 @@ FIoStatus FFileIoStoreReader::Initialize(const TCHAR* InTocFilePath, int32 InOrd
 	ContainerFile.CompressionBlocks		= MoveTemp(TocResource.CompressionBlocks);
 	ContainerFile.ContainerFlags		= TocResource.Header.ContainerFlags;
 	ContainerFile.EncryptionKeyGuid		= TocResource.Header.EncryptionKeyGuid;
-	ContainerFile.BlockSignatureTable	= new FFileIoStoreBlockSignatureTable();
-	ContainerFile.BlockSignatureTable->Hashes = MoveTemp(TocResource.ChunkBlockSignatures);
+	ContainerFile.BlockSignatureTable	= MoveTemp(TocResource.ChunkBlockSignatures);
 	ContainerFile.ContainerInstanceId	= ++GlobalContainerInstanceId;
 
 	Stats.OnTocMounted(GetTocAllocatedSize());
 
 	UE_LOG(LogIoDispatcher, Display, TEXT("Toc loaded : %s, Id=%s, Order=%d, EntryCount=%u, SignatureHash=%s"),
 		InTocFilePath, *LexToString(TocResource.Header.ContainerId), InOrder, TocResource.Header.TocEntryCount, *TocResource.SignatureHash.ToString());
-	
+
+	// Hold onto the handle and mapped regions
+	DataContainer = MoveTemp(TocStorage);
 	ContainerId = TocResource.Header.ContainerId;
 	Order = InOrder;
 	return FIoStatus::Ok;
@@ -833,11 +829,10 @@ FIoStatus FFileIoStoreReader::Close()
 
 	Stats.OnTocUnmounted(GetTocAllocatedSize());
 
-	PerfectHashMap.TocChunkHashSeeds.Empty();
-	PerfectHashMap.TocChunkIds.Empty();
-	PerfectHashMap.TocOffsetAndLengths.Empty();
+	PerfectHashMap = FFileIoStoreReader::FPerfectHashMap();
 	TocImperfectHashMapFallback.Empty();
 	ContainerFile = FFileIoStoreContainerFile();
+	DataContainer = FIoStoreTocResourceStorage();
 	ContainerId = FIoContainerId();
 	Order = INDEX_NONE;
 	bClosed = true;
@@ -999,7 +994,7 @@ TIoStatusOr<FIoContainerHeader> FFileIoStoreReader::ReadContainerHeader() const
 			const uint64 BlockSize = Align(CompressionBlockEntry->GetCompressedSize(), FAES::AESBlockSize);
 			if (bSigned)
 			{
-				const FSHAHash& SignatureHash = ContainerFile.BlockSignatureTable->Hashes[CompressedBlockIndex];
+				const FSHAHash& SignatureHash = ContainerFile.BlockSignatureTable[CompressedBlockIndex];
 				FSHAHash BlockHash;
 				FSHA1::HashBuffer(BlockData, BlockSize, BlockHash.Hash);
 				if (SignatureHash != BlockHash)
@@ -2121,9 +2116,9 @@ void FFileIoStore::ReadBlocks(FFileIoStoreResolvedRequest& ResolvedRequest)
 			CompressedBlock->CompressionMethod = ContainerFile->CompressionMethods[CompressionBlockEntry.GetCompressionMethodIndex()];
 			if (EnumHasAnyFlags(ContainerFile->ContainerFlags, EIoContainerFlags::Signed))
 			{
-				check(ContainerFile->BlockSignatureTable);
+				check(!ContainerFile->BlockSignatureTable.IsEmpty());
 				CompressedBlock->BlockSignatureTable = ContainerFile->BlockSignatureTable;
-				CompressedBlock->SignatureHash = &ContainerFile->BlockSignatureTable->Hashes[CompressedBlockIndex];
+				CompressedBlock->SignatureHash = &ContainerFile->BlockSignatureTable[CompressedBlockIndex];
 			}
 			CompressedBlock->RawSize = Align(CompressionBlockEntry.GetCompressedSize(), FAES::AESBlockSize); // The raw blocks size is always aligned to AES blocks size;
 
diff --git a/Engine/Source/Runtime/PakFile/Private/IoDispatcherFileBackend.h b/Engine/Source/Runtime/PakFile/Private/IoDispatcherFileBackend.h
--- a/Engine/Source/Runtime/PakFile/Private/IoDispatcherFileBackend.h
+++ b/Engine/Source/Runtime/PakFile/Private/IoDispatcherFileBackend.h
@@ -59,14 +59,16 @@ private:
 
 	struct FPerfectHashMap
 	{
-		TArray<int32> TocChunkHashSeeds;
-		TArray<FIoChunkId> TocChunkIds;
-		TArray<FIoOffsetAndLength> TocOffsetAndLengths;
+		TConstArrayView<int32> TocChunkHashSeeds;
+		TConstArrayView<FIoChunkId> TocChunkIds;
+		TConstArrayView<FIoOffsetAndLength> TocOffsetAndLengths;
 	};
 	FPerfectHashMap PerfectHashMap;
 	TMap<FIoChunkId, FIoOffsetAndLength> TocImperfectHashMapFallback;
 	FFileIoStoreContainerFile ContainerFile;
 	FIoContainerId ContainerId;
+	// Owns the data for ArrayViews in ContainerFile and in PerfectHash map
+	FIoStoreTocResourceStorage DataContainer;
 	int32 Order;
 	bool bClosed = false;
 	bool bHasPerfectHashMap = false;
