Fix the build breakage with missing `autortfm_check_abi`.

#rnx

diff --git a/Engine/Source/Runtime/Core/Public/AutoRTFM/AutoRTFM.h b/Engine/Source/Runtime/Core/Public/AutoRTFM/AutoRTFM.h
--- a/Engine/Source/Runtime/Core/Public/AutoRTFM/AutoRTFM.h
+++ b/Engine/Source/Runtime/Core/Public/AutoRTFM/AutoRTFM.h
@@ -356,6 +356,22 @@ static UE_AUTORTFM_FORCEINLINE void autortfm_did_free(void* ptr)
 }
 #endif
 
+// If running with AutoRTFM enabled, then perform an ABI check between the
+// AutoRTFM compiler and the AutoRTFM runtime, to ensure that memory is being
+// laid out in an identical manner between the AutoRTFM runtime and the AutoRTFM
+// compiler pass. Should not be called manually by the user, a call to this will
+// be injected by the compiler into a global constructor in the AutoRTFM compiled
+// code.
+#if UE_AUTORTFM_ENABLED
+UE_AUTORTFM_API void autortfm_check_abi(void* ptr, size_t size);
+#else
+static UE_AUTORTFM_FORCEINLINE void autortfm_check_abi(void* ptr, size_t size)
+{
+	UE_AUTORTFM_UNUSED(ptr);
+	UE_AUTORTFM_UNUSED(size);
+}
+#endif
+
 #if UE_AUTORTFM_ENABLED
 UE_AUTORTFM_API void autortfm_unreachable(const UE_AUTORTFM_TCHAR* const Message);
 #else
