Refactored TFuture classes to remove lots of duplication through specializations.
Fixed lack of dereferencing in TFuture::GetMutable() and TFuture::Consume().
UE_LIFETIMEBOUND removed for now as they will need to be split per future specialization, and will be restored later.

diff --git a/Engine/Source/Runtime/Core/Public/Async/Future.h b/Engine/Source/Runtime/Core/Public/Async/Future.h
--- a/Engine/Source/Runtime/Core/Public/Async/Future.h
+++ b/Engine/Source/Runtime/Core/Public/Async/Future.h
@@ -12,6 +12,23 @@
 #include "HAL/Event.h"
 #include "HAL/PlatformProcess.h"
 #include "Misc/ScopeLock.h"
+#include "Templates/Requires.h"
+#include <type_traits>
+
+namespace UE::Core::Private
+{
+	template <typename T>
+	using TInternalFutureState_T =
+		std::conditional_t<
+			std::is_lvalue_reference_v<T>,
+			std::remove_reference_t<T>*,
+			std::conditional_t<
+				std::is_void_v<T>,
+				int,
+				T
+			>
+		>;
+}
 
 /**
  * Base class for the internal state of asynchronous return values (futures).
@@ -379,9 +396,9 @@ template<typename ResultType> class TSharedFuture;
  */
 template<typename ResultType>
 class TFuture
-	: public TFutureBase<ResultType>
+	: public TFutureBase<UE::Core::Private::TInternalFutureState_T<ResultType>>
 {
-	typedef TFutureBase<ResultType> BaseType;
+	using BaseType = TFutureBase<UE::Core::Private::TInternalFutureState_T<ResultType>>;
 
 public:
 
@@ -420,9 +437,20 @@ public:
 	 * @return The result.
 	 * @note Not equivalent to std::future::get(). The future remains valid.
 	 */
-	const ResultType& Get() const UE_LIFETIMEBOUND
+	decltype(auto) Get() const //UE_LIFETIMEBOUND
 	{
-		return this->GetState()->GetResult();
+		if constexpr (std::is_lvalue_reference_v<ResultType>)
+		{
+			return (ResultType&)*this->GetState()->GetResult();
+		}
+		else if constexpr (std::is_void_v<ResultType>)
+		{
+			this->GetState()->GetResult();
+		}
+		else
+		{
+			return (const ResultType&)this->GetState()->GetResult();
+		}
 	}
 
 	/**
@@ -431,9 +459,20 @@ public:
 	 * @return The result.
 	 * @note Not equivalent to std::future::get(). The future remains valid.
 	 */
-	ResultType& GetMutable() UE_LIFETIMEBOUND
+	decltype(auto) GetMutable() // UE_LIFETIMEBOUND
 	{
-		return this->GetState()->GetResult();
+		if constexpr (std::is_lvalue_reference_v<ResultType>)
+		{
+			return (ResultType&)*this->GetState()->GetResult();
+		}
+		else if constexpr (std::is_void_v<ResultType>)
+		{
+			this->GetState()->GetResult();
+		}
+		else
+		{
+			return (ResultType&)this->GetState()->GetResult();
+		}
 	}
 
 	/**
@@ -442,10 +481,21 @@ public:
 	 * @return The result.
 	 * @note Equivalent to std::future::get(). Invalidates the future.
 	 */
-	ResultType Consume()
+	decltype(auto) Consume()
 	{
-		TFuture<ResultType> Local(MoveTemp(*this));
-		return MoveTemp(Local.GetState()->GetResult());
+		TFuture Local(MoveTemp(*this));
+		if constexpr (std::is_lvalue_reference_v<ResultType>)
+		{
+			return (ResultType&)*Local.GetState()->GetResult();
+		}
+		else if constexpr (std::is_void_v<ResultType>)
+		{
+			Local.GetState()->GetResult();
+		}
+		else
+		{
+			return (ResultType)MoveTemp(Local.GetState()->GetResult());
+		}
 	}
 
 	/**
@@ -478,213 +528,6 @@ public:
 };
 
 
-/**
- * Template for unshared futures (specialization for reference types).
- */
-template<typename ResultType>
-class TFuture<ResultType&>
-	: public TFutureBase<ResultType*>
-{
-	typedef TFutureBase<ResultType*> BaseType;
-
-public:
-
-	/** Default constructor. */
-	TFuture() = default;
-
-	/**
-	 * Creates and initializes a new instance.
-	 *
-	 * @param InState The shared state to initialize with.
-	 */
-	TFuture(const typename BaseType::StateType& InState)
-		: BaseType(InState)
-	{ }
-
-	/** Deleted copy constructor (futures cannot be copied). */
-	TFuture(const TFuture&) = delete;
-
-	/** Deleted copy assignment operator (futures cannot be copied). */
-	TFuture& operator=(const TFuture&) = delete;
-
-	/** Move constructor. */
-	TFuture(TFuture&&) = default;
-
-	/** Move assignment operator. */
-	TFuture& operator=(TFuture&&) = default;
-
-	/** Destructor. */
-	~TFuture() = default;
-
-public:
-
-	/**
-	 * Gets the future's result.
-	 *
-	 * @return The result.
-	 * @note Not equivalent to std::future::get(). The future remains valid.
-	 */
-	ResultType& Get() const
-	{
-		return *this->GetState()->GetResult();
-	}
-
-	/**
-	 * Gets the future's result.
-	 *
-	 * @return The result.
-	 * @note Not equivalent to std::future::get(). The future remains valid.
-	 */
-	ResultType& GetMutable()
-	{
-		return this->GetState()->GetResult();
-	}
-
-	/**
-	 * Consumes the future's result and invalidates the future.
-	 *
-	 * @return The result.
-	 * @note Equivalent to std::future::get(). Invalidates the future.
-	 */
-	ResultType& Consume()
-	{
-		TFuture<ResultType&> Local(MoveTemp(*this));
-		return Local.GetState()->GetResult();
-	}
-
-	/**
-	 * Moves this future's state into a shared future.
-	 *
-	 * @return The shared future object.
-	 */
-	TSharedFuture<ResultType&> Share()
-	{
-		return TSharedFuture<ResultType&>(MoveTemp(*this));
-	}
-
-	/**
-	 * Expose Then functionality
-	 * @see TFutureBase
-	 */
-	using BaseType::Then;
-
-	/**
-	 * Expose Next functionality
-	 * @see TFutureBase
-	 */
-	using BaseType::Next;
-
-	/**
-	 * Expose Reset functionality
-	 * @see TFutureBase
-	 */
-	using BaseType::Reset;
-};
-
-
-/**
- * Template for unshared futures (specialization for void).
- */
-template<>
-class TFuture<void>
-	: public TFutureBase<int>
-{
-	typedef TFutureBase<int> BaseType;
-
-public:
-
-	/** Default constructor. */
-	TFuture() = default;
-
-	/**
-	 * Creates and initializes a new instance.
-	 *
-	 * @param InState The shared state to initialize with.
-	 */
-	TFuture(const BaseType::StateType& InState)
-		: BaseType(InState)
-	{ }
-
-	/** Deleted copy constructor (futures cannot be copied). */
-	TFuture(const TFuture&) = delete;
-
-	/** Deleted copy assignment operator (futures cannot be copied). */
-	TFuture& operator=(const TFuture&) = delete;
-
-	/** Move constructor. */
-	TFuture(TFuture&&) = default;
-
-	/** Move assignment operator. */
-	TFuture& operator=(TFuture&&) = default;
-
-	/** Destructor. */
-	~TFuture() = default;
-
-public:
-
-	/**
-	 * Gets the future's result.
-	 *
-	 * @return The result.
-	 * @note Not equivalent to std::future::get(). The future remains valid.
-	 */
-	void Get() const
-	{
-		// Wait for the dummy result to be available.
-		GetState()->GetResult();
-	}
-
-	/**
-	 * Gets the future's result.
-	 *
-	 * @return The result.
-	 * @note Not equivalent to std::future::get(). The future remains valid.
-	 */
-	void GetMutable()
-	{
-		// Wait for the dummy result to be available.
-		GetState()->GetResult();
-	}
-
-	/**
-	 * Consumes the future's result and invalidates the future.
-	 *
-	 * @return The result.
-	 * @note Equivalent to std::future::get(). Invalidates the future.
-	 */
-	void Consume()
-	{
-		TFuture<void> Local(MoveTemp(*this));
-		// Wait for the dummy result to be available.
-		Local.GetState()->GetResult();
-	}
-
-	/**
-	 * Moves this future's state into a shared future.
-	 *
-	 * @return The shared future object.
-	 */
-	TSharedFuture<void> Share();
-
-	/**
-	 * Expose Then functionality
-	 * @see TFutureBase
-	 */
-	using BaseType::Then;
-
-	/**
-	 * Expose Next functionality
-	 * @see TFutureBase
-	 */
-	using BaseType::Next;
-
-	/**
-	 * Expose Reset functionality
-	 * @see TFutureBase
-	 */
-	using BaseType::Reset;
-};
-
 /* TSharedFuture
 *****************************************************************************/
 
@@ -693,9 +536,9 @@ public:
  */
 template<typename ResultType>
 class TSharedFuture
-	: public TFutureBase<ResultType>
+	: public TFutureBase<UE::Core::Private::TInternalFutureState_T<ResultType>>
 {
-	typedef TFutureBase<ResultType> BaseType;
+	using BaseType = TFutureBase<UE::Core::Private::TInternalFutureState_T<ResultType>>;
 
 public:
 
@@ -742,141 +585,24 @@ public:
 	 *
 	 * @return The result.
 	 */
-	const ResultType& Get() const
-	{
-		return this->GetState()->GetResult();
-	}
-};
-
-
-/**
- * Template for shared futures (specialization for reference types).
- */
-template<typename ResultType>
-class TSharedFuture<ResultType&>
-	: public TFutureBase<ResultType*>
-{
-	typedef TFutureBase<ResultType*> BaseType;
-
-public:
-
-	/** Default constructor. */
-	TSharedFuture() = default;
-
-	/**
-	 * Creates and initializes a new instance.
-	 *
-	 * @param InState The shared state to initialize from.
-	 */
-	TSharedFuture(const typename BaseType::StateType& InState)
-		: BaseType(InState)
-	{ }
-
-	/**
-	* Creates and initializes a new instances from a future object.
-	*
-	* @param Future The future object to initialize from.
-	*/
-	TSharedFuture(TFuture<ResultType>&& Future)
-		: BaseType(MoveTemp(Future))
-	{ }
-
-	/** Copy constructor. */
-	TSharedFuture(const TSharedFuture&) = default;
-
-	/** Copy assignment operator. */
-	TSharedFuture& operator=(const TSharedFuture&) = default;
-
-	/** Move constructor. */
-	TSharedFuture(TSharedFuture&&) = default;
-
-	/** Move assignment operator. */
-	TSharedFuture& operator=(TSharedFuture&&) = default;
-
-	/** Destructor. */
-	~TSharedFuture() = default;
-
-public:
-
-	/**
-	 * Gets the future's result.
-	 *
-	 * @return The result.
-	 */
-	ResultType& Get() const
+	decltype(auto) Get() const // UE_LIFETIMEBOUND
 	{
-		return *this->GetState()->GetResult();
-	}
-};
-
-
-/**
- * Template for shared futures (specialization for void).
- */
-template<>
-class TSharedFuture<void>
-	: public TFutureBase<int>
-{
-	typedef TFutureBase<int> BaseType;
-
-public:
-
-	/** Default constructor. */
-	TSharedFuture() = default;
-
-	/**
-	 * Creates and initializes a new instance from shared state.
-	 *
-	 * @param InState The shared state to initialize from.
-	 */
-	TSharedFuture(const BaseType::StateType& InState)
-		: BaseType(InState)
-	{ }
-
-	/**
-	 * Creates and initializes a new instances from a future object.
-	 *
-	 * @param Future The future object to initialize from.
-	 */
-	TSharedFuture(TFuture<void>&& Future)
-		: BaseType(MoveTemp(Future))
-	{ }
-
-	/** Copy constructor. */
-	TSharedFuture(const TSharedFuture&) = default;
-
-	/** Copy assignment operator. */
-	TSharedFuture& operator=(const TSharedFuture&) = default;
-
-	/** Move constructor. */
-	TSharedFuture(TSharedFuture&&) = default;
-
-	/** Move assignment operator. */
-	TSharedFuture& operator=(TSharedFuture&&) = default;
-
-	/** Destructor. */
-	~TSharedFuture() = default;
-
-public:
-
-	/**
-	 * Gets the future's result.
-	 *
-	 * @return The result.
-	 */
-	void Get() const
-	{
-		GetState()->GetResult();
+		if constexpr (std::is_lvalue_reference_v<ResultType>)
+		{
+			return (ResultType&)*this->GetState()->GetResult();
+		}
+		else if constexpr (std::is_void_v<ResultType>)
+		{
+			this->GetState()->GetResult();
+		}
+		else
+		{
+			return (const ResultType&)this->GetState()->GetResult();
+		}
 	}
 };
 
 
-inline TSharedFuture<void> TFuture<void>::Share()
-{
-	return TSharedFuture<void>(MoveTemp(*this));
-}
-
-
 /* TPromise
 *****************************************************************************/
 
@@ -962,11 +688,14 @@ private:
  */
 template<typename ResultType>
 class TPromise
-	: public TPromiseBase<ResultType>
+	: public TPromiseBase<UE::Core::Private::TInternalFutureState_T<ResultType>>
 {
-public:
+	using BaseType = TPromiseBase<UE::Core::Private::TInternalFutureState_T<ResultType>>;
 
-	typedef TPromiseBase<ResultType> BaseType;
+	// This is necessary because we can't form references to void for the parameter of Set().
+	using SetType = std::conditional_t<std::is_void_v<ResultType>, int, ResultType>;
+
+public:
 
 	/** Default constructor (creates a new shared state). */
 	TPromise()
@@ -1032,23 +761,45 @@ public:
 	 *
 	 * @param Result The result value to set.
 	 */
-	FORCEINLINE void SetValue(const ResultType& Result)
+#ifndef __cpp_concepts
+	template <
+		typename DummyResultType = ResultType
+		UE_REQUIRES(!std::is_void_v<DummyResultType>)
+	>
+#endif
+	FORCEINLINE void SetValue(const SetType& Result)
+#ifdef __cpp_concepts
+		requires(!std::is_void_v<ResultType>)
+#endif
 	{
 		EmplaceValue(Result);
 	}
-
-	/**
-	 * Sets the promised result (from rvalue).
-	 *
-	 * The result must be set only once. An assertion will
-	 * be triggered if this method is called a second time.
-	 *
-	 * @param Result The result value to set.
-	 */
-	FORCEINLINE void SetValue(ResultType&& Result)
+#ifndef __cpp_concepts
+	template <
+		typename DummyResultType = ResultType
+		UE_REQUIRES(!std::is_void_v<DummyResultType> && !std::is_lvalue_reference_v<DummyResultType>)
+	>
+#endif
+	FORCEINLINE void SetValue(SetType&& Result)
+#ifdef __cpp_concepts
+		requires(!std::is_void_v<ResultType> && !std::is_lvalue_reference_v<ResultType>)
+#endif
 	{
 		EmplaceValue(MoveTemp(Result));
 	}
+#ifndef __cpp_concepts
+	template <
+		typename DummyResultType = ResultType
+		UE_REQUIRES(std::is_void_v<DummyResultType>)
+	>
+#endif
+	FORCEINLINE void SetValue()
+#ifdef __cpp_concepts
+		requires(std::is_void_v<ResultType>)
+#endif
+	{
+		EmplaceValue();
+	}
 
 	/**
 	 * Sets the promised result.
@@ -1058,204 +809,43 @@ public:
 	 *
 	 * @param Args The arguments to forward to the constructor of the result.
 	 */
-	template<typename... ArgTypes>
+	template <
+		typename... ArgTypes
+#ifndef __cpp_concepts
+		, typename DummyResultType = ResultType
+		UE_REQUIRES(!std::is_void_v<DummyResultType> && !std::is_lvalue_reference_v<DummyResultType>)
+#endif
+	>
 	void EmplaceValue(ArgTypes&&... Args)
+#ifdef __cpp_concepts
+		requires (!std::is_void_v<ResultType> && !std::is_lvalue_reference_v<ResultType>)
+#endif
 	{
 		this->GetState()->EmplaceResult(Forward<ArgTypes>(Args)...);
 	}
-
-private:
-
-	/** Whether a future has already been retrieved from this promise. */
-	bool FutureRetrieved;
-};
-
-
-/**
- * Template for promises (specialization for reference types).
- */
-template<typename ResultType>
-class TPromise<ResultType&>
-	: public TPromiseBase<ResultType*>
-{
-	typedef TPromiseBase<ResultType*> BaseType;
-
-public:
-
-	/** Default constructor (creates a new shared state). */
-	TPromise()
-		: BaseType()
-		, FutureRetrieved(false)
-	{ }
-
-	/**
-	 * Move constructor.
-	 *
-	 * @param Other The promise holding the shared state to move.
-	 */
-	TPromise(TPromise&& Other)
-		: BaseType(MoveTemp(Other))
-		, FutureRetrieved(MoveTemp(Other.FutureRetrieved))
-	{ }
-
-	/**
-	 * Create and initialize a new instance with a callback.
-	 *
-	 * @param CompletionCallback A function that is called when the future state is completed.
-	 */
-	TPromise(TUniqueFunction<void()>&& CompletionCallback)
-		: BaseType(MoveTemp(CompletionCallback))
-		, FutureRetrieved(false)
-	{ }
-
-public:
-
-	/**
-	 * Move assignment operator.
-	 *
-	 * @param Other The promise holding the shared state to move.
-	 */
-	TPromise& operator=(TPromise&& Other)
-	{
-		BaseType::operator=(MoveTemp(Other));
-		FutureRetrieved = MoveTemp(Other.FutureRetrieved);
-
-		return this;
-	}
-
-public:
-
-	/**
-	 * Gets a TFuture object associated with the shared state of this promise.
-	 *
-	 * @return The TFuture object.
-	 */
-	TFuture<ResultType&> GetFuture()
-	{
-		check(!FutureRetrieved);
-		FutureRetrieved = true;
-
-		return TFuture<ResultType&>(this->GetState());
-	}
-
-	/**
-	 * Sets the promised result.
-	 *
-	 * The result must be set only once. An assertion will
-	 * be triggered if this method is called a second time.
-	 *
-	 * @param Result The result value to set.
-	 */
-	FORCEINLINE void SetValue(ResultType& Result)
-	{
-		EmplaceValue(Result);
-	}
-
-	/**
-	 * Sets the promised result.
-	 *
-	 * The result must be set only once. An assertion will
-	 * be triggered if this method is called a second time.
-	 *
-	 * @param Result The result value to set.
-	 */
-	void EmplaceValue(ResultType& Result)
+#ifndef __cpp_concepts
+	template <
+		typename DummyResultType = ResultType
+		UE_REQUIRES(std::is_lvalue_reference_v<DummyResultType>)
+	>
+#endif
+	void EmplaceValue(SetType Result)
+#ifdef __cpp_concepts
+		requires (std::is_lvalue_reference_v<ResultType>)
+#endif
 	{
 		this->GetState()->EmplaceResult(&Result);
 	}
-
-private:
-
-	/** Whether a future has already been retrieved from this promise. */
-	bool FutureRetrieved;
-};
-
-
-/**
- * Template for promises (specialization for void results).
- */
-template<>
-class TPromise<void>
-	: public TPromiseBase<int>
-{
-	typedef TPromiseBase<int> BaseType;
-
-public:
-
-	/** Default constructor (creates a new shared state). */
-	TPromise()
-		: BaseType()
-		, FutureRetrieved(false)
-	{ }
-
-	/**
-	 * Move constructor.
-	 *
-	 * @param Other The promise holding the shared state to move.
-	 */
-	TPromise(TPromise&& Other)
-		: BaseType(MoveTemp(Other))
-		, FutureRetrieved(false)
-	{ }
-
-	/**
-	 * Create and initialize a new instance with a callback.
-	 *
-	 * @param CompletionCallback A function that is called when the future state is completed.
-	 */
-	TPromise(TUniqueFunction<void()>&& CompletionCallback)
-		: BaseType(MoveTemp(CompletionCallback))
-		, FutureRetrieved(false)
-	{ }
-
-public:
-
-	/**
-	 * Move assignment operator.
-	 *
-	 * @param Other The promise holding the shared state to move.
-	 */
-	TPromise& operator=(TPromise&& Other)
-	{
-		BaseType::operator=(MoveTemp(Other));
-		FutureRetrieved = MoveTemp(Other.FutureRetrieved);
-
-		return *this;
-	}
-
-public:
-
-	/**
-	 * Gets a TFuture object associated with the shared state of this promise.
-	 *
-	 * @return The TFuture object.
-	 */
-	TFuture<void> GetFuture()
-	{
-		check(!FutureRetrieved);
-		FutureRetrieved = true;
-
-		return TFuture<void>(GetState());
-	}
-
-	/**
-	 * Sets the promised result.
-	 *
-	 * The result must be set only once. An assertion will
-	 * be triggered if this method is called a second time.
-	 */
-	FORCEINLINE void SetValue()
-	{
-		EmplaceValue();
-	}
-
-	/**
-	 * Sets the promised result.
-	 *
-	 * The result must be set only once. An assertion will
-	 * be triggered if this method is called a second time.
-	 */
+#ifndef __cpp_concepts
+	template <
+		typename DummyResultType = ResultType
+		UE_REQUIRES(std::is_void_v<DummyResultType>)
+	>
+#endif
 	void EmplaceValue()
+#ifdef __cpp_concepts
+		requires (std::is_void_v<ResultType>)
+#endif
 	{
 		this->GetState()->EmplaceResult(0);
 	}
