Removed ESPMode::ThreadSafe, since that's now the default.

diff --git a/Engine/Source/Runtime/Core/Public/Async/Future.h b/Engine/Source/Runtime/Core/Public/Async/Future.h
--- a/Engine/Source/Runtime/Core/Public/Async/Future.h
+++ b/Engine/Source/Runtime/Core/Public/Async/Future.h
@@ -302,7 +302,7 @@ public:
 
 protected:
 
-	typedef TSharedPtr<TFutureState<InternalResultType>, ESPMode::ThreadSafe> StateType;
+	typedef TSharedPtr<TFutureState<InternalResultType>> StateType;
 
 	/** Default constructor. */
 	TFutureBase() = default;
@@ -610,13 +610,13 @@ template<typename InternalResultType>
 class TPromiseBase
 	: FNoncopyable
 {
-	typedef TSharedPtr<TFutureState<InternalResultType>, ESPMode::ThreadSafe> StateType;
+	typedef TSharedPtr<TFutureState<InternalResultType>> StateType;
 
 public:
 
 	/** Default constructor. */
 	TPromiseBase()
-		: State(MakeShared<TFutureState<InternalResultType>, ESPMode::ThreadSafe>())
+		: State(MakeShared<TFutureState<InternalResultType>>())
 	{ }
 
 	/**
@@ -636,7 +636,7 @@ public:
 	 * @param CompletionCallback A function that is called when the future state is completed.
 	 */
 	TPromiseBase(TUniqueFunction<void()>&& CompletionCallback)
-		: State(MakeShared<TFutureState<InternalResultType>, ESPMode::ThreadSafe>(MoveTemp(CompletionCallback)))
+		: State(MakeShared<TFutureState<InternalResultType>>(MoveTemp(CompletionCallback)))
 	{ }
 
 public:
