Made some rule of zero/five fixes.

diff --git a/Engine/Source/Runtime/Core/Public/Async/Future.h b/Engine/Source/Runtime/Core/Public/Async/Future.h
--- a/Engine/Source/Runtime/Core/Public/Async/Future.h
+++ b/Engine/Source/Runtime/Core/Public/Async/Future.h
@@ -316,21 +316,6 @@ protected:
 		: State(InState)
 	{ }
 
-	/** Protected copy constructor. */
-	TFutureBase(const TFutureBase&) = default;
-
-	/** Protected copy assignment operator. */
-	TFutureBase& operator=(const TFutureBase&) = default;
-
-	/** Protected move constructor. */
-	TFutureBase(TFutureBase&&) = default;
-
-	/** Protected move assignment operator. */
-	TFutureBase& operator=(TFutureBase&&) = default;
-
-	/** Protected destructor. */
-	~TFutureBase() = default;
-
 protected:
 
 	/**
@@ -414,19 +399,11 @@ public:
 		: BaseType(InState)
 	{ }
 
-	/** Deleted copy constructor (futures cannot be copied). */
-	TFuture(const TFuture&) = delete;
-
-	/** Deleted copy assignment operator (futures cannot be copied). */
-	TFuture& operator=(const TFuture&) = delete;
-
-	/** Move constructor. */
+	// Movable-only
 	TFuture(TFuture&&) = default;
-
-	/** Move assignment operator. */
+	TFuture(const TFuture&) = delete;
 	TFuture& operator=(TFuture&&) = default;
-
-	/** Destructor. */
+	TFuture& operator=(const TFuture&) = delete;
 	~TFuture() = default;
 
 public:
@@ -563,21 +540,6 @@ public:
 		: BaseType(MoveTemp(Future))
 	{ }
 
-	/** Copy constructor. */
-	TSharedFuture(const TSharedFuture&) = default;
-
-	/** Copy assignment operator. */
-	TSharedFuture& operator=(const TSharedFuture&) = default;
-
-	/** Move constructor. */
-	TSharedFuture(TSharedFuture&&) = default;
-
-	/** Move assignment operator. */
-	TSharedFuture& operator=(TSharedFuture&&) = default;
-
-	/** Destructor. */
-	~TSharedFuture() = default;
-
 public:
 
 	/**
@@ -703,16 +665,6 @@ public:
 		, FutureRetrieved(false)
 	{ }
 
-	/**
-	 * Move constructor.
-	 *
-	 * @param Other The promise holding the shared state to move.
-	 */
-	TPromise(TPromise&& Other)
-		: BaseType(MoveTemp(Other))
-		, FutureRetrieved(MoveTemp(Other.FutureRetrieved))
-	{ }
-
 	/**
 	 * Create and initialize a new instance with a callback.
 	 *
@@ -723,22 +675,11 @@ public:
 		, FutureRetrieved(false)
 	{ }
 
-public:
-
-	/**
-	 * Move assignment operator.
-	 *
-	 * @param Other The promise holding the shared state to move.
-	 */
-	TPromise& operator=(TPromise&& Other)
-	{
-		BaseType::operator=(MoveTemp(Other));
-		FutureRetrieved = MoveTemp(Other.FutureRetrieved);
-
-		return *this;
-	}
-
-public:
+	// Movable-only
+	TPromise(TPromise&& Other) = default;
+	TPromise(const TPromise& Other) = delete;
+	TPromise& operator=(TPromise&& Other) = default;
+	TPromise& operator=(const TPromise& Other) = delete;
 
 	/**
 	 * Gets a TFuture object associated with the shared state of this promise.
