Fixed some bool prefixes and defaults inits.

diff --git a/Engine/Source/Runtime/Core/Public/Async/Future.h b/Engine/Source/Runtime/Core/Public/Async/Future.h
--- a/Engine/Source/Runtime/Core/Public/Async/Future.h
+++ b/Engine/Source/Runtime/Core/Public/Async/Future.h
@@ -40,8 +40,8 @@ public:
 	/** Default constructor. */
 	FFutureState()
 		: CompletionEvent(FPlatformProcess::GetSynchEventFromPool(true))
-		, Complete(false)
-	{ }
+	{
+	}
 
 	/**
 	 * Create and initialize a new instance with a callback.
@@ -51,8 +51,8 @@ public:
 	FFutureState(TUniqueFunction<void()>&& InCompletionCallback)
 		: CompletionCallback(MoveTemp(InCompletionCallback))
 		, CompletionEvent(FPlatformProcess::GetSynchEventFromPool(true))
-		, Complete(false)
-	{ }
+	{
+	}
 
 	/** Destructor. */
 	~FFutureState()
@@ -71,7 +71,7 @@ public:
 	 */
 	bool IsComplete() const
 	{
-		return Complete;
+		return bComplete;
 	}
 
 	/**
@@ -122,7 +122,7 @@ protected:
 		{
 			FScopeLock Lock(&Mutex);
 			Continuation = MoveTemp(CompletionCallback);
-			Complete = true;
+			bComplete = true;
 		}
 		CompletionEvent->Trigger();
 
@@ -143,7 +143,7 @@ private:
 	FEvent* CompletionEvent;
 
 	/** Whether the asynchronous result is available. */
-	TAtomic<bool> Complete; 
+	TAtomic<bool> bComplete{ false };
 };
 
 
@@ -662,8 +662,8 @@ public:
 	/** Default constructor (creates a new shared state). */
 	TPromise()
 		: BaseType()
-		, FutureRetrieved(false)
-	{ }
+	{
+	}
 
 	/**
 	 * Create and initialize a new instance with a callback.
@@ -672,8 +672,8 @@ public:
 	 */
 	TPromise(TUniqueFunction<void()>&& CompletionCallback)
 		: BaseType(MoveTemp(CompletionCallback))
-		, FutureRetrieved(false)
-	{ }
+	{
+	}
 
 	// Movable-only
 	TPromise(TPromise&& Other) = default;
@@ -688,8 +688,8 @@ public:
 	 */
 	TFuture<ResultType> GetFuture()
 	{
-		check(!FutureRetrieved);
-		FutureRetrieved = true;
+		check(!bFutureRetrieved);
+		bFutureRetrieved = true;
 
 		return TFuture<ResultType>(this->GetState());
 	}
@@ -794,7 +794,7 @@ public:
 private:
 
 	/** Whether a future has already been retrieved from this promise. */
-	bool FutureRetrieved;
+	bool bFutureRetrieved = false;
 };
 
 /* TFuture::Then
