Add `strtoll`, `strtoul`, and `strtoull` support into AutoRTFM as they are hit from Linux servers.

diff --git a/Engine/Source/Programs/AutoRTFMTests/Private/CString.cpp b/Engine/Source/Programs/AutoRTFMTests/Private/CString.cpp
--- a/Engine/Source/Programs/AutoRTFMTests/Private/CString.cpp
+++ b/Engine/Source/Programs/AutoRTFMTests/Private/CString.cpp
@@ -1087,6 +1087,21 @@ TEST_CASE("CString.strtol")
 	CheckStringToNumberFunction(strtol, "42.5abc", 10);
 }
 
+TEST_CASE("CString.strtoll")
+{
+	CheckStringToNumberFunction(strtoll, "42.5abc", 10);
+}
+
+TEST_CASE("CString.strtoul")
+{
+	CheckStringToNumberFunction(strtoul, "42.5abc", 10);
+}
+
+TEST_CASE("CString.strtoull")
+{
+	CheckStringToNumberFunction(strtoull, "42.5abc", 10);
+}
+
 TEST_CASE("CString.wcstol")
 {
 	CheckStringToNumberFunction(wcstol, L"42abc", 10);
diff --git a/Engine/Source/Runtime/Core/Private/AutoRTFM/StdLibSupport.cpp b/Engine/Source/Runtime/Core/Private/AutoRTFM/StdLibSupport.cpp
--- a/Engine/Source/Runtime/Core/Private/AutoRTFM/StdLibSupport.cpp
+++ b/Engine/Source/Runtime/Core/Private/AutoRTFM/StdLibSupport.cpp
@@ -271,6 +271,9 @@ RETURN RTFM_StringToInt(const CHAR* String, CHAR** EndPtr, int Radix) throw()
 }
 
 UE_AUTORTFM_REGISTER_OPEN_FUNCTION_EXPLICIT(strtol, (RTFM_StringToInt<long, char, strtol>));
+UE_AUTORTFM_REGISTER_OPEN_FUNCTION_EXPLICIT(strtoll, (RTFM_StringToInt<long long, char, strtoll>));
+UE_AUTORTFM_REGISTER_OPEN_FUNCTION_EXPLICIT(strtoul, (RTFM_StringToInt<unsigned long, char, strtoul>));
+UE_AUTORTFM_REGISTER_OPEN_FUNCTION_EXPLICIT(strtoull, (RTFM_StringToInt<unsigned long long, char, strtoull>));
 UE_AUTORTFM_REGISTER_OPEN_FUNCTION_EXPLICIT(strtof, (RTFM_StringToFloat<float, char, strtof>));
 UE_AUTORTFM_REGISTER_OPEN_FUNCTION_EXPLICIT(strtod, (RTFM_StringToFloat<double, char, strtod>));
 UE_AUTORTFM_REGISTER_OPEN_FUNCTION_EXPLICIT(wcstod, (RTFM_StringToFloat<double, wchar_t, wcstod>));
