Added an RAII FPooledSyncEvent, used in FFutureState.

diff --git a/Engine/Source/Runtime/Core/Public/Async/Future.h b/Engine/Source/Runtime/Core/Public/Async/Future.h
--- a/Engine/Source/Runtime/Core/Public/Async/Future.h
+++ b/Engine/Source/Runtime/Core/Public/Async/Future.h
@@ -38,10 +38,7 @@ class FFutureState
 public:
 
 	/** Default constructor. */
-	FFutureState()
-		: CompletionEvent(FPlatformProcess::GetSynchEventFromPool(true))
-	{
-	}
+	FFutureState() = default;
 
 	/**
 	 * Create and initialize a new instance with a callback.
@@ -50,15 +47,7 @@ public:
 	 */
 	FFutureState(TUniqueFunction<void()>&& InCompletionCallback)
 		: CompletionCallback(MoveTemp(InCompletionCallback))
-		, CompletionEvent(FPlatformProcess::GetSynchEventFromPool(true))
-	{
-	}
-
-	/** Destructor. */
-	~FFutureState()
 	{
-		FPlatformProcess::ReturnSynchEventToPool(CompletionEvent);
-		CompletionEvent = nullptr;
 	}
 
 public:
@@ -140,7 +129,7 @@ private:
 	TUniqueFunction<void()> CompletionCallback;
 
 	/** Holds an event signaling that the result is available. */
-	FEvent* CompletionEvent;
+	FPooledSyncEvent CompletionEvent{ true };
 
 	/** Whether the asynchronous result is available. */
 	TAtomic<bool> bComplete{ false };
diff --git a/Engine/Source/Runtime/Core/Public/HAL/PlatformProcess.h b/Engine/Source/Runtime/Core/Public/HAL/PlatformProcess.h
--- a/Engine/Source/Runtime/Core/Public/HAL/PlatformProcess.h
+++ b/Engine/Source/Runtime/Core/Public/HAL/PlatformProcess.h
@@ -5,3 +5,44 @@
 #include "CoreTypes.h"
 
 #include COMPILED_PLATFORM_HEADER(PlatformProcess.h)
+
+// An RAII wrapper around FPlatformProcess::GetSynchEventFromPool/ReturnSynchEventToPool.
+struct FPooledSyncEvent
+{
+	explicit FPooledSyncEvent(bool bIsManualReset = false)
+	{
+		Event = FPlatformProcess::GetSynchEventFromPool(bIsManualReset);
+	}
+	~FPooledSyncEvent()
+	{
+		*this = nullptr;
+	}
+
+	// Non-copyable
+	FPooledSyncEvent(FPooledSyncEvent&&) = delete;
+	FPooledSyncEvent(const FPooledSyncEvent&) = delete;
+	FPooledSyncEvent& operator=(FPooledSyncEvent&&) = delete;
+	FPooledSyncEvent& operator=(const FPooledSyncEvent&) = delete;
+
+	explicit operator bool() const
+	{
+		return !!this->Event;
+	}
+
+	FEvent* operator->() const
+	{
+		return this->Event;
+	}
+
+	FPooledSyncEvent& operator=(TYPE_OF_NULLPTR)
+	{
+		if (Event)
+		{
+			FPlatformProcess::ReturnSynchEventToPool(Event);
+			Event = nullptr;
+		}
+		return *this;
+	}
+
+	FEvent* Event;
+};
