Made implementation detail constructors and made friends of the classes that need them.

diff --git a/Engine/Source/Runtime/Core/Public/Async/Future.h b/Engine/Source/Runtime/Core/Public/Async/Future.h
--- a/Engine/Source/Runtime/Core/Public/Async/Future.h
+++ b/Engine/Source/Runtime/Core/Public/Async/Future.h
@@ -15,6 +15,10 @@
 #include "Templates/Requires.h"
 #include <type_traits>
 
+template<typename ResultType> class TFuture;
+template<typename ResultType> class TSharedFuture;
+template<typename ResultType> class TPromise;
+
 namespace UE::Core::Private
 {
 	template <typename T>
@@ -301,10 +305,14 @@ protected:
 	 *
 	 * @param InState The shared state to initialize with.
 	 */
+	TFutureBase(StateType&& InState)
+		: State(MoveTemp(InState))
+	{
+	}
 	TFutureBase(const StateType& InState)
 		: State(InState)
-	{ }
-
+	{
+	}
 protected:
 
 	/**
@@ -363,8 +371,6 @@ private:
 };
 
 
-template<typename ResultType> class TSharedFuture;
-
 /**
  * Template for unshared futures.
  */
@@ -372,6 +378,9 @@ template<typename ResultType>
 class TFuture
 	: public TFutureBase<UE::Core::Private::TInternalFutureState_T<ResultType>>
 {
+	template <typename>
+	friend class TPromise;
+
 	using BaseType = TFutureBase<UE::Core::Private::TInternalFutureState_T<ResultType>>;
 
 public:
@@ -379,15 +388,6 @@ public:
 	/** Default constructor. */
 	TFuture() = default;
 
-	/**
-	 * Creates and initializes a new instance.
-	 *
-	 * @param InState The shared state to initialize with.
-	 */
-	TFuture(const typename BaseType::StateType& InState)
-		: BaseType(InState)
-	{ }
-
 	// Movable-only
 	TFuture(TFuture&&) = default;
 	TFuture(const TFuture&) = delete;
@@ -395,8 +395,6 @@ public:
 	TFuture& operator=(const TFuture&) = delete;
 	~TFuture() = default;
 
-public:
-
 	/**
 	 * Gets the future's result by const reference.
 	 *
@@ -491,6 +489,16 @@ public:
 	 * @see TFutureBase
 	 */
 	using BaseType::Reset;
+
+private:
+	TFuture(const typename BaseType::StateType& InState)
+		: BaseType(InState)
+	{
+	}
+	TFuture(typename BaseType::StateType&& InState)
+		: BaseType(MoveTemp(InState))
+	{
+	}
 };
 
 
@@ -504,6 +512,9 @@ template<typename ResultType>
 class TSharedFuture
 	: public TFutureBase<UE::Core::Private::TInternalFutureState_T<ResultType>>
 {
+	template <typename>
+	friend class TFuture;
+
 	using BaseType = TFutureBase<UE::Core::Private::TInternalFutureState_T<ResultType>>;
 
 public:
@@ -511,15 +522,6 @@ public:
 	/** Default constructor. */
 	TSharedFuture() = default;
 
-	/**
-	 * Creates and initializes a new instance.
-	 *
-	 * @param InState The shared state to initialize from.
-	 */
-	TSharedFuture(const typename BaseType::StateType& InState)
-		: BaseType(InState)
-	{ }
-
 	/**
 	 * Creates and initializes a new instances from a future object.
 	 *
@@ -529,8 +531,6 @@ public:
 		: BaseType(MoveTemp(Future))
 	{ }
 
-public:
-
 	/**
 	 * Gets the future's result.
 	 *
@@ -551,6 +551,21 @@ public:
 			return (const ResultType&)this->GetState()->GetResult();
 		}
 	}
+
+private:
+	/**
+	 * Creates and initializes a new instance.
+	 *
+	 * @param InState The shared state to initialize from.
+	 */
+	TSharedFuture(const typename BaseType::StateType& InState)
+		: BaseType(InState)
+	{
+	}
+	TSharedFuture(typename BaseType::StateType&& InState)
+		: BaseType(InState)
+	{
+	}
 };
 
 
