CIS fix.

diff --git a/Engine/Source/Runtime/Core/Public/Async/Future.h b/Engine/Source/Runtime/Core/Public/Async/Future.h
--- a/Engine/Source/Runtime/Core/Public/Async/Future.h
+++ b/Engine/Source/Runtime/Core/Public/Async/Future.h
@@ -381,6 +381,9 @@ class TFuture
 	template <typename>
 	friend class TPromise;
 
+	template <typename>
+	friend class TFutureBase;
+
 	using BaseType = TFutureBase<UE::Core::Private::TInternalFutureState_T<ResultType>>;
 
 public:
