Revert recent TFuture changes while Verse depends on those implementation details.

diff --git a/Engine/Source/Runtime/Core/Public/Async/Future.h b/Engine/Source/Runtime/Core/Public/Async/Future.h
--- a/Engine/Source/Runtime/Core/Public/Async/Future.h
+++ b/Engine/Source/Runtime/Core/Public/Async/Future.h
@@ -15,10 +15,6 @@
 #include "Templates/Requires.h"
 #include <type_traits>
 
-template<typename ResultType> class TFuture;
-template<typename ResultType> class TSharedFuture;
-template<typename ResultType> class TPromise;
-
 namespace UE::Core::Private
 {
 	template <typename T>
@@ -305,14 +301,10 @@ protected:
 	 *
 	 * @param InState The shared state to initialize with.
 	 */
-	TFutureBase(StateType&& InState)
-		: State(MoveTemp(InState))
-	{
-	}
 	TFutureBase(const StateType& InState)
 		: State(InState)
-	{
-	}
+	{ }
+
 protected:
 
 	/**
@@ -371,6 +363,8 @@ private:
 };
 
 
+template<typename ResultType> class TSharedFuture;
+
 /**
  * Template for unshared futures.
  */
@@ -378,12 +372,6 @@ template<typename ResultType>
 class TFuture
 	: public TFutureBase<UE::Core::Private::TInternalFutureState_T<ResultType>>
 {
-	template <typename>
-	friend class TPromise;
-
-	template <typename>
-	friend class TFutureBase;
-
 	using BaseType = TFutureBase<UE::Core::Private::TInternalFutureState_T<ResultType>>;
 
 public:
@@ -391,6 +379,15 @@ public:
 	/** Default constructor. */
 	TFuture() = default;
 
+	/**
+	 * Creates and initializes a new instance.
+	 *
+	 * @param InState The shared state to initialize with.
+	 */
+	TFuture(const typename BaseType::StateType& InState)
+		: BaseType(InState)
+	{ }
+
 	// Movable-only
 	TFuture(TFuture&&) = default;
 	TFuture(const TFuture&) = delete;
@@ -398,6 +395,8 @@ public:
 	TFuture& operator=(const TFuture&) = delete;
 	~TFuture() = default;
 
+public:
+
 	/**
 	 * Gets the future's result by const reference.
 	 *
@@ -492,16 +491,6 @@ public:
 	 * @see TFutureBase
 	 */
 	using BaseType::Reset;
-
-private:
-	TFuture(const typename BaseType::StateType& InState)
-		: BaseType(InState)
-	{
-	}
-	TFuture(typename BaseType::StateType&& InState)
-		: BaseType(MoveTemp(InState))
-	{
-	}
 };
 
 
@@ -515,9 +504,6 @@ template<typename ResultType>
 class TSharedFuture
 	: public TFutureBase<UE::Core::Private::TInternalFutureState_T<ResultType>>
 {
-	template <typename>
-	friend class TFuture;
-
 	using BaseType = TFutureBase<UE::Core::Private::TInternalFutureState_T<ResultType>>;
 
 public:
@@ -525,6 +511,15 @@ public:
 	/** Default constructor. */
 	TSharedFuture() = default;
 
+	/**
+	 * Creates and initializes a new instance.
+	 *
+	 * @param InState The shared state to initialize from.
+	 */
+	TSharedFuture(const typename BaseType::StateType& InState)
+		: BaseType(InState)
+	{ }
+
 	/**
 	 * Creates and initializes a new instances from a future object.
 	 *
@@ -534,6 +529,8 @@ public:
 		: BaseType(MoveTemp(Future))
 	{ }
 
+public:
+
 	/**
 	 * Gets the future's result.
 	 *
@@ -554,21 +551,6 @@ public:
 			return (const ResultType&)this->GetState()->GetResult();
 		}
 	}
-
-private:
-	/**
-	 * Creates and initializes a new instance.
-	 *
-	 * @param InState The shared state to initialize from.
-	 */
-	TSharedFuture(const typename BaseType::StateType& InState)
-		: BaseType(InState)
-	{
-	}
-	TSharedFuture(typename BaseType::StateType&& InState)
-		: BaseType(InState)
-	{
-	}
 };
 
 
