Folded TPromiseBase into TPromise.

diff --git a/Engine/Source/Runtime/Core/Public/Async/Future.h b/Engine/Source/Runtime/Core/Public/Async/Future.h
--- a/Engine/Source/Runtime/Core/Public/Async/Future.h
+++ b/Engine/Source/Runtime/Core/Public/Async/Future.h
@@ -575,91 +575,14 @@ private:
 /* TPromise
 *****************************************************************************/
 
-template<typename InternalResultType>
-class TPromiseBase
-	: FNoncopyable
-{
-	typedef TSharedPtr<TFutureState<InternalResultType>> StateType;
-
-public:
-
-	/** Default constructor. */
-	TPromiseBase()
-		: State(MakeShared<TFutureState<InternalResultType>>())
-	{ }
-
-	/**
-	 * Move constructor.
-	 *
-	 * @param Other The promise holding the shared state to move.
-	 */
-	TPromiseBase(TPromiseBase&& Other)
-		: State(MoveTemp(Other.State))
-	{
-		Other.State.Reset();
-	}
-
-	/**
-	 * Create and initialize a new instance with a callback.
-	 *
-	 * @param CompletionCallback A function that is called when the future state is completed.
-	 */
-	TPromiseBase(TUniqueFunction<void()>&& CompletionCallback)
-		: State(MakeShared<TFutureState<InternalResultType>>(MoveTemp(CompletionCallback)))
-	{ }
-
-public:
-
-	/** Move assignment operator. */
-	TPromiseBase& operator=(TPromiseBase&& Other)
-	{
-		State = Other.State;
-		Other.State.Reset();
-		return *this;
-	}
-
-protected:
-
-	/** Destructor. */
-	~TPromiseBase()
-	{
-		if (State.IsValid())
-		{
-			// if you hit this assertion then your promise never had its result
-			// value set. broken promises are considered programming errors.
-			check(State->IsComplete());
-		}
-	}
-
-	/**
-	 * Gets the shared state object.
-	 *
-	 * @return The shared state object.
-	 */
-	const StateType& GetState()
-	{
-		// if you hit this assertion then your promise has an invalid state.
-		// this happens if you move the promise to another instance.
-		check(State.IsValid());
-
-		return State;
-	}
-
-private:
-
-	/** Holds the shared state object. */
-	StateType State;
-};
-
-
 /**
  * Template for promises.
  */
 template<typename ResultType>
 class TPromise
-	: public TPromiseBase<UE::Core::Private::TInternalFutureState_T<ResultType>>
 {
-	using BaseType = TPromiseBase<UE::Core::Private::TInternalFutureState_T<ResultType>>;
+	using InternalResultType = UE::Core::Private::TInternalFutureState_T<ResultType>;
+	using StateType          = TSharedPtr<TFutureState<InternalResultType>>;
 
 	// This is necessary because we can't form references to void for the parameter of Set().
 	using SetType = std::conditional_t<std::is_void_v<ResultType>, int, ResultType>;
@@ -668,7 +591,7 @@ public:
 
 	/** Default constructor (creates a new shared state). */
 	TPromise()
-		: BaseType()
+		: State(MakeShared<TFutureState<InternalResultType>>())
 	{
 	}
 
@@ -678,7 +601,7 @@ public:
 	 * @param CompletionCallback A function that is called when the future state is completed.
 	 */
 	TPromise(TUniqueFunction<void()>&& CompletionCallback)
-		: BaseType(MoveTemp(CompletionCallback))
+		: State(MakeShared<TFutureState<InternalResultType>>(MoveTemp(CompletionCallback)))
 	{
 	}
 
@@ -688,6 +611,17 @@ public:
 	TPromise& operator=(TPromise&& Other) = default;
 	TPromise& operator=(const TPromise& Other) = delete;
 
+	/** Destructor. */
+	~TPromise()
+	{
+		if (State.IsValid())
+		{
+			// if you hit this assertion then your promise never had its result
+			// value set. broken promises are considered programming errors.
+			check(State->IsComplete());
+		}
+	}
+
 	/**
 	 * Gets a TFuture object associated with the shared state of this promise.
 	 *
@@ -798,7 +732,24 @@ public:
 		this->GetState()->EmplaceResult(0);
 	}
 
+protected:
+	/**
+	 * Gets the shared state object.
+	 *
+	 * @return The shared state object.
+	 */
+	const StateType& GetState()
+	{
+		// if you hit this assertion then your promise has an invalid state.
+		// this happens if you move the promise to another instance.
+		check(State.IsValid());
+
+		return State;
+	}
+
 private:
+	/** Holds the shared state object. */
+	StateType State;
 
 	/** Whether a future has already been retrieved from this promise. */
 	bool bFutureRetrieved = false;
