[AutoRTFM] Register closed variants of std::to_chars

diff --git a/Engine/Source/Programs/AutoRTFMTests/Private/CString.cpp b/Engine/Source/Programs/AutoRTFMTests/Private/CString.cpp
--- a/Engine/Source/Programs/AutoRTFMTests/Private/CString.cpp
+++ b/Engine/Source/Programs/AutoRTFMTests/Private/CString.cpp
@@ -1,12 +1,15 @@
 // Copyright Epic Games, Inc. All Rights Reserved.
 
+#include "AutoRTFM/AutoRTFMTesting.h"
 #include "Catch2Includes.h"
 #include "CoreGlobals.h"
 #include "Logging/LogVerbosity.h"
 #include "Misc/FeedbackContext.h"
 #include <AutoRTFM/AutoRTFM.h>
 #include <catch_amalgamated.hpp>
+#include <charconv>
 #include <cstring>
+#include <system_error>
 
 #ifdef __clang__
 #define CLANG_BEGIN_DISABLE_OPTIMIZATIONS _Pragma("clang optimize off")
@@ -152,6 +155,56 @@ void CheckStringToNumberFunction(RESULT Function(const CHAR*, CHAR**, ADDITIONAL
 	}
 }
 
+// CheckStdToChars calls std::to_chars with (char*, char*, ARGS...) outside a
+// transaction, in an aborted transaction and committed transaction.
+// The value returned and the string buffer content is expected to match the
+// values when called outside the transaction.
+template<typename ... ARGS>
+void CheckStdToChars(ARGS ... Args)
+{
+	static constexpr size_t BufferOffset = 2;
+	static constexpr size_t BufferSize = 32;
+	char ExpectedBuffer[]
+	{
+		'#', '#', '#', '#', '#', '#', '#', '#', '#', '#', '#', '#', '#', '#', '#', '#',
+		'#', '#', '#', '#', '#', '#', '#', '#', '#', '#', '#', '#', '#', '#', '#', '\0',
+	};
+	static_assert(sizeof(ExpectedBuffer) == BufferSize);
+	const std::to_chars_result ExpectedResult = std::to_chars(&ExpectedBuffer[BufferOffset], &ExpectedBuffer[BufferSize-1], Args...);
+
+	char Buffer[]
+	{
+		'#', '#', '#', '#', '#', '#', '#', '#', '#', '#', '#', '#', '#', '#', '#', '#',
+		'#', '#', '#', '#', '#', '#', '#', '#', '#', '#', '#', '#', '#', '#', '#', '\0',
+	};
+	static_assert(sizeof(Buffer) == BufferSize);
+
+	SECTION("With Abort")
+	{
+		AutoRTFM::Testing::Abort([&]
+		{
+			std::to_chars_result Result = std::to_chars(&Buffer[BufferOffset], &Buffer[BufferSize-1], Args...);
+			REQUIRE(Result.ec == ExpectedResult.ec);
+			REQUIRE(Result.ptr - Buffer == ExpectedResult.ptr - ExpectedBuffer);
+			REQUIRE(std::string_view(Buffer) == std::string_view(ExpectedBuffer));
+			AutoRTFM::AbortTransaction();
+		});
+		REQUIRE(std::string_view(Buffer) == "###############################");
+	}
+
+	SECTION("With Commit")
+	{
+		AutoRTFM::Testing::Commit([&]
+		{
+			std::to_chars_result Result = std::to_chars(&Buffer[BufferOffset], &Buffer[BufferSize-1], Args...);
+			REQUIRE(Result.ec == ExpectedResult.ec);
+			REQUIRE(Result.ptr - Buffer == ExpectedResult.ptr - ExpectedBuffer);
+			REQUIRE(std::string_view(Buffer) == std::string_view(ExpectedBuffer));
+		});
+		REQUIRE(std::string_view(Buffer) == std::string_view(ExpectedBuffer));
+	}
+}
+
 // A helper class that for the lifetime of the object, intercepts and records
 // UE_LOG warnings.
 class FCaptureWarningContext : private FFeedbackContext
@@ -1122,6 +1175,29 @@ TEST_CASE("CString.wcstoull")
 	CheckStringToNumberFunction(wcstoull, L"42abc", 10);
 }
 
+TEST_CASE("std::to_chars")
+{
+	CheckStdToChars(/* value */ static_cast<int>(42), /* base */ 8);
+	CheckStdToChars(/* value */ static_cast<unsigned int>(42), /* base */ 8);
+	CheckStdToChars(/* value */ static_cast<int8_t>(42), /* base */ 8);
+	CheckStdToChars(/* value */ static_cast<uint8_t>(42), /* base */ 8);
+	CheckStdToChars(/* value */ static_cast<int16_t>(42), /* base */ 8);
+	CheckStdToChars(/* value */ static_cast<uint16_t>(42), /* base */ 8);
+	CheckStdToChars(/* value */ static_cast<int32_t>(42), /* base */ 8);
+	CheckStdToChars(/* value */ static_cast<uint32_t>(42), /* base */ 8);
+	CheckStdToChars(/* value */ static_cast<int64_t>(42), /* base */ 8);
+	CheckStdToChars(/* value */ static_cast<uint64_t>(42), /* base */ 8);
+	CheckStdToChars(/* value */ static_cast<float>(42.24f));
+	CheckStdToChars(/* value */ static_cast<float>(42.24f), /* fmt */ std::chars_format::hex);
+	CheckStdToChars(/* value */ static_cast<float>(42.24f), /* fmt */ std::chars_format::hex, /* precision */ 3);
+	CheckStdToChars(/* value */ static_cast<double>(42.24));
+	CheckStdToChars(/* value */ static_cast<double>(42.24), /* fmt */ std::chars_format::hex);
+	CheckStdToChars(/* value */ static_cast<double>(42.24), /* fmt */ std::chars_format::hex, /* precision */ 3);
+	CheckStdToChars(/* value */ static_cast<long double>(42.24));
+	CheckStdToChars(/* value */ static_cast<long double>(42.24), /* fmt */ std::chars_format::hex);
+	CheckStdToChars(/* value */ static_cast<long double>(42.24), /* fmt */ std::chars_format::hex, /* precision */ 3);
+}
+
 #if PLATFORM_WINDOWS
 TEST_CASE("CString._wcstoi64")
 {
diff --git a/Engine/Source/Runtime/Core/Private/AutoRTFM/StdLibSupport.cpp b/Engine/Source/Runtime/Core/Private/AutoRTFM/StdLibSupport.cpp
--- a/Engine/Source/Runtime/Core/Private/AutoRTFM/StdLibSupport.cpp
+++ b/Engine/Source/Runtime/Core/Private/AutoRTFM/StdLibSupport.cpp
@@ -1,5 +1,7 @@
 // Copyright Epic Games, Inc. All Rights Reserved.
 
+#include "AutoRTFM/AutoRTFM.h"
+#include <charconv>
 #if (defined(__AUTORTFM) && __AUTORTFM)
 #include "ContextInlines.h"
 #include "Utils.h"
@@ -284,6 +286,129 @@ UE_AUTORTFM_REGISTER_OPEN_FUNCTION_EXPLICIT(wcstoll, (RTFM_StringToInt<long long
 UE_AUTORTFM_REGISTER_OPEN_FUNCTION_EXPLICIT(wcstoul, (RTFM_StringToInt<unsigned long, wchar_t, wcstoul>));
 UE_AUTORTFM_REGISTER_OPEN_FUNCTION_EXPLICIT(wcstoull, (RTFM_StringToInt<unsigned long long, wchar_t, wcstoull>));
 
+template<typename T, std::to_chars_result FN(char*, char*, T)>
+std::to_chars_result RTFM_ToChars(char* First, char* Last, T Value)
+{
+    AutoRTFM::RecordOpenWrite(First, Last - First);
+    return FN(First, Last, Value);
+}
+
+template<typename T, std::to_chars_result FN(char*, char*, T, int)>
+std::to_chars_result RTFM_ToChars(char* First, char* Last, T Value, int Base)
+{
+    AutoRTFM::RecordOpenWrite(First, Last - First);
+    return FN(First, Last, Value, Base);
+}
+
+template<typename T, std::to_chars_result FN(char*, char*, T, std::chars_format)>
+std::to_chars_result RTFM_ToChars(char* First, char* Last, T Value, std::chars_format Format)
+{
+    AutoRTFM::RecordOpenWrite(First, Last - First);
+    return FN(First, Last, Value, Format);
+}
+
+template<typename T, std::to_chars_result FN(char*, char*, T, std::chars_format, int Precision)>
+std::to_chars_result RTFM_ToChars(char* First, char* Last, T Value, std::chars_format Format, int Precision)
+{
+    AutoRTFM::RecordOpenWrite(First, Last - First);
+    return FN(First, Last, Value, Format, Precision);
+}
+
+UE_AUTORTFM_REGISTER_OPEN_OVERLOADED_FUNCTION_EXPLICIT(
+    std::to_chars_result(char* First, char* Last, int Value, int Base),
+    std::to_chars,
+    (RTFM_ToChars<int, std::to_chars>));
+
+UE_AUTORTFM_REGISTER_OPEN_OVERLOADED_FUNCTION_EXPLICIT(
+    std::to_chars_result(char* First, char* Last, unsigned int Value, int Base),
+    std::to_chars,
+    (RTFM_ToChars<unsigned int, std::to_chars>));
+
+UE_AUTORTFM_REGISTER_OPEN_OVERLOADED_FUNCTION_EXPLICIT(
+    std::to_chars_result(char* First, char* Last, int8_t Value, int Base),
+    std::to_chars,
+    (RTFM_ToChars<int8_t, std::to_chars>));
+
+UE_AUTORTFM_REGISTER_OPEN_OVERLOADED_FUNCTION_EXPLICIT(
+    std::to_chars_result(char* First, char* Last, uint8_t Value, int Base),
+    std::to_chars,
+    (RTFM_ToChars<uint8_t, std::to_chars>));
+
+UE_AUTORTFM_REGISTER_OPEN_OVERLOADED_FUNCTION_EXPLICIT(
+    std::to_chars_result(char* First, char* Last, int16_t Value, int Base),
+    std::to_chars,
+    (RTFM_ToChars<int16_t, std::to_chars>));
+
+UE_AUTORTFM_REGISTER_OPEN_OVERLOADED_FUNCTION_EXPLICIT(
+    std::to_chars_result(char* First, char* Last, uint16_t Value, int Base),
+    std::to_chars,
+    (RTFM_ToChars<uint16_t, std::to_chars>));
+
+UE_AUTORTFM_REGISTER_OPEN_OVERLOADED_FUNCTION_EXPLICIT(
+    std::to_chars_result(char* First, char* Last, int32_t Value, int Base),
+    std::to_chars,
+    (RTFM_ToChars<int32_t, std::to_chars>));
+
+UE_AUTORTFM_REGISTER_OPEN_OVERLOADED_FUNCTION_EXPLICIT(
+    std::to_chars_result(char* First, char* Last, uint32_t Value, int Base),
+    std::to_chars,
+    (RTFM_ToChars<uint32_t, std::to_chars>));
+
+UE_AUTORTFM_REGISTER_OPEN_OVERLOADED_FUNCTION_EXPLICIT(
+    std::to_chars_result(char* First, char* Last, int64_t Value, int Base),
+    std::to_chars,
+    (RTFM_ToChars<int64_t, std::to_chars>));
+
+UE_AUTORTFM_REGISTER_OPEN_OVERLOADED_FUNCTION_EXPLICIT(
+    std::to_chars_result(char* First, char* Last, uint64_t Value, int Base),
+    std::to_chars,
+    (RTFM_ToChars<uint64_t, std::to_chars>));
+
+UE_AUTORTFM_REGISTER_OPEN_OVERLOADED_FUNCTION_EXPLICIT(
+    std::to_chars_result(char* First, char* Last, float Value),
+    std::to_chars,
+    (RTFM_ToChars<float, std::to_chars>));
+
+UE_AUTORTFM_REGISTER_OPEN_OVERLOADED_FUNCTION_EXPLICIT(
+    std::to_chars_result(char* First, char* Last, float Value, std::chars_format Format),
+    std::to_chars,
+    (RTFM_ToChars<float, std::to_chars>));
+
+UE_AUTORTFM_REGISTER_OPEN_OVERLOADED_FUNCTION_EXPLICIT(
+    std::to_chars_result(char* First, char* Last, float Value, std::chars_format Format, int Precision),
+    std::to_chars,
+    (RTFM_ToChars<float, std::to_chars>));
+
+UE_AUTORTFM_REGISTER_OPEN_OVERLOADED_FUNCTION_EXPLICIT(
+    std::to_chars_result(char* First, char* Last, double Value),
+    std::to_chars,
+    (RTFM_ToChars<double, std::to_chars>));
+
+UE_AUTORTFM_REGISTER_OPEN_OVERLOADED_FUNCTION_EXPLICIT(
+    std::to_chars_result(char* First, char* Last, double Value, std::chars_format Format),
+    std::to_chars,
+    (RTFM_ToChars<double, std::to_chars>));
+
+UE_AUTORTFM_REGISTER_OPEN_OVERLOADED_FUNCTION_EXPLICIT(
+    std::to_chars_result(char* First, char* Last, double Value, std::chars_format Format, int Precision),
+    std::to_chars,
+    (RTFM_ToChars<double, std::to_chars>));
+
+UE_AUTORTFM_REGISTER_OPEN_OVERLOADED_FUNCTION_EXPLICIT(
+    std::to_chars_result(char* First, char* Last, long double Value),
+    std::to_chars,
+    (RTFM_ToChars<long double, std::to_chars>));
+
+UE_AUTORTFM_REGISTER_OPEN_OVERLOADED_FUNCTION_EXPLICIT(
+    std::to_chars_result(char* First, char* Last, long double Value, std::chars_format Format),
+    std::to_chars,
+    (RTFM_ToChars<long double, std::to_chars>));
+
+UE_AUTORTFM_REGISTER_OPEN_OVERLOADED_FUNCTION_EXPLICIT(
+    std::to_chars_result(char* First, char* Last, long double Value, std::chars_format Format, int Precision),
+    std::to_chars,
+    (RTFM_ToChars<long double, std::to_chars>));
+
 #if PLATFORM_WINDOWS
 template<typename RETURN, typename CHAR, RETURN FN(const CHAR*, CHAR**, _locale_t)>
 RETURN RTFM_StringToFloat(const CHAR* String, CHAR** EndPtr, _locale_t Locale)
