[FTaskTagScope]
- Make all task scope optional to prevent mistakes

#rnx

diff --git a/Engine/Source/Runtime/Core/Private/HAL/ThreadingBase.cpp b/Engine/Source/Runtime/Core/Private/HAL/ThreadingBase.cpp
--- a/Engine/Source/Runtime/Core/Private/HAL/ThreadingBase.cpp
+++ b/Engine/Source/Runtime/Core/Private/HAL/ThreadingBase.cpp
@@ -78,13 +78,17 @@ void FTaskTagScope::SetTagStaticInit()
 	ActiveTaskTag = ETaskTag::EStaticInit;
 }
 
-FTaskTagScope::FTaskTagScope(bool InTagOnlyIfNone, ETaskTag InTag) : Tag(InTag), TagOnlyIfNone(InTagOnlyIfNone)
+FTaskTagScope::FTaskTagScope(ETaskTag InTag) : Tag(InTag)
 {
 	checkf(Tag != ETaskTag::ENone, TEXT("None cannot be used as a Tag"));
 	checkf(Tag != ETaskTag::EParallelThread, TEXT("Parallel cannot be used on it's own"));
 
+	TagOnlyIfNone = true; // Now defaults to always optional
 	if (ActiveTaskTag == ETaskTag::EStaticInit)
 	{
+		// Allow overwrite only from EStaticInit -> EGameThread
+		TagOnlyIfNone = Tag != ETaskTag::EGameThread;
+
 		checkf(Tag == ETaskTag::EGameThread, TEXT("The Gamethread can only be tagged on the inital thread of the application"));
 #if !IS_RUNNING_GAMETHREAD_ON_EXTERNAL_THREAD
 #	if PLATFORM_WINDOWS && DO_CHECK
diff --git a/Engine/Source/Runtime/Core/Public/CoreGlobals.h b/Engine/Source/Runtime/Core/Public/CoreGlobals.h
--- a/Engine/Source/Runtime/Core/Public/CoreGlobals.h
+++ b/Engine/Source/Runtime/Core/Public/CoreGlobals.h
@@ -663,7 +663,7 @@ ENUM_CLASS_FLAGS(ETaskTag);
 
 
 /**
- * This class can be used to Tag an execution context aka Thead or Job and allows us to later querry the state when we are in the callstack
+ * This class can be used to Tag an execution context aka Thead or Job and allows us to later query the state when we are in the callstack
  * It is usually used for the IsInRendering/GamethreadFunctions.
  *
  * @param CtorSignature InTag the Tag to use
@@ -697,15 +697,8 @@ public:
 	*/
 	static ETaskTag CORE_API SwapTag(ETaskTag Tag);
 
-protected:
-	CORE_API FTaskTagScope(bool InTagOnlyIfNone, ETaskTag InTag);
-
 public:
-	FTaskTagScope(ETaskTag InTag = ETaskTag::ENone) : FTaskTagScope(false, InTag)
-	{
-
-	}
-
+	CORE_API FTaskTagScope(ETaskTag InTag);
 	CORE_API ~FTaskTagScope();
 
 	static CORE_API ETaskTag GetCurrentTag();
@@ -722,9 +715,8 @@ public:
 class FOptionalTaskTagScope : public FTaskTagScope
 {
 public:
-	FOptionalTaskTagScope(ETaskTag InTag = ETaskTag::ENone) : FTaskTagScope(true, InTag)
+	FOptionalTaskTagScope(ETaskTag InTag = ETaskTag::ENone) : FTaskTagScope(InTag)
 	{
-
 	}
 };
 
