[TAtomic]
- Remove old deprecated implementation

#rnx

diff --git a/Engine/Source/Runtime/Core/Public/Templates/Atomic.h b/Engine/Source/Runtime/Core/Public/Templates/Atomic.h
--- a/Engine/Source/Runtime/Core/Public/Templates/Atomic.h
+++ b/Engine/Source/Runtime/Core/Public/Templates/Atomic.h
@@ -13,12 +13,6 @@
 // `TAtomic` is planned for deprecation. Please use `std::atomic`
 ////////////////////////////////////////////////////////////////////////////////////////////////////////////////////////////
 
-#define USE_DEPRECATED_TATOMIC 0
-
-#if !defined(USE_DEPRECATED_TATOMIC)
-#define USE_DEPRECATED_TATOMIC 0
-#endif
-
 template <typename T>
 struct TAtomicBase_Basic;
 
@@ -220,9 +214,6 @@ namespace UE::Core::Private::Atomic
 // Basic storage and implementation - only allows getting and setting via platform atomics.
 template <typename T>
 struct 
-#if USE_DEPRECATED_TATOMIC
-	alignas((UE::Core::Private::Atomic::TIsSupportedSize<sizeof(T)>::Value) ? alignof(UE::Core::Private::Atomic::TUnderlyingIntegerType_T<T>) : alignof(T))
-#endif
 	TAtomicBase_Basic
 {
 public:
@@ -235,17 +226,7 @@ public:
 	 */
 	FORCEINLINE T Load(EMemoryOrder Order = EMemoryOrder::SequentiallyConsistent) const
 	{
-#if USE_DEPRECATED_TATOMIC
-		switch (Order)
-		{
-			case EMemoryOrder::Relaxed:
-				return UE::Core::Private::Atomic::LoadRelaxed(&Element);
-		}
-
-		return UE::Core::Private::Atomic::Load(&Element);
-#else
 		return Element.load(ToStd(Order));
-#endif
 	}
 
 	/**
@@ -256,17 +237,7 @@ public:
 	 */
 	FORCEINLINE void Store(T Value, EMemoryOrder Order = EMemoryOrder::SequentiallyConsistent)
 	{
-#if USE_DEPRECATED_TATOMIC
-		switch (Order)
-		{
-			case EMemoryOrder::Relaxed:
-				return UE::Core::Private::Atomic::StoreRelaxed(&Element, Value);
-		}
-
-		return UE::Core::Private::Atomic::Store(&Element, Value);
-#else
 		Element.store(Value, ToStd(Order));
-#endif
 	}
 
 	/**
@@ -278,11 +249,7 @@ public:
 	 */
 	FORCEINLINE T Exchange(T Value)
 	{
-#if USE_DEPRECATED_TATOMIC
-		return UE::Core::Private::Atomic::Exchange(&Element, Value);
-#else
 		return Element.exchange(Value);
-#endif
 	}
 
 	/**
@@ -309,14 +276,7 @@ public:
 	 */
 	FORCEINLINE bool CompareExchange(T& Expected, T Value)
 	{
-#if USE_DEPRECATED_TATOMIC
-		T PrevValue = UE::Core::Private::Atomic::CompareExchange(&this->Element, Expected, Value);
-		bool bResult = PrevValue == Expected;
-		Expected = PrevValue;
-		return bResult;
-#else
 		return Element.compare_exchange_strong(Expected, Value);
-#endif
 	}
 
 protected:
@@ -327,11 +287,7 @@ protected:
 	{
 	}
 
-#if USE_DEPRECATED_TATOMIC
-	volatile T Element;
-#else
 	std::atomic<T> Element;
-#endif
 
 private:
 	static std::memory_order ToStd(EMemoryOrder Order)
@@ -358,11 +314,7 @@ public:
 	 */
 	FORCEINLINE T operator++()
 	{
-#if USE_DEPRECATED_TATOMIC
-		return UE::Core::Private::Atomic::IncrementExchange(&this->Element) + 1;
-#else
 		return this->Element.fetch_add(1) + 1;
-#endif
 	}
 
 	/**
@@ -372,11 +324,7 @@ public:
 	 */
 	FORCEINLINE T operator++(int)
 	{
-#if USE_DEPRECATED_TATOMIC
-		return UE::Core::Private::Atomic::IncrementExchange(&this->Element);
-#else
 		return this->Element.fetch_add(1);
-#endif
 	}
 
 	/**
@@ -388,11 +336,7 @@ public:
 	 */
 	FORCEINLINE T operator+=(DiffType Value)
 	{
-#if USE_DEPRECATED_TATOMIC
-		return UE::Core::Private::Atomic::AddExchange(&this->Element, Value) + Value;
-#else
 		return this->Element.fetch_add(Value) + Value;
-#endif
 	}
 
 	/**
@@ -402,11 +346,7 @@ public:
 	 */
 	FORCEINLINE T operator--()
 	{
-#if USE_DEPRECATED_TATOMIC
-		return UE::Core::Private::Atomic::DecrementExchange(&this->Element) - 1;
-#else
 		return this->Element.fetch_sub(1) - 1;
-#endif
 	}
 
 	/**
@@ -416,11 +356,7 @@ public:
 	 */
 	FORCEINLINE T operator--(int)
 	{
-#if USE_DEPRECATED_TATOMIC
-		return UE::Core::Private::Atomic::DecrementExchange(&this->Element);
-#else
 		return this->Element.fetch_sub(1);
-#endif
 	}
 
 	/**
@@ -432,11 +368,7 @@ public:
 	 */
 	FORCEINLINE T operator-=(DiffType Value)
 	{
-#if USE_DEPRECATED_TATOMIC
-		return UE::Core::Private::Atomic::SubExchange(&this->Element, Value) - Value;
-#else
 		return this->Element.fetch_sub(Value) - Value;
-#endif
 	}
 
 	/**
@@ -446,11 +378,7 @@ public:
 	 */
 	FORCEINLINE T IncrementExchange()
 	{
-#if USE_DEPRECATED_TATOMIC
-		return UE::Core::Private::Atomic::IncrementExchange(&this->Element);
-#else
 		return this->Element.fetch_add(1);
-#endif
 	}
 
 	/**
@@ -460,11 +388,7 @@ public:
 	 */
 	FORCEINLINE T DecrementExchange()
 	{
-#if USE_DEPRECATED_TATOMIC
-		return UE::Core::Private::Atomic::DecrementExchange(&this->Element);
-#else
 		return this->Element.fetch_sub(1);
-#endif
 	}
 
 	/**
@@ -476,11 +400,7 @@ public:
 	 */
 	FORCEINLINE T AddExchange(DiffType Value)
 	{
-#if USE_DEPRECATED_TATOMIC
-		return UE::Core::Private::Atomic::AddExchange(&this->Element, Value);
-#else
 		return this->Element.fetch_add(Value);
-#endif
 	}
 
 	/**
@@ -492,11 +412,7 @@ public:
 	 */
 	FORCEINLINE T SubExchange(DiffType Value)
 	{
-#if USE_DEPRECATED_TATOMIC
-		return UE::Core::Private::Atomic::SubExchange(&this->Element, Value);
-#else
 		return this->Element.fetch_sub(Value);
-#endif
 	}
 
 protected:
@@ -535,11 +451,7 @@ public:
 	 */
 	FORCEINLINE T operator&=(const T Value)
 	{
-#if USE_DEPRECATED_TATOMIC
-		return UE::Core::Private::Atomic::AndExchange(&this->Element, Value) & Value;
-#else
 		return this->Element.fetch_and(Value) & Value;
-#endif
 	}
 
 	/**
@@ -551,11 +463,7 @@ public:
 	 */
 	FORCEINLINE T operator|=(const T Value)
 	{
-#if USE_DEPRECATED_TATOMIC
-		return UE::Core::Private::Atomic::OrExchange(&this->Element, Value) | Value;
-#else
 		return this->Element.fetch_or(Value) | Value;
-#endif
 	}
 
 	/**
@@ -567,11 +475,7 @@ public:
 	 */
 	FORCEINLINE T operator^=(const T Value)
 	{
-#if USE_DEPRECATED_TATOMIC
-		return UE::Core::Private::Atomic::XorExchange(&this->Element, Value) ^ Value;
-#else
 		return this->Element.fetch_xor(Value) ^ Value;
-#endif
 	}
 
 	/**
@@ -583,11 +487,7 @@ public:
 	 */
 	FORCEINLINE T AndExchange(const T Value)
 	{
-#if USE_DEPRECATED_TATOMIC
-		return UE::Core::Private::Atomic::AndExchange(&this->Element, Value);
-#else
 		return this->Element.fetch_and(Value);
-#endif
 	}
 
 	/**
@@ -599,11 +499,7 @@ public:
 	 */
 	FORCEINLINE T OrExchange(const T Value)
 	{
-#if USE_DEPRECATED_TATOMIC
-		return UE::Core::Private::Atomic::OrExchange(&this->Element, Value);
-#else
 		return this->Element.fetch_or(Value);
-#endif
 	}
 
 	/**
@@ -615,11 +511,7 @@ public:
 	 */
 	FORCEINLINE T XorExchange(const T Value)
 	{
-#if USE_DEPRECATED_TATOMIC
-		return UE::Core::Private::Atomic::XorExchange(&this->Element, Value);
-#else
 		return this->Element.fetch_xor(Value);
-#endif
 	}
 
 
