- Worked around a crash with a new shadow cvar, causung issues when looking up for other platforms. Root cause unknown still, this is a bandaid to unblock people

diff --git a/Engine/Source/Runtime/Core/Private/HAL/ConsoleManager.cpp b/Engine/Source/Runtime/Core/Private/HAL/ConsoleManager.cpp
--- a/Engine/Source/Runtime/Core/Private/HAL/ConsoleManager.cpp
+++ b/Engine/Source/Runtime/Core/Private/HAL/ConsoleManager.cpp
@@ -3589,7 +3589,9 @@ void FConsoleManager::LoadAllPlatformCVars(FName PlatformName, const FString& De
 	{
 		// make sure the named cvar exists
 		IConsoleVariable* CVar = IConsoleManager::Get().FindConsoleVariable(*CVarName);
-		if (CVar == nullptr)
+		// @todo: disabling shadow cvars as this is crashing in one case, and likely aren't needed for other platforms
+		// but will debug this after checking in to unblock users
+		if (CVar == nullptr || CVar->IsShadowObject())
 		{
 			return;
 		}
