Marked TFuture, TSharedFuture and TPromise as final to prevent inheritance.

diff --git a/Engine/Source/Runtime/Core/Public/Async/Future.h b/Engine/Source/Runtime/Core/Public/Async/Future.h
--- a/Engine/Source/Runtime/Core/Public/Async/Future.h
+++ b/Engine/Source/Runtime/Core/Public/Async/Future.h
@@ -375,7 +375,7 @@ private:
  * Template for unshared futures.
  */
 template<typename ResultType>
-class TFuture
+class TFuture final
 	: public TFutureBase<UE::Core::Private::TInternalFutureState_T<ResultType>>
 {
 	template <typename>
@@ -512,7 +512,7 @@ private:
  * Template for shared futures.
  */
 template<typename ResultType>
-class TSharedFuture
+class TSharedFuture final
 	: public TFutureBase<UE::Core::Private::TInternalFutureState_T<ResultType>>
 {
 	template <typename>
@@ -579,7 +579,7 @@ private:
  * Template for promises.
  */
 template<typename ResultType>
-class TPromise
+class TPromise final
 {
 	using InternalResultType = UE::Core::Private::TInternalFutureState_T<ResultType>;
 	using StateType          = TSharedPtr<TFutureState<InternalResultType>>;
