StringBuilder: Support range-based for iteration

diff --git a/Engine/Source/Runtime/Core/Public/Misc/StringBuilder.h b/Engine/Source/Runtime/Core/Public/Misc/StringBuilder.h
--- a/Engine/Source/Runtime/Core/Public/Misc/StringBuilder.h
+++ b/Engine/Source/Runtime/Core/Public/Misc/StringBuilder.h
@@ -99,11 +99,21 @@ public:
 		Initialize(BufferPointer, BufferCapacity);
 	}
 
-	inline int32 Len() const { return int32(CurPos - Base); }
+	inline int32 Len() const
+	{
+		return int32(CurPos - Base);
+	}
 
 	/** Returns a pointer to Len() code units that are not necessarily null-terminated. */
-	inline CharType* GetData() UE_LIFETIMEBOUND { return Base; }
-	inline const CharType* GetData() const UE_LIFETIMEBOUND { return Base; }
+	inline CharType* GetData() UE_LIFETIMEBOUND
+	{
+		return Base;
+	}
+
+	inline const CharType* GetData() const UE_LIFETIMEBOUND
+	{
+		return Base;
+	}
 
 	/** Returns a pointer to a null-terminated string that is valid until the builder is mutated. */
 	inline const CharType* ToString() const UE_LIFETIMEBOUND
@@ -115,8 +125,7 @@ public:
 	/** Returns a pointer to a null-terminated string that is valid until the builder is mutated. */
 	inline const CharType* operator*() const UE_LIFETIMEBOUND
 	{
-		EnsureNulTerminated();
-		return Base;
+		return ToString();
 	}
 
 	/** Returns a view of the string that is valid until the builder is mutated. */
@@ -126,7 +135,7 @@ public:
 	}
 
 	/** Returns the last character, technically the last code unit. */
-	inline const CharType LastChar() const
+	inline CharType LastChar() const
 	{
 		return *(CurPos - 1);
 	}
@@ -218,8 +227,8 @@ public:
 	inline BuilderType& AppendChar(AppendedCharType Char)
 	{
 		if constexpr (TIsCharEncodingSimplyConvertibleTo_V<AppendedCharType, CharType>)
-	{
-		EnsureAdditionalCapacity(1);
+		{
+			EnsureAdditionalCapacity(1);
 			*CurPos++ = (CharType)Char;
 		}
 		else
@@ -228,7 +237,6 @@ public:
 			EnsureAdditionalCapacity(ConvertedLength);
 			CurPos = FPlatformString::Convert(CurPos, ConvertedLength, &Char, 1);
 		}
-
 		return *this;
 	}
 
@@ -262,7 +270,7 @@ public:
 		
 		if (Str.Len())
 		{
-			FMemory::Memcpy(Base + Pos, Str.GetData(), Str.Len() * sizeof(CharType));
+			Str.CopyString(Base + Pos, Str.Len());
 		}
 	}
 
@@ -383,6 +391,11 @@ public:
 private:
 	CORE_API static BuilderType& VARARGS AppendfImpl(BuilderType& Self, const CharType* Fmt, ...);
 
+	friend inline       CharType* begin(      TStringBuilderBase& Builder) { return Builder.GetData(); }
+	friend inline const CharType* begin(const TStringBuilderBase& Builder) { return Builder.GetData(); }
+	friend inline       CharType* end  (      TStringBuilderBase& Builder) { return Builder.GetData() + Builder.Len(); }
+	friend inline const CharType* end  (const TStringBuilderBase& Builder) { return Builder.GetData() + Builder.Len(); }
+
 protected:
 	inline void Initialize(CharType* InBase, int32 InCapacity)
 	{
