Moved the responsibility for adapting the future state into TFutureState.
Fixed some usage of TFuture<void>::Then being called with an int argument instead of no argument.

diff --git a/Engine/Plugins/Online/OnlineServices/Source/OnlineServicesCommon/Public/Online/OnlineAsyncOp.h b/Engine/Plugins/Online/OnlineServices/Source/OnlineServicesCommon/Public/Online/OnlineAsyncOp.h
--- a/Engine/Plugins/Online/OnlineServices/Source/OnlineServicesCommon/Public/Online/OnlineAsyncOp.h
+++ b/Engine/Plugins/Online/OnlineServices/Source/OnlineServicesCommon/Public/Online/OnlineAsyncOp.h
@@ -469,7 +469,7 @@ public:
 				if constexpr (TOnlineAsyncOpCallableTraits<CallableType>::bAsyncResult)
 				{
 					Callable(*PinnedOperation, MoveTempIfPossible(LastResult))
-						.Next([this, WeakOperation](const int& Value)
+						.Next([this, WeakOperation]()
 						{
 							TSharedPtr<TOnlineAsyncOp<OpType>> PinnedOperation2 = WeakOperation.Pin();
 							if (PinnedOperation2)
@@ -498,7 +498,7 @@ public:
 				if constexpr (TOnlineAsyncOpCallableTraits<CallableType>::bAsyncResult)
 				{
 					Callable(*PinnedOperation)
-						.Next([this, WeakOperation](const int& Value)
+						.Next([this, WeakOperation]()
 						{
 							TSharedPtr<TOnlineAsyncOp<OpType>> PinnedOperation2 = WeakOperation.Pin();
 							if (PinnedOperation2)
diff --git a/Engine/Plugins/Protocols/MQTT/Source/MQTTCore/Private/MQTTClient.cpp b/Engine/Plugins/Protocols/MQTT/Source/MQTTCore/Private/MQTTClient.cpp
--- a/Engine/Plugins/Protocols/MQTT/Source/MQTTCore/Private/MQTTClient.cpp
+++ b/Engine/Plugins/Protocols/MQTT/Source/MQTTCore/Private/MQTTClient.cpp
@@ -115,7 +115,7 @@ TFuture<void> FMQTTClient::Disconnect()
 		// Set immediately to prevent multiple requests
 		Connection->SetState(EMQTTState::Disconnecting);
 		return Connection->QueueOperation<FMQTTDisconnectOperation>(FMQTTDisconnectPacket())
-			.Next([Self = TWeakPtr<FMQTTClient, ESPMode::ThreadSafe>(AsShared())](int32)
+			.Next([Self = TWeakPtr<FMQTTClient, ESPMode::ThreadSafe>(AsShared())]()
 			{
 				if (!Self.IsValid())
 				{
@@ -166,7 +166,7 @@ TFuture<bool> FMQTTClient::Publish(const FString& InTopic, const TArray<uint8>&
 		{
 		case EMQTTQualityOfService::Once:
 			return Connection->QueueOperation<FMQTTPublishOperationQoS0>(MoveTemp(PublishPacket))
-			.Next([Self = TWeakPtr<FMQTTClient, ESPMode::ThreadSafe>(AsShared())](auto /*InResponse*/)
+			.Next([Self = TWeakPtr<FMQTTClient, ESPMode::ThreadSafe>(AsShared())]()
 			{
 				if(!Self.IsValid())
 				{
diff --git a/Engine/Plugins/Protocols/MQTT/Source/MQTTCore/Private/MQTTClientObject.cpp b/Engine/Plugins/Protocols/MQTT/Source/MQTTCore/Private/MQTTClientObject.cpp
--- a/Engine/Plugins/Protocols/MQTT/Source/MQTTCore/Private/MQTTClientObject.cpp
+++ b/Engine/Plugins/Protocols/MQTT/Source/MQTTCore/Private/MQTTClientObject.cpp
@@ -130,7 +130,7 @@ void UMQTTClientObject::Disconnect(const FOnDisconnectDelegate& InOnDisconnect)
 	{
 		const FGuid ClientId = MQTTClient->GetClientId();
 		MQTTClient->Disconnect()
-		.Next([&](int32)
+		.Next([&]()
 		{
 			if(IsValid(this))
 			{
diff --git a/Engine/Source/Runtime/Core/Public/Async/Future.h b/Engine/Source/Runtime/Core/Public/Async/Future.h
--- a/Engine/Source/Runtime/Core/Public/Async/Future.h
+++ b/Engine/Source/Runtime/Core/Public/Async/Future.h
@@ -12,6 +12,7 @@
 #include "HAL/Event.h"
 #include "HAL/PlatformProcess.h"
 #include "Misc/ScopeLock.h"
+#include "Templates/Models.h"
 #include "Templates/Requires.h"
 #include <type_traits>
 
@@ -21,17 +22,14 @@ template<typename ResultType> class TPromise;
 
 namespace UE::Core::Private
 {
-	template <typename T>
-	using TInternalFutureState_T =
-		std::conditional_t<
-			std::is_lvalue_reference_v<T>,
-			std::remove_reference_t<T>*,
-			std::conditional_t<
-				std::is_void_v<T>,
-				int,
-				T
-			>
-		>;
+	// A concept check for a type being callable
+	struct CIntCallable
+	{
+		template <typename Type>
+		auto Requires(Type& Callable, int32 Val) -> decltype(
+			Callable(Val)
+		);
+	};
 }
 
 /**
@@ -143,10 +141,21 @@ private:
 /**
  * Implements the internal state of asynchronous return values (futures).
  */
-template<typename InternalResultType>
+template<typename ResultType>
 class TFutureState
 	: public FFutureState
 {
+	using InternalResultType =
+		std::conditional_t<
+			std::is_lvalue_reference_v<ResultType>,
+			std::remove_reference_t<ResultType>*,
+			std::conditional_t<
+				std::is_void_v<ResultType>,
+				int,
+				ResultType
+			>
+		>;
+
 public:
 
 	/** Default constructor. */
@@ -226,7 +235,7 @@ private:
 /**
  * Abstract base template for futures and shared futures.
  */
-template<typename InternalResultType>
+template <typename ResultType>
 class TFutureBase
 {
 public:
@@ -295,7 +304,7 @@ public:
 
 protected:
 
-	typedef TSharedPtr<TFutureState<InternalResultType>> StateType;
+	typedef TSharedPtr<TFutureState<ResultType>> StateType;
 
 	/** Default constructor. */
 	TFutureBase() = default;
@@ -334,7 +343,7 @@ protected:
 	 * Set a completion callback that will be called once the future completes
 	 *	or immediately if already completed
 	 *
-	 * @param Continuation a continuation taking an argument of type TFuture<InternalResultType>
+	 * @param Continuation a continuation taking an argument of type TFuture<ResultType>
 	 * @return nothing at the moment but could return another future to allow future chaining
 	 */
 	template<typename Func>
@@ -344,7 +353,7 @@ protected:
 	 * Convenience wrapper for Then that
 	 *	set a completion callback that will be called once the future completes
 	 *	or immediately if already completed
-	 * @param Continuation a continuation taking an argument of type InternalResultType
+	 * @param Continuation a continuation taking an argument of type ResultType
 	 * @return nothing at the moment but could return another future to allow future chaining
 	 */
 	template<typename Func>
@@ -376,7 +385,7 @@ private:
  */
 template<typename ResultType>
 class TFuture final
-	: public TFutureBase<UE::Core::Private::TInternalFutureState_T<ResultType>>
+	: public TFutureBase<ResultType>
 {
 	template <typename>
 	friend class TPromise;
@@ -384,7 +393,7 @@ class TFuture final
 	template <typename>
 	friend class TFutureBase;
 
-	using BaseType = TFutureBase<UE::Core::Private::TInternalFutureState_T<ResultType>>;
+	using BaseType = TFutureBase<ResultType>;
 
 public:
 
@@ -513,12 +522,12 @@ private:
  */
 template<typename ResultType>
 class TSharedFuture final
-	: public TFutureBase<UE::Core::Private::TInternalFutureState_T<ResultType>>
+	: public TFutureBase<ResultType>
 {
 	template <typename>
 	friend class TFuture;
 
-	using BaseType = TFutureBase<UE::Core::Private::TInternalFutureState_T<ResultType>>;
+	using BaseType = TFutureBase<ResultType>;
 
 public:
 
@@ -581,8 +590,7 @@ private:
 template<typename ResultType>
 class TPromise final
 {
-	using InternalResultType = UE::Core::Private::TInternalFutureState_T<ResultType>;
-	using StateType          = TSharedPtr<TFutureState<InternalResultType>>;
+	using StateType = TSharedPtr<TFutureState<ResultType>>;
 
 	// This is necessary because we can't form references to void for the parameter of Set().
 	using SetType = std::conditional_t<std::is_void_v<ResultType>, int, ResultType>;
@@ -591,7 +599,7 @@ public:
 
 	/** Default constructor (creates a new shared state). */
 	TPromise()
-		: State(MakeShared<TFutureState<InternalResultType>>())
+		: State(MakeShared<TFutureState<ResultType>>())
 	{
 	}
 
@@ -601,7 +609,7 @@ public:
 	 * @param CompletionCallback A function that is called when the future state is completed.
 	 */
 	TPromise(TUniqueFunction<void()>&& CompletionCallback)
-		: State(MakeShared<TFutureState<InternalResultType>>(MoveTemp(CompletionCallback)))
+		: State(MakeShared<TFutureState<ResultType>>(MoveTemp(CompletionCallback)))
 	{
 	}
 
@@ -777,18 +785,18 @@ namespace FutureDetail
 }
 
 // Then implementation
-template<typename InternalResultType>
+template<typename ResultType>
 template<typename Func>
-auto TFutureBase<InternalResultType>::Then(Func Continuation) //-> TFuture<decltype(Continuation(MoveTemp(TFuture<InternalResultType>())))>
+auto TFutureBase<ResultType>::Then(Func Continuation) //-> TFuture<decltype(Continuation(MoveTemp(TFuture<ResultType>())))>
 {
 	check(IsValid());
-	using ReturnValue = decltype(Continuation(MoveTemp(TFuture<InternalResultType>())));
+	using ReturnValue = decltype(Continuation(MoveTemp(TFuture<ResultType>())));
 
 	TPromise<ReturnValue> Promise;
 	TFuture<ReturnValue> FutureResult = Promise.GetFuture();
 	TUniqueFunction<void()> Callback = [PromiseCapture = MoveTemp(Promise), ContinuationCapture = MoveTemp(Continuation), StateCapture = this->State]() mutable
 	{
-		FutureDetail::SetPromiseValue(PromiseCapture, ContinuationCapture, TFuture<InternalResultType>(MoveTemp(StateCapture)));
+		FutureDetail::SetPromiseValue(PromiseCapture, ContinuationCapture, TFuture<ResultType>(MoveTemp(StateCapture)));
 	};
 
 	// This invalidate this future.
@@ -798,13 +806,29 @@ auto TFutureBase<InternalResultType>::Then(Func Continuation) //-> TFuture<declt
 }
 
 // Next implementation
-template<typename InternalResultType>
+template<typename ResultType>
 template<typename Func>
-auto TFutureBase<InternalResultType>::Next(Func Continuation) //-> TFuture<decltype(Continuation(Consume()))>
+auto TFutureBase<ResultType>::Next(Func Continuation) //-> TFuture<decltype(Continuation(Consume()))>
 {
-	return this->Then([Continuation = MoveTemp(Continuation)](TFuture<InternalResultType> Self) mutable
+	return this->Then([Continuation = MoveTemp(Continuation)](TFuture<ResultType> Self) mutable
 	{
-		return Continuation(Self.Consume());
+		if constexpr (std::is_void_v<ResultType>)
+		{
+			Self.Consume();
+			if constexpr (TModels_V<UE::Core::Private::CIntCallable, Func>)
+			{
+				UE_STATIC_DEPRECATE(5.6, true, "Passing continuations to TFuture<void>::Next which take int parameters has been deprecated - please remove the parameter.");
+				return Continuation(1);
+			}
+			else
+			{
+				return Continuation();
+			}
+		}
+		else
+		{
+			return Continuation(Self.Consume());
+		}
 	});
 }
 
