Containers: TDirectoryTree: Add ranged-for iteration and other Iteration functions.
#rnx

diff --git a/Engine/Source/Runtime/Core/Public/Containers/ArrowWrapper.h b/Engine/Source/Runtime/Core/Public/Containers/ArrowWrapper.h
new file mode 100644
--- /dev/null
+++ b/Engine/Source/Runtime/Core/Public/Containers/ArrowWrapper.h
@@ -0,0 +1,45 @@
+// Copyright Epic Games, Inc. All Rights Reserved.
+
+#pragma once
+
+#include "Templates/UnrealTemplate.h"
+
+/**
+ * Wrapper type for Iterators that return a structure by value from operator*, and want to support Iter->Property
+ * to behave the same as (*Iter).Property.
+ * When returning a structure by reference, operator* automatically handles -> as well, but when returning by value
+ * it does not.
+ * 
+ * Example:
+ * struct FIterator
+ * {
+ * public:
+ *    ...
+ *		// The problematic return-by-value structure from operator*
+ *		// Without operator->, (*Iter).Key compiles but Iter->Key does not.
+ *		TPair<FStringView, ViewedValueType&> operator*() const;
+ * 
+ *      // TArrowWrapper takes the value from operator* and allows -> to access it.
+ *		TArrowWrapper<TPair<FStringView, ViewedValueType&>> operator->() const;
+ * };
+ */
+template <typename WrappedType>
+struct TArrowWrapper
+{
+	explicit TArrowWrapper(const WrappedType& InValue)
+		: Value(InValue)
+	{
+	}
+
+	explicit TArrowWrapper(WrappedType&& InValue)
+		: Value(MoveTemp(InValue))
+	{
+	}
+
+	const WrappedType* operator->() const
+	{
+		return &this->Value;
+	}
+
+	WrappedType Value;
+};
diff --git a/Engine/Source/Runtime/Core/Public/Containers/DirectoryTree.h b/Engine/Source/Runtime/Core/Public/Containers/DirectoryTree.h
--- a/Engine/Source/Runtime/Core/Public/Containers/DirectoryTree.h
+++ b/Engine/Source/Runtime/Core/Public/Containers/DirectoryTree.h
@@ -3,6 +3,7 @@
 #pragma once
 
 #include "Containers/Array.h"
+#include "Containers/ArrowWrapper.h"
 #include "Containers/StringView.h"
 #include "Containers/UnrealString.h"
 #include "HAL/Platform.h"
@@ -10,6 +11,7 @@
 #include "Misc/PathViews.h"
 #include "Misc/ScopeExit.h"
 #include "Misc/StringBuilder.h"
+#include "Templates/Tuple.h"
 #include "Templates/TypeCompatibleBytes.h"
 #include "Templates/UniquePtr.h"
 
@@ -81,6 +83,18 @@ ENUM_CLASS_FLAGS(EDirectoryTreeGetFlags);
 template <typename ValueType>
 class TDirectoryTree
 {
+public:
+	template <typename ViewedValueType> struct TIterator;
+	using FIterator = TIterator<ValueType>;
+	using FConstIterator = TIterator<const ValueType>;
+	template <typename ViewedValueType> struct TPointerIterator;
+	using FPointerIterator = TPointerIterator<ValueType>;
+	using FConstPointerIterator = TPointerIterator<const ValueType>;
+	struct FIterationSentinel;
+private:
+	struct FIteratorInternal;
+	struct FTreeNode;
+
 public:
 	TDirectoryTree();
 
@@ -164,6 +178,80 @@ public:
 	bool TryGetChildren(FStringView Path, TArray<FString>& OutRelativeChildNames,
 		EDirectoryTreeGetFlags Flags = EDirectoryTreeGetFlags::None) const;
 
+	/** Ranged-for accessor for elements in the tree. Same behavior as @see CreateIterator. */
+	FIterator begin();
+	FConstIterator begin() const;
+	FIterationSentinel end() const;
+
+	/** Iterator accessor for path,value pairs that were added to the tree. @see FIterator. */
+	FIterator CreateIterator();
+	FConstIterator CreateConstIterator() const;
+
+	/**
+	 * Iterator accessor for all paths, added or implied, in the tree. For paths that were added to the tree, the
+	 * ValueType* on the iterator is non-null and points to the added Value.
+	 * For implied paths that are parent directories of child paths in the tree the Value pointer of the iterator is
+	 * the value of their closest parent in the tree, or nullptr if none of their parents were added to the tree.
+	 * @see FPointerIterator.
+	 */
+	FPointerIterator CreateIteratorForImplied();
+	FConstPointerIterator CreateConstIteratorForImplied() const;
+
+public:
+	/**
+	 * Sentinel type used to mark the end of ranged-for iteration. Iterator operator!= functions
+	 * that take this type return (bool)*this.
+	 */
+	struct FIterationSentinel
+	{
+	};
+
+	/**
+	 * Iterator used when iterating paths added to the tree, and skipping implied parent directories
+	 * that were not added to the tree and have no Value data. The return value of operator* has a (possibly const)
+	 * ValueType&.
+	 */
+	template <typename ViewedValueType>
+	struct TIterator
+	{
+	public:
+		TPair<FStringView, ViewedValueType&> operator*() const;
+		bool operator!=(const FIterationSentinel&) const;
+		TIterator& operator++();
+
+		explicit operator bool() const;
+
+		// Support Iter->Key and Iter->Value.
+		TArrowWrapper<TPair<FStringView, ViewedValueType&>> operator->() const;
+	private:
+		friend class TDirectoryTree<ValueType>;
+		TIterator(FTreeNode& Root, EDirectoryTreeGetFlags InFlags, TCHAR InPathSeparator);
+		FIteratorInternal Internal;
+	};
+
+	/**
+	 * Iterator used when iterating all paths in the tree, both added or implied. The return value of
+	 * operator* has a (possibly const) ValueType*, which will be non-null for all paths that were added to the tree
+	 * and null for the implied parent directories that were not added.
+	 */
+	template <typename ViewedValueType>
+	struct TPointerIterator
+	{
+	public:
+		TPair<FStringView, ViewedValueType*> operator*() const;
+		bool operator!=(const FIterationSentinel&) const;
+		TPointerIterator& operator++();
+
+		explicit operator bool() const;
+
+		// Support Iter->Key and Iter->Value.
+		TArrowWrapper<TPair<FStringView, ViewedValueType*>> operator->() const;
+	private:
+		friend class TDirectoryTree<ValueType>;
+		TPointerIterator(FTreeNode& Root, EDirectoryTreeGetFlags InFlags, TCHAR InPathSeparator);
+		FIteratorInternal Internal;
+	};
+
 private:
 
 	/**
@@ -274,6 +362,8 @@ private:
 		void Realloc(int32 NewCapacity);
 
 	private:
+		friend struct FIteratorInternal;
+
 		TTypeCompatibleBytes<ValueType> Value;
 		TUniquePtr<FString[]> RelPaths;
 		TUniquePtr<FTreeNode[]> ChildNodes;
@@ -281,6 +371,50 @@ private:
 		int32 CapacityChildNodes = 0;
 	};
 
+	struct FIteratorInternal
+	{
+	public:
+		FIteratorInternal(FTreeNode& Root, EDirectoryTreeGetFlags InFlags, TCHAR InPathSeparator);
+		// Iterators for TDirectoryTree are heavy-weight; they keep a TArray and a TStringBuilder.
+		// Do not copy or move them use them only in for loops, or on the heap:
+		// for (TDirectoryTree<...>::FIterator Iter(Tree.CreateIterator()); Iter; ++Iter)
+		// TUniquePtr<TDirectoryTree<...>::FIterator> IterPtr = MakeUnique<TDirectoryTree<...>::FIterator>(Tree.CreateIterator())
+		FIteratorInternal(FIteratorInternal&& Other) = delete;
+		FIteratorInternal& operator=(FIteratorInternal&& Other) = delete;
+		FIteratorInternal(const FIteratorInternal& Other) = delete;
+		FIteratorInternal& operator=(const FIteratorInternal& Other) = delete;
+
+		TPair<FStringView, ValueType*> operator*() const;
+		bool operator!=(const FIterationSentinel&) const;
+		FIteratorInternal& operator++();
+
+		explicit operator bool() const;
+
+	private:
+		struct FStackData
+		{
+			explicit FStackData(FTreeNode& InNode, ValueType* ParentClosestValue);
+
+			FStringView RemainingChildRelPath;
+			FTreeNode* Node;
+			ValueType* ClosestValue;
+			int32 ChildIndex;
+			int32 PathLenBeforeChildNode;
+			bool bChildRelPathInitialized;
+			bool bChildNodeInitialized;
+		};
+
+	private:
+		void TraverseToValid();
+		void IncrementRemainingChildRelPath(FStackData* Data);
+
+	private:
+		TArray<FStackData, TInlineAllocator<10>> Stack;
+		TStringBuilder<256> Path;
+		EDirectoryTreeGetFlags Flags;
+		TCHAR PathSeparator;
+	};
+
 private:
 	bool NormalizePathForReading(FStringView& Path, FStringBuilderBase& NormalizeBuffer) const;
 	bool NormalizePathForWriting(FStringView& Path, FStringBuilderBase& NormalizeBuffer);
@@ -303,7 +437,7 @@ inline TDirectoryTree<ValueType>::TDirectoryTree()
 }
 
 template <typename ValueType>
-inline ValueType& TDirectoryTree<ValueType>::FindOrAdd(FStringView Path, bool* bOutExisted)
+ValueType& TDirectoryTree<ValueType>::FindOrAdd(FStringView Path, bool* bOutExisted)
 {
 	if (Path.IsEmpty())
 	{
@@ -354,7 +488,7 @@ inline void TDirectoryTree<ValueType>::Empty()
 }
 
 template <typename ValueType>
-inline void TDirectoryTree<ValueType>::Remove(FStringView Path, bool* bOutExisted)
+void TDirectoryTree<ValueType>::Remove(FStringView Path, bool* bOutExisted)
 {
 	bool bExisted;
 	if (Path.IsEmpty())
@@ -773,7 +907,7 @@ inline void TDirectoryTree<ValueType>::FTreeNode::RemoveValue()
 }
 
 template <typename ValueType>
-inline ValueType& TDirectoryTree<ValueType>::FTreeNode::FindOrAdd(FStringView InRelPath, bool& bOutExisted)
+ValueType& TDirectoryTree<ValueType>::FTreeNode::FindOrAdd(FStringView InRelPath, bool& bOutExisted)
 {
 	FStringView FirstComponent;
 	FStringView RemainingPath;
@@ -890,7 +1024,7 @@ inline ValueType& TDirectoryTree<ValueType>::FTreeNode::FindOrAdd(FStringView In
 }
 
 template <typename ValueType>
-inline void TDirectoryTree<ValueType>::FTreeNode::Remove(FStringView InRelPath, bool& bOutExisted)
+void TDirectoryTree<ValueType>::FTreeNode::Remove(FStringView InRelPath, bool& bOutExisted)
 {
 	check(!InRelPath.IsEmpty());
 
@@ -946,7 +1080,7 @@ inline void TDirectoryTree<ValueType>::FTreeNode::Remove(FStringView InRelPath,
 }
 
 template <typename ValueType>
-inline ValueType* TDirectoryTree<ValueType>::FTreeNode::Find(FStringView InRelPath)
+ValueType* TDirectoryTree<ValueType>::FTreeNode::Find(FStringView InRelPath)
 {
 	check(!InRelPath.IsEmpty());
 
@@ -981,7 +1115,7 @@ inline ValueType* TDirectoryTree<ValueType>::FTreeNode::Find(FStringView InRelPa
 }
 
 template<typename ValueType>
-inline bool TDirectoryTree<ValueType>::FTreeNode::ContainsChildPaths(FStringView InRelPath) const
+bool TDirectoryTree<ValueType>::FTreeNode::ContainsChildPaths(FStringView InRelPath) const
 {
 	if (InRelPath.IsEmpty())
 	{
@@ -1060,7 +1194,7 @@ inline bool TDirectoryTree<ValueType>::FTreeNode::ContainsChildPaths(FStringView
 }
 
 template <typename ValueType>
-inline bool TDirectoryTree<ValueType>::FTreeNode::TryGetChildren(FStringBuilderBase& ReportPathPrefix,
+bool TDirectoryTree<ValueType>::FTreeNode::TryGetChildren(FStringBuilderBase& ReportPathPrefix,
 	TCHAR InPathSeparator, FStringView InRelPath, TArray<FString>& OutRelativeChildNames,
 	EDirectoryTreeGetFlags Flags) const
 {
@@ -1274,7 +1408,7 @@ inline bool TDirectoryTree<ValueType>::FTreeNode::TryGetChildren(FStringBuilderB
 }
 
 template <typename ValueType>
-inline ValueType* TDirectoryTree<ValueType>::FTreeNode::TryFindClosestPath(FStringView InRelPath,
+ValueType* TDirectoryTree<ValueType>::FTreeNode::TryFindClosestPath(FStringView InRelPath,
 	FStringBuilderBase* OutPath, TCHAR InPathSeparator)
 {
 	check(!InRelPath.IsEmpty());
@@ -1386,7 +1520,7 @@ inline int32 TDirectoryTree<ValueType>::FTreeNode::FindInsertionIndex(FStringVie
 }
 
 template <typename ValueType>
-inline typename TDirectoryTree<ValueType>::FTreeNode& TDirectoryTree<ValueType>::FTreeNode::InsertChildNode(int32 InsertionIndex, FString&& RelPath,
+typename TDirectoryTree<ValueType>::FTreeNode& TDirectoryTree<ValueType>::FTreeNode::InsertChildNode(int32 InsertionIndex, FString&& RelPath,
 	FTreeNode&& ChildNode)
 {
 	int32 NumChildNodes = GetNumChildNodes();
@@ -1425,7 +1559,7 @@ inline typename TDirectoryTree<ValueType>::FTreeNode& TDirectoryTree<ValueType>:
 }
 
 template <typename ValueType>
-inline void TDirectoryTree<ValueType>::FTreeNode::RemoveChildNodeAt(int32 RemoveIndex)
+void TDirectoryTree<ValueType>::FTreeNode::RemoveChildNodeAt(int32 RemoveIndex)
 {
 	int32 NumChildNodes = GetNumChildNodes();
 	check(0 <= RemoveIndex && RemoveIndex < NumChildNodes);
@@ -1440,7 +1574,7 @@ inline void TDirectoryTree<ValueType>::FTreeNode::RemoveChildNodeAt(int32 Remove
 }
 
 template <typename ValueType>
-inline void TDirectoryTree<ValueType>::FTreeNode::ConditionalCompactNode(FString& RelPath, FTreeNode& ChildNode)
+void TDirectoryTree<ValueType>::FTreeNode::ConditionalCompactNode(FString& RelPath, FTreeNode& ChildNode)
 {
 	if (ChildNode.HasValue())
 	{
@@ -1462,7 +1596,7 @@ inline void TDirectoryTree<ValueType>::FTreeNode::ConditionalCompactNode(FString
 }
 
 template <typename ValueType>
-inline void TDirectoryTree<ValueType>::FTreeNode::Realloc(int32 NewCapacity)
+void TDirectoryTree<ValueType>::FTreeNode::Realloc(int32 NewCapacity)
 {
 	if (CapacityChildNodes == NewCapacity)
 	{
@@ -1503,3 +1637,331 @@ inline void TDirectoryTree<ValueType>::FTreeNode::FixupDirectChildrenPathSeparat
 	}
 }
 
+template <typename ValueType>
+inline TDirectoryTree<ValueType>::FIteratorInternal::FIteratorInternal(FTreeNode& Root, EDirectoryTreeGetFlags InFlags, TCHAR InPathSeparator)
+	: Flags(InFlags)
+	, PathSeparator(InPathSeparator)
+{
+	FStackData& RootData = Stack.Emplace_GetRef(Root, nullptr /* ParentClosestValue */);
+	// The root node always exists. It can have a value, if empty string was pushed into the tree. Otherwise it does
+	// not have a value, and is not treated as an implied parent. The highest parent directory of any added path other
+	// than empty string is a child node of the root node.
+	if (!RootData.Node->HasValue())
+	{
+		// To allow our node-based iteration to work and skip recording a value for the root, we need to skip past the
+		// "Node's value itself" state (ChildIndex == -1) for the root node, and start at ChildIndex=0.
+		RootData.ChildIndex = 0;
+	}
+	TraverseToValid();
+}
+
+template <typename ValueType>
+inline TPair<FStringView, ValueType*> TDirectoryTree<ValueType>::FIteratorInternal::operator*() const
+{
+	check(!Stack.IsEmpty());
+	const FStackData& Data = Stack.Last();
+	FTreeNode* Node = Data.Node;
+	return TPair<FStringView, ValueType*>(Path, Data.ClosestValue);
+}
+
+template <typename ValueType>
+inline TDirectoryTree<ValueType>::FIteratorInternal::operator bool() const
+{
+	return !Stack.IsEmpty();
+}
+
+template <typename ValueType>
+inline bool TDirectoryTree<ValueType>::FIteratorInternal::operator!=(const FIterationSentinel&) const
+{
+	return (bool)*this;
+}
+
+template <typename ValueType>
+typename TDirectoryTree<ValueType>::FIteratorInternal& TDirectoryTree<ValueType>::FIteratorInternal::operator++()
+{
+	check((bool)*this); // It is invalid to call operator++ on the iterator when it is done
+
+	// Stack.Last is pointing at an existing value, move past it
+	FStackData* Data = &Stack.Last();
+	if (Data->ChildIndex == -1)
+	{
+		// Pointing to the Node's value itself, move to the first child.
+		Data->ChildIndex = 0;
+	}
+	else
+	{
+		// Inbetween states are not possible for the top node in the stack. TraverseToValid moves past
+		// inbetween states before returning.
+		check(Data->ChildIndex < Data->Node->GetNumChildNodes());
+		check(Data->bChildRelPathInitialized);
+		// Empty RemainingChildRelPath is an inbetween state, with or without bChildNodeInitialized.
+		// TraverseToValid moves past it to add the next child node to the stack before returning.
+		check(!Data->RemainingChildRelPath.IsEmpty());
+
+		// Pointing to somewhere in the iteration of directories in the childnode's relative path.
+		IncrementRemainingChildRelPath(Data);
+	}
+
+	TraverseToValid();
+	return *this;
+}
+
+template <typename ValueType>
+void TDirectoryTree<ValueType>::FIteratorInternal::TraverseToValid()
+{
+	while (!Stack.IsEmpty())
+	{
+		FStackData* Data = &Stack.Last();
+		for (;;)
+		{
+			if (Data->ChildIndex == -1)
+			{
+				// Pointing to the Node's value itself. Return if the node has a value or implied children are allowed
+				if (EnumHasAnyFlags(Flags, EDirectoryTreeGetFlags::ImpliedChildren)
+					|| Data->Node->HasValue())
+				{
+					return;
+				}
+
+				if (!EnumHasAnyFlags(Flags, EDirectoryTreeGetFlags::Recursive) && Stack.Num() > 1)
+				{
+					// In the non-recursive case, we are iterating over the node in Stack[0], and Stack[1]
+					// is each child node that we need to report. Do not iterate over the children of Stack[1];
+					// skip to the end of its children.
+					Data->ChildIndex = Data->Node->GetNumChildNodes();
+				}
+				else
+				{
+					Data->ChildIndex = 0;
+				}
+			}
+			else if (Data->ChildIndex < Data->Node->GetNumChildNodes())
+			{
+				if (!Data->bChildRelPathInitialized)
+				{
+					Data->bChildRelPathInitialized = true;
+					Data->PathLenBeforeChildNode = Path.Len();
+
+					// If we are not returning implied children then skip straight to the recursive child node
+					if (!EnumHasAnyFlags(Flags, EDirectoryTreeGetFlags::ImpliedChildren))
+					{
+						Data->RemainingChildRelPath = FStringView();
+						int32 FixupStart = Path.Len();
+						FPathViews::Append(Path, Data->Node->RelPaths[Data->ChildIndex]);
+						UE::DirectoryTree::FixupPathSeparator(Path, FixupStart, PathSeparator);
+					}
+					else
+					{
+						Data->RemainingChildRelPath = Data->Node->RelPaths[Data->ChildIndex];
+						// Pull the first component of the child relpath into our path, and
+						// continue the loop to check for whether it is an implied parent or
+						// the child node itself.
+						IncrementRemainingChildRelPath(Data);
+					}
+				}
+				else if (!Data->RemainingChildRelPath.IsEmpty())
+				{
+					// Pointing to somewhere in the iteration of directories in the childnode's relative path.
+					// Return it as an implied node.
+					return;
+				}
+				else if (!Data->bChildNodeInitialized)
+				{
+					Data->bChildNodeInitialized = true;
+					Stack.Emplace(Data->Node->ChildNodes[Data->ChildIndex], Data->ClosestValue);
+					break; // Return to the outer loop to traverse into the node we just pushed
+				}
+				else
+				{
+					// Pointing to the recursive child node that we just popped, move to the next.
+					++Data->ChildIndex;
+
+					Data->bChildRelPathInitialized = false;
+					Data->bChildNodeInitialized = false;
+					Data->RemainingChildRelPath = FStringView();
+					Path.RemoveSuffix(Path.Len() - Data->PathLenBeforeChildNode);
+					Data->PathLenBeforeChildNode = 0;
+					// Continue on to the next child
+				}
+			}
+			else
+			{
+				Stack.Pop(EAllowShrinking::No);
+				break; // Return to the outer loop to return traversal to the parent node
+			}
+		}
+	}
+}
+
+template <typename ValueType>
+inline void TDirectoryTree<ValueType>::FIteratorInternal::IncrementRemainingChildRelPath(FStackData* Data)
+{
+	FStringView FirstComponent;
+	FStringView RemainingPath;
+	FPathViews::SplitFirstComponent(Data->RemainingChildRelPath, FirstComponent, RemainingPath);
+	Data->RemainingChildRelPath = RemainingPath;
+
+	int32 FixupStart = Path.Len();
+	FPathViews::Append(Path, FirstComponent);
+	UE::DirectoryTree::FixupPathSeparator(Path, FixupStart, PathSeparator);
+}
+
+template <typename ValueType>
+inline TDirectoryTree<ValueType>::FIteratorInternal::FStackData::FStackData(
+	FTreeNode& InNode, ValueType* ParentClosestValue)
+	: Node(&InNode)
+	, ClosestValue(InNode.HasValue() ? &InNode.GetValue() : ParentClosestValue)
+	, ChildIndex(-1)
+	, PathLenBeforeChildNode(0)
+	, bChildRelPathInitialized(false)
+	, bChildNodeInitialized(false)
+{
+}
+
+template <typename ValueType>
+template <typename ViewedValueType>
+inline TDirectoryTree<ValueType>::TIterator<ViewedValueType>::TIterator(FTreeNode& Root,
+	EDirectoryTreeGetFlags InFlags, TCHAR InPathSeparator)
+	: Internal(Root, InFlags, InPathSeparator)
+{
+}
+
+template <typename ValueType>
+template <typename ViewedValueType>
+inline TPair<FStringView, ViewedValueType&> TDirectoryTree<ValueType>::TIterator<ViewedValueType>::operator*() const
+{
+	check((bool)Internal);
+	TPair<FStringView, ValueType*> InternalPair = *Internal;
+	check(InternalPair.Value != nullptr);
+	return TPair<FStringView, ViewedValueType&>(InternalPair.Key, *InternalPair.Value);
+}
+
+template <typename ValueType>
+template <typename ViewedValueType>
+inline TArrowWrapper<TPair<FStringView, ViewedValueType&>>
+TDirectoryTree<ValueType>::TIterator<ViewedValueType>::operator->() const
+{
+	return TArrowWrapper<TPair<FStringView, ViewedValueType&>>(**this);
+}
+
+template <typename ValueType>
+template <typename ViewedValueType>
+inline bool TDirectoryTree<ValueType>::TIterator<ViewedValueType>::operator!=(const FIterationSentinel&) const
+{
+	return (bool)*this;
+}
+
+template <typename ValueType>
+template <typename ViewedValueType>
+inline typename TDirectoryTree<ValueType>::template TIterator<ViewedValueType>& 
+TDirectoryTree<ValueType>::TIterator<ViewedValueType>::operator++()
+{
+	++Internal;
+	return *this;
+}
+
+template <typename ValueType>
+template <typename ViewedValueType>
+inline TDirectoryTree<ValueType>::TIterator<ViewedValueType>::operator bool() const
+{
+	return (bool)Internal;
+}
+
+template <typename ValueType>
+template <typename ViewedValueType>
+inline TDirectoryTree<ValueType>::TPointerIterator<ViewedValueType>::TPointerIterator(FTreeNode& Root,
+	EDirectoryTreeGetFlags InFlags, TCHAR InPathSeparator)
+	: Internal(Root, InFlags, InPathSeparator)
+{
+}
+
+template <typename ValueType>
+template <typename ViewedValueType>
+inline TPair<FStringView, ViewedValueType*>
+TDirectoryTree<ValueType>::TPointerIterator<ViewedValueType>::operator*() const
+{
+	check((bool)Internal);
+	TPair<FStringView, ValueType*> InternalPair = *Internal;
+	return TPair<FStringView, ViewedValueType*>(InternalPair.Key, InternalPair.Value);
+}
+
+template <typename ValueType>
+template <typename ViewedValueType>
+inline TArrowWrapper<TPair<FStringView, ViewedValueType*>>
+TDirectoryTree<ValueType>::TPointerIterator<ViewedValueType>::operator->() const
+{
+	return TArrowWrapper<TPair<FStringView, ViewedValueType*>>(**this);
+}
+
+template <typename ValueType>
+template <typename ViewedValueType>
+inline bool TDirectoryTree<ValueType>::TPointerIterator<ViewedValueType>::operator!=(const FIterationSentinel&) const
+{
+	return (bool)Internal;
+}
+
+template <typename ValueType>
+template <typename ViewedValueType>
+inline typename TDirectoryTree<ValueType>::template TPointerIterator<ViewedValueType>&
+TDirectoryTree<ValueType>::TPointerIterator<ViewedValueType>::operator++()
+{
+	++Internal;
+	return *this;
+}
+
+template <typename ValueType>
+template <typename ViewedValueType>
+inline TDirectoryTree<ValueType>::TPointerIterator<ViewedValueType>::operator bool() const
+{
+	return (bool)Internal;
+}
+
+template <typename ValueType>
+inline typename TDirectoryTree<ValueType>::FIterator TDirectoryTree<ValueType>::CreateIterator()
+{
+	return FIterator(Root,
+		EDirectoryTreeGetFlags::Recursive,
+		PathSeparator);
+}
+
+template <typename ValueType>
+inline typename TDirectoryTree<ValueType>::FConstIterator TDirectoryTree<ValueType>::CreateConstIterator() const
+{
+	return FConstIterator(const_cast<FTreeNode&>(Root),
+		EDirectoryTreeGetFlags::Recursive,
+		PathSeparator);
+}
+
+template <typename ValueType>
+inline typename TDirectoryTree<ValueType>::FPointerIterator TDirectoryTree<ValueType>::CreateIteratorForImplied()
+{
+	return FPointerIterator(Root,
+		EDirectoryTreeGetFlags::Recursive | EDirectoryTreeGetFlags::ImpliedChildren, 
+		PathSeparator);
+}
+
+template <typename ValueType>
+inline typename TDirectoryTree<ValueType>::FConstPointerIterator TDirectoryTree<ValueType>::CreateConstIteratorForImplied() const
+{
+	return FConstPointerIterator(const_cast<FTreeNode&>(Root),
+		EDirectoryTreeGetFlags::Recursive | EDirectoryTreeGetFlags::ImpliedChildren,
+		PathSeparator);
+}
+
+template <typename ValueType>
+inline typename TDirectoryTree<ValueType>::FIterator TDirectoryTree<ValueType>::begin()
+{
+	return CreateIterator();
+}
+
+template <typename ValueType>
+inline typename TDirectoryTree<ValueType>::FConstIterator TDirectoryTree<ValueType>::begin() const
+{
+	return CreateConstIterator();
+}
+
+template <typename ValueType>
+inline typename TDirectoryTree<ValueType>::FIterationSentinel TDirectoryTree<ValueType>::end() const
+{
+	return FIterationSentinel();
+}
diff --git a/Engine/Source/Runtime/Core/Public/Containers/Map.h b/Engine/Source/Runtime/Core/Public/Containers/Map.h
--- a/Engine/Source/Runtime/Core/Public/Containers/Map.h
+++ b/Engine/Source/Runtime/Core/Public/Containers/Map.h
@@ -19,9 +19,6 @@
 
 #define ExchangeB(A,B) {bool T=A; A=B; B=T;}
 
-template <typename KeyType, typename ValueType>
-using TPair = TTuple<KeyType, ValueType>;
-
 /** An initializer type for pairs that's passed to the pair set when adding a new pair. */
 template <typename KeyInitType, typename ValueInitType>
 class TPairInitializer
diff --git a/Engine/Source/Runtime/Core/Public/Templates/Tuple.h b/Engine/Source/Runtime/Core/Public/Templates/Tuple.h
--- a/Engine/Source/Runtime/Core/Public/Templates/Tuple.h
+++ b/Engine/Source/Runtime/Core/Public/Templates/Tuple.h
@@ -42,6 +42,9 @@ class FArchive;
 template <typename... Types>
 struct TTuple;
 
+template <typename KeyType, typename ValueType>
+using TPair = TTuple<KeyType, ValueType>;
+
 template <typename... Types>
 TTuple<std::decay_t<Types>...> MakeTuple(Types&&... Args);
 
diff --git a/Engine/Source/Runtime/Core/Tests/Containers/DirectoryTreeTest.cpp b/Engine/Source/Runtime/Core/Tests/Containers/DirectoryTreeTest.cpp
--- a/Engine/Source/Runtime/Core/Tests/Containers/DirectoryTreeTest.cpp
+++ b/Engine/Source/Runtime/Core/Tests/Containers/DirectoryTreeTest.cpp
@@ -1,7 +1,5 @@
 // Copyright Epic Games, Inc. All Rights Reserved.
 
-// Additional tests are present in Core/Private/DirectoryTree.cpp
-
 #if WITH_TESTS
 
 #include "Containers/DirectoryTree.h"
@@ -604,10 +602,32 @@ TEST_CASE_NAMED(FDirectoryTreeAPITests, "System::Core::Containers::DirectoryTree
 
 	// Testing accessors
 	{
-		// GetChildren
+		// GetChildren and Iteration
 		TDirectoryTree<FMoveConstructOnly> Tree;
 		bool bExists;
 		TArray<FString> Children;
+		TArray<FString> IterKeys;
+		TArray<FMoveConstructOnly*> IterValues;
+		TArray<FString> PtrIterKeys;
+		TArray<FMoveConstructOnly*> PtrIterValues;
+
+		auto ReadIterPairs = [&Tree, &IterKeys, &IterValues, &PtrIterKeys, &PtrIterValues]()
+			{
+				IterKeys.Reset();
+				IterValues.Reset();
+				for (TDirectoryTree<FMoveConstructOnly>::FIterator Iter(Tree.CreateIterator()); Iter; ++Iter)
+				{
+					IterKeys.Emplace(Iter->Key);
+					IterValues.Add(&Iter->Value);
+				}
+				PtrIterKeys.Reset();
+				PtrIterValues.Reset();
+				for (TDirectoryTree<FMoveConstructOnly>::FPointerIterator PtrIter(Tree.CreateIteratorForImplied()); PtrIter; ++PtrIter)
+				{
+					PtrIterKeys.Emplace(PtrIter->Key);
+					PtrIterValues.Add(PtrIter->Value);
+				}
+			};
 
 		Children.Reset();
 		bExists = Tree.TryGetChildren(TEXTVIEW(""), Children,
@@ -615,6 +635,9 @@ TEST_CASE_NAMED(FDirectoryTreeAPITests, "System::Core::Containers::DirectoryTree
 		CHECK_MESSAGE(TEXT("GetChildrenEmpty, Root, !ImpliedParent"),
 			bExists == false && Children.IsEmpty());
 
+		CHECK_MESSAGE(TEXT("Iterate, Empty"), !Tree.CreateIterator());
+		CHECK_MESSAGE(TEXT("IterateImplied, Empty"), !Tree.CreateIteratorForImplied());
+
 		Children.Reset();
 		bExists = Tree.TryGetChildren(TEXTVIEW(""), Children,
 			EDirectoryTreeGetFlags::ImpliedParent |
@@ -638,6 +661,12 @@ TEST_CASE_NAMED(FDirectoryTreeAPITests, "System::Core::Containers::DirectoryTree
 		CHECK_MESSAGE(TEXT("GetChildrenRoot, !ImpliedParent, !ImpliedChildren"),
 			bExists == true && Children.IsEmpty());
 
+		ReadIterPairs();
+		CHECK_MESSAGE(TEXT("Iterate, RootNodeOnly"),
+			IterKeys.Num() == 1 && IterKeys[0] == TEXTVIEW("") && IterValues[0] && IterValues[0]->Value == 1);
+		CHECK_MESSAGE(TEXT("IterateImplied, RootNodeOnly"),
+			PtrIterKeys.Num() == 1 && PtrIterKeys[0] == TEXTVIEW("") && PtrIterValues[0] && PtrIterValues[0]->Value == 1);
+
 		Tree.Empty();
 		Tree.FindOrAdd(TEXTVIEW("/")).Value = 1;
 
@@ -669,6 +698,12 @@ TEST_CASE_NAMED(FDirectoryTreeAPITests, "System::Core::Containers::DirectoryTree
 		CHECK_MESSAGE(TEXT("GetChildrenRoot, ImpliedParent, !ImpliedChildren"),
 			bExists == true && UnorderedEquals(Children, { TEXT("/") }));
 
+		ReadIterPairs();
+		CHECK_MESSAGE(TEXT("Iterate, RootNodeOnlyWithPath"),
+			IterKeys.Num() == 1 && IterKeys[0] == TEXTVIEW("/") && IterValues[0] && IterValues[0]->Value == 1);
+		CHECK_MESSAGE(TEXT("IterateImplied, RootNodeOnlyWithPath"),
+			PtrIterKeys.Num() == 1 && PtrIterKeys[0] == TEXTVIEW("/") && PtrIterValues[0] && PtrIterValues[0]->Value == 1);
+
 		Tree.Empty();
 		Tree.FindOrAdd(TEXTVIEW("/Root/Child")).Value = 1;
 
@@ -697,6 +732,14 @@ TEST_CASE_NAMED(FDirectoryTreeAPITests, "System::Core::Containers::DirectoryTree
 		CHECK_MESSAGE(TEXT("GetChildrenRootImpliedChild, ImpliedParent, !ImpliedChildren, Recursive"),
 			bExists == true && UnorderedEquals(Children, { TEXT("/Root/Child") }));
 
+		ReadIterPairs();
+		CHECK_MESSAGE(TEXT("Iterate, RootChildSingleNode"),
+			IterKeys.Num() == 1 && IterKeys[0] == TEXTVIEW("/Root/Child") && IterValues[0] && IterValues[0]->Value == 1);
+		CHECK_MESSAGE(TEXT("IterateImplied, RootChildSingleNode"),
+			PtrIterKeys.Num() == 3 && PtrIterKeys[0] == TEXTVIEW("/") && !PtrIterValues[0]
+			&& PtrIterKeys[1] == TEXTVIEW("/Root") && !PtrIterValues[1]
+			&& PtrIterKeys[2] == TEXTVIEW("/Root/Child") && PtrIterValues[2] && PtrIterValues[2]->Value == 1);
+
 		Tree.FindOrAdd(TEXTVIEW("/Root/Child2")).Value = 1;
 
 		Tree.Empty();
@@ -717,6 +760,98 @@ TEST_CASE_NAMED(FDirectoryTreeAPITests, "System::Core::Containers::DirectoryTree
 		Tree.FindOrAdd(TEXTVIEW("/Stem/F_AddedRoot")).Value = 1;
 		Tree.FindOrAdd(TEXTVIEW("/Stem/F_AddedRoot/ImpliedChild/AddedChild")).Value = 1;
 
+		TArray<const TCHAR*> ExpectedAdded = {
+			TEXT("/Stem/A_OtherChild"),
+			TEXT("/Stem/B_ImpliedRoot/AddedChild"),
+			TEXT("/Stem/B_ImpliedRoot/AddedChild/Child"),
+			TEXT("/Stem/B_ImpliedRoot/ImpliedChild/AddedChild"),
+			TEXT("/Stem/B_ImpliedRoot/ImpliedChild/AddedChild/AddedChild"),
+			TEXT("/Stem/B_ImpliedRoot/ImpliedChild/AddedChild/ImpliedChild/AddedChild"),
+			TEXT("/Stem/C_MiddleRoot/MiddlePath/ImpliedChild/AddedChild"),
+			TEXT("/Stem/C_MiddleRoot/MiddlePath/ImpliedChild/AddedChild/Child"),
+			TEXT("/Stem/D_MiddleRoot/MiddlePath/AddedChild"),
+			TEXT("/Stem/D_MiddleRoot/MiddlePath/AddedChild/Child"),
+			TEXT("/Stem/E_AddedRoot"),
+			TEXT("/Stem/E_AddedRoot/AddedChild"),
+			TEXT("/Stem/E_AddedRoot/ImpliedChild/AddedChild"),
+			TEXT("/Stem/E_AddedRoot/ImpliedChild/AddedChild/ImpliedChild/AddedChild"),
+			TEXT("/Stem/F_AddedRoot"),
+			TEXT("/Stem/F_AddedRoot/ImpliedChild/AddedChild")
+		};
+		TArray<const TCHAR*> ExpectedImplied = {
+			TEXT("/"),
+			TEXT("/Stem"),
+			TEXT("/Stem/A_OtherChild"),
+			TEXT("/Stem/B_ImpliedRoot"),
+			TEXT("/Stem/B_ImpliedRoot/AddedChild"),
+			TEXT("/Stem/B_ImpliedRoot/AddedChild/Child"),
+			TEXT("/Stem/B_ImpliedRoot/ImpliedChild"),
+			TEXT("/Stem/B_ImpliedRoot/ImpliedChild/AddedChild"),
+			TEXT("/Stem/B_ImpliedRoot/ImpliedChild/AddedChild/AddedChild"),
+			TEXT("/Stem/B_ImpliedRoot/ImpliedChild/AddedChild/ImpliedChild"),
+			TEXT("/Stem/B_ImpliedRoot/ImpliedChild/AddedChild/ImpliedChild/AddedChild"),
+			TEXT("/Stem/C_MiddleRoot"),
+			TEXT("/Stem/C_MiddleRoot/MiddlePath"),
+			TEXT("/Stem/C_MiddleRoot/MiddlePath/ImpliedChild"),
+			TEXT("/Stem/C_MiddleRoot/MiddlePath/ImpliedChild/AddedChild"),
+			TEXT("/Stem/C_MiddleRoot/MiddlePath/ImpliedChild/AddedChild/Child"),
+			TEXT("/Stem/D_MiddleRoot"),
+			TEXT("/Stem/D_MiddleRoot/MiddlePath"),
+			TEXT("/Stem/D_MiddleRoot/MiddlePath/AddedChild"),
+			TEXT("/Stem/D_MiddleRoot/MiddlePath/AddedChild/Child"),
+			TEXT("/Stem/E_AddedRoot"),
+			TEXT("/Stem/E_AddedRoot/AddedChild"),
+			TEXT("/Stem/E_AddedRoot/ImpliedChild"),
+			TEXT("/Stem/E_AddedRoot/ImpliedChild/AddedChild"),
+			TEXT("/Stem/E_AddedRoot/ImpliedChild/AddedChild/ImpliedChild"),
+			TEXT("/Stem/E_AddedRoot/ImpliedChild/AddedChild/ImpliedChild/AddedChild"),
+			TEXT("/Stem/F_AddedRoot"),
+			TEXT("/Stem/F_AddedRoot/ImpliedChild"),
+			TEXT("/Stem/F_AddedRoot/ImpliedChild/AddedChild")
+		};
+		ReadIterPairs();
+
+		// Make sure the ranged for interface, which forwards to CreateIterator(), compiles correctly
+		TArray<FString> RangedForKeys;
+		TArray<FMoveConstructOnly*> RangedForValues;
+		for (TPair<FStringView, FMoveConstructOnly&> Pair : Tree)
+		{
+			RangedForKeys.Emplace(Pair.Key);
+			RangedForValues.Add(&Pair.Value);
+		}
+		// Make sure the const iterators work correctly
+		TArray<FString> ConstItKeys;
+		TArray<const FMoveConstructOnly*> ConstItValues;
+		TArray<FString> ConstItImpliedKeys;
+		TArray<const FMoveConstructOnly*> ConstItImpliedValues;
+		for (TDirectoryTree<FMoveConstructOnly>::FConstIterator ConstIt(Tree.CreateConstIterator()); ConstIt; ++ConstIt)
+		{
+			ConstItKeys.Emplace(ConstIt->Key);
+			ConstItValues.Add(&ConstIt->Value);
+		}
+		for (TDirectoryTree<FMoveConstructOnly>::FConstPointerIterator ConstIt(Tree.CreateConstIteratorForImplied()); ConstIt; ++ConstIt)
+		{
+			ConstItImpliedKeys.Emplace(ConstIt->Key);
+			ConstItImpliedValues.Add(ConstIt->Value);
+		}
+		CHECK_MESSAGE(TEXT("Iterate, ComplicatedTree1"), UnorderedEquals(IterKeys, ExpectedAdded));
+		CHECK_MESSAGE(TEXT("IterateImplied, ComplicatedTree1"), UnorderedEquals(PtrIterKeys, ExpectedImplied));
+		CHECK_MESSAGE(TEXT("RangedFor, ComplicatedTree1"), UnorderedEquals(RangedForKeys, ExpectedAdded));
+		CHECK_MESSAGE(TEXT("CreateConstIterator, ComplicatedTree1"), UnorderedEquals(ConstItKeys, ExpectedAdded));
+		CHECK_MESSAGE(TEXT("CreateConstIteratorForImplied, ComplicatedTree1"), UnorderedEquals(ConstItImpliedKeys, ExpectedImplied));
+
+		for (FMoveConstructOnly* Value : RangedForValues)
+		{
+			CHECK_MESSAGE(TEXT("RangedForValues, ComplicatedTree1"), Value && Value->Value == 1);
+		}
+		for (const FMoveConstructOnly* Value : ConstItValues)
+		{
+			CHECK_MESSAGE(TEXT("ConstItValues, ComplicatedTree1"), Value && Value->Value == 1);
+		}
+		for (const FMoveConstructOnly* Value : ConstItImpliedValues)
+		{
+			CHECK_MESSAGE(TEXT("ConstItImpliedValues, ComplicatedTree1"), !Value || Value->Value == 1);
+		}
 
 		// Case: Requested path is an implied path that is a stored child in the tree.
 		Children.Reset();
