- Proper fix for handling shadowed/deprecated CVars when loading another platform's CVars (because we process by name, I had to add handling to process the "true" cvar that actually holds values)

diff --git a/Engine/Source/Runtime/Core/Private/HAL/ConsoleManager.cpp b/Engine/Source/Runtime/Core/Private/HAL/ConsoleManager.cpp
--- a/Engine/Source/Runtime/Core/Private/HAL/ConsoleManager.cpp
+++ b/Engine/Source/Runtime/Core/Private/HAL/ConsoleManager.cpp
@@ -767,8 +767,16 @@ bool IConsoleManager::VisitPlatformCVarsForEmulation(FName PlatformName, const F
 	};
 
 	TMap<FString, int32> CVarSetByMap;
-	auto VisitIfAllowed = [&CVarSetByMap, &Visit](const FString& Name, const FString& Value, EConsoleVariableFlags SetBy)
+	auto VisitIfAllowed = [&CVarSetByMap, &Visit](const FString& InName, const FString& Value, EConsoleVariableFlags SetBy)
 	{
+		FString Name = InName;
+		// get the name of the real variable for any shadows
+		IConsoleVariable* CVar = IConsoleManager::Get().FindConsoleVariable(*Name, false /* bTrackFrequentCalls */);
+		if (CVar && CVar->IsShadowObject())
+		{
+			Name = IConsoleManager::Get().FindConsoleObjectName(CVar->GetShadowedObject());
+		}
+
 		int32 SetByInt = (int32)(SetBy & ECVF_SetByMask);
 		int32 WasSetBy = CVarSetByMap.FindOrAdd(Name);
 		if (SetByInt < WasSetBy)
@@ -3589,9 +3597,7 @@ void FConsoleManager::LoadAllPlatformCVars(FName PlatformName, const FString& De
 	{
 		// make sure the named cvar exists
 		IConsoleVariable* CVar = IConsoleManager::Get().FindConsoleVariable(*CVarName);
-		// @todo: disabling shadow cvars as this is crashing in one case, and likely aren't needed for other platforms
-		// but will debug this after checking in to unblock users
-		if (CVar == nullptr || CVar->IsShadowObject())
+		if (CVar == nullptr)
 		{
 			return;
 		}
@@ -4982,6 +4988,15 @@ public:
 	{
 		return true;
 	}
+	virtual IConsoleObject* GetShadowedObject() const override
+	{
+		if (Bind())
+		{
+			return RealVariable;
+		}
+		return nullptr;
+	}
+
 
 	virtual bool IsVariableBool() const override
 	{
diff --git a/Engine/Source/Runtime/Core/Public/HAL/IConsoleManager.h b/Engine/Source/Runtime/Core/Public/HAL/IConsoleManager.h
--- a/Engine/Source/Runtime/Core/Public/HAL/IConsoleManager.h
+++ b/Engine/Source/Runtime/Core/Public/HAL/IConsoleManager.h
@@ -517,6 +517,15 @@ public:
 		return false;
 	}
 
+	/**
+	 * If IsShadowObject() returns true, this function will return the "real" variable the shadow cvar is shadowing. Generally not needed,
+	 * as most functions (like GetInt, Set, etc) will operate on the real variable already
+	 */
+	virtual IConsoleObject* GetShadowedObject() const
+	{
+		return nullptr;
+	}
+
 private: // -----------------------------------------
 
 #if TRACK_CONSOLE_FIND_COUNT
