[Cleanup] Fix typo "Aquire" -> "Acquire" in various messages/logs.

diff --git a/Engine/Source/Runtime/Core/Private/Misc/MTAccessDetector.cpp b/Engine/Source/Runtime/Core/Private/Misc/MTAccessDetector.cpp
--- a/Engine/Source/Runtime/Core/Private/Misc/MTAccessDetector.cpp
+++ b/Engine/Source/Runtime/Core/Private/Misc/MTAccessDetector.cpp
@@ -39,14 +39,14 @@ bool FRWAccessDetector_ConcurentReadTest::RunTest(const FString& Parameters)
 	bool Success = false;
 	FThread AccessThread = FThread(TEXT("AccessThread"), [&]()
 	{
-		Success |= TestTrue(TEXT("Aquiring Read Access 1"), MTAccessDetector.AcquireReadAccess());
+		Success |= TestTrue(TEXT("Acquiring Read Access 1"), MTAccessDetector.AcquireReadAccess());
 		bReading1 = true;
 		while (!bReading2) { FPlatformProcess::SleepNoStats(0.01f); }
 		Success |= TestTrue(TEXT("Releasing Read Access 1"), MTAccessDetector.ReleaseReadAccess());
 		return true;
 	});
 	while (!bReading1) { FPlatformProcess::SleepNoStats(0.01f); }
-	Success |= TestTrue(TEXT("Aquiring Read Access 2"), MTAccessDetector.AcquireReadAccess());
+	Success |= TestTrue(TEXT("Acquiring Read Access 2"), MTAccessDetector.AcquireReadAccess());
 	bReading2 = true;
 	AccessThread.Join();
 	Success |= TestTrue(TEXT("Releasing Read Access 2"), MTAccessDetector.ReleaseReadAccess());
@@ -64,14 +64,14 @@ bool FRWAccessDetector_ConcurentWriteTest::RunTest(const FString& Parameters)
 	bool Success = false;
 	FThread AccessThread = FThread(TEXT("AccessThread"), [&]()
 	{
-		Success |= TestTrue(TEXT("Aquiring Write Access"), MTAccessDetector.AcquireWriteAccess());
+		Success |= TestTrue(TEXT("Acquiring Write Access"), MTAccessDetector.AcquireWriteAccess());
 		bWriting1 = true;
 		while (!bWriting2) { FPlatformProcess::SleepNoStats(0.01f); }
 		Success |= TestFalse(TEXT("Releasing Read Access"), MTAccessDetector.ReleaseWriteAccess());
 		return true;
 	});
 	while (!bWriting1) { FPlatformProcess::SleepNoStats(0.01f); }
-	Success |= TestFalse(TEXT("Aquiring Write Access"), MTAccessDetector.AcquireWriteAccess());
+	Success |= TestFalse(TEXT("Acquiring Write Access"), MTAccessDetector.AcquireWriteAccess());
 	bWriting2 = true;
 	AccessThread.Join();
 	Success |= TestTrue(TEXT("Releasing Write Access"), MTAccessDetector.ReleaseWriteAccess());
@@ -89,14 +89,14 @@ bool FRWAccessDetector_ConcurentReadWriteTest::RunTest(const FString& Parameters
 	bool Success = false;
 	FThread AccessThread = FThread(TEXT("AccessThread"), [&]()
 	{
-		Success |= TestTrue(TEXT("Aquiring Read Access"), MTAccessDetector.AcquireReadAccess());
+		Success |= TestTrue(TEXT("Acquiring Read Access"), MTAccessDetector.AcquireReadAccess());
 		bReading = true;
 		while (!bWriting) { FPlatformProcess::SleepNoStats(0.01f); }
 		Success |= TestFalse(TEXT("Releasing Read Access"), MTAccessDetector.ReleaseReadAccess());
 		return true;
 	});
 	while (!bReading) { FPlatformProcess::SleepNoStats(0.01f); }
-	Success |= TestFalse(TEXT("Aquiring Write Access"), MTAccessDetector.AcquireWriteAccess());
+	Success |= TestFalse(TEXT("Acquiring Write Access"), MTAccessDetector.AcquireWriteAccess());
 	bWriting = true;
 	AccessThread.Join();
 	Success |= TestTrue(TEXT("Releasing Write Access"), MTAccessDetector.ReleaseWriteAccess());
@@ -110,8 +110,8 @@ bool FRWAccessDetector_WriteReentrance::RunTest(const FString& Parameters)
 	UE_MT_DECLARE_RW_ACCESS_DETECTOR(MTAccessDetector);
 
 	bool Success = false;
-	Success |= TestTrue(TEXT("Aquiring 1st Write Access"), MTAccessDetector.AcquireWriteAccess());
-	Success |= TestFalse(TEXT("Aquiring 2nd Write Access"), MTAccessDetector.AcquireWriteAccess());
+	Success |= TestTrue(TEXT("Acquiring 1st Write Access"), MTAccessDetector.AcquireWriteAccess());
+	Success |= TestFalse(TEXT("Acquiring 2nd Write Access"), MTAccessDetector.AcquireWriteAccess());
 	Success |= TestFalse(TEXT("Releasing Write Access"), MTAccessDetector.ReleaseWriteAccess());
 	Success |= TestTrue(TEXT("Releasing Write Access"), MTAccessDetector.ReleaseWriteAccess());
 
@@ -124,8 +124,8 @@ bool FRWAccessDetector_ReadReentrance::RunTest(const FString& Parameters)
 	UE_MT_DECLARE_RW_ACCESS_DETECTOR(MTAccessDetector);
 
 	bool Success = false;
-	Success |= TestTrue(TEXT("Aquiring Write Access"), MTAccessDetector.AcquireWriteAccess());
-	Success |= TestFalse(TEXT("Aquiring Read Access"), MTAccessDetector.AcquireReadAccess());
+	Success |= TestTrue(TEXT("Acquiring Write Access"), MTAccessDetector.AcquireWriteAccess());
+	Success |= TestFalse(TEXT("Acquiring Read Access"), MTAccessDetector.AcquireReadAccess());
 	Success |= TestFalse(TEXT("Releasing Read Access"), MTAccessDetector.ReleaseReadAccess());
 	Success |= TestTrue(TEXT("Releasing Write Access"), MTAccessDetector.ReleaseWriteAccess());
 
@@ -145,14 +145,14 @@ bool FRWAccessDetector_RecursiveConcurentReadTest::RunTest(const FString& Parame
 	bool Success = false;
 	FThread AccessThread = FThread(TEXT("AccessThread"), [&]()
 	{
-		Success |= TestTrue(TEXT("Aquiring Read Access 1"), MTAccessDetector.AcquireReadAccess());
+		Success |= TestTrue(TEXT("Acquiring Read Access 1"), MTAccessDetector.AcquireReadAccess());
 		bReading1 = true;
 		while (!bReading2) { FPlatformProcess::SleepNoStats(0.01f); }
 		Success |= TestTrue(TEXT("Releasing Read Access 1"), MTAccessDetector.ReleaseReadAccess());
 		return true;
 	});
 	while (!bReading1) { FPlatformProcess::SleepNoStats(0.01f); }
-	Success |= TestTrue(TEXT("Aquiring Read Access 2"), MTAccessDetector.AcquireReadAccess());
+	Success |= TestTrue(TEXT("Acquiring Read Access 2"), MTAccessDetector.AcquireReadAccess());
 	bReading2 = true;
 	AccessThread.Join();
 	Success |= TestTrue(TEXT("Releasing Read Access 2"), MTAccessDetector.ReleaseReadAccess());
@@ -170,14 +170,14 @@ bool FRWAccessDetector_RecursiveConcurentWriteTest::RunTest(const FString& Param
 	bool Success = false;
 	FThread AccessThread = FThread(TEXT("AccessThread"), [&]()
 	{
-		Success |= TestTrue(TEXT("Aquiring Write Access"), MTAccessDetector.AcquireWriteAccess());
+		Success |= TestTrue(TEXT("Acquiring Write Access"), MTAccessDetector.AcquireWriteAccess());
 		bWriting1 = true;
 		while (!bWriting2) { FPlatformProcess::SleepNoStats(0.01f); }
 		Success |= TestFalse(TEXT("Releasing Read Access"), MTAccessDetector.ReleaseWriteAccess());
 		return true;
 	});
 	while (!bWriting1) { FPlatformProcess::SleepNoStats(0.01f); }
-	Success |= TestFalse(TEXT("Aquiring Write Access"), MTAccessDetector.AcquireWriteAccess());
+	Success |= TestFalse(TEXT("Acquiring Write Access"), MTAccessDetector.AcquireWriteAccess());
 	bWriting2 = true;
 	AccessThread.Join();
 	Success |= TestTrue(TEXT("Releasing Write Access"), MTAccessDetector.ReleaseWriteAccess());
@@ -195,14 +195,14 @@ bool FRWAccessDetector_RecursiveConcurentReadWriteTest::RunTest(const FString& P
 	bool Success = false;
 	FThread AccessThread = FThread(TEXT("AccessThread"), [&]()
 	{
-		Success |= TestTrue(TEXT("Aquiring Read Access"), MTAccessDetector.AcquireReadAccess());
+		Success |= TestTrue(TEXT("Acquiring Read Access"), MTAccessDetector.AcquireReadAccess());
 		bReading = true;
 		while (!bWriting) { FPlatformProcess::SleepNoStats(0.01f); }
 		Success |= TestFalse(TEXT("Releasing Read Access"), MTAccessDetector.ReleaseReadAccess());
 		return true;
 	});
 	while (!bReading) { FPlatformProcess::SleepNoStats(0.01f); }
-	Success |= TestFalse(TEXT("Aquiring Write Access"), MTAccessDetector.AcquireWriteAccess());
+	Success |= TestFalse(TEXT("Acquiring Write Access"), MTAccessDetector.AcquireWriteAccess());
 	bWriting = true;
 	AccessThread.Join();
 	Success |= TestTrue(TEXT("Releasing Write Access"), MTAccessDetector.ReleaseWriteAccess());
@@ -216,8 +216,8 @@ bool FRWAccessDetector_RecursiveWriteReentrance::RunTest(const FString& Paramete
 	UE_MT_DECLARE_RW_RECURSIVE_ACCESS_DETECTOR(MTAccessDetector);
 
 	bool Success = false;
-	Success |= TestTrue(TEXT("Aquiring 1st Write Access"), MTAccessDetector.AcquireWriteAccess());
-	Success |= TestTrue(TEXT("Aquiring 2nd Write Access"), MTAccessDetector.AcquireWriteAccess());
+	Success |= TestTrue(TEXT("Acquiring 1st Write Access"), MTAccessDetector.AcquireWriteAccess());
+	Success |= TestTrue(TEXT("Acquiring 2nd Write Access"), MTAccessDetector.AcquireWriteAccess());
 	Success |= TestTrue(TEXT("Releasing Write Access"), MTAccessDetector.ReleaseWriteAccess());
 	Success |= TestTrue(TEXT("Releasing Write Access"), MTAccessDetector.ReleaseWriteAccess());
 
@@ -230,8 +230,8 @@ bool FRWAccessDetector_RecursiveReadReentrance::RunTest(const FString& Parameter
 	UE_MT_DECLARE_RW_RECURSIVE_ACCESS_DETECTOR(MTAccessDetector);
 
 	bool Success = false;
-	Success |= TestTrue(TEXT("Aquiring Write Access"), MTAccessDetector.AcquireWriteAccess());
-	Success |= TestFalse(TEXT("Aquiring Read Access"), MTAccessDetector.AcquireReadAccess());
+	Success |= TestTrue(TEXT("Acquiring Write Access"), MTAccessDetector.AcquireWriteAccess());
+	Success |= TestFalse(TEXT("Acquiring Read Access"), MTAccessDetector.AcquireReadAccess());
 	Success |= TestFalse(TEXT("Releasing Read Access"), MTAccessDetector.ReleaseReadAccess());
 	Success |= TestTrue(TEXT("Releasing Write Access"), MTAccessDetector.ReleaseWriteAccess());
 
@@ -251,14 +251,14 @@ bool FRWAccessDetector_FullyRecursiveConcurentReadTest::RunTest(const FString& P
 	bool Success = false;
 	FThread AccessThread = FThread(TEXT("AccessThread"), [&]()
 	{
-		Success |= TestTrue(TEXT("Aquiring Read Access 1"), MTAccessDetector.AcquireReadAccess());
+		Success |= TestTrue(TEXT("Acquiring Read Access 1"), MTAccessDetector.AcquireReadAccess());
 		bReading1 = true;
 		while (!bReading2) { FPlatformProcess::SleepNoStats(0.01f); }
 		Success |= TestTrue(TEXT("Releasing Read Access 1"), MTAccessDetector.ReleaseReadAccess());
 		return true;
 	});
 	while (!bReading1) { FPlatformProcess::SleepNoStats(0.01f); }
-	Success |= TestTrue(TEXT("Aquiring Read Access 2"), MTAccessDetector.AcquireReadAccess());
+	Success |= TestTrue(TEXT("Acquiring Read Access 2"), MTAccessDetector.AcquireReadAccess());
 	bReading2 = true;
 	AccessThread.Join();
 	Success |= TestTrue(TEXT("Releasing Read Access 2"), MTAccessDetector.ReleaseReadAccess());
@@ -276,14 +276,14 @@ bool FRWAccessDetector_FullyRecursiveConcurentWriteTest::RunTest(const FString&
 	bool Success = false;
 	FThread AccessThread = FThread(TEXT("AccessThread"), [&]()
 	{
-		Success |= TestTrue(TEXT("Aquiring Write Access"), MTAccessDetector.AcquireWriteAccess());
+		Success |= TestTrue(TEXT("Acquiring Write Access"), MTAccessDetector.AcquireWriteAccess());
 		bWriting1 = true;
 		while (!bWriting2) { FPlatformProcess::SleepNoStats(0.01f); }
 		Success |= TestFalse(TEXT("Releasing Read Access"), MTAccessDetector.ReleaseWriteAccess());
 		return true;
 	});
 	while (!bWriting1) { FPlatformProcess::SleepNoStats(0.01f); }
-	Success |= TestFalse(TEXT("Aquiring Write Access"), MTAccessDetector.AcquireWriteAccess());
+	Success |= TestFalse(TEXT("Acquiring Write Access"), MTAccessDetector.AcquireWriteAccess());
 	bWriting2 = true;
 	AccessThread.Join();
 	Success |= TestTrue(TEXT("Releasing Write Access"), MTAccessDetector.ReleaseWriteAccess());
@@ -301,14 +301,14 @@ bool FRWAccessDetector_FullyRecursiveConcurentReadWriteTest::RunTest(const FStri
 	bool Success = false;
 	FThread AccessThread = FThread(TEXT("AccessThread"), [&]()
 	{
-		Success |= TestTrue(TEXT("Aquiring Read Access"), MTAccessDetector.AcquireReadAccess());
+		Success |= TestTrue(TEXT("Acquiring Read Access"), MTAccessDetector.AcquireReadAccess());
 		bReading = true;
 		while (!bWriting) { FPlatformProcess::SleepNoStats(0.01f); }
 		Success |= TestFalse(TEXT("Releasing Read Access"), MTAccessDetector.ReleaseReadAccess());
 		return true;
 	});
 	while (!bReading) { FPlatformProcess::SleepNoStats(0.01f); }
-	Success |= TestFalse(TEXT("Aquiring Write Access"), MTAccessDetector.AcquireWriteAccess());
+	Success |= TestFalse(TEXT("Acquiring Write Access"), MTAccessDetector.AcquireWriteAccess());
 	bWriting = true;
 	AccessThread.Join();
 	Success |= TestTrue(TEXT("Releasing Write Access"), MTAccessDetector.ReleaseWriteAccess());
@@ -322,8 +322,8 @@ bool FRWAccessDetector_FullyRecursiveWriteReentrance::RunTest(const FString& Par
 	UE_MT_DECLARE_RW_FULLY_RECURSIVE_ACCESS_DETECTOR(MTAccessDetector);
 
 	bool Success = false;
-	Success |= TestTrue(TEXT("Aquiring 1st Write Access"), MTAccessDetector.AcquireWriteAccess());
-	Success |= TestTrue(TEXT("Aquiring 2nd Write Access"), MTAccessDetector.AcquireWriteAccess());
+	Success |= TestTrue(TEXT("Acquiring 1st Write Access"), MTAccessDetector.AcquireWriteAccess());
+	Success |= TestTrue(TEXT("Acquiring 2nd Write Access"), MTAccessDetector.AcquireWriteAccess());
 	Success |= TestTrue(TEXT("Releasing Write Access"), MTAccessDetector.ReleaseWriteAccess());
 	Success |= TestTrue(TEXT("Releasing Write Access"), MTAccessDetector.ReleaseWriteAccess());
 
@@ -336,8 +336,8 @@ bool FRWAccessDetector_FullyRecursiveReadReentrance::RunTest(const FString& Para
 	UE_MT_DECLARE_RW_FULLY_RECURSIVE_ACCESS_DETECTOR(MTAccessDetector);
 
 	bool Success = false;
-	Success |= TestTrue(TEXT("Aquiring Write Access"), MTAccessDetector.AcquireWriteAccess());
-	Success |= TestTrue(TEXT("Aquiring Read Access"), MTAccessDetector.AcquireReadAccess());
+	Success |= TestTrue(TEXT("Acquiring Write Access"), MTAccessDetector.AcquireWriteAccess());
+	Success |= TestTrue(TEXT("Acquiring Read Access"), MTAccessDetector.AcquireReadAccess());
 	Success |= TestTrue(TEXT("Releasing Read Access"), MTAccessDetector.ReleaseReadAccess());
 	Success |= TestTrue(TEXT("Releasing Write Access"), MTAccessDetector.ReleaseWriteAccess());
 
diff --git a/Engine/Source/Runtime/Core/Public/Misc/MTAccessDetector.h b/Engine/Source/Runtime/Core/Public/Misc/MTAccessDetector.h
--- a/Engine/Source/Runtime/Core/Public/Misc/MTAccessDetector.h
+++ b/Engine/Source/Runtime/Core/Public/Misc/MTAccessDetector.h
@@ -70,7 +70,7 @@ public:
 	FORCEINLINE bool AcquireReadAccess() const
 	{
 		const bool ErrorDetected = (AtomicValue.fetch_add(1, std::memory_order_relaxed) & WriterBits) != 0;
-		checkf(!ErrorDetected || GIsAutomationTesting, TEXT("Aquiring a read access while there is already a write access"));
+		checkf(!ErrorDetected || GIsAutomationTesting, TEXT("Acquiring a read access while there is already a write access"));
 		return !ErrorDetected;
 	}
 
