Moved shared TFuture::Get() and TSharedFuture::Get() implementation into TFutureBase.

#rnx

diff --git a/Engine/Source/Runtime/Core/Public/Async/Future.h b/Engine/Source/Runtime/Core/Public/Async/Future.h
--- a/Engine/Source/Runtime/Core/Public/Async/Future.h
+++ b/Engine/Source/Runtime/Core/Public/Async/Future.h
@@ -240,6 +240,28 @@ class TFutureBase
 {
 public:
 
+	/**
+	 * Gets the future's result.
+	 *
+	 * @return The result as a const reference, or the same reference if the future holds a reference, or void if the future holds a void.
+	 * @note Not equivalent to std::future::get(). The future remains valid.
+	 */
+	decltype(auto) Get() const //UE_LIFETIMEBOUND
+	{
+		if constexpr (std::is_lvalue_reference_v<ResultType>)
+		{
+			return (ResultType&)*this->GetState()->GetResult();
+		}
+		else if constexpr (std::is_void_v<ResultType>)
+		{
+			this->GetState()->GetResult();
+		}
+		else
+		{
+			return (const ResultType&)this->GetState()->GetResult();
+		}
+	}
+
 	/**
 	 * Checks whether this future object has its value set.
 	 *
@@ -543,27 +565,6 @@ public:
 		: BaseType(MoveTemp(Future))
 	{ }
 
-	/**
-	 * Gets the future's result.
-	 *
-	 * @return The result.
-	 */
-	decltype(auto) Get() const // UE_LIFETIMEBOUND
-	{
-		if constexpr (std::is_lvalue_reference_v<ResultType>)
-		{
-			return (ResultType&)*this->GetState()->GetResult();
-		}
-		else if constexpr (std::is_void_v<ResultType>)
-		{
-			this->GetState()->GetResult();
-		}
-		else
-		{
-			return (const ResultType&)this->GetState()->GetResult();
-		}
-	}
-
 private:
 	/**
 	 * Creates and initializes a new instance.
