Exposed TSharedPtr in TFuture API for clarity of ownership.

#rnx

diff --git a/Engine/Source/Runtime/Core/Public/Async/Future.h b/Engine/Source/Runtime/Core/Public/Async/Future.h
--- a/Engine/Source/Runtime/Core/Public/Async/Future.h
+++ b/Engine/Source/Runtime/Core/Public/Async/Future.h
@@ -238,6 +238,9 @@ private:
 template <typename ResultType>
 class TFutureBase
 {
+protected:
+	using StateType = TFutureState<ResultType>;
+
 public:
 
 	/**
@@ -326,8 +329,6 @@ public:
 
 protected:
 
-	typedef TSharedPtr<TFutureState<ResultType>> StateType;
-
 	/** Default constructor. */
 	TFutureBase() = default;
 
@@ -336,22 +337,21 @@ protected:
 	 *
 	 * @param InState The shared state to initialize with.
 	 */
-	TFutureBase(StateType&& InState)
+	TFutureBase(TSharedPtr<StateType>&& InState)
 		: State(MoveTemp(InState))
 	{
 	}
-	TFutureBase(const StateType& InState)
+	TFutureBase(const TSharedPtr<StateType>& InState)
 		: State(InState)
 	{
 	}
-protected:
 
 	/**
 	 * Gets the shared state object.
 	 *
 	 * @return The shared state object.
 	 */
-	const StateType& GetState() const
+	const TSharedPtr<StateType>& GetState() const
 	{
 		// if you hit this assertion then your future has an invalid state.
 		// this happens if you have an uninitialized future or if you moved
@@ -398,7 +398,7 @@ protected:
 private:
 
 	/** Holds the future's state. */
-	StateType State;
+	TSharedPtr<StateType> State;
 };
 
 
@@ -525,14 +525,8 @@ public:
 	using BaseType::Reset;
 
 private:
-	TFuture(const typename BaseType::StateType& InState)
-		: BaseType(InState)
-	{
-	}
-	TFuture(typename BaseType::StateType&& InState)
-		: BaseType(MoveTemp(InState))
-	{
-	}
+	// Forward constructors
+	using BaseType::TFutureBase;
 };
 
 
@@ -566,19 +560,8 @@ public:
 	{ }
 
 private:
-	/**
-	 * Creates and initializes a new instance.
-	 *
-	 * @param InState The shared state to initialize from.
-	 */
-	TSharedFuture(const typename BaseType::StateType& InState)
-		: BaseType(InState)
-	{
-	}
-	TSharedFuture(typename BaseType::StateType&& InState)
-		: BaseType(InState)
-	{
-	}
+	// Forward constructors
+	using BaseType::TFutureBase;
 };
 
 
@@ -591,7 +574,7 @@ private:
 template<typename ResultType>
 class TPromise final
 {
-	using StateType = TSharedPtr<TFutureState<ResultType>>;
+	using StateType = TFutureState<ResultType>;
 
 	// This is necessary because we can't form references to void for the parameter of Set().
 	using SetType = std::conditional_t<std::is_void_v<ResultType>, int, ResultType>;
@@ -747,7 +730,7 @@ protected:
 	 *
 	 * @return The shared state object.
 	 */
-	const StateType& GetState()
+	const TSharedPtr<StateType>& GetState()
 	{
 		// if you hit this assertion then your promise has an invalid state.
 		// this happens if you move the promise to another instance.
@@ -758,7 +741,7 @@ protected:
 
 private:
 	/** Holds the shared state object. */
-	StateType State;
+	TSharedPtr<StateType> State;
 
 	/** Whether a future has already been retrieved from this promise. */
 	bool bFutureRetrieved = false;
@@ -801,7 +784,7 @@ auto TFutureBase<ResultType>::Then(Func Continuation) //-> TFuture<decltype(Cont
 	};
 
 	// This invalidate this future.
-	StateType MovedState = MoveTemp(this->State);
+	TSharedPtr<StateType> MovedState = MoveTemp(this->State);
 	MovedState->SetContinuation(MoveTemp(Callback));
 	return FutureResult;
 }
