[FMutex]
- Prevent oversubscription for mutex as we generally only want oversubscription to kick in for longer wait periods.
- This fixes an issue where tons of oversubscription could happen for heavily contented mutex, causing performance problems.

diff --git a/Engine/Source/Runtime/Core/Private/Async/Mutex.cpp b/Engine/Source/Runtime/Core/Private/Async/Mutex.cpp
--- a/Engine/Source/Runtime/Core/Private/Async/Mutex.cpp
+++ b/Engine/Source/Runtime/Core/Private/Async/Mutex.cpp
@@ -4,6 +4,7 @@
 
 #include "Async/ParkingLot.h"
 #include "HAL/PlatformProcess.h"
+#include "Async/Fundamental/Scheduler.h"
 
 namespace UE
 {
@@ -45,6 +46,10 @@ void FMutex::LockSlow()
 			CurrentState |= HasWaitingThreadsFlag;
 		}
 
+		// For mutex, we don't want to enter oversubscription during a wait since the wait period is generally too short for oversubscription to matter
+		// and it can worsen performance a lot for heavily contented locks.
+		LowLevelTasks::Private::FOversubscriptionAllowedScope _(false);
+
 		// Wait if the state has not changed. Either way, loop back and try to acquire the lock trying to wait.
 		ParkingLot::Wait(&State, [this, CurrentState] { return State.load(std::memory_order_acquire) == CurrentState; }, []{});
 	}
