[CsvStats]
- Fix TSAN warning on heavily used globals

diff --git a/Engine/Source/Runtime/Core/Private/ProfilingDebugging/CsvProfiler.cpp b/Engine/Source/Runtime/Core/Private/ProfilingDebugging/CsvProfiler.cpp
--- a/Engine/Source/Runtime/Core/Private/ProfilingDebugging/CsvProfiler.cpp
+++ b/Engine/Source/Runtime/Core/Private/ProfilingDebugging/CsvProfiler.cpp
@@ -40,6 +40,8 @@
 
 #include "HAL/PlatformMisc.h"
 
+#include <atomic>
+
 #if CSV_PROFILER
 
 #define CSV_PROFILER_INLINE FORCEINLINE
@@ -257,9 +259,9 @@ static FAutoConsoleVariableRef CVarTrackWaitsRT(TEXT("csv.trackWaitsRT"), GCsvTr
 // Categories
 //
 static const uint32 CSV_MAX_CATEGORY_COUNT = 2048;
-static bool GCsvCategoriesEnabled[CSV_MAX_CATEGORY_COUNT];
+static std::atomic<bool> GCsvCategoriesEnabled[CSV_MAX_CATEGORY_COUNT];
 
-static bool GCsvProfilerIsCapturing = false;
+static std::atomic<bool> GCsvProfilerIsCapturing = false;
 static bool GCsvProfilerIsCapturingRT = false; // Renderthread version of the above
 
 static bool GCsvProfilerIsWritingFile = false;
@@ -743,7 +745,10 @@ public:
 		if (!Instance)
 		{
 			Instance = new FCsvCategoryData;
-			FMemory::Memzero(GCsvCategoriesEnabled, sizeof(GCsvCategoriesEnabled));
+			for (int32 Index = 0; Index < UE_ARRAY_COUNT(GCsvCategoriesEnabled); ++Index)
+			{
+				GCsvCategoriesEnabled[Index].store(false, std::memory_order_relaxed);
+			}
 		}
 		return Instance;
 	}
@@ -781,7 +786,7 @@ public:
 			if (FWildcardString::IsMatch(*EnabledCategory, *CategoryNames[CategoryIndex]))
 			{
 				UE_LOG(LogCsvProfiler, Log, TEXT("Config enabled category %s"), *CategoryNames[CategoryIndex]);
-				GCsvCategoriesEnabled[CategoryIndex] = true;
+				GCsvCategoriesEnabled[CategoryIndex].store(true, std::memory_order_relaxed);
 			}
 		}
 
@@ -790,7 +795,7 @@ public:
 			if (FWildcardString::IsMatch(*DisabledCategory, *CategoryNames[CategoryIndex]))
 			{
 				UE_LOG(LogCsvProfiler, Log, TEXT("Config disabled category %s"), *CategoryNames[CategoryIndex]);
-				GCsvCategoriesEnabled[CategoryIndex] = false;
+				GCsvCategoriesEnabled[CategoryIndex].store(false, std::memory_order_relaxed);
 			}
 		}
 	}
@@ -846,7 +851,7 @@ public:
 				check(Index < CSV_MAX_CATEGORY_COUNT);
 				if (Index < CSV_MAX_CATEGORY_COUNT)
 				{
-					GCsvCategoriesEnabled[Index] = bEnableByDefault;
+					GCsvCategoriesEnabled[Index].store(bEnableByDefault, std::memory_order_relaxed);
 					CategoryNames[Index] = CategoryName;
 					CategoryNameToIndex.Add(CategoryName.ToLower(), Index);
 					UpdateCategoryFromConfig(Index);
@@ -3382,7 +3387,7 @@ FCsvProfiler::FCsvProfiler()
 
 FCsvProfiler::~FCsvProfiler()
 {
-	GCsvProfilerIsCapturing = false;
+	GCsvProfilerIsCapturing.store(false, std::memory_order_relaxed);
 	IsShuttingDown.Increment();
 	if (ProcessingThread)
 	{
@@ -3437,7 +3442,7 @@ void FCsvProfiler::BeginFrame()
 			BeginCaptureInternal(CurrentCommand);
 		}
 
-		if (GCsvProfilerIsCapturing)
+		if (GCsvProfilerIsCapturing.load(std::memory_order_relaxed))
 		{
 			GFrameBoundaries.AddBeginFrameTimestamp(ECsvTimeline::Gamethread);
 
@@ -3466,14 +3471,14 @@ void FCsvProfiler::BeginFrame()
 		CSVTest();
 	}
 
-	GCsvABTest.BeginFrameUpdate(CaptureFrameNumber, GCsvProfilerIsCapturing);
+	GCsvABTest.BeginFrameUpdate(CaptureFrameNumber, GCsvProfilerIsCapturing.load(std::memory_order_relaxed));
 #endif // CSV_PROFILER_ALLOW_DEBUG_FEATURES
 }
 
 
 void FCsvProfiler::BeginCaptureInternal(const FCsvCaptureCommand& CurrentCommand)
 {
-	if (GCsvProfilerIsCapturing)
+	if (GCsvProfilerIsCapturing.load(std::memory_order_relaxed))
 	{
 		UE_LOG(LogCsvProfiler, Warning, TEXT("Capture start requested, but a capture was already running"));
 		return;
@@ -3597,7 +3602,7 @@ void FCsvProfiler::BeginCaptureInternal(const FCsvCaptureCommand& CurrentCommand
 	// Check TLS is initialized before starting the capture. This should have happened in BeginCapture
 	check(FCsvProfilerThreadData::IsTlsSlotInitialized());
 	TRACE_CSV_PROFILER_BEGIN_CAPTURE(*Filename, RenderThreadId, RHIThreadId, GDefaultWaitStatName.StatName, GCsvStatCounts);
-	GCsvProfilerIsCapturing = true;
+	GCsvProfilerIsCapturing.store(true, std::memory_order_relaxed);
 	CaptureStartTime = FPlatformTime::Seconds();
 
 	// Run the Csv benchmark if enabled
@@ -3616,7 +3621,7 @@ void FCsvProfiler::EndFrame()
 	CSV_SCOPED_TIMING_STAT_EXCLUSIVE(CsvProfiler);
 
 	check(IsInGameThread());
-	if (GCsvProfilerIsCapturing)
+	if (GCsvProfilerIsCapturing.load(std::memory_order_relaxed))
 	{
 		// If the GT is stalled, wait
 		if (ProcessingThread)
@@ -3696,7 +3701,7 @@ void FCsvProfiler::EndFrame()
 
 bool FCsvProfiler::TryEndCaptureInternal(const FCsvCaptureCommand& CurrentCommand)
 {
-	if (GCsvProfilerIsCapturing || GCsvProfilerIsWritingFile)
+	if (GCsvProfilerIsCapturing.load(std::memory_order_relaxed) || GCsvProfilerIsWritingFile)
 	{
 		// Delay end capture by a frame to allow RT stats to catch up
 		if (CurrentCommand.FrameRequested == GCsvProfilerFrameNumber)
@@ -3712,7 +3717,7 @@ bool FCsvProfiler::TryEndCaptureInternal(const FCsvCaptureCommand& CurrentComman
 
 		// Signal to the processing thread to write the file out (if we have one).
 		GCsvProfilerIsWritingFile = true;
-		GCsvProfilerIsCapturing = false;
+		GCsvProfilerIsCapturing.store(false, std::memory_order_relaxed);
 
 		TRACE_CSV_PROFILER_END_CAPTURE();
 
@@ -3742,7 +3747,7 @@ bool FCsvProfiler::TryEndCaptureInternal(const FCsvCaptureCommand& CurrentComman
 	}
 
 	// If we get here then we're actually done
-	check(!GCsvProfilerIsCapturing && !GCsvProfilerIsWritingFile);
+	check(!GCsvProfilerIsCapturing.load(std::memory_order_relaxed) && !GCsvProfilerIsWritingFile);
 
 	// Signal the async completion callback, if one was provided when the capture was stopped.
 	if (CurrentCommand.Completion)
@@ -3804,12 +3809,12 @@ void FCsvProfiler::BeginFrameRT()
 	RenderThreadId = FPlatformTLS::GetCurrentThreadId();
 
 	check(IsInRenderingThread());
-	if (GCsvProfilerIsCapturing)
+	if (GCsvProfilerIsCapturing.load(std::memory_order_relaxed))
 	{
 		// Mark where the renderthread frames begin
 		GFrameBoundaries.AddBeginFrameTimestamp(ECsvTimeline::Renderthread);
 	}
-	GCsvProfilerIsCapturingRT = GCsvProfilerIsCapturing;
+	GCsvProfilerIsCapturingRT = GCsvProfilerIsCapturing.load(std::memory_order_relaxed);
 
 #if CSV_PROFILER_ALLOW_DEBUG_FEATURES
 	if (GCsvTestingRT)
@@ -3826,7 +3831,7 @@ void FCsvProfiler::EndFrameRT()
 {
 	LLM_SCOPE(ELLMTag::CsvProfiler);
 	check(IsInRenderingThread());
-	if (GCsvProfilerIsCapturing)
+	if (GCsvProfilerIsCapturing.load(std::memory_order_relaxed))
 	{
 		CaptureFrameNumberRT++;
 	}
@@ -3961,11 +3966,17 @@ void FCsvProfiler::SetDeviceProfileName(FString InDeviceProfileName)
 	CSV_METADATA(TEXT("DeviceProfile"), *InDeviceProfileName);
 }
 
+FORCEINLINE bool IsCsvProfilerCapturing(uint32 CategoryIndex)
+{
+	// These do not need to be atomic together. But we need them to at least be atomics separately so TSAN knows we're OK with that.
+	return GCsvProfilerIsCapturing.load(std::memory_order_relaxed) && GCsvCategoriesEnabled[CategoryIndex].load(std::memory_order_relaxed);
+}
+
 /** Push/pop events */
 void FCsvProfiler::BeginStat(const char * StatName, uint32 CategoryIndex, const char * NamedEventName)
 {
 #if !CSV_PROFILER_MINIMAL
-	if (GCsvProfilerIsCapturing && GCsvCategoriesEnabled[CategoryIndex])
+	if (IsCsvProfilerCapturing(CategoryIndex))
 	{
 #if CSV_PROFILER_SUPPORT_NAMED_EVENTS
 		if (UNLIKELY(GCsvProfilerNamedEventsTiming))
@@ -3981,7 +3992,7 @@ void FCsvProfiler::BeginStat(const char * StatName, uint32 CategoryIndex, const
 void FCsvProfiler::BeginStat(const FName& StatName, uint32 CategoryIndex)
 {
 #if !CSV_PROFILER_MINIMAL
-	if (GCsvProfilerIsCapturing && GCsvCategoriesEnabled[CategoryIndex])
+	if (IsCsvProfilerCapturing(CategoryIndex))
 	{
 		ValidateFName(StatName);
 #if CSV_PROFILER_SUPPORT_NAMED_EVENTS
@@ -3998,7 +4009,7 @@ void FCsvProfiler::BeginStat(const FName& StatName, uint32 CategoryIndex)
 void FCsvProfiler::EndStat(const char * StatName, uint32 CategoryIndex)
 {
 #if !CSV_PROFILER_MINIMAL
-	if (GCsvProfilerIsCapturing && GCsvCategoriesEnabled[CategoryIndex])
+	if (IsCsvProfilerCapturing(CategoryIndex))
 	{
 		FCsvProfilerThreadData::Get().AddTimestampEnd(StatName, CategoryIndex);
 #if CSV_PROFILER_SUPPORT_NAMED_EVENTS
@@ -4014,7 +4025,7 @@ void FCsvProfiler::EndStat(const char * StatName, uint32 CategoryIndex)
 void FCsvProfiler::EndStat(const FName& StatName, uint32 CategoryIndex)
 {
 #if !CSV_PROFILER_MINIMAL
-	if (GCsvProfilerIsCapturing && GCsvCategoriesEnabled[CategoryIndex])
+	if (IsCsvProfilerCapturing(CategoryIndex))
 	{
 		ValidateFName(StatName);
 		FCsvProfilerThreadData::Get().AddTimestampEnd(StatName, CategoryIndex);
@@ -4031,7 +4042,7 @@ void FCsvProfiler::EndStat(const FName& StatName, uint32 CategoryIndex)
 void FCsvProfiler::BeginExclusiveStat(const char * StatName, const char * NamedEventName)
 {
 #if !CSV_PROFILER_MINIMAL
-	if (GCsvProfilerIsCapturing && GCsvCategoriesEnabled[CSV_CATEGORY_INDEX(Exclusive)])
+	if (IsCsvProfilerCapturing(CSV_CATEGORY_INDEX(Exclusive)))
 	{
 #if CSV_PROFILER_SUPPORT_NAMED_EVENTS
 		if (UNLIKELY(GCsvProfilerNamedEventsExclusive))
@@ -4047,7 +4058,7 @@ void FCsvProfiler::BeginExclusiveStat(const char * StatName, const char * NamedE
 void FCsvProfiler::EndExclusiveStat(const char * StatName)
 {
 #if !CSV_PROFILER_MINIMAL
-	if (GCsvProfilerIsCapturing && GCsvCategoriesEnabled[CSV_CATEGORY_INDEX(Exclusive)])
+	if (IsCsvProfilerCapturing(CSV_CATEGORY_INDEX(Exclusive)))
 	{
 		FCsvProfilerThreadData::Get().AddTimestampExclusiveEnd(StatName);
 #if CSV_PROFILER_SUPPORT_NAMED_EVENTS
@@ -4064,7 +4075,7 @@ void FCsvProfiler::EndExclusiveStat(const char * StatName)
 void FCsvProfiler::BeginSetWaitStat(const char* StatName, const char* FormattedStatName, const char* FormattedStatNameNonCP)
 {
 #if !CSV_PROFILER_MINIMAL
-	if (GCsvProfilerIsCapturing && GCsvCategoriesEnabled[CSV_CATEGORY_INDEX(Exclusive)])
+	if (IsCsvProfilerCapturing(CSV_CATEGORY_INDEX(Exclusive)))
 	{
 #if CSV_PROFILER_SUPPORT_NAMED_EVENTS
 		if (UNLIKELY(GCsvProfilerNamedEventsExclusive))
@@ -4088,7 +4099,7 @@ void FCsvProfiler::BeginSetWaitStat(const char* StatName, const char* FormattedS
 void FCsvProfiler::EndSetWaitStat()
 {
 #if !CSV_PROFILER_MINIMAL
-	if (GCsvProfilerIsCapturing && GCsvCategoriesEnabled[CSV_CATEGORY_INDEX(Exclusive)])
+	if (IsCsvProfilerCapturing(CSV_CATEGORY_INDEX(Exclusive)))
 	{
 #if CSV_PROFILER_SUPPORT_NAMED_EVENTS
 		if (UNLIKELY(GCsvProfilerNamedEventsExclusive))
@@ -4104,7 +4115,7 @@ void FCsvProfiler::EndSetWaitStat()
 void FCsvProfiler::BeginWait()
 {
 #if !CSV_PROFILER_MINIMAL
-	if (GCsvProfilerIsCapturing && GCsvCategoriesEnabled[CSV_CATEGORY_INDEX(Exclusive)])
+	if (IsCsvProfilerCapturing(CSV_CATEGORY_INDEX(Exclusive)))
 	{
 		const FCsvWaitStatName& WaitStatName = FCsvProfilerThreadData::Get().GetWaitStatName();
 		if (WaitStatName != GIgnoreWaitStatName)
@@ -4141,7 +4152,7 @@ void FCsvProfiler::BeginWait()
 void FCsvProfiler::EndWait()
 {
 #if !CSV_PROFILER_MINIMAL
-	if (GCsvProfilerIsCapturing && GCsvCategoriesEnabled[CSV_CATEGORY_INDEX(Exclusive)])
+	if (IsCsvProfilerCapturing(CSV_CATEGORY_INDEX(Exclusive)))
 	{
 		const FCsvWaitStatName& WaitStatName = FCsvProfilerThreadData::Get().GetWaitStatName();
 		if (WaitStatName != GIgnoreWaitStatName)
@@ -4166,7 +4177,7 @@ void FCsvProfiler::EndWait()
 
 void FCsvProfiler::RecordEventfInternal(int32 CategoryIndex, const TCHAR* Fmt, ...)
 {
-	bool bIsCsvRecording = GCsvProfilerIsCapturing && GCsvCategoriesEnabled[CategoryIndex];
+	const bool bIsCsvRecording = IsCsvProfilerCapturing(CategoryIndex);
 	if (bIsCsvRecording || GStartOnEvent)
 	{
 		LLM_SCOPE(ELLMTag::CsvProfiler);
@@ -4196,7 +4207,7 @@ void FCsvProfiler::RecordEventfInternal(int32 CategoryIndex, const TCHAR* Fmt, .
 
 void FCsvProfiler::RecordEvent(int32 CategoryIndex, const FString& EventText)
 {
-	if (GCsvProfilerIsCapturing && GCsvCategoriesEnabled[CategoryIndex])
+	if (IsCsvProfilerCapturing(CategoryIndex))
 	{
 		LLM_SCOPE(ELLMTag::CsvProfiler);
 		UE_LOG(LogCsvProfiler, Display, TEXT("CSVEvent \"%s\" [Frame %d]"), *EventText, FCsvProfiler::Get()->GetCaptureFrameNumber());
@@ -4269,7 +4280,7 @@ void FCsvProfiler::SetMetadataInternal(const TCHAR* Key, const TCHAR* Value, boo
 
 void FCsvProfiler::RecordEventAtTimestamp(int32 CategoryIndex, const FString& EventText, uint64 Cycles64)
 {
-	if (GCsvProfilerIsCapturing && GCsvCategoriesEnabled[CategoryIndex])
+	if (IsCsvProfilerCapturing(CategoryIndex))
 	{
 		LLM_SCOPE(ELLMTag::CsvProfiler);
 		UE_LOG(LogCsvProfiler, Display, TEXT("CSVEvent [Frame %d] : \"%s\""), FCsvProfiler::Get()->GetCaptureFrameNumber(), *EventText);
@@ -4287,7 +4298,7 @@ void FCsvProfiler::RecordEventAtTimestamp(int32 CategoryIndex, const FString& Ev
 
 void FCsvProfiler::RecordCustomStatMinimal(const char* StatName, uint32 CategoryIndex, float Value, const ECsvCustomStatOp CustomStatOp)
 {
-	if (GCsvProfilerIsCapturing && GCsvCategoriesEnabled[CategoryIndex])
+	if (IsCsvProfilerCapturing(CategoryIndex))
 	{
 		FCsvProfilerThreadData::Get().AddCustomStat(StatName, CategoryIndex, Value, CustomStatOp);
 	}
@@ -4295,7 +4306,7 @@ void FCsvProfiler::RecordCustomStatMinimal(const char* StatName, uint32 Category
 
 void FCsvProfiler::RecordCustomStatMinimal(const char* StatName, uint32 CategoryIndex, int32 Value, const ECsvCustomStatOp CustomStatOp)
 {
-	if (GCsvProfilerIsCapturing && GCsvCategoriesEnabled[CategoryIndex])
+	if (IsCsvProfilerCapturing(CategoryIndex))
 	{
 		FCsvProfilerThreadData::Get().AddCustomStat(StatName, CategoryIndex, Value, CustomStatOp);
 	}
@@ -4311,7 +4322,7 @@ void FCsvProfiler::RecordCustomStatMinimal(const char* StatName, uint32 Category
 void FCsvProfiler::RecordCustomStat(const char * StatName, uint32 CategoryIndex, float Value, const ECsvCustomStatOp CustomStatOp)
 {
 #if !CSV_PROFILER_MINIMAL
-	if (GCsvProfilerIsCapturing && GCsvCategoriesEnabled[CategoryIndex])
+	if (IsCsvProfilerCapturing(CategoryIndex))
 	{
 		FCsvProfilerThreadData::Get().AddCustomStat(StatName, CategoryIndex, Value, CustomStatOp);
 	}
@@ -4329,7 +4340,7 @@ void FCsvProfiler::RecordCustomStat(const char* StatName, uint32 CategoryIndex,
 void FCsvProfiler::RecordCustomStat(const FName& StatName, uint32 CategoryIndex, float Value, const ECsvCustomStatOp CustomStatOp)
 {
 #if !CSV_PROFILER_MINIMAL
-	if (GCsvProfilerIsCapturing && GCsvCategoriesEnabled[CategoryIndex])
+	if (IsCsvProfilerCapturing(CategoryIndex))
 	{
 		ValidateFName(StatName);
 		FCsvProfilerThreadData::Get().AddCustomStat(StatName, CategoryIndex, Value, CustomStatOp);
@@ -4349,7 +4360,7 @@ void FCsvProfiler::RecordCustomStat(const FName& StatName, uint32 CategoryIndex,
 void FCsvProfiler::RecordCustomStat(const char * StatName, uint32 CategoryIndex, int32 Value, const ECsvCustomStatOp CustomStatOp)
 {
 #if !CSV_PROFILER_MINIMAL
-	if (GCsvProfilerIsCapturing && GCsvCategoriesEnabled[CategoryIndex])
+	if (IsCsvProfilerCapturing(CategoryIndex))
 	{
 		FCsvProfilerThreadData::Get().AddCustomStat(StatName, CategoryIndex, Value, CustomStatOp);
 	}
@@ -4359,7 +4370,7 @@ void FCsvProfiler::RecordCustomStat(const char * StatName, uint32 CategoryIndex,
 void FCsvProfiler::RecordCustomStat(const FName& StatName, uint32 CategoryIndex, int32 Value, const ECsvCustomStatOp CustomStatOp)
 {
 #if !CSV_PROFILER_MINIMAL
-	if (GCsvProfilerIsCapturing && GCsvCategoriesEnabled[CategoryIndex])
+	if (IsCsvProfilerCapturing(CategoryIndex))
 	{
 		ValidateFName(StatName);
 		FCsvProfilerThreadData::Get().AddCustomStat(StatName, CategoryIndex, Value, CustomStatOp);
@@ -4524,7 +4535,10 @@ void FCsvProfiler::Init()
 
 	if (GAllCategoriesStartDisabled)
 	{
-		FMemory::Memzero(GCsvCategoriesEnabled, sizeof(GCsvCategoriesEnabled));
+		for (int32 Index = 0; Index < UE_ARRAY_COUNT(GCsvCategoriesEnabled); ++Index)
+		{
+			GCsvCategoriesEnabled[Index].store(false, std::memory_order_relaxed);
+		}
 	}
 
 	// Always disable the CSV profiling thread if the platform does not support threading.
@@ -4546,7 +4560,7 @@ void FCsvProfiler::Init()
 bool FCsvProfiler::IsCapturing() const
 {
 	check(IsInGameThread());
-	return GCsvProfilerIsCapturing;
+	return GCsvProfilerIsCapturing.load(std::memory_order_relaxed);
 }
 
 bool FCsvProfiler::IsWritingFile() const
@@ -4600,7 +4614,7 @@ bool FCsvProfiler::EnableCategoryByString(const FString& CategoryName) const
 	if (Category >= 0)
 	{
 		UE_LOG(LogCsvProfiler, Log, TEXT("Enabled category %s"), *CategoryName);
-		GCsvCategoriesEnabled[Category] = true;
+		GCsvCategoriesEnabled[Category].store(true, std::memory_order_relaxed);
 		return true;
 	}
 	UE_LOG(LogCsvProfiler, Warning, TEXT("Error: Can't find category %s"), *CategoryName);
@@ -4610,13 +4624,13 @@ bool FCsvProfiler::EnableCategoryByString(const FString& CategoryName) const
 void FCsvProfiler::EnableCategoryByIndex(uint32 CategoryIndex, bool bEnable) const
 {
 	check(CategoryIndex < CSV_MAX_CATEGORY_COUNT);
-	GCsvCategoriesEnabled[CategoryIndex] = bEnable;
+	GCsvCategoriesEnabled[CategoryIndex].store(bEnable, std::memory_order_relaxed);
 }
 
 bool FCsvProfiler::IsCategoryEnabled(uint32 CategoryIndex) const
 {
 	check(CategoryIndex < CSV_MAX_CATEGORY_COUNT);
-	return GCsvCategoriesEnabled[CategoryIndex];
+	return GCsvCategoriesEnabled[CategoryIndex].load(std::memory_order_relaxed);
 }
 
 bool FCsvProfiler::IsCapturing_Renderthread() const
@@ -4683,8 +4697,11 @@ void CSVTest()
 {
 	if (GCsvTestCategoryOnly)
 	{
-		FMemory::Memzero(GCsvCategoriesEnabled, sizeof(GCsvCategoriesEnabled));
-		GCsvCategoriesEnabled[CSV_CATEGORY_INDEX(CsvTest)] = true;
+		for (int32 Index = 0; Index < UE_ARRAY_COUNT(GCsvCategoriesEnabled); ++Index)
+		{
+			GCsvCategoriesEnabled[Index].store(false, std::memory_order_relaxed);
+		}
+		GCsvCategoriesEnabled[CSV_CATEGORY_INDEX(CsvTest)].store(true, std::memory_order_relaxed);
 	}
 	TCsvPersistentCustomStat<float>* PersistentStatFloat = FCsvProfiler::Get()->GetOrCreatePersistentCustomStatFloat(TEXT("PersistentStatFloat"), CSV_CATEGORY_INDEX(CsvTest));
 	PersistentStatFloat->Add(0.15f);
