Fix TDeque not retaining capacity when copying/moving with an inline allocator.

#tests unit tests

diff --git a/Engine/Source/Runtime/Core/Public/Containers/Deque.h b/Engine/Source/Runtime/Core/Public/Containers/Deque.h
--- a/Engine/Source/Runtime/Core/Public/Containers/Deque.h
+++ b/Engine/Source/Runtime/Core/Public/Containers/Deque.h
@@ -361,7 +361,7 @@ public:
 		if (Capacity)
 		{
 			Storage.ResizeAllocation(0, 0, sizeof(ElementType));
-			Capacity = 0;
+			Capacity = Storage.GetInitialCapacity();
 		}
 	}
 
@@ -464,6 +464,11 @@ private:
 			Grow(Storage.CalculateSlackReserve(Other.Count, sizeof(ElementType)));
 			CopyElements(Other);
 		}
+		else
+		{
+			// Retain any inline capacity.
+			Capacity = Other.Storage.GetInitialCapacity();
+		}
 	}
 
 	void CopyUnchecked(std::initializer_list<ElementType> InList)
@@ -476,6 +481,10 @@ private:
 			ConstructItems<ElementType>(GetData(), &*InList.begin(), InCount);
 			Tail = Count = InCount;
 		}
+		else
+		{
+			Capacity = Storage.GetInitialCapacity();
+		}
 	}
 
 	/*
@@ -495,6 +504,10 @@ private:
 			Other.Capacity = Other.Storage.GetInitialCapacity();
 			Other.Count = Other.Head = Other.Tail = 0;
 		}
+		else
+		{
+			Capacity = Other.Storage.GetInitialCapacity();
+		}
 	}
 
 	/*
diff --git a/Engine/Source/Runtime/Core/Tests/Containers/DequeTest.cpp b/Engine/Source/Runtime/Core/Tests/Containers/DequeTest.cpp
--- a/Engine/Source/Runtime/Core/Tests/Containers/DequeTest.cpp
+++ b/Engine/Source/Runtime/Core/Tests/Containers/DequeTest.cpp
@@ -121,36 +121,64 @@ TEST_CASE_NAMED(FDequeReserveEmplaceLastSingleElementTest, "System::Core::Contai
 
 TEST_CASE_NAMED(FDequeResetTest, "System::Core::Containers::TDeque::Reset", "[SmokeFilter][Core][Containers][Deque]")
 {
-	int32Token::Reset();
+	SECTION("Default allocator")
 	{
-		TDeque<int32Token> Deque;
-		Deque.Reset();	// Should be innocuous
-		CHECK(!Deque.Max());
-		CHECK(!Deque.Num());
-		CHECK(Deque.IsEmpty());
+		int32Token::Reset();
+		{
+			TDeque<int32Token> Deque;
+			Deque.Reset();	// Should be innocuous
+			CHECK(!Deque.Max());
+			CHECK(!Deque.Num());
+			CHECK(Deque.IsEmpty());
+			Deque.EmplaceLast(0);
+			CHECK(Deque.Max());
+			CHECK(Deque.Num() == 1);
+			CHECK(!Deque.IsEmpty());
+			Deque.Reset();
+			CHECK(Deque.Max());
+			CHECK(!Deque.Num());
+			CHECK(Deque.IsEmpty());
+		}
+		CHECK(int32Token::EvenConstructionDestructionCalls(1));
+	}
+
+	SECTION("Inline Allocator")
+	{
+		TDeque<int32Token, TInlineAllocator<Test::DefaultCapacity>> Deque;
+		Deque.Reset();
+		CHECK(Deque.Max() == Test::DefaultCapacity);
 		Deque.EmplaceLast(0);
-		CHECK(Deque.Max());
 		CHECK(Deque.Num() == 1);
-		CHECK(!Deque.IsEmpty());
+		CHECK(Deque.Max() == Test::DefaultCapacity);
 		Deque.Reset();
-		CHECK(Deque.Max());
-		CHECK(!Deque.Num());
 		CHECK(Deque.IsEmpty());
+		CHECK(Deque.Max() == Test::DefaultCapacity);
 	}
-	CHECK(int32Token::EvenConstructionDestructionCalls(1));
 }
 
 TEST_CASE_NAMED(FDequeEmptyTest, "System::Core::Containers::TDeque::Empty", "[SmokeFilter][Core][Containers][Deque]")
 {
-	int32Token::Reset();
+	SECTION("Default allocator")
 	{
-		TDeque<int32Token> Deque;
-		Deque.Empty();	// Should be innocuous
-		CHECK(!Deque.Max());
+		int32Token::Reset();
+		{
+			TDeque<int32Token> Deque;
+			Deque.Empty();	// Should be innocuous
+			CHECK(!Deque.Max());
+			CHECK(!Deque.Num());
+			CHECK(Deque.IsEmpty());
+		}
+		CHECK(int32Token::EvenConstructionDestructionCalls(0));
+	}
+
+	SECTION("Inline allocator")
+	{
+		TDeque<int32Token, TInlineAllocator<Test::DefaultCapacity>> Deque;
+		Deque.Empty();
+		CHECK(Deque.Max() == Test::DefaultCapacity);
 		CHECK(!Deque.Num());
 		CHECK(Deque.IsEmpty());
 	}
-	CHECK(int32Token::EvenConstructionDestructionCalls(0));
 }
 
 TEST_CASE_NAMED(FDequeEmptyAfterSingleElementEmplaceLastTest, "System::Core::Containers::TDeque::Empty after single element EmplaceLast", "[SmokeFilter][Core][Containers][Deque]")
@@ -720,6 +748,48 @@ TEST_CASE_NAMED(FDequeCopyVariableSizeWithHeadTailWrapTest, "System::Core::Conta
 	CHECK(int32Token::EvenConstructionDestructionCalls());
 }
 
+TEST_CASE_NAMED(FDequeCopyWithInlineAllocator, "System::Core::Containers::TDeque Copy with inline allocator", "[SmokeFilter][Core][Containers][Deque]")
+{
+	SECTION("Copy empty TDeque")
+	{
+		int32Token::Reset();
+		TDeque<int32Token, TInlineAllocator<Test::DefaultCapacity>> Deque;
+		CHECK(Deque.Max() == Test::DefaultCapacity);
+		CHECK(Deque.Num() == 0);
+		TDeque<int32Token, TInlineAllocator<Test::DefaultCapacity>> DequeCopy(Deque);
+		CHECK_MESSAGE("Expected copy to retain inline capacity.", DequeCopy.Max() == Test::DefaultCapacity);
+	}
+
+	SECTION("Copy TDeque with 1 element")
+	{
+		int32Token::Reset();
+		TDeque<int32Token, TInlineAllocator<Test::DefaultCapacity>> Deque;
+		Deque.PushLast(1);
+		CHECK(Deque.Max() == Test::DefaultCapacity);
+		CHECK(Deque.Num() == 1);
+		TDeque<int32Token, TInlineAllocator<Test::DefaultCapacity>> DequeCopy(Deque);
+		CHECK_MESSAGE("Expected copy to retain inline capacity.", DequeCopy.Max() == Test::DefaultCapacity);
+		CHECK(DequeCopy.Num() == 1);
+	}
+
+	SECTION("Copy TDeque that exceeds inline allocation")
+	{
+		int32Token::Reset();
+		TDeque<int32Token, TInlineAllocator<Test::DefaultCapacity>> Deque;
+		const int32 ExpectedMax = (Test::DefaultCapacity + 2);
+		for (int32 i = 0; i < ExpectedMax; ++i)
+		{
+			Deque.PushLast(i);
+		}
+		CHECK(Deque.Num() == ExpectedMax);
+		CHECK(Deque.Max() >= ExpectedMax); // >= as default storage may grow by more than what we're expanding by.
+
+		TDeque<int32Token, TInlineAllocator<Test::DefaultCapacity>> DequeCopy(Deque);
+		CHECK(DequeCopy.Num() == ExpectedMax);
+		CHECK(DequeCopy.Max() >= ExpectedMax);
+	}
+}
+
 TEST_CASE_NAMED(FDequeMoveSimpleTest, "System::Core::Containers::TDeque::Move simple", "[SmokeFilter][Core][Containers][Deque]")
 {
 	int32Token::Reset();
@@ -804,6 +874,48 @@ TEST_CASE_NAMED(FDequeMoveVariableSizeWithHeadTailWrapTest, "System::Core::Conta
 	CHECK(int32Token::EvenConstructionDestructionCalls());
 }
 
+TEST_CASE_NAMED(FDequeMoveWithInlineAllocator, "System::Core::Containers::TDeque Move with inline allocator", "[SmokeFilter][Core][Containers][Deque]")
+{
+	SECTION("Move empty TDeque")
+	{
+		int32Token::Reset();
+		TDeque<int32Token, TInlineAllocator<Test::DefaultCapacity>> Deque;
+		CHECK(Deque.Max() == Test::DefaultCapacity);
+		CHECK(Deque.Num() == 0);
+		TDeque<int32Token, TInlineAllocator<Test::DefaultCapacity>> DequeCopy(MoveTemp(Deque));
+		CHECK_MESSAGE("Expected copy to retain inline capacity.", DequeCopy.Max() == Test::DefaultCapacity);
+	}
+
+	SECTION("Move TDeque with 1 element")
+	{
+		int32Token::Reset();
+		TDeque<int32Token, TInlineAllocator<Test::DefaultCapacity>> Deque;
+		Deque.PushLast(1);
+		CHECK(Deque.Max() == Test::DefaultCapacity);
+		CHECK(Deque.Num() == 1);
+		TDeque<int32Token, TInlineAllocator<Test::DefaultCapacity>> DequeCopy(MoveTemp(Deque));
+		CHECK_MESSAGE("Expected copy to retain inline capacity.", DequeCopy.Max() == Test::DefaultCapacity);
+		CHECK(DequeCopy.Num() == 1);
+	}
+
+	SECTION("Move TDeque that exceeds inline allocation")
+	{
+		int32Token::Reset();
+		TDeque<int32Token, TInlineAllocator<Test::DefaultCapacity>> Deque;
+		const int32 ExpectedMax = (Test::DefaultCapacity + 2);
+		for (int32 i = 0; i < ExpectedMax; ++i)
+		{
+			Deque.PushLast(i);
+		}
+		CHECK(Deque.Num() == ExpectedMax);
+		CHECK(Deque.Max() >= ExpectedMax); // >= as default storage may grow by more than what we're expanding by.
+
+		TDeque<int32Token, TInlineAllocator<Test::DefaultCapacity>> DequeCopy(MoveTemp(Deque));
+		CHECK(DequeCopy.Num() == ExpectedMax);
+		CHECK(DequeCopy.Max() >= ExpectedMax);
+	}
+}
+
 TEST_CASE_NAMED(FDequeIterationWIthoutWrapTest, "System::Core::Containers::TDeque::Iteration without head/tail wrap around", "[SmokeFilter][Core][Containers][Deque]")
 {
 	int32Token::Reset();
@@ -888,46 +1000,92 @@ TEST_CASE_NAMED(FDequeIteratorArithmeticTest, "System::Core::Containers::TDeque:
 
 TEST_CASE_NAMED(FDequeConstructFromStdInitializerListTest, "System::Core::Containers::TDeque::Construct from std initializer_list", "[SmokeFilter][Core][Containers][Deque]")
 {
-	int32Token::Reset();
+	SECTION("Default allocator")
 	{
-		TDeque<int32Token> Deque({0, 1, 2, 3, 4, 5});
-		int32Token Result;
-		int32 CheckValue = 0;
-		while (Deque.TryPopFirst(Result))
+		int32Token::Reset();
 		{
-			CHECK(*Result == CheckValue++);
+			TDeque<int32Token> Deque({ 0, 1, 2, 3, 4, 5 });
+			int32Token Result;
+			int32 CheckValue = 0;
+			while (Deque.TryPopFirst(Result))
+			{
+				CHECK(*Result == CheckValue++);
+			}
+			CHECK(CheckValue == 6);
 		}
-		CHECK(CheckValue == 6);
+		CHECK(int32Token::EvenConstructionDestructionCalls(6 * 2 + 1));
+	}
+
+	SECTION("Inline allocator without growing")
+	{
+		TDeque<int32Token, TInlineAllocator<8>> Deque({ 0, 1, 2, 3, 4, 5 });
+		CHECK(Deque.Num() == 6);
+		CHECK(Deque.Max() == 8);
+	}
+
+	SECTION("Inline allocator with growing")
+	{
+		TDeque<int32Token, TInlineAllocator<4>> Deque({ 0, 1, 2, 3, 4, 5 });
+		CHECK(Deque.Num() == 6);
+		CHECK(Deque.Max() == 6);
 	}
-	CHECK(int32Token::EvenConstructionDestructionCalls(6 * 2 + 1));
 }
 
 TEST_CASE_NAMED(FDequeConstructFromEmptyStdInitializerListTest, "System::Core::Containers::TDeque::Construct from empty std initializer_list", "[SmokeFilter][Core][Containers][Deque]")
 {
-	int32Token::Reset();
+	SECTION("Default allocator")
+	{
+		int32Token::Reset();
+		{
+			TDeque<int32Token> Deque({});
+			CHECK(Deque.IsEmpty());
+		}
+		CHECK(int32Token::EvenConstructionDestructionCalls(0));
+	}
+
+	SECTION("Inline allocator")
 	{
-		TDeque<int32Token> Deque({});
+		TDeque<int32Token, TInlineAllocator<8>> Deque({});
 		CHECK(Deque.IsEmpty());
+		CHECK(Deque.Max() == 8);
 	}
-	CHECK(int32Token::EvenConstructionDestructionCalls(0));
 }
 
 TEST_CASE_NAMED(FDequeAssignFromStdInitializerListTest, "System::Core::Containers::TDeque::Assign from std initializer_list", "[SmokeFilter][Core][Containers][Deque]")
 {
-	int32Token::Reset();
+	SECTION("Default allocator")
 	{
-		TDeque<int32Token> Deque;
-		Deque.EmplaceLast(0);
-		Deque = {0, 1, 2, 3, 4, 5};
-		int32Token Result;
-		int32 CheckValue = 0;
-		while (Deque.TryPopFirst(Result))
+		int32Token::Reset();
 		{
-			CHECK(*Result == CheckValue++);
+			TDeque<int32Token> Deque;
+			Deque.EmplaceLast(0);
+			Deque = { 0, 1, 2, 3, 4, 5 };
+			int32Token Result;
+			int32 CheckValue = 0;
+			while (Deque.TryPopFirst(Result))
+			{
+				CHECK(*Result == CheckValue++);
+			}
+			CHECK(CheckValue == 6);
 		}
-		CHECK(CheckValue == 6);
+		CHECK(int32Token::EvenConstructionDestructionCalls(6 * 2 + 2));
+	}
+
+	SECTION("Inline allocation without growing")
+	{
+		TDeque<int32Token, TInlineAllocator<8>> Deque;
+		Deque = { 0, 1, 2, 3, 4, 5 };
+		CHECK(Deque.Num() == 6);
+		CHECK(Deque.Max() == 8);
+	}
+
+	SECTION("Inline allocation with growing")
+	{
+		TDeque<int32Token, TInlineAllocator<4>> Deque;
+		Deque = { 0, 1, 2, 3, 4, 5 };
+		CHECK(Deque.Num() == 6);
+		CHECK(Deque.Max() == 6);
 	}
-	CHECK(int32Token::EvenConstructionDestructionCalls(6 * 2 + 2));
 }
 
 }  // namespace Deque
