Fixed a documentation typo in the comment for how the %P specifier will render 'pm'

diff --git a/Engine/Source/Runtime/Core/Public/Misc/DateTime.h b/Engine/Source/Runtime/Core/Public/Misc/DateTime.h
--- a/Engine/Source/Runtime/Core/Public/Misc/DateTime.h
+++ b/Engine/Source/Runtime/Core/Public/Misc/DateTime.h
@@ -547,7 +547,7 @@ public:
 	 *		%M - Minute, 00-59
 	 *		%S - Second, 00-60
 	 *		%p - AM or PM
-	 *		%P - am or PM
+	 *		%P - am or pm
 	 *		%j - Day of the Year, 001-366
 	 * 
 	 * @param Format The format of the returned string.
