Made TTypeCompatibleBytes non-copyable, to avoid accidental bitwise copying instead of calling the copy constructor or copy assignment operator.

diff --git a/Engine/Source/Runtime/Core/Public/Templates/TypeCompatibleBytes.h b/Engine/Source/Runtime/Core/Public/Templates/TypeCompatibleBytes.h
--- a/Engine/Source/Runtime/Core/Public/Templates/TypeCompatibleBytes.h
+++ b/Engine/Source/Runtime/Core/Public/Templates/TypeCompatibleBytes.h
@@ -22,6 +22,17 @@ template<typename ElementType>
 struct TTypeCompatibleBytes
 {
 	using ElementTypeAlias_NatVisHelper = ElementType;
+
+	// Trivially constructible and destuctible - users are responsible for managing the lifetime of the inner element.
+	TTypeCompatibleBytes() = default;
+	~TTypeCompatibleBytes() = default;
+
+	// Noncopyable
+	TTypeCompatibleBytes(TTypeCompatibleBytes&&) = delete;
+	TTypeCompatibleBytes(const TTypeCompatibleBytes&) = delete;
+	TTypeCompatibleBytes& operator=(TTypeCompatibleBytes&&) = delete;
+	TTypeCompatibleBytes& operator=(const TTypeCompatibleBytes&) = delete;
+
 	ElementType*		GetTypedPtr()		{ return (ElementType*)this;  }
 	const ElementType*	GetTypedPtr() const	{ return (const ElementType*)this; }
 
diff --git a/Engine/Source/Runtime/Core/Public/Templates/UnrealTemplate.h b/Engine/Source/Runtime/Core/Public/Templates/UnrealTemplate.h
--- a/Engine/Source/Runtime/Core/Public/Templates/UnrealTemplate.h
+++ b/Engine/Source/Runtime/Core/Public/Templates/UnrealTemplate.h
@@ -634,10 +634,15 @@ inline void Swap(T& A, T& B)
 
 	if constexpr (TUseBitwiseSwap<T>::Value)
 	{
-		TTypeCompatibleBytes<T> Temp;
-		*(TTypeCompatibleBytes<T>*)&Temp = *(TTypeCompatibleBytes<T>*)&A;
-		*(TTypeCompatibleBytes<T>*)&A    = *(TTypeCompatibleBytes<T>*)&B;
-		*(TTypeCompatibleBytes<T>*)&B    = *(TTypeCompatibleBytes<T>*)&Temp;
+		struct FAlignedBytes
+		{
+			alignas(T) char Bytes[sizeof(T)];
+		};
+
+		FAlignedBytes Temp;
+		*(FAlignedBytes*)&Temp = *(FAlignedBytes*)&A;
+		*(FAlignedBytes*)&A    = *(FAlignedBytes*)&B;
+		*(FAlignedBytes*)&B    = *(FAlignedBytes*)&Temp;
 	}
 	else
 	{
