[TaskGraph]
- Fix missing statid cycle scope on older tasks execution

[RN] Fix missing statid cycle scope on older taskgraph tasks execution

diff --git a/Engine/Source/Runtime/Core/Public/Async/TaskGraphInterfaces.h b/Engine/Source/Runtime/Core/Public/Async/TaskGraphInterfaces.h
--- a/Engine/Source/Runtime/Core/Public/Async/TaskGraphInterfaces.h
+++ b/Engine/Source/Runtime/Core/Public/Async/TaskGraphInterfaces.h
@@ -25,6 +25,7 @@
 #include "Misc/MemStack.h"
 #include "Misc/Timeout.h"
 #include "Templates/Atomic.h"
+#include "Templates/Models.h"
 #include "ProfilingDebugging/MetadataTrace.h"
 
 #include "Async/Fundamental/Task.h"
@@ -622,11 +623,28 @@ private:
 		FBaseGraphTask::Init(TEXT("GraphTask"), InPriority, InExtendedPriority);
 	}
 
+	struct CGetStatIdProvider
+	{
+		template <typename TaskType>
+		auto Requires(TaskType& Task) -> decltype(
+			Task.GetStatId()
+		);
+	};
+
 	virtual void ExecuteTask() override final
 	{
 		FGraphEventRef GraphEventRef{ this };
 		TTask* TaskObject = TaskStorage.GetTypedPtr();
-		TaskObject->DoTask(TaskObject->GetDesiredThread(), GraphEventRef);
+		// Only use this scope when GetStatId() is actually implemented on TaskObject for backward compat, sometimes it's not.
+		if constexpr (TModels_V<CGetStatIdProvider, TTask>)
+		{
+			FScopeCycleCounter Scope(TaskObject->GetStatId(), true);
+			TaskObject->DoTask(TaskObject->GetDesiredThread(), GraphEventRef);
+		}
+		else
+		{
+			TaskObject->DoTask(TaskObject->GetDesiredThread(), GraphEventRef);
+		}
 		DestructItem(TaskObject);
 	}
 
