Removed duplicate TFuture::Get() which is already implemented by TFutureBase::Get().
Replicated code comment info from TFutureBase::Get() to TFuture::GetMutable().

diff --git a/Engine/Source/Runtime/Core/Public/Async/Future.h b/Engine/Source/Runtime/Core/Public/Async/Future.h
--- a/Engine/Source/Runtime/Core/Public/Async/Future.h
+++ b/Engine/Source/Runtime/Core/Public/Async/Future.h
@@ -430,31 +430,9 @@ public:
 	~TFuture() = default;
 
 	/**
-	 * Gets the future's result by const reference.
-	 *
-	 * @return The result.
-	 * @note Not equivalent to std::future::get(). The future remains valid.
-	 */
-	decltype(auto) Get() const //UE_LIFETIMEBOUND
-	{
-		if constexpr (std::is_lvalue_reference_v<ResultType>)
-		{
-			return (ResultType&)*this->GetState()->GetResult();
-		}
-		else if constexpr (std::is_void_v<ResultType>)
-		{
-			this->GetState()->GetResult();
-		}
-		else
-		{
-			return (const ResultType&)this->GetState()->GetResult();
-		}
-	}
-
-	/**
-	 * Gets the future's result by non-const reference.
+	 * Gets the future's result.
 	 *
-	 * @return The result.
+	 * @return The result as a non-const reference, or the same reference if the future holds a reference, or void if the future holds a void.
 	 * @note Not equivalent to std::future::get(). The future remains valid.
 	 */
 	decltype(auto) GetMutable() // UE_LIFETIMEBOUND
