Added support for references and void in TTypeCompatibleBytes.
Removed a bunch of special casing from TFuture.

diff --git a/Engine/Source/Runtime/Core/Public/Async/Future.h b/Engine/Source/Runtime/Core/Public/Async/Future.h
--- a/Engine/Source/Runtime/Core/Public/Async/Future.h
+++ b/Engine/Source/Runtime/Core/Public/Async/Future.h
@@ -145,18 +145,10 @@ template<typename ResultType>
 class TFutureState
 	: public FFutureState
 {
-	using InternalResultType =
-		std::conditional_t<
-			std::is_lvalue_reference_v<ResultType>,
-			std::remove_reference_t<ResultType>*,
-			std::conditional_t<
-				std::is_void_v<ResultType>,
-				int,
-				ResultType
-			>
-		>;
-
 public:
+	using MutableResultType = typename TTypeCompatibleBytes<ResultType>::MutableGetType;
+	using ConstResultType   = typename TTypeCompatibleBytes<ResultType>::ConstGetType;
+	using RvalueResultType  = typename TTypeCompatibleBytes<ResultType>::RvalueGetType;
 
 	/** Default constructor. */
 	TFutureState()
@@ -167,7 +159,7 @@ public:
 	{
 		if (IsComplete())
 		{
-			DestructItem(Result.GetTypedPtr());
+			Result.DestroyUnchecked();
 		}
 	}
 
@@ -188,25 +180,18 @@ public:
 	 * @return The result value.
 	 * @see EmplaceResult
 	 */
-	const InternalResultType& GetResult() const UE_LIFETIMEBOUND
+	MutableResultType GetResult()
 	{
 		while (!IsComplete())
 		{
 			WaitFor(FTimespan::MaxValue());
 		}
 
-		return *Result.GetTypedPtr();
+		return Result.GetUnchecked();
 	}
-
-	/**
-	 * Gets the result (will block the calling thread until the result is available).
-	 *
-	 * @return The result value.
-	 * @see EmplaceResult
-	 */
-	InternalResultType& GetResult() UE_LIFETIMEBOUND
+	ConstResultType GetResult() const
 	{
-		return const_cast<InternalResultType&>(AsConst(*this).GetResult());
+		return const_cast<TFutureState*>(this);
 	}
 
 	/**
@@ -219,14 +204,14 @@ public:
 	void EmplaceResult(ArgTypes&&... Args)
 	{
 		check(!IsComplete());
-		new(Result.GetTypedPtr()) InternalResultType(Forward<ArgTypes>(Args)...);
+		Result.EmplaceUnchecked(Forward<ArgTypes>(Args)...);
 		MarkComplete();
 	}
 
 private:
 
 	/** Holds the asynchronous result. */
-	TTypeCompatibleBytes<InternalResultType> Result;
+	TTypeCompatibleBytes<ResultType> Result;
 };
 
 /* TFuture
@@ -242,6 +227,9 @@ protected:
 	using StateType = TFutureState<ResultType>;
 
 public:
+	using MutableResultType = typename StateType::MutableResultType;
+	using ConstResultType   = typename StateType::ConstResultType;
+	using RvalueResultType  = typename StateType::RvalueResultType;
 
 	/**
 	 * Gets the future's result.
@@ -249,20 +237,9 @@ public:
 	 * @return The result as a const reference, or the same reference if the future holds a reference, or void if the future holds a void.
 	 * @note Not equivalent to std::future::get(). The future remains valid.
 	 */
-	decltype(auto) Get() const //UE_LIFETIMEBOUND
+	ConstResultType Get() const // UE_LIFETIMEBOUND
 	{
-		if constexpr (std::is_lvalue_reference_v<ResultType>)
-		{
-			return (ResultType&)*this->GetState()->GetResult();
-		}
-		else if constexpr (std::is_void_v<ResultType>)
-		{
-			this->GetState()->GetResult();
-		}
-		else
-		{
-			return (const ResultType&)this->GetState()->GetResult();
-		}
+		return this->GetState()->GetResult();
 	}
 
 	/**
@@ -417,6 +394,10 @@ class TFuture final
 
 	using BaseType = TFutureBase<ResultType>;
 
+	using MutableResultType = typename BaseType::MutableResultType;
+	using ConstResultType   = typename BaseType::ConstResultType;
+	using RvalueResultType  = typename BaseType::RvalueResultType;
+
 public:
 
 	/** Default constructor. */
@@ -435,20 +416,9 @@ public:
 	 * @return The result as a non-const reference, or the same reference if the future holds a reference, or void if the future holds a void.
 	 * @note Not equivalent to std::future::get(). The future remains valid.
 	 */
-	decltype(auto) GetMutable() // UE_LIFETIMEBOUND
+	MutableResultType GetMutable() // UE_LIFETIMEBOUND
 	{
-		if constexpr (std::is_lvalue_reference_v<ResultType>)
-		{
-			return (ResultType&)*this->GetState()->GetResult();
-		}
-		else if constexpr (std::is_void_v<ResultType>)
-		{
-			this->GetState()->GetResult();
-		}
-		else
-		{
-			return (ResultType&)this->GetState()->GetResult();
-		}
+		return (MutableResultType)this->Get();
 	}
 
 	/**
@@ -457,21 +427,10 @@ public:
 	 * @return The result.
 	 * @note Equivalent to std::future::get(). Invalidates the future.
 	 */
-	decltype(auto) Consume()
+	ResultType Consume()
 	{
 		TFuture Local(MoveTemp(*this));
-		if constexpr (std::is_lvalue_reference_v<ResultType>)
-		{
-			return (ResultType&)*Local.GetState()->GetResult();
-		}
-		else if constexpr (std::is_void_v<ResultType>)
-		{
-			Local.GetState()->GetResult();
-		}
-		else
-		{
-			return (ResultType)MoveTemp(Local.GetState()->GetResult());
-		}
+		return (RvalueResultType)Local.Get();
 	}
 
 	/**
@@ -661,46 +620,11 @@ public:
 	 *
 	 * @param Args The arguments to forward to the constructor of the result.
 	 */
-	template <
-		typename... ArgTypes
-#ifndef __cpp_concepts
-		, typename DummyResultType = ResultType
-		UE_REQUIRES(!std::is_void_v<DummyResultType> && !std::is_lvalue_reference_v<DummyResultType>)
-#endif
-	>
+	template <typename... ArgTypes>
 	void EmplaceValue(ArgTypes&&... Args)
-#ifdef __cpp_concepts
-		requires (!std::is_void_v<ResultType> && !std::is_lvalue_reference_v<ResultType>)
-#endif
 	{
 		this->GetState()->EmplaceResult(Forward<ArgTypes>(Args)...);
 	}
-#ifndef __cpp_concepts
-	template <
-		typename DummyResultType = ResultType
-		UE_REQUIRES(std::is_lvalue_reference_v<DummyResultType>)
-	>
-#endif
-	void EmplaceValue(SetType Result)
-#ifdef __cpp_concepts
-		requires (std::is_lvalue_reference_v<ResultType>)
-#endif
-	{
-		this->GetState()->EmplaceResult(&Result);
-	}
-#ifndef __cpp_concepts
-	template <
-		typename DummyResultType = ResultType
-		UE_REQUIRES(std::is_void_v<DummyResultType>)
-	>
-#endif
-	void EmplaceValue()
-#ifdef __cpp_concepts
-		requires (std::is_void_v<ResultType>)
-#endif
-	{
-		this->GetState()->EmplaceResult(0);
-	}
 
 protected:
 	/**
diff --git a/Engine/Source/Runtime/Core/Public/Templates/TypeCompatibleBytes.h b/Engine/Source/Runtime/Core/Public/Templates/TypeCompatibleBytes.h
--- a/Engine/Source/Runtime/Core/Public/Templates/TypeCompatibleBytes.h
+++ b/Engine/Source/Runtime/Core/Public/Templates/TypeCompatibleBytes.h
@@ -5,6 +5,7 @@
 
 #include "CoreTypes.h"
 #include <string.h>
+#include <new>
 #include <type_traits>
 
 /**
@@ -33,12 +34,131 @@ struct TTypeCompatibleBytes
 	TTypeCompatibleBytes& operator=(TTypeCompatibleBytes&&) = delete;
 	TTypeCompatibleBytes& operator=(const TTypeCompatibleBytes&) = delete;
 
-	ElementType*		GetTypedPtr()		{ return (ElementType*)this;  }
-	const ElementType*	GetTypedPtr() const	{ return (const ElementType*)this; }
+	// GetTypedPtr only exists for backwards compatibility - these functions do not exist and cannot be implemented for the reference and void specializations.
+	ElementType* GetTypedPtr()
+	{
+		return (ElementType*)this;
+	}
+	const ElementType* GetTypedPtr() const
+	{
+		return (const ElementType*)this;
+	}
+
+	using MutableGetType = ElementType&;       // The type returned by Bytes.Get() where Bytes is a non-const lvalue
+	using ConstGetType   = const ElementType&; // The type returned by Bytes.Get() where Bytes is a const lvalue
+	using RvalueGetType  = ElementType&&;      // The type returned by Bytes.Get() where Bytes is an rvalue (non-const)
+
+	// Gets the inner element - no checks are performed to ensure an element is present.
+	ElementType& GetUnchecked() &
+	{
+		return *(ElementType*)this;
+	}
+	const ElementType& GetUnchecked() const&
+	{
+		return *(const ElementType*)this;
+	}
+	ElementType&& GetUnchecked() &&
+	{
+		return (ElementType&&)*(ElementType*)this;
+	}
+
+	// Emplaces an inner element.
+	// Note: no checks are possible to ensure that an element isn't already present.  DestroyUnchecked() must be called to end the element's lifetime.
+	template <typename... ArgTypes>
+	void EmplaceUnchecked(ArgTypes&&... Args)
+	{
+		new ((void*)GetTypedPtr()) ElementType((ArgTypes&&)Args...);
+	}
+
+	// Destroys the inner element.
+	// Note: no checks are possible to ensure that there is an element already present.
+	void DestroyUnchecked()
+	{
+		ElementTypeAlias_NatVisHelper* Ptr = (ElementTypeAlias_NatVisHelper*)this;
+		Ptr->ElementTypeAlias_NatVisHelper::~ElementTypeAlias_NatVisHelper();
+	}
 
 	alignas(ElementType) uint8 Pad[sizeof(ElementType)];
 };
 
+template <typename T>
+struct TTypeCompatibleBytes<T&>
+{
+	using ElementTypeAlias_NatVisHelper = T&;
+
+	// Trivially constructible and destuctible - users are responsible for managing the lifetime of the inner element.
+	TTypeCompatibleBytes() = default;
+	~TTypeCompatibleBytes() = default;
+
+	// Noncopyable
+	TTypeCompatibleBytes(TTypeCompatibleBytes&&) = delete;
+	TTypeCompatibleBytes(const TTypeCompatibleBytes&) = delete;
+	TTypeCompatibleBytes& operator=(TTypeCompatibleBytes&&) = delete;
+	TTypeCompatibleBytes& operator=(const TTypeCompatibleBytes&) = delete;
+
+	using MutableGetType = T&; // The type returned by Bytes.Get() where Bytes is a non-const lvalue
+	using ConstGetType   = T&; // The type returned by Bytes.Get() where Bytes is a const lvalue
+	using RvalueGetType  = T&; // The type returned by Bytes.Get() where Bytes is an rvalue (non-const)
+
+	// Gets the inner element - no checks are performed to ensure an element is present.
+	T& GetUnchecked() const
+	{
+		return *Ptr;
+	}
+
+	// Emplaces an inner element.
+	// Note: no checks are possible to ensure that an element isn't already present.  DestroyUnchecked() must be called to end the element's lifetime.
+	void EmplaceUnchecked(T& Ref)
+	{
+		Ptr = &Ref;
+	}
+
+	// Destroys the inner element.
+	// Note: no checks are possible to ensure that there is an element already present.
+	void DestroyUnchecked()
+	{
+	}
+
+	T* Ptr;
+};
+
+template <>
+struct TTypeCompatibleBytes<void>
+{
+	using ElementTypeAlias_NatVisHelper = void;
+
+	// Trivially constructible and destuctible - users are responsible for managing the lifetime of the inner element.
+	TTypeCompatibleBytes() = default;
+	~TTypeCompatibleBytes() = default;
+
+	// Noncopyable
+	TTypeCompatibleBytes(TTypeCompatibleBytes&&) = delete;
+	TTypeCompatibleBytes(const TTypeCompatibleBytes&) = delete;
+	TTypeCompatibleBytes& operator=(TTypeCompatibleBytes&&) = delete;
+	TTypeCompatibleBytes& operator=(const TTypeCompatibleBytes&) = delete;
+
+	using MutableGetType = void; // The type returned by Bytes.Get() where Bytes is a non-const lvalue
+	using ConstGetType   = void; // The type returned by Bytes.Get() where Bytes is a const lvalue
+	using RvalueGetType  = void; // The type returned by Bytes.Get() where Bytes is an rvalue (non-const)
+
+	// Gets the inner element - no checks are performed to ensure an element is present.
+	void GetUnchecked() const
+	{
+	}
+
+	// Emplaces an inner element.
+	// Note: no checks are possible to ensure that an element isn't already present.  DestroyUnchecked() must be called to end the element's lifetime.
+	void EmplaceUnchecked()
+	{
+	}
+
+	// Destroys the inner element.
+	// Note: no checks are possible to ensure that there is an element already present.
+	void DestroyUnchecked()
+	{
+	}
+};
+
 template <
 	typename ToType,
 	typename FromType,
