Optimize resetting a TFunction.

The original version of Reset, `*this = nullptr`, would construct a full TFunction and go through operator=(TFunction&&) to Swap against the empty TFunction. This is less efficient than calling Super::Reset() directly. By removing the overload we now get the behavior of Super::Reset() automatically.

Of course, our users might still do `Func = nullptr` instead of `Func.Reset()` and we'd like to be fast in that case as well. As an optimization, TFunction now has a dedicated `operator=(nullptr_t)` which calls Reset directly.

diff --git a/Engine/Source/Runtime/Core/Public/Templates/Function.h b/Engine/Source/Runtime/Core/Public/Templates/Function.h
--- a/Engine/Source/Runtime/Core/Public/Templates/Function.h
+++ b/Engine/Source/Runtime/Core/Public/Templates/Function.h
@@ -482,6 +482,9 @@ namespace UE::Core::Private::Function
 			Storage.Unbind();
 		}
 
+		/**
+		 * Removes any bound callable from the function, restoring it to the default 'empty' state.
+		 */
 		void Reset()
 		{
 			if (Callable)
@@ -762,11 +765,12 @@ public:
 	}
 
 	/**
-	 * Removes any bound callable from the TFunction, restoring it to the default 'empty' state.
+	 * Assignment to null is a form of Reset.
 	 */
-	void Reset()
+	TFunction& operator=(TYPE_OF_NULLPTR)
 	{
-		*this = nullptr;
+		Super::Reset();
+		return *this;
 	}
 
 	/**
