Avoid extra shrinks in CompactInvocationList.

The code is designed to only shrink manually once the size goes below a certain threshold. However, the cleanup loop forgot to pass EAllowShrinking::No, so we were also getting automatic shrinks from TArray.

diff --git a/Engine/Source/Runtime/Core/Public/Delegates/MulticastDelegateBase.h b/Engine/Source/Runtime/Core/Public/Delegates/MulticastDelegateBase.h
--- a/Engine/Source/Runtime/Core/Public/Delegates/MulticastDelegateBase.h
+++ b/Engine/Source/Runtime/Core/Public/Delegates/MulticastDelegateBase.h
@@ -366,7 +366,7 @@ private:
 			IDelegateInstance* DelegateInstance = DelegateBaseRef.GetDelegateInstanceProtected();
 			if (DelegateInstance == nullptr	|| DelegateInstance->IsCompactable())
 			{
-				InvocationList.RemoveAtSwap(InvocationListIndex);
+				InvocationList.RemoveAtSwap(InvocationListIndex, EAllowShrinking::No);
 			}
 			else
 			{
