Re-enabled UE_LIFETIMEBOUND on TFuture::Get() and TFuture::GetMutable().

diff --git a/Engine/Source/Runtime/Core/Public/Async/Future.h b/Engine/Source/Runtime/Core/Public/Async/Future.h
--- a/Engine/Source/Runtime/Core/Public/Async/Future.h
+++ b/Engine/Source/Runtime/Core/Public/Async/Future.h
@@ -237,10 +237,23 @@ public:
 	 * @return The result as a const reference, or the same reference if the future holds a reference, or void if the future holds a void.
 	 * @note Not equivalent to std::future::get(). The future remains valid.
 	 */
-	ConstResultType Get() const // UE_LIFETIMEBOUND
+#ifndef __cpp_concepts
+	ConstResultType Get() const UE_LIFETIMEBOUND // this will give false positives when getting from an rvalue TFuture<T&> or TFuture<void> until C++20
+	{
+		return this->GetState()->GetResult();
+	}
+#else
+	ConstResultType Get() const
+		requires (!std::is_object_v<ResultType>)
 	{
 		return this->GetState()->GetResult();
 	}
+	ConstResultType Get() const UE_LIFETIMEBOUND
+		requires (std::is_object_v<ResultType>)
+	{
+		return this->GetState()->GetResult();
+	}
+#endif
 
 	/**
 	 * Checks whether this future object has its value set.
@@ -416,10 +429,23 @@ public:
 	 * @return The result as a non-const reference, or the same reference if the future holds a reference, or void if the future holds a void.
 	 * @note Not equivalent to std::future::get(). The future remains valid.
 	 */
-	MutableResultType GetMutable() // UE_LIFETIMEBOUND
+#ifndef __cpp_concepts
+	MutableResultType GetMutable() UE_LIFETIMEBOUND // this will give false positives when getting from an rvalue TFuture<T&> or TFuture<void> until C++20
+	{
+		return (MutableResultType)this->Get();
+	}
+#else
+	MutableResultType GetMutable()
+		requires (!std::is_object_v<ResultType>)
 	{
 		return (MutableResultType)this->Get();
 	}
+	MutableResultType GetMutable() UE_LIFETIMEBOUND
+		requires (std::is_object_v<ResultType>)
+	{
+		return (MutableResultType)this->Get();
+	}
+#endif
 
 	/**
 	 * Consumes the future's result and invalidates the future.
