Fixed LocalPrint to avoid truncation on Windows

diff --git a/Engine/Source/Runtime/Core/Private/Windows/WindowsPlatformMisc.cpp b/Engine/Source/Runtime/Core/Private/Windows/WindowsPlatformMisc.cpp
--- a/Engine/Source/Runtime/Core/Private/Windows/WindowsPlatformMisc.cpp
+++ b/Engine/Source/Runtime/Core/Private/Windows/WindowsPlatformMisc.cpp
@@ -1329,10 +1329,35 @@ void FWindowsPlatformMisc::SetUTF8Output()
 	_setmode(_fileno(stdout), _O_U8TEXT);
 }
 
-void FWindowsPlatformMisc::LocalPrint( const TCHAR *Message )
+void FWindowsPlatformMisc::LocalPrint(const TCHAR* Message)
 {
 #if USE_DEBUG_LOGGING
-	OutputDebugString(Message);
+	// Use SRWLOCK directly until we have a trivially destructible FSharedMutex.
+	static SRWLOCK PrivateLock;
+	static PSRWLOCK Lock = (InitializeSRWLock(&PrivateLock), &PrivateLock);
+
+	// OutputDebugString has no documented maximum length but the OUTPUT_DEBUG_STRING_INFO
+	// type that receives debug strings in debuggers has a WORD count of bytes including
+	// the null terminator. Break up the string into parts if it exceeds this limit to
+	// avoid truncation in debuggers that do not scan for the null terminator.
+	constexpr int32 MaxOutputDebugStringLength = WORD_MAX / sizeof(TCHAR) - 1;
+	if (const int32 Len = FCString::Strlen(Message); LIKELY(Len <= MaxOutputDebugStringLength))
+	{
+		AcquireSRWLockShared(Lock);
+		OutputDebugString(Message);
+		ReleaseSRWLockShared(Lock);
+	}
+	else
+	{
+		AcquireSRWLockExclusive(Lock);
+		FStringView MessageView(Message, Len);
+		for (FString Part; !MessageView.IsEmpty(); MessageView.RightChopInline(MaxOutputDebugStringLength))
+		{
+			Part = MessageView.Left(MaxOutputDebugStringLength);
+			OutputDebugString(*Part);
+		}
+		ReleaseSRWLockExclusive(Lock);
+	}
 #endif
 }
 
