CIS fix.

diff --git a/Engine/Source/Runtime/Core/Public/Async/Future.h b/Engine/Source/Runtime/Core/Public/Async/Future.h
--- a/Engine/Source/Runtime/Core/Public/Async/Future.h
+++ b/Engine/Source/Runtime/Core/Public/Async/Future.h
@@ -508,6 +508,10 @@ class TSharedFuture final
 
 	using BaseType = TFutureBase<ResultType>;
 
+	using MutableResultType = typename BaseType::MutableResultType;
+	using ConstResultType   = typename BaseType::ConstResultType;
+	using RvalueResultType  = typename BaseType::RvalueResultType;
+
 public:
 
 	/** Default constructor. */
@@ -522,6 +526,19 @@ public:
 		: BaseType(MoveTemp(Future))
 	{ }
 
+	/**
+	 * Gets the future's result.
+	 *
+	 * @return The result as a const reference, or the same reference if the future holds a reference, or void if the future holds a void.
+	 * @note Not equivalent to std::future::get(). The future remains valid.
+	 */
+	ConstResultType Get() const
+	{
+		// This forwarding function is necessary to 'cancel' the UE_LIFETIMEBOUND of the base class, as
+		// other shared futures can keep the object alive.
+		return BaseType::Get();
+	}
+
 private:
 	// Forward constructors
 	using BaseType::TFutureBase;
