[AutoRTFM] Handle Linux case-insensitive filename lookup in the open.

This was failing within a transaction because the file-system cache leverages a transactionally-unsafe critical section. However, just updating the lock isn't a sufficient fix; past this point, the code would also use many other calls like `opendir` and `readdir` which may not currently be AutoRTFM-aware.

The only state involved here is a cache of properly-cased file system paths. This does not need to be rolled back if a transaction is aborted; in fact, it would be silly to roll back those lookups. Running the case fixup code in the open should resolve this issue.

diff --git a/Engine/Source/Runtime/Core/Private/Unix/UnixPlatformFile.cpp b/Engine/Source/Runtime/Core/Private/Unix/UnixPlatformFile.cpp
--- a/Engine/Source/Runtime/Core/Private/Unix/UnixPlatformFile.cpp
+++ b/Engine/Source/Runtime/Core/Private/Unix/UnixPlatformFile.cpp
@@ -8,6 +8,7 @@
 #include "Logging/LogMacros.h"
 #include "Misc/Paths.h"
 #include "Async/MappedFileHandle.h"
+#include "AutoRTFM/AutoRTFM.h"
 #include <sys/file.h>
 #include <sys/mman.h>
 
@@ -653,16 +654,19 @@ public:
 			return false;
 		}
 
-		// try the filename as given first
-		struct stat StatInfo;
-		bool bFound = stat(TCHAR_TO_UTF8(*PossiblyWrongFilename), &StatInfo) == 0;
-
-		if (bFound)
-		{
-			ExistingFilename = PossiblyWrongFilename;
-		}
-		else
+		bool bFound = false;
+		ExistingFilename = AutoRTFM::Open([this, &PossiblyWrongFilename, &bFound]
 		{
+			// Try the filename as given first.
+			struct stat StatInfo;
+			bFound = stat(TCHAR_TO_UTF8(*PossiblyWrongFilename), &StatInfo) == 0;
+
+			if (bFound)
+			{
+				return PossiblyWrongFilename;
+			}
+
+			// Next, check in the cache.
 			FileMapCache& Cache = GetFileMapCache();
 			UE::TScopeLock Lock(Cache);
 			const FileEntry* Entry = Cache.Find(PossiblyWrongFilename);
@@ -672,35 +676,37 @@ public:
 				if (Entry->IsInvalid())
 				{
 					bFound = false;
+					return FString();
 				}
 				else
 				{
-					ExistingFilename = Entry->File;
 					bFound = true;
+					return Entry->File;
 				}
 			}
-			else
+
+			// We haven't seen this path before. Perform a case-insensitive search from /
+			int MaxPathComponents = CountPathComponents(PossiblyWrongFilename);
+			if (MaxPathComponents == 0)
 			{
-				// perform a case-insensitive search from /
+				// Non-empty paths should always have at least one component; we don't expect this to happen.
+				return FString();
+			}
 
-				int MaxPathComponents = CountPathComponents(PossiblyWrongFilename);
-				if (MaxPathComponents > 0)
-				{
-					FString FoundFilename(TEXT("/"));	// start with root
-					bFound = MapFileRecursively(PossiblyWrongFilename, 0, MaxPathComponents, FoundFilename);
-					if (bFound)
-					{
-						ExistingFilename = FoundFilename;
-						Cache.AddEntry(PossiblyWrongFilename, ExistingFilename);
-					}
-					else
-					{
-						// Cache a failed to find entry, we'll look again if the call in is greater then MaxInvalidCacheTime from this cache point
-						Cache.AddEntry(PossiblyWrongFilename, "");
-					}
-				}
+			FString FoundFilename(TEXT("/"));	// Start from the root.
+			bFound = MapFileRecursively(PossiblyWrongFilename, 0, MaxPathComponents, FoundFilename);
+			if (bFound)
+			{
+				Cache.AddEntry(PossiblyWrongFilename, FoundFilename);
+				return FoundFilename;
 			}
-		}
+			else
+			{
+				// Cache a failed to find entry. We'll look again if the call in is greater then MaxInvalidCacheTime from this cache point.
+				Cache.AddEntry(PossiblyWrongFilename, "");
+				return FString();
+			}
+		});
 
 		return bFound;
 	}
