Moved FPooledSyncEvent into its own header.

diff --git a/Engine/Source/Runtime/Core/Public/Async/Future.h b/Engine/Source/Runtime/Core/Public/Async/Future.h
--- a/Engine/Source/Runtime/Core/Public/Async/Future.h
+++ b/Engine/Source/Runtime/Core/Public/Async/Future.h
@@ -10,7 +10,7 @@
 #include "Templates/SharedPointer.h"
 #include "Misc/DateTime.h"
 #include "HAL/Event.h"
-#include "HAL/PlatformProcess.h"
+#include "HAL/PooledSyncEvent.h"
 #include "Misc/ScopeLock.h"
 #include "Templates/Models.h"
 #include "Templates/Requires.h"
diff --git a/Engine/Source/Runtime/Core/Public/HAL/PlatformProcess.h b/Engine/Source/Runtime/Core/Public/HAL/PlatformProcess.h
--- a/Engine/Source/Runtime/Core/Public/HAL/PlatformProcess.h
+++ b/Engine/Source/Runtime/Core/Public/HAL/PlatformProcess.h
@@ -5,44 +5,3 @@
 #include "CoreTypes.h"
 
 #include COMPILED_PLATFORM_HEADER(PlatformProcess.h)
-
-// An RAII wrapper around FPlatformProcess::GetSynchEventFromPool/ReturnSynchEventToPool.
-struct FPooledSyncEvent
-{
-	explicit FPooledSyncEvent(bool bIsManualReset = false)
-	{
-		Event = FPlatformProcess::GetSynchEventFromPool(bIsManualReset);
-	}
-	~FPooledSyncEvent()
-	{
-		*this = nullptr;
-	}
-
-	// Non-copyable
-	FPooledSyncEvent(FPooledSyncEvent&&) = delete;
-	FPooledSyncEvent(const FPooledSyncEvent&) = delete;
-	FPooledSyncEvent& operator=(FPooledSyncEvent&&) = delete;
-	FPooledSyncEvent& operator=(const FPooledSyncEvent&) = delete;
-
-	explicit operator bool() const
-	{
-		return !!this->Event;
-	}
-
-	FEvent* operator->() const
-	{
-		return this->Event;
-	}
-
-	FPooledSyncEvent& operator=(TYPE_OF_NULLPTR)
-	{
-		if (Event)
-		{
-			FPlatformProcess::ReturnSynchEventToPool(Event);
-			Event = nullptr;
-		}
-		return *this;
-	}
-
-	FEvent* Event;
-};
diff --git a/Engine/Source/Runtime/Core/Public/HAL/PooledSyncEvent.h b/Engine/Source/Runtime/Core/Public/HAL/PooledSyncEvent.h
new file mode 100644
--- /dev/null
+++ b/Engine/Source/Runtime/Core/Public/HAL/PooledSyncEvent.h
@@ -0,0 +1,46 @@
+// Copyright Epic Games, Inc. All Rights Reserved.
+
+#pragma once
+
+#include "PlatformProcess.h"
+
+// An RAII wrapper around FPlatformProcess::GetSynchEventFromPool/ReturnSynchEventToPool.
+struct FPooledSyncEvent
+{
+	explicit FPooledSyncEvent(bool bIsManualReset = false)
+	{
+		Event = FPlatformProcess::GetSynchEventFromPool(bIsManualReset);
+	}
+	~FPooledSyncEvent()
+	{
+		*this = nullptr;
+	}
+
+	// Non-copyable
+	FPooledSyncEvent(FPooledSyncEvent&&) = delete;
+	FPooledSyncEvent(const FPooledSyncEvent&) = delete;
+	FPooledSyncEvent& operator=(FPooledSyncEvent&&) = delete;
+	FPooledSyncEvent& operator=(const FPooledSyncEvent&) = delete;
+
+	explicit operator bool() const
+	{
+		return !!this->Event;
+	}
+
+	FEvent* operator->() const
+	{
+		return this->Event;
+	}
+
+	FPooledSyncEvent& operator=(TYPE_OF_NULLPTR)
+	{
+		if (Event)
+		{
+			FPlatformProcess::ReturnSynchEventToPool(Event);
+			Event = nullptr;
+		}
+		return *this;
+	}
+
+	FEvent* Event;
+};
