[AutoRTFM] Prevent infinite recursion when asserting in FTransaction::DidFree()

The assertion can call back into DidFree(), which results in a stack overflow.

diff --git a/Engine/Source/Runtime/Core/Private/AutoRTFM/Transaction.h b/Engine/Source/Runtime/Core/Private/AutoRTFM/Transaction.h
--- a/Engine/Source/Runtime/Core/Private/AutoRTFM/Transaction.h
+++ b/Engine/Source/Runtime/Core/Private/AutoRTFM/Transaction.h
@@ -217,7 +217,7 @@ private:
     EState CurrentState{EState::Uninitialized};
     FStackRange StackRange;
     bool bIsStackScoped = false;
-    bool bIsInDidAllocate = false;
+    bool bIsInAllocateFn = false;
 };
 
 } // namespace AutoRTFM
diff --git a/Engine/Source/Runtime/Core/Private/AutoRTFM/TransactionInlines.h b/Engine/Source/Runtime/Core/Private/AutoRTFM/TransactionInlines.h
--- a/Engine/Source/Runtime/Core/Private/AutoRTFM/TransactionInlines.h
+++ b/Engine/Source/Runtime/Core/Private/AutoRTFM/TransactionInlines.h
@@ -218,23 +218,27 @@ template<unsigned SIZE> AUTORTFM_NO_ASAN FORCENOINLINE void FTransaction::Record
 
 UE_AUTORTFM_FORCEINLINE void FTransaction::DidAllocate(void* LogicalAddress, const size_t Size)
 {
-	if (0 == Size || bIsInDidAllocate)
+	if (0 == Size || bIsInAllocateFn)
 	{
 		return;
 	}
 
-	AutoRTFM::TScopedGuard<bool> RecursionGuard(bIsInDidAllocate, true);
-    const bool DidInsert = NewMemoryTracker.Insert(LogicalAddress, Size);
-    ASSERT(DidInsert);
+	AutoRTFM::TScopedGuard<bool> RecursionGuard(bIsInAllocateFn, true);
+	const bool DidInsert = NewMemoryTracker.Insert(LogicalAddress, Size);
+	ASSERT(DidInsert);
 }
 
 UE_AUTORTFM_FORCEINLINE void FTransaction::DidFree(void* LogicalAddress)
 {
-    ASSERT(bTrackAllocationLocations);
+	ASSERT(bTrackAllocationLocations);
 
-    // Checking if one byte is in the interval map is enough to ascertain if it
-    // is new memory and we should be worried.
-    ASSERT(bIsInDidAllocate || !NewMemoryTracker.Contains(LogicalAddress, 1));
+	// Checking if one byte is in the interval map is enough to ascertain if it
+	// is new memory and we should be worried.
+	if (!bIsInAllocateFn)
+	{
+		AutoRTFM::TScopedGuard<bool> RecursionGuard(bIsInAllocateFn, true);
+		ASSERT(!NewMemoryTracker.Contains(LogicalAddress, 1));
+	}
 }
 
 UE_AUTORTFM_FORCEINLINE void FTransaction::DeferUntilCommit(TFunction<void()>&& Callback)
