[AutoRTFM] Include a callstack in AutoRTFM's ASSERT()

We currently only print the assertion message and then terminate. The callstack helps identify the cause of the assertion.

diff --git a/Engine/Source/Runtime/Core/Private/AutoRTFM/Utils.h b/Engine/Source/Runtime/Core/Private/AutoRTFM/Utils.h
--- a/Engine/Source/Runtime/Core/Private/AutoRTFM/Utils.h
+++ b/Engine/Source/Runtime/Core/Private/AutoRTFM/Utils.h
@@ -64,7 +64,16 @@ FString GetFunctionDescription(TReturnType (*FunctionPtr)(TParameterTypes...))
 // and line numbers), but then pass this to another function
 // `DoAssert` which has a special calling convention that makes
 // the caller as optimal as possible (at the expense of the callee).
-#define ASSERT(exp) if (UNLIKELY(!(exp))) { auto Lambda = []() { UE_LOG(LogAutoRTFM, Fatal, TEXT("AutoRTFM assert %s!"), TEXT(#exp)); }; AutoRTFM::DoAssert(Lambda); }
+#define ASSERT(exp) \
+	if (UNLIKELY(!(exp))) \
+	{ \
+		auto Lambda = [] \
+		{ \
+			FDebug::DumpStackTraceToLog(TEXT("AutoRTFM assert backtrace"), ELogVerbosity::Error); \
+			UE_LOG(LogAutoRTFM, Fatal, TEXT("AutoRTFM assert %s!"), TEXT(#exp)); \
+		}; \
+		AutoRTFM::DoAssert(Lambda); \
+	}
 
 #if defined(__has_feature)
 #if __has_feature(address_sanitizer)
