[AutoRTFM] Handle dcpvar AutoRTFMMemoryValidationLevel=0

0 maps to the special 'Default' enum value, which should not be assigned to the underlying global variable.
Translate this to the actual default value (disabled by default)

diff --git a/Engine/Source/Runtime/Core/Private/AutoRTFM/API.cpp b/Engine/Source/Runtime/Core/Private/AutoRTFM/API.cpp
--- a/Engine/Source/Runtime/Core/Private/AutoRTFM/API.cpp
+++ b/Engine/Source/Runtime/Core/Private/AutoRTFM/API.cpp
@@ -27,6 +27,9 @@ namespace
 
 	static constexpr AutoRTFM::EMemoryValidationLevel GDefaultMemoryValidationLevel = AutoRTFM::EMemoryValidationLevel::Disabled;
 
+	// Note: GMemoryValidationLevel should never be EMemoryValidationLevel::Default.
+	// EMemoryValidationLevel::Default is a special enumerator that can be used
+	// in the public API to map to GDefaultMemoryValidationLevel.
 	int GMemoryValidationLevel = static_cast<int>(GDefaultMemoryValidationLevel);
 	
 	bool GMemoryValidationThrottlingEnabled = true;
@@ -58,6 +61,12 @@ namespace
 
 	void OnAutoRTFMMemoryValidationLevelChanged()
 	{
+		if (GMemoryValidationLevel == static_cast<int>(AutoRTFM::EMemoryValidationLevel::Default))
+		{
+			// Map EMemoryValidationLevel::Default to GDefaultMemoryValidationLevel.
+			GMemoryValidationLevel = static_cast<int>(GDefaultMemoryValidationLevel);
+		}
+
 		switch (static_cast<AutoRTFM::EMemoryValidationLevel>(GMemoryValidationLevel))
 		{
 		case AutoRTFM::EMemoryValidationLevel::Disabled:
