Enable LLM in Test by default.

Rationale:
-  absence of this requires tedious hookup when hooking up projects to automated performance testing
-  no effect on shipping builds, better introspection of test builds

diff --git a/Engine/Source/Runtime/Core/Public/HAL/LowLevelMemTrackerDefines.h b/Engine/Source/Runtime/Core/Public/HAL/LowLevelMemTrackerDefines.h
--- a/Engine/Source/Runtime/Core/Public/HAL/LowLevelMemTrackerDefines.h
+++ b/Engine/Source/Runtime/Core/Public/HAL/LowLevelMemTrackerDefines.h
@@ -8,7 +8,7 @@
 #include "Misc/Build.h"
 
 #ifndef ALLOW_LOW_LEVEL_MEM_TRACKER_IN_TEST
-	#define ALLOW_LOW_LEVEL_MEM_TRACKER_IN_TEST 0
+	#define ALLOW_LOW_LEVEL_MEM_TRACKER_IN_TEST 1	// enabled in Test by default to facilitate LLM captures in auto perf testing
 #endif
 
 // LLM_ENABLED_IN_CONFIG can be defined here or in the build environment only.
