Fix more memory unsafe functions on FName and FDebugName
- Deprecated the public APIs that could cause buffer overflow.
- Fixed up the places in the engine using the unsafe function.

Reduce size of FRHIBreadcrumb FName buffers, to reduce stack usage.
- Dropped from ~2KB to 256 bytes of stack memory per FName arg.

diff --git a/Engine/Source/Runtime/Core/Private/UObject/UnrealNames.cpp b/Engine/Source/Runtime/Core/Private/UObject/UnrealNames.cpp
--- a/Engine/Source/Runtime/Core/Private/UObject/UnrealNames.cpp
+++ b/Engine/Source/Runtime/Core/Private/UObject/UnrealNames.cpp
@@ -2299,22 +2299,22 @@ void FNameEntry::StoreName(const WIDECHAR* InName, uint32 Len)
 	Encode(WideName, Len);
 }
 
-void FNameEntry::CopyUnterminatedName(ANSICHAR* Out) const
+void FNameEntry::CopyUnterminatedName(ANSICHAR* Out, uint32 Len) const
 {
-	FPlatformMemory::Memcpy(Out, AnsiName, sizeof(ANSICHAR) * Header.Len);
-	Decode(Out, Header.Len);
+	FPlatformMemory::Memcpy(Out, AnsiName, sizeof(ANSICHAR) * Len);
+	Decode(Out, Len);
 }
 
-void FNameEntry::CopyUnterminatedName(WIDECHAR* Out) const
+void FNameEntry::CopyUnterminatedName(WIDECHAR* Out, uint32 Len) const
 {
-	FPlatformMemory::Memcpy(Out, WideName, sizeof(WIDECHAR) * Header.Len);
-	Decode(Out, Header.Len);
+	FPlatformMemory::Memcpy(Out, WideName, sizeof(WIDECHAR) * Len);
+	Decode(Out, Len);
 }
 
 FORCEINLINE const WIDECHAR* FNameEntry::GetUnterminatedName(WIDECHAR(&OptionalDecodeBuffer)[NAME_SIZE]) const
 {
 #ifdef WITH_CUSTOM_NAME_ENCODING
-	CopyUnterminatedName(OptionalDecodeBuffer);
+	CopyUnterminatedName(OptionalDecodeBuffer, Header.Len);
 	return OptionalDecodeBuffer;
 #else
 	return WideName;
@@ -2324,7 +2324,7 @@ FORCEINLINE const WIDECHAR* FNameEntry::GetUnterminatedName(WIDECHAR(&OptionalDe
 FORCEINLINE ANSICHAR const* FNameEntry::GetUnterminatedName(ANSICHAR(&OptionalDecodeBuffer)[NAME_SIZE]) const
 {
 #ifdef WITH_CUSTOM_NAME_ENCODING
-	CopyUnterminatedName(OptionalDecodeBuffer);
+	CopyUnterminatedName(OptionalDecodeBuffer, Header.Len);
 	return OptionalDecodeBuffer;
 #else
 	return AnsiName;
@@ -2344,48 +2344,47 @@ FORCEINLINE FNameStringView FNameEntry::MakeView(FNameBuffer& OptionalDecodeBuff
 
 void FNameEntry::GetUnterminatedName(TCHAR* OutName, uint32 OutLen) const
 {
-	check(static_cast<int32>(OutLen) >= GetNameLength());
-	CopyAndConvertUnterminatedName(OutName);
+	CopyAndConvertUnterminatedName(OutName, FMath::Min<uint32>(Header.Len, OutLen));
 }
 
 void FNameEntry::GetName(TCHAR(&OutName)[NAME_SIZE]) const
 {
-	CopyAndConvertUnterminatedName(OutName);
+	CopyAndConvertUnterminatedName(OutName, Header.Len);
 	OutName[GetNameLength()] = TEXT('\0');
 }
 
-void FNameEntry::CopyAndConvertUnterminatedName(TCHAR* OutName) const
+void FNameEntry::CopyAndConvertUnterminatedName(TCHAR* OutName, uint32 Len) const
 {
 	if (sizeof(TCHAR) < sizeof(WIDECHAR) && IsWide()) // Normally compiled out
 	{
 		FNameBuffer Temp;
-		CopyUnterminatedName(Temp.WideName);
-		ConvertInPlace<WIDECHAR, TCHAR>(Temp.WideName, Header.Len);
-		FPlatformMemory::Memcpy(OutName, Temp.AnsiName, Header.Len * sizeof(TCHAR));
+		CopyUnterminatedName(Temp.WideName, Len);
+		ConvertInPlace<WIDECHAR, TCHAR>(Temp.WideName, Len);
+		FPlatformMemory::Memcpy(OutName, Temp.AnsiName, Len * sizeof(TCHAR));
 	}
 	else if (IsWide())
 	{
-		CopyUnterminatedName((WIDECHAR*)OutName);
-		ConvertInPlace<WIDECHAR, TCHAR>((WIDECHAR*)OutName, Header.Len);
+		CopyUnterminatedName((WIDECHAR*)OutName, Len);
+		ConvertInPlace<WIDECHAR, TCHAR>((WIDECHAR*)OutName, Len);
 	}
 	else
 	{
-		CopyUnterminatedName((ANSICHAR*)OutName);
-		ConvertInPlace<ANSICHAR, TCHAR>((ANSICHAR*)OutName, Header.Len);
+		CopyUnterminatedName((ANSICHAR*)OutName, Len);
+		ConvertInPlace<ANSICHAR, TCHAR>((ANSICHAR*)OutName, Len);
 	}
 }
 
 void FNameEntry::GetAnsiName(ANSICHAR(&Out)[NAME_SIZE]) const
 {
 	check(!IsWide());
-	CopyUnterminatedName(Out);
+	CopyUnterminatedName(Out, Header.Len);
 	Out[Header.Len] = '\0';
 }
 
 void FNameEntry::GetWideName(WIDECHAR(&Out)[NAME_SIZE]) const
 {
 	check(IsWide());
-	CopyUnterminatedName(Out);
+	CopyUnterminatedName(Out, Header.Len);
 	Out[Header.Len] = '\0';
 }
 
@@ -2432,12 +2431,12 @@ void FNameEntry::AppendNameToString(FWideStringBuilderBase& Out) const
 	TCHAR* OutChars = Out.GetData() + Offset;
 	if (Header.bIsWide)
 	{
-		CopyUnterminatedName(reinterpret_cast<WIDECHAR*>(OutChars));
+		CopyUnterminatedName(reinterpret_cast<WIDECHAR*>(OutChars), Header.Len);
 		ConvertInPlace<WIDECHAR, TCHAR>(reinterpret_cast<WIDECHAR*>(OutChars), Header.Len);
 	}
 	else
 	{
-		CopyUnterminatedName(reinterpret_cast<ANSICHAR*>(OutChars));
+		CopyUnterminatedName(reinterpret_cast<ANSICHAR*>(OutChars), Header.Len);
 		ConvertInPlace<ANSICHAR, TCHAR>(reinterpret_cast<ANSICHAR*>(OutChars), Header.Len);
 	}
 }
@@ -2455,7 +2454,7 @@ void FNameEntry::AppendNameToString(FUtf8StringBuilderBase& Out) const
 	{
 		const int32 Offset = Out.AddUninitialized(Header.Len);
 		UTF8CHAR* OutChars = Out.GetData() + Offset;
-		CopyUnterminatedName(reinterpret_cast<ANSICHAR*>(OutChars));
+		CopyUnterminatedName(reinterpret_cast<ANSICHAR*>(OutChars), Header.Len);
 	}
 }
 
@@ -2464,7 +2463,7 @@ void FNameEntry::AppendAnsiNameToString(FAnsiStringBuilderBase& Out) const
 	checkName(Header.Len != 0);
 	check(!IsWide());
 	const int32 Offset = Out.AddUninitialized(Header.Len);
-	CopyUnterminatedName(Out.GetData() + Offset);
+	CopyUnterminatedName(Out.GetData() + Offset, Header.Len);
 }
 
 void FNameEntry::AppendNameToPathString(FString& Out) const
@@ -3581,27 +3580,35 @@ uint32 FName::GetStringLength() const
 	return NameLen;
 }
 
-uint32 FName::ToString(TCHAR* Out, uint32 OutSize) const
+uint32 FName::ToStringInternal(TCHAR* Out, const uint32 OutSize) const
 {
+	check(OutSize > 0);
+	const uint32 OutLen = OutSize - 1;
+
 	const FNameEntry& Entry = *GetDisplayNameEntry();
-	uint32 NameLen = Entry.GetNameLength();
-	Entry.GetUnterminatedName(Out, OutSize);
+	const uint32 NameLen = Entry.GetNameLength();
+	Entry.GetUnterminatedName(Out, OutLen);
 
-	if (GetNumber() == NAME_NO_NUMBER_INTERNAL)
+	if (GetNumber() != NAME_NO_NUMBER_INTERNAL && OutLen > NameLen)
 	{
-		Out[NameLen] = TEXT('\0');
-		return NameLen;
+		int32 SuffixLen = FCString::Snprintf(Out + NameLen, OutSize - NameLen, TEXT("_%d"), NAME_INTERNAL_TO_EXTERNAL(GetNumber()));
+		if (SuffixLen < 0)
+		{
+			// Suffix was truncated
+			return OutLen;
+		}
+		else
+		{
+			// Suffix was fully written
+			return NameLen + SuffixLen;
+		}
 	}
 	else
 	{
-		TCHAR NumberSuffixStr[16];
-		int32 SuffixLen = FCString::Sprintf(NumberSuffixStr, TEXT("_%d"), NAME_INTERNAL_TO_EXTERNAL(GetNumber()));
-		uint32 TotalLen = NameLen + SuffixLen;
-		check(SuffixLen > 0 && OutSize > TotalLen);
-
-		FPlatformMemory::Memcpy(Out + NameLen, NumberSuffixStr, SuffixLen * sizeof(TCHAR));
-		Out[TotalLen] = TEXT('\0');
-		return TotalLen;
+		// No suffix, or no room to add it
+		uint32 NullIndex = FMath::Min<uint32>(OutLen, NameLen);
+		Out[NullIndex] = TEXT('\0');
+		return NullIndex;
 	}
 }
 
@@ -3898,6 +3905,67 @@ void CheckLazyName(const CharType(&Literal)[N])
 
 static void TestNameBatch();
 
+#if DO_CHECK
+struct FNameAutoTestTruncateHelper
+{
+	static constexpr const TCHAR TestStr[] = TEXT("ABC_2147483646");
+	static constexpr const uint32 TestStrLen = UE_ARRAY_COUNT(TestStr) - 1;
+
+	template <uint32 RequiredLen>
+	static void TestInner()
+	{
+		TCHAR CanaryValue;
+		FMemory::Memset(&CanaryValue, 0xaa, sizeof(CanaryValue));
+
+		// Make a buffer of RequiredLen characters, plus a null terminator.
+		// Place a Canary after it in memory, to catch buffer overflow.
+		struct { TCHAR Data[RequiredLen + 1]; TCHAR Canary; } Buffer;
+		FMemory::Memset(&Buffer, 0xaa, sizeof(Buffer));
+
+		uint32 Len = FName(TestStr).ToStringTruncate(Buffer.Data);
+		if constexpr (RequiredLen >= TestStrLen)
+		{
+			// Whole string should be returned
+			check(Len == TestStrLen);
+		}
+		else
+		{
+			// String should be truncated
+			check(Len == RequiredLen);
+		}
+
+		for (uint32 Index = 0; Index <= RequiredLen; ++Index)
+		{
+			if (Index < Len)
+			{
+				check(Buffer.Data[Index] == TestStr[Index]);
+			}
+			else if (Index == Len)
+			{
+				check(Buffer.Data[Index] == 0); // null terminator
+			}
+			else
+			{
+				check(Buffer.Data[Index] == CanaryValue); // should not have been written
+			}
+		}
+
+		check(Buffer.Canary == CanaryValue);
+	}
+
+	template <std::size_t... Ints>
+	static void TestOuter(std::index_sequence<Ints...>)
+	{
+		(TestInner<Ints>(), ...);
+	}
+
+	static void Run()
+	{
+		TestOuter(std::make_index_sequence<TestStrLen + 4>());
+	}
+};
+#endif
+
 void FName::AutoTest()
 {
 #if DO_CHECK
@@ -4098,6 +4166,8 @@ void FName::AutoTest()
 	check(FName("a") == ANSITEXTVIEW("A"));
 	check(FName("A") != ANSITEXTVIEW("b"));
 
+	// Check truncating ToString functions
+	FNameAutoTestTruncateHelper::Run();
 	
 	CheckLazyName("Hej");
 	CheckLazyName("hej");
diff --git a/Engine/Source/Runtime/Core/Public/UObject/NameTypes.h b/Engine/Source/Runtime/Core/Public/UObject/NameTypes.h
--- a/Engine/Source/Runtime/Core/Public/UObject/NameTypes.h
+++ b/Engine/Source/Runtime/Core/Public/UObject/NameTypes.h
@@ -262,11 +262,7 @@ public:
 	FORCEINLINE bool IsNumbered() const 						{ return false; }
 #endif
 
-	/**
-	 * Copy unterminated name to TCHAR buffer without allocating.
-	 *
-	 * @param OutSize must be at least GetNameLength()
-	 */
+	/** Copy unterminated name to TCHAR buffer without allocating. */
 	CORE_API void GetUnterminatedName(TCHAR* OutName, uint32 OutSize) const;
 
 	/** Copy null-terminated name to TCHAR buffer without allocating. */
@@ -326,9 +322,9 @@ private:
 
 	void StoreName(const ANSICHAR* InName, uint32 Len);
 	void StoreName(const WIDECHAR* InName, uint32 Len);
-	void CopyUnterminatedName(ANSICHAR* OutName) const;
-	void CopyUnterminatedName(WIDECHAR* OutName) const;
-	void CopyAndConvertUnterminatedName(TCHAR* OutName) const;
+	void CopyUnterminatedName(ANSICHAR* OutName, uint32 Len) const;
+	void CopyUnterminatedName(WIDECHAR* OutName, uint32 Len) const;
+	void CopyAndConvertUnterminatedName(TCHAR* OutName, uint32 Len) const;
 	const ANSICHAR* GetUnterminatedName(ANSICHAR(&OptionalDecodeBuffer)[NAME_SIZE]) const;
 	const WIDECHAR* GetUnterminatedName(WIDECHAR(&OptionalDecodeBuffer)[NAME_SIZE]) const;
 
@@ -667,6 +663,15 @@ public:
 	 */
 	static constexpr inline uint32 StringBufferSize = NAME_SIZE + 1 + 10; // NAME_SIZE includes null-terminator
 
+private:
+	/**
+	 * Internal implementation of non-allocating ToString().
+	 * OutSize is the number of valid elements in the buffer, including the null terminator. Must be > 0.
+	 * Returns the length of the string, excluding the null terminator.
+	 */
+	CORE_API uint32 ToStringInternal(TCHAR* Out, uint32 OutSize) const;
+
+public:
 	/**
 	 * Convert to string buffer to avoid dynamic allocations and returns string length
 	 *
@@ -674,13 +679,51 @@ public:
 	 *
 	 * Note that a default constructed FName returns "None" instead of ""
 	 */
-	CORE_API uint32 ToString(TCHAR* Out, uint32 OutSize) const;
+	UE_DEPRECATED(5.6, 
+		"FName::ToString(TCHAR* Out, uint32 OutSize) is dangerous and can lead to buffer overflow if the provided "
+		"buffer is smaller than FName::StringBufferSize, even if the OutSize parameter indicates the buffer is "
+		"smaller than this value. Use the templated ToString() or ToStringTruncate() functions to format the name "
+		"string into a pre-allocated array, or use the allocating ToString() functions that return an FString.")
+	uint32 ToString(TCHAR* Out, uint32 OutSize) const
+	{
+		return ToStringInternal(Out, StringBufferSize);
+	}
 
-	template<int N>
+	/**
+	 * Converts the FName to a string buffer, avoiding dynamic allocations.
+	 *
+	 * Returns the length of the string, excluding the null terminator.
+	 *
+	 * Note that a default constructed FName returns "None" instead of ""
+	 */
+	template <uint32 N>
 	uint32 ToString(TCHAR (&Out)[N]) const
 	{
-		UE_STATIC_DEPRECATE(5.6, N < StringBufferSize, "FName::ToString(TCHAR (&Out)[N]) requires a buffer of size of at least FName::StringBufferSize. Passing a smaller buffer can lead to buffer-overflow and crashes.");
-		return ToString(Out, N);
+		UE_STATIC_DEPRECATE(5.6, N < StringBufferSize,
+			"FName::ToString(TCHAR (&Out)[N]) requires a buffer of size of at least FName::StringBufferSize. "
+			"Use ToStringTruncate() if a smaller buffer is required and it is safe for the returned string to be truncated.");
+
+		return ToStringInternal(Out, N);
+	}
+
+	/**
+	 * Converts the FName to a string buffer, avoiding dynamic allocations. Truncates the name if it does not fit in the specified output buffer.
+	 * Use a buffer size of at least FName::StringBufferSize to avoid truncation.
+	 *
+	 * Returns the length of the (possibly truncated) string, excluding the null terminator. 
+	 *
+	 * Note that a default constructed FName returns "None" instead of ""
+	 */
+	uint32 ToStringTruncate(TCHAR* Out, uint32 OutSize) const
+	{
+		return ToStringInternal(Out, OutSize);
+	}
+
+	template <uint32 N>
+	uint32 ToStringTruncate(TCHAR (&Out)[N]) const
+	{
+		static_assert(N > 0, "Out buffer must have at least one element.");
+		return ToStringTruncate(Out, N);
 	}
 
 	/**
diff --git a/Engine/Source/Runtime/Engine/Private/SparseVolumeTexture/SparseVolumeTextureStreamingManager.cpp b/Engine/Source/Runtime/Engine/Private/SparseVolumeTexture/SparseVolumeTextureStreamingManager.cpp
--- a/Engine/Source/Runtime/Engine/Private/SparseVolumeTexture/SparseVolumeTextureStreamingManager.cpp
+++ b/Engine/Source/Runtime/Engine/Private/SparseVolumeTexture/SparseVolumeTextureStreamingManager.cpp
@@ -567,10 +567,6 @@ const FStreamingDebugInfo* FStreamingManager::GetStreamingDebugInfo(FRDGBuilder&
 		}
 		const int32 NumFrames = SVTInfo->PerFrameInfo.Num();
 		const int32 NumInstances = SVTInfo->StreamingInstances.Num();
-		const int32 NameStrLength = SVTInfo->SVTName.GetStringLength() + 1;
-
-		TCHAR* AssetName = GraphBuilder.AllocPODArray<TCHAR>(NameStrLength);
-		SVTInfo->SVTName.ToString(AssetName, NameStrLength);
 		
 		float* ResidencyPercentages = GraphBuilder.AllocPODArray<float>(NumFrames);
 		float* StreamingPercentages = GraphBuilder.AllocPODArray<float>(NumFrames);
@@ -601,7 +597,7 @@ const FStreamingDebugInfo* FStreamingManager::GetStreamingDebugInfo(FRDGBuilder&
 		}
 
 		FStreamingDebugInfo::FSVT& SVT = DebugInfoSVTs[SVTArrayWriteSlot++];
-		SVT.AssetName = AssetName;
+		SVT.AssetName = SVTInfo->SVTName;
 		SVT.FrameResidencyPercentages = ResidencyPercentages;
 		SVT.FrameStreamingPercentages = StreamingPercentages;
 		SVT.Instances = Instances;
diff --git a/Engine/Source/Runtime/Engine/Public/SparseVolumeTexture/ISparseVolumeTextureStreamingManager.h b/Engine/Source/Runtime/Engine/Public/SparseVolumeTexture/ISparseVolumeTextureStreamingManager.h
--- a/Engine/Source/Runtime/Engine/Public/SparseVolumeTexture/ISparseVolumeTextureStreamingManager.h
+++ b/Engine/Source/Runtime/Engine/Public/SparseVolumeTexture/ISparseVolumeTextureStreamingManager.h
@@ -28,7 +28,7 @@ struct FStreamingDebugInfo
 			float InBudgetMip;
 		};
 
-		const TCHAR* AssetName;
+		FName AssetName;
 		const float* FrameResidencyPercentages;
 		const float* FrameStreamingPercentages;
 		const FInstance* Instances;
diff --git a/Engine/Source/Runtime/Experimental/Iris/Core/Private/Iris/ReplicationState/ReplicationStateDescriptorBuilder.cpp b/Engine/Source/Runtime/Experimental/Iris/Core/Private/Iris/ReplicationState/ReplicationStateDescriptorBuilder.cpp
--- a/Engine/Source/Runtime/Experimental/Iris/Core/Private/Iris/ReplicationState/ReplicationStateDescriptorBuilder.cpp
+++ b/Engine/Source/Runtime/Experimental/Iris/Core/Private/Iris/ReplicationState/ReplicationStateDescriptorBuilder.cpp
@@ -1363,13 +1363,13 @@ void FPropertyReplicationStateDescriptorBuilder::BuildMemberDebugDescriptors(FRe
 {
 	FReplicationStateMemberDebugDescriptor* CurrentMemberDebugDescriptor = MemberDebugDescriptors;
 
-	TCHAR StringBuffer[1024];
+	TCHAR StringBuffer[FName::StringBufferSize];
 	// Structs derived from a struct with a custom NetSerializer can contain some replicated properties.
 	for (const FMemberProperty& Info : Members)
 	{
 		if (const FProperty* Property = Info.Property)
 		{
-			Property->GetFName().ToString(StringBuffer, 1024);
+			Property->GetFName().ToString(StringBuffer);
 			CurrentMemberDebugDescriptor->DebugName = CreatePersistentNetDebugName(StringBuffer);
 		}
 		else
diff --git a/Engine/Source/Runtime/RHI/Private/RHI.cpp b/Engine/Source/Runtime/RHI/Private/RHI.cpp
--- a/Engine/Source/Runtime/RHI/Private/RHI.cpp
+++ b/Engine/Source/Runtime/RHI/Private/RHI.cpp
@@ -2125,16 +2125,6 @@ FDebugName& FDebugName::operator=(FName Other)
 	return *this;
 }
 
-uint32 FDebugName::ToString(TCHAR* Out, uint32 OutSize) const
-{
-	uint32 NumChars = Name.ToString(Out, OutSize);
-	if (Number != NAME_NO_NUMBER_INTERNAL)
-	{
-		NumChars += FCString::Snprintf(Out + NumChars, OutSize - NumChars, TEXT("_%u"), Number);
-	}
-	return NumChars;
-}
-
 FString FDebugName::ToString() const
 {
 	FString Out;
@@ -2146,6 +2136,33 @@ FString FDebugName::ToString() const
 	return Out;
 }
 
+uint32 FDebugName::ToStringInternal(TCHAR* Out, const uint32 OutSize) const
+{
+	check(OutSize > 0);
+	const uint32 OutLen = OutSize - 1;
+
+	uint32 NameLen = Name.ToStringTruncate(Out, OutSize);
+	if (Number != NAME_NO_NUMBER_INTERNAL && OutLen > NameLen)
+	{
+		int32 SuffixLen = FCString::Snprintf(Out + NameLen, OutSize - NameLen, TEXT("_%d"), Number);
+		if (SuffixLen < 0)
+		{
+			// Suffix was truncated.
+			return OutLen;
+		}
+		else
+		{
+			// Suffix was fully written
+			return NameLen + SuffixLen;
+		}
+	}
+	else
+	{
+		// No suffix, or no room to add it
+		return FMath::Min<uint32>(OutLen, NameLen);
+	}
+}
+
 void FDebugName::AppendString(FStringBuilderBase& Builder) const
 {
 	Name.AppendString(Builder);
diff --git a/Engine/Source/Runtime/RHI/Public/RHIBreadcrumbs.h b/Engine/Source/Runtime/RHI/Public/RHIBreadcrumbs.h
--- a/Engine/Source/Runtime/RHI/Public/RHIBreadcrumbs.h
+++ b/Engine/Source/Runtime/RHI/Public/RHIBreadcrumbs.h
@@ -604,9 +604,12 @@ namespace UE::RHI::GPUProfiler
 #endif
 		}
 
+		// The maximum length of a breadcrumb string, including the null terminator.
+		static constexpr uint32 MaxLength = 128;
+
 		struct FBuffer
 		{
-			TCHAR Data[128];
+			TCHAR Data[MaxLength];
 		};
 		virtual TCHAR const* GetTCHAR(FBuffer& Storage) const = 0;
 		virtual TCHAR const* GetTCHARNoFormat() const = 0;
@@ -729,10 +732,10 @@ namespace UE::RHI::GPUProfiler
 
 		struct FConvert
 		{
-			TCHAR Inner[FName::StringBufferSize];
+			TCHAR Inner[FRHIBreadcrumb::MaxLength];
 			FConvert(TRHIBreadcrumbValue const& Name)
 			{
-				Name.Value.ToString(Inner);
+				Name.Value.ToStringTruncate(Inner);
 			}
 		};
 	};
@@ -748,10 +751,10 @@ namespace UE::RHI::GPUProfiler
 
 		struct FConvert
 		{
-			TCHAR Inner[128];
+			TCHAR Inner[FRHIBreadcrumb::MaxLength];
 			FConvert(TRHIBreadcrumbValue const& Name)
 			{
-				Name.Value.ToString(Inner);
+				Name.Value.ToStringTruncate(Inner);
 			}
 		};
 	};
@@ -760,7 +763,7 @@ namespace UE::RHI::GPUProfiler
 	template <>
 	struct TRHIBreadcrumbValue<FString>
 	{
-		TCHAR Buffer[128];
+		TCHAR Buffer[FRHIBreadcrumb::MaxLength];
 		TRHIBreadcrumbValue(FString const& Value)
 		{
 			FCString::Strncpy(Buffer, *Value, UE_ARRAY_COUNT(Buffer));
diff --git a/Engine/Source/Runtime/RHI/Public/RHIDefinitions.h b/Engine/Source/Runtime/RHI/Public/RHIDefinitions.h
--- a/Engine/Source/Runtime/RHI/Public/RHIDefinitions.h
+++ b/Engine/Source/Runtime/RHI/Public/RHIDefinitions.h
@@ -74,6 +74,9 @@ static_assert(sizeof(void*) <= SHADER_PARAMETER_POINTER_ALIGNMENT, "The alignmen
 class FDebugName
 {
 public:
+	// Buffer size required for any null-terminated FDebugName string, i.e. [name] '_' [digits] '\0'
+	static constexpr uint32 StringBufferSize = FName::StringBufferSize + 1 + 10;
+
 	RHI_API FDebugName();
 	RHI_API FDebugName(FName InName);
 	RHI_API FDebugName(FName InName, int32 InNumber);
@@ -81,12 +84,43 @@ public:
 	RHI_API FDebugName& operator=(FName Other);
 
 	RHI_API FString ToString() const;
-	RHI_API uint32 ToString(TCHAR* Out, uint32 OutSize) const;
 
-	template<int N>
+	UE_DEPRECATED(5.6,
+		"FDebugName::ToString(TCHAR* Out, uint32 OutSize) is dangerous and can lead to buffer overflow if the provided "
+		"buffer is smaller than FDebugName::StringBufferSize, even if the OutSize parameter indicates the buffer is "
+		"smaller than this value. Use the templated ToString() or ToStringTruncate() functions to format the name "
+		"string into a pre-allocated array, or use the allocating ToString() function that returns an FString.")
+	uint32 ToString(TCHAR* Out, uint32 OutSize) const
+	{
+		check(OutSize >= StringBufferSize);
+		return ToStringInternal(Out, OutSize);
+	}
+
+	/**
+	 * Converts the FDebugName to a string buffer, avoiding dynamic allocations.
+	 * Returns the length of the string, excluding the null terminator.
+	 */
+	template <uint32 N>
 	uint32 ToString(TCHAR(&Out)[N]) const
 	{
-		return ToString(Out, N);
+		UE_STATIC_DEPRECATE(5.6, N < StringBufferSize,
+			"FDebugName::ToString(TCHAR (&Out)[N]) requires a buffer of size of at least FDebugName::StringBufferSize. "
+			"Use ToStringTruncate() if a smaller buffer is required and it is safe for the returned string to be truncated.");
+
+		return ToStringInternal(Out, N);
+	}
+
+	/**
+	 * Converts the FDebugName to a string buffer, avoiding dynamic allocations. Truncates the name if it does not fit in the specified output buffer.
+	 * Use a buffer size of at least FDebugName::StringBufferSize to avoid truncation.
+	 *
+	 * Returns the length of the (possibly truncated) string, excluding the null terminator.
+	 */
+	template <uint32 N>
+	uint32 ToStringTruncate(TCHAR(&Out)[N]) const
+	{
+		static_assert(N > 0, "Out buffer must have at least one element.");
+		return ToStringInternal(Out, N);
 	}
 
 	bool IsNone() const { return Name.IsNone() && Number == NAME_NO_NUMBER_INTERNAL; }
@@ -95,6 +129,8 @@ public:
 private:
 	FName Name;
 	uint32 Number;
+
+	RHI_API uint32 ToStringInternal(TCHAR* Out, uint32 OutSize) const;
 };
 
 enum class ERHIInterfaceType
diff --git a/Engine/Source/Runtime/RenderCore/Internal/ShaderCompilerDefinitions.h b/Engine/Source/Runtime/RenderCore/Internal/ShaderCompilerDefinitions.h
--- a/Engine/Source/Runtime/RenderCore/Internal/ShaderCompilerDefinitions.h
+++ b/Engine/Source/Runtime/RenderCore/Internal/ShaderCompilerDefinitions.h
@@ -450,7 +450,7 @@ private:
 			}
 			else
 			{
-				StringLength = Pair.Key.ToString(KeyBuffer, FName::StringBufferSize);
+				StringLength = Pair.Key.ToString(KeyBuffer);
 				for (uint32 CharIndex = 0; CharIndex < StringLength + 1; CharIndex++)
 				{
 					KeyBufferAnsi[CharIndex] = (ANSICHAR)KeyBuffer[CharIndex];
diff --git a/Engine/Source/Runtime/Renderer/Private/SparseVolumeTexture/SparseVolumeTextureStreamingVisualize.cpp b/Engine/Source/Runtime/Renderer/Private/SparseVolumeTexture/SparseVolumeTextureStreamingVisualize.cpp
--- a/Engine/Source/Runtime/Renderer/Private/SparseVolumeTexture/SparseVolumeTextureStreamingVisualize.cpp
+++ b/Engine/Source/Runtime/Renderer/Private/SparseVolumeTexture/SparseVolumeTextureStreamingVisualize.cpp
@@ -43,7 +43,7 @@ void UE::SVT::AddStreamingDebugPass(FRDGBuilder& GraphBuilder, const FSceneView&
 		{
 			const FStreamingDebugInfo::FSVT& SVT = DebugInfo->SVTs[SVTIndex];
 
-			StringBuilder.Appendf(TEXT("    Asset: '%s', Num Frames: %i\n"), SVT.AssetName, SVT.NumFrames);
+			StringBuilder.Appendf(TEXT("    Asset: '%s', Num Frames: %i\n"), *SVT.AssetName.ToString(), SVT.NumFrames);
 			FCanvasTextItem SVTTextItem(FVector2D(8.0, YOffset), FText::FromStringView(StringBuilder.ToView()), GEngine->GetSmallFont(), FLinearColor::White);
 			SVTTextItem.EnableShadow(FLinearColor::Black);
 			SVTTextItem.Draw(&Canvas);
