[AutoRTFM] Undo changes to reduce FTransactionallySafeMutex to a single byte

The single-byte implementation can corrupt memory if the mutex is destructed on a non-transactional thread.

We cannot safely revert back to the old TSharedPtr implementation as the single-byte implementation introduced the ability to asymmetrically lock and unlock the mutex, which may be depended on by the engine. TryLock() was also introduced which was not present on the old TSharedPtr implementation.
This is a new implementation (based heavily on the old TSharedPtr implementation) which attempts to keep as much as the single-byte implementation's functionality.
It is also more efficient than the old implementation, as it only registers one or two abort/commit handlers per nested transaction, instead of per lock / unlock.

Note: There are slight differences to the behavior of FTransactionallySafeMutex::IsLocked() when unlocking within a transaction. This is commentted as being unreliable on the method.

Added more tests, including one for destructing the mutex on a non-transactional thread.

diff --git a/Engine/Source/Programs/AutoRTFMTests/Private/UECoreTests.cpp b/Engine/Source/Programs/AutoRTFMTests/Private/UECoreTests.cpp
--- a/Engine/Source/Programs/AutoRTFMTests/Private/UECoreTests.cpp
+++ b/Engine/Source/Programs/AutoRTFMTests/Private/UECoreTests.cpp
@@ -27,11 +27,11 @@
 #include "Misc/PackageName.h"
 #include "StructUtils/PropertyBag.h"
 #include "Utils.h"
+#include <thread>
+#include <mutex>
 
 DEFINE_LOG_CATEGORY_STATIC(LogAutoRTFM_UECoreTests, Log, All)
 
-#include <thread>
-
 TEST_CASE("UECore.FDelegateHandle")
 {
 	FDelegateHandle Handle;
@@ -1855,7 +1855,7 @@ TEST_CASE("UECore.TransactionallySafeMutex")
 			Mutex.Lock();
 			REQUIRE(Mutex.IsLocked());
 			Mutex.Unlock();
-			REQUIRE(!Mutex.IsLocked());
+			REQUIRE(Mutex.IsLocked());
 		});
 
 		REQUIRE(!Mutex.IsLocked());
@@ -1982,31 +1982,142 @@ TEST_CASE("UECore.TransactionallySafeMutex")
 
 	SECTION("Delete Heap Allocated Mutex")
 	{
-		UE::FTransactionallySafeMutex* const Mutex = new UE::FTransactionallySafeMutex();
+		SECTION("SingleThread") // Mutex owned and destructed by this thread
+		{
+			UE::FTransactionallySafeMutex* const Mutex = new UE::FTransactionallySafeMutex();
 
-		AutoRTFM::Testing::Abort([&]
+			AutoRTFM::Testing::Abort([&]
+			{
+				Mutex->Lock();
+				delete Mutex;
+				AutoRTFM::AbortTransaction();
+			});
+
+			REQUIRE(!Mutex->IsLocked());
+
+			AutoRTFM::Testing::Commit([&]
+			{
+				Mutex->Lock();
+				delete Mutex;
+			});
+		}
+		
+		SECTION("MultiThread") // Mutex owned and destructed by another thread
 		{
-			Mutex->Lock();
-			delete Mutex;
-			AutoRTFM::AbortTransaction();
-		});
+			// This test does not support retries due to coordination with another thread.
+			AutoRTFMTestUtils::FScopedRetry NoRetry(AutoRTFM::ForTheRuntime::EAutoRTFMRetryTransactionState::NoRetry);
 
-		REQUIRE(!Mutex->IsLocked());
+			enum EState
+			{
+				InitializeThread,
+				MutexReady,
+				MutexUsed,
+				MutexDeleted,
+			};
+
+			struct FEvent
+			{
+				UE_AUTORTFM_ALWAYS_OPEN
+				void Signal(EState State)
+				{
+					{
+						std::unique_lock Lock(Mutex);
+						CurrentState = State;
+					}
+					CV.notify_one();
+				}
+
+				UE_AUTORTFM_ALWAYS_OPEN
+				void Wait(EState State)
+				{
+					std::unique_lock Lock(Mutex);
+					CV.wait(Lock, [this, State] { return CurrentState == State; });
+				}
+			private:
+				EState CurrentState = InitializeThread;
+				std::mutex Mutex;
+				std::condition_variable CV;
+			};
+			
+			UE::FTransactionallySafeMutex* Mutex = nullptr;
+			FEvent Event;
+
+			std::thread Thread([&Mutex, &Event]
+			{
+				Mutex = new UE::FTransactionallySafeMutex();
+				Event.Signal(MutexReady);
+				Event.Wait(MutexUsed);
+				delete Mutex;
+				Event.Signal(MutexDeleted);
+			});
+
+			Event.Wait(MutexReady);
+
+			SECTION("Lock, Abort, Destroy")
+			{
+				AutoRTFM::Testing::Abort([&]
+				{
+					Mutex->Lock();
+					AutoRTFM::AbortTransaction();
+				});
+
+				REQUIRE(!Mutex->IsLocked());
+
+				Event.Signal(MutexUsed);
+				Event.Wait(MutexDeleted);
+			}
+
+			SECTION("Lock, Unlock, Destroy, Abort")
+			{
+				AutoRTFM::Testing::Abort([&]
+				{
+					Mutex->Lock();
+					Mutex->Unlock();
+
+					Event.Signal(MutexUsed);
+					Event.Wait(MutexDeleted);
+
+					AutoRTFM::AbortTransaction();
+				});
+			}
+
+			SECTION("Lock, Unlock, Destroy, Commit")
+			{
+				AutoRTFM::Testing::Commit([&]
+				{
+					Mutex->Lock();
+					Mutex->Unlock();
+
+					Event.Signal(MutexUsed);
+					Event.Wait(MutexDeleted);
+				});
+			}
+
+			Thread.join();
+		}
+	}
+
+	SECTION("Lock Within, Unlock Outside")
+	{
+		UE::FTransactionallySafeMutex Mutex;
 
 		AutoRTFM::Testing::Commit([&]
 		{
-			Mutex->Lock();
-			delete Mutex;
+			Mutex.Lock();
 		});
+
+		REQUIRE(Mutex.IsLocked());
+		Mutex.Unlock();
+		REQUIRE(!Mutex.IsLocked());
 	}
 
-	SECTION("Lock Within, Unlock Outside")
+	SECTION("TryLock Within, Unlock Outside")
 	{
 		UE::FTransactionallySafeMutex Mutex;
 
 		AutoRTFM::Testing::Commit([&]
 		{
-			Mutex.Lock();
+			REQUIRE(Mutex.TryLock());
 		});
 
 		REQUIRE(Mutex.IsLocked());
@@ -2031,14 +2142,126 @@ TEST_CASE("UECore.TransactionallySafeMutex")
 	{
 		UE::FTransactionallySafeMutex Mutex;
 
-		AutoRTFM::Testing::Commit([&]
+		SECTION("Commit")
 		{
-			Mutex.Lock();
+			AutoRTFM::Testing::Commit([&]
+			{
+				Mutex.Lock();
+				Mutex.Unlock();
+				Mutex.Lock();
+			});
+
+			REQUIRE(Mutex.IsLocked());
 			Mutex.Unlock();
-			Mutex.Lock();
-		});
+		}
 
-		REQUIRE(Mutex.IsLocked());
+		SECTION("Abort")
+		{
+			AutoRTFM::Testing::Abort([&]
+			{
+				Mutex.Lock();
+				Mutex.Unlock();
+				Mutex.Lock();
+				AutoRTFM::AbortTransaction();
+			});
+
+			REQUIRE(!Mutex.IsLocked());
+		}
+	}
+
+	SECTION("TryLock, Unlock, TryLock")
+	{
+		UE::FTransactionallySafeMutex Mutex;
+
+		SECTION("Commit")
+		{
+			AutoRTFM::Testing::Commit([&]
+			{
+				REQUIRE(Mutex.TryLock());
+				Mutex.Unlock();
+				REQUIRE(Mutex.TryLock());
+			});
+
+			REQUIRE(Mutex.IsLocked());
+			Mutex.Unlock();
+		}
+		
+		SECTION("Abort")
+		{
+			AutoRTFM::Testing::Abort([&]
+			{
+				REQUIRE(Mutex.TryLock());
+				Mutex.Unlock();
+				REQUIRE(Mutex.TryLock());
+				AutoRTFM::AbortTransaction();
+			});
+
+			REQUIRE(!Mutex.IsLocked());
+		}
+	}
+
+	SECTION("Unlock, Lock, Unlock")
+	{
+		UE::FTransactionallySafeMutex Mutex;
+		Mutex.Lock();
+
+		SECTION("Commit")
+		{
+			AutoRTFM::Testing::Commit([&]
+			{
+				Mutex.Unlock();
+				Mutex.Lock();
+				Mutex.Unlock();
+			});
+
+			REQUIRE(!Mutex.IsLocked());
+		}
+
+		SECTION("Commit")
+		{
+			AutoRTFM::Testing::Abort([&]
+			{
+				Mutex.Unlock();
+				Mutex.Lock();
+				Mutex.Unlock();
+				AutoRTFM::AbortTransaction();
+			});
+
+			REQUIRE(Mutex.IsLocked());
+			Mutex.Unlock();
+		}
+	}
+
+	SECTION("Unlock, TryLock, Unlock")
+	{
+		UE::FTransactionallySafeMutex Mutex;
+		Mutex.Lock();
+
+		SECTION("Commit")
+		{
+			AutoRTFM::Testing::Commit([&]
+			{
+				Mutex.Unlock();
+				REQUIRE(Mutex.TryLock());
+				Mutex.Unlock();
+			});
+
+			REQUIRE(!Mutex.IsLocked());
+		}
+
+		SECTION("Commit")
+		{
+			AutoRTFM::Testing::Abort([&]
+			{
+				Mutex.Unlock();
+				REQUIRE(Mutex.TryLock());
+				Mutex.Unlock();
+				AutoRTFM::AbortTransaction();
+			});
+
+			REQUIRE(Mutex.IsLocked());
+			Mutex.Unlock();
+		}
 	}
 
 	SECTION("Lock Outside, Unlock Within")
@@ -2081,7 +2304,7 @@ TEST_CASE("UECore.TransactionallySafeMutex")
 				Mutex.Unlock();
 			});
 
-			REQUIRE(!Mutex.IsLocked());
+			REQUIRE(Mutex.IsLocked());
 		});
 
 		REQUIRE(!Mutex.IsLocked());
@@ -2206,7 +2429,7 @@ TEST_CASE("UECore.TransactionallySafeMutex")
 			{
 				REQUIRE(Mutex.IsLocked());
 				Mutex.Unlock();
-				REQUIRE(!Mutex.IsLocked());
+				REQUIRE(Mutex.IsLocked());
 			});
 			
 			Mutex.Lock();
@@ -2215,6 +2438,7 @@ TEST_CASE("UECore.TransactionallySafeMutex")
 		AutoRTFM::Testing::Commit([&]
 		{
 			REQUIRE(Mutex.TryLock());
+			REQUIRE(!Mutex.TryLock());
 
 			AutoRTFM::OnCommit([&]
 			{
@@ -2230,7 +2454,7 @@ TEST_CASE("UECore.TransactionallySafeMutex")
 			{
 				REQUIRE(Mutex.IsLocked());
 				Mutex.Unlock();
-				REQUIRE(!Mutex.IsLocked());
+				REQUIRE(Mutex.IsLocked());
 			});
 			
 			REQUIRE(Mutex.TryLock());
@@ -2243,7 +2467,7 @@ TEST_CASE("UECore.TransactionallySafeMutex")
 		{
 			AutoRTFM::OnCommit([&]
 			{
-				REQUIRE(!Mutex.IsLocked());
+				REQUIRE(Mutex.IsLocked());
 				Mutex.Lock();
 				REQUIRE(Mutex.IsLocked());
 			});
diff --git a/Engine/Source/Runtime/Core/Public/Async/TransactionallySafeMutex.h b/Engine/Source/Runtime/Core/Public/Async/TransactionallySafeMutex.h
--- a/Engine/Source/Runtime/Core/Public/Async/TransactionallySafeMutex.h
+++ b/Engine/Source/Runtime/Core/Public/Async/TransactionallySafeMutex.h
@@ -2,281 +2,281 @@
 
 #pragma once
 
-#include "Async/ExternalMutex.h"
+#include "Async/Mutex.h"
 #include "AutoRTFM/AutoRTFM.h"
-#include <atomic>
+#include "Templates/SharedPointer.h"
 
 namespace UE
 {
 
+#if UE_AUTORTFM
+
 // A transactionally safe mutex that works in the following novel ways:
 // - In the open (non-transactional):
 //   - Take the lock like before. Simple!
 //   - Free the lock like before too.
 // - In the closed (transactional):
-//   - We flip the internal transactional state between two bits `IsLockedWithinTransactionFlag` and
-// 	   `IsUnlockedWithinTransactionFlag`. These let us keep track of what the last action taken with
-// 	   the mutex was within the transaction.
+//   - Unlock() decrements `TransactionalLockCount` and registers the on-commit handler (if not 
+//     already registered) to perform the deferred unlock to the underlying mutex, if not rebalanced
+//     in the transaction with a lock.
+//   - During locking we query `TransactionalLockCount`:
+//     * -1 means the transaction performed an Unlock() before the lock, so the underlying mutex 
+//       must have been locked prior to the start of the transaction. `TransactionalLockCount` is
+//       incremented to 0, and a lock is emulated (i.e. TryLock() will return `true` without
+//       touching the underlying mutex).
+//     * 0 means the underlying mutex will be locked if it hasn't been already (tracked by 
+//       `bTransactionallyLocked`). Once the underlying mutex is locked, it remains locked for the
+//       entire duration of the transaction. This is done as a transactional abort will likely undo
+//       the writes to memory that is guarded by the mutex, so the lock needs be held until the
+//       transaction is either fully committed or aborted. When the underlying mutex is locked an
+//       on-abort handler and on-commit handler are registered to unlock the underlying mutex if
+//       necessary.
 //
 // Thus with this approach we will hold this lock for the *entirety* of the transactional nest should
 // we take the lock during the transaction, thus preventing non-transactional code from seeing any
 // modifications we should make.
-class FTransactionallySafeMutex final
+struct FTransactionallySafeMutex final
 {
-public:
-	constexpr FTransactionallySafeMutex() = default;
+	// Always open because the constructor will create the underlying mutex.
+	UE_AUTORTFM_ALWAYS_OPEN
+	FTransactionallySafeMutex() : State(MakeShared<FState>())
+	{
+	}
 
 	// Construct in a locked state. Avoids an expensive compare-and-swap at creation time.
-	inline explicit FTransactionallySafeMutex(UE::FAcquireLock Tag)
+	UE_AUTORTFM_ALWAYS_OPEN
+	inline explicit FTransactionallySafeMutex(UE::FAcquireLock Tag) : State(MakeShared<FState>(Tag))
 	{
-		FExternalMutex(State, Tag);
 	}
 
-	inline ~FTransactionallySafeMutex()
+	void Lock()
 	{
-		if (AutoRTFM::IsOnCurrentTransactionStack(this))
+		if (AutoRTFM::IsTransactional() || AutoRTFM::IsCommittingOrAborting())
 		{
-			// We'll have either 0 (we've destructed a mutex but never locked it), or 1 of each
-			// on-abort and on-commit handlers.
-			AutoRTFM::PopAllOnAbortHandlers(this);
-			AutoRTFM::PopAllOnCommitHandlers(this);
-		}
-	}
+			ensure(State->TransactionalLockCount >= -1 && State->TransactionalLockCount <= 0);
 
-	UE_NONCOPYABLE(FTransactionallySafeMutex);
+			State->TransactionalLockCount++;
 
-	inline bool IsLocked() const
-	{
-		if (AutoRTFM::IsTransactional())
-		{
-			return AutoRTFM::Open([this]
+			if (State->bTransactionallyLocked)
 			{
-				if (FExternalMutex(State).IsLocked())
-				{
-					return !IsTransactionallyUnlocked();
-				}
+				// Underlying mutex has already been locked for the duration of the transaction.
+				return;
+			}
 
-				return false;
-			});
-		}
-		else if (HasLockTransactionallyHeldDuringCommittingOrAborting())
-		{
-			return !IsTransactionallyUnlocked();
+			if (State->TransactionalLockCount == 0)
+			{
+				// The mutex was locked before the transaction and unlocked inside the transaction.
+				return;
+			}
+
+			// First time the mutex has been locked during the transaction.
+			// Mark the transactionally locked flag and register on-commit and on-abort handlers.
+			State->MutexLock();
+			SetTransactionallyLockedAndRegisterHandlers();
 		}
 		else
 		{
-			return FExternalMutex(State).IsLocked();
+			// Non-transactional path. Call Lock() on the mutex directly.
+			State->MutexLock();
 		}
 	}
-	
+
 	inline bool TryLock()
 	{
-		if (AutoRTFM::IsTransactional())
+		if (AutoRTFM::IsTransactional() || AutoRTFM::IsCommittingOrAborting())
 		{
-			return AutoRTFM::Open([this]
+			ensure(State->TransactionalLockCount >= -1 && State->TransactionalLockCount <= 1);
+
+			if (State->TransactionalLockCount > 0)
 			{
-				if (!DoesTransactionHoldLock())
-				{
-					if (!FExternalMutex(State).TryLock())
-					{
-						return false;
-					}
-
-					// If we get here, it means we've taken the lock for the *first* time within the
-					// transaction. We need to thus register the on-abort and on-commit handlers to
-					// correctly undo the transactionality of our locking.
-					const AutoRTFM::EContextStatus Status = AutoRTFM::Close([this]
-					{
-						DoOnAbortHandlerInClosed();
-						DoOnCommitHandlerInClosed();
-					});
-
-					check(Status == AutoRTFM::EContextStatus::OnTrack);
-				}
+				return false; // Attempting double-lock within transaction.
+			}
 
-				RecordPreviousStateAndSetTransactionalLockMode(ETransactionalLockMode::TransactionallyLocked);
+			if (State->bTransactionallyLocked)
+			{
+				// Underlying mutex has already been locked for the duration of the transaction.
+				State->TransactionalLockCount++;
+				return true;
+			}
 
+			if (State->TransactionalLockCount < 0)
+			{
+				// The mutex was locked before the transaction and unlocked inside the transaction.
+				// As unlock is deferred, bump the lock counter and return true to emulate the lock.
+				State->TransactionalLockCount++;
 				return true;
-			});
-		}
-		else if (HasLockTransactionallyHeldDuringCommittingOrAborting())
-		{
-			SetTransactionalLockMode(ETransactionalLockMode::TransactionallyLocked);
+			}
+			
+			// Attempt to lock the underlying mutex.
+			if (!State->MutexTryLock())
+			{
+				return false;
+			}
+
+			// First time the mutex has been locked during the transaction.
+			// Increment the lock count, mark the transactionally locked flag and register 
+			// on-commit and on-abort handlers.
+			State->TransactionalLockCount++;
+			SetTransactionallyLockedAndRegisterHandlers();
 			return true;
 		}
 		else
 		{
-			const bool bResult = FExternalMutex(State).TryLock();
-			ensure(!bResult || !DoesTransactionHoldLock());
-			return bResult;
+			// Non-transactional path. Call TryLock() on the mutex directly.
+			return State->Mutex.TryLock();
 		}
 	}
 
-	inline void Lock()
+	void Unlock()
 	{
-		if (AutoRTFM::IsTransactional())
+		if (AutoRTFM::IsTransactional() || AutoRTFM::IsCommittingOrAborting())
 		{
-			AutoRTFM::Open([this]
-			{
-				if (!DoesTransactionHoldLock())
-				{
-					FExternalMutex(State).Lock();
-
-					const AutoRTFM::EContextStatus Status = AutoRTFM::Close([this]
-					{
-						DoOnAbortHandlerInClosed();
-						DoOnCommitHandlerInClosed();
-					});
-
-					check(Status == AutoRTFM::EContextStatus::OnTrack);
-				}
+			ensure(State->TransactionalLockCount >= 0 && State->TransactionalLockCount <= 1);
 
-				RecordPreviousStateAndSetTransactionalLockMode(ETransactionalLockMode::TransactionallyLocked);
-			});
-		}
-		else if (HasLockTransactionallyHeldDuringCommittingOrAborting())
-		{
-			SetTransactionalLockMode(ETransactionalLockMode::TransactionallyLocked);
+			State->TransactionalLockCount--;
+			// Use an on-commit handler to unlock the underlying mutex if not balanced with a lock
+			// before the transaction commits.
+			MaybeRegisterCommitHandler();
 		}
 		else
 		{
-			FExternalMutex(State).Lock();
-			ensure(!DoesTransactionHoldLock());
+			State->MutexUnlock();
 		}
 	}
 
-	inline void Unlock()
+	/** 
+	 * This method may give surprising results and should be used with caution!
+	 * 
+	 * - You cannot safely use this function to determine whether Lock() will block.
+	 *   You may be in race with another thread which is also trying to lock the mutex.
+	 * 
+	 * - If a Mutex is locked by the AutoRTFM thread, it will not fully release the mutex
+	 *   until the transaction fully succeeds or rolls back. Because of this, IsLocked may
+	 *   return true even after the mutex has been fully unlocked.
+	 */
+	bool IsLocked() const
 	{
-		if (AutoRTFM::IsTransactional())
-		{
-			AutoRTFM::Open([this]
-			{
-				if (!DoesTransactionHoldLock())
-				{
-					const AutoRTFM::EContextStatus Status = AutoRTFM::Close([this]
-					{
-						DoOnCommitHandlerInClosed();
-					});
-
-					check(Status == AutoRTFM::EContextStatus::OnTrack);
-				}
-
-				RecordPreviousStateAndSetTransactionalLockMode(ETransactionalLockMode::TransactionallyUnlocked);
-			});
-		}
-		else if (HasLockTransactionallyHeldDuringCommittingOrAborting())
-		{
-			SetTransactionalLockMode(ETransactionalLockMode::TransactionallyUnlocked);
-		}
-		else
-		{
-			ensure(!DoesTransactionHoldLock());
-			FExternalMutex(State).Unlock();
-		}
+		return State->MutexIsLocked();
 	}
 
 private:
-	// Need `mutable` because `FExternalMutex` only works non-const!
-	mutable std::atomic<uint8> State = 0;
-
-	// Since we are going to be writing non-atomically to `State` in the write-log, prior to us
-	// doing an atomic write to actually undo the lock, we need to be sure that the atomic will
-	// always be lock-free and take no bits beyond the integer backing type!
-	static_assert(decltype(State)::is_always_lock_free && (sizeof(decltype(State)) == sizeof(uint8)));
-	
-	// These are the flags that `FExternalMutex` already reserves for the lock!
-	static constexpr uint8 IsLockedFlag = 1 << 0;
-	static constexpr uint8 HasWaitingThreadsFlag = 1 << 1;
-
-	// Need 2 extra bits ontop of the 2 that `FExternalMutex` uses, to record what
-	// the last transactional action was within the transaction.
-	static constexpr uint8 IsLockedWithinTransactionFlag = 1 << 2;
-	static constexpr uint8 IsUnlockedWithinTransactionFlag = 1 << 3;
-
-	inline bool IsTransactionallyUnlocked() const
-	{
-		const uint8 LoadedState = State.load(std::memory_order_consume);
-		return (LoadedState & IsUnlockedWithinTransactionFlag);
-	}
+	UE_NONCOPYABLE(FTransactionallySafeMutex)
 
-	inline bool DoesTransactionHoldLock() const
+	// Called when the underlying mutex is locked for the first time within a transaction nest.
+	// Sets the bTransactionallyLocked flag, calls MaybeRegisterCommitHandler() and registers an
+	// on-abort handler to unlock the underlying mutex.
+	void SetTransactionallyLockedAndRegisterHandlers()
 	{
-		const uint8 LoadedState = State.load(std::memory_order_consume);
-		const bool bTransactionallyLocked = (LoadedState & (IsLockedWithinTransactionFlag | IsUnlockedWithinTransactionFlag));
-		ensure(!bTransactionallyLocked || (LoadedState & IsLockedFlag));
-		return bTransactionallyLocked;
-	}
+		ensure(!State->bTransactionallyLocked);
+		State->bTransactionallyLocked = true;
 
-	inline void DoOnAbortHandlerInClosed()
-	{
-		check(AutoRTFM::IsClosed());
+		MaybeRegisterCommitHandler();
 
-		AutoRTFM::PushOnAbortHandler(this, [this]
+		// Capture State instead of 'this' as State is a TSharedPtr which can outlive this 
+		// FTransactionallySafeMutex.
+		AutoRTFM::OnAbort([State = this->State]
 		{
-			State.fetch_and(~(IsLockedWithinTransactionFlag | IsUnlockedWithinTransactionFlag), std::memory_order_release);
-			FExternalMutex(State).Unlock();
+			State->MutexUnlock();
+			State->ResetTransactionState();
 		});
 	}
-
-	inline void DoOnCommitHandlerInClosed()
+	
+	// Registers an on-commit handler (if it hasn't been already registered for this transaction
+	// nest) to unlock the underlying mutex if the transaction unlocked more times than locked.
+	void MaybeRegisterCommitHandler()
 	{
-		check(AutoRTFM::IsClosed());
-
-		AutoRTFM::PushOnCommitHandler(this, [this]
+		if (!State->bRegisteredCommitHandler) 
 		{
-			const bool bIsTransactionallyUnlocked = IsTransactionallyUnlocked();
-
-			State.fetch_and(~(IsLockedWithinTransactionFlag | IsUnlockedWithinTransactionFlag), std::memory_order_release);
+			State->bRegisteredCommitHandler = true;
 
-			if (bIsTransactionallyUnlocked)
+			// Capture State instead of 'this' as State is a TSharedPtr which can outlive this 
+			// FTransactionallySafeMutex.
+			AutoRTFM::OnCommit([State = this->State]
 			{
-				FExternalMutex(State).Unlock();
-			}
-		});
+				// If the transaction's lock count is negative (more unlocks than locks), or the 
+				// underlying mutex was locked and then re-balanced to an unlocked state, unlock
+				// the underlying mutex.
+				if (State->TransactionalLockCount < 0 ||
+					(State->bTransactionallyLocked && State->TransactionalLockCount == 0))
+				{
+					State->MutexUnlock();
+				}
+				State->ResetTransactionState();
+			});
+		}
 	}
 
-	enum class ETransactionalLockMode : bool
+	struct FState final
 	{
-		TransactionallyLocked = false,
-		TransactionallyUnlocked = true,
-	};
+		UE_AUTORTFM_NOAUTORTFM
+		FState() = default;
 
-	inline void SetTransactionalLockMode(ETransactionalLockMode LockMode)
-	{
-		if (LockMode == ETransactionalLockMode::TransactionallyLocked)
+		UE_AUTORTFM_NOAUTORTFM
+		explicit FState(UE::FAcquireLock Tag) : Mutex(Tag) {}
+
+		UE_AUTORTFM_NOAUTORTFM
+		~FState()
 		{
-			State |= IsLockedWithinTransactionFlag;
-			State &= ~IsUnlockedWithinTransactionFlag;
+			ensure(TransactionalLockCount == 0);
+			ensure(bTransactionallyLocked == false);
+			ensure(bRegisteredCommitHandler == false);
 		}
-		else
+
+		UE_AUTORTFM_NOAUTORTFM
+		void ResetTransactionState()
 		{
-			State |= IsUnlockedWithinTransactionFlag;
-			State &= ~IsLockedWithinTransactionFlag;
+			TransactionalLockCount = 0;
+			bTransactionallyLocked = false;
+			bRegisteredCommitHandler = false;
 		}
-	}
 
-	inline void RecordPreviousStateAndSetTransactionalLockMode(ETransactionalLockMode LockMode)
-	{
-		// Record the state into the write log. If an abort happens we will non-atomically
-		// overwrite `State`, removing the transactional flags. This will be followed-up
-		// with an atomic write in the on-abort handler to actually do the unlock though,
-		// thus ensuring data integrity.
-		AutoRTFM::RecordOpenWriteNoMemoryValidation(&State);
+		UE_AUTORTFM_ALWAYS_OPEN
+		void MutexLock()
+		{
+			Mutex.Lock();
+		}
+		
+		UE_AUTORTFM_ALWAYS_OPEN
+		bool MutexTryLock()
+		{
+			return Mutex.TryLock();
+		}
 
-		SetTransactionalLockMode(LockMode);
-	}
+		UE_AUTORTFM_ALWAYS_OPEN
+		void MutexUnlock()
+		{
+			Mutex.Unlock();
+		}
 
-	inline bool HasLockTransactionallyHeldDuringCommittingOrAborting() const
-	{
-		// When we are running an on-commit / on-abort, we *could* not have run
-		// the on-commit / on-abort handlers for the mutex yet, and so if we get
-		// here we've got a mutex that is still participating in a transactionally
-		// held lock state, and thus we need to access the data as if we are on
-		// the transactional codepath.
-		return AutoRTFM::IsCommittingOrAborting() && DoesTransactionHoldLock();
-	}
+		UE_AUTORTFM_ALWAYS_OPEN
+		bool MutexIsLocked()
+		{
+			return Mutex.IsLocked();
+		}
+
+		// The underlying mutex.
+		UE::FMutex Mutex;
+		// Counter for calls to unlock / lock within a transaction:
+		// -1: mutex has had one more unlock than lock. (e.g. transaction starts locked)
+		//  0: mutex has had the same number of locks as unlocks.
+		// +1: mutex has had one more lock than unlock.
+		// Other values indicate a double-lock or double-unlock within a transaction.
+		int8 TransactionalLockCount = 0;
+		// True if the underlying mutex has been locked for the duration of the transaction.
+		bool bTransactionallyLocked : 1 = false;
+		// True if the on-commit handler has been registered for this transaction nest.
+		bool bRegisteredCommitHandler : 1 = false;
+	};
+
+	TSharedPtr<FState> State;
 };
 
+#else
+using FTransactionallySafeMutex = FMutex;
+#endif // UE_AUTORTFM
+
 } // UE
 
+// Some existing code references FTransactionallySafeMutex without the UE
+// namespace prefix. TODO: properly qualify these and remove this.
 using FTransactionallySafeMutex = UE::FTransactionallySafeMutex;
