Added FSharedMutex and FSharedRecursiveMutex

These are alternatives to FRWLock.

FSharedMutex is a 4-byte unfair shared mutex built on top of ParkingLot.

FSharedRecursiveMutex is an 8-byte unfair shared recursive mutex built on top of ParkingLot. Its shared mode requires externally-owned state to avoid dynamic allocations.

diff --git a/Engine/Source/Runtime/Core/Private/Async/SharedMutex.cpp b/Engine/Source/Runtime/Core/Private/Async/SharedMutex.cpp
new file mode 100644
--- /dev/null
+++ b/Engine/Source/Runtime/Core/Private/Async/SharedMutex.cpp
@@ -0,0 +1,159 @@
+// Copyright Epic Games, Inc. All Rights Reserved.
+
+#include "Async/SharedMutex.h"
+
+#include "Async/Fundamental/Scheduler.h"
+#include "Async/ParkingLot.h"
+#include "HAL/PlatformProcess.h"
+
+namespace UE
+{
+
+inline const void* FSharedMutex::GetLockAddress() const
+{
+	return &State;
+}
+
+inline const void* FSharedMutex::GetSharedLockAddress() const
+{
+	// Shared locks need a distinct address from exclusive locks to allow threads waiting for exclusive ownership
+	// to be woken up without waking any threads waiting for shared ownership.
+	return (const uint8*)&State + 1;
+}
+
+void FSharedMutex::LockSlow()
+{
+	constexpr int32 SpinLimit = 40;
+	int32 SpinCount = 0;
+	for (uint32 CurrentState = State.load(std::memory_order_relaxed);;)
+	{
+		// Try to acquire the lock if it was unlocked, even if there are waiting threads.
+		// Acquiring the lock despite the waiting threads means that this lock is not FIFO and thus not fair.
+		if (!(CurrentState & (IsLockedFlag | SharedLockCountMask)))
+		{
+			if (State.compare_exchange_weak(CurrentState, CurrentState | IsLockedFlag, std::memory_order_acquire, std::memory_order_relaxed))
+			{
+				return;
+			}
+			continue;
+		}
+
+		// Spin up to the spin limit while there are no waiting threads.
+		if (!(CurrentState & HasWaitingLockFlag) && SpinCount < SpinLimit)
+		{
+			FPlatformProcess::YieldThread();
+			++SpinCount;
+			CurrentState = State.load(std::memory_order_relaxed);
+			continue;
+		}
+
+		// Store that there are waiting threads. Restart if the state has changed since it was loaded.
+		if (!(CurrentState & HasWaitingLockFlag))
+		{
+			if (!State.compare_exchange_weak(CurrentState, CurrentState | HasWaitingLockFlag, std::memory_order_relaxed))
+			{
+				continue;
+			}
+			CurrentState |= HasWaitingLockFlag;
+		}
+
+		// Do not enter oversubscription during a wait on a mutex since the wait is generally too short
+		// for it to matter and it can worsen performance a lot for heavily contended locks.
+		LowLevelTasks::Private::FOversubscriptionAllowedScope _(false);
+
+		// Wait if the state has not changed. Either way, loop back and try to acquire the lock trying to wait.
+		ParkingLot::Wait(GetLockAddress(), [this, CurrentState] { return State.load(std::memory_order_relaxed) == CurrentState; }, []{});
+		CurrentState = State.load(std::memory_order_relaxed);
+	}
+}
+
+void FSharedMutex::UnlockSlow()
+{
+	// Unlock immediately to allow other threads to acquire the lock while this thread looks for a thread to wake.
+	uint32 CurrentState = State.fetch_and(~IsLockedFlag, std::memory_order_release);
+	checkSlow(CurrentState & IsLockedFlag);
+	CurrentState &= ~IsLockedFlag;
+
+	if (CurrentState & HasWaitingLockFlag)
+	{
+		// Wake one thread that is waiting to acquire an exclusive lock.
+		ParkingLot::WakeOne(GetLockAddress(), [this](ParkingLot::FWakeState WakeState) -> uint64
+		{
+			if (!WakeState.bHasWaitingThreads)
+			{
+				State.fetch_and(~HasWaitingLockFlag, std::memory_order_relaxed);
+			}
+			return 0;
+		});
+	}
+	else if (CurrentState & HasWaitingSharedLockFlag)
+	{
+		// Wake every thread that is waiting to acquire a shared lock.
+		// The awoken threads might race against other exclusive locks.
+		if (State.fetch_and(~HasWaitingSharedLockFlag, std::memory_order_relaxed) & HasWaitingSharedLockFlag)
+		{
+			ParkingLot::WakeAll(GetSharedLockAddress());
+		}
+	}
+}
+
+void FSharedMutex::LockSharedSlow()
+{
+	constexpr int32 SpinLimit = 40;
+	int32 SpinCount = 0;
+	for (uint32 CurrentState = State.load(std::memory_order_relaxed);;)
+	{
+		// Try to acquire the lock if it is unlocked and there are no waiting threads.
+		if (!(CurrentState & (IsLockedFlag | HasWaitingLockFlag)))
+		{
+			if (State.compare_exchange_weak(CurrentState, CurrentState + (1 << SharedLockCountShift), std::memory_order_acquire, std::memory_order_relaxed))
+			{
+				return;
+			}
+			continue;
+		}
+
+		// Spin up to the spin limit while there are no waiting threads.
+		if (!(CurrentState & HasWaitingLockFlag) && SpinCount < SpinLimit)
+		{
+			FPlatformProcess::YieldThread();
+			++SpinCount;
+			CurrentState = State.load(std::memory_order_relaxed);
+			continue;
+		}
+
+		// Store that there are waiting threads. Restart if the state has changed since it was loaded.
+		if (!(CurrentState & HasWaitingSharedLockFlag))
+		{
+			if (!State.compare_exchange_weak(CurrentState, CurrentState | HasWaitingSharedLockFlag, std::memory_order_relaxed))
+			{
+				continue;
+			}
+			CurrentState |= HasWaitingSharedLockFlag;
+		}
+
+		// Do not enter oversubscription during a wait on a mutex since the wait is generally too short
+		// for it to matter and it can worsen performance a lot for heavily contended locks.
+		LowLevelTasks::Private::FOversubscriptionAllowedScope _(false);
+
+		// Wait if the state has not changed. Either way, loop back and try to acquire the lock trying to wait.
+		ParkingLot::Wait(GetSharedLockAddress(), [this, CurrentState] { return State.load(std::memory_order_relaxed) == CurrentState; }, []{});
+		CurrentState = State.load(std::memory_order_relaxed);
+	}
+}
+
+void FSharedMutex::UnlockSharedSlow()
+{
+	// This path is only taken when the last shared lock is released and there is a waiting exclusive lock.
+	// Wake one thread that is waiting to acquire an exclusive lock.
+	ParkingLot::WakeOne(GetLockAddress(), [this](ParkingLot::FWakeState WakeState) -> uint64
+	{
+		if (!WakeState.bHasWaitingThreads)
+		{
+			State.fetch_and(~HasWaitingLockFlag, std::memory_order_relaxed);
+		}
+		return 0;
+	});
+}
+
+} // UE
diff --git a/Engine/Source/Runtime/Core/Private/Async/SharedRecursiveMutex.cpp b/Engine/Source/Runtime/Core/Private/Async/SharedRecursiveMutex.cpp
new file mode 100644
--- /dev/null
+++ b/Engine/Source/Runtime/Core/Private/Async/SharedRecursiveMutex.cpp
@@ -0,0 +1,345 @@
+// Copyright Epic Games, Inc. All Rights Reserved.
+
+#include "Async/SharedRecursiveMutex.h"
+
+#include "Async/Fundamental/Scheduler.h"
+#include "Async/ParkingLot.h"
+#include "HAL/PlatformProcess.h"
+#include "HAL/PlatformTLS.h"
+
+namespace UE::Core::Private
+{
+
+struct FSharedRecursiveMutexStack
+{
+	constexpr FSharedRecursiveMutexStack() = default;
+
+	~FSharedRecursiveMutexStack()
+	{
+		checkf(!Top, TEXT("Thread %u destroyed while holding a shared lock on the FSharedRecursiveMutex at 0x%p."),
+			FPlatformTLS::GetCurrentThreadId(), Top->OwnedMutex);
+	}
+
+	FSharedRecursiveMutexLink* Top = nullptr;
+};
+
+static thread_local FSharedRecursiveMutexStack ThreadLocalSharedLocks;
+
+bool FSharedRecursiveMutexLink::Owns(const FSharedRecursiveMutex* Mutex)
+{
+	for (FSharedRecursiveMutexLink* It = ThreadLocalSharedLocks.Top; It; It = It->Next)
+	{
+		if (It->OwnedMutex == Mutex)
+		{
+			return true;
+		}
+	}
+	return false;
+}
+
+void FSharedRecursiveMutexLink::Push(const FSharedRecursiveMutex* Mutex)
+{
+	checkSlow(!OwnedMutex && !Next);
+	OwnedMutex = Mutex;
+	Next = ThreadLocalSharedLocks.Top;
+	ThreadLocalSharedLocks.Top = this;
+}
+
+void FSharedRecursiveMutexLink::Pop()
+{
+	checkSlow(OwnedMutex);
+	for (FSharedRecursiveMutexLink** Link = &ThreadLocalSharedLocks.Top; *Link; Link = &(*Link)->Next)
+	{
+		if ((*Link) == this)
+		{
+			*Link = Next;
+			OwnedMutex = nullptr;
+			Next = nullptr;
+			return;
+		}
+	}
+}
+
+} // UE::Core::Private
+
+namespace UE
+{
+
+inline const void* FSharedRecursiveMutex::GetLockAddress() const
+{
+	return &State;
+}
+
+inline const void* FSharedRecursiveMutex::GetSharedLockAddress() const
+{
+	// Shared locks need a distinct address from exclusive locks to allow threads waiting for exclusive ownership
+	// to be woken up without waking any threads waiting for shared ownership.
+	return (const uint8*)&State + 1;
+}
+
+bool FSharedRecursiveMutex::TryLock()
+{
+	const uint32 CurrentThreadId = FPlatformTLS::GetCurrentThreadId();
+	uint32 CurrentState = State.load(std::memory_order_relaxed);
+
+	// Try to acquire the lock if it was unlocked, even if there are waiting threads.
+	// Acquiring the lock despite the waiting threads means that this lock is not FIFO and thus not fair.
+	if (LIKELY(!(CurrentState & (LockCountMask | SharedLockCountMask))))
+	{
+		if (LIKELY(State.compare_exchange_weak(CurrentState, CurrentState | (1 << LockCountShift), std::memory_order_acquire, std::memory_order_relaxed)))
+		{
+			checkSlow(ThreadId.load(std::memory_order_relaxed) == 0 && (CurrentState & LockCountMask) == 0);
+			ThreadId.store(CurrentThreadId, std::memory_order_relaxed);
+			return true;
+		}
+	}
+
+	// Lock recursively if this is the thread that holds the lock.
+	if (ThreadId.load(std::memory_order_relaxed) == CurrentThreadId)
+	{
+		checkSlow((CurrentState & LockCountMask) != LockCountMask);
+		State.fetch_add(1 << LockCountShift, std::memory_order_relaxed);
+		return true;
+	}
+
+	return false;
+}
+
+void FSharedRecursiveMutex::Lock()
+{
+	const uint32 CurrentThreadId = FPlatformTLS::GetCurrentThreadId();
+	uint32 CurrentState = State.load(std::memory_order_relaxed);
+
+	// Try to acquire the lock if it was unlocked, even if there are waiting threads.
+	// Acquiring the lock despite the waiting threads means that this lock is not FIFO and thus not fair.
+	if (LIKELY(!(CurrentState & (LockCountMask | SharedLockCountMask))))
+	{
+		if (LIKELY(State.compare_exchange_weak(CurrentState, CurrentState | (1 << LockCountShift), std::memory_order_acquire, std::memory_order_relaxed)))
+		{
+			checkSlow(ThreadId.load(std::memory_order_relaxed) == 0 && (CurrentState & LockCountMask) == 0);
+			ThreadId.store(CurrentThreadId, std::memory_order_relaxed);
+			return;
+		}
+	}
+
+	// Lock recursively if this is the thread that holds the lock.
+	if (ThreadId.load(std::memory_order_relaxed) == CurrentThreadId)
+	{
+		checkSlow((CurrentState & LockCountMask) != LockCountMask);
+		State.fetch_add(1 << LockCountShift, std::memory_order_relaxed);
+		return;
+	}
+
+	LockSlow(CurrentState, CurrentThreadId);
+}
+
+FORCENOINLINE void FSharedRecursiveMutex::LockSlow(uint32 CurrentState, uint32 CurrentThreadId)
+{
+	constexpr int32 SpinLimit = 40;
+	int32 SpinCount = 0;
+	for (;;)
+	{
+		// Try to acquire the lock if it was unlocked, even if there are waiting threads.
+		// Acquiring the lock despite the waiting threads means that this lock is not FIFO and thus not fair.
+		if (!(CurrentState & (LockCountMask | SharedLockCountMask)))
+		{
+			if (State.compare_exchange_weak(CurrentState, CurrentState | (1 << LockCountShift), std::memory_order_acquire, std::memory_order_relaxed))
+			{
+				checkSlow(ThreadId.load(std::memory_order_relaxed) == 0 && (CurrentState & LockCountMask) == 0);
+				ThreadId.store(CurrentThreadId, std::memory_order_relaxed);
+				return;
+			}
+			continue;
+		}
+
+		// Spin up to the spin limit while there are no waiting threads.
+		if (!(CurrentState & HasWaitingLockFlag) && SpinCount < SpinLimit)
+		{
+			FPlatformProcess::YieldThread();
+			++SpinCount;
+			CurrentState = State.load(std::memory_order_relaxed);
+			continue;
+		}
+
+		// Store that there are waiting threads. Restart if the state has changed since it was loaded.
+		if (!(CurrentState & HasWaitingLockFlag))
+		{
+			if (!State.compare_exchange_weak(CurrentState, CurrentState | HasWaitingLockFlag, std::memory_order_relaxed))
+			{
+				continue;
+			}
+			CurrentState |= HasWaitingLockFlag;
+		}
+
+		// Do not enter oversubscription during a wait on a mutex since the wait is generally too short
+		// for it to matter and it can worsen performance a lot for heavily contended locks.
+		LowLevelTasks::Private::FOversubscriptionAllowedScope _(false);
+
+		// Wait if the state has not changed. Either way, loop back and try to acquire the lock trying to wait.
+		ParkingLot::Wait(GetLockAddress(), [this, CurrentState] { return State.load(std::memory_order_relaxed) == CurrentState; }, []{});
+		CurrentState = State.load(std::memory_order_relaxed);
+	}
+}
+
+void FSharedRecursiveMutex::Unlock()
+{
+	uint32 CurrentState = State.load(std::memory_order_relaxed);
+	checkSlow(CurrentState & LockCountMask);
+	checkSlow(ThreadId.load(std::memory_order_relaxed) == FPlatformTLS::GetCurrentThreadId());
+
+	if (LIKELY((CurrentState & LockCountMask) == (1 << LockCountShift)))
+	{
+		// Remove the association with this thread before unlocking.
+		ThreadId.store(0, std::memory_order_relaxed);
+
+		// Unlock immediately to allow other threads to acquire the lock while this thread looks for a thread to wake.
+		State.fetch_sub(1 << LockCountShift, std::memory_order_release);
+
+		// Wake one exclusive waiter or every shared waiter.
+		if (UNLIKELY(CurrentState & (HasWaitingLockFlag | HasWaitingSharedLockFlag)))
+		{
+			WakeWaitingThreads(CurrentState);
+		}
+	}
+	else
+	{
+		State.fetch_sub(1 << LockCountShift, std::memory_order_relaxed);
+	}
+}
+
+bool FSharedRecursiveMutex::TryLockShared(Core::Private::FSharedRecursiveMutexLink& Link)
+{
+	uint32 CurrentState = State.load(std::memory_order_relaxed);
+	// Recursive shared locks are quick to acquire. Check for that case and add 1 to the count.
+	if ((CurrentState & SharedLockCountMask) && (!(CurrentState & HasWaitingLockFlag) || Core::Private::FSharedRecursiveMutexLink::Owns(this)))
+	{
+		[[maybe_unused]] uint32 LastState = State.fetch_add(1 << SharedLockCountShift, std::memory_order_relaxed);
+		checkSlow((LastState & SharedLockCountMask) != SharedLockCountMask);
+		Link.Push(this);
+		return true;
+	}
+	// Try to acquire a shared lock if there is no active or waiting exclusive lock.
+	else if (!(CurrentState & (LockCountMask | HasWaitingLockFlag)) &&
+		State.compare_exchange_strong(CurrentState, CurrentState + (1 << SharedLockCountShift), std::memory_order_acquire, std::memory_order_relaxed))
+	{
+		Link.Push(this);
+		return true;
+	}
+	return false;
+}
+
+void FSharedRecursiveMutex::LockShared(Core::Private::FSharedRecursiveMutexLink& Link)
+{
+	uint32 CurrentState = State.load(std::memory_order_relaxed);
+	// Recursive shared locks are quick to acquire. Check for that case and add 1 to the count.
+	if ((CurrentState & SharedLockCountMask) && (!(CurrentState & HasWaitingLockFlag) || Core::Private::FSharedRecursiveMutexLink::Owns(this)))
+	{
+		[[maybe_unused]] uint32 LastState = State.fetch_add(1 << SharedLockCountShift, std::memory_order_relaxed);
+		checkSlow((LastState & SharedLockCountMask) != SharedLockCountMask);
+		Link.Push(this);
+		return;
+	}
+	// Try to acquire a shared lock if there is no active or waiting exclusive lock.
+	else if (!(CurrentState & (LockCountMask | HasWaitingLockFlag)) &&
+		State.compare_exchange_weak(CurrentState, CurrentState + (1 << SharedLockCountShift), std::memory_order_acquire, std::memory_order_relaxed))
+	{
+		Link.Push(this);
+		return;
+	}
+	LockSharedSlow(Link);
+}
+
+FORCENOINLINE void FSharedRecursiveMutex::LockSharedSlow(Core::Private::FSharedRecursiveMutexLink& Link)
+{
+	constexpr int32 SpinLimit = 40;
+	int32 SpinCount = 0;
+	for (uint32 CurrentState = State.load(std::memory_order_relaxed);;)
+	{
+		// Try to acquire the lock if it is unlocked and there are no waiting threads.
+		if (!(CurrentState & (LockCountMask | HasWaitingLockFlag)))
+		{
+			checkSlow((CurrentState & SharedLockCountMask) != SharedLockCountMask);
+			if (State.compare_exchange_weak(CurrentState, CurrentState + (1 << SharedLockCountShift), std::memory_order_acquire, std::memory_order_relaxed))
+			{
+				Link.Push(this);
+				return;
+			}
+			continue;
+		}
+
+		// Spin up to the spin limit while there are no waiting threads.
+		if (!(CurrentState & HasWaitingLockFlag) && SpinCount < SpinLimit)
+		{
+			FPlatformProcess::YieldThread();
+			++SpinCount;
+			CurrentState = State.load(std::memory_order_relaxed);
+			continue;
+		}
+
+		// Store that there are waiting threads. Restart if the state has changed since it was loaded.
+		if (!(CurrentState & HasWaitingSharedLockFlag))
+		{
+			if (!State.compare_exchange_weak(CurrentState, CurrentState | HasWaitingSharedLockFlag, std::memory_order_relaxed))
+			{
+				continue;
+			}
+			CurrentState |= HasWaitingSharedLockFlag;
+		}
+
+		// Do not enter oversubscription during a wait on a mutex since the wait is generally too short
+		// for it to matter and it can worsen performance a lot for heavily contended locks.
+		LowLevelTasks::Private::FOversubscriptionAllowedScope _(false);
+
+		// Wait if the state has not changed. Either way, loop back and try to acquire the lock trying to wait.
+		ParkingLot::Wait(GetSharedLockAddress(), [this, CurrentState] { return State.load(std::memory_order_relaxed) == CurrentState; }, []{});
+		CurrentState = State.load(std::memory_order_relaxed);
+	}
+}
+
+void FSharedRecursiveMutex::UnlockShared(Core::Private::FSharedRecursiveMutexLink& Link)
+{
+	Link.Pop();
+	const uint32 LastState = State.fetch_sub(1 << SharedLockCountShift, std::memory_order_release);
+	checkSlow(LastState & SharedLockCountMask);
+	constexpr uint32 WakeState = HasWaitingLockFlag | (1 << SharedLockCountShift);
+	if (UNLIKELY((LastState & ~HasWaitingSharedLockFlag) == WakeState))
+	{
+		// The last shared lock was released and there is a waiting exclusive lock.
+		// Wake one thread that is waiting to acquire an exclusive lock.
+		ParkingLot::WakeOne(GetLockAddress(), [this](ParkingLot::FWakeState WakeState) -> uint64
+		{
+			if (!WakeState.bHasWaitingThreads)
+			{
+				State.fetch_and(~HasWaitingLockFlag, std::memory_order_relaxed);
+			}
+			return 0;
+		});
+	}
+}
+
+FORCENOINLINE void FSharedRecursiveMutex::WakeWaitingThreads(uint32 CurrentState)
+{
+	if (CurrentState & HasWaitingLockFlag)
+	{
+		// Wake one thread that is waiting to acquire an exclusive lock.
+		ParkingLot::WakeOne(GetLockAddress(), [this](ParkingLot::FWakeState WakeState) -> uint64
+			{
+				if (!WakeState.bHasWaitingThreads)
+				{
+					State.fetch_and(~HasWaitingLockFlag, std::memory_order_relaxed);
+				}
+				return 0;
+			});
+	}
+	else if (CurrentState & HasWaitingSharedLockFlag)
+	{
+		// Wake every thread that is waiting to acquire a shared lock.
+		// The awoken threads might race against other exclusive locks.
+		if (State.fetch_and(~HasWaitingSharedLockFlag, std::memory_order_relaxed) & HasWaitingSharedLockFlag)
+		{
+			ParkingLot::WakeAll(GetSharedLockAddress());
+		}
+	}
+}
+
+} // UE
diff --git a/Engine/Source/Runtime/Core/Public/Async/SharedLock.h b/Engine/Source/Runtime/Core/Public/Async/SharedLock.h
new file mode 100644
--- /dev/null
+++ b/Engine/Source/Runtime/Core/Public/Async/SharedLock.h
@@ -0,0 +1,147 @@
+// Copyright Epic Games, Inc. All Rights Reserved.
+
+#pragma once
+
+#include "Async/LockTags.h"
+#include "Misc/AssertionMacros.h"
+
+namespace UE
+{
+
+// Note: FSharedRecursiveMutex has specialized these lock types.
+
+/**
+ * A basic shared mutex ownership wrapper that locks on construction and unlocks on destruction.
+ *
+ * LockType must contain LockShared() and UnlockShared() functions.
+ * 
+ * Use with mutex types like FSharedMutex and FSharedRecursiveMutex.
+ */
+template <typename LockType>
+class TSharedLock final
+{
+public:
+	TSharedLock(const TSharedLock&) = delete;
+	TSharedLock& operator=(const TSharedLock&) = delete;
+
+	[[nodiscard]] inline explicit TSharedLock(LockType& Lock)
+		: Mutex(Lock)
+	{
+		Mutex.LockShared();
+	}
+
+	inline ~TSharedLock()
+	{
+		Mutex.UnlockShared();
+	}
+
+private:
+	LockType& Mutex;
+};
+
+/**
+ * A shared mutex ownership wrapper that allows dynamic locking, unlocking, and deferred locking.
+ *
+ * LockType must contain LockShared() and UnlockShared() functions.
+ * 
+ * Use with mutex types like FSharedMutex and FSharedRecursiveMutex.
+ */
+template <typename LockType>
+class TDynamicSharedLock final
+{
+public:
+	TDynamicSharedLock() = default;
+
+	TDynamicSharedLock(const TDynamicSharedLock&) = delete;
+	TDynamicSharedLock& operator=(const TDynamicSharedLock&) = delete;
+
+	/** Wrap a mutex and lock it in shared mode. */
+	[[nodiscard]] inline explicit TDynamicSharedLock(LockType& Lock)
+		: Mutex(&Lock)
+	{
+		Mutex->LockShared();
+		bLocked = true;
+	}
+
+	/** Wrap a mutex without locking it in shared mode. */
+	[[nodiscard]] inline explicit TDynamicSharedLock(LockType& Lock, FDeferLock)
+		: Mutex(&Lock)
+	{
+	}
+
+	/** Move from another lock, transferring any ownership to this lock. */
+	[[nodiscard]] inline TDynamicSharedLock(TDynamicSharedLock&& Other)
+		: Mutex(Other.Mutex)
+		, bLocked(Other.bLocked)
+	{
+		Other.Mutex = nullptr;
+		Other.bLocked = false;
+	}
+
+	/** Move from another lock, transferring any ownership to this lock, and unlocking the previous mutex if locked. */
+	inline TDynamicSharedLock& operator=(TDynamicSharedLock&& Other)
+	{
+		if (bLocked)
+		{
+			Mutex->UnlockShared();
+		}
+		Mutex = Other.Mutex;
+		bLocked = Other.bLocked;
+		Other.Mutex = nullptr;
+		Other.bLocked = false;
+		return *this;
+	}
+
+	/** Unlock the mutex if locked. */
+	inline ~TDynamicSharedLock()
+	{
+		if (bLocked)
+		{
+			Mutex->UnlockShared();
+		}
+	}
+
+	/** Try to lock the associated mutex in shared mode. This lock must have a mutex and must not be locked. */
+	bool TryLock()
+	{
+		check(!bLocked);
+		check(Mutex);
+		bLocked = Mutex->TryLockShared();
+		return bLocked;
+	}
+
+	/** Lock the associated mutex in shared mode. This lock must have a mutex and must not be locked. */
+	void Lock()
+	{
+		check(!bLocked);
+		check(Mutex);
+		Mutex->LockShared();
+		bLocked = true;
+	}
+
+	/** Unlock the associated mutex in shared mode. This lock must have a mutex and must be locked. */
+	void Unlock()
+	{
+		check(bLocked);
+		bLocked = false;
+		Mutex->UnlockShared();
+	}
+
+	/** Returns true if this lock has its associated mutex locked. */
+	inline bool OwnsLock() const
+	{
+		return bLocked;
+	}
+
+	/** Returns true if this lock has its associated mutex locked. */
+	inline explicit operator bool() const
+	{
+		return OwnsLock();
+	}
+
+private:
+	LockType* Mutex = nullptr;
+	bool bLocked = false;
+};
+
+} // UE
diff --git a/Engine/Source/Runtime/Core/Public/Async/SharedMutex.h b/Engine/Source/Runtime/Core/Public/Async/SharedMutex.h
new file mode 100644
--- /dev/null
+++ b/Engine/Source/Runtime/Core/Public/Async/SharedMutex.h
@@ -0,0 +1,109 @@
+// Copyright Epic Games, Inc. All Rights Reserved.
+
+#pragma once
+
+#include "CoreTypes.h"
+#include "Misc/AssertionMacros.h"
+#include <atomic>
+
+#define UE_API CORE_API
+
+namespace UE
+{
+
+/**
+ * A four-byte shared mutex that is not fair and does not support recursive locking.
+ *
+ * Prefer FMutex when shared locking is not required.
+ * All new shared locks will wait when any thread is waiting to take an exclusive lock.
+ * An exclusive lock and a shared lock may not be simultaneously held by the same thread.
+ */
+class FSharedMutex final
+{
+public:
+	constexpr FSharedMutex() = default;
+
+	FSharedMutex(const FSharedMutex&) = delete;
+	FSharedMutex& operator=(const FSharedMutex&) = delete;
+
+	inline bool TryLock()
+	{
+		uint32 Expected = State.load(std::memory_order_relaxed);
+		return !(Expected & (IsLockedFlag | SharedLockCountMask)) &&
+			State.compare_exchange_strong(Expected, Expected | IsLockedFlag,
+				std::memory_order_acquire, std::memory_order_relaxed);
+	}
+
+	inline void Lock()
+	{
+		uint32 Expected = 0;
+		if (LIKELY(State.compare_exchange_weak(Expected, IsLockedFlag, std::memory_order_acquire, std::memory_order_relaxed)))
+		{
+			return;
+		}
+		LockSlow();
+	}
+
+	inline void Unlock()
+	{
+		uint32 Expected = IsLockedFlag;
+		if (LIKELY(State.compare_exchange_weak(Expected, 0, std::memory_order_release, std::memory_order_relaxed)))
+		{
+			return;
+		}
+		UnlockSlow();
+	}
+
+	inline bool TryLockShared()
+	{
+		uint32 Expected = State.load(std::memory_order_relaxed);
+		return !(Expected & (IsLockedFlag | HasWaitingLockFlag)) &&
+			State.compare_exchange_strong(Expected, Expected + (1 << SharedLockCountShift),
+				std::memory_order_acquire, std::memory_order_relaxed);
+	}
+
+	inline void LockShared()
+	{
+		uint32 Expected = State.load(std::memory_order_relaxed);
+		if (LIKELY(!(Expected & (IsLockedFlag | HasWaitingLockFlag)) &&
+			State.compare_exchange_weak(Expected, Expected + (1 << SharedLockCountShift),
+				std::memory_order_acquire, std::memory_order_relaxed)))
+		{
+			return;
+		}
+		LockSharedSlow();
+	}
+
+	inline void UnlockShared()
+	{
+		const uint32 LastState = State.fetch_sub(1 << SharedLockCountShift, std::memory_order_release);
+		checkSlow(LastState & SharedLockCountMask);
+		constexpr uint32 WakeState = HasWaitingLockFlag | (1 << SharedLockCountShift);
+		if (LIKELY((LastState & ~HasWaitingSharedLockFlag) != WakeState))
+		{
+			return;
+		}
+		UnlockSharedSlow();
+	}
+
+private:
+	UE_API void LockSlow();
+	UE_API void UnlockSlow();
+	UE_API void LockSharedSlow();
+	UE_API void UnlockSharedSlow();
+
+	const void* GetLockAddress() const;
+	const void* GetSharedLockAddress() const;
+
+	static constexpr uint32 IsLockedFlag = 1 << 0;
+	static constexpr uint32 HasWaitingLockFlag = 1 << 1;
+	static constexpr uint32 HasWaitingSharedLockFlag = 1 << 2;
+	static constexpr uint32 SharedLockCountShift = 3;
+	static constexpr uint32 SharedLockCountMask = 0xffff'fff8;
+
+	std::atomic<uint32> State = 0;
+};
+
+} // UE
+
+#undef UE_API
diff --git a/Engine/Source/Runtime/Core/Public/Async/SharedRecursiveMutex.h b/Engine/Source/Runtime/Core/Public/Async/SharedRecursiveMutex.h
new file mode 100644
--- /dev/null
+++ b/Engine/Source/Runtime/Core/Public/Async/SharedRecursiveMutex.h
@@ -0,0 +1,199 @@
+// Copyright Epic Games, Inc. All Rights Reserved.
+
+#pragma once
+
+#include "CoreTypes.h"
+#include "Async/SharedLock.h"
+#include <atomic>
+
+#define UE_API CORE_API
+
+namespace UE { class FSharedRecursiveMutex; }
+
+namespace UE::Core::Private
+{
+
+struct FSharedRecursiveMutexLink
+{
+	static bool Owns(const FSharedRecursiveMutex* Mutex);
+	void Push(const FSharedRecursiveMutex* Mutex);
+	void Pop();
+
+	const FSharedRecursiveMutex* OwnedMutex = nullptr;
+	FSharedRecursiveMutexLink* Next = nullptr;
+};
+
+} // UE::Core::Private
+
+namespace UE
+{
+
+/**
+ * An eight-byte shared mutex that is not fair and supports recursive locking.
+ *
+ * Prefer FRecursiveMutex when shared locking is not required.
+ * Prefer FSharedMutex when recursive locking is not required.
+ * All non-recursive shared locks will wait when any thread is waiting to take an exclusive lock.
+ * An exclusive lock and a shared lock may not be simultaneously held by the same thread.
+ */
+class FSharedRecursiveMutex final
+{
+public:
+	constexpr FSharedRecursiveMutex() = default;
+
+	FSharedRecursiveMutex(const FSharedRecursiveMutex&) = delete;
+	FSharedRecursiveMutex& operator=(const FSharedRecursiveMutex&) = delete;
+
+	UE_API bool TryLock();
+	UE_API void Lock();
+	UE_API void Unlock();
+
+	// Use TSharedLock or TDynamicSharedLock to acquire a shared lock.
+	UE_API bool TryLockShared(Core::Private::FSharedRecursiveMutexLink& Link);
+	UE_API void LockShared(Core::Private::FSharedRecursiveMutexLink& Link);
+	UE_API void UnlockShared(Core::Private::FSharedRecursiveMutexLink& Link);
+
+private:
+	void LockSlow(uint32 CurrentState, uint32 CurrentThreadId);
+	void LockSharedSlow(Core::Private::FSharedRecursiveMutexLink& Link);
+	void WakeWaitingThreads(uint32 CurrentState);
+
+	const void* GetLockAddress() const;
+	const void* GetSharedLockAddress() const;
+
+	static constexpr uint32 HasWaitingLockFlag = 1 << 0;
+	static constexpr uint32 HasWaitingSharedLockFlag = 1 << 1;
+	static constexpr uint32 LockCountShift = 2;
+	static constexpr uint32 LockCountMask = 0x0000'0ffc;
+	static constexpr uint32 SharedLockCountShift = 12;
+	static constexpr uint32 SharedLockCountMask = 0xffff'f000;
+
+	std::atomic<uint32> State = 0;
+	std::atomic<uint32> ThreadId = 0;
+};
+
+template <>
+class TSharedLock<FSharedRecursiveMutex> final
+{
+public:
+	TSharedLock(const TSharedLock&) = delete;
+	TSharedLock& operator=(const TSharedLock&) = delete;
+
+	[[nodiscard]] inline explicit TSharedLock(FSharedRecursiveMutex& Lock)
+		: Mutex(Lock)
+	{
+		Mutex.LockShared(Link);
+	}
+
+	inline ~TSharedLock()
+	{
+		Mutex.UnlockShared(Link);
+	}
+
+private:
+	FSharedRecursiveMutex& Mutex;
+	Core::Private::FSharedRecursiveMutexLink Link;
+};
+
+template <>
+class TDynamicSharedLock<FSharedRecursiveMutex> final
+{
+public:
+	TDynamicSharedLock() = default;
+
+	TDynamicSharedLock(const TDynamicSharedLock&) = delete;
+	TDynamicSharedLock& operator=(const TDynamicSharedLock&) = delete;
+
+	[[nodiscard]] inline explicit TDynamicSharedLock(FSharedRecursiveMutex& Lock)
+		: Mutex(&Lock)
+	{
+		Mutex->LockShared(Link);
+		bLocked = true;
+	}
+
+	[[nodiscard]] inline explicit TDynamicSharedLock(FSharedRecursiveMutex& Lock, FDeferLock)
+		: Mutex(&Lock)
+	{
+	}
+
+	[[nodiscard]] inline TDynamicSharedLock(TDynamicSharedLock&& Other)
+		: Mutex(Other.Mutex)
+		, bLocked(Other.bLocked)
+	{
+		if (bLocked)
+		{
+			Mutex->LockShared(Link);
+			Mutex->UnlockShared(Other.Link);
+		}
+		Other.Mutex = nullptr;
+		Other.bLocked = false;
+	}
+
+	inline TDynamicSharedLock& operator=(TDynamicSharedLock&& Other)
+	{
+		if (bLocked)
+		{
+			Mutex->UnlockShared(Link);
+		}
+		Mutex = Other.Mutex;
+		bLocked = Other.bLocked;
+		if (bLocked)
+		{
+			Mutex->LockShared(Link);
+			Mutex->UnlockShared(Other.Link);
+		}
+		Other.Mutex = nullptr;
+		Other.bLocked = false;
+		return *this;
+	}
+
+	inline ~TDynamicSharedLock()
+	{
+		if (bLocked)
+		{
+			Mutex->UnlockShared(Link);
+		}
+	}
+
+	bool TryLock()
+	{
+		check(!bLocked);
+		check(Mutex);
+		bLocked = Mutex->TryLockShared(Link);
+		return bLocked;
+	}
+
+	void Lock()
+	{
+		check(!bLocked);
+		check(Mutex);
+		Mutex->LockShared(Link);
+		bLocked = true;
+	}
+
+	void Unlock()
+	{
+		check(bLocked);
+		bLocked = false;
+		Mutex->UnlockShared(Link);
+	}
+
+	inline bool OwnsLock() const
+	{
+		return bLocked;
+	}
+
+	inline explicit operator bool() const
+	{
+		return OwnsLock();
+	}
+
+private:
+	FSharedRecursiveMutex* Mutex = nullptr;
+	Core::Private::FSharedRecursiveMutexLink Link;
+	bool bLocked = false;
+};
+
+} // UE
+
+#undef UE_API
diff --git a/Engine/Source/Runtime/Core/Tests/Async/SharedMutexTest.cpp b/Engine/Source/Runtime/Core/Tests/Async/SharedMutexTest.cpp
new file mode 100644
--- /dev/null
+++ b/Engine/Source/Runtime/Core/Tests/Async/SharedMutexTest.cpp
@@ -0,0 +1,282 @@
+// Copyright Epic Games, Inc. All Rights Reserved.
+
+#if WITH_TESTS
+
+#include "Async/SharedMutex.h"
+
+#include "Async/ManualResetEvent.h"
+#include "Async/SharedLock.h"
+#include "CoreGlobals.h"
+#include "HAL/PlatformProcess.h"
+#include "HAL/Thread.h"
+#include "Tests/TestHarnessAdapter.h"
+
+namespace UE
+{
+
+TEST_CASE_NAMED(FSharedMutexTest, "Core::Async::SharedMutex", "[Core][Async]")
+{
+	SECTION("SingleThread")
+	{
+		FSharedMutex Mutex;
+
+		Mutex.Lock();
+		CHECK_FALSE(TDynamicSharedLock(Mutex, DeferLock).TryLock());
+		Mutex.Unlock();
+
+		CHECK(Mutex.TryLock());
+		Mutex.Unlock();
+
+		// This test performs recursive locking because it exercises the current implementation
+		// but it is not technically supported by the mutex and can deadlock if used generally.
+		{
+			TSharedLock Lock1(Mutex);
+			TSharedLock Lock2(Mutex);
+			TSharedLock Lock3(Mutex);
+			CHECK_FALSE(Mutex.TryLock());
+		}
+
+		{
+			TDynamicSharedLock Lock1(Mutex, DeferLock);
+			TDynamicSharedLock Lock2(Mutex, DeferLock);
+			TDynamicSharedLock Lock3(Mutex, DeferLock);
+			CHECK(Lock1.TryLock());
+			CHECK(Lock2.TryLock());
+			CHECK(Lock3.TryLock());
+			CHECK_FALSE(Mutex.TryLock());
+		}
+
+		CHECK(Mutex.TryLock());
+		Mutex.Unlock();
+	}
+
+	SECTION("MultipleThreads")
+	{
+		FSharedMutex Mutex;
+		uint32 Counter = 0;
+		FManualResetEvent Events[4];
+		const auto MakeWait = [&Events](int32 Index)
+		{
+			return [&Events, Index]
+			{
+				Events[Index].Wait();
+				Events[Index].Reset();
+			};
+		};
+		const auto Wake = [&Events](int32 Index)
+		{
+			Events[Index].Notify();
+			FPlatformProcess::YieldThread();
+		};
+
+		class FCountdownEvent
+		{
+		public:
+			void Reset(int32 Count)
+			{
+				Counter.store(Count, std::memory_order_relaxed);
+				Event.Reset();
+			}
+
+			void Notify()
+			{
+				if (Counter.fetch_sub(1, std::memory_order_release) == 1)
+				{
+					Event.Notify();
+				}
+			}
+
+			void Wait()
+			{
+				Event.Wait();
+			}
+
+		private:
+			std::atomic<int32> Counter = 0;
+			FManualResetEvent Event;
+		};
+		FCountdownEvent CountdownEvent;
+
+		FThread Thread0(TEXT("SharedMutexTest0"), [&Mutex, &Counter, &CountdownEvent, &Wake, Wait = MakeWait(0)]
+		{
+			// Test 1: Exclusive w/ one waiting exclusive lock.
+			Mutex.Lock();
+			Wake(1);
+			Wait();
+			Counter = 1;
+			Mutex.Unlock();
+
+			// Test 2: Exclusive w/ one waiting shared lock.
+			Wait();
+			Wake(2);
+			Mutex.LockShared();
+			CHECK(Counter == 2);
+
+			// Test 3: Shared w/ one waiting exclusive lock.
+			Wake(1);
+			Wait();
+			Counter = 3;
+			Mutex.UnlockShared();
+
+			// Test 4: Exclusive w/ three waiting shared locks.
+			Wait();
+			Wake(1);
+			Mutex.LockShared();
+			CHECK(Counter == 4);
+			Wait();
+			Mutex.UnlockShared();
+
+			// Test 5: Shared w/ no exclusive contention.
+			CountdownEvent.Reset(3);
+			Wake(1);
+			Wake(2);
+			Wake(3);
+			for (int32 I = 0; I < 16384; ++I)
+			{
+				Mutex.LockShared();
+				FPlatformProcess::YieldThread();
+				Mutex.UnlockShared();
+			}
+			CountdownEvent.Wait();
+
+			if (GIsBuildMachine)
+			{
+				return;
+			}
+
+			// Test 6: Shared w/ one waiting exclusive lock and one waiting shared lock.
+			Mutex.LockShared();
+			Counter = 5;
+			Wake(1);
+			Wait();
+			Mutex.UnlockShared();
+		});
+
+		FThread Thread1(TEXT("SharedMutexTest1"), [&Mutex, &Counter, &CountdownEvent, &Wake, Wait = MakeWait(1)]
+		{
+			// Test 1: Exclusive w/ one waiting exclusive lock.
+			Wait();
+			Wake(2);
+			Mutex.Lock();
+			CHECK(Counter == 1);
+
+			// Test 2: Exclusive w/ one waiting shared lock.
+			Wake(0);
+			Wait();
+			Counter = 2;
+			Mutex.Unlock();
+
+			// Test 3: Shared w/ one waiting exclusive lock.
+			Wait();
+			Wake(2);
+			Mutex.Lock();
+			CHECK(Counter == 3);
+
+			// Test 4: Exclusive w/ three waiting shared locks.
+			Wake(2);
+			Wait();
+			Counter = 4;
+			Mutex.Unlock();
+
+			// Test 5: Shared w/ no exclusive contention.
+			Wait();
+			for (int32 I = 0; I < 16384; ++I)
+			{
+				Mutex.LockShared();
+				FPlatformProcess::YieldThread();
+				Mutex.UnlockShared();
+			}
+			CountdownEvent.Notify();
+
+			if (GIsBuildMachine)
+			{
+				return;
+			}
+
+			// Test 6: Shared w/ one waiting exclusive lock and one waiting shared lock.
+			Wait();
+			Wake(2);
+			Mutex.Lock();
+			CHECK(Counter == 5);
+			Counter = 6;
+			Mutex.Unlock();
+		});
+
+		FThread Thread2(TEXT("SharedMutexTest2"), [&Mutex, &Counter, &CountdownEvent, &Wake, Wait = MakeWait(2)]
+		{
+			// Test 1: Exclusive w/ one waiting exclusive lock.
+			Wait();
+			Wake(0);
+
+			// Test 2: Exclusive w/ one waiting shared lock.
+			Wait();
+			Wake(1);
+
+			// Test 3: Shared w/ one waiting exclusive lock.
+			Wait();
+			Wake(0);
+
+			// Test 4: Exclusive w/ three waiting shared locks.
+			Wait();
+			Wake(3);
+			Mutex.LockShared();
+			CHECK(Counter == 4);
+			Wake(3);
+			Mutex.UnlockShared();
+
+			// Test 5: Shared w/ no exclusive contention.
+			Wait();
+			for (int32 I = 0; I < 16384; ++I)
+			{
+				Mutex.LockShared();
+				FPlatformProcess::YieldThread();
+				Mutex.UnlockShared();
+			}
+			CountdownEvent.Notify();
+
+			if (GIsBuildMachine)
+			{
+				return;
+			}
+
+			// Test 6: Shared w/ one waiting exclusive lock and one waiting shared lock.
+			Wait();
+			FPlatformProcess::Sleep(0.001f); // Hopefully reliable enough to ensure the exclusive lock is waiting.
+			Wake(0);
+			Mutex.LockShared();
+			CHECK(Counter == 6);
+			Mutex.UnlockShared();
+		});
+
+		FThread Thread3(TEXT("SharedMutexTest3"), [&Mutex, &Counter, &CountdownEvent, &Wake, Wait = MakeWait(3)]
+		{
+			// Test 4: Exclusive w/ three waiting shared locks.
+			Wait();
+			Wake(0);
+			Mutex.LockShared();
+			CHECK(Counter == 4);
+			Wait();
+			Wake(0);
+			Mutex.UnlockShared();
+
+			// Test 5: Shared w/ no exclusive contention.
+			Wait();
+			for (int32 I = 0; I < 16384; ++I)
+			{
+				Mutex.LockShared();
+				FPlatformProcess::YieldThread();
+				Mutex.UnlockShared();
+			}
+			CountdownEvent.Notify();
+		});
+
+		Thread0.Join();
+		Thread1.Join();
+		Thread2.Join();
+		Thread3.Join();
+	}
+}
+
+} // UE
+
+#endif // WITH_TESTS
diff --git a/Engine/Source/Runtime/Core/Tests/Async/SharedRecursiveMutexTest.cpp b/Engine/Source/Runtime/Core/Tests/Async/SharedRecursiveMutexTest.cpp
new file mode 100644
--- /dev/null
+++ b/Engine/Source/Runtime/Core/Tests/Async/SharedRecursiveMutexTest.cpp
@@ -0,0 +1,294 @@
+// Copyright Epic Games, Inc. All Rights Reserved.
+
+#if WITH_TESTS
+
+#include "Async/SharedRecursiveMutex.h"
+
+#include "Async/ManualResetEvent.h"
+#include "Async/SharedLock.h"
+#include "CoreGlobals.h"
+#include "HAL/PlatformProcess.h"
+#include "HAL/Thread.h"
+#include "Tests/TestHarnessAdapter.h"
+
+namespace UE
+{
+
+TEST_CASE_NAMED(FSharedRecursiveMutexTest, "Core::Async::SharedRecursiveMutex", "[Core][Async]")
+{
+	SECTION("SingleThread")
+	{
+		FSharedRecursiveMutex Mutex;
+
+		Mutex.Lock();
+		CHECK_FALSE(TDynamicSharedLock(Mutex, DeferLock).TryLock());
+		Mutex.Unlock();
+
+		CHECK(Mutex.TryLock());
+		Mutex.Unlock();
+
+		{
+			TSharedLock Lock1(Mutex);
+			TSharedLock Lock2(Mutex);
+			TSharedLock Lock3(Mutex);
+			CHECK_FALSE(Mutex.TryLock());
+		}
+
+		{
+			TDynamicSharedLock Lock1(Mutex, DeferLock);
+			TDynamicSharedLock Lock2(Mutex, DeferLock);
+			TDynamicSharedLock Lock3(Mutex, DeferLock);
+			CHECK(Lock1.TryLock());
+			CHECK(Lock2.TryLock());
+			CHECK(Lock3.TryLock());
+			CHECK_FALSE(Mutex.TryLock());
+		}
+
+		CHECK(Mutex.TryLock());
+		CHECK(Mutex.TryLock());
+		Mutex.Unlock();
+		Mutex.Unlock();
+
+		Mutex.Lock();
+		Mutex.Lock();
+		Mutex.Unlock();
+		Mutex.Unlock();
+	}
+
+	SECTION("MultipleThreads")
+	{
+		FSharedRecursiveMutex Mutex;
+		uint32 Counter = 0;
+		FManualResetEvent Events[4];
+		const auto MakeWait = [&Events](int32 Index)
+		{
+			return [&Events, Index]
+			{
+				Events[Index].Wait();
+				Events[Index].Reset();
+			};
+		};
+		const auto Wake = [&Events](int32 Index)
+		{
+			Events[Index].Notify();
+			FPlatformProcess::YieldThread();
+		};
+
+		class FCountdownEvent
+		{
+		public:
+			void Reset(int32 Count)
+			{
+				Counter.store(Count, std::memory_order_relaxed);
+				Event.Reset();
+			}
+
+			void Notify()
+			{
+				if (Counter.fetch_sub(1, std::memory_order_release) == 1)
+				{
+					Event.Notify();
+				}
+			}
+
+			void Wait()
+			{
+				Event.Wait();
+			}
+
+		private:
+			std::atomic<int32> Counter = 0;
+			FManualResetEvent Event;
+		};
+		FCountdownEvent CountdownEvent;
+
+		FThread Thread0(TEXT("SharedMutexTest0"), [&Mutex, &Counter, &CountdownEvent, &Wake, Wait = MakeWait(0)]
+		{
+			TDynamicSharedLock SharedLock1(Mutex, DeferLock);
+			TDynamicSharedLock SharedLock2(Mutex, DeferLock);
+
+			// Test 1: Exclusive w/ one waiting exclusive lock.
+			Mutex.Lock();
+			Wake(1);
+			Wait();
+			Counter = 1;
+			Mutex.Unlock();
+
+			// Test 2: Exclusive w/ one waiting shared lock.
+			Wait();
+			Wake(2);
+			SharedLock1.Lock();
+			CHECK(Counter == 2);
+
+			// Test 3: Shared w/ one waiting exclusive lock.
+			Wake(1);
+			Wait();
+			Counter = 3;
+			CHECK(SharedLock2.TryLock());
+			SharedLock2.Unlock();
+			SharedLock1.Unlock();
+
+			// Test 4: Exclusive w/ three waiting shared locks.
+			Wait();
+			Wake(1);
+			SharedLock1.Lock();
+			CHECK(Counter == 4);
+			Wait();
+			SharedLock1.Unlock();
+
+			// Test 5: Shared w/ no exclusive contention.
+			CountdownEvent.Reset(3);
+			Wake(1);
+			Wake(2);
+			Wake(3);
+			for (int32 I = 0; I < 16384; ++I)
+			{
+				TSharedLock Lock(Mutex);
+				FPlatformProcess::YieldThread();
+			}
+			CountdownEvent.Wait();
+
+			if (GIsBuildMachine)
+			{
+				return;
+			}
+
+			// Test 6: Shared w/ one waiting exclusive lock and one waiting shared lock.
+			SharedLock1.Lock();
+			Counter = 5;
+			Wake(1);
+			Wait();
+			SharedLock2.Lock();
+			SharedLock1.Unlock();
+			SharedLock2.Unlock();
+		});
+
+		FThread Thread1(TEXT("SharedMutexTest1"), [&Mutex, &Counter, &CountdownEvent, &Wake, Wait = MakeWait(1)]
+		{
+			// Test 1: Exclusive w/ one waiting exclusive lock.
+			Wait();
+			Wake(2);
+			Mutex.Lock();
+			CHECK(Counter == 1);
+
+			// Test 2: Exclusive w/ one waiting shared lock.
+			Wake(0);
+			Wait();
+			Counter = 2;
+			Mutex.Unlock();
+
+			// Test 3: Shared w/ one waiting exclusive lock.
+			Wait();
+			Wake(2);
+			Mutex.Lock();
+			CHECK(Counter == 3);
+
+			// Test 4: Exclusive w/ three waiting shared locks.
+			Wake(2);
+			Wait();
+			Counter = 4;
+			Mutex.Unlock();
+
+			// Test 5: Shared w/ no exclusive contention.
+			Wait();
+			for (int32 I = 0; I < 16384; ++I)
+			{
+				TSharedLock Lock(Mutex);
+				FPlatformProcess::YieldThread();
+			}
+			CountdownEvent.Notify();
+
+			if (GIsBuildMachine)
+			{
+				return;
+			}
+
+			// Test 6: Shared w/ one waiting exclusive lock and one waiting shared lock.
+			Wait();
+			Wake(2);
+			Mutex.Lock();
+			CHECK(Counter == 5);
+			Counter = 6;
+			Mutex.Unlock();
+		});
+
+		FThread Thread2(TEXT("SharedMutexTest2"), [&Mutex, &Counter, &CountdownEvent, &Wake, Wait = MakeWait(2)]
+		{
+			TDynamicSharedLock SharedLock(Mutex, DeferLock);
+
+			// Test 1: Exclusive w/ one waiting exclusive lock.
+			Wait();
+			Wake(0);
+
+			// Test 2: Exclusive w/ one waiting shared lock.
+			Wait();
+			Wake(1);
+
+			// Test 3: Shared w/ one waiting exclusive lock.
+			Wait();
+			Wake(0);
+
+			// Test 4: Exclusive w/ three waiting shared locks.
+			Wait();
+			Wake(3);
+			SharedLock.Lock();
+			CHECK(Counter == 4);
+			Wake(3);
+			SharedLock.Unlock();
+
+			// Test 5: Shared w/ no exclusive contention.
+			Wait();
+			for (int32 I = 0; I < 16384; ++I)
+			{
+				TSharedLock Lock(Mutex);
+				FPlatformProcess::YieldThread();
+			}
+			CountdownEvent.Notify();
+
+			if (GIsBuildMachine)
+			{
+				return;
+			}
+
+			// Test 6: Shared w/ one waiting exclusive lock and one waiting shared lock.
+			Wait();
+			FPlatformProcess::Sleep(0.001f); // Hopefully reliable enough to ensure the exclusive lock is waiting.
+			Wake(0);
+			SharedLock.Lock();
+			CHECK(Counter == 6);
+			SharedLock.Unlock();
+		});
+
+		FThread Thread3(TEXT("SharedMutexTest3"), [&Mutex, &Counter, &CountdownEvent, &Wake, Wait = MakeWait(3)]
+		{
+			TDynamicSharedLock SharedLock(Mutex, DeferLock);
+
+			// Test 4: Exclusive w/ three waiting shared locks.
+			Wait();
+			Wake(0);
+			SharedLock.Lock();
+			CHECK(Counter == 4);
+			Wait();
+			Wake(0);
+			SharedLock.Unlock();
+
+			// Test 5: Shared w/ no exclusive contention.
+			Wait();
+			for (int32 I = 0; I < 16384; ++I)
+			{
+				TSharedLock Lock(Mutex);
+				FPlatformProcess::YieldThread();
+			}
+			CountdownEvent.Notify();
+		});
+
+		Thread0.Join();
+		Thread1.Join();
+		Thread2.Join();
+		Thread3.Join();
+	}
+}
+
+} // UE
+
+#endif // WITH_TESTS
