Fix an assert in TDeque due when initializing from an initializer list with inline capacity.

Also fix initial capacity not being correctly initialized in some cases.

#rnx
#tests unit tests

diff --git a/Engine/Source/Runtime/Core/Public/Containers/Deque.h b/Engine/Source/Runtime/Core/Public/Containers/Deque.h
--- a/Engine/Source/Runtime/Core/Public/Containers/Deque.h
+++ b/Engine/Source/Runtime/Core/Public/Containers/Deque.h
@@ -113,11 +113,13 @@ public:
 	}
 
 	[[nodiscard]] TDeque(const TDeque& Other)
+		: Capacity(Storage.GetInitialCapacity())
 	{
 		CopyUnchecked(Other);
 	}
 
 	[[nodiscard]] TDeque(std::initializer_list<ElementType> InList)
+		: Capacity(Storage.GetInitialCapacity())
 	{
 		CopyUnchecked(InList);
 	}
@@ -461,7 +463,7 @@ private:
 		checkSlow(!Count);
 		if (Other.Count)
 		{
-			Grow(Storage.CalculateSlackReserve(Other.Count, sizeof(ElementType)));
+			Reserve(Other.Count);
 			CopyElements(Other);
 		}
 		else
@@ -477,7 +479,7 @@ private:
 		checkSlow(!Count);
 		if (InCount)
 		{
-			Grow(Storage.CalculateSlackReserve(InCount, sizeof(ElementType)));
+			Reserve(InCount);
 			ConstructItems<ElementType>(GetData(), &*InList.begin(), InCount);
 			Tail = Count = InCount;
 		}
diff --git a/Engine/Source/Runtime/Core/Tests/Containers/DequeTest.cpp b/Engine/Source/Runtime/Core/Tests/Containers/DequeTest.cpp
--- a/Engine/Source/Runtime/Core/Tests/Containers/DequeTest.cpp
+++ b/Engine/Source/Runtime/Core/Tests/Containers/DequeTest.cpp
@@ -772,6 +772,17 @@ TEST_CASE_NAMED(FDequeCopyWithInlineAllocator, "System::Core::Containers::TDeque
 		CHECK(DequeCopy.Num() == 1);
 	}
 
+	SECTION("Copy Full TDeque")
+	{
+		int32Token::Reset();
+		TDeque<int32Token, TInlineAllocator<Test::DefaultCapacity>> Deque{ 0, 1, 2, 3 };
+		CHECK(Deque.Max() == Test::DefaultCapacity);
+		CHECK(Deque.Num() == 4);
+		TDeque<int32Token, TInlineAllocator<Test::DefaultCapacity>> DequeCopy(Deque);
+		CHECK_MESSAGE("Expected copy to retain inline capacity.", DequeCopy.Max() == Test::DefaultCapacity);
+		CHECK(DequeCopy.Num() == 4);
+	}
+
 	SECTION("Copy TDeque that exceeds inline allocation")
 	{
 		int32Token::Reset();
@@ -1088,6 +1099,49 @@ TEST_CASE_NAMED(FDequeAssignFromStdInitializerListTest, "System::Core::Container
 	}
 }
 
+TEST_CASE_NAMED(FDequeCopyAssignTest, "System::Core::Containers::TDeque::Assign copy", "[SmokeFilter][Core][Containers][Deque]")
+{
+	SECTION("Default allocator")
+	{
+		TDeque<int32Token> Deque;
+		TDeque<int32Token> DequeCopy{ 0, 1, 2, 3, 4, 5 };
+		Deque = DequeCopy;
+		CHECK(Deque.Num() == 6);
+		CHECK(Deque.Num() == DequeCopy.Num());
+		CHECK(Deque.Max() == DequeCopy.Max());
+	}
+
+	SECTION("Inline allocation without growing")
+	{
+		TDeque<int32Token, TInlineAllocator<8>> Deque;
+		TDeque<int32Token, TInlineAllocator<8>> DequeCopy{ 0, 1, 2, 3, 4, 5 };
+		Deque = DequeCopy;
+		CHECK(Deque.Num() == 6);
+		CHECK(Deque.Max() == 8);
+		CHECK(Deque.Num() == DequeCopy.Num());
+		CHECK(Deque.Max() == DequeCopy.Max());
+	}
+
+	SECTION("Inline allocation without growing 2")
+	{
+		TDeque<int32Token, TInlineAllocator<8>> DequeCopy{ 0, 1, 2, 3, 4, 5 };
+		TDeque<int32Token, TInlineAllocator<8>> Deque = DequeCopy;
+		CHECK(Deque.Num() == 6);
+		CHECK(Deque.Max() == 8);
+		CHECK(Deque.Num() == DequeCopy.Num());
+		CHECK(Deque.Max() == DequeCopy.Max());
+	}
+
+	SECTION("Inline allocation with growing")
+	{
+		TDeque<int32Token, TInlineAllocator<4>> Deque;
+		TDeque<int32Token, TInlineAllocator<4>> DequeCopy{ 0, 1, 2, 3, 4, 5 };
+		Deque = DequeCopy;
+		CHECK(Deque.Num() == DequeCopy.Num());
+		CHECK(Deque.Max() == DequeCopy.Max());
+	}
+}
+
 }  // namespace Deque
 
 #endif
