Changed serialization of soft package references in container header to save offset and size. Soft package references now needs to be explicitly serialized to be used.

#rnx

diff --git a/Engine/Source/Runtime/Core/Private/IO/IoContainerHeader.cpp b/Engine/Source/Runtime/Core/Private/IO/IoContainerHeader.cpp
--- a/Engine/Source/Runtime/Core/Private/IO/IoContainerHeader.cpp
+++ b/Engine/Source/Runtime/Core/Private/IO/IoContainerHeader.cpp
@@ -38,6 +38,13 @@ FArchive& operator<<(FArchive& Ar, FIoContainerHeaderSoftPackageReferences& Soft
 	return Ar;
 }
 
+FArchive& operator<<(FArchive& Ar, FIoContainerHeaderSerialInfo& SerialInfo)
+{
+	Ar << SerialInfo.Offset;
+	Ar << SerialInfo.Size;
+	return Ar;
+}
+
 FArchive& operator<<(FArchive& Ar, FIoContainerHeader& ContainerHeader)
 {
 	uint32 Signature = FIoContainerHeader::Signature;
@@ -49,8 +56,9 @@ FArchive& operator<<(FArchive& Ar, FIoContainerHeader& ContainerHeader)
 	}
 	EIoContainerHeaderVersion Version = EIoContainerHeaderVersion::Latest;
 	Ar << Version;
-	if (Ar.IsLoading() && Version != EIoContainerHeaderVersion::Latest)
+	if (Ar.IsLoading() && Version < EIoContainerHeaderVersion::NoExportInfo)
 	{
+		// This version is too old
 		Ar.SetError();
 		return Ar;
 	}
@@ -74,10 +82,36 @@ FArchive& operator<<(FArchive& Ar, FIoContainerHeader& ContainerHeader)
 	Ar << ContainerHeader.LocalizedPackages;
 	Ar << ContainerHeader.PackageRedirects;
 
-	if (Ar.IsSaving() || Version >= EIoContainerHeaderVersion::SoftPackageReferences)
+	if (Version == EIoContainerHeaderVersion::SoftPackageReferences)
 	{
 		Ar << ContainerHeader.SoftPackageReferences;
 	}
+	else if (Version >= EIoContainerHeaderVersion::SoftPackageReferencesOffset)
+	{
+		const int64 SerialInfoOffset = Ar.Tell();
+		Ar << ContainerHeader.SoftPackageReferencesSerialInfo;
+		if (Ar.IsLoading())
+		{
+			if (ContainerHeader.SoftPackageReferencesSerialInfo.Size > 0)
+			{
+				const int64 EndPos = Ar.Tell() + ContainerHeader.SoftPackageReferencesSerialInfo.Size;
+				if (EndPos > Ar.TotalSize())
+				{
+					Ar.SetError();
+					return Ar;
+				}
+				Ar.Seek(EndPos);
+			}
+		}
+		else 
+		{
+			ContainerHeader.SoftPackageReferencesSerialInfo.Offset = Ar.Tell();
+			Ar << ContainerHeader.SoftPackageReferences;
+			ContainerHeader.SoftPackageReferencesSerialInfo.Size = Ar.Tell() - ContainerHeader.SoftPackageReferencesSerialInfo.Offset;
+			FArchive::FScopeSeekTo _(Ar, SerialInfoOffset);
+			Ar << ContainerHeader.SoftPackageReferencesSerialInfo;
+		}
+	}
 
 	return Ar;
 }
diff --git a/Engine/Source/Runtime/Core/Public/IO/IoContainerHeader.h b/Engine/Source/Runtime/Core/Public/IO/IoContainerHeader.h
--- a/Engine/Source/Runtime/Core/Public/IO/IoContainerHeader.h
+++ b/Engine/Source/Runtime/Core/Public/IO/IoContainerHeader.h
@@ -82,6 +82,14 @@ struct FIoContainerHeaderSoftPackageReferences
 	CORE_API friend FArchive& operator<<(FArchive& Ar, FIoContainerHeaderSoftPackageReferences& SoftPackageReferences);
 };
 
+struct FIoContainerHeaderSerialInfo
+{
+	int64 Offset = -1;
+	int64 Size = -1;
+
+	CORE_API friend FArchive& operator<<(FArchive& Ar, FIoContainerHeaderSerialInfo& SerialInfo);
+};
+
 enum class EIoContainerHeaderVersion : uint32
 {
 	Initial = 0,
@@ -89,6 +97,7 @@ enum class EIoContainerHeaderVersion : uint32
 	OptionalSegmentPackages = 2,
 	NoExportInfo = 3,
 	SoftPackageReferences = 4,
+	SoftPackageReferencesOffset = 5,
 
 	LatestPlusOne,
 	Latest = LatestPlusOne - 1
@@ -109,6 +118,7 @@ struct FIoContainerHeader
 	TArray<FDisplayNameEntryId> RedirectsNameMap;
 	TArray<FIoContainerHeaderLocalizedPackage> LocalizedPackages;
 	TArray<FIoContainerHeaderPackageRedirect> PackageRedirects;
+	FIoContainerHeaderSerialInfo SoftPackageReferencesSerialInfo;
 	FIoContainerHeaderSoftPackageReferences SoftPackageReferences;
 
 	CORE_API friend FArchive& operator<<(FArchive& Ar, FIoContainerHeader& ContainerHeader);
diff --git a/Engine/Source/Runtime/Experimental/IoStore/OnDemand/Private/OnDemandIoStore.cpp b/Engine/Source/Runtime/Experimental/IoStore/OnDemand/Private/OnDemandIoStore.cpp
--- a/Engine/Source/Runtime/Experimental/IoStore/OnDemand/Private/OnDemandIoStore.cpp
+++ b/Engine/Source/Runtime/Experimental/IoStore/OnDemand/Private/OnDemandIoStore.cpp
@@ -139,6 +139,26 @@ static TIoStatusOr<FSharedContainerHeader> DeserializeContainerHeader(const FOnD
 	FSharedContainerHeader OutHeader = MakeShared<FIoContainerHeader>();
 	FMemoryReaderView Ar(Chunk.ValueOrDie().GetView());
 	Ar << *OutHeader;
+	if (OutHeader->SoftPackageReferencesSerialInfo.Size > 0)
+	{
+		if (OutHeader->SoftPackageReferencesSerialInfo.Offset < 0)
+		{
+			FIoStatus Status = FIoStatusBuilder(EIoErrorCode::ReadError)
+				<< FString::Printf(TEXT("Invalid soft package reference offset '" INT64_FMT "'"), OutHeader->SoftPackageReferencesSerialInfo.Offset);
+			return Status;
+		}
+		if ((OutHeader->SoftPackageReferencesSerialInfo.Offset + OutHeader->SoftPackageReferencesSerialInfo.Size) > Ar.TotalSize())
+		{
+			FIoStatus Status = FIoStatusBuilder(EIoErrorCode::ReadError)
+				<< FString::Printf(TEXT("Soft package reference offset '" INT64_FMT "' and size '" INT64_FMT "' will seek past the end of archive size '" INT64_FMT "'"),
+					OutHeader->SoftPackageReferencesSerialInfo.Offset,
+					OutHeader->SoftPackageReferencesSerialInfo.Size,
+					Ar.TotalSize());
+			return Status;
+		}
+		Ar.Seek(OutHeader->SoftPackageReferencesSerialInfo.Offset);
+		Ar << OutHeader->SoftPackageReferences;
+	}
 	Ar.Close();
 
 	if (Ar.IsError() || Ar.IsCriticalError())
