[Insights/Trace]
- Added support for FName variants of CpuProfilerTrace macros: for TRACE_CPUPROFILER_EVENT_SCOPE_STR (when the provided event name does not change) and for TRACE_CPUPROFILER_EVENT_SCOPE_TEXT (when the provided event name is dynamic). FEventScope and FDynamicEventScope can also be initialized now with an FName.
- Fixed some comments for CpuProfilerTrace API.
- Added variations for _STR macros (_STR_ON_CHANNEL_CONDITIONAL vs. _ON_CHANNEL_STR_CONDITIONAL vs _ON_CHANNEL_CONDITIONAL_STR).

diff --git a/Engine/Source/Runtime/Core/Private/ProfilingDebugging/CpuProfilerTrace.cpp b/Engine/Source/Runtime/Core/Private/ProfilingDebugging/CpuProfilerTrace.cpp
--- a/Engine/Source/Runtime/Core/Private/ProfilingDebugging/CpuProfilerTrace.cpp
+++ b/Engine/Source/Runtime/Core/Private/ProfilingDebugging/CpuProfilerTrace.cpp
@@ -1,19 +1,21 @@
 // Copyright Epic Games, Inc. All Rights Reserved.
+
 #include "ProfilingDebugging/CpuProfilerTrace.h"
-#include "Trace/Trace.inl"
+
+#if CPUPROFILERTRACE_ENABLED
+
+#include "AutoRTFM/AutoRTFM.h"
+#include "Containers/Map.h"
 #include "HAL/LowLevelMemTracker.h"
 #include "HAL/PlatformTime.h"
 #include "HAL/PlatformTLS.h"
 #include "HAL/TlsAutoCleanup.h"
-#include "ProfilingDebugging/MiscTrace.h"
-#include "Misc/Parse.h"
-#include "Containers/Map.h"
-#include "Misc/MemStack.h"
 #include "Misc/Crc.h"
+#include "Misc/MemStack.h"
+#include "Misc/Parse.h"
+#include "ProfilingDebugging/MiscTrace.h"
+#include "Trace/Trace.inl"
 #include "UObject/NameTypes.h"
-#include "AutoRTFM/AutoRTFM.h"
-
-#if CPUPROFILERTRACE_ENABLED
 
 #if !defined(CPUPROFILERTRACE_FILE_AND_LINE_ENABLED)
 	#define CPUPROFILERTRACE_FILE_AND_LINE_ENABLED 1
@@ -167,7 +169,7 @@ void FCpuProfilerTraceInternal::FlushThreadBuffer(FThreadBuffer* InThreadBuffer)
 
 uint8 FCpuProfilerTrace::OnAbortKey = 42;
 
-void FCpuProfilerTrace::GetOrCreateSpecId(uint32& InOutSpecId, const ANSICHAR* InEventString, const ANSICHAR* File, uint32 Line)
+void FCpuProfilerTrace::GetOrCreateSpecId(uint32& InOutSpecId, const ANSICHAR* InEventName, const ANSICHAR* File, uint32 Line)
 {
 	AutoRTFM::Open([&]
 	{
@@ -176,12 +178,12 @@ void FCpuProfilerTrace::GetOrCreateSpecId(uint32& InOutSpecId, const ANSICHAR* I
 		// We don't care which event type wins as long as all threads eventually converge and stop creating new ones.
 		if (FPlatformAtomics::AtomicRead_Relaxed((volatile int32*)&InOutSpecId) == 0)
 		{
-			FPlatformAtomics::AtomicStore_Relaxed((volatile int32*)&InOutSpecId, FCpuProfilerTrace::OutputEventType(InEventString, File, Line));
+			FPlatformAtomics::AtomicStore_Relaxed((volatile int32*)&InOutSpecId, FCpuProfilerTrace::OutputEventType(InEventName, File, Line));
 		}
 	});
 }
 
-void FCpuProfilerTrace::GetOrCreateSpecId(uint32& InOutSpecId, const TCHAR* InEventString, const ANSICHAR* File, uint32 Line)
+void FCpuProfilerTrace::GetOrCreateSpecId(uint32& InOutSpecId, const TCHAR* InEventName, const ANSICHAR* File, uint32 Line)
 {
 	AutoRTFM::Open([&]
 	{
@@ -190,7 +192,21 @@ void FCpuProfilerTrace::GetOrCreateSpecId(uint32& InOutSpecId, const TCHAR* InEv
 		// We don't care which event type wins as long as all threads eventually converge and stop creating new ones.
 		if (FPlatformAtomics::AtomicRead_Relaxed((volatile int32*)&InOutSpecId) == 0)
 		{
-			FPlatformAtomics::AtomicStore_Relaxed((volatile int32*)&InOutSpecId, FCpuProfilerTrace::OutputEventType(InEventString, File, Line));
+			FPlatformAtomics::AtomicStore_Relaxed((volatile int32*)&InOutSpecId, FCpuProfilerTrace::OutputEventType(InEventName, File, Line));
+		}
+	});
+}
+
+void FCpuProfilerTrace::GetOrCreateSpecId(uint32& InOutSpecId, FNameParam InEventName, const ANSICHAR* File, uint32 Line)
+{
+	AutoRTFM::Open([&]
+	{
+		// We only do relaxed here to avoid barrier cost as the worst case that can happen is multiple threads could each create an event type.
+		// At some point the last thread in the race will set the output event and no more thread will try to create new ones from then on.
+		// We don't care which event type wins as long as all threads eventually converge and stop creating new ones.
+		if (FPlatformAtomics::AtomicRead_Relaxed((volatile int32*)&InOutSpecId) == 0)
+		{
+			FPlatformAtomics::AtomicStore_Relaxed((volatile int32*)&InOutSpecId, FCpuProfilerTrace::OutputEventType(InEventName, File, Line));
 		}
 	});
 }
@@ -260,7 +276,7 @@ void FCpuProfilerTrace::OutputBeginDynamicEvent(const TCHAR* Name, const ANSICHA
 	CPUPROFILERTRACE_OUTPUTBEGINEVENT_EPILOGUE();
 }
 
-void FCpuProfilerTrace::OutputBeginDynamicEvent(const FName Name, const ANSICHAR* File, uint32 Line)
+void FCpuProfilerTrace::OutputBeginDynamicEvent(FNameParam Name, const ANSICHAR* File, uint32 Line)
 {
 	if (AutoRTFM::IsClosed())
 	{
@@ -282,7 +298,7 @@ void FCpuProfilerTrace::OutputBeginDynamicEvent(const FName Name, const ANSICHAR
 	CPUPROFILERTRACE_OUTPUTBEGINEVENT_EPILOGUE();
 }
 
-void FCpuProfilerTrace::OutputBeginDynamicEventWithId(const FName Id, const ANSICHAR* Name, const ANSICHAR* File, uint32 Line)
+void FCpuProfilerTrace::OutputBeginDynamicEventWithId(FNameParam Id, const ANSICHAR* Name, const ANSICHAR* File, uint32 Line)
 {
 	if (AutoRTFM::IsClosed())
 	{
@@ -304,7 +320,7 @@ void FCpuProfilerTrace::OutputBeginDynamicEventWithId(const FName Id, const ANSI
 	CPUPROFILERTRACE_OUTPUTBEGINEVENT_EPILOGUE();
 }
 
-void FCpuProfilerTrace::OutputBeginDynamicEventWithId(const FName Id, const TCHAR* Name, const ANSICHAR* File, uint32 Line)
+void FCpuProfilerTrace::OutputBeginDynamicEventWithId(FNameParam Id, const TCHAR* Name, const ANSICHAR* File, uint32 Line)
 {
 	if (AutoRTFM::IsClosed())
 	{
@@ -478,7 +494,7 @@ uint32 FCpuProfilerTrace::OutputEventType(const ANSICHAR* Name, const ANSICHAR*
 	return SpecId;
 }
 
-uint32 FCpuProfilerTrace::OutputEventType(const FName Name, const ANSICHAR* File, uint32 Line)
+uint32 FCpuProfilerTrace::OutputEventType(FNameParam Name, const ANSICHAR* File, uint32 Line)
 {
 	if (AutoRTFM::IsClosed())
 	{
@@ -559,7 +575,7 @@ uint32 FCpuProfilerTrace::OutputDynamicEventType(const TCHAR* Name, const ANSICH
 	return SpecId;
 }
 
-uint32 FCpuProfilerTrace::OutputDynamicEventType(const FName Name, const ANSICHAR* File, uint32 Line)
+uint32 FCpuProfilerTrace::OutputDynamicEventType(FNameParam Name, const ANSICHAR* File, uint32 Line)
 {
 	if (AutoRTFM::IsClosed())
 	{
@@ -584,7 +600,7 @@ uint32 FCpuProfilerTrace::OutputDynamicEventType(const FName Name, const ANSICHA
 	return SpecId;
 }
 
-uint32 FCpuProfilerTrace::OutputDynamicEventTypeWithId(const FName Id, const ANSICHAR* Name, const ANSICHAR* File, uint32 Line)
+uint32 FCpuProfilerTrace::OutputDynamicEventTypeWithId(FNameParam Id, const ANSICHAR* Name, const ANSICHAR* File, uint32 Line)
 {
 	if (AutoRTFM::IsClosed())
 	{
@@ -616,7 +632,7 @@ uint32 FCpuProfilerTrace::OutputDynamicEventTypeWithId(const FName Id, const ANS
 	return SpecId;
 }
 
-uint32 FCpuProfilerTrace::OutputDynamicEventTypeWithId(const FName Id, const TCHAR* Name, const ANSICHAR* File, uint32 Line)
+uint32 FCpuProfilerTrace::OutputDynamicEventTypeWithId(FNameParam Id, const TCHAR* Name, const ANSICHAR* File, uint32 Line)
 {
 	if (AutoRTFM::IsClosed())
 	{
@@ -667,4 +683,4 @@ void FCpuProfilerTrace::FlushThreadBuffer()
 	}
 }
 
-#endif
+#endif // CPUPROFILERTRACE_ENABLED
diff --git a/Engine/Source/Runtime/Core/Public/ProfilingDebugging/CpuProfilerTrace.h b/Engine/Source/Runtime/Core/Public/ProfilingDebugging/CpuProfilerTrace.h
--- a/Engine/Source/Runtime/Core/Public/ProfilingDebugging/CpuProfilerTrace.h
+++ b/Engine/Source/Runtime/Core/Public/ProfilingDebugging/CpuProfilerTrace.h
@@ -26,9 +26,9 @@ UE_TRACE_CHANNEL_EXTERN(CpuChannel, CORE_API);
 class FName;
 
 /*
- * Facilities for tracing timed CPU events. Two types of events are supported, static events where the identifier is
- * known at compile time, and dynamic event were identifiers can be constructed in runtime. Static events have lower overhead
- * so always prefer to use them if possible.
+ * Facilities for tracing timed CPU events. Two types of events are supported, static/constant events where the identifier is
+ * known at compile time or it is constant during execution, and dynamic events were identifier can change between calls.
+ * Static/constant events have lower overhead, so always prefer to use them if possible.
  *
  * Events are tracked per thread, so begin/end calls must be matched and called on the same thread. It is possible to use any channel
  * to emit the events, but both that channel and the CpuChannel must then be enabled.
@@ -37,6 +37,8 @@ class FName;
  */
 struct FCpuProfilerTrace
 {
+	typedef const FName& FNameParam;
+
 	/*
 	 * Output CPU event definition (spec).
 	 * The trace event emitted by this function is an "important event" (so all events emitted will add to the Trace System's cache).
@@ -48,7 +50,7 @@ struct FCpuProfilerTrace
 	 */
 	FORCENOINLINE CORE_API static uint32 OutputEventType(const ANSICHAR* Name, const ANSICHAR* File = nullptr, uint32 Line = 0);
 	FORCENOINLINE CORE_API static uint32 OutputEventType(const TCHAR* Name, const ANSICHAR* File = nullptr, uint32 Line = 0);
-	FORCENOINLINE CORE_API static uint32 OutputEventType(const FName Name, const ANSICHAR* File = nullptr, uint32 Line = 0);
+	FORCENOINLINE CORE_API static uint32 OutputEventType(FNameParam Name, const ANSICHAR* File = nullptr, uint32 Line = 0);
 
 	/*
 	 * Output CPU event definition (spec) for a dynamic event.
@@ -60,7 +62,7 @@ struct FCpuProfilerTrace
 	 */
 	FORCENOINLINE CORE_API static uint32 OutputDynamicEventType(const ANSICHAR* Name, const ANSICHAR* File = nullptr, uint32 Line = 0);
 	FORCENOINLINE CORE_API static uint32 OutputDynamicEventType(const TCHAR* Name, const ANSICHAR* File = nullptr, uint32 Line = 0);
-	FORCENOINLINE CORE_API static uint32 OutputDynamicEventType(const FName Name, const ANSICHAR* File = nullptr, uint32 Line = 0);
+	FORCENOINLINE CORE_API static uint32 OutputDynamicEventType(FNameParam Name, const ANSICHAR* File = nullptr, uint32 Line = 0);
 
 	/*
 	 * Output CPU event definition (spec) for a dynamic event identified by an FName.
@@ -71,8 +73,8 @@ struct FCpuProfilerTrace
 	 * @param File - The source filename
 	 * @param Line - The line number in source file
 	 */
-	FORCENOINLINE CORE_API static uint32 OutputDynamicEventTypeWithId(const FName Id, const ANSICHAR* Name, const ANSICHAR* File = nullptr, uint32 Line = 0);
-	FORCENOINLINE CORE_API static uint32 OutputDynamicEventTypeWithId(const FName Id, const TCHAR* Name, const ANSICHAR* File = nullptr, uint32 Line = 0);
+	FORCENOINLINE CORE_API static uint32 OutputDynamicEventTypeWithId(FNameParam Id, const ANSICHAR* Name, const ANSICHAR* File = nullptr, uint32 Line = 0);
+	FORCENOINLINE CORE_API static uint32 OutputDynamicEventTypeWithId(FNameParam Id, const TCHAR* Name, const ANSICHAR* File = nullptr, uint32 Line = 0);
 
 	/*
 	 * Output begin event marker for a given spec. Must always be matched with an end event.
@@ -98,7 +100,7 @@ struct FCpuProfilerTrace
 	 * @param File - The source filename
 	 * @param Line - The line number in source file
 	 */
-	CORE_API static void OutputBeginDynamicEvent(const FName Name, const ANSICHAR* File = nullptr, uint32 Line = 0);
+	CORE_API static void OutputBeginDynamicEvent(FNameParam Name, const ANSICHAR* File = nullptr, uint32 Line = 0);
 
 	/*
 	 * Output begin event marker for a dynamic event identified by an FName. This is more expensive than
@@ -109,8 +111,8 @@ struct FCpuProfilerTrace
 	 * @param File - The source filename
 	 * @param Line - The line number in source file
 	 */
-	CORE_API static void OutputBeginDynamicEventWithId(const FName Id, const ANSICHAR* Name, const ANSICHAR* File = nullptr, uint32 Line = 0);
-	CORE_API static void OutputBeginDynamicEventWithId(const FName Id, const TCHAR* Name, const ANSICHAR* File = nullptr, uint32 Line = 0);
+	CORE_API static void OutputBeginDynamicEventWithId(FNameParam Id, const ANSICHAR* Name, const ANSICHAR* File = nullptr, uint32 Line = 0);
+	CORE_API static void OutputBeginDynamicEventWithId(FNameParam Id, const TCHAR* Name, const ANSICHAR* File = nullptr, uint32 Line = 0);
 
 	/*
 	 * Output end event marker for static or dynamic event for the currently open scope.
@@ -118,20 +120,20 @@ struct FCpuProfilerTrace
 	CORE_API static void OutputEndEvent();
 
 	/*
-	* Output resume marker for a given spec. Must always be matched with an suspend event.
-	* @param SpecId - The unique Resume Event definition id
-	* @param TimerScopeDepth - updates the depth of the current OutputBeginEvent depth
-	*/
+	 * Output resume marker for a given spec. Must always be matched with an suspend event.
+	 * @param SpecId - The unique Resume Event definition id
+	 * @param TimerScopeDepth - updates the depth of the current OutputBeginEvent depth
+	 */
 	CORE_API static void OutputResumeEvent(uint64 SpecId, uint32& TimerScopeDepth);
 
 	/*
-	* Output suspend event marker for the currently open resume event.
-	*/
+	 * Output suspend event marker for the currently open resume event.
+	 */
 	CORE_API static void OutputSuspendEvent();
 
 	/*
-	* Make sure all thread data has reached the destination. Can be useful to call this before entering a wait condition that might take a while.
-	*/
+	 * Make sure all thread data has reached the destination. Can be useful to call this before entering a wait condition that might take a while.
+	 */
 	CORE_API static void FlushThreadBuffer();
 
 	/*
@@ -141,16 +143,9 @@ struct FCpuProfilerTrace
 	 * @param File - The source filename
 	 * @param Line - The line number in source file
 	 */
-	CORE_API static void GetOrCreateSpecId(uint32& InOutSpecId, const ANSICHAR* InEventString, const ANSICHAR* File, uint32 Line);
-	
-	/*
-	 * Get or create a unique event definition id for the given parameters.
-	 * @param InOutSpecId - The event definition id.
-	 * @param Name - The name of event
-	 * @param File - The source filename
-	 * @param Line - The line number in source file
-	 */
-	CORE_API static void GetOrCreateSpecId(uint32& InOutSpecId, const TCHAR* InEventString, const ANSICHAR* File, uint32 Line);
+	CORE_API static void GetOrCreateSpecId(uint32& InOutSpecId, const ANSICHAR* Name, const ANSICHAR* File, uint32 Line);
+	CORE_API static void GetOrCreateSpecId(uint32& InOutSpecId, const TCHAR* Name, const ANSICHAR* File, uint32 Line);
+	CORE_API static void GetOrCreateSpecId(uint32& InOutSpecId, FNameParam Name, const ANSICHAR* File, uint32 Line);
 
 	class FEventScope
 	{
@@ -167,28 +162,40 @@ struct FCpuProfilerTrace
 			BeginEventCommon(InSpecId);
 		}
 
-		FORCEINLINE FEventScope(uint32& InOutSpecId, const ANSICHAR* InEventString, bool bInCondition, const ANSICHAR* File, uint32 Line)
+		FORCEINLINE FEventScope(uint32& InOutSpecId, const ANSICHAR* InEventName, bool bInCondition, const ANSICHAR* File, uint32 Line)
 			: bEnabled(bInCondition && CpuChannel)
 		{
-			BeginEventCommon(InOutSpecId, InEventString, File, Line);
+			BeginEventCommon(InOutSpecId, InEventName, File, Line);
 		}
 
-		FORCEINLINE FEventScope(uint32& InOutSpecId, const ANSICHAR* InEventString, const UE::Trace::FChannel& InChannel, bool bInCondition, const ANSICHAR* File, uint32 Line)
+		FORCEINLINE FEventScope(uint32& InOutSpecId, const ANSICHAR* InEventName, const UE::Trace::FChannel& InChannel, bool bInCondition, const ANSICHAR* File, uint32 Line)
 			: bEnabled(bInCondition && (CpuChannel | InChannel))
 		{
-			BeginEventCommon(InOutSpecId, InEventString, File, Line);
+			BeginEventCommon(InOutSpecId, InEventName, File, Line);
 		}
 
-		FORCEINLINE FEventScope(uint32& InOutSpecId, const TCHAR* InEventString, bool bInCondition, const ANSICHAR* File, uint32 Line)
+		FORCEINLINE FEventScope(uint32& InOutSpecId, const TCHAR* InEventName, bool bInCondition, const ANSICHAR* File, uint32 Line)
 			: bEnabled(bInCondition && CpuChannel)
 		{
-			BeginEventCommon(InOutSpecId, InEventString, File, Line);
+			BeginEventCommon(InOutSpecId, InEventName, File, Line);
 		}
 
-		FORCEINLINE FEventScope(uint32& InOutSpecId, const TCHAR* InEventString, const UE::Trace::FChannel& InChannel, bool bInCondition, const ANSICHAR* File, uint32 Line)
+		FORCEINLINE FEventScope(uint32& InOutSpecId, const TCHAR* InEventName, const UE::Trace::FChannel& InChannel, bool bInCondition, const ANSICHAR* File, uint32 Line)
 			: bEnabled(bInCondition && (CpuChannel | InChannel))
 		{
-			BeginEventCommon(InOutSpecId, InEventString, File, Line);
+			BeginEventCommon(InOutSpecId, InEventName, File, Line);
+		}
+
+		FORCEINLINE explicit FEventScope(uint32& InOutSpecId, FNameParam InEventName, bool bInCondition, const ANSICHAR* File, uint32 Line)
+			: bEnabled(bInCondition && CpuChannel)
+		{
+			BeginEventCommon(InOutSpecId, InEventName, File, Line);
+		}
+
+		FORCEINLINE explicit FEventScope(uint32& InOutSpecId, FNameParam InEventName, const UE::Trace::FChannel& InChannel, bool bInCondition, const ANSICHAR* File, uint32 Line)
+			: bEnabled(bInCondition && (CpuChannel | InChannel))
+		{
+			BeginEventCommon(InOutSpecId, InEventName, File, Line);
 		}
 
 		FORCEINLINE ~FEventScope()
@@ -208,20 +215,29 @@ struct FCpuProfilerTrace
 			}
 		}
 
-		FORCEINLINE void BeginEventCommon(uint32& InOutSpecId, const ANSICHAR* InEventString, const ANSICHAR* File, uint32 Line)
+		FORCEINLINE void BeginEventCommon(uint32& InOutSpecId, const ANSICHAR* InEventName, const ANSICHAR* File, uint32 Line)
 		{
 			if (bEnabled)
 			{
-				GetOrCreateSpecId(InOutSpecId, InEventString, File, Line);
+				GetOrCreateSpecId(InOutSpecId, InEventName, File, Line);
 				OutputBeginEvent(FPlatformAtomics::AtomicRead_Relaxed((volatile int32*)&InOutSpecId));
 			}
 		}
 
-		FORCEINLINE void BeginEventCommon(uint32& InOutSpecId, const TCHAR* InEventString, const ANSICHAR* File, uint32 Line)
+		FORCEINLINE void BeginEventCommon(uint32& InOutSpecId, const TCHAR* InEventName, const ANSICHAR* File, uint32 Line)
 		{
 			if (bEnabled)
 			{
-				GetOrCreateSpecId(InOutSpecId, InEventString, File, Line);
+				GetOrCreateSpecId(InOutSpecId, InEventName, File, Line);
+				OutputBeginEvent(FPlatformAtomics::AtomicRead_Relaxed((volatile int32*)&InOutSpecId));
+			}
+		}
+
+		FORCEINLINE void BeginEventCommon(uint32& InOutSpecId, FNameParam InEventName, const ANSICHAR* File, uint32 Line)
+		{
+			if (bEnabled)
+			{
+				GetOrCreateSpecId(InOutSpecId, InEventName, File, Line);
 				OutputBeginEvent(FPlatformAtomics::AtomicRead_Relaxed((volatile int32*)&InOutSpecId));
 			}
 		}
@@ -267,6 +283,24 @@ struct FCpuProfilerTrace
 			}
 		}
 
+		FORCEINLINE explicit FDynamicEventScope(FNameParam InEventName, bool bInCondition, const ANSICHAR* InFile = nullptr, uint32 InLine = 0)
+			: bEnabled(bInCondition && CpuChannel)
+		{
+			if (bEnabled)
+			{
+				OutputBeginDynamicEvent(InEventName, InFile, InLine);
+			}
+		}
+
+		FORCEINLINE explicit FDynamicEventScope(FNameParam InEventName, const UE::Trace::FChannel& InChannel, bool bInCondition, const ANSICHAR* InFile = nullptr, uint32 InLine = 0)
+			: bEnabled(bInCondition && (CpuChannel | InChannel))
+		{
+			if (bEnabled)
+			{
+				OutputBeginDynamicEvent(InEventName, InFile, InLine);
+			}
+		}
+
 		FORCEINLINE ~FDynamicEventScope()
 		{
 			if (bEnabled)
@@ -317,98 +351,129 @@ private:
 #define TRACE_CPUPROFILER_EVENT_MANUAL_IS_ENABLED() \
 	bool(CpuChannel)
 
-// Conditionally trace a scoped CPU timing event providing a static string (const ANSICHAR* or const TCHAR*)
-// as the scope name and a condition under which to create the trace. It will use the CPU trace channel.
+////////////////////////////////////////////////////////////////////////////////////////////////////
+// TRACE_CPUPROFILER_EVENT_SCOPE_STR* - scope name as static/constant string
+////////////////////////////////////////////////////////////////////////////////////////////////////
+
+// Conditionally trace a scoped CPU timing event providing a static/constant string (const ANSICHAR*
+// or const TCHAR* or FName) as the scope name.
+// It will use the CPU trace channel.
+// Note: The name will be cached only once. If used with a string pointer or an FName,
+//       the name should be guaranteed to not change between calls.
 // Example: TRACE_CPUPROFILER_EVENT_SCOPE_STR_CONDITIONAL("My Scoped Timer A", Condition)
 #define TRACE_CPUPROFILER_EVENT_SCOPE_STR_CONDITIONAL(NameStr, Condition) \
 	TRACE_CPUPROFILER_EVENT_DECLARE(PREPROCESSOR_JOIN(__CpuProfilerEventSpecId, __LINE__)); \
 	TRACE_CPUPROFILER_EVENT_SCOPE_USE(PREPROCESSOR_JOIN(__CpuProfilerEventSpecId, __LINE__), NameStr, PREPROCESSOR_JOIN(__CpuProfilerEventScope, __LINE__), (Condition));
 
-// Trace a scoped CPU timing event providing a static string (const ANSICHAR* or const TCHAR*)
-// as the scope name. It will use the CPU trace channel.
+// Trace a scoped CPU timing event providing a static/constant string (const ANSICHAR* or
+// const TCHAR* or FName) as the scope name. It will use the CPU trace channel.
+// Note: The name will be cached only once. If used with a string pointer or an FName,
+//       the name should be guaranteed to not change between calls.
 // Example: TRACE_CPUPROFILER_EVENT_SCOPE_STR("My Scoped Timer A")
 #define TRACE_CPUPROFILER_EVENT_SCOPE_STR(NameStr) \
 	TRACE_CPUPROFILER_EVENT_SCOPE_STR_CONDITIONAL(NameStr, true)
 
-// Conditionally trace a scoped CPU timing event providing a static string (const ANSICHAR* or const TCHAR*)
-// as the scope name and a trace channel and a condition.
-// Example: TRACE_CPUPROFILER_EVENT_SCOPE_ON_CHANNEL_STR_CONDITIONAL("My Scoped Timer A", CpuChannel, Condition)
-// Note: The event will be emitted only if both the given channel and CpuChannel is enabled.
-#define TRACE_CPUPROFILER_EVENT_SCOPE_ON_CHANNEL_STR_CONDITIONAL(NameStr, Channel, Condition) \
+// Conditionally trace a scoped CPU timing event providing a static/constant string (const ANSICHAR*
+// or const TCHAR* or FName) as the scope name and a trace channel.
+// Note: The name will be cached only once. If used with a string pointer or an FName,
+//       the name should be guaranteed to not change between calls.
+// Note: The event will be emitted only if both the given channel and the CpuChannel are enabled.
+// Example: TRACE_CPUPROFILER_EVENT_SCOPE_STR_ON_CHANNEL_CONDITIONAL("My Scoped Timer A", CustomChannel, Condition)
+#define TRACE_CPUPROFILER_EVENT_SCOPE_STR_ON_CHANNEL_CONDITIONAL(NameStr, Channel, Condition) \
 	TRACE_CPUPROFILER_EVENT_DECLARE(PREPROCESSOR_JOIN(__CpuProfilerEventSpecId, __LINE__)); \
 	TRACE_CPUPROFILER_EVENT_SCOPE_USE_ON_CHANNEL(PREPROCESSOR_JOIN(__CpuProfilerEventSpecId, __LINE__), NameStr, PREPROCESSOR_JOIN(__CpuProfilerEventScope, __LINE__), Channel, (Condition));
-
-// Trace a scoped CPU timing event providing a static string (const ANSICHAR* or const TCHAR*)
-// as the scope name and a trace channel.
-// Example: TRACE_CPUPROFILER_EVENT_SCOPE_ON_CHANNEL_STR("My Scoped Timer A", CpuChannel)
-// Note: The event will be emitted only if both the given channel and CpuChannel is enabled.
+#define TRACE_CPUPROFILER_EVENT_SCOPE_ON_CHANNEL_STR_CONDITIONAL(NameStr, Channel, Condition) \
+	TRACE_CPUPROFILER_EVENT_SCOPE_STR_ON_CHANNEL_CONDITIONAL(NameStr, Channel, Condition)
+#define TRACE_CPUPROFILER_EVENT_SCOPE_ON_CHANNEL_CONDITIONAL_STR(NameStr, Channel, Condition) \
+	TRACE_CPUPROFILER_EVENT_SCOPE_STR_ON_CHANNEL_CONDITIONAL(NameStr, Channel, Condition)
+
+// Trace a scoped CPU timing event providing a static/constant string (const ANSICHAR* or
+// const TCHAR* or FName) as the scope name and a trace channel.
+// Note: The name will be cached only once. If used with a string pointer or an FName,
+//       the name should be guaranteed to not change between calls.
+// Note: The event will be emitted only if both the given channel and the CpuChannel are enabled.
+// Example: TRACE_CPUPROFILER_EVENT_SCOPE_STR_ON_CHANNEL("My Scoped Timer A", CustomChannel)
+#define TRACE_CPUPROFILER_EVENT_SCOPE_STR_ON_CHANNEL(NameStr, Channel) \
+	TRACE_CPUPROFILER_EVENT_SCOPE_STR_ON_CHANNEL_CONDITIONAL(NameStr, Channel, true)
 #define TRACE_CPUPROFILER_EVENT_SCOPE_ON_CHANNEL_STR(NameStr, Channel) \
-	TRACE_CPUPROFILER_EVENT_SCOPE_ON_CHANNEL_STR_CONDITIONAL(NameStr, Channel, true)
+	TRACE_CPUPROFILER_EVENT_SCOPE_STR_ON_CHANNEL(NameStr, Channel)
+
+////////////////////////////////////////////////////////////////////////////////////////////////////
+// TRACE_CPUPROFILER_EVENT_SCOPE* - scope name as plain text
+////////////////////////////////////////////////////////////////////////////////////////////////////
 
-// Conditionally trace a scoped CPU timing event providing a scope name (plain text) and a condition.
+// Conditionally trace a scoped CPU timing event providing a scope name (plain text).
 // It will use the CPU trace channel.
-// Example: TRACE_CPUPROFILER_EVENT_SCOPE_CONDITIONAL(MyScopedTimer::A, Condition)
 // Note: Do not use this macro with a static string because, in that case, additional quotes will
 //       be added around the event scope name.
+// Example: TRACE_CPUPROFILER_EVENT_SCOPE_CONDITIONAL(MyScopedTimer::A, Condition)
 #define TRACE_CPUPROFILER_EVENT_SCOPE_CONDITIONAL(Name, Condition) \
 	TRACE_CPUPROFILER_EVENT_SCOPE_STR_CONDITIONAL(#Name, (Condition))
 
 // Trace a scoped CPU timing event providing a scope name (plain text).
 // It will use the CPU trace channel.
-// Example: TRACE_CPUPROFILER_EVENT_SCOPE(MyScopedTimer::A)
 // Note: Do not use this macro with a static string because, in that case, additional quotes will
 //       be added around the event scope name.
+// Example: TRACE_CPUPROFILER_EVENT_SCOPE(MyScopedTimer::A)
 #define TRACE_CPUPROFILER_EVENT_SCOPE(Name) \
 	TRACE_CPUPROFILER_EVENT_SCOPE_CONDITIONAL(Name, true)
 
-// Conditionally trace a scoped CPU timing event providing a scope name (plain text), a trace channel, and a condition.
-// Example: TRACE_CPUPROFILER_EVENT_SCOPE_ON_CHANNEL_CONDITIONAL(MyScopedTimer::A, CpuChannel, Condition)
+// Conditionally trace a scoped CPU timing event providing a scope name (plain text) and a trace channel.
 // Note: Do not use this macro with a static string because, in that case, additional quotes will
 //       be added around the event scope name.
-// Note: The event will be emitted only if both the given channel and CpuChannel is enabled.
+// Note: The event will be emitted only if both the given channel and the CpuChannel are enabled.
+// Example: TRACE_CPUPROFILER_EVENT_SCOPE_ON_CHANNEL_CONDITIONAL(MyScopedTimer::A, CustomChannel, Condition)
 #define TRACE_CPUPROFILER_EVENT_SCOPE_ON_CHANNEL_CONDITIONAL(Name, Channel, Condition) \
-	TRACE_CPUPROFILER_EVENT_SCOPE_ON_CHANNEL_STR_CONDITIONAL(#Name, Channel, (Condition))
+	TRACE_CPUPROFILER_EVENT_SCOPE_STR_ON_CHANNEL_CONDITIONAL(#Name, Channel, (Condition))
 
 // Trace a scoped CPU timing event providing a scope name (plain text) and a trace channel.
-// Example: TRACE_CPUPROFILER_EVENT_SCOPE_ON_CHANNEL(MyScopedTimer::A, CpuChannel)
 // Note: Do not use this macro with a static string because, in that case, additional quotes will
 //       be added around the event scope name.
-// Note: The event will be emitted only if both the given channel and CpuChannel is enabled.
+// Note: The event will be emitted only if both the given channel and the CpuChannel are enabled.
+// Example: TRACE_CPUPROFILER_EVENT_SCOPE_ON_CHANNEL(MyScopedTimer::A, CustomChannel)
 #define TRACE_CPUPROFILER_EVENT_SCOPE_ON_CHANNEL(Name, Channel) \
 	TRACE_CPUPROFILER_EVENT_SCOPE_ON_CHANNEL_CONDITIONAL(Name, Channel, true)
 
-// Conditionally trace a scoped CPU timing event providing a dynamic string (const ANSICHAR* or const TCHAR*)
-// as the scope name and a trace channel.
-// Example: TRACE_CPUPROFILER_EVENT_SCOPE_TEXT_ON_CHANNEL_CONDITIONAL(*MyScopedTimerNameString, CpuChannel, Condition)
-// Note: This macro has a larger overhead compared to macro that accepts a plain text name
-//       or a static string. Use it only if scope name really needs to be a dynamic string.
-#define TRACE_CPUPROFILER_EVENT_SCOPE_TEXT_ON_CHANNEL_CONDITIONAL(Name, Channel, Condition) \
-	FCpuProfilerTrace::FDynamicEventScope PREPROCESSOR_JOIN(__CpuProfilerEventScope, __LINE__)(Name, Channel, (Condition), __FILE__, __LINE__);
+////////////////////////////////////////////////////////////////////////////////////////////////////
+// TRACE_CPUPROFILER_EVENT_SCOPE_TEXT* - scope name as dynamic string
+////////////////////////////////////////////////////////////////////////////////////////////////////
 
-// Trace a scoped CPU timing event providing a dynamic string (const ANSICHAR* or const TCHAR*)
-// as the scope name and a trace channel.
-// Example: TRACE_CPUPROFILER_EVENT_SCOPE_TEXT_ON_CHANNEL(*MyScopedTimerNameString, CpuChannel)
+// Conditionally trace a scoped CPU timing event providing a dynamic string (const ANSICHAR* or
+// const TCHAR* or FName) as the scope name. It will use the CPU trace channel.
 // Note: This macro has a larger overhead compared to macro that accepts a plain text name
 //       or a static string. Use it only if scope name really needs to be a dynamic string.
-#define TRACE_CPUPROFILER_EVENT_SCOPE_TEXT_ON_CHANNEL(Name, Channel) \
-	TRACE_CPUPROFILER_EVENT_SCOPE_TEXT_ON_CHANNEL_CONDITIONAL(Name, Channel, true)
-
-// Conditionally trace a scoped CPU timing event providing a dynamic string (const ANSICHAR* or const TCHAR*)
-// as the scope name. It will use the CPU trace channel.
 // Example: TRACE_CPUPROFILER_EVENT_SCOPE_TEXT_CONDITIONAL(*MyScopedTimerNameString, Condition)
-// Note: This macro has a larger overhead compared to macro that accepts a plain text name
-//       or a static string. Use it only if scope name really needs to be a dynamic string.
 #define TRACE_CPUPROFILER_EVENT_SCOPE_TEXT_CONDITIONAL(Name, Condition) \
 	FCpuProfilerTrace::FDynamicEventScope PREPROCESSOR_JOIN(__CpuProfilerEventScope, __LINE__)(Name, (Condition), __FILE__, __LINE__);
 
-// Trace a scoped CPU timing event providing a dynamic string (const ANSICHAR* or const TCHAR*)
-// as the scope name. It will use the CPU trace channel.
-// Example: TRACE_CPUPROFILER_EVENT_SCOPE_TEXT(*MyScopedTimerNameString)
+// Trace a scoped CPU timing event providing a dynamic string (const ANSICHAR* or
+// const TCHAR* or FName) as the scope name. It will use the CPU trace channel.
 // Note: This macro has a larger overhead compared to macro that accepts a plain text name
 //       or a static string. Use it only if scope name really needs to be a dynamic string.
+// Example: TRACE_CPUPROFILER_EVENT_SCOPE_TEXT(*MyScopedTimerNameString)
 #define TRACE_CPUPROFILER_EVENT_SCOPE_TEXT(Name) \
 	TRACE_CPUPROFILER_EVENT_SCOPE_TEXT_CONDITIONAL(Name, true)
 
+// Conditionally trace a scoped CPU timing event providing a dynamic string (const ANSICHAR* or
+// const TCHAR* or FName) as the scope name and a trace channel.
+// Note: This macro has a larger overhead compared to macro that accepts a plain text name
+//       or a static string. Use it only if scope name really needs to be a dynamic string.
+// Note: The event will be emitted only if both the given channel and the CpuChannel are enabled.
+// Example: TRACE_CPUPROFILER_EVENT_SCOPE_TEXT_ON_CHANNEL_CONDITIONAL(*MyScopedTimerNameString, CustomChannel, Condition)
+#define TRACE_CPUPROFILER_EVENT_SCOPE_TEXT_ON_CHANNEL_CONDITIONAL(Name, Channel, Condition) \
+	FCpuProfilerTrace::FDynamicEventScope PREPROCESSOR_JOIN(__CpuProfilerEventScope, __LINE__)(Name, Channel, (Condition), __FILE__, __LINE__);
+
+// Trace a scoped CPU timing event providing a dynamic string (const ANSICHAR* or const TCHAR* or FName)
+// as the scope name and a trace channel.
+// Note: This macro has a larger overhead compared to macro that accepts a plain text name
+//       or a static string. Use it only if scope name really needs to be a dynamic string.
+// Note: The event will be emitted only if both the given channel and the CpuChannel are enabled.
+// Example: TRACE_CPUPROFILER_EVENT_SCOPE_TEXT_ON_CHANNEL(*MyScopedTimerNameString, CustomChannel)
+#define TRACE_CPUPROFILER_EVENT_SCOPE_TEXT_ON_CHANNEL(Name, Channel) \
+	TRACE_CPUPROFILER_EVENT_SCOPE_TEXT_ON_CHANNEL_CONDITIONAL(Name, Channel, true)
+
+////////////////////////////////////////////////////////////////////////////////////////////////////
+
 // Make sure all thread data has reached the destination.
 // Note: Can be useful to call this before entering a wait condition that might take a while.
 #define TRACE_CPUPROFILER_EVENT_FLUSH() \
@@ -424,16 +489,19 @@ private:
 #define TRACE_CPUPROFILER_EVENT_MANUAL_IS_ENABLED() false
 #define TRACE_CPUPROFILER_EVENT_SCOPE_STR(NameStr)
 #define TRACE_CPUPROFILER_EVENT_SCOPE_STR_CONDITIONAL(NameStr, Condition)
+#define TRACE_CPUPROFILER_EVENT_SCOPE_STR_ON_CHANNEL(NameStr, Channel)
+#define TRACE_CPUPROFILER_EVENT_SCOPE_STR_ON_CHANNEL_CONDITIONAL(NameStr, Channel, Condition)
 #define TRACE_CPUPROFILER_EVENT_SCOPE_ON_CHANNEL_STR(NameStr, Channel)
 #define TRACE_CPUPROFILER_EVENT_SCOPE_ON_CHANNEL_STR_CONDITIONAL(NameStr, Channel, Condition)
+#define TRACE_CPUPROFILER_EVENT_SCOPE_ON_CHANNEL_CONDITIONAL_STR(NameStr, Channel, Condition)
 #define TRACE_CPUPROFILER_EVENT_SCOPE(Name)
 #define TRACE_CPUPROFILER_EVENT_SCOPE_CONDITIONAL(Name, Condition)
 #define TRACE_CPUPROFILER_EVENT_SCOPE_ON_CHANNEL(Name, Channel)
 #define TRACE_CPUPROFILER_EVENT_SCOPE_ON_CHANNEL_CONDITIONAL(Name, Channel, Condition)
-#define TRACE_CPUPROFILER_EVENT_SCOPE_TEXT_ON_CHANNEL(Name, Channel)
-#define TRACE_CPUPROFILER_EVENT_SCOPE_TEXT_ON_CHANNEL_CONDITIONAL(Name, Channel, Condition)
 #define TRACE_CPUPROFILER_EVENT_SCOPE_TEXT(Name)
 #define TRACE_CPUPROFILER_EVENT_SCOPE_TEXT_CONDITIONAL(Name, Condition)
+#define TRACE_CPUPROFILER_EVENT_SCOPE_TEXT_ON_CHANNEL(Name, Channel)
+#define TRACE_CPUPROFILER_EVENT_SCOPE_TEXT_ON_CHANNEL_CONDITIONAL(Name, Channel, Condition)
 #define TRACE_CPUPROFILER_EVENT_FLUSH()
 
 #endif // CPUPROFILERTRACE_ENABLED
