Fixed a missing wake in FSharedRecursiveMutex

Also removed extraneous conditions in checkSlow statements and fixed whitespace that Visual Studio messed up.

#rnx

diff --git a/Engine/Source/Runtime/Core/Private/Async/SharedRecursiveMutex.cpp b/Engine/Source/Runtime/Core/Private/Async/SharedRecursiveMutex.cpp
--- a/Engine/Source/Runtime/Core/Private/Async/SharedRecursiveMutex.cpp
+++ b/Engine/Source/Runtime/Core/Private/Async/SharedRecursiveMutex.cpp
@@ -88,7 +88,7 @@ bool FSharedRecursiveMutex::TryLock()
 	{
 		if (LIKELY(State.compare_exchange_weak(CurrentState, CurrentState | (1 << LockCountShift), std::memory_order_acquire, std::memory_order_relaxed)))
 		{
-			checkSlow(ThreadId.load(std::memory_order_relaxed) == 0 && (CurrentState & LockCountMask) == 0);
+			checkSlow(ThreadId.load(std::memory_order_relaxed) == 0);
 			ThreadId.store(CurrentThreadId, std::memory_order_relaxed);
 			return true;
 		}
@@ -116,7 +116,7 @@ void FSharedRecursiveMutex::Lock()
 	{
 		if (LIKELY(State.compare_exchange_weak(CurrentState, CurrentState | (1 << LockCountShift), std::memory_order_acquire, std::memory_order_relaxed)))
 		{
-			checkSlow(ThreadId.load(std::memory_order_relaxed) == 0 && (CurrentState & LockCountMask) == 0);
+			checkSlow(ThreadId.load(std::memory_order_relaxed) == 0);
 			ThreadId.store(CurrentThreadId, std::memory_order_relaxed);
 			return;
 		}
@@ -145,7 +145,7 @@ FORCENOINLINE void FSharedRecursiveMutex::LockSlow(uint32 CurrentState, uint32 C
 		{
 			if (State.compare_exchange_weak(CurrentState, CurrentState | (1 << LockCountShift), std::memory_order_acquire, std::memory_order_relaxed))
 			{
-				checkSlow(ThreadId.load(std::memory_order_relaxed) == 0 && (CurrentState & LockCountMask) == 0);
+				checkSlow(ThreadId.load(std::memory_order_relaxed) == 0);
 				ThreadId.store(CurrentThreadId, std::memory_order_relaxed);
 				return;
 			}
@@ -193,10 +193,10 @@ void FSharedRecursiveMutex::Unlock()
 		ThreadId.store(0, std::memory_order_relaxed);
 
 		// Unlock immediately to allow other threads to acquire the lock while this thread looks for a thread to wake.
-		State.fetch_sub(1 << LockCountShift, std::memory_order_release);
+		uint32 LastState = State.fetch_sub(1 << LockCountShift, std::memory_order_release);
 
 		// Wake one exclusive waiter or every shared waiter.
-		if (UNLIKELY(CurrentState & (HasWaitingLockFlag | HasWaitingSharedLockFlag)))
+		if (UNLIKELY(LastState & (HasWaitingLockFlag | HasWaitingSharedLockFlag)))
 		{
 			WakeWaitingThreads(CurrentState);
 		}
@@ -323,13 +323,13 @@ FORCENOINLINE void FSharedRecursiveMutex::WakeWaitingThreads(uint32 CurrentState
 	{
 		// Wake one thread that is waiting to acquire an exclusive lock.
 		ParkingLot::WakeOne(GetLockAddress(), [this](ParkingLot::FWakeState WakeState) -> uint64
+		{
+			if (!WakeState.bHasWaitingThreads)
 			{
-				if (!WakeState.bHasWaitingThreads)
-				{
-					State.fetch_and(~HasWaitingLockFlag, std::memory_order_relaxed);
-				}
-				return 0;
-			});
+				State.fetch_and(~HasWaitingLockFlag, std::memory_order_relaxed);
+			}
+			return 0;
+		});
 	}
 	else if (CurrentState & HasWaitingSharedLockFlag)
 	{
