Fixed invalid calls to LowLevelOutputDebugStringf

#rnx

diff --git a/Engine/Plugins/Runtime/GeometryProcessing/Source/DynamicMesh/Private/Solvers/Internal/MeshDiffusionIntegrator.cpp b/Engine/Plugins/Runtime/GeometryProcessing/Source/DynamicMesh/Private/Solvers/Internal/MeshDiffusionIntegrator.cpp
--- a/Engine/Plugins/Runtime/GeometryProcessing/Source/DynamicMesh/Private/Solvers/Internal/MeshDiffusionIntegrator.cpp
+++ b/Engine/Plugins/Runtime/GeometryProcessing/Source/DynamicMesh/Private/Solvers/Internal/MeshDiffusionIntegrator.cpp
@@ -72,7 +72,7 @@ void FMeshDiffusionIntegrator::Initialize(const FDynamicMesh3& DynamicMesh, cons
 	ss << Eigen::MatrixXd(M) << std::endl;
 
 	FString Foo = ss.str().c_str();
-	FPlatformMisc::LowLevelOutputDebugStringf(*Foo);
+	FPlatformMisc::LowLevelOutputDebugString(*Foo);
 #endif
 
 }
diff --git a/Engine/Source/Runtime/Core/Private/Windows/WindowsPlatformNamedPipe.cpp b/Engine/Source/Runtime/Core/Private/Windows/WindowsPlatformNamedPipe.cpp
--- a/Engine/Source/Runtime/Core/Private/Windows/WindowsPlatformNamedPipe.cpp
+++ b/Engine/Source/Runtime/Core/Private/Windows/WindowsPlatformNamedPipe.cpp
@@ -20,7 +20,7 @@ static void VerifyWinResult(BOOL bResult, const TCHAR* InMessage)
 		uint32 Error = GetLastError();
 		const TCHAR* Buffer = FWindowsPlatformMisc::GetSystemErrorMessage(ErrorBuffer, MAX_SPRINTF, Error);
 		FString Message = FString::Printf(TEXT("FAILED (%s) with GetLastError() %d: %s!\n"), InMessage, Error, ErrorBuffer);
-		FPlatformMisc::LowLevelOutputDebugStringf(*Message);
+		FPlatformMisc::LowLevelOutputDebugString(*Message);
 		UE_LOG(LogWindows, Fatal, TEXT("%s"), *Message);
 	}
 	verify(bResult != 0);
diff --git a/Engine/Source/Runtime/Engine/Private/ShaderCompiler/ShaderCompiler.cpp b/Engine/Source/Runtime/Engine/Private/ShaderCompiler/ShaderCompiler.cpp
--- a/Engine/Source/Runtime/Engine/Private/ShaderCompiler/ShaderCompiler.cpp
+++ b/Engine/Source/Runtime/Engine/Private/ShaderCompiler/ShaderCompiler.cpp
@@ -2314,7 +2314,7 @@ static void BuildErrorStringAndReport(const FShaderDiagnosticInfo& DiagInfo, FSt
 		{
 			// Using OutputDebugString to avoid any text getting added before the filename,
 			// Which will throw off VS.NET's ability to take you directly to the file and line of the error when double clicking it in the output window.
-			FPlatformMisc::LowLevelOutputDebugStringf(*UniqueErrorString);
+			FPlatformMisc::LowLevelOutputDebugString(*UniqueErrorString);
 		}
 		else
 		{
