Disable call stack dumping and minidump collection for failed checks and ensures in Windows ASAN builds due to incompatibilities.

#rnx

diff --git a/Engine/Source/Runtime/Core/Private/Misc/AssertionMacros.cpp b/Engine/Source/Runtime/Core/Private/Misc/AssertionMacros.cpp
--- a/Engine/Source/Runtime/Core/Private/Misc/AssertionMacros.cpp
+++ b/Engine/Source/Runtime/Core/Private/Misc/AssertionMacros.cpp
@@ -29,6 +29,23 @@
 #include "HAL/IConsoleManager.h"
 #include "AutoRTFM/AutoRTFM.h"
 
+// Stack walking can race with ASAN's usage of DbgHelp (which is not thread safe), so we're
+// not going to report callstacks for failed ensures and checks if ASAN is enabled.
+#if PLATFORM_WINDOWS && PLATFORM_HAS_ASAN_INCLUDE
+#	if defined(__SANITIZE_ADDRESS__)
+#		define ASSERT_DISABLE_BACKTRACE_DUE_TO_ASAN 1
+#	elif defined(__has_feature)
+#		if __has_feature(address_sanitizer)
+#			define ASSERT_DISABLE_BACKTRACE_DUE_TO_ASAN 1
+#		endif
+#	endif
+#endif
+
+#ifndef ASSERT_DISABLE_BACKTRACE_DUE_TO_ASAN
+#	define ASSERT_DISABLE_BACKTRACE_DUE_TO_ASAN 0
+#endif
+
+
 namespace 
 {
 	// used to track state of assets/ensures
@@ -224,7 +241,11 @@ FORCENOINLINE void StaticFailDebugV(
 	ANSICHAR StackTrace[4096] = "";
 	if (FPlatformProperties::AllowsCallStackDumpDuringAssert() && !bIsEnsure)
 	{
+#if ASSERT_DISABLE_BACKTRACE_DUE_TO_ASAN
+		FCStringAnsi::Snprintf(StackTrace, UE_ARRAY_COUNT(StackTrace), "(callstack disabled due to ASAN)");
+#else
 		FPlatformStackWalk::StackWalkAndDump(StackTrace, UE_ARRAY_COUNT(StackTrace), ProgramCounter);
+#endif
 	}
 
 	FErrorHistWriter Writer;
@@ -468,7 +489,11 @@ FORCENOINLINE void FDebug::EnsureFailed(const ANSICHAR* Expr, const ANSICHAR* Fi
 				SCOPE_LOG_TIME_IN_SECONDS(*StackWalkPerfMessage, nullptr)
 #endif
 				StackTrace[0] = 0;
+#if ASSERT_DISABLE_BACKTRACE_DUE_TO_ASAN
+				FCStringAnsi::Snprintf(StackTrace, StackTraceSize, "(callstack disabled due to ASAN)");
+#else
 				FPlatformStackWalk::StackWalkAndDumpEx(StackTrace, StackTraceSize, ProgramCounter, FGenericPlatformStackWalk::EStackWalkFlags::FlagsUsedWhenHandlingEnsure);
+#endif
 			}
 
 			// Also append the stack trace
@@ -761,8 +786,12 @@ FORCENOINLINE void FDebug::DumpStackTraceToLog(const TCHAR* Heading, const ELogV
 #endif
 		StackTrace[0] = 0;
 
+#if ASSERT_DISABLE_BACKTRACE_DUE_TO_ASAN
+		FCStringAnsi::Snprintf(StackTrace, StackTraceSize, "(callstack disabled due to ASAN)");
+#else
 		const int32 NumStackFramesToIgnore = 1;
 		FPlatformStackWalk::StackWalkAndDumpEx(StackTrace, StackTraceSize, NumStackFramesToIgnore, FGenericPlatformStackWalk::EStackWalkFlags::FlagsUsedWhenHandlingEnsure);
+#endif
 	}
 
 	// Dump the error and flush the log.
diff --git a/Engine/Source/Runtime/Core/Private/Windows/WindowsPlatformCrashContext.cpp b/Engine/Source/Runtime/Core/Private/Windows/WindowsPlatformCrashContext.cpp
--- a/Engine/Source/Runtime/Core/Private/Windows/WindowsPlatformCrashContext.cpp
+++ b/Engine/Source/Runtime/Core/Private/Windows/WindowsPlatformCrashContext.cpp
@@ -73,6 +73,21 @@ THIRD_PARTY_INCLUDES_END
 #define CR_CLIENT_MAX_PATH_LEN 265
 #define CR_CLIENT_MAX_ARGS_LEN 256
 
+// MiniDumpWriteDump hangs in ASAN builds and stack walking can race with ASAN's usage of DbgHelp (which is not thread safe).
+#if PLATFORM_HAS_ASAN_INCLUDE
+#	if defined(__SANITIZE_ADDRESS__)
+#		define DISABLE_DEBUG_INFO_DUE_TO_ASAN 1
+#	elif defined(__has_feature)
+#		if __has_feature(address_sanitizer)
+#			define DISABLE_DEBUG_INFO_DUE_TO_ASAN 1
+#		endif
+#	endif
+#endif
+
+#ifndef DISABLE_DEBUG_INFO_DUE_TO_ASAN
+#	define DISABLE_DEBUG_INFO_DUE_TO_ASAN 0
+#endif
+
 #pragma comment( lib, "version.lib" )
 #pragma comment( lib, "Shlwapi.lib" )
 
@@ -107,6 +122,9 @@ namespace {
 	 // #CrashReport: 2014-10-08 Move to FWindowsPlatformCrashContext
 	bool WriteMinidump(HANDLE Process, DWORD ThreadId, FWindowsPlatformCrashContext& InContext, const TCHAR* Path, LPEXCEPTION_POINTERS ExceptionInfo)
 	{
+#if DISABLE_DEBUG_INFO_DUE_TO_ASAN
+		return false;
+#else
 		// Are we calling this in process or from an external process?
 		const BOOL bIsClientPointers = Process != GetCurrentProcess();
 
@@ -149,6 +167,7 @@ namespace {
 		CloseHandle(FileHandle);
 
 		return Result == TRUE;
+#endif
 	}
 	
 	/** Returns the crash timeout in seconds. */
@@ -1547,7 +1566,11 @@ private:
 				ContextWrapper = FWindowsPlatformStackWalk::MakeThreadContextWrapper(ExceptionInfo->ContextRecord, CrashingThreadHandle);
 			}
 			
+#if DISABLE_DEBUG_INFO_DUE_TO_ASAN
+			FCStringAnsi::Snprintf(StackTrace, StackTraceSize, "(callstack disabled due to ASAN)");
+#else
 			FPlatformStackWalk::StackWalkAndDump(StackTrace, StackTraceSize, ErrorProgramCounter, ContextWrapper);
+#endif
 			
 			if (ExceptionInfo->ExceptionRecord->ExceptionCode != EnsureExceptionCode && ExceptionInfo->ExceptionRecord->ExceptionCode != AssertExceptionCode)
 			{
