Restored FLLMScope constructors and moved code into the .cpp file.

diff --git a/Engine/Source/Runtime/Core/Private/HAL/LowLevelMemTracker.cpp b/Engine/Source/Runtime/Core/Private/HAL/LowLevelMemTracker.cpp
--- a/Engine/Source/Runtime/Core/Private/HAL/LowLevelMemTracker.cpp
+++ b/Engine/Source/Runtime/Core/Private/HAL/LowLevelMemTracker.cpp
@@ -3540,89 +3540,107 @@ const UE::LLMPrivate::FTagData* FLowLevelMemTracker::FindTagData(FName TagName,
 	}
 }
 
-void FLLMScope::Init(ELLMTag TagEnum, bool bInIsStatTag, ELLMTagSet InTagSet, ELLMTracker InTracker, bool bOverride)
+FLLMScope::FLLMScope(ELLMTag TagEnum, bool bInIsStatTag, ELLMTagSet InTagSet, ELLMTracker InTracker, bool bOverride)
 {
-	// This scope does not support TagSets (use FLLMScopeDynamic instead), except for the special case of
-	// anonymized allocations, which are indicated with a TagSet and ELLMTag::EngineMisc
-	LLMCheck((!bInIsStatTag && InTagSet == ELLMTagSet::None) || (TagEnum == ELLMTag::EngineMisc));
-	// ELLMTag::FMalloc is a special tag expected to only be used by the Platform Tracker (see header where defined).
-	LLMCheck((TagEnum != ELLMTag::FMalloc) || (ELLMTracker::Platform == InTracker));
-
-	FLowLevelMemTracker& LLMRef = FLowLevelMemTracker::Get();
-	// We have to check IsEnabled() again after calling Get, because the constructor is called
-	// from Get, and will set EnabledState=Disabled if the platform doesn't support it.
 	if (!FLowLevelMemTracker::IsEnabled())
 	{
-		bEnabled = false;
 		return;
 	}
-	LLMRef.BootstrapInitialise();
 
-	if (!bOverride && LLMRef.GetTracker(InTracker)->GetActiveTagData(InTagSet) != nullptr)
+	// We run the init in the open, because we want to track LLM even in transactions.
+	UE_AUTORTFM_OPEN
 	{
-		bEnabled = false;
-		return;
-	}
+		// This scope does not support TagSets (use FLLMScopeDynamic instead), except for the special case of
+		// anonymized allocations, which are indicated with a TagSet and ELLMTag::EngineMisc
+		LLMCheck((!bInIsStatTag && InTagSet == ELLMTagSet::None) || (TagEnum == ELLMTag::EngineMisc));
+		// ELLMTag::FMalloc is a special tag expected to only be used by the Platform Tracker (see header where defined).
+		LLMCheck((TagEnum != ELLMTag::FMalloc) || (ELLMTracker::Platform == InTracker));
 
-	bEnabled = true;
-	Tracker = InTracker;
-	TagSet = InTagSet;
-	LLMRef.GetTracker(Tracker)->PushTag(TagEnum, InTagSet);
+		FLowLevelMemTracker& LLMRef = FLowLevelMemTracker::Get();
+		// We have to check IsEnabled() again after calling Get, because the constructor is called
+		// from Get, and will set EnabledState=Disabled if the platform doesn't support it.
+		if (FLowLevelMemTracker::IsEnabled())
+		{
+			LLMRef.BootstrapInitialise();
+
+			if (bOverride || LLMRef.GetTracker(InTracker)->GetActiveTagData(InTagSet) == nullptr)
+			{
+				bEnabled = true;
+				Tracker = InTracker;
+				TagSet = InTagSet;
+				LLMRef.GetTracker(Tracker)->PushTag(TagEnum, InTagSet);
+			}
+		}
+	};
+
+	// But remember that if we abort while we hold the scope, we need to destroy the scope.
+	AutoRTFM::PushOnAbortHandler(this, [this] { if (bEnabled) { Destruct(); } });
 }
 
-void FLLMScope::Init(FName TagName, bool bInIsStatTag, ELLMTagSet InTagSet, ELLMTracker InTracker, bool bOverride)
+FLLMScope::FLLMScope(FName TagName, bool bInIsStatTag, ELLMTagSet InTagSet, ELLMTracker InTracker, bool bOverride)
 {
-	FLowLevelMemTracker& LLMRef = FLowLevelMemTracker::Get();
-	// We have to check IsEnabled() again after calling Get, because the constructor is called
-	// from Get, and will set EnabledStae=Disabled if the platform doesn't support it.
 	if (!FLowLevelMemTracker::IsEnabled())
 	{
-		bEnabled = false;
-		return;
-	}
-	LLMRef.BootstrapInitialise();
-	if (!LLMRef.IsTagSetActive(InTagSet))
-	{
-		bEnabled = false;
 		return;
 	}
 
-	if (!bOverride && LLMRef.GetTracker(InTracker)->GetActiveTagData(InTagSet) != nullptr)
+	// We run the init in the open, because we want to track LLM even in transactions.
+	UE_AUTORTFM_OPEN
 	{
-		bEnabled = false;
-		return;
-	}
-
-	bEnabled = true;
-	Tracker = InTracker;
-	TagSet = InTagSet;
+		FLowLevelMemTracker& LLMRef = FLowLevelMemTracker::Get();
+		// We have to check IsEnabled() again after calling Get, because the constructor is called
+		// from Get, and will set EnabledStae=Disabled if the platform doesn't support it.
+		if (FLowLevelMemTracker::IsEnabled())
+		{
+			LLMRef.BootstrapInitialise();
+			if (LLMRef.IsTagSetActive(InTagSet))
+			{
+				if (bOverride || LLMRef.GetTracker(InTracker)->GetActiveTagData(InTagSet) == nullptr)
+				{
+					bEnabled = true;
+					Tracker = InTracker;
+					TagSet = InTagSet;
+	
+					LLMRef.GetTracker(Tracker)->PushTag(TagName, bInIsStatTag, InTagSet);
+				}
+			}
+		}
+	};
 
-	LLMRef.GetTracker(Tracker)->PushTag(TagName, bInIsStatTag, InTagSet);
+	// But remember that if we abort while we hold the scope, we need to destroy the scope.
+	AutoRTFM::PushOnAbortHandler(this, [this] { if (bEnabled) { Destruct(); } });
 }
 
-void FLLMScope::Init(const UE::LLMPrivate::FTagData* TagData, bool bInIsStatTag, ELLMTagSet InTagSet,
+FLLMScope::FLLMScope(const UE::LLMPrivate::FTagData* TagData, bool bInIsStatTag, ELLMTagSet InTagSet,
 	ELLMTracker InTracker, bool bOverride)
 {
-	FLowLevelMemTracker& LLMRef = FLowLevelMemTracker::Get();
-	// We have to check IsEnabled() again after calling Get, because the constructor is called
-	// from Get, and will set EnabledState=Disabled if the platform doesn't support it.
 	if (!FLowLevelMemTracker::IsEnabled())
 	{
-		bEnabled = false;
 		return;
 	}
-	LLMRef.BootstrapInitialise();
 
-	if (!bOverride && LLMRef.GetTracker(InTracker)->GetActiveTagData(InTagSet) != nullptr)
+	// We run the init in the open, because we want to track LLM even in transactions.
+	UE_AUTORTFM_OPEN
 	{
-		bEnabled = false;
-		return;
-	}
+		FLowLevelMemTracker& LLMRef = FLowLevelMemTracker::Get();
+		// We have to check IsEnabled() again after calling Get, because the constructor is called
+		// from Get, and will set EnabledState=Disabled if the platform doesn't support it.
+		if (FLowLevelMemTracker::IsEnabled())
+		{
+			LLMRef.BootstrapInitialise();
 
-	bEnabled = true;
-	Tracker = InTracker;
-	TagSet = InTagSet;
-	LLMRef.GetTracker(Tracker)->PushTag(TagData, InTagSet);
+			if (bOverride || LLMRef.GetTracker(InTracker)->GetActiveTagData(InTagSet) == nullptr)
+			{
+				bEnabled = true;
+				Tracker = InTracker;
+				TagSet = InTagSet;
+				LLMRef.GetTracker(Tracker)->PushTag(TagData, InTagSet);
+			}
+		}
+	};
+
+	// But remember that if we abort while we hold the scope, we need to destroy the scope.
+	AutoRTFM::PushOnAbortHandler(this, [this] { if (bEnabled) { Destruct(); } });
 }
 
 void FLLMScope::Destruct()
@@ -3631,6 +3649,17 @@ void FLLMScope::Destruct()
 	LLMRef.GetTracker(Tracker)->PopTag(TagSet);
 }
 
+FLLMScope::~FLLMScope()
+{
+	if (bEnabled)
+	{
+		// We run the destroy in the open, because we want to track LLM even in transactions.
+		UE_AUTORTFM_OPEN{ Destruct(); };
+		// But remember to pop our on-abort handler (so that we don't try and double destroy).
+		AutoRTFM::PopOnAbortHandler(this);
+	}
+}
+
 void FLLMScopeDynamic::Init(ELLMTracker InTracker, ELLMTagSet InTagSet)
 {
 	FLowLevelMemTracker& LLMRef = FLowLevelMemTracker::Get();
diff --git a/Engine/Source/Runtime/Core/Public/HAL/LowLevelMemTracker.h b/Engine/Source/Runtime/Core/Public/HAL/LowLevelMemTracker.h
--- a/Engine/Source/Runtime/Core/Public/HAL/LowLevelMemTracker.h
+++ b/Engine/Source/Runtime/Core/Public/HAL/LowLevelMemTracker.h
@@ -888,33 +888,12 @@ public:
 class FLLMScope
 {
 public:
-	template<typename TagType>
-	FLLMScope(TagType Tag, bool bIsStatTag, ELLMTagSet InTagSet, ELLMTracker InTracker, bool bOverride = true)
-	{
-		if (FLowLevelMemTracker::IsEnabled())
-		{
-			// We run the init in the open, because we want to track LLM even in transactions.
-			UE_AUTORTFM_OPEN { Init(Tag, bIsStatTag, InTagSet, InTracker, bOverride); };
-			// But remember that if we abort while we hold the scope, we need to destroy the scope.
-			AutoRTFM::PushOnAbortHandler(this, [this] { if (bEnabled) { Destruct(); } });
-		}
-	}
-
-	~FLLMScope()
-	{
-		if (bEnabled)
-		{
-			// We run the destroy in the open, because we want to track LLM even in transactions.
-			UE_AUTORTFM_OPEN { Destruct(); };
-			// But remember to pop our on-abort handler (so that we don't try and double destroy).
-			AutoRTFM::PopOnAbortHandler(this);
-		}
-	}
+	CORE_API explicit FLLMScope(FName TagName, bool bIsStatTag, ELLMTagSet InTagSet, ELLMTracker InTracker, bool bOverride = true);
+	CORE_API explicit FLLMScope(ELLMTag TagEnum, bool bIsStatTag, ELLMTagSet InTagSet, ELLMTracker InTracker, bool bOverride = true);
+	CORE_API explicit FLLMScope(const UE::LLMPrivate::FTagData* TagData, bool bIsStatTag, ELLMTagSet InTagSet, ELLMTracker InTracker, bool bOverride = true);
+	CORE_API ~FLLMScope();
 
 protected:
-	CORE_API void Init(FName TagName, bool bIsStatTag, ELLMTagSet InTagSet, ELLMTracker InTracker, bool bOverride = true);
-	CORE_API void Init(ELLMTag TagEnum, bool bIsStatTag, ELLMTagSet InTagSet, ELLMTracker InTracker, bool bOverride = true);
-	CORE_API void Init(const UE::LLMPrivate::FTagData* TagData, bool bIsStatTag, ELLMTagSet InTagSet, ELLMTracker InTracker, bool bOverride = true);
 	CORE_API void Destruct();
 
 	// All fields other than bEnabled are not initialized by constructor unless bEnabled is true
