Minor optimization to how we collect MallocBinned extended stats

diff --git a/Engine/Source/Runtime/Core/Public/HAL/MallocBinnedCommon.h b/Engine/Source/Runtime/Core/Public/HAL/MallocBinnedCommon.h
--- a/Engine/Source/Runtime/Core/Public/HAL/MallocBinnedCommon.h
+++ b/Engine/Source/Runtime/Core/Public/HAL/MallocBinnedCommon.h
@@ -932,21 +932,28 @@ protected:
 		static FName WasteStats[NumSmallPools];
 		static FName TotalMemStats[NumSmallPools];
 
-		for (int32 i = 0; i < NumSmallPools; i++)
+		if (bFirstTime)
 		{
-			if (bFirstTime)
+			for (int32 i = 0; i < NumSmallPools; i++)
 			{
+				const int BinSize = Alloc.PoolIndexToBinSize(i);
+
 				TCHAR Name[64];
-				FCString::Sprintf(Name, TEXT("FragmentationBin%d"), Alloc.PoolIndexToBinSize(i));
+				FCString::Sprintf(Name, TEXT("FragmentationBin%d"), BinSize);
 				FragmentationsStats[i] = FName(Name);
 
-				FCString::Sprintf(Name, TEXT("WasteBin%d"), Alloc.PoolIndexToBinSize(i));
+				FCString::Sprintf(Name, TEXT("WasteBin%d"), BinSize);
 				WasteStats[i] = FName(Name);
 
-				FCString::Sprintf(Name, TEXT("TotalMemBin%d"), Alloc.PoolIndexToBinSize(i));
+				FCString::Sprintf(Name, TEXT("TotalMemBin%d"), BinSize);
 				TotalMemStats[i] = FName(Name);
 			}
 
+			bFirstTime = false;
+		}
+
+		for (int32 i = 0; i < NumSmallPools; i++)
+		{
 			const float Fragmentation = 1.0f - (float)Alloc.SmallPoolTables[i].TotalUsedBins / (float)Alloc.SmallPoolTables[i].TotalAllocatedBins;
 			FCsvProfiler::RecordCustomStat(FragmentationsStats[i], CSV_CATEGORY_INDEX(MallocBinned), int(Fragmentation * 100.0f), ECsvCustomStatOp::Set);
 
@@ -960,8 +967,6 @@ protected:
 		CSV_CUSTOM_STAT(MallocBinned, TotalAllocatedSmallPoolMemoryMB,   (float)AllocatedOSSmallPoolMemory.load(std::memory_order_relaxed) / (1024.0f * 1024.0f), ECsvCustomStatOp::Set);
 		CSV_CUSTOM_STAT(MallocBinned, RequestedLargeAllocsMemoryMB,		(float)AllocatedLargePoolMemory.load(std::memory_order_relaxed) / (1024.0f * 1024.0f), ECsvCustomStatOp::Set);
 		CSV_CUSTOM_STAT(MallocBinned, TotalAllocatedLargeAllocsMemoryMB, (float)AllocatedLargePoolMemoryWAlignment.load(std::memory_order_relaxed) / (1024.0f * 1024.0f), ECsvCustomStatOp::Set);
-
-		bFirstTime = false;
 #endif
 	}
 
