Add a LowLevelFatalError in FOutputDeviceFile::Serialize in order to prevent Log calls happening before FPersistentStorageManager is actually ready to use.

diff --git a/Engine/Source/Runtime/Core/Private/HAL/IPlatformFileManagedStorageWrapper.cpp b/Engine/Source/Runtime/Core/Private/HAL/IPlatformFileManagedStorageWrapper.cpp
--- a/Engine/Source/Runtime/Core/Private/HAL/IPlatformFileManagedStorageWrapper.cpp
+++ b/Engine/Source/Runtime/Core/Private/HAL/IPlatformFileManagedStorageWrapper.cpp
@@ -197,10 +197,13 @@ void FManagedStorageScopeFileLock::Lock()
 	}
 }
 
+bool FPersistentStorageManager::bCreated = false;
+
 FPersistentStorageManager& FPersistentStorageManager::Get()
 {
 	check(IsReady());
 	static FPersistentStorageManager Singleton;
+	bCreated = true;
 	return Singleton;
 }
 
diff --git a/Engine/Source/Runtime/Core/Private/Misc/OutputDeviceFile.cpp b/Engine/Source/Runtime/Core/Private/Misc/OutputDeviceFile.cpp
--- a/Engine/Source/Runtime/Core/Private/Misc/OutputDeviceFile.cpp
+++ b/Engine/Source/Runtime/Core/Private/Misc/OutputDeviceFile.cpp
@@ -13,6 +13,7 @@
 #include "HAL/PlatformTime.h"
 #include "HAL/PlatformProcess.h"
 #include "HAL/PlatformOutputDevices.h"
+#include "HAL/IPlatformFileManagedStorageWrapper.h"
 #include "CoreGlobals.h"
 #include "HAL/FileManager.h"
 #include "Misc/Parse.h"
@@ -535,6 +536,7 @@ bool FOutputDeviceFile::CreateWriter(uint32 MaxAttempts)
 void FOutputDeviceFile::Serialize( const TCHAR* Data, ELogVerbosity::Type Verbosity, const class FName& Category, const double Time )
 {
 #if ALLOW_LOG_FILE && !NO_LOGGING
+
 	if (CategoryInclusionInternal && !CategoryInclusionInternal->IncludedCategories.Contains(Category))
 	{
 		return;
@@ -545,6 +547,14 @@ void FOutputDeviceFile::Serialize( const TCHAR* Data, ELogVerbosity::Type Verbos
 	{
 		if (!AsyncWriter && !Dead)
 		{
+			if (FPersistentStorageManager::IsReady() == false)
+			{
+				// A failure of this check means that a LOG call is happening before the FPersistentStorageManager is ready
+				// Please move the LOG that is in this stack after the creation of the FPersistentStorageManager in order to avoid memory issue 
+				LowLevelFatalError(TEXT("The LOG: %s \nis being printed before FPersistentStorageManager is ready, please move the LOG after FPersistentStorageManager creation."), Data);
+				return;
+			}
+
 			// Open log file and create the worker thread.
 			if (!CreateWriter())
 			{
diff --git a/Engine/Source/Runtime/Core/Public/HAL/IPlatformFileManagedStorageWrapper.h b/Engine/Source/Runtime/Core/Public/HAL/IPlatformFileManagedStorageWrapper.h
--- a/Engine/Source/Runtime/Core/Public/HAL/IPlatformFileManagedStorageWrapper.h
+++ b/Engine/Source/Runtime/Core/Public/HAL/IPlatformFileManagedStorageWrapper.h
@@ -354,6 +354,11 @@ public:
 		return GConfig && GConfig->IsReadyForUse() && FCommandLine::IsInitialized() && !GIsCriticalError;
 	}
 
+	static bool IsCreated()
+	{
+		return bCreated;
+	}
+
 	/** Singleton access **/
 	static CORE_API FPersistentStorageManager& Get();
 
@@ -608,6 +613,7 @@ public:
 private:
 	friend class FManagedStorageScopeFileLock; // For access to Categories
 
+	static bool bCreated;
 	bool bInitialized;
 
 	// Top level of all managed storage
