Test FPersistentStorageManager::IsReady only on non Windows platforms  in order to avoid problems with Verse autotest early log

diff --git a/Engine/Source/Runtime/Core/Private/Misc/OutputDeviceFile.cpp b/Engine/Source/Runtime/Core/Private/Misc/OutputDeviceFile.cpp
--- a/Engine/Source/Runtime/Core/Private/Misc/OutputDeviceFile.cpp
+++ b/Engine/Source/Runtime/Core/Private/Misc/OutputDeviceFile.cpp
@@ -547,6 +547,7 @@ void FOutputDeviceFile::Serialize( const TCHAR* Data, ELogVerbosity::Type Verbos
 	{
 		if (!AsyncWriter && !Dead)
 		{
+#if !PLATFORM_WINDOWS
 			if (FPersistentStorageManager::IsReady() == false)
 			{
 				// A failure of this check means that a LOG call is happening before the FPersistentStorageManager is ready
@@ -554,7 +555,7 @@ void FOutputDeviceFile::Serialize( const TCHAR* Data, ELogVerbosity::Type Verbos
 				LowLevelFatalError(TEXT("The LOG: %s \nis being printed before FPersistentStorageManager is ready, please move the LOG after FPersistentStorageManager creation."), Data);
 				return;
 			}
-
+#endif
 			// Open log file and create the worker thread.
 			if (!CreateWriter())
 			{
