Expand support for Memory Mapped Files writes on more platforms.

diff --git a/Engine/Source/Runtime/Core/Private/Android/AndroidPlatformFile.cpp b/Engine/Source/Runtime/Core/Private/Android/AndroidPlatformFile.cpp
--- a/Engine/Source/Runtime/Core/Private/Android/AndroidPlatformFile.cpp
+++ b/Engine/Source/Runtime/Core/Private/Android/AndroidPlatformFile.cpp
@@ -1453,7 +1453,7 @@ public:
 		return result;
 	}
 
-	virtual FOpenMappedResult OpenMappedEx(const TCHAR* Filename, EOpenReadFlags OpenOptions = EOpenReadFlags::None) override
+	virtual FOpenMappedResult OpenMappedEx(const TCHAR* Filename, EOpenReadFlags OpenOptions = EOpenReadFlags::None, int64 MaximumSize = 0) override
 	{
 		FString LocalPath;
 		FString AssetPath;
diff --git a/Engine/Source/Runtime/Core/Private/Apple/ApplePlatformFile.cpp b/Engine/Source/Runtime/Core/Private/Apple/ApplePlatformFile.cpp
--- a/Engine/Source/Runtime/Core/Private/Apple/ApplePlatformFile.cpp
+++ b/Engine/Source/Runtime/Core/Private/Apple/ApplePlatformFile.cpp
@@ -875,7 +875,7 @@ IFileHandle* FApplePlatformFile::OpenWrite(const TCHAR* Filename, bool bAppend,
 	}
 }
 
-FOpenMappedResult FApplePlatformFile::OpenMappedEx(const TCHAR* Filename, EOpenReadFlags OpenOptions)
+FOpenMappedResult FApplePlatformFile::OpenMappedEx(const TCHAR* Filename, EOpenReadFlags OpenOptions, int64 MaximumSize)
 {
 	FString NormalizedFilename = NormalizeFilename(Filename);
 	
diff --git a/Engine/Source/Runtime/Core/Private/GenericPlatform/GenericPlatformFile.cpp b/Engine/Source/Runtime/Core/Private/GenericPlatform/GenericPlatformFile.cpp
--- a/Engine/Source/Runtime/Core/Private/GenericPlatform/GenericPlatformFile.cpp
+++ b/Engine/Source/Runtime/Core/Private/GenericPlatform/GenericPlatformFile.cpp
@@ -548,7 +548,7 @@ IMappedFileHandle* IPlatformFile::OpenMapped(const TCHAR* Filename)
 	return Result.StealValue().Release();
 }
 
-FOpenMappedResult IPlatformFile::OpenMappedEx(const TCHAR* Filename, EOpenReadFlags OpenOptions)
+FOpenMappedResult IPlatformFile::OpenMappedEx(const TCHAR* Filename, EOpenReadFlags OpenOptions, int64 MaximumSize)
 {
 	return MakeError(TEXT("OpenMappedEx is not implemented on this platform"));
 }
diff --git a/Engine/Source/Runtime/Core/Private/GenericPlatform/GenericPlatformFileTest.cpp b/Engine/Source/Runtime/Core/Private/GenericPlatform/GenericPlatformFileTest.cpp
new file mode 100644
--- /dev/null
+++ b/Engine/Source/Runtime/Core/Private/GenericPlatform/GenericPlatformFileTest.cpp
@@ -0,0 +1,131 @@
+// Copyright Epic Games, Inc. All Rights Reserved.
+
+#include "GenericPlatform/GenericPlatformFile.h"
+#include "HAL/FileManager.h"
+#include "HAL/PlatformFile.h"
+#include "HAL/PlatformFileManager.h"
+#include "Misc/Paths.h"
+#include "Misc/PathViews.h"
+#include "Misc/AutomationTest.h"
+#include "HAL/PlatformMisc.h"
+#include "Templates/UniquePtr.h"
+
+#include "Async/AsyncFileHandle.h"
+#include "Async/MappedFileHandle.h"
+
+#if WITH_DEV_AUTOMATION_TESTS
+
+#if WITH_ENGINE
+
+IMPLEMENT_SIMPLE_AUTOMATION_TEST(FMMapFileReadTest, "System.Engine.Files.MMapFileRead", EAutomationTestFlags_ApplicationContextMask | EAutomationTestFlags::ClientContext | EAutomationTestFlags::EngineFilter);
+
+bool FMMapFileReadTest::RunTest(const FString& InParameter)
+{
+	const FString TempDir = FPaths::AutomationTransientDir();
+	const FString ReadPrefix = TEXT("MMap_FileToRead");
+	const FString TxtExtension = TEXT(".txt");
+	const FString TempFileToRead = FPaths::ConvertRelativePathToFull(FPaths::CreateTempFilename(*TempDir, *ReadPrefix, *TxtExtension));
+	const FString TestDirectory = FPaths::GetPath(TempFileToRead);
+
+	// Make sure the directory exists
+	const bool bMakeTree = true;
+	UTEST_TRUE("Making directory tree", IFileManager::Get().MakeDirectory(*TestDirectory, bMakeTree));
+
+	// Create a dummy file to read later with OpenMappedEx
+	const char FileContent[] = "Temp file to read.";
+	IFileHandle* FileHandle = nullptr;
+	IPlatformFile& PlatformFile = FPlatformFileManager::Get().GetPlatformFile();
+
+	FileHandle = PlatformFile.OpenWrite(*TempFileToRead);
+
+	if (FileHandle == nullptr)
+	{
+		return false;
+	}
+
+	FileHandle->Write((uint8*)FileContent, sizeof(FileContent));
+	FileHandle->Flush();
+	delete FileHandle;
+
+	FOpenMappedResult ResMap = FPlatformFileManager::Get().GetPlatformFile().OpenMappedEx(*TempFileToRead);
+	TUniquePtr<IMappedFileHandle> Handle = ResMap.HasError() ? nullptr : ResMap.StealValue();
+
+	if (Handle)
+	{
+		IMappedFileRegion* Region = Handle->MapRegion();
+
+		int64 Size = Region->GetMappedSize();
+		const char* Data = (const char*)Region->GetMappedPtr();
+		delete Region;
+
+		return true;
+	}
+
+	// Some platforms might not implement OpenMappedEx so the test should not fail
+	if (ResMap.HasError())
+	{
+		if (ResMap.GetError().GetMessage() == TEXT("OpenMappedEx is not implemented on this platform"))
+		{
+			return true;
+		}
+	}
+
+	return false;
+}
+
+
+IMPLEMENT_SIMPLE_AUTOMATION_TEST(FMMapFileWriteTest, "System.Engine.Files.MMapFileWrite", EAutomationTestFlags_ApplicationContextMask | EAutomationTestFlags::ClientContext | EAutomationTestFlags::EngineFilter);
+
+bool FMMapFileWriteTest::RunTest(const FString& InParameter)
+{
+	const FString TempDir = FPaths::AutomationTransientDir();
+	const FString WritePrefix = TEXT("MMap_FileToWrite");
+	const FString TxtExtension = TEXT(".txt");
+	const FString TempFileToWrite = FPaths::ConvertRelativePathToFull(FPaths::CreateTempFilename(*TempDir, *WritePrefix, *TxtExtension));
+	const int64	  MapAddressSpaceSize = 1024;
+	const FString TestDirectory = FPaths::GetPath(TempFileToWrite);
+
+	// Make sure the directory exists
+	const bool bMakeTree = true;
+	UTEST_TRUE("Making directory tree", IFileManager::Get().MakeDirectory(*TestDirectory, bMakeTree))
+
+	FOpenMappedResult ResMap = FPlatformFileManager::Get().GetPlatformFile().OpenMappedEx(*TempFileToWrite, IPlatformFile::EOpenReadFlags::AllowWrite, MapAddressSpaceSize);
+	TUniquePtr<IMappedFileHandle> Handle = ResMap.HasError() ? nullptr : ResMap.StealValue();
+
+	if (Handle)
+	{
+		IMappedFileRegion* Region = Handle->MapRegion(0, MapAddressSpaceSize, EMappedFileFlags::EFileWritable);
+
+		const char FileContent[] = "Temp file to write.";
+
+		uint8* MapPtr = (uint8*)Region->GetMappedPtr();
+		uint64 MapSize = Region->GetMappedSize();
+		
+		check(sizeof(FileContent) <= MapSize);
+		FMemory::Memcpy(MapPtr, (uint8*)FileContent, sizeof(FileContent));
+		Handle->Flush();
+
+		delete Region;
+
+		return true;
+	}
+
+	// Some platforms might not implement OpenMappedEx so the test should not fail
+	if (ResMap.HasError())
+	{
+		if (ResMap.GetError().GetMessage() == TEXT("OpenMappedEx is not implemented on this platform"))
+		{
+			return true;
+		}
+	}
+
+	return false;
+}
+
+
+#endif // WITH_ENGINE
+
+#endif // WITH_DEV_AUTOMATION_TESTS
+
+
+
diff --git a/Engine/Source/Runtime/Core/Private/HAL/IPlatformFileLogWrapper.cpp b/Engine/Source/Runtime/Core/Private/HAL/IPlatformFileLogWrapper.cpp
--- a/Engine/Source/Runtime/Core/Private/HAL/IPlatformFileLogWrapper.cpp
+++ b/Engine/Source/Runtime/Core/Private/HAL/IPlatformFileLogWrapper.cpp
@@ -88,7 +88,7 @@ FLoggedFileHandle::~FLoggedFileHandle()
 	FILE_LOG(LogPlatformFile, Log, TEXT("Close %s"), *Filename);
 }
 
-FOpenMappedResult FLoggedPlatformFile::OpenMappedEx(const TCHAR* Filename, EOpenReadFlags OpenOptions)
+FOpenMappedResult FLoggedPlatformFile::OpenMappedEx(const TCHAR* Filename, EOpenReadFlags OpenOptions, int64 MaximumSize)
 {
-	return LowerLevel->OpenMappedEx(Filename, OpenOptions);
+	return LowerLevel->OpenMappedEx(Filename, OpenOptions, MaximumSize);
 }
diff --git a/Engine/Source/Runtime/Core/Private/HAL/IPlatformFileOpenLogWrapper.cpp b/Engine/Source/Runtime/Core/Private/HAL/IPlatformFileOpenLogWrapper.cpp
--- a/Engine/Source/Runtime/Core/Private/HAL/IPlatformFileOpenLogWrapper.cpp
+++ b/Engine/Source/Runtime/Core/Private/HAL/IPlatformFileOpenLogWrapper.cpp
@@ -106,9 +106,9 @@ bool FPlatformFileOpenLog::Initialize(IPlatformFile* Inner, const TCHAR* Command
 	return true;
 }
 
-FOpenMappedResult FPlatformFileOpenLog::OpenMappedEx(const TCHAR* Filename, EOpenReadFlags OpenOptions)
+FOpenMappedResult FPlatformFileOpenLog::OpenMappedEx(const TCHAR* Filename, EOpenReadFlags OpenOptions, int64 MaximumSize)
 {
-	return LowerLevel->OpenMappedEx(Filename, OpenOptions);
+	return LowerLevel->OpenMappedEx(Filename, OpenOptions, MaximumSize);
 }
 
 void FPlatformFileOpenLog::AddLabelInternal(const TCHAR* Fmt, ...)
diff --git a/Engine/Source/Runtime/Core/Private/Unix/UnixPlatformFile.cpp b/Engine/Source/Runtime/Core/Private/Unix/UnixPlatformFile.cpp
--- a/Engine/Source/Runtime/Core/Private/Unix/UnixPlatformFile.cpp
+++ b/Engine/Source/Runtime/Core/Private/Unix/UnixPlatformFile.cpp
@@ -1244,7 +1244,7 @@ IFileHandle* FUnixPlatformFile::OpenWrite(const TCHAR* Filename, bool bAppend, b
 	return nullptr;
 }
 
-FOpenMappedResult FUnixPlatformFile::OpenMappedEx(const TCHAR* Filename, EOpenReadFlags OpenOptions)
+FOpenMappedResult FUnixPlatformFile::OpenMappedEx(const TCHAR* Filename, EOpenReadFlags OpenOptions, int64 MaximumSize)
 {
 	const FString NormalizedFilename = NormalizeFilename(Filename, false);
 
diff --git a/Engine/Source/Runtime/Core/Private/Windows/WindowsPlatformFile.cpp b/Engine/Source/Runtime/Core/Private/Windows/WindowsPlatformFile.cpp
--- a/Engine/Source/Runtime/Core/Private/Windows/WindowsPlatformFile.cpp
+++ b/Engine/Source/Runtime/Core/Private/Windows/WindowsPlatformFile.cpp
@@ -1510,20 +1510,30 @@ IFileHandle* FWindowsPlatformFile::OpenWrite(const TCHAR* Filename, bool bAppend
 	}
 }
 
-FOpenMappedResult FWindowsPlatformFile::OpenMappedEx(const TCHAR* Filename, EOpenReadFlags OpenOptions)
+FOpenMappedResult FWindowsPlatformFile::OpenMappedEx(const TCHAR* Filename, EOpenReadFlags OpenOptions, int64 MaximumSize)
 {
 	int64 Size = FileSize(Filename);
-	if (Size < 1)
+	uint32  Access = GENERIC_READ;
+	uint32  WinFlags = FILE_SHARE_READ;
+	uint32  Create = OPEN_EXISTING;
+	DWORD	Protect = PAGE_READONLY;
+
+	if (EnumHasAnyFlags(OpenOptions, EOpenReadFlags::AllowWrite))
 	{
-		return MakeError(TEXT("File size is 0"));
+		Access = GENERIC_READ | GENERIC_WRITE;
+		WinFlags = FILE_SHARE_READ | FILE_SHARE_WRITE;
+		Create = CREATE_ALWAYS;
+		Protect = PAGE_READWRITE;
+		Size = MaximumSize;
 	}
-	uint32 Access = GENERIC_READ;
-	if (EnumHasAnyFlags(OpenOptions, EOpenReadFlags::AllowWrite))
+	else
 	{
-		Access |= GENERIC_WRITE;
+		if (Size < 1)
+		{
+			return MakeError(TEXT("File size is 0"));
+		}
 	}
-	uint32 WinFlags = FILE_SHARE_READ;
-	uint32 Create = OPEN_EXISTING;
+
 	FNormalizedFilename NormalizedFilename = FNormalizedFilename(Filename);
 
 	TRACE_PLATFORMFILE_BEGIN_OPEN(Filename);
@@ -1537,7 +1547,11 @@ FOpenMappedResult FWindowsPlatformFile::OpenMappedEx(const TCHAR* Filename, EOpe
 		TRACE_PLATFORMFILE_FAIL_OPEN(Filename);
 		return MakeError(FString::Printf(TEXT("Failed to open file '%s'"), *NormalizedFilename));
 	}
-	HANDLE MappingHandle = CreateFileMapping(Handle, NULL, EnumHasAnyFlags(OpenOptions, EOpenReadFlags::AllowWrite) ? PAGE_READWRITE : PAGE_READONLY, 0, 0, NULL);
+	if (MaximumSize == 0)
+	{
+		MaximumSize = Size;
+	}
+	HANDLE MappingHandle = CreateFileMapping(Handle, NULL, Protect, (DWORD)(MaximumSize >> 32), (DWORD)(MaximumSize), NULL);
 	if (MappingHandle == NULL)
 	{
 		TRACE_PLATFORMFILE_BEGIN_CLOSE(Handle);
diff --git a/Engine/Source/Runtime/Core/Private/Windows/WindowsPlatformFilePrivate.h b/Engine/Source/Runtime/Core/Private/Windows/WindowsPlatformFilePrivate.h
--- a/Engine/Source/Runtime/Core/Private/Windows/WindowsPlatformFilePrivate.h
+++ b/Engine/Source/Runtime/Core/Private/Windows/WindowsPlatformFilePrivate.h
@@ -68,7 +68,7 @@ public:
 	virtual IFileHandle* OpenRead(const TCHAR* Filename, bool bAllowWrite = false) override;
 	virtual IFileHandle* OpenReadNoBuffering(const TCHAR* Filename, bool bAllowWrite = false) override;
 	virtual IFileHandle* OpenWrite(const TCHAR* Filename, bool bAppend = false, bool bAllowRead = false) override;
-	virtual FOpenMappedResult OpenMappedEx(const TCHAR* Filename, EOpenReadFlags OpenOptions = EOpenReadFlags::None) override;
+	virtual FOpenMappedResult OpenMappedEx(const TCHAR* Filename, EOpenReadFlags OpenOptions = EOpenReadFlags::None, int64 MaximumSize = 0) override;
 	virtual bool DirectoryExists(const TCHAR* Directory) override;
 	virtual bool CreateDirectory(const TCHAR* Directory) override;
 	virtual bool DeleteDirectory(const TCHAR* Directory) override;
diff --git a/Engine/Source/Runtime/Core/Public/Apple/ApplePlatformFile.h b/Engine/Source/Runtime/Core/Public/Apple/ApplePlatformFile.h
--- a/Engine/Source/Runtime/Core/Public/Apple/ApplePlatformFile.h
+++ b/Engine/Source/Runtime/Core/Public/Apple/ApplePlatformFile.h
@@ -41,7 +41,7 @@ public:
 
 	virtual IFileHandle* OpenRead(const TCHAR* Filename, bool bAllowWrite = false) override;
 	virtual IFileHandle* OpenWrite(const TCHAR* Filename, bool bAppend = false, bool bAllowRead = false) override;
-	virtual FOpenMappedResult OpenMappedEx(const TCHAR* Filename, EOpenReadFlags OpenOptions = EOpenReadFlags::None) override;
+	virtual FOpenMappedResult OpenMappedEx(const TCHAR* Filename, EOpenReadFlags OpenOptions = EOpenReadFlags::None, int64 MaximumSize = 0) override;
 	virtual bool DirectoryExists(const TCHAR* Directory) override;
 	virtual bool CreateDirectory(const TCHAR* Directory) override;
 	virtual bool DeleteDirectory(const TCHAR* Directory) override;
diff --git a/Engine/Source/Runtime/Core/Public/Async/MappedFileHandle.h b/Engine/Source/Runtime/Core/Public/Async/MappedFileHandle.h
--- a/Engine/Source/Runtime/Core/Public/Async/MappedFileHandle.h
+++ b/Engine/Source/Runtime/Core/Public/Async/MappedFileHandle.h
@@ -76,6 +76,17 @@ public:
 	{
 	}
 
+	/**
+	* Synchronously flush part or all of the mapped region.
+	* This is only a hint, some platforms might ignore it.
+	* There are not guarantees how long this data will stay in paged out.
+	* @param FlushOffset		Offset into this region to flush
+	* @param BytesToFlush		number of bytes to flush. This is clamped to the size of the mapped region
+	**/
+	virtual void Flush(int64 FlushOffset = 0, int64 BytesToFlush = MAX_int64)
+	{
+	}
+
 	// Non-copyable
 	IMappedFileRegion(const IMappedFileRegion&) = delete;
 	IMappedFileRegion& operator=(const IMappedFileRegion&) = delete;
@@ -133,6 +144,10 @@ public:
 	**/
 	virtual IMappedFileRegion* MapRegion(int64 Offset = 0, int64 BytesToMap = MAX_int64, FFileMappingFlags Flags = EMappedFileFlags::ENone) = 0;
 
+	virtual void Flush(void)
+	{
+	}
+
 	// Non-copyable
 	IMappedFileHandle(const IMappedFileHandle&) = delete;
 	IMappedFileHandle& operator=(const IMappedFileHandle&) = delete;
diff --git a/Engine/Source/Runtime/Core/Public/GenericPlatform/GenericPlatformFile.h b/Engine/Source/Runtime/Core/Public/GenericPlatform/GenericPlatformFile.h
--- a/Engine/Source/Runtime/Core/Public/GenericPlatform/GenericPlatformFile.h
+++ b/Engine/Source/Runtime/Core/Public/GenericPlatform/GenericPlatformFile.h
@@ -754,7 +754,7 @@ public:
 	* Null can be returned for many reasons even if the file exists. Perhaps this platform does not support mapped files, or this file is compressed in a pak file.
 	* Generally you attempt to open mapped, and if that fails, then use other file operations instead.
 	*/
-	CORE_API virtual FOpenMappedResult OpenMappedEx(const TCHAR* Filename, EOpenReadFlags OpenOptions = EOpenReadFlags::None);
+	CORE_API virtual FOpenMappedResult OpenMappedEx(const TCHAR* Filename, EOpenReadFlags OpenOptions = EOpenReadFlags::None, int64 MaximumSize = 0);
 
 	CORE_API virtual void GetTimeStampPair(const TCHAR* PathA, const TCHAR* PathB, FDateTime& OutTimeStampA, FDateTime& OutTimeStampB);
 
diff --git a/Engine/Source/Runtime/Core/Public/HAL/IPlatformFileCachedWrapper.h b/Engine/Source/Runtime/Core/Public/HAL/IPlatformFileCachedWrapper.h
--- a/Engine/Source/Runtime/Core/Public/HAL/IPlatformFileCachedWrapper.h
+++ b/Engine/Source/Runtime/Core/Public/HAL/IPlatformFileCachedWrapper.h
@@ -456,9 +456,9 @@ public:
 	{
 		return LowerLevel->OpenAsyncRead(Filename, bAllowWrite);
 	}
-	virtual FOpenMappedResult OpenMappedEx(const TCHAR* Filename, EOpenReadFlags OpenOptions = EOpenReadFlags::None) override
+	virtual FOpenMappedResult OpenMappedEx(const TCHAR* Filename, EOpenReadFlags OpenOptions = EOpenReadFlags::None, int64 MaximumSize = 0) override
 	{
-		return LowerLevel->OpenMappedEx(Filename, OpenOptions);
+		return LowerLevel->OpenMappedEx(Filename, OpenOptions, MaximumSize);
 	}
 	virtual void SetAsyncMinimumPriority(EAsyncIOPriorityAndFlags MinPriority) override
 	{
diff --git a/Engine/Source/Runtime/Core/Public/HAL/IPlatformFileLogWrapper.h b/Engine/Source/Runtime/Core/Public/HAL/IPlatformFileLogWrapper.h
--- a/Engine/Source/Runtime/Core/Public/HAL/IPlatformFileLogWrapper.h
+++ b/Engine/Source/Runtime/Core/Public/HAL/IPlatformFileLogWrapper.h
@@ -503,7 +503,7 @@ public:
 		//@todo no wrapped logging for async file handles (yet)
 		return Result;
 	}
-	virtual FOpenMappedResult OpenMappedEx(const TCHAR* Filename, EOpenReadFlags OpenOptions = EOpenReadFlags::None) override;
+	virtual FOpenMappedResult OpenMappedEx(const TCHAR* Filename, EOpenReadFlags OpenOptions = EOpenReadFlags::None, int64 MaximumSize = 0) override;
 
 	virtual void SetAsyncMinimumPriority(EAsyncIOPriorityAndFlags Priority) override
 	{
diff --git a/Engine/Source/Runtime/Core/Public/HAL/IPlatformFileOpenLogWrapper.h b/Engine/Source/Runtime/Core/Public/HAL/IPlatformFileOpenLogWrapper.h
--- a/Engine/Source/Runtime/Core/Public/HAL/IPlatformFileOpenLogWrapper.h
+++ b/Engine/Source/Runtime/Core/Public/HAL/IPlatformFileOpenLogWrapper.h
@@ -223,7 +223,7 @@ public:
 		// we must not record the "open" here...what matters is when we start reading the file!
 		return new FLoggingAsyncReadFileHandle(this, Filename, LowerLevel->OpenAsyncRead(Filename, bAllowWrite));
 	}
-	virtual FOpenMappedResult OpenMappedEx(const TCHAR* Filename, EOpenReadFlags OpenOptions = EOpenReadFlags::None) override;
+	virtual FOpenMappedResult OpenMappedEx(const TCHAR* Filename, EOpenReadFlags OpenOptions = EOpenReadFlags::None, int64 MaximumSize = 0) override;
 
 	void AddToOpenLog(const TCHAR* Filename)
 	{
diff --git a/Engine/Source/Runtime/Core/Public/Unix/UnixPlatformFile.h b/Engine/Source/Runtime/Core/Public/Unix/UnixPlatformFile.h
--- a/Engine/Source/Runtime/Core/Public/Unix/UnixPlatformFile.h
+++ b/Engine/Source/Runtime/Core/Public/Unix/UnixPlatformFile.h
@@ -46,7 +46,7 @@ public:
 
 	CORE_API virtual IFileHandle* OpenRead(const TCHAR* Filename, bool bAllowWrite = false) override;
 	CORE_API virtual IFileHandle* OpenWrite(const TCHAR* Filename, bool bAppend = false, bool bAllowRead = false) override;
-	CORE_API virtual FOpenMappedResult OpenMappedEx(const TCHAR* Filename, EOpenReadFlags OpenOptions = EOpenReadFlags::None) override;
+	CORE_API virtual FOpenMappedResult OpenMappedEx(const TCHAR* Filename, EOpenReadFlags OpenOptions = EOpenReadFlags::None, int64 MaximumSize = 0) override;
 
 	CORE_API virtual bool DirectoryExists(const TCHAR* Directory) override;
 	CORE_API virtual bool CreateDirectory(const TCHAR* Directory) override;
diff --git a/Engine/Source/Runtime/PakFile/Private/IPlatformFilePak.cpp b/Engine/Source/Runtime/PakFile/Private/IPlatformFilePak.cpp
--- a/Engine/Source/Runtime/PakFile/Private/IPlatformFilePak.cpp
+++ b/Engine/Source/Runtime/PakFile/Private/IPlatformFilePak.cpp
@@ -4624,7 +4624,7 @@ bool FPakPlatformFile::IsPakFileInstalled(const FString& InFilename)
 }
 #endif //HAS_PLATFORM_PAK_INSTALL_CHECK
 
-FOpenMappedResult FPakPlatformFile::OpenMappedEx(const TCHAR* Filename, EOpenReadFlags OpenOptions)
+FOpenMappedResult FPakPlatformFile::OpenMappedEx(const TCHAR* Filename, EOpenReadFlags OpenOptions, int64 MaximumSize)
 {
 	if (!GMMIO_Enable)
 	{
@@ -4634,7 +4634,7 @@ FOpenMappedResult FPakPlatformFile::OpenMappedEx(const TCHAR* Filename, EOpenRea
 #if !UE_BUILD_SHIPPING
 	if (bLookLooseFirst && IsNonPakFilenameAllowed(Filename))
 	{
-		FOpenMappedResult Result = LowerLevel->OpenMappedEx(Filename, OpenOptions);
+		FOpenMappedResult Result = LowerLevel->OpenMappedEx(Filename, OpenOptions, MaximumSize);
 		if (!Result.HasError())
 		{
 			return MakeValue(Result.StealValue());
@@ -4654,7 +4654,7 @@ FOpenMappedResult FPakPlatformFile::OpenMappedEx(const TCHAR* Filename, EOpenRea
 		FScopeLock Lock(&PakEntry->MappedFileHandleCriticalSection);
 		if (!PakEntry->MappedFileHandle)
 		{
-			FOpenMappedResult Result = LowerLevel->OpenMappedEx(*PakEntry->GetFilename(), OpenOptions);
+			FOpenMappedResult Result = LowerLevel->OpenMappedEx(*PakEntry->GetFilename(), OpenOptions, MaximumSize);
 			if (Result.HasError())
 			{
 				return MakeError(Result.StealError());
@@ -4667,7 +4667,7 @@ FOpenMappedResult FPakPlatformFile::OpenMappedEx(const TCHAR* Filename, EOpenRea
 	}
 	if (IsNonPakFilenameAllowed(Filename))
 	{
-		return LowerLevel->OpenMappedEx(Filename, OpenOptions);
+		return LowerLevel->OpenMappedEx(Filename, OpenOptions, MaximumSize);
 	}
 	return MakeError(FString::Printf(TEXT("Failed to open mapped file '%s'"), Filename));
 }
diff --git a/Engine/Source/Runtime/PakFile/Public/IPlatformFilePak.h b/Engine/Source/Runtime/PakFile/Public/IPlatformFilePak.h
--- a/Engine/Source/Runtime/PakFile/Public/IPlatformFilePak.h
+++ b/Engine/Source/Runtime/PakFile/Public/IPlatformFilePak.h
@@ -2547,7 +2547,7 @@ public:
 	PAKFILE_API virtual IAsyncReadFileHandle* OpenAsyncRead(const TCHAR* Filename, bool bAllowWrite = false) override;
 	PAKFILE_API virtual void SetAsyncMinimumPriority(EAsyncIOPriorityAndFlags Priority) override;
 
-	virtual FOpenMappedResult OpenMappedEx(const TCHAR* Filename, EOpenReadFlags OpenOptions = EOpenReadFlags::None) override;
+	virtual FOpenMappedResult OpenMappedEx(const TCHAR* Filename, EOpenReadFlags OpenOptions = EOpenReadFlags::None, int64 MaximumSize = 0) override;
 	/**
 	 * Converts a filename to a path inside pak file.
 	 *
diff --git a/Engine/Source/Runtime/SandboxFile/Private/IPlatformFileSandboxWrapper.cpp b/Engine/Source/Runtime/SandboxFile/Private/IPlatformFileSandboxWrapper.cpp
--- a/Engine/Source/Runtime/SandboxFile/Private/IPlatformFileSandboxWrapper.cpp
+++ b/Engine/Source/Runtime/SandboxFile/Private/IPlatformFileSandboxWrapper.cpp
@@ -996,14 +996,14 @@ IMappedFileHandle* FSandboxPlatformFile::OpenMapped(const TCHAR* Filename)
 	return Result.StealValue().Release();
 }
 
-FOpenMappedResult FSandboxPlatformFile::OpenMappedEx(const TCHAR* Filename, EOpenReadFlags OpenOptions)
+FOpenMappedResult FSandboxPlatformFile::OpenMappedEx(const TCHAR* Filename, EOpenReadFlags OpenOptions, int64 MaximumSize)
 {
 	FString UserFilename(*ConvertToSandboxPath(Filename));
 	if (!OkForInnerAccess(Filename) || LowerLevel->FileExists(*UserFilename))
 	{
-		return LowerLevel->OpenMappedEx(*UserFilename, OpenOptions);
+		return LowerLevel->OpenMappedEx(*UserFilename, OpenOptions, MaximumSize);
 	}
-	return LowerLevel->OpenMappedEx(Filename, OpenOptions);
+	return LowerLevel->OpenMappedEx(Filename, OpenOptions, MaximumSize);
 }
 
 /**
diff --git a/Engine/Source/Runtime/SandboxFile/Public/IPlatformFileSandboxWrapper.h b/Engine/Source/Runtime/SandboxFile/Public/IPlatformFileSandboxWrapper.h
--- a/Engine/Source/Runtime/SandboxFile/Public/IPlatformFileSandboxWrapper.h
+++ b/Engine/Source/Runtime/SandboxFile/Public/IPlatformFileSandboxWrapper.h
@@ -226,7 +226,7 @@ public:
 
 	UE_DEPRECATED(5.6, "Please use overloaded value that returns error values properly")
 	SANDBOXFILE_API virtual IMappedFileHandle*		OpenMapped(const TCHAR* Filename) override;
-	SANDBOXFILE_API virtual FOpenMappedResult		OpenMappedEx(const TCHAR* Filename, EOpenReadFlags OpenOptions = EOpenReadFlags::None) override;
+	SANDBOXFILE_API virtual FOpenMappedResult		OpenMappedEx(const TCHAR* Filename, EOpenReadFlags OpenOptions = EOpenReadFlags::None, int64 MaximumSize = 0) override;
 
 	friend class FSandboxVisitor;
 	friend class FSandboxStatVisitor;
diff --git a/Engine/Source/Runtime/StorageServerClient/Private/StorageServerPlatformFile.cpp b/Engine/Source/Runtime/StorageServerClient/Private/StorageServerPlatformFile.cpp
--- a/Engine/Source/Runtime/StorageServerClient/Private/StorageServerPlatformFile.cpp
+++ b/Engine/Source/Runtime/StorageServerClient/Private/StorageServerPlatformFile.cpp
@@ -1320,11 +1320,11 @@ bool FStorageServerPlatformFile::IterateDirectoryStat(const TCHAR* Directory, FD
 	return bResult;
 }
 
-FOpenMappedResult FStorageServerPlatformFile::OpenMappedEx(const TCHAR* Filename, EOpenReadFlags OpenOptions)
+FOpenMappedResult FStorageServerPlatformFile::OpenMappedEx(const TCHAR* Filename, EOpenReadFlags OpenOptions, int64 MaximumSize)
 {
 	if (LowerLevel && IsNonServerFilenameAllowed(Filename))
 	{
-		return LowerLevel->OpenMappedEx(Filename, OpenOptions);
+		return LowerLevel->OpenMappedEx(Filename, OpenOptions, MaximumSize);
 	}
 	return MakeError(FString::Printf(TEXT("Can't open mapped file '%s'"), Filename));
 }
diff --git a/Engine/Source/Runtime/StorageServerClient/Private/StorageServerPlatformFile.h b/Engine/Source/Runtime/StorageServerClient/Private/StorageServerPlatformFile.h
--- a/Engine/Source/Runtime/StorageServerClient/Private/StorageServerPlatformFile.h
+++ b/Engine/Source/Runtime/StorageServerClient/Private/StorageServerPlatformFile.h
@@ -97,7 +97,7 @@ public:
 	virtual bool IterateDirectory(const TCHAR* Directory, FDirectoryVisitor& Visitor) override;
 	virtual bool IterateDirectoryRecursively(const TCHAR* Directory, FDirectoryVisitor& Visitor) override;
 	virtual bool IterateDirectoryStat(const TCHAR* Directory, FDirectoryStatVisitor& Visitor) override;
-	virtual FOpenMappedResult OpenMappedEx(const TCHAR* Filename, EOpenReadFlags OpenOptions = EOpenReadFlags::None) override;
+	virtual FOpenMappedResult OpenMappedEx(const TCHAR* Filename, EOpenReadFlags OpenOptions = EOpenReadFlags::None, int64 MaximumSize = 0) override;
 	virtual FString GetFilenameOnDisk(const TCHAR* Filename) override;
 	virtual bool DeleteFile(const TCHAR* Filename) override;
 	virtual bool MoveFile(const TCHAR* To, const TCHAR* From) override;
