Fix CSV profiler attributing GPU stats to the wrong frames
- Custom stats reported to the CSV profiler are assigned to a frame based on their timestamp. The CSV profiler maintains two timelines (game and render thread), neither of which are correct for GPU stats, which arrive on an end-of-pipe worker thread after GPU completion.
- Added an EndOfPipe timeline, and exposed some functions for the GPU profiler to use, so that stats are counted against the correct frame.
- Since RHIEndFrame is called more frequently than the actual game thread ticks (e.g. during loading screens), we also need to ignore frame boundaries from the GPU profiler if we're running ahead of the game thread. This is handled with the EOPCounter atomic.

diff --git a/Engine/Source/Runtime/Core/Private/ProfilingDebugging/CsvProfiler.cpp b/Engine/Source/Runtime/Core/Private/ProfilingDebugging/CsvProfiler.cpp
--- a/Engine/Source/Runtime/Core/Private/ProfilingDebugging/CsvProfiler.cpp
+++ b/Engine/Source/Runtime/Core/Private/ProfilingDebugging/CsvProfiler.cpp
@@ -1233,6 +1233,7 @@ namespace ECsvTimeline
 	{
 		Gamethread,
 		Renderthread,
+		EndOfPipe,
 		Count
 	};
 }
@@ -1257,6 +1258,7 @@ public:
 			FrameBoundaryTimestamps[i].Empty();
 		}
 		CurrentReadFrameIndex = 0;
+		EOPCounter = 0;
 	}
 
 	int32 GetFrameNumberForTimestamp(ECsvTimeline::Type Timeline, uint64 Timestamp) const
@@ -1334,6 +1336,26 @@ public:
 			}
 		}
 #endif
+		//
+		// The EndOfPipe frame boundary happens more frequently than the game or render thread frame boundaries.
+		// This is because EndOfPipe happens on each RHICmdList.EndFrame() call, which includes the loading screen / video player when the game thread is not ticking.
+		// Use a counter to ignore any EndOfPipe boundaries if we haven't had a prior render thread boundary. Not doing this means the timelines get out of sync.
+		//
+		if (Timeline == ECsvTimeline::Renderthread)
+		{
+			EOPCounter++;
+		}
+		else if (Timeline == ECsvTimeline::EndOfPipe)
+		{
+			int32 Counter = EOPCounter--;
+			if (Counter == 0)
+			{
+				// Skip this frame boundary.
+				EOPCounter++;
+				return;
+			}
+		}
+
 		uint64* Element = FrameBoundaryTimestampsWriteBuffer[Timeline].ReserveElement();
 		*Element = FPlatformTime::Cycles64();
 		FrameBoundaryTimestampsWriteBuffer[Timeline].CommitElement();
@@ -1359,6 +1381,8 @@ private:
 	TSingleProducerSingleConsumerList<uint64, 16> FrameBoundaryTimestampsWriteBuffer[ECsvTimeline::Count];
 	TArray<uint64> FrameBoundaryTimestamps[ECsvTimeline::Count];
 	mutable int32 CurrentReadFrameIndex;
+
+	std::atomic<int32> EOPCounter { 0 };
 };
 static FFrameBoundaries GFrameBoundaries;
 
@@ -2271,6 +2295,8 @@ private:
 	static TArray<FWeakPtr> TlsInstances;
 	static uint32 TlsSlot;
 
+	static FSharedPtr EndOfPipeData;
+
 public:
 	static void InitTls()
 	{
@@ -2314,11 +2340,16 @@ public:
 		return *ProfilerThread;
 	}
 
+	static CSV_PROFILER_INLINE FCsvProfilerThreadData& GetEndOfPipe()
+	{
+		return *EndOfPipeData.Get();
+	}
+
 	static inline void GetTlsInstances(TArray<FSharedPtr>& OutTlsInstances)
 	{
 		QUICK_SCOPE_CYCLE_COUNTER(CSVProfiler_ThreadData_GetTlsInstances);
 		FScopeLock Lock(&TlsCS);
-		OutTlsInstances.Empty(TlsInstances.Num());
+		OutTlsInstances.Empty(TlsInstances.Num() + 1);
 
 		for (int32 Index = TlsInstances.Num() - 1; Index >= 0; --Index)
 		{
@@ -2329,12 +2360,16 @@ public:
 				OutTlsInstances.Emplace(MoveTemp(SharedPtr));
 			}
 		}
+
+		// Always add the end-of-pipe data
+		OutTlsInstances.Add(EndOfPipeData);
 	}
 
-	FCsvProfilerThreadData()
+	FCsvProfilerThreadData(TOptional<ECsvTimeline::Type> InCsvTimeline = NullOpt)
 		: ThreadId(FPlatformTLS::GetCurrentThreadId())
 		, ThreadName(FThreadManager::GetThreadName(ThreadId))
 		, bIsTaskWorkerThread(LowLevelTasks::FScheduler::Get().IsWorkerThread())
+		, CsvTimeline(InCsvTimeline)
 		, DataProcessor(nullptr)
 	{
 	}
@@ -2522,11 +2557,14 @@ public:
 	const uint32 ThreadId;
 	const FString ThreadName;
 	const bool bIsTaskWorkerThread;
+	TOptional<ECsvTimeline::Type> CsvTimeline;
 
 	class FCsvProfilerThreadDataProcessor* DataProcessor;
 	TArray<FCsvWaitStatName> WaitStatNameStack;
 };
 
+FCsvProfilerThreadData::FSharedPtr FCsvProfilerThreadData::EndOfPipeData(MakeShared<FCsvProfilerThreadData>(ECsvTimeline::EndOfPipe));
+
 uint32 FCsvProfilerThreadData::TlsSlot = FPlatformTLS::InvalidTlsSlot;
 FCriticalSection FCsvProfilerThreadData::TlsCS;
 TArray<FCsvProfilerThreadData::FWeakPtr> FCsvProfilerThreadData::TlsInstances;
@@ -3104,7 +3142,7 @@ void FCsvProfilerThreadDataProcessor::Process(FCsvProcessThreadDataStats& OutSta
 	// Flush the frame boundaries after the stat data. This way, we ensure the frame boundary data is up to date
 	// (we do not want to encounter markers from a frame which hasn't been registered yet)
 	FPlatformMisc::MemoryBarrier();
-	ECsvTimeline::Type Timeline = (ThreadData->ThreadId == RenderThreadId || ThreadData->ThreadId == RHIThreadId) ? ECsvTimeline::Renderthread : ECsvTimeline::Gamethread;
+	ECsvTimeline::Type Timeline = ThreadData->CsvTimeline.IsSet() ? *ThreadData->CsvTimeline : ((ThreadData->ThreadId == RenderThreadId || ThreadData->ThreadId == RHIThreadId) ? ECsvTimeline::Renderthread : ECsvTimeline::Gamethread);
 
 	if (ThreadMarkers.Num() > 0)
 	{
@@ -3832,6 +3870,14 @@ void FCsvProfiler::EndFrameRT()
 	}
 }
 
+void FCsvProfiler::BeginFrameEOP()
+{
+	if (GCsvProfilerIsCapturing)
+	{
+		GFrameBoundaries.AddBeginFrameTimestamp(ECsvTimeline::EndOfPipe);
+	}
+}
+
 void FCsvProfiler::BeginCapture(int InNumFramesToCapture, 
 	const FString& InDestinationFolder, 
 	const FString& InFilename,
@@ -4367,6 +4413,17 @@ void FCsvProfiler::RecordCustomStat(const FName& StatName, uint32 CategoryIndex,
 #endif
 }
 
+void FCsvProfiler::RecordEndOfPipeCustomStat(const FName& StatName, uint32 CategoryIndex, double Value, const ECsvCustomStatOp CustomStatOp)
+{
+#if !CSV_PROFILER_MINIMAL
+	if (GCsvProfilerIsCapturing && GCsvCategoriesEnabled[CategoryIndex])
+	{
+		ValidateFName(StatName);
+		FCsvProfilerThreadData::GetEndOfPipe().AddCustomStat(StatName, CategoryIndex, (float)Value, CustomStatOp);
+	}
+#endif
+}
+
 void FCsvProfiler::Init()
 {
 #if CSV_PROFILER_ALLOW_DEBUG_FEATURES
diff --git a/Engine/Source/Runtime/Core/Public/ProfilingDebugging/CsvProfiler.h b/Engine/Source/Runtime/Core/Public/ProfilingDebugging/CsvProfiler.h
--- a/Engine/Source/Runtime/Core/Public/ProfilingDebugging/CsvProfiler.h
+++ b/Engine/Source/Runtime/Core/Public/ProfilingDebugging/CsvProfiler.h
@@ -381,6 +381,8 @@ public:
 	CORE_API static void RecordCustomStat(const char * StatName, uint32 CategoryIndex, int32 Value, const ECsvCustomStatOp CustomStatOp);
 	CORE_API static void RecordCustomStat(const FName& StatName, uint32 CategoryIndex, int32 Value, const ECsvCustomStatOp CustomStatOp);
 
+	CORE_API static void RecordEndOfPipeCustomStat(const FName& StatName, uint32 CategoryIndex, double Value, const ECsvCustomStatOp CustomStatOp);
+
 	// Warning: Don't use this unless the stat is one of a handful of key stats that we want to capture in CSV_PROFILER_MINIMAL mode
 	CORE_API static void RecordCustomStatMinimal(const char* StatName, uint32 CategoryIndex, float Value, const ECsvCustomStatOp CustomStatOp);
 	CORE_API static void RecordCustomStatMinimal(const char* StatName, uint32 CategoryIndex, int32 Value, const ECsvCustomStatOp CustomStatOp);
@@ -451,6 +453,8 @@ public:
 	CORE_API void BeginFrameRT();
 	CORE_API void EndFrameRT();
 
+	CORE_API void BeginFrameEOP();
+
 	CORE_API void SetDeviceProfileName(FString InDeviceProfileName);
 
 	FString GetOutputFilename() const { return OutputFilename; }
