Temporarily increase the maximum commandline length accepted on editor builds from 16k to 20k to avoid trunctation caused by long commandlines.

diff --git a/Engine/Source/Runtime/Core/Public/Misc/CommandLine.h b/Engine/Source/Runtime/Core/Public/Misc/CommandLine.h
--- a/Engine/Source/Runtime/Core/Public/Misc/CommandLine.h
+++ b/Engine/Source/Runtime/Core/Public/Misc/CommandLine.h
@@ -23,7 +23,11 @@
 struct FCommandLine
 {
 	/** maximum size of the command line */
+#if WITH_EDITOR
+	static constexpr inline uint32 MaxCommandLineSize = 20480;
+#else
 	static constexpr inline uint32 MaxCommandLineSize = 16384;
+#endif
 
 	/** 
 	 * Resets FCommandLine to an uninitialised state as if Set() has never been called
