Improve codegen for FOversubscriptionAllowedScope

Inline accesses to the thread-local variable generate a lot of code. Adding this scope to FRecursiveMutex originally increased instruction count from 73 to 98, or 25 instructions. With this change, the increase is only 10 instructions.

diff --git a/Engine/Source/Runtime/Core/Private/Async/Fundamental/Scheduler.cpp b/Engine/Source/Runtime/Core/Private/Async/Fundamental/Scheduler.cpp
--- a/Engine/Source/Runtime/Core/Private/Async/Fundamental/Scheduler.cpp
+++ b/Engine/Source/Runtime/Core/Private/Async/Fundamental/Scheduler.cpp
@@ -28,6 +28,11 @@ namespace LowLevelTasks
 	thread_local bool FSchedulerTls::Impl::bIsStandbyWorker = false;
 	thread_local bool Private::FOversubscriptionTls::bIsOversubscriptionAllowed = false;
 
+	FORCENOINLINE bool& Private::FOversubscriptionTls::GetIsOversubscriptionAllowedRef()
+	{
+		return bIsOversubscriptionAllowed;
+	}
+
 	FScheduler FScheduler::Singleton;
 
 	TUniquePtr<FThread> FScheduler::CreateWorker(uint32 WorkerId, const TCHAR* Name, bool bPermitBackgroundWork, FThread::EForkable IsForkable, Private::FWaitEvent* ExternalWorkerEvent, FSchedulerTls::FLocalQueueType* ExternalWorkerLocalQueue, EThreadPriority Priority, uint64 InAffinity)
diff --git a/Engine/Source/Runtime/Core/Public/Async/Fundamental/Scheduler.h b/Engine/Source/Runtime/Core/Public/Async/Fundamental/Scheduler.h
--- a/Engine/Source/Runtime/Core/Public/Async/Fundamental/Scheduler.h
+++ b/Engine/Source/Runtime/Core/Public/Async/Fundamental/Scheduler.h
@@ -256,6 +256,9 @@ namespace LowLevelTasks
 			static thread_local bool bIsOversubscriptionAllowed;
 
 			friend class FOversubscriptionAllowedScope;
+
+			static bool& GetIsOversubscriptionAllowedRef();
+
 		public:
 			static bool IsOversubscriptionAllowed() { return bIsOversubscriptionAllowed; }
 		};
@@ -265,17 +268,20 @@ namespace LowLevelTasks
 			UE_NONCOPYABLE(FOversubscriptionAllowedScope);
 
 		public:
-			FOversubscriptionAllowedScope(bool bIsOversubscriptionAllowed)
+			explicit FOversubscriptionAllowedScope(bool bIsOversubscriptionAllowed)
+				: bValue(FOversubscriptionTls::GetIsOversubscriptionAllowedRef())
+				, bPreviousValue(bValue)
 			{
-				bPreviousValue = FOversubscriptionTls::bIsOversubscriptionAllowed;
-				FOversubscriptionTls::bIsOversubscriptionAllowed = bIsOversubscriptionAllowed;
+				bValue = bIsOversubscriptionAllowed;
 			}
 
 			~FOversubscriptionAllowedScope()
 			{
-				FOversubscriptionTls::bIsOversubscriptionAllowed = bPreviousValue;
+				bValue = bPreviousValue;
 			}
+
 		private:
+			bool& bValue;
 			bool bPreviousValue;
 		};
 	}
