[AutoRTFM] Silence memory validation statistics for tests by default.

Add new AutoRTFMMemoryValidationStatisticsEnabled dcpvar flag to control this.
Fix the status logging for AutoRTFMMemoryValidationThrottlingEnabled

diff --git a/Engine/Source/Programs/AutoRTFMTests/Private/AutoRTFMTests.cpp b/Engine/Source/Programs/AutoRTFMTests/Private/AutoRTFMTests.cpp
--- a/Engine/Source/Programs/AutoRTFMTests/Private/AutoRTFMTests.cpp
+++ b/Engine/Source/Programs/AutoRTFMTests/Private/AutoRTFMTests.cpp
@@ -105,6 +105,9 @@ int RunTests(int ArgC, const char* ArgV[])
 	AutoRTFM::ForTheRuntime::SetMemoryValidationLevel(AutoRTFM::EMemoryValidationLevel::Error);
 	AutoRTFM::ForTheRuntime::SetMemoryValidationThrottlingEnabled(false);
 
+	// Don't print memory validation stats.
+	AutoRTFM::ForTheRuntime::SetMemoryValidationStatisticsEnabled(false);
+
 	int Result = Session.run();
 
 	FPlatformMisc::RequestExit(false);
diff --git a/Engine/Source/Runtime/Core/Private/AutoRTFM/API.cpp b/Engine/Source/Runtime/Core/Private/AutoRTFM/API.cpp
--- a/Engine/Source/Runtime/Core/Private/AutoRTFM/API.cpp
+++ b/Engine/Source/Runtime/Core/Private/AutoRTFM/API.cpp
@@ -33,6 +33,7 @@ namespace
 	int GMemoryValidationLevel = static_cast<int>(GDefaultMemoryValidationLevel);
 	
 	bool GMemoryValidationThrottlingEnabled = true;
+	bool GMemoryValidationStatisticsEnabled = true;
 
 	bool GAutoRTFMAssertOnError = true;
 
@@ -80,13 +81,20 @@ namespace
 			return FGenericCrashContext::SetGameData(TEXT("AutoRTFMMemoryValidationLevel"), TEXT("Enabled"));
 		}
 	}
-	
+
 	void OnAutoRTFMMemoryValidationThrottlingChanged()
 	{
 		const TCHAR* Text = GMemoryValidationThrottlingEnabled ? TEXT("true") : TEXT("false");
-		UE_LOG(LogCore, Log, TEXT("AutoRTFM Memory Validation: %s"), Text);
+		UE_LOG(LogCore, Log, TEXT("AutoRTFM Memory Validation Throttling Enabled: %s"), Text);
 		return FGenericCrashContext::SetGameData(TEXT("AutoRTFMMemoryValidationThrottlingEnabled"), Text);
 	}
+
+	void OnAutoRTFMMemoryValidationStatisticsChanged()
+	{
+		const TCHAR* Text = GMemoryValidationStatisticsEnabled ? TEXT("true") : TEXT("false");
+		UE_LOG(LogCore, Log, TEXT("AutoRTFM Memory Validation Statistics Enabled: %s"), Text);
+		return FGenericCrashContext::SetGameData(TEXT("AutoRTFMMemoryValidationStatisticsEnabled"), Text);
+	}
 }
 
 #if UE_AUTORTFM
@@ -122,6 +130,14 @@ static FAutoConsoleVariableRef CVarAutoRTFMMemoryValidationThrottlingEnabled(
 	ECVF_Default
 );
 
+static FAutoConsoleVariableRef CVarAutoRTFMMemoryValidationStatisticsEnabled(
+	TEXT("AutoRTFMMemoryValidationStatisticsEnabled"),
+	GMemoryValidationStatisticsEnabled,
+	TEXT("Disable logging of memory validation statistics"),
+	FConsoleVariableDelegate::CreateLambda([](IConsoleVariable*) { OnAutoRTFMMemoryValidationStatisticsChanged(); }),
+	ECVF_Default
+);
+
 static FAutoConsoleVariableRef CVarAutoRTFMAssertOnError(
 	TEXT("AutoRTFMAssertOnError"),
 	GAutoRTFMAssertOnError,
@@ -135,6 +151,7 @@ static FDelayedAutoRegisterHelper DelayedAutoRegister(EDelayedRegisterRunPhase::
 		OnAutoRTFMRetryTransactionsChanged();
 		OnAutoRTFMMemoryValidationLevelChanged();
 		OnAutoRTFMMemoryValidationThrottlingChanged();
+		OnAutoRTFMMemoryValidationStatisticsChanged();
 	});
 #endif
 
@@ -305,7 +322,23 @@ namespace AutoRTFM
 			GMemoryValidationThrottlingEnabled = bEnabled;
 #endif
 		}
-	
+
+		bool GetMemoryValidationStatisticsEnabled()
+		{
+#if UE_AUTORTFM
+			return GMemoryValidationStatisticsEnabled;
+#else
+			return false;
+#endif
+		}
+
+		void SetMemoryValidationStatisticsEnabled(bool bEnabled)
+		{
+#if UE_AUTORTFM
+			GMemoryValidationStatisticsEnabled = bEnabled;
+#endif
+		}
+
 		UE_AUTORTFM_ALWAYS_OPEN void DebugBreakIfMemoryValidationFails()
 		{
 #if UE_AUTORTFM
diff --git a/Engine/Source/Runtime/Core/Private/AutoRTFM/OpenHashThrottler.cpp b/Engine/Source/Runtime/Core/Private/AutoRTFM/OpenHashThrottler.cpp
--- a/Engine/Source/Runtime/Core/Private/AutoRTFM/OpenHashThrottler.cpp
+++ b/Engine/Source/Runtime/Core/Private/AutoRTFM/OpenHashThrottler.cpp
@@ -192,6 +192,11 @@ void FOpenHashThrottler::UpdateThrottlingData()
 
 void FOpenHashThrottler::UpdateLogStats()
 {
+	if (!AutoRTFM::ForTheRuntime::GetMemoryValidationStatisticsEnabled())
+	{
+		return;
+	}
+
 	const double TimeSinceLastLog = LogStats.TimeSinceLastReset;
 	UE_LOG(LogAutoRTFM, Display,
 		TEXT("Transaction Hash Statistics\n")
diff --git a/Engine/Source/Runtime/Core/Public/AutoRTFM/AutoRTFM.h b/Engine/Source/Runtime/Core/Public/AutoRTFM/AutoRTFM.h
--- a/Engine/Source/Runtime/Core/Public/AutoRTFM/AutoRTFM.h
+++ b/Engine/Source/Runtime/Core/Public/AutoRTFM/AutoRTFM.h
@@ -1093,6 +1093,12 @@ namespace ForTheRuntime
 	// Sets the memory validation throttling mode.
 	UE_AUTORTFM_API void SetMemoryValidationThrottlingEnabled(bool bEnabled);
 
+	// Returns true if the memory validation statistics are enabled.
+	UE_AUTORTFM_API bool GetMemoryValidationStatisticsEnabled();
+
+	// Sets whether memory validation statistics are logged or not.
+	UE_AUTORTFM_API void SetMemoryValidationStatisticsEnabled(bool bEnabled);
+
 	// A debug helper that will break to the debugger if the hash of the memory
 	// write locations no longer matches the hash recorded when the transaction
 	// was opened. Useful for isolating where the open write happened.
