Allow plugin lists to be specified via a file instead of a list on the commandline to avoid excessively long commandlines.  Also improve the failure handling when passed a commandline that exceeds internal limits.

diff --git a/Engine/Source/Runtime/Core/Private/Misc/CommandLine.cpp b/Engine/Source/Runtime/Core/Private/Misc/CommandLine.cpp
--- a/Engine/Source/Runtime/Core/Private/Misc/CommandLine.cpp
+++ b/Engine/Source/Runtime/Core/Private/Misc/CommandLine.cpp
@@ -79,22 +79,44 @@ bool FCommandLine::Set(const TCHAR* NewCommandLine)
 
 	FCString::Strncpy( CmdLine, NewCommandLine, UE_ARRAY_COUNT(CmdLine) );
 	FCString::Strncpy(LoggingCmdLine, NewCommandLine, UE_ARRAY_COUNT(LoggingCmdLine));
+
 	// If configured as part of the build, strip out any unapproved args
 	ApplyCommandLineAllowList();
 
 	bIsInitialized = true;
 	CmdLineVersion++;
 
+	// Check that the command line doesn't exceed our internal storage limit
+	int NewCommandLineLength = FCString::Strlen(NewCommandLine);
+	if (NewCommandLineLength >= UE_ARRAY_COUNT(CmdLine))
+	{
+		FText ErrorMessage = FText::Format(NSLOCTEXT("Engine", "CmdLineTooLong", "Error: Command-line exceeds internal storage limit.\nLength = {0}\nLimit={1}"),
+			NewCommandLineLength, int32(UE_ARRAY_COUNT(CmdLine)-1));
+#if !UE_BUILD_SHIPPING
+		// We can't call FApp::IsUnattended, because command line setting can be called before FApp::IsUnattended can read from the command line
+		if (!FParse::Param(FCommandLine::Get(), TEXT("UNATTENDED")))
+		{
+			FMessageDialog::Open(EAppMsgType::Ok, ErrorMessage);
+			return false;
+		}
+#endif
+		UE_LOG(LogInit, Fatal, TEXT("%s"), *ErrorMessage.ToString());
+	}
+
 	// Check for the '-' that normal ones get converted to in Outlook. It's important to do it AFTER the command line is initialized
 	if (StringHasBadDashes(NewCommandLine))
 	{
-		FText ErrorMessage = FText::Format(NSLOCTEXT("Engine", "ComdLineHasInvalidChar", "Error: Command-line contains an invalid '-' character, likely pasted from an email.\nCmdline = {0}"), FText::FromString(NewCommandLine));
+		FText ErrorMessage = FText::Format(NSLOCTEXT("Engine", "CmdLineHasInvalidChar", "Error: Command-line contains an invalid '-' character, likely pasted from an email.\nCmdline = {0}"),
+			FText::FromString(NewCommandLine));
 #if !UE_BUILD_SHIPPING
-		FMessageDialog::Open(EAppMsgType::Ok, ErrorMessage);
-		return false;
-#else
-		UE_LOG(LogInit, Fatal, TEXT("%s"), *ErrorMessage.ToString());
+		// We can't call FApp::IsUnattended, because command line setting can be called before FApp::IsUnattended can read from the command line
+		if (!FParse::Param(FCommandLine::Get(), TEXT("UNATTENDED")))
+		{
+			FMessageDialog::Open(EAppMsgType::Ok, ErrorMessage);
+			return false;
+		}
 #endif
+		UE_LOG(LogInit, Fatal, TEXT("%s"), *ErrorMessage.ToString());
 	}
 
 	return true;
diff --git a/Engine/Source/Runtime/Core/Public/Misc/CommandLine.h b/Engine/Source/Runtime/Core/Public/Misc/CommandLine.h
--- a/Engine/Source/Runtime/Core/Public/Misc/CommandLine.h
+++ b/Engine/Source/Runtime/Core/Public/Misc/CommandLine.h
@@ -23,11 +23,7 @@
 struct FCommandLine
 {
 	/** maximum size of the command line */
-#if WITH_EDITOR
-	static constexpr inline uint32 MaxCommandLineSize = 20480;
-#else
 	static constexpr inline uint32 MaxCommandLineSize = 16384;
-#endif
 
 	/** 
 	 * Resets FCommandLine to an uninitialised state as if Set() has never been called
diff --git a/Engine/Source/Runtime/Projects/Private/PluginManager.cpp b/Engine/Source/Runtime/Projects/Private/PluginManager.cpp
--- a/Engine/Source/Runtime/Projects/Private/PluginManager.cpp
+++ b/Engine/Source/Runtime/Projects/Private/PluginManager.cpp
@@ -1444,6 +1444,25 @@ bool FPluginManager::ConfigureEnabledPlugins()
 					PluginsListStr.ParseIntoArray(PluginsList, TEXT(","));
 				}
 
+				TArray<FString> PluginListFilenames;
+				PluginsList.RemoveAll([&PluginListFilenames](const FString& ExceptPluginName)
+				{
+					if (ExceptPluginName.EndsWith(TEXT(".txt")))
+					{
+						PluginListFilenames.Add(ExceptPluginName);
+						return true;
+					}
+					return false;
+				});
+
+				for (const FString& PluginListFilename : PluginListFilenames)
+				{
+					FFileHelper::LoadFileToStringWithLineVisitor(*PluginListFilename, [&PluginsList](FStringView Line)
+					{
+						PluginsList.Add(FString(Line));
+					});
+				}
+
 				TArray<FString> WildcardPluginsList;
 				PluginsList.RemoveAll([&WildcardPluginsList](const FString& ExceptPluginName)
 				{
