ISPC: Integrated a necessary change to allow building our .ispc files without warnings with ISPC 1.24.0 and newer.

It's an integration of a corresponding change from UE 5.5 (CL 36824059) but with a modification to only apply the change when using the new ISPC. This way we may build our files with both the current version and 1.24.0 which we are to update to in the future.

#rnx

diff --git a/Engine/Source/Runtime/Core/Public/Math/Scalar.isph b/Engine/Source/Runtime/Core/Public/Math/Scalar.isph
--- a/Engine/Source/Runtime/Core/Public/Math/Scalar.isph
+++ b/Engine/Source/Runtime/Core/Public/Math/Scalar.isph
@@ -10,8 +10,17 @@ typedef int64 IntType;
 
 #define REAL_BITS doublebits
 
-#undef  FLT_MAX
-#undef  FLT_MIN
+// To make version checks easier, we combine the major and minor versions into a single number.
+// We allocate three ditits for the minor version i.e. the format is going to be xyyy with x
+// being the major version and yyy the minor version.
+// We don't incorporate the patch version (ISPC_VERSION_PATCH from the compiler code) because
+// it is not exposed to the *.ispc/h sources.
+#define ISPC_VERSION		(ISPC_MAJOR_VERSION * 1000 + ISPC_MINOR_VERSION)
+
+#if ISPC_VERSION >= 1024
+	#undef  FLT_MAX
+	#undef  FLT_MIN
+#endif
 
 #define FLOAT_PI					(3.1415926535897932f)
 #define FLOAT_SMALL_NUMBER			(1.e-8f)
@@ -26,8 +35,10 @@ typedef int64 IntType;
 #define FLOAT_ONE_HALF				0.5f
 #define FLOAT_TWO					2.0f
 
-#undef  DBL_MAX
-#undef  DBL_MIN
+#if ISPC_VERSION >= 1024
+	#undef  DBL_MAX
+	#undef  DBL_MIN
+#endif
 
 #define DOUBLE_PI					(3.141592653589793238462643383279502884197169399d)
 #define DOUBLE_SMALL_NUMBER			(1.d-8)
