Optimize FMutex and FRecursiveMutex based on FSharedMutex and FSharedRecursiveMutex

Unlock immediately to allow other threads to acquire the lock while the unlocking thread looks for a thread to wake.
Align the implementation of the shared and non-shared mutexes to prepare for extracting the similar parts into a template.
Disable oversubscription in these mutexes and in FWordMutex.

diff --git a/Engine/Source/Runtime/Core/Private/Async/ExternalMutex.cpp b/Engine/Source/Runtime/Core/Private/Async/ExternalMutex.cpp
--- a/Engine/Source/Runtime/Core/Private/Async/ExternalMutex.cpp
+++ b/Engine/Source/Runtime/Core/Private/Async/ExternalMutex.cpp
@@ -2,6 +2,7 @@
 
 #include "Async/ExternalMutex.h"
 
+#include "Async/Fundamental/Scheduler.h"
 #include "Async/ParkingLot.h"
 #include "HAL/PlatformProcess.h"
 
@@ -12,73 +13,59 @@ void FExternalMutex::LockSlow()
 {
 	constexpr int32 SpinLimit = 40;
 	int32 SpinCount = 0;
-	for (;;)
+	for (uint8 CurrentState = State.load(std::memory_order_relaxed);;)
 	{
-		uint8 CurrentState = State.load(std::memory_order_acquire);
-
 		// Try to acquire the lock if it was unlocked, even if there are waiting threads.
 		// Acquiring the lock despite the waiting threads means that this lock is not FIFO and thus not fair.
 		if (LIKELY(!(CurrentState & IsLockedFlag)))
 		{
-			if (LIKELY(State.compare_exchange_weak(CurrentState, CurrentState | IsLockedFlag, std::memory_order_acquire)))
+			if (LIKELY(State.compare_exchange_weak(CurrentState, CurrentState | IsLockedFlag, std::memory_order_acquire, std::memory_order_relaxed)))
 			{
 				return;
 			}
 			continue;
 		}
 
-		if (LIKELY(!(CurrentState & HasWaitingThreadsFlag)))
+		// Spin up to the spin limit while there are no waiting threads.
+		if (LIKELY(!(CurrentState & HasWaitingLockFlag) && SpinCount < SpinLimit))
 		{
-			// Spin up to the spin limit while there are no waiting threads.
-			if (LIKELY(SpinCount < SpinLimit))
-			{
-				FPlatformProcess::YieldThread();
-				++SpinCount;
-				continue;
-			}
+			FPlatformProcess::YieldThread();
+			++SpinCount;
+			CurrentState = State.load(std::memory_order_relaxed);
+			continue;
+		}
 
-			// Store that there are waiting threads. Restart if the state has changed since it was loaded.
-			if (UNLIKELY(!State.compare_exchange_weak(CurrentState, CurrentState | HasWaitingThreadsFlag, std::memory_order_acquire)))
+		// Store that there are waiting threads. Restart if the state has changed since it was loaded.
+		if (LIKELY(!(CurrentState & HasWaitingLockFlag)))
+		{
+			if (UNLIKELY(!State.compare_exchange_weak(CurrentState, CurrentState | HasWaitingLockFlag, std::memory_order_acquire)))
 			{
 				continue;
 			}
-			CurrentState |= HasWaitingThreadsFlag;
+			CurrentState |= HasWaitingLockFlag;
 		}
 
-		// Wait if the state has not changed. Either way, loop back and try to acquire the lock trying to wait.
-		ParkingLot::Wait(&State, [this] { 
-			const uint8 CurrentState = State.load(std::memory_order_acquire);
-			return LIKELY((CurrentState & IsLockedFlag) && (CurrentState & HasWaitingThreadsFlag));
+		// Do not enter oversubscription during a wait on a mutex since the wait is generally too short
+		// for it to matter and it can worsen performance a lot for heavily contended locks.
+		LowLevelTasks::Private::FOversubscriptionAllowedScope _(false);
+
+		// Wait if the state has not changed. Either way, loop back and try to acquire the lock after trying to wait.
+		ParkingLot::Wait(&State, [this]
+		{
+			return (State.load(std::memory_order_relaxed) & (IsLockedFlag | HasWaitingLockFlag)) == (IsLockedFlag | HasWaitingLockFlag);
 		}, []{});
+		CurrentState = State.load(std::memory_order_relaxed);
 	}
 }
 
-void FExternalMutex::UnlockSlow()
+FORCENOINLINE void FExternalMutex::WakeWaitingThread()
 {
-	uint8 CurrentState = State.load(std::memory_order_acquire);
-	checkSlow(CurrentState & IsLockedFlag);
-
-	// Spin on the fast path because there may be spurious failures.
-	while (LIKELY((CurrentState & IsLockedFlag) && !(CurrentState & HasWaitingThreadsFlag)))
+	ParkingLot::WakeOne(&State, [this](ParkingLot::FWakeState WakeState) -> uint64
 	{
-		if (LIKELY(State.compare_exchange_weak(CurrentState, CurrentState & ~IsLockedFlag, std::memory_order_release)))
+		if (!WakeState.bHasWaitingThreads)
 		{
-			return;
+			State.fetch_and(~HasWaitingLockFlag, std::memory_order_relaxed);
 		}
-	}
-
-	// There is at least one thread waiting. Wake one thread and return.
-	ParkingLot::WakeOne(&State, [this](ParkingLot::FWakeState WakeState) -> uint64
-	{
-		uint8 ExpectedState = State.load(std::memory_order_acquire);
-		checkSlow(LIKELY((ExpectedState & IsLockedFlag) && (ExpectedState & HasWaitingThreadsFlag)));
-		
-		uint8 NewState = WakeState.bHasWaitingThreads ? ExpectedState | HasWaitingThreadsFlag : ExpectedState & ~HasWaitingThreadsFlag;
-		NewState &= ~IsLockedFlag;
-
-		uint8 OldState = ExpectedState;
-		State.compare_exchange_strong(OldState, NewState, std::memory_order_release);
-		checkSlow(OldState == ExpectedState);
 		return 0;
 	});
 }
diff --git a/Engine/Source/Runtime/Core/Private/Async/Mutex.cpp b/Engine/Source/Runtime/Core/Private/Async/Mutex.cpp
--- a/Engine/Source/Runtime/Core/Private/Async/Mutex.cpp
+++ b/Engine/Source/Runtime/Core/Private/Async/Mutex.cpp
@@ -2,9 +2,9 @@
 
 #include "Async/Mutex.h"
 
+#include "Async/Fundamental/Scheduler.h"
 #include "Async/ParkingLot.h"
 #include "HAL/PlatformProcess.h"
-#include "Async/Fundamental/Scheduler.h"
 
 namespace UE
 {
@@ -13,70 +13,56 @@ void FMutex::LockSlow()
 {
 	constexpr int32 SpinLimit = 40;
 	int32 SpinCount = 0;
-	for (;;)
+	for (uint8 CurrentState = State.load(std::memory_order_relaxed);;)
 	{
-		uint8 CurrentState = State.load(std::memory_order_acquire);
-
 		// Try to acquire the lock if it was unlocked, even if there are waiting threads.
 		// Acquiring the lock despite the waiting threads means that this lock is not FIFO and thus not fair.
 		if (LIKELY(!(CurrentState & IsLockedFlag)))
 		{
-			if (LIKELY(State.compare_exchange_weak(CurrentState, CurrentState | IsLockedFlag, std::memory_order_acquire)))
+			if (LIKELY(State.compare_exchange_weak(CurrentState, CurrentState | IsLockedFlag, std::memory_order_acquire, std::memory_order_relaxed)))
 			{
 				return;
 			}
 			continue;
 		}
 
-		if (LIKELY(!(CurrentState & HasWaitingThreadsFlag)))
+		// Spin up to the spin limit while there are no waiting threads.
+		if (LIKELY(!(CurrentState & HasWaitingLockFlag) && SpinCount < SpinLimit))
 		{
-			// Spin up to the spin limit while there are no waiting threads.
-			if (LIKELY(SpinCount < SpinLimit))
-			{
-				FPlatformProcess::YieldThread();
-				++SpinCount;
-				continue;
-			}
+			FPlatformProcess::YieldThread();
+			++SpinCount;
+			CurrentState = State.load(std::memory_order_relaxed);
+			continue;
+		}
 
-			// Store that there are waiting threads. Restart if the state has changed since it was loaded.
-			if (UNLIKELY(!State.compare_exchange_weak(CurrentState, CurrentState | HasWaitingThreadsFlag, std::memory_order_acquire)))
+		// Store that there are waiting threads. Restart if the state has changed since it was loaded.
+		if (LIKELY(!(CurrentState & HasWaitingLockFlag)))
+		{
+			if (UNLIKELY(!State.compare_exchange_weak(CurrentState, CurrentState | HasWaitingLockFlag, std::memory_order_relaxed)))
 			{
 				continue;
 			}
-			CurrentState |= HasWaitingThreadsFlag;
+			CurrentState |= HasWaitingLockFlag;
 		}
 
-		// For mutex, we don't want to enter oversubscription during a wait since the wait period is generally too short for oversubscription to matter
-		// and it can worsen performance a lot for heavily contented locks.
+		// Do not enter oversubscription during a wait on a mutex since the wait is generally too short
+		// for it to matter and it can worsen performance a lot for heavily contended locks.
 		LowLevelTasks::Private::FOversubscriptionAllowedScope _(false);
 
-		// Wait if the state has not changed. Either way, loop back and try to acquire the lock trying to wait.
-		ParkingLot::Wait(&State, [this, CurrentState] { return State.load(std::memory_order_acquire) == CurrentState; }, []{});
+		// Wait if the state has not changed. Either way, loop back and try to acquire the lock after trying to wait.
+		ParkingLot::Wait(&State, [this, CurrentState] { return State.load(std::memory_order_relaxed) == CurrentState; }, []{});
+		CurrentState = State.load(std::memory_order_relaxed);
 	}
 }
 
-void FMutex::UnlockSlow()
+FORCENOINLINE void FMutex::WakeWaitingThread()
 {
-	uint8 CurrentState = State.load(std::memory_order_acquire);
-	checkSlow(CurrentState & IsLockedFlag);
-
-	// Spin on the fast path because there may be spurious failures.
-	while (LIKELY(CurrentState == IsLockedFlag))
+	ParkingLot::WakeOne(&State, [this](ParkingLot::FWakeState WakeState) -> uint64
 	{
-		if (LIKELY(State.compare_exchange_weak(CurrentState, 0, std::memory_order_release)))
+		if (!WakeState.bHasWaitingThreads)
 		{
-			return;
+			State.fetch_and(~HasWaitingLockFlag, std::memory_order_relaxed);
 		}
-	}
-
-	// There is at least one thread waiting. Wake one thread and return.
-	ParkingLot::WakeOne(&State, [this](ParkingLot::FWakeState WakeState) -> uint64
-	{
-		constexpr uint8 ExpectedState = IsLockedFlag | HasWaitingThreadsFlag;
-		const uint8 NewState = WakeState.bHasWaitingThreads ? HasWaitingThreadsFlag : 0;
-		uint8 OldState = ExpectedState;
-		State.compare_exchange_strong(OldState, NewState, std::memory_order_release);
-		checkSlow(OldState == ExpectedState);
 		return 0;
 	});
 }
diff --git a/Engine/Source/Runtime/Core/Private/Async/RecursiveMutex.cpp b/Engine/Source/Runtime/Core/Private/Async/RecursiveMutex.cpp
--- a/Engine/Source/Runtime/Core/Private/Async/RecursiveMutex.cpp
+++ b/Engine/Source/Runtime/Core/Private/Async/RecursiveMutex.cpp
@@ -2,6 +2,7 @@
 
 #include "Async/RecursiveMutex.h"
 
+#include "Async/Fundamental/Scheduler.h"
 #include "Async/ParkingLot.h"
 #include "HAL/PlatformProcess.h"
 #include "HAL/PlatformTLS.h"
@@ -9,176 +10,143 @@
 namespace UE
 {
 
-union FRecursiveMutex::FState
-{
-	struct
-	{
-		uint32 bIsLocked : 1;
-		uint32 bHasWaitingThreads : 1;
-		uint32 RecurseCount : 30;
-	};
-	uint32 Value = 0;
-
-	constexpr FState() = default;
-
-	constexpr explicit FState(uint32 State)
-		: Value(State)
-	{
-	}
-
-	constexpr bool operator==(const FState& Other) const { return Value == Other.Value; }
-};
-
 bool FRecursiveMutex::TryLock()
 {
 	const uint32 CurrentThreadId = FPlatformTLS::GetCurrentThreadId();
-	FState CurrentState(State.load(std::memory_order_acquire));
+	uint32 CurrentState = State.load(std::memory_order_relaxed);
 
 	// Try to acquire the lock if it was unlocked, even if there are waiting threads.
 	// Acquiring the lock despite the waiting threads means that this lock is not FIFO and thus not fair.
-	if (LIKELY(!CurrentState.bIsLocked))
+	if (LIKELY(!(CurrentState & LockCountMask)))
 	{
-		FState NewState = CurrentState;
-		NewState.bIsLocked = true;
-		if (LIKELY(State.compare_exchange_strong(CurrentState.Value, NewState.Value, std::memory_order_acquire)))
+		if (LIKELY(State.compare_exchange_strong(CurrentState, CurrentState | (1 << LockCountShift), std::memory_order_acquire, std::memory_order_relaxed)))
 		{
-			checkSlow(ThreadId.load(std::memory_order_relaxed) == 0 && CurrentState.RecurseCount == 0);
+			checkSlow(ThreadId.load(std::memory_order_relaxed) == 0);
 			ThreadId.store(CurrentThreadId, std::memory_order_relaxed);
 			return true;
 		}
 	}
 
-	if (LIKELY(ThreadId.load(std::memory_order_relaxed) != CurrentThreadId))
+	// Lock recursively if this is the thread that holds the lock.
+	if (ThreadId.load(std::memory_order_relaxed) == CurrentThreadId)
 	{
-		return false;
+		State.fetch_add(1 << LockCountShift, std::memory_order_relaxed);
+		return true;
 	}
 
-	// Lock recursively if this is the thread that holds the lock.
-	FState AddState;
-	AddState.RecurseCount = 1;
-	State.fetch_add(AddState.Value, std::memory_order_relaxed);
-	return true;
+	return false;
 }
 
 void FRecursiveMutex::Lock()
 {
 	const uint32 CurrentThreadId = FPlatformTLS::GetCurrentThreadId();
-	FState CurrentState(State.load(std::memory_order_acquire));
+	uint32 CurrentState = State.load(std::memory_order_relaxed);
 
 	// Try to acquire the lock if it was unlocked, even if there are waiting threads.
 	// Acquiring the lock despite the waiting threads means that this lock is not FIFO and thus not fair.
-	if (LIKELY(!CurrentState.bIsLocked))
+	if (LIKELY(!(CurrentState & LockCountMask)))
 	{
-		FState NewState = CurrentState;
-		NewState.bIsLocked = true;
-		if (LIKELY(State.compare_exchange_weak(CurrentState.Value, NewState.Value, std::memory_order_acquire)))
+		if (LIKELY(State.compare_exchange_weak(CurrentState, CurrentState | (1 << LockCountShift), std::memory_order_acquire, std::memory_order_relaxed)))
 		{
-			checkSlow(ThreadId.load(std::memory_order_relaxed) == 0 && CurrentState.RecurseCount == 0);
+			checkSlow(ThreadId.load(std::memory_order_relaxed) == 0);
 			ThreadId.store(CurrentThreadId, std::memory_order_relaxed);
 			return;
 		}
 	}
 
-	if (LIKELY(ThreadId.load(std::memory_order_relaxed) != CurrentThreadId))
+	// Lock recursively if this is the thread that holds the lock.
+	if (ThreadId.load(std::memory_order_relaxed) == CurrentThreadId)
 	{
-		return LockSlow(CurrentState, CurrentThreadId);
+		State.fetch_add(1 << LockCountShift, std::memory_order_relaxed);
+		return;
 	}
 
-	// Lock recursively if this is the thread that holds the lock.
-	FState AddState;
-	AddState.RecurseCount = 1;
-	State.fetch_add(AddState.Value, std::memory_order_relaxed);
+	LockSlow(CurrentState, CurrentThreadId);
 }
 
-FORCENOINLINE void FRecursiveMutex::LockSlow(FState CurrentState, const uint32 CurrentThreadId)
+FORCENOINLINE void FRecursiveMutex::LockSlow(uint32 CurrentState, const uint32 CurrentThreadId)
 {
 	constexpr int32 SpinLimit = 40;
 	int32 SpinCount = 0;
-
 	for (;;)
 	{
 		// Try to acquire the lock if it was unlocked, even if there are waiting threads.
 		// Acquiring the lock despite the waiting threads means that this lock is not FIFO and thus not fair.
-		if (LIKELY(!CurrentState.bIsLocked))
+		if (LIKELY(!(CurrentState & LockCountMask)))
 		{
-			FState NewState = CurrentState;
-			NewState.bIsLocked = true;
-			if (LIKELY(State.compare_exchange_weak(CurrentState.Value, NewState.Value, std::memory_order_acquire)))
+			if (LIKELY(State.compare_exchange_weak(CurrentState, CurrentState | (1 << LockCountShift), std::memory_order_acquire, std::memory_order_relaxed)))
 			{
-				checkSlow(ThreadId.load(std::memory_order_relaxed) == 0 && CurrentState.RecurseCount == 0);
+				checkSlow(ThreadId.load(std::memory_order_relaxed) == 0);
 				ThreadId.store(CurrentThreadId, std::memory_order_relaxed);
 				return;
 			}
 			continue;
 		}
 
-		if (LIKELY(!CurrentState.bHasWaitingThreads))
+		// Spin up to the spin limit while there are no waiting threads.
+		if (LIKELY(!(CurrentState & HasWaitingLockFlag) && SpinCount < SpinLimit))
 		{
-			// Spin up to the spin limit while there are no waiting threads.
-			if (LIKELY(SpinCount < SpinLimit))
-			{
-				FPlatformProcess::YieldThread();
-				++SpinCount;
-				CurrentState.Value = State.load(std::memory_order_acquire);
-				continue;
-			}
+			FPlatformProcess::YieldThread();
+			++SpinCount;
+			CurrentState = State.load(std::memory_order_relaxed);
+			continue;
+		}
 
-			// Store that there are waiting threads. Restart if the state has changed since it was loaded.
-			FState NewState = CurrentState;
-			NewState.bHasWaitingThreads = true;
-			if (UNLIKELY(!State.compare_exchange_weak(CurrentState.Value, NewState.Value, std::memory_order_acquire)))
+		// Store that there are waiting threads. Restart if the state has changed since it was loaded.
+		if (LIKELY(!(CurrentState & HasWaitingLockFlag)))
+		{
+			if (UNLIKELY(!State.compare_exchange_weak(CurrentState, CurrentState | HasWaitingLockFlag, std::memory_order_relaxed)))
 			{
 				continue;
 			}
-			CurrentState = NewState;
+			CurrentState |= HasWaitingLockFlag;
 		}
 
+		// Do not enter oversubscription during a wait on a mutex since the wait is generally too short
+		// for it to matter and it can worsen performance a lot for heavily contended locks.
+		LowLevelTasks::Private::FOversubscriptionAllowedScope _(false);
+
 		// Wait if the state has not changed. Either way, loop back and try to acquire the lock after trying to wait.
-		ParkingLot::Wait(&State, [this, CurrentState] { return State.load(std::memory_order_acquire) == CurrentState.Value; }, []{});
-		CurrentState.Value = State.load(std::memory_order_acquire);
+		ParkingLot::Wait(&State, [this, CurrentState] { return State.load(std::memory_order_relaxed) == CurrentState; }, []{});
+		CurrentState = State.load(std::memory_order_relaxed);
 	}
 }
 
 void FRecursiveMutex::Unlock()
 {
-	FState CurrentState(State.load(std::memory_order_relaxed));
-	checkSlow(CurrentState.bIsLocked);
+	uint32 CurrentState = State.load(std::memory_order_relaxed);
+	checkSlow(CurrentState & LockCountMask);
 	checkSlow(ThreadId.load(std::memory_order_relaxed) == FPlatformTLS::GetCurrentThreadId());
 
-	if (LIKELY(CurrentState.RecurseCount == 0))
+	if (LIKELY((CurrentState & LockCountMask) == (1 << LockCountShift)))
 	{
 		// Remove the association with this thread before unlocking.
 		ThreadId.store(0, std::memory_order_relaxed);
 
-		// Unlocking with no waiting threads only requires resetting the state.
-		if (LIKELY(!CurrentState.bHasWaitingThreads))
+		// Unlock immediately to allow other threads to acquire the lock while this thread looks for a thread to wake.
+		uint32 LastState = State.fetch_sub(1 << LockCountShift, std::memory_order_release);
+
+		// Wake one exclusive waiter if there are waiting threads.
+		if (UNLIKELY(LastState & HasWaitingLockFlag))
 		{
-			FState NewState;
-			if (LIKELY(State.compare_exchange_strong(CurrentState.Value, NewState.Value, std::memory_order_release, std::memory_order_relaxed)))
-			{
-				return;
-			}
+			WakeWaitingThread();
 		}
-
-		UnlockSlow(CurrentState);
-		return;
 	}
-
-	// When locked recursively, decrement the count and return.
-	FState SubState;
-	SubState.RecurseCount = 1;
-	State.fetch_sub(SubState.Value, std::memory_order_relaxed);
+	else
+	{
+		// This is recursively locked. Decrement the lock count.
+		State.fetch_sub(1 << LockCountShift, std::memory_order_relaxed);
+	}
 }
 
-FORCENOINLINE void FRecursiveMutex::UnlockSlow(FState CurrentState)
+FORCENOINLINE void FRecursiveMutex::WakeWaitingThread()
 {
-	// There is at least one thread waiting. Wake one thread and return.
 	ParkingLot::WakeOne(&State, [this](ParkingLot::FWakeState WakeState) -> uint64
 	{
-		FState NewState;
-		NewState.bHasWaitingThreads = WakeState.bHasWaitingThreads;
-		FState OldState(State.exchange(NewState.Value, std::memory_order_release));
-		checkSlow(OldState.bIsLocked && OldState.bHasWaitingThreads);
+		if (!WakeState.bHasWaitingThreads)
+		{
+			State.fetch_and(~HasWaitingLockFlag, std::memory_order_relaxed);
+		}
 		return 0;
 	});
 }
diff --git a/Engine/Source/Runtime/Core/Private/Async/SharedMutex.cpp b/Engine/Source/Runtime/Core/Private/Async/SharedMutex.cpp
--- a/Engine/Source/Runtime/Core/Private/Async/SharedMutex.cpp
+++ b/Engine/Source/Runtime/Core/Private/Async/SharedMutex.cpp
@@ -29,9 +29,9 @@ void FSharedMutex::LockSlow()
 	{
 		// Try to acquire the lock if it was unlocked, even if there are waiting threads.
 		// Acquiring the lock despite the waiting threads means that this lock is not FIFO and thus not fair.
-		if (!(CurrentState & (IsLockedFlag | SharedLockCountMask)))
+		if (LIKELY(!(CurrentState & (IsLockedFlag | SharedLockCountMask))))
 		{
-			if (State.compare_exchange_weak(CurrentState, CurrentState | IsLockedFlag, std::memory_order_acquire, std::memory_order_relaxed))
+			if (LIKELY(State.compare_exchange_weak(CurrentState, CurrentState | IsLockedFlag, std::memory_order_acquire, std::memory_order_relaxed)))
 			{
 				return;
 			}
@@ -39,7 +39,7 @@ void FSharedMutex::LockSlow()
 		}
 
 		// Spin up to the spin limit while there are no waiting threads.
-		if (!(CurrentState & HasWaitingLockFlag) && SpinCount < SpinLimit)
+		if (LIKELY(!(CurrentState & HasWaitingLockFlag) && SpinCount < SpinLimit))
 		{
 			FPlatformProcess::YieldThread();
 			++SpinCount;
@@ -48,9 +48,9 @@ void FSharedMutex::LockSlow()
 		}
 
 		// Store that there are waiting threads. Restart if the state has changed since it was loaded.
-		if (!(CurrentState & HasWaitingLockFlag))
+		if (LIKELY(!(CurrentState & HasWaitingLockFlag)))
 		{
-			if (!State.compare_exchange_weak(CurrentState, CurrentState | HasWaitingLockFlag, std::memory_order_relaxed))
+			if (UNLIKELY(!State.compare_exchange_weak(CurrentState, CurrentState | HasWaitingLockFlag, std::memory_order_relaxed)))
 			{
 				continue;
 			}
@@ -61,42 +61,12 @@ void FSharedMutex::LockSlow()
 		// for it to matter and it can worsen performance a lot for heavily contended locks.
 		LowLevelTasks::Private::FOversubscriptionAllowedScope _(false);
 
-		// Wait if the state has not changed. Either way, loop back and try to acquire the lock trying to wait.
+		// Wait if the state has not changed. Either way, loop back and try to acquire the lock after trying to wait.
 		ParkingLot::Wait(GetLockAddress(), [this, CurrentState] { return State.load(std::memory_order_relaxed) == CurrentState; }, []{});
 		CurrentState = State.load(std::memory_order_relaxed);
 	}
 }
 
-void FSharedMutex::UnlockSlow()
-{
-	// Unlock immediately to allow other threads to acquire the lock while this thread looks for a thread to wake.
-	uint32 CurrentState = State.fetch_and(~IsLockedFlag, std::memory_order_release);
-	checkSlow(CurrentState & IsLockedFlag);
-	CurrentState &= ~IsLockedFlag;
-
-	if (CurrentState & HasWaitingLockFlag)
-	{
-		// Wake one thread that is waiting to acquire an exclusive lock.
-		ParkingLot::WakeOne(GetLockAddress(), [this](ParkingLot::FWakeState WakeState) -> uint64
-		{
-			if (!WakeState.bHasWaitingThreads)
-			{
-				State.fetch_and(~HasWaitingLockFlag, std::memory_order_relaxed);
-			}
-			return 0;
-		});
-	}
-	else if (CurrentState & HasWaitingSharedLockFlag)
-	{
-		// Wake every thread that is waiting to acquire a shared lock.
-		// The awoken threads might race against other exclusive locks.
-		if (State.fetch_and(~HasWaitingSharedLockFlag, std::memory_order_relaxed) & HasWaitingSharedLockFlag)
-		{
-			ParkingLot::WakeAll(GetSharedLockAddress());
-		}
-	}
-}
-
 void FSharedMutex::LockSharedSlow()
 {
 	constexpr int32 SpinLimit = 40;
@@ -104,9 +74,9 @@ void FSharedMutex::LockSharedSlow()
 	for (uint32 CurrentState = State.load(std::memory_order_relaxed);;)
 	{
 		// Try to acquire the lock if it is unlocked and there are no waiting threads.
-		if (!(CurrentState & (IsLockedFlag | HasWaitingLockFlag)))
+		if (LIKELY(!(CurrentState & (IsLockedFlag | HasWaitingLockFlag))))
 		{
-			if (State.compare_exchange_weak(CurrentState, CurrentState + (1 << SharedLockCountShift), std::memory_order_acquire, std::memory_order_relaxed))
+			if (LIKELY(State.compare_exchange_weak(CurrentState, CurrentState + (1 << SharedLockCountShift), std::memory_order_acquire, std::memory_order_relaxed)))
 			{
 				return;
 			}
@@ -114,7 +84,7 @@ void FSharedMutex::LockSharedSlow()
 		}
 
 		// Spin up to the spin limit while there are no waiting threads.
-		if (!(CurrentState & HasWaitingLockFlag) && SpinCount < SpinLimit)
+		if (LIKELY(!(CurrentState & HasWaitingLockFlag) && SpinCount < SpinLimit))
 		{
 			FPlatformProcess::YieldThread();
 			++SpinCount;
@@ -123,9 +93,9 @@ void FSharedMutex::LockSharedSlow()
 		}
 
 		// Store that there are waiting threads. Restart if the state has changed since it was loaded.
-		if (!(CurrentState & HasWaitingSharedLockFlag))
+		if (LIKELY(!(CurrentState & HasWaitingSharedLockFlag)))
 		{
-			if (!State.compare_exchange_weak(CurrentState, CurrentState | HasWaitingSharedLockFlag, std::memory_order_relaxed))
+			if (UNLIKELY(!State.compare_exchange_weak(CurrentState, CurrentState | HasWaitingSharedLockFlag, std::memory_order_relaxed)))
 			{
 				continue;
 			}
@@ -136,16 +106,14 @@ void FSharedMutex::LockSharedSlow()
 		// for it to matter and it can worsen performance a lot for heavily contended locks.
 		LowLevelTasks::Private::FOversubscriptionAllowedScope _(false);
 
-		// Wait if the state has not changed. Either way, loop back and try to acquire the lock trying to wait.
+		// Wait if the state has not changed. Either way, loop back and try to acquire the lock after trying to wait.
 		ParkingLot::Wait(GetSharedLockAddress(), [this, CurrentState] { return State.load(std::memory_order_relaxed) == CurrentState; }, []{});
 		CurrentState = State.load(std::memory_order_relaxed);
 	}
 }
 
-void FSharedMutex::UnlockSharedSlow()
+FORCENOINLINE void FSharedMutex::WakeWaitingThread()
 {
-	// This path is only taken when the last shared lock is released and there is a waiting exclusive lock.
-	// Wake one thread that is waiting to acquire an exclusive lock.
 	ParkingLot::WakeOne(GetLockAddress(), [this](ParkingLot::FWakeState WakeState) -> uint64
 	{
 		if (!WakeState.bHasWaitingThreads)
@@ -156,4 +124,29 @@ void FSharedMutex::UnlockSharedSlow()
 	});
 }
 
+FORCENOINLINE void FSharedMutex::WakeWaitingThreads(uint32 LastState)
+{
+	if (LastState & HasWaitingLockFlag)
+	{
+		// Wake one thread that is waiting to acquire an exclusive lock.
+		ParkingLot::WakeOne(GetLockAddress(), [this](ParkingLot::FWakeState WakeState) -> uint64
+		{
+			if (!WakeState.bHasWaitingThreads)
+			{
+				State.fetch_and(~HasWaitingLockFlag, std::memory_order_relaxed);
+			}
+			return 0;
+		});
+	}
+	else if (LastState & HasWaitingSharedLockFlag)
+	{
+		// Wake every thread that is waiting to acquire a shared lock.
+		// The awoken threads might race against other exclusive locks.
+		if (State.fetch_and(~HasWaitingSharedLockFlag, std::memory_order_relaxed) & HasWaitingSharedLockFlag)
+		{
+			ParkingLot::WakeAll(GetSharedLockAddress());
+		}
+	}
+}
+
 } // UE
diff --git a/Engine/Source/Runtime/Core/Private/Async/SharedRecursiveMutex.cpp b/Engine/Source/Runtime/Core/Private/Async/SharedRecursiveMutex.cpp
--- a/Engine/Source/Runtime/Core/Private/Async/SharedRecursiveMutex.cpp
+++ b/Engine/Source/Runtime/Core/Private/Async/SharedRecursiveMutex.cpp
@@ -86,7 +86,7 @@ bool FSharedRecursiveMutex::TryLock()
 	// Acquiring the lock despite the waiting threads means that this lock is not FIFO and thus not fair.
 	if (LIKELY(!(CurrentState & (LockCountMask | SharedLockCountMask))))
 	{
-		if (LIKELY(State.compare_exchange_weak(CurrentState, CurrentState | (1 << LockCountShift), std::memory_order_acquire, std::memory_order_relaxed)))
+		if (LIKELY(State.compare_exchange_strong(CurrentState, CurrentState | (1 << LockCountShift), std::memory_order_acquire, std::memory_order_relaxed)))
 		{
 			checkSlow(ThreadId.load(std::memory_order_relaxed) == 0);
 			ThreadId.store(CurrentThreadId, std::memory_order_relaxed);
@@ -133,7 +133,7 @@ void FSharedRecursiveMutex::Lock()
 	LockSlow(CurrentState, CurrentThreadId);
 }
 
-FORCENOINLINE void FSharedRecursiveMutex::LockSlow(uint32 CurrentState, uint32 CurrentThreadId)
+FORCENOINLINE void FSharedRecursiveMutex::LockSlow(uint32 CurrentState, const uint32 CurrentThreadId)
 {
 	constexpr int32 SpinLimit = 40;
 	int32 SpinCount = 0;
@@ -141,9 +141,9 @@ FORCENOINLINE void FSharedRecursiveMutex::LockSlow(uint32 CurrentState, uint32 C
 	{
 		// Try to acquire the lock if it was unlocked, even if there are waiting threads.
 		// Acquiring the lock despite the waiting threads means that this lock is not FIFO and thus not fair.
-		if (!(CurrentState & (LockCountMask | SharedLockCountMask)))
+		if (LIKELY(!(CurrentState & (LockCountMask | SharedLockCountMask))))
 		{
-			if (State.compare_exchange_weak(CurrentState, CurrentState | (1 << LockCountShift), std::memory_order_acquire, std::memory_order_relaxed))
+			if (LIKELY(State.compare_exchange_weak(CurrentState, CurrentState | (1 << LockCountShift), std::memory_order_acquire, std::memory_order_relaxed)))
 			{
 				checkSlow(ThreadId.load(std::memory_order_relaxed) == 0);
 				ThreadId.store(CurrentThreadId, std::memory_order_relaxed);
@@ -153,7 +153,7 @@ FORCENOINLINE void FSharedRecursiveMutex::LockSlow(uint32 CurrentState, uint32 C
 		}
 
 		// Spin up to the spin limit while there are no waiting threads.
-		if (!(CurrentState & HasWaitingLockFlag) && SpinCount < SpinLimit)
+		if (LIKELY(!(CurrentState & HasWaitingLockFlag) && SpinCount < SpinLimit))
 		{
 			FPlatformProcess::YieldThread();
 			++SpinCount;
@@ -162,9 +162,9 @@ FORCENOINLINE void FSharedRecursiveMutex::LockSlow(uint32 CurrentState, uint32 C
 		}
 
 		// Store that there are waiting threads. Restart if the state has changed since it was loaded.
-		if (!(CurrentState & HasWaitingLockFlag))
+		if (LIKELY(!(CurrentState & HasWaitingLockFlag)))
 		{
-			if (!State.compare_exchange_weak(CurrentState, CurrentState | HasWaitingLockFlag, std::memory_order_relaxed))
+			if (UNLIKELY(!State.compare_exchange_weak(CurrentState, CurrentState | HasWaitingLockFlag, std::memory_order_relaxed)))
 			{
 				continue;
 			}
@@ -175,7 +175,7 @@ FORCENOINLINE void FSharedRecursiveMutex::LockSlow(uint32 CurrentState, uint32 C
 		// for it to matter and it can worsen performance a lot for heavily contended locks.
 		LowLevelTasks::Private::FOversubscriptionAllowedScope _(false);
 
-		// Wait if the state has not changed. Either way, loop back and try to acquire the lock trying to wait.
+		// Wait if the state has not changed. Either way, loop back and try to acquire the lock after trying to wait.
 		ParkingLot::Wait(GetLockAddress(), [this, CurrentState] { return State.load(std::memory_order_relaxed) == CurrentState; }, []{});
 		CurrentState = State.load(std::memory_order_relaxed);
 	}
@@ -195,7 +195,7 @@ void FSharedRecursiveMutex::Unlock()
 		// Unlock immediately to allow other threads to acquire the lock while this thread looks for a thread to wake.
 		uint32 LastState = State.fetch_sub(1 << LockCountShift, std::memory_order_release);
 
-		// Wake one exclusive waiter or every shared waiter.
+		// Wake one exclusive waiter or every shared waiter if there are waiting threads.
 		if (UNLIKELY(LastState & (HasWaitingLockFlag | HasWaitingSharedLockFlag)))
 		{
 			WakeWaitingThreads(CurrentState);
@@ -203,6 +203,7 @@ void FSharedRecursiveMutex::Unlock()
 	}
 	else
 	{
+		// This is recursively locked. Decrement the lock count.
 		State.fetch_sub(1 << LockCountShift, std::memory_order_relaxed);
 	}
 }
@@ -256,10 +257,10 @@ FORCENOINLINE void FSharedRecursiveMutex::LockSharedSlow(Core::Private::FSharedR
 	for (uint32 CurrentState = State.load(std::memory_order_relaxed);;)
 	{
 		// Try to acquire the lock if it is unlocked and there are no waiting threads.
-		if (!(CurrentState & (LockCountMask | HasWaitingLockFlag)))
+		if (LIKELY(!(CurrentState & (LockCountMask | HasWaitingLockFlag))))
 		{
 			checkSlow((CurrentState & SharedLockCountMask) != SharedLockCountMask);
-			if (State.compare_exchange_weak(CurrentState, CurrentState + (1 << SharedLockCountShift), std::memory_order_acquire, std::memory_order_relaxed))
+			if (LIKELY(State.compare_exchange_weak(CurrentState, CurrentState + (1 << SharedLockCountShift), std::memory_order_acquire, std::memory_order_relaxed)))
 			{
 				Link.Push(this);
 				return;
@@ -268,7 +269,7 @@ FORCENOINLINE void FSharedRecursiveMutex::LockSharedSlow(Core::Private::FSharedR
 		}
 
 		// Spin up to the spin limit while there are no waiting threads.
-		if (!(CurrentState & HasWaitingLockFlag) && SpinCount < SpinLimit)
+		if (LIKELY(!(CurrentState & HasWaitingLockFlag) && SpinCount < SpinLimit))
 		{
 			FPlatformProcess::YieldThread();
 			++SpinCount;
@@ -277,9 +278,9 @@ FORCENOINLINE void FSharedRecursiveMutex::LockSharedSlow(Core::Private::FSharedR
 		}
 
 		// Store that there are waiting threads. Restart if the state has changed since it was loaded.
-		if (!(CurrentState & HasWaitingSharedLockFlag))
+		if (LIKELY(!(CurrentState & HasWaitingSharedLockFlag)))
 		{
-			if (!State.compare_exchange_weak(CurrentState, CurrentState | HasWaitingSharedLockFlag, std::memory_order_relaxed))
+			if (UNLIKELY(!State.compare_exchange_weak(CurrentState, CurrentState | HasWaitingSharedLockFlag, std::memory_order_relaxed)))
 			{
 				continue;
 			}
@@ -290,7 +291,7 @@ FORCENOINLINE void FSharedRecursiveMutex::LockSharedSlow(Core::Private::FSharedR
 		// for it to matter and it can worsen performance a lot for heavily contended locks.
 		LowLevelTasks::Private::FOversubscriptionAllowedScope _(false);
 
-		// Wait if the state has not changed. Either way, loop back and try to acquire the lock trying to wait.
+		// Wait if the state has not changed. Either way, loop back and try to acquire the lock after trying to wait.
 		ParkingLot::Wait(GetSharedLockAddress(), [this, CurrentState] { return State.load(std::memory_order_relaxed) == CurrentState; }, []{});
 		CurrentState = State.load(std::memory_order_relaxed);
 	}
diff --git a/Engine/Source/Runtime/Core/Private/Async/WordMutex.cpp b/Engine/Source/Runtime/Core/Private/Async/WordMutex.cpp
--- a/Engine/Source/Runtime/Core/Private/Async/WordMutex.cpp
+++ b/Engine/Source/Runtime/Core/Private/Async/WordMutex.cpp
@@ -2,6 +2,7 @@
 
 #include "Async/WordMutex.h"
 
+#include "Async/Fundamental/Scheduler.h"
 #include "HAL/PlatformProcess.h"
 #include "HAL/PlatformManualResetEvent.h"
 
@@ -70,6 +71,10 @@ void FWordMutex::LockSlow()
 			continue;
 		}
 
+		// Do not enter oversubscription during a wait on a mutex since the wait is generally too short
+		// for it to matter and it can worsen performance a lot for heavily contended locks.
+		LowLevelTasks::Private::FOversubscriptionAllowedScope _(false);
+
 		// Wait until another thread wakes this thread, which can happen as soon as the preceding store completes.
 		Self.Event.Wait();
 
diff --git a/Engine/Source/Runtime/Core/Public/Async/ExternalMutex.h b/Engine/Source/Runtime/Core/Public/Async/ExternalMutex.h
--- a/Engine/Source/Runtime/Core/Public/Async/ExternalMutex.h
+++ b/Engine/Source/Runtime/Core/Public/Async/ExternalMutex.h
@@ -44,18 +44,15 @@ public:
 
 	inline bool TryLock()
 	{
-		uint8 Expected = State.load(std::memory_order_acquire);
-		if (LIKELY(!(Expected & IsLockedFlag) && State.compare_exchange_strong(Expected, Expected | IsLockedFlag, std::memory_order_acquire)))
-		{
-			return true;
-		}
-		return false;
+		uint8 Expected = State.load(std::memory_order_relaxed);
+		return !(Expected & IsLockedFlag) &&
+			State.compare_exchange_strong(Expected, Expected | IsLockedFlag, std::memory_order_acquire, std::memory_order_relaxed);
 	}
 
 	inline void Lock()
 	{
-		uint8 Expected = State.load(std::memory_order_acquire) & ~IsLockedFlag & ~HasWaitingThreadsFlag;
-		if (LIKELY(State.compare_exchange_weak(Expected, Expected | IsLockedFlag, std::memory_order_acquire)))
+		uint8 Expected = State.load(std::memory_order_relaxed) & ~IsLockedFlag & ~HasWaitingLockFlag;
+		if (LIKELY(State.compare_exchange_weak(Expected, Expected | IsLockedFlag, std::memory_order_acquire, std::memory_order_relaxed)))
 		{
 			return;
 		}
@@ -64,20 +61,21 @@ public:
 
 	inline void Unlock()
 	{
-		uint8 Expected = (State.load(std::memory_order_acquire) | IsLockedFlag) & ~HasWaitingThreadsFlag;
-		if (LIKELY(State.compare_exchange_weak(Expected, Expected & ~IsLockedFlag, std::memory_order_release)))
+		// Unlock immediately to allow other threads to acquire the lock while this thread looks for a thread to wake.
+		uint8 LastState = State.fetch_sub(IsLockedFlag, std::memory_order_release);
+		if (LIKELY(!(LastState & HasWaitingLockFlag)))
 		{
 			return;
 		}
-		UnlockSlow();
+		WakeWaitingThread();
 	}
 
 private:
 	UE_API void LockSlow();
-	UE_API void UnlockSlow();
+	UE_API void WakeWaitingThread();
 
 	static constexpr uint8 IsLockedFlag = 1 << 0;
-	static constexpr uint8 HasWaitingThreadsFlag = 1 << 1;
+	static constexpr uint8 HasWaitingLockFlag = 1 << 1;
 
 	std::atomic<uint8>& State;
 };
diff --git a/Engine/Source/Runtime/Core/Public/Async/Mutex.h b/Engine/Source/Runtime/Core/Public/Async/Mutex.h
--- a/Engine/Source/Runtime/Core/Public/Async/Mutex.h
+++ b/Engine/Source/Runtime/Core/Public/Async/Mutex.h
@@ -36,18 +36,15 @@ public:
 	
 	inline bool TryLock()
 	{
-		uint8 Expected = State.load(std::memory_order_acquire);
-		if (LIKELY(!(Expected & IsLockedFlag) && State.compare_exchange_strong(Expected, Expected | IsLockedFlag, std::memory_order_acquire)))
-		{
-			return true;
-		}
-		return false;
+		uint8 Expected = State.load(std::memory_order_relaxed);
+		return !(Expected & IsLockedFlag) &&
+			State.compare_exchange_strong(Expected, Expected | IsLockedFlag, std::memory_order_acquire, std::memory_order_relaxed);
 	}
 
 	inline void Lock()
 	{
 		uint8 Expected = 0;
-		if (LIKELY(State.compare_exchange_weak(Expected, IsLockedFlag, std::memory_order_acquire)))
+		if (LIKELY(State.compare_exchange_weak(Expected, IsLockedFlag, std::memory_order_acquire, std::memory_order_relaxed)))
 		{
 			return;
 		}
@@ -56,20 +53,21 @@ public:
 
 	inline void Unlock()
 	{
-		uint8 Expected = IsLockedFlag;
-		if (LIKELY(State.compare_exchange_weak(Expected, 0, std::memory_order_release)))
+		// Unlock immediately to allow other threads to acquire the lock while this thread looks for a thread to wake.
+		uint8 LastState = State.fetch_sub(IsLockedFlag, std::memory_order_release);
+		if (LIKELY(!(LastState & HasWaitingLockFlag)))
 		{
 			return;
 		}
-		UnlockSlow();
+		WakeWaitingThread();
 	}
 
 private:
 	UE_API void LockSlow();
-	UE_API void UnlockSlow();
+	UE_API void WakeWaitingThread();
 
 	static constexpr uint8 IsLockedFlag = 1 << 0;
-	static constexpr uint8 HasWaitingThreadsFlag = 1 << 1;
+	static constexpr uint8 HasWaitingLockFlag = 1 << 1;
 
 	std::atomic<uint8> State = 0;
 };
diff --git a/Engine/Source/Runtime/Core/Public/Async/RecursiveMutex.h b/Engine/Source/Runtime/Core/Public/Async/RecursiveMutex.h
--- a/Engine/Source/Runtime/Core/Public/Async/RecursiveMutex.h
+++ b/Engine/Source/Runtime/Core/Public/Async/RecursiveMutex.h
@@ -28,10 +28,12 @@ public:
 	UE_API void Unlock();
 
 private:
-	union FState;
+	void LockSlow(uint32 CurrentState, uint32 CurrentThreadId);
+	void WakeWaitingThread();
 
-	void LockSlow(FState CurrentState, uint32 CurrentThreadId);
-	void UnlockSlow(FState CurrentState);
+	static constexpr uint32 HasWaitingLockFlag = 1 << 0;
+	static constexpr uint32 LockCountShift = 1;
+	static constexpr uint32 LockCountMask = 0xffff'fffe;
 
 	std::atomic<uint32> State = 0;
 	std::atomic<uint32> ThreadId = 0;
diff --git a/Engine/Source/Runtime/Core/Public/Async/SharedMutex.h b/Engine/Source/Runtime/Core/Public/Async/SharedMutex.h
--- a/Engine/Source/Runtime/Core/Public/Async/SharedMutex.h
+++ b/Engine/Source/Runtime/Core/Public/Async/SharedMutex.h
@@ -46,12 +46,14 @@ public:
 
 	inline void Unlock()
 	{
-		uint32 Expected = IsLockedFlag;
-		if (LIKELY(State.compare_exchange_weak(Expected, 0, std::memory_order_release, std::memory_order_relaxed)))
+		// Unlock immediately to allow other threads to acquire the lock while this thread looks for a thread to wake.
+		uint32 LastState = State.fetch_sub(IsLockedFlag, std::memory_order_release);
+		checkSlow(LastState & IsLockedFlag);
+		if (LIKELY(!(LastState & (HasWaitingLockFlag | HasWaitingSharedLockFlag))))
 		{
 			return;
 		}
-		UnlockSlow();
+		WakeWaitingThreads(LastState);
 	}
 
 	inline bool TryLockShared()
@@ -76,6 +78,7 @@ public:
 
 	inline void UnlockShared()
 	{
+		// Unlock immediately to allow other threads to acquire the lock while this thread looks for a thread to wake.
 		const uint32 LastState = State.fetch_sub(1 << SharedLockCountShift, std::memory_order_release);
 		checkSlow(LastState & SharedLockCountMask);
 		constexpr uint32 WakeState = HasWaitingLockFlag | (1 << SharedLockCountShift);
@@ -83,14 +86,14 @@ public:
 		{
 			return;
 		}
-		UnlockSharedSlow();
+		WakeWaitingThread();
 	}
 
 private:
 	UE_API void LockSlow();
-	UE_API void UnlockSlow();
 	UE_API void LockSharedSlow();
-	UE_API void UnlockSharedSlow();
+	UE_API void WakeWaitingThread();
+	UE_API void WakeWaitingThreads(uint32 LastState);
 
 	const void* GetLockAddress() const;
 	const void* GetSharedLockAddress() const;
