Fixed a race waking threads in FSharedMutex and FSharedRecursiveMutex

#rnx

diff --git a/Engine/Source/Runtime/Core/Private/Async/SharedMutex.cpp b/Engine/Source/Runtime/Core/Private/Async/SharedMutex.cpp
--- a/Engine/Source/Runtime/Core/Private/Async/SharedMutex.cpp
+++ b/Engine/Source/Runtime/Core/Private/Async/SharedMutex.cpp
@@ -129,16 +129,23 @@ FORCENOINLINE void FSharedMutex::WakeWaitingThreads(uint32 LastState)
 	if (LastState & HasWaitingLockFlag)
 	{
 		// Wake one thread that is waiting to acquire an exclusive lock.
-		ParkingLot::WakeOne(GetLockAddress(), [this](ParkingLot::FWakeState WakeState) -> uint64
+		bool bDidWake = false;
+		ParkingLot::WakeOne(GetLockAddress(), [this, &bDidWake](ParkingLot::FWakeState WakeState) -> uint64
 		{
 			if (!WakeState.bHasWaitingThreads)
 			{
 				State.fetch_and(~HasWaitingLockFlag, std::memory_order_relaxed);
 			}
+			bDidWake = WakeState.bDidWake;
 			return 0;
 		});
+		if (bDidWake)
+		{
+			return;
+		}
 	}
-	else if (LastState & HasWaitingSharedLockFlag)
+
+	if (LastState & HasWaitingSharedLockFlag)
 	{
 		// Wake every thread that is waiting to acquire a shared lock.
 		// The awoken threads might race against other exclusive locks.
diff --git a/Engine/Source/Runtime/Core/Private/Async/SharedRecursiveMutex.cpp b/Engine/Source/Runtime/Core/Private/Async/SharedRecursiveMutex.cpp
--- a/Engine/Source/Runtime/Core/Private/Async/SharedRecursiveMutex.cpp
+++ b/Engine/Source/Runtime/Core/Private/Async/SharedRecursiveMutex.cpp
@@ -323,16 +323,23 @@ FORCENOINLINE void FSharedRecursiveMutex::WakeWaitingThreads(uint32 CurrentState
 	if (CurrentState & HasWaitingLockFlag)
 	{
 		// Wake one thread that is waiting to acquire an exclusive lock.
-		ParkingLot::WakeOne(GetLockAddress(), [this](ParkingLot::FWakeState WakeState) -> uint64
+		bool bDidWake = false;
+		ParkingLot::WakeOne(GetLockAddress(), [this, &bDidWake](ParkingLot::FWakeState WakeState) -> uint64
 		{
 			if (!WakeState.bHasWaitingThreads)
 			{
 				State.fetch_and(~HasWaitingLockFlag, std::memory_order_relaxed);
 			}
+			bDidWake = WakeState.bDidWake;
 			return 0;
 		});
+		if (bDidWake)
+		{
+			return;
+		}
 	}
-	else if (CurrentState & HasWaitingSharedLockFlag)
+
+	if (CurrentState & HasWaitingSharedLockFlag)
 	{
 		// Wake every thread that is waiting to acquire a shared lock.
 		// The awoken threads might race against other exclusive locks.
