Update TWeakObjectPtr::TryPin to use the newer base pointer access.

diff --git a/Engine/Source/Runtime/Core/Public/UObject/WeakObjectPtrTemplates.h b/Engine/Source/Runtime/Core/Public/UObject/WeakObjectPtrTemplates.h
--- a/Engine/Source/Runtime/Core/Public/UObject/WeakObjectPtrTemplates.h
+++ b/Engine/Source/Runtime/Core/Public/UObject/WeakObjectPtrTemplates.h
@@ -164,7 +164,7 @@ public:
 	FORCEINLINE TStrongObjectPtr<T> TryPin(bool& bOutPinValid, bool bEvenIfPendingKill) const
 	{
 		TStrongObjectPtr<T> StrongPtr;
-		StrongPtr.Attach((T*)TWeakObjectPtrBase::TryPin(bOutPinValid, bEvenIfPendingKill).Detach());
+		StrongPtr.Attach((T*)WeakPtr.TryPin(bOutPinValid, bEvenIfPendingKill).Detach());
 		return StrongPtr;
 	}
 
