GLog - Don't return a bad GLog when it has been teardown / destroyed
This is to avoid trying to use FOutputDeviceRedirector::Flush() which would failed when GLog is not valid anymore and cause reentrant checkf and eventually overflow the callstack.

#tests Local Test on Mac & Windows MegaStation including for LowLevelTests

diff --git a/Engine/Source/Developer/LowLevelTestsRunner/Private/TestCommon/Initialization.cpp b/Engine/Source/Developer/LowLevelTestsRunner/Private/TestCommon/Initialization.cpp
--- a/Engine/Source/Developer/LowLevelTestsRunner/Private/TestCommon/Initialization.cpp
+++ b/Engine/Source/Developer/LowLevelTestsRunner/Private/TestCommon/Initialization.cpp
@@ -10,7 +10,6 @@
 #include "HAL/PlatformFileManager.h"
 #include "Misc/CommandLine.h"
 #include "Misc/DelayedAutoRegister.h"
-#include "Misc/OutputDeviceRedirector.h"
 #include "Misc/Paths.h"
 
 
@@ -129,14 +128,6 @@ void InitAll(bool bAllowLogging, bool bMultithreaded)
 	GIsRunning = true;
 }
 
-void CleanupLogs()
-{
-	if (GLog)
-	{
-		GLog->TearDown();
-	}
-}
-
 void CleanupLocalization()
 {
 	FTextLocalizationManager::TearDown();
@@ -153,6 +144,5 @@ void CleanupAll()
 #endif
 	CleanupAllThreadPools();
 	CleanupTaskGraph();
-	CleanupLogs();
 	CleanupLocalization();
 }
\ No newline at end of file
diff --git a/Engine/Source/Developer/LowLevelTestsRunner/Private/TestRunner.cpp b/Engine/Source/Developer/LowLevelTestsRunner/Private/TestRunner.cpp
--- a/Engine/Source/Developer/LowLevelTestsRunner/Private/TestRunner.cpp
+++ b/Engine/Source/Developer/LowLevelTestsRunner/Private/TestRunner.cpp
@@ -6,6 +6,7 @@
 #include "HAL/PlatformOutputDevices.h"
 #include "HAL/PlatformTLS.h"
 #include "Logging/LogSuppressionInterface.h"
+#include "Misc/OutputDeviceRedirector.h"
 #include "Misc/ScopeExit.h"
 #include "Misc/StringBuilder.h"
 #include "Misc/Paths.h"
@@ -325,6 +326,11 @@ int RunTests(int32 ArgC, const ANSICHAR* ArgV[])
 		
 		FModuleManager::Get().UnloadModulesAtShutdown();
 		RequestEngineExit(TEXT("Exiting"));
+		// TearDown GLog at the end so RequestEngineExit can properly log ExitCode
+		if (GLog)
+		{
+			GLog->TearDown();
+		}
 	};
 
 	int CatchReturn = TestRunner.RunCatchSession();
diff --git a/Engine/Source/Runtime/Core/Private/Misc/OutputDeviceRedirector.cpp b/Engine/Source/Runtime/Core/Private/Misc/OutputDeviceRedirector.cpp
--- a/Engine/Source/Runtime/Core/Private/Misc/OutputDeviceRedirector.cpp
+++ b/Engine/Source/Runtime/Core/Private/Misc/OutputDeviceRedirector.cpp
@@ -778,6 +778,8 @@ void FOutputDeviceRedirector::Panic()
 	}
 }
 
+static bool GWasOutputDeviceRedirectorTornDown = false;
+
 void FOutputDeviceRedirector::TearDown()
 {
 	SetCurrentThreadAsPrimaryThread();
@@ -806,6 +808,8 @@ void FOutputDeviceRedirector::TearDown()
 	{
 		OutputDevice->TearDown();
 	}
+	
+	GWasOutputDeviceRedirectorTornDown = true;
 }
 
 bool FOutputDeviceRedirector::IsBacklogEnabled() const
@@ -816,5 +820,5 @@ bool FOutputDeviceRedirector::IsBacklogEnabled() const
 
 CORE_API FOutputDeviceRedirector* GetGlobalLogSingleton()
 {
-	return FOutputDeviceRedirector::Get();
+	return GWasOutputDeviceRedirectorTornDown ? nullptr : FOutputDeviceRedirector::Get();
 }
