[Core] Expose IsSet() as public functionality in TFunction and peers.

This is a natural sibling to Reset() and makes sense to expose directly, instead of requiring users to cast to bool to check the is-set state.

diff --git a/Engine/Source/Runtime/Core/Private/AutoRTFM/Context.cpp b/Engine/Source/Runtime/Core/Private/AutoRTFM/Context.cpp
--- a/Engine/Source/Runtime/Core/Private/AutoRTFM/Context.cpp
+++ b/Engine/Source/Runtime/Core/Private/AutoRTFM/Context.cpp
@@ -495,12 +495,12 @@ ETransactionResult FContext::Transact(void (*InstrumentedFunction)(void*), void*
 				break;
 
 			case EContextStatus::AbortedByCascadingAbort:
-				ASSERT(!CascadingRetryTransactionWork);
+				ASSERT(!CascadingRetryTransactionWork.IsSet());
 				Result = ETransactionResult::AbortedByCascade;
 				break;
 
 			case EContextStatus::AbortedByCascadingRetry:
-				ASSERT(CascadingRetryTransactionWork);
+				ASSERT(CascadingRetryTransactionWork.IsSet());
 
 				// Clean up the transaction to get it ready for re-execution.
 				ClearTransactionStatus();
diff --git a/Engine/Source/Runtime/Core/Public/Templates/Function.h b/Engine/Source/Runtime/Core/Public/Templates/Function.h
--- a/Engine/Source/Runtime/Core/Public/Templates/Function.h
+++ b/Engine/Source/Runtime/Core/Public/Templates/Function.h
@@ -781,6 +781,14 @@ public:
 		return Super::IsSet();
 	}
 
+	/**
+	 * Tests if the TFunction is callable.
+	 */
+	FORCEINLINE bool IsSet() const
+	{
+		return Super::IsSet();
+	}
+
 	/**
 	 * Nullptr equality operator.
 	 */
@@ -899,6 +907,14 @@ public:
 	{
 		return Super::IsSet();
 	}
+
+	/**
+	 * Tests if the TUniqueFunction is callable.
+	 */
+	FORCEINLINE bool IsSet() const
+	{
+		return Super::IsSet();
+	}
 };
 
 
