Fix failing check in FoundationTests due to varying growth patterns on different platforms.

#rnx

diff --git a/Engine/Source/Runtime/Core/Tests/Containers/DequeTest.cpp b/Engine/Source/Runtime/Core/Tests/Containers/DequeTest.cpp
--- a/Engine/Source/Runtime/Core/Tests/Containers/DequeTest.cpp
+++ b/Engine/Source/Runtime/Core/Tests/Containers/DequeTest.cpp
@@ -1038,7 +1038,7 @@ TEST_CASE_NAMED(FDequeConstructFromStdInitializerListTest, "System::Core::Contai
 	{
 		TDeque<int32Token, TInlineAllocator<4>> Deque({ 0, 1, 2, 3, 4, 5 });
 		CHECK(Deque.Num() == 6);
-		CHECK(Deque.Max() == 6);
+		CHECK(Deque.Max() >= 6); // Growth amount can vary per platform
 	}
 }
 
@@ -1095,7 +1095,7 @@ TEST_CASE_NAMED(FDequeAssignFromStdInitializerListTest, "System::Core::Container
 		TDeque<int32Token, TInlineAllocator<4>> Deque;
 		Deque = { 0, 1, 2, 3, 4, 5 };
 		CHECK(Deque.Num() == 6);
-		CHECK(Deque.Max() == 6);
+		CHECK(Deque.Max() >= 6); // Growth amount can vary per platform
 	}
 }
 
