Add EndOfPipe timeline to stats system
- Will be used by the new GPU profiler. Solves the problem where stats reported by the end-of-pipe thread are attributed to the wrong frame, making their values unstable.

Minor clean up of some stats system code
- Moved the TRACE_STAT_SET Insights hooks into the FThreadStats::AddMessage function, to avoid the need to add them in places that aren't using the standard STATS macros.
- Simplified the pooling of FThreadStats instances
- Removed custom stats macros in Android code, replaced with standard SET_*_STAT_FName macros
- Removed some deprecated functions

diff --git a/Engine/Plugins/FX/Niagara/Source/Niagara/Private/NiagaraSystem.cpp b/Engine/Plugins/FX/Niagara/Source/Niagara/Private/NiagaraSystem.cpp
--- a/Engine/Plugins/FX/Niagara/Source/Niagara/Private/NiagaraSystem.cpp
+++ b/Engine/Plugins/FX/Niagara/Source/Niagara/Private/NiagaraSystem.cpp
@@ -3834,12 +3834,10 @@ void UNiagaraSystem::AddToInstanceCountStat(int32 NumInstances, bool bSolo)const
 		if (bSolo)
 		{
 			FThreadStats::AddMessage(StatID_InstanceCountSolo.GetName(), EStatOperation::Add, int64(NumInstances));
-			TRACE_STAT_ADD(StatID_InstanceCountSolo.GetName(), int64(NumInstances));
 		}
 		else
 		{
 			FThreadStats::AddMessage(StatID_InstanceCount.GetName(), EStatOperation::Add, int64(NumInstances));
-			TRACE_STAT_ADD(StatID_InstanceCount.GetName(), int64(NumInstances));
 		}
 	}
 #endif
diff --git a/Engine/Source/Developer/Profiler/Private/ProfilerRawStatsForThreadView.cpp b/Engine/Source/Developer/Profiler/Private/ProfilerRawStatsForThreadView.cpp
--- a/Engine/Source/Developer/Profiler/Private/ProfilerRawStatsForThreadView.cpp
+++ b/Engine/Source/Developer/Profiler/Private/ProfilerRawStatsForThreadView.cpp
@@ -447,7 +447,7 @@ void FRawProfilerSession::ProcessStatPacketArray( const FStatPacketArray& StatPa
 				continue;
 			}
 
-			if (Op == EStatOperation::CycleScopeStart || Op == EStatOperation::CycleScopeEnd || Op == EStatOperation::AdvanceFrameEventRenderThread)
+			if (Op == EStatOperation::CycleScopeStart || Op == EStatOperation::CycleScopeEnd)
 			{
 				//check( Item.NameAndInfo.GetFlag( EStatMetaFlags::IsCycle ) );
 				if (Op == EStatOperation::CycleScopeStart)
@@ -471,11 +471,7 @@ void FRawProfilerSession::ProcessStatPacketArray( const FStatPacketArray& StatPa
 					StartStack.Add( &Item );
 					Current = ChildNode;
 				}
-				// Workaround for render thread hierarchy. EStatOperation::AdvanceFrameEventRenderThread is called within the scope.
-				if (Op == EStatOperation::AdvanceFrameEventRenderThread)
-				{
-					int k = 0; k++;
-				}
+
 				if (Op == EStatOperation::CycleScopeEnd)
 				{
 					const FStatMessage ScopeStart = *StartStack.Pop();
@@ -486,11 +482,6 @@ void FRawProfilerSession::ProcessStatPacketArray( const FStatPacketArray& StatPa
 					Current->CycleCounterStartTimeMS = MetaData->ConvertCyclesToMS( uint32(Current->CyclesStart) );
 					Current->CycleCounterEndTimeMS = MetaData->ConvertCyclesToMS( uint32(Current->CyclesEnd) );
 
-					if (Current->CycleCounterStartTimeMS > Current->CycleCounterEndTimeMS)
-					{
-						int k = 0; k++;
-					}
-
 					check( Current->CycleCounterEndTimeMS >= Current->CycleCounterStartTimeMS );
 
 					FProfilerStackNode* ChildNode = Current;
diff --git a/Engine/Source/Runtime/Core/Private/Android/AndroidStats.cpp b/Engine/Source/Runtime/Core/Private/Android/AndroidStats.cpp
--- a/Engine/Source/Runtime/Core/Private/Android/AndroidStats.cpp
+++ b/Engine/Source/Runtime/Core/Private/Android/AndroidStats.cpp
@@ -72,28 +72,6 @@ DECLARE_FLOAT_COUNTER_STAT(TEXT("Thermal Status"), STAT_ThermalStatus, STATGROUP
 CSV_DEFINE_STAT(AndroidCPU, ThermalStress);
 DECLARE_FLOAT_COUNTER_STAT(TEXT("Thermal Stress"), STAT_ThermalStress, STATGROUP_AndroidCPU);
 
-#if STATS
-
-#define SET_DWORD_STAT_BY_FNAME(Stat, Amount) \
-{\
-	if (Amount != 0 && FThreadStats::IsCollectingData()) \
-	{ \
-		FThreadStats::AddMessage(Stat, EStatOperation::Set, int64(Amount));\
-		TRACE_STAT_SET(Stat, int64(Amount)); \
-	} \
-}
-
-#define SET_FLOAT_STAT_BY_FNAME(Stat, Amount) \
-{\
-	if (Amount != 0 && FThreadStats::IsCollectingData()) \
-	{ \
-		FThreadStats::AddMessage(Stat, EStatOperation::Set, double(Amount));\
-		TRACE_STAT_SET(Stat, double(Amount)); \
-	} \
-}
-
-#endif
-
 #if CSV_PROFILER_STATS
 
 #define CSV_STAT_PTR(StatName)									&_GCsvStat_##StatName
@@ -413,18 +391,18 @@ void FAndroidStats::UpdateAndroidStats()
 	for (int32 FrequencyGroupIndex = 0; FrequencyGroupIndex < FrequencyGroups.Num(); FrequencyGroupIndex++)
 	{
 		const FFrequencyGroup& FrequencyGroup = FrequencyGroups[FrequencyGroupIndex];
-		SET_DWORD_STAT_BY_FNAME(AndroidFrequencyGroupMaxFreqStats[FrequencyGroupIndex], FrequencyGroup.MaxFrequency);
-		SET_DWORD_STAT_BY_FNAME(AndroidFrequencyGroupNumCoresStats[FrequencyGroupIndex], FrequencyGroup.CoreCount);
-		//SET_DWORD_STAT_BY_FNAME(AndroidFrequencyGroupMinFreqStats[FrequencyGroupIndex], FrequencyGroup.MinFrequency);
-		SET_DWORD_STAT_BY_FNAME(AndroidFrequencyGroupCurrentFreqStats[FrequencyGroupIndex], CurrentFrequencies[FrequencyGroupIndex]);
-		SET_FLOAT_STAT_BY_FNAME(AndroidFrequencyGroupCurrentFreqPercentageStats[FrequencyGroupIndex], CurrentFrequencies[FrequencyGroupIndex]);
-		SET_FLOAT_STAT_BY_FNAME(AndroidFrequencyGroupMaxCoresUtilizationStats[FrequencyGroupIndex], MaxSingleCoreUtilization[FrequencyGroupIndex]);
+		SET_DWORD_STAT_FName(AndroidFrequencyGroupMaxFreqStats[FrequencyGroupIndex], FrequencyGroup.MaxFrequency);
+		SET_DWORD_STAT_FName(AndroidFrequencyGroupNumCoresStats[FrequencyGroupIndex], FrequencyGroup.CoreCount);
+		//SET_DWORD_STAT_FName(AndroidFrequencyGroupMinFreqStats[FrequencyGroupIndex], FrequencyGroup.MinFrequency);
+		SET_DWORD_STAT_FName(AndroidFrequencyGroupCurrentFreqStats[FrequencyGroupIndex], CurrentFrequencies[FrequencyGroupIndex]);
+		SET_FLOAT_STAT_FName(AndroidFrequencyGroupCurrentFreqPercentageStats[FrequencyGroupIndex], CurrentFrequencies[FrequencyGroupIndex]);
+		SET_FLOAT_STAT_FName(AndroidFrequencyGroupMaxCoresUtilizationStats[FrequencyGroupIndex], MaxSingleCoreUtilization[FrequencyGroupIndex]);
 	}
 
 	static const FName CPUStatName = GET_STATFNAME(STAT_CPUTemp);
 	static const FName ThermalStatus = GET_STATFNAME(STAT_ThermalStatus);
-	SET_FLOAT_STAT_BY_FNAME(CPUStatName, CPUTemp);
-	SET_FLOAT_STAT_BY_FNAME(ThermalStatus, GThermalStatus);
+	SET_FLOAT_STAT_FName(CPUStatName, CPUTemp);
+	SET_FLOAT_STAT_FName(ThermalStatus, GThermalStatus);
 #endif
 
 	UpdateGPUStats();
diff --git a/Engine/Source/Runtime/Core/Private/Stats/StatsData.cpp b/Engine/Source/Runtime/Core/Private/Stats/StatsData.cpp
--- a/Engine/Source/Runtime/Core/Private/Stats/StatsData.cpp
+++ b/Engine/Source/Runtime/Core/Private/Stats/StatsData.cpp
@@ -633,6 +633,7 @@ FStatsThreadState::FStatsThreadState(int32 InHistoryFrames)
 	, bFindMemoryExtensiveStats(false)
 	, CurrentGameFrame(1)
 	, CurrentRenderFrame(1)
+	, CurrentEndOfPipeFrame(1)
 {
 }
 
@@ -751,6 +752,28 @@ void FStatsThreadState::ScanForAdvance(const FStatMessagesArray& Data)
 			}
 			CurrentRenderFrame = NewRenderFrame;
 		}
+		else if (Op == EStatOperation::AdvanceFrameEventEndOfPipe)
+		{
+			check(Item.NameAndInfo.GetField<EStatDataType>() == EStatDataType::ST_int64);
+			int64 NewEndOfPipeFrame = Item.GetValue_int64();
+
+			if (NewEndOfPipeFrame < 0)
+			{
+				NewEndOfPipeFrame = -NewEndOfPipeFrame;
+				BadFrames.Add(NewEndOfPipeFrame - 1);
+			}
+
+			if (CurrentEndOfPipeFrame > STAT_FRAME_SLOP && CurrentEndOfPipeFrame + 1 != NewEndOfPipeFrame)
+			{
+				// this packet has multiple advances in it. They are all bad.
+				check(CurrentEndOfPipeFrame + 1 < NewEndOfPipeFrame);
+				for (int64 Frame = CurrentEndOfPipeFrame + 1; Frame <= NewEndOfPipeFrame; Frame++)
+				{
+					BadFrames.Add(Frame - 1);
+				}
+			}
+			CurrentEndOfPipeFrame = NewEndOfPipeFrame;
+		}
 	}
 
 	// We don't care about bad frame when the raw stats are active.
@@ -780,6 +803,10 @@ void FStatsThreadState::ScanForAdvance(FStatPacketArray& NewData)
 				Packet->AssignFrame( CurrentGameFrame );
 				break;
 
+			case EThreadType::EndOfPipe:
+				Packet->AssignFrame( CurrentEndOfPipeFrame );
+				break;
+
 			case EThreadType::Other:
 				// @see FThreadStats::DetectAndUpdateCurrentGameFrame 
 				break;
@@ -840,7 +867,8 @@ void FStatsThreadState::ProcessNonFrameStats(FStatMessagesArray& Data, TSet<FNam
 				Op != EStatOperation::ChildrenEnd &&
 				Op != EStatOperation::Leaf &&
 				Op != EStatOperation::AdvanceFrameEventGameThread &&
-				Op != EStatOperation::AdvanceFrameEventRenderThread
+				Op != EStatOperation::AdvanceFrameEventRenderThread &&
+				Op != EStatOperation::AdvanceFrameEventEndOfPipe
 				))
 			{
 				UE_LOG(LogStats, Fatal, TEXT( "Stat %s was not cleared every frame, but was used with a scope cycle counter." ), *Item.NameAndInfo.GetRawName().ToString() );
@@ -914,7 +942,7 @@ void FStatsThreadState::AddToHistoryAndEmpty(FStatPacketArray& NewData)
 	History.GenerateKeyArray(Frames);
 	Frames.Sort();
 
-	int64 LatestFinishedFrame = FMath::Min<int64>(CurrentGameFrame, CurrentRenderFrame) - 1;
+	int64 LatestFinishedFrame = FMath::Min<int64>(CurrentGameFrame, CurrentRenderFrame, CurrentEndOfPipeFrame) - 1;
 
 	for (int64 FrameNum : Frames)
 	{
@@ -946,8 +974,6 @@ void FStatsThreadState::AddToHistoryAndEmpty(FStatPacketArray& NewData)
 			if (PacketToCopyForNonFrame && !BadFrames.Contains(FrameNum))
 			{
 				// add the non frame stats as a new last packet
-
-				FThreadStats* ThreadStats = FThreadStats::GetThreadStats();
 				FStatPacket* NonFramePacket = nullptr;
 
 				for (const TPair<FName, FStatMessage>& It : NotClearedEveryFrame)
@@ -1050,7 +1076,7 @@ void FStatsThreadState::AddToHistoryAndEmpty(FStatPacketArray& NewData)
 	if (BadFrames.Num() > HistoryFrames * 2 + 5)
 	{
 		UE_LOG(LogStats, Display, TEXT("BadFrames.Num(): %d, History.Num(): %d, HistoryFrames: %d"), BadFrames.Num(), History.Num(), HistoryFrames);
-		UE_LOG(LogStats, Display, TEXT("CurrentGameFrame: %lld, CurrentRenderFrame: %lld, LastFullFrameMetaAndNonFrame: %lld"), CurrentGameFrame, CurrentRenderFrame, LastFullFrameMetaAndNonFrame);
+		UE_LOG(LogStats, Display, TEXT("CurrentGameFrame: %lld, CurrentRenderFrame: %lld, CurrentEndOfPipeFrame: %lld, LastFullFrameMetaAndNonFrame: %lld"), CurrentGameFrame, CurrentRenderFrame, CurrentEndOfPipeFrame, LastFullFrameMetaAndNonFrame);
 		for (int64 BadFrame : BadFrames)
 		{
 			UE_LOG(LogStats, Display, TEXT("Bad frame: %lld"), BadFrame);
@@ -1690,6 +1716,11 @@ FName FStatsThreadState::GetStatThreadName( const FStatPacket& Packet ) const
 	{
 		ThreadName = NAME_RenderThread;
 	}
+	else if ( Packet.ThreadType == EThreadType::EndOfPipe )
+	{
+		static const FName NAME_EndOfPipe = TEXT("EndOfPipe");
+		ThreadName = NAME_EndOfPipe;
+	}
 	else if( Packet.ThreadType == EThreadType::Other )
 	{
 		TMap<uint32, FName>& MutableThreads = const_cast<TMap<uint32, FName>&>(Threads);
@@ -1714,6 +1745,7 @@ void FStatsThreadState::Condense(int64 TargetFrame, TArray<FStatMessage>& OutSta
 {
 	OutStats.Emplace(FStatConstants::AdvanceFrame.GetEncodedName(), EStatOperation::AdvanceFrameEventGameThread, TargetFrame, false);
 	OutStats.Emplace(FStatConstants::AdvanceFrame.GetEncodedName(), EStatOperation::AdvanceFrameEventRenderThread, TargetFrame, false);
+	OutStats.Emplace(FStatConstants::AdvanceFrame.GetEncodedName(), EStatOperation::AdvanceFrameEventEndOfPipe, TargetFrame, false);
 	FRawStatStackNode Root;
 	GetRawStackStats(TargetFrame, Root, &OutStats);
 	TArray<FStatMessage> StackStats;
diff --git a/Engine/Source/Runtime/Core/Private/Stats/StatsFile.cpp b/Engine/Source/Runtime/Core/Private/Stats/StatsFile.cpp
--- a/Engine/Source/Runtime/Core/Private/Stats/StatsFile.cpp
+++ b/Engine/Source/Runtime/Core/Private/Stats/StatsFile.cpp
@@ -168,6 +168,7 @@ void FStatsLoadedState::AdvanceFrameForLoad( TArray<FStatMessage>& CondensedMess
 
 	CurrentGameFrame = AdvanceGameThread.GetValue_int64();
 	CurrentRenderFrame = CurrentGameFrame;
+	CurrentEndOfPipeFrame = CurrentGameFrame;
 
 	if (CurrentGameFrame > MaxFrameSeen)
 	{
diff --git a/Engine/Source/Runtime/Core/Private/Stats/StatsSystem.cpp b/Engine/Source/Runtime/Core/Private/Stats/StatsSystem.cpp
--- a/Engine/Source/Runtime/Core/Private/Stats/StatsSystem.cpp
+++ b/Engine/Source/Runtime/Core/Private/Stats/StatsSystem.cpp
@@ -83,6 +83,7 @@ namespace UE::Stats
 		{
 			// There is no rendering thread, so this message is sufficient to make stats happy and don't leak memory.
 			FThreadStats::AddMessage(FStatConstants::AdvanceFrame.GetEncodedName(), EStatOperation::AdvanceFrameEventRenderThread, Frame);
+			FThreadStats::AddMessage(FStatConstants::AdvanceFrame.GetEncodedName(), EStatOperation::AdvanceFrameEventEndOfPipe, Frame);
 		}
 
 		FThreadStats::ExplicitFlush(bDiscardCallstack);
diff --git a/Engine/Source/Runtime/Core/Private/Stats/StatsSystemTypes.cpp b/Engine/Source/Runtime/Core/Private/Stats/StatsSystemTypes.cpp
--- a/Engine/Source/Runtime/Core/Private/Stats/StatsSystemTypes.cpp
+++ b/Engine/Source/Runtime/Core/Private/Stats/StatsSystemTypes.cpp
@@ -933,11 +933,6 @@ private:
 
 FThreadStatsPool::FThreadStatsPool()
 {
-	LLM_SCOPE(ELLMTag::Stats);
-	for (int32 Index = 0; Index < NUM_ELEMENTS_IN_POOL; ++Index)
-	{
-		Pool.Push(new FThreadStats(EConstructor::FOR_POOL));
-	}
 }
 
 FThreadStatsPool& FThreadStatsPool::Get()
@@ -949,14 +944,13 @@ FThreadStatsPool& FThreadStatsPool::Get()
 FThreadStats* FThreadStatsPool::GetFromPool()
 {
 	LLM_SCOPE(ELLMTag::Stats);
-	FPlatformMisc::MemoryBarrier();
-	FThreadStats* Address = Pool.Pop();
-	while (!Address)
+
+	if (void* Address = Pool.Pop())
 	{
-		Pool.Push(new FThreadStats(EConstructor::FOR_POOL));
-		Address = Pool.Pop();
+		return new (Address) FThreadStats;
 	}
-	return new(Address) FThreadStats();
+
+	return new FThreadStats;
 }
 
 void FThreadStatsPool::ReturnToPool(FThreadStats* Instance)
@@ -977,34 +971,32 @@ FThreadSafeCounter FThreadStats::PrimaryDisableChangeTagLock;
 bool FThreadStats::bPrimaryEnable = false;
 bool FThreadStats::bPrimaryDisableForever = false;
 bool FThreadStats::bIsRawStatsActive = false;
+FEndOfPipeStats FEndOfPipeStats::EndOfPipeStats;
 
-FThreadStats::FThreadStats() :
-	CurrentGameFrame(UE::Stats::FStats::GameThreadStatsFrame.Load(EMemoryOrder::Relaxed)),
-	ScopeCount(0),
-	bWaitForExplicitFlush(0),
-	MemoryMessageScope(0),
-	bReentranceGuard(false),
-	bSawExplicitFlush(false)
+FThreadStats::FThreadStats(EThreadType::Type InThreadType)
+	: Packet(InThreadType)
+	, CurrentGameFrame(UE::Stats::FStats::GameThreadStatsFrame.Load(EMemoryOrder::Relaxed))
 {
-	Packet.SetThreadProperties();
+	if (Packet.ThreadType == EThreadType::Invalid)
+	{
+		Packet.SetThreadProperties();
+	}
 
-	check(FPlatformTLS::IsValidTlsSlot(TlsSlot));
-	FPlatformTLS::SetTlsValue(TlsSlot, this);
+	if (Packet.ThreadType != EThreadType::EndOfPipe)
+	{
+		check(FPlatformTLS::IsValidTlsSlot(TlsSlot));
+		FPlatformTLS::SetTlsValue(TlsSlot, this);
+	}
 }
 
-FThreadStats::FThreadStats(EConstructor) :
-	CurrentGameFrame(-1),
-	ScopeCount(0),
-	bWaitForExplicitFlush(0),
-	MemoryMessageScope(0),
-	bReentranceGuard(false),
-	bSawExplicitFlush(false)
-{}
-
 void FThreadStats::CheckEnable()
 {
-	bool bOldPrimaryEnable(bPrimaryEnable);
-	bool bNewPrimaryEnable(WillEverCollectData() && (!IsRunningCommandlet() || UE::Stats::FStats::EnabledForCommandlet()) && IsThreadingReady() && (PrimaryEnableCounter.GetValue()));
+	bool bNewPrimaryEnable = 
+		   WillEverCollectData()
+		&& (!IsRunningCommandlet() || UE::Stats::FStats::EnabledForCommandlet())
+		&& IsThreadingReady()
+		&& PrimaryEnableCounter.GetValue();
+
 	if (bPrimaryEnable != bNewPrimaryEnable)
 	{
 		PrimaryDisableChangeTagLockAdd();
diff --git a/Engine/Source/Runtime/Core/Public/Stats/Stats.h b/Engine/Source/Runtime/Core/Public/Stats/Stats.h
--- a/Engine/Source/Runtime/Core/Public/Stats/Stats.h
+++ b/Engine/Source/Runtime/Core/Public/Stats/Stats.h
@@ -269,7 +269,6 @@ struct FStat_##StatName\
 	{ \
 		const FName StatName = GET_STATFNAME(Stat); \
 		FThreadStats::AddMessage(StatName, EStatOperation::Add, int64(1));\
-		TRACE_STAT_INCREMENT(StatName); \
 	} \
 }
 #define INC_FLOAT_STAT_BY(Stat, Amount) \
@@ -281,7 +280,6 @@ struct FStat_##StatName\
 		{ \
 			const FName StatName = GET_STATFNAME(Stat); \
 			FThreadStats::AddMessage(StatName, EStatOperation::Add, AddAmount); \
-			TRACE_STAT_ADD(StatName, AddAmount); \
 		} \
 	} \
 }
@@ -294,7 +292,6 @@ struct FStat_##StatName\
 		{ \
 			const FName StatName = GET_STATFNAME(Stat); \
 			FThreadStats::AddMessage(StatName, EStatOperation::Add, AddAmount); \
-			TRACE_STAT_ADD(StatName, AddAmount); \
 		} \
 	} \
 }
@@ -307,7 +304,6 @@ struct FStat_##StatName\
 		{ \
 			const FName StatName = GET_STATFNAME(Stat); \
 			FThreadStats::AddMessage(StatName, EStatOperation::Add, AddAmount); \
-			TRACE_STAT_ADD(StatName, AddAmount); \
 		} \
 	} \
 }
@@ -317,7 +313,6 @@ struct FStat_##StatName\
 	{ \
 		const FName StatName = GET_STATFNAME(Stat); \
 		FThreadStats::AddMessage(StatName, EStatOperation::Subtract, int64(1));\
-		TRACE_STAT_DECREMENT(StatName); \
 	} \
 }
 #define DEC_FLOAT_STAT_BY(Stat,Amount) \
@@ -329,7 +324,6 @@ struct FStat_##StatName\
 		{ \
 			const FName StatName = GET_STATFNAME(Stat); \
 			FThreadStats::AddMessage(StatName, EStatOperation::Subtract, SubtractAmount); \
-			TRACE_STAT_ADD(StatName, -SubtractAmount); \
 		} \
 	} \
 }
@@ -342,7 +336,6 @@ struct FStat_##StatName\
 		{ \
 			const FName StatName = GET_STATFNAME(Stat); \
 			FThreadStats::AddMessage(StatName, EStatOperation::Subtract, SubtractAmount); \
-			TRACE_STAT_ADD(StatName, -SubtractAmount); \
 		} \
 	} \
 }
@@ -355,7 +348,6 @@ struct FStat_##StatName\
 		{ \
 			const FName StatName = GET_STATFNAME(Stat); \
 			FThreadStats::AddMessage(StatName, EStatOperation::Subtract, SubtractAmount); \
-			TRACE_STAT_ADD(StatName, -SubtractAmount); \
 		} \
 	} \
 }
@@ -366,7 +358,6 @@ struct FStat_##StatName\
 		const FName StatName = GET_STATFNAME(Stat); \
 		int64 SetValue = int64(Value); \
 		FThreadStats::AddMessage(StatName, EStatOperation::Set, SetValue); \
-		TRACE_STAT_SET(StatName, SetValue); \
 	} \
 }
 #define SET_DWORD_STAT(Stat,Value) \
@@ -376,7 +367,6 @@ struct FStat_##StatName\
 		const FName StatName = GET_STATFNAME(Stat); \
 		int64 SetValue = int64(Value); \
 		FThreadStats::AddMessage(StatName, EStatOperation::Set, SetValue); \
-		TRACE_STAT_SET(StatName, SetValue); \
 	} \
 }
 #define SET_FLOAT_STAT(Stat,Value) \
@@ -386,7 +376,6 @@ struct FStat_##StatName\
 		const FName StatName = GET_STATFNAME(Stat); \
 		double SetValue = double(Value); \
 		FThreadStats::AddMessage(StatName, EStatOperation::Set, SetValue); \
-		TRACE_STAT_SET(StatName, SetValue); \
 	} \
 }
 
@@ -407,7 +396,6 @@ struct FStat_##StatName\
 #define INC_DWORD_STAT_FName(Stat) \
 {\
 	FThreadStats::AddMessage(Stat, EStatOperation::Add, int64(1));\
-	TRACE_STAT_INCREMENT(Stat); \
 }
 #define INC_FLOAT_STAT_BY_FName(Stat, Amount) \
 {\
@@ -415,7 +403,6 @@ struct FStat_##StatName\
 	if (AddAmount != 0.0) \
 	{ \
 		FThreadStats::AddMessage(Stat, EStatOperation::Add, AddAmount); \
-		TRACE_STAT_ADD(Stat, AddAmount); \
 	} \
 }
 #define INC_DWORD_STAT_BY_FName(Stat, Amount) \
@@ -424,7 +411,6 @@ struct FStat_##StatName\
 	if (AddAmount != 0) \
 	{ \
 		FThreadStats::AddMessage(Stat, EStatOperation::Add, AddAmount); \
-		TRACE_STAT_ADD(Stat, AddAmount); \
 	} \
 }
 #define INC_DWORD_STAT_FNAME_BY(Stat, Amount) INC_DWORD_STAT_BY_FName(Stat, Amount)
@@ -434,13 +420,11 @@ struct FStat_##StatName\
 	if (AddAmount != 0) \
 	{ \
 		FThreadStats::AddMessage(Stat, EStatOperation::Add, AddAmount); \
-		TRACE_STAT_ADD(Stat, AddAmount); \
 	} \
 }
 #define DEC_DWORD_STAT_FName(Stat) \
 {\
 	FThreadStats::AddMessage(Stat, EStatOperation::Subtract, int64(1));\
-	TRACE_STAT_DECREMENT(Stat); \
 }
 #define DEC_FLOAT_STAT_BY_FName(Stat,Amount) \
 {\
@@ -448,7 +432,6 @@ struct FStat_##StatName\
 	if (SubtractAmount != 0.0) \
 	{ \
 		FThreadStats::AddMessage(Stat, EStatOperation::Subtract, SubtractAmount); \
-		TRACE_STAT_ADD(Stat, -SubtractAmount); \
 	} \
 }
 #define DEC_DWORD_STAT_BY_FName(Stat,Amount) \
@@ -457,7 +440,6 @@ struct FStat_##StatName\
 	if (SubtractAmount != 0) \
 	{ \
 		FThreadStats::AddMessage(Stat, EStatOperation::Subtract, SubtractAmount); \
-		TRACE_STAT_ADD(Stat, -SubtractAmount); \
 	} \
 }
 #define DEC_DWORD_STAT_FNAME_BY(Stat,Amount) DEC_DWORD_STAT_BY_FName(Stat,Amount)
@@ -467,26 +449,22 @@ struct FStat_##StatName\
 	if (SubtractAmount != 0) \
 	{ \
 		FThreadStats::AddMessage(Stat, EStatOperation::Subtract, SubtractAmount); \
-		TRACE_STAT_ADD(Stat, -SubtractAmount); \
 	} \
 }
 #define SET_MEMORY_STAT_FName(Stat,Value) \
 {\
 	int64 SetValue = int64(Value); \
 	FThreadStats::AddMessage(Stat, EStatOperation::Set, SetValue); \
-	TRACE_STAT_SET(Stat, SetValue); \
 }
 #define SET_DWORD_STAT_FName(Stat,Value) \
 {\
 	int64 SetValue = int64(Value); \
 	FThreadStats::AddMessage(Stat, EStatOperation::Set, SetValue); \
-	TRACE_STAT_SET(Stat, SetValue); \
 }
 #define SET_FLOAT_STAT_FName(Stat,Value) \
 {\
 	double SetValue = double(Value); \
 	FThreadStats::AddMessage(Stat, EStatOperation::Set, SetValue); \
-	TRACE_STAT_SET(Stat, SetValue); \
 }
 
 #elif UE_USE_LIGHTWEIGHT_STATS // STATS
diff --git a/Engine/Source/Runtime/Core/Public/Stats/StatsData.h b/Engine/Source/Runtime/Core/Public/Stats/StatsData.h
--- a/Engine/Source/Runtime/Core/Public/Stats/StatsData.h
+++ b/Engine/Source/Runtime/Core/Public/Stats/StatsData.h
@@ -612,6 +612,8 @@ public:
 	/** New packets on the render thread as assign this frame **/
 	int64 CurrentRenderFrame;
 
+	int64 CurrentEndOfPipeFrame;
+
 	/** List of the stats, with data of all of the things that are NOT cleared every frame. We need to accumulate these over the entire life of the executable. **/
 	TMap<FName, FStatMessage> NotClearedEveryFrame;
 
diff --git a/Engine/Source/Runtime/Core/Public/Stats/StatsSystemTypes.h b/Engine/Source/Runtime/Core/Public/Stats/StatsSystemTypes.h
--- a/Engine/Source/Runtime/Core/Public/Stats/StatsSystemTypes.h
+++ b/Engine/Source/Runtime/Core/Public/Stats/StatsSystemTypes.h
@@ -164,7 +164,7 @@ private:
 
 /**
  * For packet messages, this indicated what sort of thread timing we use.
- * Game and Other use CurrentGameFrame, Renderer is CurrentRenderFrame
+ * Game and Other use CurrentGameFrame, Renderer is CurrentRenderFrame, EndOfPipe is CurrentEndOfPipeFrame
  */
 namespace EThreadType
 {
@@ -173,6 +173,7 @@ namespace EThreadType
 		Invalid,
 		Game,
 		Renderer,
+		EndOfPipe,
 		Other,
 	};
 }
@@ -218,6 +219,8 @@ struct EStatOperation
 		AdvanceFrameEventGameThread,
 		/** Special message for advancing the stats frame from the render thread. */
 		AdvanceFrameEventRenderThread,
+		/** Special message for advancing the stats frame from the end of pipe thread. */
+		AdvanceFrameEventEndOfPipe,
 		/** Indicates begin of the cycle scope. */
 		CycleScopeStart,
 		/** Indicates end of the cycle scope. */
@@ -1043,24 +1046,21 @@ typedef TChunkedArray<FStatMessage, (uint32)EStatMessagesArrayConstants::MESSAGE
 struct FStatPacket
 {
 	/** Assigned later, this is the frame number this packet is for. @see FStatsThreadState::ScanForAdvance or FThreadStats::DetectAndUpdateCurrentGameFrame */
-	int64 Frame;
+	int64 Frame = 1;
 	/** ThreadId this packet came from **/
-	uint32 ThreadId;
+	uint32 ThreadId = 0;
 	/** type of thread this packet came from **/
 	EThreadType::Type ThreadType;
 	/** true if this packet has broken callstacks **/
-	bool bBrokenCallstacks;
+	bool bBrokenCallstacks = false;
 	/** messages in this packet **/
 	FStatMessagesArray StatMessages;
 	/** Size we presize the message buffer to, currently the max of what we have seen for the last PRESIZE_MAX_NUM_ENTRIES. **/
 	TArray<int32> StatMessagesPresize;
 
 	/** constructor **/
-	FStatPacket()
-		: Frame(1)
-		, ThreadId(0)
-		, ThreadType(EThreadType::Invalid)
-		, bBrokenCallstacks(false)
+	FStatPacket(EThreadType::Type InThreadType = EThreadType::Invalid)
+		: ThreadType(InThreadType)
 	{
 	}
 
@@ -1094,7 +1094,6 @@ struct FStatPacket
 		PRAGMA_ENABLE_DEPRECATION_WARNINGS
 	}
 
-
 	void AssignFrame(int64 InFrame)
 	{
 		Frame = InFrame;
@@ -1123,18 +1122,8 @@ protected:
 struct FThreadStatsPool
 {
 private:
-	enum
-	{
-		/**
-		 *	Number of elements preallocated in the pool.
-		 *  Should be enough to handle all threads.
-		 *  Increase if you encounter a crash.
-		 */
-		NUM_ELEMENTS_IN_POOL = 512
-	};
-
 	/** Lock free pool of FThreadStats instances. */
-	TLockFreePointerListUnordered<FThreadStats, 0> Pool;
+	TLockFreePointerListUnordered<void, 0> Pool;
 
 public:
 	/** Default constructor. */
@@ -1150,12 +1139,6 @@ public:
 	CORE_API void ReturnToPool(FThreadStats* Instance);
 };
 
-/** Fake type to distinguish constructors. */
-enum EConstructor
-{
-	FOR_POOL
-};
-
 /**
 * This is thread-private information about the stats we are acquiring. Pointers to these objects are held in TLS.
 */
@@ -1188,31 +1171,27 @@ class FThreadStats : FNoncopyable
 	int32 CurrentGameFrame;
 
 	/** Tracks current stack depth for cycle counters. **/
-	int32 ScopeCount;
+	int32 ScopeCount = 0;
 
 	/** Tracks current stack depth for cycle counters. **/
-	int32 bWaitForExplicitFlush;
+	int32 bWaitForExplicitFlush = 0;
 
 	/**
 	 * Tracks the progress of adding a new memory stat message and prevents from using the memory profiler in the scope.
 	 * Mostly related to ignoring all memory allocations in AddStatMessage().
 	 * Memory usage of the stats messages is handled by the STAT_StatMessagesMemory.
 	 */
-	int32 MemoryMessageScope;
+	int32 MemoryMessageScope = 0;
 
 	/** We have a scoped cycle counter in the FlushRawStats which causes an infinite recursion, ReentranceGuard will solve that issue. */
-	bool bReentranceGuard;
+	bool bReentranceGuard = false;
 
 	/** Tracks current stack depth for cycle counters. **/
-	bool bSawExplicitFlush;
+	bool bSawExplicitFlush = false;
 
+protected:
 	/** Gathers information about the current thread and sets up the TLS value. **/
-	CORE_API FThreadStats();
-
-	/** Constructor used for the pool. */
-	CORE_API FThreadStats(EConstructor);
-
-	void SendMessage_Async(FStatPacket* ToSend);
+	CORE_API FThreadStats(EThreadType::Type InThreadType = EThreadType::Invalid);
 
 public:
 	/** Checks the TLS for a thread packet and if it isn't found, it makes a new one. **/
@@ -1259,6 +1238,9 @@ public:
 	/** Checks the command line whether we want to enable collecting startup stats. */
 	static CORE_API void CheckForCollectingStartupStats();
 
+private:
+	void SendMessage_Async(FStatPacket* ToSend);
+
 	FORCEINLINE_STATS void AddStatMessage(const FStatMessage& StatMessage)
 	{
 		LLM_SCOPE(ELLMTag::Stats);
@@ -1266,71 +1248,98 @@ public:
 		Packet.StatMessages.AddElement(StatMessage);
 	}
 
-public:
-	/** This should be called when a thread exits, this deletes FThreadStats from the heap and TLS. **/
-	static void Shutdown()
+protected:
+	/** Any non-clock operation with an ordinary payload. **/
+	template <typename TValue>
+	FORCEINLINE_STATS void AddMessageInner(FName InStatName, EStatOperation::Type InStatOperation, TValue Value, bool bIsCycle = false)
 	{
-		FThreadStats* Stats = IsThreadingReady() ? (FThreadStats*)FPlatformTLS::GetTlsValue(TlsSlot) : nullptr;
-		if (Stats)
+		AddStatMessage(FStatMessage(InStatName, InStatOperation, Value, bIsCycle));
+
+		if constexpr (std::is_same_v<std::decay_t<TValue>, double> || std::is_same_v<std::decay_t<TValue>, int64>)
 		{
-			// Send all remaining messages.
-			Stats->Flush(false, true);
-			FPlatformTLS::SetTlsValue(TlsSlot, nullptr);
-			FThreadStatsPool::Get().ReturnToPool(Stats);
+			switch (InStatOperation)
+			{
+			case EStatOperation::Set:      TRACE_STAT_SET(InStatName,  Value); break;
+			case EStatOperation::Add:      TRACE_STAT_ADD(InStatName,  Value); break;
+			case EStatOperation::Subtract: TRACE_STAT_ADD(InStatName, -Value); break;
+			}
+		}
+
+		if (!ScopeCount)
+		{
+			Flush();
+		}
+		else if (bIsRawStatsActive)
+		{
+			FlushRawStats();
 		}
 	}
 
 	/** Clock operation. **/
-	static FORCEINLINE_STATS void AddMessage(FName InStatName, EStatOperation::Type InStatOperation)
+	FORCEINLINE_STATS void AddMessageInner(FName InStatName, EStatOperation::Type InStatOperation)
 	{
 		checkStats((InStatOperation == EStatOperation::CycleScopeStart || InStatOperation == EStatOperation::CycleScopeEnd));
-		FThreadStats* ThreadStats = GetThreadStats();
+
 		// these branches are handled by the optimizer
 		if (InStatOperation == EStatOperation::CycleScopeStart)
 		{
-			ThreadStats->ScopeCount++;
-			ThreadStats->AddStatMessage(FStatMessage(InStatName, InStatOperation));
+			ScopeCount++;
+			AddStatMessage(FStatMessage(InStatName, InStatOperation));
 
 			if (bIsRawStatsActive)
 			{
-				ThreadStats->FlushRawStats();
+				FlushRawStats();
 			}
 		}
 		else if (InStatOperation == EStatOperation::CycleScopeEnd)
 		{
-			if (ThreadStats->ScopeCount > ThreadStats->bWaitForExplicitFlush)
+			if (ScopeCount > bWaitForExplicitFlush)
 			{
-				ThreadStats->AddStatMessage(FStatMessage(InStatName, InStatOperation));
-				ThreadStats->ScopeCount--;
-				if (!ThreadStats->ScopeCount)
+				AddStatMessage(FStatMessage(InStatName, InStatOperation));
+				ScopeCount--;
+				if (!ScopeCount)
 				{
-					ThreadStats->Flush();
+					Flush();
 				}
 				else if (bIsRawStatsActive)
 				{
-					ThreadStats->FlushRawStats();
+					FlushRawStats();
 				}
 			}
 			// else we dumped this frame without closing scope, so we just drop the closes on the floor
 		}
 	}
 
+public:
+	/** This should be called when a thread exits, this deletes FThreadStats from the heap and TLS. **/
+	static void Shutdown()
+	{
+		FThreadStats* Stats = IsThreadingReady() ? (FThreadStats*)FPlatformTLS::GetTlsValue(TlsSlot) : nullptr;
+		if (Stats)
+		{
+			// Send all remaining messages.
+			Stats->Flush(false, true);
+			FPlatformTLS::SetTlsValue(TlsSlot, nullptr);
+			FThreadStatsPool::Get().ReturnToPool(Stats);
+		}
+	}
+
+	/** Clock operation. **/
+	static FORCEINLINE_STATS void AddMessage(FName InStatName, EStatOperation::Type InStatOperation)
+	{
+		if (!InStatName.IsNone() && WillEverCollectData() && IsThreadingReady())
+		{
+			GetThreadStats()->AddMessageInner(InStatName, InStatOperation);
+		}
+	}
+
 	/** Any non-clock operation with an ordinary payload. **/
 	template<typename TValue>
 	static FORCEINLINE_STATS void AddMessage(FName InStatName, EStatOperation::Type InStatOperation, TValue Value, bool bIsCycle = false)
 	{
 		if (!InStatName.IsNone() && WillEverCollectData() && IsThreadingReady())
 		{
-			FThreadStats* ThreadStats = GetThreadStats();
-			ThreadStats->AddStatMessage(FStatMessage(InStatName, InStatOperation, Value, bIsCycle));
-			if (!ThreadStats->ScopeCount)
-			{
-				ThreadStats->Flush();
-			}
-			else if (bIsRawStatsActive)
-			{
-				ThreadStats->FlushRawStats();
-			}
+			GetThreadStats()->AddMessageInner(InStatName, InStatOperation, Value, bIsCycle);
 		}
 	}
 
@@ -1422,41 +1431,6 @@ public:
 		return PrimaryEnableUpdateNumber.GetValue();
 	}
 
-	/** Indicate that you would like the system to begin collecting data, if it isn't already collecting data. Think reference count. **/
-	UE_DEPRECATED(5.1, "Use PrimaryEnableAdd instead")
-		static FORCEINLINE_STATS void MasterEnableAdd(int32 Value = 1)
-	{
-		PrimaryEnableAdd(Value);
-	}
-
-	/** Indicate that you no longer need stat data, if nobody else needs stat data, then no stat data will be collected. Think reference count. **/
-	UE_DEPRECATED(5.1, "Use PrimaryEnableSubtract instead")
-		static FORCEINLINE_STATS void MasterEnableSubtract(int32 Value = 1)
-	{
-		PrimaryEnableSubtract(Value);
-	}
-
-	/** Indicate that you no longer need stat data, forever. **/
-	UE_DEPRECATED(5.1, "Use PrimaryDisableForever instead")
-		static FORCEINLINE_STATS void MasterDisableForever()
-	{
-		PrimaryDisableForever();
-	}
-
-	/** This is called before we start to change something that will invalidate. **/
-	UE_DEPRECATED(5.1, "Use PrimaryDisableChangeTagLockAdd instead")
-		static FORCEINLINE_STATS void MasterDisableChangeTagLockAdd(int32 Value = 1)
-	{
-		PrimaryDisableChangeTagLockAdd(Value);
-	}
-
-	/** Indicate that you no longer need stat data, if nobody else needs stat data, then no stat data will be collected. Think reference count. **/
-	UE_DEPRECATED(5.1, "Use PrimaryDisableChangeTagLockSubtract instead")
-		static FORCEINLINE_STATS void MasterDisableChangeTagLockSubtract(int32 Value = 1)
-	{
-		PrimaryDisableChangeTagLockSubtract(Value);
-	}
-
 	/** Call this if something disrupts data gathering. For example when the render thread is killed, data is abandoned.**/
 	static FORCEINLINE_STATS void FrameDataIsIncomplete()
 	{
@@ -1487,6 +1461,38 @@ public:
 	static CORE_API void WaitForStats();
 };
 
+/*
+ * Wrapper used by the end-of-pipe tasks to report stats on the appropriate timeline.
+ * Acts as a singleton instance of FThreadStats, since the end-of-pipe is a logical part of the frame, not a specific thread.
+ */
+class FEndOfPipeStats : private FThreadStats
+{
+	static CORE_API FEndOfPipeStats EndOfPipeStats;
+
+	FEndOfPipeStats()
+		: FThreadStats(EThreadType::EndOfPipe)
+	{}
+
+public:
+	static FEndOfPipeStats* Get()
+	{
+		return &EndOfPipeStats;
+	}
+
+	template<typename TValue>
+	void AddMessage(FName InStatName, EStatOperation::Type InStatOperation, TValue Value, bool bIsCycle = false)
+	{
+		if (!InStatName.IsNone() && WillEverCollectData() && IsThreadingReady())
+		{
+			FThreadStats::AddMessageInner(InStatName, InStatOperation, Value, bIsCycle);
+		}
+	}
+
+	void Flush()
+	{
+		FThreadStats::Flush(false, true);
+	}
+};
 
 /**
  * This is a utility class for counting the number of cycles during the
diff --git a/Engine/Source/Runtime/RenderCore/Private/ProfilingDebugging/RealtimeGPUProfiler.cpp b/Engine/Source/Runtime/RenderCore/Private/ProfilingDebugging/RealtimeGPUProfiler.cpp
--- a/Engine/Source/Runtime/RenderCore/Private/ProfilingDebugging/RealtimeGPUProfiler.cpp
+++ b/Engine/Source/Runtime/RenderCore/Private/ProfilingDebugging/RealtimeGPUProfiler.cpp
@@ -498,12 +498,10 @@ public:
 				if (bKnownStat)
 				{
 					FThreadStats::AddMessage(Event.GetStatName(), EStatOperation::Add, EventTimeMs);
-					TRACE_STAT_ADD(Event.GetStatName(), EventTimeMs);
 				}
 				else
 				{
 					FThreadStats::AddMessage(Event.GetStatName(), EStatOperation::Set, EventTimeMs);
-					TRACE_STAT_SET(Event.GetStatName(), EventTimeMs);
 				}
 #endif
 
@@ -521,7 +519,6 @@ public:
 
 #if STATS
 		FThreadStats::AddMessage(GET_STATFNAME(Stat_GPU_Total), EStatOperation::Set, TotalMs);
-		TRACE_STAT_SET(GET_STATFNAME(Stat_GPU_Total), TotalMs);
 #endif
 
 #if CSV_PROFILER_STATS
diff --git a/Engine/Source/Runtime/RenderCore/Private/RenderingThread.cpp b/Engine/Source/Runtime/RenderCore/Private/RenderingThread.cpp
--- a/Engine/Source/Runtime/RenderCore/Private/RenderingThread.cpp
+++ b/Engine/Source/Runtime/RenderCore/Private/RenderingThread.cpp
@@ -344,6 +344,10 @@ static void AdvanceRenderingThreadStats(int64 StatsFrame, int32 DisableChangeTag
 		Frame = -StatsFrame; // mark this as a bad frame
 	}
 	FThreadStats::AddMessage(FStatConstants::AdvanceFrame.GetEncodedName(), EStatOperation::AdvanceFrameEventRenderThread, Frame);
+
+	// @todo remove this when the new GPU profiler is using the end-of-pipe stats timeline
+	FThreadStats::AddMessage(FStatConstants::AdvanceFrame.GetEncodedName(), EStatOperation::AdvanceFrameEventEndOfPipe, Frame);
+
 	if( IsInActualRenderingThread() )
 	{
 		FThreadStats::ExplicitFlush();
