[Scheduler]
- Add a local queue to the game thread to improve its task queuing efficiency
- Reduce enqueuer starvation on FAAArrayQueue when there are lots of dequeuer
- Provides a 2x GT tasks queuing efficiency in micro-benchmarks and noticable improvements of GT time in real world replay runs

diff --git a/Engine/Source/Runtime/Core/Private/Async/Fundamental/Scheduler.cpp b/Engine/Source/Runtime/Core/Private/Async/Fundamental/Scheduler.cpp
--- a/Engine/Source/Runtime/Core/Private/Async/Fundamental/Scheduler.cpp
+++ b/Engine/Source/Runtime/Core/Private/Async/Fundamental/Scheduler.cpp
@@ -84,6 +84,8 @@ namespace LowLevelTasks
 
 	void FScheduler::StartWorkers(uint32 NumForegroundWorkers, uint32 NumBackgroundWorkers, FThread::EForkable IsForkable, EThreadPriority InWorkerPriority,  EThreadPriority InBackgroundPriority, uint64 InWorkerAffinity, uint64 InBackgroundAffinity)
 	{
+		// It always been a given that only the game thread should start workers so just add validation.
+		check(IsInGameThread());
 		int32 Value = 0;
 		if (FParse::Value(FCommandLine::Get(), TEXT("TaskGraphUseDynamicPrioritization="), Value))
 		{
@@ -132,6 +134,12 @@ namespace LowLevelTasks
 			const int32 MaxWorkers = MaxForegroundWorkers + MaxBackgroundWorkers;
 			const EThreadPriority ActualBackgroundPriority = GTaskGraphUseDynamicPrioritization ? WorkerPriority : BackgroundPriority;
 
+			if (GameThreadLocalQueue == nullptr)
+			{
+				GameThreadLocalQueue = MakeUnique<FSchedulerTls::FLocalQueueType>(QueueRegistry, Private::ELocalQueueType::EForeground);
+			}
+			FSchedulerTls::LocalQueue = GameThreadLocalQueue.Get();
+
 			WorkerEvents.SetNum(MaxWorkers);
 			WorkerLocalQueues.Reserve(MaxWorkers);
 			WorkerThreads.SetNum(MaxWorkers);
@@ -284,6 +292,8 @@ namespace LowLevelTasks
 
 	void FScheduler::StopWorkers(bool bDrainGlobalQueue)
 	{
+		// It always been a given that only the game thread should stop workers so just add validation.
+		check(IsInGameThread());
 		uint32 OldActiveWorkers = ActiveWorkers.load(std::memory_order_relaxed);
 		if (OldActiveWorkers != 0 && ActiveWorkers.compare_exchange_strong(OldActiveWorkers, 0, std::memory_order_relaxed))
 		{
@@ -306,6 +316,9 @@ namespace LowLevelTasks
 			WaitingQueue[0].FinishShutdown();
 			WaitingQueue[1].FinishShutdown();
 
+			GameThreadLocalQueue.Reset();
+			FSchedulerTls::LocalQueue = nullptr;
+
 			NextWorkerId = 0;
 			WorkerThreads.Reset();
 			WorkerLocalQueues.Reset();
@@ -346,6 +359,7 @@ namespace LowLevelTasks
 			const bool bIsBackgroundTask = Task.IsBackgroundTask();
 			const bool bIsBackgroundWorker = FSchedulerTls::IsBackgroundWorker();
 			const bool bIsStandbyWorker = FSchedulerTls::IsStandbyWorker();
+			FSchedulerTls::FLocalQueueType* const CachedLocalQueue = FSchedulerTls::LocalQueue;
 
 			// Standby workers always enqueue to the global queue and perform wakeup
 			// as they can go to sleep whenever the oversubscription period is done
@@ -356,11 +370,18 @@ namespace LowLevelTasks
 				QueuePreference = EQueuePreference::GlobalQueuePreference;
 			}
 
-			bWakeUpWorker |= bIsStandbyWorker || FSchedulerTls::LocalQueue == nullptr;
+			bWakeUpWorker |= bIsStandbyWorker || CachedLocalQueue == nullptr;
+
+			// Always force local queue usage when launching from the game thread to minimize cost.
+			if (CachedLocalQueue && CachedLocalQueue == GameThreadLocalQueue.Get())
+			{
+				QueuePreference = EQueuePreference::LocalQueuePreference;
+				bWakeUpWorker = true; // The game thread is never pumping its local queue directly, need to always perform a wakeup.
+			}
 
-			if (FSchedulerTls::LocalQueue && QueuePreference != EQueuePreference::GlobalQueuePreference)
+			if (CachedLocalQueue && QueuePreference != EQueuePreference::GlobalQueuePreference)
 			{
-				FSchedulerTls::LocalQueue->Enqueue(&Task, uint32(Task.GetPriority()));
+				CachedLocalQueue->Enqueue(&Task, uint32(Task.GetPriority()));
 			}
 			else
 			{
@@ -369,7 +390,7 @@ namespace LowLevelTasks
 
 			if (bWakeUpWorker)
 			{
-				if (bWakeUpWorker && !WakeUpWorker(bIsBackgroundTask) && !bIsBackgroundTask)
+				if (!WakeUpWorker(bIsBackgroundTask) && !bIsBackgroundTask)
 				{
 					WakeUpWorker(true);
 				}
@@ -476,7 +497,8 @@ namespace LowLevelTasks
 		while (true)
 		{
 			bool bExecutedSomething = false;
-			while (TryExecuteTaskFrom<FSchedulerTls::FLocalQueueType, &FSchedulerTls::FLocalQueueType::Dequeue, true>(WorkerEvent, WorkerLocalQueue, OutOfWork, bPermitBackgroundWork)
+			while (TryExecuteTaskFrom<FSchedulerTls::FLocalQueueType, &FSchedulerTls::FLocalQueueType::StealLocal, true>(WorkerEvent, GameThreadLocalQueue.Get(), OutOfWork, bPermitBackgroundWork)
+				|| TryExecuteTaskFrom<FSchedulerTls::FLocalQueueType, &FSchedulerTls::FLocalQueueType::Dequeue, true>(WorkerEvent, WorkerLocalQueue, OutOfWork, bPermitBackgroundWork)
 				|| TryExecuteTaskFrom<FSchedulerTls::FLocalQueueType, &FSchedulerTls::FLocalQueueType::DequeueSteal, true>(WorkerEvent, WorkerLocalQueue, OutOfWork, bPermitBackgroundWork))
 			{
 				bPreparingStandby = false;
@@ -517,7 +539,8 @@ namespace LowLevelTasks
 		while (true)
 		{
 			bool bExecutedSomething = false;
-			while (TryExecuteTaskFrom<FSchedulerTls::FLocalQueueType, &FSchedulerTls::FLocalQueueType::Dequeue, false>(WorkerEvent, WorkerLocalQueue, OutOfWork, bPermitBackgroundWork)
+			while (TryExecuteTaskFrom<FSchedulerTls::FLocalQueueType, &FSchedulerTls::FLocalQueueType::StealLocal, false>(WorkerEvent, GameThreadLocalQueue.Get(), OutOfWork, bPermitBackgroundWork)
+				|| TryExecuteTaskFrom<FSchedulerTls::FLocalQueueType, &FSchedulerTls::FLocalQueueType::Dequeue, false>(WorkerEvent, WorkerLocalQueue, OutOfWork, bPermitBackgroundWork)
 				|| TryExecuteTaskFrom<FSchedulerTls::FLocalQueueType, &FSchedulerTls::FLocalQueueType::DequeueSteal, false>(WorkerEvent, WorkerLocalQueue, OutOfWork, bPermitBackgroundWork))
 			{
 				bPreparingWait = false;
diff --git a/Engine/Source/Runtime/Core/Public/Async/Fundamental/LocalQueue.h b/Engine/Source/Runtime/Core/Public/Async/Fundamental/LocalQueue.h
--- a/Engine/Source/Runtime/Core/Public/Async/Fundamental/LocalQueue.h
+++ b/Engine/Source/Runtime/Core/Public/Async/Fundamental/LocalQueue.h
@@ -218,6 +218,21 @@ public:
 			}
 		}
 
+		inline FTask* StealLocal(bool GetBackGroundTasks)
+		{
+			const int32 MaxPriority = GetBackGroundTasks ? int32(ETaskPriority::Count) : int32(ETaskPriority::ForegroundCount);
+
+			for (int32 PriorityIndex = 0; PriorityIndex < MaxPriority; ++PriorityIndex)
+			{
+				FTask* Item;
+				if (LocalQueues[PriorityIndex].Steal(Item))
+				{
+					return Item;
+				}
+			}
+			return nullptr;
+		}
+
 		// Check both the local and global queue in priority order
 		inline FTask* Dequeue(bool GetBackGroundTasks)
 		{
diff --git a/Engine/Source/Runtime/Core/Public/Async/Fundamental/Scheduler.h b/Engine/Source/Runtime/Core/Public/Async/Fundamental/Scheduler.h
--- a/Engine/Source/Runtime/Core/Public/Async/Fundamental/Scheduler.h
+++ b/Engine/Source/Runtime/Core/Public/Async/Fundamental/Scheduler.h
@@ -237,6 +237,7 @@ namespace LowLevelTasks
 		TArray<std::atomic<FThread*>>                  WorkerThreads;
 		TAlignedArray<FSchedulerTls::FLocalQueueType>  WorkerLocalQueues;
 		TAlignedArray<Private::FWaitEvent>             WorkerEvents;
+		TUniquePtr<FSchedulerTls::FLocalQueueType>     GameThreadLocalQueue;
 		std::atomic_uint                               ActiveWorkers { 0 };
 		std::atomic_uint                               NextWorkerId { 0 };
 		std::atomic<int32>                             ForegroundCreationIndex{ 0 };
diff --git a/Engine/Source/Runtime/Core/Public/Experimental/Containers/FAAArrayQueue.h b/Engine/Source/Runtime/Core/Public/Experimental/Containers/FAAArrayQueue.h
--- a/Engine/Source/Runtime/Core/Public/Experimental/Containers/FAAArrayQueue.h
+++ b/Engine/Source/Runtime/Core/Public/Experimental/Containers/FAAArrayQueue.h
@@ -31,7 +31,7 @@
 
 #include <atomic>
 #include "HazardPointer.h"
-
+#include "HAL/PlatformProcess.h"
 
 /**
  * <h1> Fetch-And-Add Array Queue </h1>
@@ -262,6 +262,19 @@ private:
                 continue;
             }
 
+			// When there is more consumers than producers we can end up stealing
+			// empty slots that producers have reserved but not yet had time to write into.
+			// This leads to a lot of retries on the producers side so help this case
+			// by spinning just a little bit when we know the deqidx we got is valid
+			// and is likely to be written to very soon.
+			if (lhead->items[idx].load() == nullptr && idx <= lhead->enqidx.load())
+			{
+				for (int32 Try = 0; lhead->items[idx].load() == nullptr && Try < 10; ++Try)
+				{
+					FPlatformProcess::Yield();
+				}
+			}
+
             T* item = lhead->items[idx].exchange(GetTakenPointer());
 			if (item == nullptr)
 			{
