[AutoRTFM] Allow a post-abort callback to be used with CascadingAbortTransaction.

This will allow us to avoid a double-free hazard when triggering runtime errors from Verse in BPVM.

diff --git a/Engine/Source/Programs/AutoRTFMTests/Private/Abort.cpp b/Engine/Source/Programs/AutoRTFMTests/Private/Abort.cpp
--- a/Engine/Source/Programs/AutoRTFMTests/Private/Abort.cpp
+++ b/Engine/Source/Programs/AutoRTFMTests/Private/Abort.cpp
@@ -183,12 +183,93 @@ TEST_CASE("Abort.Cascade")
 		{
 			AutoRTFM::CascadingAbortTransaction();
 		});
+
+		FAIL("Execution should never reach this point");
 	});
 
 	REQUIRE(AutoRTFM::ETransactionResult::AbortedByCascade == Result);
 	REQUIRE(false == bTouched);
 }
 
+TEST_CASE("Abort.CascadeWithPostAbortCallback")
+{
+	bool bTouched = false;
+	bool bPostAbortCallbackWasExecuted = false;
+
+	const AutoRTFM::ETransactionResult Result = AutoRTFM::Transact([&]
+	{
+		bTouched = true;
+		AutoRTFM::Transact([&]
+		{
+			AutoRTFM::CascadingAbortTransaction([&]
+			{
+				REQUIRE(!bTouched);  // Rollbacks must occur before the post-abort callback is invoked.
+				bPostAbortCallbackWasExecuted = true;
+			});
+		});
+
+		FAIL("Execution should never reach this point");
+	});
+
+	REQUIRE(AutoRTFM::ETransactionResult::AbortedByCascade == Result);
+	REQUIRE(!bTouched);
+	REQUIRE(bPostAbortCallbackWasExecuted);
+}
+
+TEST_CASE("Abort.CascadingAbortMustRunOnAbortsBeforePostAbort")
+{
+	bool bOnAbortWasCalled = false;
+	bool bPostAbortCallbackWasExecuted = false;
+
+	const AutoRTFM::ETransactionResult Result = AutoRTFM::Transact([&]
+	{
+		AutoRTFM::OnAbort([&]
+		{
+			bOnAbortWasCalled = true;
+		});
+
+		AutoRTFM::CascadingAbortTransaction([&]
+		{
+			REQUIRE(bOnAbortWasCalled);  // OnAborts must before the post-abort callback is invoked.
+			bPostAbortCallbackWasExecuted = true;
+		});
+	});
+
+	REQUIRE(AutoRTFM::ETransactionResult::AbortedByCascade == Result);
+	REQUIRE(bOnAbortWasCalled);
+	REQUIRE(bPostAbortCallbackWasExecuted);
+}
+
+TEST_CASE("Abort.CascadingAbortInOpen")
+{
+	bool bTouched = false;
+	bool bExecutedCodeAfterCascadingAbort = false;
+	bool bPostAbortCallbackWasExecuted = false;
+
+	const AutoRTFM::ETransactionResult Result = AutoRTFM::Transact([&]
+	{
+		bTouched = true;
+
+		AutoRTFM::Open([&]
+		{
+			const AutoRTFM::ETransactionResult InnerResult = AutoRTFM::CascadingAbortTransaction([&]
+			{
+				REQUIRE(!bTouched);
+				bPostAbortCallbackWasExecuted = true;
+			});
+
+			REQUIRE(AutoRTFM::ETransactionResult::AbortedByCascade == InnerResult);
+
+			bExecutedCodeAfterCascadingAbort = true;
+		});
+	});
+
+	REQUIRE(AutoRTFM::ETransactionResult::AbortedByCascade == Result);
+	REQUIRE(!bTouched);
+	REQUIRE(bPostAbortCallbackWasExecuted);
+	REQUIRE(bExecutedCodeAfterCascadingAbort);
+}
+
 TEST_CASE("Abort.CascadeThroughOpen")
 {
 	bool bTouched = false;
@@ -215,8 +296,41 @@ TEST_CASE("Abort.CascadeThroughOpen")
 	REQUIRE(false == bTouched);
 }
 
+TEST_CASE("Abort.CascadeWithCallbackThroughOpen")
+{
+	bool bTouched = false;
+	bool bPostAbortCallbackWasExecuted = false;
+
+	const AutoRTFM::ETransactionResult Result = AutoRTFM::Transact([&]
+	{
+		bTouched = true;
+
+		AutoRTFM::Open([&]
+		{
+			const AutoRTFM::EContextStatus Status = AutoRTFM::Close([&]
+			{
+				AutoRTFM::Transact([&]
+				{
+					AutoRTFM::CascadingAbortTransaction([&]
+					{
+						REQUIRE(!bTouched);  // Rollbacks must occur before the post-abort callback is invoked.
+						bPostAbortCallbackWasExecuted = true;
+					});
+				});
+			});
+
+			REQUIRE(AutoRTFM::EContextStatus::AbortedByCascadingAbort == Status);
+		});
+	});
+
+	REQUIRE(AutoRTFM::ETransactionResult::AbortedByCascade == Result);
+	REQUIRE(!bTouched);
+	REQUIRE(bPostAbortCallbackWasExecuted);
+}
+
 TEST_CASE("Abort.CascadeThroughManualTransaction")
 {
+	bool bPostAbortCallbackWasExecuted = false;
 	bool bTouched = false;
 
 	const AutoRTFM::ETransactionResult Result = AutoRTFM::Transact([&]
@@ -229,7 +343,11 @@ TEST_CASE("Abort.CascadeThroughManualTransaction")
 
 			const AutoRTFM::EContextStatus Status = AutoRTFM::Close([&]
 			{
-				AutoRTFM::CascadingAbortTransaction();
+				AutoRTFM::CascadingAbortTransaction([&]
+				{
+					REQUIRE(!bTouched);  // Rollbacks must occur before the post-abort callback is invoked.
+					bPostAbortCallbackWasExecuted = true;
+				});
 			});
 
 			REQUIRE(AutoRTFM::EContextStatus::AbortedByCascadingAbort == Status);
@@ -243,7 +361,8 @@ TEST_CASE("Abort.CascadeThroughManualTransaction")
 	});
 
 	REQUIRE(AutoRTFM::ETransactionResult::AbortedByCascade == Result);
-	REQUIRE(false == bTouched);
+	REQUIRE(!bTouched);
+	REQUIRE(bPostAbortCallbackWasExecuted);
 }
 
 static int SAbortKey;
diff --git a/Engine/Source/Runtime/Core/Private/AutoRTFM/API.cpp b/Engine/Source/Runtime/Core/Private/AutoRTFM/API.cpp
--- a/Engine/Source/Runtime/Core/Private/AutoRTFM/API.cpp
+++ b/Engine/Source/Runtime/Core/Private/AutoRTFM/API.cpp
@@ -433,12 +433,6 @@ AUTORTFM_CORE_API autortfm_result RTFM_autortfm_abort_transaction()
 	return static_cast<autortfm_result>(Context->AbortTransaction(/* bIsClosed */ true, EContextStatus::AbortedByRequest));
 }
 
-AUTORTFM_CORE_API autortfm_result RTFM_autortfm_cascading_abort_transaction()
-{
-	FContext* const Context = FContext::Get();
-	return static_cast<autortfm_result>(Context->AbortTransaction(/* bIsClosed */ true, EContextStatus::AbortedByCascadingAbort));
-}
-
 AUTORTFM_CORE_API void RTFM_autortfm_clear_transaction_status()
 {
 	UE_LOG(LogAutoRTFM, Fatal, TEXT("The function `autortfm_clear_transaction_status` was called from closed code."));
@@ -533,11 +527,18 @@ AUTORTFM_CORE_API void RTFM_autortfm_record_open_write_err(void*, size_t)
 	UE_LOG(LogAutoRTFM, Fatal, TEXT("The function `autortfm_record_open_write` was called from closed code."));
 }
 
+AUTORTFM_CORE_API ETransactionResult RTFM_CascadingAbortTransactionInternal(TFunction<void()>&& RunAfterAbort)
+{
+	FContext* const Context = FContext::Get();
+	ASSERT(Context->GetStatus() == EContextStatus::OnTrack);
+	return Context->AbortTransactionWithPostAbortCallback(/* bIsClosed */ true, EContextStatus::AbortedByCascadingAbort, MoveTemp(RunAfterAbort));
+}
+
 AUTORTFM_CORE_API void RTFM_CascadingRetryTransactionInternal(TFunction<void()> && RunAfterAbortBeforeRetryWork)
 {
 	FContext* Context = FContext::Get();
 	ASSERT(Context->GetStatus() == EContextStatus::OnTrack);
-	Context->CascadingRetryTransaction(true, MoveTemp(RunAfterAbortBeforeRetryWork));
+	Context->AbortTransactionWithPostAbortCallback(/* bIsClosed */ true, EContextStatus::AbortedByCascadingRetry, MoveTemp(RunAfterAbortBeforeRetryWork));
 }
 
 AUTORTFM_CORE_API void RTFM_OnCommitInternal(TFunction<void()> && Work)
@@ -716,8 +717,9 @@ extern "C" bool autortfm_start_transaction()
 		return false;
 	}
 
-	UE_CLOG(!FContext::Get()->IsTransactional(), LogAutoRTFM, Fatal, TEXT("The function `autortfm_start_transaction` was called from outside a transact."));
-	FContext::Get()->StartTransaction(ForTheRuntime::GetMemoryValidationLevel());
+	FContext* Context = FContext::Get();
+	UE_CLOG(!Context->IsTransactional(), LogAutoRTFM, Fatal, TEXT("The function `autortfm_start_transaction` was called from outside a transact."));
+	Context->StartTransaction(ForTheRuntime::GetMemoryValidationLevel());
 	return true;
 }
 
@@ -740,23 +742,11 @@ extern "C" autortfm_result autortfm_abort_transaction()
 		return RTFM_autortfm_abort_transaction();
 	}
 
-	UE_CLOG(!FContext::Get()->IsTransactional(), LogAutoRTFM, Fatal, TEXT("The function `autortfm_abort_transaction` was called from outside a transact."));
 	FContext* const Context = FContext::Get();
+	UE_CLOG(!Context->IsTransactional(), LogAutoRTFM, Fatal, TEXT("The function `autortfm_abort_transaction` was called from outside a transact."));
 	return static_cast<autortfm_result>(Context->AbortTransaction(false, EContextStatus::AbortedByRequest));
 }
 
-extern "C" autortfm_result autortfm_cascading_abort_transaction()
-{
-	if (autortfm_is_closed())
-	{
-		return RTFM_autortfm_cascading_abort_transaction();
-	}
-
-	UE_CLOG(!FContext::Get()->IsTransactional(), LogAutoRTFM, Fatal, TEXT("The function `autortfm_cascading_abort_transaction` was called from outside a transact."));
-	FContext* const Context = FContext::Get();
-	return static_cast<autortfm_result>(Context->AbortTransaction(false, EContextStatus::AbortedByCascadingAbort));
-}
-
 extern "C" void autortfm_clear_transaction_status()
 {
 	if (autortfm_is_closed())
@@ -901,6 +891,18 @@ extern "C" bool autortfm_is_on_current_transaction_stack(void* Ptr)
 	}
 }
 
+ETransactionResult ForTheRuntime::CascadingAbortTransactionInternal(TFunction<void()>&& RunAfterAbort)
+{
+	if (autortfm_is_closed())
+	{
+		return RTFM_CascadingAbortTransactionInternal(MoveTemp(RunAfterAbort));
+	}
+
+	FContext* const Context = FContext::Get();
+	UE_CLOG(!Context->IsTransactional(), LogAutoRTFM, Fatal, TEXT("The function `CascadingAbortTransaction` was called from outside a transact."));
+	return Context->AbortTransactionWithPostAbortCallback(/* bIsClosed */ false, EContextStatus::AbortedByCascadingAbort, MoveTemp(RunAfterAbort));
+}
+
 void ForTheRuntime::CascadingRetryTransactionInternal(TFunction<void()>&& RunAfterAbortBeforeRetryWork)
 {
 	if (autortfm_is_closed())
diff --git a/Engine/Source/Runtime/Core/Private/AutoRTFM/Context.cpp b/Engine/Source/Runtime/Core/Private/AutoRTFM/Context.cpp
--- a/Engine/Source/Runtime/Core/Private/AutoRTFM/Context.cpp
+++ b/Engine/Source/Runtime/Core/Private/AutoRTFM/Context.cpp
@@ -174,11 +174,6 @@ ETransactionResult FContext::AbortTransaction(bool bIsClosed, EContextStatus New
 	{
 		// The optimization worked! We didn't need to allocate an FTransaction for this.
 		NumDeferredTransactions -= 1;
-		if (Status == EContextStatus::AbortedByCascadingAbort ||
-			Status == EContextStatus::AbortedByCascadingRetry)
-		{
-			Result = ETransactionResult::AbortedByCascade;
-		}
 	}
 	else
 	{
@@ -200,20 +195,36 @@ ETransactionResult FContext::AbortTransaction(bool bIsClosed, EContextStatus New
 		Throw();
 	}
 
+	// If this is a cascading abort, we should reflect that in the returned result.
+	if (Result == ETransactionResult::AbortedByRequest &&
+		(Status == EContextStatus::AbortedByCascadingAbort ||
+		 Status == EContextStatus::AbortedByCascadingRetry))
+	{
+		Result = ETransactionResult::AbortedByCascade;
+	}
+
 	return Result;
 }
 
-void FContext::CascadingRetryTransaction(bool bIsClosed, TFunction<void()>&& RunAfterAbortBeforeRetryWork)
+ETransactionResult FContext::AbortTransactionWithPostAbortCallback(bool bIsClosed, EContextStatus NewStatus, TFunction<void()>&& InCallback)
 {
-	ASSERT(!CascadingRetryTransactionWork);
-	
-	// We explicitly must copy the function here because the original was allocated
-	// within a transactional context, and thus the memory is allocating under
+	// The callback parameter is only honored by the cascading aborts.
+	ASSERT(NewStatus == EContextStatus::AbortedByCascadingAbort ||
+		   NewStatus == EContextStatus::AbortedByCascadingRetry);
+
+	// We must explicitly copy the passed-in callback here, because the original may have been
+	// allocated within a transactional context; if so, its memory was allocated under
 	// transactionalized conditions. By copying, we create an open copy of the callback.
-	TFunction<void()> Copy(RunAfterAbortBeforeRetryWork);
+	if (!PostAbortCallback.IsSet())
+	{
+		PostAbortCallback = InCallback;
+	}
+	else if (InCallback.IsSet())
+	{
+		UE_LOG(LogAutoRTFM, Fatal, TEXT("Only one post-abort callback at a time is supported."));
+	}
 
-	CascadingRetryTransactionWork = MoveTemp(Copy);
-	AbortTransaction(bIsClosed, EContextStatus::AbortedByCascadingRetry);
+	return AbortTransaction(bIsClosed, NewStatus);
 }
 
 bool FContext::IsAborting() const
@@ -487,6 +498,7 @@ ETransactionResult FContext::Transact(void (*InstrumentedFunction)(void*), void*
 				break;
 
 			case EContextStatus::AbortedByRequest:
+				ASSERT(!PostAbortCallback.IsSet());
 				Result = ETransactionResult::AbortedByRequest;
 				break;
 
@@ -495,26 +507,38 @@ ETransactionResult FContext::Transact(void (*InstrumentedFunction)(void*), void*
 				break;
 
 			case EContextStatus::AbortedByCascadingAbort:
-				ASSERT(!CascadingRetryTransactionWork.IsSet());
+				if (PostAbortCallback.IsSet())
+				{
+					// Call the post-abort callback to do whatever work the user
+					// required be done before throwing.
+					Status = EContextStatus::InPostAbort;
+
+					PostAbortCallback();
+					PostAbortCallback.Reset();
+
+					ASSERT(Status == EContextStatus::InPostAbort);
+					Status = EContextStatus::AbortedByCascadingAbort;
+				}
+
 				Result = ETransactionResult::AbortedByCascade;
 				break;
 
 			case EContextStatus::AbortedByCascadingRetry:
-				ASSERT(CascadingRetryTransactionWork.IsSet());
-
+				ASSERT(PostAbortCallback.IsSet());
+			
 				// Clean up the transaction to get it ready for re-execution.
 				ClearTransactionStatus();
 				CurrentTransaction->Reset();
 
 				ASSERT(Status == EContextStatus::OnTrack);
-				Status = EContextStatus::InRetryTransaction;
+				Status = EContextStatus::InPostAbort;
 
-				// Call the retry callback to do whatever work the user
+				// Call the post-abort callback to do whatever work the user
 				// required be done before we attempt to re-execute.
-				CascadingRetryTransactionWork();
-				CascadingRetryTransactionWork.Reset();
+				PostAbortCallback();
+				PostAbortCallback.Reset();
 
-				ASSERT(Status == EContextStatus::InRetryTransaction);
+				ASSERT(Status == EContextStatus::InPostAbort);
 				Status = EContextStatus::OnTrack;
 
 				// Then get rolling!
diff --git a/Engine/Source/Runtime/Core/Private/AutoRTFM/Context.h b/Engine/Source/Runtime/Core/Private/AutoRTFM/Context.h
--- a/Engine/Source/Runtime/Core/Private/AutoRTFM/Context.h
+++ b/Engine/Source/Runtime/Core/Private/AutoRTFM/Context.h
@@ -33,11 +33,11 @@ public:
 
 	ETransactionResult CommitTransaction();
 	ETransactionResult AbortTransaction(bool bIsClosed, EContextStatus NewStatus);
+	ETransactionResult AbortTransactionWithPostAbortCallback(bool bIsClosed, EContextStatus NewStatus, TFunction<void()>&& InCallback);
 	void ClearTransactionStatus();
 	bool IsAborting() const;
 	bool IsTransactional() const;
 	bool IsCommittingOrAborting() const;
-	void CascadingRetryTransaction(bool bIsClosed, TFunction<void()> && RunAfterAbortBeforeRetryWork);
 
     // Record that a write is about to occur at the given LogicalAddress of Size bytes.
     void RecordWrite(void* LogicalAddress, size_t Size);
@@ -163,7 +163,7 @@ private:
     EContextStatus Status{EContextStatus::Idle};
 	uint32 CurrentThreadId{ FPlatformTLS::InvalidTlsSlot };
 	uint32 StackLocalInitializerDepth = 0;
-	TFunction<void()> CascadingRetryTransactionWork;
+	TFunction<void()> PostAbortCallback;
 	FOpenHashThrottler OpenHashThrottler{
 		/* LogInterval            */ 10,   // Log every 10 seconds.
 		/* AdjustThrottleInterval */ 0.5,  // Adjust throttling probabilities every 500ms.
diff --git a/Engine/Source/Runtime/Core/Public/AutoRTFM/AutoRTFM.h b/Engine/Source/Runtime/Core/Public/AutoRTFM/AutoRTFM.h
--- a/Engine/Source/Runtime/Core/Public/AutoRTFM/AutoRTFM.h
+++ b/Engine/Source/Runtime/Core/Public/AutoRTFM/AutoRTFM.h
@@ -136,7 +136,7 @@ typedef enum
 	autortfm_status_aborted_by_cascading_abort,
 	autortfm_status_aborted_by_cascading_retry,
 	autortfm_status_in_static_local_initializer,
-	autortfm_status_in_retry_transaction
+	autortfm_status_in_post_abort
 } autortfm_status;
 
 #if UE_AUTORTFM_ENABLED
@@ -206,12 +206,6 @@ UE_AUTORTFM_API autortfm_result autortfm_abort_transaction();
 static UE_AUTORTFM_FORCEINLINE autortfm_result autortfm_abort_transaction() { return autortfm_aborted_by_request; }
 #endif
 
-#if UE_AUTORTFM_ENABLED
-UE_AUTORTFM_API autortfm_result autortfm_cascading_abort_transaction();
-#else
-static UE_AUTORTFM_FORCEINLINE autortfm_result autortfm_cascading_abort_transaction() { return autortfm_aborted_by_cascade; }
-#endif
-
 #if UE_AUTORTFM_ENABLED
 UE_AUTORTFM_API bool autortfm_start_transaction();
 #else
@@ -450,9 +444,10 @@ enum class EContextStatus : uint8_t
 	// `IsTransactional()` will return `false` when in this state.
 	InStaticLocalInitializer = autortfm_status_in_static_local_initializer,
 
-	// Means we are in a transaction retry currently, inbetween the aborted previous transaction
-	// and the re-runned retried transaction.
-	InRetryTransaction = autortfm_status_in_retry_transaction,
+	// An InPostAbort status means that a transaction has just been aborted, but we have not yet
+	// thrown back to the caller or retried the transaction. We are only in this state when 
+	// executing the post-abort callback.
+	InPostAbort = autortfm_status_in_post_abort,
 };
 
 // An enumerator of transactional memory validation levels.
@@ -477,6 +472,7 @@ enum class EMemoryValidationLevel : uint8_t
 #if UE_AUTORTFM
 namespace ForTheRuntime
 {
+	UE_AUTORTFM_API ETransactionResult CascadingAbortTransactionInternal(TFunction<void()>&& RunAfterAbort);
 	UE_AUTORTFM_API void CascadingRetryTransactionInternal(TFunction<void()>&& RunAfterAbortBeforeRetryWork);
 	UE_AUTORTFM_API void OnCommitInternal(TFunction<void()>&& Work);
 	UE_AUTORTFM_API void OnAbortInternal(TFunction<void()>&& Work);
@@ -594,11 +590,32 @@ static UE_AUTORTFM_FORCEINLINE ETransactionResult AbortTransaction()
 
 // End a transaction nest and discard all effects. This cascades, meaning
 // an abort of a nested transaction will cause all transactions in the
-// nest to abort.
+// nest to abort. Once the transaction has aborted, `RunAfterAbort` will
+// be called. Finally, control will be returned to the end of the outermost 
+// Transact.
+#if UE_AUTORTFM
+template <typename FunctorType>
+static UE_AUTORTFM_FORCEINLINE ETransactionResult CascadingAbortTransaction(FunctorType&& RunAfterAbort)
+{
+	return ForTheRuntime::CascadingAbortTransactionInternal(std::forward<FunctorType>(RunAfterAbort));
+}
+
+static UE_AUTORTFM_FORCEINLINE ETransactionResult CascadingAbortTransaction()
+{
+	return ForTheRuntime::CascadingAbortTransactionInternal({});
+}
+#else
+template <typename FunctorType>
+static UE_AUTORTFM_FORCEINLINE ETransactionResult CascadingAbortTransaction(FunctorType&&)
+{
+	return ETransactionResult::AbortedByCascade;
+}
+
 static UE_AUTORTFM_FORCEINLINE ETransactionResult CascadingAbortTransaction()
 {
-	return static_cast<ETransactionResult>(autortfm_cascading_abort_transaction());
+	return ETransactionResult::AbortedByCascade;
 }
+#endif
 
 namespace Detail
 {
@@ -791,12 +808,12 @@ template<typename FunctorType> [[nodiscard]] static UE_AUTORTFM_FORCEINLINE ECon
 			const_cast<void*>(static_cast<const void*>(&Functor))));
 }
 
-#if UE_AUTORTFM
 // Force a transaction nest to be retried. Once the transaction has aborted,
 // it will call `RunAfterAbortBeforeRetryWork`, before retrying the
 // transaction. This is an expensive operation and should thus be used with
 // extreme caution. If this is called outside of a transaction then the
 // `RunAfterAbortBeforeRetryWork` callback is never executed.
+#if UE_AUTORTFM
 template<typename FunctorType> static UE_AUTORTFM_FORCEINLINE void CascadingRetryTransaction(FunctorType&& RunAfterAbortBeforeRetryWork)
 {
 	if (autortfm_is_closed())
@@ -805,14 +822,8 @@ template<typename FunctorType> static UE_AUTORTFM_FORCEINLINE void CascadingRetr
 	}
 }
 #else
-// Force a transaction nest to be retried. Once the transaction has aborted,
-// it will call `RunAfterAbortBeforeRetryWork`, before retrying the
-// transaction. This is an expensive operation and should thus be used with
-// extreme caution. If this is called outside of a transaction then the
-// `RunAfterAbortBeforeRetryWork` callback is never executed.
-template<typename FunctorType> static UE_AUTORTFM_FORCEINLINE void CascadingRetryTransaction(FunctorType&& RunAfterAbortBeforeRetryWork)
+template<typename FunctorType> static UE_AUTORTFM_FORCEINLINE void CascadingRetryTransaction(FunctorType&&)
 {
-	UE_AUTORTFM_UNUSED(RunAfterAbortBeforeRetryWork);
 }
 #endif
 
