[Core] Cleanup in set.h
* Move FSetElement/KeyFuncs into a separate file so they can be used by other classes without depending on Set.h
* Some code cleanup to match coding standards
** Break single line functions into multiple lines
** Use default instead of empty constructor

#tests Ran foundation tests

diff --git a/Engine/Source/Runtime/Core/Public/Containers/Set.h b/Engine/Source/Runtime/Core/Public/Containers/Set.h
--- a/Engine/Source/Runtime/Core/Public/Containers/Set.h
+++ b/Engine/Source/Runtime/Core/Public/Containers/Set.h
@@ -2,158 +2,30 @@
 
 #pragma once
 
-#include "CoreTypes.h"
-#include "Misc/AssertionMacros.h"
-#include "Templates/UnrealTypeTraits.h"
-#include "Templates/UnrealTemplate.h"
+#include "Containers/Array.h"
 #include "Containers/ContainerAllocationPolicies.h"
 #include "Containers/ContainerElementTypeCompatibility.h"
-#include "Templates/Sorting.h"
-#include "Containers/Array.h"
+#include "Containers/SetUtilities.h"
+#include "Containers/SparseArray.h"
+#include "ContainersFwd.h"
+#include "CoreTypes.h"
 #include "Math/UnrealMathUtility.h"
+#include "Misc/AssertionMacros.h"
 #include "Misc/StructBuilder.h"
-#include "Templates/Function.h"
-#include <initializer_list>
-#include "Templates/TypeHash.h"
-#include "Containers/SparseArray.h"
-#include "Serialization/StructuredArchive.h"
 #include "Serialization/MemoryImageWriter.h"
-#include "ContainersFwd.h"
+#include "Serialization/StructuredArchive.h"
+#include "Templates/Function.h"
 #include "Templates/RetainedRef.h"
+#include "Templates/Sorting.h"
+#include "Templates/TypeHash.h"
+#include "Templates/UnrealTemplate.h"
+#include "Templates/UnrealTypeTraits.h"
+#include <initializer_list>
 #include <type_traits>
 
-/**
- * The base KeyFuncs type with some useful definitions for all KeyFuncs; meant to be derived from instead of used directly.
- * bInAllowDuplicateKeys=true is slightly faster because it allows the TSet to skip validating that
- * there isn't already a duplicate entry in the TSet.
-  */
-template<typename ElementType,typename InKeyType,bool bInAllowDuplicateKeys = false>
-struct BaseKeyFuncs
-{
-	typedef InKeyType KeyType;
-	typedef typename TCallTraits<InKeyType>::ParamType KeyInitType;
-	typedef typename TCallTraits<ElementType>::ParamType ElementInitType;
-
-	enum { bAllowDuplicateKeys = bInAllowDuplicateKeys };
-};
-
-/**
- * A default implementation of the KeyFuncs used by TSet which uses the element as a key.
- */
-template<typename ElementType,bool bInAllowDuplicateKeys /*= false*/>
-struct DefaultKeyFuncs : BaseKeyFuncs<ElementType,ElementType,bInAllowDuplicateKeys>
-{
-	typedef typename TTypeTraits<ElementType>::ConstPointerType KeyInitType;
-	typedef typename TCallTraits<ElementType>::ParamType ElementInitType;
-
-	/**
-	 * @return The key used to index the given element.
-	 */
-	[[nodiscard]] static FORCEINLINE KeyInitType GetSetKey(ElementInitType Element)
-	{
-		return Element;
-	}
-
-	/**
-	 * @return True if the keys match.
-	 */
-	[[nodiscard]] static FORCEINLINE bool Matches(KeyInitType A, KeyInitType B)
-	{
-		return A == B;
-	}
-
-	/**
-	 * @return True if the keys match.
-	 */
-	template<typename ComparableKey>
-	[[nodiscard]] static FORCEINLINE bool Matches(KeyInitType A, ComparableKey B)
-	{
-		return A == B;
-	}
-
-	/** Calculates a hash index for a key. */
-	[[nodiscard]] static FORCEINLINE uint32 GetKeyHash(KeyInitType Key)
-	{
-		return GetTypeHash(Key);
-	}
-
-	/** Calculates a hash index for a key. */
-	template<typename ComparableKey>
-	[[nodiscard]] static FORCEINLINE uint32 GetKeyHash(ComparableKey Key)
-	{
-		return GetTypeHash(Key);
-	}
-};
-
-/** This is used to provide type specific behavior for a move which will destroy B. */
-/** Should be in UnrealTemplate but isn't for Clang build reasons - will move later */
-template<typename T>
-FORCEINLINE void MoveByRelocate(T& A, T& B)
-{
-	// Destruct the previous value of A.
-	A.~T();
-
-	// Relocate B into the 'hole' left by the destruction of A, leaving a hole in B instead.
-	RelocateConstructItems<T>(&A, &B, 1);
-}
-
 template <typename AllocatorType, typename InDerivedType = void>
 class TScriptSet;
 
-/** Either NULL or an identifier for an element of a set. */
-class FSetElementId
-{
-public:
-
-	template<typename,typename,typename>
-	friend class TSet;
-
-	template <typename, typename>
-	friend class TScriptSet;
-
-	/** Default constructor. */
-	FORCEINLINE FSetElementId() = default;
-
-	/** @return a boolean value representing whether the id is NULL. */
-	[[nodiscard]] FORCEINLINE bool IsValidId() const
-	{
-		return Index != INDEX_NONE;
-	}
-
-	/** Comparison operators. */
-	[[nodiscard]] FORCEINLINE friend bool operator==(const FSetElementId& A,const FSetElementId& B)
-	{
-		return A.Index == B.Index;
-	}
-	[[nodiscard]] FORCEINLINE friend bool operator!=(const FSetElementId& A, const FSetElementId& B)
-	{
-		return A.Index != B.Index;
-	}
-
-	[[nodiscard]] FORCEINLINE int32 AsInteger() const
-	{
-		return Index;
-	}
-
-	[[nodiscard]] FORCEINLINE static FSetElementId FromInteger(int32 Integer)
-	{
-		return FSetElementId(Integer);
-	}
-
-private:
-	/** The index of the element in the set's element array. */
-	int32 Index = INDEX_NONE;
-
-	/** Initialization constructor. */
-	[[nodiscard]] FORCEINLINE explicit FSetElementId(int32 InIndex)
-		: Index(InIndex)
-	{
-	}
-};
-
-// This is just an int32
-DECLARE_INTRINSIC_TYPE_LAYOUT(FSetElementId);
-
 template<typename InElementType, bool bTypeLayout>
 class TSetElementBase
 {
@@ -305,9 +177,7 @@ private:
 
 public:
 	/** Initialization constructor. */
-	[[nodiscard]] FORCEINLINE TSet()
-	{
-	}
+	[[nodiscard]] FORCEINLINE TSet() = default;
 
 	/** Copy constructor. */
 	[[nodiscard]] FORCEINLINE TSet(const TSet& Copy)
@@ -631,8 +501,14 @@ public:
 	 * @param	bIsAlreadyInSetPtr	[out]	Optional pointer to bool that will be set depending on whether element is already in set
 	 * @return	A pointer to the element stored in the set.
 	 */
-	FORCEINLINE FSetElementId Add(const InElementType&  InElement, bool* bIsAlreadyInSetPtr = nullptr) { return Emplace(                   InElement , bIsAlreadyInSetPtr); }
-	FORCEINLINE FSetElementId Add(      InElementType&& InElement, bool* bIsAlreadyInSetPtr = nullptr) { return Emplace(MoveTempIfPossible(InElement), bIsAlreadyInSetPtr); }
+	FORCEINLINE FSetElementId Add(const InElementType&  InElement, bool* bIsAlreadyInSetPtr = nullptr)
+	{
+		return Emplace(InElement, bIsAlreadyInSetPtr);
+	}
+	FORCEINLINE FSetElementId Add(InElementType&& InElement, bool* bIsAlreadyInSetPtr = nullptr)
+	{
+		return Emplace(MoveTempIfPossible(InElement), bIsAlreadyInSetPtr);
+	}
 
 	/**
 	 * Adds an element to the set if not already present and returns a reference to the added or existing element.
@@ -1812,10 +1688,22 @@ public:
 	 * DO NOT USE DIRECTLY
 	 * STL-like iterators to enable range-based for loop support.
 	 */
-	[[nodiscard]] FORCEINLINE TRangedForIterator      begin()       { return TRangedForIterator     (Elements.begin()); }
-	[[nodiscard]] FORCEINLINE TRangedForConstIterator begin() const { return TRangedForConstIterator(Elements.begin()); }
-	[[nodiscard]] FORCEINLINE TRangedForIterator      end()         { return TRangedForIterator     (Elements.end());   }
-	[[nodiscard]] FORCEINLINE TRangedForConstIterator end() const   { return TRangedForConstIterator(Elements.end());   }
+	[[nodiscard]] FORCEINLINE TRangedForIterator begin()
+	{
+		return TRangedForIterator(Elements.begin());
+	}
+	[[nodiscard]] FORCEINLINE TRangedForConstIterator begin() const
+	{
+		return TRangedForConstIterator(Elements.begin());
+	}
+	[[nodiscard]] FORCEINLINE TRangedForIterator end()
+	{
+		return TRangedForIterator(Elements.end());
+	}
+	[[nodiscard]] FORCEINLINE TRangedForConstIterator end() const
+	{
+		return TRangedForConstIterator(Elements.end());
+	}
 
 	// Maps are deliberately prevented from being hashed or compared, because this would hide potentially major performance problems behind default operations.
 	friend uint32 GetTypeHash(const TSet& Set) = delete;
diff --git a/Engine/Source/Runtime/Core/Public/Containers/SetUtilities.h b/Engine/Source/Runtime/Core/Public/Containers/SetUtilities.h
new file mode 100644
--- /dev/null
+++ b/Engine/Source/Runtime/Core/Public/Containers/SetUtilities.h
@@ -0,0 +1,139 @@
+// Copyright Epic Games, Inc. All Rights Reserved.
+
+#pragma once
+
+#include "CoreTypes.h"
+#include "Serialization/MemoryLayout.h"
+#include "Templates/MemoryOps.h"
+#include "Templates/UnrealTypeTraits.h"
+
+/**
+ * The base KeyFuncs type with some useful definitions for all KeyFuncs; meant to be derived from instead of used directly.
+ * bInAllowDuplicateKeys=true is slightly faster because it allows the TSet to skip validating that
+ * there isn't already a duplicate entry in the TSet.
+  */
+template<typename ElementType,typename InKeyType,bool bInAllowDuplicateKeys = false>
+struct BaseKeyFuncs
+{
+	typedef InKeyType KeyType;
+	typedef typename TCallTraits<InKeyType>::ParamType KeyInitType;
+	typedef typename TCallTraits<ElementType>::ParamType ElementInitType;
+
+	enum { bAllowDuplicateKeys = bInAllowDuplicateKeys };
+};
+
+/**
+ * A default implementation of the KeyFuncs used by TSet which uses the element as a key.
+ */
+template<typename ElementType,bool bInAllowDuplicateKeys /*= false*/>
+struct DefaultKeyFuncs : BaseKeyFuncs<ElementType,ElementType,bInAllowDuplicateKeys>
+{
+	typedef typename TTypeTraits<ElementType>::ConstPointerType KeyInitType;
+	typedef typename TCallTraits<ElementType>::ParamType ElementInitType;
+
+	/**
+	 * @return The key used to index the given element.
+	 */
+	[[nodiscard]] static FORCEINLINE KeyInitType GetSetKey(ElementInitType Element)
+	{
+		return Element;
+	}
+
+	/**
+	 * @return True if the keys match.
+	 */
+	[[nodiscard]] static FORCEINLINE bool Matches(KeyInitType A, KeyInitType B)
+	{
+		return A == B;
+	}
+
+	/**
+	 * @return True if the keys match.
+	 */
+	template<typename ComparableKey>
+	[[nodiscard]] static FORCEINLINE bool Matches(KeyInitType A, ComparableKey B)
+	{
+		return A == B;
+	}
+
+	/** Calculates a hash index for a key. */
+	[[nodiscard]] static FORCEINLINE uint32 GetKeyHash(KeyInitType Key)
+	{
+		return GetTypeHash(Key);
+	}
+
+	/** Calculates a hash index for a key. */
+	template<typename ComparableKey>
+	[[nodiscard]] static FORCEINLINE uint32 GetKeyHash(ComparableKey Key)
+	{
+		return GetTypeHash(Key);
+	}
+};
+
+/** This is used to provide type specific behavior for a move which will destroy B. */
+/** Should be in UnrealTemplate but isn't for Clang build reasons - will move later */
+template<typename T>
+FORCEINLINE void MoveByRelocate(T& A, T& B)
+{
+	// Destruct the previous value of A.
+	A.~T();
+
+	// Relocate B into the 'hole' left by the destruction of A, leaving a hole in B instead.
+	RelocateConstructItems<T>(&A, &B, 1);
+}
+
+/** Either NULL or an identifier for an element of a set.
+ * Used to differentiate between int as an element type and an index to a specific location
+ */
+class FSetElementId
+{
+public:
+
+	template<typename,typename,typename>
+	friend class TSet;
+
+	template <typename, typename>
+	friend class TScriptSet;
+
+	/** Default constructor. */
+	FORCEINLINE FSetElementId() = default;
+
+	/** @return a boolean value representing whether the id is NULL. */
+	[[nodiscard]] FORCEINLINE bool IsValidId() const
+	{
+		return Index != INDEX_NONE;
+	}
+
+	/** Comparison operators. */
+	[[nodiscard]] FORCEINLINE friend bool operator==(const FSetElementId& A,const FSetElementId& B)
+	{
+		return A.Index == B.Index;
+	}
+	[[nodiscard]] FORCEINLINE friend bool operator!=(const FSetElementId& A, const FSetElementId& B)
+	{
+		return A.Index != B.Index;
+	}
+
+	[[nodiscard]] FORCEINLINE int32 AsInteger() const
+	{
+		return Index;
+	}
+
+	[[nodiscard]] FORCEINLINE static FSetElementId FromInteger(int32 Integer)
+	{
+		return FSetElementId(Integer);
+	}
+
+private:
+	/** The index of the element in the set's element array. */
+	int32 Index = INDEX_NONE;
+
+	/** Initialization constructor. */
+	[[nodiscard]] FORCEINLINE explicit FSetElementId(int32 InIndex)
+	: Index(InIndex)
+	{
+	}
+};
+
+// This is just an int32
+DECLARE_INTRINSIC_TYPE_LAYOUT(FSetElementId);
