Minor optimizations to FModuleManager::QueryModules

#rnx

- It is quite common to have a large number of modules, so we should resize the output array before adding to it to prevent it being resized several times.
- Move the local ModuleStatus when adding to the output array to avoid string copies.

diff --git a/Engine/Source/Runtime/Core/Private/Modules/ModuleManager.cpp b/Engine/Source/Runtime/Core/Private/Modules/ModuleManager.cpp
--- a/Engine/Source/Runtime/Core/Private/Modules/ModuleManager.cpp
+++ b/Engine/Source/Runtime/Core/Private/Modules/ModuleManager.cpp
@@ -1247,8 +1247,9 @@ bool FModuleManager::QueryModule( const FName InModuleName, FModuleStatus& OutMo
 
 void FModuleManager::QueryModules( TArray< FModuleStatus >& OutModuleStatuses ) const
 {
-	OutModuleStatuses.Reset();
 	FScopeLock Lock(&ModulesCriticalSection);
+	OutModuleStatuses.Reset(Modules.Num());
+
 	for (const TPair<FName, ModuleInfoRef>& ModuleIt : Modules)
 	{
 		const FModuleInfo& CurModule = *ModuleIt.Value;
@@ -1263,7 +1264,7 @@ void FModuleManager::QueryModules( TArray< FModuleStatus >& OutModuleStatuses )
 			ModuleStatus.bIsGameModule = CurModule.Module->IsGameModule();
 		}
 
-		OutModuleStatuses.Add( ModuleStatus );
+		OutModuleStatuses.Add(MoveTemp(ModuleStatus));
 	}
 }
 
diff --git a/Engine/Source/Runtime/Core/Public/Modules/ModuleManager.h b/Engine/Source/Runtime/Core/Public/Modules/ModuleManager.h
--- a/Engine/Source/Runtime/Core/Public/Modules/ModuleManager.h
+++ b/Engine/Source/Runtime/Core/Public/Modules/ModuleManager.h
@@ -279,7 +279,7 @@ public:
 	/**
 	 * Queries information about all of the currently known modules.
 	 *
-	 * @param OutModuleStatuses Status of all modules.
+	 * @param OutModuleStatuses Status of all modules. Note that the array will be cleared before the results are added to it.
 	 * @see QueryModule
 	 */
 	CORE_API void QueryModules( TArray<FModuleStatus>& OutModuleStatuses ) const;
