Standards and modernization pass on ExpressionParserTypes.*.

diff --git a/Engine/Source/Runtime/Core/Public/Misc/ExpressionParserTypes.h b/Engine/Source/Runtime/Core/Public/Misc/ExpressionParserTypes.h
--- a/Engine/Source/Runtime/Core/Public/Misc/ExpressionParserTypes.h
+++ b/Engine/Source/Runtime/Core/Public/Misc/ExpressionParserTypes.h
@@ -3,7 +3,6 @@
 #pragma once
 
 #include "CoreTypes.h"
-#include "Templates/PointerIsConvertibleFromTo.h"
 #include "Templates/UnrealTemplate.h"
 #include "Containers/Array.h"
 #include "Containers/UnrealString.h"
@@ -15,6 +14,7 @@
 #include "Internationalization/Internationalization.h"
 #include "Misc/Guid.h"
 #include "Templates/ValueOrError.h"
+#include <type_traits>
 
 class FExpressionNode;
 struct FExpressionError;
@@ -24,12 +24,15 @@ namespace Impl
 	struct IExpressionNodeStorage;
 }
 
-typedef TValueOrError<FExpressionNode, FExpressionError> FExpressionResult;
+using FExpressionResult = TValueOrError<FExpressionNode, FExpressionError>;
 
 /** Simple error structure used for reporting parse errors */
 struct FExpressionError
 {
-	FExpressionError(const FText& InText) : Text(InText) {}
+	FExpressionError(const FText& InText)
+		: Text(InText)
+	{
+	}
 
 	FText Text;
 };
@@ -38,7 +41,7 @@ struct FExpressionError
 class FStringToken
 {
 public:
-	FStringToken() : TokenStart(nullptr), TokenEnd(nullptr), LineNumber(0), CharacterIndex(0) {}
+	FStringToken() = default;
 
 	/** Get the string representation of this token */
 	FString GetString() const
@@ -47,32 +50,59 @@ public:
 	}
 
 	/** Check if this token is valid */
-	bool IsValid() const { return TokenEnd != TokenStart; }
+	bool IsValid() const
+	{
+		return TokenEnd != TokenStart;
+	}
+
+	/** Get the position of the start of this token in the stream */
+	const TCHAR* GetTokenStartPos() const
+	{
+		return TokenStart;
+	}
 
-	/** Get the position of the start and end of this token in the stream */
-	const TCHAR* GetTokenStartPos() const { return TokenStart; }
-	const TCHAR* GetTokenEndPos() const { return TokenEnd; }
+	/** Get the position of the end of this token in the stream */
+	const TCHAR* GetTokenEndPos() const
+	{
+		return TokenEnd;
+	}
+
+	/** Get the character index of this token in the stream */
+	int32 GetCharacterIndex() const
+	{
+		return CharacterIndex;
+	}
 
-	/** Contextual information about this token */
-	int32 GetCharacterIndex() const { return CharacterIndex; }
-	int32 GetLineNumber() const { return LineNumber; }
+	/** Get the line number of this token in the stream */
+	int32 GetLineNumber() const
+	{
+		return LineNumber;
+	}
 	
 	/** Accumulate another token into this one */
-	void Accumulate(const FStringToken& InToken) { if (InToken.TokenEnd > TokenEnd) { TokenEnd = InToken.TokenEnd; } }
+	void Accumulate(const FStringToken& InToken)
+	{
+		if (InToken.TokenEnd > TokenEnd)
+		{
+			TokenEnd = InToken.TokenEnd;
+		}
+	}
 
 protected:
 	friend class FTokenStream;
 
-	FStringToken(const TCHAR* InStart, int32 Line = 0, int32 Character = 0)
-		: TokenStart(InStart), TokenEnd(InStart), LineNumber(Line), CharacterIndex(Character)
-	{}
+	explicit FStringToken(const TCHAR* InStart, int32 Line = 0, int32 Character = 0)
+		: TokenStart(InStart)
+		, TokenEnd(InStart)
+		, LineNumber(Line)
+		, CharacterIndex(Character)
+	{
+	}
 
-	/** The start of the token */
-	const TCHAR* TokenStart;
-	/** The end of the token */
-	const TCHAR* TokenEnd;
-	/** Line number and Character index */
-	int32 LineNumber, CharacterIndex;
+	const TCHAR* TokenStart = nullptr;
+	const TCHAR* TokenEnd = nullptr;
+	int32 LineNumber = 0;
+	int32 CharacterIndex = 0;
 };
 
 /** Enum specifying how to treat the currently parsing character. */
@@ -92,7 +122,6 @@ enum class EParseState
 class FTokenStream
 {
 public:
-
 	/**
 	 * Parse out a token using the supplied predicate.
 	 * Will keep consuming characters into the resulting token provided the predicate returns EParseState::Continue or EParseState::StopAfter.
@@ -118,9 +147,8 @@ public:
 	CORE_API TOptional<FStringToken> GenerateToken(int32 NumChars, FStringToken* Accumulate = nullptr) const;
 
 public:
-
 	/** Constructor. The stream is only valid for the lifetime of the string provided */
-	CORE_API FTokenStream(const TCHAR* In);
+	CORE_API explicit FTokenStream(const TCHAR* In);
 	
 	/** Peek at the character at the specified offset from the current read position */
 	CORE_API TCHAR PeekChar(int32 Offset = 0) const;
@@ -137,9 +165,18 @@ public:
 	/** Get the current read position from the start of the stream */
 	CORE_API int32 GetPosition() const;
 
-	const TCHAR* GetStart() const { return Start; }
-	const TCHAR* GetRead() const { return ReadPos; }
-	const TCHAR* GetEnd() const { return End; }
+	const TCHAR* GetStart() const
+	{
+		return Start;
+	}
+	const TCHAR* GetRead() const
+	{
+		return ReadPos;
+	}
+	const TCHAR* GetEnd() const
+	{
+		return End;
+	}
 
 	/** Get the error context from the current read position */
 	CORE_API FString GetErrorContext() const;
@@ -148,7 +185,6 @@ public:
 	CORE_API void SetReadPos(const FStringToken& Token);
 
 private:
-
 	/** The start of the expression */
 	const TCHAR* Start;
 	/** The end of the expression */
@@ -179,25 +215,26 @@ DEFINE_EXPRESSION_NODE_TYPE(double, 0x8444A8A3, 0x19AE4E13, 0xBCFA75EE, 0x39982B
  * 	Can be constructed from any C++ type that has a corresponding DEFINE_EXPRESSION_NODE_TYPE.
  * 	Evaluation behaviour (unary/binary operator etc) is defined in the expression grammar, rather than the type itself.
  */
-class FExpressionNode : FNoncopyable
+class FExpressionNode
 {
 public:
-
 	/** Default constructor */
-	FExpressionNode() : TypeId() {}
+	FExpressionNode() = default;
 
 	/** Construction from client expression data type */
-	template<typename T>
-	FExpressionNode(T In,
-		/** @todo: make this a default function template parameter when VS2012 support goes */
-		typename TEnableIf<!TPointerIsConvertibleFromTo<T, FExpressionNode>::Value>::Type* = nullptr
-		);
+	template <
+		typename T
+		UE_REQUIRES(!std::is_convertible_v<T*, FExpressionNode*>)
+	>
+	FExpressionNode(T In);
 
 	CORE_API ~FExpressionNode();
 
-	/** Move construction/assignment */
+	// Movable, but non copyable
 	CORE_API FExpressionNode(FExpressionNode&& In);
+	FExpressionNode(const FExpressionNode&) = delete;
 	CORE_API FExpressionNode& operator=(FExpressionNode&& In);
+	FExpressionNode& operator=(const FExpressionNode&) = delete;
 
 	/** Get the type identifier of this node */
 	CORE_API const FGuid& GetTypeId() const;
@@ -210,7 +247,6 @@ public:
 	CORE_API FExpressionNode Copy() const;
 
 private:
-
 	/** The maximum size of type we will allow allocation on the stack (for efficiency). Anything larger will be allocated on the heap. */
 	static constexpr uint32 MaxStackAllocationSize = 64 - sizeof(FGuid);
 
@@ -224,18 +260,15 @@ private:
 };
 
 /** A specific token in a stream. Comprises an expression node, and the stream token it was created from */
-class FExpressionToken : FNoncopyable
+class FExpressionToken
 {
 public:
-	FExpressionToken(const FStringToken& InContext, FExpressionNode InNode)
+	explicit FExpressionToken(const FStringToken& InContext, FExpressionNode InNode)
 		: Node(MoveTemp(InNode))
 		, Context(InContext)
 	{
 	}
 
-	FExpressionToken(FExpressionToken&& In) : Node(MoveTemp(In.Node)), Context(In.Context) {}
-	FExpressionToken& operator=(FExpressionToken&& In) { Node = MoveTemp(In.Node); Context = In.Context; return *this; }
-
 	FExpressionNode Node;
 	FStringToken Context;
 };
@@ -246,12 +279,12 @@ struct FCompiledToken : FExpressionToken
 	// Todo: add callable types here?
 	enum EType { Operand, PreUnaryOperator, PostUnaryOperator, BinaryOperator, ShortCircuit, Benign };
 
-	FCompiledToken(EType InType, FExpressionToken InToken, TOptional<int32> InShortCircuitIndex = TOptional<int32>())
-		: FExpressionToken(MoveTemp(InToken)), Type(InType), ShortCircuitIndex(InShortCircuitIndex)
-	{}
-
-	FCompiledToken(FCompiledToken&& In) : FExpressionToken(MoveTemp(In)), Type(In.Type) {}
-	FCompiledToken& operator=(FCompiledToken&& In) { FExpressionToken::operator=(MoveTemp(In)); Type = In.Type; return *this; }
+	explicit FCompiledToken(EType InType, FExpressionToken InToken, TOptional<int32> InShortCircuitIndex = TOptional<int32>())
+		: FExpressionToken(MoveTemp(InToken))
+		, Type(InType)
+		, ShortCircuitIndex(InShortCircuitIndex)
+	{
+	}
 
 	EType Type;
 	TOptional<int32> ShortCircuitIndex;
@@ -359,13 +392,11 @@ struct TOperatorJumpTable
 	void MapShortCircuit(FuncType InFunc);
 
 public:
-
-	typedef TFunction<FExpressionResult(const FExpressionNode&, const ContextType* Context)> FUnaryFunction;
-	typedef TFunction<FExpressionResult(const FExpressionNode&, const FExpressionNode&, const ContextType* Context)> FBinaryFunction;
-	typedef TFunction<bool(const FExpressionNode&, const ContextType* Context)> FShortCircuit;
+	using FUnaryFunction = TFunction<FExpressionResult(const FExpressionNode&, const ContextType* Context)>;
+	using FBinaryFunction = TFunction<FExpressionResult(const FExpressionNode&, const FExpressionNode&, const ContextType* Context)>;
+	using FShortCircuit = TFunction<bool(const FExpressionNode&, const ContextType* Context)>;
 
 private:
-
 	/** Maps of unary/binary operators */
 	TMap<FOperatorFunctionID, FUnaryFunction>  PreUnaryOps;
 	TMap<FOperatorFunctionID, FUnaryFunction>  PostUnaryOps;
@@ -373,7 +404,7 @@ private:
 	TMap<FOperatorFunctionID, FShortCircuit>   BinaryShortCircuits;
 };
 
-typedef TOperatorJumpTable<> FOperatorJumpTable;
+using FOperatorJumpTable = TOperatorJumpTable<>;
 
 /** Structures used for managing the evaluation environment for operators in an expression. This class manages the evaluation context
  * to avoid templating the whole evaluation code on a context type
@@ -392,9 +423,11 @@ struct IOperatorEvaluationEnvironment
 template<typename ContextType = void>
 struct TOperatorEvaluationEnvironment : IOperatorEvaluationEnvironment
 {
-	TOperatorEvaluationEnvironment(const TOperatorJumpTable<ContextType>& InOperators, const ContextType* InContext)
-		: Operators(InOperators), Context(InContext)
-	{}
+	explicit TOperatorEvaluationEnvironment(const TOperatorJumpTable<ContextType>& InOperators, const ContextType* InContext)
+		: Operators(InOperators)
+		, Context(InContext)
+	{
+	}
 
 	virtual FExpressionResult ExecPreUnary(const FExpressionToken& Operator, const FExpressionToken& R) const override
 	{
@@ -418,14 +451,21 @@ private:
 	const ContextType* Context;
 };
 
-typedef TOperatorEvaluationEnvironment<> FOperatorEvaluationEnvironment;
+using FOperatorEvaluationEnvironment = TOperatorEvaluationEnvironment<>;
 
 /** Class used to consume tokens from a string */
 class FExpressionTokenConsumer
 {
 public:
 	/** Construction from a raw string. The consumer is only valid as long as the string is valid */
-	CORE_API FExpressionTokenConsumer(const TCHAR* InExpression);
+	CORE_API explicit FExpressionTokenConsumer(const TCHAR* InExpression);
+
+	// Non-copyable
+	FExpressionTokenConsumer(FExpressionTokenConsumer&&) = delete;
+	FExpressionTokenConsumer(const FExpressionTokenConsumer&) = delete;
+	FExpressionTokenConsumer& operator=(FExpressionTokenConsumer&&) = delete;
+	FExpressionTokenConsumer& operator=(const FExpressionTokenConsumer&) = delete;
+	~FExpressionTokenConsumer() = default;
 
 	/** Extract the list of tokens from this consumer */
 	CORE_API TArray<FExpressionToken> Extract();
@@ -439,9 +479,6 @@ public:
 	FTokenStream& GetStream() { return Stream; }
 
 private:
-	CORE_API FExpressionTokenConsumer(const FExpressionTokenConsumer&);
-	CORE_API FExpressionTokenConsumer& operator=(const FExpressionTokenConsumer&);
-
 	/** Array of added tokens */
 	TArray<FExpressionToken> Tokens;
 
@@ -454,36 +491,38 @@ private:
  * 	Definitions may add FExpressionNodes parsed from the provided consumer's stream, or return an optional error.
  *	Where a definition performs no mutable operations, subsequent token definitions will be invoked.
  */
-typedef TOptional<FExpressionError> (FExpressionDefinition)(FExpressionTokenConsumer&);
+using FExpressionDefinition = TOptional<FExpressionError>(FExpressionTokenConsumer&);
 
 
 /** A lexeme dictionary defining how to lex an expression. */
 class FTokenDefinitions
 {
 public:
-	FTokenDefinitions() : bIgnoreWhitespace(false) {}
+	FTokenDefinitions() = default;
 
 	/** Define the grammar to ignore whitespace between tokens, unless explicitly included in a token */
-	void IgnoreWhitespace() { bIgnoreWhitespace = true; }
+	void IgnoreWhitespace()
+	{
+		bIgnoreWhitespace = true;
+	}
 
 	/** Define a token by way of a function to be invoked to attempt to parse a token from a stream */
 	CORE_API void DefineToken(TFunction<FExpressionDefinition>&& Definition);
 
-public:
-
 	/** Check if the grammar ignores whitespace */
-	bool DoesIgnoreWhitespace() { return bIgnoreWhitespace; }
+	bool DoesIgnoreWhitespace()
+	{
+		return bIgnoreWhitespace;
+	}
 
 	/** Consume a token for the specified consumer */
 	CORE_API TOptional<FExpressionError> ConsumeTokens(FExpressionTokenConsumer& Consumer) const;
 
 private:
-
 	CORE_API TOptional<FExpressionError> ConsumeToken(FExpressionTokenConsumer& Consumer) const;
 
 private:
-
-	bool bIgnoreWhitespace;
+	bool bIgnoreWhitespace = false;
 	TArray<TFunction<FExpressionDefinition>> Definitions;
 };
 
@@ -524,16 +563,25 @@ class FExpressionGrammar
 {
 public:
 	/** Define a grouping operator from two expression node types */
-	template<typename TStartGroup, typename TEndGroup>
-	void DefineGrouping() { Groupings.Add(TGetExpressionNodeTypeId<TStartGroup>::GetTypeId(), TGetExpressionNodeTypeId<TEndGroup>::GetTypeId()); }
+	template<typename StartGroupType, typename EndGroupType>
+	void DefineGrouping()
+	{
+		Groupings.Add(TGetExpressionNodeTypeId<StartGroupType>::GetTypeId(), TGetExpressionNodeTypeId<EndGroupType>::GetTypeId());
+	}
 
 	/** Define a pre-unary operator for the specified symbol */
-	template<typename TExpressionNode>
-	void DefinePreUnaryOperator() { PreUnaryOperators.Add(TGetExpressionNodeTypeId<TExpressionNode>::GetTypeId()); }
+	template<typename ExpressionNodeType>
+	void DefinePreUnaryOperator()
+	{
+		PreUnaryOperators.Add(TGetExpressionNodeTypeId<ExpressionNodeType>::GetTypeId());
+	}
 
 	/** Define a post-unary operator for the specified symbol */
-	template<typename TExpressionNode>
-	void DefinePostUnaryOperator() { PostUnaryOperators.Add(TGetExpressionNodeTypeId<TExpressionNode>::GetTypeId()); }
+	template<typename ExpressionNodeType>
+	void DefinePostUnaryOperator()
+	{
+		PostUnaryOperators.Add(TGetExpressionNodeTypeId<ExpressionNodeType>::GetTypeId());
+	}
 
 	/**
 	 * Define a binary operator for the specified symbol, with the specified precedence and associativity
@@ -542,7 +590,7 @@ public:
 	 * @param InPrecedence		The precedence (priority of execution) this operator should have
 	 * @param InAssociativity	With operators of the same precedence, determines whether they execute left to right, or right to left
 	 */
-	template<typename TExpressionNode>
+	template<typename ExpressionNodeType>
 	void DefineBinaryOperator(int32 InPrecedence, EAssociativity InAssociativity=EAssociativity::RightToLeft, bool bCanShortCircuit = false)
 	{
 #if DO_CHECK
@@ -558,7 +606,7 @@ public:
 		}
 #endif
 
-		BinaryOperators.Add(TGetExpressionNodeTypeId<TExpressionNode>::GetTypeId(), FOpParameters(InPrecedence, InAssociativity, bCanShortCircuit));
+		BinaryOperators.Add(TGetExpressionNodeTypeId<ExpressionNodeType>::GetTypeId(), FOpParameters(InPrecedence, InAssociativity, bCanShortCircuit));
 	}
 
 public:
@@ -583,4 +631,8 @@ private:
 	TMap<FGuid, FOpParameters>	BinaryOperators;
 };
 
+#if UE_ENABLE_INCLUDE_ORDER_DEPRECATED_IN_5_6
+#include "Templates/PointerIsConvertibleFromTo.h"
+#endif
+
 #include "Misc/ExpressionParserTypes.inl" // IWYU pragma: export
diff --git a/Engine/Source/Runtime/Core/Public/Misc/ExpressionParserTypes.inl b/Engine/Source/Runtime/Core/Public/Misc/ExpressionParserTypes.inl
--- a/Engine/Source/Runtime/Core/Public/Misc/ExpressionParserTypes.inl
+++ b/Engine/Source/Runtime/Core/Public/Misc/ExpressionParserTypes.inl
@@ -3,8 +3,7 @@
 #pragma once
 
 #include "CoreFwd.h"
-#include "Templates/RemoveReference.h"
-#include "Templates/PointerIsConvertibleFromTo.h"
+#include <type_traits>
 
 #define LOCTEXT_NAMESPACE "ExpressionParser"
 
@@ -17,7 +16,7 @@ namespace Impl
 	/** Interface for a wrapper utility for any moveable/copyable data */
 	struct IExpressionNodeStorage
 	{
-		virtual ~IExpressionNodeStorage() {}
+		virtual ~IExpressionNodeStorage() = default;
 		/** Move this type into another unallocated buffer (move-construct a new type from our wrapped value) */
 		virtual void Reseat(uint8* Dst) = 0;
 		/** Move this type to a buffer already allocated to the same type (uses type-defined move-assignment) */
@@ -34,15 +33,27 @@ namespace Impl
 		/** The data itself, allocated on the stack */
 		T Value;
 
-		/** Constructor/destructor */
-		FInlineDataStorage(T InValue) : Value(MoveTemp(InValue)) {}
-		~FInlineDataStorage(){}
-		
+		FInlineDataStorage(T InValue)
+			: Value(MoveTemp(InValue))
+		{
+		}
+
 		const T* Access() const { return &Value; }
 
-		virtual void Reseat(uint8* Dst) override 			{ ::new((void*)Dst) FInlineDataStorage(MoveTemp(Value)); }
-		virtual void MoveAssign(uint8* Dst) override 		{ reinterpret_cast<FInlineDataStorage*>(Dst)->Value = MoveTemp(Value); }
-		virtual FExpressionNode Copy() const override		{ return Value; }
+		virtual void Reseat(uint8* Dst) override
+		{
+			::new((void*)Dst) FInlineDataStorage(MoveTemp(Value));
+		}
+
+		virtual void MoveAssign(uint8* Dst) override
+		{
+			reinterpret_cast<FInlineDataStorage*>(Dst)->Value = MoveTemp(Value);
+		}
+
+		virtual FExpressionNode Copy() const override
+		{
+			return Value;
+		}
 	};
 
 	/** Data is stored on the heap in this implementation */
@@ -53,22 +64,36 @@ namespace Impl
 		TUniquePtr<T> Value;
 		
 		/** Constructor/destructor */
-		FHeapDataStorage(T InValue) : Value(new T(MoveTemp(InValue))) {}
-		FHeapDataStorage(FHeapDataStorage&& In) : Value(MoveTemp(In.Value)) {}
-		~FHeapDataStorage(){}
-		
-		const T* Access() const { return Value.Get(); }
-
-		virtual void Reseat(uint8* Dst) override 			{ ::new((void*)Dst) FHeapDataStorage(MoveTemp(*this)); }
-		virtual void MoveAssign(uint8* Dst) override 		{ reinterpret_cast<FHeapDataStorage*>(Dst)->Value = MoveTemp(Value); }
-		virtual FExpressionNode Copy() const override		{ return *Value; }
+		FHeapDataStorage(T InValue)
+			: Value(MakeUnique<T>(MoveTemp(InValue)))
+		{
+		}
+
+		const T* Access() const
+		{
+			return Value.Get();
+		}
+
+		virtual void Reseat(uint8* Dst) override
+		{
+			::new((void*)Dst) FHeapDataStorage(MoveTemp(*this));
+		}
+
+		virtual void MoveAssign(uint8* Dst) override
+		{
+			reinterpret_cast<FHeapDataStorage*>(Dst)->Value = MoveTemp(Value);
+		}
+		virtual FExpressionNode Copy() const override
+		{
+			return *Value;
+		}
 	};
 
-	template<typename T, uint32 MaxStackAllocationSize, typename Enabled=void> struct TStorageTypeDeduction;
-	template<typename T, uint32 MaxStackAllocationSize> struct TStorageTypeDeduction<T, MaxStackAllocationSize, typename TEnableIf<(sizeof(FInlineDataStorage<T>) <= MaxStackAllocationSize)>::Type>
-		{ typedef FInlineDataStorage<T> Type; };
-	template<typename T, uint32 MaxStackAllocationSize> struct TStorageTypeDeduction<T, MaxStackAllocationSize, typename TEnableIf<(sizeof(FInlineDataStorage<T>) >  MaxStackAllocationSize)>::Type>
-		{ typedef FHeapDataStorage<T> 	Type; };
+	template<typename T, uint32 MaxStackAllocationSize>
+	struct TStorageTypeDeduction
+	{
+		using Type = std::conditional_t<sizeof(FInlineDataStorage<T>) <= MaxStackAllocationSize, FInlineDataStorage<T>, FHeapDataStorage<T>>;
+	};
 
 
 	/** Machinery required for operator mapping */
@@ -76,95 +101,165 @@ namespace Impl
 	template <typename> struct TCallableInfoImpl;
 
 	template <typename Ret_, typename T, typename Arg1_>
-	struct TCallableInfoImpl<Ret_ (T::*)(Arg1_)> { typedef Ret_ Ret; typedef Arg1_ Arg1; enum { NumArgs = 1 }; };
+	struct TCallableInfoImpl<Ret_ (T::*)(Arg1_)>
+	{
+		using Ret = Ret_;
+		using Arg1 = Arg1_;
+		static constexpr uint32 NumArgs = 1;
+	};
+
 	template <typename Ret_, typename T, typename Arg1_>
-	struct TCallableInfoImpl<Ret_ (T::*)(Arg1_) const> { typedef Ret_ Ret; typedef Arg1_ Arg1; enum { NumArgs = 1 };  };
+	struct TCallableInfoImpl<Ret_ (T::*)(Arg1_) const>
+	{
+		using Ret = Ret_;
+		using Arg1 = Arg1_;
+		static constexpr uint32 NumArgs = 1;
+	};
 
 	template <typename Ret_, typename T, typename Arg1_, typename Arg2_>
-	struct TCallableInfoImpl<Ret_ (T::*)(Arg1_, Arg2_)> { typedef Ret_ Ret; typedef Arg1_ Arg1; typedef Arg2_ Arg2; enum { NumArgs = 2 };  };
+	struct TCallableInfoImpl<Ret_ (T::*)(Arg1_, Arg2_)>
+	{
+		using Ret = Ret_;
+		using Arg1 = Arg1_;
+		using Arg2 = Arg2_;
+		static constexpr uint32 NumArgs = 2;
+	};
+
 	template <typename Ret_, typename T, typename Arg1_, typename Arg2_>
-	struct TCallableInfoImpl<Ret_ (T::*)(Arg1_, Arg2_) const> { typedef Ret_ Ret; typedef Arg1_ Arg1; typedef Arg2_ Arg2; enum { NumArgs = 2 };  };
+	struct TCallableInfoImpl<Ret_ (T::*)(Arg1_, Arg2_) const>
+	{
+		using Ret = Ret_;
+		using Arg1 = Arg1_;
+		using Arg2 = Arg2_;
+		static constexpr uint32 NumArgs = 2;
+	};
 	
 	template <typename Ret_, typename T, typename Arg1_, typename Arg2_, typename Arg3_>
-	struct TCallableInfoImpl<Ret_ (T::*)(Arg1_, Arg2_, Arg3_)> { typedef Ret_ Ret; typedef Arg1_ Arg1; typedef Arg2_ Arg2; typedef Arg3_ Arg3; enum { NumArgs = 3 }; };
-	template <typename Ret_, typename T, typename Arg1_, typename Arg2_, typename Arg3_>
-	struct TCallableInfoImpl<Ret_ (T::*)(Arg1_, Arg2_, Arg3_) const> { typedef Ret_ Ret; typedef Arg1_ Arg1; typedef Arg2_ Arg2; typedef Arg3_ Arg3; enum { NumArgs = 3 }; };
+	struct TCallableInfoImpl<Ret_ (T::*)(Arg1_, Arg2_, Arg3_)>
+	{
+		using Ret = Ret_;
+		using Arg1 = Arg1_;
+		using Arg2 = Arg2_;
+		using Arg3 = Arg3_;
+		static constexpr uint32 NumArgs = 3;
+	};
 
-	/** @todo: decltype(&T::operator()) can go directly in the specialization below if it weren't for VS2012 support */
-	template<typename T>
-	struct TGetOperatorCallPtr { typedef decltype(&T::operator()) Type; };
+	template <typename Ret_, typename T, typename Arg1_, typename Arg2_, typename Arg3_>
+	struct TCallableInfoImpl<Ret_ (T::*)(Arg1_, Arg2_, Arg3_) const>
+	{
+		using Ret = Ret_;
+		using Arg1 = Arg1_;
+		using Arg2 = Arg2_;
+		using Arg3 = Arg3_;
+		static constexpr uint32 NumArgs = 3;
+	};
 
 	template <typename T>
-	struct TCallableInfo : TCallableInfoImpl<typename TGetOperatorCallPtr<T>::Type> {};
+	struct TCallableInfo : TCallableInfoImpl<decltype(&T::operator())>
+	{
+	};
 
 	/** Overloaded function that returns an FExpressionResult, regardless of what is passed in */
 	template<typename T>
-	inline FExpressionResult ForwardReturnType(T&& Result) 			{ return MakeValue(MoveTemp(Result)); }
-	inline FExpressionResult ForwardReturnType(FExpressionResult&& Result) 	{ return MoveTemp(Result); }
-
-	/** Wrapper function for supplied functions of the signature T(A) */
-	template<typename OperandType, typename ContextType, typename FuncType>
-	inline typename TEnableIf<Impl::TCallableInfo<FuncType>::NumArgs == 1, typename TOperatorJumpTable<ContextType>::FUnaryFunction>::Type WrapUnaryFunction(FuncType In)
+	inline FExpressionResult ForwardReturnType(T&& Result)
 	{
-		// Ignore the context
-		return [=](const FExpressionNode& InOperand, const ContextType* Context) {
-			return ForwardReturnType(In(*InOperand.Cast<OperandType>()));
-		};
+		return MakeValue(MoveTemp(Result));
 	}
-
-	/** Wrapper function for supplied functions of the signature T(A, const ContextType* Context) */
-	template<typename OperandType, typename ContextType, typename FuncType>
-	inline typename TEnableIf<Impl::TCallableInfo<FuncType>::NumArgs == 2, typename TOperatorJumpTable<ContextType>::FUnaryFunction>::Type WrapUnaryFunction(FuncType In)
+	inline FExpressionResult ForwardReturnType(FExpressionResult&& Result)
 	{
-		// Ignore the context
-		return [=](const FExpressionNode& InOperand, const ContextType* Context) {
-			return ForwardReturnType(In(*InOperand.Cast<OperandType>(), Context));
-		};
+		return MoveTemp(Result);
 	}
 
-	/** Wrapper function for supplied functions of the signature T(A, B) */
-	template<typename LeftOperandType, typename RightOperandType, typename ContextType, typename FuncType>
-	inline typename TEnableIf<Impl::TCallableInfo<FuncType>::NumArgs == 2, typename TOperatorJumpTable<ContextType>::FBinaryFunction>::Type WrapBinaryFunction(FuncType In)
+	/** Wrapper function for supplied functions of the signature T(A) */
+	template<typename OperandType, typename ContextType, typename FuncType>
+	inline typename TOperatorJumpTable<ContextType>::FUnaryFunction WrapUnaryFunction(FuncType In)
 	{
-		// Ignore the context
-		return [=](const FExpressionNode& InLeftOperand, const FExpressionNode& InRightOperand, const ContextType* Context) {
-			return ForwardReturnType(In(*InLeftOperand.Cast<LeftOperandType>(), *InRightOperand.Cast<RightOperandType>()));
-		};
+		constexpr uint32 NumArgs = Impl::TCallableInfo<FuncType>::NumArgs;
+
+		if constexpr (NumArgs == 1)
+		{
+			// Ignore the context
+			return [=](const FExpressionNode& InOperand, const ContextType* Context)
+			{
+				return ForwardReturnType(In(*InOperand.Cast<OperandType>()));
+			};
+		}
+		else if constexpr (NumArgs == 2)
+		{
+			return [=](const FExpressionNode& InOperand, const ContextType* Context)
+			{
+				return ForwardReturnType(In(*InOperand.Cast<OperandType>(), Context));
+			};
+		}
+		else
+		{
+			// sizeof(OperandType) == 0 is used to create a false value based on a template parameter
+			static_assert(sizeof(OperandType) == 0, "FuncType has an unexpected number of parameters");
+		}
 	}
 
-	/** Wrapper function for supplied functions of the signature T(A, B, const ContextType* Context) */
+	/** Wrapper function for supplied functions of the signature T(A, B) */
 	template<typename LeftOperandType, typename RightOperandType, typename ContextType, typename FuncType>
-	inline typename TEnableIf<Impl::TCallableInfo<FuncType>::NumArgs == 3, typename TOperatorJumpTable<ContextType>::FBinaryFunction>::Type WrapBinaryFunction(FuncType In)
+	inline typename TOperatorJumpTable<ContextType>::FBinaryFunction WrapBinaryFunction(FuncType In)
 	{
-		// Ignore the context
-		return [=](const FExpressionNode& InLeftOperand, const FExpressionNode& InRightOperand, const ContextType* Context) {
-			return ForwardReturnType(In(*InLeftOperand.Cast<LeftOperandType>(), *InRightOperand.Cast<RightOperandType>(), Context));
-		};
+		constexpr uint32 NumArgs = Impl::TCallableInfo<FuncType>::NumArgs;
+
+		if constexpr (NumArgs == 2)
+		{
+			// Ignore the context
+			return [=](const FExpressionNode& InLeftOperand, const FExpressionNode& InRightOperand, const ContextType* Context)
+			{
+				return ForwardReturnType(In(*InLeftOperand.Cast<LeftOperandType>(), *InRightOperand.Cast<RightOperandType>()));
+			};
+		}
+		else if constexpr (NumArgs == 3)
+		{
+			return [=](const FExpressionNode& InLeftOperand, const FExpressionNode& InRightOperand, const ContextType* Context)
+			{
+				return ForwardReturnType(In(*InLeftOperand.Cast<LeftOperandType>(), *InRightOperand.Cast<RightOperandType>(), Context));
+			};
+		}
+		else
+		{
+			// sizeof(LeftOperandType) == 0 is used to create a false value based on a template parameter
+			static_assert(sizeof(LeftOperandType) == 0, "FuncType has an unexpected number of parameters");
+		}
 	}
 
 	/** Wrapper function for supplied functions of the signature bool(A, const ContextType* Context) */
 	template<typename OperandType, typename ContextType, typename FuncType>
-	inline typename TEnableIf<Impl::TCallableInfo<FuncType>::NumArgs == 1, typename TOperatorJumpTable<ContextType>::FShortCircuit>::Type WrapShortCircuitFunction(FuncType In)
+	inline typename TOperatorJumpTable<ContextType>::FShortCircuit WrapShortCircuitFunction(FuncType In)
 	{
-		// Ignore the context
-		return [=](const FExpressionNode& InOperand, const ContextType* Context) {
-			return In(*InOperand.Cast<OperandType>());
-		};
-	}
-
-	/** Wrapper function for supplied functions of the signature bool(A, const ContextType* Context) */
-	template<typename OperandType, typename ContextType, typename FuncType>
-	inline typename TEnableIf<Impl::TCallableInfo<FuncType>::NumArgs == 2, typename TOperatorJumpTable<ContextType>::FShortCircuit>::Type WrapShortCircuitFunction(FuncType In)
-	{
-		// Ignore the context
-		return [=](const FExpressionNode& InOperand, const ContextType* Context) {
-			return In(*InOperand.Cast<OperandType>(), Context);
-		};
+		constexpr uint32 NumArgs = Impl::TCallableInfo<FuncType>::NumArgs;
+
+		if constexpr (NumArgs == 1)
+		{
+			// Ignore the context
+			return [=](const FExpressionNode& InOperand, const ContextType* Context)
+			{
+				return In(*InOperand.Cast<OperandType>());
+			};
+		}
+		else if constexpr (NumArgs == 2)
+		{
+			return [=](const FExpressionNode& InOperand, const ContextType* Context)
+			{
+				return In(*InOperand.Cast<OperandType>(), Context);
+			};
+		}
+		else
+		{
+			// sizeof(OperandType) == 0 is used to create a false value based on a template parameter
+			static_assert(sizeof(OperandType) == 0, "FuncType has an unexpected number of parameters");
+		}
 	}
 }
 
-template<typename T>
-FExpressionNode::FExpressionNode(T In, typename TEnableIf<!TPointerIsConvertibleFromTo<T,FExpressionNode>::Value>::Type*)
+template <
+	typename T
+	UE_REQUIRES_DEFINITION(!std::is_convertible_v<T*, FExpressionNode*>)
+>
+FExpressionNode::FExpressionNode(T In)
 	: TypeId(TGetExpressionNodeTypeId<T>::GetTypeId())
 {
 	// Choose the relevant allocation strategy based on the size of the type
@@ -245,7 +340,7 @@ template<typename ContextType>
 template<typename OperatorType, typename FuncType>
 void TOperatorJumpTable<ContextType>::MapPreUnary(FuncType InFunc)
 {
-	typedef std::remove_const_t<typename TRemoveReference<typename Impl::TCallableInfo<FuncType>::Arg1>::Type> OperandType;
+	using OperandType = std::decay_t<typename Impl::TCallableInfo<FuncType>::Arg1>;
 
 	FOperatorFunctionID ID = {
 		TGetExpressionNodeTypeId<OperatorType>::GetTypeId(),
@@ -260,7 +355,7 @@ template<typename ContextType>
 template<typename OperatorType, typename FuncType>
 void TOperatorJumpTable<ContextType>::MapPostUnary(FuncType InFunc)
 {
-	typedef std::remove_const_t<typename TRemoveReference<typename Impl::TCallableInfo<FuncType>::Arg1>::Type> OperandType;
+	using OperandType = std::decay_t<typename Impl::TCallableInfo<FuncType>::Arg1>;
 
 	FOperatorFunctionID ID = {
 		TGetExpressionNodeTypeId<OperatorType>::GetTypeId(),
@@ -275,8 +370,8 @@ template<typename ContextType>
 template<typename OperatorType, typename FuncType>
 void TOperatorJumpTable<ContextType>::MapBinary(FuncType InFunc)
 {
-	typedef std::remove_const_t<typename TRemoveReference<typename Impl::TCallableInfo<FuncType>::Arg1>::Type> LeftOperandType;
-	typedef std::remove_const_t<typename TRemoveReference<typename Impl::TCallableInfo<FuncType>::Arg2>::Type> RightOperandType;
+	using LeftOperandType = std::decay_t<typename Impl::TCallableInfo<FuncType>::Arg1>;
+	using RightOperandType = std::decay_t<typename Impl::TCallableInfo<FuncType>::Arg2>;
 
 	FOperatorFunctionID ID = {
 		TGetExpressionNodeTypeId<OperatorType>::GetTypeId(),
@@ -291,7 +386,7 @@ template<typename ContextType>
 template<typename OperatorType, typename FuncType>
 void TOperatorJumpTable<ContextType>::MapShortCircuit(FuncType InFunc)
 {
-	typedef std::remove_const_t<typename TRemoveReference<typename Impl::TCallableInfo<FuncType>::Arg1>::Type> OperandType;
+	using OperandType = std::decay_t<typename Impl::TCallableInfo<FuncType>::Arg1>;
 
 	FOperatorFunctionID ID = {
 		TGetExpressionNodeTypeId<OperatorType>::GetTypeId(),
@@ -303,6 +398,6 @@ void TOperatorJumpTable<ContextType>::MapShortCircuit(FuncType InFunc)
 }
 
 
-typedef TOperatorEvaluationEnvironment<> FOperatorEvaluationEnvironment;
+using FOperatorEvaluationEnvironment = TOperatorEvaluationEnvironment<>;
 
 #undef LOCTEXT_NAMESPACE
