Add support for Android's built in Profiler (as of OS 15).
Only available when compiled with SDK 35 or above.

diff --git a/Engine/Source/Runtime/Core/Core.Build.cs b/Engine/Source/Runtime/Core/Core.Build.cs
--- a/Engine/Source/Runtime/Core/Core.Build.cs
+++ b/Engine/Source/Runtime/Core/Core.Build.cs
@@ -24,6 +24,8 @@ public class Core : ModuleRules
 				PrivateDependencyModuleNames.Add("HWCPipe");        // Performance counters for ARM CPUs and ARM Mali GPUs
 				PrivateDependencyModuleNames.Add("heapprofd");      // Exposes custom allocators to Google's Memory Profiler
 			}
+			string ModulePath = Utils.MakePathRelativeTo(ModuleDirectory, Target.RelativeEnginePath);
+			AdditionalPropertiesForReceipt.Add("AndroidPlugin", Path.Combine(ModulePath, "Core_APL.xml"));
 		}
 
 		PrivateDependencyModuleNames.Add("AtomicQueue");
diff --git a/Engine/Source/Runtime/Core/Core_APL.xml b/Engine/Source/Runtime/Core/Core_APL.xml
new file mode 100644
--- /dev/null
+++ b/Engine/Source/Runtime/Core/Core_APL.xml
@@ -0,0 +1,73 @@
+<?xml version="1.0" encoding="utf-8"?>
+<!--Core module additions-->
+<root xmlns:android="http://schemas.android.com/apk/res/android">
+	<!-- init section is always evaluated once per architecture -->
+	<init>
+		<log text="Core UPL init"/>
+		<!-- we cant set bAndroidProfilerAPIAvailable $S(CompileSDK) -->
+	</init>
+
+	<gameActivityImportAdditions>
+		<setBoolIsGreaterEqual result="bAndroidProfilerAPIAvailable" arg1="$S(CompileSDK)" arg2="35"/>
+		<if condition="bAndroidProfilerAPIAvailable">
+			<true>
+				<log text="Compile SDK is 35+, adding Android Profiler access."/>
+				<insert>
+					import com.epicgames.unreal.androidprofiling.ProfilerAccessor;
+				</insert>
+			</true>
+			<false>
+				<log text="Compile SDK is below 35, Android Profiler access will not available."/>
+			</false>
+		</if>
+	</gameActivityImportAdditions>
+	
+	<gameActivityOnCreateAdditions>
+		<setBoolIsGreaterEqual result="bAndroidProfilerAPIAvailable" arg1="$S(CompileSDK)" arg2="35"/>
+		<if condition="bAndroidProfilerAPIAvailable">
+			<true>
+				<insert>
+					<![CDATA[
+					ProfilerAccessor.Init(this);
+					]]>
+				</insert>
+			</true>
+
+		</if>
+	</gameActivityOnCreateAdditions>
+
+	<proguardAdditions>
+		<setBoolIsGreaterEqual result="bAndroidProfilerAPIAvailable" arg1="$S(CompileSDK)" arg2="35"/>
+		<if condition="bAndroidProfilerAPIAvailable">
+			<true>
+				<insert>
+					-keep class com.epicgames.unreal.androidprofiling.ProfilerAccessor {
+					public *;
+					}
+				</insert>
+			</true>
+		</if>
+	</proguardAdditions>
+
+	<resourceCopies>
+		<setBoolIsGreaterEqual result="bAndroidProfilerAPIAvailable" arg1="$S(CompileSDK)" arg2="35"/>
+		<if condition="bAndroidProfilerAPIAvailable">
+			<true>
+				<copyFile src="$S(EngineDir)/Source/Runtime/Core/Private/Android/Java/ProfilerAccessor.java" dst="$S(BuildDir)/src/com/epicgames/unreal/androidprofiling/ProfilerAccessor.java" />
+			</true>
+		</if>
+	</resourceCopies>
+	
+	<buildGradleAdditions>
+		<setBoolIsGreaterEqual result="bAndroidProfilerAPIAvailable" arg1="$S(CompileSDK)" arg2="35"/>
+		<if condition="bAndroidProfilerAPIAvailable">
+			<true>
+				<insert>
+				dependencies {
+					implementation('androidx.core:core:1.15.0')
+				}
+				</insert>
+			</true>
+		</if>
+	</buildGradleAdditions>
+</root>
diff --git a/Engine/Source/Runtime/Core/Private/Android/AndroidProfiler.cpp b/Engine/Source/Runtime/Core/Private/Android/AndroidProfiler.cpp
new file mode 100644
--- /dev/null
+++ b/Engine/Source/Runtime/Core/Private/Android/AndroidProfiler.cpp
@@ -0,0 +1,203 @@
+// Copyright Epic Games, Inc. All Rights Reserved.
+
+#include "Android/AndroidProfiler.h"
+
+#if USE_ANDROID_JNI
+#include "CoreGlobals.h"
+#include "CoreMinimal.h"
+#include "Android/AndroidJava.h"
+#include "Android/AndroidJavaEnv.h"
+#include "Android/AndroidProfiler.h"
+#include "HAL/IConsoleManager.h"
+#include "Misc/Paths.h"
+#include "HAL/PlatformFileManager.h"
+#include "GenericPlatform/GenericPlatformFile.h"
+
+DEFINE_LOG_CATEGORY_STATIC(LogAndroidProfiling, Log, Log);
+
+namespace AndroidProfiler
+{
+	static FAutoConsoleCommand AndroidProfileConsoleCommand(
+		TEXT("android.profile"),
+		TEXT("Invoke android's profiling api.\n")
+		TEXT("Call with no arguments to echo usage to the log.\n")
+		TEXT("Note: requires Android 15 or above (API 35).\n")
+		TEXT("(this API is rate limited, to remove the rate limit use: 'adb shell device_config put profiling_testing rate_limiter.disabled true')"),
+		FConsoleCommandWithArgsDelegate::CreateLambda([](const TArray<FString>& Args)
+			{
+				FString CombinedArgs;
+				for (const FString& Argument : Args)
+				{
+					CombinedArgs += Argument + TEXT(" ");
+				}
+				FString ProfileName = FAndroidProfiler::StartCapture(CombinedArgs, [](const FAndroidProfiler::FProfileResults& Results)
+				{
+					if (!Results.FilePath.IsEmpty())
+					{
+						const FString DestPath = FPaths::ProfilingDir() / FPaths::GetCleanFilename(Results.FilePath);
+						UE_CLOG(!Results.Error.IsEmpty(), LogAndroidProfiling, Warning, TEXT("Profile %s completed with log: %s"), *Results.ProfileName, *Results.Error);
+						IPlatformFile& PlatformFile = FPlatformFileManager::Get().GetPlatformFile();
+						// move the file to public path so adb can get to it.
+						bool bFileMoved = PlatformFile.MoveFile(*DestPath, *Results.FilePath);
+						if (!bFileMoved)
+						{
+							// copy the file if the move fails.
+							bFileMoved = PlatformFile.CopyFile(*DestPath, *Results.FilePath);
+							PlatformFile.DeleteFile(*Results.FilePath);
+						}
+
+						UE_LOG(LogAndroidProfiling, Log, TEXT("Profile %s complete, retrieve via 'adb pull %s'"), *Results.ProfileName, *DestPath);
+					}
+					else
+					{
+						UE_LOG(LogAndroidProfiling, Error, TEXT("Profile %s failed: %s"), *Results.ProfileName, *Results.Error);
+					}
+				});
+			}));
+
+	static jclass ProfilerAccessorClassId = 0;
+	static jmethodID IssueProfilerCommandMethod = 0;
+	static jmethodID StopProfilerCommandMethod = 0;
+	static void InitJNI();
+
+	static bool IsJNIAvailable()
+	{
+		InitJNI();
+		return ProfilerAccessorClassId && IssueProfilerCommandMethod && StopProfilerCommandMethod;
+	}
+
+	static void InitJNI()
+	{
+		static bool bInitJNI = false;
+		if (bInitJNI)
+		{
+			return;
+		}
+		bInitJNI = true;
+
+		if (JNIEnv* Env = AndroidJavaEnv::GetJavaEnv())
+		{
+			ProfilerAccessorClassId = AndroidJavaEnv::FindJavaClassGlobalRef("com/epicgames/unreal/androidprofiling/ProfilerAccessor");
+			if(ProfilerAccessorClassId)
+			{
+				IssueProfilerCommandMethod = Env->GetStaticMethodID(ProfilerAccessorClassId, "AndroidThunkJava_IssueProfilerCommand", "(Ljava/lang/String;)[Ljava/lang/String;");
+				StopProfilerCommandMethod = Env->GetStaticMethodID(ProfilerAccessorClassId, "AndroidThunkJava_StopProfilerCommand", "(Ljava/lang/String;)Z");
+			}
+		}
+		UE_CLOG(!IsJNIAvailable(), LogAndroidProfiling, Error, TEXT("JNI Could not find ProfilerAccessor class or methods. ProfilerAccessor is not supported when Compile SDK < 35."));
+	}
+
+	static FString CallProfiler(FString& ProfileNameOUT, const FString& CombinedArgs)
+	{
+		if (IsJNIAvailable())
+		{
+			JNIEnv* Env = AndroidJavaEnv::GetJavaEnv();
+
+			auto CombinedArgsJava = FJavaHelper::ToJavaString(Env, CombinedArgs);
+			TArray<FString> ProfilerResults = FJavaHelper::ObjectArrayToFStringTArray(Env,
+				(jobjectArray)Env->CallStaticObjectMethod(ProfilerAccessorClassId, IssueProfilerCommandMethod, *CombinedArgsJava)
+			);
+
+			const bool bException = AndroidJavaEnv::CheckJavaException();
+			UE_CLOG(bException, LogAndroidProfiling, Error, TEXT("Profiler failed due to java exception."));
+
+			if(ProfilerResults.Num() == 2)
+			{
+				ProfileNameOUT = ProfilerResults[0];
+				const FString& ProfilerMessage = ProfilerResults[1];
+				UE_LOG(LogAndroidProfiling, Log, TEXT("Profiler : %s "), ProfilerMessage.IsEmpty() ? TEXT("profile issued.") : *ProfilerMessage);
+				return ProfilerMessage;
+			}
+			else
+			{
+				UE_LOG(LogAndroidProfiling, Error, TEXT("Profiler failed."));
+			}
+		}
+		else
+		{
+			UE_LOG(LogAndroidProfiling, Error, TEXT("Profile failed, java profiler not found."));
+		}
+		return FString();
+	}
+
+	static bool CancelProfile(const FString& ProfileName)
+	{
+		if( IsJNIAvailable() )
+		{
+			JNIEnv* Env = AndroidJavaEnv::GetJavaEnv();
+			auto ProfileNameJava = FJavaHelper::ToJavaString(Env, ProfileName);
+			return (bool)Env->CallStaticBooleanMethod(ProfilerAccessorClassId, StopProfilerCommandMethod, *ProfileNameJava);
+		}
+		UE_LOG(LogAndroidProfiling, Error, TEXT("Profile cancel failed, java profiler not found."));
+		return false;
+	}
+}
+
+// this provides a static function for accessing FAndroidProfiler just to avoid making OnProfileFinish part of the public api.
+struct FAndroidProfilerInternal
+{
+	static void OnJavaProfileFinish(const FAndroidProfiler::FProfileResults& Results)
+	{
+		FAndroidProfiler::OnProfileFinish(Results);
+	}
+};
+
+// called from jvm after profile completes.
+JNI_METHOD void Java_com_epicgames_unreal_androidprofiling_ProfilerAccessor_nativeOnProfileFinish(JNIEnv* jenv, jobject /*this*/, jstring ProfileNameJava, jstring ProfileErrorJava, jstring ProfileFilepathJava)
+{
+	FAndroidProfiler::FProfileResults Results;
+	Results.ProfileName = FJavaHelper::FStringFromParam(jenv, ProfileNameJava);
+	Results.Error = FJavaHelper::FStringFromParam(jenv, ProfileErrorJava);
+	Results.FilePath = FJavaHelper::FStringFromParam(jenv, ProfileFilepathJava);
+	UE_LOG(LogAndroidProfiling, Log, TEXT("nativeOnProfileFinish (%s, %s, %s)"),*Results.ProfileName, *Results.Error, *Results.FilePath);
+	FAndroidProfilerInternal::OnJavaProfileFinish(Results);
+}
+
+// ProfilerCS is held during StartCapture, that completes before we lock and call OnFinish here.
+void FAndroidProfiler::OnProfileFinish(const FAndroidProfiler::FProfileResults& Results)
+{
+	FScopeLock Lock(&ProfilerCS);
+	FAndroidProfiler::FActiveSessions FoundSession = ActiveSessions.FindAndRemoveChecked(Results.ProfileName);
+	FoundSession.OnFinish(Results);
+}
+
+FString FAndroidProfiler::StartCapture(const FString& Args, TUniqueFunction<void(const FAndroidProfiler::FProfileResults& Results)> OnFinish)
+{
+	FScopeLock Lock(&ProfilerCS);
+
+	FString ProfileName;
+	FString Message = AndroidProfiler::CallProfiler(ProfileName, Args);
+
+	if (!ProfileName.IsEmpty())
+	{
+		check(!ActiveSessions.Contains(ProfileName));
+		ActiveSessions.Add(ProfileName, FAndroidProfiler::FActiveSessions({ MoveTemp(OnFinish) }));
+	}
+	else
+	{
+		FAndroidProfiler::FProfileResults ProfileNameInProgressError;
+		ProfileNameInProgressError.Error = FString::Printf(TEXT("Unable to launch profile. %s"), *Message);
+		OnFinish(ProfileNameInProgressError);
+	}
+	return ProfileName;
+}
+
+void FAndroidProfiler::StopCapture(const FString& ProfileName)
+{
+	FScopeLock Lock(&ProfilerCS);
+	if (ActiveSessions.Contains(ProfileName))
+	{
+		AndroidProfiler::CancelProfile(ProfileName);
+	}
+	else
+	{
+		UE_LOG(LogAndroidProfiling, Error, TEXT("StopCapture ignored, profile %s is not active."), *ProfileName);
+	}
+}
+#else
+
+FString FAndroidProfiler::StartCapture(const FString& Args, TUniqueFunction<void(const FAndroidProfiler::FProfileResults& Results)> OnFinish) { return FString(); }
+void FAndroidProfiler::StopCapture(FString ProfileName) {}
+void FAndroidProfiler::OnProfileFinish(const FAndroidProfiler::FProfileResults& Results) {}
+
+#endif
diff --git a/Engine/Source/Runtime/Core/Private/Android/Java/ProfilerAccessor.java b/Engine/Source/Runtime/Core/Private/Android/Java/ProfilerAccessor.java
new file mode 100644
--- /dev/null
+++ b/Engine/Source/Runtime/Core/Private/Android/Java/ProfilerAccessor.java
@@ -0,0 +1,318 @@
+// Copyright Epic Games, Inc. All Rights Reserved.
+package com.epicgames.unreal.androidprofiling;
+
+import android.content.Context;
+import android.os.Build;
+import android.os.ProfilingResult;
+import java.util.HashMap;
+import java.util.Map;
+import java.util.Objects;
+import java.util.function.Consumer;
+import java.util.concurrent.Executors;
+import com.epicgames.unreal.Logger;
+import com.epicgames.unreal.GameActivity;
+import androidx.annotation.RequiresApi;
+import androidx.core.os.BufferFillPolicy;
+import androidx.core.os.HeapProfileRequestBuilder;
+import androidx.core.os.ProfilingRequest;
+import androidx.core.os.StackSamplingRequestBuilder;
+import androidx.core.os.SystemTraceRequestBuilder;
+
+public class ProfilerAccessor
+{
+	public static void Init(GameActivity InActivity)
+	{
+		mActivity = InActivity;
+	}
+	private static GameActivity mActivity;
+
+	public static Logger Log = new Logger("UE", "ProfilerAccessor");
+	abstract static class BaseProfilerType
+	{
+		public String GetMode() {return "none";}
+
+		public String InvokeProfiler(Context context, Map<String, String> Args, Consumer<android.os.ProfilingResult> listener) { return ""; }
+
+		public ProfilingRequest BuildProfilerRequest(Map<String, String> Args, android.os.CancellationSignal cancellationSignal, StringBuilder OutputMessage) { return null; }
+
+		public final static String ProfileSessionName = "profilename";
+
+		public Map<String, String> GetDefaultArgs()
+		{
+			// args common to all profile types. 
+			Map<String, String> DefaultArgs = new HashMap<String, String>();
+			DefaultArgs.put("duration", "10");
+			DefaultArgs.put("buffersize", "1000");
+			DefaultArgs.put(ProfileSessionName, "profile");
+			return DefaultArgs;
+		}
+
+		public String DefaultArgsHelp()
+		{
+			StringBuilder ret = new StringBuilder();
+			for (Map.Entry<String, String> entry : GetDefaultArgs().entrySet())
+			{
+				ret.append(", [").append(entry.getKey()).append(" ").append(entry.getValue()).append("]");
+			}
+			return ret.toString();
+		}
+	};
+
+	static class CallstackProfiler extends com.epicgames.unreal.androidprofiling.ProfilerAccessor.BaseProfilerType
+	{
+		@Override
+		public String GetMode() {return "callstack";}
+
+		@Override
+		public Map<String, String> GetDefaultArgs()
+		{
+			Map<String, String> DefaultArgs = super.GetDefaultArgs();
+			DefaultArgs.put("samplingfreq", "100");
+			return DefaultArgs;
+		}
+
+		@RequiresApi(api = Build.VERSION_CODES.VANILLA_ICE_CREAM)
+		public ProfilingRequest BuildProfilerRequest(Map<String, String> Args, android.os.CancellationSignal cancellationSignal, StringBuilder OutputMessage)
+		{
+			return new StackSamplingRequestBuilder()
+					.setBufferSizeKb(Integer.parseInt(Objects.requireNonNull(Args.get("buffersize"))) /* Requested buffer size in KB */)
+					.setDurationMs(Integer.parseInt(Args.get("duration")) * 1000 /* Requested profiling duration in millisconds */)
+					.setSamplingFrequencyHz(Integer.parseInt(Args.get("samplingfreq")) /* Requested sampling frequency */)
+					.setTag(Args.get(BaseProfilerType.ProfileSessionName) /* Caller supplied tag for identification */)
+					.setCancellationSignal(cancellationSignal)
+					.build();
+		}
+	}
+
+	static class HeapProfiler extends com.epicgames.unreal.androidprofiling.ProfilerAccessor.BaseProfilerType
+	{
+		@Override
+		public String GetMode() {return "heap";}
+
+		@Override
+		public Map<String, String> GetDefaultArgs()
+		{
+			Map<String, String> DefaultArgs = super.GetDefaultArgs();
+			DefaultArgs.put("samplingintervalbytes", "100");
+			return DefaultArgs;
+		}
+
+		@RequiresApi(api = Build.VERSION_CODES.VANILLA_ICE_CREAM)
+		public ProfilingRequest BuildProfilerRequest(Map<String, String> Args, android.os.CancellationSignal cancellationSignal, StringBuilder OutputMessage)
+		{
+			return new HeapProfileRequestBuilder()
+					.setBufferSizeKb(Integer.parseInt(Objects.requireNonNull(Args.get("buffersize"))) /* Requested buffer size in KB */)
+					.setDurationMs(Integer.parseInt(Args.get("duration")) * 1000 /* Requested profiling duration in millisconds */)
+					.setTrackJavaAllocations(false)
+					.setSamplingIntervalBytes(Integer.parseInt(Args.get("samplingintervalbytes")) /* Requested sampling interval in bytes */)
+					.setTag(Args.get(BaseProfilerType.ProfileSessionName) /* Caller supplied tag for identification */)
+					.setCancellationSignal(cancellationSignal)
+					.build();
+		}
+	}
+
+	static class SystemProfiler extends com.epicgames.unreal.androidprofiling.ProfilerAccessor.BaseProfilerType
+	{
+		@Override
+		public String GetMode() {return "system";}
+
+		@Override
+		public Map<String, String> GetDefaultArgs()
+		{
+			Map<String, String> DefaultArgs = super.GetDefaultArgs();
+			DefaultArgs.put("bufferfillpolicy", "ringbuffer");
+			return DefaultArgs;
+		}
+		@RequiresApi(api = Build.VERSION_CODES.VANILLA_ICE_CREAM)
+		public ProfilingRequest BuildProfilerRequest(Map<String, String> Args, android.os.CancellationSignal cancellationSignal, StringBuilder OutputMessage)
+		{
+			BufferFillPolicy fillPolicy;
+			String desiredFillPolicy = Args.get("bufferfillpolicy");
+			switch( desiredFillPolicy )
+			{
+				case "ringbuffer":
+					fillPolicy = BufferFillPolicy.RING_BUFFER;
+					break;
+				case "discard":
+					fillPolicy = BufferFillPolicy.DISCARD;
+					break;
+				default:
+					OutputMessage.append("Ignoring unknown bufferfillpolicy ("+desiredFillPolicy+")");
+					Log.error(OutputMessage.toString());
+					fillPolicy = BufferFillPolicy.RING_BUFFER;
+					break;
+			}
+
+			return new SystemTraceRequestBuilder()
+					.setBufferSizeKb(Integer.parseInt(Args.get("buffersize")) /* Requested buffer size in KB */)
+					.setDurationMs(Integer.parseInt(Args.get("duration")) * 1000 /* Requested profiling duration in millisconds */)
+					.setBufferFillPolicy(fillPolicy /* Buffer fill policy */)
+					.setTag(Args.get(BaseProfilerType.ProfileSessionName) /* Caller supplied tag for identification */)
+					.setCancellationSignal(cancellationSignal)
+					.build();
+		}
+
+	}
+
+	static com.epicgames.unreal.androidprofiling.ProfilerAccessor.BaseProfilerType[] AllProfilers = {
+			new com.epicgames.unreal.androidprofiling.ProfilerAccessor.CallstackProfiler(),
+			new com.epicgames.unreal.androidprofiling.ProfilerAccessor.SystemProfiler(),
+			new com.epicgames.unreal.androidprofiling.ProfilerAccessor.HeapProfiler()
+	};
+
+
+	private static String[] OnError(String ErrorMessage)
+	{
+		return new String[] { null, ErrorMessage};
+	}
+
+	private static String[] OnSessionIssued(String Tag, String Message)
+	{
+		return new String[] { Tag, Message};
+	}
+
+	private static native void nativeOnProfileFinish(String Tag, String Mesage, String FilePath);
+	static class ActiveProfileSession
+	{
+		public android.os.CancellationSignal CancelSignal;
+	};
+
+	private static final Map<String, ActiveProfileSession> ActiveSessions = new HashMap<String,ActiveProfileSession>();
+
+	private static com.epicgames.unreal.androidprofiling.ProfilerAccessor.BaseProfilerType FindProfiler(String RequestedProfilerMode)
+	{
+		for ( com.epicgames.unreal.androidprofiling.ProfilerAccessor.BaseProfilerType Profiler : AllProfilers)
+		{
+			if( Profiler.GetMode().equals(RequestedProfilerMode) )
+			{
+				return Profiler;
+			}
+		}
+		return null;
+	}
+	
+	// returns String[] = { profile tag (or null on error), message }
+	public synchronized static String[] IssueProfilerCommand(String RequestedProfilerArgs)
+	{
+		if(Build.VERSION.SDK_INT<Build.VERSION_CODES.VANILLA_ICE_CREAM) // Build.VERSION.VANILLA_ICE_CREAM == 35
+		{
+			return OnError("API version must be >= "+Build.VERSION_CODES.VANILLA_ICE_CREAM+". (current API version = "+ Build.VERSION.SDK_INT+")");
+		}
+
+		int FirstCommaIdx = RequestedProfilerArgs.indexOf(" ");
+		String RequestedProfilerMode = (FirstCommaIdx == -1 ? RequestedProfilerArgs.trim() : RequestedProfilerArgs.substring(0, FirstCommaIdx)).toLowerCase();
+		com.epicgames.unreal.androidprofiling.ProfilerAccessor.BaseProfilerType FoundProfiler = FindProfiler(RequestedProfilerMode);
+
+		if( FoundProfiler == null )
+		{
+			StringBuilder message = new StringBuilder("usage:\nandroid.profiler <profiler mode> [, optional comma separated args]\n\n");
+			for ( com.epicgames.unreal.androidprofiling.ProfilerAccessor.BaseProfilerType Profiler : AllProfilers)
+			{
+				message.append("\tandroid.profiler "+Profiler.GetMode()+Profiler.DefaultArgsHelp()+"\n");
+			}
+			return OnError(message.toString());
+		}
+
+		// find all comma separated args
+		Map<String, String> ArgsMap = FoundProfiler.GetDefaultArgs();
+		for(String CombinedArg : RequestedProfilerArgs.split(" "))
+		{
+			// split the space separated args by = to get arg and parameter. 
+			String Arg = CombinedArg, Parameter = "";
+			int FirstIdx = CombinedArg.indexOf("=");
+			if( FirstIdx > 0 )
+			{
+				Arg = CombinedArg.substring(0, FirstIdx);
+				Parameter = CombinedArg.substring(FirstIdx+1);
+			}
+			ArgsMap.put(Arg, Parameter);
+		}
+
+		String sessionName = ArgsMap.get(BaseProfilerType.ProfileSessionName);
+		// profile tag must be alphanum and <=20 chars.
+		sessionName = sessionName.replaceAll("[^0-9a-zA-Z]", "");
+		sessionName = sessionName.substring(0, Math.min(sessionName.length(), 20));
+		
+		ActiveProfileSession newSession = new ActiveProfileSession();
+			
+		if( ActiveSessions.putIfAbsent(sessionName, newSession ) != null) 
+		{
+			// name in use.
+			return OnError("Profile session name '"+sessionName+"' is in use.");
+		}
+
+		Consumer<android.os.ProfilingResult> listener = profilingResult ->
+		{
+			OnSessionComplete(profilingResult);
+		};
+
+		StringBuilder InvocationMessage = new StringBuilder();
+		try {
+			java.util.concurrent.Executor executor = Executors.newSingleThreadExecutor();
+			newSession.CancelSignal = new android.os.CancellationSignal();
+			ProfilingRequest Request = FoundProfiler.BuildProfilerRequest(ArgsMap, newSession.CancelSignal, InvocationMessage);
+
+			androidx.core.os.Profiling.requestProfiling(
+					mActivity.getApplicationContext(),
+					Request,
+					executor,
+					listener
+			);
+		}
+		catch (Exception e)
+		{
+			if(!InvocationMessage.isEmpty())
+			{
+				InvocationMessage.append("\n");
+			}
+			InvocationMessage.append("profile session failed: "+e.toString()+" cause: "+ e.getCause());
+			return OnError(InvocationMessage.toString());
+		}
+		return OnSessionIssued(sessionName, InvocationMessage.toString());
+	}
+
+	@RequiresApi(api = Build.VERSION_CODES.VANILLA_ICE_CREAM)
+	private static synchronized void OnSessionComplete(ProfilingResult profilingResult) 
+	{
+		if (profilingResult.getErrorCode() == ProfilingResult.ERROR_NONE)
+		{
+			Log.debug("profile success.. (retrieve with 'adb pull "+profilingResult.getResultFilePath()+"')");
+		}
+		else
+		{
+			Log.error("profile failed: ("+profilingResult.getErrorCode()+") " +profilingResult.getErrorMessage());
+		}
+
+		String resultSessionName = profilingResult.getTag();
+		String Message = "("+profilingResult.getErrorCode()+") " +profilingResult.getErrorMessage();
+		String FilePath = profilingResult.getResultFilePath();
+
+		nativeOnProfileFinish(resultSessionName, Message, FilePath);
+
+		ActiveSessions.remove(resultSessionName);
+	}
+
+	// returns true if a session was in progress and was cancelled
+	public static synchronized boolean CancelProfile(String sessionName)
+	{
+		ActiveProfileSession foundSession = ActiveSessions.get(sessionName);
+		if( foundSession != null)
+		{
+			Log.debug("profile: ("+sessionName+") was cancelled.");
+			foundSession.CancelSignal.cancel();
+			return true;
+		}
+		Log.error("could not cancel profile: ("+sessionName+") was not found.");
+		return false;
+	}
+
+	public static String[] AndroidThunkJava_IssueProfilerCommand(String RequestedProfilerArgs)
+	{
+		return IssueProfilerCommand(RequestedProfilerArgs);
+	}
+
+	public static boolean AndroidThunkJava_StopProfilerCommand(String profileSessionName)
+	{
+		return CancelProfile(profileSessionName);
+	}
+}
\ No newline at end of file
diff --git a/Engine/Source/Runtime/Core/Private/ProfilingDebugging/Android/AndroidExternalProfiler.cpp b/Engine/Source/Runtime/Core/Private/ProfilingDebugging/Android/AndroidExternalProfiler.cpp
new file mode 100644
--- /dev/null
+++ b/Engine/Source/Runtime/Core/Private/ProfilingDebugging/Android/AndroidExternalProfiler.cpp
@@ -0,0 +1,235 @@
+// Copyright Epic Games, Inc. All Rights Reserved.
+
+#include "CoreTypes.h"
+#include "HAL/PlatformMisc.h"
+#include "HAL/PlatformProcess.h"
+#include "Misc/AssertionMacros.h"
+#include "ProfilingDebugging/ExternalProfiler.h"
+#include "Features/IModularFeatures.h"
+#include "Templates/UniquePtr.h"
+#include "Misc/Paths.h"
+#include "HAL/FileManager.h"
+#include "Android/AndroidProfiler.h"
+#include "HAL/PlatformFileManager.h"
+#include "GenericPlatform/GenericPlatformFile.h"
+#include "ProfilingDebugging/CsvProfiler.h"
+#include "Misc/CommandLine.h"
+#include "Android/AndroidPlatformFile.h"
+
+#if ANDROID_PROFILING_ENABLED && UE_EXTERNAL_PROFILING_ENABLED
+
+namespace PerfettoExternalProfiler
+{
+	FString StartProfileInternal(const FString& ProfilerArgs)
+	{		
+		FString ActiveProfileName = FAndroidProfiler::StartCapture(ProfilerArgs, [](const FAndroidProfiler::FProfileResults& Results)
+			{
+				if (!Results.FilePath.IsEmpty())
+				{
+					const FString FilenameOnly = FPaths::GetCleanFilename(Results.FilePath);
+					const FString DestPath = FPaths::ProfilingDir() / FilenameOnly;
+					UE_CLOG(!Results.Error.IsEmpty(), LogProfilingDebugging, Warning, TEXT("Profile %s completed with log: %s"), *Results.ProfileName, *Results.Error);
+					IPlatformFile& PlatformFile = FPlatformFileManager::Get().GetPlatformFile();
+					// move the file to public path so adb can get to it.
+					bool bFileMoved = PlatformFile.MoveFile(*DestPath, *Results.FilePath);
+					if (!bFileMoved)
+					{
+						// copy the file if the move fails.
+						bFileMoved = PlatformFile.CopyFile(*DestPath, *Results.FilePath);
+						PlatformFile.DeleteFile(*Results.FilePath);
+					}
+					IAndroidPlatformFile& AndroidPlatformFile = IAndroidPlatformFile::GetPlatformPhysical();
+
+					const FString DestAbsolutePath = AndroidPlatformFile.FileRootPath(*DestPath);
+					UE_LOG(LogProfilingDebugging, Log, TEXT("Profile %s complete, retrieve via 'adb pull %s'"), *Results.ProfileName, *DestAbsolutePath);
+				}
+				else
+				{
+					UE_LOG(LogProfilingDebugging, Error, TEXT("Profile %s failed: %s"), *Results.ProfileName, *Results.Error);
+				}
+			});
+		return ActiveProfileName;
+	}
+
+	void CancelProfilerInternal(const FString& ProfileName)
+	{
+		FAndroidProfiler::StopCapture(ProfileName);
+	}
+}
+
+/**
+ * Android perfetto implementation of FExternalProfiler
+ */
+class FPerfettoExternalProfiler : public FExternalProfiler
+{
+	static inline FCriticalSection ProfilerCS;
+
+public:
+
+	/** Constructor */
+	FPerfettoExternalProfiler()
+	{
+		// Register as a modular feature
+		IModularFeatures::Get().RegisterModularFeature( FExternalProfiler::GetFeatureName(), this );
+	}
+
+	/** Destructor */
+	virtual ~FPerfettoExternalProfiler()
+	{
+		IModularFeatures::Get().UnregisterModularFeature( FExternalProfiler::GetFeatureName(), this );
+	}
+
+	/** Gets the name of this profiler as a string.  This is used to allow the user to select this profiler in a system configuration file or on the command-line */
+	virtual const TCHAR* GetProfilerName() const override
+	{
+		return TEXT( "AndroidPerfetto" );
+	}
+
+	/** Pauses profiling. */
+	virtual void ProfilerPauseFunction() override
+	{
+		FScopeLock Lock(&ProfilerCS);
+		FAndroidProfiler::StopCapture( ActiveProfileName );
+	}
+
+	/** Resumes profiling. */
+	virtual void ProfilerResumeFunction() override
+	{
+		FScopeLock Lock(&ProfilerCS);
+		ActiveProfileName = PerfettoExternalProfiler::StartProfileInternal( ProfilerArgs );
+	}
+
+	virtual void Register()
+	{
+		FString CmdLineProfileArgsOverride;
+		if (FParse::Value(FCommandLine::Get(), TEXT("-androidprofilerargs="), CmdLineProfileArgsOverride))
+		{
+			ProfilerArgs = CmdLineProfileArgsOverride;
+		}
+
+		if (FParse::Value(FCommandLine::Get(), TEXT("-androidprofilercsvargs="), CmdLineProfileArgsOverride))
+		{
+			CSVProfilerArgs = CmdLineProfileArgsOverride;
+		}
+
+		FString FrameRange;
+        bool bEnableCSVProfiler = FParse::Param(FCommandLine::Get(), TEXT("-csvExtProfCpu"));
+		if (FParse::Value(FCommandLine::Get(), TEXT("-csvExtProfCpu="), FrameRange))
+		{
+			FString Start, End;
+			if (FrameRange.Split(TEXT(":"), &Start, &End))
+			{
+				FrameStart = FCString::Atoi(*Start);
+				FrameEnd = FCString::Atoi(*End);
+			}
+            bEnableCSVProfiler = true;
+		}
+		
+		if (bEnableCSVProfiler)
+		{
+#if CSV_PROFILER
+			OnCSVStartCaptureHandle = FCsvProfiler::Get()->OnCSVProfileStart().AddRaw(this, &FPerfettoExternalProfiler::StartDelegate);
+			OnCSVStopCaptureHandle = FCsvProfiler::Get()->OnCSVProfileEnd().AddRaw(this, &FPerfettoExternalProfiler::EndDelegate);
+#endif // CSV_PROFILER
+		}
+	}
+
+	void StartDelegate()
+	{
+		if (CSVCaptureMode == ECSVCaptureMode::None)
+		{
+			CSVCaptureMode = ECSVCaptureMode::Start;
+		}
+	}
+
+	void EndDelegate()
+	{
+		FScopeLock Lock(&ProfilerCS);
+		if (CSVCaptureMode != ECSVCaptureMode::None)
+		{
+			CSVCaptureMode = ECSVCaptureMode::Stop;
+		}
+	}
+
+	virtual void FrameSync() override
+	{
+#if CSV_PROFILER
+		switch (CSVCaptureMode)
+		{
+			case ECSVCaptureMode::Start:
+			{
+				if (FCsvProfiler::Get()->GetCaptureFrameNumber() >= FrameStart)
+				{
+					FScopeLock Lock(&ProfilerCS);
+					CSVActiveProfileName = PerfettoExternalProfiler::StartProfileInternal(CSVProfilerArgs);
+					CSVCaptureMode = ECSVCaptureMode::InProgress;
+				}
+			}
+			break;
+			case ECSVCaptureMode::InProgress:
+			{
+				if (FCsvProfiler::Get()->IsCapturing() && FCsvProfiler::Get()->GetCaptureFrameNumber() >= FrameEnd)
+				{
+					CSVCaptureMode = ECSVCaptureMode::Stop;
+				}
+			}
+			break;
+			case ECSVCaptureMode::Stop:
+			{
+				FScopeLock Lock(&ProfilerCS);
+				FAndroidProfiler::StopCapture(CSVActiveProfileName);
+				CSVCaptureMode = ECSVCaptureMode::None;
+			}
+			break;
+			default:
+			break;
+		}
+#endif // CSV_PROFILER
+	}
+
+	/**
+	 */
+	bool Initialize()
+	{
+		ProfilerArgs = TEXT("system duration=100 buffersize=10000");
+		// perfetto requires a duration, we ask for a long duration and expect to cancel within that period.
+		CSVProfilerArgs = TEXT("system duration=1000 buffersize=10000 profilename=csvprofile");
+		return true;
+	}
+
+private:
+	FString ProfilerArgs;
+	FString ActiveProfileName;
+
+	// CSV profile session params.
+	// note the CSV session can run concurrently with ProfilerResumeFunction etc.
+	enum class ECSVCaptureMode { None, Start, InProgress, Stop};
+	ECSVCaptureMode CSVCaptureMode = ECSVCaptureMode::None;
+	FString CSVProfilerArgs;
+	FString CSVActiveProfileName;
+	int32 FrameStart = 0;
+	int32 FrameEnd = 0xFFFFFFFF;
+
+	FDelegateHandle OnCSVStartCaptureHandle;
+	FDelegateHandle OnCSVStopCaptureHandle;
+};
+
+
+namespace PerfettoExternalProfiler
+{
+	struct FAtModuleInit
+	{
+		FAtModuleInit()
+		{
+			static TUniquePtr<FPerfettoExternalProfiler> PerfettoExternal = MakeUnique<FPerfettoExternalProfiler>();
+			if( !PerfettoExternal->Initialize() )
+			{
+				PerfettoExternal.Reset();
+			}
+		}
+	};
+
+	static FAtModuleInit AtModuleInit;
+}
+
+#endif	// ANDROID_PROFILING_ENABLED && UE_EXTERNAL_PROFILING_ENABLED
diff --git a/Engine/Source/Runtime/Core/Public/Android/AndroidPlatformMisc.h b/Engine/Source/Runtime/Core/Public/Android/AndroidPlatformMisc.h
--- a/Engine/Source/Runtime/Core/Public/Android/AndroidPlatformMisc.h
+++ b/Engine/Source/Runtime/Core/Public/Android/AndroidPlatformMisc.h
@@ -21,6 +21,10 @@ class TFunction;
 
 #define ANDROID_HAS_RTSIGNALS PLATFORM_USED_NDK_VERSION_INTEGER >= 21
 
+#ifndef ANDROID_PROFILING_ENABLED
+#define ANDROID_PROFILING_ENABLED !(UE_BUILD_SHIPPING)
+#endif
+
 enum class ECrashContextType;
 
 /**
diff --git a/Engine/Source/Runtime/Core/Public/Android/AndroidProfiler.h b/Engine/Source/Runtime/Core/Public/Android/AndroidProfiler.h
new file mode 100644
--- /dev/null
+++ b/Engine/Source/Runtime/Core/Public/Android/AndroidProfiler.h
@@ -0,0 +1,75 @@
+// Copyright Epic Games, Inc. All Rights Reserved.
+#pragma once
+
+#include "CoreTypes.h"
+#include "CoreMinimal.h"
+
+/*
+* Invokes android's built-in perfetto capture API.
+*
+* Usage:
+* FString ProfileName = FAndroidProfiler::StartCapture(ArgsString, [](const FAndroidProfiler::FProfileResults& Results)
+				{ 
+					// lambda to be called when capture completes.
+					// This can be called from a random thread.
+					// params: 
+					// Results.ProfileName	== the tag/name used when initiating the profiling session.
+					// Results.Error		== Any errors or warnings generated during the session
+					// Results.FilePath		== empty on error. otherwise contains the path to the perfetto capture.
+					//							the file is in the apps internal files folder. it is up to the application
+					//							to move or delete this file.
+				} );
+* 
+* where	ArgsString = "callstack [optional args]"
+*		ArgsString = "heap		[optional args]"
+*		ArgsString = "system	[optional args]"
+*
+* optional args can be any of the following (space delimited):
+* 
+*  duration=x								(in seconds, default 10)
+*  buffersize=x								(in kb, default 1000)
+*  profilename=x							(unique name for the profile session, default 'profile')
+*											(you cannot have concurrent sessions with the same name.) 
+*  samplingfreq=x							(callstack profile only. number of samples per second, default 100)
+*  samplingintervalbytes=x					(heap profile only. alloc granularity of heap samples, default 100)
+*  bufferfillpolicy=[ringbuffer|discard]	(system profile only. default ringbuffer)
+*
+* example: callstack samplingfreq=1000 duration=30
+*
+* to stop an in progress capture:
+*	param ProfileName is the name/tag of the profile to cancel (i.e. as returned from StartCapture)
+* FAndroidProfiler::StopCapture(ProfileName);
+*
+*/
+
+class FAndroidProfiler
+{
+public:
+
+	struct FProfileResults
+	{
+		FString ProfileName;// string used to identify the profile session.
+		FString Error;		// Any errors or warnings produced
+		FString FilePath;	// path to any output produced. empty on error.
+	};
+
+	// Returns an FString used to identify the profile session.
+	static FString StartCapture(const FString& Args, TUniqueFunction<void(const FAndroidProfiler::FProfileResults& Results)> OnFinish);
+
+	// stop an in-progress profile session.
+	static void StopCapture(const FString& ProfileName);
+
+private:
+	static inline FCriticalSection ProfilerCS;
+
+	// Called when a profile session ends. (either on success or on error)
+	static void OnProfileFinish(const FAndroidProfiler::FProfileResults& Results);
+
+	struct FActiveSessions
+	{
+		TUniqueFunction<void(const FProfileResults& Results)> OnFinish;
+	};
+	static inline TMap<FString, FActiveSessions> ActiveSessions;
+
+	friend struct FAndroidProfilerInternal;
+};
\ No newline at end of file
