CIS fix.

diff --git a/Engine/Source/Runtime/Core/Public/Misc/ExpressionParserTypes.inl b/Engine/Source/Runtime/Core/Public/Misc/ExpressionParserTypes.inl
--- a/Engine/Source/Runtime/Core/Public/Misc/ExpressionParserTypes.inl
+++ b/Engine/Source/Runtime/Core/Public/Misc/ExpressionParserTypes.inl
@@ -195,6 +195,7 @@ namespace Impl
 		{
 			// sizeof(OperandType) == 0 is used to create a false value based on a template parameter
 			static_assert(sizeof(OperandType) == 0, "FuncType has an unexpected number of parameters");
+			return nullptr;
 		}
 	}
 
@@ -223,6 +224,7 @@ namespace Impl
 		{
 			// sizeof(LeftOperandType) == 0 is used to create a false value based on a template parameter
 			static_assert(sizeof(LeftOperandType) == 0, "FuncType has an unexpected number of parameters");
+			return nullptr;
 		}
 	}
 
@@ -251,6 +253,7 @@ namespace Impl
 		{
 			// sizeof(OperandType) == 0 is used to create a false value based on a template parameter
 			static_assert(sizeof(OperandType) == 0, "FuncType has an unexpected number of parameters");
+			return nullptr;
 		}
 	}
 }
