Packed FConfigValue to reduce memory footprint

diff --git a/Engine/Source/Runtime/Core/Public/Misc/ConfigCacheIni.h b/Engine/Source/Runtime/Core/Public/Misc/ConfigCacheIni.h
--- a/Engine/Source/Runtime/Core/Public/Misc/ConfigCacheIni.h
+++ b/Engine/Source/Runtime/Core/Public/Misc/ConfigCacheIni.h
@@ -153,7 +153,7 @@ public:
 		: FConfigValue(nullptr, NAME_None)
 	{}
 
-	FConfigValue(const FConfigSection* InSection, FName InValueName, EValueType Type=EValueType::Combined)
+	FConfigValue(const FConfigSection* InSection, FName InValueName, EValueType Type = EValueType::Combined)
 		: ValueType(Type)
 #if UE_WITH_CONFIG_TRACKING
 		, SectionAccess(GetSectionAccess(InSection))
@@ -166,9 +166,9 @@ public:
 		: FConfigValue(nullptr, NAME_None, InValue, Type)
 	{}
 
-	FConfigValue(const FConfigSection* InSection, FName InValueName, const TCHAR* InValue, EValueType Type=EValueType::Combined)
-		: ValueType(Type)
-		, SavedValue(InValue)
+	FConfigValue(const FConfigSection* InSection, FName InValueName, const TCHAR* InValue, EValueType Type = EValueType::Combined)
+		: SavedValue(InValue)
+		, ValueType(Type)
 #if UE_WITH_CONFIG_TRACKING
 		, SectionAccess(GetSectionAccess(InSection))
 		, ValueName(InValueName)
@@ -182,9 +182,9 @@ public:
 		: FConfigValue(nullptr, NAME_None, InValue, Type)
 	{}
 
-	FConfigValue(const FConfigSection* InSection, FName InValueName, const FString& InValue, EValueType Type=EValueType::Combined)
-		: ValueType(Type)
-		, SavedValue(InValue)
+	FConfigValue(const FConfigSection* InSection, FName InValueName, const FString& InValue, EValueType Type = EValueType::Combined)
+		: SavedValue(InValue)
+		, ValueType(Type)
 #if UE_WITH_CONFIG_TRACKING
 		, SectionAccess(GetSectionAccess(InSection))
 		, ValueName(InValueName)
@@ -198,9 +198,9 @@ public:
 		: FConfigValue(nullptr, NAME_None, MoveTemp(InValue), Type)
 	{}
 
-	FConfigValue(const FConfigSection* InSection, FName InValueName, FString&& InValue, EValueType Type=EValueType::Combined)
-		: ValueType(Type)
-		, SavedValue(MoveTemp(InValue))
+	FConfigValue(const FConfigSection* InSection, FName InValueName, FString&& InValue, EValueType Type = EValueType::Combined)
+		: SavedValue(MoveTemp(InValue))
+		, ValueType(Type)
 #if UE_WITH_CONFIG_TRACKING
 		, SectionAccess(GetSectionAccess(InSection))
 		, ValueName(InValueName)
@@ -211,9 +211,9 @@ public:
 	}
 
 	FConfigValue(const FConfigValue& InConfigValue)
-		: ValueType(InConfigValue.ValueType)
-		, SavedValue(InConfigValue.SavedValue)
+		: SavedValue(InConfigValue.SavedValue)
 		, SavedValueHash(InConfigValue.SavedValueHash)
+		, ValueType(InConfigValue.ValueType)
 		, bExpandOnDemand(InConfigValue.bExpandOnDemand)
 #if UE_WITH_CONFIG_TRACKING
 		, SectionAccess(InConfigValue.SectionAccess)
@@ -224,9 +224,9 @@ public:
 	}
 
 	FConfigValue(FConfigValue&& InConfigValue)
-		: ValueType(InConfigValue.ValueType)
-		, SavedValue(MoveTemp(InConfigValue.SavedValue))
+		: SavedValue(MoveTemp(InConfigValue.SavedValue))
 		, SavedValueHash(InConfigValue.SavedValueHash)
+		, ValueType(InConfigValue.ValueType)
 		, bExpandOnDemand(InConfigValue.bExpandOnDemand)
 #if UE_WITH_CONFIG_TRACKING
 		, SectionAccess(MoveTemp(InConfigValue.SectionAccess))
@@ -238,9 +238,9 @@ public:
 
 	FConfigValue& operator=(FConfigValue&& RHS)
 	{
-		ValueType = RHS.ValueType;
 		SavedValue = MoveTemp(RHS.SavedValue);
 		SavedValueHash = RHS.SavedValueHash;
+		ValueType = RHS.ValueType;
 		bExpandOnDemand = RHS.bExpandOnDemand;
 #if UE_WITH_CONFIG_TRACKING
 		SectionAccess = MoveTemp(RHS.SectionAccess);
@@ -252,9 +252,9 @@ public:
 
 	FConfigValue& operator=(const FConfigValue& RHS)
 	{
-		ValueType = RHS.ValueType;
 		SavedValue = RHS.SavedValue;
 		SavedValueHash = RHS.SavedValueHash;
+		ValueType = RHS.ValueType;
 		bExpandOnDemand = RHS.bExpandOnDemand;
 #if UE_WITH_CONFIG_TRACKING
 		SectionAccess = RHS.SectionAccess;
@@ -373,9 +373,6 @@ public:
 	 */
 	CORE_API static FString CollapseValue(const FString& InExpandedValue);
 
-	// Add, subtract, stc
-	EValueType ValueType;
-
 #if CONFIG_CAN_SAVE_COMMENTS
 	FString Comment;
 #endif
@@ -402,6 +399,10 @@ private:
 
 	FString SavedValue;
 	uint32 SavedValueHash;
+public:
+	// Add, subtract, stc
+	EValueType ValueType;
+private:
 	bool bExpandOnDemand = false;
 #if UE_WITH_CONFIG_TRACKING
 	TRefCountPtr<UE::ConfigAccessTracking::FSection> SectionAccess;
