Added indirection when accessing FText::TextData with LiveCoding enabled, to workaround crashes where it leaves TextData null

#rnx

diff --git a/Engine/Source/Runtime/Core/Private/Internationalization/Text.cpp b/Engine/Source/Runtime/Core/Private/Internationalization/Text.cpp
--- a/Engine/Source/Runtime/Core/Private/Internationalization/Text.cpp
+++ b/Engine/Source/Runtime/Core/Private/Internationalization/Text.cpp
@@ -36,6 +36,17 @@ DEFINE_LOG_CATEGORY(LogText);
 
 #define LOCTEXT_NAMESPACE "Core.Text"
 
+// LiveCoding can result in FText with null TextData, so this guards against them crashing while still emitting an ensure
+#ifndef UE_ENABLE_SAFE_GET_TEXT_DATA
+	#define UE_ENABLE_SAFE_GET_TEXT_DATA (WITH_EDITOR && WITH_LIVE_CODING)
+#endif
+
+#if UE_ENABLE_SAFE_GET_TEXT_DATA
+	#define UE_SAFE_GET_TEXT_DATA(Text) (ensure((Text).TextData) ? (Text).TextData : FText::GetEmpty().TextData)
+#else
+	#define UE_SAFE_GET_TEXT_DATA(Text) ((Text).TextData)
+#endif
+
 bool FTextInspector::ShouldGatherForLocalization(const FText& Text)
 {
 	return Text.ShouldGatherForLocalization();
@@ -63,18 +74,18 @@ TOptional<FString> FTextInspector::GetKey(const FText& Text)
 
 FTextId FTextInspector::GetTextId(const FText& Text)
 {
-	return Text.TextData->GetTextHistory().GetTextId();
+	return UE_SAFE_GET_TEXT_DATA(Text)->GetTextHistory().GetTextId();
 }
 
 const FString* FTextInspector::GetSourceString(const FText& Text)
 {
-	return &Text.TextData->GetSourceString();
+	return &UE_SAFE_GET_TEXT_DATA(Text)->GetSourceString();
 }
 
 const FString& FTextInspector::GetDisplayString(const FText& Text)
 {
 	Text.Rebuild();
-	return Text.TextData->GetDisplayString();
+	return UE_SAFE_GET_TEXT_DATA(Text)->GetDisplayString();
 }
 
 bool FTextInspector::GetTableIdAndKey(const FText& Text, FName& OutTableId, FString& OutKey)
@@ -117,7 +128,7 @@ bool FTextInspector::GetHistoricNumericData(const FText& Text, FHistoricTextNume
 
 const void* FTextInspector::GetSharedDataId(const FText& Text)
 {
-	return Text.TextData.GetReference();
+	return UE_SAFE_GET_TEXT_DATA(Text).GetReference();
 }
 
 // These default values have been duplicated to the KismetTextLibrary functions for Blueprints. Please replicate any changes there!
@@ -907,7 +918,7 @@ void FText::SerializeText(FStructuredArchive::FSlot Slot, FText& Value)
 
 			if (bSerializeHistory)
 			{
-				int8 HistoryType = (int8)Value.TextData->GetTextHistory().GetType();
+				int8 HistoryType = (int8)UE_SAFE_GET_TEXT_DATA(Value)->GetTextHistory().GetType();
 				Record << SA_VALUE(TEXT("HistoryType"), HistoryType);
 			}
 			else
@@ -919,7 +930,7 @@ void FText::SerializeText(FStructuredArchive::FSlot Slot, FText& Value)
 				Record << SA_VALUE(TEXT("bHasCultureInvariantString"), bHasCultureInvariantString);
 				if (bHasCultureInvariantString)
 				{
-					FString CultureInvariantString = Value.TextData->GetSourceString();
+					FString CultureInvariantString = UE_SAFE_GET_TEXT_DATA(Value)->GetSourceString();
 					Record << SA_VALUE(TEXT("CultureInvariantString"), CultureInvariantString);
 				}
 			}
@@ -1019,7 +1030,7 @@ void FText::SerializeText(FStructuredArchive::FSlot Slot, FText& Value)
 
 		if (bSerializeHistory)
 		{
-			Value.TextData->GetMutableTextHistory().Serialize(Record);
+			UE_SAFE_GET_TEXT_DATA(Value)->GetMutableTextHistory().Serialize(Record);
 		}
 	}
 
@@ -1043,7 +1054,7 @@ void FText::SerializeText(FStructuredArchive::FSlot Slot, FText& Value)
 #if WITH_EDITORONLY_DATA
 FText FText::ChangeKey( const FTextKey& Namespace, const FTextKey& Key, const FText& Text )
 {
-	return FText(FString(*Text.TextData->GetSourceString()), Namespace, Key, 0);
+	return FText(FString(*UE_SAFE_GET_TEXT_DATA(Text)->GetSourceString()), Namespace, Key, 0);
 }
 #endif
 
@@ -1160,12 +1171,12 @@ bool FText::FindTextInLiveTable_Advanced(const FTextKey& Namespace, const FTextK
 const FString& FText::ToString() const
 {
 	Rebuild();
-	return TextData->GetDisplayString();
+	return UE_SAFE_GET_TEXT_DATA(*this)->GetDisplayString();
 }
 
 FString FText::BuildSourceString() const
 {
-	return TextData->GetTextHistory().BuildInvariantDisplayString();
+	return UE_SAFE_GET_TEXT_DATA(*this)->GetTextHistory().BuildInvariantDisplayString();
 }
 
 bool FText::IsNumeric() const
@@ -1195,7 +1206,7 @@ bool FText::EqualToCaseIgnored(const FText& Other) const
 
 void FText::Rebuild() const
 {
-	TextData->GetMutableTextHistory().UpdateDisplayStringIfOutOfDate();
+	UE_SAFE_GET_TEXT_DATA(*this)->GetMutableTextHistory().UpdateDisplayStringIfOutOfDate();
 }
 
 bool FText::IsTransient() const
@@ -1215,7 +1226,7 @@ bool FText::IsInitializedFromString() const
 
 bool FText::IsFromStringTable() const
 {
-	return TextData->GetTextHistory().GetType() == ETextHistoryType::StringTableEntry;
+	return UE_SAFE_GET_TEXT_DATA(*this)->GetTextHistory().GetType() == ETextHistoryType::StringTableEntry;
 }
 
 bool FText::ShouldGatherForLocalization() const
@@ -1225,7 +1236,7 @@ bool FText::ShouldGatherForLocalization() const
 		return false;
 	}
 
-	const FString& SourceString = TextData->GetSourceString();
+	const FString& SourceString = UE_SAFE_GET_TEXT_DATA(*this)->GetSourceString();
 
 	auto IsAllWhitespace = [](const FString& String) -> bool
 	{
@@ -1249,12 +1260,12 @@ void FText::PreloadStringTable( FName InTableId, const FTextKey& InKey )
 
 void FText::GetHistoricFormatData(TArray<FHistoricTextFormatData>& OutHistoricFormatData) const
 {
-	TextData->GetTextHistory().GetHistoricFormatData(*this, OutHistoricFormatData);
+	UE_SAFE_GET_TEXT_DATA(*this)->GetTextHistory().GetHistoricFormatData(*this, OutHistoricFormatData);
 }
 
 bool FText::GetHistoricNumericData(FHistoricTextNumericData& OutHistoricNumericData) const
 {
-	return TextData->GetTextHistory().GetHistoricNumericData(*this, OutHistoricNumericData);
+	return UE_SAFE_GET_TEXT_DATA(*this)->GetTextHistory().GetHistoricNumericData(*this, OutHistoricNumericData);
 }
 
 bool FText::IdenticalTo( const FText& Other, const ETextIdenticalModeFlags CompareModeFlags ) const
@@ -1272,8 +1283,8 @@ bool FText::IdenticalTo( const FText& Other, const ETextIdenticalModeFlags Compa
 	// This is fast as it skips a lexical compare, however it can also return false for two instances that have identical strings, but in different pointers.
 	// For instance, this method will return false for two FText objects created from FText::FromString("Wooble") as they each have unique (or null), non-shared instances.
 	{
-		FTextConstDisplayStringPtr DisplayStringPtr = TextData->GetLocalizedString();
-		FTextConstDisplayStringPtr OtherDisplayStringPtr = Other.TextData->GetLocalizedString();
+		FTextConstDisplayStringPtr DisplayStringPtr = UE_SAFE_GET_TEXT_DATA(*this)->GetLocalizedString();
+		FTextConstDisplayStringPtr OtherDisplayStringPtr = UE_SAFE_GET_TEXT_DATA(Other)->GetLocalizedString();
 		if (DisplayStringPtr && OtherDisplayStringPtr && DisplayStringPtr == OtherDisplayStringPtr)
 		{
 			return true;
@@ -1282,8 +1293,8 @@ bool FText::IdenticalTo( const FText& Other, const ETextIdenticalModeFlags Compa
 
 	if (EnumHasAnyFlags(CompareModeFlags, ETextIdenticalModeFlags::DeepCompare))
 	{
-		const FTextHistory& ThisTextHistory = TextData->GetTextHistory();
-		const FTextHistory& OtherTextHistory = Other.TextData->GetTextHistory();
+		const FTextHistory& ThisTextHistory = UE_SAFE_GET_TEXT_DATA(*this)->GetTextHistory();
+		const FTextHistory& OtherTextHistory = UE_SAFE_GET_TEXT_DATA(Other)->GetTextHistory();
 		if (ThisTextHistory.GetType() == OtherTextHistory.GetType() && ThisTextHistory.IdenticalTo(OtherTextHistory, CompareModeFlags))
 		{
 			return true;
@@ -1660,7 +1671,7 @@ FTextSnapshot::FTextSnapshot()
 
 FTextSnapshot::FTextSnapshot(const FText& InText)
 	: TextDataPtr(InText.TextData)
-	, LocalizedStringPtr(InText.TextData->GetLocalizedString())
+	, LocalizedStringPtr(UE_SAFE_GET_TEXT_DATA(InText)->GetLocalizedString())
 	, GlobalHistoryRevision(GetGlobalHistoryRevisionForText(InText))
 	, LocalHistoryRevision(GetLocalHistoryRevisionForText(InText))
 	, Flags(InText.Flags)
@@ -1675,7 +1686,7 @@ bool FTextSnapshot::IdenticalTo(const FText& InText) const
 
 	return TextDataPtr
 		&& TextDataPtr == InText.TextData
-		&& LocalizedStringPtr == InText.TextData->GetLocalizedString()
+		&& LocalizedStringPtr == UE_SAFE_GET_TEXT_DATA(InText)->GetLocalizedString()
 		&& GlobalHistoryRevision == GetGlobalHistoryRevisionForText(InText)
 		&& LocalHistoryRevision == GetLocalHistoryRevisionForText(InText)
 		&& Flags == InText.Flags;
@@ -1692,17 +1703,17 @@ bool FTextSnapshot::IsDisplayStringEqualTo(const FText& InText) const
 	return TextDataPtr
 		&& GlobalHistoryRevision == GetGlobalHistoryRevisionForText(InText)
 		&& LocalHistoryRevision == GetLocalHistoryRevisionForText(InText)
-		&& TextDataPtr->GetDisplayString().Equals(InText.TextData->GetDisplayString(), ESearchCase::CaseSensitive);
+		&& TextDataPtr->GetDisplayString().Equals(UE_SAFE_GET_TEXT_DATA(InText)->GetDisplayString(), ESearchCase::CaseSensitive);
 }
 
 uint16 FTextSnapshot::GetGlobalHistoryRevisionForText(const FText& InText)
 {
-	return (InText.IsEmpty() || InText.IsCultureInvariant()) ? 0 : InText.TextData->GetGlobalHistoryRevision();
+	return (InText.IsEmpty() || InText.IsCultureInvariant()) ? 0 : UE_SAFE_GET_TEXT_DATA(InText)->GetGlobalHistoryRevision();
 }
 
 uint16 FTextSnapshot::GetLocalHistoryRevisionForText(const FText& InText)
 {
-	return (InText.IsEmpty() || InText.IsCultureInvariant()) ? 0 : InText.TextData->GetLocalHistoryRevision();
+	return (InText.IsEmpty() || InText.IsCultureInvariant()) ? 0 : UE_SAFE_GET_TEXT_DATA(InText)->GetLocalHistoryRevision();
 }
 
 bool TextBiDi::IsControlCharacter(const TCHAR InChar)
@@ -1882,7 +1893,7 @@ void FTextStringHelper::WriteToBuffer(FString& Buffer, const FText& Value, const
 #undef LOC_DEFINE_REGION
 	}
 	// Is this text that should be written via its text history?
-	else if (Value.TextData->GetTextHistory().WriteToBuffer(Buffer, bStripPackageNamespace))
+	else if (UE_SAFE_GET_TEXT_DATA(Value)->GetTextHistory().WriteToBuffer(Buffer, bStripPackageNamespace))
 	{
 	}
 	// This isn't special text, so write as a raw string (potentially quoted)
@@ -2125,4 +2136,6 @@ const TCHAR* LexToString(ERoundingMode InValue)
 	}
 }
 
+#undef UE_SAFE_GET_TEXT_DATA
+
 #undef LOCTEXT_NAMESPACE
