Fixed use after happening while iterating in TDirectoryTree<ValueType>::FIteratorInternal::TraverseToValid()

#rnx

diff --git a/Engine/Source/Runtime/Core/Public/Containers/DirectoryTree.h b/Engine/Source/Runtime/Core/Public/Containers/DirectoryTree.h
--- a/Engine/Source/Runtime/Core/Public/Containers/DirectoryTree.h
+++ b/Engine/Source/Runtime/Core/Public/Containers/DirectoryTree.h
@@ -1768,7 +1768,7 @@ void TDirectoryTree<ValueType>::FIteratorInternal::TraverseToValid()
 				else if (!Data->bChildNodeInitialized)
 				{
 					Data->bChildNodeInitialized = true;
-					Stack.Emplace(Data->Node->ChildNodes[Data->ChildIndex], Data->ClosestValue);
+					Stack.Add(FStackData(Data->Node->ChildNodes[Data->ChildIndex], Data->ClosestValue));
 					break; // Return to the outer loop to traverse into the node we just pushed
 				}
 				else
