[Core] Simplify RWLock to Mutex.

This RWLock was only ever locked for writing, and doesn't need to support recursion. A mutex is sufficient and much lighter-weight.

diff --git a/Engine/Source/Runtime/Core/Private/HAL/FileManagerGeneric.cpp b/Engine/Source/Runtime/Core/Private/HAL/FileManagerGeneric.cpp
--- a/Engine/Source/Runtime/Core/Private/HAL/FileManagerGeneric.cpp
+++ b/Engine/Source/Runtime/Core/Private/HAL/FileManagerGeneric.cpp
@@ -9,15 +9,15 @@
 
 =============================================================================*/
 
-// for compression
 #include "HAL/FileManagerGeneric.h"
 #include "Logging/LogMacros.h"
+#include "Async/Mutex.h"
 #include "Misc/AccessDetection.h"
 #include "Misc/Parse.h"
 #include "Misc/CommandLine.h"
 #include "Misc/Paths.h"
 #include "Misc/SecureHash.h"
-#include "Misc/ScopeRWLock.h"
+#include "Misc/ScopeLock.h"
 #include "Templates/UniquePtr.h"
 #include <time.h>
 
@@ -413,11 +413,12 @@ namespace FileManagerGenericImpl
 	{
 	public:
 		TArray<FString>& Result;
-		FRWLock ResultLock;
 		FString WildCard;
 		bool bFiles;
 		bool bDirectories;
 		bool bStoreFullPath;
+		UE::FMutex ResultLock;
+
 		FFileMatch(TArray<FString>& InResult, const FString& InWildCard, bool bInFiles, bool bInDirectories, bool bInStoreFullPath = false)
 			: IPlatformFile::FDirectoryVisitor(EDirectoryVisitorFlags::ThreadSafe)
 			, Result(InResult)
@@ -442,7 +443,7 @@ namespace FileManagerGenericImpl
 					TEXT("PlatformFile.IterateDirectory needs to call ShouldVisitLeafFilename before calling Visit.")))
 				{
 					FString FullPath = bStoreFullPath ? FString(FilenameOrDirectory) : MoveTemp(Filename);
-					FWriteScopeLock ScopeLock(ResultLock);
+					UE::TScopeLock ScopeLock(ResultLock);
 					Result.Add(MoveTemp(FullPath));
 				}
 			}
