Rework TFunction to improve performance and reduce memory footprint.

This is done with two optimizations:

- Making `InlineAllocation` be the first thing in `FFunctionStorage` means codegen is optimized for the small `TFunction` case.
- Make `NUM_TFUNCTION_INLINE_BYTES` 24 bytes and dropping the use of `TAlignedBytes`. `TAlignedBytes` is bad if you have a size that isn't a multiple of alignment, as it will pad the struct to the alignment boundary. This allows us to reduce the size of `TFunction` from 48 bytes to 32 bytes, resulting in allocation wins.

diff --git a/Engine/Source/Runtime/Core/Public/Misc/CoreDefines.h b/Engine/Source/Runtime/Core/Public/Misc/CoreDefines.h
--- a/Engine/Source/Runtime/Core/Public/Misc/CoreDefines.h
+++ b/Engine/Source/Runtime/Core/Public/Misc/CoreDefines.h
@@ -58,7 +58,7 @@ Can be overridden by setting to 1 or 0 in the project's .Target.cs files.
 	#error USE_SMALL_TFUNCTIONS has been removed - please use NUM_TFUNCTION_INLINE_BYTES instead
 #else
 	#ifndef NUM_TFUNCTION_INLINE_BYTES
-		#define NUM_TFUNCTION_INLINE_BYTES 32
+		#define NUM_TFUNCTION_INLINE_BYTES 24
 	#endif
 #endif
 
diff --git a/Engine/Source/Runtime/Core/Public/Templates/Function.h b/Engine/Source/Runtime/Core/Public/Templates/Function.h
--- a/Engine/Source/Runtime/Core/Public/Templates/Function.h
+++ b/Engine/Source/Runtime/Core/Public/Templates/Function.h
@@ -195,11 +195,12 @@ namespace UE::Core::Private::Function
 			Owned->Destroy();
 		}
 
-		void* HeapAllocation;
 		#if TFUNCTION_USES_INLINE_STORAGE
 			// Inline storage for an owned object
-			TAlignedBytes<TFUNCTION_INLINE_SIZE, TFUNCTION_INLINE_ALIGNMENT> InlineAllocation;
+			alignas(TFUNCTION_INLINE_ALIGNMENT) uint8 InlineAllocation[TFUNCTION_INLINE_SIZE];
 		#endif
+
+		void* HeapAllocation;
 	};
 
 	template <bool bUnique>
