PR #11278: Fixed memory corruption in iOS MallocBinned when reallocating a pointer from nano region in a case it it went out of nano and needs to fallback to UE allocations, but new size is smaller than old size. This will result in memcpy more bytes and overwrite some neighbouring data or cause a page fault

diff --git a/Engine/Source/Runtime/Core/Private/HAL/MallocBinned.cpp b/Engine/Source/Runtime/Core/Private/HAL/MallocBinned.cpp
--- a/Engine/Source/Runtime/Core/Private/HAL/MallocBinned.cpp
+++ b/Engine/Source/Runtime/Core/Private/HAL/MallocBinned.cpp
@@ -1448,7 +1448,7 @@ void* FMallocBinned::Realloc( void* Ptr, SIZE_T NewSize, uint32 Alignment )
                 // Fall back to UE's allocator
                 Ptr = NewPtr;
                 NewPtr = Malloc(NewSize, Alignment);
-                FMemory::Memcpy(NewPtr, Ptr, OldSize);
+                FMemory::Memcpy(NewPtr, Ptr, FMath::Min(OldSize, NewSize));
                 Private::SmallOSFree(*this, Ptr, OldSize);
             }
 		}
