Formatting changes.

#rnx

diff --git a/Engine/Source/Runtime/Core/Private/Misc/StringFormatter.cpp b/Engine/Source/Runtime/Core/Private/Misc/StringFormatter.cpp
--- a/Engine/Source/Runtime/Core/Private/Misc/StringFormatter.cpp
+++ b/Engine/Source/Runtime/Core/Private/Misc/StringFormatter.cpp
@@ -75,7 +75,12 @@ void AppendToString(const FStringFormatArg& Arg, FString& StringToAppendTo)
 /** Token representing a literal string inside the string */
 struct FStringLiteral
 {
-	FStringLiteral(const FStringToken& InString) : String(InString), Len(UE_PTRDIFF_TO_INT32(InString.GetTokenEndPos() - InString.GetTokenStartPos())) {}
+	explicit FStringLiteral(const FStringToken& InString)
+		: String(InString)
+		, Len(UE_PTRDIFF_TO_INT32(InString.GetTokenEndPos() - InString.GetTokenStartPos()))
+	{
+	}
+
 	/** The string literal token */
 	FStringToken String;
 	/** Cached length of the string */
@@ -85,7 +90,12 @@ struct FStringLiteral
 /** Token representing a user-defined token, such as {Argument} */
 struct FFormatSpecifier
 {
-	FFormatSpecifier(const FStringToken& InIdentifier, const FStringToken& InEntireToken) : Identifier(InIdentifier), EntireToken(InEntireToken), Len(UE_PTRDIFF_TO_INT32(Identifier.GetTokenEndPos() - Identifier.GetTokenStartPos())) {}
+	explicit FFormatSpecifier(const FStringToken& InIdentifier, const FStringToken& InEntireToken)
+		: Identifier(InIdentifier)
+		, EntireToken(InEntireToken)
+		, Len(UE_PTRDIFF_TO_INT32(Identifier.GetTokenEndPos() - Identifier.GetTokenStartPos()))
+	{
+	}
 
 	/** The identifier part of the token */
 	FStringToken Identifier;
@@ -98,7 +108,11 @@ struct FFormatSpecifier
 /** Token representing a user-defined index token, such as {0} */
 struct FIndexSpecifier
 {
-	FIndexSpecifier(int32 InIndex, const FStringToken& InEntireToken) : Index(InIndex), EntireToken(InEntireToken) {}
+	explicit FIndexSpecifier(int32 InIndex, const FStringToken& InEntireToken)
+		: Index(InIndex)
+		, EntireToken(InEntireToken)
+	{
+	}
 
 	/** The index of the parsed token */
 	int32 Index;
@@ -109,7 +123,10 @@ struct FIndexSpecifier
 /** Token representing an escaped character */
 struct FEscapedCharacter
 {
-	FEscapedCharacter(TCHAR InChar) : Character(InChar) {}
+	explicit FEscapedCharacter(TCHAR InChar)
+		: Character(InChar)
+	{
+	}
 
 	/** The character that was escaped */
 	TCHAR Character;
@@ -134,7 +151,7 @@ TOptional<FExpressionError> ParseIndex(FExpressionTokenConsumer& Consumer, bool
 	TOptional<FStringToken> OpeningChar = Stream.ParseSymbol(TEXT('{'));
 	if (!OpeningChar.IsSet())
 	{
-		return TOptional<FExpressionError>();
+		return {};
 	}
 
 	FStringToken& EntireToken = OpeningChar.GetValue();
@@ -167,7 +184,7 @@ TOptional<FExpressionError> ParseIndex(FExpressionTokenConsumer& Consumer, bool
 		}
 		else
 		{
-			return TOptional<FExpressionError>();
+			return {};
 		}
 	}
 
@@ -183,13 +200,13 @@ TOptional<FExpressionError> ParseIndex(FExpressionTokenConsumer& Consumer, bool
 		}
 		else
 		{
-			return TOptional<FExpressionError>();
+			return {};
 		}
 	}
 
 	// Add the token to the consumer. This moves the read position in the stream to the end of the token.
 	Consumer.Add(EntireToken, FIndexSpecifier(Index.GetValue(), EntireToken));
-	return TOptional<FExpressionError>();
+	return {};
 }
 
 TOptional<FExpressionError> ParseSpecifier(FExpressionTokenConsumer& Consumer, bool bEmitErrors)
@@ -199,7 +216,7 @@ TOptional<FExpressionError> ParseSpecifier(FExpressionTokenConsumer& Consumer, b
 	TOptional<FStringToken> OpeningChar = Stream.ParseSymbol(TEXT('{'));
 	if (!OpeningChar.IsSet())
 	{
-		return TOptional<FExpressionError>();
+		return {};
 	}
 
 	FStringToken& EntireToken = OpeningChar.GetValue();
@@ -234,7 +251,7 @@ TOptional<FExpressionError> ParseSpecifier(FExpressionTokenConsumer& Consumer, b
 		}
 		else
 		{
-			return TOptional<FExpressionError>();
+			return {};
 		}
 	}
 
@@ -250,13 +267,13 @@ TOptional<FExpressionError> ParseSpecifier(FExpressionTokenConsumer& Consumer, b
 		}
 		else
 		{
-			return TOptional<FExpressionError>();
+			return {};
 		}
 	}
 
 	// Add the token to the consumer. This moves the read position in the stream to the end of the token.
 	Consumer.Add(EntireToken, FFormatSpecifier(Identifier.GetValue(), EntireToken));
-	return TOptional<FExpressionError>();
+	return {};
 }
 
 static const TCHAR EscapeChar = TEXT('`');
@@ -269,14 +286,14 @@ TOptional<FExpressionError> ParseEscapedChar(FExpressionTokenConsumer& Consumer,
 	TOptional<FStringToken> Token = Consumer.GetStream().ParseSymbol(EscapeChar);
 	if (!Token.IsSet())
 	{
-		return TOptional<FExpressionError>();
+		return {};
 	}
 
 	// Accumulate the next character into the token
 	TOptional<FStringToken> EscapedChar = Consumer.GetStream().ParseSymbol(&Token.GetValue());
 	if (!EscapedChar.IsSet())
 	{
-		return TOptional<FExpressionError>();
+		return {};
 	}
 
 	// Check for a valid escape character
@@ -285,7 +302,7 @@ TOptional<FExpressionError> ParseEscapedChar(FExpressionTokenConsumer& Consumer,
 	{
 		// Add the token to the consumer. This moves the read position in the stream to the end of the token.
 		Consumer.Add(Token.GetValue(), FEscapedCharacter(Character));
-		return TOptional<FExpressionError>();
+		return {};
 	}
 	else if (bEmitErrors)
 	{
@@ -297,7 +314,7 @@ TOptional<FExpressionError> ParseEscapedChar(FExpressionTokenConsumer& Consumer,
 	}
 	else
 	{
-		return TOptional<FExpressionError>();
+		return {};
 	}
 }
 
@@ -328,7 +345,7 @@ TOptional<FExpressionError> ParseLiteral(FExpressionTokenConsumer& Consumer, boo
 		// Add the token to the consumer. This moves the read position in the stream to the end of the token.
 		Consumer.Add(Token.GetValue(), FStringLiteral(Token.GetValue()));
 	}
-	return TOptional<FExpressionError>();
+	return {};
 }
 
 
