XSX - Boot Test - Assertion failed: !QueuedPool

diff --git a/Engine/Source/Runtime/Core/Public/Async/AsyncWork.h b/Engine/Source/Runtime/Core/Public/Async/AsyncWork.h
--- a/Engine/Source/Runtime/Core/Public/Async/AsyncWork.h
+++ b/Engine/Source/Runtime/Core/Public/Async/AsyncWork.h
@@ -363,6 +363,7 @@ protected:
 	void CheckIdle() const
 	{
 		check(WorkNotFinishedCounter.GetValue() == 0);
+		if (QueuedPool) FDebug::DumpStackTraceToLog(ELogVerbosity::Warning);
 		check(!QueuedPool);
 	}
 
