[AutoRTFM] Defer heartbeat thread suspend/resume until transaction commit time.

diff --git a/Engine/Source/Runtime/Core/Private/HAL/ThreadHeartBeat.cpp b/Engine/Source/Runtime/Core/Private/HAL/ThreadHeartBeat.cpp
--- a/Engine/Source/Runtime/Core/Private/HAL/ThreadHeartBeat.cpp
+++ b/Engine/Source/Runtime/Core/Private/HAL/ThreadHeartBeat.cpp
@@ -19,6 +19,7 @@
 #include "ProfilingDebugging/CsvProfiler.h"
 #include "Misc/App.h"
 #include "Misc/Fork.h"
+#include "AutoRTFM/AutoRTFM.h"
 
 // When enabled, the heart beat thread will call abort() when a hang
 // is detected, rather than performing stack back-traces and logging.
@@ -1146,29 +1147,41 @@ void FGameThreadHitchHeartBeatThreaded::SuspendHeartBeat()
 	if (!IsInGameThread())
 		return;
 
-	FPlatformAtomics::InterlockedIncrement(&SuspendedCount);
-	UE_LOG(LogCore, Log, TEXT("HitchHeartBeat Suspend called (count %d) - State: %s"), SuspendedCount, SuspendedCount==0 ? TEXT("Running") : TEXT("Suspended") );
+	UE_AUTORTFM_OPEN
+	{
+		FPlatformAtomics::InterlockedIncrement(&SuspendedCount);
+		UE_LOG(LogCore, Log, TEXT("HitchHeartBeat Suspend called (count %d) - State: %s"), SuspendedCount, SuspendedCount == 0 ? TEXT("Running") : TEXT("Suspended"));
+	};
+	UE_AUTORTFM_ONABORT(this)
+	{
+		ResumeHeartBeat();
+	};
 #endif
 }
+
 void FGameThreadHitchHeartBeatThreaded::ResumeHeartBeat()
 {
 #if USE_HITCH_DETECTION
 	if (!IsInGameThread())
 		return;
 
-	// Temporary workaround for suspend/resume issue
-	//check(SuspendedCount > 0);
-	if (SuspendedCount == 0)
+	// Defer heartbeat-resume to the end of our transaction, if one is running.
+	UE_AUTORTFM_ONCOMMIT(this)
 	{
-		UE_LOG(LogCore, Warning, TEXT("HitchHeartBeat Resume called when SuspendedCount was already 0! Ignoring"));
-		return;
-	}
+		// Temporary workaround for suspend/resume issue
+		//check(SuspendedCount > 0);
+		if (SuspendedCount == 0)
+		{
+			UE_LOG(LogCore, Warning, TEXT("HitchHeartBeat Resume called when SuspendedCount was already 0! Ignoring"));
+			return;
+		}
 
-	if (FPlatformAtomics::InterlockedDecrement(&SuspendedCount) == 0)
-	{
-		FrameStart(true);
-	}
-	UE_LOG(LogCore, Log, TEXT("HitchHeartBeat Resume called (count %d) - State: %s"), SuspendedCount, SuspendedCount == 0 ? TEXT("Running") : TEXT("Suspended"));
+		if (FPlatformAtomics::InterlockedDecrement(&SuspendedCount) == 0)
+		{
+			FrameStart(true);
+		}
+		UE_LOG(LogCore, Log, TEXT("HitchHeartBeat Resume called (count %d) - State: %s"), SuspendedCount, SuspendedCount == 0 ? TEXT("Running") : TEXT("Suspended"));
+	};
 #endif
 }
 
