Fix crash on startup in some targets by reading command line too early from test code.

diff --git a/Engine/Source/Runtime/Core/Private/Windows/WindowsPlatformMisc.cpp b/Engine/Source/Runtime/Core/Private/Windows/WindowsPlatformMisc.cpp
--- a/Engine/Source/Runtime/Core/Private/Windows/WindowsPlatformMisc.cpp
+++ b/Engine/Source/Runtime/Core/Private/Windows/WindowsPlatformMisc.cpp
@@ -1096,11 +1096,6 @@ int32 FWindowsPlatformMisc::GetMaxPathLength()
 
 		FLongPathsEnabled()
 		{
-			if (FParse::Param( FCommandLine::Get(), TEXT("NoLongPaths")))
-			{
-				bValue = false;
-				return;
-			}
 			HMODULE Handle = GetModuleHandle(TEXT("ntdll.dll"));
 			if (Handle == NULL)
 			{
