[Insights] Fixed memory tracing to enable tracing of tags and mem scope with "memalloc" channel (and not with "memtag" channel which is used to enable tracing from LLM).

diff --git a/Engine/Source/Runtime/Core/Private/ProfilingDebugging/MemoryTrace.cpp b/Engine/Source/Runtime/Core/Private/ProfilingDebugging/MemoryTrace.cpp
--- a/Engine/Source/Runtime/Core/Private/ProfilingDebugging/MemoryTrace.cpp
+++ b/Engine/Source/Runtime/Core/Private/ProfilingDebugging/MemoryTrace.cpp
@@ -329,6 +329,8 @@ static EMemoryTraceInit MemoryTrace_ShouldEnable(int32 ArgC, const ArgCharType*
 					if (Channel.Equals("memalloc", ESearchCase::IgnoreCase))
 					{
 						Mode |= EMemoryTraceInit::AllocEvents;
+						// The tags and memory scopes are also emitted on the "memalloc" channel.
+						Mode |= EMemoryTraceInit::Tags;
 					}
 					else if (Channel.Equals("callstack", ESearchCase::IgnoreCase))
 					{
@@ -336,7 +338,7 @@ static EMemoryTraceInit MemoryTrace_ShouldEnable(int32 ArgC, const ArgCharType*
 					}
 					else if (Channel.Equals("memtag", ESearchCase::IgnoreCase))
 					{
-						Mode |= EMemoryTraceInit::Tags;
+						// The "memtag" channel is only used to enable LLM tracing.
 					}
 					else if (Channel.Equals("probing", ESearchCase::IgnoreCase))
 					{
