Fixed typo in TIn(s)trusive list

#rnx

diff --git a/Engine/Source/Runtime/Core/Internal/IO/IoContainers.h b/Engine/Source/Runtime/Core/Internal/IO/IoContainers.h
--- a/Engine/Source/Runtime/Core/Internal/IO/IoContainers.h
+++ b/Engine/Source/Runtime/Core/Internal/IO/IoContainers.h
@@ -8,19 +8,19 @@ namespace UE::Private
 {
 
 template <typename TypeTraits>
-class TInstrusiveListIterator
+class TIntrusiveListIterator
 {
 public:
 	using ElementType = typename TypeTraits::ElementType;
 
-	explicit TInstrusiveListIterator(ElementType* InElement)
+	explicit TIntrusiveListIterator(ElementType* InElement)
 		: Element(InElement)
 	{ }
 
 	ElementType&	operator*() const	{ check(Element); return *Element; }
 	explicit		operator bool()		{ return Element != nullptr; }
 	void			operator++()		{ check(Element); Element = TypeTraits::GetNext(Element); }
-	bool			operator!=(const TInstrusiveListIterator& Other) const { return Element != Other.Element; }
+	bool			operator!=(const TIntrusiveListIterator& Other) const { return Element != Other.Element; }
 
 private:
 	ElementType* Element;
@@ -29,24 +29,24 @@ private:
 } // namespace UE::Private
 
 template <typename TypeTraits>
-class TInstrusiveList
+class TIntrusiveList
 {
 public:
 	using ElementType		= typename TypeTraits::ElementType;
-	using FIterator			= UE::Private::TInstrusiveListIterator<TypeTraits>;
-	using FConstIterator	= UE::Private::TInstrusiveListIterator<const TypeTraits>;
+	using FIterator			= UE::Private::TIntrusiveListIterator<TypeTraits>;
+	using FConstIterator	= UE::Private::TIntrusiveListIterator<const TypeTraits>;
 
-	TInstrusiveList() = default;
-	TInstrusiveList(const TInstrusiveList&) = delete;
-	TInstrusiveList(TInstrusiveList&& Other)
+	TIntrusiveList() = default;
+	TIntrusiveList(const TIntrusiveList&) = delete;
+	TIntrusiveList(TIntrusiveList&& Other)
 		: Head(Other.Head)
 		, Tail(Other.Tail)
 	{
 		Other.Head = Other.Tail = nullptr;
 	}
 
-	TInstrusiveList& operator=(const TInstrusiveList&) = delete;
-	TInstrusiveList& operator=(TInstrusiveList&& Other)
+	TIntrusiveList& operator=(const TIntrusiveList&) = delete;
+	TIntrusiveList& operator=(TIntrusiveList&& Other)
 	{
 		Head = Other.Head;
 		Tail = Other.Tail;
@@ -90,7 +90,7 @@ public:
 		}
 	}
 
-	void AddTail(TInstrusiveList&& Other)
+	void AddTail(TIntrusiveList&& Other)
 	{
 		if (!Other.IsEmpty())
 		{
diff --git a/Engine/Source/Runtime/Core/Public/IO/IoDispatcherBackend.h b/Engine/Source/Runtime/Core/Public/IO/IoDispatcherBackend.h
--- a/Engine/Source/Runtime/Core/Public/IO/IoDispatcherBackend.h
+++ b/Engine/Source/Runtime/Core/Public/IO/IoDispatcherBackend.h
@@ -136,7 +136,7 @@ struct FIoRequestListTraits
 
 } // namespace UE::Private
 
-using FIoRequestList = TInstrusiveList<UE::Private::FIoRequestListTraits>;
+using FIoRequestList = TIntrusiveList<UE::Private::FIoRequestListTraits>;
 
 DECLARE_DELEGATE(FWakeUpIoDispatcherThreadDelegate);
 
