Added ANSI and UTF-8 overloads to FText::FromString and FText::AsCultureInvariant.

diff --git a/Engine/Source/Runtime/Core/Private/Internationalization/Text.cpp b/Engine/Source/Runtime/Core/Private/Internationalization/Text.cpp
--- a/Engine/Source/Runtime/Core/Private/Internationalization/Text.cpp
+++ b/Engine/Source/Runtime/Core/Private/Internationalization/Text.cpp
@@ -2,6 +2,8 @@
 
 #include "Internationalization/Text.h"
 #include "Algo/Transform.h"
+#include "Containers/AnsiString.h"
+#include "Containers/Utf8String.h"
 #include "Misc/Parse.h"
 #include "UObject/ObjectVersion.h"
 #include "UObject/DebugSerializationFlags.h"
@@ -1068,6 +1070,26 @@ FText FText::FromName( const FName& Val)
 	return FText::FromString(Val.ToString());
 }
 
+FText FText::FromString( const ANSICHAR* String )
+{
+	return FText::FromString(FString(String));
+}
+
+FText FText::FromString( const UTF8CHAR* String )
+{
+	return FText::FromString(FString(String));
+}
+
+FText FText::FromString( const WIDECHAR* String )
+{
+	return FText::FromString(FString(String));
+}
+
+FText FText::FromString( const FAnsiString& String )
+{
+	return FText::FromString(FString::ConstructFromPtrSize(*String, String.Len()));
+}
+
 FText FText::FromString( const FString& String )
 {
 	FText NewText = String.IsEmpty() ? FText::GetEmpty() : FText(CopyTemp(String));
@@ -1094,7 +1116,25 @@ FText FText::FromString( FString&& String )
 	return NewText;
 }
 
-FText FText::FromStringView(FStringView InString)
+FText FText::FromString( const FUtf8String& String )
+{
+	return FText::FromString(FString::ConstructFromPtrSize(*String, String.Len()));
+}
+
+FText FText::FromStringView(FAnsiStringView InString)
+{
+	FText NewText = InString.IsEmpty() ? FText::GetEmpty() : FText(FString(InString));
+
+	if (!GIsEditor)
+	{
+		NewText.Flags |= ETextFlag::CultureInvariant;
+	}
+	NewText.Flags |= ETextFlag::InitializedFromString;
+
+	return NewText;
+}
+
+FText FText::FromStringView(FUtf8StringView InString)
 {
 	FText NewText = InString.IsEmpty() ? FText::GetEmpty() : FText(FString(InString));
 
@@ -1107,12 +1147,51 @@ FText FText::FromStringView(FStringView InString)
 	return NewText;
 }
 
-FText FText::AsCultureInvariant( const TCHAR* String )
+FText FText::FromStringView(FWideStringView InString)
 {
-	return FText::AsCultureInvariant(FStringView(String));
+	FText NewText = InString.IsEmpty() ? FText::GetEmpty() : FText(FString(InString));
+
+	if (!GIsEditor)
+	{
+		NewText.Flags |= ETextFlag::CultureInvariant;
+	}
+	NewText.Flags |= ETextFlag::InitializedFromString;
+
+	return NewText;
+}
+
+FText FText::AsCultureInvariant( const ANSICHAR* String )
+{
+	return FText::AsCultureInvariant(FAnsiStringView(String));
+}
+
+FText FText::AsCultureInvariant( const UTF8CHAR* String )
+{
+	return FText::AsCultureInvariant(FUtf8StringView(String));
+}
+
+FText FText::AsCultureInvariant( const WIDECHAR* String )
+{
+	return FText::AsCultureInvariant(FWideStringView(String));
+}
+
+FText FText::AsCultureInvariant( FAnsiStringView String )
+{
+	FText NewText = String.IsEmpty() ? FText::GetEmpty() : FText(FString(String));
+	NewText.Flags |= ETextFlag::CultureInvariant;
+
+	return NewText;
+}
+
+FText FText::AsCultureInvariant( FUtf8StringView String )
+{
+	FText NewText = String.IsEmpty() ? FText::GetEmpty() : FText(FString(String));
+	NewText.Flags |= ETextFlag::CultureInvariant;
+
+	return NewText;
 }
 
-FText FText::AsCultureInvariant( FStringView String )
+FText FText::AsCultureInvariant( FWideStringView String )
 {
 	FText NewText = String.IsEmpty() ? FText::GetEmpty() : FText(FString(String));
 	NewText.Flags |= ETextFlag::CultureInvariant;
diff --git a/Engine/Source/Runtime/Core/Public/Internationalization/Text.h b/Engine/Source/Runtime/Core/Public/Internationalization/Text.h
--- a/Engine/Source/Runtime/Core/Public/Internationalization/Text.h
+++ b/Engine/Source/Runtime/Core/Public/Internationalization/Text.h
@@ -482,19 +482,30 @@ public:
 	/**
 	 * Generate an FText representing the passed in string
 	 */
+	[[nodiscard]] static CORE_API FText FromString( const ANSICHAR* String );
+	[[nodiscard]] static CORE_API FText FromString( const UTF8CHAR* String );
+	[[nodiscard]] static CORE_API FText FromString( const WIDECHAR* String );
+	[[nodiscard]] static CORE_API FText FromString( const FAnsiString& String );
 	[[nodiscard]] static CORE_API FText FromString( const FString& String );
 	[[nodiscard]] static CORE_API FText FromString( FString&& String );
+	[[nodiscard]] static CORE_API FText FromString( const FUtf8String& String );
 
 	/**
 	 * Generate a FText representing the passed string view
 	 */
-	[[nodiscard]] static CORE_API FText FromStringView(FStringView InString);
+	[[nodiscard]] static CORE_API FText FromStringView(FAnsiStringView InString);
+	[[nodiscard]] static CORE_API FText FromStringView(FUtf8StringView InString);
+	[[nodiscard]] static CORE_API FText FromStringView(FWideStringView InString);
 
 	/**
 	 * Generate a culture invariant FText representing the passed in string
 	 */
-	[[nodiscard]] static CORE_API FText AsCultureInvariant( const TCHAR* String );
-	[[nodiscard]] static CORE_API FText AsCultureInvariant( FStringView String );
+	[[nodiscard]] static CORE_API FText AsCultureInvariant( const ANSICHAR* String );
+	[[nodiscard]] static CORE_API FText AsCultureInvariant( const UTF8CHAR* String );
+	[[nodiscard]] static CORE_API FText AsCultureInvariant( const WIDECHAR* String );
+	[[nodiscard]] static CORE_API FText AsCultureInvariant( FAnsiStringView String );
+	[[nodiscard]] static CORE_API FText AsCultureInvariant( FUtf8StringView String );
+	[[nodiscard]] static CORE_API FText AsCultureInvariant( FWideStringView String );
 	[[nodiscard]] static CORE_API FText AsCultureInvariant( FString&& String );
 
 	/**
