[Backout] - CL39320869
#rnx
Original CL Desc
-----------------------------------------------------------------
[AutoRTFM] Add TTaskArray tests, fix bug found

These tests check the existing behavior and will be used to verify an optimization of TTaskArray.

Bug fix: ForwardProcessEntry() and BackwardProcessEntry() can nullify the void* value of the BitMasksPerKey map entry. Currently a transaction will only ever call these once per TTaskArray lifetime, but if called twice (same or either), then we could dereference a nullptr.

diff --git a/Engine/Source/Programs/AutoRTFMTests/Private/TaskArrayTests.cpp b/Engine/Source/Programs/AutoRTFMTests/Private/TaskArrayTests.cpp
deleted file mode 100644
diff --git a/Engine/Source/Runtime/Core/Private/AutoRTFM/TaskArray.h b/Engine/Source/Runtime/Core/Private/AutoRTFM/TaskArray.h
--- a/Engine/Source/Runtime/Core/Private/AutoRTFM/TaskArray.h
+++ b/Engine/Source/Runtime/Core/Private/AutoRTFM/TaskArray.h
@@ -7,7 +7,6 @@
 #include "Containers/Map.h"
 #include "Templates/Function.h"
 #include "Utils.h"
-#include <type_traits>
 
 namespace AutoRTFM
 {
@@ -354,25 +353,19 @@ private:
 		{
 			FBitMaskPerKeyType* const KeyedElement = BitMasksPerKey.Find(*Pair.Key);
 
-			if (KeyedElement && KeyedElement->Num() > 0)
+			if (KeyedElement->Num() > 0)
 			{
 				const bool bFirst = KeyedElement->BitAt(0);
 				KeyedElement->RemoveFirst();
 
 				if (bFirst)
 				{
+					static const TFunction<void()> EmptyFunctor = [] {};
+
 					// We have a deferred signal to pop this element, so do it now.
 					*Pair.Key = SKey{};
-					if constexpr (std::is_invocable_v<InElementType>)
-					{
-						static const TFunction<void()> EmptyFunctor = [] {};
-						Pair.Val->Reset();
-						*Pair.Val = EmptyFunctor;
-					}
-					else
-					{
-						*Pair.Val = {};
-					}
+					Pair.Val->Reset();
+					*Pair.Val = EmptyFunctor;
 					return;
 				}
 			}
@@ -387,10 +380,6 @@ private:
 		if (Pair.Key)
 		{
 			FBitMaskPerKeyType* const KeyedElement = BitMasksPerKey.Find(*Pair.Key);
-			if (!KeyedElement)
-			{
-				return;
-			}
 
 			const int32 Num = KeyedElement->Num();
 
@@ -401,19 +390,12 @@ private:
 
 				if (bLast)
 				{
+					static const TFunction<void()> EmptyFunctor = [] {};
+
 					// We have a deferred signal to pop this element, so do it now.
 					*Pair.Key = SKey{};
-					if constexpr (std::is_invocable_v<InElementType>)
-					{
-						static const TFunction<void()> EmptyFunctor = [] {};
-						Pair.Val->Reset();
-						*Pair.Val = EmptyFunctor;
-					}
-					else
-					{
-						*Pair.Val = {};
-					}
-
+					Pair.Val->Reset();
+					*Pair.Val = EmptyFunctor;
 					return;
 				}
 			}
