PR #12752: Update android stats in a separate background worker thread only where necessary

diff --git a/Engine/Source/Runtime/Core/Private/Android/AndroidPlatformThermal.cpp b/Engine/Source/Runtime/Core/Private/Android/AndroidPlatformThermal.cpp
--- a/Engine/Source/Runtime/Core/Private/Android/AndroidPlatformThermal.cpp
+++ b/Engine/Source/Runtime/Core/Private/Android/AndroidPlatformThermal.cpp
@@ -38,6 +38,7 @@ namespace UE
 			static constexpr int PeriodEnumToS[NumPeriods] = { 1,5,10 };
 			static_assert(NumPeriods == 3); // update PeriodEnumToS if this changes.
 			double LastCallTime = 0;
+			FCriticalSection ThermalCS;
 
 			FAndroidPlatformThermal::EForecastPeriod PeriodToRead = FAndroidPlatformThermal::EForecastPeriod::ONE_SEC;
 
@@ -46,7 +47,6 @@ namespace UE
 
 			FAndroidThermalManagerImpl() : ActivityClass(FJavaClassObject::GetGameActivity())
 			{
-				check(IsInGameThread());
 				ThermalHeadroomMethod = ActivityClass.GetClassMethod("AndroidThunkJava_getThermalHeadroom", "(I)F");
 			}
 
@@ -64,16 +64,19 @@ namespace UE
 
 			const FThermalInfo& UpdateAndGetThermalInfo()
 			{
-				check(IsInGameThread());
 				float CurrentTime = (float)FPlatformTime::Seconds();
 				if (CurrentTime > LastCallTime + CallRateLimitSeconds)
 				{
-					/// thermal headroom docs state we should not call this API more than once per second.
-					float RecentThermalReading = ActivityClass.CallMethod<float>(ThermalHeadroomMethod, PeriodEnumToS[PeriodToRead]);
-					ThermalInfo.HeadroomForecasts[PeriodToRead] = RecentThermalReading;
+					if (ThermalCS.TryLock())
+					{
+						/// thermal headroom docs state we should not call this API more than once per second.
+						float RecentThermalReading = ActivityClass.CallMethod<float>(ThermalHeadroomMethod, PeriodEnumToS[PeriodToRead]);
+						ThermalInfo.HeadroomForecasts[PeriodToRead] = RecentThermalReading;
 
-					LastCallTime = CurrentTime;
-					PeriodToRead = (FAndroidPlatformThermal::EForecastPeriod)((int)(PeriodToRead+1) % NumPeriods);
+						LastCallTime = CurrentTime;
+						PeriodToRead = (FAndroidPlatformThermal::EForecastPeriod)((int)(PeriodToRead+1) % NumPeriods);
+						ThermalCS.Unlock();
+					}
 				}
 				return ThermalInfo;
 			}
diff --git a/Engine/Source/Runtime/Core/Private/Android/AndroidStats.cpp b/Engine/Source/Runtime/Core/Private/Android/AndroidStats.cpp
--- a/Engine/Source/Runtime/Core/Private/Android/AndroidStats.cpp
+++ b/Engine/Source/Runtime/Core/Private/Android/AndroidStats.cpp
@@ -18,6 +18,7 @@
 #	include "hwcpipe.h"
 #endif
 #include "Android/AndroidPlatformThermal.h"
+#include "Tasks/Task.h"
 
 DECLARE_STATS_GROUP(TEXT("Android CPU stats"), STATGROUP_AndroidCPU, STATCAT_Advanced);
 CSV_DEFINE_CATEGORY(AndroidCPU, true);
@@ -225,187 +226,202 @@ void FAndroidStats::OnTrimMemory(int TrimLevel)
 }
 }
 
+static std::atomic<bool> GIsStatTaskActive = false;
+
 void FAndroidStats::UpdateAndroidStats()
 {
-	if (GAndroidCPUStatsUpdateRate <= 0.0f)
+#if CSV_PROFILER || STATS
+	SCOPED_NAMED_EVENT(UpdateAndroidStats, FColor::Green);
+	
+	if (GAndroidCPUStatsUpdateRate <= 0.0f || GIsStatTaskActive.load())
 	{
 		return;
 	}
 
-	static float CPUTemp = 0.0f;
-	static uint64 LastCollectionTime = FPlatformTime::Cycles64();
-	const uint64 CurrentTime = FPlatformTime::Cycles64();
-	const bool bUpdateStats = ((FPlatformTime::ToSeconds64(CurrentTime - LastCollectionTime) >= GAndroidCPUStatsUpdateRate));
-	static FPlatformMemoryStats MemStats = FAndroidPlatformMemory::GetStats();
-	if (bUpdateStats)
+	GIsStatTaskActive = true;
+	// Run everything in a background task thread so that long system calls won't lock up game thread
+	UE::Tasks::Launch(UE_SOURCE_LOCATION, []()
 	{
-		LastCollectionTime = CurrentTime;
-		CPUTemp = FAndroidMisc::GetCPUTemperature();
+#if CSV_PROFILER
+		SCOPED_NAMED_EVENT(UpdateAndroidStats_CSV, FColor::Green);
+
+		static float CPUTemp = 0.0f;
+		static uint64 LastCollectionTime = FPlatformTime::Cycles64();
+		const uint64 CurrentTime = FPlatformTime::Cycles64();
+		const bool bUpdateStats = ((FPlatformTime::ToSeconds64(CurrentTime - LastCollectionTime) >= GAndroidCPUStatsUpdateRate));
+		static FPlatformMemoryStats MemStats = FAndroidPlatformMemory::GetStats();
+		if (bUpdateStats)
+		{
+			LastCollectionTime = CurrentTime;
+			CPUTemp = FAndroidMisc::GetCPUTemperature();
 
-		MemStats = FAndroidPlatformMemory::GetStats();
-	}
+			MemStats = FAndroidPlatformMemory::GetStats();
+		}
 
-	CSV_CUSTOM_STAT_DEFINED(CPUTemp, CPUTemp, ECsvCustomStatOp::Set);
-	CSV_CUSTOM_STAT_DEFINED(ThermalStatus, GThermalStatus, ECsvCustomStatOp::Set);
+		CSV_CUSTOM_STAT_DEFINED(CPUTemp, CPUTemp, ECsvCustomStatOp::Set);
+		CSV_CUSTOM_STAT_DEFINED(ThermalStatus, GThermalStatus, ECsvCustomStatOp::Set);
 
+		float Thermals5S = FAndroidPlatformThermal::GetThermalStress(FAndroidPlatformThermal::EForecastPeriod::FIVE_SEC);
+		CSV_CUSTOM_STAT_DEFINED(ThermalStress, Thermals5S, ECsvCustomStatOp::Set);
 
-	float Thermals5S = FAndroidPlatformThermal::GetThermalStress(FAndroidPlatformThermal::EForecastPeriod::FIVE_SEC);
-	CSV_CUSTOM_STAT_DEFINED(ThermalStress, Thermals5S, ECsvCustomStatOp::Set);
-	
-	CSV_CUSTOM_STAT_DEFINED(TrimMemoryBackgroundLevel, GTrimMemoryBackgroundLevel, ECsvCustomStatOp::Set);
-	CSV_CUSTOM_STAT_DEFINED(Mem_Swap, (int)(MemStats.VMSwap / (1024 * 1024)), ECsvCustomStatOp::Set);
-	CSV_CUSTOM_STAT_DEFINED(Mem_RSS, (int)(MemStats.VMRss / (1024 * 1024)), ECsvCustomStatOp::Set);
-	CSV_CUSTOM_STAT_DEFINED(Mem_TotalUsed, (int)(MemStats.UsedPhysical / (1024 * 1024)), ECsvCustomStatOp::Set);
-	
-	static const uint32 MaxFrequencyGroupStats = 4;
-	const int32 MaxCoresStatsSupport = 16;
-	int32 NumCores = FMath::Min(FAndroidMisc::NumberOfCores(), MaxCoresStatsSupport);
-
-	struct FFrequencyGroup
-	{
-		uint32 MinFrequency;
-		uint32 MaxFrequency;
-		uint32 CoreCount;
-	};
+		CSV_CUSTOM_STAT_DEFINED(TrimMemoryBackgroundLevel, GTrimMemoryBackgroundLevel, ECsvCustomStatOp::Set);
+		CSV_CUSTOM_STAT_DEFINED(Mem_Swap, (int)(MemStats.VMSwap / (1024 * 1024)), ECsvCustomStatOp::Set);
+		CSV_CUSTOM_STAT_DEFINED(Mem_RSS, (int)(MemStats.VMRss / (1024 * 1024)), ECsvCustomStatOp::Set);
+		CSV_CUSTOM_STAT_DEFINED(Mem_TotalUsed, (int)(MemStats.UsedPhysical / (1024 * 1024)), ECsvCustomStatOp::Set);
 
-	static uint32 UnInitializedCores = NumCores;
-	static TArray<FFrequencyGroup> FrequencyGroups;
-	static uint32 CoreFrequencyGroupIndex[MaxCoresStatsSupport] = {
-		0xFFFFFFFF, 0xFFFFFFFF,0xFFFFFFFF,0xFFFFFFFF,
-		0xFFFFFFFF, 0xFFFFFFFF,0xFFFFFFFF,0xFFFFFFFF,
-		0xFFFFFFFF, 0xFFFFFFFF,0xFFFFFFFF,0xFFFFFFFF,
-		0xFFFFFFFF, 0xFFFFFFFF,0xFFFFFFFF,0xFFFFFFFF,
-	};
+		static const uint32 MaxFrequencyGroupStats = 4;
+		const int32 MaxCoresStatsSupport = 16;
+		int32 NumCores = FMath::Min(FAndroidMisc::NumberOfCores(), MaxCoresStatsSupport);
 
-	if (UnInitializedCores != 0)
-	{
-		for (int32 CoreIndex = 0; CoreIndex < NumCores; CoreIndex++)
+		struct FFrequencyGroup
 		{
-			if (CoreFrequencyGroupIndex[CoreIndex] == 0xFFFFFFFF)
+			uint32 MinFrequency;
+			uint32 MaxFrequency;
+			uint32 CoreCount;
+		};
+
+		static uint32 UnInitializedCores = NumCores;
+		static TArray<FFrequencyGroup> FrequencyGroups;
+		static uint32 CoreFrequencyGroupIndex[MaxCoresStatsSupport] = {
+			0xFFFFFFFF, 0xFFFFFFFF,0xFFFFFFFF,0xFFFFFFFF,
+			0xFFFFFFFF, 0xFFFFFFFF,0xFFFFFFFF,0xFFFFFFFF,
+			0xFFFFFFFF, 0xFFFFFFFF,0xFFFFFFFF,0xFFFFFFFF,
+			0xFFFFFFFF, 0xFFFFFFFF,0xFFFFFFFF,0xFFFFFFFF,
+		};
+
+		if (UnInitializedCores != 0)
+		{
+			for (int32 CoreIndex = 0; CoreIndex < NumCores; CoreIndex++)
 			{
-				uint32 MinFreq = FAndroidMisc::GetCoreFrequency(CoreIndex, FAndroidMisc::ECoreFrequencyProperty::MinFrequency) / 1000;
-				uint32 MaxFreq = FAndroidMisc::GetCoreFrequency(CoreIndex, FAndroidMisc::ECoreFrequencyProperty::MaxFrequency) / 1000;
-				if (MaxFreq > 0)
+				if (CoreFrequencyGroupIndex[CoreIndex] == 0xFFFFFFFF)
 				{
-					UnInitializedCores--;
-					int32 FoundIndex = FrequencyGroups.IndexOfByPredicate([&MinFreq, &MaxFreq](const FFrequencyGroup& s) { return s.MinFrequency == MinFreq && s.MaxFrequency == MaxFreq; });
-					if (FoundIndex == INDEX_NONE)
-					{
-						FFrequencyGroup NewGroup = { MinFreq,MaxFreq, 1 };
-						CoreFrequencyGroupIndex[CoreIndex] = FrequencyGroups.Add(NewGroup);
-					}
-					else
+					uint32 MinFreq = FAndroidMisc::GetCoreFrequency(CoreIndex, FAndroidMisc::ECoreFrequencyProperty::MinFrequency) / 1000;
+					uint32 MaxFreq = FAndroidMisc::GetCoreFrequency(CoreIndex, FAndroidMisc::ECoreFrequencyProperty::MaxFrequency) / 1000;
+					if (MaxFreq > 0)
 					{
-						CoreFrequencyGroupIndex[CoreIndex] = FoundIndex;
-						FrequencyGroups[FoundIndex].CoreCount++;
+						UnInitializedCores--;
+						int32 FoundIndex = FrequencyGroups.IndexOfByPredicate([&MinFreq, &MaxFreq](const FFrequencyGroup& s) { return s.MinFrequency == MinFreq && s.MaxFrequency == MaxFreq; });
+						if (FoundIndex == INDEX_NONE)
+						{
+							FFrequencyGroup NewGroup = { MinFreq,MaxFreq, 1 };
+							CoreFrequencyGroupIndex[CoreIndex] = FrequencyGroups.Add(NewGroup);
+						}
+						else
+						{
+							CoreFrequencyGroupIndex[CoreIndex] = FoundIndex;
+							FrequencyGroups[FoundIndex].CoreCount++;
+						}
 					}
 				}
 			}
 		}
-	}
 
-	auto GetFrequencyGroupCurrentFrequency = [&](int32 FrequencyGroupIdx)
-	{
-		for (int32 CoreIdx = 0; CoreIdx < NumCores; CoreIdx++)
+		auto GetFrequencyGroupCurrentFrequency = [&](int32 FrequencyGroupIdx)
 		{
-			if (CoreFrequencyGroupIndex[CoreIdx] == FrequencyGroupIdx)
+			for (int32 CoreIdx = 0; CoreIdx < NumCores; CoreIdx++)
 			{
-				return FAndroidMisc::GetCoreFrequency(CoreIdx, FAndroidMisc::ECoreFrequencyProperty::CurrentFrequency) / 1000;
+				if (CoreFrequencyGroupIndex[CoreIdx] == FrequencyGroupIdx)
+				{
+					return FAndroidMisc::GetCoreFrequency(CoreIdx, FAndroidMisc::ECoreFrequencyProperty::CurrentFrequency) / 1000;
+				}
 			}
-		}
-		return 0u;
-	};
+			return 0u;
+		};
 
-	static int32 CurrentFrequencies[MaxFrequencyGroupStats] = { 0,0,0,0 };
-	static float CurrentFrequenciesPercentage[MaxFrequencyGroupStats] = { 0,0,0,0 };
-	for (int32 FrequencyGroupIndex = 0; FrequencyGroupIndex < FrequencyGroups.Num(); FrequencyGroupIndex++)
-	{
-		if (bUpdateStats)
+		static int32 CurrentFrequencies[MaxFrequencyGroupStats] = { 0,0,0,0 };
+		static float CurrentFrequenciesPercentage[MaxFrequencyGroupStats] = { 0,0,0,0 };
+		for (int32 FrequencyGroupIndex = 0; FrequencyGroupIndex < FrequencyGroups.Num(); FrequencyGroupIndex++)
 		{
-			CurrentFrequencies[FrequencyGroupIndex] = GetFrequencyGroupCurrentFrequency(FrequencyGroupIndex);
-			CurrentFrequenciesPercentage[FrequencyGroupIndex] = ((float)CurrentFrequencies[FrequencyGroupIndex] / (float)FrequencyGroups[FrequencyGroupIndex].MaxFrequency) * 100.0f;
+			if (bUpdateStats)
+			{
+				CurrentFrequencies[FrequencyGroupIndex] = GetFrequencyGroupCurrentFrequency(FrequencyGroupIndex);
+				CurrentFrequenciesPercentage[FrequencyGroupIndex] = ((float)CurrentFrequencies[FrequencyGroupIndex] / (float)FrequencyGroups[FrequencyGroupIndex].MaxFrequency) * 100.0f;
+			}
+			CSV_CUSTOM_STAT_DEFINED_BY_PTR(GCPUFreqStats[FrequencyGroupIndex], CurrentFrequencies[FrequencyGroupIndex], ECsvCustomStatOp::Set);
+			CSV_CUSTOM_STAT_DEFINED_BY_PTR(GCPUFreqPercentageStats[FrequencyGroupIndex], CurrentFrequenciesPercentage[FrequencyGroupIndex], ECsvCustomStatOp::Set);
 		}
-		CSV_CUSTOM_STAT_DEFINED_BY_PTR(GCPUFreqStats[FrequencyGroupIndex], CurrentFrequencies[FrequencyGroupIndex], ECsvCustomStatOp::Set);
-		CSV_CUSTOM_STAT_DEFINED_BY_PTR(GCPUFreqPercentageStats[FrequencyGroupIndex], CurrentFrequenciesPercentage[FrequencyGroupIndex], ECsvCustomStatOp::Set);
-	}
+#endif
 
 #if STATS
-	static const FName AndroidFrequencyGroupMaxFreqStats[] = {
-		GET_STATFNAME(STAT_FreqGroup0MaxFrequency),
-		GET_STATFNAME(STAT_FreqGroup1MaxFrequency),
-		GET_STATFNAME(STAT_FreqGroup2MaxFrequency),
-		GET_STATFNAME(STAT_FreqGroup3MaxFrequency),
-	};
-
-	static const FName AndroidFrequencyGroupMinFreqStats[] = {
-		GET_STATFNAME(STAT_FreqGroup0MinFrequency),
-		GET_STATFNAME(STAT_FreqGroup1MinFrequency),
-		GET_STATFNAME(STAT_FreqGroup2MinFrequency),
-		GET_STATFNAME(STAT_FreqGroup3MinFrequency),
-	};
-
-	static const FName AndroidFrequencyGroupCurrentFreqStats[] = {
-		GET_STATFNAME(STAT_FreqGroup0CurrentFrequency),
-		GET_STATFNAME(STAT_FreqGroup1CurrentFrequency),
-		GET_STATFNAME(STAT_FreqGroup2CurrentFrequency),
-		GET_STATFNAME(STAT_FreqGroup3CurrentFrequency),
-	};
-
-	
-	static const FName AndroidFrequencyGroupCurrentFreqPercentageStats[] = {
-		GET_STATFNAME(STAT_FreqGroup0CurrentFrequencyPercentage),
-		GET_STATFNAME(STAT_FreqGroup1CurrentFrequencyPercentage),
-		GET_STATFNAME(STAT_FreqGroup2CurrentFrequencyPercentage),
-		GET_STATFNAME(STAT_FreqGroup3CurrentFrequencyPercentage),
-	};
-
-	static const FName AndroidFrequencyGroupNumCoresStats[] = {
-		GET_STATFNAME(STAT_FreqGroup0NumCores),
-		GET_STATFNAME(STAT_FreqGroup1NumCores),
-		GET_STATFNAME(STAT_FreqGroup2NumCores),
-		GET_STATFNAME(STAT_FreqGroup3NumCores),
-	};
-
-	static const FName AndroidFrequencyGroupMaxCoresUtilizationStats[] = {
-		GET_STATFNAME(STAT_FreqGroup0MaxUtilization),
-		GET_STATFNAME(STAT_FreqGroup1MaxUtilization),
-		GET_STATFNAME(STAT_FreqGroup2MaxUtilization),
-		GET_STATFNAME(STAT_FreqGroup3MaxUtilization),
-	};
-
-	static float MaxSingleCoreUtilization[MaxFrequencyGroupStats] = { 0.0f };
-	if (bUpdateStats)
-	{
-		FAndroidMisc::FCPUState& AndroidCPUState = FAndroidMisc::GetCPUState();
-		for (int32 CoreIndex = 0; CoreIndex < NumCores; CoreIndex++)
+		static const FName AndroidFrequencyGroupMaxFreqStats[] = {
+			GET_STATFNAME(STAT_FreqGroup0MaxFrequency),
+			GET_STATFNAME(STAT_FreqGroup1MaxFrequency),
+			GET_STATFNAME(STAT_FreqGroup2MaxFrequency),
+			GET_STATFNAME(STAT_FreqGroup3MaxFrequency),
+		};
+
+		static const FName AndroidFrequencyGroupMinFreqStats[] = {
+			GET_STATFNAME(STAT_FreqGroup0MinFrequency),
+			GET_STATFNAME(STAT_FreqGroup1MinFrequency),
+			GET_STATFNAME(STAT_FreqGroup2MinFrequency),
+			GET_STATFNAME(STAT_FreqGroup3MinFrequency),
+		};
+
+		static const FName AndroidFrequencyGroupCurrentFreqStats[] = {
+			GET_STATFNAME(STAT_FreqGroup0CurrentFrequency),
+			GET_STATFNAME(STAT_FreqGroup1CurrentFrequency),
+			GET_STATFNAME(STAT_FreqGroup2CurrentFrequency),
+			GET_STATFNAME(STAT_FreqGroup3CurrentFrequency),
+		};
+
+		
+		static const FName AndroidFrequencyGroupCurrentFreqPercentageStats[] = {
+			GET_STATFNAME(STAT_FreqGroup0CurrentFrequencyPercentage),
+			GET_STATFNAME(STAT_FreqGroup1CurrentFrequencyPercentage),
+			GET_STATFNAME(STAT_FreqGroup2CurrentFrequencyPercentage),
+			GET_STATFNAME(STAT_FreqGroup3CurrentFrequencyPercentage),
+		};
+
+		static const FName AndroidFrequencyGroupNumCoresStats[] = {
+			GET_STATFNAME(STAT_FreqGroup0NumCores),
+			GET_STATFNAME(STAT_FreqGroup1NumCores),
+			GET_STATFNAME(STAT_FreqGroup2NumCores),
+			GET_STATFNAME(STAT_FreqGroup3NumCores),
+		};
+
+		static const FName AndroidFrequencyGroupMaxCoresUtilizationStats[] = {
+			GET_STATFNAME(STAT_FreqGroup0MaxUtilization),
+			GET_STATFNAME(STAT_FreqGroup1MaxUtilization),
+			GET_STATFNAME(STAT_FreqGroup2MaxUtilization),
+			GET_STATFNAME(STAT_FreqGroup3MaxUtilization),
+		};
+
+		static float MaxSingleCoreUtilization[MaxFrequencyGroupStats] = { 0.0f };
+		if (bUpdateStats)
 		{
-			uint32 FrequencyGroupIndex = CoreFrequencyGroupIndex[CoreIndex];
-			if (FrequencyGroupIndex != 0xFFFFFFFF)
+			FAndroidMisc::FCPUState& AndroidCPUState = FAndroidMisc::GetCPUState();
+			for (int32 CoreIndex = 0; CoreIndex < NumCores; CoreIndex++)
 			{
-				float& MaxCoreUtilization = MaxSingleCoreUtilization[FrequencyGroupIndex];
-				MaxCoreUtilization = FMath::Max((float)AndroidCPUState.Utilization[CoreIndex], MaxCoreUtilization);
+				uint32 FrequencyGroupIndex = CoreFrequencyGroupIndex[CoreIndex];
+				if (FrequencyGroupIndex != 0xFFFFFFFF)
+				{
+					float& MaxCoreUtilization = MaxSingleCoreUtilization[FrequencyGroupIndex];
+					MaxCoreUtilization = FMath::Max((float)AndroidCPUState.Utilization[CoreIndex], MaxCoreUtilization);
+				}
 			}
 		}
-	}
 
-	for (int32 FrequencyGroupIndex = 0; FrequencyGroupIndex < FrequencyGroups.Num(); FrequencyGroupIndex++)
-	{
-		const FFrequencyGroup& FrequencyGroup = FrequencyGroups[FrequencyGroupIndex];
-		SET_DWORD_STAT_FName(AndroidFrequencyGroupMaxFreqStats[FrequencyGroupIndex], FrequencyGroup.MaxFrequency);
-		SET_DWORD_STAT_FName(AndroidFrequencyGroupNumCoresStats[FrequencyGroupIndex], FrequencyGroup.CoreCount);
-		//SET_DWORD_STAT_FName(AndroidFrequencyGroupMinFreqStats[FrequencyGroupIndex], FrequencyGroup.MinFrequency);
-		SET_DWORD_STAT_FName(AndroidFrequencyGroupCurrentFreqStats[FrequencyGroupIndex], CurrentFrequencies[FrequencyGroupIndex]);
-		SET_FLOAT_STAT_FName(AndroidFrequencyGroupCurrentFreqPercentageStats[FrequencyGroupIndex], CurrentFrequencies[FrequencyGroupIndex]);
-		SET_FLOAT_STAT_FName(AndroidFrequencyGroupMaxCoresUtilizationStats[FrequencyGroupIndex], MaxSingleCoreUtilization[FrequencyGroupIndex]);
-	}
+		for (int32 FrequencyGroupIndex = 0; FrequencyGroupIndex < FrequencyGroups.Num(); FrequencyGroupIndex++)
+		{
+			const FFrequencyGroup& FrequencyGroup = FrequencyGroups[FrequencyGroupIndex];
+			SET_DWORD_STAT_FName(AndroidFrequencyGroupMaxFreqStats[FrequencyGroupIndex], FrequencyGroup.MaxFrequency);
+			SET_DWORD_STAT_FName(AndroidFrequencyGroupNumCoresStats[FrequencyGroupIndex], FrequencyGroup.CoreCount);
+			//SET_DWORD_STAT_FName(AndroidFrequencyGroupMinFreqStats[FrequencyGroupIndex], FrequencyGroup.MinFrequency);
+			SET_DWORD_STAT_FName(AndroidFrequencyGroupCurrentFreqStats[FrequencyGroupIndex], CurrentFrequencies[FrequencyGroupIndex]);
+			SET_DWORD_STAT_FName(AndroidFrequencyGroupCurrentFreqPercentageStats[FrequencyGroupIndex], CurrentFrequencies[FrequencyGroupIndex]);
+			SET_DWORD_STAT_FName(AndroidFrequencyGroupMaxCoresUtilizationStats[FrequencyGroupIndex], MaxSingleCoreUtilization[FrequencyGroupIndex]);
+		}
 
-	static const FName CPUStatName = GET_STATFNAME(STAT_CPUTemp);
-	static const FName ThermalStatus = GET_STATFNAME(STAT_ThermalStatus);
-	SET_FLOAT_STAT_FName(CPUStatName, CPUTemp);
-	SET_FLOAT_STAT_FName(ThermalStatus, GThermalStatus);
+		static const FName CPUStatName = GET_STATFNAME(STAT_CPUTemp);
+		static const FName ThermalStatus = GET_STATFNAME(STAT_ThermalStatus);
+		SET_DWORD_STAT_FName(CPUStatName, CPUTemp);
+		SET_DWORD_STAT_FName(ThermalStatus, GThermalStatus);
 #endif
 
-	UpdateGPUStats();
+		UpdateGPUStats();
+		GIsStatTaskActive = false;
+	}, UE::Tasks::ETaskPriority::BackgroundLow);
+#endif
 }
 
 static void UpdateGPUStats()
