Added TArrayView::FindLast(const ElementType&) overload, which was causing the other overload to fail to compile.

diff --git a/Engine/Source/Runtime/Core/Public/Containers/ArrayView.h b/Engine/Source/Runtime/Core/Public/Containers/ArrayView.h
--- a/Engine/Source/Runtime/Core/Public/Containers/ArrayView.h
+++ b/Engine/Source/Runtime/Core/Public/Containers/ArrayView.h
@@ -553,6 +553,25 @@ public:
 		return Index != static_cast<SizeType>(INDEX_NONE);
 	}
 
+	/**
+	 * Finds element within the array starting from the end.
+	 *
+	 * @param Item Item to look for.
+	 * @returns Index of the found element. INDEX_NONE otherwise.
+	 */
+	[[nodiscard]] SizeType FindLast(const ElementType& Item) const
+	{
+		for (const ElementType* RESTRICT Start = GetData(), *RESTRICT Data = Start + ArrayNum; Data != Start; )
+		{
+			--Data;
+			if (*Data == Item)
+			{
+				return static_cast<SizeType>(Data - Start);
+			}
+		}
+		return static_cast<SizeType>(INDEX_NONE);
+	}
+
 	/**
 	 * Finds element within the array starting from StartIndex and going backwards. Uses predicate to match element.
 	 *
