Reverted the recent mutex unlocking optimizations while investigating a potential race

Fixed an issue with FExternalMutex where unlocking would fail if another thread had changed bits of the state that are not used by the mutex.

#rnx

diff --git a/Engine/Source/Runtime/Core/Private/Async/ExternalMutex.cpp b/Engine/Source/Runtime/Core/Private/Async/ExternalMutex.cpp
--- a/Engine/Source/Runtime/Core/Private/Async/ExternalMutex.cpp
+++ b/Engine/Source/Runtime/Core/Private/Async/ExternalMutex.cpp
@@ -38,7 +38,7 @@ void FExternalMutex::LockSlow()
 		// Store that there are waiting threads. Restart if the state has changed since it was loaded.
 		if (LIKELY(!(CurrentState & HasWaitingLockFlag)))
 		{
-			if (UNLIKELY(!State.compare_exchange_weak(CurrentState, CurrentState | HasWaitingLockFlag, std::memory_order_acquire)))
+			if (UNLIKELY(!State.compare_exchange_weak(CurrentState, CurrentState | HasWaitingLockFlag, std::memory_order_relaxed)))
 			{
 				continue;
 			}
@@ -60,11 +60,35 @@ void FExternalMutex::LockSlow()
 
 FORCENOINLINE void FExternalMutex::WakeWaitingThread()
 {
+	if constexpr (!bUnlockImmediately)
+	{
+		uint8 CurrentState = State.load(std::memory_order_relaxed);
+		checkSlow(CurrentState & IsLockedFlag);
+
+		// Spin on the fast path because there may be spurious failures.
+		while (LIKELY(!(CurrentState & HasWaitingLockFlag)))
+		{
+			if (LIKELY(State.compare_exchange_weak(CurrentState, CurrentState & ~IsLockedFlag, std::memory_order_release, std::memory_order_relaxed)))
+			{
+				return;
+			}
+		}
+	}
+
 	ParkingLot::WakeOne(&State, [this](ParkingLot::FWakeState WakeState) -> uint64
 	{
-		if (!WakeState.bHasWaitingThreads)
+		if constexpr (bUnlockImmediately)
+		{
+			if (!WakeState.bHasWaitingThreads)
+			{
+				State.fetch_and(~HasWaitingLockFlag, std::memory_order_relaxed);
+			}
+		}
+		else
 		{
-			State.fetch_and(~HasWaitingLockFlag, std::memory_order_relaxed);
+			const uint8 ClearFlags = IsLockedFlag | (WakeState.bHasWaitingThreads ? 0 : HasWaitingLockFlag);
+			const uint8 OldState = State.fetch_and(~ClearFlags, std::memory_order_release);
+			checkSlow((OldState & IsLockedFlag) && (OldState & HasWaitingLockFlag));
 		}
 		return 0;
 	});
diff --git a/Engine/Source/Runtime/Core/Private/Async/Mutex.cpp b/Engine/Source/Runtime/Core/Private/Async/Mutex.cpp
--- a/Engine/Source/Runtime/Core/Private/Async/Mutex.cpp
+++ b/Engine/Source/Runtime/Core/Private/Async/Mutex.cpp
@@ -57,11 +57,35 @@ void FMutex::LockSlow()
 
 FORCENOINLINE void FMutex::WakeWaitingThread()
 {
+	if constexpr (!bUnlockImmediately)
+	{
+		uint8 CurrentState = State.load(std::memory_order_relaxed);
+		checkSlow(CurrentState & IsLockedFlag);
+
+		// Spin on the fast path because there may be spurious failures.
+		while (LIKELY(!(CurrentState & HasWaitingLockFlag)))
+		{
+			if (LIKELY(State.compare_exchange_weak(CurrentState, 0, std::memory_order_release, std::memory_order_relaxed)))
+			{
+				return;
+			}
+		}
+	}
+
 	ParkingLot::WakeOne(&State, [this](ParkingLot::FWakeState WakeState) -> uint64
 	{
-		if (!WakeState.bHasWaitingThreads)
+		if constexpr (bUnlockImmediately)
+		{
+			if (!WakeState.bHasWaitingThreads)
+			{
+				State.fetch_and(~HasWaitingLockFlag, std::memory_order_relaxed);
+			}
+		}
+		else
 		{
-			State.fetch_and(~HasWaitingLockFlag, std::memory_order_relaxed);
+			const uint8 NewState = WakeState.bHasWaitingThreads ? HasWaitingLockFlag : 0;
+			const uint8 OldState = State.exchange(NewState, std::memory_order_release);
+			checkSlow((OldState & IsLockedFlag) && (OldState & HasWaitingLockFlag));
 		}
 		return 0;
 	});
diff --git a/Engine/Source/Runtime/Core/Private/Async/RecursiveMutex.cpp b/Engine/Source/Runtime/Core/Private/Async/RecursiveMutex.cpp
--- a/Engine/Source/Runtime/Core/Private/Async/RecursiveMutex.cpp
+++ b/Engine/Source/Runtime/Core/Private/Async/RecursiveMutex.cpp
@@ -10,6 +10,8 @@
 namespace UE
 {
 
+static constexpr bool bUnlockRecursiveMutexImmediately = false;
+
 bool FRecursiveMutex::TryLock()
 {
 	const uint32 CurrentThreadId = FPlatformTLS::GetCurrentThreadId();
@@ -123,12 +125,27 @@ void FRecursiveMutex::Unlock()
 		// Remove the association with this thread before unlocking.
 		ThreadId.store(0, std::memory_order_relaxed);
 
-		// Unlock immediately to allow other threads to acquire the lock while this thread looks for a thread to wake.
-		uint32 LastState = State.fetch_sub(1 << LockCountShift, std::memory_order_release);
+		if constexpr (bUnlockRecursiveMutexImmediately)
+		{
+			// Unlock immediately to allow other threads to acquire the lock while this thread looks for a thread to wake.
+			const uint32 LastState = State.fetch_sub(1 << LockCountShift, std::memory_order_release);
 
-		// Wake one exclusive waiter if there are waiting threads.
-		if (UNLIKELY(LastState & HasWaitingLockFlag))
+			// Wake one exclusive waiter if there are waiting threads.
+			if (UNLIKELY(LastState & HasWaitingLockFlag))
+			{
+				WakeWaitingThread();
+			}
+		}
+		else
 		{
+			if (LIKELY(!(CurrentState & HasWaitingLockFlag)))
+			{
+				if (LIKELY(State.compare_exchange_strong(CurrentState, 0, std::memory_order_release, std::memory_order_relaxed)))
+				{
+					return;
+				}
+			}
+
 			WakeWaitingThread();
 		}
 	}
@@ -143,9 +160,18 @@ FORCENOINLINE void FRecursiveMutex::WakeWaitingThread()
 {
 	ParkingLot::WakeOne(&State, [this](ParkingLot::FWakeState WakeState) -> uint64
 	{
-		if (!WakeState.bHasWaitingThreads)
+		if constexpr (bUnlockRecursiveMutexImmediately)
+		{
+			if (!WakeState.bHasWaitingThreads)
+			{
+				State.fetch_and(~HasWaitingLockFlag, std::memory_order_relaxed);
+			}
+		}
+		else
 		{
-			State.fetch_and(~HasWaitingLockFlag, std::memory_order_relaxed);
+			const uint32 NewState = WakeState.bHasWaitingThreads ? HasWaitingLockFlag : 0;
+			const uint32 OldState = State.exchange(NewState, std::memory_order_release);
+			checkSlow((OldState & LockCountMask) && (OldState & HasWaitingLockFlag));
 		}
 		return 0;
 	});
diff --git a/Engine/Source/Runtime/Core/Public/Async/ExternalMutex.h b/Engine/Source/Runtime/Core/Public/Async/ExternalMutex.h
--- a/Engine/Source/Runtime/Core/Public/Async/ExternalMutex.h
+++ b/Engine/Source/Runtime/Core/Public/Async/ExternalMutex.h
@@ -61,13 +61,25 @@ public:
 
 	inline void Unlock()
 	{
-		// Unlock immediately to allow other threads to acquire the lock while this thread looks for a thread to wake.
-		uint8 LastState = State.fetch_sub(IsLockedFlag, std::memory_order_release);
-		if (LIKELY(!(LastState & HasWaitingLockFlag)))
+		if constexpr (bUnlockImmediately)
 		{
-			return;
+			// Unlock immediately to allow other threads to acquire the lock while this thread looks for a thread to wake.
+			const uint8 LastState = State.fetch_sub(IsLockedFlag, std::memory_order_release);
+			if (LIKELY(!(LastState & HasWaitingLockFlag)))
+			{
+				return;
+			}
+			WakeWaitingThread();
+		}
+		else
+		{
+			uint8 Expected = (State.load(std::memory_order_relaxed) | IsLockedFlag) & ~HasWaitingLockFlag;
+			if (LIKELY(State.compare_exchange_weak(Expected, Expected & ~IsLockedFlag, std::memory_order_release, std::memory_order_relaxed)))
+			{
+				return;
+			}
+			WakeWaitingThread();
 		}
-		WakeWaitingThread();
 	}
 
 private:
@@ -77,6 +89,8 @@ private:
 	static constexpr uint8 IsLockedFlag = 1 << 0;
 	static constexpr uint8 HasWaitingLockFlag = 1 << 1;
 
+	static constexpr bool bUnlockImmediately = false;
+
 	std::atomic<uint8>& State;
 };
 
diff --git a/Engine/Source/Runtime/Core/Public/Async/Mutex.h b/Engine/Source/Runtime/Core/Public/Async/Mutex.h
--- a/Engine/Source/Runtime/Core/Public/Async/Mutex.h
+++ b/Engine/Source/Runtime/Core/Public/Async/Mutex.h
@@ -53,19 +53,33 @@ public:
 
 	inline void Unlock()
 	{
-		// Unlock immediately to allow other threads to acquire the lock while this thread looks for a thread to wake.
-		uint8 LastState = State.fetch_sub(IsLockedFlag, std::memory_order_release);
-		if (LIKELY(!(LastState & HasWaitingLockFlag)))
+		if constexpr (bUnlockImmediately)
 		{
-			return;
+			// Unlock immediately to allow other threads to acquire the lock while this thread looks for a thread to wake.
+			const uint8 LastState = State.fetch_sub(IsLockedFlag, std::memory_order_release);
+			if (LIKELY(!(LastState & HasWaitingLockFlag)))
+			{
+				return;
+			}
+			WakeWaitingThread();
+		}
+		else
+		{
+			uint8 Expected = IsLockedFlag;
+			if (LIKELY(State.compare_exchange_weak(Expected, 0, std::memory_order_release, std::memory_order_relaxed)))
+			{
+				return;
+			}
+			WakeWaitingThread();
 		}
-		WakeWaitingThread();
 	}
 
 private:
 	UE_API void LockSlow();
 	UE_API void WakeWaitingThread();
 
+	static constexpr bool bUnlockImmediately = false;
+
 	static constexpr uint8 IsLockedFlag = 1 << 0;
 	static constexpr uint8 HasWaitingLockFlag = 1 << 1;
 
