Possible fix for missing TTokenDefinitions<char16_t>::IgnoreWhitespace symbols on Mac.

#rnx

diff --git a/Engine/Source/Runtime/Core/Private/Misc/ExpressionParser.cpp b/Engine/Source/Runtime/Core/Private/Misc/ExpressionParser.cpp
--- a/Engine/Source/Runtime/Core/Private/Misc/ExpressionParser.cpp
+++ b/Engine/Source/Runtime/Core/Private/Misc/ExpressionParser.cpp
@@ -299,12 +299,24 @@ void TExpressionTokenConsumer<CharType>::Add(const TStringToken<CharType>& Sourc
 	Tokens.Emplace(SourceToken, MoveTemp(Node));
 }
 
+template <typename CharType>
+void TTokenDefinitions<CharType>::IgnoreWhitespace()
+{
+	bIgnoreWhitespace = true;
+}
+
 template <typename CharType>
 void TTokenDefinitions<CharType>::DefineToken(TFunction<TExpressionDefinition<CharType>>&& Definition)
 {
 	Definitions.Emplace(MoveTemp(Definition));
 }
 
+template <typename CharType>
+bool TTokenDefinitions<CharType>::DoesIgnoreWhitespace()
+{
+	return bIgnoreWhitespace;
+}
+
 template <typename CharType>
 TOptional<FExpressionError> TTokenDefinitions<CharType>::ConsumeToken(TExpressionTokenConsumer<CharType>& Consumer) const
 {
diff --git a/Engine/Source/Runtime/Core/Public/Misc/ExpressionParserTypes.h b/Engine/Source/Runtime/Core/Public/Misc/ExpressionParserTypes.h
--- a/Engine/Source/Runtime/Core/Public/Misc/ExpressionParserTypes.h
+++ b/Engine/Source/Runtime/Core/Public/Misc/ExpressionParserTypes.h
@@ -504,19 +504,13 @@ public:
 	TTokenDefinitions() = default;
 
 	/** Define the grammar to ignore whitespace between tokens, unless explicitly included in a token */
-	void IgnoreWhitespace()
-	{
-		bIgnoreWhitespace = true;
-	}
+	CORE_API void IgnoreWhitespace();
 
 	/** Define a token by way of a function to be invoked to attempt to parse a token from a stream */
 	CORE_API void DefineToken(TFunction<TExpressionDefinition<CharType>>&& Definition);
 
 	/** Check if the grammar ignores whitespace */
-	bool DoesIgnoreWhitespace()
-	{
-		return bIgnoreWhitespace;
-	}
+	CORE_API bool DoesIgnoreWhitespace();
 
 	/** Consume a token for the specified consumer */
 	CORE_API TOptional<FExpressionError> ConsumeTokens(TExpressionTokenConsumer<CharType>& Consumer) const;
