Another attempted Mac missing symbol fix.

diff --git a/Engine/Source/Runtime/Core/Private/Misc/ExpressionParser.cpp b/Engine/Source/Runtime/Core/Private/Misc/ExpressionParser.cpp
--- a/Engine/Source/Runtime/Core/Private/Misc/ExpressionParser.cpp
+++ b/Engine/Source/Runtime/Core/Private/Misc/ExpressionParser.cpp
@@ -299,6 +299,12 @@ void TExpressionTokenConsumer<CharType>::Add(const TStringToken<CharType>& Sourc
 	Tokens.Emplace(SourceToken, MoveTemp(Node));
 }
 
+template <typename CharType>
+TTokenStream<CharType>& TExpressionTokenConsumer<CharType>::GetStream()
+{
+	return Stream;
+}
+
 template <typename CharType>
 void TTokenDefinitions<CharType>::IgnoreWhitespace()
 {
diff --git a/Engine/Source/Runtime/Core/Public/Misc/ExpressionParserTypes.h b/Engine/Source/Runtime/Core/Public/Misc/ExpressionParserTypes.h
--- a/Engine/Source/Runtime/Core/Public/Misc/ExpressionParserTypes.h
+++ b/Engine/Source/Runtime/Core/Public/Misc/ExpressionParserTypes.h
@@ -474,10 +474,7 @@ public:
 	CORE_API void Add(const TStringToken<CharType>& SourceToken, FExpressionNode&& Node);
 
 	/** Get the expression stream */
-	TTokenStream<CharType>& GetStream()
-	{
-		return Stream;
-	}
+	CORE_API TTokenStream<CharType>& GetStream();
 
 private:
 	/** Array of added tokens */
