[trivial] IWYU header fix (iOS)

diff --git a/Engine/Source/Runtime/Core/Private/IOS/IOSPlatformCrashContext.cpp b/Engine/Source/Runtime/Core/Private/IOS/IOSPlatformCrashContext.cpp
--- a/Engine/Source/Runtime/Core/Private/IOS/IOSPlatformCrashContext.cpp
+++ b/Engine/Source/Runtime/Core/Private/IOS/IOSPlatformCrashContext.cpp
@@ -12,6 +12,7 @@
 #include "HAL/PlatformOutputDevices.h"
 
 #include "Internationalization/Culture.h"
+#include "Internationalization/Internationalization.h"
 
 #include "Misc/App.h"
 #include "Misc/EngineVersion.h"
