[trivial] Replaced header include with forward declation

diff --git a/Engine/Source/Developer/AutomationController/Public/AutomationControllerSettings.h b/Engine/Source/Developer/AutomationController/Public/AutomationControllerSettings.h
--- a/Engine/Source/Developer/AutomationController/Public/AutomationControllerSettings.h
+++ b/Engine/Source/Developer/AutomationController/Public/AutomationControllerSettings.h
@@ -8,6 +8,7 @@
 #include "CoreMinimal.h"
 #include "HAL/Platform.h"
 #include "IAutomationReport.h"
+#include "Misc/TextFilterExpressionEvaluator.h"
 #include "UObject/Object.h"
 #include "UObject/ObjectMacros.h"
 #include "UObject/SoftObjectPath.h"
diff --git a/Engine/Source/Runtime/Core/Private/Misc/AutomationTest.cpp b/Engine/Source/Runtime/Core/Private/Misc/AutomationTest.cpp
--- a/Engine/Source/Runtime/Core/Private/Misc/AutomationTest.cpp
+++ b/Engine/Source/Runtime/Core/Private/Misc/AutomationTest.cpp
@@ -19,6 +19,7 @@
 #include "Misc/Paths.h"
 #include "Misc/ScopedSlowTask.h"
 #include "Misc/ScopeRWLock.h"
+#include "Misc/TextFilterExpressionEvaluator.h"
 #include "Modules/ModuleManager.h"
 
 DEFINE_LOG_CATEGORY(LogLatentCommands);
diff --git a/Engine/Source/Runtime/Core/Public/Misc/AutomationTest.h b/Engine/Source/Runtime/Core/Public/Misc/AutomationTest.h
--- a/Engine/Source/Runtime/Core/Public/Misc/AutomationTest.h
+++ b/Engine/Source/Runtime/Core/Public/Misc/AutomationTest.h
@@ -40,7 +40,6 @@
 #include "Misc/Guid.h"
 #include "Misc/Optional.h"
 #include "Misc/OutputDevice.h"
-#include "Misc/TextFilterExpressionEvaluator.h"
 #include "Misc/Timespan.h"
 #include "Templates/Function.h"
 #include "Templates/SharedPointer.h"
@@ -54,6 +53,7 @@ CORE_API DECLARE_LOG_CATEGORY_EXTERN(LogAutomationTestFramework, Log, All);
 
 class FAutomationTestBase;
 class FAutomationTestTags;
+class FTextFilterExpressionEvaluator;
 
 #ifndef WITH_AUTOMATION_TESTS
 	#define WITH_AUTOMATION_TESTS (WITH_DEV_AUTOMATION_TESTS || WITH_PERF_AUTOMATION_TESTS)
@@ -4716,9 +4716,11 @@ public:
 			if (HasExceededMaxTotalRunTime())
 			{
 				// Command run time has exceeded max total run time.
-				FString ErrorMessageText = FString::Printf(TEXT("%s has failed due to exceeding the max allowed run time of %f seconds. \
-This may be due to an error, or having a single command attempt to do too many things. If this is not due to an error, consider breaking \
-up this command into multiple, smaller commands."), *GetTestAndCommandName(), MaxTotalRunTimeInSeconds);
+				FString ErrorMessageText = FString::Printf(
+					TEXT("%s has failed due to exceeding the max allowed run time of %f seconds. ")
+					TEXT("This may be due to an error, or having a single command attempt to do ")
+					TEXT("too many things. If this is not due to an error, consider breaking ")
+					TEXT("up this command into multiple, smaller commands."), *GetTestAndCommandName(), MaxTotalRunTimeInSeconds);
 				CommandFailedDueToError(ErrorMessageText);
 				return true;
 			}
@@ -4791,8 +4793,10 @@ protected:
 		if (bHasUnlimitedRetries)
 		{
 			// Must log here so that Gauntlet can also pick up the error for its report. Otherwise, it will only show up in the Automation log.
-			UE_LOG(LogLatentCommands, Warning, TEXT("%s has been set to Unlimited retries. Will be using MaxTotalRunTime to prevent \
-running forever. Default time is set at 300 seconds. If this is not enough time, make sure to override this in your Execute() loop."), *GetTestAndCommandName());
+			UE_LOG(LogLatentCommands, Warning,
+				TEXT("%s has been set to Unlimited retries. Will be using MaxTotalRunTime to prevent ")
+				TEXT("running forever. Default time is set at 300 seconds. If this is not enough time, ")
+				TEXT("make sure to override this in your Execute() loop."), *GetTestAndCommandName());
 		}
 	}
 
@@ -4863,3 +4867,7 @@ private:
 	double MaxTotalRunTimeInSeconds = 300.0;
 	
 };
+
+#if UE_ENABLE_INCLUDE_ORDER_DEPRECATED_IN_5_6
+# include "TextFilterExpressionEvaluator.h"
+#endif
