[trivial] Fixed header missed in deprecation block

diff --git a/Engine/Source/Runtime/Core/Public/CoreSharedPCH.h b/Engine/Source/Runtime/Core/Public/CoreSharedPCH.h
--- a/Engine/Source/Runtime/Core/Public/CoreSharedPCH.h
+++ b/Engine/Source/Runtime/Core/Public/CoreSharedPCH.h
@@ -259,4 +259,5 @@
 
 #if UE_ENABLE_INCLUDE_ORDER_DEPRECATED_IN_5_6
 #include "Math/BasicMathExpressionEvaluator.h"
+#include "Misc/ExpressionParserTypes.h"
 #endif
