[Trace]
- Fixed "trace.status" displaying incorrect value for "Emitted - Traced" statistic.
- Fixed TraceLog messages to not duplicate the message type into the log message.

diff --git a/Engine/Source/Runtime/Core/Private/ProfilingDebugging/TraceAuxiliary.cpp b/Engine/Source/Runtime/Core/Private/ProfilingDebugging/TraceAuxiliary.cpp
--- a/Engine/Source/Runtime/Core/Private/ProfilingDebugging/TraceAuxiliary.cpp
+++ b/Engine/Source/Runtime/Core/Private/ProfilingDebugging/TraceAuxiliary.cpp
@@ -177,7 +177,6 @@ public:
 	void EnableCommandlineChannels();
 	void EnableCommandlineChannelsPostInitialize();
 	void SetTruncateFile(bool bTruncateFile);
-	void UpdateCsvStats() const;
 	void StartWorkerThread();
 	void RegisterEndFrameCallbacks();
 	bool WriteSnapshot(const TCHAR* InFilePath, const FTraceAuxiliary::FLogCategoryAlias& LogCategory);
@@ -1013,38 +1012,34 @@ void TraceAuxiliaryOnMessageCallback(const UE::Trace::FMessageEvent& Message)
 	#define LOG_WITH_VERBOSITY(Verbosity, Text) \
 		UE_LOG(LogTrace, Verbosity, TEXT("%s"), Text)
 
-	const auto TypeStr = StringCast<TCHAR>(Message.TypeStr);
 	TStringBuilder<128> LogStr;
-	if (Message.Description != nullptr)
+	if (Message.Description)
 	{
-		LogStr << Message.TypeStr << TEXT(" ") << Message.Description;
-	}
-	else
-	{
-		LogStr << Message.TypeStr;
+		LogStr << Message.Description;
 	}
+	const TCHAR* Msg = LogStr.ToString();
 
-	switch(Message.Type)
+	switch (Message.Type)
 	{
 	case UE::Trace::EMessageType::Log:
-		LOG_WITH_VERBOSITY(Log, LogStr.ToString());
+		LOG_WITH_VERBOSITY(Log, Msg);
 		break;
 	case UE::Trace::EMessageType::Display:
-		LOG_WITH_VERBOSITY(Display, LogStr.ToString());
+		LOG_WITH_VERBOSITY(Display, Msg);
 		break;
 	default:
 		{
 			if (Message.Type > UE::Trace::EMessageType::FatalStart)
 			{
-				LOG_WITH_VERBOSITY(Fatal, LogStr.ToString());
+				LOG_WITH_VERBOSITY(Fatal, Msg);
 			}
 			else if (Message.Type > UE::Trace::EMessageType::ErrorStart)
 			{
-				LOG_WITH_VERBOSITY(Error, LogStr.ToString());
+				LOG_WITH_VERBOSITY(Error, Msg);
 			}
 			else if (Message.Type > UE::Trace::EMessageType::WarningStart)
 			{
-				LOG_WITH_VERBOSITY(Warning, LogStr.ToString());
+				LOG_WITH_VERBOSITY(Warning, Msg);
 			}
 		}
 		break;
@@ -1216,31 +1211,34 @@ static void TraceAuxiliaryStatus()
 	{
 		ConnectionStr = TEXT("Not tracing");
 	}
-	UE_LOG(LogConsoleResponse, Display, TEXT("- Connection: %s"), ConnectionStr.ToString());
+	UE_LOG(LogConsoleResponse, Display, TEXT("Connection: %s"), ConnectionStr.ToString());
 
 	// Stats
 	UE::Trace::FStatistics Stats;
 	UE::Trace::GetStatistics(Stats);
 	constexpr double MiB = 1.0 / (1024.0 * 1024.0);
-	UE_LOG(LogConsoleResponse, Display, TEXT("- Memory Used: %.02f MiB"),
+	UE_LOG(LogConsoleResponse, Display, TEXT("Memory Used: %.02f MiB"),
 		double(Stats.MemoryUsed) * MiB);
-	UE_LOG(LogConsoleResponse, Display, TEXT("- Block Pool: %.02f MiB"),
+	UE_LOG(LogConsoleResponse, Display, TEXT("Block Pool: %.02f MiB"),
 		double(Stats.BlockPoolAllocated) * MiB);
-	UE_LOG(LogConsoleResponse, Display, TEXT("- Fixed Buffers: %.02f MiB"),
+	UE_LOG(LogConsoleResponse, Display, TEXT("Fixed Buffers: %.02f MiB"),
 		double(Stats.FixedBufferAllocated) * MiB);
-	UE_LOG(LogConsoleResponse, Display, TEXT("- Shared Buffers: %.02f MiB"),
+	UE_LOG(LogConsoleResponse, Display, TEXT("Shared Buffers: %.02f MiB"),
 		double(Stats.SharedBufferAllocated) * MiB);
-	UE_LOG(LogConsoleResponse, Display, TEXT("- Important Events Cache: %.02f MiB (%.02f MiB used + %0.02f MiB unused | %0.02f MiB waste)"),
+	UE_LOG(LogConsoleResponse, Display, TEXT("Important Events Cache: %.02f MiB (%.02f MiB used + %0.02f MiB unused | %0.02f MiB waste)"),
 		double(Stats.CacheAllocated) * MiB,
 		double(Stats.CacheUsed) * MiB,
 		double(Stats.CacheAllocated - Stats.CacheUsed) * MiB,
 		double(Stats.CacheWaste) * MiB);
-	UE_LOG(LogConsoleResponse, Display, TEXT("- Emitted: %.02f MiB"),
+	UE_LOG(LogConsoleResponse, Display, TEXT("Emitted: %.02f MiB"),
 		double(Stats.BytesEmitted) * MiB);
-	UE_LOG(LogConsoleResponse, Display, TEXT("- Traced: %.02f MiB (-%llu bytes)"),
-		double(Stats.BytesTraced) * MiB,
-		Stats.BytesEmitted - Stats.BytesTraced);
-	UE_LOG(LogConsoleResponse, Display, TEXT("- Sent: %.02f MiB"),
+	UE_LOG(LogConsoleResponse, Display, TEXT("Traced: %.02f MiB"),
+		double(Stats.BytesTraced) * MiB);
+	int64 BytesEmittedNotTraced = int64(Stats.BytesEmitted) - int64(Stats.BytesTraced);
+	UE_LOG(LogConsoleResponse, Display, TEXT("Emitted - Traced: %.02f MiB (%lli bytes)"),
+		double(BytesEmittedNotTraced) * MiB,
+		BytesEmittedNotTraced);
+	UE_LOG(LogConsoleResponse, Display, TEXT("Sent: %.02f MiB"),
 		double(Stats.BytesSent) * MiB);
 
 	// Channels
@@ -1276,8 +1274,8 @@ static void TraceAuxiliaryStatus()
 		FAnsiStringView NameView = FAnsiStringView(Name).LeftChop(7); // Remove "Channel" suffix
 		EnumerateUser[bEnabled ? 0 : 1].AddChannel(NameView);
 	}, EnumerateUserData);
-	UE_LOG(LogConsoleResponse, Display, TEXT("- Enabled channels: %s"), EnumerateUserData[0].Count == 0 ? TEXT("<none>") : EnumerateUserData[0].ChannelsStr.ToString());
-	UE_LOG(LogConsoleResponse, Display, TEXT("- Available channels: %s"), EnumerateUserData[1].ChannelsStr.ToString());
+	UE_LOG(LogConsoleResponse, Display, TEXT("Enabled channels: %s"), EnumerateUserData[0].Count == 0 ? TEXT("<none>") : EnumerateUserData[0].ChannelsStr.ToString());
+	UE_LOG(LogConsoleResponse, Display, TEXT("Available channels: %s"), EnumerateUserData[1].ChannelsStr.ToString());
 
 	UE_LOG(LogConsoleResponse, Display, TEXT("-----------------------------------------------------------------------"));
 }
