Console: Added optional deprecation message for console variables.

Can be used to alter the otherwise automatically generated deprecation message, warning users about change of values for instance.

diff --git a/Engine/Source/Runtime/Core/Private/HAL/ConsoleManager.cpp b/Engine/Source/Runtime/Core/Private/HAL/ConsoleManager.cpp
--- a/Engine/Source/Runtime/Core/Private/HAL/ConsoleManager.cpp
+++ b/Engine/Source/Runtime/Core/Private/HAL/ConsoleManager.cpp
@@ -4814,9 +4814,10 @@ class FConsoleObjectShadowData
 {
 protected:
 	
-	FConsoleObjectShadowData(const TCHAR* CVarToShadow, const TCHAR* InDeprecatedVersion, EShadowCVarBehavior InLookupBehavior, EShadowCVarBehavior InUsageBehavior)
+	FConsoleObjectShadowData(const TCHAR* CVarToShadow, const TCHAR* InDeprecatedVersion, EShadowCVarBehavior InLookupBehavior, EShadowCVarBehavior InUsageBehavior, const TCHAR* InDeprecationMessage = nullptr)
 		: ShadowName(CVarToShadow)
 		, DeprecatedVersion(InDeprecatedVersion)
+		, DeprecationMessage(InDeprecationMessage)
 		, LookupBehavior(InLookupBehavior)
 		, UsageBehavior(InUsageBehavior)
 		, bHasLooked(false)
@@ -4877,13 +4878,20 @@ protected:
 					Arguments.Add(TEXT("ShadowName"), FText::FromString(ShadowName));
 
 					FText Message;
-					if (LookupBehavior == EShadowCVarBehavior::Assert)
+					if (DeprecationMessage.IsEmpty())
 					{
-						Message = FText::Format(LOCTEXT("FailedShadowCVarLookup_Assert", "Attempted to delay-load real CVar '{ThisName}' for shadowed CVar '{ShadowName}' failed."), Arguments);
+						if (LookupBehavior == EShadowCVarBehavior::Assert)
+						{
+							Message = FText::Format(LOCTEXT("FailedShadowCVarLookup_Assert", "Attempted to delay-load real CVar '{ThisName}' for shadowed CVar '{ShadowName}' failed."), Arguments);
+						}
+						else
+						{
+							Message = FText::Format(LOCTEXT("FailedShadowCVarLookup", "Attempted to delay-load real CVar '{ThisName}' for shadowed CVar '{ShadowName}' failed. Uses of '{ThisName}' will do nothing."), Arguments);
+						}
 					}
 					else
 					{
-						Message = FText::Format(LOCTEXT("FailedShadowCVarLookup", "Attempted to delay-load real CVar '{ThisName}' for shadowed CVar '{ShadowName}' failed. Uses of '{ThisName}' will do nothing."), Arguments);
+						Message = FText::FromString(DeprecationMessage);
 					}
 
 					switch (LookupBehavior)
@@ -4928,14 +4936,22 @@ protected:
 				Arguments.Add(TEXT("DeprecatedVersion"), FText::FromString(DeprecatedVersion));
 
 				FText Message;
-				if (DeprecatedVersion != TEXT(""))
+				if (DeprecationMessage.IsEmpty())
 				{
-					Message = FText::Format(LOCTEXT("ShadowCVarUsage_Deprecated", "Using a deprecated (as of UE {DeprecatedVersion}) CVar: '{ThisName}'. It will be removed in the future. Change all uses to '{ShadowName}' instead."), Arguments);
+					if (DeprecatedVersion != TEXT(""))
+					{
+						Message = FText::Format(LOCTEXT("ShadowCVarUsage_Deprecated", "Using a deprecated (as of UE {DeprecatedVersion}) CVar: '{ThisName}'. It will be removed in the future. Change all uses to '{ShadowName}' instead."), Arguments);
+					}
+					else
+					{
+						Message = FText::Format(LOCTEXT("ShadowCVarUsage", "Using a shadowed CVar '{ThisName}'. It is recommended to change all uses to '{ShadowName}' instead."), Arguments);
+					}
 				}
 				else
 				{
-					Message = FText::Format(LOCTEXT("ShadowCVarUsage", "Using a shadowed CVar '{ThisName}'. It is recommended to change all uses to '{ShadowName}' instead."), Arguments);
+					Message = FText::FromString(DeprecationMessage);
 				}
+
 				switch (UsageBehavior)
 				{
 				case EShadowCVarBehavior::Warn:
@@ -4961,6 +4977,7 @@ protected:
 
 	FString ShadowName;
 	FString DeprecatedVersion;
+	FString DeprecationMessage;
 	EShadowCVarBehavior LookupBehavior;
 	EShadowCVarBehavior UsageBehavior;
 	mutable uint8 bHasLooked : 1;
@@ -4972,8 +4989,8 @@ protected:
 class FConsoleVariableShadow : public FConsoleObjectShadowData, public IConsoleVariable
 {
 public:
-	FConsoleVariableShadow(const TCHAR* CVarToShadow, const TCHAR* InDeprecatedVersion, EShadowCVarBehavior InLookupBehavior, EShadowCVarBehavior InUsageBehavior)
-		: FConsoleObjectShadowData(CVarToShadow, InDeprecatedVersion, InLookupBehavior, InUsageBehavior)
+	FConsoleVariableShadow(const TCHAR* CVarToShadow, const TCHAR* InDeprecatedVersion, EShadowCVarBehavior InLookupBehavior, EShadowCVarBehavior InUsageBehavior, const TCHAR* InDeprecationMessage)
+		: FConsoleObjectShadowData(CVarToShadow, InDeprecatedVersion, InLookupBehavior, InUsageBehavior, InDeprecationMessage)
 		, RealVariable(nullptr)
 	{
 
@@ -5267,8 +5284,8 @@ private:
 struct FConsoleCommandShadow : public FConsoleObjectShadowData, public IConsoleCommand
 {
 
-	FConsoleCommandShadow(const TCHAR* CVarToShadow, const TCHAR* InDeprecatedVersion, EShadowCVarBehavior InLookupBehavior, EShadowCVarBehavior InUsageBehavior)
-		: FConsoleObjectShadowData(CVarToShadow, InDeprecatedVersion, InLookupBehavior, InUsageBehavior)
+	FConsoleCommandShadow(const TCHAR* CVarToShadow, const TCHAR* InDeprecatedVersion, EShadowCVarBehavior InLookupBehavior, EShadowCVarBehavior InUsageBehavior, const TCHAR* InDeprecationMessage)
+		: FConsoleObjectShadowData(CVarToShadow, InDeprecatedVersion, InLookupBehavior, InUsageBehavior, InDeprecationMessage)
 		, RealCommand(nullptr)
 	{
 
@@ -5360,24 +5377,24 @@ private:
 
 
 
-FAutoConsoleVariableShadow::FAutoConsoleVariableShadow(const TCHAR* Name, const TCHAR* CVarToShadow, EShadowCVarBehavior LookupFailureBehavior)
+FAutoConsoleVariableShadow::FAutoConsoleVariableShadow(const TCHAR* Name, const TCHAR* CVarToShadow, EShadowCVarBehavior LookupFailureBehavior, const TCHAR* DeprecationMessage)
 {
 #if !NO_CVARS
-	GetManager().AddShadowConsoleObject(Name, new FConsoleVariableShadow(CVarToShadow, nullptr, LookupFailureBehavior, EShadowCVarBehavior::NoMessaging));
+	GetManager().AddShadowConsoleObject(Name, new FConsoleVariableShadow(CVarToShadow, nullptr, LookupFailureBehavior, EShadowCVarBehavior::NoMessaging, DeprecationMessage));
 #endif
 }
 
-FAutoConsoleVariableDeprecated::FAutoConsoleVariableDeprecated(const TCHAR* Name, const TCHAR* CVarToShadow, const TCHAR* DeprecatedAtVersion, EShadowCVarBehavior UsageBehavior,  EShadowCVarBehavior LookupFailureBehavior)
+FAutoConsoleVariableDeprecated::FAutoConsoleVariableDeprecated(const TCHAR* Name, const TCHAR* CVarToShadow, const TCHAR* DeprecatedAtVersion, EShadowCVarBehavior UsageBehavior,  EShadowCVarBehavior LookupFailureBehavior, const TCHAR* DeprecationMessage)
 {
 #if !NO_CVARS
-	GetManager().AddShadowConsoleObject(Name, new FConsoleVariableShadow(CVarToShadow, DeprecatedAtVersion, LookupFailureBehavior, UsageBehavior));
+	GetManager().AddShadowConsoleObject(Name, new FConsoleVariableShadow(CVarToShadow, DeprecatedAtVersion, LookupFailureBehavior, UsageBehavior, DeprecationMessage));
 #endif
 }
 
-FAutoConsoleCommandDeprecated::FAutoConsoleCommandDeprecated(const TCHAR* Name, const TCHAR* CVarToShadow, const TCHAR* DeprecatedAtVersion, EShadowCVarBehavior UsageBehavior, EShadowCVarBehavior LookupFailureBehavior)
+FAutoConsoleCommandDeprecated::FAutoConsoleCommandDeprecated(const TCHAR* Name, const TCHAR* CVarToShadow, const TCHAR* DeprecatedAtVersion, EShadowCVarBehavior UsageBehavior, EShadowCVarBehavior LookupFailureBehavior, const TCHAR* DeprecationMessage)
 {
 #if !NO_CVARS
-	GetManager().AddShadowConsoleObject(Name, new FConsoleCommandShadow(CVarToShadow, DeprecatedAtVersion, LookupFailureBehavior, UsageBehavior));
+	GetManager().AddShadowConsoleObject(Name, new FConsoleCommandShadow(CVarToShadow, DeprecatedAtVersion, LookupFailureBehavior, UsageBehavior, DeprecationMessage));
 #endif
 }
 
diff --git a/Engine/Source/Runtime/Core/Public/HAL/IConsoleManager.h b/Engine/Source/Runtime/Core/Public/HAL/IConsoleManager.h
--- a/Engine/Source/Runtime/Core/Public/HAL/IConsoleManager.h
+++ b/Engine/Source/Runtime/Core/Public/HAL/IConsoleManager.h
@@ -2244,8 +2244,9 @@ public:
 		 * @param Name The name of this cvar
 		 * @param CVarToShadow The name of the existing cvar that this one will shadow
 		 * @param LookupFailureBehavior How to handle failure of looking up the CVarToShadow name - defaults to ensure
+		 * @param DeprecationMessage Optional deprecation message - defaults to empty, auto-generated.
 		 */
-		CORE_API FAutoConsoleVariableShadow(const TCHAR* Name, const TCHAR* CVarToShadow, EShadowCVarBehavior LookupFailureBehavior=EShadowCVarBehavior::Ensure);
+		CORE_API FAutoConsoleVariableShadow(const TCHAR* Name, const TCHAR* CVarToShadow, EShadowCVarBehavior LookupFailureBehavior=EShadowCVarBehavior::Ensure, const TCHAR* DeprecationMessage = TEXT(""));
 };
 
 /**
@@ -2265,8 +2266,9 @@ public:
 	 * @param DeprecatedAtVersion THe first UE version when this was deprecated
 	 * @param UsageBehavior How to handle uses of the shadowed cvar by the deprecated name - defaults to ensure
 	 * @param LookupFailureBehavior How to handle failure of looking up the CVarToShadow name - defaults to ensure
+	 * @param DeprecationMessage Optional deprecation message - defaults to empty, auto-generated.
 	 */
-	CORE_API FAutoConsoleVariableDeprecated(const TCHAR* Name, const TCHAR* CVarToShadow, const TCHAR* DeprecatedAtVersion, EShadowCVarBehavior UsageBehavior=EShadowCVarBehavior::Ensure, EShadowCVarBehavior LookupFailureBehavior = EShadowCVarBehavior::Ensure);
+	CORE_API FAutoConsoleVariableDeprecated(const TCHAR* Name, const TCHAR* CVarToShadow, const TCHAR* DeprecatedAtVersion, EShadowCVarBehavior UsageBehavior=EShadowCVarBehavior::Ensure, EShadowCVarBehavior LookupFailureBehavior = EShadowCVarBehavior::Ensure, const TCHAR* DeprecationMessage = TEXT(""));
 };
 
 /**
@@ -2281,8 +2283,9 @@ public:
 	 * @param DeprecatedAtVersion THe first UE version when this was deprecated
 	 * @param UsageBehavior How to handle uses of the shadowed cvar by the deprecated name - defaults to ensure
 	 * @param LookupFailureBehavior How to handle failure of looking up the CVarToShadow name - defaults to ensure
+	 * @param DeprecationMessage Optional deprecation message - defaults to empty, auto-generated.
 	 */
-	CORE_API FAutoConsoleCommandDeprecated(const TCHAR* Name, const TCHAR* CVarToShadow, const TCHAR* DeprecatedAtVersion, EShadowCVarBehavior UsageBehavior = EShadowCVarBehavior::Ensure, EShadowCVarBehavior LookupFailureBehavior = EShadowCVarBehavior::Ensure);
+	CORE_API FAutoConsoleCommandDeprecated(const TCHAR* Name, const TCHAR* CVarToShadow, const TCHAR* DeprecatedAtVersion, EShadowCVarBehavior UsageBehavior = EShadowCVarBehavior::Ensure, EShadowCVarBehavior LookupFailureBehavior = EShadowCVarBehavior::Ensure, const TCHAR* DeprecationMessage = TEXT(""));
 };
 
 
