[AutoRTFM] Add TBlockAllocator

A single-threaded memory allocator that allocates out of fixed size blocks of memory.
All allocated memory is freed when the TBlockAllocator is destructed. There is no way to free individual allocations.

This will be used to optimize the TTaskArray by reducing the number of heap allocations.

diff --git a/Engine/Source/Programs/AutoRTFMTests/Private/BlockAllocatorTests.cpp b/Engine/Source/Programs/AutoRTFMTests/Private/BlockAllocatorTests.cpp
new file mode 100644
--- /dev/null
+++ b/Engine/Source/Programs/AutoRTFMTests/Private/BlockAllocatorTests.cpp
@@ -0,0 +1,88 @@
+// Copyright Epic Games, Inc. All Rights Reserved.
+
+#include "AutoRTFM/BlockAllocator.h"
+#include "Catch2Includes.h"
+#include <cstddef>
+
+TEST_CASE("BlockAllocator")
+{
+	static constexpr size_t AllocatorBlockSize = 128;
+	static constexpr size_t AllocatorAlignment = 16;
+	AutoRTFM::TBlockAllocator<AllocatorBlockSize, AllocatorAlignment> BlockAllocator;
+
+	struct FAllocationInfo
+	{
+		// The pointer returned by Allocate().
+		std::byte* Ptr;
+		// The allocation size in bytes.
+		size_t Size;
+		// The value written into the allocated memory.
+		std::byte Value;
+	};
+	std::vector<FAllocationInfo> Allocations;
+
+	// Helper that calls BlockAllocator.Allocate(Size, Alignment) and returns
+	// the allocated pointer cast to std::byte*. The allocated memory is
+	// memset() with a value, and the allocation pointer, size and memset()
+	// value is recorded into Allocations so the memory can be checked at the
+	// end of the test. The allocated pointer is checked for the requested
+	// alignment.
+	auto Allocate = [&](size_t Size, size_t Alignment)
+	{
+		std::byte* Ptr = static_cast<std::byte*>(BlockAllocator.Allocate(Size, Alignment));
+		REQUIRE(reinterpret_cast<uintptr_t>(Ptr) % Alignment == 0);
+		std::byte Value = static_cast<std::byte>(Allocations.size());
+		memset(Ptr, static_cast<int>(Value), Size);
+		Allocations.push_back(FAllocationInfo{Ptr, Size, Value});
+		return static_cast<std::byte*>(Ptr);
+	};
+
+	SECTION("Contiguous")
+	{
+		size_t AllocationSize = AllocatorBlockSize / 8;
+		size_t AllocationAlignment = 4;
+
+		std::byte* FirstInBlock = Allocate(AllocationSize, AllocationAlignment);
+		for (size_t BlockIndex = 0; BlockIndex < 4; BlockIndex++)
+		{
+			size_t Offset = AllocationSize;
+			for (; Offset < AllocatorBlockSize; Offset += AllocationSize)
+			{
+				std::byte* Allocation = Allocate(AllocationSize, AllocationAlignment);
+				REQUIRE(Allocation == FirstInBlock + Offset);
+			}
+			std::byte* NextBlock = Allocate(AllocationSize, AllocationAlignment);
+			REQUIRE(NextBlock != FirstInBlock + Offset);
+			FirstInBlock = NextBlock;
+		}
+	}
+
+	SECTION("Size")
+	{
+		for (size_t I = 0; I < 1000; I++)
+		{
+			size_t Size = std::max<size_t>(1, (I * 31) % AllocatorBlockSize);
+			Allocate(Size, /* Alignment */ 4);
+		}
+	}
+
+	SECTION("Alignment")
+	{
+		for (size_t Alignment = 1; Alignment <= AllocatorAlignment; Alignment *= 2)
+		{
+			for (size_t I = 0; I < 100; I++)
+			{
+				Allocate(/* Size */ 4, Alignment);
+			}
+		}
+	}
+
+	// Check that the allocated memory is as expected.
+	for (FAllocationInfo AllocationInfo : Allocations)
+	{
+		for (size_t Offset = 0; Offset < AllocationInfo.Size; Offset++)
+		{
+			REQUIRE(AllocationInfo.Ptr[Offset] == AllocationInfo.Value);
+		}
+	}
+}
diff --git a/Engine/Source/Runtime/Core/Private/AutoRTFM/BlockAllocator.h b/Engine/Source/Runtime/Core/Private/AutoRTFM/BlockAllocator.h
new file mode 100644
--- /dev/null
+++ b/Engine/Source/Runtime/Core/Private/AutoRTFM/BlockAllocator.h
@@ -0,0 +1,119 @@
+// Copyright Epic Games, Inc. All Rights Reserved.
+
+#pragma once
+
+#include "AutoRTFM/Utils.h"
+#include "HAL/Platform.h"
+#include "HAL/UnrealMemory.h"
+#include <cstddef>
+
+namespace AutoRTFM
+{
+
+// A single-threaded memory allocator that allocates out of fixed size blocks of
+// memory. All allocated memory is freed when the TBlockAllocator is destructed.
+// There is no way to free individual allocations.
+//
+// Template parameters:
+//   BlockDataSize - The size of the data section of each block, which is also
+//                   the maximum size that can be allocated using Allocate().
+//   DataAlignment - The maximum supported alignment of data from this allocator.
+template<size_t BlockDataSize, size_t DataAlignment = 16>
+class TBlockAllocator
+{
+public:
+	// Constructor
+	TBlockAllocator() = default;
+
+	// Destructor - frees all the memory allocated by the allocator.
+	~TBlockAllocator()
+	{
+		FreeAll();
+	}
+
+	// Frees all allocations made by the allocator.
+	void FreeAll()
+	{
+		while (Tail)
+		{
+			FBlockHeader* Prev = Tail->Prev;
+			FMemory::Free(Tail);
+			Tail = Prev;
+		}
+	}
+
+	// Allocates memory from the block allocator.
+	// Size must be no larger than BlockSize.
+	// Alignment must be a power of two and no larger than DataAlignment.
+	inline void* Allocate(size_t Size, size_t Alignment)
+	{
+		ASSERT(Size <= BlockDataSize);
+		ASSERT(Alignment <= DataAlignment);
+
+		void* Allocation = Tail ? Tail->TryAllocate(Size, Alignment) : nullptr;
+		if (LIKELY(Allocation))
+		{
+			return Allocation;
+		}
+
+		FBlockHeader* NewBlock = FBlockHeader::New();
+		NewBlock->Prev = Tail;
+		Tail = NewBlock;
+
+		Allocation = NewBlock->TryAllocate(Size, Alignment);
+		ASSERT(Allocation);
+		return Allocation;
+	}
+
+	// Constructs and returns a new T into the memory returned by calling
+	// Allocate(sizeof(T), alignof(T)).
+	template<typename T, typename... ArgTypes> T* New(ArgTypes&&... Args)
+	{
+		return new (Allocate(sizeof(T), alignof(T))) T(Forward<ArgTypes>(Args)...);
+	}
+
+private:
+	UE_NONCOPYABLE(TBlockAllocator);
+
+	static inline size_t AlignDown(size_t Value, size_t Alignment)
+	{
+		return Value & ~(Alignment-1);
+	}
+
+	struct alignas(std::max<size_t>(DataAlignment, 8)) FBlockHeader
+	{
+		// The previous FBlock in the singly linked-list
+		FBlockHeader* Prev = nullptr;
+		// The number of unallocated bytes remaining in the block
+		size_t Remaining = BlockDataSize;
+		// <padding>
+		// <data>
+
+		// Allocates, constructs and returns a new block from FMemory::Malloc.
+		static FBlockHeader* New()
+		{
+			void* Memory = FMemory::Malloc(sizeof(FBlockHeader) + BlockDataSize);
+			return new (Memory) FBlockHeader;
+		}
+
+		// Attempts to sub-allocate with out of this block. Returns a pointer to
+		// the sub-allocated memory on success, or nullptr on failure.
+		inline void* TryAllocate(size_t Size, size_t Alignment)
+		{
+			size_t RemainingAligned = AlignDown(Remaining, Alignment);
+			if (RemainingAligned < Size)
+			{
+				return nullptr;
+			}
+			void* Ptr = reinterpret_cast<std::byte*>(this) + sizeof(FBlockHeader) + BlockDataSize - RemainingAligned;
+			Remaining = RemainingAligned - Size;
+			return Ptr;
+		}
+	};
+
+	static_assert(sizeof(FBlockHeader) % DataAlignment == 0);
+
+	FBlockHeader* Tail = nullptr;
+};
+
+}
