Fixed Horde Issue 747394: Compile errors
- turned ECustomVersionSerializationFormat from a struct into an enum class

#rnx

diff --git a/Engine/Source/Developer/DirectoryWatcher/Private/FileCache.cpp b/Engine/Source/Developer/DirectoryWatcher/Private/FileCache.cpp
--- a/Engine/Source/Developer/DirectoryWatcher/Private/FileCache.cpp
+++ b/Engine/Source/Developer/DirectoryWatcher/Private/FileCache.cpp
@@ -17,7 +17,7 @@ namespace DirectoryWatcher
 {
 
 template<typename T>
-void ReadWithCustomVersions(FArchive& Ar, T& Data, ECustomVersionSerializationFormat::Type CustomVersionFormat)
+void ReadWithCustomVersions(FArchive& Ar, T& Data, ECustomVersionSerializationFormat CustomVersionFormat)
 {
 	int64 CustomVersionsOffset = 0;
 	Ar << CustomVersionsOffset;
@@ -76,7 +76,7 @@ FCustomVersionRegistration GRegisterFileCacheVersion(FFileCacheCustomVersion::Ke
 static const uint32 CacheFileMagicNumberOldCustomVersionFormat = 0x03DCCB00;
 static const uint32 CacheFileMagicNumber = 0x03DCCB03;
 
-static ECustomVersionSerializationFormat::Type GetCustomVersionFormatForFileCache(uint32 MagicNumber)
+static ECustomVersionSerializationFormat GetCustomVersionFormatForFileCache(uint32 MagicNumber)
 {
 	if (MagicNumber == CacheFileMagicNumberOldCustomVersionFormat)
 	{
diff --git a/Engine/Source/Runtime/Core/Private/Serialization/CustomVersion.cpp b/Engine/Source/Runtime/Core/Private/Serialization/CustomVersion.cpp
--- a/Engine/Source/Runtime/Core/Private/Serialization/CustomVersion.cpp
+++ b/Engine/Source/Runtime/Core/Private/Serialization/CustomVersion.cpp
@@ -312,12 +312,12 @@ void operator<<(FStructuredArchive::FSlot Slot, FCustomVersion& Version)
 	Record << SA_VALUE(TEXT("Version"), Version.Version);
 }
 
-void FCustomVersionContainer::Serialize(FArchive& Ar, ECustomVersionSerializationFormat::Type Format)
+void FCustomVersionContainer::Serialize(FArchive& Ar, ECustomVersionSerializationFormat Format)
 {
 	Serialize(FStructuredArchiveFromArchive(Ar).GetSlot(), Format);
 }
 
-void FCustomVersionContainer::Serialize(FStructuredArchive::FSlot Slot, ECustomVersionSerializationFormat::Type Format)
+void FCustomVersionContainer::Serialize(FStructuredArchive::FSlot Slot, ECustomVersionSerializationFormat Format)
 {
 	switch (Format)
 	{
diff --git a/Engine/Source/Runtime/Core/Public/Serialization/CustomVersion.h b/Engine/Source/Runtime/Core/Public/Serialization/CustomVersion.h
--- a/Engine/Source/Runtime/Core/Public/Serialization/CustomVersion.h
+++ b/Engine/Source/Runtime/Core/Public/Serialization/CustomVersion.h
@@ -18,19 +18,16 @@ class FArchive;
 class FCustomVersionContainer;
 struct FCustomVersion;
 
-struct ECustomVersionSerializationFormat
+enum class ECustomVersionSerializationFormat : uint8
 {
-	enum Type
-	{
-		Unknown,
-		Guids,
-		Enums,
-		Optimized,
-
-		// Add new versions above this comment
-		CustomVersion_Automatic_Plus_One,
-		Latest = CustomVersion_Automatic_Plus_One - 1
-	};
+	Unknown,
+	Guids,
+	Enums,
+	Optimized,
+
+	// Add new versions above this comment
+	CustomVersion_Automatic_Plus_One,
+	Latest = CustomVersion_Automatic_Plus_One - 1
 };
 
 typedef TArray<FCustomVersion> FCustomVersionArray;
@@ -154,8 +151,8 @@ public:
 	CORE_API void SetVersionUsingRegistry(FGuid CustomKey, ESetCustomVersionFlags Options = ESetCustomVersionFlags::None);
 
 	/** Serialization. */
-	CORE_API void Serialize(FArchive& Ar, ECustomVersionSerializationFormat::Type Format = ECustomVersionSerializationFormat::Latest);
-	CORE_API void Serialize(FStructuredArchive::FSlot Slot, ECustomVersionSerializationFormat::Type Format = ECustomVersionSerializationFormat::Latest);
+	CORE_API void Serialize(FArchive& Ar, ECustomVersionSerializationFormat Format = ECustomVersionSerializationFormat::Latest);
+	CORE_API void Serialize(FStructuredArchive::FSlot Slot, ECustomVersionSerializationFormat Format = ECustomVersionSerializationFormat::Latest);
 
 	/**
 	 * Gets a singleton with the registered versions.
diff --git a/Engine/Source/Runtime/CoreUObject/Private/UObject/PackageFileSummary.cpp b/Engine/Source/Runtime/CoreUObject/Private/UObject/PackageFileSummary.cpp
--- a/Engine/Source/Runtime/CoreUObject/Private/UObject/PackageFileSummary.cpp
+++ b/Engine/Source/Runtime/CoreUObject/Private/UObject/PackageFileSummary.cpp
@@ -15,9 +15,9 @@ FPackageFileSummary::FPackageFileSummary()
 }
 
 /** Converts file version to custom version system version */
-static ECustomVersionSerializationFormat::Type GetCustomVersionFormatForArchive(int32 LegacyFileVersion)
+static ECustomVersionSerializationFormat GetCustomVersionFormatForArchive(int32 LegacyFileVersion)
 {
-	ECustomVersionSerializationFormat::Type CustomVersionFormat = ECustomVersionSerializationFormat::Unknown;
+	ECustomVersionSerializationFormat CustomVersionFormat = ECustomVersionSerializationFormat::Unknown;
 	if (LegacyFileVersion == -2)
 	{
 		CustomVersionFormat = ECustomVersionSerializationFormat::Enums;
diff --git a/Engine/Source/Runtime/Engine/Private/GameplayStatics.cpp b/Engine/Source/Runtime/Engine/Private/GameplayStatics.cpp
--- a/Engine/Source/Runtime/Engine/Private/GameplayStatics.cpp
+++ b/Engine/Source/Runtime/Engine/Private/GameplayStatics.cpp
@@ -200,7 +200,7 @@ void FSaveGameHeader::Read(FMemoryReader& MemoryReader)
 		{
 			MemoryReader << CustomVersionFormat;
 
-			CustomVersions.Serialize(MemoryReader, static_cast<ECustomVersionSerializationFormat::Type>(CustomVersionFormat));
+			CustomVersions.Serialize(MemoryReader, static_cast<ECustomVersionSerializationFormat>(CustomVersionFormat));
 			MemoryReader.SetCustomVersions(CustomVersions);
 		}
 
@@ -228,7 +228,7 @@ void FSaveGameHeader::Write(FMemoryWriter& MemoryWriter)
 
 	// Write out custom version data
 	MemoryWriter << CustomVersionFormat;
-	CustomVersions.Serialize(MemoryWriter, static_cast<ECustomVersionSerializationFormat::Type>(CustomVersionFormat));
+	CustomVersions.Serialize(MemoryWriter, static_cast<ECustomVersionSerializationFormat>(CustomVersionFormat));
 
 	// Write the class name so we know what class to load to
 	MemoryWriter << SaveGameClassName;
diff --git a/Engine/Source/Runtime/Experimental/Chaos/Private/ChaosVisualDebugger/ChaosVDMemWriterReader.cpp b/Engine/Source/Runtime/Experimental/Chaos/Private/ChaosVisualDebugger/ChaosVDMemWriterReader.cpp
--- a/Engine/Source/Runtime/Experimental/Chaos/Private/ChaosVisualDebugger/ChaosVDMemWriterReader.cpp
+++ b/Engine/Source/Runtime/Experimental/Chaos/Private/ChaosVisualDebugger/ChaosVDMemWriterReader.cpp
@@ -10,7 +10,7 @@ namespace Chaos::VisualDebugger
 
 		uint8 CustomVersionFormat = static_cast<uint8>(ECustomVersionSerializationFormat::Latest);
 		Ar << CustomVersionFormat;
-		CustomVersionContainer.Serialize(Ar, static_cast<ECustomVersionSerializationFormat::Type>(CustomVersionFormat));
+		CustomVersionContainer.Serialize(Ar, static_cast<ECustomVersionSerializationFormat>(CustomVersionFormat));
 
 		return !Ar.IsError();
 	}
