fix CIS issue

diff --git a/Engine/Source/Runtime/Core/Private/Android/AndroidProfiler.cpp b/Engine/Source/Runtime/Core/Private/Android/AndroidProfiler.cpp
--- a/Engine/Source/Runtime/Core/Private/Android/AndroidProfiler.cpp
+++ b/Engine/Source/Runtime/Core/Private/Android/AndroidProfiler.cpp
@@ -197,7 +197,7 @@ void FAndroidProfiler::StopCapture(const FString& ProfileName)
 #else
 
 FString FAndroidProfiler::StartCapture(const FString& Args, TUniqueFunction<void(const FAndroidProfiler::FProfileResults& Results)> OnFinish) { return FString(); }
-void FAndroidProfiler::StopCapture(FString ProfileName) {}
+void FAndroidProfiler::StopCapture(const FString& ProfileName) {}
 void FAndroidProfiler::OnProfileFinish(const FAndroidProfiler::FProfileResults& Results) {}
 
 #endif
