Support clearing crash report on upload on all Windows builds

#rnx

diff --git a/Engine/Source/Runtime/Core/Private/Windows/WindowsPlatformCrashContext.cpp b/Engine/Source/Runtime/Core/Private/Windows/WindowsPlatformCrashContext.cpp
--- a/Engine/Source/Runtime/Core/Private/Windows/WindowsPlatformCrashContext.cpp
+++ b/Engine/Source/Runtime/Core/Private/Windows/WindowsPlatformCrashContext.cpp
@@ -492,6 +492,12 @@ FProcHandle LaunchCrashReportClient(void** OutWritePipe, void** OutReadPipe, uin
 				{
 					FCString::Strncat(CrashReporterClientArgs, TEXT(" -unattended"), CR_CLIENT_MAX_ARGS_LEN);
 				}
+
+				// Whether to clean up crash reports after send
+				if (FParse::Param(ArgList[It], TEXT("CleanCrashReports")))
+				{
+					FCString::Strncat(CrashReporterClientArgs, TEXT(" -CleanCrashReports"), CR_CLIENT_MAX_ARGS_LEN);
+				}
 			}
 		}
 
