[Core] Cleanup in set.h
* Remove redundant headers
* Remove friend from FSetElementId
* Add [[nodiscard]] to FindId

diff --git a/Engine/Source/Runtime/Core/Public/Containers/Set.h b/Engine/Source/Runtime/Core/Public/Containers/Set.h
--- a/Engine/Source/Runtime/Core/Public/Containers/Set.h
+++ b/Engine/Source/Runtime/Core/Public/Containers/Set.h
@@ -2,13 +2,11 @@
 
 #pragma once
 
-#include "Containers/Array.h"
 #include "Containers/ContainerAllocationPolicies.h"
 #include "Containers/ContainerElementTypeCompatibility.h"
 #include "Containers/SetUtilities.h"
 #include "Containers/SparseArray.h"
 #include "ContainersFwd.h"
-#include "CoreTypes.h"
 #include "Math/UnrealMathUtility.h"
 #include "Misc/AssertionMacros.h"
 #include "Misc/StructBuilder.h"
@@ -19,7 +17,6 @@
 #include "Templates/Sorting.h"
 #include "Templates/TypeHash.h"
 #include "Templates/UnrealTemplate.h"
-#include "Templates/UnrealTypeTraits.h"
 #include <initializer_list>
 #include <type_traits>
 
@@ -463,7 +460,7 @@ public:
 	 */
 	[[nodiscard]] FORCEINLINE bool IsValidId(FSetElementId Id) const
 	{
-		SizeType Index = Id.Index;
+		SizeType Index = Id.AsInteger();
 		return Index != INDEX_NONE &&
 			Index >= 0 &&
 			Index < Elements.GetMaxIndex() &&
@@ -473,25 +470,25 @@ public:
 	/** Accesses the identified element's value. Element must be valid (see @IsValidId). */
 	[[nodiscard]] FORCEINLINE ElementType& operator[](FSetElementId Id)
 	{
-		return Elements[Id.Index].Value;
+		return Elements[Id.AsInteger()].Value;
 	}
 
 	/** Accesses the identified element's value. Element must be valid (see @IsValidId). */
 	[[nodiscard]] FORCEINLINE const ElementType& operator[](FSetElementId Id) const
 	{
-		return Elements[Id.Index].Value;
+		return Elements[Id.AsInteger()].Value;
 	}
 
 	/** Accesses the identified element's value. Element must be valid (see @IsValidId). */
 	[[nodiscard]] FORCEINLINE ElementType& Get(FSetElementId Id)
 	{
-		return Elements[Id.Index].Value;
+		return Elements[Id.AsInteger()].Value;
 	}
 
 	/** Accesses the identified element's value. Element must be valid (see @IsValidId). */
 	[[nodiscard]] FORCEINLINE const ElementType& Get(FSetElementId Id) const
 	{
-		return Elements[Id.Index].Value;
+		return Elements[Id.AsInteger()].Value;
 	}
 
 	/**
@@ -637,7 +634,7 @@ public:
 		{
 			RehashOrLink(KeyHash, Element, NewHashIndex);
 		}
-		return FSetElementId(NewHashIndex);
+		return FSetElementId::FromInteger(NewHashIndex);
 	}
 	
 	/**
@@ -661,7 +658,7 @@ public:
 		{
 			RehashOrLink(KeyHash, Element, NewHashIndex);
 		}
-		return FSetElementId(NewHashIndex);
+		return FSetElementId::FromInteger(NewHashIndex);
 	}
 
 	void Append(TArrayView<const ElementType> InElements)
@@ -727,19 +724,19 @@ private:
 
 		// Remove the element from the hash.
 		FSetElementId* HashPtr              = Hash.GetAllocation();
-		SizeType*      NextElementIndexIter = &HashPtr[ElementBeingRemoved.HashIndex].Index;
+		FSetElementId* NextElementIndexIter = &HashPtr[ElementBeingRemoved.HashIndex];
 		for (;;)
 		{
-			SizeType NextElementIndex = *NextElementIndexIter;
+			SizeType NextElementIndex = NextElementIndexIter->AsInteger();
 			checkf(NextElementIndex != INDEX_NONE, TEXT("Corrupt hash"));
 
 			if (NextElementIndex == ElementIndex)
 			{
-				*NextElementIndexIter = ElementBeingRemoved.HashNextId.Index;
+				*NextElementIndexIter = ElementBeingRemoved.HashNextId;
 				break;
 			}
 
-			NextElementIndexIter = &Elements[NextElementIndex].HashNextId.Index;
+			NextElementIndexIter = &Elements[NextElementIndex].HashNextId;
 		}
 
 		// Remove the element from the elements array.
@@ -753,7 +750,7 @@ public:
 	 */
 	void Remove(FSetElementId ElementId)
 	{
-		RemoveByIndex(ElementId.Index);
+		RemoveByIndex(ElementId.AsInteger());
 	}
 
 private:
@@ -771,7 +768,7 @@ private:
 		}
 
 		FSetElementId* HashPtr      = Hash.GetAllocation();
-		SizeType       ElementIndex = HashPtr[KeyHash & (HashSize - 1)].Index;
+		SizeType       ElementIndex = HashPtr[KeyHash & (HashSize - 1)].AsInteger();
 		for (;;)
 		{
 			if (ElementIndex == INDEX_NONE)
@@ -785,7 +782,7 @@ private:
 				return ElementIndex;
 			}
 
-			ElementIndex = Elements[ElementIndex].HashNextId.Index;
+			ElementIndex = Elements[ElementIndex].HashNextId.AsInteger();
 		}
 	}
 
@@ -812,9 +809,9 @@ public:
 	 * @param Key - The key to search for.
 	 * @return The id of the set element matching the given key, or the NULL id if none matches.
 	 */
-	FSetElementId FindId(KeyInitType Key) const
+	[[nodiscard]] FSetElementId FindId(KeyInitType Key) const
 	{
-		return FSetElementId(FindIndexByHash(KeyFuncs::GetKeyHash(Key), Key));
+		return FSetElementId::FromInteger(FindIndexByHash(KeyFuncs::GetKeyHash(Key), Key));
 	}
 
 	/**
@@ -830,7 +827,7 @@ public:
 		checkSlow(KeyHash == KeyFuncs::GetKeyHash(Key));
 		PRAGMA_ENABLE_DEPRECATION_WARNINGS
 
-		return FSetElementId(FindIndexByHash(KeyHash, Key));
+		return FSetElementId::FromInteger(FindIndexByHash(KeyHash, Key));
 	}
 
 	/**
@@ -895,13 +892,14 @@ private:
 		FSetElementId* NextElementId = &GetTypedHash(KeyHash);
 		while (NextElementId->IsValidId())
 		{
-			SetElementType& Element = Elements[NextElementId->Index];
+			const int32 ElementIndex = NextElementId->AsInteger();
+			SetElementType& Element = Elements[ElementIndex];
 
 			if (KeyFuncs::Matches(KeyFuncs::GetSetKey(Element.Value), Key))
 			{
 				// This element matches the key, remove it from the set.  Note that Remove sets *NextElementId to point to the next
 				// element after the removed element in the hash bucket.
-				RemoveByIndex(NextElementId->Index);
+				RemoveByIndex(ElementIndex);
 				NumRemovedElements++;
 
 				if constexpr (!KeyFuncs::bAllowDuplicateKeys)
@@ -1033,7 +1031,7 @@ public:
 			int32 NumElementsInBucket = 0;
 			for(FSetElementId ElementId = GetTypedHash(HashIndex);
 				ElementId.IsValidId();
-				ElementId = Elements[ElementId.Index].HashNextId)
+				ElementId = Elements[ElementId.AsInteger()].HashNextId)
 			{
 				NumElementsInBucket++;
 			}
@@ -1057,7 +1055,7 @@ public:
 					bResult=false;
 					break;
 				}
-				ElementId = Elements[ElementId.Index].HashNextId;
+				ElementId = Elements[ElementId.AsInteger()].HashNextId;
 			}
 		}
 		return bResult;
@@ -1076,11 +1074,11 @@ public:
 			{
 				if( !IsValidId(ElementId) )
 				{
-					Ar.Logf(TEXT("		!!INVALID!! ElementId = %d"),ElementId.Index);
+					Ar.Logf(TEXT("		!!INVALID!! ElementId = %d"),ElementId.AsInteger());
 				}
 				else
 				{
-					Ar.Logf(TEXT("		VALID ElementId = %d"),ElementId.Index);
+					Ar.Logf(TEXT("		VALID ElementId = %d"),ElementId.AsInteger());
 				}
 				ElementId = Elements[ElementId].HashNextId;
 			}
@@ -1354,8 +1352,9 @@ private:
 		Element.HashIndex = KeyHash & (HashSize - 1);
 
 		// Link the element into the hash bucket.
-		Element.HashNextId.Index = GetTypedHash(Element.HashIndex).Index;
-		GetTypedHash(Element.HashIndex).Index = ElementIndex;
+		FSetElementId& TypedHash = GetTypedHash(Element.HashIndex);
+		Element.HashNextId = TypedHash;
+		TypedHash = FSetElementId::FromInteger(ElementIndex);
 	}
 
 	/** Hashes and links an added element to the hash chain. */
@@ -1375,7 +1374,7 @@ private:
 			// Faster path: only reset hash buckets to invalid for elements in the hash
 			for (const SetElementType& Element: Elements)
 			{
-				HashPtr[Element.HashIndex].Index = INDEX_NONE;
+				HashPtr[Element.HashIndex] = FSetElementId();
 			}
 		}
 		else
@@ -1383,7 +1382,7 @@ private:
 			/** Reset the whole hash to invalid */
 			for (int32 I = 0; I < HashSize; ++I)
 			{
-				HashPtr[I].Index = INDEX_NONE;
+				HashPtr[I] = FSetElementId();
 			}
 		}
 	}
@@ -1490,7 +1489,7 @@ private:
 		// Accessors.
 		[[nodiscard]] FORCEINLINE FSetElementId GetId() const
 		{
-			return FSetElementId(ElementIt.GetIndex());
+			return FSetElementId::FromInteger(ElementIt.GetIndex());
 		}
 		[[nodiscard]] FORCEINLINE ItElementType* operator->() const
 		{
@@ -1532,7 +1531,7 @@ private:
 		{
 			if (Set.HashSize)
 			{
-				NextIndex = Set.GetTypedHash(KeyFuncs::GetKeyHash(Key)).Index;
+				NextIndex = Set.GetTypedHash(KeyFuncs::GetKeyHash(Key)).AsInteger();
 				++(*this);
 			}
 			else
@@ -1548,7 +1547,7 @@ private:
 
 			while (Index != INDEX_NONE)
 			{
-				NextIndex = Set.Elements[Index].HashNextId.Index;
+				NextIndex = Set.Elements[Index].HashNextId.AsInteger();
 				checkSlow(Index != NextIndex);
 
 				if (KeyFuncs::Matches(KeyFuncs::GetSetKey(Set.Elements[Index].Value),Key))
@@ -1575,7 +1574,7 @@ private:
 		// Accessors.
 		[[nodiscard]] FORCEINLINE FSetElementId GetId() const
 		{
-			return FSetElementId(Index);
+			return FSetElementId::FromInteger(Index);
 		}
 		[[nodiscard]] FORCEINLINE ItElementType* operator->() const
 		{
@@ -1861,7 +1860,7 @@ public:
 		FSetElementId* HashPtr = Hash.GetAllocation();
 		for (int32 I = 0; I < HashSize; ++I)
 		{
-			HashPtr[I].Index = INDEX_NONE;
+			HashPtr[I] = FSetElementId();
 		}
 	}
 
@@ -1872,9 +1871,9 @@ public:
 		void* ElementBeingRemoved = Elements.GetData(Index, Layout.SparseArrayLayout);
 
 		// Remove the element from the hash.
-		for (FSetElementId* NextElementId = &GetTypedHash(GetHashIndexRef(ElementBeingRemoved, Layout)); NextElementId->IsValidId(); NextElementId = &GetHashNextIdRef(Elements.GetData(NextElementId->Index, Layout.SparseArrayLayout), Layout))
+		for (FSetElementId* NextElementId = &GetTypedHash(GetHashIndexRef(ElementBeingRemoved, Layout)); NextElementId->IsValidId(); NextElementId = &GetHashNextIdRef(Elements.GetData(NextElementId->AsInteger(), Layout.SparseArrayLayout), Layout))
 		{
-			if (NextElementId->Index == Index)
+			if (NextElementId->AsInteger() == Index)
 			{
 				*NextElementId = GetHashNextIdRef(ElementBeingRemoved, Layout);
 				break;
@@ -1920,7 +1919,7 @@ public:
 			{
 				if (Elements.IsValidIndex(Index))
 				{
-					FSetElementId ElementId(Index);
+					FSetElementId ElementId = FSetElementId::FromInteger(Index);
 
 					void* Element = (uint8*)Elements.GetData(Index, Layout.SparseArrayLayout);
 
@@ -1951,10 +1950,11 @@ private:
 			ElementId.IsValidId();
 			ElementId = GetHashNextIdRef(CurrentElement, Layout))
 		{
-			CurrentElement = (uint8*)Elements.GetData(ElementId.Index, Layout.SparseArrayLayout);
+			const int32 ElementIndex = ElementId.AsInteger();
+			CurrentElement = (uint8*)Elements.GetData(ElementIndex, Layout.SparseArrayLayout);
 			if (EqualityFn(Element, CurrentElement))
 			{
-				return ElementId.Index;
+				return ElementIndex;
 			}
 		}
 	
@@ -2040,7 +2040,7 @@ private:
 			FSetElementId& TypedHash = GetTypedHash(HashIndex);
 			GetHashIndexRef(ElementPtr, Layout) = HashIndex;
 			GetHashNextIdRef(ElementPtr, Layout) = TypedHash;
-			TypedHash = FSetElementId(NewElementIndex);
+			TypedHash = FSetElementId::FromInteger(NewElementIndex);
 		}
 
 		return NewElementIndex;
diff --git a/Engine/Source/Runtime/Core/Public/Containers/SetUtilities.h b/Engine/Source/Runtime/Core/Public/Containers/SetUtilities.h
--- a/Engine/Source/Runtime/Core/Public/Containers/SetUtilities.h
+++ b/Engine/Source/Runtime/Core/Public/Containers/SetUtilities.h
@@ -88,13 +88,6 @@ FORCEINLINE void MoveByRelocate(T& A, T& B)
 class FSetElementId
 {
 public:
-
-	template<typename,typename,typename>
-	friend class TSet;
-
-	template <typename, typename>
-	friend class TScriptSet;
-
 	/** Default constructor. */
 	FORCEINLINE FSetElementId() = default;
 
@@ -105,7 +98,7 @@ public:
 	}
 
 	/** Comparison operators. */
-	[[nodiscard]] FORCEINLINE friend bool operator==(const FSetElementId& A,const FSetElementId& B)
+	[[nodiscard]] FORCEINLINE friend bool operator==(const FSetElementId& A, const FSetElementId& B)
 	{
 		return A.Index == B.Index;
 	}
@@ -136,4 +129,4 @@ private:
 };
 
 // This is just an int32
-DECLARE_INTRINSIC_TYPE_LAYOUT(FSetElementId);
+DECLARE_INTRINSIC_TYPE_LAYOUT(FSetElementId);
\ No newline at end of file
