[AutoRTFM] Migrate tasks from TFunction to AutoRTFM::TTask

TTask is a small std::function like class that can wrap a functor (lambda, etc).
Use TTask for callbacks such as OnCommit and OnAbort handlers, which there may be many accumulating throughout a transaction.
For this reason TTask is designed to be compact in size, but large enough to fit the most common functors used by AutoRTFM tasks.

Provides up to 1.09x speed improvements based on benchmarking.

diff --git a/Engine/Source/Runtime/Core/Private/AutoRTFM/API.cpp b/Engine/Source/Runtime/Core/Private/AutoRTFM/API.cpp
--- a/Engine/Source/Runtime/Core/Private/AutoRTFM/API.cpp
+++ b/Engine/Source/Runtime/Core/Private/AutoRTFM/API.cpp
@@ -527,35 +527,35 @@ AUTORTFM_CORE_API void RTFM_autortfm_record_open_write_err(void*, size_t)
 	UE_LOG(LogAutoRTFM, Fatal, TEXT("The function `autortfm_record_open_write` was called from closed code."));
 }
 
-AUTORTFM_CORE_API ETransactionResult RTFM_CascadingAbortTransactionInternal(TFunction<void()>&& RunAfterAbort)
+AUTORTFM_CORE_API ETransactionResult RTFM_CascadingAbortTransactionInternal(TTask<void()>&& RunAfterAbort)
 {
 	FContext* const Context = FContext::Get();
 	ASSERT(Context->GetStatus() == EContextStatus::OnTrack);
 	return Context->AbortTransactionWithPostAbortCallback(/* bIsClosed */ true, EContextStatus::AbortedByCascadingAbort, MoveTemp(RunAfterAbort));
 }
 
-AUTORTFM_CORE_API void RTFM_CascadingRetryTransactionInternal(TFunction<void()> && RunAfterAbortBeforeRetryWork)
+AUTORTFM_CORE_API void RTFM_CascadingRetryTransactionInternal(TTask<void()> && RunAfterAbortBeforeRetryWork)
 {
 	FContext* Context = FContext::Get();
 	ASSERT(Context->GetStatus() == EContextStatus::OnTrack);
 	Context->AbortTransactionWithPostAbortCallback(/* bIsClosed */ true, EContextStatus::AbortedByCascadingRetry, MoveTemp(RunAfterAbortBeforeRetryWork));
 }
 
-AUTORTFM_CORE_API void RTFM_OnCommitInternal(TFunction<void()> && Work)
+AUTORTFM_CORE_API void RTFM_OnCommitInternal(TTask<void()> && Work)
 {
 	FContext* Context = FContext::Get();
 	ASSERT(Context->GetStatus() == EContextStatus::OnTrack);
 	Context->GetCurrentTransaction()->DeferUntilCommit(MoveTemp(Work));
 }
 
-AUTORTFM_CORE_API void RTFM_OnAbortInternal(TFunction<void()>&& Work)
+AUTORTFM_CORE_API void RTFM_OnAbortInternal(TTask<void()>&& Work)
 {
 	FContext* Context = FContext::Get();
 	ASSERT(Context->GetStatus() == EContextStatus::OnTrack);
 	Context->GetCurrentTransaction()->DeferUntilAbort(MoveTemp(Work));
 }
 
-AUTORTFM_CORE_API void RTFM_PushOnCommitHandlerInternal(const void* Key, TFunction<void()>&& Work)
+AUTORTFM_CORE_API void RTFM_PushOnCommitHandlerInternal(const void* Key, TTask<void()>&& Work)
 {
 	FContext* Context = FContext::Get();
 	ASSERT(Context->GetStatus() == EContextStatus::OnTrack);
@@ -576,7 +576,7 @@ AUTORTFM_CORE_API void RTFM_PopAllOnCommitHandlersInternal(const void* Key)
 	Context->GetCurrentTransaction()->PopAllDeferUntilCommitHandlers(Key);
 }
 
-AUTORTFM_CORE_API void RTFM_PushOnAbortHandlerInternal(const void* Key, TFunction<void()>&& Work)
+AUTORTFM_CORE_API void RTFM_PushOnAbortHandlerInternal(const void* Key, TTask<void()>&& Work)
 {
 	FContext* Context = FContext::Get();
 	ASSERT(Context->GetStatus() == EContextStatus::OnTrack);
@@ -891,7 +891,7 @@ extern "C" bool autortfm_is_on_current_transaction_stack(void* Ptr)
 	}
 }
 
-ETransactionResult ForTheRuntime::CascadingAbortTransactionInternal(TFunction<void()>&& RunAfterAbort)
+ETransactionResult ForTheRuntime::CascadingAbortTransactionInternal(TTask<void()>&& RunAfterAbort)
 {
 	if (autortfm_is_closed())
 	{
@@ -903,7 +903,7 @@ ETransactionResult ForTheRuntime::CascadingAbortTransactionInternal(TFunction<vo
 	return Context->AbortTransactionWithPostAbortCallback(/* bIsClosed */ false, EContextStatus::AbortedByCascadingAbort, MoveTemp(RunAfterAbort));
 }
 
-void ForTheRuntime::CascadingRetryTransactionInternal(TFunction<void()>&& RunAfterAbortBeforeRetryWork)
+void ForTheRuntime::CascadingRetryTransactionInternal(TTask<void()>&& RunAfterAbortBeforeRetryWork)
 {
 	if (autortfm_is_closed())
 	{
@@ -911,7 +911,7 @@ void ForTheRuntime::CascadingRetryTransactionInternal(TFunction<void()>&& RunAft
 	}
 }
 
-void ForTheRuntime::OnCommitInternal(TFunction<void()> && Work)
+void ForTheRuntime::OnCommitInternal(TTask<void()> && Work)
 {
 	if (autortfm_is_closed())
 	{
@@ -921,7 +921,7 @@ void ForTheRuntime::OnCommitInternal(TFunction<void()> && Work)
 	Work();
 }
 
-void ForTheRuntime::OnAbortInternal(TFunction<void()> && Work)
+void ForTheRuntime::OnAbortInternal(TTask<void()> && Work)
 {
 	if (autortfm_is_closed())
 	{
@@ -929,7 +929,7 @@ void ForTheRuntime::OnAbortInternal(TFunction<void()> && Work)
 	}
 }
 
-void ForTheRuntime::PushOnCommitHandlerInternal(const void* Key, TFunction<void()>&& Work)
+void ForTheRuntime::PushOnCommitHandlerInternal(const void* Key, TTask<void()>&& Work)
 {
 	if (autortfm_is_closed())
 	{
@@ -953,7 +953,7 @@ void ForTheRuntime::PopAllOnCommitHandlersInternal(const void* Key)
 	}
 }
 
-void ForTheRuntime::PushOnAbortHandlerInternal(const void* Key, TFunction<void()> && Work)
+void ForTheRuntime::PushOnAbortHandlerInternal(const void* Key, TTask<void()> && Work)
 {
 	if (autortfm_is_closed())
 	{
diff --git a/Engine/Source/Runtime/Core/Private/AutoRTFM/Context.cpp b/Engine/Source/Runtime/Core/Private/AutoRTFM/Context.cpp
--- a/Engine/Source/Runtime/Core/Private/AutoRTFM/Context.cpp
+++ b/Engine/Source/Runtime/Core/Private/AutoRTFM/Context.cpp
@@ -206,7 +206,7 @@ ETransactionResult FContext::AbortTransaction(bool bIsClosed, EContextStatus New
 	return Result;
 }
 
-ETransactionResult FContext::AbortTransactionWithPostAbortCallback(bool bIsClosed, EContextStatus NewStatus, TFunction<void()>&& InCallback)
+ETransactionResult FContext::AbortTransactionWithPostAbortCallback(bool bIsClosed, EContextStatus NewStatus, TTask<void()>&& InCallback)
 {
 	// The callback parameter is only honored by the cascading aborts.
 	ASSERT(NewStatus == EContextStatus::AbortedByCascadingAbort ||
diff --git a/Engine/Source/Runtime/Core/Private/AutoRTFM/Context.h b/Engine/Source/Runtime/Core/Private/AutoRTFM/Context.h
--- a/Engine/Source/Runtime/Core/Private/AutoRTFM/Context.h
+++ b/Engine/Source/Runtime/Core/Private/AutoRTFM/Context.h
@@ -33,7 +33,7 @@ public:
 
 	ETransactionResult CommitTransaction();
 	ETransactionResult AbortTransaction(bool bIsClosed, EContextStatus NewStatus);
-	ETransactionResult AbortTransactionWithPostAbortCallback(bool bIsClosed, EContextStatus NewStatus, TFunction<void()>&& InCallback);
+	ETransactionResult AbortTransactionWithPostAbortCallback(bool bIsClosed, EContextStatus NewStatus, TTask<void()>&& InCallback);
 	void ClearTransactionStatus();
 	bool IsAborting() const;
 	bool IsTransactional() const;
@@ -163,7 +163,7 @@ private:
     EContextStatus Status{EContextStatus::Idle};
 	uint32 CurrentThreadId{ FPlatformTLS::InvalidTlsSlot };
 	uint32 StackLocalInitializerDepth = 0;
-	TFunction<void()> PostAbortCallback;
+	TTask<void()> PostAbortCallback;
 	FOpenHashThrottler OpenHashThrottler{
 		/* LogInterval            */ 10,   // Log every 10 seconds.
 		/* AdjustThrottleInterval */ 0.5,  // Adjust throttling probabilities every 500ms.
diff --git a/Engine/Source/Runtime/Core/Private/AutoRTFM/TaskArray.h b/Engine/Source/Runtime/Core/Private/AutoRTFM/TaskArray.h
--- a/Engine/Source/Runtime/Core/Private/AutoRTFM/TaskArray.h
+++ b/Engine/Source/Runtime/Core/Private/AutoRTFM/TaskArray.h
@@ -5,7 +5,6 @@
 #include "Containers/Array.h"
 #include "Containers/BitArray.h"
 #include "Containers/Map.h"
-#include "Templates/Function.h"
 #include "Utils.h"
 
 namespace AutoRTFM
@@ -360,12 +359,9 @@ private:
 
 				if (bFirst)
 				{
-					static const TFunction<void()> EmptyFunctor = [] {};
-
 					// We have a deferred signal to pop this element, so do it now.
 					*Pair.Key = SKey{};
 					Pair.Val->Reset();
-					*Pair.Val = EmptyFunctor;
 					return;
 				}
 			}
@@ -390,12 +386,9 @@ private:
 
 				if (bLast)
 				{
-					static const TFunction<void()> EmptyFunctor = [] {};
-
 					// We have a deferred signal to pop this element, so do it now.
 					*Pair.Key = SKey{};
 					Pair.Val->Reset();
-					*Pair.Val = EmptyFunctor;
 					return;
 				}
 			}
diff --git a/Engine/Source/Runtime/Core/Private/AutoRTFM/Transaction.cpp b/Engine/Source/Runtime/Core/Private/AutoRTFM/Transaction.cpp
--- a/Engine/Source/Runtime/Core/Private/AutoRTFM/Transaction.cpp
+++ b/Engine/Source/Runtime/Core/Private/AutoRTFM/Transaction.cpp
@@ -105,7 +105,7 @@ void FTransaction::AbortWithoutThrowing()
     CommitTasks.Reset();
 
     Undo();
-	AbortTasks.ForEachBackward([&](TFunction<void()>& Task)
+	AbortTasks.ForEachBackward([&](TTask<void()>& Task)
     { 
         // Call and then reset each of the tasks in reverse order.
         // This ensures that the task and its destructor are called in reverse chronological order,
@@ -285,7 +285,7 @@ bool FTransaction::AttemptToCommitOuterNest()
     AbortTasks.Reset();
 
 
-    CommitTasks.ForEachForward([] (TFunction<void()>& Task)
+    CommitTasks.ForEachForward([] (TTask<void()>& Task)
     { 
         Task();
         Task.Reset();
diff --git a/Engine/Source/Runtime/Core/Private/AutoRTFM/Transaction.h b/Engine/Source/Runtime/Core/Private/AutoRTFM/Transaction.h
--- a/Engine/Source/Runtime/Core/Private/AutoRTFM/Transaction.h
+++ b/Engine/Source/Runtime/Core/Private/AutoRTFM/Transaction.h
@@ -90,12 +90,12 @@ public:
 
 	inline bool IsScopedTransaction() const { return bIsStackScoped; }
 
-    void DeferUntilCommit(TFunction<void()>&&);
-    void DeferUntilAbort(TFunction<void()>&&);
-	void PushDeferUntilCommitHandler(const void* Key, TFunction<void()>&&);
+    void DeferUntilCommit(TTask<void()>&&);
+    void DeferUntilAbort(TTask<void()>&&);
+	void PushDeferUntilCommitHandler(const void* Key, TTask<void()>&&);
 	void PopDeferUntilCommitHandler(const void* Key);
 	void PopAllDeferUntilCommitHandlers(const void* Key);
-    void PushDeferUntilAbortHandler(const void* Key, TFunction<void()>&&);
+    void PushDeferUntilAbortHandler(const void* Key, TTask<void()>&&);
 	void PopDeferUntilAbortHandler(const void* Key);
 	void PopAllDeferUntilAbortHandlers(const void* Key);
 
@@ -148,6 +148,7 @@ public:
 
 private:
     using FWriteHash = uint64;
+    using FTaskArray = TTaskArray<TTask<void()>>;
 
     void Undo();
 
@@ -185,10 +186,10 @@ private:
     FTransaction* Parent{nullptr};
 
     // Commit tasks run on commit in forward order.
-    TTaskArray<TFunction<void()>> CommitTasks;
+    FTaskArray CommitTasks;
 
 	// Abort tasks run on abort in reverse order.
-    TTaskArray<TFunction<void()>> AbortTasks;
+    FTaskArray AbortTasks;
 
 	// If a call to `PopOnCommitHandler` could not find a commit to pop, it is deferred
 	// and tried again on the parent transaction.
diff --git a/Engine/Source/Runtime/Core/Private/AutoRTFM/TransactionInlines.h b/Engine/Source/Runtime/Core/Private/AutoRTFM/TransactionInlines.h
--- a/Engine/Source/Runtime/Core/Private/AutoRTFM/TransactionInlines.h
+++ b/Engine/Source/Runtime/Core/Private/AutoRTFM/TransactionInlines.h
@@ -241,30 +241,30 @@ UE_AUTORTFM_FORCEINLINE void FTransaction::DidFree(void* LogicalAddress)
 	}
 }
 
-UE_AUTORTFM_FORCEINLINE void FTransaction::DeferUntilCommit(TFunction<void()>&& Callback)
+UE_AUTORTFM_FORCEINLINE void FTransaction::DeferUntilCommit(TTask<void()>&& Callback)
 {
 	// We explicitly must copy the function here because the original was allocated
 	// within a transactional context, and thus the memory is allocating under
 	// transactionalized conditions. By copying, we create an open copy of the callback.
-	TFunction<void()> Copy(Callback);
+	TTask<void()> Copy(Callback);
     CommitTasks.Add(MoveTemp(Copy));
 }
 
-UE_AUTORTFM_FORCEINLINE void FTransaction::DeferUntilAbort(TFunction<void()>&& Callback)
+UE_AUTORTFM_FORCEINLINE void FTransaction::DeferUntilAbort(TTask<void()>&& Callback)
 {
 	// We explicitly must copy the function here because the original was allocated
 	// within a transactional context, and thus the memory is allocating under
 	// transactionalized conditions. By copying, we create an open copy of the callback.
-	TFunction<void()> Copy(Callback);
+	TTask<void()> Copy(Callback);
     AbortTasks.Add(MoveTemp(Copy));
 }
 
-UE_AUTORTFM_FORCEINLINE void FTransaction::PushDeferUntilCommitHandler(const void* Key, TFunction<void()>&& Callback)
+UE_AUTORTFM_FORCEINLINE void FTransaction::PushDeferUntilCommitHandler(const void* Key, TTask<void()>&& Callback)
 {
 	// We explicitly must copy the function here because the original was allocated
 	// within a transactional context, and thus the memory was allocated under
 	// transactionalized conditions. By copying, we create an open copy of the callback.
-	TFunction<void()> Copy(Callback);
+	TTask<void()> Copy(Callback);
 	CommitTasks.AddKeyed(Key, MoveTemp(Copy));
 }
 
@@ -286,12 +286,12 @@ UE_AUTORTFM_FORCEINLINE void FTransaction::PopAllDeferUntilCommitHandlers(const
 	DeferredPopAllOnCommitHandlers.Push(Key);
 }
 
-UE_AUTORTFM_FORCEINLINE void FTransaction::PushDeferUntilAbortHandler(const void* Key, TFunction<void()>&& Callback)
+UE_AUTORTFM_FORCEINLINE void FTransaction::PushDeferUntilAbortHandler(const void* Key, TTask<void()>&& Callback)
 {
 	// We explicitly must copy the function here because the original was allocated
 	// within a transactional context, and thus the memory is allocating under
 	// transactionalized conditions. By copying, we create an open copy of the callback.
-	TFunction<void()> Copy(Callback);
+	TTask<void()> Copy(Callback);
     AbortTasks.AddKeyed(Key, MoveTemp(Copy));
 }
 
diff --git a/Engine/Source/Runtime/Core/Public/AutoRTFM/AutoRTFM.h b/Engine/Source/Runtime/Core/Public/AutoRTFM/AutoRTFM.h
--- a/Engine/Source/Runtime/Core/Public/AutoRTFM/AutoRTFM.h
+++ b/Engine/Source/Runtime/Core/Public/AutoRTFM/AutoRTFM.h
@@ -70,11 +70,9 @@
 #define LIKELY(EXPR) __builtin_expect(!!(EXPR), 1)
 #define PLATFORM_RETURN_ADDRESS() __builtin_return_address(0)
 #define UE_AUTORTFM_TCHAR char
-template <typename FuncType> class TFunction;  // Should be replaced with `std::function`
 #else
 #include <HAL/Platform.h>
 #include <HAL/PlatformMemory.h>
-#include <Templates/Function.h>
 #include <Templates/UnrealTemplate.h>
 #define UE_AUTORTFM_API CORE_API
 #define UE_AUTORTFM_FORCEINLINE FORCEINLINE
@@ -379,6 +377,7 @@ static UE_AUTORTFM_FORCEINLINE void autortfm_unreachable(const UE_AUTORTFM_TCHAR
 #ifdef __cplusplus
 }
 
+#include "AutoRTFMTask.h"
 #include <algorithm>
 #include <tuple>
 
@@ -472,14 +471,14 @@ enum class EMemoryValidationLevel : uint8_t
 #if UE_AUTORTFM
 namespace ForTheRuntime
 {
-	UE_AUTORTFM_API ETransactionResult CascadingAbortTransactionInternal(TFunction<void()>&& RunAfterAbort);
-	UE_AUTORTFM_API void CascadingRetryTransactionInternal(TFunction<void()>&& RunAfterAbortBeforeRetryWork);
-	UE_AUTORTFM_API void OnCommitInternal(TFunction<void()>&& Work);
-	UE_AUTORTFM_API void OnAbortInternal(TFunction<void()>&& Work);
-	UE_AUTORTFM_API void PushOnCommitHandlerInternal(const void* Key, TFunction<void()>&& Work);
+	UE_AUTORTFM_API ETransactionResult CascadingAbortTransactionInternal(TTask<void()>&& RunAfterAbort);
+	UE_AUTORTFM_API void CascadingRetryTransactionInternal(TTask<void()>&& RunAfterAbortBeforeRetryWork);
+	UE_AUTORTFM_API void OnCommitInternal(TTask<void()>&& Work);
+	UE_AUTORTFM_API void OnAbortInternal(TTask<void()>&& Work);
+	UE_AUTORTFM_API void PushOnCommitHandlerInternal(const void* Key, TTask<void()>&& Work);
 	UE_AUTORTFM_API void PopOnCommitHandlerInternal(const void* Key);
 	UE_AUTORTFM_API void PopAllOnCommitHandlersInternal(const void* Key);
-	UE_AUTORTFM_API void PushOnAbortHandlerInternal(const void* Key, TFunction<void()>&& Work);
+	UE_AUTORTFM_API void PushOnAbortHandlerInternal(const void* Key, TTask<void()>&& Work);
 	UE_AUTORTFM_API void PopOnAbortHandlerInternal(const void* Key);
 	UE_AUTORTFM_API void PopAllOnAbortHandlersInternal(const void* Key);
 } // namespace ForTheRuntime
