[Core] Add missing include for TFunction.

There was a transitive dependency on "Templates/Function.h" via "AutoRTFM/AutoRTFM.h" which has been removed.
Add the missing include to the .cpp file that uses TFunction.

diff --git a/Engine/Source/Runtime/Core/Tests/Serialization/Base64Test.cpp b/Engine/Source/Runtime/Core/Tests/Serialization/Base64Test.cpp
--- a/Engine/Source/Runtime/Core/Tests/Serialization/Base64Test.cpp
+++ b/Engine/Source/Runtime/Core/Tests/Serialization/Base64Test.cpp
@@ -3,6 +3,7 @@
 #if WITH_TESTS
 
 #include "Misc/Base64.h"
+#include "Templates/Function.h"
 #include "Tests/TestHarnessAdapter.h"
 
 template <size_t InputSize>
