[AutoRTFM] Make filename-based versions of FFileHelper::LoadFileToString transactionally safe.

We now run the inner file handling code in the open, making sure that everything done in the Open is cleaned up appropriately. This code can be reached in a transaction via several paths; one such path can be found in the attached Jira.

Note that the FArchive-based version of LoadFileToString was not rewritten because it changes the state of the FArchive (by reading from the archive, then closing it) and this would be difficult to undo cleanly in the event of a transaction rollback. However, this function was updated to be consistent in its treatment of the FArchive--we now consistently close it when the function returns. Previously, zero-byte archives would not be closed, and if an error occurred we would also leave the archive open, but a successful read would close the archive. The doc-comments have been updated to mention that the archive is closed; they previously did not mention this (and mislabeled the parameter as well).

diff --git a/Engine/Source/Runtime/Core/Private/Misc/FileHelper.cpp b/Engine/Source/Runtime/Core/Private/Misc/FileHelper.cpp
--- a/Engine/Source/Runtime/Core/Private/Misc/FileHelper.cpp
+++ b/Engine/Source/Runtime/Core/Private/Misc/FileHelper.cpp
@@ -1,6 +1,7 @@
 // Copyright Epic Games, Inc. All Rights Reserved.
 
 #include "Misc/FileHelper.h"
+#include "AutoRTFM/AutoRTFM.h"
 #include "Containers/StringConv.h"
 #include "Logging/LogMacros.h"
 #include "CoreGlobals.h"
@@ -161,12 +162,14 @@ bool FFileHelper::LoadFileToString(FString& Result, FArchive& Reader, EHashOptio
 	if (!Size)
 	{
 		Result.Empty();
+		Reader.Close();
 		return true;
 	}
 
 	if (Reader.Tell() != 0)
 	{
 		UE_LOG(LogStreaming, Warning, TEXT("Archive '%s' has already been read from."), *Reader.GetArchiveName());
+		Reader.Close();
 		return false;
 	}
 
@@ -176,15 +179,15 @@ bool FFileHelper::LoadFileToString(FString& Result, FArchive& Reader, EHashOptio
 
 	BufferToString(Result, Ch, (int32)Size);
 
-	// handle SHA verify of the file
+	// Handle SHA verification of the file.
 	if (EnumHasAnyFlags(VerifyFlags, EHashOptions::EnableVerify) && (EnumHasAnyFlags(VerifyFlags, EHashOptions::ErrorMissingHash) || FSHA1::GetFileSHAHash(*Reader.GetArchiveName(), nullptr)))
 	{
-		// kick off SHA verify task. this frees the buffer on close
+		// Kick off SHA verify task. This frees the buffer on close.
 		FBufferReaderWithSHA Ar(Ch, Size, true, *Reader.GetArchiveName(), false, true);
 	}
 	else
 	{
-		// free manually since not running SHA task
+		// Free manually, since the SHA task is not being run.
 		FMemory::Free(Ch);
 	}
 
@@ -198,47 +201,61 @@ bool FFileHelper::LoadFileToString(FString& Result, FArchive& Reader, EHashOptio
  * @param Filename name of the file to load
  * @param VerifyFlags flags controlling the hash verification behavior ( see EHashOptions )
  */
-bool FFileHelper::LoadFileToString( FString& Result, const TCHAR* Filename, EHashOptions VerifyFlags, uint32 ReadFlags)
+bool FFileHelper::LoadFileToString(FString& Result, const TCHAR* Filename, EHashOptions VerifyFlags, uint32 ReadFlags)
 {
-	TUniquePtr<FArchive> Reader(IFileManager::Get().CreateFileReader(Filename, ReadFlags));
-	if (!Reader)
+	bool Success = false;
+
+	Result = AutoRTFM::Open([Filename, VerifyFlags, ReadFlags, &Success]() -> FString
 	{
-		return false;
-	}
+		FString FileData;
 
-	return LoadFileToString(Result, *Reader.Get(), VerifyFlags);
+		if (TUniquePtr<FArchive> Reader(IFileManager::Get().CreateFileReader(Filename, ReadFlags)); Reader)
+		{
+			Success = LoadFileToString(FileData, *Reader, VerifyFlags);
+		}
+
+		return FileData;
+	});
+
+	return Success;
 }
 
 bool FFileHelper::LoadFileToString(FString& Result, IPlatformFile* PlatformFile, const TCHAR* Filename, EHashOptions VerifyFlags /*= EHashOptions::None*/, uint32 ReadFlags /*= 0*/)
 {
-	if (!PlatformFile)
-	{
-		return false;
-	}
+	bool Success = false;
 
-	IFileHandle* File = PlatformFile->OpenRead(Filename);
-	if (!File)
+	if (PlatformFile)
 	{
-		if (ReadFlags & FILEREAD_NoFail)
+		Result = AutoRTFM::Open([PlatformFile, Filename, VerifyFlags, ReadFlags, &Success]() -> FString
 		{
-			UE_LOG(LogStreaming, Fatal, TEXT("Failed to read file: %s"), Filename);
-		}
+			FString FileData;
 
-		if (!(ReadFlags & FILEREAD_Silent))
-		{
-			UE_LOG(LogStreaming, Warning, TEXT("Failed to read file '%s' error."), Filename);
-		}
+			if (IFileHandle* File = PlatformFile->OpenRead(Filename))
+			{
+				if (TUniquePtr<FArchive> Reader = MakeUnique<FArchiveFileReaderGeneric>(File, Filename, File->Size()))
+				{
+					Success = LoadFileToString(FileData, *Reader, VerifyFlags);
+				}
+			}
 
-		return false;
-	}
+			return FileData;
+		});
 
-	TUniquePtr<FArchive> Reader = MakeUnique<FArchiveFileReaderGeneric>(File, Filename, File->Size());
-	if (!Reader)
-	{
-		return false;
+		if (!Success)
+		{
+			if (ReadFlags & FILEREAD_NoFail)
+			{
+				UE_LOG(LogStreaming, Fatal, TEXT("Failed to read file: %s"), Filename);
+			}
+
+			if (!(ReadFlags & FILEREAD_Silent))
+			{
+				UE_LOG(LogStreaming, Warning, TEXT("Failed to read file '%s' error."), Filename);
+			}
+		}
 	}
 
-	return LoadFileToString(Result, *Reader.Get(), VerifyFlags);
+	return Success;
 }
 
 bool FFileHelper::LoadFileToStringArray( TArray<FString>& Result, const TCHAR* Filename )
diff --git a/Engine/Source/Runtime/Core/Public/Misc/FileHelper.h b/Engine/Source/Runtime/Core/Public/Misc/FileHelper.h
--- a/Engine/Source/Runtime/Core/Public/Misc/FileHelper.h
+++ b/Engine/Source/Runtime/Core/Public/Misc/FileHelper.h
@@ -78,10 +78,10 @@ struct FFileHelper
 	static CORE_API bool LoadFileToArray( TArray64<uint8>& Result, const TCHAR* Filename, uint32 Flags = 0 );
 
 	/**
-	 * Load a text file to an FString. Supports all combination of ANSI/Unicode files and platforms.
+	 * Loads a text file from an FArchive to an FString, then closes the archive. Supports all combination of ANSI/Unicode files and platforms.
 	 *
 	 * @param Result       String representation of the loaded file
-	 * @param Archive      Name of the archive to load from
+	 * @param Reader       The archive to load from
 	 * @param VerifyFlags  Flags controlling the hash verification behavior ( see EHashOptions )
 	 */
 	static CORE_API bool LoadFileToString(FString& Result, FArchive& Reader, EHashOptions VerifyFlags = EHashOptions::None);
