On some platform the size of PEVENT_RECORD::UserDataLength is bigger than FContextSwitchEvent (https://docs.microsoft.com/en-us/windows/win32/etw/cswitch)

diff --git a/Engine/Source/Runtime/Core/Private/ProfilingDebugging/Microsoft/EventTracingForWindows.cpp b/Engine/Source/Runtime/Core/Private/ProfilingDebugging/Microsoft/EventTracingForWindows.cpp
--- a/Engine/Source/Runtime/Core/Private/ProfilingDebugging/Microsoft/EventTracingForWindows.cpp
+++ b/Engine/Source/Runtime/Core/Private/ProfilingDebugging/Microsoft/EventTracingForWindows.cpp
@@ -112,7 +112,9 @@ void FPlatformEvents::TraceEventCallback(PEVENT_RECORD Event)
 	if (Opcode == FContextSwitchEvent::Type)
 	{
 		const FContextSwitchEvent* Payload = (const FContextSwitchEvent*)Event->UserData;
-		if (Payload == nullptr || Event->UserDataLength != sizeof(FContextSwitchEvent))
+		
+		// On some platforms we are getting Event->UserDataLength > sizeof(FContextSwitchEvent), the delta in size is undocumented 
+		if (Payload == nullptr || Event->UserDataLength < sizeof(FContextSwitchEvent))
 		{
 			return;
 		}
