Deleted some deprecated traits.

diff --git a/Engine/Source/Runtime/Core/Public/Templates/UnrealTypeTraits.h b/Engine/Source/Runtime/Core/Public/Templates/UnrealTypeTraits.h
--- a/Engine/Source/Runtime/Core/Public/Templates/UnrealTypeTraits.h
+++ b/Engine/Source/Runtime/Core/Public/Templates/UnrealTypeTraits.h
@@ -59,17 +59,6 @@ struct TIsDerivedFrom
 	static constexpr bool IsDerived = Value;
 };
 
-/**
- * TIsSame
- *
- * Unreal implementation of std::is_same trait.
- */
-template<typename A, typename B>
-struct UE_DEPRECATED(5.2, "TIsSame has been deprecated, please use std::is_same instead.") TIsSame
-{
-	enum { Value = std::is_same_v<A, B>	};
-};
-
 /** Gets the Nth type in a template parameter pack. N must be less than sizeof...(Types) */
 template <int32 N, typename... Types>
 struct TNthTypeFromParameterPack;
@@ -259,19 +248,6 @@ struct TCallTraitsParamTypeHelper<T*, true>
 	Helper templates for dealing with 'const' in template code
 -----------------------------------------------------------------------------*/
 
-/**
- * TRemoveConst<> is modeled after boost's implementation.  It allows you to take a templatized type
- * such as 'const Foo*' and use const_cast to convert that type to 'Foo*' without knowing about Foo.
- *
- *		MutablePtr = const_cast< RemoveConst< ConstPtrType >::Type >( ConstPtr );
- */
-template< class T >
-struct UE_DEPRECATED(5.3, "TRemoveConst has been deprecated, please use std::remove_const_t instead.") TRemoveConst
-{
-	using Type = std::remove_const_t<T>;
-};
-
-
 /*-----------------------------------------------------------------------------
  * TCallTraits
  *
diff --git a/Engine/Source/Runtime/Core/Public/Traits/IsVoidType.h b/Engine/Source/Runtime/Core/Public/Traits/IsVoidType.h
deleted file mode 100644
