Rename `Success` to `bSuccess` throughout FileHelper.

Somewhat surprisingly, this mistake was made in other FileHelper functions as well.

diff --git a/Engine/Source/Runtime/Core/Private/Misc/FileHelper.cpp b/Engine/Source/Runtime/Core/Private/Misc/FileHelper.cpp
--- a/Engine/Source/Runtime/Core/Private/Misc/FileHelper.cpp
+++ b/Engine/Source/Runtime/Core/Private/Misc/FileHelper.cpp
@@ -59,9 +59,9 @@ bool FFileHelper::LoadFileToArray( TArray<uint8>& Result, const TCHAR* Filename,
 	Result.Reset( TotalSize + 2 );
 	Result.AddUninitialized( TotalSize );
 	Reader->Serialize(Result.GetData(), Result.Num());
-	bool Success = Reader->Close();
+	bool bSuccess = Reader->Close();
 	delete Reader;
-	return Success;
+	return bSuccess;
 }
 
 /**
@@ -86,9 +86,9 @@ bool FFileHelper::LoadFileToArray(TArray64<uint8>& Result, const TCHAR* Filename
 	Result.Reset(TotalSize + 2);
 	Result.AddUninitialized(TotalSize);
 	Reader->Serialize(Result.GetData(), Result.Num());
-	bool Success = Reader->Close();
+	bool bSuccess = Reader->Close();
 	delete Reader;
-	return Success;
+	return bSuccess;
 }
 
 /**
@@ -175,7 +175,7 @@ bool FFileHelper::LoadFileToString(FString& Result, FArchive& Reader, EHashOptio
 
 	uint8* Ch = (uint8*)FMemory::Malloc(Size);
 	Reader.Serialize(Ch, Size);
-	bool Success = Reader.Close();
+	bool bSuccess = Reader.Close();
 
 	BufferToString(Result, Ch, (int32)Size);
 
@@ -191,7 +191,7 @@ bool FFileHelper::LoadFileToString(FString& Result, FArchive& Reader, EHashOptio
 		FMemory::Free(Ch);
 	}
 
-	return Success;
+	return bSuccess;
 }
 
 /**
@@ -203,30 +203,30 @@ bool FFileHelper::LoadFileToString(FString& Result, FArchive& Reader, EHashOptio
  */
 bool FFileHelper::LoadFileToString(FString& Result, const TCHAR* Filename, EHashOptions VerifyFlags, uint32 ReadFlags)
 {
-	bool Success = false;
+	bool bSuccess = false;
 
-	Result = AutoRTFM::Open([Filename, VerifyFlags, ReadFlags, &Success]() -> FString
+	Result = AutoRTFM::Open([Filename, VerifyFlags, ReadFlags, &bSuccess]() -> FString
 	{
 		FString FileData;
 
 		if (TUniquePtr<FArchive> Reader(IFileManager::Get().CreateFileReader(Filename, ReadFlags)); Reader)
 		{
-			Success = LoadFileToString(FileData, *Reader, VerifyFlags);
+			bSuccess = LoadFileToString(FileData, *Reader, VerifyFlags);
 		}
 
 		return FileData;
 	});
 
-	return Success;
+	return bSuccess;
 }
 
 bool FFileHelper::LoadFileToString(FString& Result, IPlatformFile* PlatformFile, const TCHAR* Filename, EHashOptions VerifyFlags /*= EHashOptions::None*/, uint32 ReadFlags /*= 0*/)
 {
-	bool Success = false;
+	bool bSuccess = false;
 
 	if (PlatformFile)
 	{
-		Result = AutoRTFM::Open([PlatformFile, Filename, VerifyFlags, ReadFlags, &Success]() -> FString
+		Result = AutoRTFM::Open([PlatformFile, Filename, VerifyFlags, ReadFlags, &bSuccess]() -> FString
 		{
 			FString FileData;
 
@@ -234,14 +234,14 @@ bool FFileHelper::LoadFileToString(FString& Result, IPlatformFile* PlatformFile,
 			{
 				if (TUniquePtr<FArchive> Reader = MakeUnique<FArchiveFileReaderGeneric>(File, Filename, File->Size()))
 				{
-					Success = LoadFileToString(FileData, *Reader, VerifyFlags);
+					bSuccess = LoadFileToString(FileData, *Reader, VerifyFlags);
 				}
 			}
 
 			return FileData;
 		});
 
-		if (!Success)
+		if (!bSuccess)
 		{
 			if (ReadFlags & FILEREAD_NoFail)
 			{
@@ -255,7 +255,7 @@ bool FFileHelper::LoadFileToString(FString& Result, IPlatformFile* PlatformFile,
 		}
 	}
 
-	return Success;
+	return bSuccess;
 }
 
 bool FFileHelper::LoadFileToStringArray( TArray<FString>& Result, const TCHAR* Filename )
