Remove gsl not null
[RN] gsl::not_null is deprecated for reason explain in the TNotNull header file. If you would like to use gsl in your project, you can add the GuidelinesSupportLibrary module in your project build.cs file.

diff --git a/Engine/Source/Runtime/Core/Core.Build.cs b/Engine/Source/Runtime/Core/Core.Build.cs
--- a/Engine/Source/Runtime/Core/Core.Build.cs
+++ b/Engine/Source/Runtime/Core/Core.Build.cs
@@ -33,7 +33,6 @@ public class Core : ModuleRules
 		PrivateDependencyModuleNames.Add("OodleDataCompression");
 		PrivateDependencyModuleNames.Add("xxhash");
 
-		PublicDependencyModuleNames.Add("GuidelinesSupportLibrary");
 		PublicDependencyModuleNames.Add("TraceLog");
 
 		PrivateIncludePathModuleNames.AddRange(
diff --git a/Engine/Source/Runtime/Core/Public/Templates/GSLNotNull.h b/Engine/Source/Runtime/Core/Public/Templates/GSLNotNull.h
--- a/Engine/Source/Runtime/Core/Public/Templates/GSLNotNull.h
+++ b/Engine/Source/Runtime/Core/Public/Templates/GSLNotNull.h
@@ -2,54 +2,24 @@
 
 #pragma once
 
+// HEADER_UNIT_SKIP - Deprecated
+
 #include "CoreTypes.h"
+
+#if UE_ENABLE_INCLUDE_ORDER_DEPRECATED_IN_5_6
 #include "Misc/OptionalFwd.h"
 #include "Templates/UnrealTemplate.h"
-#include <gsl/pointers>
-
-class FArchive;
-
-/** Utility function to turn an `TOptional<gsl::not_null<T>>` back into a nullable T. */
-template<typename ObjectType>
-[[nodiscard]] inline ObjectType* GetRawPointerOrNull(const TOptional<gsl::not_null<ObjectType*>>& Optional)
-{
-	return Optional.IsSet() ? static_cast<ObjectType*>(Optional.GetValue()) : nullptr;
-}
-
-/** Utility function to turn an `TOptional<gsl::strict_not_null<T>>` back into a nullable T. */
-template<typename ObjectType>
-[[nodiscard]] inline ObjectType* GetRawPointerOrNull(const TOptional<gsl::strict_not_null<ObjectType*>>& Optional)
-{
-	return Optional.IsSet() ? static_cast<ObjectType*>(Optional.GetValue()) : nullptr;
-}
-
-
-/** Utility function to serialize a `gsl::not_null<T>. */
-template<typename ObjectType>
-FArchive& operator<<(FArchive& Ar, gsl::not_null<ObjectType>& NotNull)
-{
-	Ar << *NotNull;
-	return Ar;
-}
-
-/** Utility function to serialize a `gsl::strict_not_null<T>`. */
-template<typename ObjectType>
-FArchive& operator<<(FArchive& Ar, gsl::strict_not_null<ObjectType>& NotNull)
-{
-	Ar << *NotNull;
-	return Ar;
-}
-
-/** Utility function to hash a `gsl::not_null<T>`. */
-template<typename ObjectType>
-[[nodiscard]] inline auto GetTypeHash(const gsl::not_null<ObjectType>& NotNull) -> decltype(GetTypeHash(*NotNull))
-{
-	return GetTypeHash(*NotNull);
-}
-
-/** Utility function to hash a `gsl::strict_not_null<T>`. */
-template<typename ObjectType>
-[[nodiscard]] inline auto GetTypeHash(const gsl::strict_not_null<ObjectType>& NotNull) -> decltype(GetTypeHash(*NotNull))
-{
-	return GetTypeHash(*NotNull);
-}
+#endif
+
+/**
+ * gsl::not_null is deprecated for reason explain in the TNotNull header file.
+ * If you would like to use gsl in your project, you can add the GuidelinesSupportLibrary module in your project build.cs file.
+ * Ex:
+ * 	public class MyModule : ModuleRules
+ *	{
+ *		public MyModule(ReadOnlyTargetRules Target) : base(Target)
+ *		{
+ *			PublicDependencyModuleNames.Add("GuidelinesSupportLibrary");
+ *			...
+ */
+UE_DEPRECATED_HEADER(5.6, "Including Templates/GSLNotNull.h is deprecated. Use TNotNull instead.")
diff --git a/Engine/Source/Runtime/Core/Public/Templates/NonNullPointer.h b/Engine/Source/Runtime/Core/Public/Templates/NonNullPointer.h
--- a/Engine/Source/Runtime/Core/Public/Templates/NonNullPointer.h
+++ b/Engine/Source/Runtime/Core/Public/Templates/NonNullPointer.h
@@ -9,7 +9,9 @@
 #include "Templates/Requires.h"
 #include "Templates/UnrealTypeTraits.h"
 
+#if UE_ENABLE_INCLUDE_ORDER_DEPRECATED_IN_5_6
 #include <gsl/pointers>
+#endif
 #include <type_traits>
 
 class FArchive;
@@ -49,15 +51,6 @@ public:
 		ensureMsgf(InObject, TEXT("Tried to initialize TNonNullPtr with a null pointer!"));
 	}
 
-	/**
-	 * Converts from gsl::not_null
-	 */
-	template <typename OtherType, typename = std::enable_if_t<std::is_convertible<OtherType, ObjectType*>::value>>
-	FORCEINLINE TNonNullPtr(const gsl::not_null<OtherType>& Other)
-		: Object(Other.get())
-	{
-	}
-
 	/**
 	 * Constructs a non-null pointer from another non-null pointer
 	 */
@@ -103,17 +96,6 @@ public:
 		return *this;
 	}
 
-
-	/**
-	 * Assignment operator taking a gsl::not_null
-	 */
-	template <typename OtherType, typename = std::enable_if_t<std::is_convertible<OtherType, ObjectType*>::value>>
-	FORCEINLINE TNonNullPtr& operator=(const gsl::not_null<OtherType>& Other)
-	{
-		Object = Other.Get();
-		return *this;
-	}
-
 	/**
 	 * Comparison, will also handle default constructed state
 	 */
@@ -174,15 +156,6 @@ public:
 		ensureMsgf(Object, TEXT("Tried to access null pointer!"));
 		return Object;
 	}
-	
-	/**
-	 * Converts to a gsl::not_null
-	 */
-	FORCEINLINE operator gsl::not_null<ObjectType*>() const
-	{
-		ensureMsgf(Object, TEXT("Tried to access null pointer!"));
-		return gsl::make_not_null<ObjectType*>(Object);
-	}
 
 	/**
 	 * Returns the internal pointer
diff --git a/Engine/Source/Runtime/Core/Tests/Templates/NotNullTest.cpp b/Engine/Source/Runtime/Core/Tests/Templates/NotNullTest.cpp
deleted file mode 100644
