Remove reference to GSL not null

#rnx

diff --git a/Engine/Source/Runtime/Core/Public/Templates/NonNullPointer.h b/Engine/Source/Runtime/Core/Public/Templates/NonNullPointer.h
--- a/Engine/Source/Runtime/Core/Public/Templates/NonNullPointer.h
+++ b/Engine/Source/Runtime/Core/Public/Templates/NonNullPointer.h
@@ -9,9 +9,6 @@
 #include "Templates/Requires.h"
 #include "Templates/UnrealTypeTraits.h"
 
-#if UE_ENABLE_INCLUDE_ORDER_DEPRECATED_IN_5_6
-#include <gsl/pointers>
-#endif
 #include <type_traits>
 
 class FArchive;
