[Backout] - CL39384834
#rnx
Original CL Desc
-----------------------------------------------------------------
Remove gsl not null
[RN] gsl::not_null is deprecated for reason explain in the TNotNull header file. If you would like to use gsl in your project, you can add the GuidelinesSupportLibrary module in your project build.cs file.

diff --git a/Engine/Source/Runtime/Core/Core.Build.cs b/Engine/Source/Runtime/Core/Core.Build.cs
--- a/Engine/Source/Runtime/Core/Core.Build.cs
+++ b/Engine/Source/Runtime/Core/Core.Build.cs
@@ -33,6 +33,7 @@ public class Core : ModuleRules
 		PrivateDependencyModuleNames.Add("OodleDataCompression");
 		PrivateDependencyModuleNames.Add("xxhash");
 
+		PublicDependencyModuleNames.Add("GuidelinesSupportLibrary");
 		PublicDependencyModuleNames.Add("TraceLog");
 
 		PrivateIncludePathModuleNames.AddRange(
diff --git a/Engine/Source/Runtime/Core/Public/Templates/GSLNotNull.h b/Engine/Source/Runtime/Core/Public/Templates/GSLNotNull.h
--- a/Engine/Source/Runtime/Core/Public/Templates/GSLNotNull.h
+++ b/Engine/Source/Runtime/Core/Public/Templates/GSLNotNull.h
@@ -2,24 +2,54 @@
 
 #pragma once
 
-// HEADER_UNIT_SKIP - Deprecated
-
 #include "CoreTypes.h"
-
-#if UE_ENABLE_INCLUDE_ORDER_DEPRECATED_IN_5_6
 #include "Misc/OptionalFwd.h"
 #include "Templates/UnrealTemplate.h"
-#endif
-
-/**
- * gsl::not_null is deprecated for reason explain in the TNotNull header file.
- * If you would like to use gsl in your project, you can add the GuidelinesSupportLibrary module in your project build.cs file.
- * Ex:
- * 	public class MyModule : ModuleRules
- *	{
- *		public MyModule(ReadOnlyTargetRules Target) : base(Target)
- *		{
- *			PublicDependencyModuleNames.Add("GuidelinesSupportLibrary");
- *			...
- */
-UE_DEPRECATED_HEADER(5.6, "Including Templates/GSLNotNull.h is deprecated. Use TNotNull instead.")
+#include <gsl/pointers>
+
+class FArchive;
+
+/** Utility function to turn an `TOptional<gsl::not_null<T>>` back into a nullable T. */
+template<typename ObjectType>
+[[nodiscard]] inline ObjectType* GetRawPointerOrNull(const TOptional<gsl::not_null<ObjectType*>>& Optional)
+{
+	return Optional.IsSet() ? static_cast<ObjectType*>(Optional.GetValue()) : nullptr;
+}
+
+/** Utility function to turn an `TOptional<gsl::strict_not_null<T>>` back into a nullable T. */
+template<typename ObjectType>
+[[nodiscard]] inline ObjectType* GetRawPointerOrNull(const TOptional<gsl::strict_not_null<ObjectType*>>& Optional)
+{
+	return Optional.IsSet() ? static_cast<ObjectType*>(Optional.GetValue()) : nullptr;
+}
+
+
+/** Utility function to serialize a `gsl::not_null<T>. */
+template<typename ObjectType>
+FArchive& operator<<(FArchive& Ar, gsl::not_null<ObjectType>& NotNull)
+{
+	Ar << *NotNull;
+	return Ar;
+}
+
+/** Utility function to serialize a `gsl::strict_not_null<T>`. */
+template<typename ObjectType>
+FArchive& operator<<(FArchive& Ar, gsl::strict_not_null<ObjectType>& NotNull)
+{
+	Ar << *NotNull;
+	return Ar;
+}
+
+/** Utility function to hash a `gsl::not_null<T>`. */
+template<typename ObjectType>
+[[nodiscard]] inline auto GetTypeHash(const gsl::not_null<ObjectType>& NotNull) -> decltype(GetTypeHash(*NotNull))
+{
+	return GetTypeHash(*NotNull);
+}
+
+/** Utility function to hash a `gsl::strict_not_null<T>`. */
+template<typename ObjectType>
+[[nodiscard]] inline auto GetTypeHash(const gsl::strict_not_null<ObjectType>& NotNull) -> decltype(GetTypeHash(*NotNull))
+{
+	return GetTypeHash(*NotNull);
+}
