[TaskGraphTest]
- Fix race detected by TSAN

diff --git a/Engine/Source/Runtime/Core/Tests/Async/TaskGraphTest.cpp b/Engine/Source/Runtime/Core/Tests/Async/TaskGraphTest.cpp
--- a/Engine/Source/Runtime/Core/Tests/Async/TaskGraphTest.cpp
+++ b/Engine/Source/Runtime/Core/Tests/Async/TaskGraphTest.cpp
@@ -213,13 +213,13 @@ namespace OldTaskGraphTests
 		{
 			StartTime = FPlatformTime::Seconds();
 
-			static bool Output[1000];
+			static std::atomic<bool> Output[1000];
 			FPlatformMemory::Memzero(Output, 1000);
 
 			ParallelFor(1000,
 				[](int32 Index)
 				{
-					bool& Out = Output[Index];
+					std::atomic<bool>& Out = Output[Index];
 					FFunctionGraphTask::CreateAndDispatchWhenReady([&Out] { Out = true; });
 				}
 			);
