[AutoRTFM] Use named functions instead of lambdas for AutoRTFM::TTask::FPayloadMethods

Functionally this is exactly the same, but a proper name is preferable for debugging and profiling.

diff --git a/Engine/Source/Runtime/Core/Public/AutoRTFM/AutoRTFMTask.h b/Engine/Source/Runtime/Core/Public/AutoRTFM/AutoRTFMTask.h
--- a/Engine/Source/Runtime/Core/Public/AutoRTFM/AutoRTFMTask.h
+++ b/Engine/Source/Runtime/Core/Public/AutoRTFM/AutoRTFMTask.h
@@ -42,10 +42,10 @@ private:
 		using MoveFn = void(FPayload& Dst, FPayload& Src);
 		using DestructFn = void(FPayload& Payload);
 
-		CallFn* Call = nullptr; // Call the function held by the payload.
-		CopyFn* Copy = nullptr; // Copy the payload.
-		MoveFn* Move = nullptr; // Move the payload.
-		DestructFn* Destruct = nullptr; // Destruct the payload.
+		CallFn* Call = nullptr; // Calls the function held by the payload.
+		CopyFn* Copy = nullptr; // Copies the payload.
+		MoveFn* Move = nullptr; // Moves the payload.
+		DestructFn* Destruct = nullptr; // Destructs the payload.
 	};
 	
 	// Traits helper for a functor (T).
@@ -68,47 +68,53 @@ private:
 			return bCanUseInline ? reinterpret_cast<const T*>(Payload.Inline) : reinterpret_cast<const T*>(Payload.External);
 		}
 
+		// Calls the function held by the payload.
+		static ReturnType Call(FPayload& Payload, ParameterTypes... Arguments)
+		{
+			return (*FunctionFrom(Payload))(Forward<ParameterTypes>(Arguments)...);
+		}
+
+		// Copies the payload.
+		static void Copy(FPayload& Dst, const FPayload& Src)
+		{
+			if constexpr (bCanUseInline)
+			{
+				new(Dst.Inline) T(*FunctionFrom(Src));
+			}
+			else
+			{
+				Dst.External = new T(*FunctionFrom(Src));
+			}
+		}
+
+		// Moves the payload.
+		static void Move(FPayload& Dst, FPayload& Src)
+		{
+			if constexpr (bCanUseInline)
+			{
+				new(Dst.Inline) T(MoveTemp(*FunctionFrom(Src)));
+			}
+			else
+			{
+				Dst.External = Src.External;
+			}
+		}
+
+		// Destructs the payload.
+		static void Destruct(FPayload& Payload)
+		{
+			if constexpr (bCanUseInline)
+			{
+				FunctionFrom(Payload)->~T();
+			}
+			else
+			{
+				delete FunctionFrom(Payload);
+			}
+		}
+
 		// The methods to operate on a payload for a functor of type T.
-		static constexpr FPayloadMethods PayloadMethods = 
-			bCanUseInline
-				? FPayloadMethods
-				{
-					/* Call */ [](FPayload& Payload, ParameterTypes... Arguments) -> ReturnType
-					{
-						return (*FunctionFrom(Payload))(Forward<ParameterTypes>(Arguments)...);
-					},
-					/* Copy */ [](FPayload& Dst, const FPayload& Src)
-					{
-						new(Dst.Inline) T(*FunctionFrom(Src));
-					},
-					/* Move */ [](FPayload& Dst, FPayload& Src)
-					{
-						new(Dst.Inline) T(MoveTemp(*FunctionFrom(Src)));
-					},
-					/* Destruct */ [](FPayload& Payload)
-					{
-						FunctionFrom(Payload)->~T();
-					},
-				}
-				: FPayloadMethods
-				{
-					/* Call */ [](FPayload& Payload, ParameterTypes... Arguments) -> ReturnType
-					{
-						return (*FunctionFrom(Payload))(Forward<ParameterTypes>(Arguments)...);
-					},
-					/* Copy */ [](FPayload& Dst, const FPayload& Src)
-					{
-						Dst.External = new T(*FunctionFrom(Src));
-					},
-					/* Move */ [](FPayload& Dst, FPayload& Src)
-					{
-						Dst.External = Src.External;
-					},
-					/* Destruct */ [](FPayload& Payload)
-					{
-						delete FunctionFrom(Payload);
-					},
-				};
+		static constexpr FPayloadMethods PayloadMethods { &Call, &Copy, &Move, &Destruct };
 	};
 public:
 	// Constructor.
