[AutoRTFM] TBlockAllocator: Add inline block, add growth factor, support large allocations

The inline block can avoid an initial heap allocation.
The growth factor is used to reduce the number of heap allocations when sequentially allocating a lot of objects.
The block allocator can now support allocations larger than the current block size.

diff --git a/Engine/Source/Programs/AutoRTFMTests/Private/BlockAllocatorTests.cpp b/Engine/Source/Programs/AutoRTFMTests/Private/BlockAllocatorTests.cpp
--- a/Engine/Source/Programs/AutoRTFMTests/Private/BlockAllocatorTests.cpp
+++ b/Engine/Source/Programs/AutoRTFMTests/Private/BlockAllocatorTests.cpp
@@ -8,7 +8,8 @@ TEST_CASE("BlockAllocator")
 {
 	static constexpr size_t AllocatorBlockSize = 128;
 	static constexpr size_t AllocatorAlignment = 16;
-	AutoRTFM::TBlockAllocator<AllocatorBlockSize, AllocatorAlignment> BlockAllocator;
+	static constexpr size_t AllocatorGrowthPercentage = 100; // Fixed size
+	AutoRTFM::TBlockAllocator<AllocatorBlockSize, AllocatorAlignment, AllocatorGrowthPercentage> BlockAllocator;
 
 	struct FAllocationInfo
 	{
@@ -77,6 +78,28 @@ TEST_CASE("BlockAllocator")
 		}
 	}
 
+	SECTION("LargeAllocations")
+	{
+		SECTION("Large, Small, Small")
+		{
+			Allocate(AllocatorBlockSize*4, 4);
+			Allocate(4, 4);
+			Allocate(4, 4);
+		}
+		SECTION("Small, Large, Small")
+		{
+			Allocate(4, 4);
+			Allocate(AllocatorBlockSize*4, 4);
+			Allocate(4, 4);
+		}
+		SECTION("Small, Small, Large")
+		{
+			Allocate(4, 4);
+			Allocate(4, 4);
+			Allocate(AllocatorBlockSize*4, 4);
+		}
+	}
+
 	// Check that the allocated memory is as expected.
 	for (FAllocationInfo AllocationInfo : Allocations)
 	{
diff --git a/Engine/Source/Runtime/Core/Private/AutoRTFM/BlockAllocator.h b/Engine/Source/Runtime/Core/Private/AutoRTFM/BlockAllocator.h
--- a/Engine/Source/Runtime/Core/Private/AutoRTFM/BlockAllocator.h
+++ b/Engine/Source/Runtime/Core/Private/AutoRTFM/BlockAllocator.h
@@ -3,22 +3,22 @@
 #pragma once
 
 #include "AutoRTFM/Utils.h"
-#include "HAL/Platform.h"
 #include "HAL/UnrealMemory.h"
 #include <cstddef>
 
 namespace AutoRTFM
 {
 
-// A single-threaded memory allocator that allocates out of fixed size blocks of
-// memory. All allocated memory is freed when the TBlockAllocator is destructed.
+// A single-threaded, bump memory allocator that sub-allocates out of blocks of
+// inline memory and once that's exhausted, FMemory allocated memory.
+// All allocated memory is freed when the TBlockAllocator is destructed.
 // There is no way to free individual allocations.
 //
 // Template parameters:
-//   BlockDataSize - The size of the data section of each block, which is also
-//                   the maximum size that can be allocated using Allocate().
+//   InlineBlockDataSize - The size of the data section of the inline block.
 //   DataAlignment - The maximum supported alignment of data from this allocator.
-template<size_t BlockDataSize, size_t DataAlignment = 16>
+//   GrowthPercentage - The percentage to grow the block size each reallocation.
+template<size_t InlineBlockDataSize = 256, size_t DataAlignment = 16, size_t GrowthPercentage = 200>
 class TBlockAllocator
 {
 public:
@@ -34,33 +34,33 @@ public:
 	// Frees all allocations made by the allocator.
 	void FreeAll()
 	{
-		while (Tail)
+		while (Tail != &InlineBlock)
 		{
 			FBlockHeader* Prev = Tail->Prev;
 			FMemory::Free(Tail);
 			Tail = Prev;
 		}
+		InlineBlock.Remaining = InlineBlockDataSize;
+		NextBlockSize = InlineBlockDataSize * GrowthPercentage / 100;
 	}
 
 	// Allocates memory from the block allocator.
-	// Size must be no larger than BlockSize.
 	// Alignment must be a power of two and no larger than DataAlignment.
 	inline void* Allocate(size_t Size, size_t Alignment)
 	{
-		ASSERT(Size <= BlockDataSize);
 		ASSERT(Alignment <= DataAlignment);
 
-		void* Allocation = Tail ? Tail->TryAllocate(Size, Alignment) : nullptr;
-		if (LIKELY(Allocation))
+		if (void* Allocation = Tail->TryAllocate(Size, Alignment))
 		{
 			return Allocation;
 		}
 
-		FBlockHeader* NewBlock = FBlockHeader::New();
-		NewBlock->Prev = Tail;
+		size_t NewBlockSize = std::max(NextBlockSize, Size);
+		NextBlockSize = NextBlockSize * GrowthPercentage / 100;
+		FBlockHeader* NewBlock = FBlockHeader::New(Tail, NewBlockSize);
 		Tail = NewBlock;
 
-		Allocation = NewBlock->TryAllocate(Size, Alignment);
+		void* Allocation = NewBlock->TryAllocate(Size, Alignment);
 		ASSERT(Allocation);
 		return Allocation;
 	}
@@ -74,26 +74,29 @@ public:
 
 private:
 	UE_NONCOPYABLE(TBlockAllocator);
+	
+	static constexpr size_t BlockAlignment = std::max<size_t>(DataAlignment, 8);
 
 	static inline size_t AlignDown(size_t Value, size_t Alignment)
 	{
 		return Value & ~(Alignment-1);
 	}
 
-	struct alignas(std::max<size_t>(DataAlignment, 8)) FBlockHeader
+	struct alignas(BlockAlignment) FBlockHeader
 	{
 		// The previous FBlock in the singly linked-list
-		FBlockHeader* Prev = nullptr;
+		FBlockHeader* const Prev = nullptr;
+		// The size of the block's data
+		const size_t BlockDataSize = 0;
 		// The number of unallocated bytes remaining in the block
 		size_t Remaining = BlockDataSize;
-		// <padding>
 		// <data>
 
 		// Allocates, constructs and returns a new block from FMemory::Malloc.
-		static FBlockHeader* New()
+		static FBlockHeader* New(FBlockHeader* Prev, size_t BlockDataSize)
 		{
 			void* Memory = FMemory::Malloc(sizeof(FBlockHeader) + BlockDataSize);
-			return new (Memory) FBlockHeader;
+			return new (Memory) FBlockHeader{Prev, BlockDataSize};
 		}
 
 		// Attempts to sub-allocate with out of this block. Returns a pointer to
@@ -113,7 +116,10 @@ private:
 
 	static_assert(sizeof(FBlockHeader) % DataAlignment == 0);
 
-	FBlockHeader* Tail = nullptr;
+	size_t NextBlockSize = InlineBlockDataSize * GrowthPercentage / 100;
+	FBlockHeader* Tail = &InlineBlock;
+	FBlockHeader InlineBlock{/* Prev */ nullptr, InlineBlockDataSize};
+	alignas(BlockAlignment) std::byte Data[InlineBlockDataSize];
 };
 
 }
