[AutoRTFM] Reimplement TTaskArray to optimize for keyed deletions.

The old implementation used bitfields for keeping track of the keyed tasks, and a single vector for all the tasks.
Popping a keyed task would involve flipping a bit in the bitfield, and the task would be removed next traversal.
This algorithm involved a lot of heap reallocations as certain use-cases pushed and popped with high frequency, and the underlying vector would grow repeatedly during a transaction, causing a substantial amount of profiled time to be spent in memory shuffling.

This new implementation uses a doubly-linked list to store the task entries, and a per key stack of entry pointers.
Deletion operations require unlinking an entry and popping a pointer from the stack.
Freed entries are pooled by a new task pool, so they can be efficiently reused.
Tasks that aren't in the free list are allocated from a TBlockAllocator which is optimised for small allocations that can be freed as a bulk operation.

Added tests for task array.

Reduces times for task-heavy transactions by over 6%.

diff --git a/Engine/Source/Programs/AutoRTFMTests/Private/TaskArrayTests.cpp b/Engine/Source/Programs/AutoRTFMTests/Private/TaskArrayTests.cpp
new file mode 100644
--- /dev/null
+++ b/Engine/Source/Programs/AutoRTFMTests/Private/TaskArrayTests.cpp
@@ -0,0 +1,280 @@
+// Copyright Epic Games, Inc. All Rights Reserved.
+
+#include "AutoRTFM/TaskArray.h"
+#include "Catch2Includes.h"
+#include "Templates/UnrealTemplate.h"
+#include <memory>
+#include <random>
+#include <unordered_map>
+
+TEST_CASE("TTaskArray")
+{
+	std::mt19937 Rand(0x1234);
+	using TaskArrayType = AutoRTFM::TTaskArray
+		<
+			/* TaskType */ int,
+			/* KeyType */ const void*,
+			/* AllocatorBlockSize */ 256
+		>;
+	TaskArrayType::FPool Pool;
+	auto TaskArray = std::make_unique<TaskArrayType>(Pool);
+
+	struct FEntry
+	{
+		int Value = 0;
+		FEntry* Prev = nullptr;
+		FEntry* Next = nullptr;
+	};
+
+	std::vector<void*> Keys;
+	std::unordered_map<void*, std::vector<FEntry*>> EntriesByKey;
+	FEntry* Head = nullptr;
+	FEntry* Tail = nullptr;
+
+	auto DeleteEntry = [&](FEntry* Entry)
+	{
+		if (Head == Entry) { Head = Entry->Next; }
+		if (Tail == Entry) { Tail = Entry->Prev; }
+		if (Entry->Prev) { Entry->Prev->Next = Entry->Next; }
+		if (Entry->Next) { Entry->Next->Prev = Entry->Prev; }
+		Entry->Prev = nullptr;
+		Entry->Next = nullptr;
+		delete Entry;
+	};
+
+	struct FConfig
+	{
+		size_t NumKeys = 0;
+		size_t NumElements = 0;
+	};
+
+	auto Populate = [&](FConfig Config)
+	{
+		for (size_t I = 0; I < Config.NumKeys; I++)
+		{
+			void* Key = reinterpret_cast<void*>(static_cast<uintptr_t>(I + 1) << 20);
+			Keys.push_back(Key);
+		}
+		
+		for (size_t I = 0; I < Config.NumElements; I++)
+		{
+			const int Value = static_cast<int>(Rand());
+			FEntry* const Entry = new FEntry{};
+			Entry->Value = Value;
+			if (!Head)
+			{ 
+				Head = Entry;
+			}
+			else
+			{
+				Tail->Next = Entry;
+				Entry->Prev = Tail;
+			}
+			Tail = Entry;
+			if (Config.NumKeys > 0 && Rand() & 1)
+			{
+				void* Key = Keys[Rand() % Keys.size()];
+				EntriesByKey[Key].push_back(Entry);
+				TaskArray->AddKeyed(Key, Value);
+			}
+			else
+			{
+				TaskArray->Add(Value);
+			}
+		}
+	};
+
+	// Checks the TaskArray is valid using a number of different approaches.
+	auto Check = [&]
+	{
+		if (Head)
+		{
+			const size_t ExpectedCount = TaskArray->Num();
+			size_t Count = 0;
+
+			SECTION("RemoveEachForward")
+			{
+				FEntry* Current = Head;
+				TaskArray->RemoveEachForward([&](int Got)
+				{
+					REQUIRE(Current);
+					REQUIRE(Current->Value == Got);
+					Current = Current->Next;
+					Count++;
+				});
+				REQUIRE(Count == ExpectedCount);
+				REQUIRE(!Current);
+			}
+
+			SECTION("RemoveEachBackward")
+			{
+				FEntry* Current = Tail;
+				TaskArray->RemoveEachBackward([&](int Got)
+				{
+					REQUIRE(Current);
+					REQUIRE(Current->Value == Got);
+					Current = Current->Prev;
+					Count++;
+				});
+				REQUIRE(Count == ExpectedCount);
+				REQUIRE(!Current);
+			}
+
+			SECTION("Reset")
+			{
+				TaskArray->Reset();
+			}
+			
+			REQUIRE(TaskArray->Num() == 0);
+			REQUIRE(TaskArray->IsEmpty());
+			REQUIRE(Pool.GetNumInUse() == 0);
+		}
+		else
+		{
+			REQUIRE(TaskArray->Num() == 0);
+			REQUIRE(TaskArray->IsEmpty());
+			TaskArray->RemoveEachForward([&](int)
+			{
+				FAIL("Should not be called");
+			});
+			TaskArray->RemoveEachBackward([&](int)
+			{
+				FAIL("Should not be called");
+			});
+		}
+	};
+	
+	auto CallWithConfigs = [&](auto&& Fn)
+	{
+		auto WithKeys = [&](size_t NumKeys)
+		{
+			SECTION("NumElements: 0")
+			{
+				Fn(FConfig{NumKeys, /* NumElements */ 0});
+			}
+
+			SECTION("NumElements: 10")
+			{
+				Fn(FConfig{NumKeys, /* NumElements */ 10});
+			}
+
+			SECTION("NumElements: 100")
+			{
+				Fn(FConfig{NumKeys, /* NumElements */ 100});
+			}
+
+			SECTION("NumElements: 1000")
+			{
+				Fn(FConfig{NumKeys, /* NumElements */ 1000});
+			}
+
+			SECTION("NumElements: 10000")
+			{
+				Fn(FConfig{NumKeys, /* NumElements */ 10000});
+			}
+		};
+		
+		SECTION("NumKeys: 0")
+		{
+			WithKeys(/* NumKeys */ 0);
+		}
+
+		SECTION("NumKeys: 5")
+		{
+			WithKeys(/* NumKeys */ 5);
+		}
+
+		SECTION("NumKeys: 10")
+		{
+			WithKeys(/* NumKeys */ 10);
+		}
+	};
+
+	SECTION("Add")
+	{
+		CallWithConfigs([&](FConfig Config)
+		{
+			Populate(Config);
+			Check();
+		});
+	}
+
+	SECTION("DeleteKey")
+	{
+		CallWithConfigs([&](FConfig Config)
+		{
+			if (Config.NumKeys == 0)
+			{
+				return;
+			}
+
+			Populate(Config);
+
+			for (size_t I = 0; I < 10; I++)
+			{
+				void* Key = Keys[Rand() % Config.NumKeys];
+				bool HoldsKey = false;
+				if (std::vector<FEntry*>& Entries = EntriesByKey[Key]; Entries.size())
+				{
+					HoldsKey = true;
+					DeleteEntry(Entries.back());
+					Entries.pop_back();
+				}
+				REQUIRE(TaskArray->DeleteKey(Key) == HoldsKey);
+			}
+
+			Check();
+		});
+	}
+
+	SECTION("DeleteAllMatchingKeys")
+	{
+		CallWithConfigs([&](FConfig Config)
+		{
+			if (Config.NumKeys == 0)
+			{
+				return;
+			}
+
+			Populate(Config);
+
+			for (size_t I = 0; I < 4; I++)
+			{
+				void* Key = Keys[Rand() % Config.NumKeys];
+				std::vector<FEntry*>& Entries = EntriesByKey[Key];
+				bool HoldsKey = !Entries.empty();
+				while (Entries.size())
+				{
+					DeleteEntry(Entries.back());
+					Entries.pop_back();
+				}
+				REQUIRE(TaskArray->DeleteAllMatchingKeys(Key) == HoldsKey);
+			}
+
+			Check();
+		});
+	}
+
+	SECTION("AddAll")
+	{
+		CallWithConfigs([&](FConfig Config)
+		{
+			Populate(Config);
+
+			auto Second = std::make_unique<TaskArrayType>(Pool);
+			std::swap(TaskArray, Second);
+			Populate(Config);
+			std::swap(TaskArray, Second);
+
+			TaskArray->AddAll(MoveTemp(*Second));
+
+			Check();
+		});
+	}
+
+	// Clean-up
+	while (Head)
+	{
+		DeleteEntry(Head);
+	}
+}
diff --git a/Engine/Source/Runtime/Core/Private/AutoRTFM/Context.cpp b/Engine/Source/Runtime/Core/Private/AutoRTFM/Context.cpp
--- a/Engine/Source/Runtime/Core/Private/AutoRTFM/Context.cpp
+++ b/Engine/Source/Runtime/Core/Private/AutoRTFM/Context.cpp
@@ -666,10 +666,11 @@ void FContext::Reset()
 
 	CurrentThreadId = FPlatformTLS::InvalidTlsSlot;
 	Stack = {};
-    CurrentTransaction = nullptr;
+	CurrentTransaction = nullptr;
 	CurrentNest = nullptr;
-    Status = EContextStatus::Idle;
+	Status = EContextStatus::Idle;
 	StackLocalInitializerDepth = 0;
+	TaskPool.Reset();
 }
 
 void FContext::Throw()
diff --git a/Engine/Source/Runtime/Core/Private/AutoRTFM/Context.h b/Engine/Source/Runtime/Core/Private/AutoRTFM/Context.h
--- a/Engine/Source/Runtime/Core/Private/AutoRTFM/Context.h
+++ b/Engine/Source/Runtime/Core/Private/AutoRTFM/Context.h
@@ -7,6 +7,7 @@
 #include "HAL/PlatformTLS.h"
 #include "StackRange.h"
 #include "OpenHashThrottler.h"
+#include "TaskArray.h"
 
 namespace AutoRTFM
 {
@@ -130,6 +131,8 @@ public:
 			}
 		}
 	}
+	
+	inline TTaskArray<TTask<void()>>::FPool& GetTaskPool() { return TaskPool; }
 
 private:
 	static FContext ContextSingleton;
@@ -164,6 +167,7 @@ private:
 	uint32 CurrentThreadId{ FPlatformTLS::InvalidTlsSlot };
 	uint32 StackLocalInitializerDepth = 0;
 	TTask<void()> PostAbortCallback;
+	TTaskArray<TTask<void()>>::FPool TaskPool;
 	FOpenHashThrottler OpenHashThrottler{
 		/* LogInterval            */ 10,   // Log every 10 seconds.
 		/* AdjustThrottleInterval */ 0.5,  // Adjust throttling probabilities every 500ms.
diff --git a/Engine/Source/Runtime/Core/Private/AutoRTFM/TaskArray.h b/Engine/Source/Runtime/Core/Private/AutoRTFM/TaskArray.h
--- a/Engine/Source/Runtime/Core/Private/AutoRTFM/TaskArray.h
+++ b/Engine/Source/Runtime/Core/Private/AutoRTFM/TaskArray.h
@@ -2,405 +2,377 @@
 
 #pragma once
 
-#include "Containers/Array.h"
-#include "Containers/BitArray.h"
+#include "AutoRTFM/Utils.h"
+#include "BlockAllocator.h"
 #include "Containers/Map.h"
-#include "Utils.h"
+#include "HAL/Platform.h"
+#include "Misc/CoreMiscDefines.h"
+#include "Templates/UnrealTemplate.h"
+#include <cstddef>
 
 namespace AutoRTFM
 {
 
-template<typename T> class TBackwards final
+// A sequential linked-list of tasks, with the ability to tag some tasks with a
+// key. Those tagged with a key can be removed from the list in LIFO order
+// (stack-like).
+// Tasks can be traversed bi-directionally.
+//
+// Template parameters:
+// Task            - the task data type
+// Key             - the data type used as a key for push / pop.
+// InlineTaskCount - the number of tasks that can be allocated before a performing heap allocations.
+template<typename TaskType, typename KeyType = const void*, size_t InlineTaskCount = 256>
+class TTaskArray
 {
-	T& Obj;
-public:
-	TBackwards(T& Obj) : Obj(Obj) {}
-	auto begin() { return Obj.rbegin(); }
-	auto end() { return Obj.rend(); }
-};
-
-template<typename InKeyType, typename InValType>
-struct TTaskArrayKeyValuePair
-{
-	InKeyType Key;
-	InValType Val;
-
-	TTaskArrayKeyValuePair() = default;
-	TTaskArrayKeyValuePair(const InKeyType& Key, const InValType& Val) : Key(Key), Val(Val) {}
-	TTaskArrayKeyValuePair(InKeyType&& Key, InValType&& Val) : Key(MoveTemp(Key)), Val(MoveTemp(Val)) {}
-};
-
-template<typename InElementType>
-class TTaskArray final
-{
-	using SKey = const void*;
-
-	struct FBitMaskPerKeyType 
+	// A task entry used by the array and pool.
+	// Holds an intrusive linked list as well as the task.
+	class FEntry
 	{
-		TBitArray<> Bits;
-		int32 InitialIndex = 0;
-
-		void RemoveFirst()
-		{
-			checkSlow(InitialIndex < Bits.Num());
-			InitialIndex++;
-		}
-
-		void RemoveLast()
-		{
-			checkSlow(InitialIndex < Bits.Num());
-			Bits.SetNumUninitialized(Bits.Num() - 1);
-		}
-
-		bool BitAt(int32 Index) const
+	public:
+		// The next task in the sequential list.
+		FEntry* Next = nullptr;
+		// The previous task in the sequential list.
+		FEntry* Prev = nullptr;
+
+		// Returns a reference to the task.
+		TaskType& Task()
 		{
-			return Bits[InitialIndex + Index];
+			return *reinterpret_cast<TaskType*>(TaskData);
 		}
 
-		int32 Num() const
+		// Unlinks this entry from the doubly-linked list, updating the provided
+		// head and tail pointers as necessary.
+		void Unlink(FEntry*& ListHead, FEntry*& ListTail)
 		{
-			return Bits.Num() - InitialIndex;
+			if (ListHead == this) { ListHead = Next; }
+			if (ListTail == this) { ListTail = Prev; }
+			if (Next) { Next->Prev = Prev; }
+			if (Prev) { Prev->Next = Next; }
+			Next = nullptr;
+			Prev = nullptr;
 		}
 
-		bool IsEmpty() const
+		// Links this entry to the doubly-linked list's tail pointer.
+		// This entry must be unlinked before calling.
+		void Link(FEntry*& ListHead, FEntry*& ListTail)
 		{
-			return InitialIndex == Bits.Num();
+			ASSERT(Next == nullptr && Prev == nullptr);
+			if (!ListHead)
+			{
+				ListHead = this;
+			}
+			Prev = ListTail;
+			if (Prev)
+			{
+				Prev->Next = this;
+			}
+			ListTail = this;
 		}
+	private:
+		alignas(alignof(TaskType)) std::byte TaskData[sizeof(TaskType)];
 	};
 
-	using PairType = TTaskArrayKeyValuePair<SKey*, InElementType*>;
-
-	struct FInternalArray final
+public:
+	// A pool of reusable tasks that can be shared between transactions.
+	class FPool
 	{
-		using IsKeyedArrayType = TBitArray<>;
-		using KeyArrayType = TArray<SKey>;
-		using ElementArrayType = TArray<InElementType>;
-		using SizeType = typename ElementArrayType::SizeType;
+	public:
+		// Constructor
+		FPool() = default;
 
-		// This iterator is just a zip of the data structures that make up the array.
-		template<typename IsKeyedIterType, typename KeyIterType, typename ElementIterType> struct TIterator final
+		// Acquires a new FEntry from the task pool, forward-constructing the task with Task.
+		template<typename TaskParam>
+		FEntry* TakeTask(TaskParam&& Task)
 		{
-			explicit TIterator(IsKeyedIterType IsKeyedIter, KeyIterType KeyIter, ElementIterType ElementIter) :
-				IsKeyedIter(IsKeyedIter), KeyIter(KeyIter), ElementIter(ElementIter) {}
-
-			PairType operator*() const
+			FEntry* Entry = nullptr;
+			if (FreeEntries)
 			{
-				SKey* const Key = (*IsKeyedIter) ? &*KeyIter : nullptr;
-				return PairType(Key, &*ElementIter);
+				// Take from the free list.
+				Entry = FreeEntries;
+				FreeEntries = Entry->Next;
+				Entry->Prev = nullptr;
+				Entry->Next = nullptr;
 			}
-
-			TIterator& operator++()
+			else
 			{
-				// We only bump the key iterator if we are currently on a keyed iterator!
-				if (*IsKeyedIter)
-				{
-					++KeyIter;
-				}
-				++IsKeyedIter;
-				++ElementIter;
-				return *this;
+				// Free list is empty. Allocate another entry.
+				Entry = Allocator.template New<FEntry>();
+				NumAllocated++;
 			}
 
-			bool operator==(const TIterator& Other) const
-			{
-				return !(*this != Other);
-			}
+			// Move the task into the entry.
+			new (&Entry->Task()) TaskType(Forward<TaskParam>(Task));
 
-			bool operator!=(const TIterator& Other) const
-			{
-				return (Other.IsKeyedIter != IsKeyedIter) || (Other.KeyIter != KeyIter) || (Other.ElementIter != ElementIter);
-			}
+			ASSERT(NumInUse < NumAllocated);
+			NumInUse++;
 
-		private:
-			IsKeyedIterType IsKeyedIter;
-			KeyIterType KeyIter;
-			ElementIterType ElementIter;
-		};
-
-		IsKeyedArrayType IsKeyed;
-		KeyArrayType Keys;
-		ElementArrayType Elements;
-		
-		void Add(const SKey Key, const InElementType& Element)
-		{
-			IsKeyed.Add(true);
-			Keys.Add(Key);
-			Elements.Add(Element);
+			return Entry;
 		}
 
-		void Add(const SKey Key, InElementType&& Element)
+		// Destructs the entry's Task and returns the entry back to the pool so
+		// it can be reused.
+		// Entry must be unlinked before calling.
+		void ReturnTask(FEntry* Entry)
 		{
-			IsKeyed.Add(true);
-			Keys.Add(Key);
-			Elements.Add(MoveTemp(Element));
-		}
+			// Destruct the task
+			Entry->Task().~TaskType();
 
-		void Add(const InElementType& Element)
-		{
-			IsKeyed.Add(false);
-			Elements.Add(Element);
-		}
+			// Place the entry onto the free list
+			Entry->Prev = nullptr;
+			Entry->Next = FreeEntries;
+			FreeEntries = Entry;
 
-		void Add(InElementType&& Element)
-		{
-			IsKeyed.Add(false);
-			Elements.Add(MoveTemp(Element));
+			ASSERT(NumInUse > 0);
+			NumInUse--;
 		}
 
-		bool IsEmpty() const
+		// Destructs each entry's Task and returns the linked list of entries
+		// back to the pool so they can be reused.
+		void ReleaseTasks(FEntry* List)
 		{
-			const bool bResult = Elements.IsEmpty();
-			ASSERT(bResult == IsKeyed.IsEmpty());
-			return bResult;
-		}
+			if (UNLIKELY(!List))
+			{
+				return;
+			}
 
-		SizeType Num() const
-		{
-			const SizeType Num = Elements.Num();
-			ASSERT(Num == IsKeyed.Num());
-			return Num;
+			// Walk the list to destruct each of the tasks.
+			size_t ListCount = 0;
+			for (FEntry* Entry = List; Entry; Entry = Entry->Next)
+			{
+				// Destruct the task
+				Entry->Task().~TaskType();
+				ListCount++;
+			}
+
+			// Link the entries onto the free list.
+			List->Prev = nullptr;
+			List->Next = FreeEntries;
+			FreeEntries = List;
+
+			ASSERT(NumInUse >= ListCount);
+			NumInUse -= ListCount;
 		}
 
+		// Frees the memory allocated for the task list.
+		// No tasks can be in use when calling.
 		void Reset()
 		{
-			IsKeyed.Reset();
-			Keys.Reset();
-			Elements.Reset();
+			ASSERT(NumInUse == 0);
+			Allocator.FreeAll();
+			NumAllocated = 0;
+			FreeEntries = nullptr;
 		}
 
-		using FIterator = TIterator<IsKeyedArrayType::FIterator, typename KeyArrayType::RangedForIteratorType, typename ElementArrayType::RangedForIteratorType>;
-		using FConstIterator = TIterator<IsKeyedArrayType::FConstIterator, typename KeyArrayType::RangedForConstIteratorType, typename ElementArrayType::RangedForConstIteratorType>;
-		using FReverseIterator = TIterator<IsKeyedArrayType::FReverseIterator, typename KeyArrayType::RangedForReverseIteratorType, typename ElementArrayType::RangedForReverseIteratorType>;
-		using FConstReverseIterator = TIterator<IsKeyedArrayType::FConstReverseIterator, typename KeyArrayType::RangedForConstReverseIteratorType, typename ElementArrayType::RangedForConstReverseIteratorType>;
-		
-		FIterator begin() { return TIterator(IsKeyed.begin(), Keys.begin(), Elements.begin()); }
-		FIterator end() { return TIterator(IsKeyed.end(), Keys.end(), Elements.end()); }
+		// Returns the number of tasks allocated.
+		size_t GetNumAllocated() const { return NumAllocated; }
 
-		FConstIterator begin() const { return TIterator(IsKeyed.begin(), Keys.begin(), Elements.begin()); }
-		FConstIterator end() const { return TIterator(IsKeyed.end(), Keys.end(), Elements.end()); }
+		// Returns the number of tasks currently in use (not returned).
+		size_t GetNumInUse() const { return NumInUse; }
 
-		FReverseIterator rbegin() { return TIterator(IsKeyed.rbegin(), Keys.rbegin(), Elements.rbegin()); }
-		FReverseIterator rend() { return TIterator(IsKeyed.rend(), Keys.rend(), Elements.rend()); }
+	private:
+		UE_NONCOPYABLE(FPool);
 
-		FConstReverseIterator rbegin() const { return TIterator(IsKeyed.rbegin(), Keys.rbegin(), Elements.rbegin()); }
-		FConstReverseIterator rend() const { return TIterator(IsKeyed.rend(), Keys.rend(), Elements.rend()); }
-	};
+		static constexpr size_t AllocatorInlineSize = InlineTaskCount * sizeof(TaskType);
+		static constexpr size_t AllocatorAlignment = std::max<size_t>(alignof(TaskType), alignof(void*));
 
-public:
-	TTaskArray() = default;
+		TBlockAllocator<AllocatorInlineSize, AllocatorAlignment> Allocator;
+		FEntry* FreeEntries = nullptr;
 
-    bool IsEmpty() const { return Latest.IsEmpty() && Stash.IsEmpty(); }
+		// Number of entries allocated
+		size_t NumAllocated = 0;
+		// Number of entries currently in use
+		size_t NumInUse = 0;
+	};
 
-    void Add(InElementType&& Value)
-    {
-		Latest.Add(MoveTemp(Value));
-    }
+	// Constructor. TaskPool is used to allocate new tasks.
+	TTaskArray(FPool& TaskPool) : Pool{TaskPool}
+	{}
 
-	void Add(const InElementType& Value)
+	// Destructor. Returns all allocated tasks to the pool.
+	~TTaskArray()
 	{
-		Latest.Add(Value);
+		Reset();
 	}
 
-	void AddKeyed(SKey Key, InElementType&& Value)
-    {
-		ASSERT(Key != SKey{});
-		Latest.Add(MoveTemp(Key), MoveTemp(Value));
-		BitMasksPerKey.FindOrAdd(Key).Bits.Add(false);
-    }
-
-	void AddKeyed(SKey Key, const InElementType& Value)
-    {
-		ASSERT(Key != SKey{});
-		Latest.Add(MoveTemp(Key), Value);
-		BitMasksPerKey.FindOrAdd(Key).Bits.Add(false);
-    }
-
-	bool DeleteKey(SKey Key)
+	// Returns all allocated tasks to the pool.
+	void Reset()
 	{
-		ASSERT(Key != SKey{});
-
-		FBitMaskPerKeyType* const KeyedElement = BitMasksPerKey.Find(Key);
-
-		if (UNLIKELY(nullptr == KeyedElement))
-		{
-			return false;
-		}
-
-		return KeyedElement->Bits.FindAndSetLastZeroBit() >= KeyedElement->InitialIndex;
+		Pool.ReleaseTasks(Head);
+		Head = nullptr;
+		Tail = nullptr;
+		Keyed.Empty();
+		Count = 0;
 	}
 
-	bool DeleteAllMatchingKeys(SKey Key)
+	// Adds the unkeyed task to the end of the array.
+	template<typename TaskParamType>
+	void Add(TaskParamType&& Task)
 	{
-		ASSERT(Key != SKey{});
+		FEntry* Entry = Pool.TakeTask(Forward<TaskParamType>(Task));
+		Entry->Link(Head, Tail);
+		Count++;
+	}
 
-		FBitMaskPerKeyType* const KeyedElement = BitMasksPerKey.Find(Key);
+	// Adds the keyed task to the end of the array.
+	template<typename TaskParamType>
+	void AddKeyed(KeyType Key, TaskParamType&& Task)
+	{
+		Add(Forward<TaskParamType>(Task));
+		Keyed.FindOrAdd(Key).Add(Tail);
+	}
 
-		if (UNLIKELY(nullptr == KeyedElement))
+	// Moves all the tasks from OtherArray to the end of this array, clearing
+	// OtherArray.
+	void AddAll(TTaskArray&& OtherArray)
+	{
+		if (Head)
+		{
+			// This array has tasks in the list.
+			// Link the two task linked lists together.
+			Tail->Next = OtherArray.Head;
+			if (OtherArray.Head)
+			{
+				OtherArray.Head->Prev = Tail;
+				Tail = OtherArray.Tail;
+			}
+		}
+		else
 		{
-			return false;
+			// This array holds no tasks, so replace this list with OtherArray's.
+			Head = OtherArray.Head;
+			Tail = OtherArray.Tail;
 		}
 
-		if (KeyedElement->Num() == 0)
+		// Append the keyed entries of OtherArray to this.
+		for (auto& Iterator : OtherArray.Keyed)
 		{
-			return false;
+			KeyType Key = Iterator.template Get<0>();
+			EntryStack& Stack = Iterator.template Get<1>();
+			Keyed.FindOrAdd(Key).Append(MoveTemp(Stack));
 		}
 
-		KeyedElement->Bits.SetRange(KeyedElement->InitialIndex, KeyedElement->Num(), true);
-		return true;
-	}
-
-    void AddAll(TTaskArray&& Other)
-    {
-        Canonicalize();
-
-        for (FInternalArray& StashedVectorBox : Other.Stash)
-        {
-            Stash.Push(MoveTemp(StashedVectorBox));
-        }
-        Other.Stash.Reset();
+		// Add in OtherArray's count.
+		Count += OtherArray.Count;
 
-        if (!Other.Latest.IsEmpty())
-        {
-            Stash.Push(MoveTemp(Other.Latest));
-			Other.Latest = FInternalArray{};
-        }
+		// Everything stolen from OtherArray. Reset it.
+		OtherArray.Head = nullptr;
+		OtherArray.Tail = nullptr;
+		OtherArray.Keyed.Reset();
+		OtherArray.Count = 0;
+	}
 
-		for (auto& Pair : Other.BitMasksPerKey)
+	// Removes the last added task with the given key.
+	// Returns true if the task with the given key was removed, or false if
+	// there are no remaining tasks with the given key.
+	bool DeleteKey(KeyType Key)
+	{
+		if (EntryStack* Stack = Keyed.Find(Key); Stack && !Stack->IsEmpty())
 		{
-			const SKey& OtherKey = Pair.Key;
-			FBitMaskPerKeyType& KeyedElement = BitMasksPerKey.FindOrAdd(OtherKey);
-
-			const FBitMaskPerKeyType& OtherBitMask = Pair.Value;
-			int OtherNumBits = OtherBitMask.Num();
+			Release(Stack->Pop(EAllowShrinking::No));
 
-			KeyedElement.Bits.AddRange(OtherBitMask.Bits, OtherBitMask.Num(), OtherBitMask.InitialIndex);
-		}
-		Other.BitMasksPerKey.Reset();
-    }
-
-    template<typename FunctorType>
-    void ForEachForward(const FunctorType& Func)
-    {
-        for (FInternalArray& StashedVectorBox : Stash)
-        {
-            for (PairType Pair : StashedVectorBox)
-            {
-				ForwardProcessEntry(Pair, Func);
-            }
-        }
-
-        for (PairType Pair : Latest)
-        {
-			ForwardProcessEntry(Pair, Func);
-        }
-    }
-
-    template<typename FunctorType>
-    void ForEachBackward(const FunctorType& Func)
-    {
-		for (PairType Pair : TBackwards(Latest))
-		{
-			BackwardProcessEntry(Pair, Func);
+			if (Stack->IsEmpty())
+			{
+				Keyed.Remove(Key);
+			}
+			return true;
 		}
 
-		for (FInternalArray& StashedVectorBox : TBackwards(Stash))
+		return false;
+	}
+	
+	// Removes all the tasks with the given key.
+	// Returns true if a task with the given key was removed, or false if
+	// there are no remaining tasks with the given key.
+	bool DeleteAllMatchingKeys(KeyType Key)
+	{
+		if (EntryStack* Stack = Keyed.Find(Key))
 		{
-			for (PairType Pair : TBackwards(StashedVectorBox))
+			while (!Stack->IsEmpty())
 			{
-				BackwardProcessEntry(Pair, Func);
+				Release(Stack->Pop(EAllowShrinking::No));
 			}
-		}
-    }
-
-    void Reset()
-    {
-        Latest.Reset();
-        Stash.Reset();
-		BitMasksPerKey.Reset();
-    }
-
-    size_t Num() const
-    {
-		typename FInternalArray::SizeType Result = Latest.Num();
 
-        for (typename FInternalArray::SizeType Index = 0; Index < Stash.Num(); Index++)
-        {
-            const FInternalArray& StashedVector = Stash[Index];
-            Result += StashedVector.Num();
-        }
-
-		return static_cast<size_t>(Result);
-    }
+			Keyed.Remove(Key);
+			return true;
+		}
 
-	UE_NONCOPYABLE(TTaskArray)
+		return false;
+	}
 
-private:
-    // We don't want to do this too often, as it'll involve a few allocations.
-    void Canonicalize()
-    {
-        if (!Latest.IsEmpty())
-        {
-			Stash.Push(MoveTemp(Latest));
-			Latest = FInternalArray{};
+	// Traverses the tasks from least recently added to most recently added,
+	// calling Callback with each task, and removing each from the array.
+	// Callback is a function-like type with the signature: void(TASK_TYPE&)
+	template<typename CallbackType>
+	void RemoveEachForward(CallbackType&& Callback)
+	{
+		while (FEntry* Entry = Head)
+		{
+			Callback(Entry->Task());
+			Head = Entry->Next;
+			Pool.ReturnTask(Entry);
 		}
-    }
+		Head = nullptr;
+		Tail = nullptr;
+		Keyed.Empty();
+		Count = 0;
+	}
 
-	template<typename FunctorType>
-	void ForwardProcessEntry(PairType Pair, const FunctorType& Func)
+	// Traverses the tasks from most recently added to least recently added,
+	// calling Callback with each task, and removing each from the array.
+	// Callback is a function-like type with the signature: void(TASK_TYPE&)
+	template<typename CallbackType>
+	void RemoveEachBackward(CallbackType&& Callback)
 	{
-		if (Pair.Key)
+		while (FEntry* Entry = Tail)
 		{
-			FBitMaskPerKeyType* const KeyedElement = BitMasksPerKey.Find(*Pair.Key);
-
-			if (KeyedElement->Num() > 0)
-			{
-				const bool bFirst = KeyedElement->BitAt(0);
-				KeyedElement->RemoveFirst();
-
-				if (bFirst)
-				{
-					// We have a deferred signal to pop this element, so do it now.
-					*Pair.Key = SKey{};
-					Pair.Val->Reset();
-					return;
-				}
-			}
+			Callback(Entry->Task());
+			Tail = Entry->Prev;
+			Pool.ReturnTask(Entry);
 		}
-
-		Func(*Pair.Val);
+		Head = nullptr;
+		Tail = nullptr;
+		Keyed.Empty();
+		Count = 0;
 	}
 
-	template<typename FunctorType>
-	void BackwardProcessEntry(PairType Pair, const FunctorType& Func)
+	// Returns the total number of tasks held by the array.
+	size_t Num() const
 	{
-		if (Pair.Key)
-		{
-			FBitMaskPerKeyType* const KeyedElement = BitMasksPerKey.Find(*Pair.Key);
+		return Count;
+	}
 
-			const int32 Num = KeyedElement->Num();
+	// Returns true if there are no tasks held by the array.
+	bool IsEmpty() const
+	{
+		return Count == 0;
+	}
 
-			if (Num > 0)
-			{
-				const bool bLast = KeyedElement->BitAt(Num - 1);
-				KeyedElement->RemoveLast();
-
-				if (bLast)
-				{
-					// We have a deferred signal to pop this element, so do it now.
-					*Pair.Key = SKey{};
-					Pair.Val->Reset();
-					return;
-				}
-			}
-		}
+private:
+	// An ordered vector of pointers to entries in the linked list.
+	using EntryStack = TArray<FEntry*, TInlineAllocator<8>>;
 
-		ASSERT(Pair.Val);
-		Func(*Pair.Val);
+	// Unlinks Entry and releases it back to the pool.
+	// Also decrements the count.
+	void Release(FEntry* Entry)
+	{
+		Entry->Unlink(Head, Tail);
+		Pool.ReturnTask(Entry);
+		Count--;
 	}
-    
-    FInternalArray Latest;
-    TArray<FInternalArray> Stash;
-	TMap<SKey, FBitMaskPerKeyType> BitMasksPerKey;
+
+	// The task pool. Holds the underlying allocator.
+	FPool& Pool;
+
+	// The doubly-linked list of tasks held by this array.
+	FEntry* Head = nullptr;
+	FEntry* Tail = nullptr;
+
+	// A map of keyed tasks keys to their ordered task entries in the linked list.
+	TMap<KeyType, EntryStack> Keyed;
+
+	// Total number of tasks held by this array.
+	size_t Count = 0;
 };
 
 } // namespace AutoRTFM
diff --git a/Engine/Source/Runtime/Core/Private/AutoRTFM/Transaction.cpp b/Engine/Source/Runtime/Core/Private/AutoRTFM/Transaction.cpp
--- a/Engine/Source/Runtime/Core/Private/AutoRTFM/Transaction.cpp
+++ b/Engine/Source/Runtime/Core/Private/AutoRTFM/Transaction.cpp
@@ -23,6 +23,8 @@ namespace AutoRTFM
 
 FTransaction::FTransaction(FContext* Context)
     : Context(Context)
+	, CommitTasks(Context->GetTaskPool())
+	, AbortTasks(Context->GetTaskPool())
 {
 }
 
@@ -105,13 +107,9 @@ void FTransaction::AbortWithoutThrowing()
     CommitTasks.Reset();
 
     Undo();
-	AbortTasks.ForEachBackward([&](TTask<void()>& Task)
+	AbortTasks.RemoveEachBackward([&](TTask<void()>& Task)
     { 
-        // Call and then reset each of the tasks in reverse order.
-        // This ensures that the task and its destructor are called in reverse chronological order,
-        // which is important if the function has captures with non-trivial destructors.
         Task();
-        Task.Reset();
     });
 
     if (IsNested())
@@ -284,11 +282,9 @@ bool FTransaction::AttemptToCommitOuterNest()
 
     AbortTasks.Reset();
 
-
-    CommitTasks.ForEachForward([] (TTask<void()>& Task)
+    CommitTasks.RemoveEachForward([] (TTask<void()>& Task)
     { 
         Task();
-        Task.Reset();
     });
 
     return true;
