[TaskGraph]
- Minor but noticeable optim in some scenario when adding many prerequisites or none at all
- Early out when passing empty prerequisite list to AddPrerequisites
- Lock once while adding many prerequisites
- Avoid interlocked operation in the normal case of AddPrerequisites

diff --git a/Engine/Source/Runtime/Core/Public/Tasks/TaskPrivate.h b/Engine/Source/Runtime/Core/Public/Tasks/TaskPrivate.h
--- a/Engine/Source/Runtime/Core/Public/Tasks/TaskPrivate.h
+++ b/Engine/Source/Runtime/Core/Public/Tasks/TaskPrivate.h
@@ -273,6 +273,12 @@ protected:
 			template<typename PrerequisiteCollectionType, decltype(std::declval<PrerequisiteCollectionType>().begin())* = nullptr>
 			void AddPrerequisites(const PrerequisiteCollectionType& InPrerequisites, bool bLockPrerequisite)
 			{
+				const int32 PrerequisiteCount = GetNum(InPrerequisites);
+				if (UNLIKELY(PrerequisiteCount <= 0))
+				{
+					return;
+				}
+
 				TASKGRAPH_VERBOSE_EVENT_SCOPE(FTaskBase::AddPrerequisites_Collection);
 
 				checkf(NumLocks.load(std::memory_order_relaxed) >= NumInitialLocks && NumLocks.load(std::memory_order_relaxed) < ExecutionFlag, TEXT("Prerequisites can be added only before the task is launched"));
@@ -280,9 +286,11 @@ protected:
 				// registering the task as a subsequent of the given prerequisite can cause its immediate launch by the prerequisite
 				// (if the prerequisite has been completed on another thread), so we need to keep the task locked by assuming that the 
 				// prerequisite can be added successfully, and release the lock if it wasn't
-				uint32 PrevNumLocks = NumLocks.fetch_add(GetNum(InPrerequisites), std::memory_order_relaxed); // relaxed because the following
+				uint32 PrevNumLocks = NumLocks.fetch_add(PrerequisiteCount, std::memory_order_relaxed); // relaxed because the following
 				// `AddSubsequent` provides required sync
 
+				bool bUnlockPrerequisites = false;
+
 				uint32 NumCompletedPrerequisites = 0;
 				for (auto& Prereq : InPrerequisites)
 				{
@@ -317,12 +325,12 @@ protected:
 						Prerequisite->AddRef(); // keep it alive until this task's execution
 						if (bLockPrerequisite)
 						{
-							Prerequisites.Push(Prerequisite); // release memory order
-						}
-						else
-						{
-							Prerequisites.PushNoLock(Prerequisite); // relaxed memory order
+							Prerequisites.Lock();
+							bLockPrerequisite = false; // don't lock again
+							bUnlockPrerequisites = true;
 						}
+
+						Prerequisites.PushNoLock(Prerequisite); // relaxed memory order
 					}
 					else
 					{
@@ -333,8 +341,16 @@ protected:
 				// This check is here to avoid the data dependency on PrevNumLocks.
 				checkf(PrevNumLocks + GetNum(InPrerequisites) < ExecutionFlag, TEXT("Max number of nested tasks reached: %d"), ExecutionFlag);
 
+				if (bUnlockPrerequisites)
+				{
+					Prerequisites.Unlock();
+				}
+
 				// unlock for prerequisites that weren't added
-				NumLocks.fetch_sub(NumCompletedPrerequisites, std::memory_order_release);
+				if (NumCompletedPrerequisites)
+				{
+					NumLocks.fetch_sub(NumCompletedPrerequisites, std::memory_order_release);
+				}
 			}
 public:
 			// The task will be executed only when all prerequisites are completed.
@@ -737,6 +753,11 @@ public:
 					return MoveTemp(Prerequisites);
 				}
 
+				void Lock()
+				{
+					Mutex.Lock();
+				}
+
 				void Unlock()
 				{
 					Mutex.Unlock();
