StringBuilder: Removing some usage of ToString() to move toward deprecating it

#rnx

diff --git a/Engine/Source/Developer/DerivedDataCache/Private/DerivedDataLegacyCacheStore.cpp b/Engine/Source/Developer/DerivedDataCache/Private/DerivedDataLegacyCacheStore.cpp
--- a/Engine/Source/Developer/DerivedDataCache/Private/DerivedDataLegacyCacheStore.cpp
+++ b/Engine/Source/Developer/DerivedDataCache/Private/DerivedDataLegacyCacheStore.cpp
@@ -147,7 +147,7 @@ bool FLegacyCacheKey::ReadValueTrailer(FCompositeBuffer& Value) const
 	{
 		FUtf8StringBuilderBase FullKey;
 		FullKey << GetFullKey();
-		const FMemoryView CompareKey(FullKey.ToString(), (FullKey.Len() + 1) * sizeof(UTF8CHAR));
+		const FMemoryView CompareKey(*FullKey, (FullKey.Len() + 1) * sizeof(UTF8CHAR));
 		if (Value.GetSize() < CompareKey.GetSize())
 		{
 			UE_LOG(LogDerivedDataCache, Display, TEXT("FLegacyCacheKey: Hash collision or short value for key %s."), *GetFullKey());
@@ -172,7 +172,7 @@ void FLegacyCacheKey::WriteValueTrailer(FCompositeBuffer& Value) const
 	{
 		FUtf8StringBuilderBase FullKey;
 		FullKey << GetFullKey();
-		Value = FCompositeBuffer(Value, FSharedBuffer::Clone(FullKey.ToString(), (FullKey.Len() + 1) * sizeof(UTF8CHAR)));
+		Value = FCompositeBuffer(Value, FSharedBuffer::Clone(*FullKey, (FullKey.Len() + 1) * sizeof(UTF8CHAR)));
 	}
 }
 
diff --git a/Engine/Source/Developer/DerivedDataCache/Private/Http/HttpHostBuilder.cpp b/Engine/Source/Developer/DerivedDataCache/Private/Http/HttpHostBuilder.cpp
--- a/Engine/Source/Developer/DerivedDataCache/Private/Http/HttpHostBuilder.cpp
+++ b/Engine/Source/Developer/DerivedDataCache/Private/Http/HttpHostBuilder.cpp
@@ -11,6 +11,7 @@
 #include "HAL/PlatformTime.h"
 #include "Http/HttpClient.h"
 #include "Logging/LogMacros.h"
+#include "Logging/StructuredLog.h"
 #include "Memory/MemoryFwd.h"
 #include "Misc/ScopeExit.h"
 #include "Serialization/CompactBinary.h"
@@ -111,7 +112,7 @@ void FHttpHostBuilder::AddFromEndpoint(const FAnsiStringView HostUrl, const FAns
 	FHttpConnectionPoolParams ConnectionPoolParams;
 	const THttpUniquePtr<IHttpConnectionPool> LocalConnectionPool = IHttpManager::Get().CreateConnectionPool(ConnectionPoolParams);
 	FHttpClientParams DefaultClientParams;
-	// we want to keep this timeout fairly low incase the host is not reachable 
+	// we want to keep this timeout fairly low in case the host is not reachable
 	DefaultClientParams.ConnectTimeout = 5000;
 	THttpUniquePtr<IHttpClient> LocalClient = LocalConnectionPool->CreateClient(DefaultClientParams);
 
@@ -128,7 +129,7 @@ void FHttpHostBuilder::AddFromEndpoint(const FAnsiStringView HostUrl, const FAns
 	{
 		Request->AddHeader(ANSITEXTVIEW("Authorization"), WriteToAnsiString<1024>(AccessToken));
 	}
-	Request->AddHeader(ANSITEXTVIEW("Accept"), "application/x-ue-cb");
+	Request->AddHeader(ANSITEXTVIEW("Accept"), ANSITEXTVIEW("application/x-ue-cb"));
 	TAnsiStringBuilder<256> Uri(InPlace, HostUrl, ANSITEXTVIEW("/status/peers"));
 	Request->SetUri(Uri);
 	THttpUniquePtr<IHttpResponse> Response;
@@ -136,7 +137,7 @@ void FHttpHostBuilder::AddFromEndpoint(const FAnsiStringView HostUrl, const FAns
 
 	if (Response->GetStatusCode() != 200)
 	{
-		UE_LOG(LogHttpHostBuilder, Warning, TEXT("Unsuccesful attempt to fetch hosts by endpoint from host: '%hs'. Status code was: %d"), Uri.ToString(), Response->GetStatusCode());
+		UE_LOGFMT(LogHttpHostBuilder, Warning, "Unsuccessful attempt to fetch hosts by endpoint from host: '{Host}'. Status code was: {Code}", Uri, Response->GetStatusCode());
 		return;
 	}
 
diff --git a/Engine/Source/Developer/DerivedDataCache/Private/HttpCacheStore.cpp b/Engine/Source/Developer/DerivedDataCache/Private/HttpCacheStore.cpp
--- a/Engine/Source/Developer/DerivedDataCache/Private/HttpCacheStore.cpp
+++ b/Engine/Source/Developer/DerivedDataCache/Private/HttpCacheStore.cpp
@@ -297,16 +297,15 @@ public:
 
 	const FString& GetDomain() const { return Domain; }
 	const FString& GetNamespace() const { return Namespace; }
+
 	const FString GetAccessToken() const
 	{
 		TAnsiStringBuilder<128> AccessTokenBuilder;
-
 		if (Access.IsValid())
 		{
 			AccessTokenBuilder << *Access;
 		}
-
-		return FString(ANSI_TO_TCHAR(AccessTokenBuilder.ToString()));
+		return FString(AccessTokenBuilder);
 	}
 
 private:
@@ -679,7 +678,7 @@ FString FHttpCacheStore::FHttpOperation::GetBodyAsString() const
 			TUtf8StringBuilder<1024> JsonStringBuilder;
 			const FCbObject ResponseObject(ResponseBody);
 			CompactBinaryToCompactJson(ResponseObject, JsonStringBuilder);
-			return JsonStringBuilder.ToString();
+			return FString(JsonStringBuilder);
 		}
 		return FString::Printf(TEXT("Invalid compact binary object of size %" UINT64_FMT), ResponseBodySize);
 	case EHttpMediaType::CompressedBinary:
diff --git a/Engine/Source/Developer/DerivedDataCache/Private/S3CacheStore.cpp b/Engine/Source/Developer/DerivedDataCache/Private/S3CacheStore.cpp
--- a/Engine/Source/Developer/DerivedDataCache/Private/S3CacheStore.cpp
+++ b/Engine/Source/Developer/DerivedDataCache/Private/S3CacheStore.cpp
@@ -12,6 +12,7 @@
 	#include "Microsoft/HideMicrosoftPlatformTypes.h"
 #endif
 
+#include "Algo/Transform.h"
 #include "Async/ParallelFor.h"
 #include "DerivedDataBackendInterface.h"
 #include "DerivedDataCacheRecord.h"
@@ -1632,8 +1633,9 @@ TUniquePtr<FArchive> FS3CacheStore::OpenFile(FStringBuilderBase& Path, const FSt
 
 	FSHAHash Hash;
 
-	auto AnsiString = StringCast<ANSICHAR>(*FString(Path.ToString()).ToUpper());
-	FSHA1::HashBuffer(AnsiString.Get(), AnsiString.Length(), Hash.Hash);
+	TAnsiStringBuilder<256> AnsiString;
+	Algo::Transform(Path, AnsiString, FChar::ToUpper);
+	FSHA1::HashBuffer(AnsiString.GetData(), AnsiString.Len(), Hash.Hash);
 
 	for (const FBundle& Bundle : Bundles)
 	{
@@ -1659,8 +1661,9 @@ bool FS3CacheStore::FileExists(FStringBuilderBase& Path, FRequestStats& Stats) c
 
 	FSHAHash Hash;
 
-	auto AnsiString = StringCast<ANSICHAR>(*FString(Path.ToString()).ToUpper());
-	FSHA1::HashBuffer(AnsiString.Get(), AnsiString.Length(), Hash.Hash);
+	TAnsiStringBuilder<256> AnsiString;
+	Algo::Transform(Path, AnsiString, FChar::ToUpper);
+	FSHA1::HashBuffer(AnsiString.GetData(), AnsiString.Len(), Hash.Hash);
 
 	for (const FBundle& Bundle : Bundles)
 	{
diff --git a/Engine/Source/Developer/DesktopPlatform/Private/LockFile.cpp b/Engine/Source/Developer/DesktopPlatform/Private/LockFile.cpp
--- a/Engine/Source/Developer/DesktopPlatform/Private/LockFile.cpp
+++ b/Engine/Source/Developer/DesktopPlatform/Private/LockFile.cpp
@@ -45,7 +45,7 @@ FLockFile::TryReadAndClear(const TCHAR* FileName, FString& OutContents)
 	{
 		FullFileNameBuilder.Prepend(TEXTVIEW("\\\\?\\"));
 	}
-	HANDLE Handle = CreateFileW(FullFileNameBuilder.ToString(), Access, WinFlags, NULL, Disposition, FILE_ATTRIBUTE_NORMAL, NULL);
+	HANDLE Handle = CreateFileW(*FullFileNameBuilder, Access, WinFlags, NULL, Disposition, FILE_ATTRIBUTE_NORMAL, NULL);
 	if (Handle != INVALID_HANDLE_VALUE)
 	{
 		ON_SCOPE_EXIT { CloseHandle(Handle); };
diff --git a/Engine/Source/Developer/DesktopPlatform/Private/TargetReceiptBuildWorker.cpp b/Engine/Source/Developer/DesktopPlatform/Private/TargetReceiptBuildWorker.cpp
--- a/Engine/Source/Developer/DesktopPlatform/Private/TargetReceiptBuildWorker.cpp
+++ b/Engine/Source/Developer/DesktopPlatform/Private/TargetReceiptBuildWorker.cpp
@@ -24,15 +24,15 @@ const FGuid FTargetReceiptBuildWorker::WorkerReceiptVersion(TEXT("dab5352e-a5a7-
 
 FTargetReceiptBuildWorker::FWorkerPath::FWorkerPath(FStringView InLocalPathRoot, FStringView InLocalPathSuffix, FStringView InRemotePathRoot, FStringView InRemotePathSuffix)
 {
-	TStringBuilder<128> LocalPathBuilder;
-	FPathViews::Append(LocalPathBuilder, InLocalPathRoot, InLocalPathSuffix);
-	LocalPath = LocalPathBuilder.ToString();
-
-	TStringBuilder<128> RemotePathBuilder;
-	FPathViews::Append(RemotePathBuilder,
-						InRemotePathRoot.IsEmpty() ? InLocalPathRoot : InRemotePathRoot,
-						InRemotePathSuffix.IsEmpty() ? InLocalPathSuffix : InRemotePathSuffix);
-	RemotePath = RemotePathBuilder.ToString();
+	TStringBuilder<128> PathBuilder;
+	FPathViews::Append(PathBuilder, InLocalPathRoot, InLocalPathSuffix);
+	LocalPath = PathBuilder;
+
+	PathBuilder.Reset();
+	FPathViews::Append(PathBuilder,
+		InRemotePathRoot.IsEmpty() ? InLocalPathRoot : InRemotePathRoot,
+		InRemotePathSuffix.IsEmpty() ? InLocalPathSuffix : InRemotePathSuffix);
+	RemotePath = PathBuilder;
 }
 
 class FPopulateWorkerFromTargetReceiptTask
diff --git a/Engine/Source/Developer/LowLevelTestsRunner/Private/UnrealReporter.cpp b/Engine/Source/Developer/LowLevelTestsRunner/Private/UnrealReporter.cpp
--- a/Engine/Source/Developer/LowLevelTestsRunner/Private/UnrealReporter.cpp
+++ b/Engine/Source/Developer/LowLevelTestsRunner/Private/UnrealReporter.cpp
@@ -148,7 +148,7 @@ public:
 			}
 
 			Stream << "at " << Result.getSourceInfo().file << '(' << (uint64)Result.getSourceInfo().line << ')';
-			Xml.writeText(Stream.ToString(), XmlFormatting::Newline);
+			Xml.writeText(*Stream, XmlFormatting::Newline);
 		}
 	}
 
diff --git a/Engine/Source/Developer/ShaderPreprocessor/Private/ShaderPreprocessor.cpp b/Engine/Source/Developer/ShaderPreprocessor/Private/ShaderPreprocessor.cpp
--- a/Engine/Source/Developer/ShaderPreprocessor/Private/ShaderPreprocessor.cpp
+++ b/Engine/Source/Developer/ShaderPreprocessor/Private/ShaderPreprocessor.cpp
@@ -669,7 +669,7 @@ static void AddStbDefine(stb_arena* MacroArena, macro_definition**& StbDefines,
 	Define.AppendChar(' ');
 	Define.Append(Value);
 
-	arrput(StbDefines, pp_define(MacroArena, Define.ToString()));
+	arrput(StbDefines, pp_define(MacroArena, *Define));
 }
 
 static void AddStbDefines(stb_arena* MacroArena, macro_definition**& StbDefines, const FShaderCompilerDefinitions& Defines)
diff --git a/Engine/Source/Developer/Zen/Private/ZenServerHttp.cpp b/Engine/Source/Developer/Zen/Private/ZenServerHttp.cpp
--- a/Engine/Source/Developer/Zen/Private/ZenServerHttp.cpp
+++ b/Engine/Source/Developer/Zen/Private/ZenServerHttp.cpp
@@ -138,7 +138,7 @@ namespace UE::Zen {
 	{
 		TStringBuilder<128> Sb;
 		Sb << Header << TEXTVIEW(": ") << Value;
-		Headers.Emplace(Sb.ToString());
+		Headers.Emplace(Sb);
 	}
 
 	FZenHttpRequest::Result FZenHttpRequest::PerformBlockingPut(const TCHAR* Uri, const FCompositeBuffer& Buffer, EContentType ContentType)
diff --git a/Engine/Source/Developer/Zen/Private/ZenServerInterface.cpp b/Engine/Source/Developer/Zen/Private/ZenServerInterface.cpp
--- a/Engine/Source/Developer/Zen/Private/ZenServerInterface.cpp
+++ b/Engine/Source/Developer/Zen/Private/ZenServerInterface.cpp
@@ -19,6 +19,7 @@
 #include "HAL/PlatformProcess.h"
 #include "HAL/PlatformTime.h"
 #include "Logging/LogScopedCategoryAndVerbosityOverride.h"
+#include "Logging/StructuredLog.h"
 #include "Misc/App.h"
 #include "Misc/ConfigCacheIni.h"
 #include "Misc/ConfigContext.h"
@@ -1741,12 +1742,12 @@ FZenServiceInstance::PingService(uint32_t AttemptCount)
 
 		TStringBuilder<128> ZenDomain;
 		ZenDomain << HostName << TEXT(":") << Port;
-		Zen::FZenHttpRequest Request(ZenDomain.ToString(), false, 500);
+		Zen::FZenHttpRequest Request(ZenDomain, false, 500);
 		Zen::FZenHttpRequest::Result Result = Request.PerformBlockingDownload(TEXTVIEW("health/ready"), nullptr, Zen::EContentType::Text);
 		
 		if (Result == Zen::FZenHttpRequest::Result::Success && Zen::IsSuccessCode(Request.GetResponseCode()))
 		{
-			UE_LOG(LogZenServiceInstance, Display, TEXT("Unreal Zen Storage Server HTTP service at %s status: %s."), ZenDomain.ToString(), *Request.GetResponseAsString());
+			UE_LOGFMT(LogZenServiceInstance, Display, "Unreal Zen Storage Server HTTP service at {Domain} status: {Response}.", ZenDomain, Request.GetResponseAsString());
 			return true;
 		}
 
@@ -1754,13 +1755,13 @@ FZenServiceInstance::PingService(uint32_t AttemptCount)
 		{
 			if (Attempt == AttemptCount)
 			{
-				UE_LOG(LogZenServiceInstance, Warning, TEXT("Unable to reach Unreal Zen Storage Server HTTP service at %s. Status: %d . Response: %s"), ZenDomain.ToString(), Request.GetResponseCode(), *Request.GetResponseAsString());
+				UE_LOGFMT(LogZenServiceInstance, Warning, "Unable to reach Unreal Zen Storage Server HTTP service at {Domain}. Status: {Code}. Response: {Response}", ZenDomain, Request.GetResponseCode(), Request.GetResponseAsString());
 				break;
 			}
 		}
 		else
 		{
-			UE_LOG(LogZenServiceInstance, Display, TEXT("Unable to reach Unreal Zen Storage Server HTTP service at %s. Status: %d . Response: %s"), ZenDomain.ToString(), Request.GetResponseCode(), *Request.GetResponseAsString());
+			UE_LOGFMT(LogZenServiceInstance, Display, "Unable to reach Unreal Zen Storage Server HTTP service at {Domain}. Status: {Code}. Response: {Response}", ZenDomain, Request.GetResponseCode(), Request.GetResponseAsString());
 			break;
 		}
 	}
@@ -2603,9 +2604,7 @@ FZenServiceInstance::GetGCStatus(FGCStatus& Status)
 bool 
 FZenServiceInstance::RequestGC(const bool* OverrideCollectSmallObjects, const uint32* OverrideMaxCacheDuration)
 {
-	TStringBuilder<128> ZenDomain;
-	ZenDomain << HostName << TEXT(":") << Port;
-	UE::Zen::FZenHttpRequest Request(ZenDomain.ToString(), false);
+	UE::Zen::FZenHttpRequest Request(WriteToString<128>(HostName, ':', Port), /*bLogErrors*/ false);
 
 	TCHAR Separators[] = {TEXT('?'), TEXT('&')};
 	int32 SeparatorIndex = 0;
@@ -2624,7 +2623,7 @@ FZenServiceInstance::RequestGC(const bool* OverrideCollectSmallObjects, const ui
 		SeparatorIndex = FMath::Min(SeparatorIndex + 1, (int32)UE_ARRAY_COUNT(Separators));
 	}
 
-	FZenHttpRequest::Result Result = Request.PerformBlockingPost(Query.ToString(), FMemoryView());
+	FZenHttpRequest::Result Result = Request.PerformBlockingPost(Query, FMemoryView());
 
 	if (Result == Zen::FZenHttpRequest::Result::Success && Request.GetResponseCode() == 200)
 	{
diff --git a/Engine/Source/Developer/Zen/Private/ZenServerState.cpp b/Engine/Source/Developer/Zen/Private/ZenServerState.cpp
--- a/Engine/Source/Developer/Zen/Private/ZenServerState.cpp
+++ b/Engine/Source/Developer/Zen/Private/ZenServerState.cpp
@@ -876,7 +876,7 @@ ZenLockFileData::ReadCbLockFile(const TCHAR* FileName)
 		FullFileNameBuilder.Prepend(TEXTVIEW("\\\\?\\"));
 	}
 
-	HANDLE Handle = CreateFileW(FullFileNameBuilder.ToString(), Access, WinFlags, NULL, Create, FILE_ATTRIBUTE_NORMAL, NULL);
+	HANDLE Handle = CreateFileW(*FullFileNameBuilder, Access, WinFlags, NULL, Create, FILE_ATTRIBUTE_NORMAL, NULL);
 	if (Handle != INVALID_HANDLE_VALUE)
 	{
 		ON_SCOPE_EXIT{ CloseHandle(Handle); };
diff --git a/Engine/Source/Developer/Zen/Private/ZenVersion.cpp b/Engine/Source/Developer/Zen/Private/ZenVersion.cpp
--- a/Engine/Source/Developer/Zen/Private/ZenVersion.cpp
+++ b/Engine/Source/Developer/Zen/Private/ZenVersion.cpp
@@ -81,7 +81,7 @@ FString FZenVersion::ToString(bool bDetailed) const
 		Builder.AppendChar(TEXT('-'));
 		Builder << Details;
 	}
-	return Builder.ToString();
+	return FString(Builder);
 }
 
 bool FZenVersion::operator<(FZenVersion& Other) const
diff --git a/Engine/Source/Runtime/Core/Private/IO/IoStore.cpp b/Engine/Source/Runtime/Core/Private/IO/IoStore.cpp
--- a/Engine/Source/Runtime/Core/Private/IO/IoStore.cpp
+++ b/Engine/Source/Runtime/Core/Private/IO/IoStore.cpp
@@ -990,7 +990,7 @@ public:
 				Sb.Append(FString::Printf(TEXT("_s%d"), PartitionIndex));
 			}
 			Sb.Append(TEXT(".ucas"));
-			OutPaths.Add(Sb.ToString());
+			OutPaths.Emplace(Sb);
 		}
 	}
 
diff --git a/Engine/Source/Runtime/Core/Private/Internationalization/FastDecimalFormat.cpp b/Engine/Source/Runtime/Core/Private/Internationalization/FastDecimalFormat.cpp
--- a/Engine/Source/Runtime/Core/Private/Internationalization/FastDecimalFormat.cpp
+++ b/Engine/Source/Runtime/Core/Private/Internationalization/FastDecimalFormat.cpp
@@ -535,7 +535,7 @@ FString CultureInvariantDecimalToString(const double InVal, const TCHAR*& InBuff
 	OutStr.InsertAt(0, FinalPrefixStr);
 	OutStr += FinalSuffixStr;
 	
-	return OutStr.ToString();
+	return FString(OutStr);
 }
 
 void FractionalToString(const double InVal, const FDecimalNumberFormattingRules& InFormattingRules, FNumberFormattingOptions InFormattingOptions, FString& OutString)
diff --git a/Engine/Source/Runtime/Core/Private/Internationalization/TextKey.cpp b/Engine/Source/Runtime/Core/Private/Internationalization/TextKey.cpp
--- a/Engine/Source/Runtime/Core/Private/Internationalization/TextKey.cpp
+++ b/Engine/Source/Runtime/Core/Private/Internationalization/TextKey.cpp
@@ -680,7 +680,7 @@ void FTextKey::SerializeAsString(FArchive& Ar)
 	{
 		TextKeyUtil::FInlineStringBuilder StrBuilder;
 		AppendString(StrBuilder);
-		TextKeyUtil::SaveKeyString(Ar, StrBuilder.ToString());
+		TextKeyUtil::SaveKeyString(Ar, *StrBuilder);
 	}
 }
 
@@ -710,7 +710,7 @@ void FTextKey::SerializeWithHash(FArchive& Ar)
 
 		TextKeyUtil::FInlineStringBuilder StrBuilder;
 		AppendString(StrBuilder);
-		TextKeyUtil::SaveKeyString(Ar, StrBuilder.ToString());
+		TextKeyUtil::SaveKeyString(Ar, *StrBuilder);
 	}
 }
 
@@ -740,7 +740,7 @@ void FTextKey::SerializeDiscardHash(FArchive& Ar)
 
 		TextKeyUtil::FInlineStringBuilder StrBuilder;
 		AppendString(StrBuilder);
-		TextKeyUtil::SaveKeyString(Ar, StrBuilder.ToString());
+		TextKeyUtil::SaveKeyString(Ar, *StrBuilder);
 	}
 }
 
diff --git a/Engine/Source/Runtime/Core/Private/Misc/ConfigCacheIni.cpp b/Engine/Source/Runtime/Core/Private/Misc/ConfigCacheIni.cpp
--- a/Engine/Source/Runtime/Core/Private/Misc/ConfigCacheIni.cpp
+++ b/Engine/Source/Runtime/Core/Private/Misc/ConfigCacheIni.cpp
@@ -6090,11 +6090,11 @@ void FConfigFile::UpdateSections(const TCHAR* DiskFilename, const TCHAR* IniRoot
 				FString* ExistingSectionText = SectionTexts.Find(SectionName);
 				if (ExistingSectionText)
 				{
-					ExistingSectionText->Append(SectionText.ToString(), SectionText.Len());
+					ExistingSectionText->Append(SectionText);
 				}
 				else
 				{
-					SectionTexts.Add(SectionName, FString(SectionText.ToString()));
+					SectionTexts.Emplace(SectionName, SectionText);
 				}
 			}
 
diff --git a/Engine/Source/Runtime/Core/Private/Misc/DateTime.cpp b/Engine/Source/Runtime/Core/Private/Misc/DateTime.cpp
--- a/Engine/Source/Runtime/Core/Private/Misc/DateTime.cpp
+++ b/Engine/Source/Runtime/Core/Private/Misc/DateTime.cpp
@@ -229,7 +229,7 @@ FString FDateTime::ToString(const TCHAR* Format) const
 {
 	TStringBuilder<48> Result;
 	ToString(Format, Result);
-	return Result.ToString();
+	return FString(Result);
 }
 
 void FDateTime::ToString(const TCHAR* Format, FStringBuilderBase& Result) const
@@ -332,7 +332,7 @@ FString FDateTime::ToFormattedString(const TCHAR* Format) const
 		}
 	}
 
-	return Result.ToString();
+	return FString(Result);
 }
 
 
diff --git a/Engine/Source/Runtime/Core/Private/Misc/HierarchicalLogArchive.cpp b/Engine/Source/Runtime/Core/Private/Misc/HierarchicalLogArchive.cpp
--- a/Engine/Source/Runtime/Core/Private/Misc/HierarchicalLogArchive.cpp
+++ b/Engine/Source/Runtime/Core/Private/Misc/HierarchicalLogArchive.cpp
@@ -2,39 +2,35 @@
 
 #include "Misc/HierarchicalLogArchive.h"
 
-#include "Containers/StringConv.h"
-#include "Containers/StringFwd.h"
 #include "Misc/StringBuilder.h"
 
-class FArchive;
-
 FHierarchicalLogArchive::FHierarchicalLogArchive(FArchive& InInnerArchive)
 	: FArchiveProxy(InInnerArchive)
 	, Indentation(0)
-{}
+{
+}
 
 void FHierarchicalLogArchive::WriteLine(const FString& InLine, bool bIndent)
 {
-	TStringBuilder<512> Builder;
+	TAnsiStringBuilder<512> Builder;
 
 	if (Indentation > 0)
 	{
 		for (int i = 0; i < Indentation - 1; ++i)
 		{
-			Builder += TEXT(" |  ");
+			Builder << " |  ";
 		}
 
-		Builder += TEXT(" |- ");
+		Builder << " |- ";
 	}
 	
 	if (bIndent)
 	{
-		Builder += TEXT("[+] ");
+		Builder << "[+] ";
 	}
 
-	Builder += InLine;
-	Builder += LINE_TERMINATOR;
-	FString LineEntry = Builder.ToString();
+	Builder << InLine;
+	Builder << LINE_TERMINATOR_ANSI;
 
-	Serialize(TCHAR_TO_ANSI(*LineEntry), LineEntry.Len());
-}
\ No newline at end of file
+	Serialize(Builder.GetData(), Builder.Len() * sizeof(ANSICHAR));
+}
diff --git a/Engine/Source/Runtime/Core/Private/Misc/OutputDeviceHelper.cpp b/Engine/Source/Runtime/Core/Private/Misc/OutputDeviceHelper.cpp
--- a/Engine/Source/Runtime/Core/Private/Misc/OutputDeviceHelper.cpp
+++ b/Engine/Source/Runtime/Core/Private/Misc/OutputDeviceHelper.cpp
@@ -102,7 +102,7 @@ FString FOutputDeviceHelper::FormatLogLine(
 {
 	TStringBuilder<512> Format;
 	AppendFormatLogLine(Format, Verbosity, Category, Message, LogTime, Time, OutCategoryIndex);
-	return Format.ToString();
+	return FString(Format);
 }
 
 void FOutputDeviceHelper::FormatCastAndSerializeLine(
diff --git a/Engine/Source/Runtime/Core/Private/Misc/PathPermissionList.cpp b/Engine/Source/Runtime/Core/Private/Misc/PathPermissionList.cpp
--- a/Engine/Source/Runtime/Core/Private/Misc/PathPermissionList.cpp
+++ b/Engine/Source/Runtime/Core/Private/Misc/PathPermissionList.cpp
@@ -614,5 +614,5 @@ FString FPathPermissionList::ToString() const
 		AppendList(AllowTree);
 	}
 
-	return StringBuilder.ToString();
+	return FString(StringBuilder);
 }
\ No newline at end of file
diff --git a/Engine/Source/Runtime/Core/Private/Misc/Paths.cpp b/Engine/Source/Runtime/Core/Private/Misc/Paths.cpp
--- a/Engine/Source/Runtime/Core/Private/Misc/Paths.cpp
+++ b/Engine/Source/Runtime/Core/Private/Misc/Paths.cpp
@@ -918,7 +918,7 @@ FString FPaths::FindCorrectCase(const FString& Path)
 			}
 		}
 	);
-	return Builder.ToString();
+	return FString(Builder);
 #else
 	return IFileManager::Get().GetFilenameOnDisk(*Path);
 #endif
diff --git a/Engine/Source/Runtime/Core/Private/Modules/VisualizerDebuggingState.cpp b/Engine/Source/Runtime/Core/Private/Modules/VisualizerDebuggingState.cpp
--- a/Engine/Source/Runtime/Core/Private/Modules/VisualizerDebuggingState.cpp
+++ b/Engine/Source/Runtime/Core/Private/Modules/VisualizerDebuggingState.cpp
@@ -92,7 +92,7 @@ EVisualizerDebuggingStateResult FVisualizerDebuggingState::AssignImpl(const FGui
 	UniqueId.AppendString(ThisGuidString, EGuidFormats::DigitsLower);
 
 	// Search the existing string for a potential string collision (this is extremely unlikely)
-	if (GuidString != nullptr && FCStringAnsi::Strstr(GuidString, ThisGuidString.ToString()) != nullptr)
+	if (GuidString != nullptr && FCStringAnsi::Strstr(GuidString, *ThisGuidString) != nullptr)
 	{
 		return EVisualizerDebuggingStateResult::StringCollision;
 	}
diff --git a/Engine/Source/Runtime/Core/Private/String/LexFromString.cpp b/Engine/Source/Runtime/Core/Private/String/LexFromString.cpp
--- a/Engine/Source/Runtime/Core/Private/String/LexFromString.cpp
+++ b/Engine/Source/Runtime/Core/Private/String/LexFromString.cpp
@@ -9,131 +9,131 @@
 void LexFromString(int8& OutValue, const FStringView& InString)
 {
 	TStringBuilder<64> Builder;
-	LexFromString(OutValue, (Builder << InString).ToString());
+	LexFromString(OutValue, *(Builder << InString));
 }
 
 void LexFromString(int16& OutValue, const FStringView& InString)
 {
 	TStringBuilder<64> Builder;
-	LexFromString(OutValue, (Builder << InString).ToString());
+	LexFromString(OutValue, *(Builder << InString));
 }
 
 void LexFromString(int32& OutValue, const FStringView& InString)
 {
 	TStringBuilder<64> Builder;
-	LexFromString(OutValue, (Builder << InString).ToString());
+	LexFromString(OutValue, *(Builder << InString));
 }
 
 void LexFromString(int64& OutValue, const FStringView& InString)
 {
 	TStringBuilder<64> Builder;
-	LexFromString(OutValue, (Builder << InString).ToString());
+	LexFromString(OutValue, *(Builder << InString));
 }
 
 void LexFromString(uint8& OutValue, const FStringView& InString)
 {
 	TStringBuilder<64> Builder;
-	LexFromString(OutValue, (Builder << InString).ToString());
+	LexFromString(OutValue, *(Builder << InString));
 }
 
 void LexFromString(uint16& OutValue, const FStringView& InString)
 {
 	TStringBuilder<64> Builder;
-	LexFromString(OutValue, (Builder << InString).ToString());
+	LexFromString(OutValue, *(Builder << InString));
 }
 
 void LexFromString(uint32& OutValue, const FStringView& InString)
 {
 	TStringBuilder<64> Builder;
-	LexFromString(OutValue, (Builder << InString).ToString());
+	LexFromString(OutValue, *(Builder << InString));
 }
 
 void LexFromString(uint64& OutValue, const FStringView& InString)
 {
 	TStringBuilder<64> Builder;
-	LexFromString(OutValue, (Builder << InString).ToString());
+	LexFromString(OutValue, *(Builder << InString));
 }
 
 void LexFromString(float& OutValue, const FStringView& InString)
 {
 	TStringBuilder<64> Builder;
-	LexFromString(OutValue, (Builder << InString).ToString());
+	LexFromString(OutValue, *(Builder << InString));
 }
 
 void LexFromString(double& OutValue, const FStringView& InString)
 {
 	TStringBuilder<64> Builder;
-	LexFromString(OutValue, (Builder << InString).ToString());
+	LexFromString(OutValue, *(Builder << InString));
 }
 
 void LexFromString(bool& OutValue, const FStringView& InString)
 {
 	TStringBuilder<64> Builder;
-	LexFromString(OutValue, (Builder << InString).ToString());
+	LexFromString(OutValue, *(Builder << InString));
 }
 
 void LexFromString(int8& OutValue, const FUtf8StringView& InString)
 {
 	TUtf8StringBuilder<64> Builder;
-	LexFromString(OutValue, (Builder << InString).ToString());
+	LexFromString(OutValue, *(Builder << InString));
 }
 
 void LexFromString(int16& OutValue, const FUtf8StringView& InString)
 {
 	TUtf8StringBuilder<64> Builder;
-	LexFromString(OutValue, (Builder << InString).ToString());
+	LexFromString(OutValue, *(Builder << InString));
 }
 
 void LexFromString(int32& OutValue, const FUtf8StringView& InString)
 {
 	TUtf8StringBuilder<64> Builder;
-	LexFromString(OutValue, (Builder << InString).ToString());
+	LexFromString(OutValue, *(Builder << InString));
 }
 
 void LexFromString(int64& OutValue, const FUtf8StringView& InString)
 {
 	TUtf8StringBuilder<64> Builder;
-	LexFromString(OutValue, (Builder << InString).ToString());
+	LexFromString(OutValue, *(Builder << InString));
 }
 
 void LexFromString(uint8& OutValue, const FUtf8StringView& InString)
 {
 	TUtf8StringBuilder<64> Builder;
-	LexFromString(OutValue, (Builder << InString).ToString());
+	LexFromString(OutValue, *(Builder << InString));
 }
 
 void LexFromString(uint16& OutValue, const FUtf8StringView& InString)
 {
 	TUtf8StringBuilder<64> Builder;
-	LexFromString(OutValue, (Builder << InString).ToString());
+	LexFromString(OutValue, *(Builder << InString));
 }
 
 void LexFromString(uint32& OutValue, const FUtf8StringView& InString)
 {
 	TUtf8StringBuilder<64> Builder;
-	LexFromString(OutValue, (Builder << InString).ToString());
+	LexFromString(OutValue, *(Builder << InString));
 }
 
 void LexFromString(uint64& OutValue, const FUtf8StringView& InString)
 {
 	TUtf8StringBuilder<64> Builder;
-	LexFromString(OutValue, (Builder << InString).ToString());
+	LexFromString(OutValue, *(Builder << InString));
 }
 
 void LexFromString(float& OutValue, const FUtf8StringView& InString)
 {
 	TUtf8StringBuilder<64> Builder;
-	LexFromString(OutValue, (Builder << InString).ToString());
+	LexFromString(OutValue, *(Builder << InString));
 }
 
 void LexFromString(double& OutValue, const FUtf8StringView& InString)
 {
 	TUtf8StringBuilder<64> Builder;
-	LexFromString(OutValue, (Builder << InString).ToString());
+	LexFromString(OutValue, *(Builder << InString));
 }
 
 void LexFromString(bool& OutValue, const FUtf8StringView& InString)
 {
 	TUtf8StringBuilder<64> Builder;
-	LexFromString(OutValue, (Builder << InString).ToString());
+	LexFromString(OutValue, *(Builder << InString));
 }
diff --git a/Engine/Source/Runtime/Core/Public/Misc/Timecode.h b/Engine/Source/Runtime/Core/Public/Misc/Timecode.h
--- a/Engine/Source/Runtime/Core/Public/Misc/Timecode.h
+++ b/Engine/Source/Runtime/Core/Public/Misc/Timecode.h
@@ -380,7 +380,7 @@ public:
 			int32 ClampedSubframe = static_cast<int32>(FMath::Clamp(100*Subframe,0,99));
 			Builder.Appendf(TEXT(".%02d"), ClampedSubframe);
 		}
-		return Builder.ToString();
+		return FString(Builder);
 	}
 
 	/**
diff --git a/Engine/Source/Runtime/Core/Public/ProfilingDebugging/ScopedTimers.h b/Engine/Source/Runtime/Core/Public/ProfilingDebugging/ScopedTimers.h
--- a/Engine/Source/Runtime/Core/Public/ProfilingDebugging/ScopedTimers.h
+++ b/Engine/Source/Runtime/Core/Public/ProfilingDebugging/ScopedTimers.h
@@ -302,7 +302,7 @@ public:
 			Msg.Appendf(TEXT(" (total: %s)"), *FPlatformTime::PrettyTime(TotalTime));
 		}
 
-		LogFunc(Msg.ToString());
+		LogFunc(*Msg);
 	}
 
 private:
diff --git a/Engine/Source/Runtime/Core/Tests/Misc/TimecodeTest.cpp b/Engine/Source/Runtime/Core/Tests/Misc/TimecodeTest.cpp
--- a/Engine/Source/Runtime/Core/Tests/Misc/TimecodeTest.cpp
+++ b/Engine/Source/Runtime/Core/Tests/Misc/TimecodeTest.cpp
@@ -291,7 +291,7 @@ namespace UE::TimecodeParserTest::Private
 			}
 			Builder.Appendf(TEXT(".%.*s"), NumDigits, SubFrameValue.GetData());
 		}
-		return Builder.ToString();
+		return FString(Builder);
 	};
 
 	static bool IsSame(const FTimecode& InTimecode, const FTimecode& InOtherTimecode)
diff --git a/Engine/Source/Runtime/CoreUObject/Private/Internationalization/TextPackageNamespaceUtil.cpp b/Engine/Source/Runtime/CoreUObject/Private/Internationalization/TextPackageNamespaceUtil.cpp
--- a/Engine/Source/Runtime/CoreUObject/Private/Internationalization/TextPackageNamespaceUtil.cpp
+++ b/Engine/Source/Runtime/CoreUObject/Private/Internationalization/TextPackageNamespaceUtil.cpp
@@ -194,7 +194,7 @@ FString TextNamespaceUtil::GenerateDeterministicTextKey(UObject* InTextOwner, co
 	auto GetNameKeyHash = [](const FName Name)
 	{
 		FNameBuilder Builder(Name);
-		return TextKeyUtil::HashString(Builder.ToString(), Builder.Len());
+		return TextKeyUtil::HashString(Builder.ToView());
 	};
 
 	auto GetObjectKeyHash = [](const UObject* Obj)
@@ -204,7 +204,7 @@ FString TextNamespaceUtil::GenerateDeterministicTextKey(UObject* InTextOwner, co
 		{
 			Obj->GetPathName(nullptr, Builder);
 		}
-		return TextKeyUtil::HashString(Builder.ToString(), Builder.Len());
+		return TextKeyUtil::HashString(Builder.ToView());
 	};
 
 	// Build a (hopefully) unique and deterministic key from a combination of the text owner and text property info
diff --git a/Engine/Source/Runtime/CoreUObject/Private/Misc/PackageName.cpp b/Engine/Source/Runtime/CoreUObject/Private/Misc/PackageName.cpp
--- a/Engine/Source/Runtime/CoreUObject/Private/Misc/PackageName.cpp
+++ b/Engine/Source/Runtime/CoreUObject/Private/Misc/PackageName.cpp
@@ -293,7 +293,7 @@ struct FLongPackagePathsSingleton
 		TStringBuilder<256> AbsolutePathBuilder;
 		FPathViews::ToAbsolutePath(RelativeContentPath, AbsolutePathBuilder);
 
-		return FMountPoint(FString(RootPath), MoveTemp(RelativeContentPath), AbsolutePathBuilder.ToString(), Flags);
+		return FMountPoint(FString(RootPath), MoveTemp(RelativeContentPath), FString(AbsolutePathBuilder), Flags);
 	}
 
 	void InsertMountPoint(const FString& RootPath, const FString& ContentPath)
@@ -484,7 +484,7 @@ private:
 			TStringBuilder<256> AbsolutePathBuilder;
 			FPathViews::ToAbsolutePath(ContentPathRelative, AbsolutePathBuilder);
 			FMountPoint& MountPoint = *MountPoints.Add_GetRef(MakeUnique<FMountPoint>(
-				FString(RootPath), FString(ContentPathRelative), AbsolutePathBuilder.ToString(), 
+				FString(RootPath), FString(ContentPathRelative), FString(AbsolutePathBuilder),
 				InFlags));
 
 			// We are appending to the MountPoint list, and will have only a few duplicates. If any duplicates do occur,
@@ -1014,7 +1014,7 @@ bool FPackageName::TryConvertToMountedPath(FStringView InFilePathOrPackageName,
 	}
 	if (OutExtension)
 	{
-		*OutExtension = ExtensionType == EPackageExtension::Custom ? CustomExtension.ToString() : LexToString(ExtensionType);
+		*OutExtension = ExtensionType == EPackageExtension::Custom ? *CustomExtension : LexToString(ExtensionType);
 	}
 	return true;
 }
@@ -2076,7 +2076,7 @@ FPackageName::EPackageLocationFilter FPackageName::InternalDoesPackageExistEx(co
 	FText Reason;
 	if (!FPackageName::IsValidTextForLongPackageName(PackageName, &Reason))
 	{
-		UE_LOG(LogPackageName, Error, TEXT("DoesPackageExist: DoesPackageExist FAILED: '%s' is not a long packagename name. Reason: %s"), PackageName.ToString(), *Reason.ToString());
+		UE_LOG(LogPackageName, Error, TEXT("DoesPackageExist: DoesPackageExist FAILED: '%s' is not a long packagename name. Reason: %s"), *PackageName, *Reason.ToString());
 		return EPackageLocationFilter::None;
 	}
 
diff --git a/Engine/Source/Runtime/CoreUObject/Private/UObject/PropertyBaseObject.cpp b/Engine/Source/Runtime/CoreUObject/Private/UObject/PropertyBaseObject.cpp
--- a/Engine/Source/Runtime/CoreUObject/Private/UObject/PropertyBaseObject.cpp
+++ b/Engine/Source/Runtime/CoreUObject/Private/UObject/PropertyBaseObject.cpp
@@ -385,12 +385,12 @@ bool FObjectPropertyBase::ParseObjectPropertyValue(const FProperty* Property, UO
 
 			// ignore the object class, it isn't fully qualified, and searching globally might get the wrong one!
 			// Try the find the object.
-			out_ResolvedValue = FObjectPropertyBase::FindImportedObject(Property, OwnerObject, ObjectClass, RequiredMetaClass, Temp.ToString(), PortFlags, InSerializeContext, bAllowAnyPackage);
+			out_ResolvedValue = FObjectPropertyBase::FindImportedObject(Property, OwnerObject, ObjectClass, RequiredMetaClass, *Temp, PortFlags, InSerializeContext, bAllowAnyPackage);
 		}
 		else
 		{
 			// Try the find the object.
-			out_ResolvedValue = FObjectPropertyBase::FindImportedObject(Property, OwnerObject, ObjectClass, RequiredMetaClass, Temp.ToString(), PortFlags, InSerializeContext, bAllowAnyPackage);
+			out_ResolvedValue = FObjectPropertyBase::FindImportedObject(Property, OwnerObject, ObjectClass, RequiredMetaClass, *Temp, PortFlags, InSerializeContext, bAllowAnyPackage);
 		}
 
 		if ( out_ResolvedValue != nullptr && !out_ResolvedValue.GetClass()->IsChildOf(RequiredMetaClass) )
diff --git a/Engine/Source/Runtime/RenderCore/Private/ShaderMap.cpp b/Engine/Source/Runtime/RenderCore/Private/ShaderMap.cpp
--- a/Engine/Source/Runtime/RenderCore/Private/ShaderMap.cpp
+++ b/Engine/Source/Runtime/RenderCore/Private/ShaderMap.cpp
@@ -446,7 +446,7 @@ FString FShaderMapBase::ToString() const
 		Code->ToString(String);
 	}
 
-	return String.ToString();
+	return FString(String);
 }
 
 void FShaderMapBase::DestroyContent()
diff --git a/Engine/Source/Runtime/RenderCore/Private/ShaderParameterParser.cpp b/Engine/Source/Runtime/RenderCore/Private/ShaderParameterParser.cpp
--- a/Engine/Source/Runtime/RenderCore/Private/ShaderParameterParser.cpp
+++ b/Engine/Source/Runtime/RenderCore/Private/ShaderParameterParser.cpp
@@ -85,14 +85,14 @@ static void IterateShaderParameterMembersInternal(
 					const FString RealBindingName = FString::Printf(TEXT("%s_%d"), Member.GetName(), Index);
 
 					ShaderBindingNameBuilder.Append(RealBindingName);
-					Lambda(ParametersMetadata, Member, ShaderBindingNameBuilder.ToString(), MemberOffset + Index * ElementSize);
+					Lambda(ParametersMetadata, Member, *ShaderBindingNameBuilder, MemberOffset + Index * ElementSize);
 					ShaderBindingNameBuilder.RemoveSuffix(RealBindingName.Len());
 				}
 			}
 			else
 			{
 				ShaderBindingNameBuilder.Append(Member.GetName());
-				Lambda(ParametersMetadata, Member, ShaderBindingNameBuilder.ToString(), MemberOffset);
+				Lambda(ParametersMetadata, Member, *ShaderBindingNameBuilder, MemberOffset);
 				ShaderBindingNameBuilder.RemoveSuffix(MemberNameLength);
 			}
 		}
@@ -971,7 +971,7 @@ FString FShaderParameterParser::GenerateBindlessParameterDeclaration(const FPars
 	Result << FullType << TEXT(" ") << FunctionPrefix << Name << TEXT("() { return ") << BindlessAccess << TEXT("; } ");
 	Result << TEXT("static const ") << FullType << TEXT(" ") << Name << TEXT(" = ") << FunctionPrefix << Name << TEXT("();");
 
-	return Result.ToString();
+	return FString(Result);
 }
 
 void FShaderParameterParser::ApplyBindlessModifications(FString& PreprocessedShaderSource)
