Fixed bug in memstack move constructor that would not free existing pages.

diff --git a/Engine/Source/Runtime/Core/Private/Misc/MemStack.cpp b/Engine/Source/Runtime/Core/Private/Misc/MemStack.cpp
--- a/Engine/Source/Runtime/Core/Private/Misc/MemStack.cpp
+++ b/Engine/Source/Runtime/Core/Private/Misc/MemStack.cpp
@@ -310,13 +310,6 @@ void FPageAllocator::UpdateStats()
 -----------------------------------------------------------------------------*/
  
 FMemStackBase::FMemStackBase(EPageSize InPageSize)
-: Top(nullptr)
-, End(nullptr)
-, TopChunk(nullptr)
-, TopMark(nullptr)
-, NumMarks(0)
-, PageSize(InPageSize)
-, bShouldEnforceAllocMarks(false)
 {
 	// By fetching the FPageAllocator singleton's address here we can guarantee
 	// that its lifetime spans that of all FMemStackBase objects.
diff --git a/Engine/Source/Runtime/Core/Public/Misc/MemStack.h b/Engine/Source/Runtime/Core/Public/Misc/MemStack.h
--- a/Engine/Source/Runtime/Core/Public/Misc/MemStack.h
+++ b/Engine/Source/Runtime/Core/Public/Misc/MemStack.h
@@ -96,6 +96,7 @@ public:
 
 	FMemStackBase& operator=(FMemStackBase&& Other)
 	{
+		FreeChunks(nullptr);
 		Top = Other.Top;
 		End = Other.End;
 		TopChunk = Other.TopChunk;
@@ -221,21 +222,21 @@ private:
 	CORE_API void FreeChunks( FTaggedMemory* NewTopChunk );
 
 	// Variables.
-	uint8*			Top;				// Top of current chunk (Top<=End).
-	uint8*			End;				// End of current chunk.
-	FTaggedMemory*	TopChunk;			// Only chunks 0..ActiveChunks-1 are valid.
+	uint8*			Top = nullptr;				// Top of current chunk (Top<=End).
+	uint8*			End = nullptr;				// End of current chunk.
+	FTaggedMemory*	TopChunk = nullptr;			// Only chunks 0..ActiveChunks-1 are valid.
 
 	/** The top mark on the stack. */
-	class FMemMark*	TopMark;
+	class FMemMark*	TopMark = nullptr;
 
 	/** The number of marks on this stack. */
-	int32 NumMarks;
+	int32 NumMarks = 0;
 
 	/** The page size to use when allocating. */
-	EPageSize PageSize;
+	EPageSize PageSize = EPageSize::Small;
 
 protected:
-	bool bShouldEnforceAllocMarks;	
+	bool bShouldEnforceAllocMarks = false;
 };
 
 
