[CSVProfiler] Added events trigger -csveventexeccmds, use "<Eventname>::<Command>,..." to bind multiple commands to CSV Events.

#rnx

diff --git a/Engine/Source/Runtime/Core/Private/ProfilingDebugging/CsvProfiler.cpp b/Engine/Source/Runtime/Core/Private/ProfilingDebugging/CsvProfiler.cpp
--- a/Engine/Source/Runtime/Core/Private/ProfilingDebugging/CsvProfiler.cpp
+++ b/Engine/Source/Runtime/Core/Private/ProfilingDebugging/CsvProfiler.cpp
@@ -438,6 +438,14 @@ void CsvEndNamedEvent()
 
 static TMap<uint32, TArray<FString>>* GCsvFrameExecCmds = NULL;
 
+struct FEventExecCmds
+{
+	alignas(PLATFORM_CACHE_LINE_SIZE) volatile int32 bIsActive;
+	FString EventWildcard;
+	TArray<FString> Cmds;
+};
+static TArray<FEventExecCmds>* GCsvEventExecCmds = NULL;
+
 bool IsContinuousWriteEnabled(bool bGameThread)
 {
 	int CVarValue = -1;
@@ -886,7 +894,7 @@ int32 FCsvProfiler::RegisterCategory(const FString& CategoryName, bool bEnableBy
 	return FCsvCategoryData::Get()->RegisterCategory(CategoryName, bEnableByDefault, bIsGlobal);
 }
 
-void FCsvProfiler::GetFrameExecCommands(TArray<FString>& OutFrameCommands) const
+void FCsvProfiler::GetFrameExecCommands(TArray<FString>& OutFrameCommands)
 {
 	check(IsInGameThread());
 	OutFrameCommands.Empty();
@@ -898,10 +906,19 @@ void FCsvProfiler::GetFrameExecCommands(TArray<FString>& OutFrameCommands) const
 			OutFrameCommands = *FrameCommands;
 		}
 	}
+	if (GCsvProfilerIsCapturing && GCsvEventExecCmds)
+	{
+		for (FEventExecCmds& EventPair : *GCsvEventExecCmds)
+		{
+			if (FPlatformAtomics::InterlockedExchange(&EventPair.bIsActive, (int32)0) > 0)
+			{
+				OutFrameCommands.Append(EventPair.Cmds);
+			}
+		}
+	}
 }
 
 
-
 bool IsInCsvProcessingThread()
 {
 	uint32 ProcessingThreadId = GGameThreadIsCsvProcessingThread ? GGameThreadId : GCsvProcessingThreadId;
@@ -3618,6 +3635,16 @@ void FCsvProfiler::BeginCaptureInternal(const FCsvCaptureCommand& CurrentCommand
 		}
 	}
 
+	// Reset the state of csv event triggers, because there could be some stale state from previous CSVProfile session
+	// and event were triggered at the last frame but didn't have a chance to be consumed
+	if (GCsvEventExecCmds)
+	{
+		for (FEventExecCmds& EventPair : *GCsvEventExecCmds)
+		{
+			FPlatformAtomics::AtomicStore(&EventPair.bIsActive, (int32)0);
+		}
+	}
+
 	SetMetadataInternal(TEXT("TargetFramerate"), *FString::FromInt(TargetFPS));
 	SetMetadataInternal(TEXT("StartTimestamp"), *FString::Printf(TEXT("%lld"), FDateTime::UtcNow().ToUnixTimestamp()));
 	SetMetadataInternal(TEXT("NamedEvents"), (GCycleStatsShouldEmitNamedEvents > 0) ? TEXT("1") : TEXT("0"));
@@ -4247,6 +4274,17 @@ void FCsvProfiler::RecordEvent(int32 CategoryIndex, const FString& EventText)
 		LLM_SCOPE(ELLMTag::CsvProfiler);
 		UE_LOG(LogCsvProfiler, Display, TEXT("CSVEvent \"%s\" [Frame %d]"), *EventText, FCsvProfiler::Get()->GetCaptureFrameNumber());
 		FCsvProfilerThreadData::Get().AddEvent(EventText, CategoryIndex);
+		if (GCsvEventExecCmds)
+		{
+			FStringView Match(EventText);
+			for (FEventExecCmds& EventWildcard : *GCsvEventExecCmds)
+			{
+				if (FWildcardString::IsMatchSubstring(*EventWildcard.EventWildcard, Match.begin(), Match.end(), ESearchCase::IgnoreCase))
+				{
+					FPlatformAtomics::AtomicStore(&EventWildcard.bIsActive, 1);
+				}
+			}
+		}
 	}
 }
 
@@ -4577,6 +4615,45 @@ void FCsvProfiler::Init()
 		}
 	}
 
+	if (FParse::Value(FCommandLine::Get(), TEXT("-csvEventExecCmds="), CsvExecCommandsStr, false))
+	{
+		GCsvEventExecCmds = new TArray<FEventExecCmds>();
+		
+		TArray<FString> CsvExecCommandsList;
+		if (CsvExecCommandsStr.ParseIntoArray(CsvExecCommandsList, TEXT(","), true) > 0)
+		{
+			for (FString FrameAndCommand : CsvExecCommandsList)
+			{
+				int32 ColonIndex = -1;
+				if (FrameAndCommand.Find(TEXT("::"), ColonIndex))
+				{
+					FString EventStr = FrameAndCommand.Mid(0,ColonIndex);
+					FString CommandStr = FrameAndCommand.Mid(ColonIndex+2);
+					if (EventStr.IsEmpty())
+					{
+						EventStr = TEXT("*");
+					}
+					bool bAlreadyExists = false;
+					for (FEventExecCmds& EventDesc : *GCsvEventExecCmds)
+					{
+						if (EventDesc.EventWildcard == EventStr)
+						{
+							bAlreadyExists = true;
+							EventDesc.Cmds.Add(CommandStr);
+						}
+					}
+					if (!bAlreadyExists)
+					{
+						FEventExecCmds& EventDesc = GCsvEventExecCmds->Emplace_GetRef();
+						EventDesc.EventWildcard = EventStr;
+						EventDesc.Cmds.Add(CommandStr);
+					}
+					UE_LOG(LogCsvProfiler, Display, TEXT("Added CsvEventExecCommand - event \"%s\" : %s"), *EventStr, *CommandStr);
+				}
+			}
+		}
+	}
+
 #endif // CSV_PROFILER_ALLOW_DEBUG_FEATURES
 
 	if (GAllCategoriesStartDisabled)
diff --git a/Engine/Source/Runtime/Core/Public/ProfilingDebugging/CsvProfiler.h b/Engine/Source/Runtime/Core/Public/ProfilingDebugging/CsvProfiler.h
--- a/Engine/Source/Runtime/Core/Public/ProfilingDebugging/CsvProfiler.h
+++ b/Engine/Source/Runtime/Core/Public/ProfilingDebugging/CsvProfiler.h
@@ -462,7 +462,7 @@ public:
 
 	CORE_API static bool IsWaitTrackingEnabledOnCurrentThread();
 
-	CORE_API void GetFrameExecCommands(TArray<FString>& OutFrameCommands) const;
+	CORE_API void GetFrameExecCommands(TArray<FString>& OutFrameCommands);
 
 	/** Called right before we start capturing. */
 	DECLARE_MULTICAST_DELEGATE(FOnCSVProfileStart);
