Android bug fix. Current implementation of GetStatData for AAsset* resulted in directories not being stat'able. AAssetManager_open would fail for a directory and the branch where AAssetManager_openDir was called would get skipped. Handle directories under condition that AAssetManager_open fails to open a file instead.

diff --git a/Engine/Source/Runtime/Core/Private/Android/AndroidPlatformFile.cpp b/Engine/Source/Runtime/Core/Private/Android/AndroidPlatformFile.cpp
--- a/Engine/Source/Runtime/Core/Private/Android/AndroidPlatformFile.cpp
+++ b/Engine/Source/Runtime/Core/Private/Android/AndroidPlatformFile.cpp
@@ -1831,23 +1831,29 @@ public:
 		else
 		{
 			AAsset * file = AAssetManager_open(AssetMgr, TCHAR_TO_UTF8(*AssetPath), AASSET_MODE_UNKNOWN);
-			if (nullptr != file)
+			bool isDirectory = false;
+			bool exists = (nullptr != file);
+			int64 FileSize = -1;
+			if (exists)
 			{
-				bool isDirectory = false;
-				AAssetDir * subdir = AAssetManager_openDir(AssetMgr, TCHAR_TO_UTF8(*AssetPath));
-				if (nullptr != subdir)
+				FileSize = AAsset_getLength(file);
+				AAsset_close(file);
+			}
+			else
+			{
+				AAssetDir * dir = AAssetManager_openDir(AssetMgr, TCHAR_TO_UTF8(*AssetPath));
+				exists = (nullptr != dir);
+				
+				if (exists)
 				{
 					isDirectory = true;
-					AAssetDir_close(subdir);
-				}
-
-				int64 FileSize = -1;
-				if (!isDirectory)
-				{
-					FileSize = AAsset_getLength(file);
+					AAssetDir_close(dir);
 				}
-
-				FFileStatData StatData(
+			}
+			
+			if (exists)
+			{
+				return FFileStatData(
 					FDateTime::MinValue(),				// CreationTime
 					FDateTime::MinValue(),				// AccessTime
 					FDateTime::MinValue(),				// ModificationTime
@@ -1855,10 +1861,6 @@ public:
 					isDirectory,						// bIsDirectory
 					true								// bIsReadOnly
 					);
-
-				AAsset_close(file);
-
-				return StatData;
 			}
 		}
 		
