Add retry when launching zen command line tool to determine version
Fix some concurrency issues with fields in FMonitoredProcess (no-atomic fields touched from multiple threads)
Refactor to make sure we pick up all output from FMonitoredProcess and set the state correct even if we cancel/terminate

#rnx

diff --git a/Engine/Source/Developer/Zen/Private/ZenServerInterface.cpp b/Engine/Source/Developer/Zen/Private/ZenServerInterface.cpp
--- a/Engine/Source/Developer/Zen/Private/ZenServerInterface.cpp
+++ b/Engine/Source/Developer/Zen/Private/ZenServerInterface.cpp
@@ -316,8 +316,12 @@ GetZenVersion(const FString& UtilityPath, const FString& ServicePath, const FStr
 	FZenVersion Version;
 	if (!GetZenVersion(UtilityPath, ServicePath, Version))
 	{
-		checkf(false, TEXT("Unable to determine version using zen utility executable path: '%s'."), *UtilityPath);
-		Version = GetFallbackVersion();
+		UE_LOG(LogZenServiceInstance, Warning, TEXT("Unable to determine version using zen utility executable path: '%s'. Retrying..."), *UtilityPath);
+		if (!GetZenVersion(UtilityPath, ServicePath, Version))
+		{
+			checkf(false, TEXT("Unable to determine version using zen utility executable path: '%s'."), *UtilityPath);
+			Version = GetFallbackVersion();
+		}
 	}
 	FFileHelper::SaveStringToFile(Version.ToString(), *VersionCachePath);
 	return Version;
diff --git a/Engine/Source/Runtime/Core/Private/Misc/MonitoredProcess.cpp b/Engine/Source/Runtime/Core/Private/Misc/MonitoredProcess.cpp
--- a/Engine/Source/Runtime/Core/Private/Misc/MonitoredProcess.cpp
+++ b/Engine/Source/Runtime/Core/Private/Misc/MonitoredProcess.cpp
@@ -49,12 +49,17 @@ FMonitoredProcess::~FMonitoredProcess()
 
 FTimespan FMonitoredProcess::GetDuration() const
 {
+	if (StartTime.Load() == 0)
+	{
+		// Not initialized yet - monitoring thread might not have set StartTime yet.
+		return 0;
+	}
 	if (bIsRunning)
 	{
-		return (FDateTime::UtcNow() - StartTime);
+		return (FDateTime::UtcNow() - StartTime.Load());
 	}
 
-	return (EndTime - StartTime);
+	return (EndTime.Load() - StartTime.Load());
 }
 
 
@@ -86,12 +91,23 @@ bool FMonitoredProcess::Launch()
 	Thread = FRunnableThread::Create(this, *MonitoredProcessName, 128 * 1024, TPri_AboveNormal);
 	if ( !FPlatformProcess::SupportsMultithreading() )
 	{
-		StartTime = FDateTime::UtcNow();
+		StartTime.Store(FDateTime::UtcNow());
 	}
 
 	return true;
 }
 
+int FMonitoredProcess::GetReturnCode() const
+{
+	check(!bIsRunning);
+	return ReturnCode;
+}
+
+const FString& FMonitoredProcess::GetFullOutputWithoutDelegate() const
+{
+	check((!FPlatformProcess::SupportsMultithreading()) || (!bIsRunning));	// Don't touch the output buffer while running with a background task
+	return OutputBuffer;
+}
 
 /* FMonitoredProcess implementation
  *****************************************************************************/
@@ -129,32 +145,56 @@ void FMonitoredProcess::ProcessOutput( const FString& Output )
 void FMonitoredProcess::TickInternal()
 {
 	// monitor the process
-	ProcessOutput(FPlatformProcess::ReadPipe(ReadPipe));
 
 	// kill child process if we choose to cancel or it the engine is shutting down
 	if (Canceling)
 	{
 		FPlatformProcess::TerminateProc(ProcessHandle, KillTree);
+
+		EndTime.Store(FDateTime::UtcNow());
+
+		// Make sure to read the final bytes from the pipe before closing
+		ProcessOutput(FPlatformProcess::ReadPipe(ReadPipe));
+
+		// close output pipes
+		FPlatformProcess::ClosePipe(ReadPipe, WritePipe);
+		ReadPipe = WritePipe = nullptr;
+
+		// Indicate failure in the return code
+		ReturnCode.Store(-1);
+
 		CanceledDelegate.ExecuteIfBound();
 		bIsRunning = false;
 	}
 	else if (!FPlatformProcess::IsProcRunning(ProcessHandle))
 	{
-		EndTime = FDateTime::UtcNow();
+		EndTime.Store(FDateTime::UtcNow());
+
+		// Make sure to read the final bytes from the pipe before closing
+		ProcessOutput(FPlatformProcess::ReadPipe(ReadPipe));
 
 		// close output pipes
 		FPlatformProcess::ClosePipe(ReadPipe, WritePipe);
 		ReadPipe = WritePipe = nullptr;
 
+		int ProcessReturnCode = 0;
 		// get completion status
-		if (!FPlatformProcess::GetProcReturnCode(ProcessHandle, &ReturnCode))
+		if (FPlatformProcess::GetProcReturnCode(ProcessHandle, &ProcessReturnCode))
 		{
-			ReturnCode = -1;
+			ReturnCode.Store(ProcessReturnCode);
+		}
+		else
+		{
+			ReturnCode.Store(-1);
 		}
 
 		CompletedDelegate.ExecuteIfBound(ReturnCode);
 		bIsRunning = false;
 	}
+	else
+	{
+		ProcessOutput(FPlatformProcess::ReadPipe(ReadPipe));
+	}
 }
 
 
@@ -174,7 +214,7 @@ bool FMonitoredProcess::Update()
 
 uint32 FMonitoredProcess::Run()
 {
-	StartTime = FDateTime::UtcNow();
+	StartTime.Store(FDateTime::UtcNow());
 	while (bIsRunning)
 	{
 		FPlatformProcess::Sleep(SleepInterval);
diff --git a/Engine/Source/Runtime/Core/Public/Misc/MonitoredProcess.h b/Engine/Source/Runtime/Core/Public/Misc/MonitoredProcess.h
--- a/Engine/Source/Runtime/Core/Public/Misc/MonitoredProcess.h
+++ b/Engine/Source/Runtime/Core/Public/Misc/MonitoredProcess.h
@@ -70,8 +70,8 @@ public:
 	 */
 	void Cancel( bool InKillTree = false )
 	{
-		Canceling = true;
 		KillTree = InKillTree;
+		Canceling = true;
 	}
 
 	/**
@@ -156,19 +156,13 @@ public:
 	 *
 	 * @return Process return code
 	 */
-	int GetReturnCode() const
-	{
-		return ReturnCode;
-	}
+	CORE_API int GetReturnCode() const;
 
 	/**
 	 * Returns the full output, wihtout needing to hookup a delegate and buffer it externally. Note that if OutputDelegate is 
 	 * bound, this _will not_ have the entire output
 	 */
-	const FString& GetFullOutputWithoutDelegate() const
-	{
-		return OutputBuffer;
-	}
+	CORE_API const FString& GetFullOutputWithoutDelegate() const;
 
 public:
 
@@ -212,16 +206,16 @@ protected:
 
 
 	// Whether the process is being canceled. */
-	bool Canceling = false;
+	TAtomic<bool> Canceling = false;
 
 	// Holds the time at which the process ended. */
-	FDateTime EndTime;
+	TAtomic<FDateTime> EndTime;
 
 	// Whether the window of the process should be hidden. */
 	bool Hidden = false;
 
 	// Whether to kill the entire process tree when cancelling this process. */
-	bool KillTree = false;
+	TAtomic<bool> KillTree = false;
 
 	// Holds the command line parameters. */
 	FString Params;
@@ -233,16 +227,16 @@ protected:
 	void* ReadPipe = nullptr;
 
 	// Holds the return code. */
-	int ReturnCode = 0;
+	TAtomic<int> ReturnCode = 0;
 
 	// Holds the time at which the process started. */
-	FDateTime StartTime { 0 };
+	TAtomic<FDateTime> StartTime { 0 };
 
 	// Holds the monitoring thread object. */
 	FRunnableThread* Thread = nullptr;
 
 	// Is the thread running? 
-	TSAN_ATOMIC(bool) bIsRunning;
+	TAtomic<bool> bIsRunning;
 
 	// Holds the URL of the executable to launch. */
 	FString URL;
