[AutoRTFM] Clean up Transaction.h

Remove unused header.
Convert space indentation to tabs.

diff --git a/Engine/Source/Runtime/Core/Private/AutoRTFM/Transaction.h b/Engine/Source/Runtime/Core/Private/AutoRTFM/Transaction.h
--- a/Engine/Source/Runtime/Core/Private/AutoRTFM/Transaction.h
+++ b/Engine/Source/Runtime/Core/Private/AutoRTFM/Transaction.h
@@ -3,8 +3,8 @@
 #pragma once
 
 #include "AutoRTFM/AutoRTFM.h"
+#include "AutoRTFM/AutoRTFMTask.h"
 #include "HAL/Platform.h"
-#include "Templates/Function.h"
 #include "CoreTypes.h"
 #include "HitSet.h"
 #include "IntervalTree.h"
@@ -20,176 +20,176 @@ class FContext;
 class FTransaction final
 {
 public:
-    // State flow diagram for a transaction:
-    //                       ┌──────────────────────────────────────┐
-    //                       │             Uninitialized            │
-    //                       └──────────────────────────────────────┘
-    //                               │                       │
-    //                               ▼                       ▼
-    // ┌────────────────┐    ┌────────────────┐    ┌────────────────┐    ┌────────────────┐
-    // │  OpenInactive  │ ←→ │   OpenActive   │ ←→ │  ClosedActive  │ ←→ │ ClosedInactive │
-    // └────────────────┘    └────────────────┘    └────────────────┘    └────────────────┘
-    //                               │                       │
-    //                               ▼                       ▼
-    //                       ┌──────────────────────────────────────┐
-    //                       │                 Done                 │
-    //                       └──────────────────────────────────────┘
-    enum class EState
-    {
-        // The initial state for the transaction.
-        // Can only transition to OpenActive or ClosedActive.
-        Uninitialized,
-        // The transaction is open (not recording writes) and is the current transaction.
-        // Can only transition to OpenInactive, ClosedActive or Done.
-        OpenActive,
-        // The transaction is closed (recording writes) and is the current transaction.
-        // Can only transition to ClosedInactive, OpenActive or Done.
-        ClosedActive,
-        // The transaction is open and the current transaction is a descendant.
-        // Can only transition to OpenActive.
-        OpenInactive,
-        // The transaction is closed and the current transaction is a descendant.
-        // Can only transition to ClosedActive.
-        ClosedInactive,
-        // The transaction is committed or aborted.
-        // Once in this state the transaction must be reset with Reset() or destructed.
-        Done,
-    };
-
-    // Constructor.
-    // The transaction is constructed in an initial Closed state.
-    FTransaction(FContext* Context);
-
-    // Destructor.
-    // The transaction must we in a Done state.
-    ~FTransaction();
-
-    void Initialize(FTransaction* Parent, bool bIsScoped, FStackRange StackRange);
-
-    // Clears the tracked transaction state and resets back to the default Uninitialized state.
-    void Reset();
-
-    bool IsNested() const { return !!Parent; }
-    FTransaction* GetParent() const { return Parent; }
-
-    // This should just use type displays or ranges. Maybe ranges could even work out great.
-    bool IsNestedWithin(const FTransaction* Other) const
-    {
-        for (const FTransaction* Current = this; ; Current = Current->Parent)
-        {
-            if (!Current)
-            {
-                return false;
-            }
-            else if (Current == Other)
-            {
-                return true;
-            }
-        }
-    }
+	// State flow diagram for a transaction:
+	//                       ┌──────────────────────────────────────┐
+	//                       │             Uninitialized            │
+	//                       └──────────────────────────────────────┘
+	//                               │                       │
+	//                               ▼                       ▼
+	// ┌────────────────┐    ┌────────────────┐    ┌────────────────┐    ┌────────────────┐
+	// │  OpenInactive  │ ←→ │   OpenActive   │ ←→ │  ClosedActive  │ ←→ │ ClosedInactive │
+	// └────────────────┘    └────────────────┘    └────────────────┘    └────────────────┘
+	//                               │                       │
+	//                               ▼                       ▼
+	//                       ┌──────────────────────────────────────┐
+	//                       │                 Done                 │
+	//                       └──────────────────────────────────────┘
+	enum class EState
+	{
+		// The initial state for the transaction.
+		// Can only transition to OpenActive or ClosedActive.
+		Uninitialized,
+		// The transaction is open (not recording writes) and is the current transaction.
+		// Can only transition to OpenInactive, ClosedActive or Done.
+		OpenActive,
+		// The transaction is closed (recording writes) and is the current transaction.
+		// Can only transition to ClosedInactive, OpenActive or Done.
+		ClosedActive,
+		// The transaction is open and the current transaction is a descendant.
+		// Can only transition to OpenActive.
+		OpenInactive,
+		// The transaction is closed and the current transaction is a descendant.
+		// Can only transition to ClosedActive.
+		ClosedInactive,
+		// The transaction is committed or aborted.
+		// Once in this state the transaction must be reset with Reset() or destructed.
+		Done,
+	};
+
+	// Constructor.
+	// The transaction is constructed in an initial Closed state.
+	FTransaction(FContext* Context);
+
+	// Destructor.
+	// The transaction must we in a Done state.
+	~FTransaction();
+
+	void Initialize(FTransaction* Parent, bool bIsScoped, FStackRange StackRange);
+
+	// Clears the tracked transaction state and resets back to the default Uninitialized state.
+	void Reset();
+
+	bool IsNested() const { return !!Parent; }
+	FTransaction* GetParent() const { return Parent; }
+
+	// This should just use type displays or ranges. Maybe ranges could even work out great.
+	bool IsNestedWithin(const FTransaction* Other) const
+	{
+		for (const FTransaction* Current = this; ; Current = Current->Parent)
+		{
+			if (!Current)
+			{
+				return false;
+			}
+			else if (Current == Other)
+			{
+				return true;
+			}
+		}
+	}
 
 	inline bool IsScopedTransaction() const { return bIsStackScoped; }
 
-    void DeferUntilCommit(TTask<void()>&&);
-    void DeferUntilAbort(TTask<void()>&&);
+	void DeferUntilCommit(TTask<void()>&&);
+	void DeferUntilAbort(TTask<void()>&&);
 	void PushDeferUntilCommitHandler(const void* Key, TTask<void()>&&);
 	void PopDeferUntilCommitHandler(const void* Key);
 	void PopAllDeferUntilCommitHandlers(const void* Key);
-    void PushDeferUntilAbortHandler(const void* Key, TTask<void()>&&);
+	void PushDeferUntilAbortHandler(const void* Key, TTask<void()>&&);
 	void PopDeferUntilAbortHandler(const void* Key);
 	void PopAllDeferUntilAbortHandlers(const void* Key);
 
-    [[noreturn]] void AbortAndThrow();
-    void AbortWithoutThrowing();
-    bool AttemptToCommit();
+	[[noreturn]] void AbortAndThrow();
+	void AbortWithoutThrowing();
+	bool AttemptToCommit();
 
-    // Record that a write is about to occur at the given LogicalAddress of Size bytes.
-    void RecordWrite(void* LogicalAddress, size_t Size, bool bNoMemoryValidation = false);
-    template<unsigned SIZE> void RecordWrite(void* LogicalAddress);
+	// Record that a write is about to occur at the given LogicalAddress of Size bytes.
+	void RecordWrite(void* LogicalAddress, size_t Size, bool bNoMemoryValidation = false);
+	template<unsigned SIZE> void RecordWrite(void* LogicalAddress);
 	template<unsigned SIZE> void RecordWriteInsertedSlow(void* LogicalAddress);
 	template<unsigned SIZE> void RecordWriteNotInsertedSlow(void* LogicalAddress);
 
-    void DidAllocate(void* LogicalAddress, size_t Size);
-    void DidFree(void* LogicalAddress);
-
-    void SetOpenActive(EMemoryValidationLevel MemoryValidationLevel, const void* ReturnAddress);
-    void SetClosedActive();
-    void SetOpenInactive();
-    void SetClosedInactive();
-    void SetActive();
-    void SetInactive();
-    void SetDone();
-
-    // A debug helper that will break to the debugger if the memory validation
-    // hash no longer matches. Useful for isolating where the open write happened.
-    void DebugBreakIfMemoryValidationFails();
-
-    // State querying
-    EState State() const { return CurrentState; }
-    bool IsFresh() const;
-    bool IsOpenActive() const { return CurrentState == EState::OpenActive; }
-    bool IsClosedActive() const { return CurrentState == EState::ClosedActive; }
-    bool IsOpenInactive() const { return CurrentState == EState::OpenInactive; }
-    bool IsClosedInactive() const { return CurrentState == EState::ClosedInactive; }
-    bool IsOpen() const { return IsOpenActive() || IsOpenInactive(); }
-    bool IsClosed() const { return IsClosedActive() || IsClosedInactive(); }
-    bool IsActive() const { return IsOpenActive() || IsClosedActive(); }
-    bool IsInactive() const { return IsOpenInactive() || IsClosedInactive(); }
-    bool IsDone() const { return CurrentState == EState::Done; }
-
-    EMemoryValidationLevel MemoryValidationLevel() const { return CurrentMemoryValidationLevel; }
-    const void* OpenReturnAddress() const { return CurrentOpenReturnAddress; }
-
-    // The stack range represents all stack memory inside the transaction scope
-    inline FStackRange GetStackRange() const { return StackRange; }
-
-    // Returns true if the LogicalAddress is within the stack of the transaction.
-    inline bool IsOnStack(const void* LogicalAddress) const;
+	void DidAllocate(void* LogicalAddress, size_t Size);
+	void DidFree(void* LogicalAddress);
+
+	void SetOpenActive(EMemoryValidationLevel MemoryValidationLevel, const void* ReturnAddress);
+	void SetClosedActive();
+	void SetOpenInactive();
+	void SetClosedInactive();
+	void SetActive();
+	void SetInactive();
+	void SetDone();
+
+	// A debug helper that will break to the debugger if the memory validation
+	// hash no longer matches. Useful for isolating where the open write happened.
+	void DebugBreakIfMemoryValidationFails();
+
+	// State querying
+	EState State() const { return CurrentState; }
+	bool IsFresh() const;
+	bool IsOpenActive() const { return CurrentState == EState::OpenActive; }
+	bool IsClosedActive() const { return CurrentState == EState::ClosedActive; }
+	bool IsOpenInactive() const { return CurrentState == EState::OpenInactive; }
+	bool IsClosedInactive() const { return CurrentState == EState::ClosedInactive; }
+	bool IsOpen() const { return IsOpenActive() || IsOpenInactive(); }
+	bool IsClosed() const { return IsClosedActive() || IsClosedInactive(); }
+	bool IsActive() const { return IsOpenActive() || IsClosedActive(); }
+	bool IsInactive() const { return IsOpenInactive() || IsClosedInactive(); }
+	bool IsDone() const { return CurrentState == EState::Done; }
+
+	EMemoryValidationLevel MemoryValidationLevel() const { return CurrentMemoryValidationLevel; }
+	const void* OpenReturnAddress() const { return CurrentOpenReturnAddress; }
+
+	// The stack range represents all stack memory inside the transaction scope
+	inline FStackRange GetStackRange() const { return StackRange; }
+
+	// Returns true if the LogicalAddress is within the stack of the transaction.
+	inline bool IsOnStack(const void* LogicalAddress) const;
 
 private:
-    using FWriteHash = uint64;
-    using FTaskArray = TTaskArray<TTask<void()>>;
+	using FWriteHash = uint64;
+	using FTaskArray = TTaskArray<TTask<void()>>;
 
-    void Undo();
+	void Undo();
 
-    void CommitNested();
-    bool AttemptToCommitOuterNest();
+	void CommitNested();
+	bool AttemptToCommitOuterNest();
 
-    // Calculates the hash of all the memory written during the transaction and stores it to
-    // RecordedWriteHash.
-    void RecordWriteHash();
+	// Calculates the hash of all the memory written during the transaction and stores it to
+	// RecordedWriteHash.
+	void RecordWriteHash();
 
-    // Compares the hash of all the memory that will be written on abort against the
-    // RecordedWriteHash, raising a warning or error if the hashes are not equal.
-    // See CalculateNestedWriteHash().
-    void ValidateWriteHash() const;
+	// Compares the hash of all the memory that will be written on abort against the
+	// RecordedWriteHash, raising a warning or error if the hashes are not equal.
+	// See CalculateNestedWriteHash().
+	void ValidateWriteHash() const;
 
-    // Returns a hash of all the memory written during this transaction and all parent transactions.
-    // Can be used to ensure that no transaction written memory is modified during an open.
-    FWriteHash CalculateNestedWriteHash() const;
-    
-    // Returns a hash of all the memory written during this transaction and all parent transactions.
-    // Can be used to ensure that no transaction written memory is modified during an open.
-    // NumWriteEntries is the number of write entries to hash before stopping.
-    FWriteHash CalculateNestedWriteHashWithLimit(size_t NumWriteEntries) const;
+	// Returns a hash of all the memory written during this transaction and all parent transactions.
+	// Can be used to ensure that no transaction written memory is modified during an open.
+	FWriteHash CalculateNestedWriteHash() const;
+	
+	// Returns a hash of all the memory written during this transaction and all parent transactions.
+	// Can be used to ensure that no transaction written memory is modified during an open.
+	// NumWriteEntries is the number of write entries to hash before stopping.
+	FWriteHash CalculateNestedWriteHashWithLimit(size_t NumWriteEntries) const;
 
-    void CollectStats() const;
+	void CollectStats() const;
 
-    bool ShouldRecordWrite(void* LogicalAddress) const;
+	bool ShouldRecordWrite(void* LogicalAddress) const;
 
-    template<EState NewState>
-    void SetState();
+	template<EState NewState>
+	void SetState();
 
-    FContext* Context{nullptr};
+	FContext* Context{nullptr};
 
-    // Are we nested? Then this is the parent.
-    FTransaction* Parent{nullptr};
+	// Are we nested? Then this is the parent.
+	FTransaction* Parent{nullptr};
 
-    // Commit tasks run on commit in forward order.
-    FTaskArray CommitTasks;
+	// Commit tasks run on commit in forward order.
+	FTaskArray CommitTasks;
 
 	// Abort tasks run on abort in reverse order.
-    FTaskArray AbortTasks;
+	FTaskArray AbortTasks;
 
 	// If a call to `PopOnCommitHandler` could not find a commit to pop, it is deferred
 	// and tried again on the parent transaction.
@@ -207,18 +207,18 @@ private:
 	// commits, we need to propagate this to the parent too.
 	TArray<const void*, TInlineAllocator<1>> DeferredPopAllOnAbortHandlers;
 
-    FHitSet HitSet;
-    FIntervalTree NewMemoryTracker;
-    FWriteLog WriteLog;
-    TStatStorage<uint64_t> StatDepth = 1;
-    FWriteHash RecordedWriteHash = 0;
-    size_t NumWriteLogsHashed = 0;
-    EMemoryValidationLevel CurrentMemoryValidationLevel = EMemoryValidationLevel::Disabled;
-    const void* CurrentOpenReturnAddress = nullptr;
-    EState CurrentState{EState::Uninitialized};
-    FStackRange StackRange;
-    bool bIsStackScoped = false;
-    bool bIsInAllocateFn = false;
+	FHitSet HitSet;
+	FIntervalTree NewMemoryTracker;
+	FWriteLog WriteLog;
+	TStatStorage<uint64_t> StatDepth = 1;
+	FWriteHash RecordedWriteHash = 0;
+	size_t NumWriteLogsHashed = 0;
+	EMemoryValidationLevel CurrentMemoryValidationLevel = EMemoryValidationLevel::Disabled;
+	const void* CurrentOpenReturnAddress = nullptr;
+	EState CurrentState{EState::Uninitialized};
+	FStackRange StackRange;
+	bool bIsStackScoped = false;
+	bool bIsInAllocateFn = false;
 };
 
 } // namespace AutoRTFM
