Get GetCurrentProcessId to work in a transaction.

diff --git a/Engine/Source/Programs/AutoRTFMTests/Private/UECoreTests.cpp b/Engine/Source/Programs/AutoRTFMTests/Private/UECoreTests.cpp
--- a/Engine/Source/Programs/AutoRTFMTests/Private/UECoreTests.cpp
+++ b/Engine/Source/Programs/AutoRTFMTests/Private/UECoreTests.cpp
@@ -3556,3 +3556,12 @@ TEST_CASE("UECore.LoadFileToString.IPlatformFile")
 		REQUIRE(FileData == TEXT("File Loaded"));
 	}
 }
+
+TEST_CASE("UECore.GetCurrentProcessId")
+{
+	const uint32 Outer = FPlatformProcess::GetCurrentProcessId();
+	AutoRTFM::Testing::Commit([&]
+	{
+		REQUIRE(Outer == FPlatformProcess::GetCurrentProcessId());
+	});
+}
diff --git a/Engine/Source/Runtime/Core/Private/AutoRTFM/StdLibSupport.cpp b/Engine/Source/Runtime/Core/Private/AutoRTFM/StdLibSupport.cpp
--- a/Engine/Source/Runtime/Core/Private/AutoRTFM/StdLibSupport.cpp
+++ b/Engine/Source/Runtime/Core/Private/AutoRTFM/StdLibSupport.cpp
@@ -819,6 +819,7 @@ UE_AUTORTFM_REGISTER_SELF_FUNCTION(GetSystemTime);
 UE_AUTORTFM_REGISTER_SELF_FUNCTION(QueryPerformanceCounter);
 UE_AUTORTFM_REGISTER_SELF_FUNCTION(QueryPerformanceFrequency);
 UE_AUTORTFM_REGISTER_SELF_FUNCTION(GetCurrentThreadId);
+UE_AUTORTFM_REGISTER_SELF_FUNCTION(GetCurrentProcessId);
 
 UE_AUTORTFM_REGISTER_SELF_FUNCTION(TlsGetValue);
 UE_AUTORTFM_REGISTER_SELF_FUNCTION(GetLocalTime);
@@ -840,6 +841,7 @@ UE_AUTORTFM_REGISTER_OPEN_FUNCTION(TlsSetValue);
 #endif // PLATFORM_WINDOWS
 
 #if PLATFORM_LINUX
+UE_AUTORTFM_REGISTER_SELF_FUNCTION(getpid);
 UE_AUTORTFM_REGISTER_SELF_FUNCTION(clock_gettime);
 UE_AUTORTFM_REGISTER_SELF_FUNCTION(gettimeofday);
 UE_AUTORTFM_REGISTER_SELF_FUNCTION(gmtime_r);
