[AutoRTFM] Make thread safe refcount class transactionally safe.

diff --git a/Engine/Source/Programs/AutoRTFMTests/Private/RefCountPtrTests.cpp b/Engine/Source/Programs/AutoRTFMTests/Private/RefCountPtrTests.cpp
--- a/Engine/Source/Programs/AutoRTFMTests/Private/RefCountPtrTests.cpp
+++ b/Engine/Source/Programs/AutoRTFMTests/Private/RefCountPtrTests.cpp
@@ -17,9 +17,15 @@ namespace
 		RefCountedMixinInt(int InValue) : Value(InValue) {}
 		int Value;
 	};
+
+	struct ThreadSafeRefCountedInt : public FThreadSafeRefCountedObject
+	{
+		ThreadSafeRefCountedInt(int InValue) : Value(InValue) {}
+		int Value;
+	};
 }
 
-TEMPLATE_TEST_CASE("PreviouslyAllocated", "[RefCountPtr]", RefCountedInt, RefCountedMixinInt)
+TEMPLATE_TEST_CASE("PreviouslyAllocated", "[RefCountPtr]", RefCountedInt, RefCountedMixinInt, ThreadSafeRefCountedInt)
 {
 	SECTION("Using new T()")
 	{
@@ -52,7 +58,7 @@ TEMPLATE_TEST_CASE("PreviouslyAllocated", "[RefCountPtr]", RefCountedInt, RefCou
 	}
 }
 
-TEMPLATE_TEST_CASE("AbortWithPreviouslyAllocated", "[RefCountPtr]", RefCountedInt, RefCountedMixinInt)
+TEMPLATE_TEST_CASE("AbortWithPreviouslyAllocated", "[RefCountPtr]", RefCountedInt, RefCountedMixinInt, ThreadSafeRefCountedInt)
 {
 	SECTION("Using new T()")
 	{
@@ -89,7 +95,7 @@ TEMPLATE_TEST_CASE("AbortWithPreviouslyAllocated", "[RefCountPtr]", RefCountedIn
 	}
 }
 
-TEMPLATE_TEST_CASE("AbortWithNewlyAllocated", "[RefCountPtr]", RefCountedInt, RefCountedMixinInt)
+TEMPLATE_TEST_CASE("AbortWithNewlyAllocated", "[RefCountPtr]", RefCountedInt, RefCountedMixinInt, ThreadSafeRefCountedInt)
 {
 	int Result = 42;
 
@@ -138,7 +144,7 @@ TEMPLATE_TEST_CASE("AbortWithNewlyAllocated", "[RefCountPtr]", RefCountedInt, Re
 	REQUIRE(Result == 42);
 }
 
-TEMPLATE_TEST_CASE("OnCommitCapturingRefCountPtr", "[RefCountPtr]", RefCountedInt, RefCountedMixinInt)
+TEMPLATE_TEST_CASE("OnCommitCapturingRefCountPtr", "[RefCountPtr]", RefCountedInt, RefCountedMixinInt, ThreadSafeRefCountedInt)
 {
 	TRefCountPtr<TestType> Foo = MakeRefCount<TestType>(42);
 
@@ -170,7 +176,7 @@ TEMPLATE_TEST_CASE("OnCommitCapturingRefCountPtr", "[RefCountPtr]", RefCountedIn
 	}
 }
 
-TEMPLATE_TEST_CASE("OnAbortCapturingRefCountPtr", "[RefCountPtr]", RefCountedInt, RefCountedMixinInt)
+TEMPLATE_TEST_CASE("OnAbortCapturingRefCountPtr", "[RefCountPtr]", RefCountedInt, RefCountedMixinInt, ThreadSafeRefCountedInt)
 {
 	TRefCountPtr<TestType> Foo = MakeRefCount<TestType>(42);
 
diff --git a/Engine/Source/Runtime/Core/Private/Templates/RefCounting.cpp b/Engine/Source/Runtime/Core/Private/Templates/RefCounting.cpp
--- a/Engine/Source/Runtime/Core/Private/Templates/RefCounting.cpp
+++ b/Engine/Source/Runtime/Core/Private/Templates/RefCounting.cpp
@@ -2,7 +2,7 @@
 
 #include "Templates/RefCounting.h"
 
-void FRefCountBase::CheckRefCountIsNonZero() const
+void UE::Private::FTransactionalAtomicRefCounter::CheckRefCountIsNonZero()
 {
 	check(!"Release() was called on an object which is already at a zero refcount.");
 }
\ No newline at end of file
diff --git a/Engine/Source/Runtime/Core/Public/Templates/RefCounting.h b/Engine/Source/Runtime/Core/Public/Templates/RefCounting.h
--- a/Engine/Source/Runtime/Core/Public/Templates/RefCounting.h
+++ b/Engine/Source/Runtime/Core/Public/Templates/RefCounting.h
@@ -56,39 +56,25 @@ private:
 	uint32 RefCount = 0;
 };
 
-/** A virtual interface for ref counted objects to implement. */
-class IRefCountedObject
-{
-public:
-	virtual ~IRefCountedObject() { }
-	virtual FReturnedRefCountValue AddRef() const = 0;
-
-	// TODO (SOL-7350): return FReturnedRefCountValue from Release(); clean up call sites
-	// which rely on its return value.
-	virtual uint32 Release() const = 0;
-
-	// TODO (SOL-7350): mark this function as deprecated; clean up existing callers.
-	virtual uint32 GetRefCount() const = 0;
-};
-
 /**
- * Base class implementing thread-safe reference counting.
+ * UE::Private::FTransactionalAtomicRefCounter provides transactionally-safe atomic refcounting primitives.
+ * These are used by FRefCountBase, FThreadSafeRefCountedObject and TRefCountingMixin. This isn't public API.
  */
-class FRefCountBase
+namespace UE::Private
 {
-public:
-	FRefCountBase() = default;
-	virtual ~FRefCountBase() = default;
-
-	FRefCountBase(const FRefCountBase& Rhs) = delete;
-	FRefCountBase& operator=(const FRefCountBase& Rhs) = delete;
 
-	FReturnedRefCountValue AddRef() const
+class FTransactionalAtomicRefCounter
+{
+public:
+	template <auto PtrToRefCount, auto DeleteFn, typename T>
+	static uint32 AddRef(const T* That)
 	{
+		std::atomic<uint32>& RefCount = const_cast<std::atomic<uint32>&>(That->*PtrToRefCount);
 		uint32 Refs;
+
 		UE_AUTORTFM_OPEN
 		{
-			Refs = NumRefs++;
+			Refs = RefCount++;
 		};
 
 		// If we are inside a transaction, a rollback must undo our refcount bump.
@@ -101,28 +87,31 @@ public:
 		// and special-case it.
 		if (Refs == 0)
 		{
-			UE_AUTORTFM_ONABORT(this)
+			UE_AUTORTFM_ONABORT(&RefCount)
 			{
-				--NumRefs;
+				--RefCount;
 				// The refcount is likely zero now, but leaving the object alive isn't a leak.
 				// We are restoring the object back to its initial "just-created" state.
 			};
 		}
 		else
 		{
-			UE_AUTORTFM_ONABORT(this)
+			UE_AUTORTFM_ONABORT(That)
 			{
-				Release();
+				Release<PtrToRefCount, DeleteFn, T>(That);
 			};
 		}
 
-		return FReturnedRefCountValue{Refs + 1};
+		return Refs + 1;
 	}
 
-	uint32 Release() const
+	template <auto PtrToRefCount, auto DeleteFn, typename T>
+	static uint32 Release(const T* That)
 	{
+		std::atomic<uint32>& RefCount = const_cast<std::atomic<uint32>&>(That->*PtrToRefCount);
+
 #if DO_GUARD_SLOW
-		if (NumRefs == 0)
+		if (RefCount == 0)
 		{
 			CheckRefCountIsNonZero();
 		}
@@ -139,58 +128,121 @@ public:
 			// always returns the true refcount.
 			UE_AUTORTFM_OPEN
 			{
-				RefsToReturn = uint32(NumRefs.load(std::memory_order_relaxed)) - 1;
+				RefsToReturn = uint32(RefCount.load(std::memory_order_relaxed)) - 1;
 			};
 
 			// Refcount changes and frees are deferred until the transaction is concluded.
-			UE_AUTORTFM_ONCOMMIT(this)
+			UE_AUTORTFM_ONCOMMIT(That, &RefCount)
 			{
-				ImmediatelyRelease();
+				ImmediatelyRelease<PtrToRefCount, DeleteFn, T>(That);
 			};
 		}
 		else
 		{
-			RefsToReturn = ImmediatelyRelease() - 1;
+			RefsToReturn = ImmediatelyRelease<PtrToRefCount, DeleteFn, T>(That);
 		}
 
 		return RefsToReturn;
 	}
 
-	uint32 GetRefCount() const
+	static uint32 GetRefCount(std::atomic<uint32>& RefCount)
 	{
+		// This is equivalent to https://en.cppreference.com/w/cpp/memory/shared_ptr/use_count
+		// 
 		// Inside of an AutoRTFM transaction, the returned refcount value may be higher than you'd
 		// expect, because all Release calls are deferred until the transaction commit time.
 		uint32 Refs;
 		UE_AUTORTFM_OPEN
 		{
-			Refs = uint32(NumRefs.load(std::memory_order_relaxed));
+			// A 'live' reference count is unstable by nature and so there's no benefit
+			// to try and enforce memory ordering around the reading of it.
+			Refs = RefCount.load(std::memory_order_relaxed);
 		};
 		return Refs;
 	}
 
 private:
-	uint32 ImmediatelyRelease() const
+	static void CheckRefCountIsNonZero();
+
+	template <auto PtrToRefCount, auto DeleteFn, typename T>
+	static uint32 ImmediatelyRelease(const T* That)
 	{
-		// This decrements the reference count and immediately frees the object if it has become zero.
-		// A check-failure is issued if an object is over-released.
-		// Returns the refcount _before_ it was decremented.
-		int32 RefsBeforeRelease = NumRefs--;
+		std::atomic<uint32>& RefCount = const_cast<std::atomic<uint32>&>(That->*PtrToRefCount);
+
+		// fetch_sub returns the refcount _before_ it was decremented. std::memory_order_acq_rel is
+		// used so that, if we do end up executing the destructor, it's not possible for side effects 
+		// from executing the destructor to end up being visible before we've determined that the 
+		// reference count is actually zero.
+		uint32 RefsBeforeRelease = RefCount.fetch_sub(1, std::memory_order_acq_rel);
+
 #if DO_GUARD_SLOW
+		// A check-failure is issued if an object is over-released.
 		if (RefsBeforeRelease == 0)
 		{
 			CheckRefCountIsNonZero();
 		}
 #endif
+		// We immediately free the object if its refcount has become zero.
 		if (RefsBeforeRelease == 1)
 		{
-			delete this;
+			DeleteFn(That);
 		}
-		return (uint32)RefsBeforeRelease;
+		return RefsBeforeRelease;
+	}
+};
+
+}
+
+/** A virtual interface for ref counted objects to implement. */
+class IRefCountedObject
+{
+public:
+	virtual ~IRefCountedObject() { }
+	virtual FReturnedRefCountValue AddRef() const = 0;
+
+	// TODO (SOL-7350): return FReturnedRefCountValue from Release(); clean up call sites
+	// which rely on its return value.
+	virtual uint32 Release() const = 0;
+
+	// TODO (SOL-7350): mark this function as deprecated; clean up existing callers.
+	virtual uint32 GetRefCount() const = 0;
+};
+
+/**
+ * Base class implementing thread-safe reference counting.
+ */
+class FRefCountBase
+{
+public:
+	FRefCountBase() = default;
+	virtual ~FRefCountBase() = default;
+
+	FRefCountBase(const FRefCountBase& Rhs) = delete;
+	FRefCountBase& operator=(const FRefCountBase& Rhs) = delete;
+
+	FReturnedRefCountValue AddRef() const
+	{
+		uint32 Refs = UE::Private::FTransactionalAtomicRefCounter::AddRef<&FRefCountBase::NumRefs, DeleteThis>(this);
+		return FReturnedRefCountValue{Refs};
+	}
+
+	uint32 Release() const
+	{
+		return UE::Private::FTransactionalAtomicRefCounter::Release<&FRefCountBase::NumRefs, DeleteThis>(this);
 	}
 
-	mutable std::atomic<int32> NumRefs = 0;
+	uint32 GetRefCount() const
+	{
+		return UE::Private::FTransactionalAtomicRefCounter::GetRefCount(NumRefs);
+	}
 
-	CORE_API void CheckRefCountIsNonZero() const;
+private:
+	static void DeleteThis(const FRefCountBase* This)
+	{
+		delete This;
+	}
+
+	mutable std::atomic<uint32> NumRefs = 0;
 };
 
 /**
@@ -234,29 +286,39 @@ private:
 class FThreadSafeRefCountedObject
 {
 public:
-	FThreadSafeRefCountedObject() : NumRefs(0) {}
-	virtual ~FThreadSafeRefCountedObject() { check(NumRefs.GetValue() == 0); }
+	FThreadSafeRefCountedObject() = default;
+
 	FThreadSafeRefCountedObject(const FThreadSafeRefCountedObject& Rhs) = delete;
 	FThreadSafeRefCountedObject& operator=(const FThreadSafeRefCountedObject& Rhs) = delete;
+
+	virtual ~FThreadSafeRefCountedObject()
+	{ 
+		check(RefCount == 0); 
+	}
+	
 	FReturnedRefCountValue AddRef() const
 	{
-		return FReturnedRefCountValue{uint32(NumRefs.Increment())};
+		uint32 Refs = UE::Private::FTransactionalAtomicRefCounter::AddRef<&FThreadSafeRefCountedObject::RefCount, DeleteThis>(this);
+		return FReturnedRefCountValue{Refs};
 	}
+
 	uint32 Release() const
 	{
-		uint32 Refs = uint32(NumRefs.Decrement());
-		if (Refs == 0)
-		{
-			delete this;
-		}
-		return Refs;
+		return UE::Private::FTransactionalAtomicRefCounter::Release<&FThreadSafeRefCountedObject::RefCount, DeleteThis>(this);
 	}
+
 	uint32 GetRefCount() const
 	{
-		return uint32(NumRefs.GetValue());
+		return UE::Private::FTransactionalAtomicRefCounter::GetRefCount(RefCount);
 	}
+
 private:
-	mutable FThreadSafeCounter NumRefs;
+	static void DeleteThis(const FThreadSafeRefCountedObject* This)
+	{
+		delete This;
+	}
+
+	mutable std::atomic<uint32> RefCount{0};
 };
 
 /**
@@ -311,42 +373,14 @@ public:
 
 		if constexpr (Mode == ERefCountingMode::ThreadSafe)
 		{
-			UE_AUTORTFM_OPEN
-			{
-				Refs = RefCount++;
-			};
-
-			// If we are inside a transaction, a rollback must undo our refcount bump.
-			// In general, this is best handled by Release(). However, there is one case 
-			// that needs to be handled with special care. A brand-new object has a
-			// refcount of zero, and a rollback must return it to this zero-refcount state.
-			// However, calling AddRef() followed by Release() would not accomplish this;
-			// instead, it would free the object entirely! We need to guard against this,
-			// since it could lead to a double-free, so we detect the zero-reference case
-			// and special-case it.
-			if (Refs == 0)
-			{
-				UE_AUTORTFM_ONABORT(this)
-				{
-					--RefCount;
-					// The refcount is likely zero now, but leaving the object alive isn't a leak.
-					// We are restoring the object back to its initial "just-created" state.
-				};
-			}
-			else
-			{
-				UE_AUTORTFM_ONABORT(this)
-				{
-					Release();
-				};
-			}
+			Refs = UE::Private::FTransactionalAtomicRefCounter::AddRef<&TRefCountingMixin<T>::RefCount, T::StaticDestroyMixin>(this);
 		}
 		else
 		{
-			Refs = RefCount++;
+			Refs = ++RefCount;
 		}
 
-		return FReturnedRefCountValue{Refs + 1};
+		return FReturnedRefCountValue{Refs};
 	}
 
 	uint32 Release() const
@@ -355,24 +389,13 @@ public:
 
 		if constexpr (Mode == ERefCountingMode::ThreadSafe)
 		{
-			UE_AUTORTFM_ONCOMMIT(this)
-			{
-				// std::memory_order_acq_rel is used here so that, if we do end up executing the
-				// destructor, it's not possible for side effects from executing the destructor end up
-				// being visible before we've determined that the reference count is actually zero.
-				uint32 OldCount = RefCount.fetch_sub(1, std::memory_order_acq_rel);
-				checkSlow(OldCount > 0);
-				if (OldCount == 1)
-				{
-					T::StaticDestroyObject(static_cast<const T*>(this));
-				}
-			};
+			UE::Private::FTransactionalAtomicRefCounter::Release<&TRefCountingMixin<T>::RefCount, T::StaticDestroyMixin>(this);
 		}
 		else
 		{
 			if (--RefCount == 0)
 			{
-				T::StaticDestroyObject(static_cast<const T*>(this));
+				StaticDestroyMixin(this);
 			}
 		}
 
@@ -384,20 +407,7 @@ public:
 	{
 		if constexpr (Mode == ERefCountingMode::ThreadSafe)
 		{
-			// A 'live' reference count is unstable by nature and so there's no benefit
-			// to try and enforce memory ordering around the reading of it.
-			//
-			// This is equivalent to https://en.cppreference.com/w/cpp/memory/shared_ptr/use_count
-
-			uint32 Count = 0;
-
-			UE_AUTORTFM_OPEN
-			{
-				// This reference count may be accessed by multiple threads
-				Count = RefCount.load(std::memory_order_relaxed);
-			};
-
-			return Count;
+			return UE::Private::FTransactionalAtomicRefCounter::GetRefCount(RefCount);
 		}
 		else
 		{
@@ -411,6 +421,11 @@ public:
 	}
 
 private:
+	static void StaticDestroyMixin(const TRefCountingMixin<T>* This)
+	{
+		T::StaticDestroyObject(static_cast<const T*>(This));
+	}
+
 	mutable RefCountType RefCount{0};
 };
 
