Update LoadFileToString based on review feedback.

The passed-in FArchive is no longer closed upon successfully reading the file. This is a behavioral change from the past; historically, we would close the FArchive upon a successful read of a non-zero-sized file.

The IPlatformFile overload of LoadFileToString now only logs errors if PlatformFile->OpenRead fails to open the file. Failure to read the file will not directly log an error, although `false` will be returned if Reader.IsError() is set. This restores the previous behavior.

We no longer unnecessarily create a UniquePtr (and then immediately check it for null) within the IPlatformFile overload.

diff --git a/Engine/Source/Runtime/Core/Private/Misc/FileHelper.cpp b/Engine/Source/Runtime/Core/Private/Misc/FileHelper.cpp
--- a/Engine/Source/Runtime/Core/Private/Misc/FileHelper.cpp
+++ b/Engine/Source/Runtime/Core/Private/Misc/FileHelper.cpp
@@ -162,20 +162,18 @@ bool FFileHelper::LoadFileToString(FString& Result, FArchive& Reader, EHashOptio
 	if (!Size)
 	{
 		Result.Empty();
-		Reader.Close();
 		return true;
 	}
 
 	if (Reader.Tell() != 0)
 	{
 		UE_LOG(LogStreaming, Warning, TEXT("Archive '%s' has already been read from."), *Reader.GetArchiveName());
-		Reader.Close();
 		return false;
 	}
 
 	uint8* Ch = (uint8*)FMemory::Malloc(Size);
 	Reader.Serialize(Ch, Size);
-	bool bSuccess = Reader.Close();
+	bool bSuccess = !Reader.IsError();
 
 	BufferToString(Result, Ch, (int32)Size);
 
@@ -232,27 +230,25 @@ bool FFileHelper::LoadFileToString(FString& Result, IPlatformFile* PlatformFile,
 
 			if (IFileHandle* File = PlatformFile->OpenRead(Filename))
 			{
-				if (TUniquePtr<FArchive> Reader = MakeUnique<FArchiveFileReaderGeneric>(File, Filename, File->Size()))
+				FArchiveFileReaderGeneric Reader{File, Filename, File->Size()};
+				bSuccess = LoadFileToString(FileData, Reader, VerifyFlags);
+			}
+			else
+			{
+				if (ReadFlags & FILEREAD_NoFail)
 				{
-					bSuccess = LoadFileToString(FileData, *Reader, VerifyFlags);
+					UE_LOG(LogStreaming, Fatal, TEXT("Failed to read file: %s"), Filename);
+				}
+
+				if (!(ReadFlags & FILEREAD_Silent))
+				{
+					UE_LOG(LogStreaming, Warning, TEXT("Failed to read file '%s' error."), Filename);
 				}
 			}
 
 			return FileData;
 		});
 
-		if (!bSuccess)
-		{
-			if (ReadFlags & FILEREAD_NoFail)
-			{
-				UE_LOG(LogStreaming, Fatal, TEXT("Failed to read file: %s"), Filename);
-			}
-
-			if (!(ReadFlags & FILEREAD_Silent))
-			{
-				UE_LOG(LogStreaming, Warning, TEXT("Failed to read file '%s' error."), Filename);
-			}
-		}
 	}
 
 	return bSuccess;
diff --git a/Engine/Source/Runtime/Core/Public/Misc/FileHelper.h b/Engine/Source/Runtime/Core/Public/Misc/FileHelper.h
--- a/Engine/Source/Runtime/Core/Public/Misc/FileHelper.h
+++ b/Engine/Source/Runtime/Core/Public/Misc/FileHelper.h
@@ -78,7 +78,7 @@ struct FFileHelper
 	static CORE_API bool LoadFileToArray( TArray64<uint8>& Result, const TCHAR* Filename, uint32 Flags = 0 );
 
 	/**
-	 * Loads a text file from an FArchive to an FString, then closes the archive. Supports all combination of ANSI/Unicode files and platforms.
+	 * Loads a text file from an FArchive to an FString. Supports all combination of ANSI/Unicode files and platforms.
 	 *
 	 * @param Result       String representation of the loaded file
 	 * @param Reader       The archive to load from
