[Backout] - CL39441791

Landed this accidentally, was not ready to land.

#rnx
Original CL Desc
-----------------------------------------------------------------
[AutoRTFM] Make thread safe refcount class transactionally safe.

diff --git a/Engine/Source/Programs/AutoRTFMTests/Private/RefCountPtrTests.cpp b/Engine/Source/Programs/AutoRTFMTests/Private/RefCountPtrTests.cpp
--- a/Engine/Source/Programs/AutoRTFMTests/Private/RefCountPtrTests.cpp
+++ b/Engine/Source/Programs/AutoRTFMTests/Private/RefCountPtrTests.cpp
@@ -17,15 +17,9 @@ namespace
 		RefCountedMixinInt(int InValue) : Value(InValue) {}
 		int Value;
 	};
-
-	struct ThreadSafeRefCountedInt : public FThreadSafeRefCountedObject
-	{
-		ThreadSafeRefCountedInt(int InValue) : Value(InValue) {}
-		int Value;
-	};
 }
 
-TEMPLATE_TEST_CASE("PreviouslyAllocated", "[RefCountPtr]", RefCountedInt, RefCountedMixinInt, ThreadSafeRefCountedInt)
+TEMPLATE_TEST_CASE("PreviouslyAllocated", "[RefCountPtr]", RefCountedInt, RefCountedMixinInt)
 {
 	SECTION("Using new T()")
 	{
@@ -58,7 +52,7 @@ TEMPLATE_TEST_CASE("PreviouslyAllocated", "[RefCountPtr]", RefCountedInt, RefCou
 	}
 }
 
-TEMPLATE_TEST_CASE("AbortWithPreviouslyAllocated", "[RefCountPtr]", RefCountedInt, RefCountedMixinInt, ThreadSafeRefCountedInt)
+TEMPLATE_TEST_CASE("AbortWithPreviouslyAllocated", "[RefCountPtr]", RefCountedInt, RefCountedMixinInt)
 {
 	SECTION("Using new T()")
 	{
@@ -95,7 +89,7 @@ TEMPLATE_TEST_CASE("AbortWithPreviouslyAllocated", "[RefCountPtr]", RefCountedIn
 	}
 }
 
-TEMPLATE_TEST_CASE("AbortWithNewlyAllocated", "[RefCountPtr]", RefCountedInt, RefCountedMixinInt, ThreadSafeRefCountedInt)
+TEMPLATE_TEST_CASE("AbortWithNewlyAllocated", "[RefCountPtr]", RefCountedInt, RefCountedMixinInt)
 {
 	int Result = 42;
 
@@ -144,7 +138,7 @@ TEMPLATE_TEST_CASE("AbortWithNewlyAllocated", "[RefCountPtr]", RefCountedInt, Re
 	REQUIRE(Result == 42);
 }
 
-TEMPLATE_TEST_CASE("OnCommitCapturingRefCountPtr", "[RefCountPtr]", RefCountedInt, RefCountedMixinInt, ThreadSafeRefCountedInt)
+TEMPLATE_TEST_CASE("OnCommitCapturingRefCountPtr", "[RefCountPtr]", RefCountedInt, RefCountedMixinInt)
 {
 	TRefCountPtr<TestType> Foo = MakeRefCount<TestType>(42);
 
@@ -176,7 +170,7 @@ TEMPLATE_TEST_CASE("OnCommitCapturingRefCountPtr", "[RefCountPtr]", RefCountedIn
 	}
 }
 
-TEMPLATE_TEST_CASE("OnAbortCapturingRefCountPtr", "[RefCountPtr]", RefCountedInt, RefCountedMixinInt, ThreadSafeRefCountedInt)
+TEMPLATE_TEST_CASE("OnAbortCapturingRefCountPtr", "[RefCountPtr]", RefCountedInt, RefCountedMixinInt)
 {
 	TRefCountPtr<TestType> Foo = MakeRefCount<TestType>(42);
 
diff --git a/Engine/Source/Runtime/Core/Private/Templates/RefCounting.cpp b/Engine/Source/Runtime/Core/Private/Templates/RefCounting.cpp
--- a/Engine/Source/Runtime/Core/Private/Templates/RefCounting.cpp
+++ b/Engine/Source/Runtime/Core/Private/Templates/RefCounting.cpp
@@ -2,7 +2,7 @@
 
 #include "Templates/RefCounting.h"
 
-void UE::Private::FTransactionalAtomicRefCounter::CheckRefCountIsNonZero()
+void FRefCountBase::CheckRefCountIsNonZero() const
 {
 	check(!"Release() was called on an object which is already at a zero refcount.");
 }
\ No newline at end of file
diff --git a/Engine/Source/Runtime/Core/Public/Templates/RefCounting.h b/Engine/Source/Runtime/Core/Public/Templates/RefCounting.h
--- a/Engine/Source/Runtime/Core/Public/Templates/RefCounting.h
+++ b/Engine/Source/Runtime/Core/Public/Templates/RefCounting.h
@@ -56,25 +56,39 @@ private:
 	uint32 RefCount = 0;
 };
 
-/**
- * UE::Private::FTransactionalAtomicRefCounter provides transactionally-safe atomic refcounting primitives.
- * These are used by FRefCountBase, FThreadSafeRefCountedObject and TRefCountingMixin. This isn't public API.
- */
-namespace UE::Private
+/** A virtual interface for ref counted objects to implement. */
+class IRefCountedObject
 {
+public:
+	virtual ~IRefCountedObject() { }
+	virtual FReturnedRefCountValue AddRef() const = 0;
+
+	// TODO (SOL-7350): return FReturnedRefCountValue from Release(); clean up call sites
+	// which rely on its return value.
+	virtual uint32 Release() const = 0;
+
+	// TODO (SOL-7350): mark this function as deprecated; clean up existing callers.
+	virtual uint32 GetRefCount() const = 0;
+};
 
-class FTransactionalAtomicRefCounter
+/**
+ * Base class implementing thread-safe reference counting.
+ */
+class FRefCountBase
 {
 public:
-	template <auto PtrToRefCount, auto DeleteFn, typename T>
-	static uint32 AddRef(const T* That)
+	FRefCountBase() = default;
+	virtual ~FRefCountBase() = default;
+
+	FRefCountBase(const FRefCountBase& Rhs) = delete;
+	FRefCountBase& operator=(const FRefCountBase& Rhs) = delete;
+
+	FReturnedRefCountValue AddRef() const
 	{
-		std::atomic<uint32>& RefCount = const_cast<std::atomic<uint32>&>(That->*PtrToRefCount);
 		uint32 Refs;
-
 		UE_AUTORTFM_OPEN
 		{
-			Refs = RefCount++;
+			Refs = NumRefs++;
 		};
 
 		// If we are inside a transaction, a rollback must undo our refcount bump.
@@ -87,31 +101,28 @@ public:
 		// and special-case it.
 		if (Refs == 0)
 		{
-			UE_AUTORTFM_ONABORT(&RefCount)
+			UE_AUTORTFM_ONABORT(this)
 			{
-				--RefCount;
+				--NumRefs;
 				// The refcount is likely zero now, but leaving the object alive isn't a leak.
 				// We are restoring the object back to its initial "just-created" state.
 			};
 		}
 		else
 		{
-			UE_AUTORTFM_ONABORT(That)
+			UE_AUTORTFM_ONABORT(this)
 			{
-				Release<PtrToRefCount, DeleteFn, T>(That);
+				Release();
 			};
 		}
 
-		return Refs + 1;
+		return FReturnedRefCountValue{Refs + 1};
 	}
 
-	template <auto PtrToRefCount, auto DeleteFn, typename T>
-	static uint32 Release(const T* That)
+	uint32 Release() const
 	{
-		std::atomic<uint32>& RefCount = const_cast<std::atomic<uint32>&>(That->*PtrToRefCount);
-
 #if DO_GUARD_SLOW
-		if (RefCount == 0)
+		if (NumRefs == 0)
 		{
 			CheckRefCountIsNonZero();
 		}
@@ -128,121 +139,58 @@ public:
 			// always returns the true refcount.
 			UE_AUTORTFM_OPEN
 			{
-				RefsToReturn = uint32(RefCount.load(std::memory_order_relaxed)) - 1;
+				RefsToReturn = uint32(NumRefs.load(std::memory_order_relaxed)) - 1;
 			};
 
 			// Refcount changes and frees are deferred until the transaction is concluded.
-			UE_AUTORTFM_ONCOMMIT(That, &RefCount)
+			UE_AUTORTFM_ONCOMMIT(this)
 			{
-				ImmediatelyRelease<PtrToRefCount, DeleteFn, T>(That);
+				ImmediatelyRelease();
 			};
 		}
 		else
 		{
-			RefsToReturn = ImmediatelyRelease<PtrToRefCount, DeleteFn, T>(That);
+			RefsToReturn = ImmediatelyRelease() - 1;
 		}
 
 		return RefsToReturn;
 	}
 
-	static uint32 GetRefCount(std::atomic<uint32>& RefCount)
+	uint32 GetRefCount() const
 	{
-		// This is equivalent to https://en.cppreference.com/w/cpp/memory/shared_ptr/use_count
-		// 
 		// Inside of an AutoRTFM transaction, the returned refcount value may be higher than you'd
 		// expect, because all Release calls are deferred until the transaction commit time.
 		uint32 Refs;
 		UE_AUTORTFM_OPEN
 		{
-			// A 'live' reference count is unstable by nature and so there's no benefit
-			// to try and enforce memory ordering around the reading of it.
-			Refs = RefCount.load(std::memory_order_relaxed);
+			Refs = uint32(NumRefs.load(std::memory_order_relaxed));
 		};
 		return Refs;
 	}
 
 private:
-	static void CheckRefCountIsNonZero();
-
-	template <auto PtrToRefCount, auto DeleteFn, typename T>
-	static uint32 ImmediatelyRelease(const T* That)
+	uint32 ImmediatelyRelease() const
 	{
-		std::atomic<uint32>& RefCount = const_cast<std::atomic<uint32>&>(That->*PtrToRefCount);
-
-		// fetch_sub returns the refcount _before_ it was decremented. std::memory_order_acq_rel is
-		// used so that, if we do end up executing the destructor, it's not possible for side effects 
-		// from executing the destructor to end up being visible before we've determined that the 
-		// reference count is actually zero.
-		uint32 RefsBeforeRelease = RefCount.fetch_sub(1, std::memory_order_acq_rel);
-
-#if DO_GUARD_SLOW
+		// This decrements the reference count and immediately frees the object if it has become zero.
 		// A check-failure is issued if an object is over-released.
+		// Returns the refcount _before_ it was decremented.
+		int32 RefsBeforeRelease = NumRefs--;
+#if DO_GUARD_SLOW
 		if (RefsBeforeRelease == 0)
 		{
 			CheckRefCountIsNonZero();
 		}
 #endif
-		// We immediately free the object if its refcount has become zero.
 		if (RefsBeforeRelease == 1)
 		{
-			DeleteFn(That);
+			delete this;
 		}
-		return RefsBeforeRelease;
-	}
-};
-
-}
-
-/** A virtual interface for ref counted objects to implement. */
-class IRefCountedObject
-{
-public:
-	virtual ~IRefCountedObject() { }
-	virtual FReturnedRefCountValue AddRef() const = 0;
-
-	// TODO (SOL-7350): return FReturnedRefCountValue from Release(); clean up call sites
-	// which rely on its return value.
-	virtual uint32 Release() const = 0;
-
-	// TODO (SOL-7350): mark this function as deprecated; clean up existing callers.
-	virtual uint32 GetRefCount() const = 0;
-};
-
-/**
- * Base class implementing thread-safe reference counting.
- */
-class FRefCountBase
-{
-public:
-	FRefCountBase() = default;
-	virtual ~FRefCountBase() = default;
-
-	FRefCountBase(const FRefCountBase& Rhs) = delete;
-	FRefCountBase& operator=(const FRefCountBase& Rhs) = delete;
-
-	FReturnedRefCountValue AddRef() const
-	{
-		uint32 Refs = UE::Private::FTransactionalAtomicRefCounter::AddRef<&FRefCountBase::NumRefs, DeleteThis>(this);
-		return FReturnedRefCountValue{Refs};
-	}
-
-	uint32 Release() const
-	{
-		return UE::Private::FTransactionalAtomicRefCounter::Release<&FRefCountBase::NumRefs, DeleteThis>(this);
+		return (uint32)RefsBeforeRelease;
 	}
 
-	uint32 GetRefCount() const
-	{
-		return UE::Private::FTransactionalAtomicRefCounter::GetRefCount(NumRefs);
-	}
+	mutable std::atomic<int32> NumRefs = 0;
 
-private:
-	static void DeleteThis(const FRefCountBase* This)
-	{
-		delete This;
-	}
-
-	mutable std::atomic<uint32> NumRefs = 0;
+	CORE_API void CheckRefCountIsNonZero() const;
 };
 
 /**
@@ -286,39 +234,29 @@ private:
 class FThreadSafeRefCountedObject
 {
 public:
-	FThreadSafeRefCountedObject() = default;
-
+	FThreadSafeRefCountedObject() : NumRefs(0) {}
+	virtual ~FThreadSafeRefCountedObject() { check(NumRefs.GetValue() == 0); }
 	FThreadSafeRefCountedObject(const FThreadSafeRefCountedObject& Rhs) = delete;
 	FThreadSafeRefCountedObject& operator=(const FThreadSafeRefCountedObject& Rhs) = delete;
-
-	virtual ~FThreadSafeRefCountedObject()
-	{ 
-		check(RefCount == 0); 
-	}
-	
 	FReturnedRefCountValue AddRef() const
 	{
-		uint32 Refs = UE::Private::FTransactionalAtomicRefCounter::AddRef<&FThreadSafeRefCountedObject::RefCount, DeleteThis>(this);
-		return FReturnedRefCountValue{Refs};
+		return FReturnedRefCountValue{uint32(NumRefs.Increment())};
 	}
-
 	uint32 Release() const
 	{
-		return UE::Private::FTransactionalAtomicRefCounter::Release<&FThreadSafeRefCountedObject::RefCount, DeleteThis>(this);
+		uint32 Refs = uint32(NumRefs.Decrement());
+		if (Refs == 0)
+		{
+			delete this;
+		}
+		return Refs;
 	}
-
 	uint32 GetRefCount() const
 	{
-		return UE::Private::FTransactionalAtomicRefCounter::GetRefCount(RefCount);
+		return uint32(NumRefs.GetValue());
 	}
-
 private:
-	static void DeleteThis(const FThreadSafeRefCountedObject* This)
-	{
-		delete This;
-	}
-
-	mutable std::atomic<uint32> RefCount{0};
+	mutable FThreadSafeCounter NumRefs;
 };
 
 /**
@@ -373,14 +311,42 @@ public:
 
 		if constexpr (Mode == ERefCountingMode::ThreadSafe)
 		{
-			Refs = UE::Private::FTransactionalAtomicRefCounter::AddRef<&TRefCountingMixin<T>::RefCount, T::StaticDestroyMixin>(this);
+			UE_AUTORTFM_OPEN
+			{
+				Refs = RefCount++;
+			};
+
+			// If we are inside a transaction, a rollback must undo our refcount bump.
+			// In general, this is best handled by Release(). However, there is one case 
+			// that needs to be handled with special care. A brand-new object has a
+			// refcount of zero, and a rollback must return it to this zero-refcount state.
+			// However, calling AddRef() followed by Release() would not accomplish this;
+			// instead, it would free the object entirely! We need to guard against this,
+			// since it could lead to a double-free, so we detect the zero-reference case
+			// and special-case it.
+			if (Refs == 0)
+			{
+				UE_AUTORTFM_ONABORT(this)
+				{
+					--RefCount;
+					// The refcount is likely zero now, but leaving the object alive isn't a leak.
+					// We are restoring the object back to its initial "just-created" state.
+				};
+			}
+			else
+			{
+				UE_AUTORTFM_ONABORT(this)
+				{
+					Release();
+				};
+			}
 		}
 		else
 		{
-			Refs = ++RefCount;
+			Refs = RefCount++;
 		}
 
-		return FReturnedRefCountValue{Refs};
+		return FReturnedRefCountValue{Refs + 1};
 	}
 
 	uint32 Release() const
@@ -389,13 +355,24 @@ public:
 
 		if constexpr (Mode == ERefCountingMode::ThreadSafe)
 		{
-			UE::Private::FTransactionalAtomicRefCounter::Release<&TRefCountingMixin<T>::RefCount, T::StaticDestroyMixin>(this);
+			UE_AUTORTFM_ONCOMMIT(this)
+			{
+				// std::memory_order_acq_rel is used here so that, if we do end up executing the
+				// destructor, it's not possible for side effects from executing the destructor end up
+				// being visible before we've determined that the reference count is actually zero.
+				uint32 OldCount = RefCount.fetch_sub(1, std::memory_order_acq_rel);
+				checkSlow(OldCount > 0);
+				if (OldCount == 1)
+				{
+					T::StaticDestroyObject(static_cast<const T*>(this));
+				}
+			};
 		}
 		else
 		{
 			if (--RefCount == 0)
 			{
-				StaticDestroyMixin(this);
+				T::StaticDestroyObject(static_cast<const T*>(this));
 			}
 		}
 
@@ -407,7 +384,20 @@ public:
 	{
 		if constexpr (Mode == ERefCountingMode::ThreadSafe)
 		{
-			return UE::Private::FTransactionalAtomicRefCounter::GetRefCount(RefCount);
+			// A 'live' reference count is unstable by nature and so there's no benefit
+			// to try and enforce memory ordering around the reading of it.
+			//
+			// This is equivalent to https://en.cppreference.com/w/cpp/memory/shared_ptr/use_count
+
+			uint32 Count = 0;
+
+			UE_AUTORTFM_OPEN
+			{
+				// This reference count may be accessed by multiple threads
+				Count = RefCount.load(std::memory_order_relaxed);
+			};
+
+			return Count;
 		}
 		else
 		{
@@ -421,11 +411,6 @@ public:
 	}
 
 private:
-	static void StaticDestroyMixin(const TRefCountingMixin<T>* This)
-	{
-		T::StaticDestroyObject(static_cast<const T*>(This));
-	}
-
 	mutable RefCountType RefCount{0};
 };
 
