Fix Csv profiler global destructor order crash
- The TlsCS is used in EndOfPipeData's destructor, but since they are both global objects, the destruction order is undefined, and can happen the wrong way around.
- Grouping the objects into a single global object struct enforces the correct destruction order (reverse order to their declaration in the struct).

diff --git a/Engine/Source/Runtime/Core/Private/ProfilingDebugging/CsvProfiler.cpp b/Engine/Source/Runtime/Core/Private/ProfilingDebugging/CsvProfiler.cpp
--- a/Engine/Source/Runtime/Core/Private/ProfilingDebugging/CsvProfiler.cpp
+++ b/Engine/Source/Runtime/Core/Private/ProfilingDebugging/CsvProfiler.cpp
@@ -2308,37 +2308,42 @@ private:
 	static CSV_PROFILER_INLINE uint64 GetStatID(const char* StatName) { return uint64(StatName); }
 	static CSV_PROFILER_INLINE uint64 GetStatID(const FName& StatId) { return StatId.ToUnstableInt(); }
 
-	static FCriticalSection TlsCS;
-	static TArray<FWeakPtr> TlsInstances;
-	static uint32 TlsSlot;
+	// Group these objects together to enforce correct global-scope object destruction order.
+	struct FSingleton
+	{
+		FCriticalSection TlsCS;
+		TArray<FWeakPtr> TlsInstances;
+		uint32 TlsSlot = FPlatformTLS::InvalidTlsSlot;
+
+		FSharedPtr EndOfPipeData = MakeShared<FCsvProfilerThreadData>(ECsvTimeline::EndOfPipe);
 
-	static FSharedPtr EndOfPipeData;
+	} static Singleton;
 
 public:
 	static void InitTls()
 	{
-		if (!FPlatformTLS::IsValidTlsSlot(TlsSlot))
+		if (!FPlatformTLS::IsValidTlsSlot(Singleton.TlsSlot))
 		{
-			TlsSlot = FPlatformTLS::AllocTlsSlot();
+			Singleton.TlsSlot = FPlatformTLS::AllocTlsSlot();
 			FPlatformMisc::MemoryBarrier();
 		}
 	}
 
 	static bool IsTlsSlotInitialized()
 	{
-		return FPlatformTLS::IsValidTlsSlot(TlsSlot);
+		return FPlatformTLS::IsValidTlsSlot(Singleton.TlsSlot);
 	}
 
 	static FORCENOINLINE FCsvProfilerThreadData* CreateTLSData()
 	{
 		QUICK_SCOPE_CYCLE_COUNTER(CSVProfiler_ThreadData_CreateTLSData);
-		FScopeLock Lock(&TlsCS);
+		FScopeLock Lock(&Singleton.TlsCS);
 
 		FSharedPtr ProfilerThreadPtr = MakeShareable(new FCsvProfilerThreadData());
-		FPlatformTLS::SetTlsValue(TlsSlot, ProfilerThreadPtr.Get());
+		FPlatformTLS::SetTlsValue(Singleton.TlsSlot, ProfilerThreadPtr.Get());
 
 		// Keep a weak reference to this thread data in the global array.
-		TlsInstances.Emplace(ProfilerThreadPtr);
+		Singleton.TlsInstances.Emplace(ProfilerThreadPtr);
 
 		// Register the shared ptr in the thread's TLS auto-cleanup list.
 		// When the thread exits, it will delete the shared ptr, releasing its reference.
@@ -2349,7 +2354,7 @@ public:
 
 	static CSV_PROFILER_INLINE FCsvProfilerThreadData& Get()
 	{
-		FCsvProfilerThreadData* ProfilerThread = (FCsvProfilerThreadData*)FPlatformTLS::GetTlsValue(TlsSlot);
+		FCsvProfilerThreadData* ProfilerThread = (FCsvProfilerThreadData*)FPlatformTLS::GetTlsValue(Singleton.TlsSlot);
 		if (UNLIKELY(!ProfilerThread))
 		{
 			ProfilerThread = CreateTLSData();
@@ -2359,18 +2364,18 @@ public:
 
 	static CSV_PROFILER_INLINE FCsvProfilerThreadData& GetEndOfPipe()
 	{
-		return *EndOfPipeData.Get();
+		return *Singleton.EndOfPipeData.Get();
 	}
 
 	static inline void GetTlsInstances(TArray<FSharedPtr>& OutTlsInstances)
 	{
 		QUICK_SCOPE_CYCLE_COUNTER(CSVProfiler_ThreadData_GetTlsInstances);
-		FScopeLock Lock(&TlsCS);
-		OutTlsInstances.Empty(TlsInstances.Num() + 1);
+		FScopeLock Lock(&Singleton.TlsCS);
+		OutTlsInstances.Empty(Singleton.TlsInstances.Num() + 1);
 
-		for (int32 Index = TlsInstances.Num() - 1; Index >= 0; --Index)
+		for (int32 Index = Singleton.TlsInstances.Num() - 1; Index >= 0; --Index)
 		{
-			FSharedPtr SharedPtr = TlsInstances[Index].Pin();
+			FSharedPtr SharedPtr = Singleton.TlsInstances[Index].Pin();
 			if (SharedPtr.IsValid())
 			{
 				// Thread is still alive.
@@ -2379,7 +2384,7 @@ public:
 		}
 
 		// Always add the end-of-pipe data
-		OutTlsInstances.Add(EndOfPipeData);
+		OutTlsInstances.Add(Singleton.EndOfPipeData);
 	}
 
 	FCsvProfilerThreadData(TOptional<ECsvTimeline::Type> InCsvTimeline = NullOpt)
@@ -2406,8 +2411,8 @@ public:
 
 		// Clean up dead entries in the thread data array.
 		// This will remove both the current instance, and any others that have expired.
-		FScopeLock Lock(&TlsCS);
-		for (auto Iter = TlsInstances.CreateIterator(); Iter; ++Iter)
+		FScopeLock Lock(&Singleton.TlsCS);
+		for (auto Iter = Singleton.TlsInstances.CreateIterator(); Iter; ++Iter)
 		{
 			if (!Iter->IsValid())
 			{
@@ -2580,11 +2585,7 @@ public:
 	TArray<FCsvWaitStatName> WaitStatNameStack;
 };
 
-FCsvProfilerThreadData::FSharedPtr FCsvProfilerThreadData::EndOfPipeData(MakeShared<FCsvProfilerThreadData>(ECsvTimeline::EndOfPipe));
-
-uint32 FCsvProfilerThreadData::TlsSlot = FPlatformTLS::InvalidTlsSlot;
-FCriticalSection FCsvProfilerThreadData::TlsCS;
-TArray<FCsvProfilerThreadData::FWeakPtr> FCsvProfilerThreadData::TlsInstances;
+FCsvProfilerThreadData::FSingleton FCsvProfilerThreadData::Singleton;
 
 // Responsible to collating and accumulating stats for a thread, or group of threads
 class FCsvThreadGroupStatProcessor
