Improve support for merged modular build module manager

- Add delay unload mechanism
- Properly free UObjects
- Add debugging facilities

#rnx

diff --git a/Engine/Plugins/Runtime/GameFeatures/Source/GameFeatures/Private/GameFeaturePluginStateMachine.cpp b/Engine/Plugins/Runtime/GameFeatures/Source/GameFeatures/Private/GameFeaturePluginStateMachine.cpp
--- a/Engine/Plugins/Runtime/GameFeatures/Source/GameFeatures/Private/GameFeaturePluginStateMachine.cpp
+++ b/Engine/Plugins/Runtime/GameFeatures/Source/GameFeatures/Private/GameFeaturePluginStateMachine.cpp
@@ -2050,12 +2050,22 @@ struct FGameFeaturePluginState_Unmounting : public FGameFeaturePluginState
 				IPluginManager::Get().UnmountExplicitlyLoadedPluginLocalizationData(StateProperties.PluginName);
 			}
 
+#if UE_MERGED_MODULES
+			// We normally do not allow unloading code because it's difficult to make that safe.
+			// For example, destructors can be called much later than we'd expect because of garbage collection or async work.
+			// We allow this behavior for merged modular builds because the have a smaller scope (intended for console clients) and
+			// unloading code to save memory is the main goal for that feature.
+			static constexpr bool bAllowUnloadCode = true;
+#else
+			static constexpr bool bAllowUnloadCode = false;
+#endif // UE_MERGED_MODULES
+
 			// The asset registry listens to FPackageName::OnContentPathDismounted() and 
 			// will automatically cleanup the asset registry state we added for this plugin.
 			// This will also cause any assets we added to the asset manager to be removed.
 			// Scan paths added to the asset manager should have already been cleaned up.
 			FText FailureReason;
-			if (!IPluginManager::Get().UnmountExplicitlyLoadedPlugin(StateProperties.PluginName, &FailureReason, false))
+			if (!IPluginManager::Get().UnmountExplicitlyLoadedPlugin(StateProperties.PluginName, &FailureReason, bAllowUnloadCode))
 			{
 				const FStringView ShortUrl = StateProperties.PluginIdentifier.GetIdentifyingString();
 				ensureMsgf(false, TEXT("Failed to explicitly unmount the PluginURL(%.*s) because %s"), ShortUrl.Len(), ShortUrl.GetData(), *FailureReason.ToString());
diff --git a/Engine/Source/Runtime/Core/Private/Modules/ModuleManager.cpp b/Engine/Source/Runtime/Core/Private/Modules/ModuleManager.cpp
--- a/Engine/Source/Runtime/Core/Private/Modules/ModuleManager.cpp
+++ b/Engine/Source/Runtime/Core/Private/Modules/ModuleManager.cpp
@@ -27,7 +27,45 @@ DEFINE_LOG_CATEGORY_STATIC(LogModuleManager, Log, All);
 #if WITH_HOT_RELOAD
 	/** If true, we are reloading a class for HotReload */
 	CORE_API bool GIsHotReload = false;
-#endif
+#endif // WITH_HOT_RELOAD
+
+#if UE_MERGED_MODULES && !UE_BUILD_SHIPPING
+
+	// Log user modules of merged libraries
+	FAutoConsoleCommand LogMergedLibraryUsageCommand(
+		TEXT("Modules.MergedLibraries"),
+		TEXT("Log all users of currently loaded merged libraries"),
+		FConsoleCommandDelegate::CreateLambda([]()
+		{
+			FModuleManager::Get().LogMergedLibraryUsage();
+		}));
+
+	// Force load an entire merged library
+	FAutoConsoleCommand LoadMergedLibraryCommand(
+		TEXT("Modules.LoadMergedLibrary"),
+		TEXT("Load all modules from a merged library"),
+		FConsoleCommandWithArgsDelegate::CreateLambda([](const TArray<FString>& Args)
+		{
+			if (Args.Num() == 1)
+			{
+				FModuleManager::Get().LoadAllModulesInMergedLibrary(Args[0]);
+			}
+		}));
+
+	// Force unload an entire merged library
+	FAutoConsoleCommand UnloadMergedLibraryCommand(
+		TEXT("Modules.UnloadMergedLibrary"),
+		TEXT("Unload all modules from a merged library"),
+		FConsoleCommandWithArgsDelegate::CreateLambda([](const TArray<FString>& Args)
+		{
+			if (Args.Num() == 1)
+			{
+				FModuleManager::Get().UnloadAllModulesInMergedLibrary(Args[0]);
+			}
+		}));
+
+
+#endif // UE_MERGED_MODULES && !UE_BUILD_SHIPPING
 
 #if WITH_ENGINE
 	TMap<UClass*, UClass*>& GetClassesToReinstanceForHotReload()
@@ -35,7 +73,7 @@ DEFINE_LOG_CATEGORY_STATIC(LogModuleManager, Log, All);
 		static TMap<UClass*, UClass*> Data;
 		return Data;
 	}
-#endif
+#endif // WITH_ENGINE
 
 FModuleInitializerEntry* GFirstModuleInitializerEntry;
 
@@ -153,6 +191,7 @@ FModuleManager::FModuleManager(FPrivateToken)
 	// Modules bootstrapping is useful to avoid costly directory enumeration by reloading
 	// a serialized state of the module manager. Can only be used when run in the exact
 	// same context multiple times (i.e. starting multiple shader compile workers)
+	// When using UE_MERGED_MODULES, module initialization is done with statics as if the build was monolithic.
 	FString ModulesBootstrapFilename;
 	if (FParse::Value(FCommandLine::Get(), TEXT("ModulesBootstrap="), ModulesBootstrapFilename))
 	{
@@ -208,6 +247,8 @@ void FModuleManager::FindModules(const TCHAR* WildcardWithoutExtension, TArray<F
 void FModuleManager::FindModules(const TCHAR* WildcardWithoutExtension, TArray<FModuleDiskInfo>& OutModules) const
 {
 	// @todo plugins: Try to convert existing use cases to use plugins, and get rid of this function
+
+	// Merged modular builds act as if they were monolithic as far as module discovery is concerned.
 #if !IS_MONOLITHIC && !UE_MERGED_MODULES
 
 	TMap<FName, FString> ModulePaths;
@@ -438,7 +479,7 @@ void FModuleManager::RefreshModuleFilenameFromManifestImpl(const FName InModuleN
 	{
 		return;
 	}
-	#endif
+	#endif // UE_MERGED_MODULES
 
 	if (!ensureMsgf(MatchPos != INDEX_NONE, TEXT("Could not find module name '%s' in module filename '%s'"), *InModuleName.ToString(), *ModuleFilename))
 	{
@@ -479,82 +520,82 @@ void FModuleManager::RefreshModuleFilenameFromManifest(const FName InModuleName)
 	}
 }
 
-void* FModuleManager::InternalLoadLibrary(FName ModuleName, const FString& ModuleFileToLoad)
+void* FModuleManager::InternalLoadLibrary(FName ModuleName, const FString& LibraryToLoad)
 {
-	UE_LOG(LogModuleManager, Verbose, TEXT("InternalLoadLibrary: %s"), *ModuleName.ToString());
+	UE_LOG(LogModuleManager, Verbose, TEXT("InternalLoadLibrary: '%s'"), *ModuleName.ToString());
 
 	void* Handle = nullptr;
 
-#if UE_MERGED_MODULES
-
-	// First, attempt to find a cached library handle
-	FModuleManagerLibraryTracker& LoadedDynamicLibrary = LoadedDynamicLibraries.FindOrAdd(ModuleFileToLoad);
-	if (LoadedDynamicLibrary.Handle != nullptr)
-	{
-		Handle = LoadedDynamicLibrary.Handle;
-	}
-	else
-
-#endif // UE_MERGED_MODULES
-
-	// If no cached handle exists, then just load the library
-	{
 #if CPUPROFILERTRACE_ENABLED
-		UE_TRACE_LOG_SCOPED_T(Cpu, FPlatformProcess_GetDllHandle, CpuChannel)
-			<< FPlatformProcess_GetDllHandle.Name(*ModuleFileToLoad);
+	UE_TRACE_LOG_SCOPED_T(Cpu, FPlatformProcess_GetDllHandle, CpuChannel)
+		<< FPlatformProcess_GetDllHandle.Name(*LibraryToLoad);
 #endif // CPUPROFILERTRACE_ENABLED
 
-		Handle = FPlatformProcess::GetDllHandle(*ModuleFileToLoad);
-	}
+	const uint64_t InitialMemUsed = FPlatformMemory::GetMemoryUsedFast();
+	Handle = FPlatformProcess::GetDllHandle(*LibraryToLoad);
+	const uint64_t FinalMemUsed = FPlatformMemory::GetMemoryUsedFast();
 
 #if UE_MERGED_MODULES
 
-	// While using merged modules, update the usage data for the loaded library
-	LoadedDynamicLibrary.Users.AddUnique(ModuleName);
-	if (LoadedDynamicLibrary.Handle == nullptr)
-	{
-		LoadedDynamicLibrary.Handle = Handle;
-		UE_LOG(LogModuleManager, Verbose, TEXT("InternalLoadLibrary: cached library '%s'"), *ModuleFileToLoad);
-	}
-	else
-	{
-		UE_LOG(LogModuleManager, Verbose, TEXT("InternalLoadLibrary: added module to users of '%s'"), *ModuleFileToLoad);
-	}
+	// Stop any delay unload for that library, and match the handle for the code module to the library path.
+	DelayUnloadLibraries.Remove(LibraryToLoad);
+	LibraryHandles.Add(LibraryToLoad, Handle);
 
 #endif // UE_MERGED_MODULES
 
+	UE_LOG(LogModuleManager, Verbose, TEXT("InternalLoadLibrary: used about %llu KB"), (FinalMemUsed - InitialMemUsed) / 1024);
+
 	return Handle;
 }
 
 void FModuleManager::InternalFreeLibrary(FName ModuleName, void* Handle)
 {
-	UE_LOG(LogModuleManager, Verbose, TEXT("InternalFreeLibrary: %s"), *ModuleName.ToString());
+	UE_LOG(LogModuleManager, Verbose, TEXT("InternalFreeLibrary: '%s'"), *ModuleName.ToString());
 
 #if UE_MERGED_MODULES
 
-	// Find out if we should unload a dynamic library
-	const FString* LibraryToUnload = nullptr;
-	for (TPair<FString, FModuleManagerLibraryTracker>& LibraryNameAndTracker : LoadedDynamicLibraries)
+	// When using merged modular build, unloading follows a reference-counting and delay-unload process.
+	// We count how many modules are referencing each merged dynamic library, and decrement that value when one is freed.
+	// When no module uses the library anymore, the following process will happen:
+	//     - Mark the library for delay unload
+	//     - Notify the UObject system to remove live classes in every module of that library
+	//     - Wait for garbage collection to be run
+	//     - In FModuleManager::OnObjectCleanup, after GC, actually do the removal using FreeDllHandle
+
+	ModuleInfoRef ModuleInfo = Modules[ModuleName];
+	const FString LibraryToUnload = FPaths::ConvertRelativePathToFull(ModuleInfo->Filename);
+	const int32 ModuleCount = GetLibraryUsers(LibraryToUnload).Num();
+	UE_LOG(LogModuleManager, Verbose, TEXT("InternalFreeLibrary: library '%s' has %d users"), *LibraryToUnload, ModuleCount);
+
+	if (ModuleCount == 0)
 	{
-		FModuleManagerLibraryTracker& Tracker = LibraryNameAndTracker.Value;
+		UE_LOG(LogModuleManager, Log, TEXT("InternalFreeLibrary: preparing unload for library '%s'"), *LibraryToUnload);
+
+		TArray<FName> UnloadedModules;
 
-		if (Tracker.Handle == Handle)
+		// Remove statically linked initializer info, or the next module load for that library will crash on a dangling pointer
+		for (const TPair<FName, ModuleInfoRef>& OtherModuleNameAndInfo : Modules)
 		{
-			Tracker.Users.Remove(ModuleName);
-			if (Tracker.Users.Num() == 0)
+			const ModuleInfoRef& OtherModuleInfo = OtherModuleNameAndInfo.Value;
+			const FString OtherModuleLibraryName = FPaths::ConvertRelativePathToFull(OtherModuleInfo->Filename);
+
+			if (OtherModuleLibraryName == LibraryToUnload)
 			{
-				LibraryToUnload = &LibraryNameAndTracker.Key;
+				UE_LOG(LogModuleManager, Verbose, TEXT("InternalFreeLibrary: cleaning up for module '%s'"),
+					*OtherModuleNameAndInfo.Key.ToString());
+
+				StaticallyLinkedModuleInitializers.Remove(OtherModuleNameAndInfo.Key);
+
+				UnloadedModules.Add(OtherModuleNameAndInfo.Key);
 			}
-			break;
 		}
-	}
 
-	// Unload the DLL
-	if (LibraryToUnload)
-	{
-		UE_LOG(LogModuleManager, Verbose, TEXT("InternalFreeLibrary: unloading library '%s'"), **LibraryToUnload);
-		LoadedDynamicLibraries.Remove(*LibraryToUnload);
-		FPlatformProcess::FreeDllHandle(Handle);
+		// Mark for delay unload to run after GC is done.
+		DelayUnloadLibraries.Add(LibraryToUnload, UnloadedModules);
+
+		// If the module to unload has UObjects in it, then the UObject system need be informed.
+		// UObjects should be removed and garbage collection should run.
+		RemoveLoadedObjectsCallback.Broadcast(UnloadedModules);
 	}
 
 #else
@@ -562,8 +603,106 @@ void FModuleManager::InternalFreeLibrary(FName ModuleName, void* Handle)
 #endif // UE_MERGED_MODULES
 }
 
+#if UE_MERGED_MODULES
+
+TArray<FName> FModuleManager::GetLibraryUsers(const FString& ModuleFile) const
+{
+	TArray<FName> MatchingModules;
+
+	for (const TPair<FName, ModuleInfoRef>& ModuleNameAndInfo : Modules)
+	{
+		const ModuleInfoRef& ModuleInfo = ModuleNameAndInfo.Value;
+		if (ModuleInfo->Module.IsValid() && ModuleFile == FPaths::ConvertRelativePathToFull(ModuleInfo->Filename))
+		{
+			MatchingModules.Add(ModuleNameAndInfo.Key);
+		}
+	}
+
+	return MatchingModules;
+}
+
+#endif // UE_MERGED_MODULES
+
+#if UE_MERGED_MODULES && !UE_BUILD_SHIPPING
+
+void FModuleManager::LogMergedLibraryUsage()
+{
+	TMap<FString, TArray<FName>> LibraryUsageMap;
+
+	for (const TPair<FName, ModuleInfoRef>& ModuleNameAndInfo : Modules)
+	{
+		const ModuleInfoRef& ModuleInfo = ModuleNameAndInfo.Value;
+		if (ModuleInfo->Module.IsValid() && !ModuleInfo->Filename.IsEmpty())
+		{
+			const FString LibraryName = FPaths::ConvertRelativePathToFull(ModuleInfo->Filename);
+			TArray<FName>& UsingModules = LibraryUsageMap.FindOrAdd(LibraryName);
+			UsingModules.Add(ModuleNameAndInfo.Key);
+		}
+	}
+
+	for (TPair<FString, TArray<FName>>& LibraryAndModules : LibraryUsageMap)
+	{
+		UE_LOG(LogModuleManager, Log, TEXT("Modules loaded for library '%s' (%d)"), *LibraryAndModules.Key, LibraryAndModules.Value.Num());
+
+		for (const FName& Module : LibraryAndModules.Value)
+		{
+			UE_LOG(LogModuleManager, Log, TEXT("    '%s'"), *Module.ToString());
+		}
+	}
+}
+
+void FModuleManager::LoadAllModulesInMergedLibrary(FStringView LibraryName)
+{
+	for (const TPair<FName, ModuleInfoRef>& ModuleNameAndInfo : Modules)
+	{
+		const ModuleInfoRef& ModuleInfo = ModuleNameAndInfo.Value;
+		if (ModuleInfo->Filename.Contains(LibraryName))
+		{
+			LoadModule(ModuleNameAndInfo.Key);
+		}
+	}
+}
+
+void FModuleManager::UnloadAllModulesInMergedLibrary(FStringView LibraryName)
+{
+	for (const TPair<FName, ModuleInfoRef>& ModuleNameAndInfo : Modules)
+	{
+		const ModuleInfoRef& ModuleInfo = ModuleNameAndInfo.Value;
+		if (ModuleInfo->Module.IsValid() && ModuleInfo->Filename.Contains(LibraryName))
+		{
+			UnloadModule(ModuleNameAndInfo.Key, false, true);
+		}
+	}
+}
+
+#endif // UE_MERGED_MODULES && !UE_BUILD_SHIPPING
+
 #endif	// !IS_MONOLITHIC
 
+void FModuleManager::OnObjectCleanup()
+{
+#if UE_MERGED_MODULES
+
+	// See FModuleManager::InternalFreeLibrary for a detailed rundown of the library unload process.
+	for (auto LibraryAndModulesIterator = DelayUnloadLibraries.CreateIterator(); LibraryAndModulesIterator; ++LibraryAndModulesIterator)
+	{
+		// Check that no live object still exists for the modules in this merged library
+		if (ensure(CheckLiveObjectsInModulesCallback.IsBound()) && !CheckLiveObjectsInModulesCallback.Execute(LibraryAndModulesIterator.Value()))
+		{
+			void** HandlePtr = LibraryHandles.Find(LibraryAndModulesIterator.Key());
+			if (HandlePtr)
+			{
+				UE_LOG(LogModuleManager, Log, TEXT("OnObjectCleanup: unloading dynamic library '%s'"), *LibraryAndModulesIterator.Key());
+				FPlatformProcess::FreeDllHandle(*HandlePtr);
+			}
+
+			LibraryAndModulesIterator.RemoveCurrent();
+		}
+	}
+
+#endif // UE_MERGED_MODULES
+}
+
 IModuleInterface* FModuleManager::LoadModule(const FName InModuleName, ELoadModuleFlags InLoadModuleFlags)
 {
 	EModuleLoadResult FailureReason = EModuleLoadResult::Success;
@@ -635,12 +774,30 @@ IModuleInterface* FModuleManager::LoadModuleWithFailureReason(const FName InModu
 	// Do fast check for existing module, this is the most common case
 	ModuleInfoPtr FoundModulePtr = FindModule(InModuleName);
 
+	const auto EnsureModuleFilename = [this, InModuleName, &FoundModulePtr]()
+	{
+#if !IS_MONOLITHIC
+		ModuleInfoRef ModuleInfo = FoundModulePtr.ToSharedRef();
+		if (ModuleInfo->Filename.IsEmpty() || !FPaths::FileExists(ModuleInfo->Filename))
+		{
+			TMap<FName, FString> ModulePathMap;
+			FindModulePaths(*InModuleName.ToString(), ModulePathMap);
+			if (ModulePathMap.Num() == 1)
+			{
+				ModuleInfo->Filename = MoveTemp(TMap<FName, FString>::TIterator(ModulePathMap).Value());
+			}
+		}
+#endif // !IS_MONOLITHIC
+	};
+
 	if (FoundModulePtr.IsValid())
 	{
 		LoadedModule = FoundModulePtr->Module.Get();
 
 		if (LoadedModule)
 		{
+			EnsureModuleFilename();
+
 			// note: this function does not check (bIsReady || IsInGameThread()) the way GetModule() does
 			//   that looks like a bug if called from off-game-thread
 
@@ -648,6 +805,8 @@ IModuleInterface* FModuleManager::LoadModuleWithFailureReason(const FName InModu
 		}
 	}
 
+	UE_LOG(LogModuleManager, Verbose, TEXT("LoadModuleWithFailureReason %s"), *InModuleName.ToString());
+
 	// doing LoadModule off GameThread should not be done
 	// already warned above
 	// enable this ensure when we know it's okay
@@ -737,6 +896,7 @@ IModuleInterface* FModuleManager::LoadModuleWithFailureReason(const FName InModu
 			ModulesChangedEvent.Broadcast(InModuleName, EModuleChangeReason::ModuleLoaded);
 
 			// Set the return parameter
+			EnsureModuleFilename();
 			LoadedModule = ModuleInfo->Module.Get();
 		}
 		else
@@ -864,6 +1024,7 @@ IModuleInterface* FModuleManager::LoadModuleWithFailureReason(const FName InModu
 							ModulesChangedEvent.Broadcast(InModuleName, EModuleChangeReason::ModuleLoaded);
 
 							// Set the return parameter
+							EnsureModuleFilename();
 							LoadedModule = ModuleInfo->Module.Get();
 						}
 						else
@@ -918,6 +1079,8 @@ IModuleInterface* FModuleManager::LoadModuleWithFailureReason(const FName InModu
 
 bool FModuleManager::UnloadModule( const FName InModuleName, bool bIsShutdown, bool bAllowUnloadCode)
 {
+	UE_LOG(LogModuleManager, Verbose, TEXT("UnloadModule %s %d"), *InModuleName.ToString(), bAllowUnloadCode);
+
 	// Do we even know about this module?
 	ModuleInfoPtr ModuleInfoPtr = FindModule(InModuleName);
 	if( ModuleInfoPtr.IsValid() )
@@ -937,6 +1100,19 @@ bool FModuleManager::UnloadModule( const FName InModuleName, bool bIsShutdown, b
 			ModuleInfo.Module.Reset();
 
 #if !IS_MONOLITHIC
+
+#if UE_MERGED_MODULES
+			// Double-check that we don't have a valid handle for this...
+			if (ModuleInfo.Handle == nullptr)
+			{
+				void** HandlePtr = LibraryHandles.Find(ModuleInfo.Filename);
+				if (HandlePtr)
+				{
+					ModuleInfo.Handle = *HandlePtr;
+				}
+			}
+#endif // UE_MERGED_MODULES
+
 			if( ModuleInfo.Handle != nullptr )
 			{
 				// If we're shutting down then don't bother actually unloading the DLL.  We'll simply abandon it in memory
diff --git a/Engine/Source/Runtime/Core/Public/Modules/ModuleManager.h b/Engine/Source/Runtime/Core/Public/Modules/ModuleManager.h
--- a/Engine/Source/Runtime/Core/Public/Modules/ModuleManager.h
+++ b/Engine/Source/Runtime/Core/Public/Modules/ModuleManager.h
@@ -535,6 +535,26 @@ public:
 		return ProcessLoadedObjectsCallback;
 	}
 
+	/**
+	 * Gets a multicast delegate that is executed when any UObjects need destroying before modules are unloaded.
+	 *
+	 * @return The delegate.
+	 */
+	TMulticastDelegateRegistration<void(TConstArrayView<FName>)>& OnRemoveLoadedObjectsCallback()
+	{
+		return RemoveLoadedObjectsCallback;
+	}
+
+	/**
+	 * Gets a delegate that is executed whenever we need to confirm that modules can be safely unloaded.
+	 *
+	 * @return The delegate.
+	 */
+	TDelegateRegistration<bool(TConstArrayView<FName>)>& OnCheckLiveObjectsInModulesCallback()
+	{
+		return CheckLiveObjectsInModulesCallback;
+	}
+
 	/**
 	 * Gets a delegate that is executed when a module containing UObjects has been loaded.
 	 *
@@ -630,6 +650,22 @@ public:
 	/** Clears module path cache */
 	CORE_API void ResetModulePathsCache();
 
+#if UE_MERGED_MODULES && !UE_BUILD_SHIPPING
+
+	/** Dump all module usage for each loaded merged library */
+	void LogMergedLibraryUsage();
+
+	/** Load all modules from a merged library */
+	void LoadAllModulesInMergedLibrary(FStringView LibraryName);
+
+	/** Unload all modules from a merged library */
+	void UnloadAllModulesInMergedLibrary(FStringView LibraryName);
+
+#endif // UE_MERGED_MODULES && !UE_BUILD_SHIPPING
+
+	/** Signal that objects used by an unloading dynamic library should be cleaned up now */
+	CORE_API void OnObjectCleanup();
+
 	friend FArchive& operator<<( FArchive& Ar, FModuleManager& ModuleManager );
 private:
 	static CORE_API void WarnIfItWasntSafeToLoadHere(const FName InModuleName);
@@ -667,6 +703,12 @@ private:
 	/** Attempt to free the backing library for a module */
 	void InternalFreeLibrary(FName ModuleName, void* Handle);
 
+#if UE_MERGED_MODULES
+	/** Get all modules linking to a particular file */
+	TArray<FName> GetLibraryUsers(const FString& ModuleFile) const;
+#endif // UE_MERGED_MODULES
+
+
 #endif
 
 	/** Adds pending module initializer registrations to the StaticallyLinkedModuleInitializers map. */
@@ -676,24 +718,6 @@ private:
 	/** Map of all modules.  Maps the case-insensitive module name to information about that module, loaded or not. */
 	FModuleMap Modules;
 
-#if UE_MERGED_MODULES
-
-	struct FModuleManagerLibraryTracker
-	{
-		FModuleManagerLibraryTracker()
-			: Handle(nullptr)
-			, Users()
-		{}
-
-		void* Handle;
-		TArray<FName> Users;
-	};
-
-	/** Map of loaded DLL handles for when merged modular build is in use */
-	TMap<FString, FModuleManagerLibraryTracker> LoadedDynamicLibraries;
-
-#endif // UE_MERGED_MODULES
-
 	/** Pending registrations of module names */
 	/** We use an array here to stop comparisons (and thus FNames being constructed) when they are registered. */
 	/** Instead, we validate there are no duplicates when they're inserted into StaticallyLinkedModuleInitializers. */
@@ -718,6 +742,12 @@ private:
 	/** Multicast delegate called to process any new loaded objects. */
 	TMulticastDelegate<void(FName, bool)> ProcessLoadedObjectsCallback;
 
+	/** Multicast delegate called to remove any loaded objects for a list of modules. */
+	TMulticastDelegate<void(TConstArrayView<FName>)> RemoveLoadedObjectsCallback;
+
+	/** Delegate called to check whether any module to delay-unload still has live objects. */
+	TDelegate<bool(TConstArrayView<FName>)> CheckLiveObjectsInModulesCallback;
+
 	/** When module manager is linked against an application that supports UObjects, this delegate will be primed
 		at startup to provide information about whether a UObject package is loaded into memory. */
 	FIsPackageLoadedCallback IsPackageLoaded;
@@ -735,6 +765,16 @@ private:
 
 	/** Critical section object controlling R/W access to Modules. */
 	mutable FCriticalSection ModulesCriticalSection;
+
+#if UE_MERGED_MODULES
+
+	/** Path-to-handle map for libraries */
+	TMap<FString, void*> LibraryHandles;
+
+	/** Merged libraries that should be unloaded after the next object cleanup, with their associated modules */
+	TMap<FString, TArray<FName>> DelayUnloadLibraries;
+
+#endif // UE_MERGED_MODULES
 };
 
 /**
diff --git a/Engine/Source/Runtime/CoreUObject/Private/UObject/Obj.cpp b/Engine/Source/Runtime/CoreUObject/Private/UObject/Obj.cpp
--- a/Engine/Source/Runtime/CoreUObject/Private/UObject/Obj.cpp
+++ b/Engine/Source/Runtime/CoreUObject/Private/UObject/Obj.cpp
@@ -5461,8 +5461,8 @@ void InitUObject()
 	FCoreDelegates::OnShutdownAfterError.AddStatic(StaticShutdownAfterError);
 	FCoreDelegates::OnExit.AddStatic(StaticExit);
 
-	// Always register the UObjects callback for VNI and general consistency with the callbacks ProcessNewlyLoadedUObjects calls.
-	RegisterProcessNewlyLoadedUObjects();
+	// Always register the UObjects callbacks from the module manager for dynamic loading/unloading.
+	RegisterModularObjectsProcessing();
 
 	struct Local
 	{
diff --git a/Engine/Source/Runtime/CoreUObject/Private/UObject/UObjectAllocator.cpp b/Engine/Source/Runtime/CoreUObject/Private/UObject/UObjectAllocator.cpp
--- a/Engine/Source/Runtime/CoreUObject/Private/UObject/UObjectAllocator.cpp
+++ b/Engine/Source/Runtime/CoreUObject/Private/UObject/UObjectAllocator.cpp
@@ -10,7 +10,12 @@
 /** Global UObjectBase allocator							*/
 COREUOBJECT_API FUObjectAllocator GUObjectAllocator;
 
+// We cannot use the persistent allocator with merged modular builds, because every object can be removed eventually.
+#if UE_MERGED_MODULES
+static bool GPersistentAllocatorIsDisabled = true;
+#else
 static bool GPersistentAllocatorIsDisabled = false;
+#endif // UE_MERGED_MODULES
 
 void FUObjectAllocator::DisablePersistentAllocator()
 {
diff --git a/Engine/Source/Runtime/CoreUObject/Private/UObject/UObjectBase.cpp b/Engine/Source/Runtime/CoreUObject/Private/UObject/UObjectBase.cpp
--- a/Engine/Source/Runtime/CoreUObject/Private/UObject/UObjectBase.cpp
+++ b/Engine/Source/Runtime/CoreUObject/Private/UObject/UObjectBase.cpp
@@ -19,6 +19,7 @@
 #include "UObject/UObjectIterator.h"
 #include "UObject/UObjectStats.h"
 #include "UObject/Package.h"
+#include "UObject/ReferenceChainSearch.h"
 #include "Templates/Casts.h"
 #include "UObject/GCObject.h"
 #include "UObject/LinkerLoad.h"
@@ -107,6 +108,11 @@ static TMap<FName, TArray<FPendingRegistrant*>>& GetPerModuleBootstrapMap()
 
 #endif
 
+/** Must be called before modules will unload that contains UObject classes */
+static void RemoveLoadedUObjects(TConstArrayView<FName> InModuleNames);
+
+/** Can be called to confirm full removal of UObjects in a list of modules */
+static bool CheckLiveObjectsInModules(TConstArrayView<FName> InModuleNames);
 
 /**
  * Constructor used for bootstrapping
@@ -909,28 +915,33 @@ static void UObjectLoadAllCompiledInStructs()
 	StructRegistry.DoPendingOuterRegistrations(true);
 }
 
-void RegisterProcessNewlyLoadedUObjects()
+void RegisterModularObjectsProcessing()
 {
-	static bool bHasRegistered = false;
-	if (!bHasRegistered)
+	UE_CALL_ONCE([]()
 	{
-		bHasRegistered = true;
 		FModuleManager::Get().OnProcessLoadedObjectsCallback().AddStatic(ProcessNewlyLoadedUObjects);
-	}
+		FModuleManager::Get().OnRemoveLoadedObjectsCallback().AddStatic(RemoveLoadedUObjects);
+		FModuleManager::Get().OnCheckLiveObjectsInModulesCallback().BindStatic(CheckLiveObjectsInModules);
+
+		FCoreUObjectDelegates::GarbageCollectComplete.AddLambda([]()
+		{
+			FModuleManager::Get().OnObjectCleanup();
+		});
+	});
 }
 
-void ProcessNewlyLoadedUObjects(FName Package, bool bCanProcessNewlyLoadedObjects)
+void ProcessNewlyLoadedUObjects(FName InModuleName, bool bCanProcessNewlyLoadedObjects)
 {
 	SCOPED_BOOT_TIMING("ProcessNewlyLoadedUObjects");
 #if USE_PER_MODULE_UOBJECT_BOOTSTRAP
-	if (Package != NAME_None)
+	if (InModuleName != NAME_None)
 	{
-		UObjectReleaseModuleRegistrants(Package);
+		UObjectReleaseModuleRegistrants(InModuleName);
 	}
 #endif
 	if (!bCanProcessNewlyLoadedObjects)
 	{
-		FCoreUObjectDelegates::CompiledInUObjectsRegisteredDelegate.Broadcast(Package, ECompiledInUObjectsRegisteredStatus::Delayed);
+		FCoreUObjectDelegates::CompiledInUObjectsRegisteredDelegate.Broadcast(InModuleName, ECompiledInUObjectsRegisteredStatus::Delayed);
 		return;
 	}
 	LLM_SCOPE(ELLMTag::UObject);
@@ -958,12 +969,12 @@ void ProcessNewlyLoadedUObjects(FName Package, bool bCanProcessNewlyLoadedObject
 		UObjectProcessRegistrants();
 		UObjectLoadAllCompiledInStructs();
 
-		FCoreUObjectDelegates::CompiledInUObjectsRegisteredDelegate.Broadcast(Package, ECompiledInUObjectsRegisteredStatus::PreCDO);
+		FCoreUObjectDelegates::CompiledInUObjectsRegisteredDelegate.Broadcast(InModuleName, ECompiledInUObjectsRegisteredStatus::PreCDO);
 
 		UObjectLoadAllCompiledInDefaultProperties(AllNewClasses);
 	}
 
-	FCoreUObjectDelegates::CompiledInUObjectsRegisteredDelegate.Broadcast(Package, ECompiledInUObjectsRegisteredStatus::PostCDO);
+	FCoreUObjectDelegates::CompiledInUObjectsRegisteredDelegate.Broadcast(InModuleName, ECompiledInUObjectsRegisteredStatus::PostCDO);
 
 #if WITH_RELOAD
 	IReload* Reload = GetActiveReloadInterface();
@@ -1001,6 +1012,96 @@ void ProcessNewlyLoadedUObjects(FName Package, bool bCanProcessNewlyLoadedObject
 	}
 }
 
+static UPackage* GetScriptPackageFromModuleName (FName InModuleName )
+{
+	const FName PackagePath = *(TEXT("/Script/") + InModuleName .ToString());
+
+	UPackage* Package = FindObjectFast<UPackage>(/* Outer */ nullptr, PackagePath);
+
+	return Package;
+}
+
+static void RemoveObjectsInModules(TConstArrayView<FName> InModuleNames)
+{
+	LLM_SCOPE(ELLMTag::UObject);
+	DECLARE_SCOPE_CYCLE_COUNTER(TEXT("RemoveObjectsInModules"), STAT_RemoveObjectsInModules, STATGROUP_ObjectVerbose);
+
+	FPermanentObjectPoolExtents PermanentPool;
+
+	for (const FName PackageName : InModuleNames)
+	{
+		UPackage* Package = GetScriptPackageFromModuleName (PackageName);
+		if (Package)
+		{
+			UE_LOG(LogUObjectBase, Log, TEXT("RemoveObjectsInModules: removing package %s"), *Package->GetName());
+
+			ForEachObjectWithPackage(Package, [&PermanentPool](UObject* InObject) mutable -> bool
+			{
+				check(!PermanentPool.Contains(InObject));
+
+				UE_LOG(LogUObjectBase, Verbose, TEXT("RemoveObjectsInModules: marking %s as garbage"), *InObject->GetName());
+
+				InObject->ClearInternalFlags(EInternalObjectFlags_RootFlags);
+				InObject->ClearFlags(RF_Standalone);
+				InObject->MarkAsGarbage();
+
+				return true;
+			}, true);
+		}
+	}
+}
+
+static bool CheckLiveObjectsInModules(TConstArrayView<FName> InModuleNames)
+{
+	DECLARE_SCOPE_CYCLE_COUNTER(TEXT("CheckLiveObjectsInModules"), STAT_CheckLiveObjectsInModules, STATGROUP_ObjectVerbose);
+
+	bool bObjectsToRemoveStillLive = false;
+
+	for (const FName PackageName : InModuleNames)
+	{
+		UPackage* Package = GetScriptPackageFromModuleName (PackageName);
+		if (Package)
+		{
+			ForEachObjectWithPackage(Package, [&bObjectsToRemoveStillLive](UObject* InObject) -> bool
+			{
+				UE_LOG(LogUObjectBase, Warning, TEXT("CheckLiveObjectsInModules: object %s should have been removed (0x%0lX 0x%0lX)"),
+					*InObject->GetName(), InObject->GetFlags(), InObject->GetInternalFlags());
+
+				if (InObject->IsA<UClass>())
+				{
+					ForEachObjectOfClass(Cast<UClass>(InObject), [InObject](UObject* InClassObject) -> bool
+					{
+						UE_LOG(LogUObjectBase, Warning, TEXT("CheckLiveObjectsInModules: %s has live instance %s"),
+							*InObject->GetName(), *InClassObject->GetName());
+
+						FReferenceChainSearch::FindAndPrintStaleReferencesToObject(InClassObject, EPrintStaleReferencesOptions::Log);
+						return true;
+					});
+				}
+				else
+				{
+					FReferenceChainSearch::FindAndPrintStaleReferencesToObject(InObject, EPrintStaleReferencesOptions::Log);
+				}
+
+				bObjectsToRemoveStillLive = true;
+				return true;
+			}, true);
+		}
+	}
+
+	return bObjectsToRemoveStillLive;
+}
+
+static void RemoveLoadedUObjects(TConstArrayView<FName> InModuleNames)
+{
+	LLM_SCOPE(ELLMTag::UObject);
+	DECLARE_SCOPE_CYCLE_COUNTER(TEXT("RemoveLoadedUObjects"), STAT_RemoveLoadedUObjects, STATGROUP_ObjectVerbose);
+
+	FlushAsyncLoading();
+	RemoveObjectsInModules(InModuleNames);
+	CollectGarbage(GARBAGE_COLLECTION_KEEPFLAGS, true);
+}
+
 static int32 GVarMaxObjectsNotConsideredByGC;
 static FAutoConsoleVariableRef CMaxObjectsNotConsideredByGC(
 	TEXT("gc.MaxObjectsNotConsideredByGC"),
diff --git a/Engine/Source/Runtime/CoreUObject/Public/UObject/UObjectBase.h b/Engine/Source/Runtime/CoreUObject/Public/UObject/UObjectBase.h
--- a/Engine/Source/Runtime/CoreUObject/Public/UObject/UObjectBase.h
+++ b/Engine/Source/Runtime/CoreUObject/Public/UObject/UObjectBase.h
@@ -537,11 +537,11 @@ COREUOBJECT_API void RegisterCompiledInInfo(UPackage* (*InOuterRegister)(), cons
  */
 COREUOBJECT_API void RegisterCompiledInInfo(const TCHAR* PackageName, const FClassRegisterCompiledInInfo* ClassInfo, size_t NumClassInfo, const FStructRegisterCompiledInInfo* StructInfo, size_t NumStructInfo, const FEnumRegisterCompiledInInfo* EnumInfo, size_t NumEnumInfo);
 
-/** Must be called to register ProcessNewlyLoadedUObjects with the module manager */
-COREUOBJECT_API void RegisterProcessNewlyLoadedUObjects();
+/** Must be called to register UObject-related callbacks with the module manager */
+COREUOBJECT_API void RegisterModularObjectsProcessing();
 
 /** Must be called after a module has been loaded that contains UObject classes */
-COREUOBJECT_API void ProcessNewlyLoadedUObjects(FName Package = NAME_None, bool bCanProcessNewlyLoadedObjects = true);
+COREUOBJECT_API void ProcessNewlyLoadedUObjects(FName InModuleName = NAME_None, bool bCanProcessNewlyLoadedObjects = true);
 
 /**
  * Final phase of UObject initialization. all auto register objects are added to the main data structures.
diff --git a/Engine/Source/Runtime/Engine/Private/AssetManager.cpp b/Engine/Source/Runtime/Engine/Private/AssetManager.cpp
--- a/Engine/Source/Runtime/Engine/Private/AssetManager.cpp
+++ b/Engine/Source/Runtime/Engine/Private/AssetManager.cpp
@@ -1101,6 +1101,7 @@ void UAssetManager::RemovePrimaryAssetType(FPrimaryAssetType PrimaryAssetType)
 		TSharedPtr<FPrimaryAssetTypeData> TypeDataPtr = *TypeDataRef;
 		AssetTypeMap.Remove(PrimaryAssetType);
 		TypeDataPtr->ResetAssets(AssetPathMap);
+		bObjectReferenceListDirty = true;
 	}
 }
 
diff --git a/Engine/Source/Runtime/Launch/Private/LaunchEngineLoop.cpp b/Engine/Source/Runtime/Launch/Private/LaunchEngineLoop.cpp
--- a/Engine/Source/Runtime/Launch/Private/LaunchEngineLoop.cpp
+++ b/Engine/Source/Runtime/Launch/Private/LaunchEngineLoop.cpp
@@ -4470,8 +4470,8 @@ bool FEngineLoop::LoadCoreModules()
 {
 	// Always attempt to load CoreUObject. It requires additional pre-init which is called from its module's StartupModule method.
 #if WITH_COREUOBJECT
-	// Always register the UObjects callback for VNI and general consistency with the callbacks ProcessNewlyLoadedUObjects calls.
-	RegisterProcessNewlyLoadedUObjects();
+	// Always register the UObjects callbacks from the module manager for dynamic loading/unloading.
+	RegisterModularObjectsProcessing();
 	return FModuleManager::Get().LoadModule(TEXT("CoreUObject")) != nullptr;
 #else
 	return true;
diff --git a/Engine/Source/Runtime/Projects/Private/PluginManager.cpp b/Engine/Source/Runtime/Projects/Private/PluginManager.cpp
--- a/Engine/Source/Runtime/Projects/Private/PluginManager.cpp
+++ b/Engine/Source/Runtime/Projects/Private/PluginManager.cpp
@@ -1762,8 +1762,8 @@ bool FPluginManager::ConfigureEnabledPlugins()
 			// Mark all the plugins as enabled
 			SCOPED_BOOT_TIMING("MarkEnabledPlugins");
 
-
-#if !IS_MONOLITHIC
+			// Merged modular builds act as if they were monolithic as far as plugin discovery is concerned.
+#if !IS_MONOLITHIC && !UE_MERGED_MODULES
 			TMap<FString, FString> PluginToPhysicalFile;
 
 			// If EngineBinariesRootDir is set it means that we are loading binaries from somewhere else than our basedir.
@@ -1789,14 +1789,15 @@ bool FPluginManager::ConfigureEnabledPlugins()
 					PluginToPhysicalFile.Add(PluginName, FileName);
 				}
 			}
-#endif
+#endif // !IS_MONOLITHIC && !UE_MERGED_MODULES
 
 			for (TPair<FString, FPlugin*>& Pair : EnabledPlugins)
 			{
 				FPlugin& Plugin = *Pair.Value;
 
-#if !IS_MONOLITHIC
-				// Mount the binaries directory, and check the modules are valid
+#if !IS_MONOLITHIC && !UE_MERGED_MODULES
+				// Mount the binaries directory, and check the modules are valid.
+				// We do not run this process while using merged modular build, we'll keep acting as if the build was monolithic.
 				if (Plugin.Descriptor.Modules.Num() > 0)
 				{
 					const TCHAR* PluginFile = *Plugin.FileName;
@@ -1818,7 +1819,7 @@ bool FPluginManager::ConfigureEnabledPlugins()
 					UE_LOG(LogPluginManager, Display, TEXT("Skipping load of '%s'."), *Plugin.Name);
 					continue;
 				}
-#endif
+#endif // !IS_MONOLITHIC && !UE_MERGED_MODULES
 				if (!Plugin.GetDescriptor().bExplicitlyLoaded && Plugin.GetDescriptor().LocalizationTargets.Num() > 0)
 				{
 					PluginLocalizationUtils::GetLocalizationPathsForPlugin(Plugin, AdditionalLocResPaths);
@@ -2816,7 +2817,10 @@ bool FPluginManager::LoadModulesForEnabledPlugins( const ELoadingPhase::Type Loa
 			}
 		}
 
-		// This is a workaround for crashes when delaying the loading of binaries when using pak files/iostore. Should be removed
+#if !UE_MERGED_MODULES
+
+		// This is a workaround for crashes when delaying the loading of binaries when using pak files/iostore. Should be removed.
+		// This is avoided while using merged modular build, as we don't want modules to be loaded except when we're certain to need them.
 		if (bPreloadedBinaries && LoadingPhase == ELoadingPhase::PostConfigInit)
 		{
 			UE_SCOPED_ENGINE_ACTIVITY("Preloading all plugin binaries");
@@ -2836,6 +2840,8 @@ bool FPluginManager::LoadModulesForEnabledPlugins( const ELoadingPhase::Type Loa
 				}
 			}
 		}
+
+#endif // !UE_MERGED_MODULES
 	}
 	// Some phases such as ELoadingPhase::PreEarlyLoadingScreen are potentially called multiple times,
 	// but we do not return to an earlier phase after calling LoadModulesForEnabledPlugins on a later phase
