Add a testing API to force the AutoRTFM runtime state.

diff --git a/Engine/Source/Programs/AutoRTFMTests/Private/API.cpp b/Engine/Source/Programs/AutoRTFMTests/Private/API.cpp
--- a/Engine/Source/Programs/AutoRTFMTests/Private/API.cpp
+++ b/Engine/Source/Programs/AutoRTFMTests/Private/API.cpp
@@ -1122,6 +1122,8 @@ TEST_CASE("API.CascadingRetryTransaction")
 
 TEST_CASE("API.IsAutoRTFMRuntimeEnabled")
 {
+	AutoRTFM::Testing::FEnabledStateResetterScoped _(AutoRTFM::ForTheRuntime::AutoRTFM_EnabledByDefault);
+
 	// On entry to the test we'll be enabled-by-default, so first test we can set back to disabled by default.
 	REQUIRE(AutoRTFM::ForTheRuntime::IsAutoRTFMRuntimeEnabled());
 	REQUIRE(AutoRTFM::ForTheRuntime::SetAutoRTFMRuntime(AutoRTFM::ForTheRuntime::AutoRTFM_DisabledByDefault));
diff --git a/Engine/Source/Runtime/Core/Private/AutoRTFM/API.cpp b/Engine/Source/Runtime/Core/Private/AutoRTFM/API.cpp
--- a/Engine/Source/Runtime/Core/Private/AutoRTFM/API.cpp
+++ b/Engine/Source/Runtime/Core/Private/AutoRTFM/API.cpp
@@ -97,6 +97,19 @@ namespace
 	}
 }
 
+namespace AutoRTFM
+{
+	namespace Testing
+	{
+		UE_AUTORTFM_API ForTheRuntime::EAutoRTFMEnabledState ForceSetAutoRTFMRuntime(ForTheRuntime::EAutoRTFMEnabledState State)
+		{
+			const int Original = GAutoRTFMRuntimeEnabled;
+			GAutoRTFMRuntimeEnabled = State;
+			return static_cast<ForTheRuntime::EAutoRTFMEnabledState>(Original);
+		}
+	} // Testing
+} // AutoRTFM
+
 #if UE_AUTORTFM
 FString GAutoRTFM;
 
diff --git a/Engine/Source/Runtime/Core/Public/AutoRTFM/AutoRTFMTesting.h b/Engine/Source/Runtime/Core/Public/AutoRTFM/AutoRTFMTesting.h
--- a/Engine/Source/Runtime/Core/Public/AutoRTFM/AutoRTFMTesting.h
+++ b/Engine/Source/Runtime/Core/Public/AutoRTFM/AutoRTFMTesting.h
@@ -34,6 +34,18 @@ namespace Testing
 		const ETransactionResult Result = Transact(Functor);
 		UE_AUTORTFM_CHECK(ETransactionResult::Committed != Result);
 	}
+
+	// Force set the AutoRTFM runtime state. Used in testing. Returns the old value.
+	UE_AUTORTFM_API ForTheRuntime::EAutoRTFMEnabledState ForceSetAutoRTFMRuntime(ForTheRuntime::EAutoRTFMEnabledState State);
+
+	struct FEnabledStateResetterScoped final
+	{
+		FEnabledStateResetterScoped(ForTheRuntime::EAutoRTFMEnabledState State) : Original(ForceSetAutoRTFMRuntime(State)) {}
+		~FEnabledStateResetterScoped() { ForceSetAutoRTFMRuntime(Original); }
+
+	private:
+		const ForTheRuntime::EAutoRTFMEnabledState Original;
+	};
 } // namespace Testing
 
 } // namespace AutoRTFM
