Fix link issues with modular build

#rnx

diff --git a/Engine/Source/Runtime/Core/Public/Stats/StatIgnoreList.h b/Engine/Source/Runtime/Core/Public/Stats/StatIgnoreList.h
--- a/Engine/Source/Runtime/Core/Public/Stats/StatIgnoreList.h
+++ b/Engine/Source/Runtime/Core/Public/Stats/StatIgnoreList.h
@@ -49,7 +49,7 @@ namespace UE::Stats
 	 * @param StatNameHash: The 32-bit FNV1a hash of the stat name.
 	 * @param GroupNameHash: The 32-bit FNV1a hash of the group name this stat is part of, or 0 to skip checking the group.
 	 */
-	bool IsStatOrGroupIgnoredOnCurrentThread(uint32 StatNameHash, uint32 GroupNameHash = 0);
+	bool CORE_API IsStatOrGroupIgnoredOnCurrentThread(uint32 StatNameHash, uint32 GroupNameHash = 0);
 
 } // namespace UE::Stats
 #endif // UE_STATS_ALLOW_PER_THREAD_IGNORELIST
