Automation Tests: Add support for silently suppressing errors

diff --git a/Engine/Source/Runtime/Core/Private/Misc/AutomationTest.cpp b/Engine/Source/Runtime/Core/Private/Misc/AutomationTest.cpp
--- a/Engine/Source/Runtime/Core/Private/Misc/AutomationTest.cpp
+++ b/Engine/Source/Runtime/Core/Private/Misc/AutomationTest.cpp
@@ -1540,8 +1540,7 @@ bool FAutomationTestBase::HasMetExpectedMessages(ELogVerbosity::Type VerbosityTy
 		// Avoid ambiguity of the messages below when the verbosity is "All"
 		const TCHAR* LogVerbosityString = ExpectedMessage.Verbosity == ELogVerbosity::All ? TEXT("Any") : ToString(ExpectedMessage.Verbosity);
 
-		const bool bExpectsOneOrMore = ExpectedMessage.ExpectedNumberOfOccurrences == 0;
-		if (!bExpectsOneOrMore && (ExpectedMessage.ExpectedNumberOfOccurrences != ExpectedMessage.ActualNumberOfOccurrences))
+		if (ExpectedMessage.ExpectedNumberOfOccurrences > 0 && (ExpectedMessage.ExpectedNumberOfOccurrences != ExpectedMessage.ActualNumberOfOccurrences))
 		{
 			FWriteScopeLock WLock(ActionCS);
 			bHasMetAllExpectedMessages = false;
@@ -1555,7 +1554,7 @@ bool FAutomationTestBase::HasMetExpectedMessages(ELogVerbosity::Type VerbosityTy
 					, ExpectedMessage.ActualNumberOfOccurrences)
 				, ExecutionInfo.GetContext()));
 		}
-		else if (bExpectsOneOrMore)
+		else if (ExpectedMessage.ExpectedNumberOfOccurrences == 0)
 		{
 			FWriteScopeLock WLock(ActionCS);
 			if (ExpectedMessage.ActualNumberOfOccurrences == 0)
@@ -1634,15 +1633,8 @@ void FAutomationTestBase::AddExpectedMessage(
 	int32 Occurrences,
 	bool IsRegex)
 {
-	if (Occurrences >= 0)
-	{
-		FWriteScopeLock Lock(ActionCS);
-		ExpectedMessages.Add(FAutomationExpectedMessage(ExpectedPatternString, ExpectedVerbosity, CompareType, Occurrences, IsRegex));
-	}
-	else
-	{
-		UE_LOG(LogAutomationTest, Error, TEXT("Adding expected log message matching '%s' failed: number of expected occurrences must be >= 0"), *ExpectedPatternString);
-	}
+	FWriteScopeLock Lock(ActionCS);
+	ExpectedMessages.Add(FAutomationExpectedMessage(ExpectedPatternString, ExpectedVerbosity, CompareType, Occurrences, IsRegex));
 }
 
 void FAutomationTestBase::AddExpectedMessage(
diff --git a/Engine/Source/Runtime/Core/Public/Misc/AutomationTest.h b/Engine/Source/Runtime/Core/Public/Misc/AutomationTest.h
--- a/Engine/Source/Runtime/Core/Public/Misc/AutomationTest.h
+++ b/Engine/Source/Runtime/Core/Public/Misc/AutomationTest.h
@@ -628,7 +628,8 @@ struct FAutomationExpectedMessage
 	EAutomationExpectedMessageFlags::MatchType CompareType;
 	/** 
 	 * Number of occurrences expected for message. If set greater than 0, it will cause the test to fail if the
-	 * exact number of occurrences expected is not matched. If set to 0, it will suppress all matching messages. 
+	 * exact number of occurrences expected is not matched. If set to 0, it will cause the test to fail if
+	 * there were no occurrences at all. If negative, it will suppress all matching messages.
 	 */
 	int32 ExpectedNumberOfOccurrences;
 	int32 ActualNumberOfOccurrences;
@@ -1737,6 +1738,7 @@ public:
 	* @param CompareType - How to match this string with an encountered message, should it match exactly or simply just contain the string.
 	* @param Occurrences - How many times to expect this message string to be seen. If > 0, the message must be seen the exact number of times
 	* specified or the test will fail. If == 0, the message must be seen one or more times (with no upper limit) or the test will fail.
+	* If < 0, occurrences of this message will be silently ignored.
 	* @param IsRegex - If the pattern is to be used as regex or plain string. Default is true.
 	*/
 	CORE_API void AddExpectedMessage(FString ExpectedPatternString, ELogVerbosity::Type ExpectedVerbosity, EAutomationExpectedMessageFlags::MatchType CompareType = EAutomationExpectedMessageFlags::Contains, int32 Occurrences = 1, bool IsRegex = true);
@@ -1751,6 +1753,7 @@ public:
 	* @param CompareType - How to match this string with an encountered message, should it match exactly or simply just contain the string.
 	* @param Occurrences - How many times to expect this message string to be seen. If > 0, the message must be seen the exact number of times
 	* specified or the test will fail. If == 0, the message must be seen one or more times (with no upper limit) or the test will fail.
+	* If < 0, occurrences of this message will be silently ignored.
 	*/
 	CORE_API void AddExpectedMessagePlain(FString ExpectedString, ELogVerbosity::Type ExpectedVerbosity, EAutomationExpectedMessageFlags::MatchType CompareType = EAutomationExpectedMessageFlags::Contains, int32 Occurrences = 1);
 
@@ -1762,6 +1765,7 @@ public:
 	* @param CompareType - How to match this string with an encountered message, should it match exactly or simply just contain the string.
 	* @param Occurrences - How many times to expect this message string to be seen. If > 0, the message must be seen the exact number of times
 	* specified or the test will fail. If == 0, the message must be seen one or more times (with no upper limit) or the test will fail.
+	* If < 0, occurrences of this message will be silently ignored.
 	* @param IsRegex - If the pattern is to be used as regex or plain string. Default is true.
 	*/
 	CORE_API void AddExpectedMessage(FString ExpectedPatternString, EAutomationExpectedMessageFlags::MatchType CompareType = EAutomationExpectedMessageFlags::Contains, int32 Occurrences = 1, bool IsRegex = true);
@@ -1774,6 +1778,7 @@ public:
 	* @param CompareType - How to match this string with an encountered message, should it match exactly or simply just contain the string.
 	* @param Occurrences - How many times to expect this message string to be seen. If > 0, the message must be seen the exact number of times
 	* specified or the test will fail. If == 0, the message must be seen one or more times (with no upper limit) or the test will fail.
+	* If < 0, occurrences of this message will be silently ignored.
 	*/
 	CORE_API void AddExpectedMessagePlain(FString ExpectedString, EAutomationExpectedMessageFlags::MatchType CompareType = EAutomationExpectedMessageFlags::Contains, int32 Occurrences = 1);
 
@@ -1794,6 +1799,7 @@ public:
 	* @param CompareType - How to match this string with an encountered error, should it match exactly or simply just contain the string.
 	* @param Occurrences - How many times to expect this error string to be seen. If > 0, the error must be seen the exact number of times
 	* specified or the test will fail. If == 0, the error must be seen one or more times (with no upper limit) or the test will fail.
+	* If < 0, occurrences of this message will be silently ignored.
 	* @param IsRegex - If the pattern is to be used as regex or plain string. Default is true.
 	*/
 	CORE_API void AddExpectedError(FString ExpectedPatternString, EAutomationExpectedErrorFlags::MatchType CompareType = EAutomationExpectedErrorFlags::Contains, int32 Occurrences = 1, bool IsRegex = true);
@@ -1806,6 +1812,7 @@ public:
 	* @param CompareType - How to match this string with an encountered error, should it match exactly or simply just contain the string.
 	* @param Occurrences - How many times to expect this error string to be seen. If > 0, the error must be seen the exact number of times
 	* specified or the test will fail. If == 0, the error must be seen one or more times (with no upper limit) or the test will fail.
+	* If < 0, occurrences of this message will be silently ignored.
 	*/
 	CORE_API void AddExpectedErrorPlain(FString ExpectedString, EAutomationExpectedErrorFlags::MatchType CompareType = EAutomationExpectedErrorFlags::Contains, int32 Occurrences = 1);
 
