[EventCount]
- Fix WaitFor and WaitUntil that would return a wrong value to the caller when the value changed before actually waiting.

diff --git a/Engine/Source/Runtime/Core/Public/Async/EventCount.h b/Engine/Source/Runtime/Core/Public/Async/EventCount.h
--- a/Engine/Source/Runtime/Core/Public/Async/EventCount.h
+++ b/Engine/Source/Runtime/Core/Public/Async/EventCount.h
@@ -106,7 +106,10 @@ public:
 	{
 		if ((Count.load(std::memory_order_acquire) & ~CounterType(1)) == Compare.Value)
 		{
-			return ParkingLot::WaitFor(&Count, [this, Compare] { return (Count.load(std::memory_order_acquire) & ~CounterType(1)) == Compare.Value; }, []{}, WaitTime).bDidWake;
+			const ParkingLot::FWaitState WaitState = ParkingLot::WaitFor(&Count, [this, Compare] { return (Count.load(std::memory_order_acquire) & ~CounterType(1)) == Compare.Value; }, [] {}, WaitTime);
+
+			// Make sure we return true if we did wait but also if the wait was skipped because the value actually changed before we had the opportunity to wait.
+			return WaitState.bDidWake || !WaitState.bDidWait;
 		}
 		return true;
 	}
@@ -122,7 +125,10 @@ public:
 	{
 		if ((Count.load(std::memory_order_acquire) & ~CounterType(1)) == Compare.Value)
 		{
-			return ParkingLot::WaitUntil(&Count, [this, Compare] { return (Count.load(std::memory_order_acquire) & ~CounterType(1)) == Compare.Value; }, []{}, WaitTime).bDidWake;
+			const ParkingLot::FWaitState WaitState = ParkingLot::WaitUntil(&Count, [this, Compare] { return (Count.load(std::memory_order_acquire) & ~CounterType(1)) == Compare.Value; }, [] {}, WaitTime);
+
+			// Make sure we return true if we did wait but also if the wait was skipped because the value actually changed before we had the opportunity to wait.
+			return WaitState.bDidWake || !WaitState.bDidWait;
 		}
 		return true;
 	}
