[AutoRTFM] Make FThreadSafeRefCountedObject class transactionally safe.

This class is used less often than TRefCountingMixin<T> because the mixin form is generally just as good or better on every axis.  However, it's still used in some places which can be reached by Verse, such as in the landscape collision code used by [REDACTED_PERSON]. (`struct FHeightfieldGeometryRef : public FThreadSafeRefCountedObject` in class `ULandscapeHeightfieldCollisionComponent`)

To avoid making a third copy of our transactionally-safe refcounting logic (since it was in two classes already), it's now been broken out into a UE::Private helper class which maintains the refcount.

This change necessitated a refactor of TRefCountingMixin. Previously, this class handled the Mode thread-safety template parameter via typedefs and if constexpr. Now, the thread-safe and thread-unsafe versions are entirely distinct specializations. There was almost no shared code anyway so this is perhaps easier to follow.

diff --git a/Engine/Source/Programs/AutoRTFMTests/Private/RefCountPtrTests.cpp b/Engine/Source/Programs/AutoRTFMTests/Private/RefCountPtrTests.cpp
--- a/Engine/Source/Programs/AutoRTFMTests/Private/RefCountPtrTests.cpp
+++ b/Engine/Source/Programs/AutoRTFMTests/Private/RefCountPtrTests.cpp
@@ -17,9 +17,15 @@ namespace
 		RefCountedMixinInt(int InValue) : Value(InValue) {}
 		int Value;
 	};
+
+	struct ThreadSafeRefCountedInt : public FThreadSafeRefCountedObject
+	{
+		ThreadSafeRefCountedInt(int InValue) : Value(InValue) {}
+		int Value;
+	};
 }
 
-TEMPLATE_TEST_CASE("PreviouslyAllocated", "[RefCountPtr]", RefCountedInt, RefCountedMixinInt)
+TEMPLATE_TEST_CASE("PreviouslyAllocated", "[RefCountPtr]", RefCountedInt, RefCountedMixinInt, ThreadSafeRefCountedInt)
 {
 	SECTION("Using new T()")
 	{
@@ -52,7 +58,7 @@ TEMPLATE_TEST_CASE("PreviouslyAllocated", "[RefCountPtr]", RefCountedInt, RefCou
 	}
 }
 
-TEMPLATE_TEST_CASE("AbortWithPreviouslyAllocated", "[RefCountPtr]", RefCountedInt, RefCountedMixinInt)
+TEMPLATE_TEST_CASE("AbortWithPreviouslyAllocated", "[RefCountPtr]", RefCountedInt, RefCountedMixinInt, ThreadSafeRefCountedInt)
 {
 	SECTION("Using new T()")
 	{
@@ -89,7 +95,7 @@ TEMPLATE_TEST_CASE("AbortWithPreviouslyAllocated", "[RefCountPtr]", RefCountedIn
 	}
 }
 
-TEMPLATE_TEST_CASE("AbortWithNewlyAllocated", "[RefCountPtr]", RefCountedInt, RefCountedMixinInt)
+TEMPLATE_TEST_CASE("AbortWithNewlyAllocated", "[RefCountPtr]", RefCountedInt, RefCountedMixinInt, ThreadSafeRefCountedInt)
 {
 	int Result = 42;
 
@@ -138,7 +144,7 @@ TEMPLATE_TEST_CASE("AbortWithNewlyAllocated", "[RefCountPtr]", RefCountedInt, Re
 	REQUIRE(Result == 42);
 }
 
-TEMPLATE_TEST_CASE("OnCommitCapturingRefCountPtr", "[RefCountPtr]", RefCountedInt, RefCountedMixinInt)
+TEMPLATE_TEST_CASE("OnCommitCapturingRefCountPtr", "[RefCountPtr]", RefCountedInt, RefCountedMixinInt, ThreadSafeRefCountedInt)
 {
 	TRefCountPtr<TestType> Foo = MakeRefCount<TestType>(42);
 
@@ -170,7 +176,7 @@ TEMPLATE_TEST_CASE("OnCommitCapturingRefCountPtr", "[RefCountPtr]", RefCountedIn
 	}
 }
 
-TEMPLATE_TEST_CASE("OnAbortCapturingRefCountPtr", "[RefCountPtr]", RefCountedInt, RefCountedMixinInt)
+TEMPLATE_TEST_CASE("OnAbortCapturingRefCountPtr", "[RefCountPtr]", RefCountedInt, RefCountedMixinInt, ThreadSafeRefCountedInt)
 {
 	TRefCountPtr<TestType> Foo = MakeRefCount<TestType>(42);
 
diff --git a/Engine/Source/Runtime/Core/Private/Templates/RefCounting.cpp b/Engine/Source/Runtime/Core/Private/Templates/RefCounting.cpp
--- a/Engine/Source/Runtime/Core/Private/Templates/RefCounting.cpp
+++ b/Engine/Source/Runtime/Core/Private/Templates/RefCounting.cpp
@@ -2,7 +2,7 @@
 
 #include "Templates/RefCounting.h"
 
-void FRefCountBase::CheckRefCountIsNonZero() const
+void UE::Private::CheckRefCountIsNonZero()
 {
 	check(!"Release() was called on an object which is already at a zero refcount.");
 }
\ No newline at end of file
diff --git a/Engine/Source/Runtime/Core/Public/Templates/RefCounting.h b/Engine/Source/Runtime/Core/Public/Templates/RefCounting.h
--- a/Engine/Source/Runtime/Core/Public/Templates/RefCounting.h
+++ b/Engine/Source/Runtime/Core/Public/Templates/RefCounting.h
@@ -56,39 +56,27 @@ private:
 	uint32 RefCount = 0;
 };
 
-/** A virtual interface for ref counted objects to implement. */
-class IRefCountedObject
-{
-public:
-	virtual ~IRefCountedObject() { }
-	virtual FReturnedRefCountValue AddRef() const = 0;
-
-	// TODO (SOL-7350): return FReturnedRefCountValue from Release(); clean up call sites
-	// which rely on its return value.
-	virtual uint32 Release() const = 0;
-
-	// TODO (SOL-7350): mark this function as deprecated; clean up existing callers.
-	virtual uint32 GetRefCount() const = 0;
-};
-
 /**
- * Base class implementing thread-safe reference counting.
+ * UE::Private::TTransactionalAtomicRefCount manages a transactionally-safe atomic refcount value.
+ * This is used by FRefCountBase, FThreadSafeRefCountedObject and TRefCountingMixin (in thread-safe mode).
  */
-class FRefCountBase
+namespace UE::Private
 {
-public:
-	FRefCountBase() = default;
-	virtual ~FRefCountBase() = default;
 
-	FRefCountBase(const FRefCountBase& Rhs) = delete;
-	FRefCountBase& operator=(const FRefCountBase& Rhs) = delete;
+CORE_API void CheckRefCountIsNonZero();
 
-	FReturnedRefCountValue AddRef() const
+template <typename AtomicType>
+class TTransactionalAtomicRefCount
+{
+public:
+	template <auto DeleteFn>
+	AtomicType AddRef() const
 	{
-		uint32 Refs;
+		AtomicType Refs;
+
 		UE_AUTORTFM_OPEN
 		{
-			Refs = NumRefs++;
+			Refs = RefCount++;
 		};
 
 		// If we are inside a transaction, a rollback must undo our refcount bump.
@@ -103,7 +91,7 @@ public:
 		{
 			UE_AUTORTFM_ONABORT(this)
 			{
-				--NumRefs;
+				--RefCount;
 				// The refcount is likely zero now, but leaving the object alive isn't a leak.
 				// We are restoring the object back to its initial "just-created" state.
 			};
@@ -112,23 +100,24 @@ public:
 		{
 			UE_AUTORTFM_ONABORT(this)
 			{
-				Release();
+				Release<DeleteFn>();
 			};
 		}
 
-		return FReturnedRefCountValue{Refs + 1};
+		return Refs + 1;
 	}
 
-	uint32 Release() const
+	template <auto DeleteFn>
+	AtomicType Release() const
 	{
 #if DO_GUARD_SLOW
-		if (NumRefs == 0)
+		if (RefCount == 0)
 		{
 			CheckRefCountIsNonZero();
 		}
 #endif
 
-		uint32 RefsToReturn;
+		AtomicType RefsToReturn;
 
 		if (AutoRTFM::IsClosed())
 		{
@@ -139,58 +128,118 @@ public:
 			// always returns the true refcount.
 			UE_AUTORTFM_OPEN
 			{
-				RefsToReturn = uint32(NumRefs.load(std::memory_order_relaxed)) - 1;
+				RefsToReturn = AtomicType(RefCount.load(std::memory_order_relaxed)) - 1;
 			};
 
 			// Refcount changes and frees are deferred until the transaction is concluded.
 			UE_AUTORTFM_ONCOMMIT(this)
 			{
-				ImmediatelyRelease();
+				ImmediatelyRelease<DeleteFn>();
 			};
 		}
 		else
 		{
-			RefsToReturn = ImmediatelyRelease() - 1;
+			RefsToReturn = ImmediatelyRelease<DeleteFn>();
 		}
 
 		return RefsToReturn;
 	}
 
-	uint32 GetRefCount() const
+	AtomicType GetRefCount() const
 	{
+		// This is equivalent to https://en.cppreference.com/w/cpp/memory/shared_ptr/use_count
+		// 
 		// Inside of an AutoRTFM transaction, the returned refcount value may be higher than you'd
 		// expect, because all Release calls are deferred until the transaction commit time.
-		uint32 Refs;
+		AtomicType Refs;
 		UE_AUTORTFM_OPEN
 		{
-			Refs = uint32(NumRefs.load(std::memory_order_relaxed));
+			// A 'live' reference count is unstable by nature and so there's no benefit
+			// to try and enforce memory ordering around the reading of it.
+			Refs = RefCount.load(std::memory_order_relaxed);
 		};
 		return Refs;
 	}
 
 private:
-	uint32 ImmediatelyRelease() const
+	template <auto DeleteFn>
+	AtomicType ImmediatelyRelease() const
 	{
-		// This decrements the reference count and immediately frees the object if it has become zero.
-		// A check-failure is issued if an object is over-released.
-		// Returns the refcount _before_ it was decremented.
-		int32 RefsBeforeRelease = NumRefs--;
+		// fetch_sub returns the refcount _before_ it was decremented. std::memory_order_acq_rel is
+		// used so that, if we do end up executing the destructor, it's not possible for side effects 
+		// from executing the destructor to end up being visible before we've determined that the 
+		// reference count is actually zero.
+		AtomicType RefsBeforeRelease = RefCount.fetch_sub(1, std::memory_order_acq_rel);
+
 #if DO_GUARD_SLOW
+		// A check-failure is issued if an object is over-released.
 		if (RefsBeforeRelease == 0)
 		{
 			CheckRefCountIsNonZero();
 		}
 #endif
+		// We immediately free the object if its refcount has become zero.
 		if (RefsBeforeRelease == 1)
 		{
-			delete this;
+			DeleteFn(this);
 		}
-		return (uint32)RefsBeforeRelease;
+		return RefsBeforeRelease;
+	}
+
+	mutable std::atomic<AtomicType> RefCount = 0;
+};
+
+}
+
+/** A virtual interface for ref counted objects to implement. */
+class IRefCountedObject
+{
+public:
+	virtual ~IRefCountedObject() { }
+	virtual FReturnedRefCountValue AddRef() const = 0;
+
+	// TODO (SOL-7350): return FReturnedRefCountValue from Release(); clean up call sites
+	// which rely on its return value.
+	virtual uint32 Release() const = 0;
+
+	// TODO (SOL-7350): mark this function as deprecated; clean up existing callers.
+	virtual uint32 GetRefCount() const = 0;
+};
+
+/**
+ * Base class implementing thread-safe reference counting.
+ */
+class FRefCountBase : UE::Private::TTransactionalAtomicRefCount<uint32>
+{
+public:
+	FRefCountBase() = default;
+	virtual ~FRefCountBase() = default;
+
+	FRefCountBase(const FRefCountBase& Rhs) = delete;
+	FRefCountBase& operator=(const FRefCountBase& Rhs) = delete;
+
+	FReturnedRefCountValue AddRef() const
+	{
+		return FReturnedRefCountValue{Super::AddRef<DeleteThis>()};
 	}
 
-	mutable std::atomic<int32> NumRefs = 0;
+	uint32 Release() const
+	{
+		return Super::Release<DeleteThis>();
+	}
 
-	CORE_API void CheckRefCountIsNonZero() const;
+	uint32 GetRefCount() const
+	{
+		return Super::GetRefCount();
+	}
+
+private:
+	using Super = UE::Private::TTransactionalAtomicRefCount<uint32>;
+
+	static void DeleteThis(const Super* This)
+	{
+		delete static_cast<const FRefCountBase*>(This);
+	}
 };
 
 /**
@@ -198,7 +247,6 @@ private:
  *
  * This class should not be used for new code as it does not use atomic operations to update 
  * the reference count.
- *
  */
 class FRefCountedObject
 {
@@ -229,46 +277,55 @@ private:
 };
 
 /**
- * Like FRefCountedObject, but internal ref count is thread safe
+ * Like FRefCountedObject, but the reference count is thread-safe.
  */
-class FThreadSafeRefCountedObject
+class FThreadSafeRefCountedObject : UE::Private::TTransactionalAtomicRefCount<uint32>
 {
 public:
-	FThreadSafeRefCountedObject() : NumRefs(0) {}
-	virtual ~FThreadSafeRefCountedObject() { check(NumRefs.GetValue() == 0); }
+	FThreadSafeRefCountedObject() = default;
+
 	FThreadSafeRefCountedObject(const FThreadSafeRefCountedObject& Rhs) = delete;
 	FThreadSafeRefCountedObject& operator=(const FThreadSafeRefCountedObject& Rhs) = delete;
+
+	virtual ~FThreadSafeRefCountedObject()
+	{ 
+		check(Super::GetRefCount() == 0); 
+	}
+	
 	FReturnedRefCountValue AddRef() const
 	{
-		return FReturnedRefCountValue{uint32(NumRefs.Increment())};
+		return FReturnedRefCountValue{Super::AddRef<DeleteThis>()};
 	}
+
 	uint32 Release() const
 	{
-		uint32 Refs = uint32(NumRefs.Decrement());
-		if (Refs == 0)
-		{
-			delete this;
-		}
-		return Refs;
+		return Super::Release<DeleteThis>();
 	}
+
 	uint32 GetRefCount() const
 	{
-		return uint32(NumRefs.GetValue());
+		return Super::GetRefCount();
 	}
+
 private:
-	mutable FThreadSafeCounter NumRefs;
+	using Super = UE::Private::TTransactionalAtomicRefCount<uint32>;
+
+	static void DeleteThis(const Super* This)
+	{
+		delete static_cast<const FThreadSafeRefCountedObject*>(This);
+	}
 };
 
 /**
  * ERefCountingMode is used select between either 'fast' or 'thread safe' ref-counting types.
- * This is only used by templates at compile time to generate one code path or another.
+ * This is only used at compile time to select between template specializations.
  */
 enum class ERefCountingMode : uint8
 {
 	/** Forced to be not thread-safe. */
 	NotThreadSafe = 0,
 
-	/** Thread-safe, never spin locks, but slower */
+	/** Thread-safe: never spin locks, but slower. */
 	ThreadSafe = 1
 };
 
@@ -276,8 +333,8 @@ enum class ERefCountingMode : uint8
  * Ref-counting mixin, designed to add ref-counting to an object without requiring a virtual destructor.
  * Implements support for AutoRTFM, is thread-safe by default, and can support custom deleters via T::StaticDestroyObject.
  * 
- * @note AutoRTFM means that the return value of AddRef/Release is nonsense (as the ref-count doesn't change until the
- *       transaction is committed), but this is fine for use with TRefCountPtr (as it doesn't use those return values).
+ * @note AutoRTFM means that the return value of AddRef/Release can't be trusted (as the ref-count doesn't decrement until
+ *       the transaction is committed), but this is fine for use with TRefCountPtr, as it doesn't use those return values.
  * 
  * Basic Example:
  *  struct FMyRefCountedObject : public TRefCountingMixin<FMyRefCountedObject>
@@ -295,10 +352,14 @@ enum class ERefCountingMode : uint8
  *  };
  */
 template <typename T, ERefCountingMode Mode = ERefCountingMode::ThreadSafe>
-class TRefCountingMixin
-{
-	using RefCountType = std::conditional_t<Mode == ERefCountingMode::ThreadSafe, std::atomic<uint32>, uint32>;
+class TRefCountingMixin;
 
+/**
+ * Thread-safe specialization
+ */
+template <typename T>
+class TRefCountingMixin<T, ERefCountingMode::ThreadSafe> : UE::Private::TTransactionalAtomicRefCount<uint32>
+{
 public:
 	TRefCountingMixin() = default;
 
@@ -307,73 +368,62 @@ public:
 
 	FReturnedRefCountValue AddRef() const
 	{
-		uint32 Refs;
+		return FReturnedRefCountValue{Super::template AddRef<StaticDestroyMixin>()};
+	}
 
-		if constexpr (Mode == ERefCountingMode::ThreadSafe)
-		{
-			UE_AUTORTFM_OPEN
-			{
-				Refs = RefCount++;
-			};
+	uint32 Release() const
+	{
+		Super::template Release<StaticDestroyMixin>();
 
-			// If we are inside a transaction, a rollback must undo our refcount bump.
-			// In general, this is best handled by Release(). However, there is one case 
-			// that needs to be handled with special care. A brand-new object has a
-			// refcount of zero, and a rollback must return it to this zero-refcount state.
-			// However, calling AddRef() followed by Release() would not accomplish this;
-			// instead, it would free the object entirely! We need to guard against this,
-			// since it could lead to a double-free, so we detect the zero-reference case
-			// and special-case it.
-			if (Refs == 0)
-			{
-				UE_AUTORTFM_ONABORT(this)
-				{
-					--RefCount;
-					// The refcount is likely zero now, but leaving the object alive isn't a leak.
-					// We are restoring the object back to its initial "just-created" state.
-				};
-			}
-			else
-			{
-				UE_AUTORTFM_ONABORT(this)
-				{
-					Release();
-				};
-			}
-		}
-		else
-		{
-			Refs = RefCount++;
-		}
+		// Note: TRefCountPtr doesn't use the return value
+		return 0;
+	}
+
+	uint32 GetRefCount() const
+	{
+		return Super::GetRefCount();
+	}
+
+	static void StaticDestroyObject(const T* Obj)
+	{
+		delete Obj;
+	}
 
-		return FReturnedRefCountValue{Refs + 1};
+private:
+	using Super = UE::Private::TTransactionalAtomicRefCount<uint32>;
+
+	static void StaticDestroyMixin(const Super* This)
+	{
+		// This static_cast is traversing two levels of the class hierarchy.
+		// We are casting from our parent class (TTransactionalAtomicRefCount*) to our subclass (T*).
+		T::StaticDestroyObject(static_cast<const T*>(This));
+	}
+};
+
+/**
+ * Not-thread-safe specialization
+ */
+template <typename T>
+class TRefCountingMixin<T, ERefCountingMode::NotThreadSafe> 
+{
+public:
+	TRefCountingMixin() = default;
+
+	TRefCountingMixin(const TRefCountingMixin&) = delete;
+	TRefCountingMixin& operator=(const TRefCountingMixin&) = delete;
+
+	FReturnedRefCountValue AddRef() const
+	{
+		return FReturnedRefCountValue{++RefCount};
 	}
 
 	uint32 Release() const
 	{
 		checkSlow(RefCount > 0);
 
-		if constexpr (Mode == ERefCountingMode::ThreadSafe)
-		{
-			UE_AUTORTFM_ONCOMMIT(this)
-			{
-				// std::memory_order_acq_rel is used here so that, if we do end up executing the
-				// destructor, it's not possible for side effects from executing the destructor end up
-				// being visible before we've determined that the reference count is actually zero.
-				uint32 OldCount = RefCount.fetch_sub(1, std::memory_order_acq_rel);
-				checkSlow(OldCount > 0);
-				if (OldCount == 1)
-				{
-					T::StaticDestroyObject(static_cast<const T*>(this));
-				}
-			};
-		}
-		else
+		if (--RefCount == 0)
 		{
-			if (--RefCount == 0)
-			{
-				T::StaticDestroyObject(static_cast<const T*>(this));
-			}
+			StaticDestroyMixin(this);
 		}
 
 		// Note: TRefCountPtr doesn't use the return value
@@ -382,27 +432,7 @@ public:
 
 	uint32 GetRefCount() const
 	{
-		if constexpr (Mode == ERefCountingMode::ThreadSafe)
-		{
-			// A 'live' reference count is unstable by nature and so there's no benefit
-			// to try and enforce memory ordering around the reading of it.
-			//
-			// This is equivalent to https://en.cppreference.com/w/cpp/memory/shared_ptr/use_count
-
-			uint32 Count = 0;
-
-			UE_AUTORTFM_OPEN
-			{
-				// This reference count may be accessed by multiple threads
-				Count = RefCount.load(std::memory_order_relaxed);
-			};
-
-			return Count;
-		}
-		else
-		{
-			return RefCount;
-		}
+		return RefCount;
 	}
 
 	static void StaticDestroyObject(const T* Obj)
@@ -411,7 +441,12 @@ public:
 	}
 
 private:
-	mutable RefCountType RefCount{0};
+	static void StaticDestroyMixin(const TRefCountingMixin* This)
+	{
+		T::StaticDestroyObject(static_cast<const T*>(This));
+	}
+
+	mutable uint32 RefCount;
 };
 
 /**
