Fix incorrect initializer for UMaterialExpressionVectorParameter

diff --git a/Engine/Source/Runtime/Core/Private/UObject/FortniteMainBranchObjectVersion.cpp b/Engine/Source/Runtime/Core/Private/UObject/FortniteMainBranchObjectVersion.cpp
--- a/Engine/Source/Runtime/Core/Private/UObject/FortniteMainBranchObjectVersion.cpp
+++ b/Engine/Source/Runtime/Core/Private/UObject/FortniteMainBranchObjectVersion.cpp
@@ -14,6 +14,6 @@ TMap<FGuid, FGuid> FFortniteMainBranchObjectVersion::GetSystemGuids()
 	SystemGuids.Add(DevGuids.Niagara_LatestScriptCompileVersion, FGuid("393081597FA7234C92F00869DE262D73"));
 	SystemGuids.Add(DevGuids.SkeletalMeshDerivedDataVersion, FGuid("2F90180769024483BB7D90AACEB2798D"));
 	SystemGuids.Add(DevGuids.STATICMESH_DERIVEDDATA_VER, FGuid("656E338C78C34641A5D3BFAE303997AF"));
-	SystemGuids.Add(DevGuids.MaterialTranslationDDCVersion, FGuid("2AAE29CCC8BB4553A3CEB289E14D75E7"));
+	SystemGuids.Add(DevGuids.MaterialTranslationDDCVersion, FGuid("FA259B100C434423A05372474232CB7E"));
 	return SystemGuids;
 }
diff --git a/Engine/Source/Runtime/Engine/Public/Materials/MaterialExpressionVectorParameter.h b/Engine/Source/Runtime/Engine/Public/Materials/MaterialExpressionVectorParameter.h
--- a/Engine/Source/Runtime/Engine/Public/Materials/MaterialExpressionVectorParameter.h
+++ b/Engine/Source/Runtime/Engine/Public/Materials/MaterialExpressionVectorParameter.h
@@ -15,7 +15,7 @@ class UMaterialExpressionVectorParameter : public UMaterialExpressionParameter
 	GENERATED_UCLASS_BODY()
 
 	UPROPERTY(EditAnywhere, Category=MaterialExpressionVectorParameter, meta = (ShowAsInputPin = "Primary"))
-	FLinearColor DefaultValue = FLinearColor(1.0f, 0.0f, 0.0f, 0.0f);
+	FLinearColor DefaultValue;
 
 	UPROPERTY(EditAnywhere, Category=CustomPrimitiveData)
 	bool bUseCustomPrimitiveData = false;
