[TaskGraph/Pipe]
- Fix potential use-after-free when destroying an FPipe just after a call to WaitUntilEmpty
- Fix potential race between ClearTask and WaitUntilEmpty because of missing barriers on TaskCount.

diff --git a/Engine/Source/Runtime/Core/Private/Tasks/Pipe.cpp b/Engine/Source/Runtime/Core/Private/Tasks/Pipe.cpp
--- a/Engine/Source/Runtime/Core/Private/Tasks/Pipe.cpp
+++ b/Engine/Source/Runtime/Core/Private/Tasks/Pipe.cpp
@@ -43,15 +43,21 @@ namespace UE::Tasks
 			Task.Release(); // it was still pipe's last task. now that we cleared it, release the reference
 		}
 
-		if (TaskCount.fetch_sub(1, std::memory_order_relaxed) == 1)
+		// Avoid use-after-free by taking a ref on the event before decrementing the value.
+		// Since WaitUntilEmpty only looks at TaskCount to early out at which point
+		// we could decide to get rid of the pipe object.
+		TSharedRef<UE::FEventCount> LocalEmptyEvent = EmptyEventRef;
+		if (TaskCount.fetch_sub(1, std::memory_order_release) == 1)
 		{
-			EmptyEvent.Notify();
+			// use-after-free territory!
+
+			LocalEmptyEvent->Notify();
 		}
 	}
 
 	bool FPipe::WaitUntilEmpty(FTimespan InTimeout/* = FTimespan::MaxValue()*/)
 	{
-		if (TaskCount.load(std::memory_order_relaxed) == 0)
+		if (TaskCount.load(std::memory_order_acquire) == 0)
 		{
 			return true;
 		}
@@ -61,7 +67,7 @@ namespace UE::Tasks
 		UE::FTimeout Timeout(InTimeout);
 		while (true)
 		{
-			if (TaskCount.load(std::memory_order_relaxed) == 0)
+			if (TaskCount.load(std::memory_order_acquire) == 0)
 			{
 				return true;
 			}
@@ -71,14 +77,14 @@ namespace UE::Tasks
 				return false;
 			}
 
-			UE::FEventCountToken Token = EmptyEvent.PrepareWait();
+			UE::FEventCountToken Token = EmptyEventRef->PrepareWait();
 
-			if (TaskCount.load(std::memory_order_relaxed) == 0)
+			if (TaskCount.load(std::memory_order_acquire) == 0)
 			{
 				return true;
 			}
 
-			if (!EmptyEvent.WaitFor(Token, UE::FMonotonicTimeSpan::FromMilliseconds(Timeout.GetRemainingRoundedUpMilliseconds())))
+			if (!EmptyEventRef->WaitFor(Token, UE::FMonotonicTimeSpan::FromMilliseconds(Timeout.GetRemainingRoundedUpMilliseconds())))
 			{
 				break;
 			}
diff --git a/Engine/Source/Runtime/Core/Public/Tasks/Pipe.h b/Engine/Source/Runtime/Core/Public/Tasks/Pipe.h
--- a/Engine/Source/Runtime/Core/Public/Tasks/Pipe.h
+++ b/Engine/Source/Runtime/Core/Public/Tasks/Pipe.h
@@ -33,7 +33,8 @@ namespace UE::Tasks
 		// @param InDebugName helps to identify the pipe in debugger and profiler. `UE_SOURCE_LOCATION` can be used as an auto-generated
 		// unique name.
 		explicit FPipe(const TCHAR* InDebugName)
-			: DebugName(InDebugName)
+			: EmptyEventRef(MakeShared<UE::FEventCount>())
+			, DebugName(InDebugName)
 		{}
 
 		~FPipe()
@@ -134,7 +135,7 @@ namespace UE::Tasks
 		// pipe builds a chain (a linked list) of tasks and so needs to store only the last one. the last task is null if the pipe is not blocked
 		std::atomic<Private::FTaskBase*> LastTask{ nullptr };
 		std::atomic<uint64> TaskCount { 0 };
-		UE::FEventCount EmptyEvent;
+		TSharedRef<UE::FEventCount> EmptyEventRef;
 	public:
 		FORCENOINLINE const TCHAR* GetDebugName() const
 		{
