FPlatformString and FCString: Deprecate Strcpy and Strcat that take a DestLen, because some platforms are ignoring the DestLen and allowing buffer overflows.
Add FPlatformString::Strncat to provide similar functionality (but following CRT's standard of taking SrcLen rather than DestSize).
Deprecate FCString::Strncat because it has a different meaning than c++ strncat. Provide StrncatTruncateDest and StrncatTruncateSrc to replace it.
Change FCString::Strncpy's MaxLen argument from int32 to SIZE_T to match CRT strncpy.
Change the bufferlength arguments passed into GetSystemErrorMessage from the incorrect sizeof(TCharBuffer) to the correct UE_ARRAY_COUNT(TCharBuffer) to avoid the possibility of overflowing the buffer due to passing in twice the character count of the buffer.
[RN] Minor, Core

diff --git a/Engine/Plugins/Experimental/BackChannel/Source/BackChannel/Private/Protocol/BackChannelOSCBundle.cpp b/Engine/Plugins/Experimental/BackChannel/Source/BackChannel/Private/Protocol/BackChannelOSCBundle.cpp
--- a/Engine/Plugins/Experimental/BackChannel/Source/BackChannel/Private/Protocol/BackChannelOSCBundle.cpp
+++ b/Engine/Plugins/Experimental/BackChannel/Source/BackChannel/Private/Protocol/BackChannelOSCBundle.cpp
@@ -88,7 +88,7 @@ void FBackChannelOSCBundle::WriteToBuffer(TArray<uint8>& OutBuffer) const
 	uint8* DataPtr = (OutBuffer.GetData() + BufferSize);
 
 	// write header (fixed size of 8-bytes)
-	FCStringAnsi::Strcpy((ANSICHAR*)DataPtr, Size, BundleHeader);
+	FCStringAnsi::Strncpy((ANSICHAR*)DataPtr, BundleHeader, Size);
 	DataPtr[7] = 0;
 	DataPtr += 8;
 
diff --git a/Engine/Plugins/Experimental/BackChannel/Source/BackChannel/Private/Protocol/BackChannelOSCMessage.cpp b/Engine/Plugins/Experimental/BackChannel/Source/BackChannel/Private/Protocol/BackChannelOSCMessage.cpp
--- a/Engine/Plugins/Experimental/BackChannel/Source/BackChannel/Private/Protocol/BackChannelOSCMessage.cpp
+++ b/Engine/Plugins/Experimental/BackChannel/Source/BackChannel/Private/Protocol/BackChannelOSCMessage.cpp
@@ -261,7 +261,7 @@ void FBackChannelOSCMessage::WriteToBuffer(TArray<uint8>& OutBuffer) const
 	pOutBuffer[kAddressLength] = 0;
 	pOutBuffer += AddressSize;
 
-	FCStringAnsi::Strcpy(pOutBuffer, kTagLength, TCHAR_TO_ANSI(*FinalTagString));
+	FCStringAnsi::Strncpy(pOutBuffer, TCHAR_TO_ANSI(*FinalTagString), kTagLength);
 	pOutBuffer[kTagLength] = 0;
 	pOutBuffer += TagSize;
 
diff --git a/Engine/Plugins/Online/OnlineServicesEOSGS/Source/Private/Online/LobbiesEOSGSTypes.h b/Engine/Plugins/Online/OnlineServicesEOSGS/Source/Private/Online/LobbiesEOSGSTypes.h
--- a/Engine/Plugins/Online/OnlineServicesEOSGS/Source/Private/Online/LobbiesEOSGSTypes.h
+++ b/Engine/Plugins/Online/OnlineServicesEOSGS/Source/Private/Online/LobbiesEOSGSTypes.h
@@ -402,7 +402,7 @@ public:
 	FLobbyInviteIdEOS() = default;
 	FLobbyInviteIdEOS(const char* InviteIdEOS)
 	{
-		FPlatformString::Strcpy(Data, MaxLobbyInviteIdSize, InviteIdEOS);
+		FPlatformString::Strncpy(Data, InviteIdEOS, MaxLobbyInviteIdSize);
 	}
 
 	const char* Get() const { return Data; }
diff --git a/Engine/Plugins/Online/OnlineSubsystem/Source/Private/OnlineKeyValuePair.cpp b/Engine/Plugins/Online/OnlineSubsystem/Source/Private/OnlineKeyValuePair.cpp
--- a/Engine/Plugins/Online/OnlineSubsystem/Source/Private/OnlineKeyValuePair.cpp
+++ b/Engine/Plugins/Online/OnlineSubsystem/Source/Private/OnlineKeyValuePair.cpp
@@ -114,7 +114,7 @@ void FVariantData::SetValue(const TCHAR* InData)
 		if (StrLen > 0)
 		{
 			// Copy the data
-			FCString::Strcpy(Value.AsTCHAR, StrLen + 1, InData);
+			FCString::Strncpy(Value.AsTCHAR, InData, StrLen + 1);
 		}
 		else
 		{
diff --git a/Engine/Plugins/Runtime/OpenXR/Source/OpenXRInput/Private/OpenXRInput.cpp b/Engine/Plugins/Runtime/OpenXR/Source/OpenXRInput/Private/OpenXRInput.cpp
--- a/Engine/Plugins/Runtime/OpenXR/Source/OpenXRInput/Private/OpenXRInput.cpp
+++ b/Engine/Plugins/Runtime/OpenXR/Source/OpenXRInput/Private/OpenXRInput.cpp
@@ -120,7 +120,7 @@ FOpenXRInputPlugin::FOpenXRAction::FOpenXRAction(XrActionSet InActionSet,
 	}
 	else
 	{
-		FCStringAnsi::Strcpy(Info.localizedActionName, XR_MAX_LOCALIZED_ACTION_NAME_SIZE, ActionName);
+		FCStringAnsi::Strncpy(Info.localizedActionName, ActionName, XR_MAX_LOCALIZED_ACTION_NAME_SIZE);
 	}
 
 	if (OpenXRHMD)
@@ -171,7 +171,7 @@ FOpenXRInputPlugin::FOpenXRActionSet::FOpenXRActionSet(XrInstance InInstance,
 	}
 	else
 	{
-		FCStringAnsi::Strcpy(Info.localizedActionSetName, XR_MAX_LOCALIZED_ACTION_SET_NAME_SIZE, ActionName);
+		FCStringAnsi::Strncpy(Info.localizedActionSetName, ActionName, XR_MAX_LOCALIZED_ACTION_SET_NAME_SIZE);
 	}
 	Info.priority = InPriority;
 
@@ -211,26 +211,26 @@ FOpenXRInputPlugin::FOpenXRController::FOpenXRController(XrActionSet InActionSet
 	Info.countSubactionPaths = 0;
 	Info.subactionPaths = nullptr;
 
-	FCStringAnsi::Strcpy(Info.localizedActionName, XR_MAX_ACTION_NAME_SIZE, InName);
-	FCStringAnsi::Strcat(Info.localizedActionName, XR_MAX_ACTION_NAME_SIZE, " Grip Pose");
+	FCStringAnsi::Strncpy(Info.localizedActionName, InName, XR_MAX_ACTION_NAME_SIZE);
+	FCStringAnsi::StrncatTruncateDest(Info.localizedActionName, XR_MAX_ACTION_NAME_SIZE, " Grip Pose");
 	FilterActionName(Info.localizedActionName, Info.actionName);
 	Info.actionType = XR_ACTION_TYPE_POSE_INPUT;
 	XR_ENSURE(xrCreateAction(ActionSet, &Info, &GripAction));
 
-	FCStringAnsi::Strcpy(Info.localizedActionName, XR_MAX_ACTION_NAME_SIZE, InName);
-	FCStringAnsi::Strcat(Info.localizedActionName, XR_MAX_ACTION_NAME_SIZE, " Aim Pose");
+	FCStringAnsi::Strncpy(Info.localizedActionName, InName, XR_MAX_ACTION_NAME_SIZE);
+	FCStringAnsi::StrncatTruncateDest(Info.localizedActionName, XR_MAX_ACTION_NAME_SIZE, " Aim Pose");
 	FilterActionName(Info.localizedActionName, Info.actionName);
 	Info.actionType = XR_ACTION_TYPE_POSE_INPUT;
 	XR_ENSURE(xrCreateAction(ActionSet, &Info, &AimAction));
 
-	FCStringAnsi::Strcpy(Info.localizedActionName, XR_MAX_ACTION_NAME_SIZE, InName);
-	FCStringAnsi::Strcat(Info.localizedActionName, XR_MAX_ACTION_NAME_SIZE, " Palm Pose");
+	FCStringAnsi::Strncpy(Info.localizedActionName, InName, XR_MAX_ACTION_NAME_SIZE);
+	FCStringAnsi::StrncatTruncateDest(Info.localizedActionName, XR_MAX_ACTION_NAME_SIZE, " Palm Pose");
 	FilterActionName(Info.localizedActionName, Info.actionName);
 	Info.actionType = XR_ACTION_TYPE_POSE_INPUT;
 	XR_ENSURE(xrCreateAction(ActionSet, &Info, &PalmAction));
 
-	FCStringAnsi::Strcpy(Info.localizedActionName, XR_MAX_ACTION_NAME_SIZE, InName);
-	FCStringAnsi::Strcat(Info.localizedActionName, XR_MAX_ACTION_NAME_SIZE, " Vibration");
+	FCStringAnsi::Strncpy(Info.localizedActionName, InName, XR_MAX_ACTION_NAME_SIZE);
+	FCStringAnsi::StrncatTruncateDest(Info.localizedActionName, XR_MAX_ACTION_NAME_SIZE, " Vibration");
 	FilterActionName(Info.localizedActionName, Info.actionName);
 	Info.actionType = XR_ACTION_TYPE_VIBRATION_OUTPUT;
 	XR_ENSURE(xrCreateAction(ActionSet, &Info, &VibrationAction));
diff --git a/Engine/Plugins/Runtime/OpenXREyeTracker/Source/OpenXREyeTracker/Private/OpenXREyeTracker.cpp b/Engine/Plugins/Runtime/OpenXREyeTracker/Source/OpenXREyeTracker/Private/OpenXREyeTracker.cpp
--- a/Engine/Plugins/Runtime/OpenXREyeTracker/Source/OpenXREyeTracker/Private/OpenXREyeTracker.cpp
+++ b/Engine/Plugins/Runtime/OpenXREyeTracker/Source/OpenXREyeTracker/Private/OpenXREyeTracker.cpp
@@ -93,8 +93,8 @@ void FOpenXREyeTracker::AttachActionSets(TSet<XrActionSet>& OutActionSets)
 		XrActionSetCreateInfo Info;
 		Info.type = XR_TYPE_ACTION_SET_CREATE_INFO;
 		Info.next = nullptr;
-		FCStringAnsi::Strcpy(Info.actionSetName, XR_MAX_ACTION_SET_NAME_SIZE, "openxreyetrackeractionset");
-		FCStringAnsi::Strcpy(Info.localizedActionSetName, XR_MAX_LOCALIZED_ACTION_SET_NAME_SIZE, "OpenXR Eye Tracker Action Set");
+		FCStringAnsi::Strncpy(Info.actionSetName, "openxreyetrackeractionset", XR_MAX_ACTION_SET_NAME_SIZE);
+		FCStringAnsi::Strncpy(Info.localizedActionSetName, "OpenXR Eye Tracker Action Set", XR_MAX_LOCALIZED_ACTION_SET_NAME_SIZE);
 		Info.priority = 0;
 		XR_ENSURE(xrCreateActionSet(Instance, &Info, &EyeTrackerActionSet));
 	}
@@ -106,8 +106,8 @@ void FOpenXREyeTracker::AttachActionSets(TSet<XrActionSet>& OutActionSets)
 		Info.next = nullptr;
 		Info.countSubactionPaths = 0;
 		Info.subactionPaths = nullptr;
-		FCStringAnsi::Strcpy(Info.actionName, XR_MAX_ACTION_NAME_SIZE, "openxreyetrackeraction");
-		FCStringAnsi::Strcpy(Info.localizedActionName, XR_MAX_LOCALIZED_ACTION_NAME_SIZE, "OpenXR Eye Tracker Action");
+		FCStringAnsi::Strncpy(Info.actionName, "openxreyetrackeraction", XR_MAX_ACTION_NAME_SIZE);
+		FCStringAnsi::Strncpy(Info.localizedActionName, "OpenXR Eye Tracker Action", XR_MAX_LOCALIZED_ACTION_NAME_SIZE);
 		Info.actionType = XR_ACTION_TYPE_POSE_INPUT;
 		XR_ENSURE(xrCreateAction(EyeTrackerActionSet, &Info, &EyeTrackerAction));
 	}
diff --git a/Engine/Plugins/Runtime/OpenXRViveTracker/Source/OpenXRViveTracker/Private/OpenXRViveTracker.cpp b/Engine/Plugins/Runtime/OpenXRViveTracker/Source/OpenXRViveTracker/Private/OpenXRViveTracker.cpp
--- a/Engine/Plugins/Runtime/OpenXRViveTracker/Source/OpenXRViveTracker/Private/OpenXRViveTracker.cpp
+++ b/Engine/Plugins/Runtime/OpenXRViveTracker/Source/OpenXRViveTracker/Private/OpenXRViveTracker.cpp
@@ -100,14 +100,14 @@ FOpenXRViveTracker::FViveTracker::FViveTracker(XrActionSet InActionSet, FOpenXRP
 	Info.countSubactionPaths = 0;
 	Info.subactionPaths = nullptr;
 
-	FCStringAnsi::Strcpy(Info.localizedActionName, XR_MAX_ACTION_NAME_SIZE, InName);
-	FCStringAnsi::Strcat(Info.localizedActionName, XR_MAX_ACTION_NAME_SIZE, " Pose");
+	FCStringAnsi::Strncpy(Info.localizedActionName, InName, XR_MAX_ACTION_NAME_SIZE);
+	FCStringAnsi::StrncatTruncateDest(Info.localizedActionName, XR_MAX_ACTION_NAME_SIZE, " Pose");
 	FilterActionName(Info.localizedActionName, Info.actionName);
 	Info.actionType = XR_ACTION_TYPE_POSE_INPUT;
 	XR_ENSURE(xrCreateAction(ActionSet, &Info, &GripAction));
 
-	FCStringAnsi::Strcpy(Info.localizedActionName, XR_MAX_ACTION_NAME_SIZE, InName);
-	FCStringAnsi::Strcat(Info.localizedActionName, XR_MAX_ACTION_NAME_SIZE, " Vibration");
+	FCStringAnsi::Strncpy(Info.localizedActionName, InName, XR_MAX_ACTION_NAME_SIZE);
+	FCStringAnsi::StrncatTruncateDest(Info.localizedActionName, XR_MAX_ACTION_NAME_SIZE, " Vibration");
 	FilterActionName(Info.localizedActionName, Info.actionName);
 	Info.actionType = XR_ACTION_TYPE_VIBRATION_OUTPUT;
 	XR_ENSURE(xrCreateAction(ActionSet, &Info, &VibrationAction));
@@ -193,9 +193,9 @@ const void* FOpenXRViveTracker::OnCreateSession(XrInstance InInstance, XrSystemI
 	}
 
 	XrActionSetCreateInfo CreateInfo = { XR_TYPE_ACTION_SET_CREATE_INFO };
-	FCStringAnsi::Strcpy(CreateInfo.actionSetName, XR_MAX_ACTION_SET_NAME_SIZE, "trackers");
-	FCStringAnsi::Strcpy(CreateInfo.localizedActionSetName,
-		XR_MAX_LOCALIZED_ACTION_SET_NAME_SIZE, "Vive Trackers");
+	FCStringAnsi::Strncpy(CreateInfo.actionSetName, "trackers", XR_MAX_ACTION_SET_NAME_SIZE);
+	FCStringAnsi::Strncpy(CreateInfo.localizedActionSetName, "Vive Trackers",
+		XR_MAX_LOCALIZED_ACTION_SET_NAME_SIZE);
 	XR_ENSURE(xrCreateActionSet(InInstance, &CreateInfo, &TrackerActionSet));
 
 	Trackers.Add(EControllerHand::Pad, FViveTracker(TrackerActionSet, FOpenXRPath("/user/vive_tracker_htcx/role/keyboard"), "Keyboard"));
diff --git a/Engine/Plugins/Runtime/XR/XRScribe/Source/XRScribe/Private/XRScribeEmulationLayer.cpp b/Engine/Plugins/Runtime/XR/XRScribe/Source/XRScribe/Private/XRScribeEmulationLayer.cpp
--- a/Engine/Plugins/Runtime/XR/XRScribe/Source/XRScribe/Private/XRScribeEmulationLayer.cpp
+++ b/Engine/Plugins/Runtime/XR/XRScribe/Source/XRScribe/Private/XRScribeEmulationLayer.cpp
@@ -92,7 +92,7 @@ struct FOpenXREmulatedInstance
 	FOpenXREmulatedInstance()
 	{
 		FOpenXREmulatedPath NullPathString;
-		FCStringAnsi::Strcpy(NullPathString.GetData(), XR_MAX_PATH_LENGTH, "XR_NULL_PATH");
+		FCStringAnsi::Strncpy(NullPathString.GetData(), "XR_NULL_PATH", XR_MAX_PATH_LENGTH);
 		PathList.Add(MoveTemp(NullPathString));
 		ANSIStringToPathMap.Add(NullPathString.GetData(), 0);
 	}
@@ -2072,7 +2072,7 @@ XrResult FOpenXREmulationLayer::XrLayerStringToPath(XrInstance instance, const c
 
 		// TODO: well-formed path string filtering
 
-		FCStringAnsi::Strcpy(EmulatedPathString.GetData(), XR_MAX_PATH_LENGTH, pathString);
+		FCStringAnsi::Strncpy(EmulatedPathString.GetData(), pathString, XR_MAX_PATH_LENGTH);
 		{
 			FWriteScopeLock InstanceLock(InstanceMutex);
 			const int32 AddedIndex = CurrentInstance->PathList.Add(MoveTemp(EmulatedPathString));
@@ -2176,8 +2176,8 @@ XrResult FOpenXREmulationLayer::XrLayerCreateActionSet(XrInstance instance, cons
 
 	TUniquePtr<FOpenXREmulatedActionSet> ActionSet = MakeUnique<FOpenXREmulatedActionSet>();
 	ActionSet->Priority = createInfo->priority;
-	FCStringAnsi::Strcpy(ActionSet->ActionSetName.GetData(), XR_MAX_ACTION_NAME_SIZE, createInfo->actionSetName);
-	FCStringAnsi::Strcpy(ActionSet->LocalizedActionSetName.GetData(), XR_MAX_LOCALIZED_ACTION_NAME_SIZE, createInfo->localizedActionSetName);
+	FCStringAnsi::Strncpy(ActionSet->ActionSetName.GetData(), createInfo->actionSetName, XR_MAX_ACTION_NAME_SIZE);
+	FCStringAnsi::Strncpy(ActionSet->LocalizedActionSetName.GetData(), createInfo->localizedActionSetName, XR_MAX_LOCALIZED_ACTION_NAME_SIZE);
 
 	{
 		FWriteScopeLock InstanceLock(InstanceMutex);
@@ -2271,8 +2271,8 @@ XrResult FOpenXREmulationLayer::XrLayerCreateAction(XrActionSet actionSet, const
 	
 	Action->OwningActionSet = EmulatedActionSet;
 	Action->ActionType = createInfo->actionType;
-	FCStringAnsi::Strcpy(Action->ActionName.GetData(), XR_MAX_ACTION_NAME_SIZE, createInfo->actionName);
-	FCStringAnsi::Strcpy(Action->LocalizedActionName.GetData(), XR_MAX_LOCALIZED_ACTION_NAME_SIZE, createInfo->localizedActionName);
+	FCStringAnsi::Strncpy(Action->ActionName.GetData(), createInfo->actionName, XR_MAX_ACTION_NAME_SIZE);
+	FCStringAnsi::Strncpy(Action->LocalizedActionName.GetData(), createInfo->localizedActionName, XR_MAX_LOCALIZED_ACTION_NAME_SIZE);
 	Action->SubactionPaths.Append(createInfo->subactionPaths, createInfo->countSubactionPaths);
 
 	*action = reinterpret_cast<XrAction>(Action.Get());
diff --git a/Engine/Plugins/Runtime/nDisplay/Source/DisplayClusterProjection/Private/Policy/VIOSO/Windows/DisplayClusterProjectionVIOSOWarper.cpp b/Engine/Plugins/Runtime/nDisplay/Source/DisplayClusterProjection/Private/Policy/VIOSO/Windows/DisplayClusterProjectionVIOSOWarper.cpp
--- a/Engine/Plugins/Runtime/nDisplay/Source/DisplayClusterProjection/Private/Policy/VIOSO/Windows/DisplayClusterProjectionVIOSOWarper.cpp
+++ b/Engine/Plugins/Runtime/nDisplay/Source/DisplayClusterProjection/Private/Policy/VIOSO/Windows/DisplayClusterProjectionVIOSOWarper.cpp
@@ -285,7 +285,7 @@ bool FDisplayClusterProjectionVIOSOWarper::Initialize(void* pDxDevice, const FVe
 			}
 
 			// Assign calib file
-			FPlatformString::Strcpy(pWarper->calibFile, sizeof(VWB_Warper::calibFile), TCHAR_TO_ANSI(*VIOSOConfigData.CalibrationFile));
+			FPlatformString::Strncpy(pWarper->calibFile, TCHAR_TO_ANSI(*VIOSOConfigData.CalibrationFile), UE_ARRAY_COUNT(VWB_Warper::calibFile));
 
 			/// set to true to make the world turn and move with view direction and eye position, this is the case if the viewer gets
 			/// moved by a motion platform, defaults to false
diff --git a/Engine/Source/Developer/Horde/Private/Compute/ComputeBuffer.cpp b/Engine/Source/Developer/Horde/Private/Compute/ComputeBuffer.cpp
--- a/Engine/Source/Developer/Horde/Private/Compute/ComputeBuffer.cpp
+++ b/Engine/Source/Developer/Horde/Private/Compute/ComputeBuffer.cpp
@@ -314,7 +314,7 @@ struct FComputeBufferDetail
 		static_assert(sizeof(FChunkStatePtr) == sizeof(long long), "Incorrect size of FChunkStatePtr; check union is declared correctly.");
 		static_assert(sizeof(FWriterStatePtr) == sizeof(long long), "Incorrect size of FWriterStatePtr; check union is declared correctly.");
 
-		FCStringAnsi::Strcpy(Name, FComputeBuffer::MaxNameLength, InName);
+		FCStringAnsi::Strncpy(Name, InName, FComputeBuffer::MaxNameLength);
 	}
 
 	~FComputeBufferDetail()
diff --git a/Engine/Source/Developer/Horde/Private/Compute/ComputeSocket.cpp b/Engine/Source/Developer/Horde/Private/Compute/ComputeSocket.cpp
--- a/Engine/Source/Developer/Horde/Private/Compute/ComputeSocket.cpp
+++ b/Engine/Source/Developer/Horde/Private/Compute/ComputeSocket.cpp
@@ -211,7 +211,7 @@ size_t FWorkerComputeSocket::ReadString(const unsigned char* Pos, char* OutText,
 	unsigned int TextLen;
 
 	size_t Len = ReadVarUInt(Pos, &TextLen);
-	FCStringAnsi::Strcpy(OutText, OutTextMaxLen, (const char*)Pos + Len);
+	FCStringAnsi::Strncpy(OutText, (const char*)Pos + Len, OutTextMaxLen);
 
 	return Len + TextLen;
 }
diff --git a/Engine/Source/Developer/Horde/Private/HordePlatform.cpp b/Engine/Source/Developer/Horde/Private/HordePlatform.cpp
--- a/Engine/Source/Developer/Horde/Private/HordePlatform.cpp
+++ b/Engine/Source/Developer/Horde/Private/HordePlatform.cpp
@@ -66,7 +66,7 @@ bool FHordePlatform::GetEnvironmentVariable(const char* Name, char* Buffer, size
 	char* Value = getenv(Name);
 	if (Value != nullptr)
 	{
-		FCStringAnsi::Strcpy(Buffer, BufferLen, Value);
+		FCStringAnsi::Strncpy(Buffer, Value, BufferLen);
 		return true;
 	}
 	return false;
diff --git a/Engine/Source/Developer/ShaderCompilerCommon/Private/HlslUtils.h b/Engine/Source/Developer/ShaderCompilerCommon/Private/HlslUtils.h
--- a/Engine/Source/Developer/ShaderCompilerCommon/Private/HlslUtils.h
+++ b/Engine/Source/Developer/ShaderCompilerCommon/Private/HlslUtils.h
@@ -93,7 +93,7 @@ namespace CrossCompiler
 			const auto Length = FCString::Strlen(String);
 			const auto Size = (Length + 1) * sizeof(TCHAR);
 			auto* Data = (TCHAR*)Alloc(Size, sizeof(TCHAR));
-			FCString::Strcpy(Data, Length, String);
+			FCString::Strncpy(Data, String, Length+1);
 			return Data;
 		}
 
diff --git a/Engine/Source/Developer/ShaderCompilerCommon/Private/ShaderCompilerCommon.cpp b/Engine/Source/Developer/ShaderCompilerCommon/Private/ShaderCompilerCommon.cpp
--- a/Engine/Source/Developer/ShaderCompilerCommon/Private/ShaderCompilerCommon.cpp
+++ b/Engine/Source/Developer/ShaderCompilerCommon/Private/ShaderCompilerCommon.cpp
@@ -1861,7 +1861,7 @@ void CompileShaderOffline(const FShaderCompilerInput& Input,
 					FString StatsOutput = CompilerCommand + FString("\n") + StdOut;
 					const int32 StatsLen = StatsOutput.Len();
 					TSharedPtr<ANSICHAR> ShaderStats = MakeShareable(new ANSICHAR[StatsLen + 1]);
-					FCStringAnsi::Strcpy(ShaderStats.Get(), StatsLen + 1, TCHAR_TO_ANSI(*StatsOutput));
+					FCStringAnsi::Strncpy(ShaderStats.Get(), TCHAR_TO_ANSI(*StatsOutput), StatsLen + 1);
 					ArOutput->Serialize((void*)ShaderStats.Get(), StatsLen);
 					delete ArOutput;
 				}
diff --git a/Engine/Source/Developer/ShaderFormatOpenGL/Private/OpenGLShaderCompiler.cpp b/Engine/Source/Developer/ShaderFormatOpenGL/Private/OpenGLShaderCompiler.cpp
--- a/Engine/Source/Developer/ShaderFormatOpenGL/Private/OpenGLShaderCompiler.cpp
+++ b/Engine/Source/Developer/ShaderFormatOpenGL/Private/OpenGLShaderCompiler.cpp
@@ -3261,7 +3261,7 @@ static TSharedPtr<ANSICHAR> PrepareCodeForOfflineCompilationInternal(const GLSLV
 
 	const int32 SourceLen = StrOutSource.Len();
 	TSharedPtr<ANSICHAR> RetShaderSource = MakeShareable(new ANSICHAR[SourceLen + 1]);
-	FCStringAnsi::Strcpy(RetShaderSource.Get(), SourceLen + 1, TCHAR_TO_ANSI(*StrOutSource));
+	FCStringAnsi::Strncpy(RetShaderSource.Get(), TCHAR_TO_ANSI(*StrOutSource), SourceLen + 1);
 
 	return RetShaderSource;
 }
diff --git a/Engine/Source/Developer/SlateFileDialogs/Private/SlateFileDlgWindow.cpp b/Engine/Source/Developer/SlateFileDialogs/Private/SlateFileDlgWindow.cpp
--- a/Engine/Source/Developer/SlateFileDialogs/Private/SlateFileDlgWindow.cpp
+++ b/Engine/Source/Developer/SlateFileDialogs/Private/SlateFileDlgWindow.cpp
@@ -945,14 +945,14 @@ void SSlateFileOpenDlg::BuildDirectoryPath()
 			BuiltPath = BuiltPath + TEXT("/") + AbsPath.Left(Idx);
 		}
 
-		FCString::Strcpy(Temp, UE_ARRAY_COUNT(Temp), &AbsPath[Idx < AbsPath.Len() - 1 ? Idx + 1 : Idx]);
+		FCString::Strncpy(Temp, &AbsPath[Idx < AbsPath.Len() - 1 ? Idx + 1 : Idx], UE_ARRAY_COUNT(Temp));
 
 		DirectoryNodesArray.Add(FDirNode(AbsPath.Left(Idx), nullptr));
 	}
 	else if (PLATFORM_LINUX)
 	{
 		// start with system base directory
-		FCString::Strcpy(Temp, UE_ARRAY_COUNT(Temp), *AbsPath);
+		FCString::Strncpy(Temp, *AbsPath, UE_ARRAY_COUNT(Temp));
 
 		BuiltPath = "/";
 		DirectoryNodesArray.Add(FDirNode(FString(TEXT("/")), nullptr));
@@ -1614,7 +1614,7 @@ void SSlateFileOpenDlg::ParseFilters()
 			if (FilterNameArray.Num() == 0)
 			{
 				TCHAR Temp[MAX_FILTER_LENGTH] = {0};
-				FCString::Strcpy(Temp, UE_ARRAY_COUNT(Temp), *Filters);
+				FCString::Strncpy(Temp, *Filters, UE_ARRAY_COUNT(Temp));
 
 				// break path into tokens
 				TCHAR *ContextStr = nullptr;
diff --git a/Engine/Source/Developer/StandaloneRenderer/Private/OpenGL/SlateOpenGLShaders.cpp b/Engine/Source/Developer/StandaloneRenderer/Private/OpenGL/SlateOpenGLShaders.cpp
--- a/Engine/Source/Developer/StandaloneRenderer/Private/OpenGL/SlateOpenGLShaders.cpp
+++ b/Engine/Source/Developer/StandaloneRenderer/Private/OpenGL/SlateOpenGLShaders.cpp
@@ -128,10 +128,10 @@ void FSlateOpenGLShader::CompileShader( const FString& Filename, GLenum ShaderTy
 	ANSICHAR* Chars[2] = {0};
 	
 	Chars[0] = new ANSICHAR[Header.Len()+1];
-	FCStringAnsi::Strcpy(Chars[0], Header.Len() + 1, TCHAR_TO_ANSI(*Header));
+	FCStringAnsi::Strncpy(Chars[0], TCHAR_TO_ANSI(*Header), Header.Len() + 1);
 	
 	Chars[1] = new ANSICHAR[Source.Len()+1];
-	FCStringAnsi::Strcpy(Chars[1], Source.Len() + 1, TCHAR_TO_ANSI(*Source));
+	FCStringAnsi::Strncpy(Chars[1], TCHAR_TO_ANSI(*Source), Source.Len() + 1);
 
 	// give opengl the source code for the shader
 	glShaderSource( ShaderID, 2, (const ANSICHAR**)Chars, NULL );
diff --git a/Engine/Source/Developer/TraceServices/Private/Common/FormatArgs.cpp b/Engine/Source/Developer/TraceServices/Private/Common/FormatArgs.cpp
--- a/Engine/Source/Developer/TraceServices/Private/Common/FormatArgs.cpp
+++ b/Engine/Source/Developer/TraceServices/Private/Common/FormatArgs.cpp
@@ -353,7 +353,7 @@ void FFormatArgsHelper::Format(TCHAR* Out, uint64 MaxOut, TCHAR* Temp, uint64 Ma
 	InitArgumentStream(ArgumentStream, FormatArgs);
 	if (ArgumentStream.ArgumentCount == 0)
 	{
-		FCString::Strcpy(Out, MaxOut, FormatString);
+		FCString::Strncpy(Out, FormatString, MaxOut);
 		return;
 	}
 
diff --git a/Engine/Source/Developer/Virtualization/Private/VirtualizationUtilities.cpp b/Engine/Source/Developer/Virtualization/Private/VirtualizationUtilities.cpp
--- a/Engine/Source/Developer/Virtualization/Private/VirtualizationUtilities.cpp
+++ b/Engine/Source/Developer/Virtualization/Private/VirtualizationUtilities.cpp
@@ -51,7 +51,7 @@ void GetFormattedSystemError(FStringBuilderBase& SystemErrorMessage)
 	if (SystemError != 0)
 	{
 		TCHAR SystemErrorMsg[MAX_SPRINTF] = { 0 };
-		FPlatformMisc::GetSystemErrorMessage(SystemErrorMsg, sizeof(SystemErrorMsg), SystemError);
+		FPlatformMisc::GetSystemErrorMessage(SystemErrorMsg, UE_ARRAY_COUNT(SystemErrorMsg), SystemError);
 
 		SystemErrorMessage.Appendf(TEXT("'%s' (%d)"), SystemErrorMsg, SystemError);
 	}
diff --git a/Engine/Source/Editor/MaterialEditor/Private/MaterialStats.cpp b/Engine/Source/Editor/MaterialEditor/Private/MaterialStats.cpp
--- a/Engine/Source/Editor/MaterialEditor/Private/MaterialStats.cpp
+++ b/Engine/Source/Editor/MaterialEditor/Private/MaterialStats.cpp
@@ -1573,7 +1573,7 @@ void FMaterialStats::ExtractHLSLCode()
 
 		// Make a copy so we can insert null terminators.
 		TCHAR* MarkupSourceCopy = new TCHAR[MarkupSource.Len() + 1];
-		FCString::Strcpy(MarkupSourceCopy, MarkupSource.Len() + 1, *MarkupSource);
+		FCString::Strncpy(MarkupSourceCopy, *MarkupSource, MarkupSource.Len() + 1);
 
 		TCHAR* Ptr = MarkupSourceCopy;
 		while (Ptr && *Ptr != '\0')
diff --git a/Engine/Source/Editor/UnrealEd/Private/Cooker/DiffWriterArchive.cpp b/Engine/Source/Editor/UnrealEd/Private/Cooker/DiffWriterArchive.cpp
--- a/Engine/Source/Editor/UnrealEd/Private/Cooker/DiffWriterArchive.cpp
+++ b/Engine/Source/Editor/UnrealEd/Private/Cooker/DiffWriterArchive.cpp
@@ -226,7 +226,7 @@ ANSICHAR* FCallstacks::AddUniqueCallstack(bool bIsCollectingCallstacks, UObject*
 		{
 			const int32 Len = FCStringAnsi::Strlen(StackTrace.Get()) + 1;
 			TUniquePtr<ANSICHAR[]> NewCallstack = MakeUnique<ANSICHAR[]>(Len);
-			FCStringAnsi::Strcpy(NewCallstack.Get(), Len, StackTrace.Get());
+			FCStringAnsi::Strncpy(NewCallstack.Get(), StackTrace.Get(), Len);
 			FCallstackData& NewEntry = UniqueCallstacks.Add(OutCallstackCRC, FCallstackData(MoveTemp(NewCallstack), SerializedObject, SerializedProperty));
 			Callstack = NewEntry.Callstack.Get();
 		}
@@ -265,7 +265,7 @@ void FCallstacks::Add(
 		//if we have a debug name stack, plaster it onto the end of the current stack buffer so that it's a part of the unique stack entry.
 		if (DebugDataStack.Num() > 0)
 		{
-			FCStringAnsi::Strcat(StackTrace.Get(), StackTraceSize, DebugDataStackMarker);
+			FCStringAnsi::StrncatTruncateDest(StackTrace.Get(), StackTraceSize, DebugDataStackMarker);
 
 			const FString SubIndent = FString(IndentToken) + FString(TEXT("    "));
 
@@ -274,17 +274,17 @@ void FCallstacks::Add(
 			{
 				if (bIsIndenting)
 				{
-					FCStringAnsi::Strcat(StackTrace.Get(), StackTraceSize, TCHAR_TO_ANSI(*SubIndent));
+					FCStringAnsi::StrncatTruncateDest(StackTrace.Get(), StackTraceSize, TCHAR_TO_ANSI(*SubIndent));
 				}
 
 				ANSICHAR DebugName[NAME_SIZE];
 				DebugData.GetPlainANSIString(DebugName);
-				FCStringAnsi::Strcat(StackTrace.Get(), StackTraceSize, DebugName);
+				FCStringAnsi::StrncatTruncateDest(StackTrace.Get(), StackTraceSize, DebugName);
 
 				//these are special-cased, as we assume they'll be followed by object/property names and want the names on the same line for readability's sake.
 				const bool bIsPropertyLabel = (DebugData == TEXT("SerializeScriptProperties") || DebugData == TEXT("PropertySerialize") || DebugData == TEXT("SerializeTaggedProperty"));
 				const ANSICHAR* const LineEnd = bIsPropertyLabel ? ": " : "\r\n";
-				FCStringAnsi::Strcat(StackTrace.Get(), StackTraceSize, LineEnd);
+				FCStringAnsi::StrncatTruncateDest(StackTrace.Get(), StackTraceSize, LineEnd);
 				bIsIndenting = !bIsPropertyLabel;
 			}
 		}
diff --git a/Engine/Source/Editor/UnrealEd/Private/Fbx/FbxMainImport.cpp b/Engine/Source/Editor/UnrealEd/Private/Fbx/FbxMainImport.cpp
--- a/Engine/Source/Editor/UnrealEd/Private/Fbx/FbxMainImport.cpp
+++ b/Engine/Source/Editor/UnrealEd/Private/Fbx/FbxMainImport.cpp
@@ -2013,7 +2013,7 @@ int32 GetNodeNameDiffScore(const char* MeshName, const FbxNode* FbxMeshNode)
 FbxNode* FFbxImporter::GetMeshNodesFromName(const FString& ReimportMeshName, TArray<FbxNode*>& FbxMeshArray)
 {
 	char MeshName[2048];
-	FCStringAnsi::Strcpy(MeshName, 2048, TCHAR_TO_UTF8(*ReimportMeshName));
+	FCStringAnsi::Strncpy(MeshName, TCHAR_TO_UTF8(*ReimportMeshName), UE_ARRAY_COUNT(MeshName));
 
 	// find the Fbx mesh node that the Unreal Mesh matches according to name
 	int32 BestMatchIndex = INDEX_NONE;
diff --git a/Engine/Source/Programs/ShaderCompileWorker/Private/ShaderCompileWorker.cpp b/Engine/Source/Programs/ShaderCompileWorker/Private/ShaderCompileWorker.cpp
--- a/Engine/Source/Programs/ShaderCompileWorker/Private/ShaderCompileWorker.cpp
+++ b/Engine/Source/Programs/ShaderCompileWorker/Private/ShaderCompileWorker.cpp
@@ -552,7 +552,7 @@ private:
 				TArray<ANSICHAR> AnsiString;
 				AnsiString.SetNumZeroed(Name.Len() + 1);
 				ANSICHAR* Dest = &AnsiString[0];
-				FCStringAnsi::Strcpy(Dest, Name.Len() + 1, TCHAR_TO_ANSI(*Name));
+				FCStringAnsi::Strncpy(Dest, TCHAR_TO_ANSI(*Name), Name.Len() + 1);
 
 				AllocatedAnsiStrings.Add(MakeUnique<TArray<ANSICHAR>>(AnsiString));
 				CharName = &(*AllocatedAnsiStrings.Last())[0];
diff --git a/Engine/Source/Runtime/ApplicationCore/Private/Linux/LinuxApplication.cpp b/Engine/Source/Runtime/ApplicationCore/Private/Linux/LinuxApplication.cpp
--- a/Engine/Source/Runtime/ApplicationCore/Private/Linux/LinuxApplication.cpp
+++ b/Engine/Source/Runtime/ApplicationCore/Private/Linux/LinuxApplication.cpp
@@ -1970,7 +1970,7 @@ bool FLinuxApplication::IsMouseAttached() const
 	int rc;
 	char Mouse[64] = "/sys/class/input/mouse0";
 	int MouseIdx = strlen(Mouse) - 1;
-	FCStringAnsi::Strncat(Mouse, "/device/name", sizeof(Mouse) - 1);
+	FCStringAnsi::StrncatTruncateDest(Mouse, UE_ARRAY_COUNT(Mouse), "/device/name");
 
 	for (int i=0; i<9; i++)
 	{
diff --git a/Engine/Source/Runtime/ApplicationCore/Private/Windows/WindowsPlatformApplicationMisc.cpp b/Engine/Source/Runtime/ApplicationCore/Private/Windows/WindowsPlatformApplicationMisc.cpp
--- a/Engine/Source/Runtime/ApplicationCore/Private/Windows/WindowsPlatformApplicationMisc.cpp
+++ b/Engine/Source/Runtime/ApplicationCore/Private/Windows/WindowsPlatformApplicationMisc.cpp
@@ -391,7 +391,7 @@ void FWindowsPlatformApplicationMisc::ClipboardCopy(const TCHAR* Str)
 		GlobalMem = GlobalAlloc( GMEM_MOVEABLE, sizeof(TCHAR)*(StrLen+1) );
 		check(GlobalMem);
 		TCHAR* Data = (TCHAR*) GlobalLock( GlobalMem );
-		FCString::Strcpy( Data, (StrLen+1), Str );
+		FCString::Strncpy( Data, Str, (StrLen+1) );
 		GlobalUnlock( GlobalMem );
 		if( SetClipboardData( CF_UNICODETEXT, GlobalMem ) == NULL )
 			UE_LOG(LogWindows, Fatal,TEXT("SetClipboardData failed with error code %i"), (uint32)GetLastError() );
diff --git a/Engine/Source/Runtime/Core/Private/Android/AndroidHeapProfiling.cpp b/Engine/Source/Runtime/Core/Private/Android/AndroidHeapProfiling.cpp
--- a/Engine/Source/Runtime/Core/Private/Android/AndroidHeapProfiling.cpp
+++ b/Engine/Source/Runtime/Core/Private/Android/AndroidHeapProfiling.cpp
@@ -31,7 +31,7 @@
 		const int AllocatorNameBufferSize = 256;
 		char AllocatorNameBuffer[AllocatorNameBufferSize] = "epicgames.";
 		const int Len = FCStringAnsi::Strlen(AllocatorNameBuffer);
-		FCStringAnsi::Strcpy(AllocatorNameBuffer + Len, AllocatorNameBufferSize - Len, TCHAR_TO_ANSI(AllocatorName));
+		FCStringAnsi::Strncpy(AllocatorNameBuffer + Len, TCHAR_TO_ANSI(AllocatorName), AllocatorNameBufferSize - Len);
 		return AHeapProfileRegisterHeap(AHeapInfoCreate(AllocatorNameBuffer));
 	}
 #endif
diff --git a/Engine/Source/Runtime/Core/Private/Android/AndroidPlatformCrashContext.cpp b/Engine/Source/Runtime/Core/Private/Android/AndroidPlatformCrashContext.cpp
--- a/Engine/Source/Runtime/Core/Private/Android/AndroidPlatformCrashContext.cpp
+++ b/Engine/Source/Runtime/Core/Private/Android/AndroidPlatformCrashContext.cpp
@@ -176,15 +176,15 @@ void FAndroidCrashContext::GenerateReportDirectoryName(char(&DirectoryNameOUT)[C
 {
 	FGuid ReportGUID = FGuid::NewGuid();
 	FCStringAnsi::Strncpy(DirectoryNameOUT, GAndroidCrashInfo.AndroidCrashReportPath, CrashReportMaxPathSize);
-	FCStringAnsi::Strcat(DirectoryNameOUT, CrashReportMaxPathSize, "/CrashReport-UE-");
-	FCStringAnsi::Strcat(DirectoryNameOUT, CrashReportMaxPathSize, GAndroidCrashInfo.ProjectNameUTF8);
-	FCStringAnsi::Strcat(DirectoryNameOUT, CrashReportMaxPathSize, "-pid-");
-	FCStringAnsi::Strcat(DirectoryNameOUT, CrashReportMaxPathSize, ItoANSI((uint32)getpid(), 10));
-	FCStringAnsi::Strcat(DirectoryNameOUT, CrashReportMaxPathSize, "-");
-	FCStringAnsi::Strcat(DirectoryNameOUT, CrashReportMaxPathSize, ItoANSI(ReportGUID.A, 16, 8));
-	FCStringAnsi::Strcat(DirectoryNameOUT, CrashReportMaxPathSize, ItoANSI(ReportGUID.B, 16, 8));
-	FCStringAnsi::Strcat(DirectoryNameOUT, CrashReportMaxPathSize, ItoANSI(ReportGUID.C, 16, 8));
-	FCStringAnsi::Strcat(DirectoryNameOUT, CrashReportMaxPathSize, ItoANSI(ReportGUID.D, 16, 8));
+	FCStringAnsi::StrncatTruncateDest(DirectoryNameOUT, CrashReportMaxPathSize, "/CrashReport-UE-");
+	FCStringAnsi::StrncatTruncateDest(DirectoryNameOUT, CrashReportMaxPathSize, GAndroidCrashInfo.ProjectNameUTF8);
+	FCStringAnsi::StrncatTruncateDest(DirectoryNameOUT, CrashReportMaxPathSize, "-pid-");
+	FCStringAnsi::StrncatTruncateDest(DirectoryNameOUT, CrashReportMaxPathSize, ItoANSI((uint32)getpid(), 10));
+	FCStringAnsi::StrncatTruncateDest(DirectoryNameOUT, CrashReportMaxPathSize, "-");
+	FCStringAnsi::StrncatTruncateDest(DirectoryNameOUT, CrashReportMaxPathSize, ItoANSI(ReportGUID.A, 16, 8));
+	FCStringAnsi::StrncatTruncateDest(DirectoryNameOUT, CrashReportMaxPathSize, ItoANSI(ReportGUID.B, 16, 8));
+	FCStringAnsi::StrncatTruncateDest(DirectoryNameOUT, CrashReportMaxPathSize, ItoANSI(ReportGUID.C, 16, 8));
+	FCStringAnsi::StrncatTruncateDest(DirectoryNameOUT, CrashReportMaxPathSize, ItoANSI(ReportGUID.D, 16, 8));
 }
 
 static void CrashReportFileCopy(const char* DestPath, const char* SourcePath)
@@ -206,18 +206,18 @@ static void CrashReportFileCopy(const char* DestPath, const char* SourcePath)
 void FAndroidCrashContext::StoreCrashInfo(bool bWriteLog) const
 {
 	char FilePath[CrashReportMaxPathSize] = { 0 };
-	FCStringAnsi::Strcpy(FilePath, ReportDirectory);
-	FCStringAnsi::Strcat(FilePath, "/");
-	FCStringAnsi::Strcat(FilePath, FGenericCrashContext::CrashContextRuntimeXMLNameA);
+	FCStringAnsi::Strncpy(FilePath, ReportDirectory, CrashReportMaxPathSize);
+	FCStringAnsi::StrncatTruncateDest(FilePath, CrashReportMaxPathSize, "/");
+	FCStringAnsi::StrncatTruncateDest(FilePath, CrashReportMaxPathSize, FGenericCrashContext::CrashContextRuntimeXMLNameA);
 	SerializeAsXML(*FString(FilePath)); // CreateFileWriter will also create destination directory.
 
 	if(bWriteLog)
 	{
 		// copy log:
-		FCStringAnsi::Strcpy(FilePath, ReportDirectory);
-		FCStringAnsi::Strcat(FilePath, "/");
-		FCStringAnsi::Strcat(FilePath, FCStringAnsi::Strlen(GAndroidCrashInfo.AppName) ? GAndroidCrashInfo.AppName : "UE4");
-		FCStringAnsi::Strcat(FilePath, ".log");
+		FCStringAnsi::Strncpy(FilePath, ReportDirectory, CrashReportMaxPathSize);
+		FCStringAnsi::StrncatTruncateDest(FilePath, CrashReportMaxPathSize, "/");
+		FCStringAnsi::StrncatTruncateDest(FilePath, CrashReportMaxPathSize, FCStringAnsi::Strlen(GAndroidCrashInfo.AppName) ? GAndroidCrashInfo.AppName : "UE4");
+		FCStringAnsi::StrncatTruncateDest(FilePath, CrashReportMaxPathSize, ".log");
 		CrashReportFileCopy(FilePath, GAndroidCrashInfo.AppLogPath);
 	}
 }
@@ -233,8 +233,8 @@ void FAndroidCrashContext::DumpAllThreadCallstacks(FAsyncThreadBackTrace* BackTr
 	}
 
 	char FilePath[FAndroidCrashContext::CrashReportMaxPathSize] = { 0 };
-	FCStringAnsi::Strcpy(FilePath, ReportDirectory);
-	FCStringAnsi::Strcat(FilePath, FAndroidCrashContext::CrashReportMaxPathSize, "/AllThreads.txt");
+	FCStringAnsi::Strncpy(FilePath, ReportDirectory, FAndroidCrashContext::CrashReportMaxPathSize);
+	FCStringAnsi::StrncatTruncateDest(FilePath, FAndroidCrashContext::CrashReportMaxPathSize, "/AllThreads.txt");
 	TArray<FCrashStackFrame> CrashStackFrames;
 	CrashStackFrames.Empty(32);
 	int DestHandle = open(FilePath, O_CREAT | O_WRONLY, S_IRUSR | S_IWUSR);
@@ -276,10 +276,10 @@ void FAndroidCrashContext::DumpAllThreadCallstacks(FAsyncThreadBackTrace* BackTr
 				for (const FCrashStackFrame& CrashStackFrame : CrashStackFrames)
 				{
 					FCStringAnsi::Strncpy(Line, TCHAR_TO_UTF8(*CrashStackFrame.ModuleName), UE_ARRAY_COUNT(Line));
-					FCStringAnsi::Strcat(Line, " 0x");
-					FCStringAnsi::Strcat(Line, ItoANSI(CrashStackFrame.BaseAddress, (uint64)16, 16));
-					FCStringAnsi::Strcat(Line, " + ");
-					FCStringAnsi::Strcat(Line, ItoANSI(CrashStackFrame.Offset, (uint64)16, 16));
+					FCStringAnsi::StrncatTruncateDest(Line, UE_ARRAY_COUNT(Line), " 0x");
+					FCStringAnsi::StrncatTruncateDest(Line, UE_ARRAY_COUNT(Line), ItoANSI(CrashStackFrame.BaseAddress, (uint64)16, 16));
+					FCStringAnsi::StrncatTruncateDest(Line, UE_ARRAY_COUNT(Line), " + ");
+					FCStringAnsi::StrncatTruncateDest(Line, UE_ARRAY_COUNT(Line), ItoANSI(CrashStackFrame.Offset, (uint64)16, 16));
 					Writeln(DestHandle, Line);
 				}
 				Writeln(DestHandle, "</CallStack>");
diff --git a/Engine/Source/Runtime/Core/Private/Android/AndroidPlatformMisc.cpp b/Engine/Source/Runtime/Core/Private/Android/AndroidPlatformMisc.cpp
--- a/Engine/Source/Runtime/Core/Private/Android/AndroidPlatformMisc.cpp
+++ b/Engine/Source/Runtime/Core/Private/Android/AndroidPlatformMisc.cpp
@@ -694,7 +694,7 @@ const TCHAR* FAndroidMisc::GetSystemErrorMessage(TCHAR* OutBuffer, int32 BufferC
 	char ErrorBuffer[1024];
 	if (strerror_r(Error, ErrorBuffer, 1024) == 0)
 	{
-		FCString::Strcpy(OutBuffer, BufferCount, UTF8_TO_TCHAR((const ANSICHAR*)ErrorBuffer));
+		FCString::Strncpy(OutBuffer, UTF8_TO_TCHAR((const ANSICHAR*)ErrorBuffer), BufferCount);
 	}
 	else
 	{
diff --git a/Engine/Source/Runtime/Core/Private/Android/AndroidPlatformStackWalk.cpp b/Engine/Source/Runtime/Core/Private/Android/AndroidPlatformStackWalk.cpp
--- a/Engine/Source/Runtime/Core/Private/Android/AndroidPlatformStackWalk.cpp
+++ b/Engine/Source/Runtime/Core/Private/Android/AndroidPlatformStackWalk.cpp
@@ -356,28 +356,28 @@ bool FAndroidPlatformStackWalk::SymbolInfoToHumanReadableString(const FProgramCo
 		// Start with address
 		ANSICHAR PCAddress[MAX_TEMP_SPRINTF] = { 0 };
 		FCStringAnsi::Snprintf(PCAddress, MAX_TEMP_SPRINTF, "0x%016llX ", SymbolInfo.ProgramCounter);
-		FCStringAnsi::Strncat(StackLine, PCAddress, MAX_SPRINTF);
+		FCStringAnsi::StrncatTruncateDest(StackLine, MAX_SPRINTF, PCAddress);
 
 		// Module if it's present
 		const bool bHasValidModuleName = FCStringAnsi::Strlen(StrippedModuleName) > 0;
 		if (bHasValidModuleName)
 		{
-			FCStringAnsi::Strncat(StackLine, StrippedModuleName, MAX_SPRINTF);
+			FCStringAnsi::StrncatTruncateDest(StackLine, MAX_SPRINTF, StrippedModuleName);
 		}
 
 		FCStringAnsi::Snprintf(PCAddress, MAX_TEMP_SPRINTF, "(0x%016llX)", SymbolInfo.OffsetInModule);
-		FCStringAnsi::Strncat(StackLine, PCAddress, MAX_SPRINTF);
-		FCStringAnsi::Strncat(StackLine, "!", MAX_SPRINTF);
+		FCStringAnsi::StrncatTruncateDest(StackLine, MAX_SPRINTF, PCAddress);
+		FCStringAnsi::StrncatTruncateDest(StackLine, MAX_SPRINTF, "!");
 
 		// Function if it's available, unknown if it's not
 		const bool bHasValidFunctionName = FCStringAnsi::Strlen(SymbolInfo.FunctionName) > 0;
 		if (bHasValidFunctionName)
 		{
-			FCStringAnsi::Strncat(StackLine, SymbolInfo.FunctionName, MAX_SPRINTF);
+			FCStringAnsi::StrncatTruncateDest(StackLine, MAX_SPRINTF, SymbolInfo.FunctionName);
 		}
 		else
 		{
-			FCStringAnsi::Strncat(StackLine, "UnknownFunction", MAX_SPRINTF);
+			FCStringAnsi::StrncatTruncateDest(StackLine, MAX_SPRINTF, "UnknownFunction");
 		}
 
 		// file info
@@ -386,15 +386,15 @@ bool FAndroidPlatformStackWalk::SymbolInfoToHumanReadableString(const FProgramCo
 		{
 			ANSICHAR FilenameAndLineNumber[MAX_TEMP_SPRINTF] = { 0 };
 			FCStringAnsi::Snprintf(FilenameAndLineNumber, MAX_TEMP_SPRINTF, " [%s:%i]", SymbolInfo.Filename, SymbolInfo.LineNumber);
-			FCStringAnsi::Strncat(StackLine, FilenameAndLineNumber, MAX_SPRINTF);
+			FCStringAnsi::StrncatTruncateDest(StackLine, MAX_SPRINTF, FilenameAndLineNumber);
 		}
 		else
 		{
-			FCStringAnsi::Strncat(StackLine, " []", MAX_SPRINTF);
+			FCStringAnsi::StrncatTruncateDest(StackLine, MAX_SPRINTF, " []");
 		}
 
 		// Append the stack line.
-		FCStringAnsi::Strncat(HumanReadableString, StackLine, HumanReadableStringSize);
+		FCStringAnsi::StrncatTruncateDest(HumanReadableString, HumanReadableStringSize, StackLine);
 
 		// Return true, if we have a valid function name.
 		return bHasValidFunctionName;
diff --git a/Engine/Source/Runtime/Core/Private/Apple/ApplePlatformMisc.cpp b/Engine/Source/Runtime/Core/Private/Apple/ApplePlatformMisc.cpp
--- a/Engine/Source/Runtime/Core/Private/Apple/ApplePlatformMisc.cpp
+++ b/Engine/Source/Runtime/Core/Private/Apple/ApplePlatformMisc.cpp
@@ -88,7 +88,7 @@ const TCHAR* FApplePlatformMisc::GetSystemErrorMessage(TCHAR* OutBuffer, int32 B
 	char* ErrorBuffer = (char*)alloca(BufferCount);
 	if (strerror_r(Error, ErrorBuffer, 1024) == 0)
 	{
-		FCString::Strcpy(OutBuffer, BufferCount, UTF8_TO_TCHAR((const ANSICHAR*)ErrorBuffer));
+		FCString::Strncpy(OutBuffer, UTF8_TO_TCHAR((const ANSICHAR*)ErrorBuffer), BufferCount);
 	}
 	else
 	{
diff --git a/Engine/Source/Runtime/Core/Private/Apple/ApplePlatformStackWalk.cpp b/Engine/Source/Runtime/Core/Private/Apple/ApplePlatformStackWalk.cpp
--- a/Engine/Source/Runtime/Core/Private/Apple/ApplePlatformStackWalk.cpp
+++ b/Engine/Source/Runtime/Core/Private/Apple/ApplePlatformStackWalk.cpp
@@ -232,7 +232,7 @@ bool FApplePlatformStackWalk::ProgramCounterToHumanReadableString( int32 Current
 	// Write out prefix, address, module, and function na,e
 	FCStringAnsi::Sprintf(TempArray, "0x%08x %s!%s ", (uint32)ProgramCounter, SymbolInfo.ModuleName, SymbolInfo.FunctionName);
 	
-	FCStringAnsi::Strcat(HumanReadableString, HumanReadableStringSize, TempArray);
+	FCStringAnsi::StrncatTruncateDest(HumanReadableString, HumanReadableStringSize, TempArray);
 
 	// Get filename.
 	{
@@ -249,11 +249,11 @@ bool FApplePlatformStackWalk::ProgramCounterToHumanReadableString( int32 Current
 			FCStringAnsi::Sprintf(FileNameLine, " [%s:%d] ", SymbolInfo.Filename, SymbolInfo.LineNumber);
 		}
 		
-		FCStringAnsi::Strcat(HumanReadableString, HumanReadableStringSize, FileNameLine);
+		FCStringAnsi::StrncatTruncateDest(HumanReadableString, HumanReadableStringSize, FileNameLine);
 	}
 
 	// For the crash reporting code this needs a Windows line ending, the caller is responsible for the '\n'
-	FCStringAnsi::Strcat(HumanReadableString, HumanReadableStringSize, "\r");
+	FCStringAnsi::StrncatTruncateDest(HumanReadableString, HumanReadableStringSize, "\r");
 
 	return true;
 }
@@ -450,7 +450,7 @@ int32 ReportCrash(ucontext_t *Context, int32 Signal, struct __siginfo* Info)
 		// Walk the stack and dump it to the allocated memory.
 		FPlatformStackWalk::StackWalkAndDump( StackTrace, StackTraceSize, Info->si_addr, Context );
 #if WITH_EDITORONLY_DATA
-        FCString::Strncat( GErrorHist, ANSI_TO_TCHAR(StackTrace), UE_ARRAY_COUNT(GErrorHist) - 1 );
+		FCString::StrncatTruncateDest( GErrorHist, UE_ARRAY_COUNT(GErrorHist), ANSI_TO_TCHAR(StackTrace) );
 		CreateExceptionInfoString(Signal, Info);
 #endif
 	}
diff --git a/Engine/Source/Runtime/Core/Private/Apple/ApplePlatformSymbolication.cpp b/Engine/Source/Runtime/Core/Private/Apple/ApplePlatformSymbolication.cpp
--- a/Engine/Source/Runtime/Core/Private/Apple/ApplePlatformSymbolication.cpp
+++ b/Engine/Source/Runtime/Core/Private/Apple/ApplePlatformSymbolication.cpp
@@ -413,7 +413,7 @@ bool FApplePlatformSymbolication::SymbolInfoForStrippedSymbol(FApplePlatformSymb
 				ANSICHAR const* FunctionName = CSSymbolGetName(Symbol);
 				if(FunctionName)
 				{
-					FCStringAnsi::Strcpy(Info.FunctionName, UE_ARRAY_COUNT(Info.FunctionName), FunctionName);
+					FCStringAnsi::Strncpy(Info.FunctionName, FunctionName, UE_ARRAY_COUNT(Info.FunctionName));
 				}
 				
 				CSRange Range = CSSymbolGetRange(Symbol);
@@ -429,7 +429,7 @@ bool FApplePlatformSymbolication::SymbolInfoForStrippedSymbol(FApplePlatformSymb
 					ANSICHAR const* FileName = CSSourceInfoGetPath(SymbolInfo);
 					if(FileName)
 					{
-						FCStringAnsi::Strcpy(Info.Filename, UE_ARRAY_COUNT(Info.FunctionName), FileName);
+						FCStringAnsi::Strncpy(Info.Filename, FileName, UE_ARRAY_COUNT(Info.FunctionName));
 					}
 				}
 				
@@ -459,7 +459,7 @@ bool FApplePlatformSymbolication::SymbolInfoForAddress(uint64 ProgramCounter, FP
 			{
 				out_SymbolInfo.LineNumber = CSSourceInfoGetLineNumber(Symbol);
 				FGenericPlatformStackWalk::CopyVirtualPathToLocal(out_SymbolInfo.Filename, UE_ARRAY_COUNT(out_SymbolInfo.Filename), CSSourceInfoGetPath(Symbol));
-				FCStringAnsi::Strcpy(out_SymbolInfo.FunctionName, UE_ARRAY_COUNT(out_SymbolInfo.FunctionName), CSSymbolGetName(CSSourceInfoGetSymbol(Symbol)));
+				FCStringAnsi::Strncpy(out_SymbolInfo.FunctionName, CSSymbolGetName(CSSourceInfoGetSymbol(Symbol)), UE_ARRAY_COUNT(out_SymbolInfo.FunctionName));
 				CSRange CodeRange = CSSourceInfoGetRange(Symbol);
 				out_SymbolInfo.SymbolDisplacement = (ProgramCounter - CodeRange.Location);
 				
@@ -467,7 +467,7 @@ bool FApplePlatformSymbolication::SymbolInfoForAddress(uint64 ProgramCounter, FP
 				if(!CSIsNull(Owner))
 				{
 					ANSICHAR const* DylibName = CSSymbolOwnerGetName(Owner);
-					FCStringAnsi::Strcpy(out_SymbolInfo.ModuleName, UE_ARRAY_COUNT(out_SymbolInfo.ModuleName), DylibName);
+					FCStringAnsi::Strncpy(out_SymbolInfo.ModuleName, DylibName, UE_ARRAY_COUNT(out_SymbolInfo.ModuleName));
 					
 					bOK = out_SymbolInfo.LineNumber != 0;
 				}
diff --git a/Engine/Source/Runtime/Core/Private/GenericPlatform/GenericPlatformCrashContext.cpp b/Engine/Source/Runtime/Core/Private/GenericPlatform/GenericPlatformCrashContext.cpp
--- a/Engine/Source/Runtime/Core/Private/GenericPlatform/GenericPlatformCrashContext.cpp
+++ b/Engine/Source/Runtime/Core/Private/GenericPlatform/GenericPlatformCrashContext.cpp
@@ -617,8 +617,8 @@ void FGenericCrashContext::CopySharedCrashContext(FSharedCrashContext& Dst)
 	Dst.EnabledPluginsNum = NCached::EnabledPluginsList.Num();
 	for (const FString& Plugin : NCached::EnabledPluginsList)
 	{
-		FCString::Strncat(DynamicDataPtr, *Plugin, CR_DYNAMIC_BUFFER_REMAIN);
-		FCString::Strncat(DynamicDataPtr, CR_PAIR_DELIM, CR_DYNAMIC_BUFFER_REMAIN);
+		FCString::StrncatTruncateDest(DynamicDataPtr, CR_DYNAMIC_BUFFER_REMAIN, *Plugin);
+		FCString::StrncatTruncateDest(DynamicDataPtr, CR_DYNAMIC_BUFFER_REMAIN, CR_PAIR_DELIM);
 	}
 	DynamicDataPtr += FCString::Strlen(DynamicDataPtr) + 1;
 #else
@@ -630,10 +630,10 @@ void FGenericCrashContext::CopySharedCrashContext(FSharedCrashContext& Dst)
 	Dst.EngineDataNum = NCached::EngineData.Num();
 	for (const TPair<FString, FString>& Pair : NCached::EngineData)
 	{
-		FCString::Strncat(DynamicDataPtr, *Pair.Key, CR_DYNAMIC_BUFFER_REMAIN);
-		FCString::Strncat(DynamicDataPtr, CR_PAIR_EQ, CR_DYNAMIC_BUFFER_REMAIN);
-		FCString::Strncat(DynamicDataPtr, *Pair.Value, CR_DYNAMIC_BUFFER_REMAIN);
-		FCString::Strncat(DynamicDataPtr, CR_PAIR_DELIM, CR_DYNAMIC_BUFFER_REMAIN);
+		FCString::StrncatTruncateDest(DynamicDataPtr, CR_DYNAMIC_BUFFER_REMAIN, *Pair.Key);
+		FCString::StrncatTruncateDest(DynamicDataPtr, CR_DYNAMIC_BUFFER_REMAIN, CR_PAIR_EQ);
+		FCString::StrncatTruncateDest(DynamicDataPtr, CR_DYNAMIC_BUFFER_REMAIN, *Pair.Value);
+		FCString::StrncatTruncateDest(DynamicDataPtr, CR_DYNAMIC_BUFFER_REMAIN, CR_PAIR_DELIM);
 	}
 	DynamicDataPtr += FCString::Strlen(DynamicDataPtr) + 1;
 
@@ -641,10 +641,10 @@ void FGenericCrashContext::CopySharedCrashContext(FSharedCrashContext& Dst)
 	Dst.GameDataNum = NCached::GameData.Num();
 	for (const TPair<FString, FString>& Pair : NCached::GameData)
 	{
-		FCString::Strncat(DynamicDataPtr, *Pair.Key, CR_DYNAMIC_BUFFER_REMAIN);
-		FCString::Strncat(DynamicDataPtr, CR_PAIR_EQ, CR_DYNAMIC_BUFFER_REMAIN);
-		FCString::Strncat(DynamicDataPtr, *Pair.Value, CR_DYNAMIC_BUFFER_REMAIN);
-		FCString::Strncat(DynamicDataPtr, CR_PAIR_DELIM, CR_DYNAMIC_BUFFER_REMAIN);
+		FCString::StrncatTruncateDest(DynamicDataPtr, CR_DYNAMIC_BUFFER_REMAIN, *Pair.Key);
+		FCString::StrncatTruncateDest(DynamicDataPtr, CR_DYNAMIC_BUFFER_REMAIN, CR_PAIR_EQ);
+		FCString::StrncatTruncateDest(DynamicDataPtr, CR_DYNAMIC_BUFFER_REMAIN, *Pair.Value);
+		FCString::StrncatTruncateDest(DynamicDataPtr, CR_DYNAMIC_BUFFER_REMAIN, CR_PAIR_DELIM);
 	}
 	DynamicDataPtr += FCString::Strlen(DynamicDataPtr) + 1;
 
diff --git a/Engine/Source/Runtime/Core/Private/GenericPlatform/GenericPlatformMemory.cpp b/Engine/Source/Runtime/Core/Private/GenericPlatform/GenericPlatformMemory.cpp
--- a/Engine/Source/Runtime/Core/Private/GenericPlatform/GenericPlatformMemory.cpp
+++ b/Engine/Source/Runtime/Core/Private/GenericPlatform/GenericPlatformMemory.cpp
@@ -533,7 +533,7 @@ FGenericPlatformMemory::FSharedMemoryRegion::FSharedMemoryRegion(const FString&
 	,	Address(InAddress)
 	,	Size(InSize)
 {
-	FCString::Strcpy(Name, sizeof(Name) - 1, *InName);
+	FCString::Strncpy(Name, *InName, UE_ARRAY_COUNT(Name));
 }
 
 FGenericPlatformMemory::FSharedMemoryRegion * FGenericPlatformMemory::MapNamedSharedMemoryRegion(const FString& Name, bool bCreate, uint32 AccessMode, SIZE_T Size)
diff --git a/Engine/Source/Runtime/Core/Private/GenericPlatform/GenericPlatformMisc.cpp b/Engine/Source/Runtime/Core/Private/GenericPlatform/GenericPlatformMisc.cpp
--- a/Engine/Source/Runtime/Core/Private/GenericPlatform/GenericPlatformMisc.cpp
+++ b/Engine/Source/Runtime/Core/Private/GenericPlatform/GenericPlatformMisc.cpp
@@ -968,7 +968,7 @@ const TCHAR* FGenericPlatformMisc::GetSystemErrorMessage(TCHAR* OutBuffer, int32
 	const TCHAR* Message = TEXT("No system errors available on this platform.");
 	check(OutBuffer && BufferCount > 80);
 	Error = 0;
-	FCString::Strcpy(OutBuffer, BufferCount - 1, Message);
+	FCString::Strncpy(OutBuffer, Message, BufferCount);
 	return OutBuffer;
 }
 
diff --git a/Engine/Source/Runtime/Core/Private/GenericPlatform/GenericPlatformOutputDevices.cpp b/Engine/Source/Runtime/Core/Private/GenericPlatform/GenericPlatformOutputDevices.cpp
--- a/Engine/Source/Runtime/Core/Private/GenericPlatform/GenericPlatformOutputDevices.cpp
+++ b/Engine/Source/Runtime/Core/Private/GenericPlatform/GenericPlatformOutputDevices.cpp
@@ -78,7 +78,7 @@ FString FGenericPlatformOutputDevices::GetAbsoluteLogFilename()
 
 	if (CachedAbsoluteFilename[0] == 0)
 	{
-		FCString::Strcpy(CachedAbsoluteFilename, UE_ARRAY_COUNT(CachedAbsoluteFilename), *FPaths::ProjectLogDir());
+		FCString::Strncpy(CachedAbsoluteFilename, *FPaths::ProjectLogDir(), UE_ARRAY_COUNT(CachedAbsoluteFilename));
 		FString LogFilename;
 		const bool bShouldStopOnSeparator = false;
 		if (!FParse::Value(FCommandLine::Get(), TEXT("LOG="), LogFilename, bShouldStopOnSeparator))
@@ -121,7 +121,7 @@ FString FGenericPlatformOutputDevices::GetAbsoluteLogFilename()
 			LogFilename += TEXT(".log");
 		}
 
-		FCString::Strcat(CachedAbsoluteFilename, UE_ARRAY_COUNT(CachedAbsoluteFilename) - FCString::Strlen(CachedAbsoluteFilename), *LogFilename);
+		FCString::StrncatTruncateDest(CachedAbsoluteFilename, UE_ARRAY_COUNT(CachedAbsoluteFilename), *LogFilename);
 	}
 
 	return CachedAbsoluteFilename;
diff --git a/Engine/Source/Runtime/Core/Private/GenericPlatform/GenericPlatformProcess.cpp b/Engine/Source/Runtime/Core/Private/GenericPlatform/GenericPlatformProcess.cpp
--- a/Engine/Source/Runtime/Core/Private/GenericPlatform/GenericPlatformProcess.cpp
+++ b/Engine/Source/Runtime/Core/Private/GenericPlatform/GenericPlatformProcess.cpp
@@ -647,7 +647,7 @@ FGenericPlatformProcess::FSemaphore::FSemaphore(const FString& InName)
 
 FGenericPlatformProcess::FSemaphore::FSemaphore(const TCHAR* InName) 
 {
-	FCString::Strcpy(Name, sizeof(Name) - 1, InName);
+	FCString::Strncpy(Name, InName, UE_ARRAY_COUNT(Name));
 }
 
 FGenericPlatformProcess::FSemaphore* FGenericPlatformProcess::NewInterprocessSynchObject(const FString& Name, bool bCreate, uint32 MaxLocks)
diff --git a/Engine/Source/Runtime/Core/Private/GenericPlatform/GenericPlatformStackWalk.cpp b/Engine/Source/Runtime/Core/Private/GenericPlatform/GenericPlatformStackWalk.cpp
--- a/Engine/Source/Runtime/Core/Private/GenericPlatform/GenericPlatformStackWalk.cpp
+++ b/Engine/Source/Runtime/Core/Private/GenericPlatform/GenericPlatformStackWalk.cpp
@@ -85,25 +85,25 @@ bool FGenericPlatformStackWalk::SymbolInfoToHumanReadableString( const FProgramC
 		// Start with address
 		ANSICHAR PCAddress[MAX_TEMP_SPRINTF] = { 0 };
 		FCStringAnsi::Snprintf(PCAddress, MAX_TEMP_SPRINTF, "0x%016llx ", SymbolInfo.ProgramCounter);
-		FCStringAnsi::Strncat(StackLine, PCAddress, MAX_SPRINTF);
+		FCStringAnsi::StrncatTruncateDest(StackLine, MAX_SPRINTF, PCAddress);
 		
 		// Module if it's present
 		const bool bHasValidModuleName = FCStringAnsi::Strlen(StrippedModuleName) > 0;
 		if (bHasValidModuleName)
 		{
-			FCStringAnsi::Strncat(StackLine, StrippedModuleName, MAX_SPRINTF);
-			FCStringAnsi::Strncat(StackLine, "!", MAX_SPRINTF);
+			FCStringAnsi::StrncatTruncateDest(StackLine, MAX_SPRINTF, StrippedModuleName);
+			FCStringAnsi::StrncatTruncateDest(StackLine, MAX_SPRINTF, "!");
 		}
 
 		// Function if it's available, unknown if it's not
 		const bool bHasValidFunctionName = FCStringAnsi::Strlen( SymbolInfo.FunctionName ) > 0;
 		if( bHasValidFunctionName )
 		{
-			FCStringAnsi::Strncat(StackLine, SymbolInfo.FunctionName, MAX_SPRINTF);
+			FCStringAnsi::StrncatTruncateDest(StackLine, MAX_SPRINTF, SymbolInfo.FunctionName);
 		}
 		else
 		{
-			FCStringAnsi::Strncat(StackLine, "UnknownFunction", MAX_SPRINTF);
+			FCStringAnsi::StrncatTruncateDest(StackLine, MAX_SPRINTF, "UnknownFunction");
 		}
 
 		// file info
@@ -112,15 +112,15 @@ bool FGenericPlatformStackWalk::SymbolInfoToHumanReadableString( const FProgramC
 		{
 			ANSICHAR FilenameAndLineNumber[MAX_TEMP_SPRINTF] = {0};
 			FCStringAnsi::Snprintf( FilenameAndLineNumber, MAX_TEMP_SPRINTF, " [%s:%i]", SymbolInfo.Filename, SymbolInfo.LineNumber );
-			FCStringAnsi::Strncat(StackLine, FilenameAndLineNumber, MAX_SPRINTF);
+			FCStringAnsi::StrncatTruncateDest(StackLine, MAX_SPRINTF, FilenameAndLineNumber);
 		}
 		else
 		{
-			FCStringAnsi::Strncat(StackLine, " []", MAX_SPRINTF);
+			FCStringAnsi::StrncatTruncateDest(StackLine, MAX_SPRINTF, " []");
 		}
 
 		// Append the stack line.
-		FCStringAnsi::Strncat(HumanReadableString, StackLine, (int32)HumanReadableStringSize);
+		FCStringAnsi::StrncatTruncateDest(HumanReadableString, (int32)HumanReadableStringSize, StackLine);
 
 		// Return true, if we have a valid function name.
 		return bHasValidFunctionName;
@@ -196,7 +196,7 @@ void FGenericPlatformStackWalk::StackWalkAndDump( ANSICHAR* HumanReadableString,
 	while( CurrentDepth < Depth )
 	{
 		FPlatformStackWalk::ProgramCounterToHumanReadableString( CurrentDepth, StackTrace[CurrentDepth], HumanReadableString, HumanReadableStringSize, reinterpret_cast< FGenericCrashContext* >( Context ) );
-		FCStringAnsi::Strncat(HumanReadableString, LINE_TERMINATOR_ANSI, (int32)HumanReadableStringSize);
+		FCStringAnsi::StrncatTruncateDest(HumanReadableString, (int32)HumanReadableStringSize, LINE_TERMINATOR_ANSI);
 		CurrentDepth++;
 	}
 }
@@ -224,7 +224,7 @@ void FGenericPlatformStackWalk::StackWalkAndDump( ANSICHAR* HumanReadableString,
 	for (; CurrentDepth < Depth; CurrentDepth++)
 	{
 		FPlatformStackWalk::ProgramCounterToHumanReadableString( CurrentDepth, StackTrace[CurrentDepth], HumanReadableString, HumanReadableStringSize, reinterpret_cast< FGenericCrashContext* >( Context ) );
-		FCStringAnsi::Strncat(HumanReadableString, LINE_TERMINATOR_ANSI, (int32)HumanReadableStringSize);
+		FCStringAnsi::StrncatTruncateDest(HumanReadableString, (int32)HumanReadableStringSize, LINE_TERMINATOR_ANSI);
 	}
 }
 
@@ -298,7 +298,7 @@ void FGenericPlatformStackWalk::CopyVirtualPathToLocal(char* Dest, int32 DestCap
 			int32 LocalPathLen = int32(LocalPathEnd - LocalPath);
 
 			FCStringAnsi::Strncpy(Dest, LocalPath, FMath::Min(DestCapacity, LocalPathLen+1));
-			FCStringAnsi::Strncat(Dest, Source + VirtualPathLen, DestCapacity);
+			FCStringAnsi::StrncatTruncateDest(Dest, DestCapacity, Source + VirtualPathLen);
 			return;
 		}
 
diff --git a/Engine/Source/Runtime/Core/Private/GenericPlatform/GenericWidePlatformString.cpp b/Engine/Source/Runtime/Core/Private/GenericPlatform/GenericWidePlatformString.cpp
--- a/Engine/Source/Runtime/Core/Private/GenericPlatform/GenericWidePlatformString.cpp
+++ b/Engine/Source/Runtime/Core/Private/GenericPlatform/GenericWidePlatformString.cpp
@@ -16,11 +16,11 @@
 
 DEFINE_LOG_CATEGORY_STATIC(LogStandardPlatformString, Log, All);
 
-WIDECHAR* FGenericWidePlatformString::Strcpy(WIDECHAR* Dest, SIZE_T DestCount, const WIDECHAR* Src)
+WIDECHAR* FGenericWidePlatformString::Strcpy(WIDECHAR* Dest, const WIDECHAR* Src)
 {
 	TCHAR *BufPtr = Dest;
 	
-	while (*Src && --DestCount)
+	while (*Src)
 	{
 		*BufPtr++ = *Src++;
 	}
@@ -59,27 +59,35 @@ WIDECHAR* FGenericWidePlatformString::Strncpy(WIDECHAR* Dest, const WIDECHAR* Sr
 	return Dest;
 }
 
-WIDECHAR* FGenericWidePlatformString::Strcat(WIDECHAR* Dest, SIZE_T DestCount, const WIDECHAR* Src)
+WIDECHAR* FGenericWidePlatformString::Strcat(WIDECHAR* Dest, const WIDECHAR* Src)
 {
-	TCHAR *String = Dest;
-	
-	while (*String != 0 && DestCount > 1)
+	if (!*Src)
 	{
-		String++;
-		// remove how much we can copy in the lower loop
-		DestCount--;
+		return Dest;
 	}
-	
-	while (*Src != 0 && DestCount > 1)
+	TCHAR *NewDest = Dest + Strlen(Dest);
+	do
 	{
-		*String++ = *Src++;
-		DestCount--;
+		*NewDest++ = *Src++;
+	} while (*Src != 0);
+	*NewDest = 0;
+	return Dest;
+}
+
+WIDECHAR* FGenericWidePlatformString::Strncat(WIDECHAR* Dest, const WIDECHAR* Src, SIZE_T SrcLen)
+{
+	if (!SrcLen || !*Src)
+	{
+		return Dest;
 	}
-	
-//	// rewind 1 for trailing zero
-//	String--;
-	*String = 0;
-	
+	TCHAR* NewDest = Dest + Strlen(Dest);
+	SIZE_T AppendedCount = 0;
+	do
+	{
+		*NewDest++ = *Src++;
+		++AppendedCount;
+	} while (AppendedCount < SrcLen && *Src);
+	*NewDest = 0;
 	return Dest;
 }
 
diff --git a/Engine/Source/Runtime/Core/Private/HAL/MallocCallstackHandler.cpp b/Engine/Source/Runtime/Core/Private/HAL/MallocCallstackHandler.cpp
--- a/Engine/Source/Runtime/Core/Private/HAL/MallocCallstackHandler.cpp
+++ b/Engine/Source/Runtime/Core/Private/HAL/MallocCallstackHandler.cpp
@@ -194,7 +194,7 @@ FORCENOINLINE void FMallocCallstackHandler::DumpStackTraceToLog(int32 StackIndex
 		while (CurrentDepth < MaxCallStackDepth && CallStackInfoArray[StackIndex].FramePointers[CurrentDepth] != 0)
 		{
 			FPlatformStackWalk::ProgramCounterToHumanReadableString(CurrentDepth, CallStackInfoArray[StackIndex].FramePointers[CurrentDepth], &StackTraceString[0], StackTraceStringSize, reinterpret_cast<FGenericCrashContext*>(0));
-			FCStringAnsi::Strncat(StackTraceString, LINE_TERMINATOR_ANSI, StackTraceStringSize);
+			FCStringAnsi::StrncatTruncateDest(StackTraceString, StackTraceStringSize, LINE_TERMINATOR_ANSI);
 			CurrentDepth++;
 		}
 	}
diff --git a/Engine/Source/Runtime/Core/Private/IO/IoStore.cpp b/Engine/Source/Runtime/Core/Private/IO/IoStore.cpp
--- a/Engine/Source/Runtime/Core/Private/IO/IoStore.cpp
+++ b/Engine/Source/Runtime/Core/Private/IO/IoStore.cpp
@@ -1579,7 +1579,7 @@ TIoStatusOr<uint64> FIoStoreTocResource::Write(
 			continue;
 		}
 		FMemory::Memzero(AnsiMethodName, FIoStoreTocResource::CompressionMethodNameLen);
-		FCStringAnsi::Strcpy(AnsiMethodName, FIoStoreTocResource::CompressionMethodNameLen, TCHAR_TO_ANSI(*MethodName.ToString()));
+		FCStringAnsi::Strncpy(AnsiMethodName, TCHAR_TO_ANSI(*MethodName.ToString()), FIoStoreTocResource::CompressionMethodNameLen);
 
 		if (!TocFileHandle->Write(reinterpret_cast<const uint8*>(AnsiMethodName), FIoStoreTocResource::CompressionMethodNameLen))
 		{
diff --git a/Engine/Source/Runtime/Core/Private/IOS/IOSPlatformCrashContext.cpp b/Engine/Source/Runtime/Core/Private/IOS/IOSPlatformCrashContext.cpp
--- a/Engine/Source/Runtime/Core/Private/IOS/IOSPlatformCrashContext.cpp
+++ b/Engine/Source/Runtime/Core/Private/IOS/IOSPlatformCrashContext.cpp
@@ -131,7 +131,7 @@ void FIOSCrashContext::GenerateInfoInFolder(char const* const InfoFolder, bool b
     {
         char FilePath[PATH_MAX] = {};
         FCStringAnsi::Strncpy(FilePath, CrashInfoFolder, PATH_MAX);
-        FCStringAnsi::Strcat(FilePath, PATH_MAX, "/report.wer");
+        FCStringAnsi::StrncatTruncateDest(FilePath, PATH_MAX, "/report.wer");
         int ReportFile = open(FilePath, O_CREAT|O_WRONLY, 0766);
         if (ReportFile != -1)
         {
@@ -148,12 +148,12 @@ void FIOSCrashContext::GenerateInfoInFolder(char const* const InfoFolder, bool b
                 
         // generate "minidump" (Apple crash log format)
         FCStringAnsi::Strncpy(FilePath, CrashInfoFolder, PATH_MAX);
-        FCStringAnsi::Strcat(FilePath, PATH_MAX, "/minidump.dmp");
+        FCStringAnsi::StrncatTruncateDest(FilePath, PATH_MAX, "/minidump.dmp");
         CopyMinidump(FilePath, GIOSAppInfo.PLCrashReportPath);
         
         // generate "info.txt" custom data for our server
         FCStringAnsi::Strncpy(FilePath, CrashInfoFolder, PATH_MAX);
-        FCStringAnsi::Strcat(FilePath, PATH_MAX, "/info.txt");
+        FCStringAnsi::StrncatTruncateDest(FilePath, PATH_MAX, "/info.txt");
         ReportFile = open(FilePath, O_CREAT|O_WRONLY, 0766);
         if (ReportFile != -1)
         {
@@ -179,15 +179,15 @@ void FIOSCrashContext::GenerateInfoInFolder(char const* const InfoFolder, bool b
         
         // Introduces a new runtime crash context. Will replace all Windows related crash reporting.
         FCStringAnsi::Strncpy(FilePath, CrashInfoFolder, PATH_MAX);
-        FCStringAnsi::Strcat(FilePath, PATH_MAX, "/" );
-        FCStringAnsi::Strcat(FilePath, PATH_MAX, FGenericCrashContext::CrashContextRuntimeXMLNameA );
+        FCStringAnsi::StrncatTruncateDest(FilePath, PATH_MAX, "/" );
+        FCStringAnsi::StrncatTruncateDest(FilePath, PATH_MAX, FGenericCrashContext::CrashContextRuntimeXMLNameA );
         SerializeAsXML(*FString(FilePath));
         
         // copy log
         FCStringAnsi::Strncpy(FilePath, CrashInfoFolder, PATH_MAX);
-        FCStringAnsi::Strcat(FilePath, PATH_MAX, "/");
-        FCStringAnsi::Strcat(FilePath, PATH_MAX, (!GIOSAppInfo.AppName.IsEmpty() ? GIOSAppInfo.AppNameUTF8 : "Unreal"));
-        FCStringAnsi::Strcat(FilePath, PATH_MAX, ".log");
+        FCStringAnsi::StrncatTruncateDest(FilePath, PATH_MAX, "/");
+        FCStringAnsi::StrncatTruncateDest(FilePath, PATH_MAX, (!GIOSAppInfo.AppName.IsEmpty() ? GIOSAppInfo.AppNameUTF8 : "Unreal"));
+        FCStringAnsi::StrncatTruncateDest(FilePath, PATH_MAX, ".log");
         
         int LogSrc = open(GIOSAppInfo.AppLogPath, O_RDONLY);
         int LogDst = open(FilePath, O_CREAT|O_WRONLY, 0766);
@@ -201,8 +201,8 @@ void FIOSCrashContext::GenerateInfoInFolder(char const* const InfoFolder, bool b
         
         // If present, include the crash report config file to pass config values to the CRC
         FCStringAnsi::Strncpy(FilePath, CrashInfoFolder, PATH_MAX);
-        FCStringAnsi::Strcat(FilePath, PATH_MAX, "/");
-        FCStringAnsi::Strcat(FilePath, PATH_MAX, FGenericCrashContext::CrashConfigFileNameA);
+        FCStringAnsi::StrncatTruncateDest(FilePath, PATH_MAX, "/");
+        FCStringAnsi::StrncatTruncateDest(FilePath, PATH_MAX, FGenericCrashContext::CrashConfigFileNameA);
         int ConfigSrc = open(TCHAR_TO_ANSI(GetCrashConfigFilePath()), O_RDONLY);
         int ConfigDst = open(FilePath, O_CREAT | O_WRONLY, 0766);
         
@@ -229,15 +229,15 @@ void FIOSCrashContext::GenerateCrashInfo() const
     // create a crash-specific directory
     char CrashInfoFolder[PATH_MAX] = {};
     FCStringAnsi::Strncpy(CrashInfoFolder, GIOSAppInfo.CrashReportPath, PATH_MAX);
-    FCStringAnsi::Strcat(CrashInfoFolder, PATH_MAX, "/CrashReport-UE-");
-    FCStringAnsi::Strcat(CrashInfoFolder, PATH_MAX, GIOSAppInfo.AppNameUTF8);
-    FCStringAnsi::Strcat(CrashInfoFolder, PATH_MAX, "-pid-");
-    FCStringAnsi::Strcat(CrashInfoFolder, PATH_MAX, ItoANSI(getpid(), 10));
-    FCStringAnsi::Strcat(CrashInfoFolder, PATH_MAX, "-");
-    FCStringAnsi::Strcat(CrashInfoFolder, PATH_MAX, ItoANSI(GIOSAppInfo.RunUUID.A, 16));
-    FCStringAnsi::Strcat(CrashInfoFolder, PATH_MAX, ItoANSI(GIOSAppInfo.RunUUID.B, 16));
-    FCStringAnsi::Strcat(CrashInfoFolder, PATH_MAX, ItoANSI(GIOSAppInfo.RunUUID.C, 16));
-    FCStringAnsi::Strcat(CrashInfoFolder, PATH_MAX, ItoANSI(GIOSAppInfo.RunUUID.D, 16));
+    FCStringAnsi::StrncatTruncateDest(CrashInfoFolder, PATH_MAX, "/CrashReport-UE-");
+    FCStringAnsi::StrncatTruncateDest(CrashInfoFolder, PATH_MAX, GIOSAppInfo.AppNameUTF8);
+    FCStringAnsi::StrncatTruncateDest(CrashInfoFolder, PATH_MAX, "-pid-");
+    FCStringAnsi::StrncatTruncateDest(CrashInfoFolder, PATH_MAX, ItoANSI(getpid(), 10));
+    FCStringAnsi::StrncatTruncateDest(CrashInfoFolder, PATH_MAX, "-");
+    FCStringAnsi::StrncatTruncateDest(CrashInfoFolder, PATH_MAX, ItoANSI(GIOSAppInfo.RunUUID.A, 16));
+    FCStringAnsi::StrncatTruncateDest(CrashInfoFolder, PATH_MAX, ItoANSI(GIOSAppInfo.RunUUID.B, 16));
+    FCStringAnsi::StrncatTruncateDest(CrashInfoFolder, PATH_MAX, ItoANSI(GIOSAppInfo.RunUUID.C, 16));
+    FCStringAnsi::StrncatTruncateDest(CrashInfoFolder, PATH_MAX, ItoANSI(GIOSAppInfo.RunUUID.D, 16));
         
     const bool bIsEnsure = false;
     GenerateInfoInFolder(CrashInfoFolder, bIsEnsure);
@@ -327,7 +327,7 @@ void FIOSApplicationInfo::Init()
 	SCOPED_AUTORELEASE_POOL;
 	
 	AppName = FApp::GetProjectName();
-	FCStringAnsi::Strcpy(AppNameUTF8, PATH_MAX+1, TCHAR_TO_UTF8(*AppName));
+	FCStringAnsi::Strncpy(AppNameUTF8, TCHAR_TO_UTF8(*AppName), PATH_MAX + 1);
 	
 	ExecutableName = FPlatformProcess::ExecutableName();
 	
@@ -345,7 +345,7 @@ void FIOSApplicationInfo::Init()
 	
 	OSXVersion = [[NSProcessInfo processInfo] operatingSystemVersion];
 	OSVersion = FString::Printf(TEXT("%ld.%ld.%ld"), OSXVersion.majorVersion, OSXVersion.minorVersion, OSXVersion.patchVersion);
-	FCStringAnsi::Strcpy(OSVersionUTF8, PATH_MAX+1, TCHAR_TO_UTF8(*OSVersion));
+	FCStringAnsi::Strncpy(OSVersionUTF8, TCHAR_TO_UTF8(*OSVersion), PATH_MAX + 1);
 	
 	// macOS build number is only accessible on non-sandboxed applications as it resides outside the accessible sandbox
 	if(!bIsSandboxed)
@@ -384,13 +384,13 @@ void FIOSApplicationInfo::Init()
 	CommandLine = FCommandLine::Get();
 	FString LogPath = FGenericPlatformOutputDevices::GetAbsoluteLogFilename();
 	LogPath = IFileManager::Get().ConvertToAbsolutePathForExternalAppForWrite(*LogPath);
-	FCStringAnsi::Strcpy(AppLogPath, PATH_MAX + 1, TCHAR_TO_UTF8(*LogPath));
+	FCStringAnsi::Strncpy(AppLogPath, TCHAR_TO_UTF8(*LogPath), PATH_MAX + 1);
 	
 	// Cache & create the crash report folder.
 	FString ReportPath = FPaths::ConvertRelativePathToFull(FString::Printf(TEXT("%s"), *(FPaths::GameAgnosticSavedDir() / TEXT("Crashes"))));
 	IFileManager::Get().MakeDirectory(*ReportPath, true);
 	ReportPath = IFileManager::Get().ConvertToAbsolutePathForExternalAppForWrite(*ReportPath);
-	FCStringAnsi::Strcpy(CrashReportPath, PATH_MAX+1, TCHAR_TO_UTF8(*ReportPath));
+	FCStringAnsi::Strncpy(CrashReportPath, TCHAR_TO_UTF8(*ReportPath), PATH_MAX + 1);
 	
 	NSString* PLCrashReportFile = [TemporaryCrashReportFolder().GetNSString() stringByAppendingPathComponent:TemporaryCrashReportName().GetNSString()];
 	[PLCrashReportFile getCString:PLCrashReportPath maxLength:PATH_MAX encoding:NSUTF8StringEncoding];
diff --git a/Engine/Source/Runtime/Core/Private/IOS/IOSPlatformMisc.cpp b/Engine/Source/Runtime/Core/Private/IOS/IOSPlatformMisc.cpp
--- a/Engine/Source/Runtime/Core/Private/IOS/IOSPlatformMisc.cpp
+++ b/Engine/Source/Runtime/Core/Private/IOS/IOSPlatformMisc.cpp
@@ -2179,15 +2179,15 @@ FString FIOSCrashContext::CreateCrashFolder() const
 	// create a crash-specific directory
 	char CrashInfoFolder[PATH_MAX] = {};
 	FCStringAnsi::Strncpy(CrashInfoFolder, GIOSAppInfo.CrashReportPath, PATH_MAX);
-	FCStringAnsi::Strcat(CrashInfoFolder, PATH_MAX, "/CrashReport-UE-");
-	FCStringAnsi::Strcat(CrashInfoFolder, PATH_MAX, GIOSAppInfo.AppNameUTF8);
-	FCStringAnsi::Strcat(CrashInfoFolder, PATH_MAX, "-pid-");
-	FCStringAnsi::Strcat(CrashInfoFolder, PATH_MAX, ItoANSI(getpid(), 10));
-	FCStringAnsi::Strcat(CrashInfoFolder, PATH_MAX, "-");
-	FCStringAnsi::Strcat(CrashInfoFolder, PATH_MAX, ItoANSI(GIOSAppInfo.RunUUID.A, 16));
-	FCStringAnsi::Strcat(CrashInfoFolder, PATH_MAX, ItoANSI(GIOSAppInfo.RunUUID.B, 16));
-	FCStringAnsi::Strcat(CrashInfoFolder, PATH_MAX, ItoANSI(GIOSAppInfo.RunUUID.C, 16));
-	FCStringAnsi::Strcat(CrashInfoFolder, PATH_MAX, ItoANSI(GIOSAppInfo.RunUUID.D, 16));
+	FCStringAnsi::StrncatTruncateDest(CrashInfoFolder, PATH_MAX, "/CrashReport-UE-");
+	FCStringAnsi::StrncatTruncateDest(CrashInfoFolder, PATH_MAX, GIOSAppInfo.AppNameUTF8);
+	FCStringAnsi::StrncatTruncateDest(CrashInfoFolder, PATH_MAX, "-pid-");
+	FCStringAnsi::StrncatTruncateDest(CrashInfoFolder, PATH_MAX, ItoANSI(getpid(), 10));
+	FCStringAnsi::StrncatTruncateDest(CrashInfoFolder, PATH_MAX, "-");
+	FCStringAnsi::StrncatTruncateDest(CrashInfoFolder, PATH_MAX, ItoANSI(GIOSAppInfo.RunUUID.A, 16));
+	FCStringAnsi::StrncatTruncateDest(CrashInfoFolder, PATH_MAX, ItoANSI(GIOSAppInfo.RunUUID.B, 16));
+	FCStringAnsi::StrncatTruncateDest(CrashInfoFolder, PATH_MAX, ItoANSI(GIOSAppInfo.RunUUID.C, 16));
+	FCStringAnsi::StrncatTruncateDest(CrashInfoFolder, PATH_MAX, ItoANSI(GIOSAppInfo.RunUUID.D, 16));
 	
 	return FString(ANSI_TO_TCHAR(CrashInfoFolder));
 }
diff --git a/Engine/Source/Runtime/Core/Private/Linux/LinuxPlatformProcess.cpp b/Engine/Source/Runtime/Core/Private/Linux/LinuxPlatformProcess.cpp
--- a/Engine/Source/Runtime/Core/Private/Linux/LinuxPlatformProcess.cpp
+++ b/Engine/Source/Runtime/Core/Private/Linux/LinuxPlatformProcess.cpp
@@ -25,9 +25,8 @@ const TCHAR* FLinuxPlatformProcess::BaseDir()
 		}
 		SelfPath[UE_ARRAY_COUNT(SelfPath) - 1] = 0;
 
-		FCString::Strncpy(CachedResult, UTF8_TO_TCHAR(dirname(SelfPath)), UE_ARRAY_COUNT(CachedResult) - 1);
-		CachedResult[UE_ARRAY_COUNT(CachedResult) - 1] = 0;
-		FCString::Strncat(CachedResult, TEXT("/"), UE_ARRAY_COUNT(CachedResult) - 1);
+		FCString::Strncpy(CachedResult, UTF8_TO_TCHAR(dirname(SelfPath)), UE_ARRAY_COUNT(CachedResult));
+		FCString::StrncatTruncateDest(CachedResult, UE_ARRAY_COUNT(CachedResult), TEXT("/"));
 
 #ifdef UE_RELATIVE_BASE_DIR
 		FString CollapseResult(CachedResult);
diff --git a/Engine/Source/Runtime/Core/Private/Mac/MacErrorOutputDevice.cpp b/Engine/Source/Runtime/Core/Private/Mac/MacErrorOutputDevice.cpp
--- a/Engine/Source/Runtime/Core/Private/Mac/MacErrorOutputDevice.cpp
+++ b/Engine/Source/Runtime/Core/Private/Mac/MacErrorOutputDevice.cpp
@@ -40,7 +40,7 @@ void FMacErrorOutputDevice::Serialize( const TCHAR* Msg, ELogVerbosity::Type Ver
 		if (FCString::Strncmp(GErrorHist, Msg, FMath::Min<int32>(UE_ARRAY_COUNT(GErrorHist), FCString::Strlen(Msg))))
 		{
 			FCString::Strncpy(GErrorHist, Msg, UE_ARRAY_COUNT(GErrorHist) - 5);
-			FCString::Strncat(GErrorHist, TEXT("\r\n\r\n"), UE_ARRAY_COUNT(GErrorHist) - 1);
+			FCString::StrncatTruncateDest(GErrorHist, UE_ARRAY_COUNT(GErrorHist), TEXT("\r\n\r\n"));
 		}
 		ErrorPos = FCString::Strlen(GErrorHist);
 	}
diff --git a/Engine/Source/Runtime/Core/Private/Mac/MacPlatformMisc.cpp b/Engine/Source/Runtime/Core/Private/Mac/MacPlatformMisc.cpp
--- a/Engine/Source/Runtime/Core/Private/Mac/MacPlatformMisc.cpp
+++ b/Engine/Source/Runtime/Core/Private/Mac/MacPlatformMisc.cpp
@@ -228,7 +228,7 @@ struct FMacApplicationInfo
 		}
 		
 		AppName = FApp::GetProjectName();
-		FCStringAnsi::Strcpy(AppNameUTF8, PATH_MAX+1, TCHAR_TO_UTF8(*AppName));
+		FCStringAnsi::Strncpy(AppNameUTF8, TCHAR_TO_UTF8(*AppName), PATH_MAX + 1);
 		
 		ExecutableName = FPlatformProcess::ExecutableName();
 		
@@ -249,7 +249,7 @@ struct FMacApplicationInfo
 		RunUUID = RunGUID();
 		
 		MacPlatformGetOSProductVersion(OSVersion);
-		FCStringAnsi::Strcpy(OSVersionUTF8, PATH_MAX+1, TCHAR_TO_UTF8(*OSVersion));
+		FCStringAnsi::Strncpy(OSVersionUTF8, TCHAR_TO_UTF8(*OSVersion), PATH_MAX + 1);
 		
 		MacPlatformGetOSVersion(OSBuild);
 
@@ -334,20 +334,20 @@ struct FMacApplicationInfo
 		// Use the log file specified on the commandline if there is one
 		CommandLine = FCommandLine::Get();
 		FString LogPath = FGenericPlatformOutputDevices::GetAbsoluteLogFilename();
-		FCStringAnsi::Strcpy(AppLogPath, PATH_MAX + 1, TCHAR_TO_UTF8(*LogPath));
+		FCStringAnsi::Strncpy(AppLogPath, TCHAR_TO_UTF8(*LogPath), PATH_MAX + 1);
 
 		FString UserCrashVideoPath = IFileManager::Get().ConvertToAbsolutePathForExternalAppForWrite(*CrashVideoPath);
-		FCStringAnsi::Strcpy(CrashReportVideo, PATH_MAX+1, TCHAR_TO_UTF8(*UserCrashVideoPath));
+		FCStringAnsi::Strncpy(CrashReportVideo, TCHAR_TO_UTF8(*UserCrashVideoPath), PATH_MAX + 1);
 		
 		// Cache & create the crash report folder.
 		FString ReportPath = FPaths::ConvertRelativePathToFull(FString::Printf(TEXT("%s"), *(FPaths::GameAgnosticSavedDir() / TEXT("Crashes"))));
-		FCStringAnsi::Strcpy(CrashReportPath, PATH_MAX+1, TCHAR_TO_UTF8(*ReportPath));
+		FCStringAnsi::Strncpy(CrashReportPath, TCHAR_TO_UTF8(*ReportPath), PATH_MAX + 1);
 #if WITH_EDITOR
 		FString ReportClient = FPaths::ConvertRelativePathToFull(FPlatformProcess::GenerateApplicationPath(TEXT("CrashReportClientEditor"), EBuildConfiguration::Development));
 #else
 		FString ReportClient = FPaths::ConvertRelativePathToFull(FPlatformProcess::GenerateApplicationPath(TEXT("CrashReportClient"), EBuildConfiguration::Development));
 #endif
-		FCStringAnsi::Strcpy(CrashReportClient, PATH_MAX+1, TCHAR_TO_UTF8(*ReportClient));
+		FCStringAnsi::Strncpy(CrashReportClient, TCHAR_TO_UTF8(*ReportClient), PATH_MAX + 1);
 		IFileManager::Get().MakeDirectory(*ReportPath, true);
 		
 		// Notification handler to check we are running from a battery - this only applies to MacBook's.
@@ -592,7 +592,7 @@ void FMacPlatformMisc::PlatformInit()
 
 	// At the point, the log system is booted and the log file is likely created and the actual pathname used set. (Until the file is opened for writing, the name is subject to change)
 	FString LogPath = FGenericPlatformOutputDevices::GetAbsoluteLogFilename();
-	FCStringAnsi::Strcpy(GMacAppInfo.AppLogPath, PATH_MAX + 1, TCHAR_TO_UTF8(*LogPath));
+	FCStringAnsi::Strncpy(GMacAppInfo.AppLogPath, TCHAR_TO_UTF8(*LogPath), PATH_MAX + 1);
 
 	const FPlatformMemoryConstants& MemoryConstants = FPlatformMemory::GetConstants();
 	UE_LOG(LogInit, Log, TEXT("CPU Page size=%i, Cores=%i, HT=%i"), MemoryConstants.PageSize, FPlatformMisc::NumberOfCores(), FPlatformMisc::NumberOfCoresIncludingHyperthreads() );
@@ -2196,12 +2196,12 @@ void FMacCrashContext::GenerateInfoInFolder(char const* const InfoFolder) const
 		
 		// generate "minidump" (Apple crash log format)
 		FCStringAnsi::Strncpy(FilePath, CrashInfoFolder, PATH_MAX);
-		FCStringAnsi::Strcat(FilePath, PATH_MAX, "/minidump.dmp");
+		FCStringAnsi::StrncatTruncateDest(FilePath, PATH_MAX, "/minidump.dmp");
 		CopyMinidump(FilePath, GMacAppInfo.PLCrashReportPath);
 		
 		// generate "info.txt" custom data for our server
 		FCStringAnsi::Strncpy(FilePath, CrashInfoFolder, PATH_MAX);
-		FCStringAnsi::Strcat(FilePath, PATH_MAX, "/info.txt");
+		FCStringAnsi::StrncatTruncateDest(FilePath, PATH_MAX, "/info.txt");
 		int ReportFile = open(FilePath, O_CREAT|O_WRONLY, 0766);
 		if (ReportFile != -1)
 		{
@@ -2227,15 +2227,15 @@ void FMacCrashContext::GenerateInfoInFolder(char const* const InfoFolder) const
 		
 		// Introduces a new runtime crash context. Will replace all Windows related crash reporting.
 		FCStringAnsi::Strncpy(FilePath, CrashInfoFolder, PATH_MAX);
-		FCStringAnsi::Strcat(FilePath, PATH_MAX, "/" );
-		FCStringAnsi::Strcat(FilePath, PATH_MAX, FGenericCrashContext::CrashContextRuntimeXMLNameA );
+		FCStringAnsi::StrncatTruncateDest(FilePath, PATH_MAX, "/" );
+		FCStringAnsi::StrncatTruncateDest(FilePath, PATH_MAX, FGenericCrashContext::CrashContextRuntimeXMLNameA );
 		SerializeAsXML( *FString(FilePath) );
 		
 		// copy log
 		FCStringAnsi::Strncpy(FilePath, CrashInfoFolder, PATH_MAX);
-		FCStringAnsi::Strcat(FilePath, PATH_MAX, "/");
-		FCStringAnsi::Strcat(FilePath, PATH_MAX, (!GMacAppInfo.AppName.IsEmpty() ? GMacAppInfo.AppNameUTF8 : "Unreal"));
-		FCStringAnsi::Strcat(FilePath, PATH_MAX, ".log");
+		FCStringAnsi::StrncatTruncateDest(FilePath, PATH_MAX, "/");
+		FCStringAnsi::StrncatTruncateDest(FilePath, PATH_MAX, (!GMacAppInfo.AppName.IsEmpty() ? GMacAppInfo.AppNameUTF8 : "Unreal"));
+		FCStringAnsi::StrncatTruncateDest(FilePath, PATH_MAX, ".log");
 		int LogSrc = open(GMacAppInfo.AppLogPath, O_RDONLY);
 		int LogDst = open(FilePath, O_CREAT|O_WRONLY, 0766);
 		
@@ -2248,8 +2248,8 @@ void FMacCrashContext::GenerateInfoInFolder(char const* const InfoFolder) const
 		
 		// If present, include the crash report config file to pass config values to the CRC
 		FCStringAnsi::Strncpy(FilePath, CrashInfoFolder, PATH_MAX);
-		FCStringAnsi::Strcat(FilePath, PATH_MAX, "/");
-		FCStringAnsi::Strcat(FilePath, PATH_MAX, FGenericCrashContext::CrashConfigFileNameA);
+		FCStringAnsi::StrncatTruncateDest(FilePath, PATH_MAX, "/");
+		FCStringAnsi::StrncatTruncateDest(FilePath, PATH_MAX, FGenericCrashContext::CrashConfigFileNameA);
 		int ConfigSrc = open(TCHAR_TO_ANSI(GetCrashConfigFilePath()), O_RDONLY);
 		int ConfigDst = open(FilePath, O_CREAT | O_WRONLY, 0766);
 
@@ -2271,7 +2271,7 @@ void FMacCrashContext::GenerateInfoInFolder(char const* const InfoFolder) const
 				if (strstr(DirEntry.d_name, ".gpuRestart"))
 				{
 					FCStringAnsi::Strncpy(FilePath, "/Library/Logs/DiagnosticReports/", PATH_MAX);
-					FCStringAnsi::Strcat(FilePath, PATH_MAX, DirEntry.d_name);
+					FCStringAnsi::StrncatTruncateDest(FilePath, PATH_MAX, DirEntry.d_name);
 					if (access(FilePath, R_OK|F_OK) == 0)
 					{
 						char const* SysLogHeader = "\nAppending GPU Restart Log: ";
@@ -2321,8 +2321,8 @@ void FMacCrashContext::GenerateInfoInFolder(char const* const InfoFolder) const
 		if ( access(GMacAppInfo.CrashReportVideo, R_OK|F_OK) == 0 )
 		{
 			FCStringAnsi::Strncpy(FilePath, CrashInfoFolder, PATH_MAX);
-			FCStringAnsi::Strcat(FilePath, PATH_MAX, "/");
-			FCStringAnsi::Strcat(FilePath, PATH_MAX, "CrashVideo.avi");
+			FCStringAnsi::StrncatTruncateDest(FilePath, PATH_MAX, "/");
+			FCStringAnsi::StrncatTruncateDest(FilePath, PATH_MAX, "CrashVideo.avi");
 			int VideoSrc = open(GMacAppInfo.CrashReportVideo, O_RDONLY);
 			int VideoDst = open(FilePath, O_CREAT|O_WRONLY, 0766);
 			
diff --git a/Engine/Source/Runtime/Core/Private/Mac/MacPlatformProcess.cpp b/Engine/Source/Runtime/Core/Private/Mac/MacPlatformProcess.cpp
--- a/Engine/Source/Runtime/Core/Private/Mac/MacPlatformProcess.cpp
+++ b/Engine/Source/Runtime/Core/Private/Mac/MacPlatformProcess.cpp
@@ -1042,8 +1042,8 @@ const TCHAR* FMacPlatformProcess::BaseDir()
 			}
 #endif
 
-			FCString::Strcpy(Result, MAC_MAX_PATH, *FString(BasePath));
-			FCString::Strcat(Result, TEXT("/"));
+			FCString::Strncpy(Result, *FString(BasePath), MAC_MAX_PATH);
+			FCString::StrncatTruncateDest(Result, MAC_MAX_PATH, TEXT("/"));
 		}
 	}
 	return Result;
diff --git a/Engine/Source/Runtime/Core/Private/Misc/AssertionMacros.cpp b/Engine/Source/Runtime/Core/Private/Misc/AssertionMacros.cpp
--- a/Engine/Source/Runtime/Core/Private/Misc/AssertionMacros.cpp
+++ b/Engine/Source/Runtime/Core/Private/Misc/AssertionMacros.cpp
@@ -497,7 +497,7 @@ FORCENOINLINE void FDebug::EnsureFailed(const ANSICHAR* Expr, const ANSICHAR* Fi
 			}
 
 			// Also append the stack trace
-			FCString::Strncat(ErrorMsg, ANSI_TO_TCHAR(StackTrace), UE_ARRAY_COUNT(ErrorMsg) - 1);
+			FCString::StrncatTruncateDest(ErrorMsg, UE_ARRAY_COUNT(ErrorMsg), ANSI_TO_TCHAR(StackTrace));
 			FMemory::SystemFree(StackTrace);
 
 			// Dump the error and flush the log.
diff --git a/Engine/Source/Runtime/Core/Private/Misc/CStringTest.cpp b/Engine/Source/Runtime/Core/Private/Misc/CStringTest.cpp
--- a/Engine/Source/Runtime/Core/Private/Misc/CStringTest.cpp
+++ b/Engine/Source/Runtime/Core/Private/Misc/CStringTest.cpp
@@ -225,4 +225,303 @@ bool FCStringStrnstrTest::RunTest(const FString& Parameters)
 	return true;
 }
 
+IMPLEMENT_SIMPLE_AUTOMATION_TEST(FCStringStrcpyTest, "System.Core.Misc.CString.Strcpy", EAutomationTestFlags_ApplicationContextMask | EAutomationTestFlags::EngineFilter)
+bool FCStringStrcpyTest::RunTest(const FString& Parameters)
+{
+	constexpr int32 BufferLen = 32;
+	WIDECHAR WideBuffer[BufferLen];
+	UTF8CHAR Utf8Buffer[BufferLen];
+	ANSICHAR AnsiBuffer[BufferLen];
+
+	const WIDECHAR* WideTest = WIDETEXT("12345");
+	const UTF8CHAR* Utf8Test = UTF8TEXT("12345");
+	const ANSICHAR* AnsiTest = "12345";
+	constexpr int32 TestLen = 5;
+
+	auto Reset = [&WideBuffer, &Utf8Buffer, &AnsiBuffer, BufferLen]()
+		{
+			for (int32 n = 0; n < BufferLen; ++n)
+			{
+				WideBuffer[n] = '%';
+				Utf8Buffer[n] = UTF8CHAR('%');
+				AnsiBuffer[n] = '%';
+			}
+		};
+
+	Reset();
+	FCStringWide::Strcpy(WideBuffer, WideTest);
+	FCStringUtf8::Strcpy(Utf8Buffer, Utf8Test);
+	FCStringAnsi::Strcpy(AnsiBuffer, AnsiTest);
+	this->TestTrue(TEXT("WideStrcpy"),
+		WideBuffer[TestLen] == 0 && FCStringWide::Strcmp(WideTest, WideBuffer) == 0 && WideBuffer[TestLen + 1] == '%');
+	this->TestTrue(TEXT("Utf8Strcpy"),
+		Utf8Buffer[TestLen] == 0 && FCStringUtf8::Strcmp(Utf8Test, Utf8Buffer) == 0 && Utf8Buffer[TestLen + 1] == '%');
+	this->TestTrue(TEXT("AnsiStrcpy"),
+		AnsiBuffer[TestLen] == 0 && FCStringAnsi::Strcmp(AnsiTest, AnsiBuffer) == 0 && AnsiBuffer[TestLen + 1] == '%');
+
+	Reset();
+	FCStringWide::Strncpy(WideBuffer, WideTest, TestLen + 10);
+	FCStringUtf8::Strncpy(Utf8Buffer, Utf8Test, TestLen + 10);
+	FCStringAnsi::Strncpy(AnsiBuffer, AnsiTest, TestLen + 10);
+	this->TestTrue(TEXT("WideStrncpyTestLenPlus10"),
+		WideBuffer[TestLen] == 0 && FCStringWide::Strcmp(WideTest, WideBuffer) == 0 && WideBuffer[TestLen + 10] == '%');
+	this->TestTrue(TEXT("Utf8StrncpyTestLenPlus10"),
+		Utf8Buffer[TestLen] == 0 && FCStringUtf8::Strcmp(Utf8Test, Utf8Buffer) == 0 && Utf8Buffer[TestLen + 10] == '%');
+	this->TestTrue(TEXT("AnsiStrncpyTestLenPlus10"),
+		AnsiBuffer[TestLen] == 0 && FCStringAnsi::Strcmp(AnsiTest, AnsiBuffer) == 0 && AnsiBuffer[TestLen + 10] == '%');
+
+	Reset();
+	FCStringWide::Strncpy(WideBuffer, WideTest, TestLen + 1);
+	FCStringUtf8::Strncpy(Utf8Buffer, Utf8Test, TestLen + 1);
+	FCStringAnsi::Strncpy(AnsiBuffer, AnsiTest, TestLen + 1);
+	this->TestTrue(TEXT("WideStrncpyTestLenPlus1"),
+		WideBuffer[TestLen] == 0 && FCStringWide::Strcmp(WideTest, WideBuffer) == 0 && WideBuffer[TestLen + 1] == '%');
+	this->TestTrue(TEXT("Utf8StrncpyTestLenPlus1"),
+		Utf8Buffer[TestLen] == 0 && FCStringUtf8::Strcmp(Utf8Test, Utf8Buffer) == 0 && Utf8Buffer[TestLen + 1] == '%');
+	this->TestTrue(TEXT("AnsiStrncpyTestLenPlus1"),
+		AnsiBuffer[TestLen] == 0 && FCStringAnsi::Strcmp(AnsiTest, AnsiBuffer) == 0 && AnsiBuffer[TestLen + 1] == '%');
+
+	Reset();
+	FCStringWide::Strncpy(WideBuffer, WideTest, TestLen);
+	FCStringUtf8::Strncpy(Utf8Buffer, Utf8Test, TestLen);
+	FCStringAnsi::Strncpy(AnsiBuffer, AnsiTest, TestLen);
+	this->TestTrue(TEXT("WideStrncpyTestLen"),
+		WideBuffer[TestLen - 1] == 0 && WideBuffer[TestLen] == '%' && FCStringWide::Strncmp(WideTest, WideBuffer, TestLen - 1) == 0 && WideBuffer[TestLen + 1] == '%');
+	this->TestTrue(TEXT("Utf8StrncpyTestLen"),
+		Utf8Buffer[TestLen - 1] == 0 && Utf8Buffer[TestLen] == '%' && FCStringUtf8::Strncmp(Utf8Test, Utf8Buffer, TestLen - 1) == 0 && Utf8Buffer[TestLen + 1] == '%');
+	this->TestTrue(TEXT("AnsiStrncpyTestLen"),
+		AnsiBuffer[TestLen - 1] == 0 && AnsiBuffer[TestLen] == '%' && FCStringAnsi::Strncmp(AnsiTest, AnsiBuffer, TestLen - 1) == 0 && AnsiBuffer[TestLen + 1] == '%');
+
+	Reset();
+	FCStringWide::Strncpy(WideBuffer, WideTest, TestLen - 1);
+	FCStringUtf8::Strncpy(Utf8Buffer, Utf8Test, TestLen - 1);
+	FCStringAnsi::Strncpy(AnsiBuffer, AnsiTest, TestLen - 1);
+	this->TestTrue(TEXT("WideStrncpyTestLenMinus1"),
+		WideBuffer[TestLen - 2] == 0 && WideBuffer[TestLen - 1] == '%' && FCStringWide::Strncmp(WideTest, WideBuffer, TestLen - 2) == 0 && WideBuffer[TestLen + 1] == '%');
+	this->TestTrue(TEXT("Utf8StrncpyTestLenMinus1"),
+		Utf8Buffer[TestLen - 2] == 0 && Utf8Buffer[TestLen - 1] == '%' && FCStringUtf8::Strncmp(Utf8Test, Utf8Buffer, TestLen - 2) == 0 && Utf8Buffer[TestLen + 1] == '%');
+	this->TestTrue(TEXT("AnsiStrncpyTestLenMinus1"),
+		AnsiBuffer[TestLen - 2] == 0 && AnsiBuffer[TestLen - 1] == '%' && FCStringAnsi::Strncmp(AnsiTest, AnsiBuffer, TestLen - 2) == 0 && AnsiBuffer[TestLen + 1] == '%');
+
+	Reset();
+	FCStringWide::Strncpy(WideBuffer, WideTest, 2);
+	FCStringUtf8::Strncpy(Utf8Buffer, Utf8Test, 2);
+	FCStringAnsi::Strncpy(AnsiBuffer, AnsiTest, 2);
+	this->TestTrue(TEXT("WideStrncpyTwoLen"),
+		WideBuffer[0] == WideTest[0] && WideBuffer[1] == 0 && WideBuffer[2] == '%' && WideBuffer[TestLen] == '%');
+	this->TestTrue(TEXT("Utf8StrncpyTwoLen"),
+		Utf8Buffer[0] == Utf8Test[0] && Utf8Buffer[1] == 0 && Utf8Buffer[2] == '%' && Utf8Buffer[TestLen] == '%');
+	this->TestTrue(TEXT("AnsiStrncpyTwoLen"),
+		AnsiBuffer[0] == AnsiTest[0] && AnsiBuffer[1] == 0 && AnsiBuffer[2] == '%' && AnsiBuffer[TestLen] == '%');
+
+	Reset();
+	FCStringWide::Strncpy(WideBuffer, WideTest, 1);
+	FCStringUtf8::Strncpy(Utf8Buffer, Utf8Test, 1);
+	FCStringAnsi::Strncpy(AnsiBuffer, AnsiTest, 1);
+	this->TestTrue(TEXT("WideStrncpyOneLen"),
+		WideBuffer[0] == 0 && WideBuffer[1] == '%' && WideBuffer[TestLen] == '%');
+	this->TestTrue(TEXT("Utf8StrncpyOneLen"),
+		Utf8Buffer[0] == 0 && Utf8Buffer[1] == '%' && Utf8Buffer[TestLen] == '%');
+	this->TestTrue(TEXT("AnsiStrncpyOneLen"),
+		AnsiBuffer[0] == 0 && AnsiBuffer[1] == '%' && AnsiBuffer[TestLen] == '%');
+
+	// ZeroLen Strncpy is undefined
+	return true;
+}
+
+IMPLEMENT_SIMPLE_AUTOMATION_TEST(FCStringStrcatTest, "System.Core.Misc.CString.Strcat", EAutomationTestFlags_ApplicationContextMask | EAutomationTestFlags::EngineFilter)
+bool FCStringStrcatTest::RunTest(const FString& Parameters)
+{
+	constexpr int32 BufferLen = 32;
+	WIDECHAR WideBuffer[BufferLen];
+	UTF8CHAR Utf8Buffer[BufferLen];
+	ANSICHAR AnsiBuffer[BufferLen];
+
+	const WIDECHAR* WidePrefix = WIDETEXT("ABCD");
+	const UTF8CHAR* Utf8Prefix = UTF8TEXT("ABCD");
+	const ANSICHAR* AnsiPrefix = "ABCD";
+	const WIDECHAR* WideTest = WIDETEXT("12345");
+	const UTF8CHAR* Utf8Test = UTF8TEXT("12345");
+	const ANSICHAR* AnsiTest = "12345";
+	const WIDECHAR* WidePrefixPlusTest = WIDETEXT("ABCD12345");
+	const UTF8CHAR* Utf8PrefixPlusTest = UTF8TEXT("ABCD12345");
+	const ANSICHAR* AnsiPrefixPlusTest = "ABCD12345";
+	constexpr int32 PrefixLen = 4;
+	constexpr int32 TestLen = 5;
+	constexpr int32 PrefixPlusTestLen = 9;
+
+	auto Reset = [&WideBuffer, &Utf8Buffer, &AnsiBuffer, WidePrefix, Utf8Prefix, AnsiPrefix, BufferLen]()
+		{
+			int32 n = 0;
+			for (; n < PrefixLen+1; ++n)
+			{
+				WideBuffer[n] = WidePrefix[n];
+				Utf8Buffer[n] = Utf8Prefix[n];
+				AnsiBuffer[n] = AnsiPrefix[n];
+			}
+			for (; n < BufferLen; ++n)
+			{
+				WideBuffer[n] = '%';
+				Utf8Buffer[n] = UTF8CHAR('%'); 
+				AnsiBuffer[n] = '%';
+			}
+		};
+
+	Reset();
+	FCStringWide::Strcat(WideBuffer, WideTest);
+	FCStringUtf8::Strcat(Utf8Buffer, Utf8Test);
+	FCStringAnsi::Strcat(AnsiBuffer, AnsiTest);
+	this->TestTrue(TEXT("WideStrcat"),
+		WideBuffer[PrefixPlusTestLen] == 0 && FCStringWide::Strcmp(WidePrefixPlusTest, WideBuffer) == 0 && WideBuffer[PrefixPlusTestLen + 1] == '%');
+	this->TestTrue(TEXT("Utf8Strcat"),
+		Utf8Buffer[PrefixPlusTestLen] == 0 && FCStringUtf8::Strcmp(Utf8PrefixPlusTest, Utf8Buffer) == 0 && Utf8Buffer[PrefixPlusTestLen + 1] == '%');
+	this->TestTrue(TEXT("AnsiStrcat"),
+		AnsiBuffer[PrefixPlusTestLen] == 0 && FCStringAnsi::Strcmp(AnsiPrefixPlusTest, AnsiBuffer) == 0 && AnsiBuffer[PrefixPlusTestLen + 1] == '%');
+
+	Reset();
+	FCStringWide::StrncatTruncateDest(WideBuffer, PrefixPlusTestLen + 10, WideTest);
+	FCStringUtf8::StrncatTruncateDest(Utf8Buffer, PrefixPlusTestLen + 10, Utf8Test);
+	FCStringAnsi::StrncatTruncateDest(AnsiBuffer, PrefixPlusTestLen + 10, AnsiTest);
+	this->TestTrue(TEXT("WideStrncatTruncateDestTestLenPlus10"),
+		WideBuffer[PrefixPlusTestLen] == 0 && FCStringWide::Strcmp(WidePrefixPlusTest, WideBuffer) == 0 && WideBuffer[PrefixPlusTestLen + 1] == '%');
+	this->TestTrue(TEXT("Utf8StrncatTruncateDestTestLenPlus10"),
+		Utf8Buffer[PrefixPlusTestLen] == 0 && FCStringUtf8::Strcmp(Utf8PrefixPlusTest, Utf8Buffer) == 0 && Utf8Buffer[PrefixPlusTestLen + 1] == '%');
+	this->TestTrue(TEXT("AnsiStrncatTruncateDestTestLenPlus10"),
+		AnsiBuffer[PrefixPlusTestLen] == 0 && FCStringAnsi::Strcmp(AnsiPrefixPlusTest, AnsiBuffer) == 0 && AnsiBuffer[PrefixPlusTestLen + 1] == '%');
+
+	Reset();
+	FCStringWide::StrncatTruncateDest(WideBuffer, PrefixPlusTestLen + 1, WideTest);
+	FCStringUtf8::StrncatTruncateDest(Utf8Buffer, PrefixPlusTestLen + 1, Utf8Test);
+	FCStringAnsi::StrncatTruncateDest(AnsiBuffer, PrefixPlusTestLen + 1, AnsiTest);
+	this->TestTrue(TEXT("WideStrncatTruncateDestTestLenPlus1"),
+		WideBuffer[PrefixPlusTestLen] == 0 && FCStringWide::Strcmp(WidePrefixPlusTest, WideBuffer) == 0 && WideBuffer[PrefixPlusTestLen + 1] == '%');
+	this->TestTrue(TEXT("Utf8StrncatTruncateDestTestLenPlus1"),
+		Utf8Buffer[PrefixPlusTestLen] == 0 && FCStringUtf8::Strcmp(Utf8PrefixPlusTest, Utf8Buffer) == 0 && Utf8Buffer[PrefixPlusTestLen + 1] == '%');
+	this->TestTrue(TEXT("AnsiStrncatTruncateDestTestLenPlus1"),
+		AnsiBuffer[PrefixPlusTestLen] == 0 && FCStringAnsi::Strcmp(AnsiPrefixPlusTest, AnsiBuffer) == 0 && AnsiBuffer[PrefixPlusTestLen + 1] == '%');
+
+	Reset();
+	FCStringWide::StrncatTruncateDest(WideBuffer, PrefixPlusTestLen, WideTest);
+	FCStringUtf8::StrncatTruncateDest(Utf8Buffer, PrefixPlusTestLen, Utf8Test);
+	FCStringAnsi::StrncatTruncateDest(AnsiBuffer, PrefixPlusTestLen, AnsiTest);
+	this->TestTrue(TEXT("WideStrncatTruncateDestTestLen"),
+		WideBuffer[PrefixPlusTestLen - 1] == 0 && WideBuffer[PrefixPlusTestLen] == '%' && FCStringWide::Strncmp(WidePrefixPlusTest, WideBuffer, PrefixPlusTestLen - 1) == 0 && WideBuffer[PrefixPlusTestLen + 1] == '%');
+	this->TestTrue(TEXT("Utf8StrncatTruncateDestTestLen"),
+		Utf8Buffer[PrefixPlusTestLen - 1] == 0 && Utf8Buffer[PrefixPlusTestLen] == '%' && FCStringUtf8::Strncmp(Utf8PrefixPlusTest, Utf8Buffer, PrefixPlusTestLen - 1) == 0 && Utf8Buffer[PrefixPlusTestLen + 1] == '%');
+	this->TestTrue(TEXT("AnsiStrncatTruncateDestTestLen"),
+		AnsiBuffer[PrefixPlusTestLen - 1] == 0 && AnsiBuffer[PrefixPlusTestLen] == '%' && FCStringAnsi::Strncmp(AnsiPrefixPlusTest, AnsiBuffer, PrefixPlusTestLen - 1) == 0 && AnsiBuffer[PrefixPlusTestLen + 1] == '%');
+
+	Reset();
+	FCStringWide::StrncatTruncateDest(WideBuffer, PrefixPlusTestLen - 1, WideTest);
+	FCStringUtf8::StrncatTruncateDest(Utf8Buffer, PrefixPlusTestLen - 1, Utf8Test);
+	FCStringAnsi::StrncatTruncateDest(AnsiBuffer, PrefixPlusTestLen - 1, AnsiTest);
+	this->TestTrue(TEXT("WideStrncatTruncateDestTestLenMinus1"),
+		WideBuffer[PrefixPlusTestLen - 2] == 0 && WideBuffer[PrefixPlusTestLen - 1] == '%' && FCStringWide::Strncmp(WidePrefixPlusTest, WideBuffer, PrefixPlusTestLen - 2) == 0 && WideBuffer[PrefixPlusTestLen + 1] == '%');
+	this->TestTrue(TEXT("Utf8StrncatTruncateDestTestLenMinus1"),
+		Utf8Buffer[PrefixPlusTestLen - 2] == 0 && Utf8Buffer[PrefixPlusTestLen - 1] == '%' && FCStringUtf8::Strncmp(Utf8PrefixPlusTest, Utf8Buffer, PrefixPlusTestLen - 2) == 0 && Utf8Buffer[PrefixPlusTestLen + 1] == '%');
+	this->TestTrue(TEXT("AnsiStrncatTruncateDestTestLenMinus1"),
+		AnsiBuffer[PrefixPlusTestLen - 2] == 0 && AnsiBuffer[PrefixPlusTestLen - 1] == '%' && FCStringAnsi::Strncmp(AnsiPrefixPlusTest, AnsiBuffer, PrefixPlusTestLen - 2) == 0 && AnsiBuffer[PrefixPlusTestLen + 1] == '%');
+
+	Reset();
+	FCStringWide::StrncatTruncateDest(WideBuffer, PrefixLen + 2, WideTest);
+	FCStringUtf8::StrncatTruncateDest(Utf8Buffer, PrefixLen + 2, Utf8Test);
+	FCStringAnsi::StrncatTruncateDest(AnsiBuffer, PrefixLen + 2, AnsiTest);
+	this->TestTrue(TEXT("WideStrncatTruncateDestTwoLen"),
+		WideBuffer[PrefixLen + 1] == 0 && WideBuffer[PrefixLen + 2] == '%' && FCStringWide::Strncmp(WidePrefixPlusTest, WideBuffer, PrefixLen + 1) == 0 && WideBuffer[PrefixPlusTestLen + 1] == '%');
+	this->TestTrue(TEXT("Utf8StrncatTruncateDestTwoLen"),
+		Utf8Buffer[PrefixLen + 1] == 0 && Utf8Buffer[PrefixLen + 2] == '%' && FCStringUtf8::Strncmp(Utf8PrefixPlusTest, Utf8Buffer, PrefixLen + 1) == 0 && Utf8Buffer[PrefixPlusTestLen + 1] == '%');
+	this->TestTrue(TEXT("AnsiStrncatTruncateDestTwoLen"),
+		AnsiBuffer[PrefixLen + 1] == 0 && AnsiBuffer[PrefixLen + 2] == '%' && FCStringAnsi::Strncmp(AnsiPrefixPlusTest, AnsiBuffer, PrefixLen + 1) == 0 && AnsiBuffer[PrefixPlusTestLen + 1] == '%');
+
+	Reset();
+	FCStringWide::StrncatTruncateDest(WideBuffer, PrefixLen + 1, WideTest);
+	FCStringUtf8::StrncatTruncateDest(Utf8Buffer, PrefixLen + 1, Utf8Test);
+	FCStringAnsi::StrncatTruncateDest(AnsiBuffer, PrefixLen + 1, AnsiTest);
+	this->TestTrue(TEXT("WideStrncatTruncateDestOneLen"),
+		WideBuffer[PrefixLen] == 0 && WideBuffer[PrefixLen + 1] == '%' && FCStringWide::Strncmp(WidePrefixPlusTest, WideBuffer, PrefixLen) == 0 && WideBuffer[PrefixPlusTestLen + 1] == '%');
+	this->TestTrue(TEXT("Utf8StrncatTruncateDestOneLen"),
+		Utf8Buffer[PrefixLen] == 0 && Utf8Buffer[PrefixLen + 1] == '%' && FCStringUtf8::Strncmp(Utf8PrefixPlusTest, Utf8Buffer, PrefixLen) == 0 && Utf8Buffer[PrefixPlusTestLen + 1] == '%');
+	this->TestTrue(TEXT("AnsiStrncatTruncateDestOneLen"),
+		AnsiBuffer[PrefixLen] == 0 && AnsiBuffer[PrefixLen + 1] == '%' && FCStringAnsi::Strncmp(AnsiPrefixPlusTest, AnsiBuffer, PrefixLen) == 0 && AnsiBuffer[PrefixPlusTestLen + 1] == '%');
+
+	Reset();
+	FCStringWide::StrncatTruncateDest(WideBuffer, PrefixLen, WideTest);
+	FCStringUtf8::StrncatTruncateDest(Utf8Buffer, PrefixLen, Utf8Test);
+	FCStringAnsi::StrncatTruncateDest(AnsiBuffer, PrefixLen, AnsiTest);
+	this->TestTrue(TEXT("WideStrncatTruncateDestZeroLen"),
+		WideBuffer[PrefixLen] == 0 && WideBuffer[PrefixLen + 1] == '%' && FCStringWide::Strncmp(WidePrefixPlusTest, WideBuffer, PrefixLen) == 0 && WideBuffer[PrefixPlusTestLen + 1] == '%');
+	this->TestTrue(TEXT("Utf8StrncatTruncateDestZeroLen"),
+		Utf8Buffer[PrefixLen] == 0 && Utf8Buffer[PrefixLen + 1] == '%' && FCStringUtf8::Strncmp(Utf8PrefixPlusTest, Utf8Buffer, PrefixLen) == 0 && Utf8Buffer[PrefixPlusTestLen + 1] == '%');
+	this->TestTrue(TEXT("AnsiStrncatTruncateDestZeroLen"),
+		AnsiBuffer[PrefixLen] == 0 && AnsiBuffer[PrefixLen + 1] == '%' && FCStringAnsi::Strncmp(AnsiPrefixPlusTest, AnsiBuffer, PrefixLen) == 0 && AnsiBuffer[PrefixPlusTestLen + 1] == '%');
+
+	Reset();
+	FCStringWide::StrncatTruncateSrc(WideBuffer, WideTest, TestLen + 10);
+	FCStringUtf8::StrncatTruncateSrc(Utf8Buffer, Utf8Test, TestLen + 10);
+	FCStringAnsi::StrncatTruncateSrc(AnsiBuffer, AnsiTest, TestLen + 10);
+	this->TestTrue(TEXT("WideStrncatTruncateSrcTestLenPlus10"),
+		WideBuffer[PrefixPlusTestLen] == 0 && FCStringWide::Strcmp(WidePrefixPlusTest, WideBuffer) == 0 && WideBuffer[PrefixPlusTestLen + 1] == '%');
+	this->TestTrue(TEXT("Utf8StrncatTruncateSrcTestLenPlus10"),
+		Utf8Buffer[PrefixPlusTestLen] == 0 && FCStringUtf8::Strcmp(Utf8PrefixPlusTest, Utf8Buffer) == 0 && Utf8Buffer[PrefixPlusTestLen + 1] == '%');
+	this->TestTrue(TEXT("AnsiStrncatTruncateSrcTestLenPlus10"),
+		AnsiBuffer[PrefixPlusTestLen] == 0 && FCStringAnsi::Strcmp(AnsiPrefixPlusTest, AnsiBuffer) == 0 && AnsiBuffer[PrefixPlusTestLen + 1] == '%');
+
+	Reset();
+	FCStringWide::StrncatTruncateSrc(WideBuffer, WideTest, TestLen);
+	FCStringUtf8::StrncatTruncateSrc(Utf8Buffer, Utf8Test, TestLen);
+	FCStringAnsi::StrncatTruncateSrc(AnsiBuffer, AnsiTest, TestLen);
+	this->TestTrue(TEXT("WideStrncatTruncateSrcTestLenPlus1"),
+		WideBuffer[PrefixPlusTestLen] == 0 && FCStringWide::Strcmp(WidePrefixPlusTest, WideBuffer) == 0 && WideBuffer[PrefixPlusTestLen + 1] == '%');
+	this->TestTrue(TEXT("Utf8StrncatTruncateSrcTestLenPlus1"),
+		Utf8Buffer[PrefixPlusTestLen] == 0 && FCStringUtf8::Strcmp(Utf8PrefixPlusTest, Utf8Buffer) == 0 && Utf8Buffer[PrefixPlusTestLen + 1] == '%');
+	this->TestTrue(TEXT("AnsiStrncatTruncateSrcTestLenPlus1"),
+		AnsiBuffer[PrefixPlusTestLen] == 0 && FCStringAnsi::Strcmp(AnsiPrefixPlusTest, AnsiBuffer) == 0 && AnsiBuffer[PrefixPlusTestLen + 1] == '%');
+
+	Reset();
+	FCStringWide::StrncatTruncateSrc(WideBuffer, WideTest, TestLen - 1);
+	FCStringUtf8::StrncatTruncateSrc(Utf8Buffer, Utf8Test, TestLen - 1);
+	FCStringAnsi::StrncatTruncateSrc(AnsiBuffer, AnsiTest, TestLen - 1);
+	this->TestTrue(TEXT("WideStrncatTruncateSrcTestLenMinus1"),
+		WideBuffer[PrefixPlusTestLen - 1] == 0 && WideBuffer[PrefixPlusTestLen] == '%' && FCStringWide::Strncmp(WidePrefixPlusTest, WideBuffer, PrefixPlusTestLen - 1) == 0 && WideBuffer[PrefixPlusTestLen + 1] == '%');
+	this->TestTrue(TEXT("Utf8StrncatTruncateSrcTestLenMinus1"),
+		Utf8Buffer[PrefixPlusTestLen - 1] == 0 && Utf8Buffer[PrefixPlusTestLen] == '%' && FCStringUtf8::Strncmp(Utf8PrefixPlusTest, Utf8Buffer, PrefixPlusTestLen - 1) == 0 && Utf8Buffer[PrefixPlusTestLen + 1] == '%');
+	this->TestTrue(TEXT("AnsiStrncatTruncateSrcTestLenMinus1"),
+		AnsiBuffer[PrefixPlusTestLen - 1] == 0 && AnsiBuffer[PrefixPlusTestLen] == '%' && FCStringAnsi::Strncmp(AnsiPrefixPlusTest, AnsiBuffer, PrefixPlusTestLen - 1) == 0 && AnsiBuffer[PrefixPlusTestLen + 1] == '%');
+
+	Reset();
+	FCStringWide::StrncatTruncateSrc(WideBuffer, WideTest, TestLen - 2);
+	FCStringUtf8::StrncatTruncateSrc(Utf8Buffer, Utf8Test, TestLen - 2);
+	FCStringAnsi::StrncatTruncateSrc(AnsiBuffer, AnsiTest, TestLen - 2);
+	this->TestTrue(TEXT("WideStrncatTruncateSrcTestLenMinus2"),
+		WideBuffer[PrefixPlusTestLen - 2] == 0 && WideBuffer[PrefixPlusTestLen - 1] == '%' && FCStringWide::Strncmp(WidePrefixPlusTest, WideBuffer, PrefixPlusTestLen - 2) == 0 && WideBuffer[PrefixPlusTestLen + 1] == '%');
+	this->TestTrue(TEXT("Utf8StrncatTruncateSrcTestLenMinus2"),
+		Utf8Buffer[PrefixPlusTestLen - 2] == 0 && Utf8Buffer[PrefixPlusTestLen - 1] == '%' && FCStringUtf8::Strncmp(Utf8PrefixPlusTest, Utf8Buffer, PrefixPlusTestLen - 2) == 0 && Utf8Buffer[PrefixPlusTestLen + 1] == '%');
+	this->TestTrue(TEXT("AnsiStrncatTruncateSrcTestLenMinus2"),
+		AnsiBuffer[PrefixPlusTestLen - 2] == 0 && AnsiBuffer[PrefixPlusTestLen - 1] == '%' && FCStringAnsi::Strncmp(AnsiPrefixPlusTest, AnsiBuffer, PrefixPlusTestLen - 2) == 0 && AnsiBuffer[PrefixPlusTestLen + 1] == '%');
+
+	Reset();
+	FCStringWide::StrncatTruncateSrc(WideBuffer, WideTest, 1);
+	FCStringUtf8::StrncatTruncateSrc(Utf8Buffer, Utf8Test, 1);
+	FCStringAnsi::StrncatTruncateSrc(AnsiBuffer, AnsiTest, 1);
+	this->TestTrue(TEXT("WideStrncatTruncateSrcOneLen"),
+		WideBuffer[PrefixLen + 1] == 0 && WideBuffer[PrefixLen + 2] == '%' && FCStringWide::Strncmp(WidePrefixPlusTest, WideBuffer, PrefixLen + 1) == 0 && WideBuffer[PrefixPlusTestLen + 1] == '%');
+	this->TestTrue(TEXT("Utf8StrncatTruncateSrcOneLen"),
+		Utf8Buffer[PrefixLen + 1] == 0 && Utf8Buffer[PrefixLen + 2] == '%' && FCStringUtf8::Strncmp(Utf8PrefixPlusTest, Utf8Buffer, PrefixLen + 1) == 0 && Utf8Buffer[PrefixPlusTestLen + 1] == '%');
+	this->TestTrue(TEXT("AnsiStrncatTruncateSrcOneLen"),
+		AnsiBuffer[PrefixLen + 1] == 0 && AnsiBuffer[PrefixLen + 2] == '%' && FCStringAnsi::Strncmp(AnsiPrefixPlusTest, AnsiBuffer, PrefixLen + 1) == 0 && AnsiBuffer[PrefixPlusTestLen + 1] == '%');
+
+	Reset();
+	FCStringWide::StrncatTruncateSrc(WideBuffer, WideTest, 0);
+	FCStringUtf8::StrncatTruncateSrc(Utf8Buffer, Utf8Test, 0);
+	FCStringAnsi::StrncatTruncateSrc(AnsiBuffer, AnsiTest, 0);
+	this->TestTrue(TEXT("WideStrncatTruncateSrcZeroLen"),
+		WideBuffer[PrefixLen] == 0 && WideBuffer[PrefixLen + 1] == '%' && FCStringWide::Strncmp(WidePrefixPlusTest, WideBuffer, PrefixLen) == 0 && WideBuffer[PrefixPlusTestLen + 1] == '%');
+	this->TestTrue(TEXT("Utf8StrncatTruncateSrcZeroLen"),
+		Utf8Buffer[PrefixLen] == 0 && Utf8Buffer[PrefixLen + 1] == '%' && FCStringUtf8::Strncmp(Utf8PrefixPlusTest, Utf8Buffer, PrefixLen) == 0 && Utf8Buffer[PrefixPlusTestLen + 1] == '%');
+	this->TestTrue(TEXT("AnsiStrncatTruncateSrcZeroLen"),
+		AnsiBuffer[PrefixLen] == 0 && AnsiBuffer[PrefixLen + 1] == '%' && FCStringAnsi::Strncmp(AnsiPrefixPlusTest, AnsiBuffer, PrefixLen) == 0 && AnsiBuffer[PrefixPlusTestLen + 1] == '%');
+
+	return true;
+}
+
 #endif // WITH_DEV_AUTOMATION_TESTS
diff --git a/Engine/Source/Runtime/Core/Private/Misc/CommandLine.cpp b/Engine/Source/Runtime/Core/Private/Misc/CommandLine.cpp
--- a/Engine/Source/Runtime/Core/Private/Misc/CommandLine.cpp
+++ b/Engine/Source/Runtime/Core/Private/Misc/CommandLine.cpp
@@ -124,7 +124,7 @@ bool FCommandLine::Set(const TCHAR* NewCommandLine)
 
 void FCommandLine::Append(const TCHAR* AppendString)
 {
-	FCString::Strncat( CmdLine, AppendString, UE_ARRAY_COUNT(CmdLine) );
+	FCString::StrncatTruncateDest( CmdLine, UE_ARRAY_COUNT(CmdLine), AppendString );
 	CmdLineVersion++;
 	// If configured as part of the build, strip out any unapproved args
 	ApplyCommandLineAllowList();
diff --git a/Engine/Source/Runtime/Core/Private/Misc/OutputDeviceAnsiError.cpp b/Engine/Source/Runtime/Core/Private/Misc/OutputDeviceAnsiError.cpp
--- a/Engine/Source/Runtime/Core/Private/Misc/OutputDeviceAnsiError.cpp
+++ b/Engine/Source/Runtime/Core/Private/Misc/OutputDeviceAnsiError.cpp
@@ -32,7 +32,7 @@ void FOutputDeviceAnsiError::Serialize( const TCHAR* Msg, ELogVerbosity::Type Ve
 		GIsCriticalError = 1;
 		UE_LOG(LogHAL, Error, TEXT("appError called: %s"), Msg );
 		FCString::Strncpy( GErrorHist, Msg, UE_ARRAY_COUNT(GErrorHist) );
-		FCString::Strncat( GErrorHist, TEXT("\r\n\r\n"), UE_ARRAY_COUNT(GErrorHist) );
+		FCString::StrncatTruncateDest( GErrorHist, UE_ARRAY_COUNT(GErrorHist), TEXT("\r\n\r\n") );
 		ErrorPos = FCString::Strlen(GErrorHist);
 	}
 	else
diff --git a/Engine/Source/Runtime/Core/Private/Misc/OutputDeviceFile.cpp b/Engine/Source/Runtime/Core/Private/Misc/OutputDeviceFile.cpp
--- a/Engine/Source/Runtime/Core/Private/Misc/OutputDeviceFile.cpp
+++ b/Engine/Source/Runtime/Core/Private/Misc/OutputDeviceFile.cpp
@@ -494,7 +494,7 @@ bool FOutputDeviceFile::CreateWriter(uint32 MaxAttempts)
 			// Continue to increment indices until a valid filename is found
 			FinalFilename = FilenamePart + FString::FromInt(FileIndex++) + ExtensionPart;
 			CreateBackupCopy(*FinalFilename);
-			FCString::Strcpy(Filename, UE_ARRAY_COUNT(Filename), *FinalFilename);
+			FCString::Strncpy(Filename, *FinalFilename, UE_ARRAY_COUNT(Filename));
 			Ar = IFileManager::Get().CreateFileWriter(*FinalFilename, WriteFlags);
 		} while (!Ar && FileIndex < MaxAttempts);
 	}
diff --git a/Engine/Source/Runtime/Core/Private/Stats/StatsSystemTypes.cpp b/Engine/Source/Runtime/Core/Private/Stats/StatsSystemTypes.cpp
--- a/Engine/Source/Runtime/Core/Private/Stats/StatsSystemTypes.cpp
+++ b/Engine/Source/Runtime/Core/Private/Stats/StatsSystemTypes.cpp
@@ -507,7 +507,7 @@ public:
 				auto StatDescriptionWide = StringCast<WIDECHAR>(*StatDescription, StatDescLen);
 				int32 StatDescriptionWideLength = StatDescriptionWide.Length();
 				TUniquePtr<WIDECHAR[]> StatDescWide = MakeUnique<WIDECHAR[]>(StatDescriptionWideLength + 1);
-				TCString<WIDECHAR>::Strcpy(StatDescWide.Get(), StatDescriptionWideLength + 1, StatDescriptionWide.Get());
+				TCString<WIDECHAR>::Strncpy(StatDescWide.Get(), StatDescriptionWide.Get(), StatDescriptionWideLength + 1);
 				Result->StatDescriptionWide = MoveTemp(StatDescWide);
 
 				MemoryAllocated += StatDescriptionWideLength * sizeof(WIDECHAR);
@@ -518,7 +518,7 @@ public:
 				auto StatDescriptionAnsi = StringCast<ANSICHAR>(*StatDescription, StatDescLen);
 				int32 StatDescriptionAnsiLength = StatDescriptionAnsi.Length();
 				TUniquePtr<ANSICHAR[]> StatDescAnsi = MakeUnique<ANSICHAR[]>(StatDescriptionAnsiLength + 1);
-				TCString<ANSICHAR>::Strcpy(StatDescAnsi.Get(), StatDescriptionAnsiLength + 1, StatDescriptionAnsi.Get());
+				TCString<ANSICHAR>::Strncpy(StatDescAnsi.Get(), StatDescriptionAnsi.Get(), StatDescriptionAnsiLength + 1);
 				Result->StatDescriptionAnsi = MoveTemp(StatDescAnsi);
 
 				MemoryAllocated += StatDescriptionAnsiLength * sizeof(ANSICHAR);
diff --git a/Engine/Source/Runtime/Core/Private/Unix/UnixErrorOutputDevice.cpp b/Engine/Source/Runtime/Core/Private/Unix/UnixErrorOutputDevice.cpp
--- a/Engine/Source/Runtime/Core/Private/Unix/UnixErrorOutputDevice.cpp
+++ b/Engine/Source/Runtime/Core/Private/Unix/UnixErrorOutputDevice.cpp
@@ -47,7 +47,7 @@ void FUnixErrorOutputDevice::Serialize(const TCHAR* Msg, ELogVerbosity::Type Ver
 		if (FCString::Strncmp(GErrorHist, Msg, FMath::Min<int32>(UE_ARRAY_COUNT(GErrorHist), FCString::Strlen(Msg))))
 		{
 			FCString::Strncpy(GErrorHist, Msg, UE_ARRAY_COUNT(GErrorHist) - 5);
-			FCString::Strncat(GErrorHist, TEXT("\r\n\r\n"), UE_ARRAY_COUNT(GErrorHist) - 1);
+			FCString::StrncatTruncateDest(GErrorHist, UE_ARRAY_COUNT(GErrorHist), TEXT("\r\n\r\n"));
 		}
 		ErrorPos = FCString::Strlen(GErrorHist);
 	}
diff --git a/Engine/Source/Runtime/Core/Private/Unix/UnixPlatformCrashContext.cpp b/Engine/Source/Runtime/Core/Private/Unix/UnixPlatformCrashContext.cpp
--- a/Engine/Source/Runtime/Core/Private/Unix/UnixPlatformCrashContext.cpp
+++ b/Engine/Source/Runtime/Core/Private/Unix/UnixPlatformCrashContext.cpp
@@ -113,7 +113,7 @@ void FUnixCrashContext::InitFromSignal(int32 InSignal, siginfo_t* InInfo, void*
 	Info = InInfo;
 	Context = reinterpret_cast< ucontext_t* >( InContext );
 
-	FCString::Strcat(SignalDescription, UE_ARRAY_COUNT( SignalDescription ) - 1, *DescribeSignal(Signal, Info, Context));
+	FCString::StrncatTruncateDest(SignalDescription, UE_ARRAY_COUNT( SignalDescription ), *DescribeSignal(Signal, Info, Context));
 }
 
 void FUnixCrashContext::InitFromDiagnostics(const void* InAddress)
@@ -128,7 +128,7 @@ void FUnixCrashContext::InitFromDiagnostics(const void* InAddress)
 	Context = nullptr;
 
 	// set signal description to a more human-readable one for ensures
-	FCString::Strcpy(SignalDescription, UE_ARRAY_COUNT(SignalDescription) - 1, ErrorMessage);
+	FCString::Strncpy(SignalDescription, ErrorMessage, UE_ARRAY_COUNT(SignalDescription));
 
 	// only need the first string
 	for (int Idx = 0; Idx < UE_ARRAY_COUNT(SignalDescription); ++Idx)
@@ -301,7 +301,7 @@ void FUnixCrashContext::CaptureStackTrace(void* ErrorProgramCounter)
 		printf("StackTrace:\n%s\n", StackTrace);
 #endif
 
-		FCString::Strncat( GErrorHist, UTF8_TO_TCHAR(StackTrace), UE_ARRAY_COUNT(GErrorHist) - 1 );
+		FCString::StrncatTruncateDest( GErrorHist, UE_ARRAY_COUNT(GErrorHist), UTF8_TO_TCHAR(StackTrace) );
 		CreateExceptionInfoString(Signal, Info, Context);
 	}
 }
@@ -324,7 +324,7 @@ void FUnixCrashContext::CaptureThreadStackTrace(uint32_t ThreadId)
 		for ( int i=0; i<CallStack.Num(); i++ )
 		{
 			FPlatformStackWalk::ProgramCounterToHumanReadableString( i, CallStack[i].BaseAddress + CallStack[i].Offset, StackTrace, StackTraceSize, this );
-			FCStringAnsi::Strncat(StackTrace, LINE_TERMINATOR_ANSI, (int32)StackTraceSize);
+			FCStringAnsi::StrncatTruncateDest(StackTrace, (int32)StackTraceSize, LINE_TERMINATOR_ANSI);
 		}
 
 		// We do not set the ExceptionInfo string here as it just gets details for the exact callsite,
diff --git a/Engine/Source/Runtime/Core/Private/Unix/UnixPlatformProcess.cpp b/Engine/Source/Runtime/Core/Private/Unix/UnixPlatformProcess.cpp
--- a/Engine/Source/Runtime/Core/Private/Unix/UnixPlatformProcess.cpp
+++ b/Engine/Source/Runtime/Core/Private/Unix/UnixPlatformProcess.cpp
@@ -193,7 +193,7 @@ const TCHAR* FUnixPlatformProcess::ComputerName()
 			SysName = "Unix Computer";
 		}
 
-		FCString::Strcpy(CachedResult, UE_ARRAY_COUNT(CachedResult) - 1, UTF8_TO_TCHAR(SysName));
+		FCString::Strncpy(CachedResult, UTF8_TO_TCHAR(SysName), UE_ARRAY_COUNT(CachedResult));
 		CachedResult[UE_ARRAY_COUNT(CachedResult) - 1] = 0;
 		bHaveResult = true;
 	}
@@ -284,7 +284,7 @@ const TCHAR* FUnixPlatformProcess::UserDir()
 				{
 					DocPath[DocLen] = '\0';
 					FCString::Strncpy(Result, UTF8_TO_TCHAR(DocPath), UE_ARRAY_COUNT(Result));
-					FCString::Strncat(Result, TEXT("/"), UE_ARRAY_COUNT(Result));
+					FCString::StrncatTruncateDest(Result, UE_ARRAY_COUNT(Result), TEXT("/"));
 				}
 			}
 			pclose(FilePtr);
@@ -294,7 +294,7 @@ const TCHAR* FUnixPlatformProcess::UserDir()
 		if (!Result[0])
 		{
 			FCString::Strncpy(Result, FPlatformProcess::UserHomeDir(), UE_ARRAY_COUNT(Result));
-			FCString::Strncat(Result, TEXT("/Documents/"), UE_ARRAY_COUNT(Result));
+			FCString::StrncatTruncateDest(Result, UE_ARRAY_COUNT(Result), TEXT("/Documents/"));
 		}
 	}
 	return Result;
@@ -312,18 +312,18 @@ const TCHAR* FUnixPlatformProcess::UserHomeDir()
 		const char * VarValue = secure_getenv("HOME");
 		if (VarValue && VarValue[0] != '\0')
 		{
-			FCString::Strcpy(CachedResult, UE_ARRAY_COUNT(CachedResult) - 1, UTF8_TO_TCHAR(VarValue));
+			FCString::Strncpy(CachedResult, UTF8_TO_TCHAR(VarValue), UE_ARRAY_COUNT(CachedResult));
 		}
 		else
 		{
 			struct passwd * UserInfo = getpwuid(geteuid());
 			if (NULL != UserInfo && NULL != UserInfo->pw_dir && UserInfo->pw_dir[0] != '\0')
 			{
-				FCString::Strcpy(CachedResult, UE_ARRAY_COUNT(CachedResult) - 1, UTF8_TO_TCHAR(UserInfo->pw_dir));
+				FCString::Strncpy(CachedResult, UTF8_TO_TCHAR(UserInfo->pw_dir), UE_ARRAY_COUNT(CachedResult));
 			}
 			else
 			{
-				FCString::Strcpy(CachedResult, UE_ARRAY_COUNT(CachedResult) - 1, FUnixPlatformProcess::UserTempDir());
+				FCString::Strncpy(CachedResult, FUnixPlatformProcess::UserTempDir(), UE_ARRAY_COUNT(CachedResult));
 				UE_LOG(LogInit, Warning, TEXT("Could get determine user home directory.  Using temporary directory: %s"), CachedResult);
 			}
 		}
@@ -347,7 +347,7 @@ const TCHAR* FUnixPlatformProcess::ApplicationSettingsDir()
 	if (!Result[0])
 	{
 		FCString::Strncpy(Result, FPlatformProcess::UserHomeDir(), UE_ARRAY_COUNT(Result));
-		FCString::Strncat(Result, TEXT("/.config/Epic/"), UE_ARRAY_COUNT(Result));
+		FCString::StrncatTruncateDest(Result, UE_ARRAY_COUNT(Result), TEXT("/.config/Epic/"));
 	}
 	return Result;
 }
@@ -360,11 +360,11 @@ FString FUnixPlatformProcess::GetApplicationSettingsDir(const ApplicationSetting
 	FCString::Strncpy(Result, FPlatformProcess::UserHomeDir(), UE_ARRAY_COUNT(Result));
 	if (Settings.bIsEpic)
 	{
-		FCString::Strncat(Result, TEXT("/.config/Epic/"), UE_ARRAY_COUNT(Result));
+		FCString::StrncatTruncateDest(Result, UE_ARRAY_COUNT(Result), TEXT("/.config/Epic/"));
 	}
 	else
 	{
-		FCString::Strncat(Result, TEXT("/.config/"), UE_ARRAY_COUNT(Result));
+		FCString::StrncatTruncateDest(Result, UE_ARRAY_COUNT(Result), TEXT("/.config/"));
 	}
 	return FString(Result);
 }
@@ -438,7 +438,7 @@ const TCHAR* FUnixPlatformProcess::ExecutablePath()
 		}
 		SelfPath[UE_ARRAY_COUNT(SelfPath) - 1] = 0;
 
-		FCString::Strcpy(CachedResult, UE_ARRAY_COUNT(CachedResult) - 1, UTF8_TO_TCHAR(SelfPath));
+		FCString::Strncpy(CachedResult, UTF8_TO_TCHAR(SelfPath), UE_ARRAY_COUNT(CachedResult));
 		CachedResult[UE_ARRAY_COUNT(CachedResult) - 1] = 0;
 		bHaveResult = true;
 	}
@@ -461,7 +461,7 @@ const TCHAR* FUnixPlatformProcess::ExecutableName(bool bRemoveExtension)
 		}
 		SelfPath[UE_ARRAY_COUNT(SelfPath) - 1] = 0;
 
-		FCString::Strcpy(CachedResult, UE_ARRAY_COUNT(CachedResult) - 1, UTF8_TO_TCHAR(basename(SelfPath)));
+		FCString::Strncpy(CachedResult, UTF8_TO_TCHAR(basename(SelfPath)), UE_ARRAY_COUNT(CachedResult));
 		CachedResult[UE_ARRAY_COUNT(CachedResult) - 1] = 0;
 		bHaveResult = true;
 	}
diff --git a/Engine/Source/Runtime/Core/Private/Unix/UnixPlatformStackWalk.cpp b/Engine/Source/Runtime/Core/Private/Unix/UnixPlatformStackWalk.cpp
--- a/Engine/Source/Runtime/Core/Private/Unix/UnixPlatformStackWalk.cpp
+++ b/Engine/Source/Runtime/Core/Private/Unix/UnixPlatformStackWalk.cpp
@@ -563,7 +563,7 @@ bool FUnixPlatformStackWalk::ProgramCounterToHumanReadableString( int32 CurrentC
 			{
 				FCStringAnsi::Sprintf(TempArray, "0x%08x ", (uint32) ProgramCounter);
 			}
-			FCStringAnsi::Strncat(HumanReadableString, TempArray, HumanReadableStringSize);
+			FCStringAnsi::StrncatTruncateDest(HumanReadableString, HumanReadableStringSize, TempArray);
 
 			// won't be able to display names here
 		}
@@ -577,7 +577,7 @@ bool FUnixPlatformStackWalk::ProgramCounterToHumanReadableString( int32 CurrentC
 			{
 				FCStringAnsi::Sprintf(TempArray, "0x%08x ", (uint32) ProgramCounter);
 			}
-			FCStringAnsi::Strncat(HumanReadableString, TempArray, HumanReadableStringSize);
+			FCStringAnsi::StrncatTruncateDest(HumanReadableString, HumanReadableStringSize, TempArray);
 
 			// Get filename, source file and line number
 			FUnixCrashContext* UnixContext = static_cast< FUnixCrashContext* >( Context );
@@ -603,19 +603,19 @@ bool FUnixPlatformStackWalk::ProgramCounterToHumanReadableString( int32 CurrentC
 			if (bAddDetailedInfo)
 			{
 				// append Module!FunctionName() [Source.cpp:X] to HumanReadableString
-				FCStringAnsi::Strncat(HumanReadableString, TempSymbolInfo.ModuleName, HumanReadableStringSize);
-				FCStringAnsi::Strncat(HumanReadableString, "!", HumanReadableStringSize);
-				FCStringAnsi::Strncat(HumanReadableString, TempSymbolInfo.FunctionName, HumanReadableStringSize);
+				FCStringAnsi::StrncatTruncateDest(HumanReadableString, HumanReadableStringSize, TempSymbolInfo.ModuleName);
+				FCStringAnsi::StrncatTruncateDest(HumanReadableString, HumanReadableStringSize, "!");
+				FCStringAnsi::StrncatTruncateDest(HumanReadableString, HumanReadableStringSize, TempSymbolInfo.FunctionName);
 				FCStringAnsi::Sprintf(TempArray, " [%s:%d]", TempSymbolInfo.Filename, TempSymbolInfo.LineNumber);
-				FCStringAnsi::Strncat(HumanReadableString, TempArray, HumanReadableStringSize);
+				FCStringAnsi::StrncatTruncateDest(HumanReadableString, HumanReadableStringSize, TempArray);
 
 				if (UnixContext)
 				{
 					// append Module!FunctioName [Source.cpp:X] to MinidumpCallstackInfo
-					FCStringAnsi::Strncat(UnixContext->MinidumpCallstackInfo, TempSymbolInfo.ModuleName, UE_ARRAY_COUNT( UnixContext->MinidumpCallstackInfo ) - 1);
-					FCStringAnsi::Strncat(UnixContext->MinidumpCallstackInfo, "!", UE_ARRAY_COUNT( UnixContext->MinidumpCallstackInfo ) - 1);
-					FCStringAnsi::Strncat(UnixContext->MinidumpCallstackInfo, TempSymbolInfo.FunctionName, UE_ARRAY_COUNT( UnixContext->MinidumpCallstackInfo ) - 1);
-					FCStringAnsi::Strncat(UnixContext->MinidumpCallstackInfo, TempArray, UE_ARRAY_COUNT( UnixContext->MinidumpCallstackInfo ) - 1);
+					FCStringAnsi::StrncatTruncateDest(UnixContext->MinidumpCallstackInfo, UE_ARRAY_COUNT(UnixContext->MinidumpCallstackInfo), TempSymbolInfo.ModuleName);
+					FCStringAnsi::StrncatTruncateDest(UnixContext->MinidumpCallstackInfo, UE_ARRAY_COUNT(UnixContext->MinidumpCallstackInfo), "!");
+					FCStringAnsi::StrncatTruncateDest(UnixContext->MinidumpCallstackInfo, UE_ARRAY_COUNT(UnixContext->MinidumpCallstackInfo), TempSymbolInfo.FunctionName);
+					FCStringAnsi::StrncatTruncateDest(UnixContext->MinidumpCallstackInfo, UE_ARRAY_COUNT(UnixContext->MinidumpCallstackInfo), TempArray);
 				}
 			}
 			else
@@ -634,41 +634,41 @@ bool FUnixPlatformStackWalk::ProgramCounterToHumanReadableString( int32 CurrentC
 					FunctionName = TempSymbolInfo.FunctionName;
 				}
 
-				FCStringAnsi::Strncat(HumanReadableString, ModuleName != nullptr ? ModuleName : "", HumanReadableStringSize);
-				FCStringAnsi::Strncat(HumanReadableString, "!", HumanReadableStringSize);
-				FCStringAnsi::Strncat(HumanReadableString, FunctionName != nullptr ? FunctionName : "UnknownFunction", HumanReadableStringSize);
-				FCStringAnsi::Strncat(HumanReadableString, FunctionName && TempSymbolInfo.SymbolDisplacement ? "(+": "(", HumanReadableStringSize);
+				FCStringAnsi::StrncatTruncateDest(HumanReadableString, HumanReadableStringSize, ModuleName != nullptr ? ModuleName : "");
+				FCStringAnsi::StrncatTruncateDest(HumanReadableString, HumanReadableStringSize, "!");
+				FCStringAnsi::StrncatTruncateDest(HumanReadableString, HumanReadableStringSize, FunctionName != nullptr ? FunctionName : "UnknownFunction");
+				FCStringAnsi::StrncatTruncateDest(HumanReadableString, HumanReadableStringSize, FunctionName && TempSymbolInfo.SymbolDisplacement ? "(+": "(");
 
 				if (UnixContext)
 				{
-					FCStringAnsi::Strncat(UnixContext->MinidumpCallstackInfo, ModuleName != nullptr ? ModuleName : "Unknown", UE_ARRAY_COUNT( UnixContext->MinidumpCallstackInfo ) - 1);
-					FCStringAnsi::Strncat(UnixContext->MinidumpCallstackInfo, "!", UE_ARRAY_COUNT( UnixContext->MinidumpCallstackInfo ) - 1);
-					FCStringAnsi::Strncat(UnixContext->MinidumpCallstackInfo, FunctionName != nullptr ? FunctionName : "UnknownFunction", UE_ARRAY_COUNT( UnixContext->MinidumpCallstackInfo ) - 1);
-					FCStringAnsi::Strncat(UnixContext->MinidumpCallstackInfo, FunctionName && TempSymbolInfo.SymbolDisplacement ? "(+": "(", UE_ARRAY_COUNT( UnixContext->MinidumpCallstackInfo ) - 1);
+					FCStringAnsi::StrncatTruncateDest(UnixContext->MinidumpCallstackInfo, UE_ARRAY_COUNT(UnixContext->MinidumpCallstackInfo), ModuleName != nullptr ? ModuleName : "Unknown");
+					FCStringAnsi::StrncatTruncateDest(UnixContext->MinidumpCallstackInfo, UE_ARRAY_COUNT(UnixContext->MinidumpCallstackInfo), "!");
+					FCStringAnsi::StrncatTruncateDest(UnixContext->MinidumpCallstackInfo, UE_ARRAY_COUNT(UnixContext->MinidumpCallstackInfo), FunctionName != nullptr ? FunctionName : "UnknownFunction");
+					FCStringAnsi::StrncatTruncateDest(UnixContext->MinidumpCallstackInfo, UE_ARRAY_COUNT(UnixContext->MinidumpCallstackInfo), FunctionName && TempSymbolInfo.SymbolDisplacement ? "(+": "(");
 				}
 
 				if (TempSymbolInfo.SymbolDisplacement > 0x0)
 				{
 					FCStringAnsi::Sprintf(TempArray, "%p", TempSymbolInfo.SymbolDisplacement);
-					FCStringAnsi::Strncat(HumanReadableString, TempArray, HumanReadableStringSize);
+					FCStringAnsi::StrncatTruncateDest(HumanReadableString, HumanReadableStringSize, TempArray);
 
 					if (UnixContext)
 					{
-						FCStringAnsi::Strncat(UnixContext->MinidumpCallstackInfo, TempArray, UE_ARRAY_COUNT( UnixContext->MinidumpCallstackInfo ) - 1);
+						FCStringAnsi::StrncatTruncateDest(UnixContext->MinidumpCallstackInfo, UE_ARRAY_COUNT(UnixContext->MinidumpCallstackInfo), TempArray);
 					}
 				}
 
-				FCStringAnsi::Strncat(HumanReadableString, ")", HumanReadableStringSize);
+				FCStringAnsi::StrncatTruncateDest(HumanReadableString, HumanReadableStringSize, ")");
 
 				if (UnixContext)
 				{
-					FCStringAnsi::Strncat(UnixContext->MinidumpCallstackInfo, ")", UE_ARRAY_COUNT( UnixContext->MinidumpCallstackInfo ) - 1);
+					FCStringAnsi::StrncatTruncateDest(UnixContext->MinidumpCallstackInfo, UE_ARRAY_COUNT(UnixContext->MinidumpCallstackInfo), ")");
 				}
 			}
 
 			if (UnixContext)
 			{
-				FCStringAnsi::Strncat(UnixContext->MinidumpCallstackInfo, "\r\n", UE_ARRAY_COUNT( UnixContext->MinidumpCallstackInfo ) - 1);	// this one always uses Windows line terminators
+				FCStringAnsi::StrncatTruncateDest(UnixContext->MinidumpCallstackInfo, UE_ARRAY_COUNT(UnixContext->MinidumpCallstackInfo), "\r\n");	// this one always uses Windows line terminators
 			}
 		}
 		return true;
diff --git a/Engine/Source/Runtime/Core/Private/Unix/UnixPlatformTime.cpp b/Engine/Source/Runtime/Core/Private/Unix/UnixPlatformTime.cpp
--- a/Engine/Source/Runtime/Core/Private/Unix/UnixPlatformTime.cpp
+++ b/Engine/Source/Runtime/Core/Private/Unix/UnixPlatformTime.cpp
@@ -243,7 +243,7 @@ uint64 FUnixTime::CallsPerSecondBenchmark(clockid_t BenchClockId, const char * B
 	if (clock_gettime(BenchClockId, &ts) == -1)
 	{
 		FCStringAnsi::Snprintf(Buffer, sizeof(Buffer), "Clock_id %d (%s) is not supported on this system, clock_gettime() fails.\n", BenchClockId, BenchClockIdName);
-		FCStringAnsi::Strncat(CalibrationLog, Buffer, sizeof(CalibrationLog));
+		FCStringAnsi::StrncatTruncateDest(CalibrationLog, sizeof(CalibrationLog), Buffer);
 		return 0;	// unsupported clock id either
 	}
 
@@ -264,7 +264,7 @@ uint64 FUnixTime::CallsPerSecondBenchmark(clockid_t BenchClockId, const char * B
 		if (NewEndTimeStamp < EndTimeStamp)
 		{
 			FCStringAnsi::Snprintf(Buffer, sizeof(Buffer), "Clock_id %d (%s) is unusable, can go backwards.\n", BenchClockId, BenchClockIdName);
-			FCStringAnsi::Strncat(CalibrationLog, Buffer, sizeof(CalibrationLog));
+			FCStringAnsi::StrncatTruncateDest(CalibrationLog, sizeof(CalibrationLog), Buffer);
 			return 0;
 		}
 		else if (NewEndTimeStamp == EndTimeStamp)
@@ -275,7 +275,7 @@ uint64 FUnixTime::CallsPerSecondBenchmark(clockid_t BenchClockId, const char * B
 			if (NumZeroDeltas > kHardLimitOnZeroDeltas)
 			{
 				FCStringAnsi::Snprintf(Buffer, sizeof(Buffer), "Clock_id %d (%s) is unusable, too many (%llu) zero deltas.\n", BenchClockId, BenchClockIdName, NumZeroDeltas);
-				FCStringAnsi::Strncat(CalibrationLog, Buffer, sizeof(CalibrationLog));
+				FCStringAnsi::StrncatTruncateDest(CalibrationLog, sizeof(CalibrationLog), Buffer);
 				return 0;
 			}
 		}
@@ -297,7 +297,7 @@ uint64 FUnixTime::CallsPerSecondBenchmark(clockid_t BenchClockId, const char * B
 	FCStringAnsi::Snprintf(Buffer, sizeof(Buffer), (" - %s (id=%d) can sustain %llu (%lluK, %lluM) calls per second %s.\n"), BenchClockIdName, BenchClockId,
 		RealNumCalls, (RealNumCalls + 500) / 1000, (RealNumCalls + 500000) / 1000000,
 		NumZeroDeltas ? ZeroDeltasBuf : "without zero deltas");
-	FCStringAnsi::Strncat(CalibrationLog, Buffer, sizeof(CalibrationLog));
+	FCStringAnsi::StrncatTruncateDest(CalibrationLog, sizeof(CalibrationLog), Buffer);
 
 	// if clock had zero deltas, we don't want it
 	if (NumZeroDeltas)
@@ -362,13 +362,13 @@ int FUnixTime::CalibrateAndSelectClock()
 
 		FCStringAnsi::Snprintf(Buffer, sizeof(Buffer), "Selected clock_id %d (%s) since it is the fastest support clock without zero deltas.\n",
 			Clocks[ChosenClock].Id, Clocks[ChosenClock].Desc);
-		FCStringAnsi::Strncat(CalibrationLog, Buffer, sizeof(CalibrationLog));
+		FCStringAnsi::StrncatTruncateDest(CalibrationLog, sizeof(CalibrationLog), Buffer);
 
 		// Warn if our current clock source cannot be called at least 1M times a second (<30k a frame) as this may affect tight loops
 		if (ChosenClock != 0 && Clocks[ChosenClock].Rate < 1000000)
 		{
 			FCStringAnsi::Snprintf(Buffer, sizeof(Buffer), "The clock source is too slow on this machine, performance may be affected.\n");
-			FCStringAnsi::Strncat(CalibrationLog, Buffer, sizeof(CalibrationLog));
+			FCStringAnsi::StrncatTruncateDest(CalibrationLog, sizeof(CalibrationLog), Buffer);
 		}
 
 		return Clocks[ChosenClock].Id;
diff --git a/Engine/Source/Runtime/Core/Private/Windows/WindowsErrorOutputDevice.cpp b/Engine/Source/Runtime/Core/Private/Windows/WindowsErrorOutputDevice.cpp
--- a/Engine/Source/Runtime/Core/Private/Windows/WindowsErrorOutputDevice.cpp
+++ b/Engine/Source/Runtime/Core/Private/Windows/WindowsErrorOutputDevice.cpp
@@ -58,7 +58,7 @@ void FWindowsErrorOutputDevice::Serialize( const TCHAR* Msg, ELogVerbosity::Type
 		if (FCString::Strncmp(GErrorHist, Msg, FMath::Min<int32>(UE_ARRAY_COUNT(GErrorHist), FCString::Strlen(Msg))))
 		{
 			FCString::Strncpy(GErrorHist, Msg, UE_ARRAY_COUNT(GErrorHist) - 5);
-			FCString::Strncat(GErrorHist, TEXT("\r\n\r\n"), UE_ARRAY_COUNT(GErrorHist) - 1);
+			FCString::StrncatTruncateDest(GErrorHist, UE_ARRAY_COUNT(GErrorHist), TEXT("\r\n\r\n"));
 		}
 	}
 	else
diff --git a/Engine/Source/Runtime/Core/Private/Windows/WindowsPlatformCrashContext.cpp b/Engine/Source/Runtime/Core/Private/Windows/WindowsPlatformCrashContext.cpp
--- a/Engine/Source/Runtime/Core/Private/Windows/WindowsPlatformCrashContext.cpp
+++ b/Engine/Source/Runtime/Core/Private/Windows/WindowsPlatformCrashContext.cpp
@@ -369,11 +369,11 @@ bool CreateCrashReportClientPath(TCHAR* OutClientPath, int32 MaxLength)
 
 		// Find the path to crash reporter binary. Avoid creating FStrings.
 		*OutClientPath = TCHAR('\0');
-		FCString::Strncat(OutClientPath, EngineDir, MaxLength);
-		FCString::Strncat(OutClientPath, TEXT("Binaries/"), MaxLength);
-		FCString::Strncat(OutClientPath, BinariesDir, MaxLength);
-		FCString::Strncat(OutClientPath, TEXT("/"), MaxLength);
-		FCString::Strncat(OutClientPath, CrashReportClientExeName, MaxLength);
+		FCString::StrncatTruncateDest(OutClientPath, MaxLength, EngineDir);
+		FCString::StrncatTruncateDest(OutClientPath, MaxLength, TEXT("Binaries/"));
+		FCString::StrncatTruncateDest(OutClientPath, MaxLength, BinariesDir);
+		FCString::StrncatTruncateDest(OutClientPath, MaxLength, TEXT("/"));
+		FCString::StrncatTruncateDest(OutClientPath, MaxLength, CrashReportClientExeName);
 
 		const DWORD Results = GetFileAttributesW(OutClientPath);
 		return Results != INVALID_FILE_ATTRIBUTES;
@@ -446,7 +446,7 @@ FProcHandle LaunchCrashReportClient(void** OutWritePipe, void** OutReadPipe, uin
 	{
 		TCHAR PidStr[128] = { 0 };
 		FCString::Sprintf(PidStr, TEXT(" -MONITOR=%u"), FPlatformProcess::GetCurrentProcessId());
-		FCString::Strncat(CrashReporterClientArgs, PidStr, CR_CLIENT_MAX_ARGS_LEN);
+		FCString::StrncatTruncateDest(CrashReporterClientArgs, CR_CLIENT_MAX_ARGS_LEN, PidStr);
 	}
 
 	{
@@ -454,12 +454,12 @@ FProcHandle LaunchCrashReportClient(void** OutWritePipe, void** OutReadPipe, uin
 		// The summaries are matched together using this process group ID.
 		TCHAR AnalyticsSummaryGuid[128] = { 0 };
 		FCString::Sprintf(AnalyticsSummaryGuid, TEXT(" -ProcessGroupId=%s"), *FApp::GetInstanceId().ToString(EGuidFormats::Digits));
-		FCString::Strncat(CrashReporterClientArgs, AnalyticsSummaryGuid, CR_CLIENT_MAX_ARGS_LEN);
+		FCString::StrncatTruncateDest(CrashReporterClientArgs, CR_CLIENT_MAX_ARGS_LEN, AnalyticsSummaryGuid);
 	}
 
 #if DISABLE_CRC_SUBMIT_AND_RESTART_BUTTON
 	// Prevent showing the button if the application cannot be restarted.
-	FCString::Strncat(CrashReporterClientArgs, TEXT(" -HideSubmitAndRestart"), CR_CLIENT_MAX_ARGS_LEN);
+	FCString::StrncatTruncateDest(CrashReporterClientArgs, CR_CLIENT_MAX_ARGS_LEN, TEXT(" -HideSubmitAndRestart"));
 #endif
 
 	// Parse commandline arguments relevant to pass to the client. Note that since we run this from static initialization
@@ -475,8 +475,8 @@ FProcHandle LaunchCrashReportClient(void** OutWritePipe, void** OutReadPipe, uin
 				TCHAR Path[MAX_PATH];
 				if (FParse::Value(ArgList[It], TEXT("abscrashreportclientlog="), Path, UE_ARRAY_COUNT(Path)))
 				{
-					FCString::Strncat(CrashReporterClientArgs, TEXT(" -abslog="), CR_CLIENT_MAX_ARGS_LEN);
-					FCString::Strncat(CrashReporterClientArgs, Path, CR_CLIENT_MAX_ARGS_LEN);
+					FCString::StrncatTruncateDest(CrashReporterClientArgs, CR_CLIENT_MAX_ARGS_LEN, TEXT(" -abslog="));
+					FCString::StrncatTruncateDest(CrashReporterClientArgs, CR_CLIENT_MAX_ARGS_LEN, Path);
 				}
 				
 			#if !USE_NULL_RHI
@@ -484,19 +484,19 @@ FProcHandle LaunchCrashReportClient(void** OutWritePipe, void** OutReadPipe, uin
 				if (bHasNullRHIOnCommandline)
 			#endif
 				{
-					FCString::Strncat(CrashReporterClientArgs, TEXT(" -nullrhi"), CR_CLIENT_MAX_ARGS_LEN);
+					FCString::StrncatTruncateDest(CrashReporterClientArgs, CR_CLIENT_MAX_ARGS_LEN, TEXT(" -nullrhi"));
 				}
 
 				// Pass through any unattended flag
 				if (FParse::Param(ArgList[It], TEXT("unattended")))
 				{
-					FCString::Strncat(CrashReporterClientArgs, TEXT(" -unattended"), CR_CLIENT_MAX_ARGS_LEN);
+					FCString::StrncatTruncateDest(CrashReporterClientArgs, CR_CLIENT_MAX_ARGS_LEN, TEXT(" -unattended"));
 				}
 
 				// Whether to clean up crash reports after send
 				if (FParse::Param(ArgList[It], TEXT("CleanCrashReports")))
 				{
-					FCString::Strncat(CrashReporterClientArgs, TEXT(" -CleanCrashReports"), CR_CLIENT_MAX_ARGS_LEN);
+					FCString::StrncatTruncateDest(CrashReporterClientArgs, CR_CLIENT_MAX_ARGS_LEN, TEXT(" -CleanCrashReports"));
 				}
 			}
 		}
@@ -504,7 +504,7 @@ FProcHandle LaunchCrashReportClient(void** OutWritePipe, void** OutReadPipe, uin
 		// Tell CRC if we're using low integrity paths (AppData\LocalLow instead of AppData\Local).
 		if (FWindowsPlatformProcess::ShouldExpectLowIntegrityLevel())
 		{
-			FCString::Strncat(CrashReporterClientArgs, TEXT(" -ExpectLowIntegrityLevel"), CR_CLIENT_MAX_ARGS_LEN);
+			FCString::StrncatTruncateDest(CrashReporterClientArgs, CR_CLIENT_MAX_ARGS_LEN, TEXT(" -ExpectLowIntegrityLevel"));
 		}
 	}
 
@@ -514,7 +514,7 @@ FProcHandle LaunchCrashReportClient(void** OutWritePipe, void** OutReadPipe, uin
 		FString DisasterRecoveryServiceCommandLine;
 		DisasterRecoveryServiceCommandLine += FString::Printf(TEXT(" -ConcertServer=\"%s\""), *RecoveryService::GetRecoveryServerName()); // Must be in-sync with disaster recovery client module.
 
-		FCString::Strncat(CrashReporterClientArgs, *DisasterRecoveryServiceCommandLine, CR_CLIENT_MAX_ARGS_LEN);
+		FCString::StrncatTruncateDest(CrashReporterClientArgs, CR_CLIENT_MAX_ARGS_LEN, *DisasterRecoveryServiceCommandLine);
 	}
 #endif
 
@@ -1581,11 +1581,11 @@ private:
 			if (ExceptionInfo->ExceptionRecord->ExceptionCode != EnsureExceptionCode && ExceptionInfo->ExceptionRecord->ExceptionCode != AssertExceptionCode)
 			{
 				CreateExceptionInfoString(ExceptionInfo->ExceptionRecord, GErrorExceptionDescription, UE_ARRAY_COUNT(GErrorExceptionDescription));
-				FCString::Strncat(GErrorHist, GErrorExceptionDescription, UE_ARRAY_COUNT(GErrorHist));
-				FCString::Strncat(GErrorHist, TEXT("\r\n\r\n"), UE_ARRAY_COUNT(GErrorHist));
+				FCString::StrncatTruncateDest(GErrorHist, UE_ARRAY_COUNT(GErrorHist), GErrorExceptionDescription);
+				FCString::StrncatTruncateDest(GErrorHist, UE_ARRAY_COUNT(GErrorHist), TEXT("\r\n\r\n"));
 			}
 
-			FCString::Strncat(GErrorHist, ANSI_TO_TCHAR(StackTrace), UE_ARRAY_COUNT(GErrorHist));
+			FCString::StrncatTruncateDest(GErrorHist, UE_ARRAY_COUNT(GErrorHist), ANSI_TO_TCHAR(StackTrace));
 
 			GMalloc->Free(StackTrace);
 		}
diff --git a/Engine/Source/Runtime/Core/Private/Windows/WindowsPlatformMisc.cpp b/Engine/Source/Runtime/Core/Private/Windows/WindowsPlatformMisc.cpp
--- a/Engine/Source/Runtime/Core/Private/Windows/WindowsPlatformMisc.cpp
+++ b/Engine/Source/Runtime/Core/Private/Windows/WindowsPlatformMisc.cpp
@@ -343,8 +343,8 @@ int32 GetOSVersionsHelper( TCHAR* OutOSVersionLabel, int32 OSVersionLabelLength,
 		ErrorCode |= FWindowsOSVersionHelper::ERROR_GETVERSIONEX_FAILED;
 	}
 
-	FCString::Strcpy(OutOSVersionLabel, OSVersionLabelLength, *OSVersionLabel);
-	FCString::Strcpy(OutOSSubVersionLabel, OSSubVersionLabelLength, *OSSubVersionLabel);
+	FCString::Strncpy(OutOSVersionLabel, *OSVersionLabel, OSVersionLabelLength);
+	FCString::Strncpy(OutOSSubVersionLabel, *OSSubVersionLabel, OSSubVersionLabelLength);
 
 	return ErrorCode;
 }
diff --git a/Engine/Source/Runtime/Core/Private/Windows/WindowsPlatformProcess.cpp b/Engine/Source/Runtime/Core/Private/Windows/WindowsPlatformProcess.cpp
--- a/Engine/Source/Runtime/Core/Private/Windows/WindowsPlatformProcess.cpp
+++ b/Engine/Source/Runtime/Core/Private/Windows/WindowsPlatformProcess.cpp
@@ -76,7 +76,7 @@ namespace WindowsPlatformProcess
 
 		TCHAR MutexName[MAX_SPRINTF] = TEXT("");
 
-		FCString::Strcpy(MutexName, MAX_SPRINTF, TEXT("UnrealEngine4"));
+		FCString::Strncpy(MutexName, TEXT("UnrealEngine4"), MAX_SPRINTF);
 
 		GNamedMutex = CreateMutex(NULL, true, MutexName);
 
diff --git a/Engine/Source/Runtime/Core/Private/Windows/WindowsPlatformStackWalk.cpp b/Engine/Source/Runtime/Core/Private/Windows/WindowsPlatformStackWalk.cpp
--- a/Engine/Source/Runtime/Core/Private/Windows/WindowsPlatformStackWalk.cpp
+++ b/Engine/Source/Runtime/Core/Private/Windows/WindowsPlatformStackWalk.cpp
@@ -511,8 +511,8 @@ UE_AUTORTFM_ALWAYS_OPEN void FWindowsPlatformStackWalk::EnumerateSymbolInfosForP
 					if (SymFromInlineContext(ProcessHandle, ProgramCounter, InlineContext, &Displacement, Symbol))
 					{
 						FCStringAnsi::Strncpy(SymbolInfo.FunctionName, "[Inline Frame] ", FProgramCounterSymbolInfo::MAX_NAME_LENGTH);
-						FCStringAnsi::Strncat(SymbolInfo.FunctionName, Symbol->Name, FProgramCounterSymbolInfo::MAX_NAME_LENGTH);
-						FCStringAnsi::Strncat(SymbolInfo.FunctionName, "()", FProgramCounterSymbolInfo::MAX_NAME_LENGTH);
+						FCStringAnsi::StrncatTruncateDest(SymbolInfo.FunctionName, FProgramCounterSymbolInfo::MAX_NAME_LENGTH, Symbol->Name);
+						FCStringAnsi::StrncatTruncateDest(SymbolInfo.FunctionName, FProgramCounterSymbolInfo::MAX_NAME_LENGTH, "()");
 
 						SymbolInfo.OffsetInModule = ProgramCounter - Symbol->ModBase;
 
@@ -544,7 +544,7 @@ UE_AUTORTFM_ALWAYS_OPEN void FWindowsPlatformStackWalk::EnumerateSymbolInfosForP
 
 		// Write out function name.
 		FCStringAnsi::Strncpy(SymbolInfo.FunctionName, Symbol->Name + Offset, FProgramCounterSymbolInfo::MAX_NAME_LENGTH);
-		FCStringAnsi::Strncat(SymbolInfo.FunctionName, "()", FProgramCounterSymbolInfo::MAX_NAME_LENGTH);
+		FCStringAnsi::StrncatTruncateDest(SymbolInfo.FunctionName, FProgramCounterSymbolInfo::MAX_NAME_LENGTH, "()");
 	}
 
 	// Get filename and line number.
diff --git a/Engine/Source/Runtime/Core/Private/Windows/WindowsRunnableThread.cpp b/Engine/Source/Runtime/Core/Private/Windows/WindowsRunnableThread.cpp
--- a/Engine/Source/Runtime/Core/Private/Windows/WindowsRunnableThread.cpp
+++ b/Engine/Source/Runtime/Core/Private/Windows/WindowsRunnableThread.cpp
@@ -88,8 +88,8 @@ uint32 FRunnableThreadWin::GuardedRun()
 				GWarn->Flush();
 
 				// Append the thread name at the end of the error report.
-				FCString::Strncat( GErrorHist, TEXT(LINE_TERMINATOR_ANSI "Crash in runnable thread " ), UE_ARRAY_COUNT( GErrorHist ) );
-				FCString::Strncat( GErrorHist, *ThreadName, UE_ARRAY_COUNT( GErrorHist ) );
+				FCString::StrncatTruncateDest( GErrorHist, UE_ARRAY_COUNT(GErrorHist), TEXT(LINE_TERMINATOR_ANSI "Crash in runnable thread " ) );
+				FCString::StrncatTruncateDest( GErrorHist, UE_ARRAY_COUNT(GErrorHist), *ThreadName );
 
 				// Crashed.
 				ExitCode = 1;
diff --git a/Engine/Source/Runtime/Core/Public/Android/AndroidPlatformString.h b/Engine/Source/Runtime/Core/Public/Android/AndroidPlatformString.h
--- a/Engine/Source/Runtime/Core/Public/Android/AndroidPlatformString.h
+++ b/Engine/Source/Runtime/Core/Public/Android/AndroidPlatformString.h
@@ -53,10 +53,12 @@ struct FAndroidPlatformString : public FGenericPlatformString
 	/** 
 	 * Widechar implementation 
 	 **/
-	static FORCEINLINE WIDECHAR* Strcpy(WIDECHAR* Dest, SIZE_T DestCount, const WIDECHAR* Src)
+	static FORCEINLINE WIDECHAR* Strcpy(WIDECHAR* Dest, const WIDECHAR* Src)
 	{
 		if (!Dest || !Src)
-			return NULL;
+		{
+			return Dest;
+		}
 
 		int Pos = 0;
 		while (Src[Pos])
@@ -70,10 +72,20 @@ struct FAndroidPlatformString : public FGenericPlatformString
 		return Dest;
 	}
 
+	// This version was deprecated because Strcpy taking a size field does not match the CRT strcpy, and on some platforms the size field was being ignored.
+	// If the memzeroing causes a performance problem, request a new function StrcpyTruncate.
+	UE_DEPRECATED(5.6, "Use Strncpy instead. Note that Strncpy has a behavior difference from Strcpy: it memzeroes the entire DestCount-sized buffer after the end of string.")
+	static FORCEINLINE WIDECHAR* Strcpy(WIDECHAR* Dest, SIZE_T DestCount, const WIDECHAR* Src)
+	{
+		return Strcpy(Dest, Src);
+	}
+
 	static FORCEINLINE WIDECHAR* Strncpy(WIDECHAR* Dest, const WIDECHAR* Src, SIZE_T MaxLen)
 	{
 		if (!Dest || !Src)
-			return NULL;
+		{
+			return Dest;
+		}
 
 		int Pos = 0;
 		while ((Pos < MaxLen) && Src[Pos])
@@ -92,23 +104,51 @@ struct FAndroidPlatformString : public FGenericPlatformString
 		return Dest;
 	}
 
-	static FORCEINLINE WIDECHAR* Strcat(WIDECHAR* Dest, SIZE_T DestCount, const WIDECHAR* Src)
+	static FORCEINLINE WIDECHAR* Strcat(WIDECHAR* Dest, const WIDECHAR* Src)
 	{
 		if (!Dest || !Src)
-			return NULL;
-
-		int DestPos = Strlen(Dest);
-		int SrcPos = 0;
-
-		while (Src[SrcPos])
 		{
-			Dest[DestPos] = Src[SrcPos];
-			++SrcPos;
-			++DestPos;
+			return Dest;
+		}
+		if (!*Src)
+		{
+			return Dest;
 		}
 
-		Dest[DestPos] = 0;
+		WIDECHAR* NewDest = Dest + Strlen(Dest);
+		do
+		{
+			*NewDest++ = *Src++;
+		} while (*Src);
+		*NewDest = 0;
+		return Dest;
+	}
 
+	// This version was deprecated because Strcat taking a size field does not match the CRT strcat, and on some platforms the size field was being ignored.
+	UE_DEPRECATED(5.6, "Use Strncat instead. !!NOTE THAT STRNCAT takes SrcLen rather than DestCount. You must call Strncat(Dest, Src, DestCount - Strlen(Dest) - 1).")
+	static FORCEINLINE WIDECHAR* Strcat(WIDECHAR* Dest, SIZE_T DestCount, const WIDECHAR* Src)
+	{
+		return Strcat(Dest, Src);
+	}
+
+	static FORCEINLINE WIDECHAR* Strncat(WIDECHAR* Dest, const WIDECHAR* Src, SIZE_T SrcLen)
+	{
+		if (!Dest || !Src)
+		{
+			return Dest;
+		}
+		if (!SrcLen || !*Src)
+		{
+			return Dest;
+		}
+		WIDECHAR* NewDest = Dest + Strlen(Dest);
+		SIZE_T AppendedCount = 0;
+		do
+		{
+			*NewDest++ = *Src++;
+			++AppendedCount;
+		} while (AppendedCount < SrcLen && *Src);
+		*NewDest = 0;
 		return Dest;
 	}
 
@@ -462,11 +502,19 @@ struct FAndroidPlatformString : public FGenericPlatformString
 	/** 
 	 * Ansi implementation 
 	 **/
-	static FORCEINLINE ANSICHAR* Strcpy(ANSICHAR* Dest, SIZE_T DestCount, const ANSICHAR* Src)
+	static FORCEINLINE ANSICHAR* Strcpy(ANSICHAR* Dest, const ANSICHAR* Src)
 	{
 		return strcpy( Dest, Src );
 	}
 
+	// This version was deprecated because Strcpy taking a size field does not match the CRT strcpy, and on some platforms the size field was being ignored.
+	// If the memzeroing causes a performance problem, request a new function StrcpyTruncate.
+	UE_DEPRECATED(5.6, "Use Strncpy instead. Note that Strncpy has a behavior difference from Strcpy: it memzeroes the entire DestCount-sized buffer after the end of string.")
+	static FORCEINLINE ANSICHAR* Strcpy(ANSICHAR* Dest, SIZE_T DestCount, const ANSICHAR* Src)
+	{
+		return Strcpy(Dest, Src);
+	}
+
 	static FORCEINLINE ANSICHAR* Strncpy(ANSICHAR* Dest, const ANSICHAR* Src, SIZE_T MaxLen)
 	{
 		::strncpy(Dest, Src, MaxLen);
@@ -474,11 +522,23 @@ struct FAndroidPlatformString : public FGenericPlatformString
 		return Dest;
 	}
 
-	static FORCEINLINE ANSICHAR* Strcat(ANSICHAR* Dest, SIZE_T DestCount, const ANSICHAR* Src)
+	static FORCEINLINE ANSICHAR* Strcat(ANSICHAR* Dest, const ANSICHAR* Src)
 	{
 		return strcat( Dest, Src );
 	}
 
+	// This version was deprecated because Strcat taking a size field does not match the CRT strcat, and on some platforms the size field was being ignored.
+	UE_DEPRECATED(5.6, "Use Strncat instead. !!NOTE THAT STRNCAT takes SrcLen rather than DestCount. You must call Strncat(Dest, Src, DestCount - Strlen(Dest) - 1).")
+	static FORCEINLINE ANSICHAR* Strcat(ANSICHAR* Dest, SIZE_T DestCount, const ANSICHAR* Src)
+	{
+		return Strcat(Dest, Src);
+	}
+
+	static FORCEINLINE ANSICHAR* Strncat(ANSICHAR* Dest, const ANSICHAR* Src, SIZE_T SrcLen)
+	{
+		return strncat(Dest, Src, SrcLen);
+	}
+
 	static FORCEINLINE int32 Strcmp( const ANSICHAR* String1, const ANSICHAR* String2 )
 	{
 		return strcmp(String1, String2);
@@ -563,9 +623,18 @@ struct FAndroidPlatformString : public FGenericPlatformString
 	/**
 	 * UTF8CHAR implementation.
 	 */
+	static FORCEINLINE UTF8CHAR* Strcpy(UTF8CHAR* Dest, const UTF8CHAR* Src)
+	{
+		return (UTF8CHAR*)Strcpy((ANSICHAR*)Dest, (const ANSICHAR*)Src);
+	}
+	// This version was deprecated because Strcpy taking a size field does not match the CRT strcpy, and on some platforms the size field was being ignored.
+	// If the memzeroing causes a performance problem, request a new function StrcpyTruncate.
+	UE_DEPRECATED(5.6, "Use Strncpy instead. Note that Strncpy has a behavior difference from Strcpy: it memzeroes the entire DestCount-sized buffer after the end of string.")
 	static FORCEINLINE UTF8CHAR* Strcpy(UTF8CHAR* Dest, SIZE_T DestCount, const UTF8CHAR* Src)
 	{
+		PRAGMA_DISABLE_DEPRECATION_WARNINGS
 		return (UTF8CHAR*)Strcpy((ANSICHAR*)Dest, DestCount, (const ANSICHAR*)Src);
+		PRAGMA_ENABLE_DEPRECATION_WARNINGS
 	}
 
 	static FORCEINLINE UTF8CHAR* Strncpy(UTF8CHAR* Dest, const UTF8CHAR* Src, SIZE_T MaxLen)
@@ -573,9 +642,23 @@ struct FAndroidPlatformString : public FGenericPlatformString
 		return (UTF8CHAR*)Strncpy((ANSICHAR*)Dest, (const ANSICHAR*)Src, MaxLen);
 	}
 
+	static FORCEINLINE UTF8CHAR* Strcat(UTF8CHAR* Dest, const UTF8CHAR* Src)
+	{
+		return (UTF8CHAR*)Strcat((ANSICHAR*)Dest, (const ANSICHAR*)Src);
+	}
+
+	// This version was deprecated because Strcat taking a size field does not match the CRT strcat, and on some platforms the size field was being ignored.
+	UE_DEPRECATED(5.6, "Use Strncat instead. !!NOTE THAT STRNCAT takes SrcLen rather than DestCount. You must call Strncat(Dest, Src, DestCount - Strlen(Dest) - 1).")
 	static FORCEINLINE UTF8CHAR* Strcat(UTF8CHAR* Dest, SIZE_T DestCount, const UTF8CHAR* Src)
 	{
+		PRAGMA_DISABLE_DEPRECATION_WARNINGS
 		return (UTF8CHAR*)Strcat((ANSICHAR*)Dest, DestCount, (const ANSICHAR*)Src);
+		PRAGMA_ENABLE_DEPRECATION_WARNINGS
+	}
+
+	static FORCEINLINE UTF8CHAR* Strncat(UTF8CHAR* Dest, const UTF8CHAR* Src, SIZE_T SrcLen)
+	{
+		return (UTF8CHAR*)Strncat((ANSICHAR*)Dest, (const ANSICHAR*)Src, SrcLen);
 	}
 
 	static FORCEINLINE int32 Strcmp(const UTF8CHAR* String1, const UTF8CHAR* String2)
diff --git a/Engine/Source/Runtime/Core/Public/GenericPlatform/GenericWidePlatformString.h b/Engine/Source/Runtime/Core/Public/GenericPlatform/GenericWidePlatformString.h
--- a/Engine/Source/Runtime/Core/Public/GenericPlatform/GenericWidePlatformString.h
+++ b/Engine/Source/Runtime/Core/Public/GenericPlatform/GenericWidePlatformString.h
@@ -36,9 +36,29 @@ public:
 	/**
 	 * Unicode implementation
 	 **/
-	CORE_API static WIDECHAR* Strcpy(WIDECHAR* Dest, SIZE_T DestCount, const WIDECHAR* Src);
+	CORE_API static WIDECHAR* Strcpy(WIDECHAR* Dest, const WIDECHAR* Src);
+	// This version was deprecated because Strcpy taking a size field does not match the CRT strcpy, and on some platforms the size field was being ignored.
+	// If the memzeroing causes a performance problem, request a new function StrcpyTruncate.
+	UE_DEPRECATED(5.6, "Use Strncpy instead. Note that Strncpy has a behavior difference from Strcpy: it memzeroes the entire DestCount-sized buffer after the end of string.")
+	static WIDECHAR* Strcpy(WIDECHAR* Dest, SIZE_T DestCount, const WIDECHAR* Src)
+	{
+		return Strcpy(Dest, Src);
+	}
 	CORE_API static WIDECHAR* Strncpy(WIDECHAR* Dest, const WIDECHAR* Src, SIZE_T MaxLen);
-	CORE_API static WIDECHAR* Strcat(WIDECHAR* Dest, SIZE_T DestCount, const WIDECHAR* Src);
+	CORE_API static WIDECHAR* Strcat(WIDECHAR* Dest, const WIDECHAR* Src);
+	// This version was deprecated because Strcat taking a size field does not match the CRT strcat, and on some platforms the size field was being ignored.
+	UE_DEPRECATED(5.6, "Use Strncat instead. !!NOTE THAT STRNCAT takes SrcLen rather than DestCount. You must call Strncat(Dest, Src, DestCount - Strlen(Dest) - 1).")
+	static WIDECHAR* Strcat(WIDECHAR* Dest, SIZE_T DestCount, const WIDECHAR* Src)
+	{
+		SIZE_T DestLen = Strlen(Dest);
+		return Strncat(Dest, Src, DestLen + 1 <= DestCount ? DestCount - DestLen - 1 : 0);
+	}
+	/**
+	 * Appends the first SrcLen characters of Src to Dest, plus a terminating null-character.
+	 * If the length of the string in Src is less than SrcLen, only the content up to the terminating
+	 * null-character is copied.
+	 */
+	CORE_API static WIDECHAR* Strncat(WIDECHAR* Dest, const WIDECHAR* Src, SIZE_T SrcLen);
 
 	static int32 Strcmp( const WIDECHAR* String1, const WIDECHAR* String2 )
 	{
@@ -192,7 +212,7 @@ public:
 	/**
 	 * Ansi implementation
 	 **/
-	static FORCEINLINE ANSICHAR* Strcpy(ANSICHAR* Dest, SIZE_T DestCount, const ANSICHAR* Src)
+	static FORCEINLINE ANSICHAR* Strcpy(ANSICHAR* Dest, const ANSICHAR* Src)
 	{
 // Skip suggestions about using strcpy_s instead.
 PRAGMA_DISABLE_DEPRECATION_WARNINGS
@@ -200,6 +220,14 @@ PRAGMA_DISABLE_DEPRECATION_WARNINGS
 PRAGMA_ENABLE_DEPRECATION_WARNINGS
 	}
 
+	// This version was deprecated because Strcpy taking a size field does not match the CRT strcpy, and on some platforms the size field was being ignored.
+	// If the memzeroing causes a performance problem, request a new function StrcpyTruncate.
+	UE_DEPRECATED(5.6, "Use Strncpy instead. Note that Strncpy has a behavior difference from Strcpy: it memzeroes the entire DestCount-sized buffer after the end of string.")
+	static FORCEINLINE ANSICHAR* Strcpy(ANSICHAR* Dest, SIZE_T DestCount, const ANSICHAR* Src)
+	{
+		return Strcpy(Dest, Src);
+	}
+
 	static FORCEINLINE ANSICHAR* Strncpy(ANSICHAR* Dest, const ANSICHAR* Src, SIZE_T MaxLen)
 	{
 // Skip suggestions about using strncpy_s instead.
@@ -210,7 +238,7 @@ PRAGMA_ENABLE_DEPRECATION_WARNINGS
 		return Dest;
 	}
 
-	static FORCEINLINE ANSICHAR* Strcat(ANSICHAR* Dest, SIZE_T DestCount, const ANSICHAR* Src)
+	static FORCEINLINE ANSICHAR* Strcat(ANSICHAR* Dest, const ANSICHAR* Src)
 	{
 // Skip suggestions about using strcat_s instead.
 PRAGMA_DISABLE_DEPRECATION_WARNINGS
@@ -218,6 +246,26 @@ PRAGMA_DISABLE_DEPRECATION_WARNINGS
 PRAGMA_ENABLE_DEPRECATION_WARNINGS
 	}
 
+	// This version was deprecated because Strcat taking a size field does not match the CRT strcat, and on some platforms the size field was being ignored.
+	UE_DEPRECATED(5.6, "Use Strncat instead. !!NOTE THAT STRNCAT takes SrcLen rather than DestCount. You must call Strncat(Dest, Src, DestCount - Strlen(Dest) - 1).")
+	static FORCEINLINE ANSICHAR* Strcat(ANSICHAR* Dest, SIZE_T DestCount, const ANSICHAR* Src)
+	{
+		return Strcat(Dest, Src);
+	}
+
+	/**
+	 * Appends the first SrcLen characters of Src to Dest, plus a terminating null-character.
+	 * If the length of the string in Src is less than SrcLen, only the content up to the terminating
+	 * null-character is copied.
+	 */
+	static FORCEINLINE ANSICHAR* Strncat(ANSICHAR* Dest, const ANSICHAR* Src, SIZE_T SrcLen)
+	{
+// Skip suggestions about using strcat_s instead.
+PRAGMA_DISABLE_DEPRECATION_WARNINGS
+		return strncat(Dest, Src, SrcLen);
+PRAGMA_ENABLE_DEPRECATION_WARNINGS
+	}
+
 	static FORCEINLINE int32 Strcmp( const ANSICHAR* String1, const ANSICHAR* String2 )
 	{
 		return strcmp(String1, String2);
@@ -305,6 +353,14 @@ PRAGMA_ENABLE_DEPRECATION_WARNINGS
 	/**
 	 * UTF8CHAR implementation.
 	 */
+	static FORCEINLINE UTF8CHAR* Strcpy(UTF8CHAR* Dest, const UTF8CHAR* Src)
+	{
+		return (UTF8CHAR*)Strcpy((ANSICHAR*)Dest, (const ANSICHAR*)Src);
+	}
+
+	// This version was deprecated because Strcpy taking a size field does not match the CRT strcpy, and on some platforms the size field was being ignored.
+	// If the memzeroing causes a performance problem, request a new function StrcpyTruncate.
+	UE_DEPRECATED(5.6, "Use Strncpy instead. Note that Strncpy has a behavior difference from Strcpy: it memzeroes the entire DestCount-sized buffer after the end of string.")
 	static FORCEINLINE UTF8CHAR* Strcpy(UTF8CHAR* Dest, SIZE_T DestCount, const UTF8CHAR* Src)
 	{
 		return (UTF8CHAR*)Strcpy((ANSICHAR*)Dest, DestCount, (const ANSICHAR*)Src);
@@ -315,9 +371,26 @@ PRAGMA_ENABLE_DEPRECATION_WARNINGS
 		return (UTF8CHAR*)Strncpy((ANSICHAR*)Dest, (const ANSICHAR*)Src, MaxLen);
 	}
 
+	static FORCEINLINE UTF8CHAR* Strcat(UTF8CHAR* Dest, const UTF8CHAR* Src)
+	{
+		return (UTF8CHAR*)Strcat((ANSICHAR*)Dest, (const ANSICHAR*)Src);
+	}
+
+	// This version was deprecated because Strcat taking a size field does not match the CRT strcat, and on some platforms the size field was being ignored.
+	UE_DEPRECATED(5.6, "Use Strncat instead. !!NOTE THAT STRNCAT takes SrcLen rather than DestCount. You must call Strncat(Dest, Src, DestCount - Strlen(Dest) - 1).")
 	static FORCEINLINE UTF8CHAR* Strcat(UTF8CHAR* Dest, SIZE_T DestCount, const UTF8CHAR* Src)
 	{
-		return (UTF8CHAR*)Strcat((ANSICHAR*)Dest, DestCount, (const ANSICHAR*)Src);
+		return Strcat(Dest, Src);
+	}
+
+	/**
+	 * Appends the first SrcLen characters of Src to Dest, plus a terminating null-character.
+	 * If the length of the string in Src is less than SrcLen, only the content up to the terminating
+	 * null-character is copied.
+	 */
+	static FORCEINLINE UTF8CHAR* Strncat(UTF8CHAR* Dest, const UTF8CHAR* Src, SIZE_T SrcLen)
+	{
+		return (UTF8CHAR*)Strncat((ANSICHAR*)Dest, (const ANSICHAR*)Src, SrcLen);
 	}
 
 	static FORCEINLINE int32 Strcmp(const UTF8CHAR* String1, const UTF8CHAR* String2)
diff --git a/Engine/Source/Runtime/Core/Public/GenericPlatform/MicrosoftPlatformString.h b/Engine/Source/Runtime/Core/Public/GenericPlatform/MicrosoftPlatformString.h
--- a/Engine/Source/Runtime/Core/Public/GenericPlatform/MicrosoftPlatformString.h
+++ b/Engine/Source/Runtime/Core/Public/GenericPlatform/MicrosoftPlatformString.h
@@ -54,9 +54,17 @@ struct FMicrosoftPlatformString :
 	}
 #endif
 
-	/** 
+	/**
 	 * Wide character implementation 
 	 **/
+	static FORCEINLINE WIDECHAR* Strcpy(WIDECHAR* Dest, const WIDECHAR* Src)
+	{
+		return (WIDECHAR*)_tcscpy(Dest, Src);
+	}
+
+	// This version was deprecated because Strcpy taking a size field does not match the CRT strcpy, and on some platforms the size field was being ignored.
+	// If the memzeroing causes a performance problem, request a new function StrcpyTruncate.
+	UE_DEPRECATED(5.6, "Use Strncpy instead. Note that Strncpy has a behavior difference from Strcpy: it memzeroes the entire DestCount-sized buffer after the end of string.")
 	static FORCEINLINE WIDECHAR* Strcpy(WIDECHAR* Dest, SIZE_T DestCount, const WIDECHAR* Src)
 	{
 		return (WIDECHAR*)_tcscpy(Dest, Src);
@@ -69,11 +77,38 @@ struct FMicrosoftPlatformString :
 		return Dest;
 	}
 
+	static FORCEINLINE WIDECHAR* Strcat(WIDECHAR* Dest, const WIDECHAR* Src)
+	{
+		return (WIDECHAR*)_tcscat(Dest, Src);
+	}
+
+	// This version was deprecated because Strcat taking a size field does not match the CRT strcat, and on some platforms the size field was being ignored.
+	UE_DEPRECATED(5.6, "Use Strncat instead. !!NOTE THAT STRNCAT takes SrcLen rather than DestCount. You must call Strncat(Dest, Src, DestCount - Strlen(Dest) - 1).")
 	static FORCEINLINE WIDECHAR* Strcat(WIDECHAR* Dest, SIZE_T DestCount, const WIDECHAR* Src)
 	{
 		return (WIDECHAR*)_tcscat(Dest, Src);
 	}
 
+	static FORCEINLINE WIDECHAR* Strncat(WIDECHAR* Dest, const WIDECHAR* Src, SIZE_T SrcLen)
+	{
+		// The Microsoft library has _tcscat_s but not a strncat equivalent. strncat takes SrcLen
+		// but tcscat_s takes size of the destination string buffer (which must also contain null terminator).
+		// We therefore provide our own implementation.
+		if (!SrcLen || !*Src)
+		{
+			return Dest;
+		}
+		WIDECHAR* NewDest = Dest + Strlen(Dest);
+		SIZE_T AppendedCount = 0;
+		do
+		{
+			*NewDest++ = *Src++;
+			++AppendedCount;
+		} while (AppendedCount < SrcLen && *Src);
+		*NewDest = 0;
+		return Dest;
+	}
+
 	static FORCEINLINE int32 Strcmp( const WIDECHAR* String1, const WIDECHAR* String2 )
 	{
 		return _tcscmp(String1, String2);
@@ -163,6 +198,14 @@ struct FMicrosoftPlatformString :
 	/** 
 	 * Ansi implementation 
 	 **/
+	static FORCEINLINE ANSICHAR* Strcpy(ANSICHAR* Dest, const ANSICHAR* Src)
+	{
+		return (ANSICHAR*)strcpy(Dest, Src);
+	}
+
+	// This version was deprecated because Strcpy taking a size field does not match the CRT strcpy, and on some platforms the size field was being ignored.
+	// If the memzeroing causes a performance problem, request a new function StrcpyTruncate.
+	UE_DEPRECATED(5.6, "Use Strncpy instead. Note that Strncpy has a behavior difference from Strcpy: it memzeroes the entire DestCount-sized buffer after the end of string.")
 	static FORCEINLINE ANSICHAR* Strcpy(ANSICHAR* Dest, SIZE_T DestCount, const ANSICHAR* Src)
 	{
 		return (ANSICHAR*)strcpy(Dest, Src);
@@ -175,11 +218,23 @@ struct FMicrosoftPlatformString :
 		return Dest;
 	}
 
-	static FORCEINLINE ANSICHAR* Strcat(ANSICHAR* Dest, SIZE_T DestCount, const ANSICHAR* Src)
+	static FORCEINLINE ANSICHAR* Strcat(ANSICHAR* Dest, const ANSICHAR* Src)
 	{
 		return (ANSICHAR*)strcat( Dest, Src );
 	}
 
+	// This version was deprecated because Strcat taking a size field does not match the CRT strcat, and on some platforms the size field was being ignored.
+	UE_DEPRECATED(5.6, "Use Strncat instead. !!NOTE THAT STRNCAT takes SrcLen rather than DestSize. You must call Strncat(Dest, Src, DestCount - Strlen(Dest) - 1).")
+	static FORCEINLINE ANSICHAR* Strcat(ANSICHAR* Dest, SIZE_T DestCount, const ANSICHAR* Src)
+	{
+		return (ANSICHAR*)strcat(Dest, Src);
+	}
+
+	static FORCEINLINE ANSICHAR* Strncat(ANSICHAR* Dest, const ANSICHAR* Src, SIZE_T SrcLen)
+	{
+		return (ANSICHAR*)strncat(Dest, Src, SrcLen);
+	}
+
 	static FORCEINLINE int32 Strcmp( const ANSICHAR* String1, const ANSICHAR* String2 )
 	{
 		return strcmp(String1, String2);
@@ -278,9 +333,19 @@ struct FMicrosoftPlatformString :
 	/**
 	 * UTF8CHAR implementation.
 	 */
+	static FORCEINLINE UTF8CHAR* Strcpy(UTF8CHAR* Dest, const UTF8CHAR* Src)
+	{
+		return (UTF8CHAR*)Strcpy((ANSICHAR*)Dest, (const ANSICHAR*)Src);
+	}
+
+	// This version was deprecated because Strcpy taking a size field does not match the CRT strcpy, and on some platforms the size field was being ignored.
+	// If the memzeroing causes a performance problem, request a new function StrcpyTruncate.
+	UE_DEPRECATED(5.6, "Use Strncpy instead. Note that Strncpy has a behavior difference from Strcpy: it memzeroes the entire DestCount-sized buffer after the end of string.")
 	static FORCEINLINE UTF8CHAR* Strcpy(UTF8CHAR* Dest, SIZE_T DestCount, const UTF8CHAR* Src)
 	{
+		PRAGMA_DISABLE_DEPRECATION_WARNINGS
 		return (UTF8CHAR*)Strcpy((ANSICHAR*)Dest, DestCount, (const ANSICHAR*)Src);
+		PRAGMA_ENABLE_DEPRECATION_WARNINGS
 	}
 
 	static FORCEINLINE UTF8CHAR* Strncpy(UTF8CHAR* Dest, const UTF8CHAR* Src, SIZE_T MaxLen)
@@ -288,9 +353,23 @@ struct FMicrosoftPlatformString :
 		return (UTF8CHAR*)Strncpy((ANSICHAR*)Dest, (const ANSICHAR*)Src, MaxLen);
 	}
 
+	static FORCEINLINE UTF8CHAR* Strcat(UTF8CHAR* Dest, const UTF8CHAR* Src)
+	{
+		return (UTF8CHAR*)Strcat((ANSICHAR*)Dest, (const ANSICHAR*)Src);
+	}
+
+	// This version was deprecated because Strcat taking a size field does not match the CRT strcat, and on some platforms the size field was being ignored.
+	UE_DEPRECATED(5.6, "Use Strncat instead. !!NOTE THAT STRNCAT takes SrcLen rather than DestSize. You must call Strncat(Dest, Src, DestCount - Strlen(Dest) - 1).")
 	static FORCEINLINE UTF8CHAR* Strcat(UTF8CHAR* Dest, SIZE_T DestCount, const UTF8CHAR* Src)
 	{
+		PRAGMA_DISABLE_DEPRECATION_WARNINGS
 		return (UTF8CHAR*)Strcat((ANSICHAR*)Dest, DestCount, (const ANSICHAR*)Src);
+		PRAGMA_ENABLE_DEPRECATION_WARNINGS
+	}
+
+	static FORCEINLINE UTF8CHAR* Strncat(UTF8CHAR* Dest, const UTF8CHAR* Src, SIZE_T SrcLen)
+	{
+		return (UTF8CHAR*)Strncat((ANSICHAR*)Dest, (const ANSICHAR*)Src, SrcLen);
 	}
 
 	static FORCEINLINE int32 Strcmp(const UTF8CHAR* String1, const UTF8CHAR* String2)
diff --git a/Engine/Source/Runtime/Core/Public/GenericPlatform/StandardPlatformString.h b/Engine/Source/Runtime/Core/Public/GenericPlatform/StandardPlatformString.h
--- a/Engine/Source/Runtime/Core/Public/GenericPlatform/StandardPlatformString.h
+++ b/Engine/Source/Runtime/Core/Public/GenericPlatform/StandardPlatformString.h
@@ -40,7 +40,7 @@ public:
 	/**
 	 * Unicode implementation
 	 **/
-	static FORCEINLINE WIDECHAR* Strcpy(WIDECHAR* Dest, SIZE_T DestCount, const WIDECHAR* Src)
+	static FORCEINLINE WIDECHAR* Strcpy(WIDECHAR* Dest, const WIDECHAR* Src)
 	{
 // Skip suggestions about using wcscpy_s instead.
 PRAGMA_DISABLE_DEPRECATION_WARNINGS
@@ -48,6 +48,17 @@ PRAGMA_DISABLE_DEPRECATION_WARNINGS
 PRAGMA_ENABLE_DEPRECATION_WARNINGS
 	}
 
+	// This version was deprecated because Strcpy taking a size field does not match the CRT strcpy, and on some platforms the size field was being ignored.
+	// If the memzeroing causes a performance problem, request a new function StrcpyTruncate.
+	UE_DEPRECATED(5.6, "Use Strncpy instead. Note that Strncpy has a behavior difference from Strcpy: it memzeroes the entire DestCount-sized buffer after the end of string.")
+	static FORCEINLINE WIDECHAR* Strcpy(WIDECHAR* Dest, SIZE_T DestCount, const WIDECHAR* Src)
+	{
+		// Skip suggestions about using wcscpy_s instead.
+PRAGMA_DISABLE_DEPRECATION_WARNINGS
+		return wcscpy(Dest, Src);
+PRAGMA_ENABLE_DEPRECATION_WARNINGS
+	}
+
 	static FORCEINLINE WIDECHAR* Strncpy(WIDECHAR* Dest, const WIDECHAR* Src, SIZE_T MaxLen)
 	{
 // Skip suggestions about using wcsncpy_s instead.
@@ -58,7 +69,7 @@ PRAGMA_ENABLE_DEPRECATION_WARNINGS
 		return Dest;
 	}
 
-	static FORCEINLINE WIDECHAR* Strcat(WIDECHAR* Dest, SIZE_T DestCount, const WIDECHAR* Src)
+	static FORCEINLINE WIDECHAR* Strcat(WIDECHAR* Dest, const WIDECHAR* Src)
 	{
 // Skip suggestions about using wcscat_s instead.
 PRAGMA_DISABLE_DEPRECATION_WARNINGS
@@ -66,6 +77,33 @@ PRAGMA_DISABLE_DEPRECATION_WARNINGS
 PRAGMA_ENABLE_DEPRECATION_WARNINGS
 	}
 
+	// This version was deprecated because Strcat taking a size field does not match the CRT strcat, and on some platforms the size field was being ignored.
+	UE_DEPRECATED(5.6, "Use Strncat instead. !!NOTE THAT STRNCAT takes SrcLen rather than DestCount. You must call Strncat(Dest, Src, DestCount - Strlen(Dest) - 1).")
+	static FORCEINLINE WIDECHAR* Strcat(WIDECHAR* Dest, SIZE_T DestCount, const WIDECHAR* Src)
+	{
+		// Skip suggestions about using wcscat_s instead.
+PRAGMA_DISABLE_DEPRECATION_WARNINGS
+		return (WIDECHAR*)wcscat(Dest, Src);
+PRAGMA_ENABLE_DEPRECATION_WARNINGS
+	}
+
+	static FORCEINLINE WIDECHAR* Strncat(WIDECHAR* Dest, const WIDECHAR* Src, SIZE_T SrcLen)
+	{
+		if (!SrcLen || !*Src)
+		{
+			return Dest;
+		}
+		WIDECHAR* NewDest = Dest + Strlen(Dest);
+		SIZE_T AppendedCount = 0;
+		do
+		{
+			*NewDest++ = *Src++;
+			++AppendedCount;
+		} while (AppendedCount < SrcLen && *Src);
+		*NewDest = 0;
+		return Dest;
+	}
+
 	static FORCEINLINE int32 Strcmp( const WIDECHAR* String1, const WIDECHAR* String2 )
 	{
 		return wcscmp(String1, String2);
@@ -211,7 +249,7 @@ PRAGMA_ENABLE_DEPRECATION_WARNINGS
 	/**
 	 * Ansi implementation
 	 **/
-	static FORCEINLINE ANSICHAR* Strcpy(ANSICHAR* Dest, SIZE_T DestCount, const ANSICHAR* Src)
+	static FORCEINLINE ANSICHAR* Strcpy(ANSICHAR* Dest, const ANSICHAR* Src)
 	{
 // Skip suggestions about using strcpy_s instead.
 PRAGMA_DISABLE_DEPRECATION_WARNINGS
@@ -219,6 +257,17 @@ PRAGMA_DISABLE_DEPRECATION_WARNINGS
 PRAGMA_ENABLE_DEPRECATION_WARNINGS
 	}
 
+	// This version was deprecated because Strcpy taking a size field does not match the CRT strcpy, and on some platforms the size field was being ignored.
+	// If the memzeroing causes a performance problem, request a new function StrcpyTruncate.
+	UE_DEPRECATED(5.6, "Use Strncpy instead. Note that Strncpy has a behavior difference from Strcpy: it memzeroes the entire DestCount-sized buffer after the end of string.")
+	static FORCEINLINE ANSICHAR* Strcpy(ANSICHAR* Dest, SIZE_T DestCount, const ANSICHAR* Src)
+	{
+		// Skip suggestions about using strcpy_s instead.
+PRAGMA_DISABLE_DEPRECATION_WARNINGS
+		return strcpy(Dest, Src);
+PRAGMA_ENABLE_DEPRECATION_WARNINGS
+	}
+
 	static FORCEINLINE ANSICHAR* Strncpy(ANSICHAR* Dest, const ANSICHAR* Src, SIZE_T MaxLen)
 	{
 // Skip suggestions about using strncpy_s instead.
@@ -229,7 +278,7 @@ PRAGMA_ENABLE_DEPRECATION_WARNINGS
 		return Dest;
 	}
 
-	static FORCEINLINE ANSICHAR* Strcat(ANSICHAR* Dest, SIZE_T DestCount, const ANSICHAR* Src)
+	static FORCEINLINE ANSICHAR* Strcat(ANSICHAR* Dest, const ANSICHAR* Src)
 	{
 // Skip suggestions about using strcat_s instead.
 PRAGMA_DISABLE_DEPRECATION_WARNINGS
@@ -237,6 +286,21 @@ PRAGMA_DISABLE_DEPRECATION_WARNINGS
 PRAGMA_ENABLE_DEPRECATION_WARNINGS
 	}
 
+	// This version was deprecated because Strcat taking a size field does not match the CRT strcat, and on some platforms the size field was being ignored.
+	UE_DEPRECATED(5.6, "Use Strncat instead. !!NOTE THAT STRNCAT takes SrcLen rather than DestSize. You must call Strncat(Dest, Src, DestCount - Strlen(Dest) - 1).")
+	static FORCEINLINE ANSICHAR* Strcat(ANSICHAR* Dest, SIZE_T DestCount, const ANSICHAR* Src)
+	{
+// Skip suggestions about using strcat_s instead.
+PRAGMA_DISABLE_DEPRECATION_WARNINGS
+		return strcat(Dest, Src);
+PRAGMA_ENABLE_DEPRECATION_WARNINGS
+	}
+
+	static FORCEINLINE ANSICHAR* Strncat(ANSICHAR* Dest, const ANSICHAR* Src, SIZE_T SrcLen)
+	{
+		return (ANSICHAR*)strncat(Dest, Src, SrcLen);
+	}
+
 	static FORCEINLINE int32 Strcmp( const ANSICHAR* String1, const ANSICHAR* String2 )
 	{
 		return strcmp(String1, String2);
@@ -324,9 +388,19 @@ PRAGMA_ENABLE_DEPRECATION_WARNINGS
 	/**
 	 * UTF8CHAR implementation.
 	 */
+	static FORCEINLINE UTF8CHAR* Strcpy(UTF8CHAR* Dest, const UTF8CHAR* Src)
+	{
+		return (UTF8CHAR*)Strcpy((ANSICHAR*)Dest, (const ANSICHAR*)Src);
+	}
+
+	// This version was deprecated because Strcpy taking a size field does not match the CRT strcpy, and on some platforms the size field was being ignored.
+	// If the memzeroing causes a performance problem, request a new function StrcpyTruncate.
+	UE_DEPRECATED(5.6, "Use Strncpy instead. Note that Strncpy has a behavior difference from Strcpy: it memzeroes the entire DestCount-sized buffer after the end of string.")
 	static FORCEINLINE UTF8CHAR* Strcpy(UTF8CHAR* Dest, SIZE_T DestCount, const UTF8CHAR* Src)
 	{
+PRAGMA_DISABLE_DEPRECATION_WARNINGS
 		return (UTF8CHAR*)Strcpy((ANSICHAR*)Dest, DestCount, (const ANSICHAR*)Src);
+PRAGMA_ENABLE_DEPRECATION_WARNINGS
 	}
 
 	static FORCEINLINE UTF8CHAR* Strncpy(UTF8CHAR* Dest, const UTF8CHAR* Src, SIZE_T MaxLen)
@@ -334,9 +408,23 @@ PRAGMA_ENABLE_DEPRECATION_WARNINGS
 		return (UTF8CHAR*)Strncpy((ANSICHAR*)Dest, (const ANSICHAR*)Src, MaxLen);
 	}
 
+	static FORCEINLINE UTF8CHAR* Strcat(UTF8CHAR* Dest, const UTF8CHAR* Src)
+	{
+		return (UTF8CHAR*)Strcat((ANSICHAR*)Dest, (const ANSICHAR*)Src);
+	}
+
+	// This version was deprecated because Strcat taking a size field does not match the CRT strcat, and on some platforms the size field was being ignored.
+	UE_DEPRECATED(5.6, "Use Strncat instead. !!NOTE THAT STRNCAT takes SrcLen rather than DestSize. You must call Strncat(Dest, Src, DestCount - Strlen(Dest) - 1).")
 	static FORCEINLINE UTF8CHAR* Strcat(UTF8CHAR* Dest, SIZE_T DestCount, const UTF8CHAR* Src)
 	{
+PRAGMA_DISABLE_DEPRECATION_WARNINGS
 		return (UTF8CHAR*)Strcat((ANSICHAR*)Dest, DestCount, (const ANSICHAR*)Src);
+PRAGMA_ENABLE_DEPRECATION_WARNINGS
+	}
+
+	static FORCEINLINE UTF8CHAR* Strncat(UTF8CHAR* Dest, const UTF8CHAR* Src, SIZE_T SrcLen)
+	{
+		return (UTF8CHAR*)Strncat((ANSICHAR*)Dest, (const ANSICHAR*)Src, SrcLen);
 	}
 
 	static FORCEINLINE int32 Strcmp(const UTF8CHAR* String1, const UTF8CHAR* String2)
diff --git a/Engine/Source/Runtime/Core/Public/HAL/Platform.h b/Engine/Source/Runtime/Core/Public/HAL/Platform.h
--- a/Engine/Source/Runtime/Core/Public/HAL/Platform.h
+++ b/Engine/Source/Runtime/Core/Public/HAL/Platform.h
@@ -970,7 +970,7 @@ int32 main(int32 ArgC, ANSICHAR* Utf8ArgV[]) \
 	{ \
 		FUTF8ToTCHAR ConvertFromUtf8(Utf8ArgV[a]); \
 		ArgV[a] = new TCHAR[ConvertFromUtf8.Length() + 1]; \
-		FCString::Strcpy(ArgV[a], ConvertFromUtf8.Length() + 1, ConvertFromUtf8.Get()); \
+		FCString::Strncpy(ArgV[a], ConvertFromUtf8.Get(), ConvertFromUtf8.Length() + 1); \
 	} \
 	int32 Result = tchar_main(ArgC, ArgV); \
 	for (int32 a = 0; a < ArgC; a++) \
diff --git a/Engine/Source/Runtime/Core/Public/IOS/IOSCommandLineHelper.h b/Engine/Source/Runtime/Core/Public/IOS/IOSCommandLineHelper.h
--- a/Engine/Source/Runtime/Core/Public/IOS/IOSCommandLineHelper.h
+++ b/Engine/Source/Runtime/Core/Public/IOS/IOSCommandLineHelper.h
@@ -28,7 +28,7 @@ class FIOSCommandLineHelper
 				// handle the easy - case, just use the saved command line part as the start, in case it
 				// started with a ?
 				FString CombinedCommandLine = GSavedCommandLine + CommandLine;
-				FCString::Strcpy(CommandLine, CMD_LINE_MAX, *CombinedCommandLine);
+				FCString::Strncpy(CommandLine, *CombinedCommandLine, CMD_LINE_MAX);
 			}
 			else
 			{
@@ -37,16 +37,16 @@ class FIOSCommandLineHelper
 				if (Space == NULL)
 				{
 					// if there is only one token (no spaces), just append us after it
-					FCString::Strcat(CommandLine, CMD_LINE_MAX, *GSavedCommandLine);
+					FCString::StrncatTruncateDest(CommandLine, CMD_LINE_MAX, *GSavedCommandLine);
 				}
 				else
 				{
 					// save off what's after the space (include the space for pasting later)
 					FString AfterSpace(Space);
 					// copy the save part where the space was
-					FCString::Strcpy(Space, CMD_LINE_MAX - (Space - CommandLine), *GSavedCommandLine);
+					FCString::Strncpy(Space, *GSavedCommandLine, CMD_LINE_MAX - (Space - CommandLine));
 					// now put back the 2nd and so on token
-					FCString::Strcat(CommandLine, CMD_LINE_MAX, *AfterSpace);
+					FCString::StrncatTruncateDest(CommandLine, CMD_LINE_MAX, *AfterSpace);
 				}
 			}
 		}
diff --git a/Engine/Source/Runtime/Core/Public/Misc/AutomationTest.h b/Engine/Source/Runtime/Core/Public/Misc/AutomationTest.h
--- a/Engine/Source/Runtime/Core/Public/Misc/AutomationTest.h
+++ b/Engine/Source/Runtime/Core/Public/Misc/AutomationTest.h
@@ -66,7 +66,7 @@ class FTextFilterExpressionEvaluator;
 	{																										\
 		/* This is a rare failure that can occur in some circumstances */									\
 		FProgramCounterSymbolInfo& Info = VariableName.Emplace_GetRef();									\
-		TCString<ANSICHAR>::Strcpy(Info.Filename, FProgramCounterSymbolInfo::MAX_NAME_LENGTH, "Unknown");	\
+		TCString<ANSICHAR>::Strncpy(Info.Filename, "Unknown", FProgramCounterSymbolInfo::MAX_NAME_LENGTH);	\
 		Info.LineNumber = 1;																				\
 	}
 
@@ -3702,7 +3702,7 @@ private:
 		LLM_SCOPE_BYNAME(TEXT("AutomationTest/Framework"));
 		TArray<FProgramCounterSymbolInfo> Stack;
 		FProgramCounterSymbolInfo First;
-		TCString<ANSICHAR>::Strcpy(First.Filename, FProgramCounterSymbolInfo::MAX_NAME_LENGTH, "Unknown");
+		TCString<ANSICHAR>::Strncpy(First.Filename, "Unknown", FProgramCounterSymbolInfo::MAX_NAME_LENGTH);
 		First.LineNumber = 0;
 		Stack.Add(First);
 
diff --git a/Engine/Source/Runtime/Core/Public/Misc/CString.h b/Engine/Source/Runtime/Core/Public/Misc/CString.h
--- a/Engine/Source/Runtime/Core/Public/Misc/CString.h
+++ b/Engine/Source/Runtime/Core/Public/Misc/CString.h
@@ -173,7 +173,12 @@ struct TCString
 	 * @param Src - source string
 	 * @return destination string
 	 */
-	static FORCEINLINE CharType* Strcpy( CharType* Dest, SIZE_T DestCount, const CharType* Src );
+	static FORCEINLINE CharType* Strcpy( CharType* Dest, const CharType* Src );
+
+	// This version was deprecated because Strcpy taking a size field does not match the CRT strcpy, and on some platforms the size field was being ignored.
+	// If the memzeroing causes a performance problem, request a new function StrcpyTruncate.
+	UE_DEPRECATED(5.6, "Use Strncpy instead. Note that Strncpy has a behavior difference from Strcpy: it memzeroes the entire DestCount-sized buffer after the end of string.")
+	static FORCEINLINE CharType* Strcpy(CharType* Dest, SIZE_T DestCount, const CharType* Src);
 
 	/** 
 	 * Copy a string with length checking. Behavior differs from strncpy in that last character is zeroed. 
@@ -183,7 +188,7 @@ struct TCString
 	 * @param MaxLen - max length of the buffer (including null-terminator)
 	 * @return pointer to resulting string buffer
 	 */
-	static FORCEINLINE CharType* Strncpy( CharType* Dest, const CharType* Src, int32 MaxLen );
+	static FORCEINLINE CharType* Strncpy( CharType* Dest, const CharType* Src, SIZE_T MaxLen );
 
 	/**
 	 * strcpy wrapper
@@ -196,7 +201,7 @@ struct TCString
 	template<SIZE_T DestCount>
 	static FORCEINLINE CharType* Strcpy( CharType (&Dest)[DestCount], const CharType* Src ) 
 	{
-		return Strcpy( Dest, DestCount, Src );
+		return Strncpy( Dest, Src, DestCount);
 	}
 
 	/**
@@ -207,7 +212,11 @@ struct TCString
 	 * @param Src - source string
 	 * @return destination string
 	 */
-	static FORCEINLINE CharType* Strcat( CharType* Dest, SIZE_T DestCount, const CharType* Src );
+	static FORCEINLINE CharType* Strcat( CharType* Dest, const CharType* Src );
+
+	// This version was deprecated because Strcat taking a size field does not match the CRT strcat, and on some platforms the size field was being ignored.
+	UE_DEPRECATED(5.6, "Use StrncatTruncateDest instead.")
+	static FORCEINLINE CharType* Strcat(CharType* Dest, SIZE_T DestSize, const CharType* Src);
 
 	/**
 	 * strcat wrapper
@@ -220,25 +229,73 @@ struct TCString
 	template<SIZE_T DestCount>
 	static FORCEINLINE CharType* Strcat( CharType (&Dest)[DestCount], const CharType* Src ) 
 	{ 
-		return Strcat( Dest, DestCount, Src );
+		return StrncatTruncateDest( Dest, DestCount, Src);
 	}
 
-	/** 
-	 * Concatenate a string with length checking.
+	// This version was deprecated because it had a different interpretation of the length argument than CRT strncat.
+	UE_DEPRECATED(5.6, "Use StrncatTruncateDest instead.")
+	static inline CharType* Strncat( CharType* Dest, const CharType* Src, int32 DestSize )
+	{
+		return StrncatTruncateDest(Dest, DestSize, Src);
+	}
+
+	/**
+	 * Append at most SrcLen characters from Src to Dest. Append a null terminator after the last character from Src.
+	 * This version matches the behavior of CRT strncat: the size given is the maximum number to append from the Src str.
+	 * If you instead want to specify the length of the destination buffer, call StrncatTruncateDest.
 	 *
 	 * @param Dest - destination buffer to append to
 	 * @param Src - source buffer to copy from
-	 * @param MaxLen - max length of the buffer
+	 * @param SrcLen - maximum number of characters to copy from Src. If Src length is less than SrcLen,
+	 *        all characters up to the null terminator, and the null terminator, will be copied into Dest.
 	 * @return pointer to resulting string buffer
 	 */
-	static inline CharType* Strncat( CharType* Dest, const CharType* Src, int32 MaxLen )
+	static inline CharType* StrncatTruncateSrc(CharType* Dest, const CharType* Src, int32 SrcLen)
 	{
+		if (SrcLen <= 0 || !*Src)
+		{
+			return Dest;
+		}
+		CharType* NewDest = Dest + Strlen(Dest);
+		int32 AppendedCount = 0;
+		do
+		{
+			*NewDest++ = *Src++;
+			++AppendedCount;
+		} while (AppendedCount < SrcLen && *Src);
+		*NewDest = CharType(0);
+		return Dest;
+	}
+
+	/**
+	 * Concatenate a string with length checking. Unlike the CRT strncat, the size argument is interpreted
+	 * as the size of the Dest buffer (which must include space for the null terminator) rather than the max
+	 * number of Src characters to append.
+	 *
+	 * @param Dest - destination buffer to append to
+	 * @param Src - source buffer to copy from
+	 * @param DestSize - max size of the dest buffer, including null terminator. Max number of non-terminating
+	 *        characters written into Dest is DestSize - 1.
+	 * @return pointer to resulting string buffer
+	 */
+	static inline CharType* StrncatTruncateDest(CharType* Dest, int32 DestSize, const CharType* Src)
+	{
+		if (!*Src)
+		{
+			return Dest;
+		}
 		int32 Len = Strlen(Dest);
-		CharType* NewDest = Dest + Len;
-		if( (MaxLen-=Len) > 0 )
+		if (Len + 1 >= DestSize)
 		{
-			Strncpy( NewDest, Src, MaxLen );
+			return Dest;
 		}
+		CharType* NewDest = Dest + Len;
+		do
+		{
+			*NewDest++ = *Src++;
+			++Len;
+		} while (*Src && Len + 1 < DestSize);
+		*NewDest = CharType(0);
 		return Dest;
 	}
 
@@ -858,24 +915,39 @@ const typename TCString<T>::CharType* TCString<T>::Strnstr(const CharType* Str,
 	return nullptr;
 }
 
+template <typename T> FORCEINLINE
+typename TCString<T>::CharType* TCString<T>::Strcpy(CharType* Dest, const CharType* Src)
+{
+	return FPlatformString::Strcpy(Dest, Src);
+}
+
 template <typename T> FORCEINLINE
 typename TCString<T>::CharType* TCString<T>::Strcpy(CharType* Dest, SIZE_T DestCount, const CharType* Src)
 {
+	PRAGMA_DISABLE_DEPRECATION_WARNINGS
 	return FPlatformString::Strcpy(Dest, DestCount, Src);
+	PRAGMA_ENABLE_DEPRECATION_WARNINGS
 }
 
 template <typename T> FORCEINLINE
-typename TCString<T>::CharType* TCString<T>::Strncpy( CharType* Dest, const CharType* Src, int32 MaxLen )
+typename TCString<T>::CharType* TCString<T>::Strncpy( CharType* Dest, const CharType* Src, SIZE_T MaxLen )
 {
-	check(MaxLen > 0);
 	FPlatformString::Strncpy(Dest, Src, MaxLen);
 	return Dest;
 }
 
+template <typename T> FORCEINLINE
+typename TCString<T>::CharType* TCString<T>::Strcat(CharType* Dest, const CharType* Src)
+{
+	return FPlatformString::Strcat(Dest, Src);
+}
+
 template <typename T> FORCEINLINE
 typename TCString<T>::CharType* TCString<T>::Strcat(CharType* Dest, SIZE_T DestCount, const CharType* Src)
 {
+	PRAGMA_DISABLE_DEPRECATION_WARNINGS
 	return FPlatformString::Strcat(Dest, DestCount, Src);
+	PRAGMA_ENABLE_DEPRECATION_WARNINGS
 }
 
 template <typename T> FORCEINLINE
diff --git a/Engine/Source/Runtime/Core/Tests/HAL/PlatformStringTest.cpp b/Engine/Source/Runtime/Core/Tests/HAL/PlatformStringTest.cpp
--- a/Engine/Source/Runtime/Core/Tests/HAL/PlatformStringTest.cpp
+++ b/Engine/Source/Runtime/Core/Tests/HAL/PlatformStringTest.cpp
@@ -49,4 +49,220 @@ TEST_CASE_NAMED(FPlatformStringTestStrnlen, "System::Core::HAL::PlatformString::
 	CHECK_MESSAGE(TEXT("Strnlen(PaddedBuffer)"), FPlatformString::Strnlen(Buffer, UE_ARRAY_COUNT(Buffer)) == 9);
 }
 
+TEST_CASE_NAMED(FPlatformStringTestStrcpy, "System::Core::HAL::PlatformString::Strcpy", "[ApplicationContextMask][EngineFilter]")
+{
+	constexpr int32 BufferLen = 32;
+	WIDECHAR WideBuffer[BufferLen];
+	UTF8CHAR Utf8Buffer[BufferLen];
+	ANSICHAR AnsiBuffer[BufferLen];
+
+	const WIDECHAR* WideTest = WIDETEXT("12345");
+	const UTF8CHAR* Utf8Test = UTF8TEXT("12345");
+	const ANSICHAR* AnsiTest = "12345";
+	constexpr int32 TestLen = 5;
+
+	auto Reset = [&WideBuffer, &Utf8Buffer, &AnsiBuffer, BufferLen]()
+		{
+			for (int32 n = 0; n < BufferLen; ++n)
+			{
+				WideBuffer[n] = '%';
+				Utf8Buffer[n] = UTF8CHAR('%');
+				AnsiBuffer[n] = '%';
+			}
+		};
+
+	Reset();
+	FPlatformString::Strcpy(WideBuffer, WideTest);
+	FPlatformString::Strcpy(Utf8Buffer, Utf8Test);
+	FPlatformString::Strcpy(AnsiBuffer, AnsiTest);
+	CHECK_MESSAGE(TEXT("WideStrcpy"),
+		WideBuffer[TestLen] == 0 && FPlatformString::Strcmp(WideTest, WideBuffer) == 0 && WideBuffer[TestLen + 1] == '%');
+	CHECK_MESSAGE(TEXT("Utf8Strcpy"),
+		Utf8Buffer[TestLen] == 0 && FPlatformString::Strcmp(Utf8Test, Utf8Buffer) == 0 && Utf8Buffer[TestLen + 1] == '%');
+	CHECK_MESSAGE(TEXT("AnsiStrcpy"),
+		AnsiBuffer[TestLen] == 0 && FPlatformString::Strcmp(AnsiTest, AnsiBuffer) == 0 && AnsiBuffer[TestLen + 1] == '%');
+
+	Reset();
+	FPlatformString::Strncpy(WideBuffer, WideTest, TestLen + 10);
+	FPlatformString::Strncpy(Utf8Buffer, Utf8Test, TestLen + 10);
+	FPlatformString::Strncpy(AnsiBuffer, AnsiTest, TestLen + 10);
+	CHECK_MESSAGE(TEXT("WideStrncpyTestLenPlus10"),
+		WideBuffer[TestLen] == 0 && FPlatformString::Strcmp(WideTest, WideBuffer) == 0 && WideBuffer[TestLen + 10] == '%');
+	CHECK_MESSAGE(TEXT("Utf8StrncpyTestLenPlus10"),
+		Utf8Buffer[TestLen] == 0 && FPlatformString::Strcmp(Utf8Test, Utf8Buffer) == 0 && Utf8Buffer[TestLen + 10] == '%');
+	CHECK_MESSAGE(TEXT("AnsiStrncpyTestLenPlus10"),
+		AnsiBuffer[TestLen] == 0 && FPlatformString::Strcmp(AnsiTest, AnsiBuffer) == 0 && AnsiBuffer[TestLen + 10] == '%');
+
+	Reset();
+	FPlatformString::Strncpy(WideBuffer, WideTest, TestLen + 1);
+	FPlatformString::Strncpy(Utf8Buffer, Utf8Test, TestLen + 1);
+	FPlatformString::Strncpy(AnsiBuffer, AnsiTest, TestLen + 1);
+	CHECK_MESSAGE(TEXT("WideStrncpyTestLenPlus1"),
+		WideBuffer[TestLen] == 0 && FPlatformString::Strcmp(WideTest, WideBuffer) == 0 && WideBuffer[TestLen + 1] == '%');
+	CHECK_MESSAGE(TEXT("Utf8StrncpyTestLenPlus1"),
+		Utf8Buffer[TestLen] == 0 && FPlatformString::Strcmp(Utf8Test, Utf8Buffer) == 0 && Utf8Buffer[TestLen + 1] == '%');
+	CHECK_MESSAGE(TEXT("AnsiStrncpyTestLenPlus1"),
+		AnsiBuffer[TestLen] == 0 && FPlatformString::Strcmp(AnsiTest, AnsiBuffer) == 0 && AnsiBuffer[TestLen + 1] == '%');
+
+	Reset();
+	FPlatformString::Strncpy(WideBuffer, WideTest, TestLen);
+	FPlatformString::Strncpy(Utf8Buffer, Utf8Test, TestLen);
+	FPlatformString::Strncpy(AnsiBuffer, AnsiTest, TestLen);
+	CHECK_MESSAGE(TEXT("WideStrncpyTestLen"),
+		WideBuffer[TestLen - 1] == 0 && WideBuffer[TestLen] == '%' && FPlatformString::Strncmp(WideTest, WideBuffer, TestLen - 1) == 0 && WideBuffer[TestLen + 1] == '%');
+	CHECK_MESSAGE(TEXT("Utf8StrncpyTestLen"),
+		Utf8Buffer[TestLen - 1] == 0 && Utf8Buffer[TestLen] == '%' && FPlatformString::Strncmp(Utf8Test, Utf8Buffer, TestLen - 1) == 0 && Utf8Buffer[TestLen + 1] == '%');
+	CHECK_MESSAGE(TEXT("AnsiStrncpyTestLen"),
+		AnsiBuffer[TestLen - 1] == 0 && AnsiBuffer[TestLen] == '%' && FPlatformString::Strncmp(AnsiTest, AnsiBuffer, TestLen - 1) == 0 && AnsiBuffer[TestLen + 1] == '%');
+
+	Reset();
+	FPlatformString::Strncpy(WideBuffer, WideTest, TestLen - 1);
+	FPlatformString::Strncpy(Utf8Buffer, Utf8Test, TestLen - 1);
+	FPlatformString::Strncpy(AnsiBuffer, AnsiTest, TestLen - 1);
+	CHECK_MESSAGE(TEXT("WideStrncpyTestLenMinus1"),
+		WideBuffer[TestLen - 2] == 0 && WideBuffer[TestLen - 1] == '%' && FPlatformString::Strncmp(WideTest, WideBuffer, TestLen - 2) == 0 && WideBuffer[TestLen + 1] == '%');
+	CHECK_MESSAGE(TEXT("Utf8StrncpyTestLenMinus1"),
+		Utf8Buffer[TestLen - 2] == 0 && Utf8Buffer[TestLen - 1] == '%' && FPlatformString::Strncmp(Utf8Test, Utf8Buffer, TestLen - 2) == 0 && Utf8Buffer[TestLen + 1] == '%');
+	CHECK_MESSAGE(TEXT("AnsiStrncpyTestLenMinus1"),
+		AnsiBuffer[TestLen - 2] == 0 && AnsiBuffer[TestLen - 1] == '%' && FPlatformString::Strncmp(AnsiTest, AnsiBuffer, TestLen - 2) == 0 && AnsiBuffer[TestLen + 1] == '%');
+
+	Reset();
+	FPlatformString::Strncpy(WideBuffer, WideTest, 2);
+	FPlatformString::Strncpy(Utf8Buffer, Utf8Test, 2);
+	FPlatformString::Strncpy(AnsiBuffer, AnsiTest, 2);
+	CHECK_MESSAGE(TEXT("WideStrncpyTwoLen"),
+		WideBuffer[0] == WideTest[0] && WideBuffer[1] == 0 && WideBuffer[2] == '%' && WideBuffer[TestLen] == '%');
+	CHECK_MESSAGE(TEXT("Utf8StrncpyTwoLen"),
+		Utf8Buffer[0] == Utf8Test[0] && Utf8Buffer[1] == 0 && Utf8Buffer[2] == '%' && Utf8Buffer[TestLen] == '%');
+	CHECK_MESSAGE(TEXT("AnsiStrncpyTwoLen"),
+		AnsiBuffer[0] == AnsiTest[0] && AnsiBuffer[1] == 0 && AnsiBuffer[2] == '%' && AnsiBuffer[TestLen] == '%');
+
+	Reset();
+	FPlatformString::Strncpy(WideBuffer, WideTest, 1);
+	FPlatformString::Strncpy(Utf8Buffer, Utf8Test, 1);
+	FPlatformString::Strncpy(AnsiBuffer, AnsiTest, 1);
+	CHECK_MESSAGE(TEXT("WideStrncpyOneLen"),
+		WideBuffer[0] == 0 && WideBuffer[1] == '%' && WideBuffer[TestLen] == '%');
+	CHECK_MESSAGE(TEXT("Utf8StrncpyOneLen"),
+		Utf8Buffer[0] == 0 && Utf8Buffer[1] == '%' && Utf8Buffer[TestLen] == '%');
+	CHECK_MESSAGE(TEXT("AnsiStrncpyOneLen"),
+		AnsiBuffer[0] == 0 && AnsiBuffer[1] == '%' && AnsiBuffer[TestLen] == '%');
+
+	// ZeroLen Strncpy is undefined
+}
+
+TEST_CASE_NAMED(FPlatformStringTestStrcat, "System::Core::HAL::PlatformString::Strcat", "[ApplicationContextMask][EngineFilter]")
+{
+	constexpr int32 BufferLen = 32;
+	WIDECHAR WideBuffer[BufferLen];
+	UTF8CHAR Utf8Buffer[BufferLen];
+	ANSICHAR AnsiBuffer[BufferLen];
+
+	const WIDECHAR* WidePrefix = WIDETEXT("ABCD");
+	const UTF8CHAR* Utf8Prefix = UTF8TEXT("ABCD");
+	const ANSICHAR* AnsiPrefix = "ABCD";
+	const WIDECHAR* WideTest = WIDETEXT("12345");
+	const UTF8CHAR* Utf8Test = UTF8TEXT("12345");
+	const ANSICHAR* AnsiTest = "12345";
+	const WIDECHAR* WidePrefixPlusTest = WIDETEXT("ABCD12345");
+	const UTF8CHAR* Utf8PrefixPlusTest = UTF8TEXT("ABCD12345");
+	const ANSICHAR* AnsiPrefixPlusTest = "ABCD12345";
+	constexpr int32 PrefixLen = 4;
+	constexpr int32 TestLen = 5;
+	constexpr int32 PrefixPlusTestLen = 9;
+
+	auto Reset = [&WideBuffer, &Utf8Buffer, &AnsiBuffer, WidePrefix, Utf8Prefix, AnsiPrefix, BufferLen]()
+		{
+			int32 n = 0;
+			for (; n < PrefixLen + 1; ++n)
+			{
+				WideBuffer[n] = WidePrefix[n];
+				Utf8Buffer[n] = Utf8Prefix[n];
+				AnsiBuffer[n] = AnsiPrefix[n];
+			}
+			for (; n < BufferLen; ++n)
+			{
+				WideBuffer[n] = '%';
+				Utf8Buffer[n] = UTF8CHAR('%');
+				AnsiBuffer[n] = '%';
+			}
+		};
+
+	Reset();
+	FPlatformString::Strcat(WideBuffer, WideTest);
+	FPlatformString::Strcat(Utf8Buffer, Utf8Test);
+	FPlatformString::Strcat(AnsiBuffer, AnsiTest);
+	CHECK_MESSAGE(TEXT("WideStrcat"),
+		WideBuffer[PrefixPlusTestLen] == 0 && FPlatformString::Strcmp(WidePrefixPlusTest, WideBuffer) == 0 && WideBuffer[PrefixPlusTestLen + 1] == '%');
+	CHECK_MESSAGE(TEXT("Utf8Strcat"),
+		Utf8Buffer[PrefixPlusTestLen] == 0 && FPlatformString::Strcmp(Utf8PrefixPlusTest, Utf8Buffer) == 0 && Utf8Buffer[PrefixPlusTestLen + 1] == '%');
+	CHECK_MESSAGE(TEXT("AnsiStrcat"),
+		AnsiBuffer[PrefixPlusTestLen] == 0 && FPlatformString::Strcmp(AnsiPrefixPlusTest, AnsiBuffer) == 0 && AnsiBuffer[PrefixPlusTestLen + 1] == '%');
+
+	Reset();
+	FPlatformString::Strncat(WideBuffer, WideTest, TestLen + 10);
+	FPlatformString::Strncat(Utf8Buffer, Utf8Test, TestLen + 10);
+	FPlatformString::Strncat(AnsiBuffer, AnsiTest, TestLen + 10);
+	CHECK_MESSAGE(TEXT("WideStrncatTestLenPlus10"),
+		WideBuffer[PrefixPlusTestLen] == 0 && FPlatformString::Strcmp(WidePrefixPlusTest, WideBuffer) == 0 && WideBuffer[PrefixPlusTestLen + 1] == '%');
+	CHECK_MESSAGE(TEXT("Utf8StrncatTestLenPlus10"),
+		Utf8Buffer[PrefixPlusTestLen] == 0 && FPlatformString::Strcmp(Utf8PrefixPlusTest, Utf8Buffer) == 0 && Utf8Buffer[PrefixPlusTestLen + 1] == '%');
+	CHECK_MESSAGE(TEXT("AnsiStrncatTestLenPlus10"),
+		AnsiBuffer[PrefixPlusTestLen] == 0 && FPlatformString::Strcmp(AnsiPrefixPlusTest, AnsiBuffer) == 0 && AnsiBuffer[PrefixPlusTestLen + 1] == '%');
+
+	Reset();
+	FPlatformString::Strncat(WideBuffer, WideTest, TestLen);
+	FPlatformString::Strncat(Utf8Buffer, Utf8Test, TestLen);
+	FPlatformString::Strncat(AnsiBuffer, AnsiTest, TestLen);
+	CHECK_MESSAGE(TEXT("WideStrncatTestLenPlus1"),
+		WideBuffer[PrefixPlusTestLen] == 0 && FPlatformString::Strcmp(WidePrefixPlusTest, WideBuffer) == 0 && WideBuffer[PrefixPlusTestLen + 1] == '%');
+	CHECK_MESSAGE(TEXT("Utf8StrncatTestLenPlus1"),
+		Utf8Buffer[PrefixPlusTestLen] == 0 && FPlatformString::Strcmp(Utf8PrefixPlusTest, Utf8Buffer) == 0 && Utf8Buffer[PrefixPlusTestLen + 1] == '%');
+	CHECK_MESSAGE(TEXT("AnsiStrncatTestLenPlus1"),
+		AnsiBuffer[PrefixPlusTestLen] == 0 && FPlatformString::Strcmp(AnsiPrefixPlusTest, AnsiBuffer) == 0 && AnsiBuffer[PrefixPlusTestLen + 1] == '%');
+
+	Reset();
+	FPlatformString::Strncat(WideBuffer, WideTest, TestLen - 1);
+	FPlatformString::Strncat(Utf8Buffer, Utf8Test, TestLen - 1);
+	FPlatformString::Strncat(AnsiBuffer, AnsiTest, TestLen - 1);
+	CHECK_MESSAGE(TEXT("WideStrncatTestLenMinus1"),
+		WideBuffer[PrefixPlusTestLen - 1] == 0 && WideBuffer[PrefixPlusTestLen] == '%' && FPlatformString::Strncmp(WidePrefixPlusTest, WideBuffer, PrefixPlusTestLen - 1) == 0 && WideBuffer[PrefixPlusTestLen + 1] == '%');
+	CHECK_MESSAGE(TEXT("Utf8StrncatTestLenMinus1"),
+		Utf8Buffer[PrefixPlusTestLen - 1] == 0 && Utf8Buffer[PrefixPlusTestLen] == '%' && FPlatformString::Strncmp(Utf8PrefixPlusTest, Utf8Buffer, PrefixPlusTestLen - 1) == 0 && Utf8Buffer[PrefixPlusTestLen + 1] == '%');
+	CHECK_MESSAGE(TEXT("AnsiStrncatTestLenMinus1"),
+		AnsiBuffer[PrefixPlusTestLen - 1] == 0 && AnsiBuffer[PrefixPlusTestLen] == '%' && FPlatformString::Strncmp(AnsiPrefixPlusTest, AnsiBuffer, PrefixPlusTestLen - 1) == 0 && AnsiBuffer[PrefixPlusTestLen + 1] == '%');
+
+	Reset();
+	FPlatformString::Strncat(WideBuffer, WideTest, TestLen - 2);
+	FPlatformString::Strncat(Utf8Buffer, Utf8Test, TestLen - 2);
+	FPlatformString::Strncat(AnsiBuffer, AnsiTest, TestLen - 2);
+	CHECK_MESSAGE(TEXT("WideStrncatTestLenMinus2"),
+		WideBuffer[PrefixPlusTestLen - 2] == 0 && WideBuffer[PrefixPlusTestLen - 1] == '%' && FPlatformString::Strncmp(WidePrefixPlusTest, WideBuffer, PrefixPlusTestLen - 2) == 0 && WideBuffer[PrefixPlusTestLen + 1] == '%');
+	CHECK_MESSAGE(TEXT("Utf8StrncatTestLenMinus2"),
+		Utf8Buffer[PrefixPlusTestLen - 2] == 0 && Utf8Buffer[PrefixPlusTestLen - 1] == '%' && FPlatformString::Strncmp(Utf8PrefixPlusTest, Utf8Buffer, PrefixPlusTestLen - 2) == 0 && Utf8Buffer[PrefixPlusTestLen + 1] == '%');
+	CHECK_MESSAGE(TEXT("AnsiStrncatTestLenMinus2"),
+		AnsiBuffer[PrefixPlusTestLen - 2] == 0 && AnsiBuffer[PrefixPlusTestLen - 1] == '%' && FPlatformString::Strncmp(AnsiPrefixPlusTest, AnsiBuffer, PrefixPlusTestLen - 2) == 0 && AnsiBuffer[PrefixPlusTestLen + 1] == '%');
+
+	Reset();
+	FPlatformString::Strncat(WideBuffer, WideTest, 1);
+	FPlatformString::Strncat(Utf8Buffer, Utf8Test, 1);
+	FPlatformString::Strncat(AnsiBuffer, AnsiTest, 1);
+	CHECK_MESSAGE(TEXT("WideStrncatOneLen"),
+		WideBuffer[PrefixLen + 1] == 0 && WideBuffer[PrefixLen + 2] == '%' && FPlatformString::Strncmp(WidePrefixPlusTest, WideBuffer, PrefixLen + 1) == 0 && WideBuffer[PrefixPlusTestLen + 1] == '%');
+	CHECK_MESSAGE(TEXT("Utf8StrncatOneLen"),
+		Utf8Buffer[PrefixLen + 1] == 0 && Utf8Buffer[PrefixLen + 2] == '%' && FPlatformString::Strncmp(Utf8PrefixPlusTest, Utf8Buffer, PrefixLen + 1) == 0 && Utf8Buffer[PrefixPlusTestLen + 1] == '%');
+	CHECK_MESSAGE(TEXT("AnsiStrncatOneLen"),
+		AnsiBuffer[PrefixLen + 1] == 0 && AnsiBuffer[PrefixLen + 2] == '%' && FPlatformString::Strncmp(AnsiPrefixPlusTest, AnsiBuffer, PrefixLen + 1) == 0 && AnsiBuffer[PrefixPlusTestLen + 1] == '%');
+
+	Reset();
+	FPlatformString::Strncat(WideBuffer, WideTest, 0);
+	FPlatformString::Strncat(Utf8Buffer, Utf8Test, 0);
+	FPlatformString::Strncat(AnsiBuffer, AnsiTest, 0);
+	CHECK_MESSAGE(TEXT("WideStrncatZeroLen"),
+		WideBuffer[PrefixLen] == 0 && WideBuffer[PrefixLen + 1] == '%' && FPlatformString::Strncmp(WidePrefixPlusTest, WideBuffer, PrefixLen) == 0 && WideBuffer[PrefixPlusTestLen + 1] == '%');
+	CHECK_MESSAGE(TEXT("Utf8StrncatZeroLen"),
+		Utf8Buffer[PrefixLen] == 0 && Utf8Buffer[PrefixLen + 1] == '%' && FPlatformString::Strncmp(Utf8PrefixPlusTest, Utf8Buffer, PrefixLen) == 0 && Utf8Buffer[PrefixPlusTestLen + 1] == '%');
+	CHECK_MESSAGE(TEXT("AnsiStrncatZeroLen"),
+		AnsiBuffer[PrefixLen] == 0 && AnsiBuffer[PrefixLen + 1] == '%' && FPlatformString::Strncmp(AnsiPrefixPlusTest, AnsiBuffer, PrefixLen) == 0 && AnsiBuffer[PrefixPlusTestLen + 1] == '%');
+}
 #endif //WITH_TESTS
diff --git a/Engine/Source/Runtime/CoreUObject/Private/Serialization/EditorBulkData.cpp b/Engine/Source/Runtime/CoreUObject/Private/Serialization/EditorBulkData.cpp
--- a/Engine/Source/Runtime/CoreUObject/Private/Serialization/EditorBulkData.cpp
+++ b/Engine/Source/Runtime/CoreUObject/Private/Serialization/EditorBulkData.cpp
@@ -140,7 +140,7 @@ static void LogPackageOpenFailureMessage(const FPackagePath& PackagePath, EPacka
 	if (SystemError != 0)
 	{
 		TCHAR SystemErrorMsg[2048] = { 0 };
-		FPlatformMisc::GetSystemErrorMessage(SystemErrorMsg, sizeof(SystemErrorMsg), SystemError);
+		FPlatformMisc::GetSystemErrorMessage(SystemErrorMsg, UE_ARRAY_COUNT(SystemErrorMsg), SystemError);
 		UE_LOG(LogSerialization, Error, TEXT("Could not open the file '%s' for reading due to system error: '%s' (%d))"), *PackagePath.GetLocalFullPath(PackageSegment), SystemErrorMsg, SystemError);
 	}
 	else
diff --git a/Engine/Source/Runtime/CoreUObject/Private/UObject/PackageTrailer.cpp b/Engine/Source/Runtime/CoreUObject/Private/UObject/PackageTrailer.cpp
--- a/Engine/Source/Runtime/CoreUObject/Private/UObject/PackageTrailer.cpp
+++ b/Engine/Source/Runtime/CoreUObject/Private/UObject/PackageTrailer.cpp
@@ -50,7 +50,7 @@ void LogPackageOpenFailureMessage(const FString& DebugName)
 	if (SystemError != 0)
 	{
 		TCHAR SystemErrorMsg[2048] = { 0 };
-		FPlatformMisc::GetSystemErrorMessage(SystemErrorMsg, sizeof(SystemErrorMsg), SystemError);
+		FPlatformMisc::GetSystemErrorMessage(SystemErrorMsg, UE_ARRAY_COUNT(SystemErrorMsg), SystemError);
 		UE_LOG(LogSerialization, Error, TEXT("Could not open the file '%s' for reading due to system error: '%s' (%d))"), *DebugName, SystemErrorMsg, SystemError);
 	}
 	else
diff --git a/Engine/Source/Runtime/Engine/Private/AsyncCompilationHelpers.cpp b/Engine/Source/Runtime/Engine/Private/AsyncCompilationHelpers.cpp
--- a/Engine/Source/Runtime/Engine/Private/AsyncCompilationHelpers.cpp
+++ b/Engine/Source/Runtime/Engine/Private/AsyncCompilationHelpers.cpp
@@ -442,7 +442,7 @@ namespace AsyncCompilationHelpers
 			for (int32 Index = 2; Index < FStackData::MaxDepth && StackData->Backtrace[Index]; ++Index)
 			{
 				FPlatformStackWalk::ProgramCounterToHumanReadableString(Index, StackData->Backtrace[Index], HumanReadableString, HumanReadableStringSize);
-				FCStringAnsi::Strncat(HumanReadableString, LINE_TERMINATOR_ANSI, HumanReadableStringSize);
+				FCStringAnsi::StrncatTruncateDest(HumanReadableString, HumanReadableStringSize, LINE_TERMINATOR_ANSI);
 			}
 
 			TotalCount += StackData->Count;
diff --git a/Engine/Source/Runtime/Engine/Private/DebugCameraController.cpp b/Engine/Source/Runtime/Engine/Private/DebugCameraController.cpp
--- a/Engine/Source/Runtime/Engine/Private/DebugCameraController.cpp
+++ b/Engine/Source/Runtime/Engine/Private/DebugCameraController.cpp
@@ -287,7 +287,7 @@ FString ADebugCameraController::ConsoleCommand(const FString& Cmd,bool bWriteToL
 
 		const TCHAR* Command = CommandBuffer;
 		// copy the command into a modifiable buffer
-		FCString::Strcpy(CommandBuffer, (CmdLen+1), *Cmd.Left(CmdLen)); 
+		FCString::Strncpy(CommandBuffer, *Cmd.Left(CmdLen), (CmdLen+1));
 
 		// iterate over the line, breaking up on |'s
 		while (FParse::Line(&Command, Line, CmdLen+1))	// The FParse::Line function expects the full array size, including the NULL character.
diff --git a/Engine/Source/Runtime/Engine/Private/Player.cpp b/Engine/Source/Runtime/Engine/Private/Player.cpp
--- a/Engine/Source/Runtime/Engine/Private/Player.cpp
+++ b/Engine/Source/Runtime/Engine/Private/Player.cpp
@@ -42,7 +42,7 @@ FString UPlayer::ConsoleCommand(const FString& Cmd, bool bWriteToLog)
 
 	const TCHAR* Command = CommandBuffer;
 	// copy the command into a modifiable buffer
-	FCString::Strcpy(CommandBuffer, (CmdLen + 1), *Cmd.Left(CmdLen));
+	FCString::Strncpy(CommandBuffer, *Cmd.Left(CmdLen), (CmdLen + 1));
 
 	// iterate over the line, breaking up on |'s
 	while (FParse::Line(&Command, Line, CmdLen + 1))	// The FParse::Line function expects the full array size, including the NULL character.
diff --git a/Engine/Source/Runtime/Engine/Private/SystemSettings.cpp b/Engine/Source/Runtime/Engine/Private/SystemSettings.cpp
--- a/Engine/Source/Runtime/Engine/Private/SystemSettings.cpp
+++ b/Engine/Source/Runtime/Engine/Private/SystemSettings.cpp
@@ -43,11 +43,13 @@ static inline const TCHAR* GetSectionName(bool bIsEditor)
 		static int32 SystemSettingsStrLen = FCString::Strlen(ReturnedOverrideName);
 		if (FCString::Strnicmp(*OverrideSubName, TEXT("SystemSettings"), SystemSettingsStrLen) == 0)
 		{
-			FCString::Strcpy(ReturnedOverrideName + SystemSettingsStrLen, UE_ARRAY_COUNT(ReturnedOverrideName), (*OverrideSubName + SystemSettingsStrLen));
+			FCString::Strncpy(ReturnedOverrideName + SystemSettingsStrLen, (*OverrideSubName + SystemSettingsStrLen),
+				UE_ARRAY_COUNT(ReturnedOverrideName) - SystemSettingsStrLen);
 		}
 		else
 		{
-			FCString::Strcpy(ReturnedOverrideName + SystemSettingsStrLen, UE_ARRAY_COUNT(ReturnedOverrideName), *OverrideSubName);
+			FCString::Strncpy(ReturnedOverrideName + SystemSettingsStrLen, *OverrideSubName,
+				UE_ARRAY_COUNT(ReturnedOverrideName) - SystemSettingsStrLen);
 		}
 		return ReturnedOverrideName;
 	}
diff --git a/Engine/Source/Runtime/Engine/Private/URL.cpp b/Engine/Source/Runtime/Engine/Private/URL.cpp
--- a/Engine/Source/Runtime/Engine/Private/URL.cpp
+++ b/Engine/Source/Runtime/Engine/Private/URL.cpp
@@ -153,7 +153,7 @@ FURL::FURL( FURL* Base, const TCHAR* TextURL, ETravelType Type )
 	TCHAR* TempURL = new TCHAR[URLLength+1];
 	TCHAR* URL = TempURL;
 
-	FCString::Strcpy( TempURL, URLLength + 1, TextURL );
+	FCString::Strncpy( TempURL, TextURL, URLLength + 1 );
 
 	// Copy Base.
 	if( Type==TRAVEL_Relative )
diff --git a/Engine/Source/Runtime/Experimental/IoStore/OnDemand/Private/IoStoreOnDemand.cpp b/Engine/Source/Runtime/Experimental/IoStore/OnDemand/Private/IoStoreOnDemand.cpp
--- a/Engine/Source/Runtime/Experimental/IoStore/OnDemand/Private/IoStoreOnDemand.cpp
+++ b/Engine/Source/Runtime/Experimental/IoStore/OnDemand/Private/IoStoreOnDemand.cpp
@@ -518,7 +518,7 @@ void FIoStoreOnDemandModule::StartupModule()
 		else
 		{
 			TCHAR SystemErrorMsg[MAX_SPRINTF] = { 0 };
-			FPlatformMisc::GetSystemErrorMessage(SystemErrorMsg, sizeof(SystemErrorMsg), Result);
+			FPlatformMisc::GetSystemErrorMessage(SystemErrorMsg, UE_ARRAY_COUNT(SystemErrorMsg), Result);
 
 			UE_LOG(LogIoStoreOnDemand, Error, TEXT("WSAStartup failed due to: %s (%d)"), SystemErrorMsg, Result);
 		}
@@ -563,7 +563,7 @@ void FIoStoreOnDemandModule::ShutdownModule()
 			const uint32 SystemError = FPlatformMisc::GetLastError();
 
 			TCHAR SystemErrorMsg[MAX_SPRINTF] = { 0 };
-			FPlatformMisc::GetSystemErrorMessage(SystemErrorMsg, sizeof(SystemErrorMsg), SystemError);
+			FPlatformMisc::GetSystemErrorMessage(SystemErrorMsg, UE_ARRAY_COUNT(SystemErrorMsg), SystemError);
 
 			UE_LOG(LogIoStoreOnDemand, Error, TEXT("WSACleanup failed due to: %s (%u)"), SystemErrorMsg, SystemError);
 		}
diff --git a/Engine/Source/Runtime/Online/HTTP/Private/Curl/CurlHttp.h b/Engine/Source/Runtime/Online/HTTP/Private/Curl/CurlHttp.h
--- a/Engine/Source/Runtime/Online/HTTP/Private/Curl/CurlHttp.h
+++ b/Engine/Source/Runtime/Online/HTTP/Private/Curl/CurlHttp.h
@@ -98,7 +98,7 @@ namespace
 			Copy = reinterpret_cast<char*>(FMemory::Malloc(StrLen + 1));
 			if (Copy)
 			{
-				FCStringAnsi::Strcpy(Copy, StrLen, ZeroTerminatedString);
+				FCStringAnsi::Strncpy(Copy, ZeroTerminatedString, StrLen + 1);
 				check(FCStringAnsi::Strcmp(Copy, ZeroTerminatedString) == 0);
 			}
 		}
diff --git a/Engine/Source/Runtime/Online/WebSockets/Private/Lws/LwsWebSocket.cpp b/Engine/Source/Runtime/Online/WebSockets/Private/Lws/LwsWebSocket.cpp
--- a/Engine/Source/Runtime/Online/WebSockets/Private/Lws/LwsWebSocket.cpp
+++ b/Engine/Source/Runtime/Online/WebSockets/Private/Lws/LwsWebSocket.cpp
@@ -154,7 +154,7 @@ void FLwsWebSocket::Close(int32 Code, const FString& Reason)
 	// We are doing this conversion here so we don't have to do it on the ws thread
 	FTCHARToUTF8 Convert(*Reason);
 	ANSICHAR* ANSIReason = static_cast<ANSICHAR*>(FMemory::Malloc(Convert.Length() + 1));
-	FCStringAnsi::Strcpy(ANSIReason, Convert.Length(), (const ANSICHAR*)Convert.Get());
+	FCStringAnsi::Strncpy(ANSIReason, (const ANSICHAR*)Convert.Get(), Convert.Length() + 1);
 	ANSIReason[Convert.Length()] = 0;
 
 	UE_LOG(LogWebSockets, Verbose, TEXT("FLwsWebSocket[%d]::Close: Close queued with code=%d reason=%s"), Identifier, Code, *Reason);
diff --git a/Engine/Source/Runtime/Online/WebSockets/Private/Lws/LwsWebSocketsManager.cpp b/Engine/Source/Runtime/Online/WebSockets/Private/Lws/LwsWebSocketsManager.cpp
--- a/Engine/Source/Runtime/Online/WebSockets/Private/Lws/LwsWebSocketsManager.cpp
+++ b/Engine/Source/Runtime/Online/WebSockets/Private/Lws/LwsWebSocketsManager.cpp
@@ -113,7 +113,7 @@ void FLwsWebSocketsManager::InitWebSockets(TArrayView<const FString> Protocols)
 
 		// We need to hold on to the converted strings
 		ANSICHAR* Converted = static_cast<ANSICHAR*>(FMemory::Malloc(ConvertName.Length() + 1));
-		FCStringAnsi::Strcpy(Converted, ConvertName.Length(), (const ANSICHAR*)ConvertName.Get());
+		FCStringAnsi::Strncpy(Converted, (const ANSICHAR*)ConvertName.Get(), ConvertName.Length() + 1);
 		lws_protocols LwsProtocol;
 		FMemory::Memset(&LwsProtocol, 0, sizeof(LwsProtocol));
 		LwsProtocol.name = Converted;
diff --git a/Engine/Source/Runtime/PakFile/Public/IPlatformFilePak.h b/Engine/Source/Runtime/PakFile/Public/IPlatformFilePak.h
--- a/Engine/Source/Runtime/PakFile/Public/IPlatformFilePak.h
+++ b/Engine/Source/Runtime/PakFile/Public/IPlatformFilePak.h
@@ -301,7 +301,7 @@ struct FPakInfo
 				for (int32 Index = 1; Index < CompressionMethods.Num(); Index++)
 				{
 					ANSICHAR* MethodString = &Methods[(Index - 1) * CompressionMethodNameLen];
-					FCStringAnsi::Strcpy(MethodString, CompressionMethodNameLen, TCHAR_TO_ANSI(*CompressionMethods[Index].ToString()));
+					FCStringAnsi::Strncpy(MethodString, TCHAR_TO_ANSI(*CompressionMethods[Index].ToString()), CompressionMethodNameLen);
 				}
 				Ar.Serialize(Methods, BufferSize);
 			}
diff --git a/Engine/Source/Runtime/RHI/Private/TextureProfiler.cpp b/Engine/Source/Runtime/RHI/Private/TextureProfiler.cpp
--- a/Engine/Source/Runtime/RHI/Private/TextureProfiler.cpp
+++ b/Engine/Source/Runtime/RHI/Private/TextureProfiler.cpp
@@ -281,7 +281,7 @@ const char* FTextureProfiler::GetTextureNameString(FName TextureName)
 			const auto ConvertedTextureNameString = StringCast<ANSICHAR>(NewTextureNameString);
 			int Length = ConvertedTextureNameString.Length();
 			TextureNameString = new char[Length + 1];
-			FCStringAnsi::Strcpy(const_cast<char*>(TextureNameString), Length + 1, ConvertedTextureNameString.Get());
+			FCStringAnsi::Strncpy(const_cast<char*>(TextureNameString), ConvertedTextureNameString.Get(), Length + 1);
 		}
 		else
 		{
@@ -289,7 +289,7 @@ const char* FTextureProfiler::GetTextureNameString(FName TextureName)
 			TextureName.GetPlainANSIString(NewTextureNameString);
 			int Length = FCStringAnsi::Strlen(NewTextureNameString);
 			TextureNameString = new char[Length + 1];
-			FCStringAnsi::Strcpy(const_cast<char*>(TextureNameString), Length + 1, NewTextureNameString);
+			FCStringAnsi::Strncpy(const_cast<char*>(TextureNameString), NewTextureNameString, Length + 1);
 		}
 	}
 
diff --git a/Engine/Source/Runtime/RHI/Public/RHICommandList.h b/Engine/Source/Runtime/RHI/Public/RHICommandList.h
--- a/Engine/Source/Runtime/RHI/Public/RHICommandList.h
+++ b/Engine/Source/Runtime/RHI/Public/RHICommandList.h
@@ -567,7 +567,7 @@ public:
 	{
 		int32 Len = FCString::Strlen(Name) + 1;
 		TCHAR* NameCopy  = (TCHAR*)Alloc(Len * (int32)sizeof(TCHAR), (int32)sizeof(TCHAR));
-		FCString::Strcpy(NameCopy, Len, Name);
+		FCString::Strncpy(NameCopy, Name, Len);
 		return NameCopy;
 	}
 
diff --git a/Engine/Source/Runtime/RenderCore/Private/DynamicRenderScaling.cpp b/Engine/Source/Runtime/RenderCore/Private/DynamicRenderScaling.cpp
--- a/Engine/Source/Runtime/RenderCore/Private/DynamicRenderScaling.cpp
+++ b/Engine/Source/Runtime/RenderCore/Private/DynamicRenderScaling.cpp
@@ -128,7 +128,7 @@ FBudget::FBudget(const TCHAR* InName, FHeuristicSettings (*InHeuristicSettingsGe
 	check(InName);
 
 	AnsiName.SetNumZeroed(FCString::Strlen(InName) + 1);
-	FCStringAnsi::Strcpy(AnsiName.GetData(), AnsiName.Num(), TCHAR_TO_ANSI(InName));
+	FCStringAnsi::Strncpy(AnsiName.GetData(), TCHAR_TO_ANSI(InName), AnsiName.Num());
 
 	GlobalListLink.LinkHead(FBudget::GetGlobalList());
 	BudgetId = GRDSGlobalListLinkSize;
diff --git a/Engine/Source/Runtime/RenderCore/Private/ShaderParameterMetadata.cpp b/Engine/Source/Runtime/RenderCore/Private/ShaderParameterMetadata.cpp
--- a/Engine/Source/Runtime/RenderCore/Private/ShaderParameterMetadata.cpp
+++ b/Engine/Source/Runtime/RenderCore/Private/ShaderParameterMetadata.cpp
@@ -940,7 +940,7 @@ static uint32 CacheResourceTableEntriesRecursive(
 			uint32 StringStorageCount = Prefix.Len() + 1;
 			if (MemberNameBuffer)
 			{
-				FCString::Strcpy(MemberNameBuffer, StringStorageCount, Prefix.ToString());
+				FCString::Strncpy(MemberNameBuffer, Prefix.ToString(), StringStorageCount);
 
 				check(Entries);
 				Entries->Add({
diff --git a/Engine/Source/Runtime/Renderer/Private/SceneTextures.cpp b/Engine/Source/Runtime/Renderer/Private/SceneTextures.cpp
--- a/Engine/Source/Runtime/Renderer/Private/SceneTextures.cpp
+++ b/Engine/Source/Runtime/Renderer/Private/SceneTextures.cpp
@@ -562,7 +562,7 @@ FRDGTextureRef FMinimalSceneTextures::FindOrAddUserSceneTexture(FRDGBuilder& Gra
 			FString::Printf(TEXT("UST.%s"), *Name.ToString()) :
 			FString::Printf(TEXT("UST.%s[%d]"), *Name.ToString(), TransientUserTexture.AllocationOrder);
 		TCHAR* HeapDebugName = new TCHAR[DebugName.Len() + 1];
-		FCString::Strcpy(HeapDebugName, DebugName.Len() + 1, *DebugName);
+		FCString::Strncpy(HeapDebugName, *DebugName, DebugName.Len() + 1);
 
 		TransientUserTexture.Texture = GraphBuilder.CreateTexture(Desc, HeapDebugName);
 		TransientUserTexture.Texture->SetDebugNameIsHeapAllocated();
diff --git a/Engine/Source/Runtime/SlateCore/Private/Trace/SlateTrace.cpp b/Engine/Source/Runtime/SlateCore/Private/Trace/SlateTrace.cpp
--- a/Engine/Source/Runtime/SlateCore/Private/Trace/SlateTrace.cpp
+++ b/Engine/Source/Runtime/SlateCore/Private/Trace/SlateTrace.cpp
@@ -450,7 +450,7 @@ void FSlateTrace::WidgetInvalidated(const SWidget* Widget, const SWidget* Invest
 							FProgramCounterSymbolInfo SymbolInfo;
 							FPlatformStackWalk::ProgramCounterToSymbolInfo(StackTrace[StackDepth], SymbolInfo);
 							FPlatformStackWalk::SymbolInfoToHumanReadableString(SymbolInfo, CallStack, CallStackSize);
-							FCStringAnsi::Strncat(CallStack, LINE_TERMINATOR_ANSI, CallStackSize);
+							FCStringAnsi::StrncatTruncateDest(CallStack, CallStackSize, LINE_TERMINATOR_ANSI);
 						}
 						StackDepth++;
 					}
diff --git a/Engine/Source/Runtime/TraceLog/Private/Trace/Control.cpp b/Engine/Source/Runtime/TraceLog/Private/Trace/Control.cpp
--- a/Engine/Source/Runtime/TraceLog/Private/Trace/Control.cpp
+++ b/Engine/Source/Runtime/TraceLog/Private/Trace/Control.cpp
@@ -314,7 +314,7 @@ void Writer_InitializeControl()
 			ANSICHAR Channels[BufferSize] = {};
 			ANSICHAR* Ctx;
 			const bool bState = (ArgV[1][0] != '0');
-			FCStringAnsi::Strcpy(Channels, BufferSize, ArgV[0]);
+			FPlatformString::Strncpy(Channels, ArgV[0], BufferSize);
 			ANSICHAR* Channel = FCStringAnsi::Strtok(Channels, ",", &Ctx);
 			while (Channel)
 			{
diff --git a/Engine/Source/Runtime/VulkanRHI/Private/VulkanMemory.cpp b/Engine/Source/Runtime/VulkanRHI/Private/VulkanMemory.cpp
--- a/Engine/Source/Runtime/VulkanRHI/Private/VulkanMemory.cpp
+++ b/Engine/Source/Runtime/VulkanRHI/Private/VulkanMemory.cpp
@@ -325,7 +325,7 @@ FString VulkanTrackGetString(FVulkanTrackInfo& Track)
 		for (int32 Index = 0; Index < -Track.SizeOrLine; ++Index)
 		{
 			FPlatformStackWalk::ProgramCounterToHumanReadableString(Index, Stack[Index], StackTraceString, StringSize, 0);
-			FCStringAnsi::Strncat(StackTraceString, LINE_TERMINATOR_ANSI, (int32)StringSize);
+			FCStringAnsi::StrncatTruncateDest(StackTraceString, (int32)StringSize, LINE_TERMINATOR_ANSI);
 		}
 		FString Out = FString::Printf(TEXT("\n%S\n"), StackTraceString);
 		return Out;
diff --git a/Engine/Source/Runtime/VulkanRHI/Private/VulkanRHI.cpp b/Engine/Source/Runtime/VulkanRHI/Private/VulkanRHI.cpp
--- a/Engine/Source/Runtime/VulkanRHI/Private/VulkanRHI.cpp
+++ b/Engine/Source/Runtime/VulkanRHI/Private/VulkanRHI.cpp
@@ -265,7 +265,7 @@ static uint32 GetVulkanApiVersionForFeatureLevel(ERHIFeatureLevel::Type FeatureL
 	const FString ProfileName = FVulkanPlatform::GetVulkanProfileNameForFeatureLevel(FeatureLevel, bRaytracing);
 	VpProfileProperties ProfileProperties;
 	FMemory::Memzero(ProfileProperties);
-	FCStringAnsi::Strcpy(ProfileProperties.profileName, VP_MAX_PROFILE_NAME_SIZE, TCHAR_TO_ANSI(*ProfileName));
+	FCStringAnsi::Strncpy(ProfileProperties.profileName, TCHAR_TO_ANSI(*ProfileName), VP_MAX_PROFILE_NAME_SIZE);
 
 	const uint32 minApiVersion = vpGetProfileAPIVersion(&ProfileProperties);
 	if (minApiVersion)
@@ -296,7 +296,7 @@ static bool CheckVulkanProfile(ERHIFeatureLevel::Type FeatureLevel, bool bRaytra
 
 	VpProfileProperties ProfileProperties;
 	FMemory::Memzero(ProfileProperties);
-	FCStringAnsi::Strcpy(ProfileProperties.profileName, VP_MAX_PROFILE_NAME_SIZE, TCHAR_TO_ANSI(*ProfileName));
+	FCStringAnsi::Strncpy(ProfileProperties.profileName, TCHAR_TO_ANSI(*ProfileName), VP_MAX_PROFILE_NAME_SIZE);
 
 	VkBool32 bInstanceSupported = VK_FALSE;
 	VkResult InstanceResult = vpGetInstanceProfileSupport(nullptr, &ProfileProperties, &bInstanceSupported);  // :todo-jn: no VERIFYVULKANRESULT, this can fail and it's fine
