FPlatformString and FCString: Strcpy -> Strncpy callsite changes to handle deprecation.
#rnx

diff --git a/Engine/Source/Runtime/Core/Private/Linux/LinuxPlatformProcess.cpp b/Engine/Source/Runtime/Core/Private/Linux/LinuxPlatformProcess.cpp
--- a/Engine/Source/Runtime/Core/Private/Linux/LinuxPlatformProcess.cpp
+++ b/Engine/Source/Runtime/Core/Private/Linux/LinuxPlatformProcess.cpp
@@ -39,7 +39,7 @@ const TCHAR* FLinuxPlatformProcess::BaseDir()
 		}
 
 		FPaths::CollapseRelativeDirectories(CollapseResult);
-		FCString::Strcpy(CachedResult, UNIX_MAX_PATH, *CollapseResult);
+		FCString::Strncpy(CachedResult, *CollapseResult, UNIX_MAX_PATH);
 #endif
 
 		bHaveResult = true;
diff --git a/Engine/Source/Runtime/Core/Tests/Misc/CommandLineTests.cpp b/Engine/Source/Runtime/Core/Tests/Misc/CommandLineTests.cpp
--- a/Engine/Source/Runtime/Core/Tests/Misc/CommandLineTests.cpp
+++ b/Engine/Source/Runtime/Core/Tests/Misc/CommandLineTests.cpp
@@ -215,8 +215,8 @@ TEST_CASE("CommandLine::Filtering::FilterCLI_Comparison", "[Smoke]")
 		TCHAR NewMethod[1024];
 		TCHAR OldMethod[1024];
 
-		FCString::Strcpy(NewMethod, UE_ARRAY_COUNT(NewMethod), Original);
-		FCString::Strcpy(OldMethod, UE_ARRAY_COUNT(NewMethod), Original);
+		FCString::Strncpy(NewMethod, Original, UE_ARRAY_COUNT(NewMethod));
+		FCString::Strncpy(OldMethod, Original, UE_ARRAY_COUNT(OldMethod));
 
 		FCommandLine::FilterCLIUsingGrammarBasedParser(NewMethod, UE_ARRAY_COUNT(NewMethod), NewMethod, ApprovedArgs);
 		OldFilterMethod::FilterCLI(OldMethod, UE_ARRAY_COUNT(OldMethod), OldMethod, ApprovedArgs);
