Rename `GMalloc` to `UE::Private::GMalloc`.

Unless you are interested in the low-level guts of the allocator, you are better served to use `FMemory::Malloc` and friends. Moving GMalloc to a private namespace should signpost this a little better for our users.

The global namespace `GMalloc` continues to exist but will be marked deprecated in a followup.

diff --git a/Engine/Source/Editor/UnrealEd/Private/Commandlets/CookCommandlet.cpp b/Engine/Source/Editor/UnrealEd/Private/Commandlets/CookCommandlet.cpp
--- a/Engine/Source/Editor/UnrealEd/Private/Commandlets/CookCommandlet.cpp
+++ b/Engine/Source/Editor/UnrealEd/Private/Commandlets/CookCommandlet.cpp
@@ -722,7 +722,7 @@ void UCookCommandlet::ConditionalCollectGarbage(uint32 TickResults, UCookOnTheFl
 	const FPlatformMemoryStats MemStatsBeforeGC = FPlatformMemory::GetStats();
 	int32 NumObjectsBeforeGC = GUObjectArray.GetObjectArrayNumMinusAvailable();
 	FGenericMemoryStats AllocatorStatsBeforeGC;
-	GMalloc->GetAllocatorStats(AllocatorStatsBeforeGC);
+	UE::Private::GMalloc->GetAllocatorStats(AllocatorStatsBeforeGC);
 #if ENABLE_LOW_LEVEL_MEM_TRACKER
 	FLowLevelMemTracker::Get().UpdateStatsPerFrame();
 #endif
@@ -761,7 +761,7 @@ void UCookCommandlet::ConditionalCollectGarbage(uint32 TickResults, UCookOnTheFl
 	FPlatformMemoryStats MemStatsAfterGC = FPlatformMemory::GetStats();
 	int32 NumObjectsAfterGC = GUObjectArray.GetObjectArrayNumMinusAvailable();
 	FGenericMemoryStats AllocatorStatsAfterGC;
-	GMalloc->GetAllocatorStats(AllocatorStatsAfterGC);
+	UE::Private::GMalloc->GetAllocatorStats(AllocatorStatsAfterGC);
 #if ENABLE_LOW_LEVEL_MEM_TRACKER
 	FLowLevelMemTracker::Get().UpdateStatsPerFrame();
 #endif
diff --git a/Engine/Source/Programs/TestPAL/Private/Main.cpp b/Engine/Source/Programs/TestPAL/Private/Main.cpp
--- a/Engine/Source/Programs/TestPAL/Private/Main.cpp
+++ b/Engine/Source/Programs/TestPAL/Private/Main.cpp
@@ -617,8 +617,8 @@ int32 GetAllocationSizeTest(const TCHAR* CommandLine)
 	SIZE_T TotalMemoryRequested = 0, TotalMemoryAllocated = 0;
 
 	// force proxy malloc
-	FMalloc* OldGMalloc = GMalloc;
-	GMalloc = new FMallocPoisonProxy(GMalloc);
+	FMalloc* OldGMalloc = UE::Private::GMalloc;
+	UE::Private::GMalloc = new FMallocPoisonProxy(UE::Private::GMalloc);
 
 	// allocate all the memory and initialize with 0
 	for (uint32 Size = 16; Size < 4096; Size += 16)
@@ -630,8 +630,8 @@ int32 GetAllocationSizeTest(const TCHAR* CommandLine)
 			Allocation New;
 			New.RequestedSize = Size;
 			New.Alignment = Alignment;
-			New.Memory = GMalloc->Malloc(New.RequestedSize, New.Alignment);
-			if (!GMalloc->GetAllocationSize(New.Memory, New.ActualSize))
+			New.Memory = UE::Private::GMalloc->Malloc(New.RequestedSize, New.Alignment);
+			if (!UE::Private::GMalloc->GetAllocationSize(New.Memory, New.ActualSize))
 			{
 				UE_LOG(LogTestPAL, Fatal, TEXT("Could not get allocation size for %p"), New.Memory);
 			}
@@ -652,7 +652,7 @@ int32 GetAllocationSizeTest(const TCHAR* CommandLine)
 		{
 			// resize
 			Alloc.RequestedSize += 16;
-			Alloc.Memory = GMalloc->Realloc(Alloc.Memory, Alloc.RequestedSize, Alloc.Alignment);
+			Alloc.Memory = UE::Private::GMalloc->Realloc(Alloc.Memory, Alloc.RequestedSize, Alloc.Alignment);
 			FMemory::Memzero(Alloc.Memory, Alloc.RequestedSize);
 		}
 	}
@@ -687,9 +687,9 @@ int32 GetAllocationSizeTest(const TCHAR* CommandLine)
 
 	for (Allocation & Alloc : Allocs)
 	{
-		GMalloc->Free(Alloc.Memory);
+		UE::Private::GMalloc->Free(Alloc.Memory);
 	}
-	GMalloc = OldGMalloc;
+	UE::Private::GMalloc = OldGMalloc;
 
 	Allocs.Empty();
 
@@ -723,7 +723,7 @@ int32 StringsAllocationTest(const TCHAR* CommandLine)
 
 	if (GWarn)
 	{
-		GMalloc->DumpAllocatorStats(*GWarn);
+		UE::Private::GMalloc->DumpAllocatorStats(*GWarn);
 	}
 
 	for(int32 i = 0; i < NumOfStrings; ++i)
@@ -732,7 +732,7 @@ int32 StringsAllocationTest(const TCHAR* CommandLine)
 	}
 	delete [] Strings;
 
-	// GMalloc = OldGMalloc;
+	// UE::Private::GMalloc = OldGMalloc;
 	FEngineLoop::AppExit();
 	return 0;
 }
@@ -861,7 +861,7 @@ int32 CreateGuidTest(const TCHAR* CommandLine)
 			*GuidPlatform.ToString(EGuidFormats::DigitsWithHyphensInBraces));
 	}
 
-	// GMalloc = OldGMalloc;
+	// UE::Private::GMalloc = OldGMalloc;
 	FEngineLoop::AppExit();
 	return 0;
 }
@@ -1006,7 +1006,7 @@ int32 MallocThreadingTest(const TCHAR* CommandLine)
 	}
 
 	UE_LOG(LogTestPAL, Display, TEXT("Running malloc threading test using %s malloc and %d threads, each allocating %sup to %d KB (%d MB) memory %d times."),
-		bUseSystemMalloc ? TEXT("libc") : GMalloc->GetDescriptiveName(),
+		bUseSystemMalloc ? TEXT("libc") : UE::Private::GMalloc->GetDescriptiveName(),
 		NumTestThreads,
 		GMallocTestTouchTMemory ? TEXT("and touching ") : TEXT(""),
 		GMallocTestMemoryPerThreadKB, GMallocTestMemoryPerThreadKB / 1024,
@@ -1037,7 +1037,7 @@ int32 MallocThreadingTest(const TCHAR* CommandLine)
 			RunnableArray[Idx] = nullptr;
 		}
 	}
-	UE_LOG(LogTestPAL, Display, TEXT("--- Results for %s malloc ---"), bUseSystemMalloc ? TEXT("libc") : GMalloc->GetDescriptiveName());
+	UE_LOG(LogTestPAL, Display, TEXT("--- Results for %s malloc ---"), bUseSystemMalloc ? TEXT("libc") : UE::Private::GMalloc->GetDescriptiveName());
 	UE_LOG(LogTestPAL, Display, TEXT("Total wall time:        %f seconds"), WallTimeDuration);
 	UE_LOG(LogTestPAL, Display, TEXT("Total allocations done: %llu"), GTotalAllocsDone.GetValue());
 	UE_LOG(LogTestPAL, Display, TEXT("Total memory allocated: %llu bytes (%llu KB, %llu MB)"), GTotalMemoryAllocated.GetValue(),
diff --git a/Engine/Source/Programs/UnrealLightmass/Private/CPUSolver/CPUSolver.cpp b/Engine/Source/Programs/UnrealLightmass/Private/CPUSolver/CPUSolver.cpp
--- a/Engine/Source/Programs/UnrealLightmass/Private/CPUSolver/CPUSolver.cpp
+++ b/Engine/Source/Programs/UnrealLightmass/Private/CPUSolver/CPUSolver.cpp
@@ -157,7 +157,7 @@ void ReportLogFile()
 void BuildStaticLighting(const FGuid& SceneGuid, int32 NumThreads, bool bDumpTextures)
 {
 	// Place a marker in the memory profile data.
-	GMalloc->Exec(NULL, TEXT("SNAPSHOTMEMORY"), *GLog);
+	UE::Private::GMalloc->Exec(NULL, TEXT("SNAPSHOTMEMORY"), *GLog);
 
 	UE_LOG(LogLightmass, Log, TEXT("Building static lighting..."));
 
@@ -193,7 +193,7 @@ void BuildStaticLighting(const FGuid& SceneGuid, int32 NumThreads, bool bDumpTex
 
 
 	// Place a marker in the memory profile data.
-	GMalloc->Exec(NULL, TEXT("SNAPSHOTMEMORY"), *GLog);
+	UE::Private::GMalloc->Exec(NULL, TEXT("SNAPSHOTMEMORY"), *GLog);
 
 	double LightTimeStart = FPlatformTime::Seconds();
 
@@ -203,7 +203,7 @@ void BuildStaticLighting(const FGuid& SceneGuid, int32 NumThreads, bool bDumpTex
 	GStatistics.TotalTimeEnd = FPlatformTime::Seconds();
 
 	// Place a marker in the memory profile data.
-	GMalloc->Exec(NULL, TEXT("SNAPSHOTMEMORY"), *GLog);
+	UE::Private::GMalloc->Exec(NULL, TEXT("SNAPSHOTMEMORY"), *GLog);
 
 	// Report back statistics over Swarm.
 	ReportStatistics();
@@ -226,7 +226,7 @@ void BuildStaticLighting(const FGuid& SceneGuid, int32 NumThreads, bool bDumpTex
 	delete Swarm;
 	
 	// Write out memory profiling data to the .mprof file.
-	GMalloc->Exec(NULL, TEXT("DUMPALLOCSTOFILE"), *GLog);
+	UE::Private::GMalloc->Exec(NULL, TEXT("DUMPALLOCSTOFILE"), *GLog);
 }
 
 } //namespace Lightmass
diff --git a/Engine/Source/Runtime/ApplicationCore/Private/Mac/MacPlatformApplicationMisc.cpp b/Engine/Source/Runtime/ApplicationCore/Private/Mac/MacPlatformApplicationMisc.cpp
--- a/Engine/Source/Runtime/ApplicationCore/Private/Mac/MacPlatformApplicationMisc.cpp
+++ b/Engine/Source/Runtime/ApplicationCore/Private/Mac/MacPlatformApplicationMisc.cpp
@@ -514,7 +514,7 @@ void FMacPlatformApplicationMisc::ClipboardCopy(const TCHAR* Str)
 {
 	// Don't attempt to copy the text to the clipboard if we've crashed or we'll crash again & become unkillable.
 	// The MallocZone used for crash reporting will be enabled before this call if we've crashed so that will do for testing.
-	if ( GMalloc != GCrashMalloc )
+	if (UE::Private::GMalloc != GCrashMalloc)
 	{
 		SCOPED_AUTORELEASE_POOL;
 
diff --git a/Engine/Source/Runtime/Core/Private/Experimental/ConcurrentLinearAllocator.cpp b/Engine/Source/Runtime/Core/Private/Experimental/ConcurrentLinearAllocator.cpp
--- a/Engine/Source/Runtime/Core/Private/Experimental/ConcurrentLinearAllocator.cpp
+++ b/Engine/Source/Runtime/Core/Private/Experimental/ConcurrentLinearAllocator.cpp
@@ -4,18 +4,18 @@
 
 void* FOsAllocator::Malloc(SIZE_T Size, uint32 Alignment)
 {
-	if (UNLIKELY(!GMalloc))
+	if (UNLIKELY(!UE::Private::GMalloc))
 	{
 		// There is no public function to create GMalloc and this will do it for us.
 		FMemory::Free(FMemory::Malloc(0));
-		check(GMalloc);
+		check(UE::Private::GMalloc);
 	}
-	return GMalloc->Malloc(Size, Alignment);
+	return UE::Private::GMalloc->Malloc(Size, Alignment);
 }
 
 void FOsAllocator::Free(void* Pointer, SIZE_T Size)
 {
-	GMalloc->Free(Pointer);
+	UE::Private::GMalloc->Free(Pointer);
 }
 
 FORCENOINLINE void UE::Core::Private::OnInvalidConcurrentLinearArrayAllocatorNum(int32 NewNum, SIZE_T NumBytesPerElement)
diff --git a/Engine/Source/Runtime/Core/Private/GenericPlatform/GenericPlatformMallocCrash.cpp b/Engine/Source/Runtime/Core/Private/GenericPlatform/GenericPlatformMallocCrash.cpp
--- a/Engine/Source/Runtime/Core/Private/GenericPlatform/GenericPlatformMallocCrash.cpp
+++ b/Engine/Source/Runtime/Core/Private/GenericPlatform/GenericPlatformMallocCrash.cpp
@@ -288,7 +288,7 @@ FGenericPlatformMallocCrash& FGenericPlatformMallocCrash::Get( FMalloc* MainMall
 void FGenericPlatformMallocCrash::SetAsGMalloc()
 {
 	InternalLock.Lock();
-	GMalloc = this;
+	UE::Private::GMalloc = this;
 	if (PLATFORM_USES_FIXED_GMalloc_CLASS && GFixedMallocLocationPtr)
 	{
 		*GFixedMallocLocationPtr = nullptr; // this disables any fast-path inline allocators
@@ -298,7 +298,7 @@ void FGenericPlatformMallocCrash::SetAsGMalloc()
 
 bool FGenericPlatformMallocCrash::IsActive() const
 {
-	return GMalloc == this;
+	return UE::Private::GMalloc == this;
 }
 
 void* FGenericPlatformMallocCrash::Malloc( SIZE_T Size, uint32 Alignment )
@@ -646,12 +646,12 @@ void FGenericStackBasedMallocCrash::SetAsGMalloc()
 	{
 		*GFixedMallocLocationPtr = nullptr; // this disables any fast-path inline allocators
 	}
-	GMalloc = this;
+	UE::Private::GMalloc = this;
 }
 
 bool FGenericStackBasedMallocCrash::IsActive() const
 {
-	return GMalloc == this;
+	return UE::Private::GMalloc == this;
 }
 
 void* FGenericStackBasedMallocCrash::Malloc(SIZE_T Size, uint32 Alignment)
diff --git a/Engine/Source/Runtime/Core/Private/GenericPlatform/GenericPlatformMemory.cpp b/Engine/Source/Runtime/Core/Private/GenericPlatform/GenericPlatformMemory.cpp
--- a/Engine/Source/Runtime/Core/Private/GenericPlatform/GenericPlatformMemory.cpp
+++ b/Engine/Source/Runtime/Core/Private/GenericPlatform/GenericPlatformMemory.cpp
@@ -253,7 +253,7 @@ void FGenericPlatformMemory::OnOutOfMemory(uint64 Size, uint32 Alignment)
 		);
 		if (GWarn)
 		{
-			GMalloc->DumpAllocatorStats(*GWarn);
+			UE::Private::GMalloc->DumpAllocatorStats(*GWarn);
 		}
 
 		// let any registered handlers go
@@ -359,7 +359,7 @@ void FGenericPlatformMemory::DumpStats( class FOutputDevice& Ar )
 void FGenericPlatformMemory::DumpPlatformAndAllocatorStats( class FOutputDevice& Ar )
 {
 	FPlatformMemory::DumpStats( Ar );
-	GMalloc->DumpAllocatorStats( Ar );
+	UE::Private::GMalloc->DumpAllocatorStats( Ar );
 }
 
 EPlatformMemorySizeBucket FGenericPlatformMemory::GetMemorySizeBucket()
diff --git a/Engine/Source/Runtime/Core/Private/HAL/LowLevelMemTracker.cpp b/Engine/Source/Runtime/Core/Private/HAL/LowLevelMemTracker.cpp
--- a/Engine/Source/Runtime/Core/Private/HAL/LowLevelMemTracker.cpp
+++ b/Engine/Source/Runtime/Core/Private/HAL/LowLevelMemTracker.cpp
@@ -2004,7 +2004,7 @@ void FLowLevelMemTracker::TickInternal()
 	// Calculate FMalloc unused stat and set it in the Default tracker.
 	const int64 FMallocAmount = DefaultTracker.GetAllocTypeAmount(ELLMAllocType::FMalloc);
 	const int64 FMallocPlatformAmount = PlatformTracker.GetTagAmount(FindOrAddTagData(ELLMTag::FMalloc));
-	const int64 CachedFreeMemory = GMalloc ? GMalloc->GetTotalFreeCachedMemorySize() : 0;
+	const int64 CachedFreeMemory = UE::Private::GMalloc ? UE::Private::GMalloc->GetTotalFreeCachedMemorySize() : 0;
 	int64 FMallocUnused = FMallocPlatformAmount - FMallocAmount - CachedFreeMemory;
 	if (FMallocPlatformAmount == 0)
 	{
@@ -2031,7 +2031,7 @@ void FLowLevelMemTracker::TickInternal()
 	// remove the MemoryUsageCurrentOverhead from the "Total" for the default LLM as it's not something anyone needs
 	// to investigate when finding what to reduce the platform LLM will have the info 
 	const int64 DefaultProcessMemory = PlatformProcessMemory - (!FPlatformMemory::TracksLLMAllocations() ? MemoryUsageCurrentOverhead : 0);
-	const int64 ImmediatelyFreeableCachedMemory = GMalloc ? GMalloc->GetImmediatelyFreeableCachedMemorySize() : 0;
+	const int64 ImmediatelyFreeableCachedMemory = UE::Private::GMalloc ? UE::Private::GMalloc->GetImmediatelyFreeableCachedMemorySize() : 0;
 	// we need to subtract diff between total cached free mem and immediately freeable cached memory as it is not accounted anywhere in the LLM and would artificially inflate Untracked.
 	// while any new allocations from that cached, but not immediately freeable memory would be properly accounted by the LLM
 	const int64 DefaultUntracked = FMath::Max<int64>(0, DefaultProcessMemory - TrackedTotal - CachedFreeMemory + ImmediatelyFreeableCachedMemory);
diff --git a/Engine/Source/Runtime/Core/Private/HAL/UnrealMemory.cpp b/Engine/Source/Runtime/Core/Private/HAL/UnrealMemory.cpp
--- a/Engine/Source/Runtime/Core/Private/HAL/UnrealMemory.cpp
+++ b/Engine/Source/Runtime/Core/Private/HAL/UnrealMemory.cpp
@@ -256,9 +256,9 @@ void FMemory::EnablePurgatoryTests()
 	bOnce = true;
 	while (true)
 	{
-		FMalloc* LocalGMalloc = GMalloc;
+		FMalloc* LocalGMalloc = UE::Private::GMalloc;
 		FMalloc* Proxy = new FMallocPurgatoryProxy(LocalGMalloc);
-		if (FPlatformAtomics::InterlockedCompareExchangePointer((void**)&GMalloc, Proxy, LocalGMalloc) == LocalGMalloc)
+		if (FPlatformAtomics::InterlockedCompareExchangePointer((void**)&UE::Private::GMalloc, Proxy, LocalGMalloc) == LocalGMalloc)
 		{
 			UE_LOG(LogConsoleResponse, Display, TEXT("Purgatory proxy is now on."));
 			return;
@@ -288,9 +288,9 @@ void FMemory::EnablePoisonTests()
 	bOnce = true;
 	while (true)
 	{
-		FMalloc* LocalGMalloc = GMalloc;
+		FMalloc* LocalGMalloc = UE::Private::GMalloc;
 		FMalloc* Proxy = new FMallocPoisonProxy(LocalGMalloc);
-		if (FPlatformAtomics::InterlockedCompareExchangePointer((void**)&GMalloc, Proxy, LocalGMalloc) == LocalGMalloc)
+		if (FPlatformAtomics::InterlockedCompareExchangePointer((void**)&UE::Private::GMalloc, Proxy, LocalGMalloc) == LocalGMalloc)
 		{
 			UE_LOG(LogConsoleResponse, Display, TEXT("Poison proxy is now on."));
 			return;
@@ -343,7 +343,7 @@ FConsoleCommandDelegate::CreateStatic(&FMemory::EnablePoisonTests)
 /** Helper function called on first allocation to create and initialize GMalloc */
 static int FMemory_GCreateMalloc_ThreadUnsafe()
 {
-	GMalloc = FPlatformMemory::BaseAllocator();
+	UE::Private::GMalloc = FPlatformMemory::BaseAllocator();
 
 #if !PLATFORM_MAC
 	FPlatformMemoryStats Stats = FPlatformMemory::GetStats();
@@ -351,49 +351,49 @@ static int FMemory_GCreateMalloc_ThreadUnsafe()
 #endif
 
 	// Setup malloc crash as soon as possible.
-	FPlatformMallocCrash::Get( GMalloc );
+	FPlatformMallocCrash::Get(UE::Private::GMalloc);
 
 #if PLATFORM_USES_FIXED_GMalloc_CLASS
 #if MALLOC_VERIFY || MALLOC_LEAKDETECTION || UE_USE_MALLOC_FILL_BYTES
 #error "Turn off PLATFORM_USES_FIXED_GMalloc_CLASS in order to use special allocator proxies"
 #endif
-	if (!GMalloc->IsInternallyThreadSafe())
+	if (!UE::Private::GMalloc->IsInternallyThreadSafe())
 	{
 		UE_LOG(LogMemory, Fatal, TEXT("PLATFORM_USES_FIXED_GMalloc_CLASS only makes sense for allocators that are internally threadsafe."));
 	}
 #else
 
 #if UE_MEMORY_TRACE_ENABLED
-	FMalloc* TraceMalloc = MemoryTrace_Create(GMalloc);
-	if (TraceMalloc != GMalloc)
+	FMalloc* TraceMalloc = MemoryTrace_Create(UE::Private::GMalloc);
+	if (TraceMalloc != UE::Private::GMalloc)
 	{
-		GMalloc = TraceMalloc;
+		UE::Private::GMalloc = TraceMalloc;
 		MemoryTrace_Initialize();
 	}
 #endif // UE_MEMORY_TRACE_ENABLED
 
 #if WITH_MALLOC_STOMP2
-	GMalloc = FMallocStomp2::OverrideIfEnabled(GMalloc);
+	UE::Private::GMalloc = FMallocStomp2::OverrideIfEnabled(UE::Private::GMalloc);
 #endif
 
 	// if the allocator is already thread safe, there is no need for the thread safe proxy
-	if (!GMalloc->IsInternallyThreadSafe())
+	if (!UE::Private::GMalloc->IsInternallyThreadSafe())
 	{
-		GMalloc = new FMallocThreadSafeProxy( GMalloc );
+		UE::Private::GMalloc = new FMallocThreadSafeProxy(UE::Private::GMalloc);
 	}
 
 #if MALLOC_VERIFY
 	// Add the verifier
-	GMalloc = new FMallocVerifyProxy( GMalloc );
+	UE::Private::GMalloc = new FMallocVerifyProxy(UE::Private::GMalloc);
 #endif
 
 #if MALLOC_LEAKDETECTION
-	GMalloc = new FMallocLeakDetectionProxy(GMalloc);
+	UE::Private::GMalloc = new FMallocLeakDetectionProxy(UE::Private::GMalloc);
 #endif
 
 	// poison memory allocations in Debug and Development non-editor/non-editoronly data builds
 #if UE_USE_MALLOC_FILL_BYTES
-	GMalloc = new FMallocPoisonProxy(GMalloc);
+	UE::Private::GMalloc = new FMallocPoisonProxy(UE::Private::GMalloc);
 #endif
 
 #endif
@@ -406,10 +406,10 @@ static int FMemory_GCreateMalloc_ThreadUnsafe()
 	UE_LOG(LogMemory, Display, TEXT("Used memory before allocating anything was %.2fMB"), SizeInMb);
 #endif
 
-	GMalloc = FMallocDoubleFreeFinder::OverrideIfEnabled(GMalloc);
-	GMalloc = FMallocFrameProfiler::OverrideIfEnabled(GMalloc);
+	UE::Private::GMalloc = FMallocDoubleFreeFinder::OverrideIfEnabled(UE::Private::GMalloc);
+	UE::Private::GMalloc = FMallocFrameProfiler::OverrideIfEnabled(UE::Private::GMalloc);
 
-	GMalloc->OnMallocInitialized();
+	UE::Private::GMalloc->OnMallocInitialized();
 
 	return 0;
 }
@@ -417,8 +417,8 @@ static int FMemory_GCreateMalloc_ThreadUnsafe()
 void FMemory::ExplicitInit(FMalloc& Allocator)
 {
 #if defined(REQUIRE_EXPLICIT_GMALLOC_INIT) && REQUIRE_EXPLICIT_GMALLOC_INIT
-	check(!GMalloc);
-	GMalloc = &Allocator;
+	check(!UE::Private::GMalloc);
+	UE::Private::GMalloc = &Allocator;
 #else
 	checkf(false, TEXT("ExplicitInit() forbidden when global allocator is created lazily"));
 #endif
@@ -490,75 +490,75 @@ void FScopedMallocTimer::Spew()
 
 void* FMemory::MallocExternal(SIZE_T Count, uint32 Alignment)
 {
-	if (!GMalloc)
+	if (!UE::Private::GMalloc)
 	{
 		GCreateMalloc();
-		CA_ASSUME(GMalloc != NULL);	// Don't want to assert, but suppress static analysis warnings about potentially NULL GMalloc
+		CA_ASSUME(UE::Private::GMalloc != nullptr);	// Don't want to assert, but suppress static analysis warnings about potentially NULL GMalloc
 	}
-	return GMalloc->Malloc(Count, Alignment);
+	return UE::Private::GMalloc->Malloc(Count, Alignment);
 }
 
 void* FMemory::ReallocExternal(void* Original, SIZE_T Count, uint32 Alignment)
 {
-	if (!GMalloc)
+	if (!UE::Private::GMalloc)
 	{
 		GCreateMalloc();
-		CA_ASSUME(GMalloc != NULL);	// Don't want to assert, but suppress static analysis warnings about potentially NULL GMalloc
+		CA_ASSUME(UE::Private::GMalloc != nullptr);	// Don't want to assert, but suppress static analysis warnings about potentially NULL GMalloc
 	}
-	return GMalloc->Realloc(Original, Count, Alignment);
+	return UE::Private::GMalloc->Realloc(Original, Count, Alignment);
 }
 
 void FMemory::FreeExternal(void* Original)
 {
-	if (!GMalloc)
+	if (!UE::Private::GMalloc)
 	{
 		GCreateMalloc();
-		CA_ASSUME(GMalloc != NULL);	// Don't want to assert, but suppress static analysis warnings about potentially NULL GMalloc
+		CA_ASSUME(UE::Private::GMalloc != nullptr);	// Don't want to assert, but suppress static analysis warnings about potentially NULL GMalloc
 	}
 	if (Original)
 	{
-		GMalloc->Free(Original);
+		UE::Private::GMalloc->Free(Original);
 }
 }
 
 SIZE_T FMemory::GetAllocSizeExternal(void* Original)
 { 
-	if (!GMalloc)
+	if (!UE::Private::GMalloc)
 	{
 		GCreateMalloc();
-		CA_ASSUME(GMalloc != NULL);	// Don't want to assert, but suppress static analysis warnings about potentially NULL GMalloc
+		CA_ASSUME(UE::Private::GMalloc != nullptr);	// Don't want to assert, but suppress static analysis warnings about potentially NULL GMalloc
 	}
 	SIZE_T Size = 0;
-	return GMalloc->GetAllocationSize(Original, Size) ? Size : 0;
+	return UE::Private::GMalloc->GetAllocationSize(Original, Size) ? Size : 0;
 }
 
 void* FMemory::MallocZeroedExternal(SIZE_T Count, uint32 Alignment)
 {
-	if (!GMalloc)
+	if (!UE::Private::GMalloc)
 	{
 		GCreateMalloc();
-		CA_ASSUME(GMalloc != NULL);	// Don't want to assert, but suppress static analysis warnings about potentially NULL GMalloc
+		CA_ASSUME(UE::Private::GMalloc != nullptr);	// Don't want to assert, but suppress static analysis warnings about potentially NULL GMalloc
 	}
-	return GMalloc->MallocZeroed(Count, Alignment);
+	return UE::Private::GMalloc->MallocZeroed(Count, Alignment);
 }
 
 SIZE_T FMemory::QuantizeSizeExternal(SIZE_T Count, uint32 Alignment)
 { 
-	if (!GMalloc)
+	if (!UE::Private::GMalloc)
 	{
 		GCreateMalloc();	
-		CA_ASSUME(GMalloc != NULL);	// Don't want to assert, but suppress static analysis warnings about potentially NULL GMalloc
+		CA_ASSUME(UE::Private::GMalloc != nullptr);	// Don't want to assert, but suppress static analysis warnings about potentially NULL GMalloc
 	}
-	return GMalloc->QuantizeSize(Count, Alignment);
+	return UE::Private::GMalloc->QuantizeSize(Count, Alignment);
 }	
 
 
 void FMemory::Trim(bool bTrimThreadCaches)
 {
-	if (!GMalloc)
+	if (!UE::Private::GMalloc)
 	{
 		GCreateMalloc();
-		CA_ASSUME(GMalloc != NULL);	// Don't want to assert, but suppress static analysis warnings about potentially NULL GMalloc
+		CA_ASSUME(UE::Private::GMalloc != nullptr);	// Don't want to assert, but suppress static analysis warnings about potentially NULL GMalloc
 	}
 	TRACE_CPUPROFILER_EVENT_SCOPE(FMemory::Trim);
 	QUICK_SCOPE_CYCLE_COUNTER(STAT_FMemory_Trim);
@@ -567,50 +567,50 @@ void FMemory::Trim(bool bTrimThreadCaches)
 		FCoreDelegates::GetMemoryTrimDelegate().Broadcast();
 	}
 	QUICK_SCOPE_CYCLE_COUNTER(STAT_FMemory_Trim_GMalloc);
-	GMalloc->Trim(bTrimThreadCaches);
+	UE::Private::GMalloc->Trim(bTrimThreadCaches);
 }
 
 void FMemory::SetupTLSCachesOnCurrentThread()
 {
-	if (!GMalloc)
+	if (!UE::Private::GMalloc)
 	{
 		GCreateMalloc();	
-		CA_ASSUME(GMalloc != NULL);	// Don't want to assert, but suppress static analysis warnings about potentially NULL GMalloc
+		CA_ASSUME(UE::Private::GMalloc != nullptr);	// Don't want to assert, but suppress static analysis warnings about potentially NULL GMalloc
 	}
-	GMalloc->SetupTLSCachesOnCurrentThread();
+	UE::Private::GMalloc->SetupTLSCachesOnCurrentThread();
 }
 	
 void FMemory::ClearAndDisableTLSCachesOnCurrentThread()
 {
-	if (GMalloc)
+	if (UE::Private::GMalloc)
 	{
-		GMalloc->ClearAndDisableTLSCachesOnCurrentThread();
+		UE::Private::GMalloc->ClearAndDisableTLSCachesOnCurrentThread();
 	}
 }
 
 void FMemory::MarkTLSCachesAsUsedOnCurrentThread()
 {
-	if (GMalloc)
+	if (UE::Private::GMalloc)
 	{
-		GMalloc->MarkTLSCachesAsUsedOnCurrentThread();
+		UE::Private::GMalloc->MarkTLSCachesAsUsedOnCurrentThread();
 	}
 }
 
 void FMemory::MarkTLSCachesAsUnusedOnCurrentThread()
 {
-	if (GMalloc)
+	if (UE::Private::GMalloc)
 	{
-		GMalloc->MarkTLSCachesAsUnusedOnCurrentThread();
+		UE::Private::GMalloc->MarkTLSCachesAsUnusedOnCurrentThread();
 	}
 }
 
 void FMemory::TestMemory()
 {
 #if !UE_BUILD_SHIPPING
-	if( !GMalloc )
+	if( !UE::Private::GMalloc)
 	{
 		GCreateMalloc();	
-		CA_ASSUME( GMalloc != NULL );	// Don't want to assert, but suppress static analysis warnings about potentially NULL GMalloc
+		CA_ASSUME(UE::Private::GMalloc != nullptr);	// Don't want to assert, but suppress static analysis warnings about potentially NULL GMalloc
 	}
 
 	// track the pointers to free next call to the function
diff --git a/Engine/Source/Runtime/Core/Private/IOS/IOSMallocZone.cpp b/Engine/Source/Runtime/Core/Private/IOS/IOSMallocZone.cpp
--- a/Engine/Source/Runtime/Core/Private/IOS/IOSMallocZone.cpp
+++ b/Engine/Source/Runtime/Core/Private/IOS/IOSMallocZone.cpp
@@ -64,11 +64,11 @@ const TCHAR* FIOSMallocZone::GetDescriptiveName()
 
 FIOSMallocCrashHandler::FIOSMallocCrashHandler( uint64 const InitialSize )
 : FIOSMallocZone( InitialSize )
-, OriginalHeap( GMalloc )
+, OriginalHeap(UE::Private::GMalloc)
 , CrashContext( nullptr )
 , ThreadId( -1 )
 {
-	check( GMalloc );
+	check(UE::Private::GMalloc);
 }
 
 FIOSMallocCrashHandler::~FIOSMallocCrashHandler()
@@ -85,12 +85,12 @@ void FIOSMallocCrashHandler::Enable( FIOSCrashContext* Context, uint32 CrashedTh
 	check(Context);
 	CrashContext = Context;
 	ThreadId = CrashedThreadId;
-	OriginalHeap = GMalloc;
+	OriginalHeap = UE::Private::GMalloc;
 	if (PLATFORM_USES_FIXED_GMalloc_CLASS && GFixedMallocLocationPtr)
 	{
 		*GFixedMallocLocationPtr = nullptr; // this disables any fast-path inline allocators
 	}
-	GMalloc = this;
+	UE::Private::GMalloc = this;
 }
 	
 void FIOSMallocCrashHandler::SetContext( FIOSCrashContext* Context )
diff --git a/Engine/Source/Runtime/Core/Private/Mac/MacMallocZone.cpp b/Engine/Source/Runtime/Core/Private/Mac/MacMallocZone.cpp
--- a/Engine/Source/Runtime/Core/Private/Mac/MacMallocZone.cpp
+++ b/Engine/Source/Runtime/Core/Private/Mac/MacMallocZone.cpp
@@ -64,11 +64,11 @@ const TCHAR* FMacMallocZone::GetDescriptiveName()
 
 FMacMallocCrashHandler::FMacMallocCrashHandler( uint64 const InitialSize )
 : FMacMallocZone( InitialSize )
-, OriginalHeap( GMalloc )
+, OriginalHeap(UE::Private::GMalloc)
 , CrashContext( nullptr )
 , ThreadId( -1 )
 {
-	check( GMalloc );
+	check(UE::Private::GMalloc);
 }
 
 FMacMallocCrashHandler::~FMacMallocCrashHandler()
@@ -85,12 +85,12 @@ void FMacMallocCrashHandler::Enable( FMacCrashContext* Context, uint32 CrashedTh
 	check(Context);
 	CrashContext = Context;
 	ThreadId = CrashedThreadId;
-	OriginalHeap = GMalloc;
+	OriginalHeap = UE::Private::GMalloc;
 	if (PLATFORM_USES_FIXED_GMalloc_CLASS && GFixedMallocLocationPtr)
 	{
 		*GFixedMallocLocationPtr = nullptr; // this disables any fast-path inline allocators
 	}
-	GMalloc = this;
+	UE::Private::GMalloc = this;
 }
 	
 void* FMacMallocCrashHandler::Malloc( SIZE_T Size, uint32 Alignment )
diff --git a/Engine/Source/Runtime/Core/Private/Mac/MacPlatformMisc.cpp b/Engine/Source/Runtime/Core/Private/Mac/MacPlatformMisc.cpp
--- a/Engine/Source/Runtime/Core/Private/Mac/MacPlatformMisc.cpp
+++ b/Engine/Source/Runtime/Core/Private/Mac/MacPlatformMisc.cpp
@@ -432,7 +432,7 @@ struct FMacApplicationInfo
 	
 	~FMacApplicationInfo()
 	{
-		if(GMalloc != GCrashMalloc)
+		if (UE::Private::GMalloc != GCrashMalloc)
 		{
 			delete GCrashMalloc;
 		}
diff --git a/Engine/Source/Runtime/Core/Private/Misc/App.cpp b/Engine/Source/Runtime/Core/Private/Misc/App.cpp
--- a/Engine/Source/Runtime/Core/Private/Misc/App.cpp
+++ b/Engine/Source/Runtime/Core/Private/Misc/App.cpp
@@ -450,7 +450,7 @@ void FApp::PrintStartupLogMessages()
 	UE_LOG(LogInit, Log, TEXT("%sCommand Line: %s"), *FilteredString, FCommandLine::GetForLogging());
 	UE_LOG(LogInit, Log, TEXT("Base Directory: %s"), FPlatformProcess::BaseDir());
 	//UE_LOG(LogInit, Log, TEXT("Character set: %s"), sizeof(TCHAR)==1 ? TEXT("ANSI") : TEXT("Unicode") );
-	UE_LOG(LogInit, Log, TEXT("Allocator: %s"), GMalloc->GetDescriptiveName());
+	UE_LOG(LogInit, Log, TEXT("Allocator: %s"), UE::Private::GMalloc->GetDescriptiveName());
 	UE_LOG(LogInit, Log, TEXT("Installed Engine Build: %d"), FApp::IsEngineInstalled() ? 1 : 0);
 	UE_LOG(LogInit, Log, TEXT("This binary is optimized with LTO: %s, PGO: %s, instrumented for PGO data collection: %s"),
 		PLATFORM_COMPILER_OPTIMIZATION_LTCG ? TEXT("yes") : TEXT("no"),
diff --git a/Engine/Source/Runtime/Core/Private/Misc/CoreGlobals.cpp b/Engine/Source/Runtime/Core/Private/Misc/CoreGlobals.cpp
--- a/Engine/Source/Runtime/Core/Private/Misc/CoreGlobals.cpp
+++ b/Engine/Source/Runtime/Core/Private/Misc/CoreGlobals.cpp
@@ -49,12 +49,12 @@ IMPLEMENT_MODULE( FCoreModule, Core );
 /*-----------------------------------------------------------------------------
 	Global variables.
 -----------------------------------------------------------------------------*/
-CORE_API FFeedbackContext*	GWarn						= nullptr;		/* User interaction and non critical warnings */
-FConfigCacheIni*			GConfig						= nullptr;		/* Configuration database cache */
-ITransaction*				GUndo						= nullptr;		/* Transaction tracker, non-NULL when a transaction is in progress */
-FOutputDeviceConsole*		GLogConsole					= nullptr;		/* Console log hook */
-CORE_API FMalloc*			GMalloc						= nullptr;		/* Memory allocator */
+CORE_API FFeedbackContext*	GWarn					= nullptr;		/* User interaction and non critical warnings */
+FConfigCacheIni*			GConfig					= nullptr;		/* Configuration database cache */
+ITransaction*				GUndo					= nullptr;		/* Transaction tracker, non-NULL when a transaction is in progress */
+FOutputDeviceConsole*		GLogConsole				= nullptr;		/* Console log hook */
 CORE_API FMalloc**			GFixedMallocLocationPtr = nullptr;		/* Memory allocator pointer location when PLATFORM_USES_FIXED_GMalloc_CLASS is true */
+CORE_API FMalloc*			UE::Private::GMalloc	= nullptr;		/* Memory allocator */
 
 class FPropertyWindowManager*	GPropertyWindowManager	= nullptr;		/* Manages and tracks property editing windows */
 
diff --git a/Engine/Source/Runtime/Core/Private/Misc/Fork.cpp b/Engine/Source/Runtime/Core/Private/Misc/Fork.cpp
--- a/Engine/Source/Runtime/Core/Private/Misc/Fork.cpp
+++ b/Engine/Source/Runtime/Core/Private/Misc/Fork.cpp
@@ -87,7 +87,7 @@ void FForkProcessHelper::OnForkingOccured()
 {
 	if (SupportsMultithreadingPostFork())
 	{
-		ensureMsgf(GMalloc->IsInternallyThreadSafe(), TEXT("The BaseAllocator %s is not threadsafe. Switch to a multithread allocator or ensure the FMallocThreadSafeProxy wraps it."), GMalloc->GetDescriptiveName());
+		ensureMsgf(UE::Private::GMalloc->IsInternallyThreadSafe(), TEXT("The BaseAllocator %s is not threadsafe. Switch to a multithread allocator or ensure the FMallocThreadSafeProxy wraps it."), UE::Private::GMalloc->GetDescriptiveName());
 
 		ForkInternal::bIsForkedMultithreadInstance = true;
 
@@ -126,7 +126,7 @@ bool FForkProcessHelper::SupportsMultithreadingPostFork()
 
 void FForkProcessHelper::LowLevelPreFork()
 {
-	GMalloc->OnPreFork();
+	UE::Private::GMalloc->OnPreFork();
 }
 
 void FForkProcessHelper::LowLevelPostForkParent()
@@ -137,5 +137,5 @@ void FForkProcessHelper::LowLevelPostForkParent()
 void FForkProcessHelper::LowLevelPostForkChild(uint16 ChildIndex)
 {
 	FForkProcessHelper::SetIsForkedChildProcess(ChildIndex);
-	GMalloc->OnPostFork();
+	UE::Private::GMalloc->OnPostFork();
 }
\ No newline at end of file
diff --git a/Engine/Source/Runtime/Core/Private/ProfilingDebugging/TraceAuxiliary.cpp b/Engine/Source/Runtime/Core/Private/ProfilingDebugging/TraceAuxiliary.cpp
--- a/Engine/Source/Runtime/Core/Private/ProfilingDebugging/TraceAuxiliary.cpp
+++ b/Engine/Source/Runtime/Core/Private/ProfilingDebugging/TraceAuxiliary.cpp
@@ -1958,7 +1958,7 @@ void FTraceAuxiliary::Initialize(const TCHAR* CommandLine)
 	GTraceAuxiliary.RegisterEndFrameCallbacks();
 
 	// Initialize callstack tracing with the regular malloc (it might have already been initialized by memory tracing).
-	CallstackTrace_Create(GMalloc);
+	CallstackTrace_Create(UE::Private::GMalloc);
 	CallstackTrace_Initialize();
 
 	// By default use 1 msec for stack sampling interval.
diff --git a/Engine/Source/Runtime/Core/Private/Sanitizer/RaceDetector.cpp b/Engine/Source/Runtime/Core/Private/Sanitizer/RaceDetector.cpp
--- a/Engine/Source/Runtime/Core/Private/Sanitizer/RaceDetector.cpp
+++ b/Engine/Source/Runtime/Core/Private/Sanitizer/RaceDetector.cpp
@@ -1217,9 +1217,9 @@ namespace UE::Sanitizer::RaceDetector {
 	{
 		while (true)
 		{
-			FMalloc* LocalGMalloc = GMalloc;
+			FMalloc* LocalGMalloc = UE::Private::GMalloc;
 			FMalloc* Proxy = new FMallocInstrumentation(LocalGMalloc);
-			if (FPlatformAtomics::InterlockedCompareExchangePointer((void**)&GMalloc, Proxy, LocalGMalloc) == LocalGMalloc)
+			if (FPlatformAtomics::InterlockedCompareExchangePointer((void**)&UE::Private::GMalloc, Proxy, LocalGMalloc) == LocalGMalloc)
 			{
 				return;
 			}
diff --git a/Engine/Source/Runtime/Core/Private/Stats/StatsMallocProfilerProxy.cpp b/Engine/Source/Runtime/Core/Private/Stats/StatsMallocProfilerProxy.cpp
--- a/Engine/Source/Runtime/Core/Private/Stats/StatsMallocProfilerProxy.cpp
+++ b/Engine/Source/Runtime/Core/Private/Stats/StatsMallocProfilerProxy.cpp
@@ -44,7 +44,7 @@ FStatsMallocProfilerProxy* FStatsMallocProfilerProxy::Get()
 	static FStatsMallocProfilerProxy* Instance = nullptr;
 	if( Instance == nullptr )
 	{
-		Instance = new FStatsMallocProfilerProxy( GMalloc );
+		Instance = new FStatsMallocProfilerProxy(UE::Private::GMalloc);
 		// Initialize stats metadata.
 		// We need to do here, after all hardcoded names have been initialized.
 		Instance->InitializeStatsMetadata();
diff --git a/Engine/Source/Runtime/Core/Private/Unix/UnixForkPageProtector.cpp b/Engine/Source/Runtime/Core/Private/Unix/UnixForkPageProtector.cpp
--- a/Engine/Source/Runtime/Core/Private/Unix/UnixForkPageProtector.cpp
+++ b/Engine/Source/Runtime/Core/Private/Unix/UnixForkPageProtector.cpp
@@ -718,7 +718,7 @@ bool FForkPageProtector::DumpCallstackInfoToFile()
 
 void FForkPageProtector::OverrideGMalloc()
 {
-	GMalloc = new FMallocLinked(GMalloc);
+	UE::Private::GMalloc = new FMallocLinked(UE::Private::GMalloc);
 }
 
 int (*RealPThreadCreate)(pthread_t*, const pthread_attr_t*, void*(*start) (void*), void*);
diff --git a/Engine/Source/Runtime/Core/Private/Unix/UnixForkPageProtector.h b/Engine/Source/Runtime/Core/Private/Unix/UnixForkPageProtector.h
--- a/Engine/Source/Runtime/Core/Private/Unix/UnixForkPageProtector.h
+++ b/Engine/Source/Runtime/Core/Private/Unix/UnixForkPageProtector.h
@@ -157,7 +157,7 @@ public:
 	bool HandleNewCrashAddress(void* CrashAddress);
 
 	/* 
-	 * Overrides GMalloc with our own allocator so we dont stomp on possible protected pages allocating
+	 * Overrides GMalloc with our own allocator so we don't stomp on possible protected pages when allocating
 	 */
 	static void OverrideGMalloc();
 
diff --git a/Engine/Source/Runtime/Core/Private/Unix/UnixPlatformMemory.cpp b/Engine/Source/Runtime/Core/Private/Unix/UnixPlatformMemory.cpp
--- a/Engine/Source/Runtime/Core/Private/Unix/UnixPlatformMemory.cpp
+++ b/Engine/Source/Runtime/Core/Private/Unix/UnixPlatformMemory.cpp
@@ -1382,7 +1382,7 @@ void FUnixPlatformMemory::OnOutOfMemory(uint64 Size, uint32 Alignment)
 		TCHAR ErrorMsg[ErrorMsgSize];
 		FPlatformMisc::GetSystemErrorMessage(ErrorMsg, ErrorMsgSize, 0);
 
-		FMalloc* Prev = GMalloc;
+		FMalloc* Prev = UE::Private::GMalloc;
 		FPlatformMallocCrash::Get().SetAsGMalloc();
 
 		FPlatformMemoryStats PlatformMemoryStats = FPlatformMemory::GetStats();
diff --git a/Engine/Source/Runtime/Core/Private/Unix/UnixPlatformMisc.cpp b/Engine/Source/Runtime/Core/Private/Unix/UnixPlatformMisc.cpp
--- a/Engine/Source/Runtime/Core/Private/Unix/UnixPlatformMisc.cpp
+++ b/Engine/Source/Runtime/Core/Private/Unix/UnixPlatformMisc.cpp
@@ -209,7 +209,7 @@ void FUnixPlatformMisc::PlatformInit()
 		UE_LOG(LogInit, Log, TEXT(" - GPU Brand Info: %s"), *GPUBrandInfo);
 	}
 
-	UE_LOG(LogInit, Log, TEXT(" - Memory allocator used: %s"), GMalloc->GetDescriptiveName());
+	UE_LOG(LogInit, Log, TEXT(" - Memory allocator used: %s"), UE::Private::GMalloc->GetDescriptiveName());
 	UE_LOG(LogInit, Log, TEXT(" - This is %s build."), BuildSettings::IsLicenseeVersion() ? TEXT("a licensee") : TEXT("an internal"));
 
 	FPlatformTime::PrintCalibrationLog();
diff --git a/Engine/Source/Runtime/Core/Private/Windows/WindowsPlatformCrashContext.cpp b/Engine/Source/Runtime/Core/Private/Windows/WindowsPlatformCrashContext.cpp
--- a/Engine/Source/Runtime/Core/Private/Windows/WindowsPlatformCrashContext.cpp
+++ b/Engine/Source/Runtime/Core/Private/Windows/WindowsPlatformCrashContext.cpp
@@ -1564,7 +1564,7 @@ private:
 		if (bGenerateRuntimeCallstack)
 		{
 			const SIZE_T StackTraceSize = 65535;
-			ANSICHAR* StackTrace = (ANSICHAR*)GMalloc->Malloc(StackTraceSize);
+			ANSICHAR* StackTrace = (ANSICHAR*)UE::Private::GMalloc->Malloc(StackTraceSize);
 			StackTrace[0] = 0;
 			// Walk the stack and dump it to the allocated memory. This process usually allocates a lot of memory.
 			if (!ContextWrapper)
@@ -1587,7 +1587,7 @@ private:
 
 			FCString::StrncatTruncateDest(GErrorHist, UE_ARRAY_COUNT(GErrorHist), ANSI_TO_TCHAR(StackTrace));
 
-			GMalloc->Free(StackTrace);
+			UE::Private::GMalloc->Free(StackTrace);
 		}
 
 		// Make sure any thread context wrapper is released
diff --git a/Engine/Source/Runtime/Core/Public/HAL/FMemory.inl b/Engine/Source/Runtime/Core/Public/HAL/FMemory.inl
--- a/Engine/Source/Runtime/Core/Public/HAL/FMemory.inl
+++ b/Engine/Source/Runtime/Core/Public/HAL/FMemory.inl
@@ -5,7 +5,7 @@
 #endif
 
 #if !defined(FMEMORY_INLINE_GMalloc)
-	#define FMEMORY_INLINE_GMalloc GMalloc
+	#define FMEMORY_INLINE_GMalloc UE::Private::GMalloc
 #endif
 
 #include "HAL/LowLevelMemTracker.h"
diff --git a/Engine/Source/Runtime/Core/Public/HAL/MemoryBase.h b/Engine/Source/Runtime/Core/Public/HAL/MemoryBase.h
--- a/Engine/Source/Runtime/Core/Public/HAL/MemoryBase.h
+++ b/Engine/Source/Runtime/Core/Public/HAL/MemoryBase.h
@@ -27,9 +27,22 @@ enum
 	MIN_ALIGNMENT = 8,
 };
 
+	
+/** 
+ * The global memory allocator.
+ * Most callers should use FMemory::Malloc instead of accessing GMalloc directly.
+ * FMemory wraps GMalloc but also provides low-level memory tracking and AutoRTFM support.
+ */
+namespace UE::Private
+{
+	CORE_API extern class FMalloc* GMalloc;
+}
+
+/** GMalloc remains exposed for backwards compatibility. */
+/*UE_DEPRECATED(5.6, "Use FMemory::Malloc for allocation instead of accessing GMalloc directly.")*/
+inline constexpr FMalloc*& GMalloc = UE::Private::GMalloc;
 
-/** The global memory allocator. */
-CORE_API extern class FMalloc* GMalloc;
+/** Memory allocator pointer location when PLATFORM_USES_FIXED_GMalloc_CLASS is true. */
 CORE_API extern class FMalloc** GFixedMallocLocationPtr;
 
 /** Holds generic memory stats, internally implemented as a map. */
diff --git a/Engine/Source/Runtime/Core/Public/Tasks/TaskPrivate.h b/Engine/Source/Runtime/Core/Public/Tasks/TaskPrivate.h
--- a/Engine/Source/Runtime/Core/Public/Tasks/TaskPrivate.h
+++ b/Engine/Source/Runtime/Core/Public/Tasks/TaskPrivate.h
@@ -941,13 +941,13 @@ protected:
 				else
 				{
 					TASKGRAPH_VERBOSE_EVENT_SCOPE(TExecutableTask::LargeAlloc);
-					return GMalloc->Malloc(sizeof(TExecutableTask), LargeTaskAlignment);
+					return UE::Private::GMalloc->Malloc(sizeof(TExecutableTask), LargeTaskAlignment);
 				}
 			}
 
 			static void operator delete(void* Ptr, size_t Size)
 			{
-				Size <= SmallTaskSize ? SmallTaskAllocator.Free(Ptr) : GMalloc->Free(Ptr);
+				Size <= SmallTaskSize ? SmallTaskAllocator.Free(Ptr) : UE::Private::GMalloc->Free(Ptr);
 			}
 		};
 
diff --git a/Engine/Source/Runtime/Engine/Private/ProfilingDebugging/HealthSnapshot.cpp b/Engine/Source/Runtime/Engine/Private/ProfilingDebugging/HealthSnapshot.cpp
--- a/Engine/Source/Runtime/Engine/Private/ProfilingDebugging/HealthSnapshot.cpp
+++ b/Engine/Source/Runtime/Engine/Private/ProfilingDebugging/HealthSnapshot.cpp
@@ -69,12 +69,12 @@ void FHealthSnapshot::CaptureMemoryStats()
 
 	// Get used CPU memory from the allocator
 	FGenericMemoryStats MallocStats;
-	if (GMalloc)
+	if (UE::Private::GMalloc)
 	{
 		static SIZE_T MaxAllocated = 0;
 		SIZE_T Allocated = 0;
 
-		GMalloc->GetAllocatorStats(MallocStats);
+		UE::Private::GMalloc->GetAllocatorStats(MallocStats);
 		const SIZE_T *MallocTotalAllocated = nullptr;
 		if ((MallocTotalAllocated = MallocStats.Find(TEXT("TotalAllocated"))) != nullptr)
 		{
diff --git a/Engine/Source/Runtime/Engine/Private/UnrealEngine.cpp b/Engine/Source/Runtime/Engine/Private/UnrealEngine.cpp
--- a/Engine/Source/Runtime/Engine/Private/UnrealEngine.cpp
+++ b/Engine/Source/Runtime/Engine/Private/UnrealEngine.cpp
@@ -5096,7 +5096,7 @@ bool UEngine::Exec( UWorld* InWorld, const TCHAR* Cmd, FOutputDevice& Ar )
 		return true;
 	}
 
-	if (GMalloc && (GMalloc->Exec( InWorld, Cmd,Ar) == true))
+	if (UE::Private::GMalloc && (UE::Private::GMalloc->Exec( InWorld, Cmd,Ar) == true))
 	{
 		return true;
 	}
@@ -8421,7 +8421,7 @@ bool UEngine::HandleMemCommand( const TCHAR* Cmd, FOutputDevice& Ar )
 #endif
 	FPlatformMemory::DumpStats( Ar );
 	Ar.CategorizedLogf( CategoryName, ELogVerbosity::Log, TEXT("") );
-	GMalloc->DumpAllocatorStats( Ar );
+	UE::Private::GMalloc->DumpAllocatorStats( Ar );
 
 	if( bDetailed || bReport)
 	{
@@ -10061,13 +10061,13 @@ bool UEngine::HandleTrackParticleRenderingStatsCommand( const TCHAR* Cmd, FOutpu
 
 bool UEngine::HandleDumpAllocatorStats( const TCHAR* Cmd, FOutputDevice& Ar )
 {
-	GMalloc->DumpAllocatorStats(Ar);
+	UE::Private::GMalloc->DumpAllocatorStats(Ar);
 	return true;
 }
 
 bool UEngine::HandleHeapCheckCommand( const TCHAR* Cmd, FOutputDevice& Ar )
 {
-	GMalloc->ValidateHeap();
+	UE::Private::GMalloc->ValidateHeap();
 	return true;
 }
 
diff --git a/Engine/Source/Runtime/Launch/Private/LaunchEngineLoop.cpp b/Engine/Source/Runtime/Launch/Private/LaunchEngineLoop.cpp
--- a/Engine/Source/Runtime/Launch/Private/LaunchEngineLoop.cpp
+++ b/Engine/Source/Runtime/Launch/Private/LaunchEngineLoop.cpp
@@ -2070,7 +2070,7 @@ int32 FEngineLoop::PreInitPreStartupScreen(const TCHAR* CmdLine)
 			UE_LOG(LogMemory, Fatal, TEXT("Cannot do malloc profiling with PLATFORM_USES_FIXED_GMalloc_CLASS."));
 		}
 		// Assumes no concurrency here.
-		GMalloc = FStatsMallocProfilerProxy::Get();
+		UE::Private::GMalloc = FStatsMallocProfilerProxy::Get();
 	}
 #endif // STATS && UE_STATS_MEMORY_PROFILER_ENABLED
 
@@ -2132,7 +2132,7 @@ int32 FEngineLoop::PreInitPreStartupScreen(const TCHAR* CmdLine)
 	if (FParse::Param(FCommandLine::Get(), TEXT("mallocframeprofiler")))
 	{
 		GMallocFrameProfilerEnabled = true;
-		GMalloc = FMallocFrameProfiler::OverrideIfEnabled(GMalloc);
+		UE::Private::GMalloc = FMallocFrameProfiler::OverrideIfEnabled(UE::Private::GMalloc);
 	}
 #endif // !UE_BUILD_SHIPPING
 
@@ -5799,7 +5799,7 @@ void FEngineLoop::Tick()
 		// update memory allocator stats
 		{
 			QUICK_SCOPE_CYCLE_COUNTER(STAT_FEngineLoop_Malloc_UpdateStats);
-			GMalloc->UpdateStats();
+			UE::Private::GMalloc->UpdateStats();
 		}
 #endif
 	}
