Add an option to capture a screenshot at the time of a GPU crash and include it as an optional attachment in crash reports (Windows cooked client only).

This is a best-effort attempt at retrieving the pixels directly from the window using Window's PrintWindow API. Since a GPU crash renders the GPU memory inaccessible, the composited window is the last resort to obtain the last rendered frame. Note that PrintWindow is not supposed to work for windows that render content with DirectX or other hardware-accelerated APIs, but there's a well-known (yet undocumented) flag that allows it to work in those cases too. The flag is supported since Windows 8.1 which is above our Windows minimum target version.

If enabled, the screenshot is not included in the crash report like other files commonly are. Rather, this change introduces the concept of "optional attachments" for Windows crash reports, which are files that are stored in the crash report output folder under a dedicated subdirectory. These optional attachments are uploaded by the crash report client only when the user explicitly chooses to do so. A follow-up change will provide the changes to the Crash Report Client to handle optional attachments.

#tests screenshot retrieval functionality tested by QA on a variety of window options (fullscreen, borderless, windowed) and Windows versions (10 and 11).

diff --git a/Engine/Source/Runtime/ApplicationCore/Private/Windows/WindowsWindow.cpp b/Engine/Source/Runtime/ApplicationCore/Private/Windows/WindowsWindow.cpp
--- a/Engine/Source/Runtime/ApplicationCore/Private/Windows/WindowsWindow.cpp
+++ b/Engine/Source/Runtime/ApplicationCore/Private/Windows/WindowsWindow.cpp
@@ -2,6 +2,7 @@
 
 #include "Windows/WindowsWindow.h"
 
+#include "Algo/Transform.h"
 #include "Math/UnrealMathUtility.h"
 #include "HAL/UnrealMemory.h"
 #include "Containers/UnrealString.h"
@@ -1010,6 +1011,74 @@ bool FWindowsWindow::IsEnabled()
 	return !!::IsWindowEnabled( HWnd );
 }
 
+TArray<FColor> FWindowsWindow::GetWindowPixels(uint32& OutWidth, uint32& OutHeight)
+{
+	TArray<FColor> Pixels;
+
+	RECT Size;
+	if (!GetClientRect(HWnd, &Size))
+	{
+		return Pixels;
+	}
+
+	LONG WindowWidth = Size.right - Size.left;
+	LONG WindowHeight = Size.bottom - Size.top;
+
+	HDC WindowDC = GetDC(HWnd);
+	if (!WindowDC)
+	{
+		return Pixels;
+	}
+
+	HDC CaptureDC = CreateCompatibleDC(WindowDC);
+	if (!CaptureDC)
+	{
+		ReleaseDC(HWnd, WindowDC);
+		return Pixels;
+	}
+
+	HBITMAP CaptureBitmap = CreateCompatibleBitmap(WindowDC, WindowWidth, WindowHeight);
+	if (!CaptureBitmap)
+	{
+		ReleaseDC(HWnd, WindowDC);
+		DeleteDC(CaptureDC);
+		return Pixels;
+	}
+
+	if (SelectObject(CaptureDC, CaptureBitmap))
+	{
+		// Use the undocumented flag PW_RENDERFULLCONTENT (2) which is available since Windows 8.1,
+		// and that allows to capture applications that use hardware-accelerated graphics.
+		UINT Flags = PW_CLIENTONLY | 2;
+		if (PrintWindow(HWnd, CaptureDC, Flags))
+		{
+			BITMAPINFO BitmapInfo = {};
+			BitmapInfo.bmiHeader.biSize = sizeof(BitmapInfo.bmiHeader);
+			BitmapInfo.bmiHeader.biWidth = WindowWidth;
+			BitmapInfo.bmiHeader.biHeight = -WindowHeight; // Negative height to indicate we want top-to-bottom scanlines.
+			BitmapInfo.bmiHeader.biPlanes = 1;
+			BitmapInfo.bmiHeader.biBitCount = 32;
+			BitmapInfo.bmiHeader.biCompression = BI_RGB;
+
+			// Get bitmap pixels and convert to RGB.
+			TArray<RGBQUAD> BitmapPixels;
+			BitmapPixels.SetNumUninitialized(WindowWidth * WindowHeight);
+			uint32 ActualHeight = GetDIBits(CaptureDC, CaptureBitmap, 0, WindowHeight, BitmapPixels.GetData(), &BitmapInfo, DIB_RGB_COLORS);
+			if (ActualHeight > 0)
+			{
+				Algo::Transform(BitmapPixels, Pixels, [](const RGBQUAD& Pixel) { return FColor(Pixel.rgbRed, Pixel.rgbGreen, Pixel.rgbBlue); });
+				OutWidth = WindowWidth;
+				OutHeight = ActualHeight;
+			}
+		}
+	}
+
+	ReleaseDC(HWnd, WindowDC);
+	DeleteDC(CaptureDC);
+	DeleteObject(CaptureBitmap);
+	return Pixels;
+}
+
 /** @return true if native window exists underneath the coordinates */
 bool FWindowsWindow::IsPointInWindow( int32 X, int32 Y ) const
 {
diff --git a/Engine/Source/Runtime/ApplicationCore/Public/Windows/WindowsWindow.h b/Engine/Source/Runtime/ApplicationCore/Public/Windows/WindowsWindow.h
--- a/Engine/Source/Runtime/ApplicationCore/Public/Windows/WindowsWindow.h
+++ b/Engine/Source/Runtime/ApplicationCore/Public/Windows/WindowsWindow.h
@@ -94,6 +94,10 @@ public:
 	/** @return True if the window is enabled */
 	APPLICATIONCORE_API bool IsEnabled();
 
+	// Attempts to retrieve the window's pixels using PrintWindow (https://learn.microsoft.com/en-us/windows/win32/api/winuser/nf-winuser-printwindow).
+	// Only the client area's pixels are retrieved, which excludes the title bar, border, etc.
+	APPLICATIONCORE_API TArray<FColor> GetWindowPixels(uint32& OutWidth, uint32& OutHeight);
+
 public:
 
 	// FGenericWindow interface
diff --git a/Engine/Source/Runtime/Core/Private/GenericPlatform/GenericPlatformCrashContext.cpp b/Engine/Source/Runtime/Core/Private/GenericPlatform/GenericPlatformCrashContext.cpp
--- a/Engine/Source/Runtime/Core/Private/GenericPlatform/GenericPlatformCrashContext.cpp
+++ b/Engine/Source/Runtime/Core/Private/GenericPlatform/GenericPlatformCrashContext.cpp
@@ -236,12 +236,15 @@ const TCHAR* const FGenericCrashContext::ConfigSectionName = TEXT("CrashReportCl
 const TCHAR* const FGenericCrashContext::CrashConfigPurgeDays = TEXT("CrashConfigPurgeDays");
 const TCHAR* const FGenericCrashContext::CrashGUIDRootPrefix = TEXT("UECC-");
 
+const TCHAR* const FGenericCrashContext::CrashReportOptionalAttachmentsDirName = TEXT("OptionalAttachments");
+
 const TCHAR* const FGenericCrashContext::CrashContextExtension = TEXT(".runtime-xml");
 const TCHAR* const FGenericCrashContext::RuntimePropertiesTag = TEXT( "RuntimeProperties" );
 const TCHAR* const FGenericCrashContext::DeploymentNameTag = TEXT( "DeploymentName" );
 const TCHAR* const FGenericCrashContext::PlatformPropertiesTag = TEXT( "PlatformProperties" );
 const TCHAR* const FGenericCrashContext::EngineDataTag = TEXT( "EngineData" );
 const TCHAR* const FGenericCrashContext::GameDataTag = TEXT( "GameData" );
+const TCHAR* const FGenericCrashContext::RequestDataTag = TEXT( "RequestData" );
 const TCHAR* const FGenericCrashContext::GameNameTag = TEXT( "GameName" );
 const TCHAR* const FGenericCrashContext::EnabledPluginsTag = TEXT("EnabledPlugins");
 const TCHAR* const FGenericCrashContext::CrashVersionTag = TEXT("CrashVersion");
diff --git a/Engine/Source/Runtime/Core/Private/Windows/WindowsPlatformCrashContext.cpp b/Engine/Source/Runtime/Core/Private/Windows/WindowsPlatformCrashContext.cpp
--- a/Engine/Source/Runtime/Core/Private/Windows/WindowsPlatformCrashContext.cpp
+++ b/Engine/Source/Runtime/Core/Private/Windows/WindowsPlatformCrashContext.cpp
@@ -21,6 +21,7 @@
 #include "Misc/EngineBuildSettings.h"
 #include "Misc/EngineVersion.h"
 #include "Misc/FeedbackContext.h"
+#include "Misc/FileHelper.h"
 #include "Misc/MessageDialog.h"
 #include "Misc/OutputDeviceArchiveWrapper.h"
 #include "Misc/OutputDeviceFile.h"
@@ -209,6 +210,85 @@ namespace UE::Core::Private
 		false,
 		TEXT("If true, force the crash report dialog to not be displayed in the event of a crash."),
 		ECVF_Default);
+
+	static TAutoConsoleVariable<bool> CVarScreenshotOnGPUCrash(
+		TEXT("WindowsPlatformCrashContext.ScreenshotOnGPUCrash"),
+		false,
+		TEXT("Take a screenshot of the game window to include in crash reports for GPU crashes (Windows only, cooked client only)\n"),
+		ECVF_Default
+	);
+}
+
+/** Optional attachments */
+TArray<FString> OptionalAttachmentFilepaths;
+
+CORE_API void FWindowsPlatformCrashContext::AddOptionalAttachment(const FString& OptionalAttachmentFilepath)
+{
+	OptionalAttachmentFilepaths.Add(OptionalAttachmentFilepath);
+}
+
+void FWindowsPlatformCrashContext::CopyOptionalAttachments(const TCHAR* BaseOutputDirectory) const
+{
+	// If there are any optional attachments, copy them to the appropriate subdirectory.
+	if (OptionalAttachmentFilepaths.IsEmpty())
+	{
+		return;
+	}
+
+	FString OptionalAttachmentsDirectory = FPaths::Combine(BaseOutputDirectory, CrashReportOptionalAttachmentsDirName);
+	if (!IFileManager::Get().MakeDirectory(*OptionalAttachmentsDirectory))
+	{
+		UE_LOG(LogWindows, Error, TEXT("Error creating crash report's optional attachments directory %s"), *OptionalAttachmentsDirectory);
+		return;
+	}
+
+	bool bHasOptionalAttachments = false;
+	for (const FString& OptionalAttachmentFilepath : OptionalAttachmentFilepaths)
+	{
+		FString OutputFilepath = FPaths::Combine(OptionalAttachmentsDirectory, FPaths::GetCleanFilename(OptionalAttachmentFilepath));
+		if (IFileManager::Get().Move(*OutputFilepath, *OptionalAttachmentFilepath))
+		{
+			bHasOptionalAttachments = true;
+
+		}
+		else
+		{
+			UE_LOG(LogWindows, Error, TEXT("Error moving crash report's optional attachment file %s to output directory"), *OptionalAttachmentFilepath);
+		}
+	}
+
+	if (bHasOptionalAttachments)
+	{
+		FString OptionalAttachmentsXML = FPaths::Combine(OptionalAttachmentsDirectory, FPlatformCrashContext::CrashContextRuntimeXMLNameW);
+		WriteOptionalAttachmentsXML(OptionalAttachmentsXML);
+	}
+}
+
+void FWindowsPlatformCrashContext::WriteOptionalAttachmentsXML(const FString& Filepath) const
+{
+	FString Buffer;
+
+	TCHAR CrashGUID[CrashGUIDLength];
+	GetUniqueCrashName(CrashGUID, CrashGUIDLength);
+
+	AddHeader(Buffer);
+
+	// Use the same GUID from the base crash report.
+	BeginSection(Buffer, RuntimePropertiesTag);
+	AddCrashProperty(Buffer, CrashGuidTag, (const TCHAR*)CrashGUID);
+	EndSection(Buffer, RuntimePropertiesTag);
+
+	// Mark this as an update ("patch") request.
+	BeginSection(Buffer, RequestDataTag);
+	AddCrashProperty(Buffer, TEXT("Type"), TEXT("PATCH"));
+	EndSection(Buffer, RequestDataTag);
+
+	AddFooter(Buffer);
+
+	if (!FFileHelper::SaveStringToFile(Buffer, *Filepath, FFileHelper::EEncodingOptions::AutoDetect))
+	{
+		UE_LOG(LogWindows, Error, TEXT("Error writing optional attachments XML file %s"), *Filepath);
+	}
 }
 
 /**
@@ -331,6 +411,8 @@ void FWindowsPlatformCrashContext::CopyPlatformSpecificFiles(const TCHAR* Output
 			}
 		}
 	}
+
+	CopyOptionalAttachments(OutputDirectory);
 }
 
 
diff --git a/Engine/Source/Runtime/Core/Public/GenericPlatform/GenericPlatformCrashContext.h b/Engine/Source/Runtime/Core/Public/GenericPlatform/GenericPlatformCrashContext.h
--- a/Engine/Source/Runtime/Core/Public/GenericPlatform/GenericPlatformCrashContext.h
+++ b/Engine/Source/Runtime/Core/Public/GenericPlatform/GenericPlatformCrashContext.h
@@ -398,12 +398,15 @@ public:
 	CORE_API static const TCHAR* const CrashConfigPurgeDays;
 	CORE_API static const TCHAR* const CrashGUIDRootPrefix;
 
+	CORE_API static const TCHAR* const CrashReportOptionalAttachmentsDirName;
+
 	CORE_API static const TCHAR* const CrashContextExtension;
 	CORE_API static const TCHAR* const RuntimePropertiesTag;
 	CORE_API static const TCHAR* const DeploymentNameTag;
 	CORE_API static const TCHAR* const PlatformPropertiesTag;
 	CORE_API static const TCHAR* const EngineDataTag;
 	CORE_API static const TCHAR* const GameDataTag;
+	CORE_API static const TCHAR* const RequestDataTag;
 	CORE_API static const TCHAR* const GameNameTag;
 	CORE_API static const TCHAR* const EnabledPluginsTag;
 	CORE_API static const TCHAR* const UEMinidumpName;
diff --git a/Engine/Source/Runtime/Core/Public/Windows/WindowsPlatformCrashContext.h b/Engine/Source/Runtime/Core/Public/Windows/WindowsPlatformCrashContext.h
--- a/Engine/Source/Runtime/Core/Public/Windows/WindowsPlatformCrashContext.h
+++ b/Engine/Source/Runtime/Core/Public/Windows/WindowsPlatformCrashContext.h
@@ -16,6 +16,16 @@ struct FWindowsPlatformCrashContext : public FMicrosoftPlatformCrashContext
 	
 	CORE_API virtual void AddPlatformSpecificProperties() const override;
 	CORE_API virtual void CopyPlatformSpecificFiles(const TCHAR* OutputDirectory, void* Context) override;
+
+	// Windows crash contexts support "optional attachments", which are
+	// extra files copied into a subdirectory within the main crash report folder.
+	// Optional attachments are specifically distinguished from other attachments
+	// as they will not be uploaded by the crash report client by default.
+	CORE_API static void AddOptionalAttachment(const FString& OptionalAttachmentFilepath);
+	
+protected:
+	void CopyOptionalAttachments(const TCHAR* BaseOutputDirectory) const;
+	void WriteOptionalAttachmentsXML(const FString& Filepath) const;
 };
 
 typedef FWindowsPlatformCrashContext FPlatformCrashContext;
diff --git a/Engine/Source/Runtime/D3D12RHI/Private/D3D12Util.cpp b/Engine/Source/Runtime/D3D12RHI/Private/D3D12Util.cpp
--- a/Engine/Source/Runtime/D3D12RHI/Private/D3D12Util.cpp
+++ b/Engine/Source/Runtime/D3D12RHI/Private/D3D12Util.cpp
@@ -9,6 +9,7 @@ D3D12Util.h: D3D RHI utility implementation.
 #include "EngineModule.h"
 #include "RendererInterface.h"
 #include "CoreGlobals.h"
+#include "Engine/GameEngine.h"
 #include "Misc/OutputDeviceRedirector.h"
 #include "HAL/ExceptionHandling.h"
 #if PLATFORM_WINDOWS
@@ -958,6 +959,12 @@ void FD3D12DynamicRHI::TerminateOnGPUCrash()
 	FTextBuilder ErrorMessage;
 	OutputGPUCrashReport(ErrorMessage);
 
+	UGameEngine* GameEngine = Cast<UGameEngine>(GEngine);
+	if (GameEngine)
+	{
+		GameEngine->OnGPUCrash();
+	}
+
 	// Show message box or trace information
 #if !(UE_BUILD_SHIPPING || UE_BUILD_TEST)
 	if (!FApp::IsUnattended() && !IsDebuggerPresent())
diff --git a/Engine/Source/Runtime/Engine/Classes/Engine/GameEngine.h b/Engine/Source/Runtime/Engine/Classes/Engine/GameEngine.h
--- a/Engine/Source/Runtime/Engine/Classes/Engine/GameEngine.h
+++ b/Engine/Source/Runtime/Engine/Classes/Engine/GameEngine.h
@@ -121,6 +121,11 @@ public:
 
 	ENGINE_API void OnViewportResized(FViewport* Viewport, uint32 Unused);
 
+	/**
+	 * Called when the RHI detects a GPU crash. This can be called from any thread.
+	 */
+	ENGINE_API void OnGPUCrash();
+
 public:
 
 	// UObject interface
diff --git a/Engine/Source/Runtime/Engine/Private/GameEngine.cpp b/Engine/Source/Runtime/Engine/Private/GameEngine.cpp
--- a/Engine/Source/Runtime/Engine/Private/GameEngine.cpp
+++ b/Engine/Source/Runtime/Engine/Private/GameEngine.cpp
@@ -52,6 +52,8 @@
 #include "Components/ReflectionCaptureComponent.h"
 #include "GameFramework/GameUserSettings.h"
 #include "GameDelegates.h"
+#include "ImageUtils.h"
+#include "Misc/FileHelper.h"
 #include "Misc/EmbeddedCommunication.h"
 #include "Engine/CoreSettings.h"
 #include "EngineAnalytics.h"
@@ -66,6 +68,11 @@
 #include "CustomResourcePool.h"
 #include "ComponentRecreateRenderStateContext.h"
 
+#if PLATFORM_WINDOWS
+#include "Windows/WindowsWindow.h"
+#include "Windows/WindowsPlatformCrashContext.h"
+#endif
+
 #if WITH_EDITOR
 #include "PIEPreviewDeviceProfileSelectorModule.h"
 #endif
@@ -798,6 +805,52 @@ void UGameEngine::OnViewportResized(FViewport* Viewport, uint32 Unused)
 	}
 }
 
+void UGameEngine::OnGPUCrash()
+{
+#if PLATFORM_WINDOWS
+	static IConsoleVariable* ScreenshotOnCrash = IConsoleManager::Get().FindConsoleVariable(TEXT("WindowsPlatformCrashContext.ScreenshotOnGPUCrash"));
+	if (GIsEditor || !ScreenshotOnCrash || !ScreenshotOnCrash->GetBool())
+	{
+		return;
+	}
+
+	// Find the game window and attempt to retrieve its pixels.
+	uint32 Width = 0, Height = 0;
+	TArray<FColor> Pixels;
+
+	TSharedPtr<SWindow> Window = GameViewportWindow.Pin();
+	if (Window)
+	{
+		TSharedPtr<FGenericWindow> NativeWindow = Window->GetNativeWindow();
+		if (NativeWindow && NativeWindow->IsDefinitionValid())
+		{
+			const TSharedPtr<FWindowsWindow> WindowsWindow = StaticCastSharedPtr<FWindowsWindow>(NativeWindow);
+			Pixels = WindowsWindow->GetWindowPixels(Width, Height);
+		}
+	}
+
+	if (!Pixels.IsEmpty())
+	{
+		check(Width > 0 && Height > 0);
+		TArray64<uint8> JPGData;
+		FImageView Image((void*)Pixels.GetData(), Width, Height, ERawImageFormat::BGRA8);
+		FImageUtils::CompressImage(JPGData, TEXT("jpg"), Image, 70);
+
+		FString ScreenshotPath = FPaths::Combine(FPaths::ProjectLogDir(), TEXT("CrashCapture.jpg"));
+		if (FFileHelper::SaveArrayToFile(JPGData, *ScreenshotPath))
+		{
+			FWindowsPlatformCrashContext::SetEngineData(TEXT("CrashScreenshotWidth"), FString::FromInt(Width));
+			FWindowsPlatformCrashContext::SetEngineData(TEXT("CrashScreenshotHeight"), FString::FromInt(Height));
+			FWindowsPlatformCrashContext::AddOptionalAttachment(ScreenshotPath);
+		}
+		else
+		{
+			UE_LOG(LogEngine, Warning, TEXT("Error saving crash screenshot to file: %s"), *ScreenshotPath);
+		}
+	}
+#endif
+}
+
 /*-----------------------------------------------------------------------------
 	Game init and exit.
 -----------------------------------------------------------------------------*/
diff --git a/Engine/Source/Runtime/VulkanRHI/Private/VulkanUtil.cpp b/Engine/Source/Runtime/VulkanRHI/Private/VulkanUtil.cpp
--- a/Engine/Source/Runtime/VulkanRHI/Private/VulkanUtil.cpp
+++ b/Engine/Source/Runtime/VulkanRHI/Private/VulkanUtil.cpp
@@ -9,6 +9,7 @@
 #include "VulkanPendingState.h"
 #include "VulkanContext.h"
 #include "VulkanMemory.h"
+#include "Engine/GameEngine.h"
 #include "Misc/App.h"
 #include "Misc/OutputDeviceRedirector.h"
 #include "RHIValidationContext.h"
@@ -931,6 +932,12 @@ namespace VulkanRHI
 #if NV_AFTERMATH
 			UE::RHICore::Nvidia::Aftermath::OnGPUCrash();
 #endif
+
+			UGameEngine* GameEngine = Cast<UGameEngine>(GEngine);
+			if (GameEngine)
+			{
+				GameEngine->OnGPUCrash();
+			}
 		}
 
 #if UE_BUILD_DEBUG || UE_BUILD_DEVELOPMENT
diff --git a/Engine/Source/Runtime/Windows/D3D11RHI/Private/D3D11Util.cpp b/Engine/Source/Runtime/Windows/D3D11RHI/Private/D3D11Util.cpp
--- a/Engine/Source/Runtime/Windows/D3D11RHI/Private/D3D11Util.cpp
+++ b/Engine/Source/Runtime/Windows/D3D11RHI/Private/D3D11Util.cpp
@@ -7,6 +7,7 @@
 #include "D3D11Util.h"
 #include "D3D11RHIPrivate.h"
 #include "EngineModule.h"
+#include "Engine/GameEngine.h"
 #include "RendererInterface.h"
 #include "ProfilingDebugging/ScopedDebugInfo.h"
 #include "HAL/ExceptionHandling.h"
@@ -98,6 +99,12 @@ static void TerminateOnDeviceRemoved(HRESULT D3DResult, ID3D11Device* Direct3DDe
 		UE::RHICore::Nvidia::Aftermath::OnGPUCrash();
 	#endif
 
+		UGameEngine* GameEngine = Cast<UGameEngine>(GEngine);
+		if (GameEngine)
+		{
+			GameEngine->OnGPUCrash();
+		}
+
 		// Report the GPU crash which will raise the exception
 		ReportGPUCrash(TEXT("GPU Crash dump Triggered"), nullptr);
 
