Change to `using` for GMalloc declaration at global scope.

The constexpr approach fails in some cases. `using` the symbol at global scope should be more reliable. (Downside: this name can't be marked UE_DEPRECATED.)

diff --git a/Engine/Source/Runtime/Core/Public/HAL/MemoryBase.h b/Engine/Source/Runtime/Core/Public/HAL/MemoryBase.h
--- a/Engine/Source/Runtime/Core/Public/HAL/MemoryBase.h
+++ b/Engine/Source/Runtime/Core/Public/HAL/MemoryBase.h
@@ -38,9 +38,12 @@ namespace UE::Private
 	CORE_API extern class FMalloc* GMalloc;
 }
 
-/** GMalloc remains exposed for backwards compatibility. */
-/*UE_DEPRECATED(5.6, "Use FMemory::Malloc for allocation instead of accessing GMalloc directly.")*/
-inline constexpr FMalloc*& GMalloc = UE::Private::GMalloc;
+/** 
+ * GMalloc remains exposed for backwards compatibility.
+ * Most callers should use FMemory::Malloc instead of accessing GMalloc directly.
+ * FMemory wraps GMalloc but also provides low-level memory tracking and AutoRTFM support.
+ */
+using UE::Private::GMalloc;
 
 /** Memory allocator pointer location when PLATFORM_USES_FIXED_GMalloc_CLASS is true. */
 CORE_API extern class FMalloc** GFixedMallocLocationPtr;
