Color: Add new CAM16 color space definition for ACES 2.0.

diff --git a/Engine/Source/Runtime/Core/Private/ColorManagement/ColorSpace.cpp b/Engine/Source/Runtime/Core/Private/ColorManagement/ColorSpace.cpp
--- a/Engine/Source/Runtime/Core/Private/ColorManagement/ColorSpace.cpp
+++ b/Engine/Source/Runtime/Core/Private/ColorManagement/ColorSpace.cpp
@@ -198,6 +198,13 @@ TStaticArray<FVector2d, 4> FColorSpace::MakeChromaticities(UE::Color::EColorSpac
 		Chromaticities[2] = FVector2d(0.0366, 0.0001);
 		Chromaticities[3] = FVector2d(0.4254, 0.4044);
 		break;
+	case EColorSpace::ACESCAM16:
+		Chromaticities[0] = FVector2d(0.8336, 0.1735);
+		Chromaticities[1] = FVector2d(2.3854, -1.4659);
+		Chromaticities[2] = FVector2d(0.087, -0.125);
+		Chromaticities[3] = FVector2d(0.333, 0.333);
+		break;
+
 	default:
 		checkNoEntry();
 		break;
diff --git a/Engine/Source/Runtime/Core/Public/ColorManagement/ColorManagementDefines.h b/Engine/Source/Runtime/Core/Public/ColorManagement/ColorManagementDefines.h
--- a/Engine/Source/Runtime/Core/Public/ColorManagement/ColorManagementDefines.h
+++ b/Engine/Source/Runtime/Core/Public/ColorManagement/ColorManagementDefines.h
@@ -59,6 +59,7 @@ enum class EColorSpace : uint8
 	GoProProtuneNative = 12,
 	PanasonicVGamut = 13,
 	PLASA_E1_54 = 14,
+	ACESCAM16 = 15,
 	Max,
 };
 
