Fixed some type and variable naming.

#rnx

diff --git a/Engine/Source/Runtime/Core/Public/String/FormatStringSan.h b/Engine/Source/Runtime/Core/Public/String/FormatStringSan.h
--- a/Engine/Source/Runtime/Core/Public/String/FormatStringSan.h
+++ b/Engine/Source/Runtime/Core/Public/String/FormatStringSan.h
@@ -80,7 +80,7 @@ namespace UE::Core::Private
 		template <typename T>
 		inline constexpr bool bIsTCharPtr = bIsPointerTo<T, TCHAR>;
 		template <typename T>
-		inline constexpr bool bIsFloatOrDouble = std::is_same_v<float, T> || std::is_same_v<double, T>;
+		inline constexpr bool TIsFloatOrDouble_V = std::is_same_v<float, T> || std::is_same_v<double, T>;
 
 		// Returns true only when the type is TEnumAsByte
 		template <typename T>
@@ -110,23 +110,23 @@ namespace UE::Core::Private
 		template <typename Arg, typename... Args>
 		struct TCheckFormatString<TFmtArgTypes<Arg, Args...>>
 		{
-			static constexpr bool CharIsDigit(const TCHAR C)
+			static constexpr bool CharIsDigit(const TCHAR Char)
 			{
-				return C >= TEXT('0') && C <= TEXT('9');
+				return Char >= TEXT('0') && Char <= TEXT('9');
 			}
 
-			static constexpr const TCHAR* SkipInteger(const TCHAR* P)
+			static constexpr const TCHAR* SkipInteger(const TCHAR* Fmt)
 			{
-				while (CharIsDigit(*P))
+				while (CharIsDigit(*Fmt))
 				{
-					++P;
+					++Fmt;
 				}
-				return P;
+				return Fmt;
 			}
 
-			static constexpr bool CharIsIntegerFormatSpecifier(const TCHAR C)
+			static constexpr bool CharIsIntegerFormatSpecifier(const TCHAR Char)
 			{
-				switch (C)
+				switch (Char)
 				{
 				case TEXT('i'): case TEXT('d'): case TEXT('u'): case TEXT('X'): case TEXT('x'):
 					return true;
@@ -135,41 +135,40 @@ namespace UE::Core::Private
 				}
 			}
 
-			static constexpr FResult HandleDynamicLengthSpecifier(int CurArgPos, const TCHAR* P)
+			static constexpr FResult HandleDynamicLengthSpecifier(int CurArgPos, const TCHAR* Fmt)
 			{
-					if constexpr (!(std::is_integral_v<Arg> || bIsIntegralEnum<Arg>))
-					{
-						return {EFormatStringSanStatus::DynamicLengthSpecNeedsIntegerArg, CurArgPos};
-					}
-					else
-					{
-						return TCheckFormatString<TFmtArgTypes<Args...>>::Check(true, CurArgPos + 1, P + 1);
-					}
-				};
+				if constexpr (!(std::is_integral_v<Arg> || bIsIntegralEnum<Arg>))
+				{
+					return {EFormatStringSanStatus::DynamicLengthSpecNeedsIntegerArg, CurArgPos};
+				}
+				else
+				{
+					return TCheckFormatString<TFmtArgTypes<Args...>>::Check(true, CurArgPos + 1, Fmt + 1);
+				}
+			};
 
 			static constexpr FResult Check(bool bInsideFormatSpec, int CurArgPos, const TCHAR* Fmt)
 			{
-				const TCHAR* P = Fmt;
-				if (!P[0] && !bInsideFormatSpec)
+				if (!Fmt[0] && !bInsideFormatSpec)
 				{
 					return {EFormatStringSanStatus::Ok, 0};
 				}
 
 				if (!bInsideFormatSpec)
 				{
-					while (*P && *P != TEXT('%'))
+					while (*Fmt && *Fmt != TEXT('%'))
 					{
-						++P;
+						++Fmt;
 					}
-					if (*P == TEXT('%'))
+					if (*Fmt == TEXT('%'))
 					{
-						++P;
+						++Fmt;
 					}
 				}
 
-				if (!*P)
+				if (!*Fmt)
 				{
-					if (P[-1] == TEXT('%') || bInsideFormatSpec) // nb: P[-1] is safe here because zero-length fmt strings are addressed above
+					if (Fmt[-1] == TEXT('%') || bInsideFormatSpec) // nb: P[-1] is safe here because zero-length fmt strings are addressed above
 					{
 						return {EFormatStringSanStatus::IncompleteFormatSpecifierOrUnescapedPercent, CurArgPos};
 					}
@@ -179,42 +178,42 @@ namespace UE::Core::Private
 					}
 				}
 
-				while (*P == TEXT('+') || *P == TEXT('#') || *P == TEXT(' ') || *P == TEXT('0'))
+				while (*Fmt == TEXT('+') || *Fmt == TEXT('#') || *Fmt == TEXT(' ') || *Fmt == TEXT('0'))
 				{
-					++P;
+					++Fmt;
 				}
 
-				if (*P == TEXT('-'))
+				if (*Fmt == TEXT('-'))
 				{
-					++P;
+					++Fmt;
 				}
-				if (*P == TEXT('*'))
+				if (*Fmt == TEXT('*'))
 				{
-					return HandleDynamicLengthSpecifier(CurArgPos, P);
+					return HandleDynamicLengthSpecifier(CurArgPos, Fmt);
 				}
-				else if (CharIsDigit(*P))
+				else if (CharIsDigit(*Fmt))
 				{
-					P = SkipInteger(P);
+					Fmt = SkipInteger(Fmt);
 				}
 
-				if (*P == TEXT('.'))
+				if (*Fmt == TEXT('.'))
 				{
-					++P;
+					++Fmt;
 				}
-				if (*P == TEXT('*'))
+				if (*Fmt == TEXT('*'))
 				{
-					return HandleDynamicLengthSpecifier(CurArgPos, P);
+					return HandleDynamicLengthSpecifier(CurArgPos, Fmt);
 				}
-				else if (CharIsDigit(*P))
+				else if (CharIsDigit(*Fmt))
 				{
-					P = SkipInteger(P);
+					Fmt = SkipInteger(Fmt);
 				}
 
-				if (P[0] == TEXT('l') && P[1] == TEXT('s'))
+				if (Fmt[0] == TEXT('l') && Fmt[1] == TEXT('s'))
 				{
-					++P;
+					++Fmt;
 				}
-				else if (P[0] == TEXT('h') && P[1] == TEXT('s'))
+				else if (Fmt[0] == TEXT('h') && Fmt[1] == TEXT('s'))
 				{
 					if constexpr (bIsPointerTo<Arg, UTF8CHAR>)
 					{
@@ -226,14 +225,14 @@ namespace UE::Core::Private
 					}
 					else
 					{
-						return TCheckFormatString<TFmtArgTypes<Args...>>::Check(false, CurArgPos + 1, P + 2);
+						return TCheckFormatString<TFmtArgTypes<Args...>>::Check(false, CurArgPos + 1, Fmt + 2);
 					}
 				}
 
-				switch (P[0])
+				switch (Fmt[0])
 				{
 				case TEXT('%'):
-					return TCheckFormatString<TFmtArgTypes<Arg, Args...>>::Check(false, CurArgPos, P + 1);
+					return TCheckFormatString<TFmtArgTypes<Arg, Args...>>::Check(false, CurArgPos, Fmt + 1);
 				case TEXT('s'):
 					if constexpr (bIsPointerTo<Arg, UTF8CHAR>)
 					{
@@ -253,7 +252,7 @@ namespace UE::Core::Private
 					}
 					else
 					{
-						return TCheckFormatString<TFmtArgTypes<Args...>>::Check(false, CurArgPos + 1, P + 1);
+						return TCheckFormatString<TFmtArgTypes<Args...>>::Check(false, CurArgPos + 1, Fmt + 1);
 					}
 				case TEXT('c'):
 					if constexpr (!(std::is_integral_v<Arg> && sizeof(Arg) <= sizeof(void*)))
@@ -262,7 +261,7 @@ namespace UE::Core::Private
 					}
 					else
 					{
-						return TCheckFormatString<TFmtArgTypes<Args...>>::Check(false, CurArgPos + 1, P + 1);
+						return TCheckFormatString<TFmtArgTypes<Args...>>::Check(false, CurArgPos + 1, Fmt + 1);
 					}
 				case TEXT('d'):
 				case TEXT('i'):
@@ -279,10 +278,10 @@ namespace UE::Core::Private
 					}
 					else
 					{
-						return TCheckFormatString<TFmtArgTypes<Args...>>::Check(false, CurArgPos + 1, P + 1);
+						return TCheckFormatString<TFmtArgTypes<Args...>>::Check(false, CurArgPos + 1, Fmt + 1);
 					}
 				case TEXT('z'):
-					if (!CharIsIntegerFormatSpecifier(P[1]))
+					if (!CharIsIntegerFormatSpecifier(Fmt[1]))
 					{
 						return {EFormatStringSanStatus::ZNeedsIntegerSpec, CurArgPos};
 					}
@@ -292,7 +291,7 @@ namespace UE::Core::Private
 					}
 					else
 					{
-						return TCheckFormatString<TFmtArgTypes<Args...>>::Check(false, CurArgPos + 1, P + 1);
+						return TCheckFormatString<TFmtArgTypes<Args...>>::Check(false, CurArgPos + 1, Fmt + 1);
 					}
 				case TEXT('p'):
 					if constexpr (!std::is_pointer_v<Arg>)
@@ -301,14 +300,14 @@ namespace UE::Core::Private
 					}
 					else
 					{
-						return TCheckFormatString<TFmtArgTypes<Args...>>::Check(false, CurArgPos + 1, P + 1);
+						return TCheckFormatString<TFmtArgTypes<Args...>>::Check(false, CurArgPos + 1, Fmt + 1);
 					}
 				case TEXT('I'):
-					if (!(P[1] == TEXT('6') && P[2] == TEXT('4')))
+					if (!(Fmt[1] == TEXT('6') && Fmt[2] == TEXT('4')))
 					{
 						return {EFormatStringSanStatus::I64BadSpec, CurArgPos};
 					}
-					else if (!CharIsIntegerFormatSpecifier(P[3]))
+					else if (!CharIsIntegerFormatSpecifier(Fmt[3]))
 					{
 						return {EFormatStringSanStatus::I64BadSpec, CurArgPos};
 					}
@@ -318,11 +317,11 @@ namespace UE::Core::Private
 					}
 					else
 					{
-						return TCheckFormatString<TFmtArgTypes<Args...>>::Check(false, CurArgPos + 1, P + 4);
+						return TCheckFormatString<TFmtArgTypes<Args...>>::Check(false, CurArgPos + 1, Fmt + 4);
 					}
 
 				case TEXT('l'):
-					if (CharIsIntegerFormatSpecifier(P[1]))
+					if (CharIsIntegerFormatSpecifier(Fmt[1]))
 					{
 						if constexpr (!(bIsIntegralEnum<Arg> || std::is_integral_v<Arg>))
 						{
@@ -330,25 +329,25 @@ namespace UE::Core::Private
 						}
 						else
 						{
-							return TCheckFormatString<TFmtArgTypes<Args...>>::Check(false, CurArgPos + 1, P + 2);
+							return TCheckFormatString<TFmtArgTypes<Args...>>::Check(false, CurArgPos + 1, Fmt + 2);
 						}
 					}
-					else if (P[1] == TEXT('f'))
+					else if (Fmt[1] == TEXT('f'))
 					{
-						if constexpr (!bIsFloatOrDouble<Arg>)
+						if constexpr (!TIsFloatOrDouble_V<Arg>)
 						{
 							return {EFormatStringSanStatus::FNeedsFloatOrDoubleArg, CurArgPos};
 						}
 						else
 						{
-							return TCheckFormatString<TFmtArgTypes<Args...>>::Check(false, CurArgPos + 1, P + 2);
+							return TCheckFormatString<TFmtArgTypes<Args...>>::Check(false, CurArgPos + 1, Fmt + 2);
 						}
 					}
-					else if (P[1] != TEXT('l'))
+					else if (Fmt[1] != TEXT('l'))
 					{
 						return {EFormatStringSanStatus::InvalidFormatSpec, CurArgPos};
 					}
-					else if (!CharIsIntegerFormatSpecifier(P[2]))
+					else if (!CharIsIntegerFormatSpecifier(Fmt[2]))
 					{
 						return {EFormatStringSanStatus::LLNeedsIntegerSpec, CurArgPos};
 					}
@@ -358,23 +357,23 @@ namespace UE::Core::Private
 					}
 					else
 					{
-						return TCheckFormatString<TFmtArgTypes<Args...>>::Check(false, CurArgPos + 1, P + 3);
+						return TCheckFormatString<TFmtArgTypes<Args...>>::Check(false, CurArgPos + 1, Fmt + 3);
 					}
 
 				case TEXT('h'):
-					if (CharIsIntegerFormatSpecifier(P[1]))
+					if (CharIsIntegerFormatSpecifier(Fmt[1]))
 					{
 						if constexpr (!(bIsIntegralEnum<Arg> || std::is_integral_v<Arg>))
 						{
 							return {EFormatStringSanStatus::HNeedsIntegerArg, CurArgPos};
 						}
-						return TCheckFormatString<TFmtArgTypes<Args...>>::Check(false, CurArgPos + 1, P + 2);
+						return TCheckFormatString<TFmtArgTypes<Args...>>::Check(false, CurArgPos + 1, Fmt + 2);
 					}
-					else if (P[1] != TEXT('h'))
+					else if (Fmt[1] != TEXT('h'))
 					{
 						return {EFormatStringSanStatus::InvalidFormatSpec, CurArgPos};
 					}
-					else if (!CharIsIntegerFormatSpecifier(P[2]))
+					else if (!CharIsIntegerFormatSpecifier(Fmt[2]))
 					{
 						return {EFormatStringSanStatus::HHNeedsIntegerSpec, CurArgPos};
 					}
@@ -384,19 +383,19 @@ namespace UE::Core::Private
 					}
 					else
 					{
-						return TCheckFormatString<TFmtArgTypes<Args...>>::Check(false, CurArgPos + 1, P + 3);
+						return TCheckFormatString<TFmtArgTypes<Args...>>::Check(false, CurArgPos + 1, Fmt + 3);
 					}
 
 				case TEXT('f'):
 				case TEXT('e'):
 				case TEXT('g'):
-					if constexpr (!bIsFloatOrDouble<Arg>)
+					if constexpr (!TIsFloatOrDouble_V<Arg>)
 					{
 						return {EFormatStringSanStatus::FNeedsFloatOrDoubleArg, CurArgPos};
 					}
 					else
 					{
-						return TCheckFormatString<TFmtArgTypes<Args...>>::Check(false, CurArgPos + 1, P + 1);
+						return TCheckFormatString<TFmtArgTypes<Args...>>::Check(false, CurArgPos + 1, Fmt + 1);
 					}
 
 				case TEXT('S'):
@@ -414,7 +413,7 @@ namespace UE::Core::Private
 					}
 					else
 					{
-						return TCheckFormatString<TFmtArgTypes<Args...>>::Check(false, CurArgPos + 1, P + 1);
+						return TCheckFormatString<TFmtArgTypes<Args...>>::Check(false, CurArgPos + 1, Fmt + 1);
 					}
 
 				case TEXT(' '):
@@ -429,28 +428,28 @@ namespace UE::Core::Private
 		template <>
 		struct TCheckFormatString<TFmtArgTypes<>>
 		{
-			static constexpr FResult Check(bool bInsideFormatSpec, int CurArgPos, const TCHAR* P)
+			static constexpr FResult Check(bool bInsideFormatSpec, int CurArgPos, const TCHAR* Fmt)
 			{
 				if (bInsideFormatSpec)
 				{
 					return {EFormatStringSanStatus::IncompleteFormatSpecifierOrUnescapedPercent, CurArgPos};
 				}
 
-				while (*P)
+				while (*Fmt)
 				{
-					if (P[0] != TEXT('%'))
+					if (Fmt[0] != TEXT('%'))
 					{
-						++P;
+						++Fmt;
 						continue;
 					}
 
-					if (P[1] == TEXT('%'))
+					if (Fmt[1] == TEXT('%'))
 					{
-						P += 2;
+						Fmt += 2;
 						continue;
 					}
 
-					if (P[1] == 0)
+					if (Fmt[1] == 0)
 					{
 						return {EFormatStringSanStatus::IncompleteFormatSpecifierOrUnescapedPercent, CurArgPos};
 					}
@@ -469,11 +468,12 @@ namespace UE::Core::Private
 
 #if UE_VALIDATE_FORMAT_STRINGS && defined(__cpp_consteval)
 			template <int32 N>
-			consteval TCheckedFormatStringPrivate(const CharType (&Fmt)[N]) : FormatString{Fmt}
+			consteval TCheckedFormatStringPrivate(const CharType (&Fmt)[N])
+				: FormatString{Fmt}
 			{
 				if constexpr (std::is_convertible_v<const CharType*, const TCHAR*>)
 				{
-					typedef TCheckFormatString<TFmtArgTypes<std::decay_t<ArgTypes>...>> Checker;
+					using Checker = TCheckFormatString<TFmtArgTypes<std::decay_t<ArgTypes>...>>;
 					FResult Result = Checker::Check(false, 0, Fmt);
 
 					switch (Result.Status)
@@ -501,7 +501,8 @@ namespace UE::Core::Private
 #else
 
 			template <int32 N>
-			TCheckedFormatStringPrivate(const CharType(&Fmt)[N]) : FormatString{Fmt}
+			TCheckedFormatStringPrivate(const CharType(&Fmt)[N])
+				: FormatString{Fmt}
 			{
 				static_assert((TIsValidVariadicFunctionArg<ArgTypes>::Value && ...), "Invalid argument(s) passed to Printf");
 			}
