Refactored TCheckFormatString to be parameterizable by character type.

diff --git a/Engine/Source/Runtime/Core/Public/String/FormatStringSan.h b/Engine/Source/Runtime/Core/Public/String/FormatStringSan.h
--- a/Engine/Source/Runtime/Core/Public/String/FormatStringSan.h
+++ b/Engine/Source/Runtime/Core/Public/String/FormatStringSan.h
@@ -104,18 +104,22 @@ namespace UE::Core::Private
 		template <typename>
 		struct TCheckFormatString
 		{
-			static constexpr FResult Check(bool, int, const TCHAR*);
+			using CharType = TCHAR;
+
+			static constexpr FResult Check(bool, int, const CharType*);
 		};
 
 		template <typename Arg, typename... Args>
 		struct TCheckFormatString<TFmtArgTypes<Arg, Args...>>
 		{
-			static constexpr bool CharIsDigit(const TCHAR Char)
+			using CharType = TCHAR;
+
+			static constexpr bool CharIsDigit(const CharType Char)
 			{
-				return Char >= TEXT('0') && Char <= TEXT('9');
+				return Char >= CHARTEXT(CharType, '0') && Char <= CHARTEXT(CharType, '9');
 			}
 
-			static constexpr const TCHAR* SkipInteger(const TCHAR* Fmt)
+			static constexpr const CharType* SkipInteger(const CharType* Fmt)
 			{
 				while (CharIsDigit(*Fmt))
 				{
@@ -124,18 +128,18 @@ namespace UE::Core::Private
 				return Fmt;
 			}
 
-			static constexpr bool CharIsIntegerFormatSpecifier(const TCHAR Char)
+			static constexpr bool CharIsIntegerFormatSpecifier(const CharType Char)
 			{
 				switch (Char)
 				{
-				case TEXT('i'): case TEXT('d'): case TEXT('u'): case TEXT('X'): case TEXT('x'):
+				case CHARTEXT(CharType, 'i'): case CHARTEXT(CharType, 'd'): case CHARTEXT(CharType, 'u'): case CHARTEXT(CharType, 'X'): case CHARTEXT(CharType, 'x'):
 					return true;
 				default:
 					return false;
 				}
 			}
 
-			static constexpr FResult HandleDynamicLengthSpecifier(int CurArgPos, const TCHAR* Fmt)
+			static constexpr FResult HandleDynamicLengthSpecifier(int CurArgPos, const CharType* Fmt)
 			{
 				if constexpr (!(std::is_integral_v<Arg> || bIsIntegralEnum<Arg>))
 				{
@@ -147,7 +151,7 @@ namespace UE::Core::Private
 				}
 			};
 
-			static constexpr FResult Check(bool bInsideFormatSpec, int CurArgPos, const TCHAR* Fmt)
+			static constexpr FResult Check(bool bInsideFormatSpec, int CurArgPos, const CharType* Fmt)
 			{
 				if (!Fmt[0] && !bInsideFormatSpec)
 				{
@@ -156,11 +160,11 @@ namespace UE::Core::Private
 
 				if (!bInsideFormatSpec)
 				{
-					while (*Fmt && *Fmt != TEXT('%'))
+					while (*Fmt && *Fmt != CHARTEXT(CharType, '%'))
 					{
 						++Fmt;
 					}
-					if (*Fmt == TEXT('%'))
+					if (*Fmt == CHARTEXT(CharType, '%'))
 					{
 						++Fmt;
 					}
@@ -168,7 +172,7 @@ namespace UE::Core::Private
 
 				if (!*Fmt)
 				{
-					if (Fmt[-1] == TEXT('%') || bInsideFormatSpec) // nb: P[-1] is safe here because zero-length fmt strings are addressed above
+					if (Fmt[-1] == CHARTEXT(CharType, '%') || bInsideFormatSpec) // nb: P[-1] is safe here because zero-length fmt strings are addressed above
 					{
 						return {EFormatStringSanStatus::IncompleteFormatSpecifierOrUnescapedPercent, CurArgPos};
 					}
@@ -178,16 +182,16 @@ namespace UE::Core::Private
 					}
 				}
 
-				while (*Fmt == TEXT('+') || *Fmt == TEXT('#') || *Fmt == TEXT(' ') || *Fmt == TEXT('0'))
+				while (*Fmt == CHARTEXT(CharType, '+') || *Fmt == CHARTEXT(CharType, '#') || *Fmt == CHARTEXT(CharType, ' ') || *Fmt == CHARTEXT(CharType, '0'))
 				{
 					++Fmt;
 				}
 
-				if (*Fmt == TEXT('-'))
+				if (*Fmt == CHARTEXT(CharType, '-'))
 				{
 					++Fmt;
 				}
-				if (*Fmt == TEXT('*'))
+				if (*Fmt == CHARTEXT(CharType, '*'))
 				{
 					return HandleDynamicLengthSpecifier(CurArgPos, Fmt);
 				}
@@ -196,11 +200,11 @@ namespace UE::Core::Private
 					Fmt = SkipInteger(Fmt);
 				}
 
-				if (*Fmt == TEXT('.'))
+				if (*Fmt == CHARTEXT(CharType, '.'))
 				{
 					++Fmt;
 				}
-				if (*Fmt == TEXT('*'))
+				if (*Fmt == CHARTEXT(CharType, '*'))
 				{
 					return HandleDynamicLengthSpecifier(CurArgPos, Fmt);
 				}
@@ -209,11 +213,11 @@ namespace UE::Core::Private
 					Fmt = SkipInteger(Fmt);
 				}
 
-				if (Fmt[0] == TEXT('l') && Fmt[1] == TEXT('s'))
+				if (Fmt[0] == CHARTEXT(CharType, 'l') && Fmt[1] == CHARTEXT(CharType, 's'))
 				{
 					++Fmt;
 				}
-				else if (Fmt[0] == TEXT('h') && Fmt[1] == TEXT('s'))
+				else if (Fmt[0] == CHARTEXT(CharType, 'h') && Fmt[1] == CHARTEXT(CharType, 's'))
 				{
 					if constexpr (bIsPointerTo<Arg, UTF8CHAR>)
 					{
@@ -231,9 +235,9 @@ namespace UE::Core::Private
 
 				switch (Fmt[0])
 				{
-				case TEXT('%'):
+				case CHARTEXT(CharType, '%'):
 					return TCheckFormatString<TFmtArgTypes<Arg, Args...>>::Check(false, CurArgPos, Fmt + 1);
-				case TEXT('s'):
+				case CHARTEXT(CharType, 's'):
 					if constexpr (bIsPointerTo<Arg, UTF8CHAR>)
 					{
 						return {EFormatStringSanStatus::UTF8NeedsStringCast, CurArgPos};
@@ -254,7 +258,7 @@ namespace UE::Core::Private
 					{
 						return TCheckFormatString<TFmtArgTypes<Args...>>::Check(false, CurArgPos + 1, Fmt + 1);
 					}
-				case TEXT('c'):
+				case CHARTEXT(CharType, 'c'):
 					if constexpr (!(std::is_integral_v<Arg> && sizeof(Arg) <= sizeof(void*)))
 					{
 						return {EFormatStringSanStatus::CNeedsTCHARArg, CurArgPos};
@@ -263,11 +267,11 @@ namespace UE::Core::Private
 					{
 						return TCheckFormatString<TFmtArgTypes<Args...>>::Check(false, CurArgPos + 1, Fmt + 1);
 					}
-				case TEXT('d'):
-				case TEXT('i'):
-				case TEXT('X'):
-				case TEXT('x'):
-				case TEXT('u'):
+				case CHARTEXT(CharType, 'd'):
+				case CHARTEXT(CharType, 'i'):
+				case CHARTEXT(CharType, 'X'):
+				case CHARTEXT(CharType, 'x'):
+				case CHARTEXT(CharType, 'u'):
 					if constexpr (bIsTCharPtr<Arg>)
 					{
 						return {EFormatStringSanStatus::DNeedsIntegerArg, CurArgPos};
@@ -280,7 +284,7 @@ namespace UE::Core::Private
 					{
 						return TCheckFormatString<TFmtArgTypes<Args...>>::Check(false, CurArgPos + 1, Fmt + 1);
 					}
-				case TEXT('z'):
+				case CHARTEXT(CharType, 'z'):
 					if (!CharIsIntegerFormatSpecifier(Fmt[1]))
 					{
 						return {EFormatStringSanStatus::ZNeedsIntegerSpec, CurArgPos};
@@ -293,7 +297,7 @@ namespace UE::Core::Private
 					{
 						return TCheckFormatString<TFmtArgTypes<Args...>>::Check(false, CurArgPos + 1, Fmt + 1);
 					}
-				case TEXT('p'):
+				case CHARTEXT(CharType, 'p'):
 					if constexpr (!std::is_pointer_v<Arg>)
 					{
 						return {EFormatStringSanStatus::PNeedsPointerArg, CurArgPos};
@@ -302,8 +306,8 @@ namespace UE::Core::Private
 					{
 						return TCheckFormatString<TFmtArgTypes<Args...>>::Check(false, CurArgPos + 1, Fmt + 1);
 					}
-				case TEXT('I'):
-					if (!(Fmt[1] == TEXT('6') && Fmt[2] == TEXT('4')))
+				case CHARTEXT(CharType, 'I'):
+					if (!(Fmt[1] == CHARTEXT(CharType, '6') && Fmt[2] == CHARTEXT(CharType, '4')))
 					{
 						return {EFormatStringSanStatus::I64BadSpec, CurArgPos};
 					}
@@ -320,7 +324,7 @@ namespace UE::Core::Private
 						return TCheckFormatString<TFmtArgTypes<Args...>>::Check(false, CurArgPos + 1, Fmt + 4);
 					}
 
-				case TEXT('l'):
+				case CHARTEXT(CharType, 'l'):
 					if (CharIsIntegerFormatSpecifier(Fmt[1]))
 					{
 						if constexpr (!(bIsIntegralEnum<Arg> || std::is_integral_v<Arg>))
@@ -332,7 +336,7 @@ namespace UE::Core::Private
 							return TCheckFormatString<TFmtArgTypes<Args...>>::Check(false, CurArgPos + 1, Fmt + 2);
 						}
 					}
-					else if (Fmt[1] == TEXT('f'))
+					else if (Fmt[1] == CHARTEXT(CharType, 'f'))
 					{
 						if constexpr (!TIsFloatOrDouble_V<Arg>)
 						{
@@ -343,7 +347,7 @@ namespace UE::Core::Private
 							return TCheckFormatString<TFmtArgTypes<Args...>>::Check(false, CurArgPos + 1, Fmt + 2);
 						}
 					}
-					else if (Fmt[1] != TEXT('l'))
+					else if (Fmt[1] != CHARTEXT(CharType, 'l'))
 					{
 						return {EFormatStringSanStatus::InvalidFormatSpec, CurArgPos};
 					}
@@ -360,7 +364,7 @@ namespace UE::Core::Private
 						return TCheckFormatString<TFmtArgTypes<Args...>>::Check(false, CurArgPos + 1, Fmt + 3);
 					}
 
-				case TEXT('h'):
+				case CHARTEXT(CharType, 'h'):
 					if (CharIsIntegerFormatSpecifier(Fmt[1]))
 					{
 						if constexpr (!(bIsIntegralEnum<Arg> || std::is_integral_v<Arg>))
@@ -369,7 +373,7 @@ namespace UE::Core::Private
 						}
 						return TCheckFormatString<TFmtArgTypes<Args...>>::Check(false, CurArgPos + 1, Fmt + 2);
 					}
-					else if (Fmt[1] != TEXT('h'))
+					else if (Fmt[1] != CHARTEXT(CharType, 'h'))
 					{
 						return {EFormatStringSanStatus::InvalidFormatSpec, CurArgPos};
 					}
@@ -386,9 +390,9 @@ namespace UE::Core::Private
 						return TCheckFormatString<TFmtArgTypes<Args...>>::Check(false, CurArgPos + 1, Fmt + 3);
 					}
 
-				case TEXT('f'):
-				case TEXT('e'):
-				case TEXT('g'):
+				case CHARTEXT(CharType, 'f'):
+				case CHARTEXT(CharType, 'e'):
+				case CHARTEXT(CharType, 'g'):
 					if constexpr (!TIsFloatOrDouble_V<Arg>)
 					{
 						return {EFormatStringSanStatus::FNeedsFloatOrDoubleArg, CurArgPos};
@@ -398,7 +402,7 @@ namespace UE::Core::Private
 						return TCheckFormatString<TFmtArgTypes<Args...>>::Check(false, CurArgPos + 1, Fmt + 1);
 					}
 
-				case TEXT('S'):
+				case CHARTEXT(CharType, 'S'):
 					using OtherCharType = std::conditional_t<std::is_same_v<TCHAR, ANSICHAR>, WIDECHAR, ANSICHAR>;
 					if constexpr (!bIsPointerTo<Arg, OtherCharType>)
 					{
@@ -416,7 +420,7 @@ namespace UE::Core::Private
 						return TCheckFormatString<TFmtArgTypes<Args...>>::Check(false, CurArgPos + 1, Fmt + 1);
 					}
 
-				case TEXT(' '):
+				case CHARTEXT(CharType, ' '):
 					return {EFormatStringSanStatus::IncompleteFormatSpecifierOrUnescapedPercent, CurArgPos};
 
 				default:
@@ -428,7 +432,9 @@ namespace UE::Core::Private
 		template <>
 		struct TCheckFormatString<TFmtArgTypes<>>
 		{
-			static constexpr FResult Check(bool bInsideFormatSpec, int CurArgPos, const TCHAR* Fmt)
+			using CharType = TCHAR;
+
+			static constexpr FResult Check(bool bInsideFormatSpec, int CurArgPos, const CharType* Fmt)
 			{
 				if (bInsideFormatSpec)
 				{
@@ -437,13 +443,13 @@ namespace UE::Core::Private
 
 				while (*Fmt)
 				{
-					if (Fmt[0] != TEXT('%'))
+					if (Fmt[0] != CHARTEXT(CharType, '%'))
 					{
 						++Fmt;
 						continue;
 					}
 
-					if (Fmt[1] == TEXT('%'))
+					if (Fmt[1] == CHARTEXT(CharType, '%'))
 					{
 						Fmt += 2;
 						continue;
