[AutoRTFM] Add an initializer function which takes function pointers to platform APIs

This is the first of many changes to cut AutoRTFM dependencies on Core. This is required to allow other modules that Core depends on to also use AutoRTFM.

The autortfm_extern_api structure currently consists of heap allocation function pointers, but will be expanded in future CLs to include logging and other hooks.

diff --git a/Engine/Source/Programs/AutoRTFMEngineTests/Private/AutoRTFMEngineTests.cpp b/Engine/Source/Programs/AutoRTFMEngineTests/Private/AutoRTFMEngineTests.cpp
--- a/Engine/Source/Programs/AutoRTFMEngineTests/Private/AutoRTFMEngineTests.cpp
+++ b/Engine/Source/Programs/AutoRTFMEngineTests/Private/AutoRTFMEngineTests.cpp
@@ -1,6 +1,7 @@
 // Copyright Epic Games, Inc. All Rights Reserved.
 
 #include "AutoRTFMEngineTests.h"
+#include "AutoRTFM/AutoRTFMUE.h"
 #include "RequiredProgramMainCPPInclude.h"
 #include "Styling/UMGCoreStyle.h"
 #include "AutomationTestRunner.h"
@@ -197,6 +198,9 @@ static void PreInit()
 {
 	PreventLoadingOfEditorOnlyData();
 
+	// Register the Unreal Engine memory allocator.
+	AutoRTFM::InitializeForUE();
+
 	// We enable the AutoRTFM runtime as the tests depend on it.
 	AutoRTFM::ForTheRuntime::SetAutoRTFMRuntime(AutoRTFM::ForTheRuntime::AutoRTFM_Enabled);
 
diff --git a/Engine/Source/Runtime/Core/Private/AutoRTFM/API.cpp b/Engine/Source/Runtime/Core/Private/AutoRTFM/API.cpp
--- a/Engine/Source/Runtime/Core/Private/AutoRTFM/API.cpp
+++ b/Engine/Source/Runtime/Core/Private/AutoRTFM/API.cpp
@@ -3,6 +3,7 @@
 #include "AutoRTFM/AutoRTFM.h"
 #include "AutoRTFM/Context.h"
 #include "CoreGlobals.h"
+#include "ExternAPI.h"
 #include "HAL/IConsoleManager.h"
 #include "GenericPlatform/GenericPlatformCrashContext.h"
 #include "HAL/Platform.h"
@@ -725,6 +726,26 @@ AUTORTFM_CORE_API void RTFM_autortfm_did_free(void* Ptr)
 }  // anonymous namespace
 
 // The AutoRTFM public API.
+
+extern "C" void autortfm_initialize(const autortfm_extern_api* ExternAPI)
+{
+	if (ExternAPI)
+	{
+		AutoRTFM::Allocate = ExternAPI->Allocate
+			? ExternAPI->Allocate
+			: &AutoRTFM::ExternAPIDefaults::Allocate;
+		AutoRTFM::AllocateZeroed = ExternAPI->AllocateZeroed
+			? ExternAPI->AllocateZeroed
+			: &AutoRTFM::ExternAPIDefaults::AllocateZeroed;
+		AutoRTFM::Reallocate = ExternAPI->Reallocate
+			? ExternAPI->Reallocate
+			: &AutoRTFM::ExternAPIDefaults::Reallocate;
+		AutoRTFM::Free = ExternAPI->Free
+			? ExternAPI->Free
+			: &AutoRTFM::ExternAPIDefaults::Free;
+	}
+}
+
 // Each function will be forked by the compiler into an open and closed variant.
 // autortfm_is_closed() is used to branch to the closed variants declared above.
 extern "C" bool autortfm_is_transactional()
@@ -982,6 +1003,11 @@ extern "C" bool autortfm_is_on_current_transaction_stack(void* Ptr)
 	}
 }
 
+void ForTheRuntime::Initialize(const ForTheRuntime::FExternAPI& ExternAPI)
+{
+	autortfm_initialize(&ExternAPI);
+}
+
 ETransactionResult ForTheRuntime::CascadingAbortTransactionInternal(TTask<void()>&& RunAfterAbort)
 {
 	if (autortfm_is_closed())
diff --git a/Engine/Source/Runtime/Core/Private/AutoRTFM/AutoRTFMUE.cpp b/Engine/Source/Runtime/Core/Private/AutoRTFM/AutoRTFMUE.cpp
new file mode 100644
--- /dev/null
+++ b/Engine/Source/Runtime/Core/Private/AutoRTFM/AutoRTFMUE.cpp
@@ -0,0 +1,28 @@
+// Copyright Epic Games, Inc. All Rights Reserved.
+
+#if (defined(__AUTORTFM) && __AUTORTFM)
+
+#include "AutoRTFM/AutoRTFMUE.h"
+#include "AutoRTFM/AutoRTFM.h"
+#include "HAL/UnrealMemory.h"
+
+void AutoRTFM::InitializeForUE()
+{
+	AutoRTFM::ForTheRuntime::FExternAPI ExternAPI;
+	ExternAPI.Allocate = [](size_t Size, size_t Alignment)
+	{
+		return FMemory::Malloc(Size, Alignment);
+	};
+	ExternAPI.Reallocate = [](void* Pointer, size_t Size, size_t Alignment)
+	{
+		return FMemory::Realloc(Pointer, Size, Alignment);
+	};
+	ExternAPI.AllocateZeroed = [](size_t Size, size_t Alignment)
+	{
+		return FMemory::MallocZeroed(Size, Alignment);
+	};
+	ExternAPI.Free = FMemory::Free;
+	AutoRTFM::ForTheRuntime::Initialize(ExternAPI);
+}
+
+#endif  // (defined(__AUTORTFM) && __AUTORTFM)
diff --git a/Engine/Source/Runtime/Core/Private/AutoRTFM/BlockAllocator.h b/Engine/Source/Runtime/Core/Private/AutoRTFM/BlockAllocator.h
--- a/Engine/Source/Runtime/Core/Private/AutoRTFM/BlockAllocator.h
+++ b/Engine/Source/Runtime/Core/Private/AutoRTFM/BlockAllocator.h
@@ -3,14 +3,14 @@
 #pragma once
 
 #include "AutoRTFM/Utils.h"
-#include "HAL/UnrealMemory.h"
+#include "ExternAPI.h"
 #include <cstddef>
 
 namespace AutoRTFM
 {
 
 // A single-threaded, bump memory allocator that sub-allocates out of blocks of
-// inline memory and once that's exhausted, FMemory allocated memory.
+// inline memory and once that's exhausted, from heap allocated memory.
 // All allocated memory is freed when the TBlockAllocator is destructed.
 // There is no way to free individual allocations.
 //
@@ -37,7 +37,7 @@ public:
 		while (Tail != &InlineBlock)
 		{
 			FBlockHeader* Prev = Tail->Prev;
-			FMemory::Free(Tail);
+			AutoRTFM::Free(Tail);
 			Tail = Prev;
 		}
 		InlineBlock.Remaining = InlineBlockDataSize;
@@ -92,10 +92,10 @@ private:
 		size_t Remaining = BlockDataSize;
 		// <data>
 
-		// Allocates, constructs and returns a new block from FMemory::Malloc.
+		// Allocates, constructs and returns a new block from the heap-allocated memory.
 		static FBlockHeader* New(FBlockHeader* Prev, size_t BlockDataSize)
 		{
-			void* Memory = FMemory::Malloc(sizeof(FBlockHeader) + BlockDataSize);
+			void* Memory = AutoRTFM::Allocate(sizeof(FBlockHeader) + BlockDataSize, BlockAlignment);
 			return new (Memory) FBlockHeader{Prev, BlockDataSize};
 		}
 
diff --git a/Engine/Source/Runtime/Core/Private/AutoRTFM/ExternAPI.cpp b/Engine/Source/Runtime/Core/Private/AutoRTFM/ExternAPI.cpp
new file mode 100644
--- /dev/null
+++ b/Engine/Source/Runtime/Core/Private/AutoRTFM/ExternAPI.cpp
@@ -0,0 +1,42 @@
+// Copyright Epic Games, Inc. All Rights Reserved.
+
+#if (defined(__AUTORTFM) && __AUTORTFM)
+
+#include "ExternAPI.h"
+#include "Utils.h"
+#include <cstdlib>
+
+namespace AutoRTFM
+{
+
+void* ExternAPIDefaults::Allocate(size_t Size, size_t Alignment)
+{
+	ASSERT(Alignment <= 16); // If this fires, we need to over-allocate and manually align.
+	return malloc(Size);
+}
+
+void* ExternAPIDefaults::Reallocate(void* Ptr, size_t Size, size_t Alignment)
+{
+	ASSERT(Alignment <= 16); // If this fires, we need to over-allocate and manually align.
+	return realloc(Ptr, Size);
+}
+
+void* ExternAPIDefaults::AllocateZeroed(size_t Size, size_t Alignment)
+{
+	ASSERT(Alignment <= 16); // If this fires, we need to over-allocate and manually align.
+	return calloc(Size, 1);
+}
+
+void ExternAPIDefaults::Free(void* Ptr)
+{
+	free(Ptr);
+}
+
+FAllocateFn Allocate = ExternAPIDefaults::Allocate;
+FAllocateFn AllocateZeroed = ExternAPIDefaults::AllocateZeroed;
+FReallocateFn Reallocate = ExternAPIDefaults::Reallocate;
+FFreeFn Free = ExternAPIDefaults::Free;
+
+} // namespace AutoRTFM
+
+#endif // defined(__AUTORTFM) && __AUTORTFM
diff --git a/Engine/Source/Runtime/Core/Private/AutoRTFM/ExternAPI.h b/Engine/Source/Runtime/Core/Private/AutoRTFM/ExternAPI.h
new file mode 100644
--- /dev/null
+++ b/Engine/Source/Runtime/Core/Private/AutoRTFM/ExternAPI.h
@@ -0,0 +1,28 @@
+// Copyright Epic Games, Inc. All Rights Reserved.
+
+#pragma once
+
+#include <cstddef>
+
+namespace AutoRTFM
+{
+
+namespace ExternAPIDefaults
+{
+void* Allocate(size_t Size, size_t Alignment);
+void* Reallocate(void* Pointer, size_t Size, size_t Alignment);
+void* AllocateZeroed(size_t Size, size_t Alignment);
+void Free(void* Pointer);
+};
+
+using FAllocateFn = decltype(ExternAPIDefaults::Allocate)*;
+using FReallocateFn = decltype(ExternAPIDefaults::Reallocate)*;
+using FFreeFn = decltype(ExternAPIDefaults::Free)*;
+
+extern FAllocateFn Allocate;
+extern FReallocateFn Reallocate;
+extern FAllocateFn AllocateZeroed;
+extern FFreeFn Free;
+
+} // namespace AutoRTFM
+
diff --git a/Engine/Source/Runtime/Core/Private/AutoRTFM/HitSet.h b/Engine/Source/Runtime/Core/Private/AutoRTFM/HitSet.h
--- a/Engine/Source/Runtime/Core/Private/AutoRTFM/HitSet.h
+++ b/Engine/Source/Runtime/Core/Private/AutoRTFM/HitSet.h
@@ -2,6 +2,7 @@
 
 #pragma once
 
+#include "ExternAPI.h"
 #include "Utils.h"
 #include <cstdint>
 
@@ -107,7 +108,7 @@ namespace AutoRTFM
 			{
 				if (Payload != SmallPayload)
 				{
-					FMemory::Free(Payload);
+					AutoRTFM::Free(Payload);
 					Payload = SmallPayload;
 				}
 				memset(SmallPayload, 0, sizeof(SmallPayload));
@@ -167,7 +168,7 @@ namespace AutoRTFM
 
 				IncreaseCapacity();
 
-				Payload = static_cast<uintptr_t*>(FMemory::MallocZeroed(Capacity() * sizeof(uintptr_t)));
+				Payload = static_cast<uintptr_t*>(AutoRTFM::AllocateZeroed(Capacity() * sizeof(uintptr_t), alignof(uintptr_t)));
 				ASSERT(nullptr != Payload);
 
 				// Now we need to rehash and reinsert all the items.
@@ -196,7 +197,7 @@ namespace AutoRTFM
 				{
 					if (Payload != SmallPayload)
 					{
-						FMemory::Free(Payload);
+						AutoRTFM::Free(Payload);
 					}
 					continue;
 				}
@@ -208,7 +209,7 @@ namespace AutoRTFM
 
 			if (OldPayload != SmallPayload)
 			{
-				FMemory::Free(OldPayload);
+				AutoRTFM::Free(OldPayload);
 			}
 
 			return true;
diff --git a/Engine/Source/Runtime/Core/Public/AutoRTFM/AutoRTFM.h b/Engine/Source/Runtime/Core/Public/AutoRTFM/AutoRTFM.h
--- a/Engine/Source/Runtime/Core/Public/AutoRTFM/AutoRTFM.h
+++ b/Engine/Source/Runtime/Core/Public/AutoRTFM/AutoRTFM.h
@@ -137,6 +137,35 @@ typedef enum
 	autortfm_status_in_post_abort
 } autortfm_status;
 
+// Function pointers used by AutoRTFM for heap allocations, etc.
+typedef struct
+{
+	// The function used to allocate memory from the heap.
+	// If null, then a default implementation will be used.
+	void* (*Allocate)(size_t Size, size_t Alignment);
+	// The function used to reallocate memory from the heap.
+	// If null, then a default implementation will be used.
+	void* (*Reallocate)(void* Pointer, size_t Size, size_t Alignment);
+	// The function used to allocate zeroed memory from the heap.
+	// If null, then a default implementation will be used.
+	void* (*AllocateZeroed)(size_t Size, size_t Alignment);
+	// The function used to free memory allocated by Allocate() and AllocateZeroed().
+	// If null, then a default implementation will be used.
+	void (*Free)(void* Pointer);
+} autortfm_extern_api;
+
+#if UE_AUTORTFM_ENABLED
+// Initialize the AutoRTFM library.
+// Parameters:
+//   ExternAPI - if not null, then AutoRTFM will use the provided function
+//               pointers for heap allocations, etc.
+UE_AUTORTFM_API void autortfm_initialize(const autortfm_extern_api* ExternAPI);
+#else
+static UE_AUTORTFM_FORCEINLINE void autortfm_initialize(const autortfm_extern_api* ExternAPI)
+{
+}
+#endif
+
 #if UE_AUTORTFM_ENABLED
 // Note: There is no implementation of this function.
 // The AutoRTFM compiler will replace all calls to this function with a constant boolean value.
@@ -471,6 +500,8 @@ enum class EMemoryValidationLevel : uint8_t
 #if UE_AUTORTFM
 namespace ForTheRuntime
 {
+	using FExternAPI = autortfm_extern_api;
+	UE_AUTORTFM_API void Initialize(const FExternAPI& ExternAPI);
 	UE_AUTORTFM_API ETransactionResult CascadingAbortTransactionInternal(TTask<void()>&& RunAfterAbort);
 	UE_AUTORTFM_API void CascadingRetryTransactionInternal(TTask<void()>&& RunAfterAbortBeforeRetryWork);
 	UE_AUTORTFM_API void OnCommitInternal(TTask<void()>&& Work);
diff --git a/Engine/Source/Runtime/Core/Public/AutoRTFM/AutoRTFMUE.h b/Engine/Source/Runtime/Core/Public/AutoRTFM/AutoRTFMUE.h
new file mode 100644
--- /dev/null
+++ b/Engine/Source/Runtime/Core/Public/AutoRTFM/AutoRTFMUE.h
@@ -0,0 +1,15 @@
+// Copyright Epic Games, Inc. All Rights Reserved.
+
+#pragma once
+
+namespace AutoRTFM
+{
+
+#if (defined(__AUTORTFM_ENABLED) && __AUTORTFM_ENABLED)
+// Initializes AutoRTFM for use by the Unreal Engine.
+void InitializeForUE();
+#else
+inline void InitializeForUE() {}
+#endif // (defined(__AUTORTFM_ENABLED) && __AUTORTFM_ENABLED)
+
+}
