Added a FWideString alias and TIsTString trait for testing if a type is a UE string class.

diff --git a/Engine/Source/Runtime/Core/Public/Containers/ContainersFwd.h b/Engine/Source/Runtime/Core/Public/Containers/ContainersFwd.h
--- a/Engine/Source/Runtime/Core/Public/Containers/ContainersFwd.h
+++ b/Engine/Source/Runtime/Core/Public/Containers/ContainersFwd.h
@@ -16,6 +16,14 @@ class FString;
 class FAnsiString;
 class FUtf8String;
 
+// FWideString is an alias for FString.
+//
+// This is so ANSICHAR/UTF8CHAR/WIDECHAR can be matched with FAnsiString/FUtf8String/FWideString when overloaded or specialized.
+// 
+// FWideString should be the 'real' string class and FString should be the alias, but can't be for legacy reasons.
+// Forward declarations of FString expect it to be a class and changing it would affect ABIs and bloat PDBs.
+using FWideString = FString;
+
 template<> struct TIsContiguousContainer<FString>     { static constexpr bool Value = true; };
 template<> struct TIsContiguousContainer<FAnsiString> { static constexpr bool Value = true; };
 template<> struct TIsContiguousContainer<FUtf8String> { static constexpr bool Value = true; };
@@ -26,9 +34,11 @@ namespace UE::Core::Private
 	struct TCharTypeToStringType;
 
 	template <>
-	struct TCharTypeToStringType<TCHAR>
+	struct TCharTypeToStringType<WIDECHAR>
 	{
-		using Type = FString;
+		// TCHAR has been wide for too long to fix
+		static_assert(sizeof(TCHAR) == sizeof(WIDECHAR), "TCHAR is expected to be wide");
+		using Type = FWideString;
 	};
 
 	template <>
diff --git a/Engine/Source/Runtime/Core/Public/Traits/IsTString.h b/Engine/Source/Runtime/Core/Public/Traits/IsTString.h
new file mode 100644
--- /dev/null
+++ b/Engine/Source/Runtime/Core/Public/Traits/IsTString.h
@@ -0,0 +1,25 @@
+// Copyright Epic Games, Inc. All Rights Reserved.
+
+#pragma once
+
+#include "Containers/ContainersFwd.h"
+
+/**
+ * Traits classes and constants which test if a type is a UE string class.
+ */
+
+template <typename T> inline constexpr bool TIsTString_V              = false;
+template <>           inline constexpr bool TIsTString_V<FAnsiString> = true;
+template <>           inline constexpr bool TIsTString_V<FUtf8String> = true;
+template <>           inline constexpr bool TIsTString_V<FWideString> = true;
+
+template <typename T> inline constexpr bool TIsTString_V<const          T> = TIsTString_V<T>;
+template <typename T> inline constexpr bool TIsTString_V<      volatile T> = TIsTString_V<T>;
+template <typename T> inline constexpr bool TIsTString_V<const volatile T> = TIsTString_V<T>;
+
+template <typename T>
+struct TIsTString
+{
+	static constexpr bool Value = TIsTString_V<T>;
+	static constexpr bool value = Value;
+};
