Use polynomial root finding to find cubic curve closest points, which should be faster and more accurate

diff --git a/Engine/Source/Runtime/Core/Public/Math/InterpCurve.h b/Engine/Source/Runtime/Core/Public/Math/InterpCurve.h
--- a/Engine/Source/Runtime/Core/Public/Math/InterpCurve.h
+++ b/Engine/Source/Runtime/Core/Public/Math/InterpCurve.h
@@ -6,6 +6,7 @@
 #include "Misc/AssertionMacros.h"
 #include "Algo/MinElement.h"
 #include "Containers/Array.h"
+#include "Math/PolynomialRootSolver.h"
 #include "Math/UnrealMathUtility.h"
 #include "Math/Color.h"
 #include "Math/Vector2D.h"
@@ -615,136 +616,51 @@ float FInterpCurve<T>::FindNearestOnSegment(const T& PointInSpace, int32 PtIdx,
 	}
 
 	{
-		// To be accurate and fast, the algorithm is divided into 2 steps:
-		// A first step starting from the 4 start coordinates with a large tolerance to find a good approximation in few iterations
-		// if needed, a last step starting from the best approximation with a small tolerance to be accurate is performed
-		// Few small differences in the Newton algorithms:
-		// - in the first step, the break is done before the last evaluation (not needed) (BreakBeforeEvaluateStep1)
-		// - in the second step, the last evaluation is done before the break as the result point is the final point needed to evaluate the error (BreakAfterEvaluateStep2)
-		// - in the first step, to avoid duplicated test, if the new coordinate is:
-		//    - before the range ("start value" - 1/3), the process is break
-		//    - after the range ("start value" + 1/3), the next "start value" is not test
-
-		constexpr float ToleranceStep1 = 1.e-2f;
-		constexpr float ToleranceStep2 = 1.e-5f;
-		constexpr float ToleranceToLaunchStep2 = 1.e-6f; // Step1 break before the evaluation. So if the last move of step1 is not so small, a last evaluation has to be done
-
-		constexpr int32 MaxIteration = 20;
-		constexpr float InvThree = 1.f / 3.f;
-		constexpr float TwoInvThree = 2.f / 3.f;
-
-		// The parametric space is divided into 3 areas. The search starts from the four limits of the areas
-		// So Newton's methods is repeated 4 times
-		constexpr int32 NewtonIterationCount = 4;
-
-		constexpr float StartValuesT[] = { 0.f, InvThree, TwoInvThree, 1.f };
-		float ValuesT[] = {0.f, InvThree, TwoInvThree, 1.f};
-		float NextMovesT[] = { 0.f, 0.f, 0.f, 0.f };
-		float DistancesSq[NewtonIterationCount];
-
-		float LastValue = UE_BIG_NUMBER;
-		float MinValue = 0.f;
-		float MaxValue = UE_BIG_NUMBER;
-		int32 NextValue = 1;
-
 		const TArray<FInterpCurvePoint<T>>& PointsT = Points;
 
-		const auto BreakIfConverged = [&LastValue, &MinValue, &MaxValue, &NextValue, &StartValuesT, NewtonIterationCount](float& Value, const float Tolerance) -> bool
-		{
-			if (FMath::IsNearlyEqual(LastValue, Value, Tolerance))
-			{
-				// Has converged
-				Value = LastValue;
-				return true;
-			}
-
-			if (Value > MaxValue)
-			{
-				// The new coordinate is after the range ("start value" + 1/3), the next "start value" is not test
-				// Value is clamp into [0.,1.], so value cannot be > to 1.
-				++NextValue;
-				MaxValue = StartValuesT[NextValue >= NewtonIterationCount ? NewtonIterationCount - 1 : NextValue];
-			}
-
-			if (Value < MinValue)
-			{
-				// The new coordinate is before the range, the process is break as the area has been proceed 
-				return true;
-			}
-
-			LastValue = Value;
-
-			return false;
-		};
-
-		const auto NeverBreak = [](float& Value, const float Tolerance) -> bool
+		// Get the cubic's control points, shifted so PointInSpace is at the origin
+		const T P0 = PointsT[PtIdx].OutVal - PointInSpace;
+		const T P1 = PointsT[NextPtIdx].OutVal - PointInSpace;
+		const T T0 = PointsT[PtIdx].LeaveTangent * Diff;
+		const T T1 = PointsT[NextPtIdx].ArriveTangent * Diff;
+		const T CubicCoeffs[4] = { P0, T0, -3 * P0 - 2 * T0 - T1 + 3 * P1, 2 * P0 + T0 + T1 - 2 * P1 };
+		
+		// Curve is closest to PointInSpace when (Curve - PointInSpace).Dot(CurveDerivative) == 0
+		// Since we pre-subtracted PointInSpace, this becomes Curve.Dot(CurveDerivative) == 0,
+		// which expands out to the degree-5 polynomial w/ the below coefficients
+		TStaticArray<double, 6> Coeffs;
+		Coeffs[5] = 3 * CubicCoeffs[3].Dot(CubicCoeffs[3]);
+		Coeffs[4] = 5 * CubicCoeffs[3].Dot(CubicCoeffs[2]);
+		Coeffs[3] = 4 * CubicCoeffs[3].Dot(CubicCoeffs[1]) + 2 * CubicCoeffs[2].Dot(CubicCoeffs[2]);
+		Coeffs[2] = 3 * CubicCoeffs[2].Dot(CubicCoeffs[1]) + 3 * CubicCoeffs[3].Dot(CubicCoeffs[0]);
+		Coeffs[1] = CubicCoeffs[1].Dot(CubicCoeffs[1]) + 2 * CubicCoeffs[0].Dot(CubicCoeffs[2]);
+		Coeffs[0] = CubicCoeffs[0].Dot(CubicCoeffs[1]);
+
+		// Test the endpoints first -- recall P0 and P1 already have PointInSpace subtracted out
+		float BestDistSq = static_cast<float>(P0.SizeSquared());
+		float BestParam = 0;
+		float EndDistSq = static_cast<float>(P1.SizeSquared());
+		if (EndDistSq < BestDistSq)
 		{
-			return false;
-		};
-
-		const auto Newton =
-			[&PointsT, &PtIdx, &NextPtIdx, &PointInSpace, &Diff, MaxIteration, InvThree]
-			(float& Value, float& Move, TFunctionRef<bool(float&, const float)> BreakBeforeEvaluate, TFunctionRef<bool(float&, const float)> BreakAfterEvaluate, const float Tolerance) -> float
+			BestParam = 1;
+			BestDistSq = EndDistSq;
+		}
+		// Check the roots
+		UE::Math::TPolynomialRootSolver<double, 5> RootSolver(Coeffs, 0, 1);
+		for (int32 RootIdx = 0; RootIdx < RootSolver.Roots.Num(); ++RootIdx)
 		{
-			// Initialize FoundPoint to some meaningful points on the curve
-			T FoundPoint = Value > 0.5 ? PointsT[NextPtIdx].OutVal : PointsT[PtIdx].OutVal;
-
-			for (int32 Iter = 0; Iter < MaxIteration; ++Iter)
+			T FoundPoint = FMath::CubicInterp(P0, T0, P1, T1, RootSolver.Roots[RootIdx]);
+			float RootDistSq = static_cast<float>(FoundPoint.SizeSquared());
+			if (RootDistSq < BestDistSq)
 			{
-				Value += Move;
-				Value = FMath::Clamp(Value, 0.0f, 1.0f);
-
-				if (BreakBeforeEvaluate(Value, Tolerance))
-				{
-					break;
-				}
-
-				FoundPoint = FMath::CubicInterp(PointsT[PtIdx].OutVal, PointsT[PtIdx].LeaveTangent * Diff, PointsT[NextPtIdx].OutVal, PointsT[NextPtIdx].ArriveTangent * Diff, Value);
-
-				if (BreakAfterEvaluate(Value, Tolerance))
-				{
-					break;
-				}
-
-				const T Tangent = FMath::CubicInterpDerivative(PointsT[PtIdx].OutVal, PointsT[PtIdx].LeaveTangent * Diff, PointsT[NextPtIdx].OutVal, PointsT[NextPtIdx].ArriveTangent * Diff, Value);
-				const T Delta = (PointInSpace - FoundPoint);
-				Move = static_cast<float>(Tangent.Dot(Delta) / Tangent.SizeSquared());
+				BestParam = static_cast<float>(RootSolver.Roots[RootIdx]);
+				BestDistSq = RootDistSq;
 			}
-			return static_cast<float>((FoundPoint - PointInSpace).SizeSquared());
-		};
-
-		// A first step starting from the 4 start coordinates with a large tolerance to find a good approximation in few iterations
-		for (int32 Index = 0; Index < NewtonIterationCount; ++Index)
-		{
-			MinValue = StartValuesT[Index == 0 ? 0 : Index - 1];
-			MaxValue = StartValuesT[NextValue >= NewtonIterationCount ? Index : NextValue];
-
-			float& Value = ValuesT[Index];
-			float& Move = NextMovesT[Index];
-
-			DistancesSq[Index] = Newton(Value, Move, BreakIfConverged, NeverBreak, ToleranceStep1);
-
-			// Reset LastValue after calculation so that each newton iteration actually converge
-			LastValue = UE_BIG_NUMBER;
-		}
-
-		// Find the index of the best approximation
-		const float* MinValuePtr = Algo::MinElement(DistancesSq);
-		const int32 IndexOfTheSmallest = (int32)(MinValuePtr - DistancesSq);
-
-		float& Move = NextMovesT[IndexOfTheSmallest];
-
-		// if needed, a last step starting from the best approximation with a small tolerance to be accurate is performed
-		if (!FMath::IsNearlyZero(Move, ToleranceToLaunchStep2))
-		{
-			MinValue = 0.f;
-			MaxValue = 1.f;
-			float& Value = ValuesT[IndexOfTheSmallest];
-			DistancesSq[IndexOfTheSmallest] = Newton(Value, Move, NeverBreak, BreakIfConverged, ToleranceStep2);
 		}
+		// Return the best parameter value for the segment, shifted back to overall curve parameter space
+		OutSquaredDistance = BestDistSq;
+		return BestParam * Diff + Points[PtIdx].InVal;
 
-		OutSquaredDistance = DistancesSq[IndexOfTheSmallest];
-		return ValuesT[IndexOfTheSmallest] * Diff + Points[PtIdx].InVal;
 	}
 }
 
