Moved TIsTEnumAsByte into its own header and added TIsTEnumAsByte_V.

diff --git a/Engine/Source/Runtime/Core/Public/Containers/EnumAsByte.h b/Engine/Source/Runtime/Core/Public/Containers/EnumAsByte.h
--- a/Engine/Source/Runtime/Core/Public/Containers/EnumAsByte.h
+++ b/Engine/Source/Runtime/Core/Public/Containers/EnumAsByte.h
@@ -125,17 +125,6 @@ template<class T>
 
 template<class T> struct TIsPODType<TEnumAsByte<T>> { enum { Value = true }; };
 
-template <typename T>
-struct TIsTEnumAsByte
-{
-	static constexpr bool Value = false;
-};
-
-template <typename T> struct TIsTEnumAsByte<               TEnumAsByte<T>> { static constexpr bool Value = true; };
-template <typename T> struct TIsTEnumAsByte<const          TEnumAsByte<T>> { static constexpr bool Value = true; };
-template <typename T> struct TIsTEnumAsByte<      volatile TEnumAsByte<T>> { static constexpr bool Value = true; };
-template <typename T> struct TIsTEnumAsByte<const volatile TEnumAsByte<T>> { static constexpr bool Value = true; };
-
-#if UE_ENABLE_INCLUDE_ORDER_DEPRECATED_IN_5_4
-#include "Templates/IsEnumClass.h"
+#if UE_ENABLE_INCLUDE_ORDER_DEPRECATED_IN_5_6
+#include "Traits/IsTEnumAsByte.h"
 #endif
diff --git a/Engine/Source/Runtime/Core/Public/String/FormatStringSan.h b/Engine/Source/Runtime/Core/Public/String/FormatStringSan.h
--- a/Engine/Source/Runtime/Core/Public/String/FormatStringSan.h
+++ b/Engine/Source/Runtime/Core/Public/String/FormatStringSan.h
@@ -6,6 +6,7 @@
 #include "Templates/Requires.h"
 #include "Templates/Identity.h"
 #include "Templates/IsValidVariadicFunctionArg.h"
+#include "Traits/IsTEnumAsByte.h"
 #include "Containers/ContainersFwd.h"
 
 #define UE_CHECK_FORMAT_STRING(Fmt, ...) \
@@ -82,15 +83,8 @@ namespace UE::Core::Private
 		template <typename T>
 		inline constexpr bool TIsFloatOrDouble_V = std::is_same_v<float, T> || std::is_same_v<double, T>;
 
-		// Returns true only when the type is TEnumAsByte
 		template <typename T>
-		struct TIsEnumAsByte { static constexpr bool IsEnumAsByte = false; };
-		template <typename T>
-		struct TIsEnumAsByte<TEnumAsByte<T>> { static constexpr bool IsEnumAsByte = true; };
-
-
-		template <typename T>
-		inline constexpr bool bIsIntegralEnum = std::is_enum_v<T> || TIsEnumAsByte<T>::IsEnumAsByte;
+		inline constexpr bool bIsIntegralEnum = std::is_enum_v<T> || TIsTEnumAsByte_V<T>;
 
 		template <typename... Ts>
 		struct TFmtArgTypes {};
diff --git a/Engine/Source/Runtime/Core/Public/Traits/IsTEnumAsByte.h b/Engine/Source/Runtime/Core/Public/Traits/IsTEnumAsByte.h
new file mode 100644
--- /dev/null
+++ b/Engine/Source/Runtime/Core/Public/Traits/IsTEnumAsByte.h
@@ -0,0 +1,23 @@
+// Copyright Epic Games, Inc. All Rights Reserved.
+
+#pragma once
+
+/**
+ * Traits classes and constants which test if a type is a TEnumAsByte.
+ */
+
+template <class InEnumType>
+class TEnumAsByte;
+
+template <typename T> inline constexpr bool TIsTEnumAsByte_V                                = false;
+template <typename T> inline constexpr bool TIsTEnumAsByte_V<               TEnumAsByte<T>> = true;
+template <typename T> inline constexpr bool TIsTEnumAsByte_V<const          TEnumAsByte<T>> = true;
+template <typename T> inline constexpr bool TIsTEnumAsByte_V<      volatile TEnumAsByte<T>> = true;
+template <typename T> inline constexpr bool TIsTEnumAsByte_V<const volatile TEnumAsByte<T>> = true;
+
+template <typename T>
+struct TIsTEnumAsByte
+{
+	static constexpr bool Value = TIsTEnumAsByte_V<T>;
+	static constexpr bool value = Value;
+};
diff --git a/Engine/Source/Runtime/PropertyPath/Public/PropertyTypeCompatibility.h b/Engine/Source/Runtime/PropertyPath/Public/PropertyTypeCompatibility.h
--- a/Engine/Source/Runtime/PropertyPath/Public/PropertyTypeCompatibility.h
+++ b/Engine/Source/Runtime/PropertyPath/Public/PropertyTypeCompatibility.h
@@ -3,6 +3,7 @@
 #pragma once
 
 #include "CoreMinimal.h"
+#include "Traits/IsTEnumAsByte.h"
 #include "UObject/Class.h"
 #include "UObject/WeakObjectPtr.h"
 #include "UObject/UnrealType.h"
@@ -388,7 +389,7 @@ struct FConcreteTypeCompatibleWithReflectedTypeHelper<TLazyObjectPtr<T>>
 template<typename T>
 static inline bool IsConcreteTypeCompatibleWithReflectedType(FProperty* Property)
 {
-	if constexpr (TIsEnum<T>::Value || TIsTEnumAsByte<T>::Value)
+	if constexpr (TIsEnum<T>::Value || TIsTEnumAsByte_V<T>)
 	{
 		return FConcreteTypeCompatibleWithReflectedTypeHelper<uint8>::IsConcreteTypeCompatibleWithReflectedType(Property);
 	}
