Reverting changelist 39551817 due to unusual intermittent syntax errors.

#rnx

diff --git a/Engine/Source/Runtime/Core/Public/String/FormatStringSan.h b/Engine/Source/Runtime/Core/Public/String/FormatStringSan.h
--- a/Engine/Source/Runtime/Core/Public/String/FormatStringSan.h
+++ b/Engine/Source/Runtime/Core/Public/String/FormatStringSan.h
@@ -98,22 +98,18 @@ namespace UE::Core::Private
 		template <typename>
 		struct TCheckFormatString
 		{
-			using CharType = TCHAR;
-
-			static constexpr FResult Check(bool, int, const CharType*);
+			static constexpr FResult Check(bool, int, const TCHAR*);
 		};
 
 		template <typename Arg, typename... Args>
 		struct TCheckFormatString<TFmtArgTypes<Arg, Args...>>
 		{
-			using CharType = TCHAR;
-
-			static constexpr bool CharIsDigit(const CharType Char)
+			static constexpr bool CharIsDigit(const TCHAR Char)
 			{
-				return Char >= CHARTEXT(CharType, '0') && Char <= CHARTEXT(CharType, '9');
+				return Char >= TEXT('0') && Char <= TEXT('9');
 			}
 
-			static constexpr const CharType* SkipInteger(const CharType* Fmt)
+			static constexpr const TCHAR* SkipInteger(const TCHAR* Fmt)
 			{
 				while (CharIsDigit(*Fmt))
 				{
@@ -122,18 +118,18 @@ namespace UE::Core::Private
 				return Fmt;
 			}
 
-			static constexpr bool CharIsIntegerFormatSpecifier(const CharType Char)
+			static constexpr bool CharIsIntegerFormatSpecifier(const TCHAR Char)
 			{
 				switch (Char)
 				{
-				case CHARTEXT(CharType, 'i'): case CHARTEXT(CharType, 'd'): case CHARTEXT(CharType, 'u'): case CHARTEXT(CharType, 'X'): case CHARTEXT(CharType, 'x'):
+				case TEXT('i'): case TEXT('d'): case TEXT('u'): case TEXT('X'): case TEXT('x'):
 					return true;
 				default:
 					return false;
 				}
 			}
 
-			static constexpr FResult HandleDynamicLengthSpecifier(int CurArgPos, const CharType* Fmt)
+			static constexpr FResult HandleDynamicLengthSpecifier(int CurArgPos, const TCHAR* Fmt)
 			{
 				if constexpr (!(std::is_integral_v<Arg> || bIsIntegralEnum<Arg>))
 				{
@@ -145,7 +141,7 @@ namespace UE::Core::Private
 				}
 			};
 
-			static constexpr FResult Check(bool bInsideFormatSpec, int CurArgPos, const CharType* Fmt)
+			static constexpr FResult Check(bool bInsideFormatSpec, int CurArgPos, const TCHAR* Fmt)
 			{
 				if (!Fmt[0] && !bInsideFormatSpec)
 				{
@@ -154,11 +150,11 @@ namespace UE::Core::Private
 
 				if (!bInsideFormatSpec)
 				{
-					while (*Fmt && *Fmt != CHARTEXT(CharType, '%'))
+					while (*Fmt && *Fmt != TEXT('%'))
 					{
 						++Fmt;
 					}
-					if (*Fmt == CHARTEXT(CharType, '%'))
+					if (*Fmt == TEXT('%'))
 					{
 						++Fmt;
 					}
@@ -166,7 +162,7 @@ namespace UE::Core::Private
 
 				if (!*Fmt)
 				{
-					if (Fmt[-1] == CHARTEXT(CharType, '%') || bInsideFormatSpec) // nb: P[-1] is safe here because zero-length fmt strings are addressed above
+					if (Fmt[-1] == TEXT('%') || bInsideFormatSpec) // nb: P[-1] is safe here because zero-length fmt strings are addressed above
 					{
 						return {EFormatStringSanStatus::IncompleteFormatSpecifierOrUnescapedPercent, CurArgPos};
 					}
@@ -176,16 +172,16 @@ namespace UE::Core::Private
 					}
 				}
 
-				while (*Fmt == CHARTEXT(CharType, '+') || *Fmt == CHARTEXT(CharType, '#') || *Fmt == CHARTEXT(CharType, ' ') || *Fmt == CHARTEXT(CharType, '0'))
+				while (*Fmt == TEXT('+') || *Fmt == TEXT('#') || *Fmt == TEXT(' ') || *Fmt == TEXT('0'))
 				{
 					++Fmt;
 				}
 
-				if (*Fmt == CHARTEXT(CharType, '-'))
+				if (*Fmt == TEXT('-'))
 				{
 					++Fmt;
 				}
-				if (*Fmt == CHARTEXT(CharType, '*'))
+				if (*Fmt == TEXT('*'))
 				{
 					return HandleDynamicLengthSpecifier(CurArgPos, Fmt);
 				}
@@ -194,11 +190,11 @@ namespace UE::Core::Private
 					Fmt = SkipInteger(Fmt);
 				}
 
-				if (*Fmt == CHARTEXT(CharType, '.'))
+				if (*Fmt == TEXT('.'))
 				{
 					++Fmt;
 				}
-				if (*Fmt == CHARTEXT(CharType, '*'))
+				if (*Fmt == TEXT('*'))
 				{
 					return HandleDynamicLengthSpecifier(CurArgPos, Fmt);
 				}
@@ -207,11 +203,11 @@ namespace UE::Core::Private
 					Fmt = SkipInteger(Fmt);
 				}
 
-				if (Fmt[0] == CHARTEXT(CharType, 'l') && Fmt[1] == CHARTEXT(CharType, 's'))
+				if (Fmt[0] == TEXT('l') && Fmt[1] == TEXT('s'))
 				{
 					++Fmt;
 				}
-				else if (Fmt[0] == CHARTEXT(CharType, 'h') && Fmt[1] == CHARTEXT(CharType, 's'))
+				else if (Fmt[0] == TEXT('h') && Fmt[1] == TEXT('s'))
 				{
 					if constexpr (bIsPointerTo<Arg, UTF8CHAR>)
 					{
@@ -229,9 +225,9 @@ namespace UE::Core::Private
 
 				switch (Fmt[0])
 				{
-				case CHARTEXT(CharType, '%'):
+				case TEXT('%'):
 					return TCheckFormatString<TFmtArgTypes<Arg, Args...>>::Check(false, CurArgPos, Fmt + 1);
-				case CHARTEXT(CharType, 's'):
+				case TEXT('s'):
 					if constexpr (bIsPointerTo<Arg, UTF8CHAR>)
 					{
 						return {EFormatStringSanStatus::UTF8NeedsStringCast, CurArgPos};
@@ -252,7 +248,7 @@ namespace UE::Core::Private
 					{
 						return TCheckFormatString<TFmtArgTypes<Args...>>::Check(false, CurArgPos + 1, Fmt + 1);
 					}
-				case CHARTEXT(CharType, 'c'):
+				case TEXT('c'):
 					if constexpr (!(std::is_integral_v<Arg> && sizeof(Arg) <= sizeof(void*)))
 					{
 						return {EFormatStringSanStatus::CNeedsTCHARArg, CurArgPos};
@@ -261,11 +257,11 @@ namespace UE::Core::Private
 					{
 						return TCheckFormatString<TFmtArgTypes<Args...>>::Check(false, CurArgPos + 1, Fmt + 1);
 					}
-				case CHARTEXT(CharType, 'd'):
-				case CHARTEXT(CharType, 'i'):
-				case CHARTEXT(CharType, 'X'):
-				case CHARTEXT(CharType, 'x'):
-				case CHARTEXT(CharType, 'u'):
+				case TEXT('d'):
+				case TEXT('i'):
+				case TEXT('X'):
+				case TEXT('x'):
+				case TEXT('u'):
 					if constexpr (bIsTCharPtr<Arg>)
 					{
 						return {EFormatStringSanStatus::DNeedsIntegerArg, CurArgPos};
@@ -278,7 +274,7 @@ namespace UE::Core::Private
 					{
 						return TCheckFormatString<TFmtArgTypes<Args...>>::Check(false, CurArgPos + 1, Fmt + 1);
 					}
-				case CHARTEXT(CharType, 'z'):
+				case TEXT('z'):
 					if (!CharIsIntegerFormatSpecifier(Fmt[1]))
 					{
 						return {EFormatStringSanStatus::ZNeedsIntegerSpec, CurArgPos};
@@ -291,7 +287,7 @@ namespace UE::Core::Private
 					{
 						return TCheckFormatString<TFmtArgTypes<Args...>>::Check(false, CurArgPos + 1, Fmt + 1);
 					}
-				case CHARTEXT(CharType, 'p'):
+				case TEXT('p'):
 					if constexpr (!std::is_pointer_v<Arg>)
 					{
 						return {EFormatStringSanStatus::PNeedsPointerArg, CurArgPos};
@@ -300,8 +296,8 @@ namespace UE::Core::Private
 					{
 						return TCheckFormatString<TFmtArgTypes<Args...>>::Check(false, CurArgPos + 1, Fmt + 1);
 					}
-				case CHARTEXT(CharType, 'I'):
-					if (!(Fmt[1] == CHARTEXT(CharType, '6') && Fmt[2] == CHARTEXT(CharType, '4')))
+				case TEXT('I'):
+					if (!(Fmt[1] == TEXT('6') && Fmt[2] == TEXT('4')))
 					{
 						return {EFormatStringSanStatus::I64BadSpec, CurArgPos};
 					}
@@ -318,7 +314,7 @@ namespace UE::Core::Private
 						return TCheckFormatString<TFmtArgTypes<Args...>>::Check(false, CurArgPos + 1, Fmt + 4);
 					}
 
-				case CHARTEXT(CharType, 'l'):
+				case TEXT('l'):
 					if (CharIsIntegerFormatSpecifier(Fmt[1]))
 					{
 						if constexpr (!(bIsIntegralEnum<Arg> || std::is_integral_v<Arg>))
@@ -330,7 +326,7 @@ namespace UE::Core::Private
 							return TCheckFormatString<TFmtArgTypes<Args...>>::Check(false, CurArgPos + 1, Fmt + 2);
 						}
 					}
-					else if (Fmt[1] == CHARTEXT(CharType, 'f'))
+					else if (Fmt[1] == TEXT('f'))
 					{
 						if constexpr (!TIsFloatOrDouble_V<Arg>)
 						{
@@ -341,7 +337,7 @@ namespace UE::Core::Private
 							return TCheckFormatString<TFmtArgTypes<Args...>>::Check(false, CurArgPos + 1, Fmt + 2);
 						}
 					}
-					else if (Fmt[1] != CHARTEXT(CharType, 'l'))
+					else if (Fmt[1] != TEXT('l'))
 					{
 						return {EFormatStringSanStatus::InvalidFormatSpec, CurArgPos};
 					}
@@ -358,7 +354,7 @@ namespace UE::Core::Private
 						return TCheckFormatString<TFmtArgTypes<Args...>>::Check(false, CurArgPos + 1, Fmt + 3);
 					}
 
-				case CHARTEXT(CharType, 'h'):
+				case TEXT('h'):
 					if (CharIsIntegerFormatSpecifier(Fmt[1]))
 					{
 						if constexpr (!(bIsIntegralEnum<Arg> || std::is_integral_v<Arg>))
@@ -367,7 +363,7 @@ namespace UE::Core::Private
 						}
 						return TCheckFormatString<TFmtArgTypes<Args...>>::Check(false, CurArgPos + 1, Fmt + 2);
 					}
-					else if (Fmt[1] != CHARTEXT(CharType, 'h'))
+					else if (Fmt[1] != TEXT('h'))
 					{
 						return {EFormatStringSanStatus::InvalidFormatSpec, CurArgPos};
 					}
@@ -384,9 +380,9 @@ namespace UE::Core::Private
 						return TCheckFormatString<TFmtArgTypes<Args...>>::Check(false, CurArgPos + 1, Fmt + 3);
 					}
 
-				case CHARTEXT(CharType, 'f'):
-				case CHARTEXT(CharType, 'e'):
-				case CHARTEXT(CharType, 'g'):
+				case TEXT('f'):
+				case TEXT('e'):
+				case TEXT('g'):
 					if constexpr (!TIsFloatOrDouble_V<Arg>)
 					{
 						return {EFormatStringSanStatus::FNeedsFloatOrDoubleArg, CurArgPos};
@@ -396,7 +392,7 @@ namespace UE::Core::Private
 						return TCheckFormatString<TFmtArgTypes<Args...>>::Check(false, CurArgPos + 1, Fmt + 1);
 					}
 
-				case CHARTEXT(CharType, 'S'):
+				case TEXT('S'):
 					using OtherCharType = std::conditional_t<std::is_same_v<TCHAR, ANSICHAR>, WIDECHAR, ANSICHAR>;
 					if constexpr (!bIsPointerTo<Arg, OtherCharType>)
 					{
@@ -414,7 +410,7 @@ namespace UE::Core::Private
 						return TCheckFormatString<TFmtArgTypes<Args...>>::Check(false, CurArgPos + 1, Fmt + 1);
 					}
 
-				case CHARTEXT(CharType, ' '):
+				case TEXT(' '):
 					return {EFormatStringSanStatus::IncompleteFormatSpecifierOrUnescapedPercent, CurArgPos};
 
 				default:
@@ -426,9 +422,7 @@ namespace UE::Core::Private
 		template <>
 		struct TCheckFormatString<TFmtArgTypes<>>
 		{
-			using CharType = TCHAR;
-
-			static constexpr FResult Check(bool bInsideFormatSpec, int CurArgPos, const CharType* Fmt)
+			static constexpr FResult Check(bool bInsideFormatSpec, int CurArgPos, const TCHAR* Fmt)
 			{
 				if (bInsideFormatSpec)
 				{
@@ -437,13 +431,13 @@ namespace UE::Core::Private
 
 				while (*Fmt)
 				{
-					if (Fmt[0] != CHARTEXT(CharType, '%'))
+					if (Fmt[0] != TEXT('%'))
 					{
 						++Fmt;
 						continue;
 					}
 
-					if (Fmt[1] == CHARTEXT(CharType, '%'))
+					if (Fmt[1] == TEXT('%'))
 					{
 						Fmt += 2;
 						continue;
