Trigger a deprecation warning if code accesses GMalloc directly.

You can still use UE::Private::GMalloc if FMemory does not offer what you need.

diff --git a/Engine/Source/Runtime/Core/Private/Misc/CoreGlobals.cpp b/Engine/Source/Runtime/Core/Private/Misc/CoreGlobals.cpp
--- a/Engine/Source/Runtime/Core/Private/Misc/CoreGlobals.cpp
+++ b/Engine/Source/Runtime/Core/Private/Misc/CoreGlobals.cpp
@@ -49,12 +49,13 @@ IMPLEMENT_MODULE( FCoreModule, Core );
 /*-----------------------------------------------------------------------------
 	Global variables.
 -----------------------------------------------------------------------------*/
-CORE_API FFeedbackContext*	GWarn					= nullptr;		/* User interaction and non critical warnings */
-FConfigCacheIni*			GConfig					= nullptr;		/* Configuration database cache */
-ITransaction*				GUndo					= nullptr;		/* Transaction tracker, non-NULL when a transaction is in progress */
-FOutputDeviceConsole*		GLogConsole				= nullptr;		/* Console log hook */
-CORE_API FMalloc**			GFixedMallocLocationPtr = nullptr;		/* Memory allocator pointer location when PLATFORM_USES_FIXED_GMalloc_CLASS is true */
-CORE_API FMalloc*			UE::Private::GMalloc	= nullptr;		/* Memory allocator */
+CORE_API FFeedbackContext*	GWarn					= nullptr;				/* User interaction and non critical warnings */
+FConfigCacheIni*			GConfig					= nullptr;				/* Configuration database cache */
+ITransaction*				GUndo					= nullptr;				/* Transaction tracker, non-NULL when a transaction is in progress */
+FOutputDeviceConsole*		GLogConsole				= nullptr;				/* Console log hook */
+CORE_API FMalloc**			GFixedMallocLocationPtr = nullptr;				/* Memory allocator pointer location when PLATFORM_USES_FIXED_GMalloc_CLASS is true */
+CORE_API FMalloc*			UE::Private::GMalloc	= nullptr;				/* Memory allocator */
+CORE_API FMalloc* const&	GMalloc	                = UE::Private::GMalloc;	/* Provides backwards compatibility for external users of un-namespaced GMalloc. */
 
 class FPropertyWindowManager*	GPropertyWindowManager	= nullptr;		/* Manages and tracks property editing windows */
 
diff --git a/Engine/Source/Runtime/Core/Public/HAL/MemoryBase.h b/Engine/Source/Runtime/Core/Public/HAL/MemoryBase.h
--- a/Engine/Source/Runtime/Core/Public/HAL/MemoryBase.h
+++ b/Engine/Source/Runtime/Core/Public/HAL/MemoryBase.h
@@ -39,11 +39,12 @@ namespace UE::Private
 }
 
 /** 
- * GMalloc remains exposed for backwards compatibility.
+ * Un-namespaced GMalloc remains exposed for backwards compatibility.
  * Most callers should use FMemory::Malloc instead of accessing GMalloc directly.
  * FMemory wraps GMalloc but also provides low-level memory tracking and AutoRTFM support.
  */
-using UE::Private::GMalloc;
+UE_DEPRECATED(5.6, "Consider using FMemory::Malloc instead of accessing GMalloc directly.")
+CORE_API extern class FMalloc* const& GMalloc;
 
 /** Memory allocator pointer location when PLATFORM_USES_FIXED_GMalloc_CLASS is true. */
 CORE_API extern class FMalloc** GFixedMallocLocationPtr;
