[AutoRTFM] Extract TPool from TTaskArray

This is generally useful, and reduces the complexity of TTaskArray.

diff --git a/Engine/Source/Programs/AutoRTFMTests/Private/PoolTests.cpp b/Engine/Source/Programs/AutoRTFMTests/Private/PoolTests.cpp
new file mode 100644
--- /dev/null
+++ b/Engine/Source/Programs/AutoRTFMTests/Private/PoolTests.cpp
@@ -0,0 +1,74 @@
+// Copyright Epic Games, Inc. All Rights Reserved.
+
+#include "AutoRTFM/Pool.h"
+#include "Catch2Includes.h"
+#include "ObjectLifetimeHelper.h"
+#include <algorithm>
+#include <vector>
+
+using TrivialPool4 = AutoRTFM::TPool<int, 4>;
+using TrivialPool64 = AutoRTFM::TPool<int, 64>;
+using NonTrivialPool4 = AutoRTFM::TPool<AutoRTFMTestUtils::FObjectLifetimeHelper, 4>;
+using NonTrivialPool64 = AutoRTFM::TPool<AutoRTFMTestUtils::FObjectLifetimeHelper, 64>;
+
+TEMPLATE_TEST_CASE("Pool", "",  TrivialPool4, TrivialPool64, NonTrivialPool4, NonTrivialPool64)
+{
+	std::mt19937 Rand(0x1234);
+	TestType Pool;
+	using TItem = typename TestType::Item;
+
+
+	REQUIRE(AutoRTFMTestUtils::FObjectLifetimeHelper::ConstructorCalls == 0);
+	REQUIRE(AutoRTFMTestUtils::FObjectLifetimeHelper::DestructorCalls == 0);
+
+	for (int Count : {1, 2, 3, 4, 5, 8, 10, 100})
+	{
+		// Take Count items from the pool
+		std::vector<TItem*> Items(Count);
+		for (int I = 0; I < Count; I++)
+		{
+			TItem* Item = Pool.Take(I);
+			REQUIRE(*Item == I);
+			Items[I] = Item;
+		}
+
+		// Check no changes to the items as new items are taken.
+		for (int I = 0; I < Count; I++)
+		{
+			REQUIRE(*Items[I] == I);
+		}
+
+		// Shuffle the order of the items
+		std::shuffle(Items.begin(), Items.end(), Rand);
+
+		// Return the items, but keep track of the pointers.
+		for (TItem* Item : Items)
+		{
+			Pool.Return(Item);
+		}
+
+		// Retake the items and check for LIFO ordering.
+		for (int I = 0; I < Count; I++)
+		{
+			TItem* Item = Pool.Take(1000 + I);
+			REQUIRE(*Item == 1000 + I);
+			REQUIRE(Items[Count - I - 1] == Item); // Check for LIFO
+		}
+
+		// Check no changes to the items as new items are taken.
+		for (int I = 0; I < Count; I++)
+		{
+			REQUIRE(*Items[Count - I - 1] == 1000 + I);
+		}
+
+		// Return all the items.
+		for (TItem* Item : Items)
+		{
+			Pool.Return(Item);
+		}
+	}
+	
+	REQUIRE(AutoRTFMTestUtils::FObjectLifetimeHelper::ConstructorCalls == AutoRTFMTestUtils::FObjectLifetimeHelper::DestructorCalls);
+	AutoRTFMTestUtils::FObjectLifetimeHelper::ConstructorCalls = 0;
+	AutoRTFMTestUtils::FObjectLifetimeHelper::DestructorCalls = 0;
+}
diff --git a/Engine/Source/Programs/AutoRTFMTests/Private/TaskArrayTests.cpp b/Engine/Source/Programs/AutoRTFMTests/Private/TaskArrayTests.cpp
--- a/Engine/Source/Programs/AutoRTFMTests/Private/TaskArrayTests.cpp
+++ b/Engine/Source/Programs/AutoRTFMTests/Private/TaskArrayTests.cpp
@@ -10,13 +10,14 @@
 TEST_CASE("TTaskArray")
 {
 	std::mt19937 Rand(0x1234);
+	using TaskType = int;
 	using TaskArrayType = AutoRTFM::TTaskArray
 		<
 			/* TaskType */ int,
-			/* KeyType */ const void*,
-			/* AllocatorBlockSize */ 256
+			/* KeyType */ const void*
 		>;
-	TaskArrayType::FPool Pool;
+	using TaskPoolType = typename TaskArrayType::FEntryPool;
+	TaskPoolType Pool;
 	auto TaskArray = std::make_unique<TaskArrayType>(Pool);
 
 	struct FEntry
diff --git a/Engine/Source/Runtime/Core/Private/AutoRTFM/Context.h b/Engine/Source/Runtime/Core/Private/AutoRTFM/Context.h
--- a/Engine/Source/Runtime/Core/Private/AutoRTFM/Context.h
+++ b/Engine/Source/Runtime/Core/Private/AutoRTFM/Context.h
@@ -8,6 +8,7 @@
 #include "StackRange.h"
 #include "OpenHashThrottler.h"
 #include "TaskArray.h"
+#include "Pool.h"
 
 namespace AutoRTFM
 {
@@ -18,6 +19,8 @@ class FCallNest;
 
 class FContext
 {
+	using FTaskPool = typename TTaskArray<TTask<void()>>::FEntryPool;
+
 public:
     static FContext* Get();
     
@@ -131,8 +134,8 @@ public:
 			}
 		}
 	}
-	
-	inline TTaskArray<TTask<void()>>::FPool& GetTaskPool() { return TaskPool; }
+
+	inline FTaskPool& GetTaskPool() { return TaskPool; }
 
 private:
 	static FContext ContextSingleton;
@@ -167,7 +170,7 @@ private:
 	uint32 CurrentThreadId{ FPlatformTLS::InvalidTlsSlot };
 	uint32 StackLocalInitializerDepth = 0;
 	TTask<void()> PostAbortCallback;
-	TTaskArray<TTask<void()>>::FPool TaskPool;
+	FTaskPool TaskPool;
 	FOpenHashThrottler OpenHashThrottler{
 		/* LogInterval            */ 10,   // Log every 10 seconds.
 		/* AdjustThrottleInterval */ 0.5,  // Adjust throttling probabilities every 500ms.
diff --git a/Engine/Source/Runtime/Core/Private/AutoRTFM/Pool.h b/Engine/Source/Runtime/Core/Private/AutoRTFM/Pool.h
new file mode 100644
--- /dev/null
+++ b/Engine/Source/Runtime/Core/Private/AutoRTFM/Pool.h
@@ -0,0 +1,115 @@
+// Copyright Epic Games, Inc. All Rights Reserved.
+
+#pragma once
+
+#include "BlockAllocator.h"
+#include <cstddef>
+#include <utility>
+
+namespace AutoRTFM
+{
+
+// An allocator of objects of a single type, backed by a TBlockAllocator with
+// and internal linked list for free items.
+// Items can be constructed and returned from the pool with Take(). 
+// The memory of items returned to the pool with Return() will be reused for
+// the next calls to Take() in a LIFO order.
+//
+// Template parameters:
+//   ItemType - the pooled item type.
+//   InlineCapacity - the number of items that can be held before spilling to
+//                    the heap.
+template<typename ItemType, size_t InlineCapacity>
+class TPool
+{
+public:
+	using Item = ItemType;
+
+	// Constructor
+	TPool() = default;
+
+	// Acquires a new Item from the pool, forward-constructing the item with Arguments.
+	template<typename ... ArgumentTypes>
+	Item* Take(ArgumentTypes&& ... Arguments)
+	{
+		void* Memory = nullptr;
+		if (FreeList)
+		{
+			// Unlink a free entry from the free list.
+			Memory = FreeList;
+			FreeList = FreeList->Next;
+		}
+		else
+		{
+			// Free list is empty. Allocate another item.
+			Memory = Allocator.Allocate(EntrySize, EntryAlignment);
+			NumAllocated++;
+		}
+
+		ASSERT(NumInUse < NumAllocated);
+		NumInUse++;
+
+		return new (Memory) Item(std::forward<ArgumentTypes>(Arguments)...);
+	}
+
+	// Destructs the item and returns it back to the pool so the memory can be
+	// reused.
+	void Return(Item* Item)
+	{
+		// Destruct the item
+		Item->~Item();
+
+		// Place the entry onto the free list
+		TFreeEntry* FreeEntry = reinterpret_cast<TFreeEntry*>(Item);
+		FreeEntry->Next = FreeList;
+		FreeList = FreeEntry;
+
+		ASSERT(NumInUse > 0);
+		NumInUse--;
+	}
+
+	// Frees the memory allocated for the item list.
+	// No items can be in use when calling.
+	void Reset(bool bIgnoreNonReturned = false)
+	{
+		ASSERT(bIgnoreNonReturned || NumInUse == 0);
+		Allocator.FreeAll();
+		NumInUse = 0;
+		NumAllocated = 0;
+		FreeList = nullptr;
+	}
+
+	// Returns the number of items allocated.
+	size_t GetNumAllocated() const { return NumAllocated; }
+
+	// Returns the number of items currently in use (not returned).
+	size_t GetNumInUse() const { return NumInUse; }
+
+private:
+	struct TFreeEntry
+	{
+		TFreeEntry* Next;
+	};
+
+	// The memory of each entry is used for both an Item (when in use) and a
+	// TFreeEntry (when in the pool), so the size and alignment needs to be
+	// the max of both.
+	static constexpr size_t EntrySize = std::max<size_t>(sizeof(Item), sizeof(TFreeEntry));
+	static constexpr size_t EntryAlignment = std::max<size_t>(alignof(Item), alignof(TFreeEntry));
+
+	TPool(TPool&&) = delete;
+	TPool(const TPool&) = delete;
+	TPool& operator=(const TPool&) = delete;
+	TPool& operator=(TPool&&) = delete;
+
+	// The underlying allocator for the pool.
+	TBlockAllocator<InlineCapacity * EntrySize, EntryAlignment> Allocator;
+	// Free list.
+	TFreeEntry* FreeList = nullptr;
+	// Number of entries allocated.
+	size_t NumAllocated = 0;
+	// Number of entries currently in use.
+	size_t NumInUse = 0;
+};
+
+}  // namespace AutoRTFM
diff --git a/Engine/Source/Runtime/Core/Private/AutoRTFM/TaskArray.h b/Engine/Source/Runtime/Core/Private/AutoRTFM/TaskArray.h
--- a/Engine/Source/Runtime/Core/Private/AutoRTFM/TaskArray.h
+++ b/Engine/Source/Runtime/Core/Private/AutoRTFM/TaskArray.h
@@ -3,10 +3,8 @@
 #pragma once
 
 #include "AutoRTFM/Utils.h"
-#include "BlockAllocator.h"
+#include "Pool.h"
 #include "Containers/Map.h"
-#include "HAL/Platform.h"
-#include "Misc/CoreMiscDefines.h"
 #include "Templates/UnrealTemplate.h"
 #include <cstddef>
 
@@ -22,25 +20,18 @@ namespace AutoRTFM
 // Task            - the task data type
 // Key             - the data type used as a key for push / pop.
 // InlineTaskCount - the number of tasks that can be allocated before a performing heap allocations.
-template<typename TaskType, typename KeyType = const void*, size_t InlineTaskCount = 256>
+template<typename TaskType, typename KeyType = const void*>
 class TTaskArray
 {
-	// A task entry used by the array and pool.
-	// Holds an intrusive linked list as well as the task.
-	class FEntry
+	struct FEntry
 	{
-	public:
-		// The next task in the sequential list.
+		// The task function.
+		TaskType Task;
+		// The next sequential task.
 		FEntry* Next = nullptr;
-		// The previous task in the sequential list.
+		// The previous sequential task.
 		FEntry* Prev = nullptr;
 
-		// Returns a reference to the task.
-		TaskType& Task()
-		{
-			return *reinterpret_cast<TaskType*>(TaskData);
-		}
-
 		// Unlinks this entry from the doubly-linked list, updating the provided
 		// head and tail pointers as necessary.
 		void Unlink(FEntry*& ListHead, FEntry*& ListTail)
@@ -69,124 +60,15 @@ class TTaskArray
 			}
 			ListTail = this;
 		}
-	private:
-		alignas(alignof(TaskType)) std::byte TaskData[sizeof(TaskType)];
 	};
-
 public:
-	// A pool of reusable tasks that can be shared between transactions.
-	class FPool
-	{
-	public:
-		// Constructor
-		FPool() = default;
-
-		// Acquires a new FEntry from the task pool, forward-constructing the task with Task.
-		template<typename TaskParam>
-		FEntry* TakeTask(TaskParam&& Task)
-		{
-			FEntry* Entry = nullptr;
-			if (FreeEntries)
-			{
-				// Take from the free list.
-				Entry = FreeEntries;
-				FreeEntries = Entry->Next;
-				Entry->Prev = nullptr;
-				Entry->Next = nullptr;
-			}
-			else
-			{
-				// Free list is empty. Allocate another entry.
-				Entry = Allocator.template New<FEntry>();
-				NumAllocated++;
-			}
-
-			// Move the task into the entry.
-			new (&Entry->Task()) TaskType(Forward<TaskParam>(Task));
-
-			ASSERT(NumInUse < NumAllocated);
-			NumInUse++;
-
-			return Entry;
-		}
-
-		// Destructs the entry's Task and returns the entry back to the pool so
-		// it can be reused.
-		// Entry must be unlinked before calling.
-		void ReturnTask(FEntry* Entry)
-		{
-			// Destruct the task
-			Entry->Task().~TaskType();
-
-			// Place the entry onto the free list
-			Entry->Prev = nullptr;
-			Entry->Next = FreeEntries;
-			FreeEntries = Entry;
-
-			ASSERT(NumInUse > 0);
-			NumInUse--;
-		}
-
-		// Destructs each entry's Task and returns the linked list of entries
-		// back to the pool so they can be reused.
-		void ReleaseTasks(FEntry* List)
-		{
-			if (UNLIKELY(!List))
-			{
-				return;
-			}
+	// The shared pool type for task arrays.
+	using FEntryPool = TPool<FEntry, 256>;
 
-			// Walk the list to destruct each of the tasks.
-			size_t ListCount = 0;
-			for (FEntry* Entry = List; Entry; Entry = Entry->Next)
-			{
-				// Destruct the task
-				Entry->Task().~TaskType();
-				ListCount++;
-			}
-
-			// Link the entries onto the free list.
-			List->Prev = nullptr;
-			List->Next = FreeEntries;
-			FreeEntries = List;
-
-			ASSERT(NumInUse >= ListCount);
-			NumInUse -= ListCount;
-		}
-
-		// Frees the memory allocated for the task list.
-		// No tasks can be in use when calling.
-		void Reset()
-		{
-			ASSERT(NumInUse == 0);
-			Allocator.FreeAll();
-			NumAllocated = 0;
-			FreeEntries = nullptr;
-		}
-
-		// Returns the number of tasks allocated.
-		size_t GetNumAllocated() const { return NumAllocated; }
-
-		// Returns the number of tasks currently in use (not returned).
-		size_t GetNumInUse() const { return NumInUse; }
-
-	private:
-		UE_NONCOPYABLE(FPool);
-
-		static constexpr size_t AllocatorInlineSize = InlineTaskCount * sizeof(TaskType);
-		static constexpr size_t AllocatorAlignment = std::max<size_t>(alignof(TaskType), alignof(void*));
-
-		TBlockAllocator<AllocatorInlineSize, AllocatorAlignment> Allocator;
-		FEntry* FreeEntries = nullptr;
-
-		// Number of entries allocated
-		size_t NumAllocated = 0;
-		// Number of entries currently in use
-		size_t NumInUse = 0;
-	};
-
-	// Constructor. TaskPool is used to allocate new tasks.
-	TTaskArray(FPool& TaskPool) : Pool{TaskPool}
+	// Constructor.
+	// TaskEntryPool is used to allocate new task entries and can be shared
+	// between task arrays.
+	TTaskArray(FEntryPool& TaskEntryPool) : EntryPool{TaskEntryPool}
 	{}
 
 	// Destructor. Returns all allocated tasks to the pool.
@@ -195,11 +77,14 @@ public:
 		Reset();
 	}
 
-	// Returns all allocated tasks to the pool.
+	// Clears the array and returns all allocated tasks to the pool.
 	void Reset()
 	{
-		Pool.ReleaseTasks(Head);
-		Head = nullptr;
+		while (FEntry* Entry = Head)
+		{
+			Head = Head->Next;
+			EntryPool.Return(Entry);
+		}
 		Tail = nullptr;
 		Keyed.Empty();
 		Count = 0;
@@ -209,7 +94,7 @@ public:
 	template<typename TaskParamType>
 	void Add(TaskParamType&& Task)
 	{
-		FEntry* Entry = Pool.TakeTask(Forward<TaskParamType>(Task));
+		FEntry* Entry = EntryPool.Take(std::forward<TaskParamType>(Task));
 		Entry->Link(Head, Tail);
 		Count++;
 	}
@@ -218,8 +103,8 @@ public:
 	template<typename TaskParamType>
 	void AddKeyed(KeyType Key, TaskParamType&& Task)
 	{
-		Add(Forward<TaskParamType>(Task));
-		Keyed.FindOrAdd(Key).Add(Tail);
+		Add(std::forward<TaskParamType>(Task));
+		Keyed.FindOrAdd(Key).Push(Tail);
 	}
 
 	// Moves all the tasks from OtherArray to the end of this array, clearing
@@ -248,8 +133,8 @@ public:
 		for (auto& Iterator : OtherArray.Keyed)
 		{
 			KeyType Key = Iterator.template Get<0>();
-			EntryStack& Stack = Iterator.template Get<1>();
-			Keyed.FindOrAdd(Key).Append(MoveTemp(Stack));
+			TEntryStack& EntryStack = Iterator.template Get<1>();
+			Keyed.FindOrAdd(Key).Append(MoveTemp(EntryStack));
 		}
 
 		// Add in OtherArray's count.
@@ -267,11 +152,10 @@ public:
 	// there are no remaining tasks with the given key.
 	bool DeleteKey(KeyType Key)
 	{
-		if (EntryStack* Stack = Keyed.Find(Key); Stack && !Stack->IsEmpty())
+		if (TEntryStack* EntryStack = Keyed.Find(Key); EntryStack)
 		{
-			Release(Stack->Pop(EAllowShrinking::No));
-
-			if (Stack->IsEmpty())
+			Release(EntryStack->Pop());
+			if (EntryStack->IsEmpty())
 			{
 				Keyed.Remove(Key);
 			}
@@ -286,13 +170,12 @@ public:
 	// there are no remaining tasks with the given key.
 	bool DeleteAllMatchingKeys(KeyType Key)
 	{
-		if (EntryStack* Stack = Keyed.Find(Key))
+		if (TEntryStack* EntryStack = Keyed.Find(Key); EntryStack)
 		{
-			while (!Stack->IsEmpty())
+			while (!EntryStack->IsEmpty())
 			{
-				Release(Stack->Pop(EAllowShrinking::No));
+				Release(EntryStack->Pop());
 			}
-
 			Keyed.Remove(Key);
 			return true;
 		}
@@ -308,9 +191,9 @@ public:
 	{
 		while (FEntry* Entry = Head)
 		{
-			Callback(Entry->Task());
+			Callback(Entry->Task);
 			Head = Entry->Next;
-			Pool.ReturnTask(Entry);
+			EntryPool.Return(Entry);
 		}
 		Head = nullptr;
 		Tail = nullptr;
@@ -326,12 +209,11 @@ public:
 	{
 		while (FEntry* Entry = Tail)
 		{
-			Callback(Entry->Task());
+			Callback(Entry->Task);
 			Tail = Entry->Prev;
-			Pool.ReturnTask(Entry);
+			EntryPool.Return(Entry);
 		}
 		Head = nullptr;
-		Tail = nullptr;
 		Keyed.Empty();
 		Count = 0;
 	}
@@ -349,27 +231,27 @@ public:
 	}
 
 private:
-	// An ordered vector of pointers to entries in the linked list.
-	using EntryStack = TArray<FEntry*, TInlineAllocator<8>>;
+	// An stack of pointers to entries in the linked list.
+	using TEntryStack = TArray<FEntry*, TInlineAllocator<8>>;
 
 	// Unlinks Entry and releases it back to the pool.
 	// Also decrements the count.
 	void Release(FEntry* Entry)
 	{
 		Entry->Unlink(Head, Tail);
-		Pool.ReturnTask(Entry);
+		EntryPool.Return(Entry);
 		Count--;
 	}
 
-	// The task pool. Holds the underlying allocator.
-	FPool& Pool;
+	// The entry pool. Holds the underlying allocator.
+	FEntryPool& EntryPool;
 
 	// The doubly-linked list of tasks held by this array.
 	FEntry* Head = nullptr;
 	FEntry* Tail = nullptr;
 
 	// A map of keyed tasks keys to their ordered task entries in the linked list.
-	TMap<KeyType, EntryStack> Keyed;
+	TMap<KeyType, TEntryStack> Keyed;
 
 	// Total number of tasks held by this array.
 	size_t Count = 0;
