[AutoRTFM] Fix MSVC build: Rename type alias

MSVC gets confused over calling a destructor when the type name is shadowed by a variable name

diff --git a/Engine/Source/Programs/AutoRTFMTests/Private/PoolTests.cpp b/Engine/Source/Programs/AutoRTFMTests/Private/PoolTests.cpp
--- a/Engine/Source/Programs/AutoRTFMTests/Private/PoolTests.cpp
+++ b/Engine/Source/Programs/AutoRTFMTests/Private/PoolTests.cpp
@@ -15,7 +15,7 @@ TEMPLATE_TEST_CASE("Pool", "",  TrivialPool4, TrivialPool64, NonTrivialPool4, No
 {
 	std::mt19937 Rand(0x1234);
 	TestType Pool;
-	using TItem = typename TestType::Item;
+	using TItem = typename TestType::FItem;
 
 
 	REQUIRE(AutoRTFMTestUtils::FObjectLifetimeHelper::ConstructorCalls == 0);
diff --git a/Engine/Source/Runtime/Core/Private/AutoRTFM/Pool.h b/Engine/Source/Runtime/Core/Private/AutoRTFM/Pool.h
--- a/Engine/Source/Runtime/Core/Private/AutoRTFM/Pool.h
+++ b/Engine/Source/Runtime/Core/Private/AutoRTFM/Pool.h
@@ -23,14 +23,14 @@ template<typename ItemType, size_t InlineCapacity>
 class TPool
 {
 public:
-	using Item = ItemType;
+	using FItem = ItemType;
 
 	// Constructor
 	TPool() = default;
 
 	// Acquires a new Item from the pool, forward-constructing the item with Arguments.
 	template<typename ... ArgumentTypes>
-	Item* Take(ArgumentTypes&& ... Arguments)
+	FItem* Take(ArgumentTypes&& ... Arguments)
 	{
 		void* Memory = nullptr;
 		if (FreeList)
@@ -49,15 +49,15 @@ public:
 		ASSERT(NumInUse < NumAllocated);
 		NumInUse++;
 
-		return new (Memory) Item(std::forward<ArgumentTypes>(Arguments)...);
+		return new (Memory) FItem(std::forward<ArgumentTypes>(Arguments)...);
 	}
 
 	// Destructs the item and returns it back to the pool so the memory can be
 	// reused.
-	void Return(Item* Item)
+	void Return(FItem* Item)
 	{
 		// Destruct the item
-		Item->~Item();
+		Item->~FItem();
 
 		// Place the entry onto the free list
 		TFreeEntry* FreeEntry = reinterpret_cast<TFreeEntry*>(Item);
@@ -94,8 +94,8 @@ private:
 	// The memory of each entry is used for both an Item (when in use) and a
 	// TFreeEntry (when in the pool), so the size and alignment needs to be
 	// the max of both.
-	static constexpr size_t EntrySize = std::max<size_t>(sizeof(Item), sizeof(TFreeEntry));
-	static constexpr size_t EntryAlignment = std::max<size_t>(alignof(Item), alignof(TFreeEntry));
+	static constexpr size_t EntrySize = std::max<size_t>(sizeof(FItem), sizeof(TFreeEntry));
+	static constexpr size_t EntryAlignment = std::max<size_t>(alignof(FItem), alignof(TFreeEntry));
 
 	TPool(TPool&&) = delete;
 	TPool(const TPool&) = delete;
