[AutoRTFM] Register stdlib maths functions with explicit overload signature.

If an earlier .cpp file in the unity build happens to include <functional>, then we'll end up with an ambiguous overload error.
Remove the platform preprocessor conditional for registering these maths functions and explicitly register the 'float', 'double' and 'long double' variants for each.

diff --git a/Engine/Source/Runtime/Core/Private/AutoRTFM/StdLibSupport.cpp b/Engine/Source/Runtime/Core/Private/AutoRTFM/StdLibSupport.cpp
--- a/Engine/Source/Runtime/Core/Private/AutoRTFM/StdLibSupport.cpp
+++ b/Engine/Source/Runtime/Core/Private/AutoRTFM/StdLibSupport.cpp
@@ -1,10 +1,13 @@
 // Copyright Epic Games, Inc. All Rights Reserved.
 
-#include "AutoRTFM/AutoRTFM.h"
-#include <charconv>
 #if (defined(__AUTORTFM) && __AUTORTFM)
+
+#include "AutoRTFM/AutoRTFM.h"
 #include "ContextInlines.h"
 #include "Utils.h"
+
+#include <charconv>
+#include <functional> // note: introduces additional math overloads
 #include "Memcpy.h"
 #include <stdarg.h>
 #include <stdio.h>
@@ -479,71 +482,66 @@ UE_AUTORTFM_REGISTER_SELF_FUNCTION(iswxdigit);
 UE_AUTORTFM_REGISTER_SELF_FUNCTION(iswspace);
 UE_AUTORTFM_REGISTER_SELF_FUNCTION(iswcntrl);
 
-#if PLATFORM_WINDOWS
-UE_AUTORTFM_REGISTER_SELF_FUNCTION(sqrt);
-UE_AUTORTFM_REGISTER_SELF_FUNCTION(sin);
-UE_AUTORTFM_REGISTER_SELF_FUNCTION(cos);
-UE_AUTORTFM_REGISTER_SELF_FUNCTION(tan);
-UE_AUTORTFM_REGISTER_SELF_FUNCTION(asin);
-UE_AUTORTFM_REGISTER_SELF_FUNCTION(acos);
-UE_AUTORTFM_REGISTER_SELF_FUNCTION(atan);
-UE_AUTORTFM_REGISTER_SELF_FUNCTION(atan2);
-UE_AUTORTFM_REGISTER_SELF_FUNCTION(sinh);
-UE_AUTORTFM_REGISTER_SELF_FUNCTION(cosh);
-UE_AUTORTFM_REGISTER_SELF_FUNCTION(tanh);
-UE_AUTORTFM_REGISTER_SELF_FUNCTION(asinh);
-UE_AUTORTFM_REGISTER_SELF_FUNCTION(acosh);
-UE_AUTORTFM_REGISTER_SELF_FUNCTION(atanh);
-UE_AUTORTFM_REGISTER_SELF_FUNCTION(exp);
-UE_AUTORTFM_REGISTER_SELF_FUNCTION(log);
-UE_AUTORTFM_REGISTER_SELF_FUNCTION(pow);
-UE_AUTORTFM_REGISTER_SELF_FUNCTION(llrint);
-UE_AUTORTFM_REGISTER_SELF_FUNCTION(fmod);
-UE_AUTORTFM_REGISTER_SELF_FUNCTION(modf);
-// Linux (likely Mac) have ambiguous overrides to these math functions
-#else
 UE_AUTORTFM_REGISTER_SELF_OVERLOADED_FUNCTION(float (float), sqrt);
+UE_AUTORTFM_REGISTER_SELF_OVERLOADED_FUNCTION(double (double), sqrt);
 UE_AUTORTFM_REGISTER_SELF_OVERLOADED_FUNCTION(long double (long double), sqrt);
 UE_AUTORTFM_REGISTER_SELF_OVERLOADED_FUNCTION(float (float), sin);
+UE_AUTORTFM_REGISTER_SELF_OVERLOADED_FUNCTION(double (double), sin);
 UE_AUTORTFM_REGISTER_SELF_OVERLOADED_FUNCTION(long double (long double), sin);
 UE_AUTORTFM_REGISTER_SELF_OVERLOADED_FUNCTION(float (float), cos);
+UE_AUTORTFM_REGISTER_SELF_OVERLOADED_FUNCTION(double (double), cos);
 UE_AUTORTFM_REGISTER_SELF_OVERLOADED_FUNCTION(long double (long double), cos);
 UE_AUTORTFM_REGISTER_SELF_OVERLOADED_FUNCTION(float (float), tan);
+UE_AUTORTFM_REGISTER_SELF_OVERLOADED_FUNCTION(double (double), tan);
 UE_AUTORTFM_REGISTER_SELF_OVERLOADED_FUNCTION(long double (long double), tan);
 UE_AUTORTFM_REGISTER_SELF_OVERLOADED_FUNCTION(float (float), asin);
+UE_AUTORTFM_REGISTER_SELF_OVERLOADED_FUNCTION(double (double), asin);
 UE_AUTORTFM_REGISTER_SELF_OVERLOADED_FUNCTION(long double (long double), asin);
 UE_AUTORTFM_REGISTER_SELF_OVERLOADED_FUNCTION(float (float), acos);
+UE_AUTORTFM_REGISTER_SELF_OVERLOADED_FUNCTION(double (double), acos);
 UE_AUTORTFM_REGISTER_SELF_OVERLOADED_FUNCTION(long double (long double), acos);
 UE_AUTORTFM_REGISTER_SELF_OVERLOADED_FUNCTION(float (float), atan);
+UE_AUTORTFM_REGISTER_SELF_OVERLOADED_FUNCTION(double (double), atan);
 UE_AUTORTFM_REGISTER_SELF_OVERLOADED_FUNCTION(long double (long double), atan);
 UE_AUTORTFM_REGISTER_SELF_OVERLOADED_FUNCTION(float (float, float), atan2);
+UE_AUTORTFM_REGISTER_SELF_OVERLOADED_FUNCTION(double (double, double), atan2);
 UE_AUTORTFM_REGISTER_SELF_OVERLOADED_FUNCTION(long double (long double, long double), atan2);
 UE_AUTORTFM_REGISTER_SELF_OVERLOADED_FUNCTION(float (float), sinh);
+UE_AUTORTFM_REGISTER_SELF_OVERLOADED_FUNCTION(double (double), sinh);
 UE_AUTORTFM_REGISTER_SELF_OVERLOADED_FUNCTION(long double (long double), sinh);
 UE_AUTORTFM_REGISTER_SELF_OVERLOADED_FUNCTION(float (float), cosh);
+UE_AUTORTFM_REGISTER_SELF_OVERLOADED_FUNCTION(double (double), cosh);
 UE_AUTORTFM_REGISTER_SELF_OVERLOADED_FUNCTION(long double (long double), cosh);
 UE_AUTORTFM_REGISTER_SELF_OVERLOADED_FUNCTION(float (float), tanh);
+UE_AUTORTFM_REGISTER_SELF_OVERLOADED_FUNCTION(double (double), tanh);
 UE_AUTORTFM_REGISTER_SELF_OVERLOADED_FUNCTION(long double (long double), tanh);
 UE_AUTORTFM_REGISTER_SELF_OVERLOADED_FUNCTION(float (float), asinh);
+UE_AUTORTFM_REGISTER_SELF_OVERLOADED_FUNCTION(double (double), asinh);
 UE_AUTORTFM_REGISTER_SELF_OVERLOADED_FUNCTION(long double (long double), asinh);
 UE_AUTORTFM_REGISTER_SELF_OVERLOADED_FUNCTION(float (float), acosh);
+UE_AUTORTFM_REGISTER_SELF_OVERLOADED_FUNCTION(double (double), acosh);
 UE_AUTORTFM_REGISTER_SELF_OVERLOADED_FUNCTION(long double (long double), acosh);
 UE_AUTORTFM_REGISTER_SELF_OVERLOADED_FUNCTION(float (float), atanh);
+UE_AUTORTFM_REGISTER_SELF_OVERLOADED_FUNCTION(double (double), atanh);
 UE_AUTORTFM_REGISTER_SELF_OVERLOADED_FUNCTION(long double (long double), atanh);
 UE_AUTORTFM_REGISTER_SELF_OVERLOADED_FUNCTION(float (float), exp);
+UE_AUTORTFM_REGISTER_SELF_OVERLOADED_FUNCTION(double (double), exp);
 UE_AUTORTFM_REGISTER_SELF_OVERLOADED_FUNCTION(long double (long double), exp);
 UE_AUTORTFM_REGISTER_SELF_OVERLOADED_FUNCTION(float (float), log);
+UE_AUTORTFM_REGISTER_SELF_OVERLOADED_FUNCTION(double (double), log);
 UE_AUTORTFM_REGISTER_SELF_OVERLOADED_FUNCTION(long double (long double), log);
 UE_AUTORTFM_REGISTER_SELF_OVERLOADED_FUNCTION(float (float, float), pow);
+UE_AUTORTFM_REGISTER_SELF_OVERLOADED_FUNCTION(double (double, double), pow);
 UE_AUTORTFM_REGISTER_SELF_OVERLOADED_FUNCTION(long double (long double, long double), pow);
 UE_AUTORTFM_REGISTER_SELF_OVERLOADED_FUNCTION(long long (float), llrint);
+UE_AUTORTFM_REGISTER_SELF_OVERLOADED_FUNCTION(long long (double), llrint);
 UE_AUTORTFM_REGISTER_SELF_OVERLOADED_FUNCTION(long long (long double), llrint);
 UE_AUTORTFM_REGISTER_SELF_OVERLOADED_FUNCTION(float (float, float), fmod);
+UE_AUTORTFM_REGISTER_SELF_OVERLOADED_FUNCTION(double (double, double), fmod);
 UE_AUTORTFM_REGISTER_SELF_OVERLOADED_FUNCTION(long double (long double, long double), fmod);
 UE_AUTORTFM_REGISTER_SELF_OVERLOADED_FUNCTION(float (float, float*), modf);
-UE_AUTORTFM_REGISTER_SELF_OVERLOADED_FUNCTION(long double (long double, long double*), modf);
 UE_AUTORTFM_REGISTER_SELF_OVERLOADED_FUNCTION(double (double, double*), modf);
-#endif
+UE_AUTORTFM_REGISTER_SELF_OVERLOADED_FUNCTION(long double (long double, long double*), modf);
 
 // Self register Math functions
 UE_AUTORTFM_REGISTER_SELF_FUNCTION(sqrtf);
