Refactored TCheckFormatString to be parameterizable by character type, avoiding CHARTEXT which causes problems with an MSVC bug.

diff --git a/Engine/Source/Runtime/Core/Public/String/FormatStringSan.h b/Engine/Source/Runtime/Core/Public/String/FormatStringSan.h
--- a/Engine/Source/Runtime/Core/Public/String/FormatStringSan.h
+++ b/Engine/Source/Runtime/Core/Public/String/FormatStringSan.h
@@ -98,18 +98,22 @@ namespace UE::Core::Private
 		template <typename>
 		struct TCheckFormatString
 		{
-			static constexpr FResult Check(bool, int, const TCHAR*);
+			using CharType = TCHAR;
+
+			static constexpr FResult Check(bool, int, const CharType*);
 		};
 
 		template <typename Arg, typename... Args>
 		struct TCheckFormatString<TFmtArgTypes<Arg, Args...>>
 		{
-			static constexpr bool CharIsDigit(const TCHAR Char)
+			using CharType = TCHAR;
+
+			static constexpr bool CharIsDigit(const CharType Char)
 			{
-				return Char >= TEXT('0') && Char <= TEXT('9');
+				return Char >= (CharType)'0' && Char <= (CharType)'9';
 			}
 
-			static constexpr const TCHAR* SkipInteger(const TCHAR* Fmt)
+			static constexpr const CharType* SkipInteger(const CharType* Fmt)
 			{
 				while (CharIsDigit(*Fmt))
 				{
@@ -118,18 +122,18 @@ namespace UE::Core::Private
 				return Fmt;
 			}
 
-			static constexpr bool CharIsIntegerFormatSpecifier(const TCHAR Char)
+			static constexpr bool CharIsIntegerFormatSpecifier(const CharType Char)
 			{
 				switch (Char)
 				{
-				case TEXT('i'): case TEXT('d'): case TEXT('u'): case TEXT('X'): case TEXT('x'):
+				case (CharType)'i': case (CharType)'d': case (CharType)'u': case (CharType)'X': case (CharType)'x':
 					return true;
 				default:
 					return false;
 				}
 			}
 
-			static constexpr FResult HandleDynamicLengthSpecifier(int CurArgPos, const TCHAR* Fmt)
+			static constexpr FResult HandleDynamicLengthSpecifier(int CurArgPos, const CharType* Fmt)
 			{
 				if constexpr (!(std::is_integral_v<Arg> || bIsIntegralEnum<Arg>))
 				{
@@ -141,7 +145,7 @@ namespace UE::Core::Private
 				}
 			};
 
-			static constexpr FResult Check(bool bInsideFormatSpec, int CurArgPos, const TCHAR* Fmt)
+			static constexpr FResult Check(bool bInsideFormatSpec, int CurArgPos, const CharType* Fmt)
 			{
 				if (!Fmt[0] && !bInsideFormatSpec)
 				{
@@ -150,11 +154,11 @@ namespace UE::Core::Private
 
 				if (!bInsideFormatSpec)
 				{
-					while (*Fmt && *Fmt != TEXT('%'))
+					while (*Fmt && *Fmt != (CharType)'%')
 					{
 						++Fmt;
 					}
-					if (*Fmt == TEXT('%'))
+					if (*Fmt == (CharType)'%')
 					{
 						++Fmt;
 					}
@@ -162,7 +166,7 @@ namespace UE::Core::Private
 
 				if (!*Fmt)
 				{
-					if (Fmt[-1] == TEXT('%') || bInsideFormatSpec) // nb: P[-1] is safe here because zero-length fmt strings are addressed above
+					if (Fmt[-1] == (CharType)'%' || bInsideFormatSpec) // nb: P[-1] is safe here because zero-length fmt strings are addressed above
 					{
 						return {EFormatStringSanStatus::IncompleteFormatSpecifierOrUnescapedPercent, CurArgPos};
 					}
@@ -172,16 +176,16 @@ namespace UE::Core::Private
 					}
 				}
 
-				while (*Fmt == TEXT('+') || *Fmt == TEXT('#') || *Fmt == TEXT(' ') || *Fmt == TEXT('0'))
+				while (*Fmt == (CharType)'+' || *Fmt == (CharType)'#' || *Fmt == (CharType)' ' || *Fmt == (CharType)'0')
 				{
 					++Fmt;
 				}
 
-				if (*Fmt == TEXT('-'))
+				if (*Fmt == (CharType)'-')
 				{
 					++Fmt;
 				}
-				if (*Fmt == TEXT('*'))
+				if (*Fmt == (CharType)'*')
 				{
 					return HandleDynamicLengthSpecifier(CurArgPos, Fmt);
 				}
@@ -190,11 +194,11 @@ namespace UE::Core::Private
 					Fmt = SkipInteger(Fmt);
 				}
 
-				if (*Fmt == TEXT('.'))
+				if (*Fmt == (CharType)'.')
 				{
 					++Fmt;
 				}
-				if (*Fmt == TEXT('*'))
+				if (*Fmt == (CharType)'*')
 				{
 					return HandleDynamicLengthSpecifier(CurArgPos, Fmt);
 				}
@@ -203,11 +207,11 @@ namespace UE::Core::Private
 					Fmt = SkipInteger(Fmt);
 				}
 
-				if (Fmt[0] == TEXT('l') && Fmt[1] == TEXT('s'))
+				if (Fmt[0] == (CharType)'l' && Fmt[1] == (CharType)'s')
 				{
 					++Fmt;
 				}
-				else if (Fmt[0] == TEXT('h') && Fmt[1] == TEXT('s'))
+				else if (Fmt[0] == (CharType)'h' && Fmt[1] == (CharType)'s')
 				{
 					if constexpr (bIsPointerTo<Arg, UTF8CHAR>)
 					{
@@ -225,9 +229,9 @@ namespace UE::Core::Private
 
 				switch (Fmt[0])
 				{
-				case TEXT('%'):
+				case (CharType)'%':
 					return TCheckFormatString<TFmtArgTypes<Arg, Args...>>::Check(false, CurArgPos, Fmt + 1);
-				case TEXT('s'):
+				case (CharType)'s':
 					if constexpr (bIsPointerTo<Arg, UTF8CHAR>)
 					{
 						return {EFormatStringSanStatus::UTF8NeedsStringCast, CurArgPos};
@@ -248,7 +252,7 @@ namespace UE::Core::Private
 					{
 						return TCheckFormatString<TFmtArgTypes<Args...>>::Check(false, CurArgPos + 1, Fmt + 1);
 					}
-				case TEXT('c'):
+				case (CharType)'c':
 					if constexpr (!(std::is_integral_v<Arg> && sizeof(Arg) <= sizeof(void*)))
 					{
 						return {EFormatStringSanStatus::CNeedsTCHARArg, CurArgPos};
@@ -257,11 +261,11 @@ namespace UE::Core::Private
 					{
 						return TCheckFormatString<TFmtArgTypes<Args...>>::Check(false, CurArgPos + 1, Fmt + 1);
 					}
-				case TEXT('d'):
-				case TEXT('i'):
-				case TEXT('X'):
-				case TEXT('x'):
-				case TEXT('u'):
+				case (CharType)'d':
+				case (CharType)'i':
+				case (CharType)'X':
+				case (CharType)'x':
+				case (CharType)'u':
 					if constexpr (bIsTCharPtr<Arg>)
 					{
 						return {EFormatStringSanStatus::DNeedsIntegerArg, CurArgPos};
@@ -274,7 +278,7 @@ namespace UE::Core::Private
 					{
 						return TCheckFormatString<TFmtArgTypes<Args...>>::Check(false, CurArgPos + 1, Fmt + 1);
 					}
-				case TEXT('z'):
+				case (CharType)'z':
 					if (!CharIsIntegerFormatSpecifier(Fmt[1]))
 					{
 						return {EFormatStringSanStatus::ZNeedsIntegerSpec, CurArgPos};
@@ -287,7 +291,7 @@ namespace UE::Core::Private
 					{
 						return TCheckFormatString<TFmtArgTypes<Args...>>::Check(false, CurArgPos + 1, Fmt + 1);
 					}
-				case TEXT('p'):
+				case (CharType)'p':
 					if constexpr (!std::is_pointer_v<Arg>)
 					{
 						return {EFormatStringSanStatus::PNeedsPointerArg, CurArgPos};
@@ -296,8 +300,8 @@ namespace UE::Core::Private
 					{
 						return TCheckFormatString<TFmtArgTypes<Args...>>::Check(false, CurArgPos + 1, Fmt + 1);
 					}
-				case TEXT('I'):
-					if (!(Fmt[1] == TEXT('6') && Fmt[2] == TEXT('4')))
+				case (CharType)'I':
+					if (!(Fmt[1] == (CharType)'6' && Fmt[2] == (CharType)'4'))
 					{
 						return {EFormatStringSanStatus::I64BadSpec, CurArgPos};
 					}
@@ -314,7 +318,7 @@ namespace UE::Core::Private
 						return TCheckFormatString<TFmtArgTypes<Args...>>::Check(false, CurArgPos + 1, Fmt + 4);
 					}
 
-				case TEXT('l'):
+				case (CharType)'l':
 					if (CharIsIntegerFormatSpecifier(Fmt[1]))
 					{
 						if constexpr (!(bIsIntegralEnum<Arg> || std::is_integral_v<Arg>))
@@ -326,7 +330,7 @@ namespace UE::Core::Private
 							return TCheckFormatString<TFmtArgTypes<Args...>>::Check(false, CurArgPos + 1, Fmt + 2);
 						}
 					}
-					else if (Fmt[1] == TEXT('f'))
+					else if (Fmt[1] == (CharType)'f')
 					{
 						if constexpr (!TIsFloatOrDouble_V<Arg>)
 						{
@@ -337,7 +341,7 @@ namespace UE::Core::Private
 							return TCheckFormatString<TFmtArgTypes<Args...>>::Check(false, CurArgPos + 1, Fmt + 2);
 						}
 					}
-					else if (Fmt[1] != TEXT('l'))
+					else if (Fmt[1] != (CharType)'l')
 					{
 						return {EFormatStringSanStatus::InvalidFormatSpec, CurArgPos};
 					}
@@ -354,7 +358,7 @@ namespace UE::Core::Private
 						return TCheckFormatString<TFmtArgTypes<Args...>>::Check(false, CurArgPos + 1, Fmt + 3);
 					}
 
-				case TEXT('h'):
+				case (CharType)'h':
 					if (CharIsIntegerFormatSpecifier(Fmt[1]))
 					{
 						if constexpr (!(bIsIntegralEnum<Arg> || std::is_integral_v<Arg>))
@@ -363,7 +367,7 @@ namespace UE::Core::Private
 						}
 						return TCheckFormatString<TFmtArgTypes<Args...>>::Check(false, CurArgPos + 1, Fmt + 2);
 					}
-					else if (Fmt[1] != TEXT('h'))
+					else if (Fmt[1] != (CharType)'h')
 					{
 						return {EFormatStringSanStatus::InvalidFormatSpec, CurArgPos};
 					}
@@ -380,9 +384,9 @@ namespace UE::Core::Private
 						return TCheckFormatString<TFmtArgTypes<Args...>>::Check(false, CurArgPos + 1, Fmt + 3);
 					}
 
-				case TEXT('f'):
-				case TEXT('e'):
-				case TEXT('g'):
+				case (CharType)'f':
+				case (CharType)'e':
+				case (CharType)'g':
 					if constexpr (!TIsFloatOrDouble_V<Arg>)
 					{
 						return {EFormatStringSanStatus::FNeedsFloatOrDoubleArg, CurArgPos};
@@ -392,7 +396,7 @@ namespace UE::Core::Private
 						return TCheckFormatString<TFmtArgTypes<Args...>>::Check(false, CurArgPos + 1, Fmt + 1);
 					}
 
-				case TEXT('S'):
+				case (CharType)'S':
 					using OtherCharType = std::conditional_t<std::is_same_v<TCHAR, ANSICHAR>, WIDECHAR, ANSICHAR>;
 					if constexpr (!bIsPointerTo<Arg, OtherCharType>)
 					{
@@ -410,7 +414,7 @@ namespace UE::Core::Private
 						return TCheckFormatString<TFmtArgTypes<Args...>>::Check(false, CurArgPos + 1, Fmt + 1);
 					}
 
-				case TEXT(' '):
+				case (CharType)' ':
 					return {EFormatStringSanStatus::IncompleteFormatSpecifierOrUnescapedPercent, CurArgPos};
 
 				default:
@@ -422,7 +426,9 @@ namespace UE::Core::Private
 		template <>
 		struct TCheckFormatString<TFmtArgTypes<>>
 		{
-			static constexpr FResult Check(bool bInsideFormatSpec, int CurArgPos, const TCHAR* Fmt)
+			using CharType = TCHAR;
+
+			static constexpr FResult Check(bool bInsideFormatSpec, int CurArgPos, const CharType* Fmt)
 			{
 				if (bInsideFormatSpec)
 				{
@@ -431,13 +437,13 @@ namespace UE::Core::Private
 
 				while (*Fmt)
 				{
-					if (Fmt[0] != TEXT('%'))
+					if (Fmt[0] != (CharType)'%')
 					{
 						++Fmt;
 						continue;
 					}
 
-					if (Fmt[1] == TEXT('%'))
+					if (Fmt[1] == (CharType)'%')
 					{
 						Fmt += 2;
 						continue;
