Fixed character comparisons with zero, which won't work with UTF8CHAR as an enum class.

diff --git a/Engine/Source/Runtime/Core/Public/String/FormatStringSan.h b/Engine/Source/Runtime/Core/Public/String/FormatStringSan.h
--- a/Engine/Source/Runtime/Core/Public/String/FormatStringSan.h
+++ b/Engine/Source/Runtime/Core/Public/String/FormatStringSan.h
@@ -147,14 +147,14 @@ namespace UE::Core::Private
 
 			static constexpr FResult Check(bool bInsideFormatSpec, int CurArgPos, const CharType* Fmt)
 			{
-				if (!Fmt[0] && !bInsideFormatSpec)
+				if (Fmt[0] == (CharType)'\0' && !bInsideFormatSpec)
 				{
 					return {EFormatStringSanStatus::Ok, 0};
 				}
 
 				if (!bInsideFormatSpec)
 				{
-					while (*Fmt && *Fmt != (CharType)'%')
+					while (*Fmt != (CharType)'\0' && *Fmt != (CharType)'%')
 					{
 						++Fmt;
 					}
@@ -164,9 +164,9 @@ namespace UE::Core::Private
 					}
 				}
 
-				if (!*Fmt)
+				if (*Fmt == (CharType)'\0')
 				{
-					if (Fmt[-1] == (CharType)'%' || bInsideFormatSpec) // nb: P[-1] is safe here because zero-length fmt strings are addressed above
+					if (Fmt[-1] == (CharType)'%' || bInsideFormatSpec) // nb: Fmt[-1] is safe here because zero-length fmt strings are addressed above
 					{
 						return {EFormatStringSanStatus::IncompleteFormatSpecifierOrUnescapedPercent, CurArgPos};
 					}
@@ -449,7 +449,7 @@ namespace UE::Core::Private
 						continue;
 					}
 
-					if (Fmt[1] == 0)
+					if (Fmt[1] == (CharType)'\0')
 					{
 						return {EFormatStringSanStatus::IncompleteFormatSpecifierOrUnescapedPercent, CurArgPos};
 					}
