Removed unnecessary variadic typelist wrapper from FormatStringSan.

diff --git a/Engine/Source/Runtime/Core/Public/String/FormatStringSan.h b/Engine/Source/Runtime/Core/Public/String/FormatStringSan.h
--- a/Engine/Source/Runtime/Core/Public/String/FormatStringSan.h
+++ b/Engine/Source/Runtime/Core/Public/String/FormatStringSan.h
@@ -12,14 +12,13 @@
 #define UE_CHECK_FORMAT_STRING(Fmt, ...) \
 	do { \
 		namespace UCFS = ::UE::Core::Private::FormatStringSan; \
-		using UCFS_FArgs = decltype(UCFS::GetFmtArgTypes(__VA_ARGS__)); \
-		using UCFS_FChecker = UCFS::TCheckFormatString<UCFS_FArgs>; \
+		using UCFS_FChecker = decltype(UCFS::GetFmtArgCheckerType(__VA_ARGS__)); \
 		constexpr UCFS::FResult UCFS_Result = UCFS_FChecker::Check(false, 0, Fmt); \
 		(UCFS::AssertFormatStatus<UCFS_Result.Status, UCFS::TAtArgPos<UCFS_Result.ArgPos>>()); \
 	} while(false);
 
 #define UE_CHECK_FORMAT_STRING_ERR(Err, Fmt, ...) \
-	(::UE::Core::Private::FormatStringSan::TCheckFormatString<decltype(::UE::Core::Private::FormatStringSan::GetFmtArgTypes(__VA_ARGS__))>::Check(false, 0, Fmt).Status == Err)
+	(decltype(::UE::Core::Private::FormatStringSan::GetFmtArgCheckerType(__VA_ARGS__))::Check(false, 0, Fmt).Status == Err)
 
 
 #if UE_VALIDATE_FORMAT_STRINGS
@@ -28,8 +27,6 @@
 	#define UE_VALIDATE_FORMAT_STRING(Format, ...)
 #endif
 
-template <typename T> class TEnumAsByte;
-
 /// implementation
 namespace UE::Core::Private
 {
@@ -86,25 +83,14 @@ namespace UE::Core::Private
 		template <typename T>
 		inline constexpr bool bIsIntegralEnum = std::is_enum_v<T> || TIsTEnumAsByte_V<T>;
 
-		template <typename... Ts>
-		struct TFmtArgTypes {};
+		template <typename...>
+		struct TCheckFormatString;
 
 		template <typename... Ts>
-		constexpr TFmtArgTypes<std::decay_t<Ts>...> GetFmtArgTypes(Ts...)
-		{
-			return {};
-		}
-
-		template <typename>
-		struct TCheckFormatString
-		{
-			using CharType = TCHAR;
-
-			static constexpr FResult Check(bool, int, const CharType*);
-		};
+		constexpr TCheckFormatString<std::decay_t<Ts>...> GetFmtArgCheckerType(Ts...);
 
 		template <typename Arg, typename... Args>
-		struct TCheckFormatString<TFmtArgTypes<Arg, Args...>>
+		struct TCheckFormatString<Arg, Args...>
 		{
 			using CharType = TCHAR;
 
@@ -141,7 +127,7 @@ namespace UE::Core::Private
 				}
 				else
 				{
-					return TCheckFormatString<TFmtArgTypes<Args...>>::Check(true, CurArgPos + 1, Fmt + 1);
+					return TCheckFormatString<Args...>::Check(true, CurArgPos + 1, Fmt + 1);
 				}
 			};
 
@@ -223,14 +209,14 @@ namespace UE::Core::Private
 					}
 					else
 					{
-						return TCheckFormatString<TFmtArgTypes<Args...>>::Check(false, CurArgPos + 1, Fmt + 2);
+						return TCheckFormatString<Args...>::Check(false, CurArgPos + 1, Fmt + 2);
 					}
 				}
 
 				switch (Fmt[0])
 				{
 				case (CharType)'%':
-					return TCheckFormatString<TFmtArgTypes<Arg, Args...>>::Check(false, CurArgPos, Fmt + 1);
+					return TCheckFormatString<Arg, Args...>::Check(false, CurArgPos, Fmt + 1);
 				case (CharType)'s':
 					if constexpr (bIsPointerTo<Arg, UTF8CHAR>)
 					{
@@ -250,7 +236,7 @@ namespace UE::Core::Private
 					}
 					else
 					{
-						return TCheckFormatString<TFmtArgTypes<Args...>>::Check(false, CurArgPos + 1, Fmt + 1);
+						return TCheckFormatString<Args...>::Check(false, CurArgPos + 1, Fmt + 1);
 					}
 				case (CharType)'c':
 					if constexpr (!(std::is_integral_v<Arg> && sizeof(Arg) <= sizeof(void*)))
@@ -259,7 +245,7 @@ namespace UE::Core::Private
 					}
 					else
 					{
-						return TCheckFormatString<TFmtArgTypes<Args...>>::Check(false, CurArgPos + 1, Fmt + 1);
+						return TCheckFormatString<Args...>::Check(false, CurArgPos + 1, Fmt + 1);
 					}
 				case (CharType)'d':
 				case (CharType)'i':
@@ -276,7 +262,7 @@ namespace UE::Core::Private
 					}
 					else
 					{
-						return TCheckFormatString<TFmtArgTypes<Args...>>::Check(false, CurArgPos + 1, Fmt + 1);
+						return TCheckFormatString<Args...>::Check(false, CurArgPos + 1, Fmt + 1);
 					}
 				case (CharType)'z':
 					if (!CharIsIntegerFormatSpecifier(Fmt[1]))
@@ -289,7 +275,7 @@ namespace UE::Core::Private
 					}
 					else
 					{
-						return TCheckFormatString<TFmtArgTypes<Args...>>::Check(false, CurArgPos + 1, Fmt + 1);
+						return TCheckFormatString<Args...>::Check(false, CurArgPos + 1, Fmt + 1);
 					}
 				case (CharType)'p':
 					if constexpr (!std::is_pointer_v<Arg>)
@@ -298,7 +284,7 @@ namespace UE::Core::Private
 					}
 					else
 					{
-						return TCheckFormatString<TFmtArgTypes<Args...>>::Check(false, CurArgPos + 1, Fmt + 1);
+						return TCheckFormatString<Args...>::Check(false, CurArgPos + 1, Fmt + 1);
 					}
 				case (CharType)'I':
 					if (!(Fmt[1] == (CharType)'6' && Fmt[2] == (CharType)'4'))
@@ -315,7 +301,7 @@ namespace UE::Core::Private
 					}
 					else
 					{
-						return TCheckFormatString<TFmtArgTypes<Args...>>::Check(false, CurArgPos + 1, Fmt + 4);
+						return TCheckFormatString<Args...>::Check(false, CurArgPos + 1, Fmt + 4);
 					}
 
 				case (CharType)'l':
@@ -327,7 +313,7 @@ namespace UE::Core::Private
 						}
 						else
 						{
-							return TCheckFormatString<TFmtArgTypes<Args...>>::Check(false, CurArgPos + 1, Fmt + 2);
+							return TCheckFormatString<Args...>::Check(false, CurArgPos + 1, Fmt + 2);
 						}
 					}
 					else if (Fmt[1] == (CharType)'f')
@@ -338,7 +324,7 @@ namespace UE::Core::Private
 						}
 						else
 						{
-							return TCheckFormatString<TFmtArgTypes<Args...>>::Check(false, CurArgPos + 1, Fmt + 2);
+							return TCheckFormatString<Args...>::Check(false, CurArgPos + 1, Fmt + 2);
 						}
 					}
 					else if (Fmt[1] != (CharType)'l')
@@ -355,7 +341,7 @@ namespace UE::Core::Private
 					}
 					else
 					{
-						return TCheckFormatString<TFmtArgTypes<Args...>>::Check(false, CurArgPos + 1, Fmt + 3);
+						return TCheckFormatString<Args...>::Check(false, CurArgPos + 1, Fmt + 3);
 					}
 
 				case (CharType)'h':
@@ -365,7 +351,7 @@ namespace UE::Core::Private
 						{
 							return {EFormatStringSanStatus::HNeedsIntegerArg, CurArgPos};
 						}
-						return TCheckFormatString<TFmtArgTypes<Args...>>::Check(false, CurArgPos + 1, Fmt + 2);
+						return TCheckFormatString<Args...>::Check(false, CurArgPos + 1, Fmt + 2);
 					}
 					else if (Fmt[1] != (CharType)'h')
 					{
@@ -381,7 +367,7 @@ namespace UE::Core::Private
 					}
 					else
 					{
-						return TCheckFormatString<TFmtArgTypes<Args...>>::Check(false, CurArgPos + 1, Fmt + 3);
+						return TCheckFormatString<Args...>::Check(false, CurArgPos + 1, Fmt + 3);
 					}
 
 				case (CharType)'f':
@@ -393,7 +379,7 @@ namespace UE::Core::Private
 					}
 					else
 					{
-						return TCheckFormatString<TFmtArgTypes<Args...>>::Check(false, CurArgPos + 1, Fmt + 1);
+						return TCheckFormatString<Args...>::Check(false, CurArgPos + 1, Fmt + 1);
 					}
 
 				case (CharType)'S':
@@ -411,7 +397,7 @@ namespace UE::Core::Private
 					}
 					else
 					{
-						return TCheckFormatString<TFmtArgTypes<Args...>>::Check(false, CurArgPos + 1, Fmt + 1);
+						return TCheckFormatString<Args...>::Check(false, CurArgPos + 1, Fmt + 1);
 					}
 
 				case (CharType)' ':
@@ -424,7 +410,7 @@ namespace UE::Core::Private
 		};
 
 		template <>
-		struct TCheckFormatString<TFmtArgTypes<>>
+		struct TCheckFormatString<>
 		{
 			using CharType = TCHAR;
 
@@ -473,7 +459,7 @@ namespace UE::Core::Private
 			{
 				if constexpr (std::is_convertible_v<const CharType*, const TCHAR*>)
 				{
-					using Checker = TCheckFormatString<TFmtArgTypes<std::decay_t<ArgTypes>...>>;
+					using Checker = TCheckFormatString<std::decay_t<ArgTypes>...>;
 					FResult Result = Checker::Check(false, 0, Fmt);
 
 					switch (Result.Status)
