Moved all string format handling together.

diff --git a/Engine/Source/Runtime/Core/Public/String/FormatStringSan.h b/Engine/Source/Runtime/Core/Public/String/FormatStringSan.h
--- a/Engine/Source/Runtime/Core/Public/String/FormatStringSan.h
+++ b/Engine/Source/Runtime/Core/Public/String/FormatStringSan.h
@@ -195,7 +195,26 @@ namespace UE::Core::Private
 
 				if (Fmt[0] == (CharType)'l' && Fmt[1] == (CharType)'s')
 				{
-					++Fmt;
+					if constexpr (bIsPointerTo<Arg, UTF8CHAR>)
+					{
+						return {EFormatStringSanStatus::UTF8NeedsStringCast, CurArgPos};
+					}
+					else if constexpr (std::is_same_v<Arg, FString>)
+					{
+						return {EFormatStringSanStatus::SNeedsDereferencedFString, CurArgPos};
+					}
+					else if constexpr (std::is_same_v<Arg, TCHAR>)
+					{
+						return {EFormatStringSanStatus::SNeedsTCharPtrArgButGotTCHAR, CurArgPos};
+					}
+					else if constexpr (!bIsTCharPtr<Arg>)
+					{
+						return {EFormatStringSanStatus::SNeedsTCHARPtrArg, CurArgPos};
+					}
+					else
+					{
+						return TCheckFormatString<Args...>::Check(false, CurArgPos + 1, Fmt + 1);
+					}
 				}
 				else if (Fmt[0] == (CharType)'h' && Fmt[1] == (CharType)'s')
 				{
@@ -212,12 +231,8 @@ namespace UE::Core::Private
 						return TCheckFormatString<Args...>::Check(false, CurArgPos + 1, Fmt + 2);
 					}
 				}
-
-				switch (Fmt[0])
+				else if (Fmt[0] == (CharType)'s')
 				{
-				case (CharType)'%':
-					return TCheckFormatString<Arg, Args...>::Check(false, CurArgPos, Fmt + 1);
-				case (CharType)'s':
 					if constexpr (bIsPointerTo<Arg, UTF8CHAR>)
 					{
 						return {EFormatStringSanStatus::UTF8NeedsStringCast, CurArgPos};
@@ -238,6 +253,31 @@ namespace UE::Core::Private
 					{
 						return TCheckFormatString<Args...>::Check(false, CurArgPos + 1, Fmt + 1);
 					}
+				}
+				else if (Fmt[0] == (CharType)'S')
+				{
+					using OtherCharType = std::conditional_t<std::is_same_v<TCHAR, ANSICHAR>, WIDECHAR, ANSICHAR>;
+					if constexpr (!bIsPointerTo<Arg, OtherCharType>)
+					{
+						if constexpr(std::is_same_v<OtherCharType, ANSICHAR>)
+						{
+							return {EFormatStringSanStatus::CapitalSNeedsCharPtr, CurArgPos};
+						}
+						else if constexpr(std::is_same_v<OtherCharType, WIDECHAR>)
+						{
+							return {EFormatStringSanStatus::CapitalSNeedsTCHARPtr, CurArgPos};
+						}
+					}
+					else
+					{
+						return TCheckFormatString<Args...>::Check(false, CurArgPos + 1, Fmt + 1);
+					}
+				}
+
+				switch (Fmt[0])
+				{
+				case (CharType)'%':
+					return TCheckFormatString<Arg, Args...>::Check(false, CurArgPos, Fmt + 1);
 				case (CharType)'c':
 					if constexpr (!(std::is_integral_v<Arg> && sizeof(Arg) <= sizeof(void*)))
 					{
@@ -382,24 +422,6 @@ namespace UE::Core::Private
 						return TCheckFormatString<Args...>::Check(false, CurArgPos + 1, Fmt + 1);
 					}
 
-				case (CharType)'S':
-					using OtherCharType = std::conditional_t<std::is_same_v<TCHAR, ANSICHAR>, WIDECHAR, ANSICHAR>;
-					if constexpr (!bIsPointerTo<Arg, OtherCharType>)
-					{
-						if constexpr(std::is_same_v<OtherCharType, ANSICHAR>)
-						{
-							return {EFormatStringSanStatus::CapitalSNeedsCharPtr, CurArgPos};
-						}
-						else if constexpr(std::is_same_v<OtherCharType, WIDECHAR>)
-						{
-							return {EFormatStringSanStatus::CapitalSNeedsTCHARPtr, CurArgPos};
-						}
-					}
-					else
-					{
-						return TCheckFormatString<Args...>::Check(false, CurArgPos + 1, Fmt + 1);
-					}
-
 				case (CharType)' ':
 					return {EFormatStringSanStatus::IncompleteFormatSpecifierOrUnescapedPercent, CurArgPos};
 
