[AutoRTFM] Add explicit constructors for TTaskArray::FEntry

Some compilers don't like constructing these with a Task.

diff --git a/Engine/Source/Runtime/Core/Private/AutoRTFM/TaskArray.h b/Engine/Source/Runtime/Core/Private/AutoRTFM/TaskArray.h
--- a/Engine/Source/Runtime/Core/Private/AutoRTFM/TaskArray.h
+++ b/Engine/Source/Runtime/Core/Private/AutoRTFM/TaskArray.h
@@ -32,6 +32,11 @@ class TTaskArray
 		// The previous sequential task.
 		FEntry* Prev = nullptr;
 
+		// Constructor using copied Task
+		FEntry(const TaskType& Task) : Task{Task} {}
+		// Constructor using moved Task
+		FEntry(TaskType&& Task) : Task{std::move(Task)} {}
+
 		// Unlinks this entry from the doubly-linked list, updating the provided
 		// head and tail pointers as necessary.
 		void Unlink(FEntry*& ListHead, FEntry*& ListTail)
