[AutoRTFM] Fix lit compilation errors when including AutoRTFMTask.h

Remove the include into Core.
Replace use of Forward<> with std::forward<>.
Replace use of MoveTemp<> with std::move<>.

diff --git a/Engine/Source/Runtime/Core/Public/AutoRTFM/AutoRTFMTask.h b/Engine/Source/Runtime/Core/Public/AutoRTFM/AutoRTFMTask.h
--- a/Engine/Source/Runtime/Core/Public/AutoRTFM/AutoRTFMTask.h
+++ b/Engine/Source/Runtime/Core/Public/AutoRTFM/AutoRTFMTask.h
@@ -2,7 +2,6 @@
 
 #pragma once
 
-#include "Templates/UnrealTemplate.h"
 #include <cstddef>
 #include <type_traits>
 
@@ -71,7 +70,7 @@ private:
 		// Calls the function held by the payload.
 		static ReturnType Call(FPayload& Payload, ParameterTypes... Arguments)
 		{
-			return (*FunctionFrom(Payload))(Forward<ParameterTypes>(Arguments)...);
+			return (*FunctionFrom(Payload))(std::forward<ParameterTypes>(Arguments)...);
 		}
 
 		// Copies the payload.
@@ -92,7 +91,7 @@ private:
 		{
 			if constexpr (bCanUseInline)
 			{
-				new(Dst.Inline) T(MoveTemp(*FunctionFrom(Src)));
+				new(Dst.Inline) T(std::move(*FunctionFrom(Src)));
 			}
 			else
 			{
@@ -151,12 +150,14 @@ public:
 	// Constructor from functor.
 	template
 	<
-		typename FunctorType
-		UE_REQUIRES(
+		typename FunctorType,
+		typename = std::enable_if_t
+		<
 			// Use the explicit TTask copy / move constructors instead of this constructor
 			!std::is_same_v<std::decay_t<FunctorType>, TTask> &&
 			// The templated argument must be a functor with the correct signature
-			std::is_invocable_r_v<ReturnType, std::decay_t<FunctorType>, ParameterTypes...>)
+			std::is_invocable_r_v<ReturnType, std::decay_t<FunctorType>, ParameterTypes...>
+		>
 	>
 	TTask(FunctorType&& Other)
 	{
@@ -165,11 +166,11 @@ public:
 		PayloadMethods = &Traits::PayloadMethods;
 		if (Traits::bCanUseInline)
 		{
-			new(Payload.Inline) Decayed(Forward<FunctorType>(Other));
+			new(Payload.Inline) Decayed(std::forward<FunctorType>(Other));
 		}
 		else
 		{
-			Payload.External = new Decayed(Forward<FunctorType>(Other));
+			Payload.External = new Decayed(std::forward<FunctorType>(Other));
 		}
 	}
 
@@ -208,7 +209,7 @@ public:
 	template<typename... ArgumentTypes>
 	ReturnType operator()(ArgumentTypes&&... Arguments) const
 	{
-		return PayloadMethods->Call(Payload, Forward<ArgumentTypes>(Arguments)...);
+		return PayloadMethods->Call(Payload, std::forward<ArgumentTypes>(Arguments)...);
 	}
 
 	// Reset the task to an unset state.
