Fixed skipping over %ls format specifier.

diff --git a/Engine/Source/Runtime/Core/Public/String/FormatStringSan.h b/Engine/Source/Runtime/Core/Public/String/FormatStringSan.h
--- a/Engine/Source/Runtime/Core/Public/String/FormatStringSan.h
+++ b/Engine/Source/Runtime/Core/Public/String/FormatStringSan.h
@@ -213,7 +213,7 @@ namespace UE::Core::Private
 					}
 					else
 					{
-						return TCheckFormatString<Args...>::Check(false, CurArgPos + 1, Fmt + 1);
+						return TCheckFormatString<Args...>::Check(false, CurArgPos + 1, Fmt + 2);
 					}
 				}
 				else if (Fmt[0] == (CharType)'h' && Fmt[1] == (CharType)'s')
