Fixed width check for %c format specifier.

https://en.cppreference.com/w/cpp/language/variadic_arguments

> bool, char, short, and unscoped enumerations are converted to int or wider integer types as in integral promotion.

diff --git a/Engine/Source/Runtime/Core/Public/String/FormatStringSan.h b/Engine/Source/Runtime/Core/Public/String/FormatStringSan.h
--- a/Engine/Source/Runtime/Core/Public/String/FormatStringSan.h
+++ b/Engine/Source/Runtime/Core/Public/String/FormatStringSan.h
@@ -279,7 +279,7 @@ namespace UE::Core::Private
 				case (CharType)'%':
 					return TCheckFormatString<Arg, Args...>::Check(false, CurArgPos, Fmt + 1);
 				case (CharType)'c':
-					if constexpr (!(std::is_integral_v<Arg> && sizeof(Arg) <= sizeof(void*)))
+					if constexpr (!(std::is_integral_v<Arg> && sizeof(Arg) <= sizeof(int)))
 					{
 						return {EFormatStringSanStatus::CNeedsTCHARArg, CurArgPos};
 					}
