Fixing recently failing AutomationSpec test where adding an expected error with a negative occurrence was getting picked up when fetching for all expected errors. Added `GetExpectedMessages` method with the ability to fetch either all expected messages or only those that won't be suppressed.

#rnx

diff --git a/Engine/Plugins/Tests/TestFramework/Source/TestFramework/Private/AutomationExpectedMessage.spec.cpp b/Engine/Plugins/Tests/TestFramework/Source/TestFramework/Private/AutomationExpectedMessage.spec.cpp
--- a/Engine/Plugins/Tests/TestFramework/Source/TestFramework/Private/AutomationExpectedMessage.spec.cpp
+++ b/Engine/Plugins/Tests/TestFramework/Source/TestFramework/Private/AutomationExpectedMessage.spec.cpp
@@ -14,13 +14,35 @@ void FAutomationExpectedErrorTest::Define()
 			// Suppress error logged when adding entry with invalid occurrence count
 			AddExpectedError(TEXT("number of expected occurrences must be >= 0"), EAutomationExpectedErrorFlags::Contains, 1);
 
-			AddExpectedError(TEXT("The two values are not equal"), EAutomationExpectedErrorFlags::Contains,  -1);
+			AddExpectedError(TEXT("The two values are not equal"), EAutomationExpectedErrorFlags::Contains, -1);
 
 			TArray<FAutomationExpectedMessage> Errors;
 			GetExpectedMessages(Errors, ELogVerbosity::Warning);
 
 			// Only the first expected error should exist in the list.
 			TestEqual("Expected Errors Count", Errors.Num(), 1);
+
+			// Add the expected errors to ensure all test conditions pass
+			AddError(TEXT("number of expected occurrences must be >= 0"));
+			AddError(TEXT("The two values are not equal"));
+		});
+
+		It("can fetch an error with a number of occurrences less than zero", [this]()
+		{
+			AddExpectedError(TEXT("number of expected occurrences must be >= 0"), EAutomationExpectedErrorFlags::Contains, 1);
+
+			AddExpectedError(TEXT("The two values are not equal"), EAutomationExpectedErrorFlags::Contains, -1);
+
+			const bool bIncludeSuppressed = true;
+			TArray<FAutomationExpectedMessage> Errors;
+			GetExpectedMessages(Errors, bIncludeSuppressed, ELogVerbosity::Warning);
+
+			// The suppressed error should show up in the list.
+			TestEqual("Expected Errors Count", Errors.Num(), 2);
+
+			// Add the expected errors to ensure all test conditions pass
+			AddError(TEXT("number of expected occurrences must be >= 0"));
+			AddError(TEXT("The two values are not equal"));
 		});
 
 		It("will add an error with a number of occurrences equal to zero", [this]()
diff --git a/Engine/Source/Runtime/Core/Private/Misc/AutomationTest.cpp b/Engine/Source/Runtime/Core/Private/Misc/AutomationTest.cpp
--- a/Engine/Source/Runtime/Core/Private/Misc/AutomationTest.cpp
+++ b/Engine/Source/Runtime/Core/Private/Misc/AutomationTest.cpp
@@ -1667,18 +1667,21 @@ void FAutomationTestBase::GetExpectedMessages(
 	TArray<FAutomationExpectedMessage>& OutInfo,
 	ELogVerbosity::Type Verbosity) const
 {
-	if (Verbosity == ELogVerbosity::All)
-	{
-		OutInfo = ExpectedMessages.Array();
-	}
-	else
-	{
-		OutInfo.Reserve(ExpectedMessages.Num());
-		Algo::CopyIf(ExpectedMessages, OutInfo, [Verbosity](const FAutomationExpectedMessage& Message)
+	// Do not include any suppressed messages
+	GetExpectedMessages(OutInfo, false, Verbosity);
+}
+
+void FAutomationTestBase::GetExpectedMessages(
+	TArray<FAutomationExpectedMessage>& OutInfo,
+	bool IncludeSuppressed,
+	ELogVerbosity::Type Verbosity) const
+{
+	OutInfo.Reserve(ExpectedMessages.Num());
+	Algo::CopyIf(ExpectedMessages, OutInfo, [IncludeSuppressed, Verbosity](const FAutomationExpectedMessage& Message)
 		{
-			return FAutomationTestBase::LogCategoryMatchesSeverityInclusive(Message.Verbosity, Verbosity);
+			const bool bIsIncluded = (Message.ExpectedNumberOfOccurrences < 0) ? IncludeSuppressed : true;
+			return bIsIncluded && FAutomationTestBase::LogCategoryMatchesSeverityInclusive(Message.Verbosity, Verbosity);
 		});
-	}
 	OutInfo.Sort();
 }
 
diff --git a/Engine/Source/Runtime/Core/Public/Misc/AutomationTest.h b/Engine/Source/Runtime/Core/Public/Misc/AutomationTest.h
--- a/Engine/Source/Runtime/Core/Public/Misc/AutomationTest.h
+++ b/Engine/Source/Runtime/Core/Public/Misc/AutomationTest.h
@@ -1786,11 +1786,22 @@ public:
 	/**
 	* Populate the provided expected log messages object with the expected messages contained within the test. Not particularly efficient,
 	* but providing direct access to the test's private execution messages list could result in errors.
-	* @param Verbosity - Optionally filter the returned messages by verbosity. This is inclusive, so Warning will return Warnings, Errors, etc.
 	* @param OutInfo - Array of Expected Messages to be populated with the same data contained within this test's expected messages list
+	* @param Verbosity - Optionally filter the returned messages by verbosity. This is inclusive, so Warning will return Warnings, Errors, etc.
+	* 
+	* @note This will not include any messages that are supposed to be suppressed
 	*/
 	CORE_API void GetExpectedMessages(TArray<FAutomationExpectedMessage>& OutInfo, ELogVerbosity::Type Verbosity = ELogVerbosity::All) const;
 
+	/**
+	* Populate the provided expected log messages object with the expected messages contained within the test. Not particularly efficient,
+	* but providing direct access to the test's private execution messages list could result in errors.
+	* @param OutInfo - Array of Expected Messages to be populated with the same data contained within this test's expected messages list
+	* @param IncludeSuppressed - Include the expected errors that have been marked as being suppressed.
+	* @param Verbosity - Optionally filter the returned messages by verbosity. This is inclusive, so Warning will return Warnings, Errors, etc.
+	*/
+	CORE_API void GetExpectedMessages(TArray<FAutomationExpectedMessage>& OutInfo, bool IncludeSuppressed, ELogVerbosity::Type Verbosity = ELogVerbosity::All) const;
+
 	/**
 	* Adds a regex pattern to an internal list that this test will expect to encounter in error or warning logs during its execution. If an expected pattern
 	* is not encountered, it will cause this test to fail.
