Task optimizations for tick system:
Add ProcessUntilTasksComplete to the task graph to better handle the tick case of 95% game thread tasks
Reduce general task scheduling overhead
Fix tick batching to correctly skip prerequisites and not batch async ticks

diff --git a/Engine/Source/Runtime/Core/Private/Async/TaskGraph.cpp b/Engine/Source/Runtime/Core/Private/Async/TaskGraph.cpp
--- a/Engine/Source/Runtime/Core/Private/Async/TaskGraph.cpp
+++ b/Engine/Source/Runtime/Core/Private/Async/TaskGraph.cpp
@@ -1442,12 +1442,13 @@ private:
 		TRACE_CPUPROFILER_EVENT_SCOPE(WaitUntilTasksComplete);
 
 		// try retraction first and only invoke named thread waiting as a last resort
+		const UE::FTimeout NeverTimeout = UE::FTimeout::Never();
 		bool bAllTasksCompleted = true;
 		for (const FGraphEventRef& Task : Tasks)
 		{
 			if (Task.IsValid())
 			{
-				Task->TryRetractAndExecute(UE::FTimeout::Never());
+				Task->TryRetractAndExecute(NeverTimeout);
 
 				bAllTasksCompleted &= Task->IsCompleted();
 			}
@@ -1476,25 +1477,7 @@ private:
 
 		if (CurrentThreadIfKnown != ENamedThreads::AnyThread && CurrentThreadIfKnown < NumNamedThreads && !IsThreadProcessingTasks(CurrentThread))
 		{
-			if (Tasks.Num() < 8) // don't bother to check for completion if there are lots of prereqs...too expensive to check
-			{
-				bool bAnyPending = false;
-				for (int32 Index = 0; Index < Tasks.Num(); Index++)
-				{
-					FGraphEvent* Task = Tasks[Index].GetReference();
-					if (Task && !Task->IsComplete())
-					{
-						bAnyPending = true;
-						break;
-					}
-				}
-				if (!bAnyPending)
-				{
-					return;
-				}
-			}
-
-			// named thread process tasks while we wait
+			// we wouldn't get past the loop above if all tasks were complete so process named tasks while we wait
 			TGraphTask<FReturnGraphTask>::CreateTask(&Tasks, CurrentThread).ConstructAndDispatchWhenReady(CurrentThread);
 			ProcessThreadUntilRequestReturn(CurrentThread);
 			return;
@@ -1524,6 +1507,50 @@ private:
 		TriggerEventWhenTasksComplete(Event.Get(), Tasks, CurrentThreadIfKnown);
 	}
 
+	void ProcessUntilTasksComplete(const FGraphEventArray& Tasks, ENamedThreads::Type CurrentThreadIfKnown, const TFunction<void()>& ExecuteOnIdle) final override
+	{
+		ENamedThreads::Type CurrentThread = CurrentThreadIfKnown;
+		if (ENamedThreads::GetThreadIndex(CurrentThreadIfKnown) == ENamedThreads::AnyThread)
+		{
+			bool bIsHiPri = !!ENamedThreads::GetTaskPriority(CurrentThreadIfKnown);
+			int32 Priority = ENamedThreads::GetThreadPriorityIndex(CurrentThreadIfKnown);
+			check(!ENamedThreads::GetQueueIndex(CurrentThreadIfKnown));
+			CurrentThreadIfKnown = ENamedThreads::GetThreadIndex(GetCurrentThread());
+			CurrentThread = ENamedThreads::SetPriorities(CurrentThreadIfKnown, Priority, bIsHiPri);
+		}
+
+		// First process until this named thread (could be a local queue) is idle, which should complete most of the tasks
+		ProcessThreadUntilIdle(CurrentThread);
+
+		// See which tasks are still active
+		FGraphEventArray RemainingTasks;
+		for (const FGraphEventRef& Task : Tasks)
+		{
+			if (Task.IsValid())
+			{
+				if (!Task->IsCompleted())
+				{
+					RemainingTasks.Add(Task);
+				}
+			}
+		}
+
+		if (RemainingTasks.Num() == 0)
+		{
+			// They're all done so no need to wait
+			return;
+		}
+
+		// Call the idle function, which can be used to try and resolve deadlocks
+		if (ExecuteOnIdle)
+		{
+			ExecuteOnIdle();
+		}
+
+		// Now wait on this named thread until all the tasks are done
+		WaitUntilTasksComplete(RemainingTasks, CurrentThread);
+	}
+
 	void TriggerEventWhenTasksComplete(FEvent* InEvent, const FGraphEventArray& Tasks, ENamedThreads::Type CurrentThreadIfKnown = ENamedThreads::AnyThread, ENamedThreads::Type TriggerThread = ENamedThreads::AnyHiPriThreadHiPriTask) final override
 	{
 		check(InEvent);
diff --git a/Engine/Source/Runtime/Core/Private/Tasks/TaskPrivate.cpp b/Engine/Source/Runtime/Core/Private/Tasks/TaskPrivate.cpp
--- a/Engine/Source/Runtime/Core/Private/Tasks/TaskPrivate.cpp
+++ b/Engine/Source/Runtime/Core/Private/Tasks/TaskPrivate.cpp
@@ -17,6 +17,10 @@ namespace UE::Tasks
 {
 	namespace Private
 	{
+		// Specific translate functions, defined below
+		ENamedThreads::Type TranslatePriority(EExtendedTaskPriority Priority);
+		ENamedThreads::Type TranslatePriority(ETaskPriority Priority);
+		
 		FExecutableTaskAllocator SmallTaskAllocator;
 		FTaskEventBaseAllocator TaskEventBaseAllocator;
 
@@ -26,25 +30,7 @@ namespace UE::Tasks
 
 			if (IsNamedThreadTask())
 			{
-				ENamedThreads::Type ConversionMap[] =
-				{
-					ENamedThreads::GameThread,
-					(ENamedThreads::Type)(ENamedThreads::GameThread | ENamedThreads::HighTaskPriority),
-					(ENamedThreads::Type)(ENamedThreads::GameThread | ENamedThreads::LocalQueue),
-					(ENamedThreads::Type)(ENamedThreads::GameThread | ENamedThreads::HighTaskPriority | ENamedThreads::LocalQueue),
-
-					ENamedThreads::GetRenderThread(),
-					(ENamedThreads::Type)(ENamedThreads::GetRenderThread() | ENamedThreads::HighTaskPriority),
-					(ENamedThreads::Type)(ENamedThreads::GetRenderThread() | ENamedThreads::LocalQueue),
-					(ENamedThreads::Type)(ENamedThreads::GetRenderThread() | ENamedThreads::HighTaskPriority | ENamedThreads::LocalQueue),
-
-					ENamedThreads::RHIThread,
-					(ENamedThreads::Type)(ENamedThreads::RHIThread | ENamedThreads::HighTaskPriority),
-					(ENamedThreads::Type)(ENamedThreads::RHIThread | ENamedThreads::LocalQueue),
-					(ENamedThreads::Type)(ENamedThreads::RHIThread | ENamedThreads::HighTaskPriority | ENamedThreads::LocalQueue)
-				};
-
-				FTaskGraphInterface::Get().QueueTask(static_cast<FBaseGraphTask*>(this), true, ConversionMap[(int32)ExtendedPriority - (int32)EExtendedTaskPriority::GameThreadNormalPri]);
+				FTaskGraphInterface::Get().QueueTask(static_cast<FBaseGraphTask*>(this), true, TranslatePriority(ExtendedPriority));
 				return;
 			}
 			
@@ -460,7 +446,7 @@ namespace UE::Tasks
 			ENamedThreads::Type ThreadIndex = ENamedThreads::GetThreadIndex(ThreadType);
 			if (ThreadIndex != ENamedThreads::AnyThread)
 			{
-				check(ThreadIndex == ENamedThreads::GameThread || ThreadIndex == ENamedThreads::GetRenderThread() || ThreadIndex == ENamedThreads::RHIThread);
+				check(ThreadIndex == ENamedThreads::GameThread || ThreadIndex == ENamedThreads::ActualRenderingThread || ThreadIndex == ENamedThreads::RHIThread);
 				EExtendedTaskPriority ConversionMap[] =
 				{
 					EExtendedTaskPriority::RHIThreadNormalPri,
@@ -506,27 +492,40 @@ namespace UE::Tasks
 		// task priority translation from the new API to the old API
 		ENamedThreads::Type TranslatePriority(EExtendedTaskPriority Priority)
 		{
-			checkf(Priority >= EExtendedTaskPriority::GameThreadNormalPri && Priority < EExtendedTaskPriority::Count, TEXT("only named threads can call this method: %d"), Priority);
-
-			int32 ConversionMap[] =
+			// Switch is faster than table because of render thread
+			switch (Priority)
 			{
-					ENamedThreads::GameThread, // GameThreadNormalPri
-					ENamedThreads::GameThread | ENamedThreads::HighTaskPriority, // GameThreadHiPri
-					ENamedThreads::GameThread | ENamedThreads::LocalQueue, // GameThreadNormalPriLocalQueue
-					ENamedThreads::GameThread | ENamedThreads::HighTaskPriority | ENamedThreads::LocalQueue, // GameThreadHiPriLocalQueue
-
-					ENamedThreads::GetRenderThread(), // RenderThreadNormalPri
-					ENamedThreads::GetRenderThread() | ENamedThreads::HighTaskPriority, // RenderThreadHiPri
-					ENamedThreads::GetRenderThread() | ENamedThreads::LocalQueue, // RenderThreadNormalPriLocalQueue
-					ENamedThreads::GetRenderThread() | ENamedThreads::HighTaskPriority | ENamedThreads::LocalQueue, // RenderThreadHiPriLocalQueue
-
-					ENamedThreads::RHIThread, // RHIThreadNormalPri
-					ENamedThreads::RHIThread | ENamedThreads::HighTaskPriority, // RHIThreadHiPri
-					ENamedThreads::RHIThread | ENamedThreads::LocalQueue, // RHIThreadNormalPriLocalQueue
-					ENamedThreads::RHIThread | ENamedThreads::HighTaskPriority | ENamedThreads::LocalQueue // RHIThreadHiPriLocalQueue
-			};
-
-			return (ENamedThreads::Type)ConversionMap[(int32)Priority - (int32)EExtendedTaskPriority::GameThreadNormalPri];
+			case EExtendedTaskPriority::GameThreadNormalPri:
+				return ENamedThreads::GameThread;
+			case EExtendedTaskPriority::GameThreadHiPri:
+				return (ENamedThreads::Type)(ENamedThreads::GameThread | ENamedThreads::HighTaskPriority);
+			case EExtendedTaskPriority::GameThreadNormalPriLocalQueue:
+				return (ENamedThreads::Type)(ENamedThreads::GameThread | ENamedThreads::LocalQueue);
+			case EExtendedTaskPriority::GameThreadHiPriLocalQueue:
+				return (ENamedThreads::Type)(ENamedThreads::GameThread | ENamedThreads::HighTaskPriority | ENamedThreads::LocalQueue);
+
+			case EExtendedTaskPriority::RenderThreadNormalPri:
+				return ENamedThreads::GetRenderThread();
+			case EExtendedTaskPriority::RenderThreadHiPri:
+				return (ENamedThreads::Type)(ENamedThreads::GetRenderThread() | ENamedThreads::HighTaskPriority);
+			case EExtendedTaskPriority::RenderThreadNormalPriLocalQueue:
+				return (ENamedThreads::Type)(ENamedThreads::GetRenderThread() | ENamedThreads::LocalQueue);
+			case EExtendedTaskPriority::RenderThreadHiPriLocalQueue:
+				return (ENamedThreads::Type)(ENamedThreads::GetRenderThread() | ENamedThreads::HighTaskPriority | ENamedThreads::LocalQueue);
+
+			case EExtendedTaskPriority::RHIThreadNormalPri:
+				return ENamedThreads::RHIThread;
+			case EExtendedTaskPriority::RHIThreadHiPri:
+				return (ENamedThreads::Type)(ENamedThreads::RHIThread | ENamedThreads::HighTaskPriority);
+			case EExtendedTaskPriority::RHIThreadNormalPriLocalQueue:
+				return (ENamedThreads::Type)(ENamedThreads::RHIThread | ENamedThreads::LocalQueue);
+			case EExtendedTaskPriority::RHIThreadHiPriLocalQueue:
+				return (ENamedThreads::Type)(ENamedThreads::RHIThread | ENamedThreads::HighTaskPriority | ENamedThreads::LocalQueue);
+
+			default:
+				checkf(Priority >= EExtendedTaskPriority::GameThreadNormalPri && Priority < EExtendedTaskPriority::Count, TEXT("only named threads can call this method: %d"), Priority);
+				return ENamedThreads::AnyThread;
+			}
 		}
 
 		ENamedThreads::Type TranslatePriority(ETaskPriority Priority, EExtendedTaskPriority ExtendedPriority)
diff --git a/Engine/Source/Runtime/Core/Public/Async/TaskGraphInterfaces.h b/Engine/Source/Runtime/Core/Public/Async/TaskGraphInterfaces.h
--- a/Engine/Source/Runtime/Core/Public/Async/TaskGraphInterfaces.h
+++ b/Engine/Source/Runtime/Core/Public/Async/TaskGraphInterfaces.h
@@ -371,10 +371,20 @@ public:
 
 	/** 
 	 *	Requests that a named thread, which must be this thread, run until a list of tasks is complete.
+	 *	This function assumes that most of the tasks will be processed on a different thread, but will process if needed.
 	 *	@param	Tasks - tasks to wait for
 	 *	@param	CurrentThread - This thread, must be a named thread
 	**/
 	virtual void WaitUntilTasksComplete(const FGraphEventArray& Tasks, ENamedThreads::Type CurrentThreadIfKnown = ENamedThreads::AnyThread)=0;
+	
+	/**
+	 *	Requests that a named thread actively attempt to process a list of tasks, and calls a function when idle.
+	 *	This function assumes that many of the tasks can be completed on this thread.
+	 *	@param	Tasks - tasks to wait for
+	 *	@param	CurrentThread - This thread, must be a named thread
+	 *	@param	ExecuteOnIdle - If set, call this function the first time this thread has run out of work to do, before continuing to wait
+	**/
+	virtual void ProcessUntilTasksComplete(const FGraphEventArray& Tasks, ENamedThreads::Type CurrentThreadIfKnown, const TFunction<void()>& ExecuteOnIdle = {}) = 0;
 
 	/** 
 	 *	When a set of tasks complete, fire a scoped event
diff --git a/Engine/Source/Runtime/Core/Public/Async/TaskTrace.h b/Engine/Source/Runtime/Core/Public/Async/TaskTrace.h
--- a/Engine/Source/Runtime/Core/Public/Async/TaskTrace.h
+++ b/Engine/Source/Runtime/Core/Public/Async/TaskTrace.h
@@ -26,7 +26,7 @@ namespace ENamedThreads
 
 namespace TaskTrace
 {
-	UE_TRACE_CHANNEL_EXTERN(TaskChannel);
+	CORE_API UE_TRACE_CHANNEL_EXTERN(TaskChannel);
 
 	using FId = uint64;
 
diff --git a/Engine/Source/Runtime/Core/Public/Tasks/TaskPrivate.h b/Engine/Source/Runtime/Core/Public/Tasks/TaskPrivate.h
--- a/Engine/Source/Runtime/Core/Public/Tasks/TaskPrivate.h
+++ b/Engine/Source/Runtime/Core/Public/Tasks/TaskPrivate.h
@@ -389,7 +389,13 @@ public:
 			// In this case the task will be automatically scheduled when all dependencies are completed.
 			bool TryLaunch(uint64 TaskSize)
 			{
-				TaskTrace::Launched(GetTraceId(), LowLevelTask.GetDebugName(), true, TranslatePriority(LowLevelTask.GetPriority(), ExtendedPriority), TaskSize);
+#if UE_TASK_TRACE_ENABLED
+				if (UE_TRACE_CHANNELEXPR_IS_ENABLED(TaskTrace::TaskChannel))
+				{
+					// The params to this are slow to compute
+					TaskTrace::Launched(GetTraceId(), LowLevelTask.GetDebugName(), true, TranslatePriority(LowLevelTask.GetPriority(), ExtendedPriority), TaskSize);
+				}
+#endif
 
 				bool bWakeUpWorker = true;
 				return TryUnlock(bWakeUpWorker);
diff --git a/Engine/Source/Runtime/Engine/Private/TickTaskManager.cpp b/Engine/Source/Runtime/Engine/Private/TickTaskManager.cpp
--- a/Engine/Source/Runtime/Engine/Private/TickTaskManager.cpp
+++ b/Engine/Source/Runtime/Engine/Private/TickTaskManager.cpp
@@ -236,7 +236,7 @@ struct FTickGroupCondition
 		: StartGroup(TickFunction->GetActualTickGroup())
 		, EndGroup(TickFunction->GetActualEndTickGroup())
 		, bHighPriority(TickFunction->bHighPriority)
-		, bIsBatch(TickFunction->bAllowTickBatching)
+		, bIsBatch(TickFunction->bAllowTickBatching && !TickFunction->bRunOnAnyThread)
 	{
 	}
 
@@ -741,18 +741,24 @@ public:
 					BatchInfo = TickBatches[TickBatchesNum].Value.Get();
 					TickBatchesNum++;
 
-					BatchInfo->TickPrerequisites = Prerequisites;
 					check(BatchInfo->TickTask == nullptr);
-
-					// Create the batched task now
-					FGraphEventArray PrerequisiteEvents;
-					for (FTickFunction* Prereq : Prerequisites)
+					
+					FTickContext UseContext = SetupTickContext(TickFunction, TickContext);
+					if (Prerequisites.Num() > 0)
 					{
-						PrerequisiteEvents.Add(Prereq->GetCompletionHandle());
+						BatchInfo->TickPrerequisites = Prerequisites;
+						FGraphEventArray PrerequisiteEvents;
+						for (FTickFunction* Prereq : Prerequisites)
+						{
+							PrerequisiteEvents.Add(Prereq->GetCompletionHandle());
+						}
+						BatchInfo->TickTask = TGraphTask<FBatchTickFunctionTask>::CreateTask(&PrerequisiteEvents, ENamedThreads::GameThread).ConstructAndHold(BatchInfo, &UseContext);
 					}
-
-					FTickContext UseContext = SetupTickContext(TickFunction, TickContext);
-					BatchInfo->TickTask = TGraphTask<FBatchTickFunctionTask>::CreateTask(&PrerequisiteEvents, ENamedThreads::GameThread).ConstructAndHold(BatchInfo, &UseContext);
+					else
+					{
+						BatchInfo->TickTask = TGraphTask<FBatchTickFunctionTask>::CreateTask(nullptr, ENamedThreads::GameThread).ConstructAndHold(BatchInfo, &UseContext);
+					}
+					
 
 					AddTickTaskCompletion(Condition.StartGroup, Condition.EndGroup, BatchInfo->TickTask, Condition.bHighPriority);
 				}
@@ -766,14 +772,20 @@ public:
 		}
 
 		// No batching, create a single task
-		// FGraphEventArray array has some inline members so it is faster to not explicitly reserve space
-		FGraphEventArray PrerequisiteEvents;
-		for (FTickFunction* Prereq : Prerequisites)
+		if (Prerequisites.Num() > 0)
 		{
-			PrerequisiteEvents.Add(Prereq->GetCompletionHandle());
-		}
+			FGraphEventArray PrerequisiteEvents;
+			for (FTickFunction* Prereq : Prerequisites)
+			{
+				PrerequisiteEvents.Add(Prereq->GetCompletionHandle());
+			}
 
-		QueueTickTask(&PrerequisiteEvents, TickFunction, TickContext);
+			QueueTickTask(&PrerequisiteEvents, TickFunction, TickContext);
+		}
+		else
+		{
+			QueueTickTask(nullptr, TickFunction, TickContext);
+		}
 
 		return nullptr;
 	}
@@ -846,7 +858,7 @@ public:
 				if (TickCompletionEvents[Block].Num())
 				{
 					TRACE_CPUPROFILER_EVENT_SCOPE(TickCompletionEvents);
-					FTaskGraphInterface::Get().WaitUntilTasksComplete(TickCompletionEvents[Block], ENamedThreads::GameThread);
+					FTaskGraphInterface::Get().ProcessUntilTasksComplete(TickCompletionEvents[Block], ENamedThreads::GameThread);
 					if (bSingleThreadMode || Block == TG_NewlySpawned || CVarAllowAsyncTickCleanup.GetValueOnGameThread() == 0 || TickCompletionEvents[Block].Num() < 50)
 					{
 						ResetTickGroup(Block);
