CIS fix to declare the extern trace channel properly

diff --git a/Engine/Source/Runtime/Core/Public/Async/TaskTrace.h b/Engine/Source/Runtime/Core/Public/Async/TaskTrace.h
--- a/Engine/Source/Runtime/Core/Public/Async/TaskTrace.h
+++ b/Engine/Source/Runtime/Core/Public/Async/TaskTrace.h
@@ -26,7 +26,7 @@ namespace ENamedThreads
 
 namespace TaskTrace
 {
-	CORE_API UE_TRACE_CHANNEL_EXTERN(TaskChannel);
+	UE_TRACE_CHANNEL_EXTERN(TaskChannel, CORE_API);
 
 	using FId = uint64;
 
