[Insights] Fixed parameters to UE_TRACE_METADATA_SCOPE macro when UE_TRACE_METADATA_ENABLED is 0.

diff --git a/Engine/Source/Runtime/Core/Public/ProfilingDebugging/MetadataTrace.h b/Engine/Source/Runtime/Core/Public/ProfilingDebugging/MetadataTrace.h
--- a/Engine/Source/Runtime/Core/Public/ProfilingDebugging/MetadataTrace.h
+++ b/Engine/Source/Runtime/Core/Public/ProfilingDebugging/MetadataTrace.h
@@ -98,7 +98,7 @@ private:
 #define UE_TRACE_METADATA_EVENT_BEGIN(ScopeName, ...)
 #define UE_TRACE_METADATA_EVENT_FIELD(Type, Name)
 #define UE_TRACE_METADATA_EVENT_END()
-#define UE_TRACE_METADATA_SCOPE(ScopeName)
+#define UE_TRACE_METADATA_SCOPE(ScopeName, ChannelExpr)
 
 #define UE_TRACE_METADATA_CLEAR_SCOPE()
 #define UE_TRACE_METADATA_SAVE_STACK() 0
