Split out some check-independent functions.

diff --git a/Engine/Source/Runtime/Core/Public/String/FormatStringSan.h b/Engine/Source/Runtime/Core/Public/String/FormatStringSan.h
--- a/Engine/Source/Runtime/Core/Public/String/FormatStringSan.h
+++ b/Engine/Source/Runtime/Core/Public/String/FormatStringSan.h
@@ -89,35 +89,38 @@ namespace UE::Core::Private
 		template <typename... Ts>
 		constexpr TCheckFormatString<std::decay_t<Ts>...> GetFmtArgCheckerType(Ts...);
 
-		template <typename Arg, typename... Args>
-		struct TCheckFormatString<Arg, Args...>
+		template <typename CharType>
+		constexpr bool CharIsDigit(CharType Char)
 		{
-			using CharType = TCHAR;
+			return Char >= (CharType)'0' && Char <= (CharType)'9';
+		}
 
-			static constexpr bool CharIsDigit(const CharType Char)
+		template <typename CharType>
+		constexpr const CharType* SkipInteger(const CharType* Fmt)
+		{
+			while (CharIsDigit(*Fmt))
 			{
-				return Char >= (CharType)'0' && Char <= (CharType)'9';
+				++Fmt;
 			}
+			return Fmt;
+		}
 
-			static constexpr const CharType* SkipInteger(const CharType* Fmt)
+		template <typename CharType>
+		constexpr bool CharIsIntegerFormatSpecifier(CharType Char)
+		{
+			switch (Char)
 			{
-				while (CharIsDigit(*Fmt))
-				{
-					++Fmt;
-				}
-				return Fmt;
+			case (CharType)'i': case (CharType)'d': case (CharType)'u': case (CharType)'X': case (CharType)'x':
+				return true;
+			default:
+				return false;
 			}
+		}
 
-			static constexpr bool CharIsIntegerFormatSpecifier(const CharType Char)
-			{
-				switch (Char)
-				{
-				case (CharType)'i': case (CharType)'d': case (CharType)'u': case (CharType)'X': case (CharType)'x':
-					return true;
-				default:
-					return false;
-				}
-			}
+		template <typename Arg, typename... Args>
+		struct TCheckFormatString<Arg, Args...>
+		{
+			using CharType = TCHAR;
 
 			static constexpr FResult HandleDynamicLengthSpecifier(int CurArgPos, const CharType* Fmt)
 			{
@@ -129,7 +132,7 @@ namespace UE::Core::Private
 				{
 					return TCheckFormatString<Args...>::Check(true, CurArgPos + 1, Fmt + 1);
 				}
-			};
+			}
 
 			static constexpr FResult Check(bool bInsideFormatSpec, int CurArgPos, const CharType* Fmt)
 			{
