Fixed trait naming.

diff --git a/Engine/Source/Runtime/Core/Public/String/FormatStringSan.h b/Engine/Source/Runtime/Core/Public/String/FormatStringSan.h
--- a/Engine/Source/Runtime/Core/Public/String/FormatStringSan.h
+++ b/Engine/Source/Runtime/Core/Public/String/FormatStringSan.h
@@ -81,7 +81,7 @@ namespace UE::Core::Private
 		inline constexpr bool TIsFloatOrDouble_V = std::is_same_v<float, T> || std::is_same_v<double, T>;
 
 		template <typename T>
-		inline constexpr bool bIsIntegralEnum = std::is_enum_v<T> || TIsTEnumAsByte_V<T>;
+		inline constexpr bool TIsIntegralEnum_V = std::is_enum_v<T> || TIsTEnumAsByte_V<T>;
 
 		template <typename...>
 		struct TCheckFormatString;
@@ -124,7 +124,7 @@ namespace UE::Core::Private
 
 			static constexpr FResult HandleDynamicLengthSpecifier(int CurArgPos, const CharType* Fmt)
 			{
-				if constexpr (!(std::is_integral_v<Arg> || bIsIntegralEnum<Arg>))
+				if constexpr (!(std::is_integral_v<Arg> || TIsIntegralEnum_V<Arg>))
 				{
 					return {EFormatStringSanStatus::DynamicLengthSpecNeedsIntegerArg, CurArgPos};
 				}
@@ -299,7 +299,7 @@ namespace UE::Core::Private
 					{
 						return {EFormatStringSanStatus::DNeedsIntegerArg, CurArgPos};
 					}
-					else if constexpr (!(bIsIntegralEnum<Arg> || std::is_integral_v<Arg> || std::is_pointer_v<Arg>))
+					else if constexpr (!(TIsIntegralEnum_V<Arg> || std::is_integral_v<Arg> || std::is_pointer_v<Arg>))
 					{
 						return {EFormatStringSanStatus::DNeedsIntegerArg, CurArgPos};
 					}
@@ -312,7 +312,7 @@ namespace UE::Core::Private
 					{
 						return {EFormatStringSanStatus::ZNeedsIntegerSpec, CurArgPos};
 					}
-					else if constexpr (!(bIsIntegralEnum<Arg> || std::is_integral_v<Arg>))
+					else if constexpr (!(TIsIntegralEnum_V<Arg> || std::is_integral_v<Arg>))
 					{
 						return {EFormatStringSanStatus::ZNeedsIntegerArg, CurArgPos};
 					}
@@ -338,7 +338,7 @@ namespace UE::Core::Private
 					{
 						return {EFormatStringSanStatus::I64BadSpec, CurArgPos};
 					}
-					else if constexpr (!(bIsIntegralEnum<Arg> || std::is_integral_v<Arg>))
+					else if constexpr (!(TIsIntegralEnum_V<Arg> || std::is_integral_v<Arg>))
 					{
 						return {EFormatStringSanStatus::I64NeedsIntegerArg, CurArgPos};
 					}
@@ -350,7 +350,7 @@ namespace UE::Core::Private
 				case (CharType)'l':
 					if (CharIsIntegerFormatSpecifier(Fmt[1]))
 					{
-						if constexpr (!(bIsIntegralEnum<Arg> || std::is_integral_v<Arg>))
+						if constexpr (!(TIsIntegralEnum_V<Arg> || std::is_integral_v<Arg>))
 						{
 							return {EFormatStringSanStatus::LNeedsIntegerArg, CurArgPos};
 						}
@@ -378,7 +378,7 @@ namespace UE::Core::Private
 					{
 						return {EFormatStringSanStatus::LLNeedsIntegerSpec, CurArgPos};
 					}
-					else if constexpr (!(bIsIntegralEnum<Arg> || std::is_integral_v<Arg> || std::is_pointer_v<Arg>))
+					else if constexpr (!(TIsIntegralEnum_V<Arg> || std::is_integral_v<Arg> || std::is_pointer_v<Arg>))
 					{
 						return {EFormatStringSanStatus::LLNeedsIntegerArg, CurArgPos};
 					}
@@ -390,7 +390,7 @@ namespace UE::Core::Private
 				case (CharType)'h':
 					if (CharIsIntegerFormatSpecifier(Fmt[1]))
 					{
-						if constexpr (!(bIsIntegralEnum<Arg> || std::is_integral_v<Arg>))
+						if constexpr (!(TIsIntegralEnum_V<Arg> || std::is_integral_v<Arg>))
 						{
 							return {EFormatStringSanStatus::HNeedsIntegerArg, CurArgPos};
 						}
@@ -404,7 +404,7 @@ namespace UE::Core::Private
 					{
 						return {EFormatStringSanStatus::HHNeedsIntegerSpec, CurArgPos};
 					}
-					else if constexpr (!(bIsIntegralEnum<Arg> || std::is_integral_v<Arg>))
+					else if constexpr (!(TIsIntegralEnum_V<Arg> || std::is_integral_v<Arg>))
 					{
 						return {EFormatStringSanStatus::HHNeedsIntegerArg, CurArgPos};
 					}
