Cache transactions to improve runtime by 27.3% in a key benchmark.

diff --git a/Engine/Source/Runtime/Core/Private/AutoRTFM/Context.cpp b/Engine/Source/Runtime/Core/Private/AutoRTFM/Context.cpp
--- a/Engine/Source/Runtime/Core/Private/AutoRTFM/Context.cpp
+++ b/Engine/Source/Runtime/Core/Private/AutoRTFM/Context.cpp
@@ -298,7 +298,7 @@ FTransaction* FContext::PushTransaction(
 		CurrentTransaction->SetInactive();
 	}
 
-	FTransaction* NewTransaction = new FTransaction(this);
+	FTransaction* NewTransaction = TransactionPool.Take(this);
 	NewTransaction->Initialize(
 		/* Parent */ CurrentTransaction, 
 		/* bIsScoped */ bIsScoped, 
@@ -333,7 +333,7 @@ void FContext::PopTransaction()
 		ASSERT(CurrentTransaction->IsInactive());
 		CurrentTransaction->SetActive();
 	}
-	delete OldTransaction;
+	TransactionPool.Return(OldTransaction);
 }
 
 void FContext::ClearTransactionStatus()
diff --git a/Engine/Source/Runtime/Core/Private/AutoRTFM/Context.h b/Engine/Source/Runtime/Core/Private/AutoRTFM/Context.h
--- a/Engine/Source/Runtime/Core/Private/AutoRTFM/Context.h
+++ b/Engine/Source/Runtime/Core/Private/AutoRTFM/Context.h
@@ -5,21 +5,23 @@
 #include "AutoRTFM/AutoRTFM.h"
 #include "ContextStatus.h"
 #include "HAL/PlatformTLS.h"
+#include "IntrusivePool.h"
 #include "StackRange.h"
 #include "OpenHashThrottler.h"
 #include "TaskArray.h"
+#include "Transaction.h"
 #include "Pool.h"
 
 namespace AutoRTFM
 {
 
 class FLineLock;
-class FTransaction;
 class FCallNest;
 
-class FContext
+class FContext final
 {
 	using FTaskPool = typename TTaskArray<TTask<void()>>::FEntryPool;
+	using FTransactionPool = TIntrusivePool<FTransaction, 16>;
 
 public:
     static FContext* Get();
@@ -139,6 +141,7 @@ public:
 	}
 
 	inline FTaskPool& GetTaskPool() { return TaskPool; }
+	inline FTransactionPool& GetTransactionPool() { return TransactionPool; }
 
 private:
 	static FContext ContextSingleton;
@@ -174,6 +177,7 @@ private:
 	uint32 StackLocalInitializerDepth = 0;
 	TTask<void()> PostAbortCallback;
 	FTaskPool TaskPool;
+	FTransactionPool TransactionPool;
 	FOpenHashThrottler OpenHashThrottler{
 		/* LogInterval            */ 10,   // Log every 10 seconds.
 		/* AdjustThrottleInterval */ 0.5,  // Adjust throttling probabilities every 500ms.
diff --git a/Engine/Source/Runtime/Core/Private/AutoRTFM/IntrusivePool.h b/Engine/Source/Runtime/Core/Private/AutoRTFM/IntrusivePool.h
new file mode 100644
--- /dev/null
+++ b/Engine/Source/Runtime/Core/Private/AutoRTFM/IntrusivePool.h
@@ -0,0 +1,94 @@
+// Copyright Epic Games, Inc. All Rights Reserved.
+
+#pragma once
+
+#include "BlockAllocator.h"
+#include <cstddef>
+#include <utility>
+
+namespace AutoRTFM
+{
+
+template<typename ItemType, size_t InlineCapacity>
+class TIntrusivePool final
+{
+public:
+	using FItem = ItemType;
+
+	TIntrusivePool() = default;
+	
+	~TIntrusivePool()
+	{
+		while (FItem* const Item = FreeList)
+		{
+			FreeList = *Item->GetIntrusiveAddress();
+
+			Item->~FItem();
+		}
+	}
+
+	// Acquires a new Item from the pool. If the item was previously returned
+	// to the pool, we'll call `Ressurect` on it with `Arguments`. Otherwise
+	// create a new item with `Arguments`.
+	template<typename ... ArgumentTypes>
+	FItem* Take(ArgumentTypes&& ... Arguments)
+	{
+		FItem* Item = nullptr;
+
+		if (FreeList)
+		{
+			// Unlink a free entry from the free list.
+			Item = FreeList;
+			FreeList = *(FreeList->GetIntrusiveAddress());
+			
+			Item->Ressurect(std::forward<ArgumentTypes>(Arguments)...);
+		}
+		else
+		{
+			// Free list is empty. Allocate another item.
+			void* const Memory = Allocator.Allocate(EntrySize, EntryAlignment);
+			NumAllocated++;
+
+			Item = new (Memory) FItem(std::forward<ArgumentTypes>(Arguments)...);
+		}
+
+		ASSERT(NumInUse < NumAllocated);
+		NumInUse++;
+
+		return Item;
+	}
+
+	// Calls `Suppress` on `Item` and returns it to the pool to be reused.
+	void Return(FItem* const Item)
+	{
+		// Suppress the item.
+		Item->Suppress();
+
+		// Place the entry onto the free list.
+		*Item->GetIntrusiveAddress() = FreeList;
+		FreeList = Item;
+
+		ASSERT(NumInUse > 0);
+		NumInUse--;
+	}
+    
+private:
+	static constexpr size_t EntrySize = sizeof(FItem);
+	static constexpr size_t EntryAlignment = alignof(FItem);
+	
+	// The underlying allocator for the pool.
+	TBlockAllocator<InlineCapacity * EntrySize, EntryAlignment> Allocator;
+	// Free list.
+	FItem* FreeList = nullptr;
+	// Number of entries allocated.
+	size_t NumAllocated = 0;
+	// Number of entries currently in use.
+	size_t NumInUse = 0;
+
+	TIntrusivePool(TIntrusivePool&&) = delete;
+	TIntrusivePool(const TIntrusivePool&) = delete;
+	TIntrusivePool& operator=(const TIntrusivePool&) = delete;
+	TIntrusivePool& operator=(TIntrusivePool&&) = delete;
+};
+
+}  // namespace AutoRTFM
diff --git a/Engine/Source/Runtime/Core/Private/AutoRTFM/Transaction.cpp b/Engine/Source/Runtime/Core/Private/AutoRTFM/Transaction.cpp
--- a/Engine/Source/Runtime/Core/Private/AutoRTFM/Transaction.cpp
+++ b/Engine/Source/Runtime/Core/Private/AutoRTFM/Transaction.cpp
@@ -20,17 +20,43 @@
 
 namespace AutoRTFM
 {
+	
+void FTransaction::Ressurect(FContext* const InContext)
+{
+	ASSERT(Context == InContext)
+	Parent = nullptr;
+	StatDepth = 1;
+	RecordedWriteHash = 0;
+	NumWriteLogsHashed = 0;
+	CurrentMemoryValidationLevel = EMemoryValidationLevel::Disabled;
+	CurrentOpenReturnAddress = nullptr;
+	CurrentState = EState::Uninitialized;
+	bIsStackScoped = false;
+	bIsInAllocateFn = false;
+}
+
+void FTransaction::Suppress()
+{
+	CurrentState = EState::Done;
+	Reset();
+}
+
+FTransaction** FTransaction::GetIntrusiveAddress()
+{
+	return &Parent;
+}
 
 FTransaction::FTransaction(FContext* Context)
     : Context(Context)
 	, CommitTasks(Context->GetTaskPool())
 	, AbortTasks(Context->GetTaskPool())
 {
+	Ressurect(Context);
 }
 
 FTransaction::~FTransaction()
 {
-	ASSERT(IsDone());
+	Suppress();
 }
 
 void FTransaction::Initialize(FTransaction* Parent_, bool bIsStackScoped_, FStackRange StackRange_)
diff --git a/Engine/Source/Runtime/Core/Private/AutoRTFM/Transaction.h b/Engine/Source/Runtime/Core/Private/AutoRTFM/Transaction.h
--- a/Engine/Source/Runtime/Core/Private/AutoRTFM/Transaction.h
+++ b/Engine/Source/Runtime/Core/Private/AutoRTFM/Transaction.h
@@ -8,6 +8,7 @@
 #include "CoreTypes.h"
 #include "HitSet.h"
 #include "IntervalTree.h"
+#include "IntrusivePool.h"
 #include "StackRange.h"
 #include "Stats.h"
 #include "TaskArray.h"
@@ -19,6 +20,20 @@ class FContext;
 
 class FTransaction final
 {
+	friend class TIntrusivePool<FTransaction, 16>;
+
+	void Ressurect(FContext* const InContext);
+	void Suppress();
+	FTransaction** GetIntrusiveAddress();
+
+	// Constructor.
+	// The transaction is constructed in an initial Closed state.
+	FTransaction(FContext* const Context);
+
+	// Destructor.
+	// The transaction must we in a Done state.
+	~FTransaction();
+
 public:
 	// State flow diagram for a transaction:
 	//                       ┌──────────────────────────────────────┐
@@ -56,14 +71,6 @@ public:
 		Done,
 	};
 
-	// Constructor.
-	// The transaction is constructed in an initial Closed state.
-	FTransaction(FContext* Context);
-
-	// Destructor.
-	// The transaction must we in a Done state.
-	~FTransaction();
-
 	void Initialize(FTransaction* Parent, bool bIsScoped, FStackRange StackRange);
 
 	// Clears the tracked transaction state and resets back to the default Uninitialized state.
@@ -180,10 +187,10 @@ private:
 	template<EState NewState>
 	void SetState();
 
-	FContext* Context{nullptr};
+	FContext* Context;
 
 	// Are we nested? Then this is the parent.
-	FTransaction* Parent{nullptr};
+	FTransaction* Parent;
 
 	// Commit tasks run on commit in forward order.
 	FTaskArray CommitTasks;
@@ -211,14 +218,14 @@ private:
 	FIntervalTree NewMemoryTracker;
 	FWriteLog WriteLog;
 	TStatStorage<uint64_t> StatDepth = 1;
-	FWriteHash RecordedWriteHash = 0;
-	size_t NumWriteLogsHashed = 0;
-	EMemoryValidationLevel CurrentMemoryValidationLevel = EMemoryValidationLevel::Disabled;
-	const void* CurrentOpenReturnAddress = nullptr;
-	EState CurrentState{EState::Uninitialized};
+	FWriteHash RecordedWriteHash;
+	size_t NumWriteLogsHashed;
+	EMemoryValidationLevel CurrentMemoryValidationLevel;
+	const void* CurrentOpenReturnAddress;
+	EState CurrentState;
 	FStackRange StackRange;
-	bool bIsStackScoped = false;
-	bool bIsInAllocateFn = false;
+	bool bIsStackScoped;
+	bool bIsInAllocateFn;
 };
 
 } // namespace AutoRTFM
