Prevent FSlowTask::ForceRefresh from working whenthen slow task is supposed to be disabled.

#rnx

- ForceRefresh should do nothing is the object is disabled. This could either be because the caller has decided that the FSlowTask shouldn't do anything or (more importantly) we have detected that the FSlowTask is not on the GameThread and would cause a crash if we continued.
- This makes FSlowTask::ForceRefresh safe to call in code that may be called on background threads as well as the GameThread.

diff --git a/Engine/Source/Runtime/Core/Private/Misc/SlowTask.cpp b/Engine/Source/Runtime/Core/Private/Misc/SlowTask.cpp
--- a/Engine/Source/Runtime/Core/Private/Misc/SlowTask.cpp
+++ b/Engine/Source/Runtime/Core/Private/Misc/SlowTask.cpp
@@ -217,7 +217,10 @@ void FSlowTask::TickProgress()
 
 void FSlowTask::ForceRefresh()
 {
-	ForceRefresh(Context);
+	if (bEnabled)
+	{
+		ForceRefresh(Context);
+	}
 }
 
 const FText& FSlowTask::GetCurrentMessage() const
