Adds a new Niagara custom version to make sure that the MiscUsageBitMask is serialized.  These structures need to be cleaned
#tests Evaluate MiscUsageBitMask in DDC scenarios, verify that it gets serialized

diff --git a/Engine/Plugins/FX/Niagara/Source/NiagaraCore/Public/NiagaraCustomVersion.h b/Engine/Plugins/FX/Niagara/Source/NiagaraCore/Public/NiagaraCustomVersion.h
--- a/Engine/Plugins/FX/Niagara/Source/NiagaraCore/Public/NiagaraCustomVersion.h
+++ b/Engine/Plugins/FX/Niagara/Source/NiagaraCore/Public/NiagaraCustomVersion.h
@@ -210,6 +210,8 @@ struct FNiagaraCustomVersion
 		HierarchyEditorScriptSupport, // Hierarchy Editor was implemented
 
 		EmitterStateAddLoopDelayEnabled,	// Added loop delay enabled to emitter state
+
+		SerializeUsageBitMaskToGPUFunctionInfo, // Make sure that we serialize the UsageBitMask with the GPUFunctionInfo
 		
 		// DO NOT ADD A NEW VERSION UNLESS YOU HAVE TALKED TO THE NIAGARA LEAD. Mismanagement of these versions can lead to data loss if it is adjusted in multiple streams simultaneously.
 		// -----<new versions can be added above this line>  -------------------------------------------------
diff --git a/Engine/Plugins/FX/Niagara/Source/NiagaraShader/Private/NiagaraShader.cpp b/Engine/Plugins/FX/Niagara/Source/NiagaraShader/Private/NiagaraShader.cpp
--- a/Engine/Plugins/FX/Niagara/Source/NiagaraShader/Private/NiagaraShader.cpp
+++ b/Engine/Plugins/FX/Niagara/Source/NiagaraShader/Private/NiagaraShader.cpp
@@ -1264,6 +1264,12 @@ bool FNiagaraDataInterfaceGeneratedFunction::Serialize(FArchive& Ar)
 		Ar << VariadicInputs;
 		Ar << VariadicOutputs;
 	}
+
+	if (NiagaraVer >= FNiagaraCustomVersion::SerializeUsageBitMaskToGPUFunctionInfo)
+	{
+		Ar << MiscUsageBitMask;
+	}
+
 	return true;
 }
 
diff --git a/Engine/Source/Runtime/Core/Private/UObject/FortniteMainBranchObjectVersion.cpp b/Engine/Source/Runtime/Core/Private/UObject/FortniteMainBranchObjectVersion.cpp
--- a/Engine/Source/Runtime/Core/Private/UObject/FortniteMainBranchObjectVersion.cpp
+++ b/Engine/Source/Runtime/Core/Private/UObject/FortniteMainBranchObjectVersion.cpp
@@ -10,7 +10,7 @@ TMap<FGuid, FGuid> FFortniteMainBranchObjectVersion::GetSystemGuids()
 	SystemGuids.Add(DevGuids.LANDSCAPE_MOBILE_COOK_VERSION, FGuid("32D02EF867C74B71A0D4E0FA41392732"));
 	SystemGuids.Add(DevGuids.MATERIALSHADERMAP_DERIVEDDATA_VER, FGuid("0E6202C151DE402CA451802671AE9347"));
 	SystemGuids.Add(DevGuids.NANITE_DERIVEDDATA_VER, FGuid("42C46B24E1586986880AFE62B0F575AB"));
-	SystemGuids.Add(DevGuids.NIAGARASHADERMAP_DERIVEDDATA_VER, FGuid("C417A50271B3427D98785388266E3B66"));
+	SystemGuids.Add(DevGuids.NIAGARASHADERMAP_DERIVEDDATA_VER, FGuid("016F98966CE4455EBA36999548E22A39"));
 	SystemGuids.Add(DevGuids.Niagara_LatestScriptCompileVersion, FGuid("393081597FA7234C92F00869DE262D73"));
 	SystemGuids.Add(DevGuids.SkeletalMeshDerivedDataVersion, FGuid("2F90180769024483BB7D90AACEB2798D"));
 	SystemGuids.Add(DevGuids.STATICMESH_DERIVEDDATA_VER, FGuid("656E338C78C34641A5D3BFAE303997AF"));
