[Core]
* Fixed superluminal integration typo where functions were implemented where they only should have been declared

diff --git a/Engine/Source/Runtime/Core/Private/ProfilingDebugging/Windows/SuperluminalExternalProfiler.cpp b/Engine/Source/Runtime/Core/Private/ProfilingDebugging/Windows/SuperluminalExternalProfiler.cpp
--- a/Engine/Source/Runtime/Core/Private/ProfilingDebugging/Windows/SuperluminalExternalProfiler.cpp
+++ b/Engine/Source/Runtime/Core/Private/ProfilingDebugging/Windows/SuperluminalExternalProfiler.cpp
@@ -11,9 +11,9 @@
 #include "Features/IModularFeatures.h"
 #include "Templates/UniquePtr.h"
 
-void SuperLuminalStartScopedEventWide(const wchar_t* Text, uint32_t Color) {}
-void SuperLuminalStartScopedEvent(const char* Text, uint32_t Color) {}
-void SuperLuminalEndScopedEvent() {}
+void SuperLuminalStartScopedEventWide(const wchar_t* Text, uint32_t Color);
+void SuperLuminalStartScopedEvent(const char* Text, uint32_t Color);
+void SuperLuminalEndScopedEvent();
 
 // Use a custom color pack function here because our colors are full intensity
 //  and are painful/difficult to read behind text in the app's flame graph
