Add Emplace to TChunkedArray which constructs a new item to the end of the chunked array and return the index.

diff --git a/Engine/Source/Runtime/Core/Public/Containers/ChunkedArray.h b/Engine/Source/Runtime/Core/Public/Containers/ChunkedArray.h
--- a/Engine/Source/Runtime/Core/Public/Containers/ChunkedArray.h
+++ b/Engine/Source/Runtime/Core/Public/Containers/ChunkedArray.h
@@ -191,6 +191,19 @@ public:
 		return this->NumElements - 1;
 	}
 
+	/**
+	 * Constructs a new item to the end of the chunked array.
+	 *
+	 * @param Args	The arguments to forward to the constructor of the new item.
+	 * @return		Index to the new item
+	 */
+	template <typename... ArgsType>
+	int32 Emplace(ArgsType&&... Args)
+	{
+		new(*this) ElementType(Forward<ArgsType>(Args)...);
+		return this->NumElements - 1;
+	}
+
 	/**
 	 * Appends the specified array to this array.
 	 * Cannot append to self.
