Removed tolerance from ExtractScaling in SetFromMatrix, as this breaks transforms with small scaling.

diff --git a/Engine/Source/Runtime/Core/Public/Math/TransformNonVectorized.h b/Engine/Source/Runtime/Core/Public/Math/TransformNonVectorized.h
--- a/Engine/Source/Runtime/Core/Public/Math/TransformNonVectorized.h
+++ b/Engine/Source/Runtime/Core/Public/Math/TransformNonVectorized.h
@@ -1138,7 +1138,7 @@ public:
 		TMatrix<T> M = InMatrix;
 
 		// Get the 3D scale from the matrix
-		Scale3D = M.ExtractScaling();
+		Scale3D = M.ExtractScaling(0);
 
 		// If there is negative scaling going on, we handle that here
 		if (InMatrix.Determinant() < 0.f)
diff --git a/Engine/Source/Runtime/Core/Public/Math/TransformVectorized.h b/Engine/Source/Runtime/Core/Public/Math/TransformVectorized.h
--- a/Engine/Source/Runtime/Core/Public/Math/TransformVectorized.h
+++ b/Engine/Source/Runtime/Core/Public/Math/TransformVectorized.h
@@ -1250,7 +1250,7 @@ public:
 		TMatrix<T> M = InMatrix;
 
 		// Get the 3D scale from the matrix
-		TVector<T> InScale = M.ExtractScaling();
+		TVector<T> InScale = M.ExtractScaling(0);
 		Scale3D = VectorLoadFloat3_W0(&InScale);
 
 		// If there is negative scaling going on, we handle that here
