Have a single empty string reference rather than 4 copies of it.

diff --git a/Engine/Source/Runtime/Core/Private/Internationalization/FastDecimalFormat.cpp b/Engine/Source/Runtime/Core/Private/Internationalization/FastDecimalFormat.cpp
--- a/Engine/Source/Runtime/Core/Private/Internationalization/FastDecimalFormat.cpp
+++ b/Engine/Source/Runtime/Core/Private/Internationalization/FastDecimalFormat.cpp
@@ -149,8 +149,7 @@ public:
 
 	const FString& GetNegativePrefixString() const
 	{
-		static const FString EmptyStr = FString();
-		return NegativePrefixStringPtr ? *NegativePrefixStringPtr : EmptyStr;
+		return NegativePrefixStringPtr ? *NegativePrefixStringPtr : GetEmptyString();
 	}
 
 	bool HasNegativeSuffixString() const
@@ -160,8 +159,7 @@ public:
 
 	const FString& GetNegativeSuffixString() const
 	{
-		static const FString EmptyStr = FString();
-		return NegativeSuffixStringPtr ? *NegativeSuffixStringPtr : EmptyStr;
+		return NegativeSuffixStringPtr ? *NegativeSuffixStringPtr : GetEmptyString();
 	}
 
 	bool HasPositivePrefixString() const
@@ -171,8 +169,7 @@ public:
 
 	const FString& GetPositivePrefixString() const
 	{
-		static const FString EmptyStr = FString();
-		return PositivePrefixStringPtr ? *PositivePrefixStringPtr : EmptyStr;
+		return PositivePrefixStringPtr ? *PositivePrefixStringPtr : GetEmptyString();
 	}
 
 	bool HasPositiveSuffixString() const
@@ -182,8 +179,7 @@ public:
 
 	const FString& GetPositiveSuffixString() const
 	{
-		static const FString EmptyStr = FString();
-		return PositiveSuffixStringPtr ? *PositiveSuffixStringPtr : EmptyStr;
+		return PositiveSuffixStringPtr ? *PositiveSuffixStringPtr : GetEmptyString();
 	}
 
 private:
@@ -196,6 +192,12 @@ private:
 	FString GeneratedNegativeSuffixString;
 	FString GeneratedPositivePrefixString;
 	FString GeneratedPositiveSuffixString;
+
+	const FString& GetEmptyString() const
+	{
+		static const FString EmptyStr = FString();
+		return EmptyStr;
+	}
 };
 
 void SanitizeNumberFormattingOptions(FNumberFormattingOptions& InOutFormattingOptions)
