[AutoRTFM] Fix return value from FTransactionallySafeRWLock::TryWriteLock().

We actually had code already to set up a valid return value, but somehow neglected to use it.

Some test cases are still not passing, but this failure mode (mixing locks and unlocks across open/closed) extends past RWLock to almost all of our locking code in AutoRTFM, so I've cut a separate Jira for those. (SOL-7661)

diff --git a/Engine/Source/Programs/AutoRTFMTests/Private/RWLockTests.cpp b/Engine/Source/Programs/AutoRTFMTests/Private/RWLockTests.cpp
--- a/Engine/Source/Programs/AutoRTFMTests/Private/RWLockTests.cpp
+++ b/Engine/Source/Programs/AutoRTFMTests/Private/RWLockTests.cpp
@@ -407,18 +407,6 @@ TEST_CASE("TransactionallySafeRWLock.MultipleTryWriteLocks")
 		RWLock.WriteUnlock();
 	}
 
-	// Test case does not pass; RWLock trashes its internal state: #jira SOL-7658
-	SECTION("TryWrite, Commit(TryWrite, Unlock)") 
-	{
-//		FTransactionallySafeRWLock RWLock;
-//		REQUIRE(RWLock.TryWriteLock());
-//		AutoRTFM::Testing::Commit([&]
-//		{
-//			REQUIRE(!RWLock.TryWriteLock());
-//			RWLock.WriteUnlock();
-//		}); 
-	}
-
 	SECTION("TryWrite, Abort(TryWrite), Unlock")
 	{
 		FTransactionallySafeRWLock RWLock;
@@ -431,27 +419,76 @@ TEST_CASE("TransactionallySafeRWLock.MultipleTryWriteLocks")
 		RWLock.WriteUnlock();
 	}
 
-	// Test case does not pass; RWLock trashes its internal state: #jira SOL-7658
-	SECTION("Commit(TryWrite, TryWrite), Unlock")
+	// Test case does not pass. RWLock cannot be locked in the open and unlocked in the closed: #jira SOL-7661
+	SECTION("TryWrite, Commit(TryWrite, Unlock)")
 	{
 //		FTransactionallySafeRWLock RWLock;
+//		REQUIRE(RWLock.TryWriteLock());
 //		AutoRTFM::Testing::Commit([&]
 //		{
-//			REQUIRE(RWLock.TryWriteLock());
 //			REQUIRE(!RWLock.TryWriteLock());
-//		});
+//			RWLock.WriteUnlock();
+//		}); 
+	}
+
+	// Test case does not pass. RWLock cannot be locked in the open and unlocked in the closed: #jira SOL-7661
+	SECTION("TryWrite, Abort(TryWrite, Unlock)")
+	{
+//		FTransactionallySafeRWLock RWLock;
+//		REQUIRE(RWLock.TryWriteLock());
+//		AutoRTFM::Testing::Abort([&]
+//		{
+//			REQUIRE(!RWLock.TryWriteLock());
+//			RWLock.WriteUnlock();
+//			AutoRTFM::AbortTransaction();
+//		}); 
 //		RWLock.WriteUnlock();
 	}
 
-	// Test case does not pass; RWLock trashes its internal state: #jira SOL-7658
-	SECTION("Commit(TryWrite, TryWrite, Unlock)")
+
+	// Test case does not pass. RWLock cannot be locked in the closed and unlocked in the open: #jira SOL-7661
+	SECTION("Commit(TryWrite, TryWrite), Unlock")
 	{
+//		FTransactionallySafeRWLock RWLock;
 //		AutoRTFM::Testing::Commit([&]
 //		{
-//			FTransactionallySafeRWLock RWLock;
 //			REQUIRE(RWLock.TryWriteLock());
 //			REQUIRE(!RWLock.TryWriteLock());
-//			RWLock.WriteUnlock();
 //		});
+//		RWLock.WriteUnlock();
+	}
+
+	SECTION("Abort(TryWrite, TryWrite)")
+	{
+		FTransactionallySafeRWLock RWLock;
+		AutoRTFM::Testing::Abort([&]
+		{
+			REQUIRE(RWLock.TryWriteLock());
+			REQUIRE(!RWLock.TryWriteLock());
+			AutoRTFM::AbortTransaction();
+		});
+	}
+
+	SECTION("Commit(TryWrite, TryWrite, Unlock)")
+	{
+		AutoRTFM::Testing::Commit([&]
+		{
+			FTransactionallySafeRWLock RWLock;
+			REQUIRE(RWLock.TryWriteLock());
+			REQUIRE(!RWLock.TryWriteLock());
+			RWLock.WriteUnlock();
+		});
+	}
+
+	SECTION("Abort(TryWrite, TryWrite, Unlock)")
+	{
+		AutoRTFM::Testing::Abort([&]
+		{
+			FTransactionallySafeRWLock RWLock;
+			REQUIRE(RWLock.TryWriteLock());
+			REQUIRE(!RWLock.TryWriteLock());
+			RWLock.WriteUnlock();
+			AutoRTFM::AbortTransaction();
+		});
 	}
 }
diff --git a/Engine/Source/Runtime/Core/Public/Misc/TransactionallySafeRWLock.h b/Engine/Source/Runtime/Core/Public/Misc/TransactionallySafeRWLock.h
--- a/Engine/Source/Runtime/Core/Public/Misc/TransactionallySafeRWLock.h
+++ b/Engine/Source/Runtime/Core/Public/Misc/TransactionallySafeRWLock.h
@@ -75,23 +75,23 @@ struct FTransactionallySafeRWLockDefinition
 		if (AutoRTFM::IsTransactional() || AutoRTFM::IsCommittingOrAborting())
 		{
 			UE_AUTORTFM_OPEN
+			{
+				// The transactional system which can increment TransactionalLockCount
+				// is always single-threaded, thus this is safe to check without atomicity.
+				if (0 == State->TransactionalLockCount)
 				{
-					// The transactional system which can increment TransactionalLockCount
-					// is always single-threaded, thus this is safe to check without atomicity.
-					if (0 == State->TransactionalLockCount)
-					{
-						State->Lock.WriteLock();
-					}
+					State->Lock.WriteLock();
+				}
 
-					State->TransactionalLockCount += 1;
-				};
+				State->TransactionalLockCount += 1;
+			};
 
 			// We explicitly copy the state here for the case that `this` was stack
 			// allocated and has already died before the on-abort is hit.
-			AutoRTFM::OnAbort([State = this->State]
-				{
-					State->WriteUnlock();
-				});
+			UE_AUTORTFM_ONABORT(State = this->State)
+			{
+				State->WriteUnlock();
+			};
 		}
 		else
 		{
@@ -106,10 +106,10 @@ struct FTransactionallySafeRWLockDefinition
 		{
 			// We explicitly copy the state here for the case that `this` was stack
 			// allocated and has already died before the on-commit is hit.
-			AutoRTFM::OnCommit([State = this->State]
-				{
-					State->WriteUnlock();
-				});
+			UE_AUTORTFM_ONCOMMIT(State = this->State)
+			{
+				State->WriteUnlock();
+			};
 		}
 		else
 		{
@@ -125,33 +125,33 @@ struct FTransactionallySafeRWLockDefinition
 			bool Ret = false;
 
 			UE_AUTORTFM_OPEN
+			{
+				// The transactional system which can increment TransactionalLockCount
+				// is always single-threaded, thus this is safe to check without atomicity.
+				// For TryWriteLock we should only lock when we have a 0 count as no one owns this lock
+				if (0 == State->TransactionalLockCount)
 				{
-					// The transactional system which can increment TransactionalLockCount
-					// is always single-threaded, thus this is safe to check without atomicity.
-					// For TryWriteLock we should only lock when we have a 0 count as no one owns this lock
-					if (0 == State->TransactionalLockCount)
-					{
-						Ret = State->Lock.TryWriteLock();
-					}
-
-					if (Ret)
-					{
-						State->TransactionalLockCount += 1;
-					}
-				};
+					Ret = State->Lock.TryWriteLock();
+				}
+
+				if (Ret)
+				{
+					State->TransactionalLockCount += 1;
+				}
+			};
 
 			// Only setup the OnAbort if we *did* grab a lock, otherwise we will not want to do anything with the count or lock
 			if (Ret)
 			{
 				// We explicitly copy the state here for the case that `this` was stack
 				// allocated and has already died before the on-abort is hit.
-				AutoRTFM::OnAbort([State = this->State]
-					{
-						State->WriteUnlock();
-					});
+				UE_AUTORTFM_ONABORT(State = this->State)
+				{
+					State->WriteUnlock();
+				};
 			}
 
-			return State->TransactionalLockCount != 0;
+			return Ret;
 		}
 
 		return State->Lock.TryWriteLock();
