Migrate from deprecated openURL:  to openURL:options:completionHandler. As of iOS 17, openURL is documented as "Calling this method has no effect.", however, in practice, it sill appears to still have lingering functionality, however, schemes shuch as shareddocuments are not handled. Drive by add URL scheme detection parity WRT Mac/Android, adding http:// if no scheme is present in the requested URL.

diff --git a/Engine/Source/Runtime/Core/Private/IOS/IOSPlatformProcess.cpp b/Engine/Source/Runtime/Core/Private/IOS/IOSPlatformProcess.cpp
--- a/Engine/Source/Runtime/Core/Private/IOS/IOSPlatformProcess.cpp
+++ b/Engine/Source/Runtime/Core/Private/IOS/IOSPlatformProcess.cpp
@@ -55,13 +55,35 @@ void FIOSPlatformProcess::LaunchURL( const TCHAR* URL, const TCHAR* Parms, FStri
 		return;
 	}
 
+	dispatch_semaphore_t Semaphore = dispatch_semaphore_create(0);
+	__block BOOL Result;
+	FString SchemeName;
+	bool bHasSchemeName = FParse::SchemeNameFromURI(URL, SchemeName);
+		
 	NSString* CFUrl = (NSString*)FPlatformString::TCHARToCFString( URL );
-	PRAGMA_DISABLE_DEPRECATION_WARNINGS
-	bool Result = [[UIApplication sharedApplication] openURL: [NSURL URLWithString: CFUrl]];
-	PRAGMA_ENABLE_DEPRECATION_WARNINGS
+	if(!bHasSchemeName)
+	{
+		CFUrl = [@"http://%@" stringByAppendingString:CFUrl];
+	}
+	[[UIApplication sharedApplication] openURL: [NSURL URLWithString: CFUrl] options:@{} completionHandler:
+		^(BOOL Success)
+		{
+			Result = Success;
+			dispatch_semaphore_signal(Semaphore);
+		}
+	];
+	
+	// wait for a result, but timeout after 1s
+	if (dispatch_semaphore_wait(Semaphore, dispatch_time(DISPATCH_TIME_NOW, NSEC_PER_SEC)))
+	{
+		// We timed out, error unknown
+		Result = TRUE;
+	}
+	dispatch_release(Semaphore);
+	
 	if (Error != nullptr)
 	{
-		*Error = Result ? TEXT("") : TEXT("unable to open url");
+		*Error = Result ? TEXT("") : TEXT("Unable to open url");
 	}
 #endif
 }
