Correct typo in scheme addition for iOS. Lingering format parameters got left behind from integration from Mac but had no side effects for local tests.

diff --git a/Engine/Source/Runtime/Core/Private/IOS/IOSPlatformProcess.cpp b/Engine/Source/Runtime/Core/Private/IOS/IOSPlatformProcess.cpp
--- a/Engine/Source/Runtime/Core/Private/IOS/IOSPlatformProcess.cpp
+++ b/Engine/Source/Runtime/Core/Private/IOS/IOSPlatformProcess.cpp
@@ -63,7 +63,7 @@ void FIOSPlatformProcess::LaunchURL( const TCHAR* URL, const TCHAR* Parms, FStri
 	NSString* CFUrl = (NSString*)FPlatformString::TCHARToCFString( URL );
 	if(!bHasSchemeName)
 	{
-		CFUrl = [@"http://%@" stringByAppendingString:CFUrl];
+		CFUrl = [@"http://" stringByAppendingString:CFUrl];
 	}
 	[[UIApplication sharedApplication] openURL: [NSURL URLWithString: CFUrl] options:@{} completionHandler:
 		^(BOOL Success)
