[AutoRTFM] Migrate from UE's TArray to AutoRTFM::TStack

Part of the work to separate AutoRTFM out of core to break cyclic build dependency issues.

diff --git a/Engine/Source/Runtime/Core/Private/AutoRTFM/IntervalTree.h b/Engine/Source/Runtime/Core/Private/AutoRTFM/IntervalTree.h
--- a/Engine/Source/Runtime/Core/Private/AutoRTFM/IntervalTree.h
+++ b/Engine/Source/Runtime/Core/Private/AutoRTFM/IntervalTree.h
@@ -2,68 +2,13 @@
 
 #pragma once
 
-#include "Containers/Array.h"
 #include "Containers/BitArray.h"
+#include "Stack.h"
 #include "Utils.h"
 
 namespace AutoRTFM
 {
 
-template<typename T> struct TIntervalTreeArray final
-{
-	FORCEINLINE void Empty()
-	{
-		Payload.Empty();
-	}
-
-	FORCEINLINE uint32 Num() const
-	{
-		return Payload.Num();
-	}
-
-	FORCEINLINE bool IsEmpty() const
-	{
-		return Payload.IsEmpty();
-	}
-
-	FORCEINLINE void Push(T Thing)
-	{
-		Payload.Push(Thing);
-	}
-
-	FORCEINLINE uint32 Add(T Thing)
-	{
-		return Payload.Add(Thing);
-	}
-
-	FORCEINLINE T Pop()
-	{
-		return Payload.Pop();
-	}
-
-	FORCEINLINE void Reset()
-	{
-		Payload.Reset();
-	}
-
-	FORCEINLINE T& operator[](uint32 Index)
-	{
-		// A bit nasty - but since we know we are in bounds, skip the check.
-		// This change alone accounted for a 13% uplift in performance in a
-		// tight loop of critical performance importance!
-		return Payload.GetData()[Index];
-	}
-
-	FORCEINLINE const T& operator[](uint32 Index) const
-	{
-		// Same as non-const version.
-		return Payload.GetData()[Index];
-	}
-
-private:
-	TArray<T> Payload;
-};
-
 struct FIntervalTree final
 {
     FIntervalTree() = default;
@@ -135,7 +80,8 @@ struct FIntervalTree final
 
         do
         {
-            const FIntervalTreeNodeIndex Current = ToProcess.Pop();
+            const FIntervalTreeNodeIndex Current = ToProcess.Back();
+            ToProcess.Pop();
 
             FRange Range = Other.NodeRanges[Current];
 
@@ -157,7 +103,7 @@ struct FIntervalTree final
     }
 
 private:
-	template<typename T> using ArrayType = TIntervalTreeArray<T>;
+	template<typename T> using ArrayType = TStack<T, 8, EStackValidation::Disabled>;
 
     struct FRange final
     {
@@ -189,20 +135,22 @@ private:
 
 	FORCEINLINE FIntervalTreeNodeIndex AddNode(FRange&& NewRange)
 	{
-		const FIntervalTreeNodeIndex Index = NodeRanges.Add(NewRange);
-		NodeLefts.Add(IntervalTreeNodeIndexNone);
-		NodeRights.Add(IntervalTreeNodeIndexNone);
-		NodeParents.Add(IntervalTreeNodeIndexNone);
+		const FIntervalTreeNodeIndex Index = NodeRanges.Num();
+        NodeRanges.Push(NewRange);
+		NodeLefts.Push(IntervalTreeNodeIndexNone);
+		NodeRights.Push(IntervalTreeNodeIndexNone);
+		NodeParents.Push(IntervalTreeNodeIndexNone);
 		NodeColors.Add(EColor::Black);
 		return Index;
 	}
 
 	FORCEINLINE FIntervalTreeNodeIndex AddNode(FRange&& NewRange, FIntervalTreeNodeIndex Parent)
 	{
-		const FIntervalTreeNodeIndex Index = NodeRanges.Add(NewRange);
-		NodeLefts.Add(IntervalTreeNodeIndexNone);
-		NodeRights.Add(IntervalTreeNodeIndexNone);
-		NodeParents.Add(Parent);
+		const FIntervalTreeNodeIndex Index = NodeRanges.Num();
+        NodeRanges.Push(NewRange);
+		NodeLefts.Push(IntervalTreeNodeIndexNone);
+		NodeRights.Push(IntervalTreeNodeIndexNone);
+		NodeParents.Push(Parent);
 		NodeColors.Add(EColor::Red);
 		return Index;
 	}
diff --git a/Engine/Source/Runtime/Core/Private/AutoRTFM/TaskArray.h b/Engine/Source/Runtime/Core/Private/AutoRTFM/TaskArray.h
--- a/Engine/Source/Runtime/Core/Private/AutoRTFM/TaskArray.h
+++ b/Engine/Source/Runtime/Core/Private/AutoRTFM/TaskArray.h
@@ -4,6 +4,7 @@
 
 #include "HashMap.h"
 #include "Pool.h"
+#include "Stack.h"
 #include "Utils.h"
 #include <cstddef>
 
@@ -138,7 +139,7 @@ public:
 		{
 			KeyType Key = Iterator.Key;
 			TEntryStack& EntryStack = Iterator.Value;
-			Keyed.FindOrAdd(Key).Append(MoveTemp(EntryStack));
+			Keyed.FindOrAdd(Key).PushAll(MoveTemp(EntryStack));
 		}
 
 		// Add in OtherArray's count.
@@ -158,7 +159,8 @@ public:
 	{
 		if (TEntryStack* EntryStack = Keyed.Find(Key); EntryStack)
 		{
-			Release(EntryStack->Pop());
+			Release(EntryStack->Back());
+			EntryStack->Pop();
 			if (EntryStack->IsEmpty())
 			{
 				Keyed.Remove(Key);
@@ -178,7 +180,8 @@ public:
 		{
 			while (!EntryStack->IsEmpty())
 			{
-				Release(EntryStack->Pop());
+				Release(EntryStack->Back());
+				EntryStack->Pop();
 			}
 			Keyed.Remove(Key);
 			return true;
@@ -236,7 +239,7 @@ public:
 
 private:
 	// An stack of pointers to entries in the linked list.
-	using TEntryStack = TArray<FEntry*, TInlineAllocator<8>>;
+	using TEntryStack = TStack<FEntry*, 8>;
 
 	// Unlinks Entry and releases it back to the pool.
 	// Also decrements the count.
diff --git a/Engine/Source/Runtime/Core/Private/AutoRTFM/Transaction.h b/Engine/Source/Runtime/Core/Private/AutoRTFM/Transaction.h
--- a/Engine/Source/Runtime/Core/Private/AutoRTFM/Transaction.h
+++ b/Engine/Source/Runtime/Core/Private/AutoRTFM/Transaction.h
@@ -5,10 +5,10 @@
 #include "AutoRTFM/AutoRTFM.h"
 #include "AutoRTFM/AutoRTFMTask.h"
 #include "HAL/Platform.h"
-#include "CoreTypes.h"
 #include "HitSet.h"
 #include "IntervalTree.h"
 #include "IntrusivePool.h"
+#include "Stack.h"
 #include "StackRange.h"
 #include "Stats.h"
 #include "TaskArray.h"
@@ -200,19 +200,19 @@ private:
 
 	// If a call to `PopOnCommitHandler` could not find a commit to pop, it is deferred
 	// and tried again on the parent transaction.
-	TArray<const void*, TInlineAllocator<8>> DeferredPopOnCommitHandlers;
+	TStack<const void*, 8> DeferredPopOnCommitHandlers;
 
 	// If a call to `PopOnAbortHandler` could not find an abort to pop, it is deferred
 	// and tried again on the parent transaction.
-	TArray<const void*, TInlineAllocator<8>> DeferredPopOnAbortHandlers;
+	TStack<const void*, 8> DeferredPopOnAbortHandlers;
 
 	// If a call to `PopAllOnCommitHandlers` was used and our transaction successfully
 	// commits, we need to propagate this to the parent too.
-	TArray<const void*, TInlineAllocator<1>> DeferredPopAllOnCommitHandlers;
+	TStack<const void*, 1> DeferredPopAllOnCommitHandlers;
 
 	// If a call to `PopAllOnAbortHandlers` was used and our transaction successfully
 	// commits, we need to propagate this to the parent too.
-	TArray<const void*, TInlineAllocator<1>> DeferredPopAllOnAbortHandlers;
+	TStack<const void*, 1> DeferredPopAllOnAbortHandlers;
 
 	FHitSet HitSet;
 	FIntervalTree NewMemoryTracker;
