Adding FResourceBulkDataInterface::GetBulkDataView() as a shortcut to creating a TConstArrayView<> of its data.

diff --git a/Engine/Source/Runtime/Core/Public/Containers/ResourceArray.h b/Engine/Source/Runtime/Core/Public/Containers/ResourceArray.h
--- a/Engine/Source/Runtime/Core/Public/Containers/ResourceArray.h
+++ b/Engine/Source/Runtime/Core/Public/Containers/ResourceArray.h
@@ -114,6 +114,12 @@ public:
 	 */
 	[[nodiscard]] virtual uint32 GetResourceBulkDataSize() const = 0;
 
+	template<typename TElement>
+	[[nodiscard]] TConstArrayView<TElement> GetBulkDataView() const
+	{
+		return TConstArrayView<TElement>(reinterpret_cast<const TElement*>(GetResourceBulkData()), GetResourceBulkDataSize());
+	}
+
 	/**
 	 * Free memory after it has been used to initialize RHI resource 
 	 */
