Improved the comment for FText::AsDate/Time to make it clearer that it expects a UTC timestamp by default

#rnx

diff --git a/Engine/Source/Runtime/Core/Public/Internationalization/Text.h b/Engine/Source/Runtime/Core/Public/Internationalization/Text.h
--- a/Engine/Source/Runtime/Core/Public/Internationalization/Text.h
+++ b/Engine/Source/Runtime/Core/Public/Internationalization/Text.h
@@ -444,13 +444,21 @@ public:
 	[[nodiscard]] static CORE_API FText AsPercent(double Val,	const FNumberFormattingOptions* const Options = NULL, const FCulturePtr& TargetCulture = NULL);
 
 	/**
-	 * Generate an FText that represents the passed number as a date and/or time in the current culture
-	 * @note The overload using a custom pattern uses strftime-like syntax (see FDateTime::ToFormattedString)
+	 * Generate an FText that represents the given timestamp as a date/time in the current culture using the given timezone (default is the local timezone).
+	 * 
+	 * This function will convert the given timestamp from UTC (eg, FDateTime::UtcNow()) to the given timezone (taking into account DST).
+	 * If using a local timestamp (eg, FDateTime::Now()) then you MUST use FText::GetInvariantTimeZone() as the TimeZone argument to avoid additional timezone conversion.
+	 * 
+	 * @note The overload using a custom pattern uses strftime-like syntax (see FDateTime::ToFormattedString).
 	 */
 	[[nodiscard]] static CORE_API FText AsDate(const FDateTime& DateTime, const EDateTimeStyle::Type DateStyle = EDateTimeStyle::Default, const FString& TimeZone = FString(), const FCulturePtr& TargetCulture = NULL);
 	[[nodiscard]] static CORE_API FText AsDateTime(const FDateTime& DateTime, const EDateTimeStyle::Type DateStyle = EDateTimeStyle::Default, const EDateTimeStyle::Type TimeStyle = EDateTimeStyle::Default, const FString& TimeZone = FString(), const FCulturePtr& TargetCulture = NULL);
 	[[nodiscard]] static CORE_API FText AsDateTime(const FDateTime& DateTime, const FString& CustomPattern, const FString& TimeZone = FString(), const FCulturePtr& TargetCulture = NULL);
 	[[nodiscard]] static CORE_API FText AsTime(const FDateTime& DateTime, const EDateTimeStyle::Type TimeStyle = EDateTimeStyle::Default, const FString& TimeZone = FString(), const FCulturePtr& TargetCulture = NULL);
+
+	/**
+	 * Generate an FText that represents the given timespan in the current culture, as a delta-time representation (eg, {Hours}:{Minutes}:{Seconds} or {Minutes}:{Seconds}).
+	 */
 	[[nodiscard]] static CORE_API FText AsTimespan(const FTimespan& Timespan, const FCulturePtr& TargetCulture = NULL);
 
 	/**
