[AutoRTFM] Move CVar logic from API.cpp to AutoRTFMUE.cpp

Unlike the bulk of AutoRTFM, this stuff definately belongs in Core.

Part of the work to separate AutoRTFM out of core to break cyclic build dependency issues.

diff --git a/Engine/Source/Runtime/Core/Private/AutoRTFM/API.cpp b/Engine/Source/Runtime/Core/Private/AutoRTFM/API.cpp
--- a/Engine/Source/Runtime/Core/Private/AutoRTFM/API.cpp
+++ b/Engine/Source/Runtime/Core/Private/AutoRTFM/API.cpp
@@ -3,8 +3,7 @@
 #include "AutoRTFM/AutoRTFM.h"
 #include "AutoRTFM/Context.h"
 #include "CoreGlobals.h"
-#include "HAL/IConsoleManager.h"
-#include "GenericPlatform/GenericPlatformCrashContext.h"
+#include "ExternAPI.h"
 #include "HAL/Platform.h"
 #include "Transaction.h"
 
@@ -14,6 +13,8 @@ static_assert(UE_AUTORTFM_ENABLED, "AutoRTFM/API.cpp requires the compiler flag
 
 namespace
 {
+	static constexpr AutoRTFM::EMemoryValidationLevel GDefaultMemoryValidationLevel = AutoRTFM::EMemoryValidationLevel::Disabled;
+
 	// Move this to a local only and use functions to access this
 #if UE_AUTORTFM_ENABLED_RUNTIME_BY_DEFAULT
 	int GAutoRTFMRuntimeEnabled = AutoRTFM::ForTheRuntime::EAutoRTFMEnabledState::AutoRTFM_EnabledByDefault;
@@ -22,19 +23,15 @@ namespace
 #endif // UE_AUTORTFM_ENABLED_RUNTIME_BY_DEFAULT
 
 	int GAutoRTFMInternalAbortAction = AutoRTFM::ForTheRuntime::EAutoRTFMInternalAbortActionState::Crash;
-
 	int GAutoRTFMRetryTransactions = AutoRTFM::ForTheRuntime::EAutoRTFMRetryTransactionState::NoRetry;
 
-	static constexpr AutoRTFM::EMemoryValidationLevel GDefaultMemoryValidationLevel = AutoRTFM::EMemoryValidationLevel::Disabled;
-
 	// Note: GMemoryValidationLevel should never be EMemoryValidationLevel::Default.
 	// EMemoryValidationLevel::Default is a special enumerator that can be used
 	// in the public API to map to GDefaultMemoryValidationLevel.
-	int GMemoryValidationLevel = static_cast<int>(GDefaultMemoryValidationLevel);
+	AutoRTFM::EMemoryValidationLevel GMemoryValidationLevel = GDefaultMemoryValidationLevel;
 	
 	bool GMemoryValidationThrottlingEnabled = true;
 	bool GMemoryValidationStatisticsEnabled = true;
-
 	bool GAutoRTFMEnsureOnInternalAbort = true;
 
 	// A linked-list of open->closed function tables populated by
@@ -42,65 +39,6 @@ namespace
 	// ProcessAllPendingOpenToClosedRegistrations() when autortfm_initialize()
 	// is called.
 	autortfm_open_to_closed_table* GPendingOpenToClosedRegistrations = nullptr;
-
-	void OnIsAutoRTFMEnabledChanged()
-	{
-		const bool bEnabled = AutoRTFM::ForTheRuntime::IsAutoRTFMRuntimeEnabled();
-		UE_LOG(LogAutoRTFM, Log, TEXT("AutoRTFM: %s"), bEnabled ? TEXT("enabled") : TEXT("disabled"));
-		FGenericCrashContext::SetGameData(TEXT("IsAutoRTFMRuntimeEnabled"), bEnabled ? TEXT("true") : TEXT("false"));
-	}
-
-	void OnAutoRTFMRetryTransactionsChanged()
-	{
-		switch (GAutoRTFMRetryTransactions)
-		{
-		case AutoRTFM::ForTheRuntime::EAutoRTFMRetryTransactionState::NoRetry:
-			UE_LOG(LogAutoRTFM, Log, TEXT("AutoRTFM Retry Transactions: disabled"));
-			return FGenericCrashContext::SetGameData(TEXT("AutoRTFMRetryTransactionState"), TEXT("NoRetry"));
-		case AutoRTFM::ForTheRuntime::EAutoRTFMRetryTransactionState::RetryNonNested:
-			UE_LOG(LogAutoRTFM, Log, TEXT("AutoRTFM Retry Transactions: retry-non-nested"));
-			return FGenericCrashContext::SetGameData(TEXT("AutoRTFMRetryTransactionState"), TEXT("RetryNonNested"));
-		case AutoRTFM::ForTheRuntime::EAutoRTFMRetryTransactionState::RetryNestedToo:
-			UE_LOG(LogAutoRTFM, Log, TEXT("AutoRTFM Retry Transactions: retry-nested-too"));
-			return FGenericCrashContext::SetGameData(TEXT("AutoRTFMRetryTransactionState"), TEXT("RetryNestedToo"));
-		}
-	}
-
-	void OnAutoRTFMMemoryValidationLevelChanged()
-	{
-		if (GMemoryValidationLevel == static_cast<int>(AutoRTFM::EMemoryValidationLevel::Default))
-		{
-			// Map EMemoryValidationLevel::Default to GDefaultMemoryValidationLevel.
-			GMemoryValidationLevel = static_cast<int>(GDefaultMemoryValidationLevel);
-		}
-
-		switch (static_cast<AutoRTFM::EMemoryValidationLevel>(GMemoryValidationLevel))
-		{
-		case AutoRTFM::EMemoryValidationLevel::Disabled:
-			UE_LOG(LogAutoRTFM, Log, TEXT("AutoRTFM Memory Validation: disabled"));
-			return FGenericCrashContext::SetGameData(TEXT("AutoRTFMMemoryValidationLevel"), TEXT("Disabled"));
-		case AutoRTFM::EMemoryValidationLevel::Warn:
-			UE_LOG(LogAutoRTFM, Log, TEXT("AutoRTFM Memory Validation: enabled as warning"));
-			return FGenericCrashContext::SetGameData(TEXT("AutoRTFMMemoryValidationLevel"), TEXT("Warn"));
-		case AutoRTFM::EMemoryValidationLevel::Error:
-			UE_LOG(LogAutoRTFM, Log, TEXT("AutoRTFM Memory Validation: enabled as error"));
-			return FGenericCrashContext::SetGameData(TEXT("AutoRTFMMemoryValidationLevel"), TEXT("Enabled"));
-		}
-	}
-
-	void OnAutoRTFMMemoryValidationThrottlingChanged()
-	{
-		const TCHAR* Text = GMemoryValidationThrottlingEnabled ? TEXT("true") : TEXT("false");
-		UE_LOG(LogAutoRTFM, Log, TEXT("AutoRTFM Memory Validation Throttling Enabled: %s"), Text);
-		return FGenericCrashContext::SetGameData(TEXT("AutoRTFMMemoryValidationThrottlingEnabled"), Text);
-	}
-
-	void OnAutoRTFMMemoryValidationStatisticsChanged()
-	{
-		const TCHAR* Text = GMemoryValidationStatisticsEnabled ? TEXT("true") : TEXT("false");
-		UE_LOG(LogAutoRTFM, Log, TEXT("AutoRTFM Memory Validation Statistics Enabled: %s"), Text);
-		return FGenericCrashContext::SetGameData(TEXT("AutoRTFMMemoryValidationStatisticsEnabled"), Text);
-	}
 }
 
 namespace AutoRTFM
@@ -110,126 +48,21 @@ namespace AutoRTFM
 		UE_AUTORTFM_API ForTheRuntime::EAutoRTFMEnabledState ForceSetAutoRTFMRuntime(ForTheRuntime::EAutoRTFMEnabledState State)
 		{
 			const int Original = GAutoRTFMRuntimeEnabled;
-			GAutoRTFMRuntimeEnabled = State;
+			if (GAutoRTFMRuntimeEnabled != State)
+			{
+				GAutoRTFMRuntimeEnabled = State;
+#if UE_AUTORTFM
+				if (GExternAPI.OnRuntimeEnabledChanged)
+				{
+					GExternAPI.OnRuntimeEnabledChanged();
+				}
+#endif
+			}
 			return static_cast<ForTheRuntime::EAutoRTFMEnabledState>(Original);
 		}
 	} // Testing
 } // AutoRTFM
 
-#if UE_AUTORTFM
-FString GAutoRTFM;
-
-void OnAutoRTFMCVarChanged()
-{
-	if (GAutoRTFM.ToLower() == "forceon")
-	{
-		AutoRTFM::ForTheRuntime::SetAutoRTFMRuntime(AutoRTFM::ForTheRuntime::AutoRTFM_ForcedEnabled);
-	}
-	else if (GAutoRTFM.ToLower() == "forceoff")
-	{
-		AutoRTFM::ForTheRuntime::SetAutoRTFMRuntime(AutoRTFM::ForTheRuntime::AutoRTFM_ForcedDisabled);
-	}
-	else if (GAutoRTFM == "1") // The CVar system converts On to '1'.
-	{
-		AutoRTFM::ForTheRuntime::SetAutoRTFMRuntime(AutoRTFM::ForTheRuntime::AutoRTFM_Enabled);
-	}
-	else if (GAutoRTFM == "0") // The CVar system converts Off to '0'.
-	{
-		AutoRTFM::ForTheRuntime::SetAutoRTFMRuntime(AutoRTFM::ForTheRuntime::AutoRTFM_Disabled);
-	}
-	else if (GAutoRTFM == "2")
-	{
-		AutoRTFM::ForTheRuntime::SetAutoRTFMRuntime(AutoRTFM::ForTheRuntime::AutoRTFM_ForcedDisabled);
-	}
-	else if (GAutoRTFM == "3")
-	{
-		AutoRTFM::ForTheRuntime::SetAutoRTFMRuntime(AutoRTFM::ForTheRuntime::AutoRTFM_ForcedEnabled);
-	}
-	else
-	{
-		UE_LOG(LogAutoRTFM, Fatal, TEXT("'AutoRTFMRuntimeEnabled' CVar was set to '%s' which is not one of 'ForceOn', 'ForceOff', 'On', or 'Off'!"), *GAutoRTFM);
-	}
-}
-
-static FAutoConsoleVariableRef CVarAutoRTFMRuntimeEnabled(
-	TEXT("AutoRTFMRuntimeEnabled"),
-	GAutoRTFM,
-	TEXT("Enables the AutoRTFM runtime"),
-	FConsoleVariableDelegate::CreateLambda([] (IConsoleVariable*) { OnAutoRTFMCVarChanged(); }),
-	ECVF_Default
-);
-
-static FAutoConsoleVariableRef CVarAutoRTFMRetryTransactions(
-	TEXT("AutoRTFMRetryTransactions"),
-	GAutoRTFMRetryTransactions,
-	TEXT("Enables the AutoRTFM sanitizer-like mode where we can force an abort-and-retry on transactions (useful to test abort codepaths work as intended)"),
-	FConsoleVariableDelegate::CreateLambda([](IConsoleVariable*) { OnAutoRTFMRetryTransactionsChanged(); }),
-	ECVF_Default
-);
-
-static FAutoConsoleVariableRef CVarAutoRTFMMemoryValidationLevel(
-	TEXT("AutoRTFMMemoryValidationLevel"),
-	GMemoryValidationLevel,
-	TEXT("Detects potential memory corruption due to writes made both by a transaction and open-code"),
-	FConsoleVariableDelegate::CreateLambda([](IConsoleVariable*) { OnAutoRTFMMemoryValidationLevelChanged(); }),
-	ECVF_Default
-);
-
-static FAutoConsoleVariableRef CVarAutoRTFMMemoryValidationThrottlingEnabled(
-	TEXT("AutoRTFMMemoryValidationThrottlingEnabled"),
-	GMemoryValidationThrottlingEnabled,
-	TEXT("Automatically skips memory validation on opens if validation is taking an excessive amount of time"),
-	FConsoleVariableDelegate::CreateLambda([](IConsoleVariable*) { OnAutoRTFMMemoryValidationThrottlingChanged(); }),
-	ECVF_Default
-);
-
-static FAutoConsoleVariableRef CVarAutoRTFMMemoryValidationStatisticsEnabled(
-	TEXT("AutoRTFMMemoryValidationStatisticsEnabled"),
-	GMemoryValidationStatisticsEnabled,
-	TEXT("Disable logging of memory validation statistics"),
-	FConsoleVariableDelegate::CreateLambda([](IConsoleVariable*) { OnAutoRTFMMemoryValidationStatisticsChanged(); }),
-	ECVF_Default
-);
-
-FString GAutoRTFMInternalAbortActionStr;
-
-void OnAutoRTFMInternalAbortActionCVarChanged()
-{
-	if (GAutoRTFMInternalAbortActionStr.ToLower() == "crash")
-	{
-		AutoRTFM::ForTheRuntime::SetInternalAbortAction(AutoRTFM::ForTheRuntime::EAutoRTFMInternalAbortActionState::Crash);
-	}
-	else if (GAutoRTFMInternalAbortActionStr.ToLower() == "retry")
-	{
-		AutoRTFM::ForTheRuntime::SetInternalAbortAction(AutoRTFM::ForTheRuntime::EAutoRTFMInternalAbortActionState::Retry);
-	}
-	else if (GAutoRTFMInternalAbortActionStr.ToLower() == "abort")
-	{
-		AutoRTFM::ForTheRuntime::SetInternalAbortAction(AutoRTFM::ForTheRuntime::EAutoRTFMInternalAbortActionState::Abort);
-	}
-	else
-	{
-		UE_LOG(LogAutoRTFM, Fatal, TEXT("'AutoRTFMInternalAbortAction' CVar was set to '%s' which is not one of 'Crash', 'Retry', or 'Abort'!"), *GAutoRTFMInternalAbortActionStr);
-	}
-}
-
-static FAutoConsoleVariableRef CVarAutoRTFMInternalAbortAction(
-	TEXT("AutoRTFMInternalAbortAction"),
-	GAutoRTFMInternalAbortActionStr,
-	TEXT("If true when we hit an AutoRTFM issue assert over ensuring"),
-	FConsoleVariableDelegate::CreateLambda([](IConsoleVariable*) { OnAutoRTFMInternalAbortActionCVarChanged(); })
-);
-
-static FDelayedAutoRegisterHelper DelayedAutoRegister(EDelayedRegisterRunPhase::EndOfEngineInit, []
-	{
-		OnIsAutoRTFMEnabledChanged();
-		OnAutoRTFMRetryTransactionsChanged();
-		OnAutoRTFMMemoryValidationLevelChanged();
-		OnAutoRTFMMemoryValidationThrottlingChanged();
-		OnAutoRTFMMemoryValidationStatisticsChanged();
-	});
-#endif
-
 namespace AutoRTFM
 {
 	namespace ForTheRuntime
@@ -296,9 +129,14 @@ namespace AutoRTFM
 				}
 			}
 
-			GAutoRTFMRuntimeEnabled = State;
-
-			OnIsAutoRTFMEnabledChanged();
+			if (GAutoRTFMRuntimeEnabled != State)
+			{
+				GAutoRTFMRuntimeEnabled = State;
+				if (GExternAPI.OnRuntimeEnabledChanged)
+				{
+					GExternAPI.OnRuntimeEnabledChanged();
+				}
+			}
 
 			return true;
 #else
@@ -349,8 +187,14 @@ namespace AutoRTFM
 		void SetRetryTransaction(EAutoRTFMRetryTransactionState State)
 		{
 #if UE_AUTORTFM
-			GAutoRTFMRetryTransactions = State;
-			OnAutoRTFMRetryTransactionsChanged();
+			if (GAutoRTFMRetryTransactions != State)
+			{
+				GAutoRTFMRetryTransactions = State;
+				if (GExternAPI.OnRetryTransactionsChanged)
+				{
+					GExternAPI.OnRetryTransactionsChanged();
+				}
+			}
 #endif
 		}
 
@@ -397,7 +241,7 @@ namespace AutoRTFM
 		EMemoryValidationLevel GetMemoryValidationLevel()
 		{
 #if UE_AUTORTFM
-			return static_cast<EMemoryValidationLevel>(GMemoryValidationLevel);
+			return GMemoryValidationLevel;
 #else
 			return EMemoryValidationLevel::Disabled;
 #endif
@@ -406,8 +250,18 @@ namespace AutoRTFM
 		void SetMemoryValidationLevel(EMemoryValidationLevel Level)
 		{
 #if UE_AUTORTFM
-			GMemoryValidationLevel = static_cast<int>(
-				Level == EMemoryValidationLevel::Default ? GDefaultMemoryValidationLevel : Level);
+			if (Level == EMemoryValidationLevel::Default)
+			{
+				Level = GDefaultMemoryValidationLevel;
+			}
+			if (GMemoryValidationLevel != Level)
+			{
+				GMemoryValidationLevel = Level;
+				if (GExternAPI.OnMemoryValidationLevelChanged)
+				{
+					GExternAPI.OnMemoryValidationLevelChanged();
+				}
+			}
 #endif
 		}
 
@@ -423,7 +277,14 @@ namespace AutoRTFM
 		void SetMemoryValidationThrottlingEnabled(bool bEnabled)
 		{
 #if UE_AUTORTFM
-			GMemoryValidationThrottlingEnabled = bEnabled;
+			if (GMemoryValidationThrottlingEnabled != bEnabled)
+			{
+				GMemoryValidationThrottlingEnabled = bEnabled;
+				if (GExternAPI.OnMemoryValidationThrottlingChanged)
+				{
+					GExternAPI.OnMemoryValidationThrottlingChanged();
+				}
+			}
 #endif
 		}
 
@@ -439,7 +300,14 @@ namespace AutoRTFM
 		void SetMemoryValidationStatisticsEnabled(bool bEnabled)
 		{
 #if UE_AUTORTFM
-			GMemoryValidationStatisticsEnabled = bEnabled;
+			if (GMemoryValidationStatisticsEnabled != bEnabled)
+			{
+				GMemoryValidationStatisticsEnabled = bEnabled;
+				if (GExternAPI.OnMemoryValidationStatisticsChanged)
+				{
+					GExternAPI.OnMemoryValidationStatisticsChanged();
+				}
+			}
 #endif
 		}
 
diff --git a/Engine/Source/Runtime/Core/Private/AutoRTFM/AutoRTFMUE.cpp b/Engine/Source/Runtime/Core/Private/AutoRTFM/AutoRTFMUE.cpp
--- a/Engine/Source/Runtime/Core/Private/AutoRTFM/AutoRTFMUE.cpp
+++ b/Engine/Source/Runtime/Core/Private/AutoRTFM/AutoRTFMUE.cpp
@@ -4,8 +4,201 @@
 
 #include "AutoRTFM/AutoRTFMUE.h"
 #include "AutoRTFM/AutoRTFM.h"
+
+#include "Containers/UnrealString.h"
+#include "GenericPlatform/GenericPlatformCrashContext.h"
+#include "HAL/IConsoleManager.h"
 #include "HAL/UnrealMemory.h"
 
+#include "Utils.h" // TODO(SOL-7652) Remove this include once all logging is abstracted
+
+namespace
+{
+
+void OnAutoRTFMRuntimeEnabledChanged()
+{
+	const bool bEnabled = AutoRTFM::ForTheRuntime::IsAutoRTFMRuntimeEnabled();
+	UE_LOG(LogAutoRTFM, Display, TEXT("AutoRTFM: %s"), bEnabled ? TEXT("enabled") : TEXT("disabled"));
+	FGenericCrashContext::SetGameData(TEXT("IsAutoRTFMRuntimeEnabled"), bEnabled ? TEXT("true") : TEXT("false"));
+}
+
+void OnAutoRTFMRetryTransactionsChanged()
+{
+	AutoRTFM::ForTheRuntime::EAutoRTFMRetryTransactionState Value = AutoRTFM::ForTheRuntime::GetRetryTransaction();
+	switch (Value)
+	{
+	case AutoRTFM::ForTheRuntime::EAutoRTFMRetryTransactionState::NoRetry:
+		UE_LOG(LogAutoRTFM, Display, TEXT("AutoRTFM Retry Transactions: disabled"));
+		return FGenericCrashContext::SetGameData(TEXT("AutoRTFMRetryTransactionState"), TEXT("NoRetry"));
+	case AutoRTFM::ForTheRuntime::EAutoRTFMRetryTransactionState::RetryNonNested:
+		UE_LOG(LogAutoRTFM, Display, TEXT("AutoRTFM Retry Transactions: retry-non-nested"));
+		return FGenericCrashContext::SetGameData(TEXT("AutoRTFMRetryTransactionState"), TEXT("RetryNonNested"));
+	case AutoRTFM::ForTheRuntime::EAutoRTFMRetryTransactionState::RetryNestedToo:
+		UE_LOG(LogAutoRTFM, Display, TEXT("AutoRTFM Retry Transactions: retry-nested-too"));
+		return FGenericCrashContext::SetGameData(TEXT("AutoRTFMRetryTransactionState"), TEXT("RetryNestedToo"));
+	}
+}
+
+void OnAutoRTFMMemoryValidationLevelChanged()
+{
+	AutoRTFM::EMemoryValidationLevel Value = AutoRTFM::ForTheRuntime::GetMemoryValidationLevel();
+	switch (Value)
+	{
+	case AutoRTFM::EMemoryValidationLevel::Disabled:
+		UE_LOG(LogAutoRTFM, Display, TEXT("AutoRTFM Memory Validation: disabled"));
+		return FGenericCrashContext::SetGameData(TEXT("AutoRTFMMemoryValidationLevel"), TEXT("Disabled"));
+	case AutoRTFM::EMemoryValidationLevel::Warn:
+		UE_LOG(LogAutoRTFM, Display, TEXT("AutoRTFM Memory Validation: enabled as warning"));
+		return FGenericCrashContext::SetGameData(TEXT("AutoRTFMMemoryValidationLevel"), TEXT("Warn"));
+	case AutoRTFM::EMemoryValidationLevel::Error:
+		UE_LOG(LogAutoRTFM, Display, TEXT("AutoRTFM Memory Validation: enabled as error"));
+		return FGenericCrashContext::SetGameData(TEXT("AutoRTFMMemoryValidationLevel"), TEXT("Enabled"));
+	}
+}
+
+void OnAutoRTFMMemoryValidationThrottlingChanged()
+{
+	bool bEnabled = AutoRTFM::ForTheRuntime::GetMemoryValidationThrottlingEnabled();
+	const TCHAR* Text = bEnabled ? TEXT("true") : TEXT("false");
+	UE_LOG(LogAutoRTFM, Display, TEXT("AutoRTFM Memory Validation Throttling Enabled: %s"), Text);
+	return FGenericCrashContext::SetGameData(TEXT("AutoRTFMMemoryValidationThrottlingEnabled"), Text);
+}
+
+void OnAutoRTFMMemoryValidationStatisticsChanged()
+{
+	bool bEnabled = AutoRTFM::ForTheRuntime::GetMemoryValidationStatisticsEnabled();
+	const TCHAR* Text = bEnabled ? TEXT("true") : TEXT("false");
+	UE_LOG(LogAutoRTFM, Display, TEXT("AutoRTFM Memory Validation Statistics Enabled: %s"), Text);
+	return FGenericCrashContext::SetGameData(TEXT("AutoRTFMMemoryValidationStatisticsEnabled"), Text);
+}
+
+static FAutoConsoleVariable CVarAutoRTFMRuntimeEnabled(
+	TEXT("AutoRTFMRuntimeEnabled"),
+	TEXT("default"),
+	TEXT("Enables the AutoRTFM runtime"),
+	FConsoleVariableDelegate::CreateLambda([] (IConsoleVariable* Variable)
+	{ 
+		FString Value = Variable->GetString().ToLower();
+		if (Value == "default")
+		{
+			return;
+		}
+		if (Value == "forceon")
+		{
+			AutoRTFM::ForTheRuntime::SetAutoRTFMRuntime(AutoRTFM::ForTheRuntime::AutoRTFM_ForcedEnabled);
+			return;
+		}
+		if (Value == "forceoff")
+		{
+			AutoRTFM::ForTheRuntime::SetAutoRTFMRuntime(AutoRTFM::ForTheRuntime::AutoRTFM_ForcedDisabled);
+			return;
+		}
+		if (Value == "1") // The CVar system converts On to '1'.
+		{
+			AutoRTFM::ForTheRuntime::SetAutoRTFMRuntime(AutoRTFM::ForTheRuntime::AutoRTFM_Enabled);
+			return;
+		}
+		if (Value == "0") // The CVar system converts Off to '0'.
+		{
+			AutoRTFM::ForTheRuntime::SetAutoRTFMRuntime(AutoRTFM::ForTheRuntime::AutoRTFM_Disabled);
+			return;
+		}
+		if (Value == "2")
+		{
+			AutoRTFM::ForTheRuntime::SetAutoRTFMRuntime(AutoRTFM::ForTheRuntime::AutoRTFM_ForcedDisabled);
+			return;
+		}
+		if (Value == "3")
+		{
+			AutoRTFM::ForTheRuntime::SetAutoRTFMRuntime(AutoRTFM::ForTheRuntime::AutoRTFM_ForcedEnabled);
+			return;
+		}
+		UE_LOG(LogAutoRTFM, Fatal, TEXT("'AutoRTFMRuntimeEnabled' CVar was set to '%s' which is not one of 'ForceOn', 'ForceOff', 'On', or 'Off'!"), *Value);
+	}),
+	ECVF_Default
+);
+
+static FAutoConsoleVariable CVarAutoRTFMInternalAbortAction(
+	TEXT("AutoRTFMInternalAbortAction"),
+	TEXT("default"),
+	TEXT("If true when we hit an AutoRTFM issue assert over ensuring"),
+	FConsoleVariableDelegate::CreateLambda([](IConsoleVariable* Variable)
+	{
+		FString Value = Variable->GetString();
+		if (Value == "default")
+		{
+			return;
+		}
+		if (Value == "crash")
+		{
+			AutoRTFM::ForTheRuntime::SetInternalAbortAction(AutoRTFM::ForTheRuntime::EAutoRTFMInternalAbortActionState::Crash);
+			return;
+		}
+		if (Value == "retry")
+		{
+			AutoRTFM::ForTheRuntime::SetInternalAbortAction(AutoRTFM::ForTheRuntime::EAutoRTFMInternalAbortActionState::Retry);
+			return;
+		}
+		if (Value == "abort")
+		{
+			AutoRTFM::ForTheRuntime::SetInternalAbortAction(AutoRTFM::ForTheRuntime::EAutoRTFMInternalAbortActionState::Abort);
+			return;
+		}
+		UE_LOG(LogAutoRTFM, Fatal, TEXT("'AutoRTFMInternalAbortAction' CVar was set to '%s' which is not one of 'Crash', 'Retry', or 'Abort'!"), *Value);
+
+	})
+);
+
+static FAutoConsoleVariable CVarAutoRTFMRetryTransactions(
+	TEXT("AutoRTFMRetryTransactions"),
+	static_cast<int>(AutoRTFM::ForTheRuntime::GetRetryTransaction()),
+	TEXT("Enables the AutoRTFM sanitizer-like mode where we can force an abort-and-retry on transactions (useful to test abort codepaths work as intended)"),
+	FConsoleVariableDelegate::CreateLambda([](IConsoleVariable* Variable)
+	{
+		AutoRTFM::ForTheRuntime::EAutoRTFMRetryTransactionState Value = static_cast<AutoRTFM::ForTheRuntime::EAutoRTFMRetryTransactionState>(Variable->GetInt());
+		AutoRTFM::ForTheRuntime::SetRetryTransaction(Value);
+	}),
+	ECVF_Default
+);
+
+static FAutoConsoleVariable CVarAutoRTFMMemoryValidationLevel(
+	TEXT("AutoRTFMMemoryValidationLevel"),
+	static_cast<int>(AutoRTFM::ForTheRuntime::GetMemoryValidationLevel()),
+	TEXT("Detects potential memory corruption due to writes made both by a transaction and open-code"),
+	FConsoleVariableDelegate::CreateLambda([](IConsoleVariable* Variable)
+	{
+		int Value = Variable->GetInt();
+		AutoRTFM::ForTheRuntime::SetMemoryValidationLevel(static_cast<AutoRTFM::EMemoryValidationLevel>(Value));
+	}),
+	ECVF_Default
+);
+
+static FAutoConsoleVariable CVarAutoRTFMMemoryValidationThrottlingEnabled(
+	TEXT("AutoRTFMMemoryValidationThrottlingEnabled"),
+	AutoRTFM::ForTheRuntime::GetMemoryValidationThrottlingEnabled(),
+	TEXT("Automatically skips memory validation on opens if validation is taking an excessive amount of time"),
+	FConsoleVariableDelegate::CreateLambda([](IConsoleVariable* Variable)
+	{
+		bool Value = Variable->GetBool();
+		AutoRTFM::ForTheRuntime::SetMemoryValidationThrottlingEnabled(Value);
+	}),
+	ECVF_Default
+);
+
+static FAutoConsoleVariable CVarAutoRTFMMemoryValidationStatisticsEnabled(
+	TEXT("AutoRTFMMemoryValidationStatisticsEnabled"),
+	AutoRTFM::ForTheRuntime::GetMemoryValidationStatisticsEnabled(),
+	TEXT("Disable logging of memory validation statistics"),
+	FConsoleVariableDelegate::CreateLambda([](IConsoleVariable* Variable)
+	{
+		bool Value = Variable->GetBool();
+		AutoRTFM::ForTheRuntime::SetMemoryValidationStatisticsEnabled(Value);
+	}),
+	ECVF_Default
+);
+
+} // anonymous namespace
+
 void AutoRTFM::InitializeForUE()
 {
 	AutoRTFM::ForTheRuntime::FExternAPI ExternAPI;
@@ -22,7 +215,21 @@ void AutoRTFM::InitializeForUE()
 		return FMemory::MallocZeroed(Size, Alignment);
 	};
 	ExternAPI.Free = FMemory::Free;
+	ExternAPI.OnRuntimeEnabledChanged = OnAutoRTFMRuntimeEnabledChanged;
+	ExternAPI.OnRetryTransactionsChanged = OnAutoRTFMRetryTransactionsChanged;
+	ExternAPI.OnMemoryValidationLevelChanged = OnAutoRTFMMemoryValidationLevelChanged;
+	ExternAPI.OnMemoryValidationThrottlingChanged = OnAutoRTFMMemoryValidationThrottlingChanged;
+	ExternAPI.OnMemoryValidationStatisticsChanged = OnAutoRTFMMemoryValidationStatisticsChanged;
+
 	AutoRTFM::ForTheRuntime::Initialize(ExternAPI);
+	
+	// Call the OnAutoRTFMXXXChanged() handlers now so that values are logged
+	// and crash context data is set.
+	OnAutoRTFMRuntimeEnabledChanged();
+	OnAutoRTFMRetryTransactionsChanged();
+	OnAutoRTFMMemoryValidationLevelChanged();
+	OnAutoRTFMMemoryValidationThrottlingChanged();
+	OnAutoRTFMMemoryValidationStatisticsChanged();
 }
 
 #endif  // (defined(__AUTORTFM) && __AUTORTFM)
diff --git a/Engine/Source/Runtime/Core/Public/AutoRTFM/AutoRTFM.h b/Engine/Source/Runtime/Core/Public/AutoRTFM/AutoRTFM.h
--- a/Engine/Source/Runtime/Core/Public/AutoRTFM/AutoRTFM.h
+++ b/Engine/Source/Runtime/Core/Public/AutoRTFM/AutoRTFM.h
@@ -141,17 +141,38 @@ typedef enum
 typedef struct
 {
 	// The function used to allocate memory from the heap.
-	// If null, then a default implementation will be used.
+	// Must not be null.
 	void* (*Allocate)(size_t Size, size_t Alignment);
 	// The function used to reallocate memory from the heap.
-	// If null, then a default implementation will be used.
+	// Must not be null.
 	void* (*Reallocate)(void* Pointer, size_t Size, size_t Alignment);
 	// The function used to allocate zeroed memory from the heap.
-	// If null, then a default implementation will be used.
+	// Must not be null.
 	void* (*AllocateZeroed)(size_t Size, size_t Alignment);
 	// The function used to free memory allocated by Allocate() and AllocateZeroed().
-	// If null, then a default implementation will be used.
+	// Must not be null.
 	void (*Free)(void* Pointer);
+	
+	// Optional callback to be informed when the value returned by
+	// ForTheRuntime::IsAutoRTFMRuntimeEnabled() changes.
+	// Can be null.
+	void (*OnRuntimeEnabledChanged)();
+	// Optional callback to be informed when the value returned by
+	// ForTheRuntime::GetRetryTransaction() changes.
+	// Can be null.
+	void (*OnRetryTransactionsChanged)();
+	// Optional callback to be informed when the value returned by
+	// ForTheRuntime::GetMemoryValidationLevel() changes.
+	// Can be null.
+	void (*OnMemoryValidationLevelChanged)();
+	// Optional callback to be informed when the value returned by
+	// ForTheRuntime::GetMemoryValidationThrottlingEnabled() changes.
+	// Can be null.
+	void (*OnMemoryValidationThrottlingChanged)();
+	// Optional callback to be informed when the value returned by
+	// ForTheRuntime::GetMemoryValidationStatisticsEnabled() changes.
+	// Can be null.
+	void (*OnMemoryValidationStatisticsChanged)();
 } autortfm_extern_api;
 
 #if UE_AUTORTFM_ENABLED
