Fix VisionOS deprecation warning and TVOS compilation (rotation APIs not available there).

diff --git a/Engine/Source/Runtime/Core/Private/IOS/IOSPlatformMisc.cpp b/Engine/Source/Runtime/Core/Private/IOS/IOSPlatformMisc.cpp
--- a/Engine/Source/Runtime/Core/Private/IOS/IOSPlatformMisc.cpp
+++ b/Engine/Source/Runtime/Core/Private/IOS/IOSPlatformMisc.cpp
@@ -376,18 +376,22 @@ void FIOSPlatformMisc::SetAllowedDeviceOrientation(EDeviceScreenOrientation NewA
 	
 	// Notify the view controller about a change in supported interface orientations.
 	
+#if !PLATFORM_TVOS
 	dispatch_async(dispatch_get_main_queue(),
 		^{
-			if (@available(iOS 16.0, *))
+			if (@available(iOS 16.0, visionOS 1.0, *))
 			{
 				[[IOSAppDelegate GetDelegate].IOSController setNeedsUpdateOfSupportedInterfaceOrientations];
 			}
+#if !PLATFORM_VISIONOS
 			else
 			{
 				[UIViewController attemptRotationToDeviceOrientation];
 			}
+#endif
 		}
 	);
+#endif
 }
 
 bool FIOSPlatformMisc::HasPlatformFeature(const TCHAR* FeatureName)
