Mobile Vulkan: Support CVarAllowFrameTimestamps, which was previously only considered by GLES.

[RN]

diff --git a/Engine/Source/Runtime/Core/Private/Android/AndroidPlatformMisc.cpp b/Engine/Source/Runtime/Core/Private/Android/AndroidPlatformMisc.cpp
--- a/Engine/Source/Runtime/Core/Private/Android/AndroidPlatformMisc.cpp
+++ b/Engine/Source/Runtime/Core/Private/Android/AndroidPlatformMisc.cpp
@@ -2331,42 +2331,45 @@ static EDeviceVulkanSupportStatus AttemptVulkanInit(void* VulkanLib)
 
 	//for now we are allowing devices without the timing extension to run with a basic CPU frame pacer.
 #if 0
-	bool bHasVKGoogleDisplayTiming = false;
+	if (FAndroidPlatformRHIFramePacer::CVarAllowFrameTimestamps.GetValueOnAnyThread() != 0)
 	{
-		TArray<VkExtensionProperties> ExtensionProps;
-		do
+		bool bHasVKGoogleDisplayTiming = false;
 		{
-			uint32 Count = 0;
-			Result = vkEnumerateDeviceExtensionProperties(PhysicalDevices[0], nullptr, &Count, nullptr);
-			check(Result >= VK_SUCCESS);
-
-			if (Count > 0)
+			TArray<VkExtensionProperties> ExtensionProps;
+			do
 			{
-				ExtensionProps.Empty(Count);
-				ExtensionProps.AddUninitialized(Count);
-				Result = vkEnumerateDeviceExtensionProperties(PhysicalDevices[0], nullptr, &Count, ExtensionProps.GetData());
+				uint32 Count = 0;
+				Result = vkEnumerateDeviceExtensionProperties(PhysicalDevices[0], nullptr, &Count, nullptr);
 				check(Result >= VK_SUCCESS);
-			}
-		} while (Result == VK_INCOMPLETE);		
-		check(Result >= VK_SUCCESS);
 
-		
-		for (int32 i = 0; i < ExtensionProps.Num(); ++i)
-		{
-			UE_LOG(LogAndroid, Log, TEXT("Checking extension: %s."), ANSI_TO_TCHAR(ExtensionProps[i].extensionName));
-			if (!FCStringAnsi::Strcmp(VK_GOOGLE_DISPLAY_TIMING_EXTENSION_NAME, ExtensionProps[i].extensionName))
+				if (Count > 0)
+				{
+					ExtensionProps.Empty(Count);
+					ExtensionProps.AddUninitialized(Count);
+					Result = vkEnumerateDeviceExtensionProperties(PhysicalDevices[0], nullptr, &Count, ExtensionProps.GetData());
+					check(Result >= VK_SUCCESS);
+				}
+			} while (Result == VK_INCOMPLETE);		
+			check(Result >= VK_SUCCESS);
+
+			
+			for (int32 i = 0; i < ExtensionProps.Num(); ++i)
 			{
-				bHasVKGoogleDisplayTiming = true;
-				break;
+				UE_LOG(LogAndroid, Log, TEXT("Checking extension: %s."), ANSI_TO_TCHAR(ExtensionProps[i].extensionName));
+				if (!FCStringAnsi::Strcmp(VK_GOOGLE_DISPLAY_TIMING_EXTENSION_NAME, ExtensionProps[i].extensionName))
+				{
+					bHasVKGoogleDisplayTiming = true;
+					break;
+				}
 			}
 		}
-	}
-	if (!bHasVKGoogleDisplayTiming)
-	{
-		vkDestroyInstance(Instance, nullptr);
+		if (!bHasVKGoogleDisplayTiming)
+		{
+			vkDestroyInstance(Instance, nullptr);
 
-		UE_LOG(LogAndroid, Log, TEXT("Vulkan not supported, cannot find VK_GOOGLE_display_timing extension."));
-		return EDeviceVulkanSupportStatus::NotSupported;
+			UE_LOG(LogAndroid, Log, TEXT("Vulkan not supported, cannot find VK_GOOGLE_display_timing extension."));
+			return EDeviceVulkanSupportStatus::NotSupported;
+		}
 	}
 #endif
 
diff --git a/Engine/Source/Runtime/VulkanRHI/Private/Android/VulkanAndroidPlatform.cpp b/Engine/Source/Runtime/VulkanRHI/Private/Android/VulkanAndroidPlatform.cpp
--- a/Engine/Source/Runtime/VulkanRHI/Private/Android/VulkanAndroidPlatform.cpp
+++ b/Engine/Source/Runtime/VulkanRHI/Private/Android/VulkanAndroidPlatform.cpp
@@ -434,7 +434,11 @@ void FVulkanAndroidPlatform::GetInstanceExtensions(FVulkanInstanceExtensionArray
 	OutExtensions.Add(MakeUnique<FVulkanInstanceExtension>(VK_KHR_ANDROID_SURFACE_EXTENSION_NAME, VULKAN_EXTENSION_ENABLED, VULKAN_EXTENSION_NOT_PROMOTED));
 
 	// VK_GOOGLE_display_timing (as instance extension?)
-	OutExtensions.Add(MakeUnique<FVulkanInstanceExtension>(VK_GOOGLE_DISPLAY_TIMING_EXTENSION_NAME, VULKAN_EXTENSION_ENABLED, VULKAN_EXTENSION_NOT_PROMOTED));
+	if (FAndroidPlatformRHIFramePacer::CVarAllowFrameTimestamps.GetValueOnAnyThread() != 0)
+	{
+		OutExtensions.Add(
+			MakeUnique<FVulkanInstanceExtension>(VK_GOOGLE_DISPLAY_TIMING_EXTENSION_NAME, VULKAN_EXTENSION_ENABLED, VULKAN_EXTENSION_NOT_PROMOTED));
+	}
 }
 
 void FVulkanAndroidPlatform::GetInstanceLayers(TArray<const ANSICHAR*>& OutLayers)
@@ -474,7 +478,11 @@ static FAutoConsoleVariableRef CVarVulkanUseASTCDecodeMode(
 void FVulkanAndroidPlatform::GetDeviceExtensions(FVulkanDevice* Device, FVulkanDeviceExtensionArray& OutExtensions)
 {
 	OutExtensions.Add(MakeUnique<FVulkanDeviceExtension>(Device, VK_KHR_ANDROID_SURFACE_EXTENSION_NAME, VULKAN_EXTENSION_ENABLED, VULKAN_EXTENSION_NOT_PROMOTED));
-	OutExtensions.Add(MakeUnique<FVulkanDeviceExtension>(Device, VK_GOOGLE_DISPLAY_TIMING_EXTENSION_NAME, VULKAN_EXTENSION_ENABLED, VULKAN_EXTENSION_NOT_PROMOTED));
+	if (FAndroidPlatformRHIFramePacer::CVarAllowFrameTimestamps.GetValueOnAnyThread() != 0)
+	{
+		OutExtensions.Add(MakeUnique<FVulkanDeviceExtension>(
+			Device, VK_GOOGLE_DISPLAY_TIMING_EXTENSION_NAME, VULKAN_EXTENSION_ENABLED, VULKAN_EXTENSION_NOT_PROMOTED));
+	}
 	if (GVulkanUseASTCDecodeMode)
 	{
 		OutExtensions.Add(MakeUnique<FVulkanDeviceExtension>(Device, VK_EXT_ASTC_DECODE_MODE_EXTENSION_NAME, VULKAN_SUPPORTS_ASTC_DECODE_MODE, VULKAN_EXTENSION_NOT_PROMOTED, DEVICE_EXT_FLAG_SETTER(HasEXTASTCDecodeMode)));
