[AutoRTFM] Lower logging level for AutoRTFM initialization messages

Reduces test log spam.

diff --git a/Engine/Source/Runtime/Core/Private/AutoRTFM/AutoRTFMUE.cpp b/Engine/Source/Runtime/Core/Private/AutoRTFM/AutoRTFMUE.cpp
--- a/Engine/Source/Runtime/Core/Private/AutoRTFM/AutoRTFMUE.cpp
+++ b/Engine/Source/Runtime/Core/Private/AutoRTFM/AutoRTFMUE.cpp
@@ -18,7 +18,7 @@ namespace
 void OnAutoRTFMRuntimeEnabledChanged()
 {
 	const bool bEnabled = AutoRTFM::ForTheRuntime::IsAutoRTFMRuntimeEnabled();
-	UE_LOG(LogAutoRTFM, Display, TEXT("AutoRTFM: %s"), bEnabled ? TEXT("enabled") : TEXT("disabled"));
+	UE_LOG(LogAutoRTFM, Log, TEXT("AutoRTFM: %s"), bEnabled ? TEXT("enabled") : TEXT("disabled"));
 	FGenericCrashContext::SetGameData(TEXT("IsAutoRTFMRuntimeEnabled"), bEnabled ? TEXT("true") : TEXT("false"));
 }
 
@@ -28,13 +28,13 @@ void OnAutoRTFMRetryTransactionsChanged()
 	switch (Value)
 	{
 	case AutoRTFM::ForTheRuntime::EAutoRTFMRetryTransactionState::NoRetry:
-		UE_LOG(LogAutoRTFM, Display, TEXT("AutoRTFM Retry Transactions: disabled"));
+		UE_LOG(LogAutoRTFM, Log, TEXT("AutoRTFM Retry Transactions: disabled"));
 		return FGenericCrashContext::SetGameData(TEXT("AutoRTFMRetryTransactionState"), TEXT("NoRetry"));
 	case AutoRTFM::ForTheRuntime::EAutoRTFMRetryTransactionState::RetryNonNested:
-		UE_LOG(LogAutoRTFM, Display, TEXT("AutoRTFM Retry Transactions: retry-non-nested"));
+		UE_LOG(LogAutoRTFM, Log, TEXT("AutoRTFM Retry Transactions: retry-non-nested"));
 		return FGenericCrashContext::SetGameData(TEXT("AutoRTFMRetryTransactionState"), TEXT("RetryNonNested"));
 	case AutoRTFM::ForTheRuntime::EAutoRTFMRetryTransactionState::RetryNestedToo:
-		UE_LOG(LogAutoRTFM, Display, TEXT("AutoRTFM Retry Transactions: retry-nested-too"));
+		UE_LOG(LogAutoRTFM, Log, TEXT("AutoRTFM Retry Transactions: retry-nested-too"));
 		return FGenericCrashContext::SetGameData(TEXT("AutoRTFMRetryTransactionState"), TEXT("RetryNestedToo"));
 	}
 }
@@ -45,13 +45,13 @@ void OnAutoRTFMMemoryValidationLevelChanged()
 	switch (Value)
 	{
 	case AutoRTFM::EMemoryValidationLevel::Disabled:
-		UE_LOG(LogAutoRTFM, Display, TEXT("AutoRTFM Memory Validation: disabled"));
+		UE_LOG(LogAutoRTFM, Log, TEXT("AutoRTFM Memory Validation: disabled"));
 		return FGenericCrashContext::SetGameData(TEXT("AutoRTFMMemoryValidationLevel"), TEXT("Disabled"));
 	case AutoRTFM::EMemoryValidationLevel::Warn:
-		UE_LOG(LogAutoRTFM, Display, TEXT("AutoRTFM Memory Validation: enabled as warning"));
+		UE_LOG(LogAutoRTFM, Log, TEXT("AutoRTFM Memory Validation: enabled as warning"));
 		return FGenericCrashContext::SetGameData(TEXT("AutoRTFMMemoryValidationLevel"), TEXT("Warn"));
 	case AutoRTFM::EMemoryValidationLevel::Error:
-		UE_LOG(LogAutoRTFM, Display, TEXT("AutoRTFM Memory Validation: enabled as error"));
+		UE_LOG(LogAutoRTFM, Log, TEXT("AutoRTFM Memory Validation: enabled as error"));
 		return FGenericCrashContext::SetGameData(TEXT("AutoRTFMMemoryValidationLevel"), TEXT("Enabled"));
 	}
 }
@@ -60,7 +60,7 @@ void OnAutoRTFMMemoryValidationThrottlingChanged()
 {
 	bool bEnabled = AutoRTFM::ForTheRuntime::GetMemoryValidationThrottlingEnabled();
 	const TCHAR* Text = bEnabled ? TEXT("true") : TEXT("false");
-	UE_LOG(LogAutoRTFM, Display, TEXT("AutoRTFM Memory Validation Throttling Enabled: %s"), Text);
+	UE_LOG(LogAutoRTFM, Log, TEXT("AutoRTFM Memory Validation Throttling Enabled: %s"), Text);
 	return FGenericCrashContext::SetGameData(TEXT("AutoRTFMMemoryValidationThrottlingEnabled"), Text);
 }
 
@@ -68,7 +68,7 @@ void OnAutoRTFMMemoryValidationStatisticsChanged()
 {
 	bool bEnabled = AutoRTFM::ForTheRuntime::GetMemoryValidationStatisticsEnabled();
 	const TCHAR* Text = bEnabled ? TEXT("true") : TEXT("false");
-	UE_LOG(LogAutoRTFM, Display, TEXT("AutoRTFM Memory Validation Statistics Enabled: %s"), Text);
+	UE_LOG(LogAutoRTFM, Log, TEXT("AutoRTFM Memory Validation Statistics Enabled: %s"), Text);
 	return FGenericCrashContext::SetGameData(TEXT("AutoRTFMMemoryValidationStatisticsEnabled"), Text);
 }
 
