Made FUtf8String::Printf and FUtf8String::Format accept UTF8CHAR and ANSICHAR formats.

diff --git a/Engine/Source/Runtime/Core/Private/Misc/StringFormatter.cpp b/Engine/Source/Runtime/Core/Private/Misc/StringFormatter.cpp
--- a/Engine/Source/Runtime/Core/Private/Misc/StringFormatter.cpp
+++ b/Engine/Source/Runtime/Core/Private/Misc/StringFormatter.cpp
@@ -511,37 +511,37 @@ TStringFormatter<CharType>& GetDefaultFormatter()
 	return DefaultFormatter;
 }
 
-FAnsiString FAnsiString::Format(const ANSICHAR* InFormatString, const FAnsiStringFormatNamedArguments& InNamedArguments)
+FAnsiString FAnsiString::FormatImpl(const ANSICHAR* InFormatString, const FAnsiStringFormatNamedArguments& InNamedArguments)
 {
 	TStringFormatter<ANSICHAR>& DefaultFormatter = GetDefaultFormatter<ANSICHAR>();
 	return DefaultFormatter.Format(InFormatString, InNamedArguments);
 }
 
-FAnsiString FAnsiString::Format(const ANSICHAR* InFormatString, const FStringFormatOrderedArguments& InOrderedArguments)
+FAnsiString FAnsiString::FormatImpl(const ANSICHAR* InFormatString, const FStringFormatOrderedArguments& InOrderedArguments)
 {
 	TStringFormatter<ANSICHAR>& DefaultFormatter = GetDefaultFormatter<ANSICHAR>();
 	return DefaultFormatter.Format(InFormatString, InOrderedArguments);
 }
 
-FString FString::Format(const TCHAR* InFormatString, const FStringFormatNamedArguments& InNamedArguments)
+FString FString::FormatImpl(const TCHAR* InFormatString, const FStringFormatNamedArguments& InNamedArguments)
 {
 	TStringFormatter<TCHAR>& DefaultFormatter = GetDefaultFormatter<TCHAR>();
 	return DefaultFormatter.Format(InFormatString, InNamedArguments);
 }
 
-FString FString::Format(const TCHAR* InFormatString, const FStringFormatOrderedArguments& InOrderedArguments)
+FString FString::FormatImpl(const TCHAR* InFormatString, const FStringFormatOrderedArguments& InOrderedArguments)
 {
 	TStringFormatter<TCHAR>& DefaultFormatter = GetDefaultFormatter<TCHAR>();
 	return DefaultFormatter.Format(InFormatString, InOrderedArguments);
 }
 
-FUtf8String FUtf8String::Format(const UTF8CHAR* InFormatString, const FUtf8StringFormatNamedArguments& InNamedArguments)
+FUtf8String FUtf8String::FormatImpl(const UTF8CHAR* InFormatString, const FUtf8StringFormatNamedArguments& InNamedArguments)
 {
 	TStringFormatter<UTF8CHAR>& DefaultFormatter = GetDefaultFormatter<UTF8CHAR>();
 	return DefaultFormatter.Format(InFormatString, InNamedArguments);
 }
 
-FUtf8String FUtf8String::Format(const UTF8CHAR* InFormatString, const FStringFormatOrderedArguments& InOrderedArguments)
+FUtf8String FUtf8String::FormatImpl(const UTF8CHAR* InFormatString, const FStringFormatOrderedArguments& InOrderedArguments)
 {
 	TStringFormatter<UTF8CHAR>& DefaultFormatter = GetDefaultFormatter<UTF8CHAR>();
 	return DefaultFormatter.Format(InFormatString, InOrderedArguments);
diff --git a/Engine/Source/Runtime/Core/Public/Containers/UnrealString.h.inl b/Engine/Source/Runtime/Core/Public/Containers/UnrealString.h.inl
--- a/Engine/Source/Runtime/Core/Public/Containers/UnrealString.h.inl
+++ b/Engine/Source/Runtime/Core/Public/Containers/UnrealString.h.inl
@@ -58,6 +58,11 @@ public:
 	using ElementType   = UE_STRING_CHARTYPE;
 	using FmtCharType   = UE_STRING_PRINTF_FMT_CHARTYPE;
 
+#ifdef UE_STRING_PRINTF_FMT_CHARTYPE2
+	// Alternate formatting string type, particularly used for FUtf8String which can accept both "..." and UTF8TEXT("...") for formatting strings
+	using FmtCharType2 = UE_STRING_PRINTF_FMT_CHARTYPE2;
+#endif
+
 private:
 	/** Array holding the character data */
 	typedef TArray<ElementType, AllocatorType> DataType;
@@ -1571,6 +1576,13 @@ public:
 	{
 		return PrintfImpl((const ElementType*)Fmt.FormatString, Args...);
 	}
+#ifdef UE_STRING_PRINTF_FMT_CHARTYPE2
+	template <typename... Types>
+	[[nodiscard]] static UE_STRING_CLASS Printf(UE::Core::TCheckedFormatString<FmtCharType2, Types...> Fmt, Types... Args)
+	{
+		return PrintfImpl((const ElementType*)Fmt.FormatString, Args...);
+	}
+#endif
 
 	/**
 	 * Just like Printf, but appends the formatted text to the existing string instead.
@@ -1582,6 +1594,14 @@ public:
 		AppendfImpl(*this, (const ElementType*)Fmt.FormatString, Args...);
 		return *this;
 	}
+#ifdef UE_STRING_PRINTF_FMT_CHARTYPE2
+	template <typename... Types>
+	UE_STRING_CLASS& Appendf(UE::Core::TCheckedFormatString<FmtCharType2, Types...> Fmt, Types... Args)
+	{
+		AppendfImpl(*this, (const ElementType*)Fmt.FormatString, Args...);
+		return *this;
+	}
+#endif
 
 private:
 	static CORE_API UE_STRING_CLASS VARARGS PrintfImpl(const ElementType* Fmt, ...);
@@ -1594,7 +1614,16 @@ public:
 	 * @param InNamedArguments		A map of named arguments that match the tokens specified in InExpression
 	 * @return A string containing the formatted text
 	 */
-	[[nodiscard]] static CORE_API UE_STRING_CLASS Format(const ElementType* InFormatString, const PREPROCESSOR_JOIN(UE_STRING_CLASS, FormatNamedArguments)& InNamedArguments);
+	[[nodiscard]] static UE_STRING_CLASS Format(const FmtCharType* InFormatString, const PREPROCESSOR_JOIN(UE_STRING_CLASS, FormatNamedArguments)& InNamedArguments)
+	{
+		return FormatImpl((const ElementType*)InFormatString, InNamedArguments);
+	}
+#ifdef UE_STRING_PRINTF_FMT_CHARTYPE2
+	[[nodiscard]] static UE_STRING_CLASS Format(const FmtCharType2* InFormatString, const PREPROCESSOR_JOIN(UE_STRING_CLASS, FormatNamedArguments)& InNamedArguments)
+	{
+		return FormatImpl((const ElementType*)InFormatString, InNamedArguments);
+	}
+#endif
 
 	/**
 	 * Format the specified string using the specified arguments. Replaces instances of {0} with indices from the given array matching the index specified in the token
@@ -1602,7 +1631,21 @@ public:
 	 * @param InOrderedArguments	An array of ordered arguments that match the tokens specified in InExpression
 	 * @return A string containing the formatted text
 	 */
-	[[nodiscard]] static CORE_API UE_STRING_CLASS Format(const ElementType* InFormatString, const FStringFormatOrderedArguments& InOrderedArguments);
+	[[nodiscard]] static UE_STRING_CLASS Format(const FmtCharType* InFormatString, const FStringFormatOrderedArguments& InOrderedArguments)
+	{
+		return FormatImpl((const ElementType*)InFormatString, InOrderedArguments);
+	}
+#ifdef UE_STRING_PRINTF_FMT_CHARTYPE2
+	[[nodiscard]] static UE_STRING_CLASS Format(const FmtCharType2* InFormatString, const FStringFormatOrderedArguments& InOrderedArguments)
+	{
+		return FormatImpl((const ElementType*)InFormatString, InOrderedArguments);
+	}
+#endif
+
+private:
+	static CORE_API UE_STRING_CLASS FormatImpl(const ElementType* InFormatString, const PREPROCESSOR_JOIN(UE_STRING_CLASS, FormatNamedArguments)& InNamedArguments);
+	static CORE_API UE_STRING_CLASS FormatImpl(const ElementType* InFormatString, const FStringFormatOrderedArguments& InOrderedArguments);
+public:
 
 	/** Returns a string containing only the Ch character */
 	[[nodiscard]] static CORE_API UE_STRING_CLASS Chr( ElementType Ch );
diff --git a/Engine/Source/Runtime/Core/Public/Containers/Utf8String.h b/Engine/Source/Runtime/Core/Public/Containers/Utf8String.h
--- a/Engine/Source/Runtime/Core/Public/Containers/Utf8String.h
+++ b/Engine/Source/Runtime/Core/Public/Containers/Utf8String.h
@@ -11,9 +11,11 @@
 #define UE_STRING_CHARTYPE                     UTF8CHAR
 #define UE_STRING_CHARTYPE_IS_TCHAR            0
 #define UE_STRING_PRINTF_FMT_CHARTYPE          ANSICHAR
+#define UE_STRING_PRINTF_FMT_CHARTYPE2         UTF8CHAR
 #define UE_STRING_DEPRECATED(Version, Message) UE_DEPRECATED(Version, Message)
 	#include "Containers/UnrealString.h.inl"
 #undef UE_STRING_DEPRECATED
+#undef UE_STRING_PRINTF_FMT_CHARTYPE2
 #undef UE_STRING_PRINTF_FMT_CHARTYPE
 #undef UE_STRING_CHARTYPE_IS_TCHAR
 #undef UE_STRING_CHARTYPE
