Fix a spelling issue in the new AutoRTFM runtime API.

diff --git a/Engine/Source/Runtime/Core/Private/AutoRTFM/IntrusivePool.h b/Engine/Source/Runtime/Core/Private/AutoRTFM/IntrusivePool.h
--- a/Engine/Source/Runtime/Core/Private/AutoRTFM/IntrusivePool.h
+++ b/Engine/Source/Runtime/Core/Private/AutoRTFM/IntrusivePool.h
@@ -28,7 +28,7 @@ public:
 	}
 
 	// Acquires a new Item from the pool. If the item was previously returned
-	// to the pool, we'll call `Ressurect` on it with `Arguments`. Otherwise
+	// to the pool, we'll call `Resurrect` on it with `Arguments`. Otherwise
 	// create a new item with `Arguments`.
 	template<typename ... ArgumentTypes>
 	FItem* Take(ArgumentTypes&& ... Arguments)
@@ -41,7 +41,7 @@ public:
 			Item = FreeList;
 			FreeList = *(FreeList->GetIntrusiveAddress());
 			
-			Item->Ressurect(std::forward<ArgumentTypes>(Arguments)...);
+			Item->Resurrect(std::forward<ArgumentTypes>(Arguments)...);
 		}
 		else
 		{
diff --git a/Engine/Source/Runtime/Core/Private/AutoRTFM/Transaction.cpp b/Engine/Source/Runtime/Core/Private/AutoRTFM/Transaction.cpp
--- a/Engine/Source/Runtime/Core/Private/AutoRTFM/Transaction.cpp
+++ b/Engine/Source/Runtime/Core/Private/AutoRTFM/Transaction.cpp
@@ -21,7 +21,7 @@
 namespace AutoRTFM
 {
 	
-void FTransaction::Ressurect(FContext* const InContext)
+void FTransaction::Resurrect(FContext* const InContext)
 {
 	ASSERT(Context == InContext)
 	Parent = nullptr;
@@ -51,7 +51,7 @@ FTransaction::FTransaction(FContext* Context)
 	, CommitTasks(Context->GetTaskPool())
 	, AbortTasks(Context->GetTaskPool())
 {
-	Ressurect(Context);
+	Resurrect(Context);
 }
 
 FTransaction::~FTransaction()
diff --git a/Engine/Source/Runtime/Core/Private/AutoRTFM/Transaction.h b/Engine/Source/Runtime/Core/Private/AutoRTFM/Transaction.h
--- a/Engine/Source/Runtime/Core/Private/AutoRTFM/Transaction.h
+++ b/Engine/Source/Runtime/Core/Private/AutoRTFM/Transaction.h
@@ -22,7 +22,7 @@ class FTransaction final
 {
 	friend class TIntrusivePool<FTransaction, 16>;
 
-	void Ressurect(FContext* const InContext);
+	void Resurrect(FContext* const InContext);
 	void Suppress();
 	FTransaction** GetIntrusiveAddress();
 
