[AutoRTFM] Add UE_AUTORTFM_API to InitializeForUE()

Fixes linker errors with modular builds.

Also clean up preprocessor conditionals now that AutoRTFMUE.h is including AutoRTFM.h

diff --git a/Engine/Source/Runtime/Core/Private/AutoRTFM/AutoRTFMUE.cpp b/Engine/Source/Runtime/Core/Private/AutoRTFM/AutoRTFMUE.cpp
--- a/Engine/Source/Runtime/Core/Private/AutoRTFM/AutoRTFMUE.cpp
+++ b/Engine/Source/Runtime/Core/Private/AutoRTFM/AutoRTFMUE.cpp
@@ -1,10 +1,10 @@
 // Copyright Epic Games, Inc. All Rights Reserved.
 
-#if (defined(__AUTORTFM) && __AUTORTFM)
-
 #include "AutoRTFM/AutoRTFMUE.h"
 #include "AutoRTFM/AutoRTFM.h"
 
+#if UE_AUTORTFM
+
 #include "Containers/UnrealString.h"
 #include "GenericPlatform/GenericPlatformCrashContext.h"
 #include "HAL/IConsoleManager.h"
@@ -12,6 +12,8 @@
 
 #include "Utils.h" // TODO(SOL-7652) Remove this include once all logging is abstracted
 
+static_assert(UE_AUTORTFM_ENABLED, "AutoRTFM/AutoRTFMUE.cpp requires the compiler flag '-fautortfm'");
+
 namespace
 {
 
@@ -199,7 +201,7 @@ static FAutoConsoleVariable CVarAutoRTFMMemoryValidationStatisticsEnabled(
 
 } // anonymous namespace
 
-void AutoRTFM::InitializeForUE()
+UE_AUTORTFM_API void AutoRTFM::InitializeForUE()
 {
 	AutoRTFM::ForTheRuntime::FExternAPI ExternAPI;
 	ExternAPI.Allocate = [](size_t Size, size_t Alignment)
@@ -232,4 +234,4 @@ void AutoRTFM::InitializeForUE()
 	OnAutoRTFMMemoryValidationStatisticsChanged();
 }
 
-#endif  // (defined(__AUTORTFM) && __AUTORTFM)
+#endif // UE_AUTORTFM
diff --git a/Engine/Source/Runtime/Core/Public/AutoRTFM/AutoRTFMUE.h b/Engine/Source/Runtime/Core/Public/AutoRTFM/AutoRTFMUE.h
--- a/Engine/Source/Runtime/Core/Public/AutoRTFM/AutoRTFMUE.h
+++ b/Engine/Source/Runtime/Core/Public/AutoRTFM/AutoRTFMUE.h
@@ -2,14 +2,20 @@
 
 #pragma once
 
+#include "AutoRTFM.h"
+
 namespace AutoRTFM
 {
 
-#if (defined(__AUTORTFM_ENABLED) && __AUTORTFM_ENABLED)
+#if UE_AUTORTFM
+
 // Initializes AutoRTFM for use by the Unreal Engine.
-void InitializeForUE();
+UE_AUTORTFM_API void InitializeForUE();
+
 #else
+
 inline void InitializeForUE() {}
-#endif // (defined(__AUTORTFM_ENABLED) && __AUTORTFM_ENABLED)
+
+#endif
 
 }
