Document the string format for FPlatformTime::StrDate, StrTime, and StrTimestamp

diff --git a/Engine/Source/Runtime/Core/Public/GenericPlatform/GenericPlatformTime.h b/Engine/Source/Runtime/Core/Public/GenericPlatform/GenericPlatformTime.h
--- a/Engine/Source/Runtime/Core/Public/GenericPlatform/GenericPlatformTime.h
+++ b/Engine/Source/Runtime/Core/Public/GenericPlatform/GenericPlatformTime.h
@@ -84,7 +84,7 @@ struct FGenericPlatformTime
 	 * 
 	 * @param Dest Destination buffer to copy to
 	 * @param DestSize Size of destination buffer in characters
-	 * @return Date string
+	 * @return Date string in the format "MM/DD/YY"
 	 */
 	static CORE_API TCHAR* StrDate( TCHAR* Dest, SIZE_T DestSize );
 	/**
@@ -92,7 +92,7 @@ struct FGenericPlatformTime
 	 * 
 	 * @param Dest Destination buffer to copy to
 	 * @param DestSize Size of destination buffer in characters
-	 * @return Time string
+	 * @return Time string in the format "HH::mm::ss"
 	 */
 	static CORE_API TCHAR* StrTime( TCHAR* Dest, SIZE_T DestSize );
 
@@ -100,7 +100,7 @@ struct FGenericPlatformTime
 	 * Returns a timestamp string built from the current date and time.
 	 * NOTE: Only one return value is valid at a time!
 	 *
-	 * @return timestamp string
+	 * @return timestamp string in the format "MM/DD/YY HH::mm::ss"
 	 */
 	static CORE_API const TCHAR* StrTimestamp();
 
