Improve descriptions for DumpHelp, DumpConsoleCommands, DumpCVars, and DumpCCmds.
- Include all dump commands in DumpHelp.
- DumpCCmds said it dumped CVars, which is incorrect.
- Clarify that DumpCCmds only dumps registered commands, while DumpConsoleCommands dumps all registered and unregistered commands.
- Ideally we combine these together, but that's a little more work. DumpCCmds supports more features and is implemented in what seems to be a better location. But it can't scrape unregistered commands the way DumpConsoleCommands does without some refactoring.

Fixed DumpConsoleCommands not filtering results using the provided prefix.
- It was used when gathering unregistered commands, but not when gathering registered commands.
- Don't eagerly append a '*' to the command because we'd just need to peel it back off for the console object visitor.

Fix DumpConsoleCommands printing CVars in addition to CCmds
- The visitor visits all console objects, so we need to filter CVars manually.

Fixed DumpCVars tripping assertions when there are deprecated variables. GetString counts as usage and deprecated objects assert when used by default.
- Added FConsoleObjectShadowData::IsDeprecated. DumpObjects skips deprecated objects by default.
- Added -Deprecated option. When used, DumpObjects will log only deprecated objects and is careful not to read any values and trip the assert.
- Added tests to prevent regression.

Removed the deprecation message from FAutoConsoleVariableShadow because it shouldn't be used to deprecate variables.
- Noticed there's a missing console command equivalent. Might be worth adding.

diff --git a/Engine/Source/Runtime/Core/Private/HAL/ConsoleManager.cpp b/Engine/Source/Runtime/Core/Private/HAL/ConsoleManager.cpp
--- a/Engine/Source/Runtime/Core/Private/HAL/ConsoleManager.cpp
+++ b/Engine/Source/Runtime/Core/Private/HAL/ConsoleManager.cpp
@@ -5,28 +5,29 @@ ConsoleManager.cpp: console command handling
 =============================================================================*/
 
 #include "HAL/ConsoleManager.h"
+
 #include "AutoRTFM/AutoRTFM.h"
 #include "Containers/Array.h"
 #include "Containers/Map.h"
 #include "Containers/UnrealString.h"
 #include "CoreTypes.h"
 #include "HAL/CriticalSection.h"
-#include "Stats/Stats.h"
+#include "HAL/FileManager.h"
+#include "HAL/PlatformProcess.h"
+#include "Logging/MessageLog.h"
 #include "Misc/ConfigCacheIni.h"
 #include "Misc/ConfigContext.h"
 #include "Misc/ConfigUtilities.h"
 #include "Misc/CoreDelegates.h"
-#include "Misc/OutputDeviceFile.h"
+#include "Misc/OutputDeviceNull.h"
 #include "Misc/Paths.h"
 #include "Misc/RemoteConfigIni.h"
+#include "Misc/ScopeExit.h"
 #include "Misc/ScopeLock.h"
 #include "Misc/ScopeRWLock.h"
 #include "Misc/TransactionallySafeCriticalSection.h"
 #include "Modules/ModuleManager.h"
-#include "HAL/PlatformProcess.h"
 #include "ProfilingDebugging/CsvProfiler.h"
-#include "HAL/FileManager.h"
-#include "Logging/MessageLog.h"
 #include "Trace/Trace.h"
 #include "Trace/Trace.inl"
 
@@ -2759,6 +2760,8 @@ static void MultiLogf(FOutputDevice* Device, FArchive* File, const FmtType& Fmt,
 static void DumpObjects(const TMap<FString, IConsoleObject*>& ConsoleObjects, const TCHAR* Params, FOutputDevice& InAr, bool bDisplayCommands)
 {
 	bool bShowHelp = FParse::Param(Params, TEXT("showhelp"));
+	bool bShowDeprecated = FParse::Param(Params, TEXT("deprecated"));
+
 	FString CSVFilename;
 	bool bWriteToCSV = FParse::Value(Params, TEXT("-csv="), CSVFilename);
 	bWriteToCSV = bWriteToCSV || FParse::Param(Params, TEXT("csv"));
@@ -2814,6 +2817,19 @@ static void DumpObjects(const TMap<FString, IConsoleObject*>& ConsoleObjects, co
 			IConsoleVariable* CVar = Obj->AsVariable();
 			IConsoleCommand* CCmd = Obj->AsCommand();
 
+			bool bDisplay = true;
+			bDisplay &= bShowDeprecated == Obj->IsDeprecated();
+			bDisplay &= bDisplayCommands ? CCmd != nullptr : CVar != nullptr;
+			if (!bDisplay)
+			{
+				continue;
+			}
+
+			FString Annotation =
+				Obj->IsDeprecated() ? TEXT(" [DEPRECATED]") :
+				Obj->IsShadowObject() ? TEXT(" [SHADOW]") :
+				TEXT("");
+
 			// process optional help
 			FString Help;
 			if (bShowHelp)
@@ -2830,28 +2846,28 @@ static void DumpObjects(const TMap<FString, IConsoleObject*>& ConsoleObjects, co
 				}
 			}
 
-			if (bDisplayCommands && CCmd != nullptr)
+			if (CCmd != nullptr)
 			{
-				MultiLogf(Log, CSV, TEXT("%s%s"), *Key, *Help);
+				MultiLogf(Log, CSV, TEXT("%s%s%s"), *Key, *Annotation, *Help);
 			}
-			if (!bDisplayCommands && CVar != nullptr)
+			if (CVar != nullptr)
 			{
-				if (FilterSetBy.Len() > 0 && GetConsoleVariableSetByName(CVar->GetFlags()) != FilterSetBy)
+				const TCHAR* SetBy = GetConsoleVariableSetByName(CVar->GetFlags());
+				if (FilterSetBy.Len() > 0 && SetBy != FilterSetBy)
 				{
 					continue;
 				}
-				FString ShadowString;
-				if (Obj->IsShadowObject())
-				{
-					ShadowString = TEXT(" [SHADOW]");
-				}
+
+				// Don't read the value when deprecated. It will trip the UsageBehavior, which may assert.
+				FString Value = Obj->IsDeprecated() ? TEXT("") : CVar->GetString();
+
 				if (bWriteToCSV)
 				{
-					MultiLogf(Log, CSV, TEXT("%s%s,%s,%s%s"), *Key, *ShadowString, *CVar->GetString(), GetConsoleVariableSetByName(CVar->GetFlags()), *Help);
+					MultiLogf(Log, CSV, TEXT("%s%s,%s,%s%s"), *Key, *Annotation, *Value, SetBy, *Help);
 				}
 				else
 				{
-					MultiLogf(Log, CSV, TEXT("%s%s = \"%s\"      LastSetBy: %s%s"), *Key, *ShadowString, *CVar->GetString(), GetConsoleVariableSetByName(CVar->GetFlags()), *Help);
+					MultiLogf(Log, CSV, TEXT("%s%s = \"%s\"      LastSetBy: %s%s"), *Key, *Annotation, *Value, SetBy, *Help);
 				}
 			}
 		}
@@ -3881,6 +3897,36 @@ void FConsoleManager::Test()
 
 	UnregisterConsoleVariableSink_Handle(TestSinkCallbackHandle);
 
+	FConsoleManager& Manager = GetManager();
+
+	// Deprecated console variables should not assert in dump commands
+	{
+		constexpr TCHAR OldName[] = TEXT("TestVar.Old");
+		constexpr TCHAR NewName[] = TEXT("TestVar.New");
+
+		FAutoConsoleVariable NewVar(NewName, false, TEXT(""));
+		FAutoConsoleVariableDeprecated OldVar(OldName, NewName, TEXT("0.0"));
+		ON_SCOPE_EXIT { Manager.UnregisterConsoleObject(OldName, false); };
+
+		FOutputDeviceNull OutNull;
+		DumpObjects(Manager.ConsoleObjects, TEXT("-ShowHelp"), OutNull, false);
+		DumpObjects(Manager.ConsoleObjects, TEXT("-ShowHelp -Deprecated"), OutNull, false);
+	}
+
+	// Deprecated console commands should not assert in dump commands
+	{
+		constexpr TCHAR OldName[] = TEXT("TestCmd.Old");
+		constexpr TCHAR NewName[] = TEXT("TestCmd.New");
+
+		FAutoConsoleCommand NewCmd(NewName, TEXT(""), FConsoleCommandDelegate());
+		FAutoConsoleCommandDeprecated OldCmd(OldName, NewName, TEXT("0.0"));
+		ON_SCOPE_EXIT { Manager.UnregisterConsoleObject(OldName, false); };
+
+		FOutputDeviceNull OutNull;
+		DumpObjects(Manager.ConsoleObjects, TEXT("-ShowHelp"), OutNull, true);
+		DumpObjects(Manager.ConsoleObjects, TEXT("-ShowHelp -Deprecated"), OutNull, true);
+	}
+
 #endif // !(UE_BUILD_SHIPPING || UE_BUILD_TEST)
 }
 
@@ -3925,8 +3971,7 @@ void CreateConsoleVariables()
 		ECVF_Default);
 #endif
 
-#if	!UE_BUILD_SHIPPING
-	IConsoleManager::Get().RegisterConsoleCommand( TEXT( "DumpConsoleCommands" ), TEXT( "Dumps all console vaiables and commands and all exec that can be discovered to the log/console" ), ECVF_Default );
+#if !UE_BUILD_SHIPPING
 	IConsoleManager::Get().RegisterConsoleCommand( TEXT( "RedirectToFile" ),
 		TEXT( "Creates a file inside Project's Saved folder and outputs command result into it as well as into the log.\n" )
 		TEXT( "Usage: RedirectToFile <filepath/filename> <command> [command arguments]\n" )
@@ -3935,12 +3980,30 @@ void CreateConsoleVariables()
 		ECVF_Default );
 
 	IConsoleManager::Get().RegisterConsoleCommand(TEXT("DumpCVars"),
-		TEXT("Lists all CVars (or a subset) and their values. Can also show help, and can save to .csv.\nUsage: DumpCVars [Prefix] [-showhelp] [-csv=[path]]\nIf -csv does not have a file specified, it will create a file in the Project Logs directory"),
+		TEXT("List all console variables, their values, and how they were last modified.\n")
+		TEXT("Usage: DumpCVars [prefix] [-showhelp] [-deprecated] [-csv[=path]]\n")
+		TEXT("prefix      - Only list variables that start with this string\n")
+		TEXT("-showhelp   - Show help information for each variable\n")
+		TEXT("-deprecated - Only list deprecated variables. Values are not shown for deprecated variables.\n")
+		TEXT("-csv        - Save results to a csv file. If a file is not specified, it will be created in the Project Logs directory."),
 		ECVF_Default);
+
 	IConsoleManager::Get().RegisterConsoleCommand(TEXT("DumpCCmds"),
-		TEXT("Lists all CVars (or a subset) and their values. Can also show help, and can save to .csv.\nUsage: DumpCCmds [Prefix] [-showhelp] [-csv=[path]]\nIf -csv does not have a file specified, it will create a file in the Project Logs directory"),
+		TEXT("List all registered console commands.\n")
+		TEXT("Not all commands are registered. Unregistered commands will not be listed by this command\n")
+		TEXT("Use DumpConsoleCommands if you need to list unregistered commands as well\n")
+		TEXT("Usage: DumpCCmds [prefix] [-showhelp] [-deprecated] [-csv[=path]]\n")
+		TEXT("prefix      - Only list commands that start with this string\n")
+		TEXT("-showhelp   - Show help information for each command\n")
+		TEXT("-deprecated - Only list deprecated commands\n")
+		TEXT("-csv        - Save results to a csv file. If a file is not specified, it will be created in the Project Logs directory."),
 		ECVF_Default);
 
+	IConsoleManager::Get().RegisterConsoleCommand(TEXT("DumpConsoleCommands"),
+		TEXT("List all console commands, both registered and unregistered\n")
+		TEXT("Usage: DumpConsoleCommands [prefix]\n")
+		TEXT("prefix      - Only list commands that start with this string\n"),
+		ECVF_Default);
 #endif // !UE_BUILD_SHIPPING
 
 	// testing code
@@ -4814,7 +4877,14 @@ class FConsoleObjectShadowData
 {
 protected:
 	
-	FConsoleObjectShadowData(const TCHAR* CVarToShadow, const TCHAR* InDeprecatedVersion, EShadowCVarBehavior InLookupBehavior, EShadowCVarBehavior InUsageBehavior, const TCHAR* InDeprecationMessage = nullptr)
+	FConsoleObjectShadowData(
+		const TCHAR* CVarToShadow,
+		EShadowCVarBehavior InLookupBehavior,
+		EShadowCVarBehavior InUsageBehavior,
+		bool bInIsDeprecated,
+		const TCHAR* InDeprecatedVersion,
+		const TCHAR* InDeprecationMessage = nullptr
+	)
 		: ShadowName(CVarToShadow)
 		, DeprecatedVersion(InDeprecatedVersion)
 		, DeprecationMessage(InDeprecationMessage)
@@ -4823,6 +4893,7 @@ protected:
 		, bHasLooked(false)
 		, bHasMessagedForUsage(false)
 		, bHasMessagedEditorForUsage(false)
+		, bIsDeprecated(bInIsDeprecated)
 	{
 	}
 
@@ -4983,17 +5054,24 @@ protected:
 	mutable uint8 bHasLooked : 1;
 	mutable uint8 bHasMessagedForUsage : 1;
 	mutable uint8 bHasMessagedEditorForUsage : 1;
+	uint8 bIsDeprecated : 1;
 };
 
 
 class FConsoleVariableShadow : public FConsoleObjectShadowData, public IConsoleVariable
 {
 public:
-	FConsoleVariableShadow(const TCHAR* CVarToShadow, const TCHAR* InDeprecatedVersion, EShadowCVarBehavior InLookupBehavior, EShadowCVarBehavior InUsageBehavior, const TCHAR* InDeprecationMessage)
-		: FConsoleObjectShadowData(CVarToShadow, InDeprecatedVersion, InLookupBehavior, InUsageBehavior, InDeprecationMessage)
+	FConsoleVariableShadow(
+		const TCHAR* CVarToShadow,
+		EShadowCVarBehavior InLookupBehavior,
+		EShadowCVarBehavior InUsageBehavior,
+		bool bInIsDeprecated,
+		const TCHAR* InDeprecatedVersion,
+		const TCHAR* InDeprecationMessage
+	)
+		: FConsoleObjectShadowData(CVarToShadow, InLookupBehavior, InUsageBehavior, bInIsDeprecated, InDeprecatedVersion, InDeprecationMessage)
 		, RealVariable(nullptr)
 	{
-
 	}
 
 	virtual class IConsoleVariable* AsVariable() override
@@ -5013,6 +5091,10 @@ public:
 		}
 		return nullptr;
 	}
+	virtual bool IsDeprecated() const override
+	{
+		return bIsDeprecated;
+	}
 
 
 	virtual bool IsVariableBool() const override
@@ -5281,14 +5363,20 @@ private:
 	mutable IConsoleVariable* RealVariable;
 };
 
-struct FConsoleCommandShadow : public FConsoleObjectShadowData, public IConsoleCommand
+class FConsoleCommandShadow : public FConsoleObjectShadowData, public IConsoleCommand
 {
-
-	FConsoleCommandShadow(const TCHAR* CVarToShadow, const TCHAR* InDeprecatedVersion, EShadowCVarBehavior InLookupBehavior, EShadowCVarBehavior InUsageBehavior, const TCHAR* InDeprecationMessage)
-		: FConsoleObjectShadowData(CVarToShadow, InDeprecatedVersion, InLookupBehavior, InUsageBehavior, InDeprecationMessage)
+public:
+	FConsoleCommandShadow(
+		const TCHAR* CVarToShadow,
+		EShadowCVarBehavior InLookupBehavior,
+		EShadowCVarBehavior InUsageBehavior,
+		bool bInIsDeprecated,
+		const TCHAR* InDeprecatedVersion,
+		const TCHAR* InDeprecationMessage
+	)
+		: FConsoleObjectShadowData(CVarToShadow, InLookupBehavior, InUsageBehavior, bInIsDeprecated, InDeprecatedVersion, InDeprecationMessage)
 		, RealCommand(nullptr)
 	{
-
 	}
 
 	virtual struct IConsoleCommand* AsCommand() override
@@ -5296,6 +5384,11 @@ struct FConsoleCommandShadow : public FConsoleObjectShadowData, public IConsoleC
 		return this;
 	}
 
+	virtual bool IsDeprecated() const override
+	{
+		return bIsDeprecated;
+	}
+
 	virtual const TCHAR* GetHelp() const override
 	{
 		if (Bind())
@@ -5372,29 +5465,28 @@ protected:
 private:
 
 	mutable IConsoleCommand* RealCommand;
-
 };
 
 
 
-FAutoConsoleVariableShadow::FAutoConsoleVariableShadow(const TCHAR* Name, const TCHAR* CVarToShadow, EShadowCVarBehavior LookupFailureBehavior, const TCHAR* DeprecationMessage)
+FAutoConsoleVariableShadow::FAutoConsoleVariableShadow(const TCHAR* Name, const TCHAR* CVarToShadow, EShadowCVarBehavior LookupFailureBehavior)
 {
 #if !NO_CVARS
-	GetManager().AddShadowConsoleObject(Name, new FConsoleVariableShadow(CVarToShadow, nullptr, LookupFailureBehavior, EShadowCVarBehavior::NoMessaging, DeprecationMessage));
+	GetManager().AddShadowConsoleObject(Name, new FConsoleVariableShadow(CVarToShadow, LookupFailureBehavior, EShadowCVarBehavior::NoMessaging, false, nullptr, TEXT("")));
 #endif
 }
 
 FAutoConsoleVariableDeprecated::FAutoConsoleVariableDeprecated(const TCHAR* Name, const TCHAR* CVarToShadow, const TCHAR* DeprecatedAtVersion, EShadowCVarBehavior UsageBehavior,  EShadowCVarBehavior LookupFailureBehavior, const TCHAR* DeprecationMessage)
 {
 #if !NO_CVARS
-	GetManager().AddShadowConsoleObject(Name, new FConsoleVariableShadow(CVarToShadow, DeprecatedAtVersion, LookupFailureBehavior, UsageBehavior, DeprecationMessage));
+	GetManager().AddShadowConsoleObject(Name, new FConsoleVariableShadow(CVarToShadow, LookupFailureBehavior, UsageBehavior, true, DeprecatedAtVersion, DeprecationMessage));
 #endif
 }
 
 FAutoConsoleCommandDeprecated::FAutoConsoleCommandDeprecated(const TCHAR* Name, const TCHAR* CVarToShadow, const TCHAR* DeprecatedAtVersion, EShadowCVarBehavior UsageBehavior, EShadowCVarBehavior LookupFailureBehavior, const TCHAR* DeprecationMessage)
 {
 #if !NO_CVARS
-	GetManager().AddShadowConsoleObject(Name, new FConsoleCommandShadow(CVarToShadow, DeprecatedAtVersion, LookupFailureBehavior, UsageBehavior, DeprecationMessage));
+	GetManager().AddShadowConsoleObject(Name, new FConsoleCommandShadow(CVarToShadow, LookupFailureBehavior, UsageBehavior, true, DeprecatedAtVersion, DeprecationMessage));
 #endif
 }
 
diff --git a/Engine/Source/Runtime/Core/Private/Misc/Parse.cpp b/Engine/Source/Runtime/Core/Private/Misc/Parse.cpp
--- a/Engine/Source/Runtime/Core/Private/Misc/Parse.cpp
+++ b/Engine/Source/Runtime/Core/Private/Misc/Parse.cpp
@@ -31,27 +31,26 @@
 class ConsoleCommandLibrary
 {
 public:
-	ConsoleCommandLibrary(const FString& InPattern);
+	ConsoleCommandLibrary(FStringView InPrefix);
 
 	~ConsoleCommandLibrary();
 
-	void OnParseCommand(const TCHAR* Match)
+	void OnParseCommand(const TCHAR* Cmd)
 	{
-		// -1 to not take the "*" after the pattern into account
-		if(FCString::Strnicmp(Match, *Pattern, Pattern.Len() - 1) == 0)
+		if (FCString::Strnicmp(Cmd, Prefix.GetData(), Prefix.Len()) == 0)
 		{
-			KnownNames.Add(Match);
+			KnownNames.Add(Cmd);
 		}
 	}
 
-	const FString&		Pattern;
-	TSet<FString>		KnownNames;
+	FStringView   Prefix;
+	TSet<FString> KnownNames;
 };
 
 // 0 if gathering of names is deactivated
 ConsoleCommandLibrary* GConsoleCommandLibrary;
 
-ConsoleCommandLibrary::ConsoleCommandLibrary(const FString& InPattern) :Pattern(InPattern)
+ConsoleCommandLibrary::ConsoleCommandLibrary(FStringView InPrefix) : Prefix(InPrefix)
 {
 	// activate name gathering
 	GConsoleCommandLibrary = this;
@@ -65,14 +64,17 @@ ConsoleCommandLibrary::~ConsoleCommandLibrary()
 
 
 
-class FConsoleVariableDumpVisitor 
+class FConsoleCommandDumpVisitor
 {
 public:
-	// @param Name must not be 0
-	// @param " must not be 0
-	static void OnConsoleVariable(const TCHAR *Name, IConsoleObject* CVar,TSet<FString>* Sink)
+	/**
+	 * @param Name must not be null
+	 * @param Object must not be null
+	 * @param Sink must not be null
+	 */
+	static void OnConsoleObject(const TCHAR *Name, IConsoleObject* Object, TSet<FString>* Sink)
 	{
-		if(CVar->TestFlags(ECVF_Unregistered))
+		if (!Object->AsCommand() || Object->TestFlags(ECVF_Unregistered))
 		{
 			return;
 		}
@@ -83,19 +85,33 @@ public:
 
 bool ConsoleCommandLibrary_DumpLibrary(UWorld* InWorld, FExec& SubSystem, const FString& Pattern, FOutputDevice& Ar)
 {
-	ConsoleCommandLibrary LocalConsoleCommandLibrary(Pattern);
-
-	FOutputDeviceNull Null;
+	TStringBuilder<32> Prefix;
+	Prefix << FStringView(Pattern).LeftChop(1);
+	return ConsoleCommandLibrary_DumpLibrary(InWorld, SubSystem, *Prefix, Ar);
+}
 
-	bool bExecuted = SubSystem.Exec( InWorld, *Pattern, Null);
+bool ConsoleCommandLibrary_DumpLibrary(UWorld* InWorld, FExec& SubSystem, const TCHAR* Prefix, FOutputDevice& Ar)
+{
+	// Install a global handler to scrape unregistered commands as FExec implementations call FParse::Command
+	ConsoleCommandLibrary LocalConsoleCommandLibrary(Prefix);
 
+	// Gather unregistered commands
+	bool bExecuted = false;
 	{
-		IConsoleManager::Get().ForEachConsoleObjectThatStartsWith(
-			FConsoleObjectVisitor::CreateStatic(
-			&FConsoleVariableDumpVisitor::OnConsoleVariable,
-			&LocalConsoleCommandLibrary.KnownNames ) );
+ 		TStringBuilder<32> FakeCmd;
+		FakeCmd << Prefix << '*';
+
+		FOutputDeviceNull Null;
+		bExecuted = SubSystem.Exec(InWorld, *FakeCmd, Null);
 	}
 
+	// Gather registered commands
+	IConsoleManager::Get().ForEachConsoleObjectThatStartsWith(
+		FConsoleObjectVisitor::CreateStatic(
+			&FConsoleCommandDumpVisitor::OnConsoleObject,
+			&LocalConsoleCommandLibrary.KnownNames),
+		Prefix);
+
 	LocalConsoleCommandLibrary.KnownNames.Sort( TLess<FString>() );
 
 	for(TSet<FString>::TConstIterator It(LocalConsoleCommandLibrary.KnownNames); It; ++It)
@@ -106,7 +122,7 @@ bool ConsoleCommandLibrary_DumpLibrary(UWorld* InWorld, FExec& SubSystem, const
 	}
 	Ar.Logf(TEXT(""));
 
-	// the pattern (e.g. Motion*) should not really trigger the execution
+	// the fake command (e.g. Motion*) should not really trigger the execution
 	if(bExecuted)
 	{
 		Ar.Logf(TEXT("ERROR: The function was supposed to only find matching commands but not have any side effect."));
@@ -118,20 +134,26 @@ bool ConsoleCommandLibrary_DumpLibrary(UWorld* InWorld, FExec& SubSystem, const
 
 bool ConsoleCommandLibrary_DumpLibraryHTML(UWorld* InWorld, FExec& SubSystem, const FString& OutPath)
 {
-	const FString& Pattern(TEXT("*"));
-	ConsoleCommandLibrary LocalConsoleCommandLibrary(Pattern);
-
-	FOutputDeviceNull Null;
+	FStringView Prefix = TEXT("");
 
-	bool bExecuted = SubSystem.Exec( InWorld, *LocalConsoleCommandLibrary.Pattern, Null);
+	// Install a global handler to scrape unregistered commands as FExec implementations call FParse::Command
+	ConsoleCommandLibrary LocalConsoleCommandLibrary(Prefix);
 
+	// Gather unregistered commands
+	bool bExecuted = false;
 	{
-		IConsoleManager::Get().ForEachConsoleObjectThatStartsWith(
-			FConsoleObjectVisitor::CreateStatic(
-			&FConsoleVariableDumpVisitor::OnConsoleVariable,
-			&LocalConsoleCommandLibrary.KnownNames ) );
+		const TCHAR* FakeCmd = TEXT("*");
+
+		FOutputDeviceNull Null;
+		bExecuted = SubSystem.Exec(InWorld, FakeCmd, Null);
 	}
 
+	// Gather registered variables and commands
+	IConsoleManager::Get().ForEachConsoleObjectThatStartsWith(
+		FConsoleObjectVisitor::CreateStatic(
+			&FConsoleCommandDumpVisitor::OnConsoleObject,
+			&LocalConsoleCommandLibrary.KnownNames));
+
 	LocalConsoleCommandLibrary.KnownNames.Sort( TLess<FString>() );
 
 	FString TemplateFilename = FPaths::Combine(FPlatformProcess::BaseDir(), TEXT("../../Documentation/Extras"), TEXT("ConsoleHelpTemplate.html"));
diff --git a/Engine/Source/Runtime/Core/Public/HAL/IConsoleManager.h b/Engine/Source/Runtime/Core/Public/HAL/IConsoleManager.h
--- a/Engine/Source/Runtime/Core/Public/HAL/IConsoleManager.h
+++ b/Engine/Source/Runtime/Core/Public/HAL/IConsoleManager.h
@@ -526,6 +526,11 @@ public:
 		return nullptr;
 	}
 
+	virtual bool IsDeprecated() const
+	{
+		return false;
+	}
+
 private: // -----------------------------------------
 
 #if TRACK_CONSOLE_FIND_COUNT
@@ -2231,7 +2236,9 @@ enum class EShadowCVarBehavior : uint8
 };
 
 /**
- * A class that can be used to shadow a CVar, where both are valid to use. To deprecate a CVar, it's suggested to use the FAutoConsoleVariableDeprecated class below
+ * A class that can be used to shadow (alias) a CVar, where both names are valid to use.
+ * To deprecate a CVar, use FAutoConsoleVariableDeprecated.
+ * To deprecate a CCmd, use FAutoConsoleCommandDeprecated.
  * 
  * FAutoConsoleVariableShadow ShadowCVar(TEXT("r.oldname"), TEXT("r.newname"));
  * 
@@ -2240,13 +2247,12 @@ class FAutoConsoleVariableShadow
 {
 public:
 
-		/**
-		 * @param Name The name of this cvar
-		 * @param CVarToShadow The name of the existing cvar that this one will shadow
-		 * @param LookupFailureBehavior How to handle failure of looking up the CVarToShadow name - defaults to ensure
-		 * @param DeprecationMessage Optional deprecation message - defaults to empty, auto-generated.
-		 */
-		CORE_API FAutoConsoleVariableShadow(const TCHAR* Name, const TCHAR* CVarToShadow, EShadowCVarBehavior LookupFailureBehavior=EShadowCVarBehavior::Ensure, const TCHAR* DeprecationMessage = TEXT(""));
+	/**
+	 * @param Name The name of this cvar
+	 * @param CVarToShadow The name of the existing cvar that this one will shadow
+	 * @param LookupFailureBehavior How to handle failure of looking up the CVarToShadow name - defaults to ensure
+	 */
+	CORE_API FAutoConsoleVariableShadow(const TCHAR* Name, const TCHAR* CVarToShadow, EShadowCVarBehavior LookupFailureBehavior = EShadowCVarBehavior::Ensure);
 };
 
 /**
diff --git a/Engine/Source/Runtime/Core/Public/Misc/Parse.h b/Engine/Source/Runtime/Core/Public/Misc/Parse.h
--- a/Engine/Source/Runtime/Core/Public/Misc/Parse.h
+++ b/Engine/Source/Runtime/Core/Public/Misc/Parse.h
@@ -231,8 +231,13 @@ ENUM_CLASS_FLAGS(FParse::EGrammarBasedParseFlags);
 ENUM_CLASS_FLAGS(FParse::ELineExtendedFlags)
 
 #if !UE_BUILD_SHIPPING
+
+UE_DEPRECATED("5.6", "Use the overload that takes the entire command instead of a pattern.")
+CORE_API bool ConsoleCommandLibrary_DumpLibrary(class UWorld* InWorld, FExec& SubSystem, const FString& Cmd, FOutputDevice& Ar);
+
 /** Needed for the console command "DumpConsoleCommands" */
-CORE_API bool ConsoleCommandLibrary_DumpLibrary(class UWorld* InWorld, FExec& SubSystem, const FString& Pattern, FOutputDevice& Ar);
+CORE_API bool ConsoleCommandLibrary_DumpLibrary(class UWorld* InWorld, FExec& SubSystem, const TCHAR* Cmd, FOutputDevice& Ar);
+
 /** Needed for the console command "Help" */
 CORE_API bool ConsoleCommandLibrary_DumpLibraryHTML(class UWorld* InWorld, FExec& SubSystem, const FString& OutPath);
 #endif // !UE_BUILD_SHIPPING
diff --git a/Engine/Source/Runtime/Engine/Private/Tests/ExecCommandsTests.cpp b/Engine/Source/Runtime/Engine/Private/Tests/ExecCommandsTests.cpp
--- a/Engine/Source/Runtime/Engine/Private/Tests/ExecCommandsTests.cpp
+++ b/Engine/Source/Runtime/Engine/Private/Tests/ExecCommandsTests.cpp
@@ -262,7 +262,6 @@ bool FCommonCommandsExecCommandTest::RunTest(const FString& Parameters)
 
 	TArray<FString> CommandList;
 	ConsoleCommandsDump.ParseIntoArrayLines(CommandList);
-	CommandList.RemoveAt(0); // Remove the first entry as its not sorted and is something like DumpConsoleCommands:*
 
 	for ( const FString& CommonCommand : CommonCommands )
 	{
@@ -273,4 +272,4 @@ bool FCommonCommandsExecCommandTest::RunTest(const FString& Parameters)
 	return true;
 }
 
-#endif // WITH_DEV_AUTOMATION_TESTS && UE_ALLOW_EXEC_COMMANDS
\ No newline at end of file
+#endif // WITH_DEV_AUTOMATION_TESTS && UE_ALLOW_EXEC_COMMANDS
diff --git a/Engine/Source/Runtime/Engine/Private/UnrealEngine.cpp b/Engine/Source/Runtime/Engine/Private/UnrealEngine.cpp
--- a/Engine/Source/Runtime/Engine/Private/UnrealEngine.cpp
+++ b/Engine/Source/Runtime/Engine/Private/UnrealEngine.cpp
@@ -5839,28 +5839,31 @@ bool UEngine::HandleHotReloadCommand( const TCHAR* Cmd, FOutputDevice& Ar )
 
 static void DumpHelp(UWorld* InWorld)
 {
-	UE_LOG(LogEngine, Display, TEXT("Console Help:"));
-	UE_LOG(LogEngine, Display, TEXT("============="));
-	UE_LOG(LogEngine, Display, TEXT(" "));
-	UE_LOG(LogEngine, Display, TEXT("A console variable is a engine wide key value pair. The key is a string usually starting with the subsystem prefix followed"));
-	UE_LOG(LogEngine, Display, TEXT("by '.' e.g. r.BloomQuality. The value can be of different tpe (e.g. float, int, string). A console command has no state associated with"));
-	UE_LOG(LogEngine, Display, TEXT("and gets executed immediately."));
-	UE_LOG(LogEngine, Display, TEXT(" "));
-	UE_LOG(LogEngine, Display, TEXT("Console variables can be put into ini files (e.g. ConsoleVariables.ini or BaseEngine.ini) with this syntax:"));
-	UE_LOG(LogEngine, Display, TEXT("<Console variable> = <value>"));
-	UE_LOG(LogEngine, Display, TEXT(" "));
-	UE_LOG(LogEngine, Display, TEXT("DumpConsoleCommands         Lists all console variables and commands that are registered (Some are not registered)"));
-	UE_LOG(LogEngine, Display, TEXT("<Console variable>          Get the console variable state"));
-	UE_LOG(LogEngine, Display, TEXT("<Console variable> ?        Get the console variable help text"));
-	UE_LOG(LogEngine, Display, TEXT("<Console variable> <value>  Set the console variable value"));
-	UE_LOG(LogEngine, Display, TEXT("<Console command> [Params]  Execute the console command with optional parameters"));
-
-	UE_LOG(LogEngine, Display, TEXT(" "));
-
 	FString FilePath = FPaths::ProjectSavedDir() + TEXT("ConsoleHelp.html");
 
-	UE_LOG(LogEngine, Display, TEXT("To browse console variables open this: '%s'"), *FilePath);
-	UE_LOG(LogEngine, Display, TEXT(" "));
+	UE_LOG(LogEngine, Display,
+		TEXT("Console Help:\n")
+		TEXT("=============\n")
+		TEXT("\n")
+		TEXT("A console variable is an engine wide key-value pair. The key is a string usually starting with the subsystem\n")
+		TEXT("prefix followed by '.' e.g. r.BloomQuality. Values can have different types such as float, int, or string.\n")
+		TEXT("\n")
+		TEXT("Console variables can be put into ini files (e.g. ConsoleVariables.ini or BaseEngine.ini) with this syntax:\n")
+		TEXT("<Console variable> = <value>\n")
+		TEXT("\n")
+		TEXT("A console command has no state associated with it and gets executed immediately.\n")
+		TEXT("\n")
+		TEXT("DumpCVars                   Lists all console variables\n")
+		TEXT("DumpCCmds                   Lists all registerd console commands. Not all commands are registered.\n")
+		TEXT("DumpConsoleCommands         Lists all console commands, registered and unregistered\n")
+		TEXT("\n")
+		TEXT("<Console variable>          Get the console variable value and how it was last set\n")
+		TEXT("<Console variable> ?        Get the console variable help text\n")
+		TEXT("<Console variable> <value>  Set the console variable value\n")
+		TEXT("\n")
+		TEXT("<Console command> [params]  Execute the console command with optional parameters\n")
+		TEXT("\n")
+		TEXT("To browse console variables open this file: '%s'"), *FilePath);
 
 	const bool bSuccess = ConsoleCommandLibrary_DumpLibraryHTML(InWorld, *GEngine, FilePath);
 
@@ -5901,7 +5904,7 @@ static void DumpHelp(UWorld* InWorld)
 		const FString DocTemplateFile = FPaths::Combine(FPaths::EngineDir(), "Documentation/Extras/ConsoleHelpTemplate.html");
 		if (!FPaths::FileExists(DocTemplateFile))
 		{
-				Info.SubText = NSLOCTEXT("UnrealEd", "ConsoleHelpExportFailedTip", "The template HTML file could not be found in the Engine/Documentation/Extras folder.\n\nIf you use a version control system or Unreal Game Sync (UGS), check that this folder is not excluded by your sync settings.");
+			Info.SubText = NSLOCTEXT("UnrealEd", "ConsoleHelpExportFailedTip", "The template HTML file could not be found in the Engine/Documentation/Extras folder.\n\nIf you use a version control system or Unreal Game Sync (UGS), check that this folder is not excluded by your sync settings.");
 		}
 		Info.bFireAndForget = true;
 		Info.ExpireDuration = 10.0f;
@@ -5922,7 +5925,7 @@ bool UEngine::HandleDumpConsoleCommandsCommand( const TCHAR* Cmd, FOutputDevice&
 {
 	Ar.Logf(TEXT("DumpConsoleCommands: %s*"), Cmd);
 	Ar.Logf(TEXT(""));
-	ConsoleCommandLibrary_DumpLibrary( InWorld, *GEngine, FString(Cmd) + TEXT("*"), Ar);
+	ConsoleCommandLibrary_DumpLibrary(InWorld, *GEngine, Cmd, Ar);
 	return true;
 }
 
