Removed the need for FDecimalNumberSigningStrings::GetEmptyString() - avoids some bad codegen in the AutoRTFM compiler.

diff --git a/Engine/Source/Runtime/Core/Private/Internationalization/FastDecimalFormat.cpp b/Engine/Source/Runtime/Core/Private/Internationalization/FastDecimalFormat.cpp
--- a/Engine/Source/Runtime/Core/Private/Internationalization/FastDecimalFormat.cpp
+++ b/Engine/Source/Runtime/Core/Private/Internationalization/FastDecimalFormat.cpp
@@ -75,22 +75,10 @@ public:
 	FDecimalNumberSigningStrings(const FDecimalNumberFormattingRules& InFormattingRules, const EDecimalNumberSigningStringsFlags InFlags)
 	{
 		// Resolve out the default cases
-		if (InFormattingRules.NegativePrefixString.Len() > 0)
-		{
-			NegativePrefixStringPtr = &InFormattingRules.NegativePrefixString;
-		}
-		if (InFormattingRules.NegativeSuffixString.Len() > 0)
-		{
-			NegativeSuffixStringPtr = &InFormattingRules.NegativeSuffixString;
-		}
-		if (InFormattingRules.PositivePrefixString.Len() > 0)
-		{
-			PositivePrefixStringPtr = &InFormattingRules.PositivePrefixString;
-		}
-		if (InFormattingRules.PositiveSuffixString.Len() > 0)
-		{
-			PositiveSuffixStringPtr = &InFormattingRules.PositiveSuffixString;
-		}
+		NegativePrefixStringPtr = &InFormattingRules.NegativePrefixString;
+		NegativeSuffixStringPtr = &InFormattingRules.NegativeSuffixString;
+		PositivePrefixStringPtr = &InFormattingRules.PositivePrefixString;
+		PositiveSuffixStringPtr = &InFormattingRules.PositiveSuffixString;
 
 		// If we should always sign this number we can use the negative signing strings to synthesize a positive version
 		if (EnumHasAnyFlags(InFlags, EDecimalNumberSigningStringsFlags::AlwaysSign))
@@ -144,60 +132,54 @@ public:
 
 	bool HasNegativePrefixString() const
 	{
-		return NegativePrefixStringPtr && NegativePrefixStringPtr->Len() > 0;
+		return NegativePrefixStringPtr->Len() > 0;
 	}
 
 	const FString& GetNegativePrefixString() const
 	{
-		return NegativePrefixStringPtr ? *NegativePrefixStringPtr : GetEmptyString();
+		return *NegativePrefixStringPtr;
 	}
 
 	bool HasNegativeSuffixString() const
 	{
-		return NegativeSuffixStringPtr && NegativeSuffixStringPtr->Len() > 0;
+		return NegativeSuffixStringPtr->Len() > 0;
 	}
 
 	const FString& GetNegativeSuffixString() const
 	{
-		return NegativeSuffixStringPtr ? *NegativeSuffixStringPtr : GetEmptyString();
+		return *NegativeSuffixStringPtr;
 	}
 
 	bool HasPositivePrefixString() const
 	{
-		return PositivePrefixStringPtr && PositivePrefixStringPtr->Len() > 0;
+		return PositivePrefixStringPtr->Len() > 0;
 	}
 
 	const FString& GetPositivePrefixString() const
 	{
-		return PositivePrefixStringPtr ? *PositivePrefixStringPtr : GetEmptyString();
+		return *PositivePrefixStringPtr;
 	}
 
 	bool HasPositiveSuffixString() const
 	{
-		return PositiveSuffixStringPtr && PositiveSuffixStringPtr->Len() > 0;
+		return PositiveSuffixStringPtr->Len() > 0;
 	}
 
 	const FString& GetPositiveSuffixString() const
 	{
-		return PositiveSuffixStringPtr ? *PositiveSuffixStringPtr : GetEmptyString();
+		return *PositiveSuffixStringPtr;
 	}
 
 private:
-	const FString* NegativePrefixStringPtr = nullptr;
-	const FString* NegativeSuffixStringPtr = nullptr;
-	const FString* PositivePrefixStringPtr = nullptr;
-	const FString* PositiveSuffixStringPtr = nullptr;
+	const FString* NegativePrefixStringPtr;
+	const FString* NegativeSuffixStringPtr;
+	const FString* PositivePrefixStringPtr;
+	const FString* PositiveSuffixStringPtr;
 
 	FString GeneratedNegativePrefixString;
 	FString GeneratedNegativeSuffixString;
 	FString GeneratedPositivePrefixString;
 	FString GeneratedPositiveSuffixString;
-
-	const FString& GetEmptyString() const
-	{
-		static const FString EmptyStr = FString();
-		return EmptyStr;
-	}
 };
 
 void SanitizeNumberFormattingOptions(FNumberFormattingOptions& InOutFormattingOptions)
