Added Peek, Remove and base class for IoIntrusiveList

diff --git a/Engine/Source/Runtime/Core/Internal/IO/IoContainers.h b/Engine/Source/Runtime/Core/Internal/IO/IoContainers.h
--- a/Engine/Source/Runtime/Core/Internal/IO/IoContainers.h
+++ b/Engine/Source/Runtime/Core/Internal/IO/IoContainers.h
@@ -28,6 +28,22 @@ private:
 
 } // namespace UE::Private
 
+template <typename Type>
+struct TIntrusiveListElement
+{
+	using ElementType = Type; 
+
+	static Type* GetNext(const ElementType* Element)
+	{
+		return Element->Next;
+	}
+
+	static void SetNext(ElementType* Element, ElementType* Next)
+	{
+		Element->Next = Next;
+	}
+};
+
 template <typename TypeTraits>
 class TIntrusiveList
 {
@@ -44,6 +60,10 @@ public:
 	{
 		Other.Head = Other.Tail = nullptr;
 	}
+	explicit TIntrusiveList(ElementType* Element)
+	{
+		Head = Tail = Element;
+	}
 
 	TIntrusiveList& operator=(const TIntrusiveList&) = delete;
 	TIntrusiveList& operator=(TIntrusiveList&& Other)
@@ -115,6 +135,47 @@ public:
 		return Element;
 	}
 
+	ElementType* PeekHead()
+	{
+		return Head;
+	}
+
+	bool Remove(ElementType* Element)
+	{
+		if (Element == nullptr || IsEmpty())
+		{
+			return false;
+		}
+
+		if (Element == Head)
+		{
+			PopHead();
+			return true;
+		}
+
+		ElementType* It = Head;
+		ElementType* NextElement = TypeTraits::GetNext(It);
+		while (NextElement != nullptr && NextElement != Element)
+		{
+			It = NextElement;
+			NextElement = TypeTraits::GetNext(It);
+		}
+
+		if (NextElement != Element)
+		{
+			return false;
+		}
+
+		It->Next = TypeTraits::GetNext(Element);
+		TypeTraits::SetNext(Element, nullptr);
+		if (Element == Tail)
+		{
+			Tail = It;
+		}
+
+		return true;
+	}
+
 	bool				IsEmpty() const { return Head == nullptr; }
 	ElementType*		GetHead()		{ return Head; }
 	const ElementType*	GetHead() const { return Head; }
