Remove dead code from FDecimalNumberSigningStrings class.

Also, replaced `String.Len() > 0` with `!String.IsEmpty()` in a handful of places.

diff --git a/Engine/Source/Runtime/Core/Private/Internationalization/FastDecimalFormat.cpp b/Engine/Source/Runtime/Core/Private/Internationalization/FastDecimalFormat.cpp
--- a/Engine/Source/Runtime/Core/Private/Internationalization/FastDecimalFormat.cpp
+++ b/Engine/Source/Runtime/Core/Private/Internationalization/FastDecimalFormat.cpp
@@ -65,13 +65,9 @@ enum class EDecimalNumberSigningStringsFlags : uint8
 };
 ENUM_CLASS_FLAGS(EDecimalNumberSigningStringsFlags);
 
-struct FDecimalNumberSigningStrings
+class FDecimalNumberSigningStrings
 {
 public:
-	FDecimalNumberSigningStrings()
-	{
-	}
-
 	FDecimalNumberSigningStrings(const FDecimalNumberFormattingRules& InFormattingRules, const EDecimalNumberSigningStringsFlags InFlags)
 	{
 		// Resolve out the default cases
@@ -111,60 +107,40 @@ public:
 			GeneratedPositivePrefixString = GetPositivePrefixString().Replace(*InFormattingRules.PlusString,  TEXT("+"), ESearchCase::CaseSensitive);
 			GeneratedPositiveSuffixString = GetPositiveSuffixString().Replace(*InFormattingRules.PlusString,  TEXT("+"), ESearchCase::CaseSensitive);
 
-			if (InFormattingRules.NegativePrefixString.Len() > 0)
+			if (!InFormattingRules.NegativePrefixString.IsEmpty())
 			{
 				NegativePrefixStringPtr = &GeneratedNegativePrefixString;
 			}
-			if (InFormattingRules.NegativeSuffixString.Len() > 0)
+			if (!InFormattingRules.NegativeSuffixString.IsEmpty())
 			{
 				NegativeSuffixStringPtr = &GeneratedNegativeSuffixString;
 			}
-			if (InFormattingRules.PositivePrefixString.Len() > 0)
+			if (!InFormattingRules.PositivePrefixString.IsEmpty())
 			{
 				PositivePrefixStringPtr = &GeneratedPositivePrefixString;
 			}
-			if (InFormattingRules.PositiveSuffixString.Len() > 0)
+			if (!InFormattingRules.PositiveSuffixString.IsEmpty())
 			{
 				PositiveSuffixStringPtr = &GeneratedPositiveSuffixString;
 			}
 		}
 	}
 
-	bool HasNegativePrefixString() const
-	{
-		return NegativePrefixStringPtr->Len() > 0;
-	}
-
 	const FString& GetNegativePrefixString() const
 	{
 		return *NegativePrefixStringPtr;
 	}
 
-	bool HasNegativeSuffixString() const
-	{
-		return NegativeSuffixStringPtr->Len() > 0;
-	}
-
 	const FString& GetNegativeSuffixString() const
 	{
 		return *NegativeSuffixStringPtr;
 	}
 
-	bool HasPositivePrefixString() const
-	{
-		return PositivePrefixStringPtr->Len() > 0;
-	}
-
 	const FString& GetPositivePrefixString() const
 	{
 		return *PositivePrefixStringPtr;
 	}
 
-	bool HasPositiveSuffixString() const
-	{
-		return PositiveSuffixStringPtr->Len() > 0;
-	}
-
 	const FString& GetPositiveSuffixString() const
 	{
 		return *PositiveSuffixStringPtr;
