Added support decoding a single encoded I/O store chunk block

diff --git a/Engine/Source/Runtime/Core/Private/IO/IoChunkEncoding.cpp b/Engine/Source/Runtime/Core/Private/IO/IoChunkEncoding.cpp
--- a/Engine/Source/Runtime/Core/Private/IO/IoChunkEncoding.cpp
+++ b/Engine/Source/Runtime/Core/Private/IO/IoChunkEncoding.cpp
@@ -352,6 +352,63 @@ bool FIoChunkEncoding::Decode(const FIoChunkDecodingParams& Params, FMemoryView
 	return true;
 }
 
+FIoChunkBlockDecodeResult FIoChunkEncoding::DecodeBlock(const FIoChunkBlockDecodingParams Params, FMutableMemoryView EncodedBlock, FMutableMemoryView DecodedBlock)
+{
+	// Assume that the underlying memory size is aligned to AES block size 
+	const uint32 AlignedEncodedSize = Align(uint32(EncodedBlock.GetSize()), FAES::AESBlockSize);
+
+	if (Params.BlockHash.IsEmpty() == false)
+	{
+		check(Params.HashFunction == EIoSignatureHashFunction::Sha1);
+		FSHAHash ShaHash;
+		FSHA1::HashBuffer(EncodedBlock.GetData(), AlignedEncodedSize, ShaHash.Hash);
+		FMemoryView BlockHash(ShaHash.Hash, sizeof(ShaHash.Hash));
+
+		if (0 != Params.BlockHash.CompareBytes(BlockHash))
+		{
+			return FIoChunkBlockDecodeResult
+			{
+				.ErrorCode = EIoErrorCode::SignatureError
+			};
+		}
+	}
+
+	if (Params.EncryptionKey.IsEmpty() == false)
+	{
+		FAES::FAESKey AESKey;
+		FMutableMemoryView(AESKey.Key, FAES::FAESKey::KeySize).CopyFrom(Params.EncryptionKey);
+		FAES::DecryptData((uint8*)EncodedBlock.GetData(), AlignedEncodedSize, AESKey);
+	}
+
+	if (Params.CompressionFormat.IsNone())
+	{
+		return FIoChunkBlockDecodeResult
+		{
+			.DecodedBlock	= FMutableMemoryView(EncodedBlock.GetData(), DecodedBlock.GetSize()),
+			.ErrorCode		= EIoErrorCode::Ok
+		};
+	}
+
+	if (LIKELY(FCompression::UncompressMemory(
+		Params.CompressionFormat,
+		DecodedBlock.GetData(),
+		IntCastChecked<int32>(DecodedBlock.GetSize()),
+		EncodedBlock.GetData(),
+		IntCastChecked<int32>(EncodedBlock.GetSize()))))
+	{
+		return FIoChunkBlockDecodeResult
+		{
+			.DecodedBlock	= FMutableMemoryView(DecodedBlock.GetData(), DecodedBlock.GetSize()),
+			.ErrorCode		= EIoErrorCode::Ok
+		};
+	}
+
+	return FIoChunkBlockDecodeResult
+	{
+		.ErrorCode = EIoErrorCode::CompressionError
+	};
+}
+
 ////////////////////////////////////////////////////////////////////////////////
 TIoStatusOr<FIoOffsetAndLength> FIoChunkEncoding::GetChunkRange(uint64 TotalRawSize, uint32 RawBlockSize, TConstArrayView<uint32> EncodedBlockSize, uint64 RawOffset, uint64 RawSize)
 {
diff --git a/Engine/Source/Runtime/Core/Public/IO/IoChunkEncoding.h b/Engine/Source/Runtime/Core/Public/IO/IoChunkEncoding.h
--- a/Engine/Source/Runtime/Core/Public/IO/IoChunkEncoding.h
+++ b/Engine/Source/Runtime/Core/Public/IO/IoChunkEncoding.h
@@ -46,9 +46,30 @@ enum class EIoDecodeFlags : uint8
 	/** Avoid operations that read from the output buffer provided. This option may use more temp scratch memory */
 	WriteOnly	= 1 << 0
 };
-
 ENUM_CLASS_FLAGS(EIoDecodeFlags);
 
+/** List of supported signature hash functions used when decoding an I/O store chunk block. */
+enum class EIoSignatureHashFunction : uint8
+{
+	Sha1
+};
+
+/** Parameters for decoding a single I/O store chunk block. */
+struct FIoChunkBlockDecodingParams
+{
+	FMemoryView					EncryptionKey;
+	FMemoryView					BlockHash;
+	FName						CompressionFormat;
+	EIoSignatureHashFunction	HashFunction = EIoSignatureHashFunction::Sha1;
+};
+
+/** Holds the result from decoding a single I/O store chunk block. */
+struct FIoChunkBlockDecodeResult
+{
+	FMutableMemoryView	DecodedBlock;
+	EIoErrorCode		ErrorCode = EIoErrorCode::Unknown;
+};
+
 /**
  * Encodes data into a set of encrypted and compressed blocks.
  * The chunk encoding information is enocoded into a 16 byte
@@ -87,6 +108,19 @@ public:
 	static CORE_API bool Encode(const FIoChunkEncodingParams& Params, FMemoryView RawData, FIoBuffer& OutHeader, FIoBuffer& OutEncodedBlocks);
 	static CORE_API bool Decode(const FIoChunkDecodingParams& Params, FMemoryView EncodedBlocks, FMutableMemoryView OutRawData, EIoDecodeFlags Flags = EIoDecodeFlags::None);
 	static CORE_API bool Decode(FMemoryView EncodedData, FName CompressionFormat, FMemoryView EncryptionKey, FMutableMemoryView OutRawData, uint64 Offset = 0);
+	/**
+	 * Decode a single I/O store chunk block.
+	 * 
+	 * @param Params			Controls whether this function:
+	 *								- computes and verifys the block signature using the provided block hash
+	 *								- decrypts the block using the provided encryption key
+	 *								- decompress the source buffer to the destinatino buffer
+	 * @param EncodedBlock		Memory view of the encoded source buffer. Assumes the underlying size is aligned to AES block size.
+	 * @param OutDecodedBlock	Memory view of the destination buffer.
+	 *
+	 * @note This function does not copy the bytes from the encoded block unless it is compressed.
+	 */
+	static CORE_API FIoChunkBlockDecodeResult DecodeBlock(const FIoChunkBlockDecodingParams Params, FMutableMemoryView EncodedBlock, FMutableMemoryView OutDecodedBlock);
 
 	static CORE_API TIoStatusOr<FIoOffsetAndLength> GetChunkRange(uint64 TotalRawSize, uint32 RawBlockSize, TConstArrayView<uint32> EncodedBlockSize, uint64 RawOffset, uint64 RawSize);
 	static CORE_API TIoStatusOr<FIoOffsetAndLength> GetChunkRange(const FIoChunkDecodingParams& Params, uint64 RawSize);
