Added interface for decoding I/O store chunk blocks asynchronously.

diff --git a/Engine/Source/Runtime/Core/Private/IO/IoChunkEncoding.cpp b/Engine/Source/Runtime/Core/Private/IO/IoChunkEncoding.cpp
--- a/Engine/Source/Runtime/Core/Private/IO/IoChunkEncoding.cpp
+++ b/Engine/Source/Runtime/Core/Private/IO/IoChunkEncoding.cpp
@@ -468,3 +468,18 @@ FIoBlockHash FIoChunkEncoding::HashBlock(FMemoryView Block)
 	FMemory::Memcpy(&BlockHash, &Hash, sizeof(FIoBlockHash));
 	return BlockHash;
 }
+
+////////////////////////////////////////////////////////////////////////////////
+void FIoChunkBlockDecoder::Set(IIoChunkBlockDecoder& InDefault)
+{
+	check(Default == nullptr);
+	Default = &InDefault;
+}
+
+IIoChunkBlockDecoder& FIoChunkBlockDecoder::Get()
+{
+	check(Default);
+	return *Default;
+}
+
+IIoChunkBlockDecoder* FIoChunkBlockDecoder::Default = nullptr;
diff --git a/Engine/Source/Runtime/Core/Public/IO/IoChunkEncoding.h b/Engine/Source/Runtime/Core/Public/IO/IoChunkEncoding.h
--- a/Engine/Source/Runtime/Core/Public/IO/IoChunkEncoding.h
+++ b/Engine/Source/Runtime/Core/Public/IO/IoChunkEncoding.h
@@ -6,6 +6,7 @@
 #include "IO/IoStatus.h"
 #include "Memory/MemoryView.h"
 #include "Misc/EnumClassFlags.h"
+#include "Templates/FunctionFwd.h"
 #include "UObject/NameTypes.h"
 
 class FIoBuffer;
@@ -127,3 +128,25 @@ public:
 	static CORE_API uint64 GetTotalEncodedSize(TConstArrayView<uint32> EncodedBlockSize);
 	static CORE_API FIoBlockHash HashBlock(FMemoryView Block);
 };
+
+/** Callback used when decoding I/O store chunk blocks asynchronously. */
+using FOnBlockDecoded = TUniqueFunction<void(FIoChunkBlockDecodeResult&&)>;
+
+/** Interface for decoding I/O store chunk blocks asynchronously. */
+class IIoChunkBlockDecoder
+{
+public:
+	virtual			~IIoChunkBlockDecoder() = default;
+	virtual void	Enqueue(FIoChunkBlockDecodingParams Params, FMutableMemoryView EncodedBlock, FMutableMemoryView OutDecodedBlock, FOnBlockDecoded&& OnCompleted) = 0;
+};
+
+/** Provides access to the chunk block decoder. */
+class FIoChunkBlockDecoder
+{
+public:
+	CORE_API static void					Set(IIoChunkBlockDecoder& Default); 
+	CORE_API static IIoChunkBlockDecoder&	Get();
+
+private:
+	static IIoChunkBlockDecoder* Default;
+};
