Made FGenericPlatformMemory::GetMemorySizeBucket() thread safe to avoid it wrongly returning EPlatformMemorySizeBucket::Default in race conditions

#rnx

diff --git a/Engine/Source/Runtime/Core/Private/GenericPlatform/GenericPlatformMemory.cpp b/Engine/Source/Runtime/Core/Private/GenericPlatform/GenericPlatformMemory.cpp
--- a/Engine/Source/Runtime/Core/Private/GenericPlatform/GenericPlatformMemory.cpp
+++ b/Engine/Source/Runtime/Core/Private/GenericPlatform/GenericPlatformMemory.cpp
@@ -364,6 +364,10 @@ void FGenericPlatformMemory::DumpPlatformAndAllocatorStats( class FOutputDevice&
 
 EPlatformMemorySizeBucket FGenericPlatformMemory::GetMemorySizeBucket()
 {
+	// Make sure this won't be called by more than one thread to avoid race conditions in the returned value
+	static FCriticalSection Lock;
+	FScopeLock ScopeLock(&Lock);
+	
 	static bool bCalculatedBucket = false;
 	static EPlatformMemorySizeBucket Bucket = EPlatformMemorySizeBucket::Default;
 
