Disable recently added FConsoleManager tests because they cause an issue on PS4 due to static initialization order and logging grabbing a TLS slot that other code epxects to use

diff --git a/Engine/Source/Runtime/Core/Private/HAL/ConsoleManager.cpp b/Engine/Source/Runtime/Core/Private/HAL/ConsoleManager.cpp
--- a/Engine/Source/Runtime/Core/Private/HAL/ConsoleManager.cpp
+++ b/Engine/Source/Runtime/Core/Private/HAL/ConsoleManager.cpp
@@ -3897,6 +3897,11 @@ void FConsoleManager::Test()
 
 	UnregisterConsoleVariableSink_Handle(TestSinkCallbackHandle);
 
+
+	// TODO: These tests end up causing an issue on PS4 due to static initialization order and
+	// logging grabbing a TLS slot the a PS4 tool expects to use. Try moving them to proper unit
+	// tests instead.
+	#if false
 	FConsoleManager& Manager = GetManager();
 
 	// Deprecated console variables should not assert in dump commands
@@ -3926,6 +3931,7 @@ void FConsoleManager::Test()
 		DumpObjects(Manager.ConsoleObjects, TEXT("-ShowHelp"), OutNull, true);
 		DumpObjects(Manager.ConsoleObjects, TEXT("-ShowHelp -Deprecated"), OutNull, true);
 	}
+	#endif
 
 #endif // !(UE_BUILD_SHIPPING || UE_BUILD_TEST)
 }
