ShaderCompilerWorker: Fix for crashes whiile loading a system DLL on one of the platforms caused by us illegaly setting value at TLS index 0 (caused by static initialization fiasco).

#rnx

diff --git a/Engine/Source/Runtime/Core/Private/Logging/LogScopedCategoryAndVerbosityOverride.cpp b/Engine/Source/Runtime/Core/Private/Logging/LogScopedCategoryAndVerbosityOverride.cpp
--- a/Engine/Source/Runtime/Core/Private/Logging/LogScopedCategoryAndVerbosityOverride.cpp
+++ b/Engine/Source/Runtime/Core/Private/Logging/LogScopedCategoryAndVerbosityOverride.cpp
@@ -18,14 +18,23 @@ FLogScopedCategoryAndVerbosityOverride::~FLogScopedCategoryAndVerbosityOverride(
 	*TLS = Backup;
 }
 
-static uint32 OverrrideTLSID = FPlatformTLS::AllocTlsSlot();
+// OverrideTLSID is stored as a local static to make sure we initialize before use, even
+// it it happens during the static initialization. Otherwise we could end up setting
+// values to the 0 TLS slot which doesn't belong to us.
+uint32 GetOverrideTLSID()
+{
+	static uint32 OverrrideTLSID = FPlatformTLS::AllocTlsSlot();
+
+	return OverrrideTLSID;
+}
+
 FLogScopedCategoryAndVerbosityOverride::FOverride* FLogScopedCategoryAndVerbosityOverride::GetTLSCurrent()
 {
-	FOverride* TLS = (FOverride*)FPlatformTLS::GetTlsValue(OverrrideTLSID);
+	FOverride* TLS = (FOverride*)FPlatformTLS::GetTlsValue(GetOverrideTLSID());
 	if (!TLS)
 	{
 		TLS = new FOverride;
-		FPlatformTLS::SetTlsValue(OverrrideTLSID, TLS);
+		FPlatformTLS::SetTlsValue(GetOverrideTLSID(), TLS);
 	}	
 	return TLS;
 }
