Optimize hitset hashing to improve performance by 5% across our benchmarks.

diff --git a/Engine/Source/Runtime/Core/Private/AutoRTFM/HitSet.h b/Engine/Source/Runtime/Core/Private/AutoRTFM/HitSet.h
--- a/Engine/Source/Runtime/Core/Private/AutoRTFM/HitSet.h
+++ b/Engine/Source/Runtime/Core/Private/AutoRTFM/HitSet.h
@@ -40,16 +40,14 @@ namespace AutoRTFM
 		static constexpr size_t MaxSize = 16;
 
 	private:
-		// TODO: Revisit a good probe depth for the hashset.
-		static constexpr uint64_t LinearProbeDepth = 16;
-
-		// TODO: Revisit a good initial capacity for the hitset.
-		static constexpr uint64_t LogInitialCapacity = 4;
-
+		static constexpr uint8_t LinearProbeDepth = 32;
+		static constexpr uint8_t LogInitialCapacity = 6;
 		static constexpr uint64_t InitialCapacity = static_cast<uint64_t>(1) << LogInitialCapacity;
 
 		// TODO: Revisit a good max capacity for the hitset.
-		static constexpr uint64_t MaxCapacity = (128 * 1024 * 1024) / sizeof(FHitSetEntry);
+		static constexpr uint64_t MaxCapacity = static_cast<uint64_t>(128 * 1024 * 1024) / sizeof(FHitSetEntry);
+
+		static_assert(LinearProbeDepth <= InitialCapacity);
 	public:
 
 		~FHitSet()
@@ -130,11 +128,12 @@ namespace AutoRTFM
 	private:
 		uintptr_t SmallPayload[InitialCapacity]{};
 		uintptr_t* Payload = SmallPayload;
-		uint64_t SixtyFourMinusLogCapacity = 64 - LogInitialCapacity;
 		uint64_t Count = 0;
+		uint32_t SixtyFourMinusLogCapacity = 64 - LogInitialCapacity;
 
 		UE_AUTORTFM_FORCEINLINE uint64_t Capacity() const
 		{
+			UE_ASSUME(SixtyFourMinusLogCapacity < 64);
 			return static_cast<uint64_t>(1) << (64 - SixtyFourMinusLogCapacity);
 		}
 
@@ -158,7 +157,7 @@ namespace AutoRTFM
 			}
 
 			uintptr_t* const OldPayload = Payload;
-			const uint64_t OldSize = Count;
+			const uint64_t OldCount = Count;
 
 			while (true)
 			{
@@ -205,7 +204,7 @@ namespace AutoRTFM
 				break;
 			}
 
-			ASSERT(OldSize == Count);
+			ASSERT(OldCount == Count);
 
 			if (OldPayload != SmallPayload)
 			{
@@ -215,27 +214,12 @@ namespace AutoRTFM
 			return true;
 		}
 
-		UE_AUTORTFM_FORCEINLINE uintptr_t FirstHashInRange(const uintptr_t Hashee) const
+		UE_AUTORTFM_FORCEINLINE uintptr_t FibonacciHash(const uintptr_t Hashee) const
 		{
 			constexpr uintptr_t Fibonacci = 0x9E3779B97F4A7C15;
 			return (Hashee * Fibonacci) >> SixtyFourMinusLogCapacity;
 		}
 
-		UE_AUTORTFM_FORCEINLINE uintptr_t SecondHash(const uintptr_t Hashee) const
-		{
-			// TODO: Find a great hashing function for the set.
-			constexpr uintptr_t A = 0x65d200ce55b19ad8;
-			constexpr uintptr_t B = 0x4f2162926e40c299;
-			constexpr uintptr_t C = 0x162dd799029970f8;
-
-			const uintptr_t Low = Hashee & 0x00000000FFFFFFFF;
-			const uintptr_t High = Hashee >> 32;
-
-			const uintptr_t Hash = A * Low + B * High + C;
-
-			return Hash;
-		}
-
 		UE_AUTORTFM_FORCEINLINE EInsertResult TryInsertAtIndex(const uintptr_t Raw, const uintptr_t I)
 		{
 			// We have a free location in the set.
@@ -259,9 +243,12 @@ namespace AutoRTFM
 		// succeeded (EG. the key was not already in the set).
 		UE_AUTORTFM_FORCEINLINE EInsertResult InsertNoResize(const uintptr_t Raw)
 		{
-			// First try to use our simple first hash to get a location.
+			const uintptr_t Hash = FibonacciHash(Raw);
+
+			// The first index will always be in the range of capacity because we
+			// are using a fibonacci hash. So just check if we can insert here.
 			{
-				const uintptr_t I = FirstHashInRange(Raw);
+				const uintptr_t I = Hash;
 
 				const EInsertResult R = TryInsertAtIndex(Raw, I);
 
@@ -269,19 +256,21 @@ namespace AutoRTFM
 				{
 					return R;
 				}
+
+				// Otherwise we need to do a linear probe...
 			}
 
-			// Then, if that fails, we use a second more complicated hash with a linear probe.
-			const uintptr_t Hash = SecondHash(Raw);
+			const uintptr_t Mask = Capacity() - 1;
 
 			// Clang goes way over the top with loop unrolling, and makes this code noticably
 			// slower as a result. So we disable it!
 #ifdef __clang__
 #pragma clang loop unroll(disable)
 #endif
-			for (uint64_t D = 0; D < LinearProbeDepth; D++)
+			for (uint8_t D = 1; D < LinearProbeDepth; D++)
 			{
-				const uintptr_t I = FirstHashInRange(Hash + D);
+				// Capacity is always a power of 2, so we can just mask out the bits.
+				const uintptr_t I = (Hash + D) & Mask;
 
 				const EInsertResult R = TryInsertAtIndex(Raw, I);
 
