Bumping translator guid to hopefully fix a shader compilation issue due to not picked up changes from materialtemplates.

#rnx

diff --git a/Engine/Source/Runtime/Core/Private/UObject/FortniteMainBranchObjectVersion.cpp b/Engine/Source/Runtime/Core/Private/UObject/FortniteMainBranchObjectVersion.cpp
--- a/Engine/Source/Runtime/Core/Private/UObject/FortniteMainBranchObjectVersion.cpp
+++ b/Engine/Source/Runtime/Core/Private/UObject/FortniteMainBranchObjectVersion.cpp
@@ -14,6 +14,6 @@ TMap<FGuid, FGuid> FFortniteMainBranchObjectVersion::GetSystemGuids()
 	SystemGuids.Add(DevGuids.Niagara_LatestScriptCompileVersion, FGuid("670A2EBDEE1A4794BB77BBFF331EA7EF"));
 	SystemGuids.Add(DevGuids.SkeletalMeshDerivedDataVersion, FGuid("2F90180769024483BB7D90AACEB2798D"));
 	SystemGuids.Add(DevGuids.STATICMESH_DERIVEDDATA_VER, FGuid("656E338C78C34641A5D3BFAE303997AF"));
-	SystemGuids.Add(DevGuids.MaterialTranslationDDCVersion, FGuid("38a1c24b1eb84bccbc79ff8fd7b73a72"));
+	SystemGuids.Add(DevGuids.MaterialTranslationDDCVersion, FGuid("1F1845B6EF8F49ED880DE4EA60BF2EA4"));
 	return SystemGuids;
 }
