[Truncation Warnings] Fix a truncation warning in ConcurrentLinearAllocator

#rnx

diff --git a/Engine/Source/Runtime/Core/Public/Experimental/ConcurrentLinearAllocator.h b/Engine/Source/Runtime/Core/Public/Experimental/ConcurrentLinearAllocator.h
--- a/Engine/Source/Runtime/Core/Public/Experimental/ConcurrentLinearAllocator.h
+++ b/Engine/Source/Runtime/Core/Public/Experimental/ConcurrentLinearAllocator.h
@@ -767,7 +767,7 @@ public:
 	inline void* Malloc(SIZE_T Size, uint32 Alignment)
 	{
 		SIZE_T TotalSize = Align(sizeof(FAllocation), Alignment) + Size;
-		FAllocation* Allocation = new (TConcurrentLinearAllocator<BlockAllocationTag>::Malloc(TotalSize, FMath::Max<size_t>(alignof(FAllocation), Alignment))) FAllocation();
+		FAllocation* Allocation = new (TConcurrentLinearAllocator<BlockAllocationTag>::Malloc(TotalSize, FMath::Max(static_cast<uint32>(alignof(FAllocation)), Alignment))) FAllocation();
 
 		void* Alloc = Align(Allocation + 1, Alignment);
 		checkSlow(uintptr_t(Alloc) + Size - uintptr_t(Allocation) <= TotalSize);
