Adding a way to check for Windows 11 support at runtime.

diff --git a/Engine/Source/Runtime/Core/Private/Windows/WindowsPlatformCrashContext.cpp b/Engine/Source/Runtime/Core/Private/Windows/WindowsPlatformCrashContext.cpp
--- a/Engine/Source/Runtime/Core/Private/Windows/WindowsPlatformCrashContext.cpp
+++ b/Engine/Source/Runtime/Core/Private/Windows/WindowsPlatformCrashContext.cpp
@@ -309,6 +309,7 @@ void FWindowsPlatformCrashContext::AddPlatformSpecificProperties() const
 {
 	AddCrashProperty(TEXT("PlatformIsRunningWindows"), 1);
 	AddCrashProperty(TEXT("PlatformIsRunningWine"), FWindowsPlatformMisc::IsWine());
+	AddCrashProperty(TEXT("PlatformSupportsWindows11"), FWindowsPlatformMisc::SupportsWindows11());
 	AddCrashProperty(TEXT("IsRunningOnBattery"), FPlatformMisc::IsRunningOnBattery());
 	WIDECHAR DriveName = 0;
 	const TCHAR* BaseDir = FWindowsPlatformProcess::BaseDir();
diff --git a/Engine/Source/Runtime/Core/Private/Windows/WindowsPlatformMisc.cpp b/Engine/Source/Runtime/Core/Private/Windows/WindowsPlatformMisc.cpp
--- a/Engine/Source/Runtime/Core/Private/Windows/WindowsPlatformMisc.cpp
+++ b/Engine/Source/Runtime/Core/Private/Windows/WindowsPlatformMisc.cpp
@@ -3333,6 +3333,127 @@ bool FWindowsPlatformMisc::HasAVX2InstructionSupport()
 #endif
 }
 
+#include "Windows/AllowWindowsPlatformTypes.h"
+#include <tbs.h>
+bool FWindowsPlatformMisc::HasTPM2Support()
+{
+	auto HasTPM2SupportInternal = []() -> bool
+	{
+#if (NTDDI_VERSION >= NTDDI_WIN8)
+		TPM_DEVICE_INFO TpmInfo;
+		FMemory::Memzero(TpmInfo);
+		if (Tbsi_GetDeviceInfo(sizeof(TpmInfo), &TpmInfo) == TBS_SUCCESS)
+		{
+			return TpmInfo.tpmVersion == TPM_VERSION_20;
+		}
+#else
+		// Until _WIN32_WINNT is set to _WIN32_WINNT_WIN8 we will need to check TPM status dynamically
+		if (void* hTbs = FPlatformProcess::GetDllHandle(TEXT("tbs.dll")))
+		{
+			struct TPM_DEVICE_INFO
+			{
+				UINT32 structVersion;
+				UINT32 tpmVersion; // 1.2 = 1, 2.0 = 2
+				UINT32 tpmInterfaceType;
+				UINT32 tpmImpRevision;
+			};
+			// https://learn.microsoft.com/en-us/windows/win32/api/tbs/nf-tbs-tbsi_getdeviceinfo
+			typedef TBS_RESULT (WINAPI *pfnTbsiGetDeviceInfo)(UINT32 Size, PVOID Info);
+			if (const pfnTbsiGetDeviceInfo Tbsi_GetDeviceInfo = reinterpret_cast<pfnTbsiGetDeviceInfo>(FPlatformProcess::GetDllExport(hTbs, TEXT("Tbsi_GetDeviceInfo"))))
+			{
+				TPM_DEVICE_INFO TpmInfo;
+				FMemory::Memzero(TpmInfo);
+				if (Tbsi_GetDeviceInfo(sizeof(TpmInfo), &TpmInfo) == TBS_SUCCESS)
+				{
+					FPlatformProcess::FreeDllHandle(hTbs);
+					return TpmInfo.tpmVersion == 2;
+				}
+				UE_LOG(LogWindows, Warning, TEXT("Failed to retrieve TPM device information"));
+			}
+			else
+			{
+				UE_LOG(LogWindows, Warning, TEXT("Failed to load proc address for Tbsi_GetDeviceInfo"));
+			}
+			FPlatformProcess::FreeDllHandle(hTbs);
+		}
+		else
+		{
+			UE_LOG(LogWindows, Warning, TEXT("Failed to load tbs.dll"));
+		}
+#endif
+		return false;
+	};
+	static bool bSupportsTpm2 = HasTPM2SupportInternal();
+	return bSupportsTpm2;
+}
+
+#include "winternl.h"
+#include "ntstatus.h"
+bool FWindowsPlatformMisc::HasSecureBootSupport()
+{
+	auto HasSecureBootSupportInternal = []() -> bool
+	{
+		if (void* hNtDll = FPlatformProcess::GetDllHandle(TEXT("ntdll.dll")))
+		{
+			// Using an undocumented SYSTEM_INFORMATION_CLASS type SYSTEM_SECUREBOOT_INFORMATION
+			static SYSTEM_INFORMATION_CLASS SecureBootInformationClassId = static_cast<SYSTEM_INFORMATION_CLASS>(0x91);
+			struct SYSTEM_SECUREBOOT_INFORMATION
+			{
+				bool bSecureBootEnabled;
+				bool bSecureBootCapable;
+			};
+			// https://learn.microsoft.com/en-us/windows/win32/api/winternl/nf-winternl-ntquerysysteminformation
+			typedef NTSTATUS (*pfnNtQuerySystemInformation)(SYSTEM_INFORMATION_CLASS SystemInformationClass, PVOID SystemInformation, ULONG SystemInformationLength, PULONG ReturnLength);
+			if (const pfnNtQuerySystemInformation fnNtQuerySystemInformation = reinterpret_cast<pfnNtQuerySystemInformation>(FPlatformProcess::GetDllExport(hNtDll, TEXT("NtQuerySystemInformation"))))
+			{
+				SYSTEM_SECUREBOOT_INFORMATION SecureBootInfo;
+				FMemory::Memzero(SecureBootInfo);
+				if (fnNtQuerySystemInformation(SecureBootInformationClassId, &SecureBootInfo, sizeof(SecureBootInfo), nullptr) == STATUS_SUCCESS)
+				{
+					FPlatformProcess::FreeDllHandle(hNtDll);
+					return SecureBootInfo.bSecureBootCapable;
+				}
+				UE_LOG(LogWindows, Warning, TEXT("Failed to retrieve secure boot information"));
+			}
+			else
+			{
+				UE_LOG(LogWindows, Warning, TEXT("Failed to load proc address for NtQuerySystemInformation"));
+			}
+			FPlatformProcess::FreeDllHandle(hNtDll);
+		}
+		else
+		{
+			UE_LOG(LogWindows, Warning, TEXT("Failed to load ntdll.dll"));
+		}
+		return false;
+	};
+	static bool bHasSecureBootSupport = HasSecureBootSupportInternal();
+	return bHasSecureBootSupport;
+}
+#include "Windows/HideWindowsPlatformTypes.h"
+
+bool FWindowsPlatformMisc::SupportsWindows11()
+{
+	static struct FWindows11SupportCheck
+	{
+		bool bSupported;
+		
+		FWindows11SupportCheck()
+		{
+			// https://www.microsoft.com/en-us/windows/windows-11-specifications?r=1
+			// Note that this check is mostly complete, but does not check for:
+			// - Whether the CPU is on Microsoft's approved processor list https://aka.ms/CPUlist
+			// - DX12 Support
+			bSupported = HasTPM2Support();
+			bSupported &= HasSecureBootSupport();
+			bSupported &= NumberOfCores() >= 2;
+			bSupported &= FPlatformMemory::GetPhysicalGBRam() >= 4;
+		}
+	} Win11Check;
+	
+	return Win11Check.bSupported;
+}
+
 #include "Windows/AllowWindowsPlatformTypes.h"
 FString FWindowsPlatformMisc::GetPrimaryGPUBrand()
 {
diff --git a/Engine/Source/Runtime/Core/Public/Windows/WindowsPlatformMisc.h b/Engine/Source/Runtime/Core/Public/Windows/WindowsPlatformMisc.h
--- a/Engine/Source/Runtime/Core/Public/Windows/WindowsPlatformMisc.h
+++ b/Engine/Source/Runtime/Core/Public/Windows/WindowsPlatformMisc.h
@@ -244,6 +244,14 @@ struct FWindowsPlatformMisc
 	 * @return	Returns true if instruction-set is supported
 	 */
 	static CORE_API bool HasAVX2InstructionSupport();
+	
+	static CORE_API bool HasTPM2Support();
+	static CORE_API bool HasSecureBootSupport();
+	/**
+	 * Checks whether this device has the required hardware and software capabilities required to run Windows 11.
+	 * @return True if the Windows 11 min spec is met.
+	 */
+	static CORE_API bool SupportsWindows11();
 
 	static CORE_API FString GetCPUVendor();
 	static CORE_API FString GetCPUBrand();
