Removed deprecated functionality from GenericPlatformString.h.

diff --git a/Engine/Source/Runtime/Core/Public/GenericPlatform/GenericPlatformString.h b/Engine/Source/Runtime/Core/Public/GenericPlatform/GenericPlatformString.h
--- a/Engine/Source/Runtime/Core/Public/GenericPlatform/GenericPlatformString.h
+++ b/Engine/Source/Runtime/Core/Public/GenericPlatform/GenericPlatformString.h
@@ -59,39 +59,6 @@ static_assert(sizeof(UNICODE_BOGUS_CHAR_CODEPOINT) <= sizeof(ANSICHAR) && (UNICO
  */
 struct FGenericPlatformString : public FGenericPlatformStricmp
 {
-	/**
-	 * Tests whether an encoding has fixed-width characters
-	 */
-	template <typename Encoding>
-	UE_DEPRECATED(5.1, "FPlatformString::IsFixedWidthEncoding<T>() has been deprecated in favor of TIsFixedWidthCharEncoding_V<T>")
-	static constexpr bool IsFixedWidthEncoding()
-	{
-		return TIsFixedWidthCharEncoding_V<Encoding>;
-	}
-
-	/**
-	 * Function which returns whether one encoding type is binary compatible with another.
-	 *
-	 * Unlike TAreEncodingsCompatible, this is not commutative.  For example, ANSI is compatible with
-	 * UTF-8, but UTF-8 is not compatible with ANSI.
-	 */
-	template <typename SrcEncoding, typename DestEncoding>
-	UE_DEPRECATED(5.1, "FPlatformString::IsCharEncodingCompatibleWith<A, B>() has been deprecated in favor of TIsCharEncodingCompatibleWith_V<A, B>")
-	static constexpr bool IsCharEncodingCompatibleWith()
-	{
-		return TIsCharEncodingCompatibleWith_V<SrcEncoding, DestEncoding>;
-	}
-
-	/**
-	 * Tests whether you can simply (i.e. by assignment) encode code units from the source encoding as the destination encoding.
-	 */
-	template <typename SourceEncoding, typename DestEncoding>
-	UE_DEPRECATED(5.1, "FPlatformString::IsCharEncodingSimplyConvertibleTo<A, B>() has been deprecated in favor of TIsCharEncodingSimplyConvertibleTo_V<A, B>")
-	static constexpr bool IsCharEncodingSimplyConvertibleTo()
-	{
-		return TIsCharEncodingSimplyConvertibleTo_V<SourceEncoding, DestEncoding>;
-	}
-
 	/**
 	 * Tests whether a particular codepoint can be converted to the destination encoding.
 	 *
@@ -158,31 +125,6 @@ struct FGenericPlatformString : public FGenericPlatformStricmp
 	static constexpr bool IsUnicodeEncoded = true;
 
 
-	/**
-	 * Metafunction which tests whether a given character type represents a fixed-width encoding.
-	 */
-	template <typename T>
-	struct UE_DEPRECATED(5.0, "TIsFixedWidthEncoding is deprecated, use TIsFixedWidthCharEncoding_V<T> instead.") TIsFixedWidthEncoding
-	{
-		enum { Value = TIsFixedWidthCharEncoding_V<T> };
-	};
-
-
-	/**
-	 * Metafunction which tests whether two encodings are compatible.
-	 *
-	 * We'll say the encodings are compatible if they're both fixed-width and have the same size.  This
-	 * should be good enough and catches things like UCS2CHAR and WIDECHAR being equivalent.
-	 * Specializations of this template can be provided for any other special cases.
-	 * Same size is a minimum requirement.
-	 */
-	template <typename EncodingA, typename EncodingB>
-	struct UE_DEPRECATED(5.0, "TAreEncodingsCompatible is deprecated, use TIsCharEncodingCompatibleWith_V<SrcEncoding, DestEncoding> instead.") TAreEncodingsCompatible
-	{
-		enum { Value = TIsFixedWidthCharEncoding_V<EncodingA> && TIsFixedWidthCharEncoding_V<EncodingB> && sizeof(EncodingA) == sizeof(EncodingB) };
-	};
-
-
 	/**
 	 * Converts the null-terminated Src string range from SourceEncoding to DestEncoding and writes it to the [Dest, Dest+DestSize) range, including a null terminator.
 	 * If the Dest range is not big enough to hold the converted output, NULL is returned.  In this case, nothing should be assumed about the contents of Dest.
