Fixed minor file name typo in comment

#rnx

diff --git a/Engine/Source/Runtime/Core/Private/Stats/StatsSystem.cpp b/Engine/Source/Runtime/Core/Private/Stats/StatsSystem.cpp
--- a/Engine/Source/Runtime/Core/Private/Stats/StatsSystem.cpp
+++ b/Engine/Source/Runtime/Core/Private/Stats/StatsSystem.cpp
@@ -15,7 +15,7 @@
 #include "Stats/StatsSystemTypes.h"
 #include "Stats/StatIgnoreList.h"
 
-// Declared in GlobalStats.ini
+// Declared in GlobalStats.inl
 DEFINE_STAT(STAT_FrameTime);
 DEFINE_STAT(STAT_NamedMarker);
 DEFINE_STAT(STAT_SecondsPerCycle);
