Memory allocated in EventRef could be freed after allocators are torn down when -nothreading was used.

- Functional part of the fix is ensuring FGraphEventRef is cleared in StopThread()
- Waiting and reseting as well as logic about which barrier primitive to use was encapulated in one place to mitigate against future errors.
- Moved code to its own part of the file instead of arbitrarily between two FThreadStats methods.

#rnx

diff --git a/Engine/Source/Runtime/Core/Private/Stats/StatsSystemTypes.cpp b/Engine/Source/Runtime/Core/Private/Stats/StatsSystemTypes.cpp
--- a/Engine/Source/Runtime/Core/Private/Stats/StatsSystemTypes.cpp
+++ b/Engine/Source/Runtime/Core/Private/Stats/StatsSystemTypes.cpp
@@ -807,6 +807,46 @@ bool FStatNameAndInfo::GetSortByNameFrom(FName InLongName)
 	FStatsThread
 -----------------------------------------------------------------------------*/
 
+class FLastFrameBarrier
+{
+public:
+	void Reset();
+	void Wait();
+
+private:
+	UE::Tasks::FTask Barrier;
+	FGraphEventRef Event;
+};
+
+void FLastFrameBarrier::Reset()
+{
+	Event = {};
+	Barrier = {};
+}
+
+void FLastFrameBarrier::Wait()
+{
+	if (FPlatformProcess::SupportsMultithreading())
+	{
+		Barrier.Wait();
+		Barrier = GStatsPipe.Launch(UE_SOURCE_LOCATION, [] {});
+	}
+	else
+	{
+		if (Event.GetReference())
+		{
+			FTaskGraphInterface::Get().WaitUntilTaskCompletes(Event, ENamedThreads::GameThread_Local);
+		}
+
+		Event = TGraphTask<FNullGraphTask>::CreateTask(NULL, ENamedThreads::GameThread).
+			ConstructAndDispatchWhenReady(TStatId{}, ENamedThreads::GameThread);
+	}
+
+}
+
+static int32				CurrentEventIndex = 0;
+static FLastFrameBarrier	LastFrameBarriers[MAX_STAT_LAG];
+
 static TAutoConsoleVariable<int32> CVarDumpStatPackets(TEXT("DumpStatPackets"), 0, TEXT("If true, dump stat packets."));
 
 /** The rendering thread runnable object. */
@@ -1277,12 +1317,6 @@ void FThreadStats::StartThread()
 	UE_LOG(LogStats, Log, TEXT("Stats thread started at %f"), FPlatformTime::Seconds() - GStartTime);
 }
 
-static int32 CurrentEventIndex = 0;
-
-static UE::Tasks::FTask LastFramesBarriers[MAX_STAT_LAG];
-
-static FGraphEventRef LastFramesEvents[MAX_STAT_LAG];
-
 void FThreadStats::StopThread()
 {
 	// Nothing to stop if it was never started
@@ -1307,9 +1341,11 @@ void FThreadStats::StopThread()
 	{
 		FPlatformProcess::Yield();
 	}
-	for (int32 Index = 0; Index < MAX_STAT_LAG; Index++)
+
+	// release memory before the allocator is destroyed
+	for (FLastFrameBarrier& Barrier : LastFrameBarriers)
 	{
-		LastFramesBarriers[Index] = {}; // release memory before the allocator is destroyed
+		Barrier.Reset();
 	}
 }
 
@@ -1323,31 +1359,13 @@ void FThreadStats::WaitForStats()
 	check(IsInGameThread());
 	if (IsThreadingReady() && !bPrimaryDisableForever)
 	{
-		DECLARE_CYCLE_STAT(TEXT("WaitForStats"), STAT_WaitForStats, STATGROUP_Engine);
-
-		int32 EventIndex = (CurrentEventIndex + MAX_STAT_LAG - 1) % MAX_STAT_LAG;
-
-		if (FPlatformProcess::SupportsMultithreading())
 		{
+			DECLARE_CYCLE_STAT(TEXT("WaitForStats"), STAT_WaitForStats, STATGROUP_Engine);
 			SCOPE_CYCLE_COUNTER(STAT_WaitForStats);
-			LastFramesBarriers[EventIndex].Wait();
-			LastFramesBarriers[EventIndex] = GStatsPipe.Launch(UE_SOURCE_LOCATION, [] {});
+			int32 EventIndex = (CurrentEventIndex + MAX_STAT_LAG - 1) % MAX_STAT_LAG;
+			LastFrameBarriers[EventIndex].Wait();
+			CurrentEventIndex++;
 		}
-		else
-		{
-			{
-				SCOPE_CYCLE_COUNTER(STAT_WaitForStats);
-				if (LastFramesEvents[EventIndex].GetReference())
-				{
-					FTaskGraphInterface::Get().WaitUntilTaskCompletes(LastFramesEvents[EventIndex], ENamedThreads::GameThread_Local);
-				}
-			}
-
-			LastFramesEvents[EventIndex] = TGraphTask<FNullGraphTask>::CreateTask(NULL, ENamedThreads::GameThread).
-				ConstructAndDispatchWhenReady(TStatId{}, ENamedThreads::GameThread);
-		}
-
-		CurrentEventIndex++;
 
 #if	!UE_BUILD_SHIPPING
 		DebugLeakTest();
