[AutoRTFM] Add workaround for non-POD thread_local variables.

Ideally, the compiler would fix these up automatically with no intervention, but this will take a while to fully polish. In the short term, I've added an AutoRTFM macro which wraps thread_local variables in ALWAYS_OPEN. It compiles away to a no-op when AutoRTFM is off. This can be removed once we have proper compiler support.

diff --git a/Engine/Plugins/Animation/ControlRig/Source/ControlRig/Private/Rigs/FKControlRig.cpp b/Engine/Plugins/Animation/ControlRig/Source/ControlRig/Private/Rigs/FKControlRig.cpp
--- a/Engine/Plugins/Animation/ControlRig/Source/ControlRig/Private/Rigs/FKControlRig.cpp
+++ b/Engine/Plugins/Animation/ControlRig/Source/ControlRig/Private/Rigs/FKControlRig.cpp
@@ -2,6 +2,8 @@
 
 #include "Rigs/FKControlRig.h"
 #include "Animation/SmartName.h"
+#include "AutoRTFM/AutoRTFM.h"
+#include "Containers/StaticArray.h"
 #include "Engine/SkeletalMesh.h"
 #include "IControlRigObjectBinding.h"
 #include "Components/SkeletalMeshComponent.h"
@@ -29,7 +31,8 @@ FName UFKControlRig::GetControlName(const FName& InName, const ERigElementType&
 	{
 		if((InType == ERigElementType::Bone || InType == ERigElementType::Curve))
 		{
-			static thread_local TMap<FName, FName> NameToControlMappings[2];
+			using NameToControlMappingCache = TStaticArray<TMap<FName, FName>, 2>;
+			UE_AUTORTFM_DECLARE_THREAD_LOCAL_VAR(NameToControlMappingCache, NameToControlMappings);
 			TMap<FName, FName>& NameToControlMapping = NameToControlMappings[InType == ERigElementType::Bone ? 0 : 1];
 			if (const FName* CachedName = NameToControlMapping.Find(InName))
 			{
@@ -71,7 +74,9 @@ FName UFKControlRig::GetControlTargetName(const FName& InName, const ERigElement
 	if (InName != NAME_None)
 	{
 		check(InType == ERigElementType::Bone || InType == ERigElementType::Curve);
-		static thread_local TMap<FName, FName> NameToTargetMappings[2];
+
+		using NameToTargetMappingCache = TStaticArray<TMap<FName, FName>, 2>;
+		UE_AUTORTFM_DECLARE_THREAD_LOCAL_VAR(NameToTargetMappingCache, NameToTargetMappings);
 		TMap<FName, FName>& NameToTargetMapping = NameToTargetMappings[InType == ERigElementType::Bone ? 0 : 1];
 		if (const FName* CachedName = NameToTargetMapping.Find(InName))
 		{
diff --git a/Engine/Source/Programs/AutoRTFMTests/Private/Tests.cpp b/Engine/Source/Programs/AutoRTFMTests/Private/Tests.cpp
--- a/Engine/Source/Programs/AutoRTFMTests/Private/Tests.cpp
+++ b/Engine/Source/Programs/AutoRTFMTests/Private/Tests.cpp
@@ -2,6 +2,7 @@
 
 #include "Catch2Includes.h"
 #include <AutoRTFM/AutoRTFM.h>
+#include "AutoRTFM/AutoRTFMTesting.h"
 #include <thread>
 #include <atomic>
 #include <vector>
@@ -491,3 +492,55 @@ TEST_CASE("Tests.__local_stdio_scanf_options")
 }
 
 #endif // PLATFORM_WINDOWS
+
+TEST_CASE("Tests.thread_local")
+{
+	auto TLSInt = []() -> int&
+	{
+		UE_AUTORTFM_DECLARE_THREAD_LOCAL_VAR(int, MyInt);
+		return MyInt;
+	};
+	auto TLSString = []() -> FString&
+	{
+		UE_AUTORTFM_DECLARE_THREAD_LOCAL_VAR(FString, MyString);
+		return MyString;
+	};
+
+	SECTION("Abort on Initial Access")
+	{
+		AutoRTFM::Testing::Abort([&]
+		{
+			TLSInt() = 42;
+			TLSString() = "Cat";
+			AutoRTFM::AbortTransaction();
+		});
+
+		REQUIRE(TLSInt() == 0);
+		REQUIRE(TLSString() == "");
+	}
+
+	SECTION("Abort")
+	{
+		TLSInt() = 123;
+		TLSString() = "Pickle";
+		AutoRTFM::Testing::Abort([&]
+		{
+			TLSInt() = 456;
+			TLSString() = "Peanut";
+			AutoRTFM::AbortTransaction();
+		});
+		REQUIRE(TLSInt() == 123);
+		REQUIRE(TLSString() == "Pickle");
+	}
+
+	SECTION("Commit")
+	{
+		AutoRTFM::Testing::Commit([&]
+		{
+			TLSInt() = 42;
+			TLSString() = "Cat";
+		});
+		REQUIRE(TLSInt() == 42);
+		REQUIRE(TLSString() == "Cat");
+	}
+}
\ No newline at end of file
diff --git a/Engine/Source/Runtime/Core/Public/AutoRTFM/AutoRTFM.h b/Engine/Source/Runtime/Core/Public/AutoRTFM/AutoRTFM.h
--- a/Engine/Source/Runtime/Core/Public/AutoRTFM/AutoRTFM.h
+++ b/Engine/Source/Runtime/Core/Public/AutoRTFM/AutoRTFM.h
@@ -1396,6 +1396,18 @@ namespace AutoRTFM::Private
 			AutoRTFM::Transact(UE_AUTORTFM_MOVE(F));
 		}
 	};
+	namespace /* must have internal linkage */
+	{
+		struct FThreadLocalHelper
+		{
+			template <typename Type, int Unique>
+			UE_AUTORTFM_ALWAYS_OPEN static Type& Get()
+			{
+				thread_local Type Data;
+				return Data;
+			}
+		};
+	}
 } // namespace AutoRTFM::Private
 
 #if defined(__clang__) && __has_warning("-Wdeprecated-this-capture")
@@ -1406,6 +1418,8 @@ namespace AutoRTFM::Private
 #define UE_AUTORTFM_END_DISABLE_WARNINGS
 #endif
 
+#define UE_AUTORTFM_DECLARE_THREAD_LOCAL_VAR_IMPL(Type, Name) Type& Name = ::AutoRTFM::Private::FThreadLocalHelper::Get<Type, __COUNTER__>()
+
 #define UE_AUTORTFM_OPEN_IMPL          ::AutoRTFM::Private::FOpenHelper{} + [&]()
 #define UE_AUTORTFM_OPEN_NO_VALIDATION_IMPL ::AutoRTFM::Private::FOpenNoMemoryValidationHelper{} + [&]()
 #define UE_AUTORTFM_ONABORT_IMPL(...)  ::AutoRTFM::Private::FOnAbortHelper{} + [__VA_ARGS__]()
@@ -1414,6 +1428,7 @@ namespace AutoRTFM::Private
 #else
 
 // Do nothing, these should be followed by blocks that should be either executed or not executed
+#define UE_AUTORTFM_DECLARE_THREAD_LOCAL_VAR_IMPL(Type, Name) thread_local Type Name
 #define UE_AUTORTFM_OPEN_IMPL
 #define UE_AUTORTFM_OPEN_NO_VALIDATION_IMPL
 #define UE_AUTORTFM_ONABORT_IMPL(...) while (false)
@@ -1421,6 +1436,12 @@ namespace AutoRTFM::Private
 #define UE_AUTORTFM_TRANSACT_IMPL
 #endif
 
+// Declares an AutoRTFM-aware thread local variable. `thread_local` variables are not yet natively supported (#jira SOL-7684)
+// Calls should be written like this: 
+//     UE_AUTORTFM_DECLARE_THREAD_LOCAL_VAR(FString, MyThreadLocalString); 
+//     MyThreadLocalString = "Hello";
+#define UE_AUTORTFM_DECLARE_THREAD_LOCAL_VAR(Type, Name) UE_AUTORTFM_DECLARE_THREAD_LOCAL_VAR_IMPL(Type, Name)
+
 // Runs a block of code in the open, non-transactionally. Anything performed in the open will not be undone if a transaction fails.
 // Calls should be written like this: UE_AUTORTFM_OPEN { ... code ... };
 #define UE_AUTORTFM_OPEN  UE_AUTORTFM_OPEN_IMPL
