Tasks - Don't trace subsequents if they're closed.

#rnx
#tests local insights traces

diff --git a/Engine/Source/Runtime/Core/Public/Tasks/TaskPrivate.h b/Engine/Source/Runtime/Core/Public/Tasks/TaskPrivate.h
--- a/Engine/Source/Runtime/Core/Public/Tasks/TaskPrivate.h
+++ b/Engine/Source/Runtime/Core/Public/Tasks/TaskPrivate.h
@@ -366,8 +366,12 @@ public:
 			// returns false if the task is already completed and the subsequent wasn't added
 			bool AddSubsequent(FTaskBase& Subsequent)
 			{
-				TaskTrace::SubsequentAdded(GetTraceId(), Subsequent.GetTraceId()); // doesn't matter if we suceeded below, we need to record task dependency
-				return Subsequents.PushIfNotClosed(&Subsequent);
+				if (Subsequents.PushIfNotClosed(&Subsequent))
+				{
+					TaskTrace::SubsequentAdded(GetTraceId(), Subsequent.GetTraceId());
+					return true;
+				}
+				return false;
 			}
 
 			// A piped task is executed after the previous task from this pipe is completed. Tasks from the same pipe are not executed
