GenericPlatformMiscTest: Add missing header includes.
#rnx

diff --git a/Engine/Source/Runtime/Core/Tests/GenericPlatform/GenericPlatformMiscTest.cpp b/Engine/Source/Runtime/Core/Tests/GenericPlatform/GenericPlatformMiscTest.cpp
--- a/Engine/Source/Runtime/Core/Tests/GenericPlatform/GenericPlatformMiscTest.cpp
+++ b/Engine/Source/Runtime/Core/Tests/GenericPlatform/GenericPlatformMiscTest.cpp
@@ -3,6 +3,9 @@
 #if WITH_TESTS 
 
 #include "GenericPlatform/GenericPlatformMisc.h"
+
+#include "Containers/StringView.h"
+#include "Containers/UnrealString.h"
 #include "Tests/TestHarnessAdapter.h"
 
 TEST_CASE_NAMED(FGenericPlatformMiscTest, "System::Core::GenericPlatform", "[Core][GenericPlatform][SmokeFilter]")
