Platform BeginNamedEvent: Uniformly call the FGenericPlatformMisc version of BeginNamedEvent from all platforms, it is required to support external profilers. Provide a define for ENABLE_GENERIC_NAMED_EVENTS so that platform configurations can FORCEINLINE their functions to completely strip them from UE_BUILD_SHIPPING configurations. Move stat event calls and FramePro profiler calls into FGenericPlatformMisc.
#rnx

diff --git a/Engine/Source/Runtime/Core/Private/Android/AndroidPlatformMisc.cpp b/Engine/Source/Runtime/Core/Private/Android/AndroidPlatformMisc.cpp
--- a/Engine/Source/Runtime/Core/Private/Android/AndroidPlatformMisc.cpp
+++ b/Engine/Source/Runtime/Core/Private/Android/AndroidPlatformMisc.cpp
@@ -2705,13 +2705,6 @@ bool FAndroidMisc::IsDebuggerPresent()
 
 #if STATS || ENABLE_STATNAMEDEVENTS
 
-void FAndroidMisc::BeginNamedEventFrame()
-{
-#if FRAMEPRO_ENABLED
-	FFrameProProfiler::FrameStart();
-#endif // FRAMEPRO_ENABLED
-}
-
 static void WriteTraceMarkerEvent(const ANSICHAR* Text, int32 TraceMarkerFileDescriptor)
 {
 	if (bUseNativeSystrace)
@@ -2730,36 +2723,19 @@ static void WriteTraceMarkerEvent(const ANSICHAR* Text, int32 TraceMarkerFileDes
 
 void FAndroidMisc::BeginNamedEvent(const struct FColor& Color, const TCHAR* Text)
 {
-#if FRAMEPRO_ENABLED
-	FFrameProProfiler::PushEvent(Text);
-#endif // FRAMEPRO_ENABLED
+	FGenericPlatformMisc::BeginNamedEvent(Color, Text);
 	if (bUseNativeSystrace ? !ATrace_isEnabled() : TraceMarkerFileDescriptor == -1)
 	{
 		return;
 	}
 
 	const int MAX_TRACE_MESSAGE_LENGTH = 256;
-
-	// not static since may be called by different threads
-	ANSICHAR TextBuffer[MAX_TRACE_MESSAGE_LENGTH];
-
-	const TCHAR* SourcePtr = Text;
-	ANSICHAR* WritePtr = TextBuffer;
-	int32 RemainingSpace = MAX_TRACE_MESSAGE_LENGTH;
-	while (*SourcePtr && --RemainingSpace > 0)
-	{
-		*WritePtr++ = static_cast<ANSICHAR>(*SourcePtr++);
-	}
-	*WritePtr = '\0';
-
-	WriteTraceMarkerEvent(TextBuffer, TraceMarkerFileDescriptor);
+	WriteTraceMarkerEvent(StringCast<ANSICHAR, MAX_TRACE_MESSAGE_LENGTH>(Text).Get(), TraceMarkerFileDescriptor);
 }
 
 void FAndroidMisc::BeginNamedEvent(const struct FColor& Color, const ANSICHAR* Text)
 {
-#if FRAMEPRO_ENABLED
-	FFrameProProfiler::PushEvent(Text);
-#endif // FRAMEPRO_ENABLED
+	FGenericPlatformMisc::BeginNamedEvent(Color, Text);
 	if (bUseNativeSystrace ? !ATrace_isEnabled() : TraceMarkerFileDescriptor == -1)
 	{
 		return;
@@ -2770,9 +2746,7 @@ void FAndroidMisc::BeginNamedEvent(const struct FColor& Color, const ANSICHAR* T
 
 void FAndroidMisc::EndNamedEvent()
 {
-#if FRAMEPRO_ENABLED
-	FFrameProProfiler::PopEvent();
-#endif // FRAMEPRO_ENABLED
+	FGenericPlatformMisc::EndNamedEvent();
 	if (bUseNativeSystrace ? !ATrace_isEnabled() : TraceMarkerFileDescriptor == -1)
 	{
 		return;
@@ -2789,16 +2763,6 @@ void FAndroidMisc::EndNamedEvent()
 	}
 }
 
-void FAndroidMisc::CustomNamedStat(const TCHAR* Text, float Value, const TCHAR* Graph, const TCHAR* Unit)
-{
-	FRAMEPRO_DYNAMIC_CUSTOM_STAT(TCHAR_TO_WCHAR(Text), Value, TCHAR_TO_WCHAR(Graph), TCHAR_TO_WCHAR(Unit), FRAMEPRO_COLOUR(255,255,255) );
-}
-
-void FAndroidMisc::CustomNamedStat(const ANSICHAR* Text, float Value, const ANSICHAR* Graph, const ANSICHAR* Unit)
-{
-	FRAMEPRO_DYNAMIC_CUSTOM_STAT(Text, Value, Graph, Unit, FRAMEPRO_COLOUR(255,255,255));
-}
-
 #endif // STATS || ENABLE_STATNAMEDEVENTS
 
 int FAndroidMisc::GetVolumeState(double* OutTimeOfChangeInSec)
diff --git a/Engine/Source/Runtime/Core/Private/Apple/ApplePlatformMisc.cpp b/Engine/Source/Runtime/Core/Private/Apple/ApplePlatformMisc.cpp
--- a/Engine/Source/Runtime/Core/Private/Apple/ApplePlatformMisc.cpp
+++ b/Engine/Source/Runtime/Core/Private/Apple/ApplePlatformMisc.cpp
@@ -354,50 +354,30 @@ bool FApplePlatformMisc::IsOSAtLeastVersion(const uint32 MacOSVersion[3], const
 	return true;
 }
 
-#if STATS || ENABLE_STATNAMEDEVENTS || APPLE_PROFILING_ENABLED
-
-void FApplePlatformMisc::BeginNamedEventFrame()
-{
-#if FRAMEPRO_ENABLED
-	FFrameProProfiler::FrameStart();
-#endif // FRAMEPRO_ENABLED
-}
+#if APPLE_PROFILING_ENABLED
 
 void FApplePlatformMisc::BeginNamedEvent(const struct FColor& Color, const TCHAR* Text)
 {
-#if FRAMEPRO_ENABLED
-	FFrameProProfiler::PushEvent(Text);
-#elif APPLE_PROFILING_ENABLED
+	FGenericPlatformMisc::BeginNamedEvent(Color, Text);
+#if !FRAMEPRO_ENABLED && APPLE_PROFILING_ENABLED
 	FApplePlatformDebugEvents::BeginNamedEvent(Color, Text);
 #endif // FRAMEPRO_ENABLED
 }
 
 void FApplePlatformMisc::BeginNamedEvent(const struct FColor& Color, const ANSICHAR* Text)
 {
-#if FRAMEPRO_ENABLED
-	FFrameProProfiler::PushEvent(Text);
-#elif APPLE_PROFILING_ENABLED
+	FGenericPlatformMisc::BeginNamedEvent(Color, Text);
+#if !FRAMEPRO_ENABLED && APPLE_PROFILING_ENABLED
 	FApplePlatformDebugEvents::BeginNamedEvent(Color, Text);
 #endif // FRAMEPRO_ENABLED
 }
 
 void FApplePlatformMisc::EndNamedEvent()
 {
-#if FRAMEPRO_ENABLED
-	FFrameProProfiler::PopEvent();
-#elif APPLE_PROFILING_ENABLED
+	FGenericPlatformMisc::EndNamedEvent();
+#if !FRAMEPRO_ENABLED && APPLE_PROFILING_ENABLED
 	FApplePlatformDebugEvents::EndNamedEvent();
 #endif // FRAMEPRO_ENABLED
 }
 
-void FApplePlatformMisc::CustomNamedStat(const TCHAR* Text, float Value, const TCHAR* Graph, const TCHAR* Unit)
-{
-	FRAMEPRO_DYNAMIC_CUSTOM_STAT(TCHAR_TO_WCHAR(Text), Value, TCHAR_TO_WCHAR(Graph), TCHAR_TO_WCHAR(Unit), FRAMEPRO_COLOUR(255,255,255));
-}
-
-void FApplePlatformMisc::CustomNamedStat(const ANSICHAR* Text, float Value, const ANSICHAR* Graph, const ANSICHAR* Unit)
-{
-	FRAMEPRO_DYNAMIC_CUSTOM_STAT(Text, Value, Graph, Unit, FRAMEPRO_COLOUR(255,255,255));
-}
-
-#endif // STATS || ENABLE_STATNAMEDEVENTS || APPLE_PROFILING_ENABLED
+#endif // APPLE_PROFILING_ENABLED
diff --git a/Engine/Source/Runtime/Core/Private/GenericPlatform/GenericPlatformMisc.cpp b/Engine/Source/Runtime/Core/Private/GenericPlatform/GenericPlatformMisc.cpp
--- a/Engine/Source/Runtime/Core/Private/GenericPlatform/GenericPlatformMisc.cpp
+++ b/Engine/Source/Runtime/Core/Private/GenericPlatform/GenericPlatformMisc.cpp
@@ -43,6 +43,7 @@
 #include "ProfilingDebugging/CsvProfiler.h"
 #include "Misc/CoreDelegates.h"
 #include "AutoRTFM/AutoRTFM.h"
+#include "FramePro/FrameProProfiler.h"
 
 #if UE_ENABLE_ICU
 	THIRD_PARTY_INCLUDES_START
@@ -69,7 +70,6 @@ static std::atomic<uint32> GNamedEventMarkers(0L);
 #endif
 
 
-#define LOG_NAMED_EVENTS 0
 #if LOG_NAMED_EVENTS
 struct FNameEventEntry
 {
@@ -777,10 +777,24 @@ void FGenericPlatformMisc::LogNameEventStatsInit()
 #endif
 }
 
+#if ENABLE_GENERIC_NAMED_EVENTS
+
+void FGenericPlatformMisc::BeginNamedEventFrame()
+{
+	FGenericPlatformMisc::TickStatNamedEvents();
+#if FRAMEPRO_ENABLED
+	FFrameProProfiler::FrameStart();
+#endif
+}
+
 void FGenericPlatformMisc::BeginNamedEvent(const struct FColor& Color, const ANSICHAR* Text)
 {
 #if UE_EXTERNAL_PROFILING_ENABLED
 	FExternalProfilerTrace::StartScopedEvent(Color, Text);
+#endif
+	StatNamedEvent(Text);
+#if FRAMEPRO_ENABLED
+	FFrameProProfiler::PushEvent(Text);
 #endif
 }
 
@@ -788,6 +802,10 @@ void FGenericPlatformMisc::BeginNamedEvent(const struct FColor& Color, const TCH
 {
 #if UE_EXTERNAL_PROFILING_ENABLED
 	FExternalProfilerTrace::StartScopedEvent(Color, Text);
+#endif
+	StatNamedEvent(Text);
+#if FRAMEPRO_ENABLED
+	FFrameProProfiler::PushEvent(Text);
 #endif
 }
 
@@ -796,8 +814,23 @@ void FGenericPlatformMisc::EndNamedEvent()
 #if UE_EXTERNAL_PROFILING_ENABLED
 	FExternalProfilerTrace::EndScopedEvent();
 #endif
+#if FRAMEPRO_ENABLED
+	FFrameProProfiler::PopEvent();
+#endif
 }
 
+void FGenericPlatformMisc::CustomNamedStat(const TCHAR* Text, float Value, const TCHAR* Graph, const TCHAR* Unit)
+{
+	FRAMEPRO_DYNAMIC_CUSTOM_STAT(TCHAR_TO_WCHAR(Text), Value, TCHAR_TO_WCHAR(Graph), TCHAR_TO_WCHAR(Unit), FRAMEPRO_COLOUR(255,255,255) );
+}
+
+void FGenericPlatformMisc::CustomNamedStat(const ANSICHAR* Text, float Value, const ANSICHAR* Graph, const ANSICHAR* Unit)
+{
+	FRAMEPRO_DYNAMIC_CUSTOM_STAT(Text, Value, Graph, Unit, FRAMEPRO_COLOUR(255,255,255));
+}
+
+#endif // ENABLE_GENERIC_NAMED_EVENTS
+
 bool FGenericPlatformMisc::SetStoredValues(const FString& InStoreId, const FString& InSectionName, const TMap<FString, FString>& InKeyValues)
 {
 	for (const TPair<FString, FString>& InKeyValue : InKeyValues)
diff --git a/Engine/Source/Runtime/Core/Private/Unix/UnixPlatformMisc.cpp b/Engine/Source/Runtime/Core/Private/Unix/UnixPlatformMisc.cpp
--- a/Engine/Source/Runtime/Core/Private/Unix/UnixPlatformMisc.cpp
+++ b/Engine/Source/Runtime/Core/Private/Unix/UnixPlatformMisc.cpp
@@ -1710,46 +1710,6 @@ void FUnixPlatformMisc::CreateGuid(FGuid& Result)
 	}
 }
 
-#if STATS || ENABLE_STATNAMEDEVENTS
-void FUnixPlatformMisc::BeginNamedEventFrame()
-{
-#if FRAMEPRO_ENABLED
-	FFrameProProfiler::FrameStart();
-#endif // FRAMEPRO_ENABLED
-}
-
-void FUnixPlatformMisc::BeginNamedEvent(const struct FColor& Color, const TCHAR* Text)
-{
-#if FRAMEPRO_ENABLED
-	FFrameProProfiler::PushEvent(Text);
-#endif
-}
-
-void FUnixPlatformMisc::BeginNamedEvent(const struct FColor& Color, const ANSICHAR* Text)
-{
-#if FRAMEPRO_ENABLED
-	FFrameProProfiler::PushEvent(Text);
-#endif
-}
-
-void FUnixPlatformMisc::EndNamedEvent()
-{
-#if FRAMEPRO_ENABLED
-	FFrameProProfiler::PopEvent();
-#endif
-}
-
-void FUnixPlatformMisc::CustomNamedStat(const TCHAR* Text, float Value, const TCHAR* Graph, const TCHAR* Unit)
-{
-	FRAMEPRO_DYNAMIC_CUSTOM_STAT(TCHAR_TO_WCHAR(Text), Value, TCHAR_TO_WCHAR(Graph), TCHAR_TO_WCHAR(Unit), FRAMEPRO_COLOUR(255,255,255));
-}
-
-void FUnixPlatformMisc::CustomNamedStat(const ANSICHAR* Text, float Value, const ANSICHAR* Graph, const ANSICHAR* Unit)
-{
-	FRAMEPRO_DYNAMIC_CUSTOM_STAT(Text, Value, Graph, Unit, FRAMEPRO_COLOUR(255,255,255));
-}
-#endif
-
 CORE_API TFunction<void()> UngrabAllInputCallback;
 
 void FUnixPlatformMisc::UngrabAllInput()
diff --git a/Engine/Source/Runtime/Core/Private/Windows/WindowsPlatformMisc.cpp b/Engine/Source/Runtime/Core/Private/Windows/WindowsPlatformMisc.cpp
--- a/Engine/Source/Runtime/Core/Private/Windows/WindowsPlatformMisc.cpp
+++ b/Engine/Source/Runtime/Core/Private/Windows/WindowsPlatformMisc.cpp
@@ -1263,61 +1263,6 @@ EProcessDiagnosticFlags FWindowsPlatformMisc::GetProcessDiagnostics()
 }
 #endif //!UE_BUILD_SHIPPING
 
-#if STATS || ENABLE_STATNAMEDEVENTS
-void FWindowsPlatformMisc::CustomNamedStat(const TCHAR* Text, float Value, const TCHAR* Graph, const TCHAR* Unit)
-{
-	FRAMEPRO_DYNAMIC_CUSTOM_STAT(TCHAR_TO_WCHAR(Text), Value, TCHAR_TO_WCHAR(Graph), TCHAR_TO_WCHAR(Unit), FRAMEPRO_COLOUR(255,255,255));
-}
-
-void FWindowsPlatformMisc::CustomNamedStat(const ANSICHAR* Text, float Value, const ANSICHAR* Graph, const ANSICHAR* Unit)
-{
-	FRAMEPRO_DYNAMIC_CUSTOM_STAT(Text, Value, Graph, Unit, FRAMEPRO_COLOUR(255,255,255));
-}
-
-void FWindowsPlatformMisc::BeginNamedEventFrame()
-{
-	FGenericPlatformMisc::TickStatNamedEvents();
-
-#if FRAMEPRO_ENABLED
-	FFrameProProfiler::FrameStart();
-#endif
-}
-
-void FWindowsPlatformMisc::BeginNamedEvent(const struct FColor& Color, const TCHAR* Text)
-{
-	FGenericPlatformMisc::StatNamedEvent(Text);
-
-#if FRAMEPRO_ENABLED
-	FFrameProProfiler::PushEvent(Text);
-#endif
-#if UE_EXTERNAL_PROFILING_ENABLED
-	FExternalProfilerTrace::StartScopedEvent(Color, Text);
-#endif
-}
-
-void FWindowsPlatformMisc::BeginNamedEvent(const struct FColor& Color, const ANSICHAR* Text)
-{
-	FGenericPlatformMisc::StatNamedEvent(Text);
-
-#if FRAMEPRO_ENABLED
-	FFrameProProfiler::PushEvent(Text);
-#endif
-#if UE_EXTERNAL_PROFILING_ENABLED
-	FExternalProfilerTrace::StartScopedEvent(Color, Text);
-#endif
-}
-
-void FWindowsPlatformMisc::EndNamedEvent()
-{
-#if FRAMEPRO_ENABLED
-	FFrameProProfiler::PopEvent();
-#endif
-#if UE_EXTERNAL_PROFILING_ENABLED
-	FExternalProfilerTrace::EndScopedEvent();
-#endif
-}
-#endif // STATS || ENABLE_STATNAMEDEVENTS
-
 bool FWindowsPlatformMisc::IsRemoteSession()
 {
 	return ::GetSystemMetrics(SM_REMOTESESSION) != 0;
diff --git a/Engine/Source/Runtime/Core/Public/Android/AndroidPlatformMisc.h b/Engine/Source/Runtime/Core/Public/Android/AndroidPlatformMisc.h
--- a/Engine/Source/Runtime/Core/Public/Android/AndroidPlatformMisc.h
+++ b/Engine/Source/Runtime/Core/Public/Android/AndroidPlatformMisc.h
@@ -250,12 +250,10 @@ public:
 
 
 #if STATS || ENABLE_STATNAMEDEVENTS
-	static CORE_API void BeginNamedEventFrame();
+	/** Override to add Android features when stats are on, otherwise use version on FGenericPlatformMisc. */
 	static CORE_API void BeginNamedEvent(const struct FColor& Color, const TCHAR* Text);
 	static CORE_API void BeginNamedEvent(const struct FColor& Color, const ANSICHAR* Text);
 	static CORE_API void EndNamedEvent();
-	static CORE_API void CustomNamedStat(const TCHAR* Text, float Value, const TCHAR* Graph, const TCHAR* Unit);
-	static CORE_API void CustomNamedStat(const ANSICHAR* Text, float Value, const ANSICHAR* Graph, const ANSICHAR* Unit);
 #endif
 
 #if (STATS || ENABLE_STATNAMEDEVENTS)
diff --git a/Engine/Source/Runtime/Core/Public/Apple/ApplePlatformMisc.h b/Engine/Source/Runtime/Core/Public/Apple/ApplePlatformMisc.h
--- a/Engine/Source/Runtime/Core/Public/Apple/ApplePlatformMisc.h
+++ b/Engine/Source/Runtime/Core/Public/Apple/ApplePlatformMisc.h
@@ -90,13 +90,11 @@ struct CORE_API FApplePlatformMisc : public FGenericPlatformMisc
 
 	static bool IsOSAtLeastVersion(const uint32 MacOSVersion[3], const uint32 IOSVersion[3], const uint32 TVOSVersion[3]);
 
-#if STATS || ENABLE_STATNAMEDEVENTS || APPLE_PROFILING_ENABLED
-	static void BeginNamedEventFrame();
+#if APPLE_PROFILING_ENABLED
+	/** Override to add Apple features when AppleProfiling is on, otherwise use version on FGenericPlatformMisc. */
 	static void BeginNamedEvent(const struct FColor& Color, const TCHAR* Text);
 	static void BeginNamedEvent(const struct FColor& Color, const ANSICHAR* Text);
 	static void EndNamedEvent();
-	static void CustomNamedStat(const TCHAR* Text, float Value, const TCHAR* Graph, const TCHAR* Unit);
-	static void CustomNamedStat(const ANSICHAR* Text, float Value, const ANSICHAR* Graph, const ANSICHAR* Unit);
 #endif
 
 	//////// Platform specific
diff --git a/Engine/Source/Runtime/Core/Public/FramePro/FramePro.h b/Engine/Source/Runtime/Core/Public/FramePro/FramePro.h
--- a/Engine/Source/Runtime/Core/Public/FramePro/FramePro.h
+++ b/Engine/Source/Runtime/Core/Public/FramePro/FramePro.h
@@ -28,16 +28,10 @@ OR IN CONNECTION WITH THE USE OR PERFORMANCE OF THIS SOFTWARE.
 #define FRAMEPRO_H_INCLUDED
 
 //------------------------------------------------------------------------
-#ifndef FRAMEPRO_ENABLED
-	// *** enable/disable FramePro here
-	//@EPIC BEGIN: Only enable for engine / game, not external tools
-	#ifdef __UNREAL__
-		#define FRAMEPRO_ENABLED 0
-	#else
-		#define FRAMEPRO_ENABLED 1
-	#endif
-	//@EPIC END
-#endif
+// *** enable/disable FramePro here
+//@EPIC BEGIN: Enable through minimal config file, and enable for engine / game, not external tools
+#include "FramePro/FrameProConfig.h"
+//@EPIC END
 
 //------------------------------------------------------------------------
 // Scroll down for detailed documentation on all of these macros.
diff --git a/Engine/Source/Runtime/Core/Public/FramePro/FrameProConfig.h b/Engine/Source/Runtime/Core/Public/FramePro/FrameProConfig.h
new file mode 100644
--- /dev/null
+++ b/Engine/Source/Runtime/Core/Public/FramePro/FrameProConfig.h
@@ -0,0 +1,28 @@
+/*
+Copyright 2019 PureDev Software Limited
+
+Permission to use, copy, modify, and/or distribute this software for any
+purpose with or without fee is hereby granted, provided that the above
+copyright notice and this permission notice appear in all copies.
+
+THE SOFTWARE IS PROVIDED "AS IS" AND THE AUTHOR DISCLAIMS ALL WARRANTIES
+WITH REGARD TO THIS SOFTWARE INCLUDING ALL IMPLIED WARRANTIES OF
+MERCHANTABILITY AND FITNESS. IN NO EVENT SHALL THE AUTHOR BE LIABLE FOR
+ANY SPECIAL, DIRECT, INDIRECT, OR CONSEQUENTIAL DAMAGES OR ANY DAMAGES
+WHATSOEVER RESULTING FROM LOSS OF USE, DATA OR PROFITS, WHETHER IN AN
+ACTION OF CONTRACT, NEGLIGENCE OR OTHER TORTIOUS ACTION, ARISING OUT OF
+OR IN CONNECTION WITH THE USE OR PERFORMANCE OF THIS SOFTWARE.
+*/
+
+#pragma once
+
+//@EPIC BEGIN: enable for engine / game, not external tools. engine/game enables it in UnrealBuildTool.
+#ifndef FRAMEPRO_ENABLED
+	// *** enable/disable FramePro here
+	#ifdef __UNREAL__
+		#define FRAMEPRO_ENABLED 0
+	#else
+		#define FRAMEPRO_ENABLED 1
+	#endif
+#endif
+//@EPIC END
diff --git a/Engine/Source/Runtime/Core/Public/GenericPlatform/GenericPlatformMisc.h b/Engine/Source/Runtime/Core/Public/GenericPlatform/GenericPlatformMisc.h
--- a/Engine/Source/Runtime/Core/Public/GenericPlatform/GenericPlatformMisc.h
+++ b/Engine/Source/Runtime/Core/Public/GenericPlatform/GenericPlatformMisc.h
@@ -5,10 +5,13 @@
 #include "Containers/StringFwd.h"
 #include "CoreFwd.h"
 #include "CoreTypes.h"
+#include "FramePro/FrameProConfig.h"
 #include "HAL/PlatformCrt.h"
 #include "Math/NumericLimits.h"
 #include "Misc/CompressionFlags.h"
 #include "Misc/EnumClassFlags.h"
+#include "ProfilingDebugging/CsvProfilerConfig.h"
+#include "ProfilingDebugging/ExternalProfilerConfig.h"
 
 #if PLATFORM_CPU_X86_FAMILY
 #include <xmmintrin.h> // _mm_prefetch
@@ -44,6 +47,11 @@ class TDelegate;
 #define UE_DEBUG_BREAK() ((void)(FPlatformMisc::IsDebuggerPresent() && ([] () { UE_DEBUG_BREAK_IMPL(); } (), 1)))
 #endif
 
+// Enable this manually for debugging
+#define LOG_NAMED_EVENTS 0
+
+#define ENABLE_GENERIC_NAMED_EVENTS ((UE_EXTERNAL_PROFILING_ENABLED) || (STATS) || (ENABLE_STATNAMEDEVENTS) || ((CSV_PROFILER_STATS) && !UE_BUILD_SHIPPING) || (LOG_NAMED_EVENTS) || (FRAMEPRO_ENABLED))
+
 /**
  * Available build configurations. Mirorred from UnrealTargetConfiguration.
  */
@@ -874,11 +882,38 @@ struct FGenericPlatformMisc
 
 public:
 
-	/**
-	 * Platform specific function for adding a named event that can be viewed in external tool
-	 */
+#if ENABLE_GENERIC_NAMED_EVENTS
+	/** Indicates the start of a frame for named events, both generic and platform-specific. */
+	static CORE_API void BeginNamedEventFrame();
+	/** Add a named event that can be viewed in external tool, both generic and platform-specific. */
 	static CORE_API void BeginNamedEvent(const struct FColor& Color, const TCHAR* Text);
 	static CORE_API void BeginNamedEvent(const struct FColor& Color, const ANSICHAR* Text);
+	/** Close a named event that can be viewed in external tool, both generic and platform-specific. */
+	static CORE_API void EndNamedEvent();
+	/** Set the value for a named custom stat that can be viewed in external tool, both generic and platform-specific. */
+	static CORE_API void CustomNamedStat(const TCHAR* Text, float Value, const TCHAR* Graph, const TCHAR* Unit);
+	static CORE_API void CustomNamedStat(const ANSICHAR* Text, float Value, const ANSICHAR* Graph, const ANSICHAR* Unit);
+#else
+	FORCEINLINE static void BeginNamedEventFrame()
+	{
+	}
+	FORCEINLINE static void BeginNamedEvent(const struct FColor& Color, const TCHAR* Text)
+	{
+	}
+	FORCEINLINE static void BeginNamedEvent(const struct FColor& Color, const ANSICHAR* Text)
+	{
+	}
+	FORCEINLINE static void EndNamedEvent()
+	{
+	}
+	FORCEINLINE static void CustomNamedStat(const TCHAR* Text, float Value, const TCHAR* Graph, const TCHAR* Unit)
+	{
+	}
+	FORCEINLINE static void CustomNamedStat(const ANSICHAR* Text, float Value, const ANSICHAR* Graph, const ANSICHAR* Unit)
+	{
+	}
+#endif
+
 
 	/**
 	 * Stats logging for named event 
@@ -888,15 +923,6 @@ public:
 	static CORE_API void TickStatNamedEvents();
 	static CORE_API void LogNameEventStatsInit();
 
-	/**
-	 * Platform specific function for closing a named event that can be viewed in external tool
-	 */
-	static CORE_API void EndNamedEvent();
-
-	/** Platform specific function for adding a named custom stat that can be viewed in external tool */
-	static void CustomNamedStat(const TCHAR* Text, float Value, const TCHAR* Graph, const TCHAR* Unit) {}
-	static void CustomNamedStat(const ANSICHAR* Text, float Value, const ANSICHAR* Graph, const ANSICHAR* Unit) {}
-
 	/**
 	 * Profiler enter background events. These are for timing backgrounding. Some platforms have a timeout for this
 	 */
@@ -904,11 +930,6 @@ public:
 	static void EndEnterBackgroundEvent() {}
 
 
-	/** Indicates the start of a frame for named events */
-	FORCEINLINE static void BeginNamedEventFrame()
-	{
-	}
-
 	/**
 	 * Platform specific function for initializing storage of tagged memory buffers
 	 */
diff --git a/Engine/Source/Runtime/Core/Public/ProfilingDebugging/CsvProfilerConfig.h b/Engine/Source/Runtime/Core/Public/ProfilingDebugging/CsvProfilerConfig.h
--- a/Engine/Source/Runtime/Core/Public/ProfilingDebugging/CsvProfilerConfig.h
+++ b/Engine/Source/Runtime/Core/Public/ProfilingDebugging/CsvProfilerConfig.h
@@ -2,6 +2,8 @@
 
 #pragma once
 
+#include "Misc/Build.h"
+
 // Whether to allow the CSV profiler in shipping builds.
 // Enable in a .Target.cs file if required.
 #ifndef CSV_PROFILER_ENABLE_IN_SHIPPING
diff --git a/Engine/Source/Runtime/Core/Public/ProfilingDebugging/ExternalProfiler.h b/Engine/Source/Runtime/Core/Public/ProfilingDebugging/ExternalProfiler.h
--- a/Engine/Source/Runtime/Core/Public/ProfilingDebugging/ExternalProfiler.h
+++ b/Engine/Source/Runtime/Core/Public/ProfilingDebugging/ExternalProfiler.h
@@ -6,12 +6,9 @@
 #include "CoreTypes.h"
 #include "Features/IModularFeature.h"
 #include "Misc/Build.h"
+#include "ProfilingDebugging/ExternalProfilerConfig.h"
 #include "UObject/NameTypes.h"
 
-#ifndef UE_EXTERNAL_PROFILING_ENABLED
-#define UE_EXTERNAL_PROFILING_ENABLED (!UE_BUILD_SHIPPING)
-#endif
-
 #if UE_EXTERNAL_PROFILING_ENABLED
 
 /**
diff --git a/Engine/Source/Runtime/Core/Public/ProfilingDebugging/ExternalProfilerConfig.h b/Engine/Source/Runtime/Core/Public/ProfilingDebugging/ExternalProfilerConfig.h
new file mode 100644
--- /dev/null
+++ b/Engine/Source/Runtime/Core/Public/ProfilingDebugging/ExternalProfilerConfig.h
@@ -0,0 +1,9 @@
+// Copyright Epic Games, Inc. All Rights Reserved.
+
+#pragma once
+
+#include "Misc/Build.h"
+
+#ifndef UE_EXTERNAL_PROFILING_ENABLED
+#define UE_EXTERNAL_PROFILING_ENABLED (!UE_BUILD_SHIPPING)
+#endif
diff --git a/Engine/Source/Runtime/Core/Public/Unix/UnixPlatformMisc.h b/Engine/Source/Runtime/Core/Public/Unix/UnixPlatformMisc.h
--- a/Engine/Source/Runtime/Core/Public/Unix/UnixPlatformMisc.h
+++ b/Engine/Source/Runtime/Core/Public/Unix/UnixPlatformMisc.h
@@ -137,15 +137,6 @@ struct FUnixPlatformMisc : public FGenericPlatformMisc
 
 	static CORE_API int32 NumberOfWorkerThreadsToSpawn();
 
-#if STATS || ENABLE_STATNAMEDEVENTS
-	static CORE_API void BeginNamedEventFrame();
-	static CORE_API void BeginNamedEvent(const struct FColor& Color, const TCHAR* Text);
-	static CORE_API void BeginNamedEvent(const struct FColor& Color, const ANSICHAR* Text);
-	static CORE_API void EndNamedEvent();
-	static CORE_API void CustomNamedStat(const TCHAR* Text, float Value, const TCHAR* Graph, const TCHAR* Unit);
-	static CORE_API void CustomNamedStat(const ANSICHAR* Text, float Value, const ANSICHAR* Graph, const ANSICHAR* Unit);
-#endif
-
 	/* Explicitly call this function to setup syscall filters based on a file passed in from the command line
 	 * -allowsyscallfilterfile=PATH_TO_FILE
 	 */
diff --git a/Engine/Source/Runtime/Core/Public/Windows/WindowsPlatformMisc.h b/Engine/Source/Runtime/Core/Public/Windows/WindowsPlatformMisc.h
--- a/Engine/Source/Runtime/Core/Public/Windows/WindowsPlatformMisc.h
+++ b/Engine/Source/Runtime/Core/Public/Windows/WindowsPlatformMisc.h
@@ -119,15 +119,6 @@ struct FWindowsPlatformMisc
 	static CORE_API EProcessDiagnosticFlags GetProcessDiagnostics();
 #endif
 
-#if STATS || ENABLE_STATNAMEDEVENTS
-	static CORE_API void BeginNamedEventFrame();
-	static CORE_API void BeginNamedEvent(const struct FColor& Color, const TCHAR* Text);
-	static CORE_API void BeginNamedEvent(const struct FColor& Color, const ANSICHAR* Text);
-	static CORE_API void EndNamedEvent();
-	static CORE_API void CustomNamedStat(const TCHAR* Text, float Value, const TCHAR* Graph, const TCHAR* Unit);
-	static CORE_API void CustomNamedStat(const ANSICHAR* Text, float Value, const ANSICHAR* Graph, const ANSICHAR* Unit);
-#endif
-
 	FORCEINLINE static void MemoryBarrier() 
 	{
 #if PLATFORM_CPU_X86_FAMILY
