Convert FConsoleManager tests to regular unit tests
- Previously these would automatically run when the class was first accessed. This ended up happening during static initialization due to auto-registered console objects. In turn, this caused an issue on PS4 due to TLS usage in the logging system. This issue was fixed in the logging system, but there's not a strong reason these tests need to be run during static init and during every boot.
- I tried to avoid changing any more than necessary. The tests could be organized better to play nicely with SECTIONs (there's shared state woven through the various tests) but that's best handled separately.
- I moved the definition of FConsoleManager to a private header so the tests can continue to access `GetThreadPropagationCallback` and `DumpObjects`.
- It would be nice to be able to instantiate a temporary FConsoleManager so we don't break, or get broken by, the real runtime instance. This is currently prevented because `IConsoleObject::OnCVarChange` accesses the singleton.

diff --git a/Engine/Source/Runtime/Core/Private/HAL/ConsoleManager.cpp b/Engine/Source/Runtime/Core/Private/HAL/ConsoleManager.cpp
--- a/Engine/Source/Runtime/Core/Private/HAL/ConsoleManager.cpp
+++ b/Engine/Source/Runtime/Core/Private/HAL/ConsoleManager.cpp
@@ -6,12 +6,6 @@ ConsoleManager.cpp: console command handling
 
 #include "HAL/ConsoleManager.h"
 
-#include "AutoRTFM/AutoRTFM.h"
-#include "Containers/Array.h"
-#include "Containers/Map.h"
-#include "Containers/UnrealString.h"
-#include "CoreTypes.h"
-#include "HAL/CriticalSection.h"
 #include "HAL/FileManager.h"
 #include "HAL/PlatformProcess.h"
 #include "Logging/MessageLog.h"
@@ -19,15 +13,9 @@ ConsoleManager.cpp: console command handling
 #include "Misc/ConfigContext.h"
 #include "Misc/ConfigUtilities.h"
 #include "Misc/CoreDelegates.h"
-#include "Misc/OutputDeviceNull.h"
-#include "Misc/Paths.h"
-#include "Misc/RemoteConfigIni.h"
-#include "Misc/ScopeExit.h"
-#include "Misc/ScopeLock.h"
-#include "Misc/ScopeRWLock.h"
-#include "Misc/TransactionallySafeCriticalSection.h"
 #include "Modules/ModuleManager.h"
 #include "ProfilingDebugging/CsvProfiler.h"
+#include "Stats/Stats.h"
 #include "Trace/Trace.h"
 #include "Trace/Trace.inl"
 
@@ -41,177 +29,14 @@ DEFINE_LOG_CATEGORY_STATIC(LogConsoleManager, Log, All);
 class FOutputDevice;
 class UWorld;
 
-class FConsoleManager :public IConsoleManager
-{
-public:
-	/** constructor */
-	FConsoleManager()
-		: bHistoryWasLoaded(false)
-		, ThreadPropagationCallback(0)
-		, bCallAllConsoleVariableSinks(true)
-	{
-	}
-
-	/** destructor */
-	~FConsoleManager()
-	{
-		for (TMap<FString, IConsoleObject*>::TConstIterator PairIt(ConsoleObjects); PairIt; ++PairIt)
-		{
-			IConsoleObject* Var = PairIt.Value();
-
-			delete Var;
-		}
-	}
-
-
-
-	// internally needed or ECVF_RenderThreadSafe
-	IConsoleThreadPropagation* GetThreadPropagationCallback();
-	// internally needed or ECVF_RenderThreadSafe
-	bool IsThreadPropagationThread();
-
-	/** @param InVar must not be 0 */
-	virtual FString FindConsoleObjectName(const IConsoleObject* Obj) const override;
-
-	/** Can be moved out into some automated testing system */
-	void Test();
-
-	void OnCVarChanged();
-
-	virtual FConsoleVariableMulticastDelegate& OnCVarUnregistered()override;
-	virtual FConsoleObjectWithNameMulticastDelegate& OnConsoleObjectUnregistered() override;
-
-	// interface IConsoleManager -----------------------------------
-
-	virtual IConsoleVariable* RegisterConsoleVariable(const TCHAR* Name, bool DefaultValue, const TCHAR* Help, uint32 Flags) override;
-	virtual IConsoleVariable* RegisterConsoleVariable(const TCHAR* Name, int32 DefaultValue, const TCHAR* Help, uint32 Flags) override;
-	virtual IConsoleVariable* RegisterConsoleVariable(const TCHAR* Name, float DefaultValue, const TCHAR* Help, uint32 Flags) override;
-	virtual IConsoleVariable* RegisterConsoleVariable(const TCHAR* Name, const TCHAR* DefaultValue, const TCHAR* Help, uint32 Flags) override;
-	virtual IConsoleVariable* RegisterConsoleVariable(const TCHAR* Name, const FString& DefaultValue, const TCHAR* Help, uint32 Flags) override;
-
-	virtual IConsoleVariable* RegisterConsoleVariableRef(const TCHAR* Name, bool& RefValue, const TCHAR* Help, uint32 Flags) override;
-	virtual IConsoleVariable* RegisterConsoleVariableRef(const TCHAR* Name, int32& RefValue, const TCHAR* Help, uint32 Flags) override;
-	virtual IConsoleVariable* RegisterConsoleVariableRef(const TCHAR* Name, float& RefValue, const TCHAR* Help, uint32 Flags) override;
-	virtual IConsoleVariable* RegisterConsoleVariableRef(const TCHAR* Name, FString& RefValue, const TCHAR* Help, uint32 Flags) override;
-	virtual IConsoleVariable* RegisterConsoleVariableRef(const TCHAR* Name, FName& RefValue, const TCHAR* Help, uint32 Flags) override;
-	virtual IConsoleVariable* RegisterConsoleVariableBitRef(const TCHAR* CVarName, const TCHAR* FlagName, uint32 BitNumber, uint8* Force0MaskPtr, uint8* Force1MaskPtr, const TCHAR* Help, uint32 Flags) override;
-
-
-	virtual void CallAllConsoleVariableSinks() override;
-
-	virtual FConsoleVariableSinkHandle RegisterConsoleVariableSink_Handle(const FConsoleCommandDelegate& Command) override;
-	virtual void UnregisterConsoleVariableSink_Handle(FConsoleVariableSinkHandle Handle) override;
-
-	virtual IConsoleCommand* RegisterConsoleCommand(const TCHAR* Name, const TCHAR* Help, const FConsoleCommandDelegate& Command, uint32 Flags) override;
-	virtual IConsoleCommand* RegisterConsoleCommand(const TCHAR* Name, const TCHAR* Help, const FConsoleCommandWithArgsDelegate& Command, uint32 Flags) override;
-	virtual IConsoleCommand* RegisterConsoleCommand(const TCHAR* Name, const TCHAR* Help, const FConsoleCommandWithWorldDelegate& Command, uint32 Flags) override;
-	virtual IConsoleCommand* RegisterConsoleCommand(const TCHAR* Name, const TCHAR* Help, const FConsoleCommandWithWorldAndArgsDelegate& Command, uint32 Flags) override;
-	virtual IConsoleCommand* RegisterConsoleCommand(const TCHAR* Name, const TCHAR* Help, const FConsoleCommandWithArgsAndOutputDeviceDelegate& Command, uint32 Flags) override;
-	virtual IConsoleCommand* RegisterConsoleCommand(const TCHAR* Name, const TCHAR* Help, const FConsoleCommandWithWorldArgsAndOutputDeviceDelegate& Command, uint32 Flags) override;
-	virtual IConsoleCommand* RegisterConsoleCommand(const TCHAR* Name, const TCHAR* Help, const FConsoleCommandWithOutputDeviceDelegate& Command, uint32 Flags) override;
-	virtual IConsoleCommand* RegisterConsoleCommand(const TCHAR* Name, const TCHAR* Help, uint32 Flags) override;
-	virtual IConsoleObject* FindConsoleObject(const TCHAR* Name, bool bTrackFrequentCalls = true) const override;
-	virtual IConsoleVariable* FindConsoleVariable(const TCHAR* Name, bool bTrackFrequentCalls = true) const override;
-	virtual void ForEachConsoleObjectThatStartsWith(const FConsoleObjectVisitor& Visitor, const TCHAR* ThatStartsWith) const override;
-	virtual void ForEachConsoleObjectThatContains(const FConsoleObjectVisitor& Visitor, const TCHAR* ThatContains) const override;
-	virtual bool ProcessUserConsoleInput(const TCHAR* InInput, FOutputDevice& Ar, UWorld* InWorld) override;
-	virtual void AddConsoleHistoryEntry(const TCHAR* Key, const TCHAR* Input) override;
-	virtual void GetConsoleHistory(const TCHAR* Key, TArray<FString>& Out) override;
-	virtual bool IsNameRegistered(const TCHAR* Name) const override;
-	virtual void RegisterThreadPropagation(uint32 ThreadId, IConsoleThreadPropagation* InCallback) override;
-	virtual void UnregisterConsoleObject(IConsoleObject* Object, bool bKeepState) override;
-	virtual void UnsetAllConsoleVariablesWithTag(FName Tag, EConsoleVariableFlags Priority) override;
-	virtual void BatchUpdateTag(FName Tag, const TMap<FName, FString>& CVarsAndValues) override;
-#if ALLOW_OTHER_PLATFORM_CONFIG
-	virtual void LoadAllPlatformCVars(FName PlatformName, const FString& DeviceProfileName = FString()) override;
-	virtual void ClearAllPlatformCVars(FName PlatformName = NAME_None, const FString& DeviceProfileName = FString()) override;
-	virtual void PreviewPlatformCVars(FName PlatformName, const FString& DeviceProfileName, FName PreviewModeTag) override;
-	virtual void StompPlatformCVars(FName PlatformName, const FString& DeviceProfileName, FName Tag, EConsoleVariableFlags SetBy, EConsoleVariableFlags RequiredFlags, EConsoleVariableFlags DisallowedFlags) override;
-#endif
-
-	// Internal functionality to this file -----------------------------------
-
-	/** Map of console variables and commands, indexed by the name of that command or variable */
-	// [name] = pointer (pointer must not be 0)
-	TMap<FString, IConsoleObject*> ConsoleObjects;
-
-	bool bHistoryWasLoaded;
-	TMap<FString, TArray<FString>>	HistoryEntriesMap;
-	TArray<FConsoleCommandDelegate>	ConsoleVariableChangeSinks;
-
-	FConsoleVariableMulticastDelegate ConsoleVariableUnregisteredDelegate;
-	FConsoleObjectWithNameMulticastDelegate ConsoleObjectUnregisteredDelegate;
-
-	IConsoleThreadPropagation* ThreadPropagationCallback;
-
-	FCriticalSection CachedPlatformsAndDeviceProfilesLock;
-	TSet<FName> CachedPlatformsAndDeviceProfiles;
-
-	// if true the next call to CallAllConsoleVariableSinks() we will call all registered sinks
-	bool bCallAllConsoleVariableSinks;
-
-	/**
-		* Used to prevent concurrent access to ConsoleObjects.
-		* We don't aim to solve all concurrency problems (for example registering and unregistering a cvar on different threads, or reading a cvar from one thread while writing it from a different thread).
-		* Rather we just ensure that operations on a cvar from one thread will not conflict with operations on another cvar from another thread.
-	**/
-	mutable FTransactionallySafeCriticalSection ConsoleObjectsSynchronizationObject;
-
-	/**
-	 * @param Name must not be 0, must not be empty
-	 * @param Obj must not be 0
-	 * @return 0 if the name was already in use
-	 */
-	IConsoleObject* AddConsoleObject(const TCHAR* Name, IConsoleObject* Obj);
-	
-	/**
-	 * Similar to AddConsoleObject, but it just adds it without any flag checking or preexisting var checking
-	 */
-	void AddShadowConsoleObject(const TCHAR* Name, IConsoleObject* Obj);
-
-	/**
-	 * @param Stream must not be 0
-	 * @param Pattern must not be 0
-	 */
-	static bool MatchPartialName(const TCHAR* Stream, const TCHAR* Pattern);
-
-	/** Returns true if Pattern is found in Stream, case insensitive. */
-	static bool MatchSubstring(const TCHAR* Stream, const TCHAR* Pattern);
-
-	/**
-	 * Get string till whitespace, jump over whitespace
-	 * inefficient but this code is not performance critical
-	 */
-	static FString GetTextSection(const TCHAR*& It);
-
-	/** same as FindConsoleObject() but ECVF_CreatedFromIni are not filtered out (for internal use) */
-	IConsoleObject* FindConsoleObjectUnfiltered(const TCHAR* Name) const;
-
-	/**
-	 * Unregisters a console variable or command, if that object was registered.  For console variables, this will
-	 * actually only "deactivate" the variable so if it becomes registered again the state may persist
-	 * (unless bKeepState is false).
-	 *
-	 * @param	Name	Name of the console object to remove (not case sensitive)
-	 * @param	bKeepState	if the current state is kept in memory until a cvar with the same name is registered
-	 */
-	void UnregisterConsoleObject(const TCHAR* Name, bool bKeepState);
-
-	// reads HistoryEntriesMap from the .ini file (if not already loaded)
-	void LoadHistoryIfNeeded();
-
-	// writes HistoryEntriesMap to the .ini file
-	void SaveHistory();
-};
-
 static FConsoleManager& GetManager()
 {
-	return (FConsoleManager&)IConsoleManager::Get();
+	return static_cast<FConsoleManager&>(IConsoleManager::Get());
 }
 
 namespace UE::ConsoleManager::Private
 {
-	// this tracks the cvars that were  added dynamically with a tag (via plugin or similar)
+	// this tracks the cvars that were added dynamically with a tag (via plugin or similar)
 	// we use this structure to unset the cvars and update the value when the plugin unloads
 	TMap<FName, TSet<IConsoleVariable*>*> TaggedCVars;
 
@@ -624,13 +449,13 @@ private: // -----------------------------------------
 template <class T>
 void OnCVarChange(T& Dst, const T& Src, EConsoleVariableFlags Flags, EConsoleVariableFlags SetBy)
 {
-    // for the SetOnly case, just copy over the source to the dest
-    if (SetBy & ECVF_Set_SetOnly_Unsafe)
-    {
-        Dst = Src;
-        return;
-    }
-    
+	// for the SetOnly case, just copy over the source to the dest
+	if (SetBy & ECVF_Set_SetOnly_Unsafe)
+	{
+		Dst = Src;
+		return;
+	}
+	
 	FConsoleManager& ConsoleManager = GetManager();
 
 #if WITH_RELOAD
@@ -2074,65 +1899,13 @@ TAutoRangedConsoleVariable<float>::TAutoRangedConsoleVariable(const TCHAR* Name,
 {
 }
 
-//FAutoRangedConsoleVariable GTestVar(TEXT("josh.range"), 2, 0, 4, TEXT("Help str"), ECVF_Default);
-
-
-#if !(UE_BUILD_SHIPPING || UE_BUILD_TEST)
-// part of the automated test for console variables
-static TAutoConsoleVariable<int32> CVarDebugEarlyDefault(
-	TEXT("con.DebugEarlyDefault"),
-	21,
-	TEXT("used internally to test the console variable system"),
-	ECVF_Default);
-// part of the automated test for console variables
-static TAutoConsoleVariable<int32> CVarDebugEarlyCheat(
-	TEXT("con.DebugEarlyCheat"),
-	22,
-	TEXT("used internally to test the console variable system"),
-	ECVF_Cheat);
-#endif
-
 void FConsoleManager::CallAllConsoleVariableSinks()
 {
 	QUICK_SCOPE_CYCLE_COUNTER(ConsoleManager_CallAllConsoleVariableSinks);
 
-#if !(UE_BUILD_SHIPPING || UE_BUILD_TEST)
-	check(IsInGameThread());
-
-	// part of the automated test for console variables
-	// test the console variable system behavior with the ECVF_Cheat flag
+	if (bCallAllConsoleVariableSinks)
 	{
-		static uint32 LocalCounter = 0;
-
-		// after a few calls we assume the ini files are loaded
-		if(LocalCounter == 10)
-		{
-			IConsoleVariable* VarC = IConsoleManager::Get().RegisterConsoleVariable(TEXT("con.DebugLateDefault"), 23, TEXT("used internally to test the console variable system"), ECVF_Default);
-			IConsoleVariable* VarD = IConsoleManager::Get().RegisterConsoleVariable(TEXT("con.DebugLateCheat"), 24, TEXT("used internally to test the console variable system"), ECVF_Cheat);
-
-			int32 ValA = CVarDebugEarlyDefault.GetValueOnGameThread();
-			int32 ValB = CVarDebugEarlyCheat.GetValueOnGameThread();
-			int32 ValC = VarC->GetInt();
-			int32 ValD = VarD->GetInt();
-				
-			// in BaseEngine.ini we set all 4 cvars to "True" but only the non cheat one should pick up the value
-			check(ValA == 1);
-			check(ValB == 22);
-			check(ValC == 1);
-			check(ValD == 24);
-		}
-
-		// count up to 100 and don't warp around
-		if(LocalCounter < 100)
-		{
-			++LocalCounter;
-		}
-	}
-#endif
-
-	if(bCallAllConsoleVariableSinks)
-	{
-		for(uint32 i = 0; i < (uint32)ConsoleVariableChangeSinks.Num(); ++i)
+		for (int32 i = 0; i < ConsoleVariableChangeSinks.Num(); ++i)
 		{
 			ConsoleVariableChangeSinks[i].ExecuteIfBound();
 		}
@@ -2380,6 +2153,22 @@ public:
 	}
 };
 
+
+FConsoleManager::FConsoleManager()
+	: bHistoryWasLoaded(false)
+	, ThreadPropagationCallback(0)
+	, bCallAllConsoleVariableSinks(true)
+{}
+
+FConsoleManager::~FConsoleManager()
+{
+	for (TMap<FString, IConsoleObject*>::TConstIterator PairIt(ConsoleObjects); PairIt; ++PairIt)
+	{
+		IConsoleObject* Var = PairIt.Value();
+		delete Var;
+	}
+}
+
 IConsoleVariable* FConsoleManager::RegisterConsoleVariable(const TCHAR* Name, bool DefaultValue, const TCHAR* Help, uint32 Flags)
 {
 	LLM_SCOPE_BYNAME(TEXT("EngineMisc/ConsoleCommands"));
@@ -2757,7 +2546,7 @@ static void MultiLogf(FOutputDevice* Device, FArchive* File, const FmtType& Fmt,
 	}
 }
 
-static void DumpObjects(const TMap<FString, IConsoleObject*>& ConsoleObjects, const TCHAR* Params, FOutputDevice& InAr, bool bDisplayCommands)
+void FConsoleManager::DumpObjects(const TCHAR* Params, FOutputDevice& InAr, bool bDisplayCommands) const
 {
 	bool bShowHelp = FParse::Param(Params, TEXT("showhelp"));
 	bool bShowDeprecated = FParse::Param(Params, TEXT("deprecated"));
@@ -2992,12 +2781,12 @@ bool FConsoleManager::ProcessUserConsoleInput(const TCHAR* InInput, FOutputDevic
 
 	if (FParse::Command(&InInput, TEXT("dumpcvars")))
 	{
-		DumpObjects(ConsoleObjects, InInput, Ar, false);
+		DumpObjects(InInput, Ar, false);
 		return true;
 	}
 	if (FParse::Command(&InInput, TEXT("dumpccmds")))
 	{
-		DumpObjects(ConsoleObjects, InInput, Ar, true);
+		DumpObjects(InInput, Ar, true);
 		return true;
 	}
 	if (FParse::Command(&InInput, TEXT("setcvar")))
@@ -3706,236 +3495,6 @@ void FConsoleManager::ClearAllPlatformCVars(FName PlatformName, const FString& D
 
 #endif
 
-#if !(UE_BUILD_SHIPPING || UE_BUILD_TEST)
-uint32 GConsoleManagerSinkTestCounter = 0;
-void TestSinkCallback()
-{
-	++GConsoleManagerSinkTestCounter;
-}
-uint32 GConsoleVariableCallbackTestCounter = 0;
-void TestConsoleVariableCallback(IConsoleVariable* Var)
-{
-	check(Var);
-
-	float Value = Var->GetFloat();
-	check(FMath::IsNearlyEqual(Value, 3.1f, UE_KINDA_SMALL_NUMBER));
-
-	++GConsoleVariableCallbackTestCounter;
-}
-#endif // !(UE_BUILD_SHIPPING || UE_BUILD_TEST)
-
-void FConsoleManager::Test()
-{
-#if !(UE_BUILD_SHIPPING || UE_BUILD_TEST)
-	check(IsInGameThread());
-
-	// at this time we don't want to test with threading
-	check(GetThreadPropagationCallback() == 0);
-
-	// init ---------
-
-	GConsoleManagerSinkTestCounter = 0;
-	IConsoleManager::Get().CallAllConsoleVariableSinks();
-
-	// setup ---------
-
-	auto TestSinkCallbackHandle = RegisterConsoleVariableSink_Handle(FConsoleCommandDelegate::CreateStatic(&TestSinkCallback));
-
-	// start tests ---------
-
-	// no change should be triggered
-	IConsoleManager::Get().CallAllConsoleVariableSinks();
-	check(GConsoleManagerSinkTestCounter == 0);
-
-	for(uint32 Pass = 0; Pass < 2; ++Pass)
-	{
-		// we only test the main thread side of ECVF_RenderThreadSafe so we expect the same results
-		uint32 Flags = Pass ? ECVF_Default : ECVF_RenderThreadSafe;
-
-		int32 RefD = 2;
-		float RefE = 2.1f;
-
-		IConsoleVariable* VarA = IConsoleManager::Get().RegisterConsoleVariable(TEXT("TestNameA"), 1, TEXT("TestHelpA"), Flags);
-		IConsoleVariable* VarB = IConsoleManager::Get().RegisterConsoleVariable(TEXT("TestNameB"), 1.2f, TEXT("TestHelpB"), Flags);
-		IConsoleVariable* VarD = IConsoleManager::Get().RegisterConsoleVariableRef(TEXT("TestNameD"), RefD, TEXT("TestHelpD"), Flags);
-		IConsoleVariable* VarE = IConsoleManager::Get().RegisterConsoleVariableRef(TEXT("TestNameE"), RefE, TEXT("TestHelpE"), Flags);
-
-		// at the moment ECVF_SetByConstructor has to be 0 or we set ECVF_Default to ECVF_SetByConstructor
-		check((VarA->GetFlags() & ECVF_SetByMask) == ECVF_SetByConstructor);
-
-		GConsoleVariableCallbackTestCounter = 0;
-		VarB->SetOnChangedCallback(FConsoleVariableDelegate::CreateStatic(&TestConsoleVariableCallback));
-		check(GConsoleVariableCallbackTestCounter == 0);
-
-		// make sure the vars are there
-
-		check(VarA == IConsoleManager::Get().FindConsoleVariable(TEXT("TestNameA")));
-		check(VarB == IConsoleManager::Get().FindConsoleVariable(TEXT("TestNameB")));
-		check(VarD == IConsoleManager::Get().FindConsoleVariable(TEXT("TestNameD")));
-		check(VarE == IConsoleManager::Get().FindConsoleVariable(TEXT("TestNameE")));
-
-		// test Get()
-
-		check(VarA->GetInt() == 1);
-		check(VarA->GetFloat() == 1);
-		check(VarA->GetString() == FString(TEXT("1")));
-
-		check(VarB->GetInt() == 1);
-		check(FMath::IsNearlyEqual(VarB->GetFloat(), 1.2f, UE_KINDA_SMALL_NUMBER));
-		check(VarB->GetString() == FString(TEXT("1.2")));
-
-		check(RefD == 2);
-		check(VarD->GetInt() == 2);
-		check(VarD->GetFloat() == (float)2);
-		check(VarD->GetString() == FString(TEXT("2")));
-
-		check(FMath::IsNearlyEqual(RefE, 2.1f, UE_KINDA_SMALL_NUMBER));
-		check(VarE->GetInt() == (int32)RefE);
-		check(VarE->GetFloat() == RefE);
-		check(VarE->GetString() == FString(TEXT("2.1")));
-
-		// call Set(string)
-
-		VarA->Set(TEXT("3.1"), ECVF_SetByConsoleVariablesIni);
-		VarB->Set(TEXT("3.1"), ECVF_SetByConsoleVariablesIni);
-		VarD->Set(TEXT("3.1"), ECVF_SetByConsoleVariablesIni);
-		VarE->Set(TEXT("3.1"), ECVF_SetByConsoleVariablesIni);
-
-		check(GConsoleVariableCallbackTestCounter == 1);
-
-		// verify Set()
-
-		check(VarA->GetString() == FString(TEXT("3")));
-		check(VarB->GetString() == FString(TEXT("3.1")));
-		check(VarD->GetString() == FString(TEXT("3")));
-		check(VarE->GetString() == FString(TEXT("3.1")));
-		check(RefD == 3);
-		check(RefE == 3.1f);
-
-		VarB->Set(TEXT("3.1"), ECVF_SetByConsoleVariablesIni);
-		check(GConsoleVariableCallbackTestCounter == 2);
-
-		// unregister
-
-		IConsoleManager::Get().UnregisterConsoleObject(VarA);
-		IConsoleManager::Get().UnregisterConsoleObject(VarB, false);
-		UnregisterConsoleObject(TEXT("TestNameD"), false);
-		UnregisterConsoleObject(TEXT("TestNameE"), false);
-
-		check(!IConsoleManager::Get().FindConsoleVariable(TEXT("TestNameA")));
-		check(!IConsoleManager::Get().FindConsoleVariable(TEXT("TestNameB")));
-		check(!IConsoleManager::Get().FindConsoleVariable(TEXT("TestNameD")));
-		check(!IConsoleManager::Get().FindConsoleVariable(TEXT("TestNameE")));
-
-		// re-register but maintain state
-		IConsoleVariable* SecondVarA = IConsoleManager::Get().RegisterConsoleVariable(TEXT("TestNameA"), 1234, TEXT("TestHelpSecondA"), Flags);
-		check(SecondVarA == VarA);
-		check(SecondVarA->GetInt() == 3);
-		check(IConsoleManager::Get().FindConsoleVariable(TEXT("TestNameA")));
-
-		UnregisterConsoleObject(TEXT("TestNameA"), false);
-		check(!IConsoleManager::Get().FindConsoleVariable(TEXT("TestNameA")));
-
-		if((Flags & ECVF_RenderThreadSafe) == 0)
-		{
-			// string is not supported with the flag ECVF_RenderThreadSafe
-			IConsoleVariable* VarC = IConsoleManager::Get().RegisterConsoleVariable(TEXT("TestNameC"), TEXT("1.23"), TEXT("TestHelpC"), Flags);
-			check(VarC == IConsoleManager::Get().FindConsoleVariable(TEXT("TestNameC")));
-			check(VarC->GetInt() == 1);
-			// note: exact comparison fails in Win32 release
-			check(FMath::IsNearlyEqual(VarC->GetFloat(), 1.23f, UE_KINDA_SMALL_NUMBER));
-			check(VarC->GetString() == FString(TEXT("1.23")));
-			VarC->Set(TEXT("3.1"), ECVF_SetByConsole);
-			check(VarC->GetString() == FString(TEXT("3.1")));
-			UnregisterConsoleObject(TEXT("TestNameC"), false);
-			check(!IConsoleManager::Get().FindConsoleVariable(TEXT("TestNameC")));
-		}
-
-		// verify priority
-		{
-			IConsoleVariable* VarX = IConsoleManager::Get().RegisterConsoleVariable(TEXT("TestNameX"), 1, TEXT("TestHelpX"), Flags);
-			check(((uint32)VarX->GetFlags() & ECVF_SetByMask) == ECVF_SetByConstructor);
-
-			VarX->Set(TEXT("3.1"), ECVF_SetByConsoleVariablesIni);
-			check(((uint32)VarX->GetFlags() & ECVF_SetByMask) == ECVF_SetByConsoleVariablesIni);
-
-			// lower should fail
-			VarX->Set(TEXT("111"), ECVF_SetByScalability);
-			check(VarX->GetString() == FString(TEXT("3")));
-			check(((uint32)VarX->GetFlags() & ECVF_SetByMask) == ECVF_SetByConsoleVariablesIni);
-
-			// higher should work
-			VarX->Set(TEXT("222"), ECVF_SetByCommandline);
-			check(VarX->GetString() == FString(TEXT("222")));
-			check(((uint32)VarX->GetFlags() & ECVF_SetByMask) == ECVF_SetByCommandline);
-
-			// lower should fail
-			VarX->Set(TEXT("333"), ECVF_SetByConsoleVariablesIni);
-			check(VarX->GetString() == FString(TEXT("222")));
-			check(((uint32)VarX->GetFlags() & ECVF_SetByMask) == ECVF_SetByCommandline);
-
-			// higher should work
-			VarX->Set(TEXT("444"), ECVF_SetByConsole);
-			check(VarX->GetString() == FString(TEXT("444")));
-			check(((uint32)VarX->GetFlags() & ECVF_SetByMask) == ECVF_SetByConsole);
-
-			IConsoleManager::Get().UnregisterConsoleObject(VarX, false);
-		}
-	}
-
-	// this should trigger the callback
-	IConsoleManager::Get().CallAllConsoleVariableSinks();
-	check(GConsoleManagerSinkTestCounter == 1);
-
-	// this should not trigger the callback
-	IConsoleManager::Get().CallAllConsoleVariableSinks();
-	check(GConsoleManagerSinkTestCounter == 1);
-
-	// this should also trigger the callback
-	TestSinkCallback();
-	check(GConsoleManagerSinkTestCounter == 2);
-
-	UnregisterConsoleVariableSink_Handle(TestSinkCallbackHandle);
-
-
-	// TODO: These tests end up causing an issue on PS4 due to static initialization order and
-	// logging grabbing a TLS slot the a PS4 tool expects to use. Try moving them to proper unit
-	// tests instead.
-	#if false
-	FConsoleManager& Manager = GetManager();
-
-	// Deprecated console variables should not assert in dump commands
-	{
-		constexpr TCHAR OldName[] = TEXT("TestVar.Old");
-		constexpr TCHAR NewName[] = TEXT("TestVar.New");
-
-		FAutoConsoleVariable NewVar(NewName, false, TEXT(""));
-		FAutoConsoleVariableDeprecated OldVar(OldName, NewName, TEXT("0.0"));
-		ON_SCOPE_EXIT { Manager.UnregisterConsoleObject(OldName, false); };
-
-		FOutputDeviceNull OutNull;
-		DumpObjects(Manager.ConsoleObjects, TEXT("-ShowHelp"), OutNull, false);
-		DumpObjects(Manager.ConsoleObjects, TEXT("-ShowHelp -Deprecated"), OutNull, false);
-	}
-
-	// Deprecated console commands should not assert in dump commands
-	{
-		constexpr TCHAR OldName[] = TEXT("TestCmd.Old");
-		constexpr TCHAR NewName[] = TEXT("TestCmd.New");
-
-		FAutoConsoleCommand NewCmd(NewName, TEXT(""), FConsoleCommandDelegate());
-		FAutoConsoleCommandDeprecated OldCmd(OldName, NewName, TEXT("0.0"));
-		ON_SCOPE_EXIT { Manager.UnregisterConsoleObject(OldName, false); };
-
-		FOutputDeviceNull OutNull;
-		DumpObjects(Manager.ConsoleObjects, TEXT("-ShowHelp"), OutNull, true);
-		DumpObjects(Manager.ConsoleObjects, TEXT("-ShowHelp -Deprecated"), OutNull, true);
-	}
-	#endif
-
-#endif // !(UE_BUILD_SHIPPING || UE_BUILD_TEST)
-}
-
 // These don't belong here, but they belong here more than they belong in launch engine loop.
 void CreateConsoleVariables()
 {
@@ -4012,12 +3571,6 @@ void CreateConsoleVariables()
 		ECVF_Default);
 #endif // !UE_BUILD_SHIPPING
 
-	// testing code
-	{
-		FConsoleManager& ConsoleManager = GetManager();
-
-		ConsoleManager.Test();
-	}
 #endif
 }
 
diff --git a/Engine/Source/Runtime/Core/Private/HAL/ConsoleManager.h b/Engine/Source/Runtime/Core/Private/HAL/ConsoleManager.h
new file mode 100644
--- /dev/null
+++ b/Engine/Source/Runtime/Core/Private/HAL/ConsoleManager.h
@@ -0,0 +1,147 @@
+// Copyright Epic Games, Inc. All Rights Reserved.
+
+#pragma once
+
+#include "HAL/IConsoleManager.h"
+
+class FConsoleManager : public IConsoleManager
+{
+public:
+
+	FConsoleManager();
+	~FConsoleManager();
+
+	virtual IConsoleVariable* RegisterConsoleVariable(const TCHAR* Name, bool DefaultValue, const TCHAR* Help, uint32 Flags) override;
+	virtual IConsoleVariable* RegisterConsoleVariable(const TCHAR* Name, int32 DefaultValue, const TCHAR* Help, uint32 Flags) override;
+	virtual IConsoleVariable* RegisterConsoleVariable(const TCHAR* Name, float DefaultValue, const TCHAR* Help, uint32 Flags) override;
+	virtual IConsoleVariable* RegisterConsoleVariable(const TCHAR* Name, const TCHAR* DefaultValue, const TCHAR* Help, uint32 Flags) override;
+	virtual IConsoleVariable* RegisterConsoleVariable(const TCHAR* Name, const FString& DefaultValue, const TCHAR* Help, uint32 Flags) override;
+
+	virtual IConsoleVariable* RegisterConsoleVariableRef(const TCHAR* Name, bool& RefValue, const TCHAR* Help, uint32 Flags) override;
+	virtual IConsoleVariable* RegisterConsoleVariableRef(const TCHAR* Name, int32& RefValue, const TCHAR* Help, uint32 Flags) override;
+	virtual IConsoleVariable* RegisterConsoleVariableRef(const TCHAR* Name, float& RefValue, const TCHAR* Help, uint32 Flags) override;
+	virtual IConsoleVariable* RegisterConsoleVariableRef(const TCHAR* Name, FString& RefValue, const TCHAR* Help, uint32 Flags) override;
+	virtual IConsoleVariable* RegisterConsoleVariableRef(const TCHAR* Name, FName& RefValue, const TCHAR* Help, uint32 Flags) override;
+	virtual IConsoleVariable* RegisterConsoleVariableBitRef(const TCHAR* CVarName, const TCHAR* FlagName, uint32 BitNumber, uint8* Force0MaskPtr, uint8* Force1MaskPtr, const TCHAR* Help, uint32 Flags) override;
+
+	virtual IConsoleCommand* RegisterConsoleCommand(const TCHAR* Name, const TCHAR* Help, const FConsoleCommandDelegate& Command, uint32 Flags) override;
+	virtual IConsoleCommand* RegisterConsoleCommand(const TCHAR* Name, const TCHAR* Help, const FConsoleCommandWithArgsDelegate& Command, uint32 Flags) override;
+	virtual IConsoleCommand* RegisterConsoleCommand(const TCHAR* Name, const TCHAR* Help, const FConsoleCommandWithWorldDelegate& Command, uint32 Flags) override;
+	virtual IConsoleCommand* RegisterConsoleCommand(const TCHAR* Name, const TCHAR* Help, const FConsoleCommandWithWorldAndArgsDelegate& Command, uint32 Flags) override;
+	virtual IConsoleCommand* RegisterConsoleCommand(const TCHAR* Name, const TCHAR* Help, const FConsoleCommandWithArgsAndOutputDeviceDelegate& Command, uint32 Flags) override;
+	virtual IConsoleCommand* RegisterConsoleCommand(const TCHAR* Name, const TCHAR* Help, const FConsoleCommandWithWorldArgsAndOutputDeviceDelegate& Command, uint32 Flags) override;
+	virtual IConsoleCommand* RegisterConsoleCommand(const TCHAR* Name, const TCHAR* Help, const FConsoleCommandWithOutputDeviceDelegate& Command, uint32 Flags) override;
+	virtual IConsoleCommand* RegisterConsoleCommand(const TCHAR* Name, const TCHAR* Help, uint32 Flags) override;
+
+	virtual FConsoleVariableMulticastDelegate& OnCVarUnregistered()override;
+	virtual FConsoleObjectWithNameMulticastDelegate& OnConsoleObjectUnregistered() override;
+
+	virtual FString FindConsoleObjectName(const IConsoleObject* Obj) const override;
+	virtual IConsoleObject* FindConsoleObject(const TCHAR* Name, bool bTrackFrequentCalls = true) const override;
+	virtual IConsoleVariable* FindConsoleVariable(const TCHAR* Name, bool bTrackFrequentCalls = true) const override;
+
+	virtual FConsoleVariableSinkHandle RegisterConsoleVariableSink_Handle(const FConsoleCommandDelegate& Command) override;
+	virtual void UnregisterConsoleVariableSink_Handle(FConsoleVariableSinkHandle Handle) override;
+	virtual void CallAllConsoleVariableSinks() override;
+
+	virtual void ForEachConsoleObjectThatStartsWith(const FConsoleObjectVisitor& Visitor, const TCHAR* ThatStartsWith) const override;
+	virtual void ForEachConsoleObjectThatContains(const FConsoleObjectVisitor& Visitor, const TCHAR* ThatContains) const override;
+	virtual bool ProcessUserConsoleInput(const TCHAR* InInput, FOutputDevice& Ar, UWorld* InWorld) override;
+	virtual void AddConsoleHistoryEntry(const TCHAR* Key, const TCHAR* Input) override;
+	virtual void GetConsoleHistory(const TCHAR* Key, TArray<FString>& Out) override;
+	virtual bool IsNameRegistered(const TCHAR* Name) const override;
+	virtual void RegisterThreadPropagation(uint32 ThreadId, IConsoleThreadPropagation* InCallback) override;
+	virtual void UnregisterConsoleObject(IConsoleObject* Object, bool bKeepState) override;
+	virtual void UnsetAllConsoleVariablesWithTag(FName Tag, EConsoleVariableFlags Priority) override;
+	virtual void BatchUpdateTag(FName Tag, const TMap<FName, FString>& CVarsAndValues) override;
+
+#if ALLOW_OTHER_PLATFORM_CONFIG
+	virtual void LoadAllPlatformCVars(FName PlatformName, const FString& DeviceProfileName = FString()) override;
+	virtual void ClearAllPlatformCVars(FName PlatformName = NAME_None, const FString& DeviceProfileName = FString()) override;
+	virtual void PreviewPlatformCVars(FName PlatformName, const FString& DeviceProfileName, FName PreviewModeTag) override;
+	virtual void StompPlatformCVars(FName PlatformName, const FString& DeviceProfileName, FName Tag, EConsoleVariableFlags SetBy, EConsoleVariableFlags RequiredFlags, EConsoleVariableFlags DisallowedFlags) override;
+#endif
+
+	/**
+	 * @param Name must not be 0, must not be empty
+	 * @param Obj must not be 0
+	 * @return 0 if the name was already in use
+	 */
+	 IConsoleObject* AddConsoleObject(const TCHAR* Name, IConsoleObject* Obj);
+
+	/**
+	 * Similar to AddConsoleObject, but it just adds it without any flag checking or preexisting var checking
+	 */
+	void AddShadowConsoleObject(const TCHAR* Name, IConsoleObject* Obj);
+
+	/** Internally needed for ECVF_RenderThreadSafe */
+	IConsoleThreadPropagation* GetThreadPropagationCallback();
+
+	/** Internally needed for ECVF_RenderThreadSafe */
+	bool IsThreadPropagationThread();
+
+	void OnCVarChanged();
+	void DumpObjects(const TCHAR* Params, FOutputDevice& InAr, bool bDisplayCommands) const;
+
+private:
+
+	/**
+	 * @param Stream must not be 0
+	 * @param Pattern must not be 0
+	 */
+	static bool MatchPartialName(const TCHAR* Stream, const TCHAR* Pattern);
+
+	/** Returns true if Pattern is found in Stream, case insensitive. */
+	static bool MatchSubstring(const TCHAR* Stream, const TCHAR* Pattern);
+
+	/**
+	 * Get string till whitespace, jump over whitespace
+	 * inefficient but this code is not performance critical
+	 */
+	static FString GetTextSection(const TCHAR*& It);
+
+	/** same as FindConsoleObject() but ECVF_CreatedFromIni are not filtered out (for internal use) */
+	IConsoleObject* FindConsoleObjectUnfiltered(const TCHAR* Name) const;
+
+	/**
+	 * Unregisters a console variable or command, if that object was registered. For console variables, this will
+	 * actually only "deactivate" the variable so if it becomes registered again the state may persist (unless
+	 * bKeepState is false).
+	 *
+	 * @param Name Name of the console object to remove (not case sensitive)
+	 * @param bKeepState if the current state is kept in memory until a cvar with the same name is registered
+	 */
+	void UnregisterConsoleObject(const TCHAR* Name, bool bKeepState);
+
+	/** Reads HistoryEntriesMap from the .ini file (if not already loaded) */
+	void LoadHistoryIfNeeded();
+
+	/** Writes HistoryEntriesMap to the .ini file */
+	void SaveHistory();
+
+	/**
+	 * Map of console variables and commands, indexed by the name of that command or variable
+	 * [name] = pointer (pointer must not be 0)
+	 */
+	TMap<FString, IConsoleObject*> ConsoleObjects;
+
+	bool bHistoryWasLoaded;
+	TMap<FString, TArray<FString>> HistoryEntriesMap;
+	TArray<FConsoleCommandDelegate> ConsoleVariableChangeSinks;
+	FConsoleVariableMulticastDelegate ConsoleVariableUnregisteredDelegate;
+	FConsoleObjectWithNameMulticastDelegate ConsoleObjectUnregisteredDelegate;
+	IConsoleThreadPropagation* ThreadPropagationCallback;
+	FCriticalSection CachedPlatformsAndDeviceProfilesLock;
+	TSet<FName> CachedPlatformsAndDeviceProfiles;
+
+	/** If true the next call to CallAllConsoleVariableSinks() we will call all registered sinks */
+	bool bCallAllConsoleVariableSinks;
+
+	/**
+	 * Used to prevent concurrent access to ConsoleObjects. We don't aim to solve all concurrency problems (for example
+	 * registering and unregistering a cvar on different threads, or reading a cvar from one thread while writing it
+	 * from a different thread). Rather we just ensure that operations on a cvar from one thread will not conflict with
+	 * operations on another cvar from another thread.
+	 */
+	mutable FTransactionallySafeCriticalSection ConsoleObjectsSynchronizationObject;
+};
diff --git a/Engine/Source/Runtime/Core/Public/HAL/IConsoleManager.h b/Engine/Source/Runtime/Core/Public/HAL/IConsoleManager.h
--- a/Engine/Source/Runtime/Core/Public/HAL/IConsoleManager.h
+++ b/Engine/Source/Runtime/Core/Public/HAL/IConsoleManager.h
@@ -1098,7 +1098,7 @@ struct IConsoleManager
 
 	/**
 	 * Lookup the name of a console object by its pointer 
-	 * @param Object to lookup
+	 * @param Object to lookup. Must not be null.
 	 * @return Name of the object, or an empty string if the object can't be found
 	 */
 	virtual FString FindConsoleObjectName(const IConsoleObject* Obj) const = 0;
diff --git a/Engine/Source/Runtime/Core/Tests/HAL/ConsoleManagerTest.cpp b/Engine/Source/Runtime/Core/Tests/HAL/ConsoleManagerTest.cpp
new file mode 100644
--- /dev/null
+++ b/Engine/Source/Runtime/Core/Tests/HAL/ConsoleManagerTest.cpp
@@ -0,0 +1,267 @@
+// Copyright Epic Games, Inc. All Rights Reserved.
+
+#if WITH_TESTS
+
+#include "HAL/ConsoleManager.h"
+
+#include "Misc/OutputDeviceNull.h"
+#include "Misc/ScopeExit.h"
+#include "Tests/TestHarnessAdapter.h"
+
+// TODO: Try to move these into the test
+static TAutoConsoleVariable<int32> CVarDebugEarlyDefault(
+	TEXT("con.DebugEarlyDefault"),
+	21,
+	TEXT("used internally to test the console variable system"),
+	ECVF_Default);
+
+static TAutoConsoleVariable<int32> CVarDebugEarlyCheat(
+	TEXT("con.DebugEarlyCheat"),
+	22,
+	TEXT("used internally to test the console variable system"),
+	ECVF_Cheat);
+
+TEST_CASE_NAMED(FConsoleManagerTest, "System::Core::HAL::ConsoleManager", "[ApplicationContextMask][Core]")
+{
+	// FConsoleManager ManagerImpl{};
+	// IConsoleManager& Manager = ManagerImpl;
+
+	IConsoleManager& Manager = IConsoleManager::Get();
+	FConsoleManager& ManagerImpl = static_cast<FConsoleManager&>(Manager);
+
+	// HACK: Temporarily disable any thread propagation callback. This causes setting a variable to be deferred which
+	// the write-read test below to become a race condition. The proper solution here is to create a temporary
+	// FConsoleManager instance so we aren't mucking with the real runtime instance. However,
+	// IConsoleObject::OnCVarChange looks up the global IConsoleManager::Singleton, instead of using our temporary
+	// instance. Changing console objects to contain a pointer to their manager instead of doing a global lookup would
+	// solve that problem, at the expense of making objects 8 bytes larger.
+	IConsoleThreadPropagation* Prop = ManagerImpl.GetThreadPropagationCallback();
+	Manager.RegisterThreadPropagation(0, nullptr);
+	ON_SCOPE_EXIT { Manager.RegisterThreadPropagation(0, Prop); };
+
+	// we only test the main thread side of ECVF_RenderThreadSafe so we expect the same results
+	for (EConsoleVariableFlags Flags : {ECVF_Default, ECVF_RenderThreadSafe})
+	{
+		int32 RefD = 2;
+		float RefE = 2.1f;
+
+		IConsoleVariable* VarA = Manager.RegisterConsoleVariable(TEXT("TestNameA"), 1, TEXT("TestHelpA"), Flags);
+		IConsoleVariable* VarB = Manager.RegisterConsoleVariable(TEXT("TestNameB"), 1.2f, TEXT("TestHelpB"), Flags);
+		IConsoleVariable* VarD = Manager.RegisterConsoleVariableRef(TEXT("TestNameD"), RefD, TEXT("TestHelpD"), Flags);
+		IConsoleVariable* VarE = Manager.RegisterConsoleVariableRef(TEXT("TestNameE"), RefE, TEXT("TestHelpE"), Flags);
+
+		SECTION("Console Variable Sinks")
+		{
+			// Sinks are initiallially queued
+			Manager.CallAllConsoleVariableSinks();
+
+			uint32 SinkCounter = 0;
+			auto SinkCallback = [&SinkCounter]() { ++SinkCounter; };
+			FConsoleCommandDelegate SinkDelegate = FConsoleCommandDelegate::CreateLambda(SinkCallback);
+			FConsoleVariableSinkHandle SinkHandle = Manager.RegisterConsoleVariableSink_Handle(SinkDelegate);
+			ON_SCOPE_EXIT { Manager.UnregisterConsoleVariableSink_Handle(SinkHandle); };
+
+			IConsoleVariable* Var = Manager.RegisterConsoleVariable(TEXT("TestNameX"), 1, TEXT("TestHelpX"), Flags);
+			ON_SCOPE_EXIT { Manager.UnregisterConsoleObject(Var, false); };
+
+			Manager.CallAllConsoleVariableSinks();
+			CHECK(SinkCounter == 0);
+
+			Var->Set(2);
+
+			// this should trigger the callback
+			Manager.CallAllConsoleVariableSinks();
+			CHECK(SinkCounter == 1);
+
+			// this should not trigger the callback
+			Manager.CallAllConsoleVariableSinks();
+			CHECK(SinkCounter == 1);
+		}
+
+		uint32 ChangeCounter = 0;
+		auto ChangeCallback = [&ChangeCounter](IConsoleVariable* Var)
+		{
+			CHECK(Var);
+			if (Var)
+			{
+				float Value = Var->GetFloat();
+				CHECK(FMath::IsNearlyEqual(Value, 3.1f, UE_KINDA_SMALL_NUMBER));
+				++ChangeCounter;
+			}
+		};
+		FConsoleVariableDelegate ChangeDelegate = FConsoleVariableDelegate::CreateLambda(ChangeCallback);
+		VarB->SetOnChangedCallback(ChangeDelegate);
+		CHECK(ChangeCounter == 0);
+
+		SECTION("Register variables")
+		{
+			// at the moment ECVF_SetByConstructor has to be 0 or we set ECVF_Default to ECVF_SetByConstructor
+			CHECK((VarA->GetFlags() & ECVF_SetByMask) == ECVF_SetByConstructor);
+
+			CHECK(VarA == Manager.FindConsoleVariable(TEXT("TestNameA")));
+			CHECK(VarB == Manager.FindConsoleVariable(TEXT("TestNameB")));
+			CHECK(VarD == Manager.FindConsoleVariable(TEXT("TestNameD")));
+			CHECK(VarE == Manager.FindConsoleVariable(TEXT("TestNameE")));
+		}
+
+		SECTION("Get variable values")
+		{
+			CHECK(VarA->GetInt() == 1);
+			CHECK(VarA->GetFloat() == 1);
+			CHECK(VarA->GetString() == FString(TEXT("1")));
+
+			CHECK(VarB->GetInt() == 1);
+			CHECK(FMath::IsNearlyEqual(VarB->GetFloat(), 1.2f, UE_KINDA_SMALL_NUMBER));
+			CHECK(VarB->GetString() == FString(TEXT("1.2")));
+
+			CHECK(RefD == 2);
+			CHECK(VarD->GetInt() == 2);
+			CHECK(VarD->GetFloat() == (float)2);
+			CHECK(VarD->GetString() == FString(TEXT("2")));
+
+			CHECK(FMath::IsNearlyEqual(RefE, 2.1f, UE_KINDA_SMALL_NUMBER));
+			CHECK(VarE->GetInt() == (int32)RefE);
+			CHECK(VarE->GetFloat() == RefE);
+			CHECK(VarE->GetString() == FString(TEXT("2.1")));
+		}
+
+		SECTION("Set variable values (string)")
+		{
+			VarA->Set(TEXT("3.1"), ECVF_SetByConsoleVariablesIni);
+			VarB->Set(TEXT("3.1"), ECVF_SetByConsoleVariablesIni);
+			VarD->Set(TEXT("3.1"), ECVF_SetByConsoleVariablesIni);
+			VarE->Set(TEXT("3.1"), ECVF_SetByConsoleVariablesIni);
+			CHECK(ChangeCounter == 1);
+
+			CHECK(VarA->GetString() == FString(TEXT("3")));
+			CHECK(VarB->GetString() == FString(TEXT("3.1")));
+			CHECK(VarD->GetString() == FString(TEXT("3")));
+			CHECK(VarE->GetString() == FString(TEXT("3.1")));
+			CHECK(RefD == 3);
+			CHECK(RefE == 3.1f);
+
+			VarB->Set(TEXT("3.1"), ECVF_SetByConsoleVariablesIni);
+			CHECK(ChangeCounter == 2);
+		}
+
+		if ((Flags & ECVF_RenderThreadSafe) == 0)
+		{
+			// string is not supported with the flag ECVF_RenderThreadSafe
+			IConsoleVariable* VarC = Manager.RegisterConsoleVariable(TEXT("TestNameC"), TEXT("1.23"), TEXT("TestHelpC"), Flags);
+			CHECK(VarC == Manager.FindConsoleVariable(TEXT("TestNameC")));
+			CHECK(VarC->GetInt() == 1);
+			// note: exact comparison fails in Win32 release
+			CHECK(FMath::IsNearlyEqual(VarC->GetFloat(), 1.23f, UE_KINDA_SMALL_NUMBER));
+			CHECK(VarC->GetString() == FString(TEXT("1.23")));
+			VarC->Set(TEXT("3.1"), ECVF_SetByConsole);
+			CHECK(VarC->GetString() == FString(TEXT("3.1")));
+
+			Manager.UnregisterConsoleObject(TEXT("TestNameC"), false);
+			CHECK(!Manager.FindConsoleVariable(TEXT("TestNameC")));
+		}
+
+		Manager.UnregisterConsoleObject(VarA);
+		Manager.UnregisterConsoleObject(VarB, false);
+		Manager.UnregisterConsoleObject(TEXT("TestNameD"), false);
+		Manager.UnregisterConsoleObject(TEXT("TestNameE"), false);
+
+		SECTION("Unregister variables")
+		{
+			CHECK(!Manager.FindConsoleVariable(TEXT("TestNameA")));
+			CHECK(!Manager.FindConsoleVariable(TEXT("TestNameB")));
+			CHECK(!Manager.FindConsoleVariable(TEXT("TestNameD")));
+			CHECK(!Manager.FindConsoleVariable(TEXT("TestNameE")));
+		}
+
+		SECTION("Re-register variables but maintain state")
+		{
+			IConsoleVariable* SecondVarA = Manager.RegisterConsoleVariable(TEXT("TestNameA"), 1234, TEXT("TestHelpSecondA"), Flags);
+			CHECK(SecondVarA == VarA);
+			CHECK(SecondVarA->GetInt() == 3);
+			CHECK(Manager.FindConsoleVariable(TEXT("TestNameA")));
+
+			Manager.UnregisterConsoleObject(TEXT("TestNameA"), false);
+			CHECK(!Manager.FindConsoleVariable(TEXT("TestNameA")));
+		}
+
+		SECTION("Priority")
+		{
+			IConsoleVariable* VarX = Manager.RegisterConsoleVariable(TEXT("TestNameX"), 1, TEXT("TestHelpX"), Flags);
+			ON_SCOPE_EXIT { Manager.UnregisterConsoleObject(VarX, false); };
+
+			CHECK(((uint32)VarX->GetFlags() & ECVF_SetByMask) == ECVF_SetByConstructor);
+
+			VarX->Set(TEXT("3.1"), ECVF_SetByConsoleVariablesIni);
+			CHECK(((uint32)VarX->GetFlags() & ECVF_SetByMask) == ECVF_SetByConsoleVariablesIni);
+
+			// lower should fail
+			VarX->Set(TEXT("111"), ECVF_SetByScalability);
+			CHECK(VarX->GetString() == FString(TEXT("3")));
+			CHECK(((uint32)VarX->GetFlags() & ECVF_SetByMask) == ECVF_SetByConsoleVariablesIni);
+
+			// higher should work
+			VarX->Set(TEXT("222"), ECVF_SetByCommandline);
+			CHECK(VarX->GetString() == FString(TEXT("222")));
+			CHECK(((uint32)VarX->GetFlags() & ECVF_SetByMask) == ECVF_SetByCommandline);
+
+			// lower should fail
+			VarX->Set(TEXT("333"), ECVF_SetByConsoleVariablesIni);
+			CHECK(VarX->GetString() == FString(TEXT("222")));
+			CHECK(((uint32)VarX->GetFlags() & ECVF_SetByMask) == ECVF_SetByCommandline);
+
+			// higher should work
+			VarX->Set(TEXT("444"), ECVF_SetByConsole);
+			CHECK(VarX->GetString() == FString(TEXT("444")));
+			CHECK(((uint32)VarX->GetFlags() & ECVF_SetByMask) == ECVF_SetByConsole);
+		}
+	}
+
+	// We don't load config in Low Level Tests
+	if constexpr (!WITH_LOW_LEVEL_TESTS)
+	{
+		SECTION("ECVF_Cheat")
+		{
+			IConsoleVariable* VarC = Manager.RegisterConsoleVariable(TEXT("con.DebugLateDefault"), 23, TEXT(""), ECVF_Default);
+			ON_SCOPE_EXIT { Manager.UnregisterConsoleObject(VarC, true); };
+
+			IConsoleVariable* VarD = Manager.RegisterConsoleVariable(TEXT("con.DebugLateCheat"), 24, TEXT(""), ECVF_Cheat);
+			ON_SCOPE_EXIT { Manager.UnregisterConsoleObject(VarD, true); };
+
+			// in BaseEngine.ini we set all 4 cvars to "True" but only the non cheat one should pick up the value
+			CHECK(CVarDebugEarlyDefault.GetValueOnGameThread() == 1);
+			CHECK(CVarDebugEarlyCheat.GetValueOnGameThread() == 22);
+			CHECK(VarC->GetInt() == 1);
+			CHECK(VarD->GetInt() == 24);
+		}
+	}
+
+	SECTION("Deprecated console variables should not assert in dump commands")
+	{
+		constexpr TCHAR OldName[] = TEXT("TestVar.Old");
+		constexpr TCHAR NewName[] = TEXT("TestVar.New");
+
+		FAutoConsoleVariable NewVar(NewName, false, TEXT(""));
+		FAutoConsoleVariableDeprecated OldVar(OldName, NewName, TEXT("0.0"));
+		ON_SCOPE_EXIT { Manager.UnregisterConsoleObject(OldName, false); };
+
+		FOutputDeviceNull OutNull;
+		ManagerImpl.DumpObjects(TEXT("-ShowHelp"), OutNull, false);
+		ManagerImpl.DumpObjects(TEXT("-ShowHelp -Deprecated"), OutNull, false);
+	}
+
+	SECTION("Deprecated console commands should not assert in dump commands")
+	{
+		constexpr TCHAR OldName[] = TEXT("TestCmd.Old");
+		constexpr TCHAR NewName[] = TEXT("TestCmd.New");
+
+		FAutoConsoleCommand NewCmd(NewName, TEXT(""), FConsoleCommandDelegate());
+		FAutoConsoleCommandDeprecated OldCmd(OldName, NewName, TEXT("0.0"));
+		ON_SCOPE_EXIT { Manager.UnregisterConsoleObject(OldName, false); };
+
+		FOutputDeviceNull OutNull;
+		ManagerImpl.DumpObjects(TEXT("-ShowHelp"), OutNull, true);
+		ManagerImpl.DumpObjects(TEXT("-ShowHelp -Deprecated"), OutNull, true);
+	}
+};
+
+#endif
