Fix float->int conversion warning in FoundationTests on Linux

diff --git a/Engine/Source/Runtime/Core/Tests/HAL/ConsoleManagerTest.cpp b/Engine/Source/Runtime/Core/Tests/HAL/ConsoleManagerTest.cpp
--- a/Engine/Source/Runtime/Core/Tests/HAL/ConsoleManagerTest.cpp
+++ b/Engine/Source/Runtime/Core/Tests/HAL/ConsoleManagerTest.cpp
@@ -107,7 +107,7 @@ TEST_CASE_NAMED(FConsoleManagerTest, "System::Core::HAL::ConsoleManager", "[Appl
 		SECTION("Get variable values")
 		{
 			CHECK(VarA->GetInt() == 1);
-			CHECK(VarA->GetFloat() == 1);
+			CHECK(VarA->GetFloat() == (float)1);
 			CHECK(VarA->GetString() == FString(TEXT("1")));
 
 			CHECK(VarB->GetInt() == 1);
