[AutoRTFM] Annotate AutomationTest functions with UE_AUTORTFM_ALWAYS_OPEN

Without this, if a test check fails when called from the closed, we'll hit AutoRTFM errors as these use a non-transactionally safe lock. This hides the reason for the test failure which isn't very helpful.

diff --git a/Engine/Source/Runtime/Core/Private/Misc/AutomationTest.cpp b/Engine/Source/Runtime/Core/Private/Misc/AutomationTest.cpp
--- a/Engine/Source/Runtime/Core/Private/Misc/AutomationTest.cpp
+++ b/Engine/Source/Runtime/Core/Private/Misc/AutomationTest.cpp
@@ -4,6 +4,7 @@
 
 #include <inttypes.h>
 
+#include "AutoRTFM/AutoRTFM.h"
 #include "Algo/Copy.h"
 #include "HAL/FileManager.h"
 #include "HAL/PlatformStackWalk.h"
@@ -1418,6 +1419,7 @@ void FAutomationTestBase::ClearExecutionInfo()
 	ExecutionInfo.Clear();
 }
 
+UE_AUTORTFM_ALWAYS_OPEN
 void FAutomationTestBase::AddError(const FString& InError, int32 StackOffset)
 {
 	if( !IsExpectedMessage(InError, ELogVerbosity::Warning))
@@ -1427,6 +1429,7 @@ void FAutomationTestBase::AddError(const FString& InError, int32 StackOffset)
 	}
 }
 
+UE_AUTORTFM_ALWAYS_OPEN
 bool FAutomationTestBase::AddErrorIfFalse(bool bCondition, const FString& InError, int32 StackOffset)
 {
 	if (!bCondition)
@@ -1436,6 +1439,7 @@ bool FAutomationTestBase::AddErrorIfFalse(bool bCondition, const FString& InErro
 	return bCondition;
 }
 
+UE_AUTORTFM_ALWAYS_OPEN
 void FAutomationTestBase::AddErrorS(const FString& InError, const FString& InFilename, int32 InLineNumber)
 {
 	if ( !IsExpectedMessage(InError, ELogVerbosity::Warning))
@@ -1445,6 +1449,7 @@ void FAutomationTestBase::AddErrorS(const FString& InError, const FString& InFil
 	}
 }
 
+UE_AUTORTFM_ALWAYS_OPEN
 void FAutomationTestBase::AddWarningS(const FString& InWarning, const FString& InFilename, int32 InLineNumber)
 {
 	if ( !IsExpectedMessage(InWarning, ELogVerbosity::Warning))
@@ -1454,6 +1459,7 @@ void FAutomationTestBase::AddWarningS(const FString& InWarning, const FString& I
 	}
 }
 
+UE_AUTORTFM_ALWAYS_OPEN
 void FAutomationTestBase::AddWarning( const FString& InWarning, int32 StackOffset )
 {
 	if ( !IsExpectedMessage(InWarning, ELogVerbosity::Warning))
@@ -1463,6 +1469,7 @@ void FAutomationTestBase::AddWarning( const FString& InWarning, int32 StackOffse
 	}
 }
 
+UE_AUTORTFM_ALWAYS_OPEN
 void FAutomationTestBase::AddInfo( const FString& InLogItem, int32 StackOffset, bool bCaptureStack )
 {
 	if ( !IsExpectedMessage(InLogItem, ELogVerbosity::Display))
@@ -1472,18 +1479,21 @@ void FAutomationTestBase::AddInfo( const FString& InLogItem, int32 StackOffset,
 	}
 }
 
+UE_AUTORTFM_ALWAYS_OPEN
 void FAutomationTestBase::AddAnalyticsItem(const FString& InAnalyticsItem)
 {
 	FWriteScopeLock Lock(ActionCS);
 	ExecutionInfo.AnalyticsItems.Add(InAnalyticsItem);
 }
 
+UE_AUTORTFM_ALWAYS_OPEN
 void FAutomationTestBase::AddTelemetryData(const FString& DataPoint, double Measurement, const FString& Context)
 {
 	FWriteScopeLock Lock(ActionCS);
 	ExecutionInfo.TelemetryItems.Add(FAutomationTelemetryData(DataPoint, Measurement, Context));
 }
 
+UE_AUTORTFM_ALWAYS_OPEN
 void FAutomationTestBase::AddTelemetryData(const TMap <FString, double>& ValuePairs, const FString& Context)
 {
 	FWriteScopeLock Lock(ActionCS);
@@ -1498,6 +1508,7 @@ void FAutomationTestBase::SetTelemetryStorage(const FString& StorageName)
 	ExecutionInfo.TelemetryStorage = StorageName;
 }
 
+UE_AUTORTFM_ALWAYS_OPEN
 void FAutomationTestBase::AddEvent(const FAutomationEvent& InEvent, int32 StackOffset, bool bCaptureStack)
 {
 	ELogVerbosity::Type LogType = ELogVerbosity::Display;
