Changed I/O store decode chunk block request signature to allow schedule more requests w/o recursion.

diff --git a/Engine/Source/Runtime/Core/Public/IO/IoChunkEncoding.h b/Engine/Source/Runtime/Core/Public/IO/IoChunkEncoding.h
--- a/Engine/Source/Runtime/Core/Public/IO/IoChunkEncoding.h
+++ b/Engine/Source/Runtime/Core/Public/IO/IoChunkEncoding.h
@@ -6,7 +6,7 @@
 #include "IO/IoStatus.h"
 #include "Memory/MemoryView.h"
 #include "Misc/EnumClassFlags.h"
-#include "Templates/FunctionFwd.h"
+#include "Templates/Function.h"
 #include "UObject/NameTypes.h"
 
 class FIoBuffer;
@@ -129,15 +129,28 @@ public:
 	static CORE_API FIoBlockHash HashBlock(FMemoryView Block);
 };
 
-/** Callback used when decoding I/O store chunk blocks asynchronously. */
-using FOnBlockDecoded = TUniqueFunction<void(FIoChunkBlockDecodeResult&&)>;
+/**
+ * Callback used when decoding I/O store chunk blocks asynchronously. The optional out
+ * paramater can be used to schedule more decode requests without recursively calling IIoChunkBlockDecoder::Enqueue.
+ */
+using FIoBlockDecoded = TUniqueFunction<void(FIoChunkBlockDecodeResult&&, struct FIoChunkBlockDecodeRequest&)>;
+
+struct FIoChunkBlockDecodeRequest
+{
+	bool						IsValid() const { return OnDecoded.IsSet(); }
+
+	FIoChunkBlockDecodingParams Params;
+	FMutableMemoryView			EncodedBlock;
+	FMutableMemoryView			DecodedBlock;
+	FIoBlockDecoded 			OnDecoded;
+};
 
 /** Interface for decoding I/O store chunk blocks asynchronously. */
 class IIoChunkBlockDecoder
 {
 public:
 	virtual			~IIoChunkBlockDecoder() = default;
-	virtual void	Enqueue(FIoChunkBlockDecodingParams Params, FMutableMemoryView EncodedBlock, FMutableMemoryView OutDecodedBlock, FOnBlockDecoded&& OnCompleted) = 0;
+	virtual void	Enqueue(FIoChunkBlockDecodeRequest&& DecodeRequest) = 0;
 };
 
 /** Provides access to the chunk block decoder. */
