Fixed incorrect format specifier.

diff --git a/Engine/Source/Runtime/Core/Private/Apple/ApplePlatformStackWalk.cpp b/Engine/Source/Runtime/Core/Private/Apple/ApplePlatformStackWalk.cpp
--- a/Engine/Source/Runtime/Core/Private/Apple/ApplePlatformStackWalk.cpp
+++ b/Engine/Source/Runtime/Core/Private/Apple/ApplePlatformStackWalk.cpp
@@ -403,11 +403,11 @@ void CreateExceptionInfoString(int32 Signal, struct __siginfo* Info)
 	{
 		case SIGSEGV:
 			ErrorString += TEXT("SIGSEGV segmentation violation at address ");
-			ErrorString += FString::Printf(TEXT("0x%08x"), (uint32*)Info->si_addr);
+			ErrorString += FString::Printf(TEXT("0x%016" UINT64_x_FMT), (uint64)(uint32*)Info->si_addr);
 			break;
 		case SIGBUS:
 			ErrorString += TEXT("SIGBUS bus error at address ");
-			ErrorString += FString::Printf(TEXT("0x%08x"), (uint32*)Info->si_addr);
+			ErrorString += FString::Printf(TEXT("0x%016" UINT64_x_FMT), (uint64)(uint32*)Info->si_addr);
 			break;
 			
 		HANDLE_CASE(SIGINT, "interrupt program")
