Added support for UTF-8 to FormatStringSan.

diff --git a/Engine/Source/Runtime/Core/Public/HAL/Platform.h b/Engine/Source/Runtime/Core/Public/HAL/Platform.h
--- a/Engine/Source/Runtime/Core/Public/HAL/Platform.h
+++ b/Engine/Source/Runtime/Core/Public/HAL/Platform.h
@@ -4,6 +4,7 @@
 
 #include "Misc/Build.h"
 #include "Misc/LargeWorldCoordinates.h"
+#include <type_traits>
 
 // define all other platforms to be zero
 //@port Define the platform here to be zero when compiling for other platforms
@@ -1202,15 +1203,76 @@ namespace TypeTests
 	#define TEXT(x) TEXT_PASTE(x)
 #endif
 
+// UE_UTF8TEXT_IS_CONSTEXPR is only enabled because FormatStringSanTests.cpp needs to use UTF8TEXT in static_asserts.
+// It is not enabled in general because of a concern for compile times and PDB sizes.  This can be removed when we can guarantee C++20 and
+// when UTF8CHAR == char8_t and UTF8TEXT simply prepends u8 to the literal.
+#if WITH_LOW_LEVEL_TESTS && defined(__cpp_nontype_template_args)
+	#ifdef __clang__
+		// Clang supports the necessary language features to support this in C++20, but hasn't updated the version
+		// number of the feature test macro due to unrelated deficiencies, so we just assume we can enable it under C++20 in general:
+		//
+		// [REDACTED_PROFILE_URL]/llvm-project/issues/54297#issuecomment-1408671519
+		#define UE_UTF8TEXT_IS_CONSTEXPR (__cplusplus >= 202002L)
+	#else
+		#define UE_UTF8TEXT_IS_CONSTEXPR (__cpp_nontype_template_args >= 201911L)
+	#endif
+#else
+	#define UE_UTF8TEXT_IS_CONSTEXPR 0
+#endif
+
 namespace UE::Core::Private
 {
+#if UE_UTF8TEXT_IS_CONSTEXPR
+	template <SIZE_T N = 1>
+	struct TUtf8StringLiteral
+	{
+#if defined(__cpp_char8_t)
+		using Utf8CharLiteralType = char8_t;
+#else
+		using Utf8CharLiteralType = char;
+#endif
+
+		constexpr TUtf8StringLiteral(const Utf8CharLiteralType(&InArray)[N])
+		{
+			for (SIZE_T Index = 0; Index != N; ++Index)
+			{
+				Array[Index] = (UTF8CHAR)InArray[Index];
+			}
+		}
+
+		constexpr TUtf8StringLiteral(unsigned long long Val)
+		{
+			Array[0] = (UTF8CHAR)Val;
+		}
+
+		UTF8CHAR Array[N];
+	};
+
+	template <typename ArgType, TUtf8StringLiteral Literal, typename ArraySelector = void>
+	struct TUtf8StringLiteralStatic
+	{
+		static constexpr auto StaticLiteral = Literal;
+		static constexpr const auto& Value = Literal.Array;
+	};
+	template <typename ArgType, TUtf8StringLiteral Literal>
+	struct TUtf8StringLiteralStatic<ArgType, Literal, std::enable_if_t<!std::is_array_v<std::remove_reference_t<ArgType>>>>
+	{
+		static constexpr auto StaticLiteral = Literal;
+		static constexpr auto Value = Literal.Array[0];
+	};
+
+	#define UTF8TEXT(x) (UE::Core::Private::TUtf8StringLiteralStatic<decltype(x), UTF8TEXT_PASTE(x)>::Value)
+#else
 	// Can't be constexpr because it involves casts
 	template <SIZE_T N>
 	FORCEINLINE auto ToUTF8Literal(const char(&Array)[N]) -> const UTF8CHAR(&)[N]
 	{
 		return (const UTF8CHAR(&)[N])Array;
 	}
-
+	FORCEINLINE constexpr UTF8CHAR ToUTF8Literal(unsigned long long Ch)
+	{
+		return (UTF8CHAR)Ch;
+	}
 #if defined(__cpp_char8_t)
 	// Can't be constexpr because it involves casts
 	template <SIZE_T N>
@@ -1220,10 +1282,8 @@ namespace UE::Core::Private
 	}
 #endif
 
-	FORCEINLINE constexpr UTF8CHAR ToUTF8Literal(unsigned long long Ch)
-	{
-		return (UTF8CHAR)Ch;
-	}
+	#define UTF8TEXT(x) (UE::Core::Private::ToUTF8Literal(UTF8TEXT_PASTE(x)))
+#endif
 
 	template <typename CharType>
 	void CharTextStaticAssert()
@@ -1232,7 +1292,6 @@ namespace UE::Core::Private
 	}
 }
 
-#define UTF8TEXT(x) (UE::Core::Private::ToUTF8Literal(UTF8TEXT_PASTE(x)))
 #define UTF16TEXT(x) UTF16TEXT_PASTE(x)
 #define WIDETEXT(str) WIDETEXT_PASTE(str)
 
diff --git a/Engine/Source/Runtime/Core/Public/String/FormatStringSan.h b/Engine/Source/Runtime/Core/Public/String/FormatStringSan.h
--- a/Engine/Source/Runtime/Core/Public/String/FormatStringSan.h
+++ b/Engine/Source/Runtime/Core/Public/String/FormatStringSan.h
@@ -6,19 +6,21 @@
 #include "Templates/Requires.h"
 #include "Templates/Identity.h"
 #include "Templates/IsValidVariadicFunctionArg.h"
+#include "Traits/IsCharType.h"
 #include "Traits/IsTEnumAsByte.h"
+#include "Traits/IsTString.h"
 #include "Containers/ContainersFwd.h"
 
 #define UE_CHECK_FORMAT_STRING(Fmt, ...) \
 	do { \
 		namespace UCFS = ::UE::Core::Private::FormatStringSan; \
-		using UCFS_FChecker = decltype(UCFS::GetFmtArgCheckerType(__VA_ARGS__)); \
+		using UCFS_FChecker = decltype(UCFS::GetFmtArgCheckerType(Fmt, ##__VA_ARGS__)); \
 		constexpr UCFS::FResult UCFS_Result = UCFS_FChecker::Check(false, 0, Fmt); \
 		(UCFS::AssertFormatStatus<UCFS_Result.Status, UCFS::TAtArgPos<UCFS_Result.ArgPos>>()); \
 	} while(false);
 
 #define UE_CHECK_FORMAT_STRING_ERR(Err, Fmt, ...) \
-	(decltype(::UE::Core::Private::FormatStringSan::GetFmtArgCheckerType(__VA_ARGS__))::Check(false, 0, Fmt).Status == Err)
+	(decltype(::UE::Core::Private::FormatStringSan::GetFmtArgCheckerType(Fmt, ##__VA_ARGS__))::Check(false, 0, Fmt).Status == Err)
 
 
 #if UE_VALIDATE_FORMAT_STRINGS
@@ -70,24 +72,17 @@ namespace UE::Core::Private
 			int ArgPos{};
 		};
 
-		template <typename T, typename V>
-		inline constexpr bool bIsPointerTo =
-			std::is_pointer_v<T> && std::is_same_v<std::remove_cv_t<std::remove_pointer_t<T>>, V>;
-		template <typename T>
-		inline constexpr bool bIsCharPtr = bIsPointerTo<T, char>;
-		template <typename T>
-		inline constexpr bool bIsTCharPtr = bIsPointerTo<T, TCHAR>;
 		template <typename T>
 		inline constexpr bool TIsFloatOrDouble_V = std::is_same_v<float, T> || std::is_same_v<double, T>;
 
 		template <typename T>
 		inline constexpr bool TIsIntegralEnum_V = std::is_enum_v<T> || TIsTEnumAsByte_V<T>;
 
-		template <typename...>
+		template <typename CharType, typename...>
 		struct TCheckFormatString;
 
-		template <typename... Ts>
-		constexpr TCheckFormatString<std::decay_t<Ts>...> GetFmtArgCheckerType(Ts...);
+		template <typename CharType, typename... Ts>
+		constexpr TCheckFormatString<CharType, std::decay_t<Ts>...> GetFmtArgCheckerType(const CharType*, Ts...);
 
 		template <typename CharType>
 		constexpr bool CharIsDigit(CharType Char)
@@ -117,11 +112,9 @@ namespace UE::Core::Private
 			}
 		}
 
-		template <typename Arg, typename... Args>
-		struct TCheckFormatString<Arg, Args...>
+		template <typename CharType, typename Arg, typename... Args>
+		struct TCheckFormatString<CharType, Arg, Args...>
 		{
-			using CharType = TCHAR;
-
 			static constexpr FResult HandleDynamicLengthSpecifier(int CurArgPos, const CharType* Fmt)
 			{
 				if constexpr (!(std::is_integral_v<Arg> || TIsIntegralEnum_V<Arg>))
@@ -130,12 +123,14 @@ namespace UE::Core::Private
 				}
 				else
 				{
-					return TCheckFormatString<Args...>::Check(true, CurArgPos + 1, Fmt + 1);
+					return TCheckFormatString<CharType, Args...>::Check(true, CurArgPos + 1, Fmt + 1);
 				}
 			}
 
 			static constexpr FResult Check(bool bInsideFormatSpec, int CurArgPos, const CharType* Fmt)
 			{
+				using ArgPointeeType = std::remove_pointer_t<Arg>;
+
 				if (Fmt[0] == (CharType)'\0' && !bInsideFormatSpec)
 				{
 					return {EFormatStringSanStatus::Ok, 0};
@@ -198,104 +193,116 @@ namespace UE::Core::Private
 
 				if (Fmt[0] == (CharType)'l' && Fmt[1] == (CharType)'s')
 				{
-					if constexpr (bIsPointerTo<Arg, UTF8CHAR>)
+					if constexpr (TIsTString_V<Arg>)
 					{
-						return {EFormatStringSanStatus::UTF8NeedsStringCast, CurArgPos};
+						return {EFormatStringSanStatus::LSNeedsDereferencedWideString, CurArgPos};
 					}
-					else if constexpr (std::is_same_v<Arg, FString>)
+					else if constexpr (TIsCharType_V<Arg>)
 					{
-						return {EFormatStringSanStatus::SNeedsDereferencedFString, CurArgPos};
+						return {EFormatStringSanStatus::LSNeedsPtrButGotChar, CurArgPos};
 					}
-					else if constexpr (std::is_same_v<Arg, TCHAR>)
+					else if constexpr (!std::is_pointer_v<Arg> || !TIsCharType_V<ArgPointeeType>)
 					{
-						return {EFormatStringSanStatus::SNeedsTCharPtrArgButGotTCHAR, CurArgPos};
+						return {EFormatStringSanStatus::LSNeedsWideCharPtrArg, CurArgPos};
 					}
-					else if constexpr (!bIsTCharPtr<Arg>)
+					else if constexpr (sizeof(ArgPointeeType) != sizeof(WIDECHAR))
 					{
-						return {EFormatStringSanStatus::SNeedsTCHARPtrArg, CurArgPos};
+						return {(sizeof(CharType) == 1) ? EFormatStringSanStatus::LSNeedsWideCharPtrArgButGotNarrowOnNarrowString : EFormatStringSanStatus::LSNeedsWideCharPtrArgButGotNarrowOnWideString, CurArgPos};
 					}
 					else
 					{
-						return TCheckFormatString<Args...>::Check(false, CurArgPos + 1, Fmt + 2);
+						return TCheckFormatString<CharType, Args...>::Check(false, CurArgPos + 1, Fmt + 2);
 					}
 				}
 				else if (Fmt[0] == (CharType)'h' && Fmt[1] == (CharType)'s')
 				{
-					if constexpr (bIsPointerTo<Arg, UTF8CHAR>)
+					if constexpr (TIsTString_V<Arg>)
 					{
-						return {EFormatStringSanStatus::UTF8NeedsStringCastAndS, CurArgPos};
+						return {EFormatStringSanStatus::HSNeedsDereferencedNarrowString, CurArgPos};
 					}
-					else if constexpr (!bIsCharPtr<Arg>)
+					else if constexpr (TIsCharType_V<Arg>)
 					{
-						return {EFormatStringSanStatus::HSNeedsCharPtrArg, CurArgPos};
+						return {EFormatStringSanStatus::HSNeedsPtrButGotChar, CurArgPos};
+					}
+					else if constexpr (!std::is_pointer_v<Arg> || !TIsCharType_V<ArgPointeeType>)
+					{
+						return {EFormatStringSanStatus::HSNeedsNarrowCharPtrArg, CurArgPos};
+					}
+					else if constexpr (sizeof(ArgPointeeType) != sizeof(ANSICHAR))
+					{
+						return {(sizeof(CharType) == 1) ? EFormatStringSanStatus::HSNeedsNarrowCharPtrArgButGotWideOnNarrowString : EFormatStringSanStatus::HSNeedsNarrowCharPtrArgButGotWideOnWideString, CurArgPos};
 					}
 					else
 					{
-						return TCheckFormatString<Args...>::Check(false, CurArgPos + 1, Fmt + 2);
+						return TCheckFormatString<CharType, Args...>::Check(false, CurArgPos + 1, Fmt + 2);
 					}
 				}
 				else if (Fmt[0] == (CharType)'s')
 				{
-					if constexpr (bIsPointerTo<Arg, UTF8CHAR>)
+					if constexpr (TIsTString_V<Arg>)
 					{
-						return {EFormatStringSanStatus::UTF8NeedsStringCast, CurArgPos};
+						return {(sizeof(CharType) == 1) ? EFormatStringSanStatus::SNeedsDereferencedNarrowString : EFormatStringSanStatus::SNeedsDereferencedWideString, CurArgPos};
 					}
-					else if constexpr (std::is_same_v<Arg, FString>)
+					else if constexpr (TIsCharType_V<Arg>)
 					{
-						return {EFormatStringSanStatus::SNeedsDereferencedFString, CurArgPos};
+						return {EFormatStringSanStatus::SNeedsPtrButGotChar, CurArgPos};
 					}
-					else if constexpr (std::is_same_v<Arg, TCHAR>)
+					else if constexpr (!std::is_pointer_v<Arg> || !TIsCharType_V<ArgPointeeType>)
 					{
-						return {EFormatStringSanStatus::SNeedsTCharPtrArgButGotTCHAR, CurArgPos};
+						return {(sizeof(CharType) == 1) ? EFormatStringSanStatus::SNeedsNarrowCharPtrArg : EFormatStringSanStatus::SNeedsWideCharPtrArg, CurArgPos};
 					}
-					else if constexpr (!bIsTCharPtr<Arg>)
+					else if constexpr (sizeof(ArgPointeeType) != sizeof(CharType))
 					{
-						return {EFormatStringSanStatus::SNeedsTCHARPtrArg, CurArgPos};
+						return {(sizeof(CharType) == 1) ? EFormatStringSanStatus::SNeedsNarrowCharPtrArgButGotWide : EFormatStringSanStatus::SNeedsWideCharPtrArgButGotNarrow, CurArgPos};
 					}
 					else
 					{
-						return TCheckFormatString<Args...>::Check(false, CurArgPos + 1, Fmt + 1);
+						return TCheckFormatString<CharType, Args...>::Check(false, CurArgPos + 1, Fmt + 1);
 					}
 				}
 				else if (Fmt[0] == (CharType)'S')
 				{
-					using OtherCharType = std::conditional_t<std::is_same_v<TCHAR, ANSICHAR>, WIDECHAR, ANSICHAR>;
-					if constexpr (!bIsPointerTo<Arg, OtherCharType>)
+					if constexpr (TIsTString_V<Arg>)
 					{
-						if constexpr(std::is_same_v<OtherCharType, ANSICHAR>)
-						{
-							return {EFormatStringSanStatus::CapitalSNeedsCharPtr, CurArgPos};
-						}
-						else if constexpr(std::is_same_v<OtherCharType, WIDECHAR>)
-						{
-							return {EFormatStringSanStatus::CapitalSNeedsTCHARPtr, CurArgPos};
-						}
+						return {(sizeof(CharType) == 1) ? EFormatStringSanStatus::CapitalSNeedsDereferencedWideString : EFormatStringSanStatus::CapitalSNeedsDereferencedNarrowString, CurArgPos};
+					}
+					else if constexpr (TIsCharType_V<Arg>)
+					{
+						return {EFormatStringSanStatus::CapitalSNeedsPtrButGotChar, CurArgPos};
+					}
+					else if constexpr (!std::is_pointer_v<Arg> || !TIsCharType_V<ArgPointeeType>)
+					{
+						return {(sizeof(CharType) == 1) ? EFormatStringSanStatus::CapitalSNeedsWideCharPtrArg : EFormatStringSanStatus::CapitalSNeedsNarrowCharPtrArg, CurArgPos};
+					}
+					else if constexpr (sizeof(ArgPointeeType) == sizeof(CharType))
+					{
+						return {(sizeof(CharType) == 1) ? EFormatStringSanStatus::CapitalSNeedsWideCharPtrArgButGotNarrow : EFormatStringSanStatus::CapitalSNeedsNarrowCharPtrArgButGotWide, CurArgPos};
 					}
 					else
 					{
-						return TCheckFormatString<Args...>::Check(false, CurArgPos + 1, Fmt + 1);
+						return TCheckFormatString<CharType, Args...>::Check(false, CurArgPos + 1, Fmt + 1);
 					}
 				}
 
 				switch (Fmt[0])
 				{
 				case (CharType)'%':
-					return TCheckFormatString<Arg, Args...>::Check(false, CurArgPos, Fmt + 1);
+					return TCheckFormatString<CharType, Arg, Args...>::Check(false, CurArgPos, Fmt + 1);
 				case (CharType)'c':
-					if constexpr (!(std::is_integral_v<Arg> && sizeof(Arg) <= sizeof(int)))
+					if constexpr (!std::is_same_v<Arg, UTF8CHAR> && !(std::is_integral_v<Arg> && sizeof(Arg) <= sizeof(int)))
 					{
-						return {EFormatStringSanStatus::CNeedsTCHARArg, CurArgPos};
+						return {(sizeof(CharType) == 1) ? EFormatStringSanStatus::CNeedsCharArgOnNarrowString : EFormatStringSanStatus::CNeedsCharArgOnWideString, CurArgPos};
 					}
 					else
 					{
-						return TCheckFormatString<Args...>::Check(false, CurArgPos + 1, Fmt + 1);
+						return TCheckFormatString<CharType, Args...>::Check(false, CurArgPos + 1, Fmt + 1);
 					}
 				case (CharType)'d':
 				case (CharType)'i':
 				case (CharType)'X':
 				case (CharType)'x':
 				case (CharType)'u':
-					if constexpr (bIsTCharPtr<Arg>)
+					if constexpr (std::is_pointer_v<Arg> && TIsCharType_V<ArgPointeeType>)
 					{
 						return {EFormatStringSanStatus::DNeedsIntegerArg, CurArgPos};
 					}
@@ -305,7 +312,7 @@ namespace UE::Core::Private
 					}
 					else
 					{
-						return TCheckFormatString<Args...>::Check(false, CurArgPos + 1, Fmt + 1);
+						return TCheckFormatString<CharType, Args...>::Check(false, CurArgPos + 1, Fmt + 1);
 					}
 				case (CharType)'z':
 					if (!CharIsIntegerFormatSpecifier(Fmt[1]))
@@ -318,7 +325,7 @@ namespace UE::Core::Private
 					}
 					else
 					{
-						return TCheckFormatString<Args...>::Check(false, CurArgPos + 1, Fmt + 1);
+						return TCheckFormatString<CharType, Args...>::Check(false, CurArgPos + 1, Fmt + 1);
 					}
 				case (CharType)'p':
 					if constexpr (!std::is_pointer_v<Arg>)
@@ -327,7 +334,7 @@ namespace UE::Core::Private
 					}
 					else
 					{
-						return TCheckFormatString<Args...>::Check(false, CurArgPos + 1, Fmt + 1);
+						return TCheckFormatString<CharType, Args...>::Check(false, CurArgPos + 1, Fmt + 1);
 					}
 				case (CharType)'I':
 					if (!(Fmt[1] == (CharType)'6' && Fmt[2] == (CharType)'4'))
@@ -344,7 +351,7 @@ namespace UE::Core::Private
 					}
 					else
 					{
-						return TCheckFormatString<Args...>::Check(false, CurArgPos + 1, Fmt + 4);
+						return TCheckFormatString<CharType, Args...>::Check(false, CurArgPos + 1, Fmt + 4);
 					}
 
 				case (CharType)'l':
@@ -356,7 +363,7 @@ namespace UE::Core::Private
 						}
 						else
 						{
-							return TCheckFormatString<Args...>::Check(false, CurArgPos + 1, Fmt + 2);
+							return TCheckFormatString<CharType, Args...>::Check(false, CurArgPos + 1, Fmt + 2);
 						}
 					}
 					else if (Fmt[1] == (CharType)'f')
@@ -367,7 +374,7 @@ namespace UE::Core::Private
 						}
 						else
 						{
-							return TCheckFormatString<Args...>::Check(false, CurArgPos + 1, Fmt + 2);
+							return TCheckFormatString<CharType, Args...>::Check(false, CurArgPos + 1, Fmt + 2);
 						}
 					}
 					else if (Fmt[1] != (CharType)'l')
@@ -384,7 +391,7 @@ namespace UE::Core::Private
 					}
 					else
 					{
-						return TCheckFormatString<Args...>::Check(false, CurArgPos + 1, Fmt + 3);
+						return TCheckFormatString<CharType, Args...>::Check(false, CurArgPos + 1, Fmt + 3);
 					}
 
 				case (CharType)'h':
@@ -394,7 +401,7 @@ namespace UE::Core::Private
 						{
 							return {EFormatStringSanStatus::HNeedsIntegerArg, CurArgPos};
 						}
-						return TCheckFormatString<Args...>::Check(false, CurArgPos + 1, Fmt + 2);
+						return TCheckFormatString<CharType, Args...>::Check(false, CurArgPos + 1, Fmt + 2);
 					}
 					else if (Fmt[1] != (CharType)'h')
 					{
@@ -410,7 +417,7 @@ namespace UE::Core::Private
 					}
 					else
 					{
-						return TCheckFormatString<Args...>::Check(false, CurArgPos + 1, Fmt + 3);
+						return TCheckFormatString<CharType, Args...>::Check(false, CurArgPos + 1, Fmt + 3);
 					}
 
 				case (CharType)'f':
@@ -422,7 +429,7 @@ namespace UE::Core::Private
 					}
 					else
 					{
-						return TCheckFormatString<Args...>::Check(false, CurArgPos + 1, Fmt + 1);
+						return TCheckFormatString<CharType, Args...>::Check(false, CurArgPos + 1, Fmt + 1);
 					}
 
 				case (CharType)' ':
@@ -434,11 +441,9 @@ namespace UE::Core::Private
 			}
 		};
 
-		template <>
-		struct TCheckFormatString<>
+		template <typename CharType>
+		struct TCheckFormatString<CharType>
 		{
-			using CharType = TCHAR;
-
 			static constexpr FResult Check(bool bInsideFormatSpec, int CurArgPos, const CharType* Fmt)
 			{
 				if (bInsideFormatSpec)
@@ -482,9 +487,7 @@ namespace UE::Core::Private
 			consteval TCheckedFormatStringPrivate(const CharType (&Fmt)[N])
 				: FormatString{Fmt}
 			{
-				if constexpr (std::is_convertible_v<const CharType*, const TCHAR*>)
-				{
-					using Checker = TCheckFormatString<std::decay_t<ArgTypes>...>;
+				using Checker = TCheckFormatString<CharType, std::decay_t<ArgTypes>...>;
 					FResult Result = Checker::Check(false, 0, Fmt);
 
 					switch (Result.Status)
@@ -497,11 +500,6 @@ namespace UE::Core::Private
 						#undef X
 					}
 				}
-				else
-				{
-					// not supported yet
-				}
-			}
 
 			void PRINTF_FORMAT_STRING_ERROR(const char*)
 			{
diff --git a/Engine/Source/Runtime/Core/Public/String/FormatStringSanErrors.inl b/Engine/Source/Runtime/Core/Public/String/FormatStringSanErrors.inl
--- a/Engine/Source/Runtime/Core/Public/String/FormatStringSanErrors.inl
+++ b/Engine/Source/Runtime/Core/Public/String/FormatStringSanErrors.inl
@@ -6,16 +6,35 @@
   details of the specific error.
  */
 
-X(SNeedsTCHARPtrArg, "'%s' expects `TCHAR*`; use '%hs' for `const char*`.")
-X(SNeedsDereferencedFString, "FString should be dereferenced before passing to '%s'.")
-X(SNeedsTCharPtrArgButGotTCHAR, "'%s' expects `TCHAR*` but got `TCHAR`. (Is there an extra '*' somewhere?)")
-X(HSNeedsCharPtrArg, "'%hs' expects `char*`; use '%s' for `const TCHAR*`, `TEXT(...)`, etc.")
+X(SNeedsDereferencedNarrowString, "'%s' expects dereferenced `FAnsiString` or `FUtf8String`.")
+X(SNeedsDereferencedWideString, "'%s' expects dereferenced `FString`.")
+X(SNeedsPtrButGotChar, "'%s' expects character pointer but got character. (Is there an extra '*' somewhere?)")
+X(SNeedsNarrowCharPtrArg, "'%s' expects `ANSICHAR*` or `UTF8CHAR*`.")
+X(SNeedsWideCharPtrArg, "'%s' expects `TCHAR*`.")
+X(SNeedsNarrowCharPtrArgButGotWide, "'%s' expects `ANSICHAR*` or `UTF8CHAR*` but got `TCHAR*`; use '%ls'.")
+X(SNeedsWideCharPtrArgButGotNarrow, "'%s' expects `TCHAR*` but got `ANSICHAR*` or `UTF8CHAR*`; use '%hs'.")
+X(CapitalSNeedsDereferencedNarrowString, "'%S' expects dereferenced `FAnsiString` or `FUtf8String`.")
+X(CapitalSNeedsDereferencedWideString, "'%S' expects dereferenced `FString`.")
+X(CapitalSNeedsPtrButGotChar, "'%S' expects character pointer but got character. (Is there an extra '*' somewhere?)")
+X(CapitalSNeedsNarrowCharPtrArg, "'%S' expects `ANSICHAR*` or `UTF8CHAR*`.")
+X(CapitalSNeedsWideCharPtrArg, "'%S' expects `TCHAR*`.")
+X(CapitalSNeedsNarrowCharPtrArgButGotWide, "'%S' expects `ANSICHAR*` or `UTF8CHAR*` but got `TCHAR*`; use '%s'.")
+X(CapitalSNeedsWideCharPtrArgButGotNarrow, "'%S' expects `TCHAR*` but got `ANSICHAR*` or `UTF8CHAR*`; use '%s'.")
+X(HSNeedsDereferencedNarrowString, "'%hs' expects dereferenced `FAnsiString` or `FUtf8String`.")
+X(HSNeedsPtrButGotChar, "'%hs' expects character pointer but got character. (Is there an extra '*' somewhere?)")
+X(HSNeedsNarrowCharPtrArg, "'%hs' expects `ANSICHAR*` or `UTF8CHAR*`.")
+X(HSNeedsNarrowCharPtrArgButGotWideOnNarrowString, "'%hs' expects `ANSICHAR*` or `UTF8CHAR*` but got `TCHAR*`; use '%ls'.")
+X(HSNeedsNarrowCharPtrArgButGotWideOnWideString, "'%hs' expects `ANSICHAR*` or `UTF8CHAR*` but got `TCHAR*`; use '%s'.")
+X(LSNeedsDereferencedWideString, "'%ls' expects dereferenced `FString`.")
+X(LSNeedsPtrButGotChar, "'%ls' expects character pointer but got character. (Is there an extra '*' somewhere?)")
+X(LSNeedsWideCharPtrArg, "'%ls' expects `TCHAR*`.")
+X(LSNeedsWideCharPtrArgButGotNarrowOnNarrowString, "'%ls' expects `TCHAR*` but got `ANSICHAR*` or `UTF8CHAR*`; use '%s'.")
+X(LSNeedsWideCharPtrArgButGotNarrowOnWideString, "'%ls' expects `TCHAR*` but got `ANSICHAR*` or `UTF8CHAR*`; use '%hs'.")
 X(IncompleteFormatSpecifierOrUnescapedPercent, "incomplete '%' escape sequence. (to print a percent sign, write '%%'.)")
-X(CNeedsTCHARArg, "'%c' expects `TCHAR`. (construct a `TCHAR` from a `char` like this: 'TCHAR(my_char)'.)")
+X(CNeedsCharArgOnNarrowString, "'%c' expects `ANSICHAR` or `UTF8CHAR`.")
+X(CNeedsCharArgOnWideString, "'%c' expects `TCHAR`.")
 X(PNeedsPointerArg, "'%p' expects a pointer.")
 X(FNeedsFloatOrDoubleArg, "'%f' expects `float` or `double`.")
-X(CapitalSNeedsCharPtr, "'%S' expects 'char*' (or a `UTF8CHAR*` if enabled).")
-X(CapitalSNeedsTCHARPtr, "'%S' expects 'TCHAR*'.")
 X(InvalidFormatSpec, "unsupported '%' format specifier. (to print a percent sign, write '%%'.)")
 X(NotEnoughArguments, "not enough arguments provided to format string.")
 X(DNeedsIntegerArg, "'%d' expects integral arg (eg. `char`, `int`, `long`, etc.)")
@@ -30,5 +49,3 @@ X(LLNeedsIntegerArg, "'%ll[ ]' expects integral arg (eg. `char`, `int`, `long`,
 X(I64BadSpec, "'%I' must appear as '%I64' with an integral suffix (eg. '%I64d', '%I64u', etc.)")
 X(I64NeedsIntegerArg, "'%I64[ ]' expects integral arg (eg. `char`, `int`, `long`, etc.)")
 X(DynamicLengthSpecNeedsIntegerArg, "dynamic field width specifier '*' expects integral arg (eg. `char`, `int`, `long`, etc.)")
-X(UTF8NeedsStringCastAndS, "Pass UTF8 strings to `StringCast<TCHAR>(MyUTF8String).Get()` and use %s.")
-X(UTF8NeedsStringCast, "'%s' expects `TCHAR*`; pass UTF8 strings to `StringCast<TCHAR>(MyUTF8String).Get()` first.")
diff --git a/Engine/Source/Runtime/Core/Tests/String/FormatStringSanTest.cpp b/Engine/Source/Runtime/Core/Tests/String/FormatStringSanTest.cpp
--- a/Engine/Source/Runtime/Core/Tests/String/FormatStringSanTest.cpp
+++ b/Engine/Source/Runtime/Core/Tests/String/FormatStringSanTest.cpp
@@ -1,8 +1,11 @@
 // Copyright Epic Games, Inc. All Rights Reserved.
 
-#if WITH_TESTS
+// Tests need UE_UTF8TEXT_IS_CONSTEXPR because of the use of STATIC_CHECK with UTFTEXT.
+#if WITH_TESTS && UE_UTF8TEXT_IS_CONSTEXPR
 
 #include "String/FormatStringSan.h"
+#include "Containers/AnsiString.h"
+#include "Containers/Utf8String.h"
 
 #include "Containers/EnumAsByte.h"
 #include "Tests/TestHarnessAdapter.h"
@@ -14,16 +17,245 @@ TEST_CASE_NAMED(FFormatStringValidatorTest, "System::Core::String::FormatStringS
 	using namespace UE::Core::Private;
 	using E = FormatStringSan::EFormatStringSanStatus;
 
+	SECTION("%s")
+	{
+		// Test passing string pointers of the wrong element size
+		STATIC_CHECK(UE_CHECK_FORMAT_STRING_ERR(E::SNeedsNarrowCharPtrArgButGotWide,          "Test %s",  WIDETEXT("wrong")));
+		STATIC_CHECK(UE_CHECK_FORMAT_STRING_ERR(E::SNeedsNarrowCharPtrArgButGotWide, UTF8TEXT("Test %s"), WIDETEXT("wrong")));
+		STATIC_CHECK(UE_CHECK_FORMAT_STRING_ERR(E::SNeedsWideCharPtrArgButGotNarrow, WIDETEXT("Test %s"),          "wrong"));
+		STATIC_CHECK(UE_CHECK_FORMAT_STRING_ERR(E::SNeedsWideCharPtrArgButGotNarrow, WIDETEXT("Test %s"), UTF8TEXT("wrong")));
+
+		// Test passing non-string pointers
+		STATIC_CHECK(UE_CHECK_FORMAT_STRING_ERR(E::SNeedsNarrowCharPtrArg,          "Test %s",  (int*)nullptr));
+		STATIC_CHECK(UE_CHECK_FORMAT_STRING_ERR(E::SNeedsNarrowCharPtrArg, UTF8TEXT("Test %s"), (int*)nullptr));
+		STATIC_CHECK(UE_CHECK_FORMAT_STRING_ERR(E::SNeedsWideCharPtrArg,   WIDETEXT("Test %s"), (int*)nullptr));
+
+		// Test passing non-pointers
+		STATIC_CHECK(UE_CHECK_FORMAT_STRING_ERR(E::SNeedsNarrowCharPtrArg,          "Test %s",  5));
+		STATIC_CHECK(UE_CHECK_FORMAT_STRING_ERR(E::SNeedsNarrowCharPtrArg, UTF8TEXT("Test %s"), 5));
+		STATIC_CHECK(UE_CHECK_FORMAT_STRING_ERR(E::SNeedsWideCharPtrArg,   WIDETEXT("Test %s"), 5));
+
+		// Test TString passed instead of pointer
+		STATIC_CHECK(UE_CHECK_FORMAT_STRING_ERR(E::SNeedsDereferencedNarrowString,          "Test %s",  FAnsiString()));
+		STATIC_CHECK(UE_CHECK_FORMAT_STRING_ERR(E::SNeedsDereferencedNarrowString,          "Test %s",  FUtf8String()));
+		STATIC_CHECK(UE_CHECK_FORMAT_STRING_ERR(E::SNeedsDereferencedNarrowString,          "Test %s",  FWideString()));
+		STATIC_CHECK(UE_CHECK_FORMAT_STRING_ERR(E::SNeedsDereferencedNarrowString, UTF8TEXT("Test %s"), FAnsiString()));
+		STATIC_CHECK(UE_CHECK_FORMAT_STRING_ERR(E::SNeedsDereferencedNarrowString, UTF8TEXT("Test %s"), FUtf8String()));
+		STATIC_CHECK(UE_CHECK_FORMAT_STRING_ERR(E::SNeedsDereferencedNarrowString, UTF8TEXT("Test %s"), FWideString()));
+		STATIC_CHECK(UE_CHECK_FORMAT_STRING_ERR(E::SNeedsDereferencedWideString,   WIDETEXT("Test %s"), FAnsiString()));
+		STATIC_CHECK(UE_CHECK_FORMAT_STRING_ERR(E::SNeedsDereferencedWideString,   WIDETEXT("Test %s"), FUtf8String()));
+		STATIC_CHECK(UE_CHECK_FORMAT_STRING_ERR(E::SNeedsDereferencedWideString,   WIDETEXT("Test %s"), FWideString()));
+
+		// Test char passed instead of char pointer
+		STATIC_CHECK(UE_CHECK_FORMAT_STRING_ERR(E::SNeedsPtrButGotChar,          "Test %s",           'x'));
+		STATIC_CHECK(UE_CHECK_FORMAT_STRING_ERR(E::SNeedsPtrButGotChar,          "Test %s",  UTF8TEXT('x')));
+		STATIC_CHECK(UE_CHECK_FORMAT_STRING_ERR(E::SNeedsPtrButGotChar,          "Test %s",  WIDETEXT('x')));
+		STATIC_CHECK(UE_CHECK_FORMAT_STRING_ERR(E::SNeedsPtrButGotChar, UTF8TEXT("Test %s"),          'x'));
+		STATIC_CHECK(UE_CHECK_FORMAT_STRING_ERR(E::SNeedsPtrButGotChar, UTF8TEXT("Test %s"), UTF8TEXT('x')));
+		STATIC_CHECK(UE_CHECK_FORMAT_STRING_ERR(E::SNeedsPtrButGotChar, UTF8TEXT("Test %s"), WIDETEXT('x')));
+		STATIC_CHECK(UE_CHECK_FORMAT_STRING_ERR(E::SNeedsPtrButGotChar, WIDETEXT("Test %s"),          'x'));
+		STATIC_CHECK(UE_CHECK_FORMAT_STRING_ERR(E::SNeedsPtrButGotChar, WIDETEXT("Test %s"), UTF8TEXT('x')));
+		STATIC_CHECK(UE_CHECK_FORMAT_STRING_ERR(E::SNeedsPtrButGotChar, WIDETEXT("Test %s"), WIDETEXT('x')));
+
+		// Test matching arguments
+		STATIC_CHECK(UE_CHECK_FORMAT_STRING_ERR(E::Ok,          "Test %s",           "hello"));
+		STATIC_CHECK(UE_CHECK_FORMAT_STRING_ERR(E::Ok,          "Test %s",  UTF8TEXT("hello")));
+		STATIC_CHECK(UE_CHECK_FORMAT_STRING_ERR(E::Ok, UTF8TEXT("Test %s"),          "hello"));
+		STATIC_CHECK(UE_CHECK_FORMAT_STRING_ERR(E::Ok, UTF8TEXT("Test %s"), UTF8TEXT("hello")));
+		STATIC_CHECK(UE_CHECK_FORMAT_STRING_ERR(E::Ok, WIDETEXT("Test %s"), WIDETEXT("hello")));
+	}
+
+	SECTION("%S")
+	{
+		// Test passing string pointers of the wrong element size
+		STATIC_CHECK(UE_CHECK_FORMAT_STRING_ERR(E::CapitalSNeedsWideCharPtrArgButGotNarrow,          "Test %S",           "wrong"));
+		STATIC_CHECK(UE_CHECK_FORMAT_STRING_ERR(E::CapitalSNeedsWideCharPtrArgButGotNarrow, UTF8TEXT("Test %S"),          "wrong"));
+		STATIC_CHECK(UE_CHECK_FORMAT_STRING_ERR(E::CapitalSNeedsWideCharPtrArgButGotNarrow,          "Test %S",  UTF8TEXT("wrong")));
+		STATIC_CHECK(UE_CHECK_FORMAT_STRING_ERR(E::CapitalSNeedsWideCharPtrArgButGotNarrow, UTF8TEXT("Test %S"), UTF8TEXT("wrong")));
+		STATIC_CHECK(UE_CHECK_FORMAT_STRING_ERR(E::CapitalSNeedsNarrowCharPtrArgButGotWide, WIDETEXT("Test %S"), WIDETEXT("wrong")));
+
+		// Test passing non-string pointers
+		STATIC_CHECK(UE_CHECK_FORMAT_STRING_ERR(E::CapitalSNeedsWideCharPtrArg,            "Test %S",  (int*)nullptr));
+		STATIC_CHECK(UE_CHECK_FORMAT_STRING_ERR(E::CapitalSNeedsWideCharPtrArg,   UTF8TEXT("Test %S"), (int*)nullptr));
+		STATIC_CHECK(UE_CHECK_FORMAT_STRING_ERR(E::CapitalSNeedsNarrowCharPtrArg, WIDETEXT("Test %S"), (int*)nullptr));
+
+		// Test passing non-pointers
+		STATIC_CHECK(UE_CHECK_FORMAT_STRING_ERR(E::CapitalSNeedsWideCharPtrArg,            "Test %S",  5));
+		STATIC_CHECK(UE_CHECK_FORMAT_STRING_ERR(E::CapitalSNeedsWideCharPtrArg,   UTF8TEXT("Test %S"), 5));
+		STATIC_CHECK(UE_CHECK_FORMAT_STRING_ERR(E::CapitalSNeedsNarrowCharPtrArg, WIDETEXT("Test %S"), 5));
+
+		// Test TString passed instead of pointer
+		STATIC_CHECK(UE_CHECK_FORMAT_STRING_ERR(E::CapitalSNeedsDereferencedWideString,            "Test %S",  FAnsiString()));
+		STATIC_CHECK(UE_CHECK_FORMAT_STRING_ERR(E::CapitalSNeedsDereferencedWideString,            "Test %S",  FUtf8String()));
+		STATIC_CHECK(UE_CHECK_FORMAT_STRING_ERR(E::CapitalSNeedsDereferencedWideString,            "Test %S",  FWideString()));
+		STATIC_CHECK(UE_CHECK_FORMAT_STRING_ERR(E::CapitalSNeedsDereferencedWideString,   UTF8TEXT("Test %S"), FAnsiString()));
+		STATIC_CHECK(UE_CHECK_FORMAT_STRING_ERR(E::CapitalSNeedsDereferencedWideString,   UTF8TEXT("Test %S"), FUtf8String()));
+		STATIC_CHECK(UE_CHECK_FORMAT_STRING_ERR(E::CapitalSNeedsDereferencedWideString,   UTF8TEXT("Test %S"), FWideString()));
+		STATIC_CHECK(UE_CHECK_FORMAT_STRING_ERR(E::CapitalSNeedsDereferencedNarrowString, WIDETEXT("Test %S"), FAnsiString()));
+		STATIC_CHECK(UE_CHECK_FORMAT_STRING_ERR(E::CapitalSNeedsDereferencedNarrowString, WIDETEXT("Test %S"), FUtf8String()));
+		STATIC_CHECK(UE_CHECK_FORMAT_STRING_ERR(E::CapitalSNeedsDereferencedNarrowString, WIDETEXT("Test %S"), FWideString()));
+
+		// Test char passed instead of char pointer
+		STATIC_CHECK(UE_CHECK_FORMAT_STRING_ERR(E::CapitalSNeedsPtrButGotChar,          "Test %S",           'x'));
+		STATIC_CHECK(UE_CHECK_FORMAT_STRING_ERR(E::CapitalSNeedsPtrButGotChar,          "Test %S",  UTF8TEXT('x')));
+		STATIC_CHECK(UE_CHECK_FORMAT_STRING_ERR(E::CapitalSNeedsPtrButGotChar,          "Test %S",  WIDETEXT('x')));
+		STATIC_CHECK(UE_CHECK_FORMAT_STRING_ERR(E::CapitalSNeedsPtrButGotChar, UTF8TEXT("Test %S"),          'x'));
+		STATIC_CHECK(UE_CHECK_FORMAT_STRING_ERR(E::CapitalSNeedsPtrButGotChar, UTF8TEXT("Test %S"), UTF8TEXT('x')));
+		STATIC_CHECK(UE_CHECK_FORMAT_STRING_ERR(E::CapitalSNeedsPtrButGotChar, UTF8TEXT("Test %S"), WIDETEXT('x')));
+		STATIC_CHECK(UE_CHECK_FORMAT_STRING_ERR(E::CapitalSNeedsPtrButGotChar, WIDETEXT("Test %S"),          'x'));
+		STATIC_CHECK(UE_CHECK_FORMAT_STRING_ERR(E::CapitalSNeedsPtrButGotChar, WIDETEXT("Test %S"), UTF8TEXT('x')));
+		STATIC_CHECK(UE_CHECK_FORMAT_STRING_ERR(E::CapitalSNeedsPtrButGotChar, WIDETEXT("Test %S"), WIDETEXT('x')));
+
+		// Test matching arguments
+		STATIC_CHECK(UE_CHECK_FORMAT_STRING_ERR(E::Ok,          "Test %S",  WIDETEXT("hello")));
+		STATIC_CHECK(UE_CHECK_FORMAT_STRING_ERR(E::Ok, UTF8TEXT("Test %S"), WIDETEXT("hello")));
+		STATIC_CHECK(UE_CHECK_FORMAT_STRING_ERR(E::Ok, WIDETEXT("Test %S"),          "hello"));
+		STATIC_CHECK(UE_CHECK_FORMAT_STRING_ERR(E::Ok, WIDETEXT("Test %S"), UTF8TEXT("hello")));
+	}
+
+	SECTION("%hs")
+	{
+		// Test passing string pointers of the wrong element size
+		STATIC_CHECK(UE_CHECK_FORMAT_STRING_ERR(E::HSNeedsNarrowCharPtrArgButGotWideOnNarrowString,          "Test %hs",  WIDETEXT("wrong")));
+		STATIC_CHECK(UE_CHECK_FORMAT_STRING_ERR(E::HSNeedsNarrowCharPtrArgButGotWideOnNarrowString, UTF8TEXT("Test %hs"), WIDETEXT("wrong")));
+		STATIC_CHECK(UE_CHECK_FORMAT_STRING_ERR(E::HSNeedsNarrowCharPtrArgButGotWideOnWideString,   WIDETEXT("Test %hs"), WIDETEXT("wrong")));
+
+		// Test passing non-string pointers
+		STATIC_CHECK(UE_CHECK_FORMAT_STRING_ERR(E::HSNeedsNarrowCharPtrArg,          "Test %hs",  (int*)nullptr));
+		STATIC_CHECK(UE_CHECK_FORMAT_STRING_ERR(E::HSNeedsNarrowCharPtrArg, UTF8TEXT("Test %hs"), (int*)nullptr));
+		STATIC_CHECK(UE_CHECK_FORMAT_STRING_ERR(E::HSNeedsNarrowCharPtrArg, WIDETEXT("Test %hs"), (int*)nullptr));
+
+		// Test passing non-pointers
+		STATIC_CHECK(UE_CHECK_FORMAT_STRING_ERR(E::HSNeedsNarrowCharPtrArg,          "Test %hs",  5));
+		STATIC_CHECK(UE_CHECK_FORMAT_STRING_ERR(E::HSNeedsNarrowCharPtrArg, UTF8TEXT("Test %hs"), 5));
+		STATIC_CHECK(UE_CHECK_FORMAT_STRING_ERR(E::HSNeedsNarrowCharPtrArg, WIDETEXT("Test %hs"), 5));
+
+		// Test TString passed instead of pointer
+		STATIC_CHECK(UE_CHECK_FORMAT_STRING_ERR(E::HSNeedsDereferencedNarrowString,          "Test %hs",  FAnsiString()));
+		STATIC_CHECK(UE_CHECK_FORMAT_STRING_ERR(E::HSNeedsDereferencedNarrowString,          "Test %hs",  FUtf8String()));
+		STATIC_CHECK(UE_CHECK_FORMAT_STRING_ERR(E::HSNeedsDereferencedNarrowString,          "Test %hs",  FWideString()));
+		STATIC_CHECK(UE_CHECK_FORMAT_STRING_ERR(E::HSNeedsDereferencedNarrowString, UTF8TEXT("Test %hs"), FAnsiString()));
+		STATIC_CHECK(UE_CHECK_FORMAT_STRING_ERR(E::HSNeedsDereferencedNarrowString, UTF8TEXT("Test %hs"), FUtf8String()));
+		STATIC_CHECK(UE_CHECK_FORMAT_STRING_ERR(E::HSNeedsDereferencedNarrowString, UTF8TEXT("Test %hs"), FWideString()));
+		STATIC_CHECK(UE_CHECK_FORMAT_STRING_ERR(E::HSNeedsDereferencedNarrowString, WIDETEXT("Test %hs"), FAnsiString()));
+		STATIC_CHECK(UE_CHECK_FORMAT_STRING_ERR(E::HSNeedsDereferencedNarrowString, WIDETEXT("Test %hs"), FUtf8String()));
+		STATIC_CHECK(UE_CHECK_FORMAT_STRING_ERR(E::HSNeedsDereferencedNarrowString, WIDETEXT("Test %hs"), FWideString()));
+
+		// Test char passed instead of char pointer
+		STATIC_CHECK(UE_CHECK_FORMAT_STRING_ERR(E::HSNeedsPtrButGotChar,          "Test %hs",           'x'));
+		STATIC_CHECK(UE_CHECK_FORMAT_STRING_ERR(E::HSNeedsPtrButGotChar,          "Test %hs",  UTF8TEXT('x')));
+		STATIC_CHECK(UE_CHECK_FORMAT_STRING_ERR(E::HSNeedsPtrButGotChar,          "Test %hs",  WIDETEXT('x')));
+		STATIC_CHECK(UE_CHECK_FORMAT_STRING_ERR(E::HSNeedsPtrButGotChar, UTF8TEXT("Test %hs"),          'x'));
+		STATIC_CHECK(UE_CHECK_FORMAT_STRING_ERR(E::HSNeedsPtrButGotChar, UTF8TEXT("Test %hs"), UTF8TEXT('x')));
+		STATIC_CHECK(UE_CHECK_FORMAT_STRING_ERR(E::HSNeedsPtrButGotChar, UTF8TEXT("Test %hs"), WIDETEXT('x')));
+		STATIC_CHECK(UE_CHECK_FORMAT_STRING_ERR(E::HSNeedsPtrButGotChar, WIDETEXT("Test %hs"),          'x'));
+		STATIC_CHECK(UE_CHECK_FORMAT_STRING_ERR(E::HSNeedsPtrButGotChar, WIDETEXT("Test %hs"), UTF8TEXT('x')));
+		STATIC_CHECK(UE_CHECK_FORMAT_STRING_ERR(E::HSNeedsPtrButGotChar, WIDETEXT("Test %hs"), WIDETEXT('x')));
+
+		// Test matching arguments
+		STATIC_CHECK(UE_CHECK_FORMAT_STRING_ERR(E::Ok,          "Test %hs",           "hello"));
+		STATIC_CHECK(UE_CHECK_FORMAT_STRING_ERR(E::Ok,          "Test %hs",  UTF8TEXT("hello")));
+		STATIC_CHECK(UE_CHECK_FORMAT_STRING_ERR(E::Ok, UTF8TEXT("Test %hs"),          "hello"));
+		STATIC_CHECK(UE_CHECK_FORMAT_STRING_ERR(E::Ok, UTF8TEXT("Test %hs"), UTF8TEXT("hello")));
+		STATIC_CHECK(UE_CHECK_FORMAT_STRING_ERR(E::Ok, WIDETEXT("Test %hs"),          "hello"));
+		STATIC_CHECK(UE_CHECK_FORMAT_STRING_ERR(E::Ok, WIDETEXT("Test %hs"), UTF8TEXT("hello")));
+	}
+
+	SECTION("%ls")
+	{
+		// Test passing string pointers of the wrong element size
+		STATIC_CHECK(UE_CHECK_FORMAT_STRING_ERR(E::LSNeedsWideCharPtrArgButGotNarrowOnNarrowString,          "Test %ls",           "wrong"));
+		STATIC_CHECK(UE_CHECK_FORMAT_STRING_ERR(E::LSNeedsWideCharPtrArgButGotNarrowOnNarrowString,          "Test %ls",  UTF8TEXT("wrong")));
+		STATIC_CHECK(UE_CHECK_FORMAT_STRING_ERR(E::LSNeedsWideCharPtrArgButGotNarrowOnNarrowString, UTF8TEXT("Test %ls"),          "wrong"));
+		STATIC_CHECK(UE_CHECK_FORMAT_STRING_ERR(E::LSNeedsWideCharPtrArgButGotNarrowOnNarrowString, UTF8TEXT("Test %ls"), UTF8TEXT("wrong")));
+		STATIC_CHECK(UE_CHECK_FORMAT_STRING_ERR(E::LSNeedsWideCharPtrArgButGotNarrowOnWideString,   WIDETEXT("Test %ls"),          "wrong"));
+		STATIC_CHECK(UE_CHECK_FORMAT_STRING_ERR(E::LSNeedsWideCharPtrArgButGotNarrowOnWideString,   WIDETEXT("Test %ls"), UTF8TEXT("wrong")));
+
+		// Test passing non-string pointers
+		STATIC_CHECK(UE_CHECK_FORMAT_STRING_ERR(E::LSNeedsWideCharPtrArg,          "Test %ls",  (int*)nullptr));
+		STATIC_CHECK(UE_CHECK_FORMAT_STRING_ERR(E::LSNeedsWideCharPtrArg, UTF8TEXT("Test %ls"), (int*)nullptr));
+		STATIC_CHECK(UE_CHECK_FORMAT_STRING_ERR(E::LSNeedsWideCharPtrArg, WIDETEXT("Test %ls"), (int*)nullptr));
+
+		// Test passing non-pointers
+		STATIC_CHECK(UE_CHECK_FORMAT_STRING_ERR(E::LSNeedsWideCharPtrArg,          "Test %ls",  5));
+		STATIC_CHECK(UE_CHECK_FORMAT_STRING_ERR(E::LSNeedsWideCharPtrArg, UTF8TEXT("Test %ls"), 5));
+		STATIC_CHECK(UE_CHECK_FORMAT_STRING_ERR(E::LSNeedsWideCharPtrArg, WIDETEXT("Test %ls"), 5));
+
+		// Test TString passed instead of pointer
+		STATIC_CHECK(UE_CHECK_FORMAT_STRING_ERR(E::LSNeedsDereferencedWideString,          "Test %ls",  FAnsiString()));
+		STATIC_CHECK(UE_CHECK_FORMAT_STRING_ERR(E::LSNeedsDereferencedWideString,          "Test %ls",  FUtf8String()));
+		STATIC_CHECK(UE_CHECK_FORMAT_STRING_ERR(E::LSNeedsDereferencedWideString,          "Test %ls",  FWideString()));
+		STATIC_CHECK(UE_CHECK_FORMAT_STRING_ERR(E::LSNeedsDereferencedWideString, UTF8TEXT("Test %ls"), FAnsiString()));
+		STATIC_CHECK(UE_CHECK_FORMAT_STRING_ERR(E::LSNeedsDereferencedWideString, UTF8TEXT("Test %ls"), FUtf8String()));
+		STATIC_CHECK(UE_CHECK_FORMAT_STRING_ERR(E::LSNeedsDereferencedWideString, UTF8TEXT("Test %ls"), FWideString()));
+		STATIC_CHECK(UE_CHECK_FORMAT_STRING_ERR(E::LSNeedsDereferencedWideString, WIDETEXT("Test %ls"), FAnsiString()));
+		STATIC_CHECK(UE_CHECK_FORMAT_STRING_ERR(E::LSNeedsDereferencedWideString, WIDETEXT("Test %ls"), FUtf8String()));
+		STATIC_CHECK(UE_CHECK_FORMAT_STRING_ERR(E::LSNeedsDereferencedWideString, WIDETEXT("Test %ls"), FWideString()));
+
+		// Test char passed instead of char pointer
+		STATIC_CHECK(UE_CHECK_FORMAT_STRING_ERR(E::LSNeedsPtrButGotChar,          "Test %ls",           'x'));
+		STATIC_CHECK(UE_CHECK_FORMAT_STRING_ERR(E::LSNeedsPtrButGotChar,          "Test %ls",  UTF8TEXT('x')));
+		STATIC_CHECK(UE_CHECK_FORMAT_STRING_ERR(E::LSNeedsPtrButGotChar,          "Test %ls",  WIDETEXT('x')));
+		STATIC_CHECK(UE_CHECK_FORMAT_STRING_ERR(E::LSNeedsPtrButGotChar, UTF8TEXT("Test %ls"),          'x'));
+		STATIC_CHECK(UE_CHECK_FORMAT_STRING_ERR(E::LSNeedsPtrButGotChar, UTF8TEXT("Test %ls"), UTF8TEXT('x')));
+		STATIC_CHECK(UE_CHECK_FORMAT_STRING_ERR(E::LSNeedsPtrButGotChar, UTF8TEXT("Test %ls"), WIDETEXT('x')));
+		STATIC_CHECK(UE_CHECK_FORMAT_STRING_ERR(E::LSNeedsPtrButGotChar, WIDETEXT("Test %ls"),          'x'));
+		STATIC_CHECK(UE_CHECK_FORMAT_STRING_ERR(E::LSNeedsPtrButGotChar, WIDETEXT("Test %ls"), UTF8TEXT('x')));
+		STATIC_CHECK(UE_CHECK_FORMAT_STRING_ERR(E::LSNeedsPtrButGotChar, WIDETEXT("Test %ls"), WIDETEXT('x')));
+
+		// Test matching arguments
+		STATIC_CHECK(UE_CHECK_FORMAT_STRING_ERR(E::Ok,          "Test %ls",  WIDETEXT("hello")));
+		STATIC_CHECK(UE_CHECK_FORMAT_STRING_ERR(E::Ok, UTF8TEXT("Test %ls"), WIDETEXT("hello")));
+		STATIC_CHECK(UE_CHECK_FORMAT_STRING_ERR(E::Ok, WIDETEXT("Test %ls"), WIDETEXT("hello")));
+	}
+
+	SECTION("%c")
+	{
+		enum class EFakeChar : uint8 {};
+
+		// Test passing non-integer numbers
+		STATIC_CHECK(UE_CHECK_FORMAT_STRING_ERR(E::CNeedsCharArgOnNarrowString,          "Hello %c",  42.0));
+		STATIC_CHECK(UE_CHECK_FORMAT_STRING_ERR(E::CNeedsCharArgOnNarrowString, UTF8TEXT("Hello %c"), 42.0));
+		STATIC_CHECK(UE_CHECK_FORMAT_STRING_ERR(E::CNeedsCharArgOnWideString,   WIDETEXT("Hello %c"), 42.0));
+		STATIC_CHECK(UE_CHECK_FORMAT_STRING_ERR(E::CNeedsCharArgOnNarrowString,          "Hello %c",  42.0f));
+		STATIC_CHECK(UE_CHECK_FORMAT_STRING_ERR(E::CNeedsCharArgOnNarrowString, UTF8TEXT("Hello %c"), 42.0f));
+		STATIC_CHECK(UE_CHECK_FORMAT_STRING_ERR(E::CNeedsCharArgOnWideString,   WIDETEXT("Hello %c"), 42.0f));
+
+		// Test passing an enum that isn't UTF8CHAR
+		STATIC_CHECK(UE_CHECK_FORMAT_STRING_ERR(E::CNeedsCharArgOnNarrowString,          "Hello %c",  EFakeChar{}));
+		STATIC_CHECK(UE_CHECK_FORMAT_STRING_ERR(E::CNeedsCharArgOnNarrowString, UTF8TEXT("Hello %c"), EFakeChar{}));
+		STATIC_CHECK(UE_CHECK_FORMAT_STRING_ERR(E::CNeedsCharArgOnWideString,   WIDETEXT("Hello %c"), EFakeChar{}));
+
+		// Test passing a 64-bit integer
+		STATIC_CHECK(UE_CHECK_FORMAT_STRING_ERR(E::CNeedsCharArgOnNarrowString,          "Hello %c",  2147483648));
+		STATIC_CHECK(UE_CHECK_FORMAT_STRING_ERR(E::CNeedsCharArgOnNarrowString, UTF8TEXT("Hello %c"), 2147483648));
+		STATIC_CHECK(UE_CHECK_FORMAT_STRING_ERR(E::CNeedsCharArgOnWideString,   WIDETEXT("Hello %c"), 2147483648));
+
+		// Test passing a 32-bit integer
+		STATIC_CHECK(UE_CHECK_FORMAT_STRING_ERR(E::Ok,          "Hello %c",  42));
+		STATIC_CHECK(UE_CHECK_FORMAT_STRING_ERR(E::Ok, UTF8TEXT("Hello %c"), 42));
+		STATIC_CHECK(UE_CHECK_FORMAT_STRING_ERR(E::Ok, WIDETEXT("Hello %c"), 42));
+		STATIC_CHECK(UE_CHECK_FORMAT_STRING_ERR(E::Ok,          "Hello %c",  2147483647));
+		STATIC_CHECK(UE_CHECK_FORMAT_STRING_ERR(E::Ok, UTF8TEXT("Hello %c"), 2147483647));
+		STATIC_CHECK(UE_CHECK_FORMAT_STRING_ERR(E::Ok, WIDETEXT("Hello %c"), 2147483647));
+
+		// Test passing a char of any type
+		STATIC_CHECK(UE_CHECK_FORMAT_STRING_ERR(E::Ok,          "Hello %c",           'x'));
+		STATIC_CHECK(UE_CHECK_FORMAT_STRING_ERR(E::Ok, UTF8TEXT("Hello %c"),          'x'));
+		STATIC_CHECK(UE_CHECK_FORMAT_STRING_ERR(E::Ok, WIDETEXT("Hello %c"),          'x'));
+		STATIC_CHECK(UE_CHECK_FORMAT_STRING_ERR(E::Ok,          "Hello %c",  UTF8TEXT('x')));
+		STATIC_CHECK(UE_CHECK_FORMAT_STRING_ERR(E::Ok, UTF8TEXT("Hello %c"), UTF8TEXT('x')));
+		STATIC_CHECK(UE_CHECK_FORMAT_STRING_ERR(E::Ok, WIDETEXT("Hello %c"), UTF8TEXT('x')));
+		STATIC_CHECK(UE_CHECK_FORMAT_STRING_ERR(E::Ok,          "Hello %c",  WIDETEXT('x')));
+		STATIC_CHECK(UE_CHECK_FORMAT_STRING_ERR(E::Ok, UTF8TEXT("Hello %c"), WIDETEXT('x')));
+		STATIC_CHECK(UE_CHECK_FORMAT_STRING_ERR(E::Ok, WIDETEXT("Hello %c"), WIDETEXT('x')));
+	}
+
 	SECTION("Error Handling")
 	{
-		STATIC_CHECK(UE_CHECK_FORMAT_STRING_ERR(E::SNeedsTCHARPtrArg, TEXT("Test %s"), "wrong"));
-		STATIC_CHECK(UE_CHECK_FORMAT_STRING_ERR(E::HSNeedsCharPtrArg, TEXT("Test %hs"), TEXT("wrong")));
 		STATIC_CHECK(UE_CHECK_FORMAT_STRING_ERR(E::IncompleteFormatSpecifierOrUnescapedPercent, TEXT("Hello %")));
-		STATIC_CHECK(UE_CHECK_FORMAT_STRING_ERR(E::CNeedsTCHARArg, TEXT("Hello %c"), 42.2));
 		STATIC_CHECK(UE_CHECK_FORMAT_STRING_ERR(E::PNeedsPointerArg, TEXT("Hello %p"), 42));
 		STATIC_CHECK(UE_CHECK_FORMAT_STRING_ERR(E::FNeedsFloatOrDoubleArg, TEXT("Hello %f"), 42));
-		STATIC_CHECK(std::is_same_v<TCHAR, WIDECHAR> == UE_CHECK_FORMAT_STRING_ERR(E::CapitalSNeedsCharPtr, TEXT("Hello %S"), TEXT("wrong")));
-		STATIC_CHECK(std::is_same_v<TCHAR, ANSICHAR> == UE_CHECK_FORMAT_STRING_ERR(E::CapitalSNeedsTCHARPtr, TEXT("Hello %S"), "wrong"));
 		STATIC_CHECK(UE_CHECK_FORMAT_STRING_ERR(E::InvalidFormatSpec, TEXT("Hello %k"), 42));
 		STATIC_CHECK(UE_CHECK_FORMAT_STRING_ERR(E::NotEnoughArguments, TEXT("Hello %s")));
 		STATIC_CHECK(UE_CHECK_FORMAT_STRING_ERR(E::DNeedsIntegerArg, TEXT("Hello %d"), 42.0));
@@ -38,7 +270,7 @@ TEST_CASE_NAMED(FFormatStringValidatorTest, "System::Core::String::FormatStringS
 		STATIC_CHECK(UE_CHECK_FORMAT_STRING_ERR(E::LLNeedsIntegerArg, TEXT("Hello %lld"), 43.0f));
 		STATIC_CHECK(UE_CHECK_FORMAT_STRING_ERR(E::I64BadSpec, TEXT("Hello %I32d"), 44));
 		STATIC_CHECK(UE_CHECK_FORMAT_STRING_ERR(E::I64BadSpec, TEXT("Hello %I64p"), 44));
-		STATIC_CHECK(UE_CHECK_FORMAT_STRING_ERR(E::I64NeedsIntegerArg	, TEXT("Hello %I64u"), 44.0f));
+		STATIC_CHECK(UE_CHECK_FORMAT_STRING_ERR(E::I64NeedsIntegerArg, TEXT("Hello %I64u"), 44.0f));
 		STATIC_CHECK(UE_CHECK_FORMAT_STRING_ERR(E::InvalidFormatSpec, TEXT("%l^"), 42));
 		STATIC_CHECK(UE_CHECK_FORMAT_STRING_ERR(E::InvalidFormatSpec, TEXT("%h^"), 42));
 		STATIC_CHECK(UE_CHECK_FORMAT_STRING_ERR(E::IncompleteFormatSpecifierOrUnescapedPercent, TEXT("%-*"), 42));		
@@ -57,7 +289,6 @@ TEST_CASE_NAMED(FFormatStringValidatorTest, "System::Core::String::FormatStringS
 		STATIC_CHECK(UE_CHECK_FORMAT_STRING_ERR(E::Ok, TEXT("Test %d"), long(32)));
 		STATIC_CHECK(UE_CHECK_FORMAT_STRING_ERR(E::Ok, TEXT("Test %s"), TEXT("hello")));
 		STATIC_CHECK(UE_CHECK_FORMAT_STRING_ERR(E::Ok, TEXT("Test percent %% more")));
-		STATIC_CHECK(UE_CHECK_FORMAT_STRING_ERR(E::Ok, TEXT("Test %hs"), "hi"));
 
 		STATIC_CHECK(UE_CHECK_FORMAT_STRING_ERR(E::Ok, TEXT("%-8d %f"), 42, 4.4));
 		STATIC_CHECK(UE_CHECK_FORMAT_STRING_ERR(E::Ok, TEXT("%8d %f"), 42, 4.4));
