Fixed missing include.

#rnx

diff --git a/Engine/Source/Runtime/Core/Tests/String/FormatStringSanTest.cpp b/Engine/Source/Runtime/Core/Tests/String/FormatStringSanTest.cpp
--- a/Engine/Source/Runtime/Core/Tests/String/FormatStringSanTest.cpp
+++ b/Engine/Source/Runtime/Core/Tests/String/FormatStringSanTest.cpp
@@ -1,5 +1,7 @@
 // Copyright Epic Games, Inc. All Rights Reserved.
 
+#include "CoreTypes.h"
+
 // Tests need UE_UTF8TEXT_IS_CONSTEXPR because of the use of STATIC_CHECK with UTFTEXT.
 #if WITH_TESTS && UE_UTF8TEXT_IS_CONSTEXPR
 
