[AutoRTFM] Unregister `autortfm_open_to_closed_table` structures on atexit() for windows

DLLs can be unloaded before autortfm_initialize() is called, which leave invalid function table pointers in the linked list to register.
To prevent this, register an atexit() callback to remove the table.

diff --git a/Engine/Source/Runtime/Core/Private/AutoRTFM/API.cpp b/Engine/Source/Runtime/Core/Private/AutoRTFM/API.cpp
--- a/Engine/Source/Runtime/Core/Private/AutoRTFM/API.cpp
+++ b/Engine/Source/Runtime/Core/Private/AutoRTFM/API.cpp
@@ -901,6 +901,35 @@ extern "C" UE_AUTORTFM_NOAUTORTFM void autortfm_register_open_to_closed_function
 	}
 }
 
+extern "C" UE_AUTORTFM_NOAUTORTFM void autortfm_unregister_open_to_closed_functions(autortfm_open_to_closed_table* Table)
+{
+	if (Table == GPendingOpenToClosedRegistrations)
+	{
+		GPendingOpenToClosedRegistrations = Table->Next;
+	}
+	if (Table->Next)
+	{
+		Table->Next->Prev = Table->Prev;
+	}
+	if (Table->Prev)
+	{
+		Table->Prev->Next = Table->Next;
+	}
+	Table->Prev = nullptr;
+	Table->Next = nullptr;
+
+	// Note: If AutoRTFM is already initialized, we currently do *not* remove
+	// the registered functions from the function map. The reason for this is
+	// that we can register the same open address multiple times, where the
+	// closed address uses the value of the last register call.
+	// To support unregistering these cleanly, we'd need to increase the
+	// complexity of the function map - either by storing a list of all the
+	// closed functions that were registered for an open, or entirely rebuilding
+	// the map from the autortfm_open_to_closed_table lists. So far, keeping
+	// stale mappings has not been an issue, but if it does become an issue,
+	// then something will need to be done here.
+}
+
 extern "C" bool autortfm_is_on_current_transaction_stack(void* Ptr)
 {
 	if (!FContext::Get()->IsTransactional())
diff --git a/Engine/Source/Runtime/Core/Public/AutoRTFM/AutoRTFM.h b/Engine/Source/Runtime/Core/Public/AutoRTFM/AutoRTFM.h
--- a/Engine/Source/Runtime/Core/Public/AutoRTFM/AutoRTFM.h
+++ b/Engine/Source/Runtime/Core/Public/AutoRTFM/AutoRTFM.h
@@ -351,6 +351,9 @@ struct autortfm_open_to_closed_table
 {
 	// Null-terminated open function to closed function mapping table.
 	const autortfm_open_to_closed_mapping* Mappings;
+	// An intrusive linked-list pointer to the previous autortfm_open_to_closed_table.
+	// Used by autortfm_register_open_to_closed_functions().
+	autortfm_open_to_closed_table* Prev;
 	// An intrusive linked-list pointer to the next autortfm_open_to_closed_table.
 	// Used by autortfm_register_open_to_closed_functions().
 	autortfm_open_to_closed_table* Next;
@@ -358,11 +361,16 @@ struct autortfm_open_to_closed_table
 
 #if UE_AUTORTFM_ENABLED
 UE_AUTORTFM_API void autortfm_register_open_to_closed_functions(autortfm_open_to_closed_table* Table);
+UE_AUTORTFM_API void autortfm_unregister_open_to_closed_functions(autortfm_open_to_closed_table* Table);
 #else
 static UE_AUTORTFM_FORCEINLINE void autortfm_register_open_to_closed_functions(autortfm_open_to_closed_table* Table)
 {
 	UE_AUTORTFM_UNUSED(Table);
 }
+static UE_AUTORTFM_FORCEINLINE void autortfm_unregister_open_to_closed_functions(autortfm_open_to_closed_table* Table)
+{
+	UE_AUTORTFM_UNUSED(Table);
+}
 #endif
 
 #if UE_AUTORTFM_ENABLED
@@ -1315,29 +1323,51 @@ namespace ForTheRuntime
 
 	// A helper that on construction, registers all the open to closed function
 	// mappings expressed as the variadic TOpenToClosedMapping template
-	// arguments.
+	// arguments. On destruction the mappings are unregistered.
 	// Used by UE_AUTORTFM_REGISTER_OPEN_TO_CLOSED_FUNCTIONS()
 	template <typename ... OpenToClosedMappings>
 	class TAutoRegisterOpenToClosedFunctions
 	{
 	public:
-		TAutoRegisterOpenToClosedFunctions()
+		UE_AUTORTFM_NOAUTORTFM
+		inline TAutoRegisterOpenToClosedFunctions()
+		{
+			autortfm_register_open_to_closed_functions(GetTable());
+		}
+
+		UE_AUTORTFM_NOAUTORTFM
+		inline ~TAutoRegisterOpenToClosedFunctions()
+		{
+			autortfm_unregister_open_to_closed_functions(GetTable());
+		}
+
+	private:
+		struct FStaticData
 		{
-			static autortfm_open_to_closed_mapping Mappings[sizeof...(OpenToClosedMappings) + 1]{};
-			// Some functions may dynamically resolve to nullptr, which we want to filter out
-			// as this will be treated as a null-terminator.
-			size_t Count = 0;
-			auto AddMappingIfValid = [&](const auto& Mapping)
+			FStaticData()
 			{
-				if (Mapping.IsValid())
+				// Some functions may dynamically resolve to nullptr, which we want to filter out
+				// as this will be treated as a null-terminator.
+				size_t Count = 0;
+				auto AddMappingIfValid = [&](const auto& Mapping)
 				{
-					Mappings[Count++] = Mapping.Get();
-				}
-			};
-			(AddMappingIfValid(OpenToClosedMappings{}), ...);
-			Mappings[Count++] = {nullptr, nullptr};
-			static autortfm_open_to_closed_table Table{Mappings, nullptr};
-			autortfm_register_open_to_closed_functions(&Table);
+					if (Mapping.IsValid())
+					{
+						Mappings[Count++] = Mapping.Get();
+					}
+				};
+				(AddMappingIfValid(OpenToClosedMappings{}), ...);
+				Mappings[Count++] = {nullptr, nullptr};
+			}
+			autortfm_open_to_closed_mapping Mappings[sizeof...(OpenToClosedMappings) + 1]{};
+			autortfm_open_to_closed_table Table{Mappings, nullptr};
+		};
+
+		UE_AUTORTFM_NOAUTORTFM
+		static autortfm_open_to_closed_table* GetTable()
+		{
+			static FStaticData StaticData;
+			return &StaticData.Table;
 		}
 	};
 
