Fix `CSV_CUSTOM_STAT` and `CSV_CUSTOM_STAT_GLOBAL` under AutoRTFM and add tests.

diff --git a/Engine/Source/Programs/AutoRTFMTests/Private/CsvProfilerTests.cpp b/Engine/Source/Programs/AutoRTFMTests/Private/CsvProfilerTests.cpp
--- a/Engine/Source/Programs/AutoRTFMTests/Private/CsvProfilerTests.cpp
+++ b/Engine/Source/Programs/AutoRTFMTests/Private/CsvProfilerTests.cpp
@@ -82,4 +82,24 @@ TEST_CASE("CsvProfiler")
 		CSV_EVENT(MyTrueCategory, TEXT("HEREWEGO"));
 		CSV_SCOPED_TIMING_STAT(MyTrueCategory, Event);
 	});
+
+	AutoRTFM::Testing::Commit([&]
+	{
+		CSV_CUSTOM_STAT(MyFalseCategory, SomeStat, 42, ECsvCustomStatOp::Max);
+	});
+
+	AutoRTFM::Testing::Commit([&]
+	{
+		CSV_CUSTOM_STAT(MyTrueCategory, SomeStat, 42, ECsvCustomStatOp::Max);
+	});
+
+	AutoRTFM::Testing::Commit([&]
+	{
+		CSV_CUSTOM_STAT_GLOBAL(SomeStat, 42, ECsvCustomStatOp::Max);
+	});
+
+	AutoRTFM::Testing::Commit([&]
+	{
+		CSV_CUSTOM_STAT_GLOBAL(SomeStat, 42, ECsvCustomStatOp::Max);
+	});
 }
diff --git a/Engine/Source/Runtime/Core/Public/ProfilingDebugging/CsvProfiler.h b/Engine/Source/Runtime/Core/Public/ProfilingDebugging/CsvProfiler.h
--- a/Engine/Source/Runtime/Core/Public/ProfilingDebugging/CsvProfiler.h
+++ b/Engine/Source/Runtime/Core/Public/ProfilingDebugging/CsvProfiler.h
@@ -128,11 +128,17 @@ struct FCsvDeclaredStat;
 
 	// Custom Stats
 	#define CSV_CUSTOM_STAT(Category,StatName,Value,Op) \
-		TRACE_CSV_PROFILER_INLINE_STAT(#StatName, CSV_CATEGORY_INDEX(Category)); \
-		FCsvProfiler::RecordCustomStat(#StatName, CSV_CATEGORY_INDEX(Category), Value, Op);
+		UE_AUTORTFM_OPEN \
+		{ \
+			TRACE_CSV_PROFILER_INLINE_STAT(#StatName, CSV_CATEGORY_INDEX(Category)); \
+			FCsvProfiler::RecordCustomStat(#StatName, CSV_CATEGORY_INDEX(Category), Value, Op); \
+		};
 	#define CSV_CUSTOM_STAT_GLOBAL(StatName,Value,Op) \
-		TRACE_CSV_PROFILER_INLINE_STAT(#StatName, CSV_CATEGORY_INDEX_GLOBAL); \
-		FCsvProfiler::RecordCustomStat(#StatName, CSV_CATEGORY_INDEX_GLOBAL, Value, Op); 
+		UE_AUTORTFM_OPEN \
+		{ \
+			TRACE_CSV_PROFILER_INLINE_STAT(#StatName, CSV_CATEGORY_INDEX_GLOBAL); \
+			FCsvProfiler::RecordCustomStat(#StatName, CSV_CATEGORY_INDEX_GLOBAL, Value, Op); \
+		};
 
 	// Stats declarations/definitions (for FName stats)
 	#define CSV_DEFINE_STAT(Category,StatName)						FCsvDeclaredStat _GCsvStat_##StatName((TCHAR*)TEXT(#StatName), CSV_CATEGORY_INDEX(Category));
