Make Maths functions [[nodiscard]]

Issues found with this change:
https://p4-swarm.epicgames.net/changes/39264399
https://p4-swarm.epicgames.net/changes/39254483
https://p4-swarm.epicgames.net/changes/39254407
https://p4-swarm.epicgames.net/changes/39066690
https://p4-swarm.epicgames.net/changes/39051787
https://p4-swarm.epicgames.net/changes/39051738
https://p4-swarm.epicgames.net/changes/38946145
https://p4-swarm.epicgames.net/changes/38945497
https://p4-swarm.epicgames.net/changes/38934542
https://p4-swarm.epicgames.net/changes/38924578
https://p4-swarm.epicgames.net/changes/38924493

[RN] Make Maths functions [[nodiscard]]

diff --git a/Engine/Source/Runtime/Core/Private/Math/UnrealMath.cpp b/Engine/Source/Runtime/Core/Private/Math/UnrealMath.cpp
--- a/Engine/Source/Runtime/Core/Private/Math/UnrealMath.cpp
+++ b/Engine/Source/Runtime/Core/Private/Math/UnrealMath.cpp
@@ -221,7 +221,7 @@ UE::Math::TRotator<T> UE::Math::TVector<T>::ToOrientationRotator() const
 }
 
 template<typename T>
-UE::Math::TVector<T> UE::Math::TVector<T>::SlerpVectorToDirection(UE::Math::TVector<T>& V, UE::Math::TVector<T>& Direction, T Alpha)
+UE::Math::TVector<T> UE::Math::TVector<T>::SlerpVectorToDirection(const UE::Math::TVector<T>& V, const UE::Math::TVector<T>& Direction, T Alpha)
 {
 	using TVector = UE::Math::TVector<T>;
 	using TQuat = UE::Math::TQuat<T>;
@@ -238,7 +238,7 @@ UE::Math::TVector<T> UE::Math::TVector<T>::SlerpVectorToDirection(UE::Math::TVec
 }
 
 template<typename T>
-UE::Math::TVector<T> UE::Math::TVector<T>::SlerpNormals(UE::Math::TVector<T>& NormalA, UE::Math::TVector<T>& NormalB, T Alpha)
+UE::Math::TVector<T> UE::Math::TVector<T>::SlerpNormals(const UE::Math::TVector<T>& NormalA, const UE::Math::TVector<T>& NormalB, T Alpha)
 {
 	using TVector = UE::Math::TVector<T>;
 	using TQuat = UE::Math::TQuat<T>;
diff --git a/Engine/Source/Runtime/Core/Public/Math/BasicMathExpressionEvaluator.h b/Engine/Source/Runtime/Core/Public/Math/BasicMathExpressionEvaluator.h
--- a/Engine/Source/Runtime/Core/Public/Math/BasicMathExpressionEvaluator.h
+++ b/Engine/Source/Runtime/Core/Public/Math/BasicMathExpressionEvaluator.h
@@ -83,10 +83,10 @@ class FBasicMathExpressionEvaluator
 {
 public:
 	/** Constructor that sets up the parser's lexer and compiler */
-	CORE_API FBasicMathExpressionEvaluator();
+	[[nodiscard]] CORE_API FBasicMathExpressionEvaluator();
 
 	/** Evaluate the given expression, resulting in either a double value, or an error */
-	CORE_API TValueOrError<double, FExpressionError> Evaluate(const TCHAR* InExpression, double InExistingValue = 0) const;
+	[[nodiscard]] CORE_API TValueOrError<double, FExpressionError> Evaluate(const TCHAR* InExpression, double InExistingValue = 0) const;
 
 	/**
 	 * Utils that can be re-used as part of more complex evaluators
diff --git a/Engine/Source/Runtime/Core/Public/Math/BigInt.h b/Engine/Source/Runtime/Core/Public/Math/BigInt.h
--- a/Engine/Source/Runtime/Core/Public/Math/BigInt.h
+++ b/Engine/Source/Runtime/Core/Public/Math/BigInt.h
@@ -79,12 +79,12 @@ private:
 
 public:
 
-	FORCEINLINE uint32* GetBits()
+	[[nodiscard]] FORCEINLINE uint32* GetBits()
 	{
 		return Bits;
 	}
 
-	FORCEINLINE const uint32* GetBits() const
+	[[nodiscard]] FORCEINLINE const uint32* GetBits() const
 	{
 		return Bits;
 	}
@@ -108,7 +108,7 @@ public:
 	}
 
 	/** Default constructor. Initializes the number to zero. */
-	TBigInt()
+	[[nodiscard]] TBigInt()
 	{
 		Zero();
 	}
@@ -118,7 +118,7 @@ public:
 	 *
 	 * @param Other The value to set.
 	 */
-	TBigInt(int64 Other)
+	[[nodiscard]] TBigInt(int64 Other)
 	{
 		Set(Other);
 	}
@@ -126,7 +126,7 @@ public:
 	/**
 	 * Constructor. Initializes this big int with an array of words.
 	 */
-	explicit TBigInt(const uint32* InBits)
+	[[nodiscard]] explicit TBigInt(const uint32* InBits)
 	{
 		FMemory::Memcpy(Bits, InBits, sizeof(Bits));
 	}
@@ -134,7 +134,7 @@ public:
 	/**
 	 * Constructor. Initializes this big int with an array of bytes.
 	 */
-	explicit TBigInt(const uint8* InData, uint32 InNumBytes)
+	[[nodiscard]] explicit TBigInt(const uint8* InData, uint32 InNumBytes)
 	{
 		// If we weren't given enough data to fill the int, zero first
 		// TODO: Only zero the bits we aren't going to copy over
@@ -149,7 +149,7 @@ public:
 	/**
 	 * Constructor. Initializes this big int with a string representing a hex value.
 	 */
-	explicit TBigInt(const FString& Value)
+	[[nodiscard]] explicit TBigInt(const FString& Value)
 	{
 		Parse(Value);
 	}
@@ -282,7 +282,7 @@ public:
 	/**
 	 * Checks if this integer is negative.
 	 */
-	bool IsNegative() const
+	[[nodiscard]] bool IsNegative() const
 	{
 		if (bSigned)
 		{
@@ -297,7 +297,7 @@ public:
 	/**
 	 * Returns the sign of this integer.
 	 */
-	int32 Sign() const
+	[[nodiscard]] int32 Sign() const
 	{
 		return IsNegative() ? -1 : 1;
 	}
@@ -317,7 +317,7 @@ public:
 	/**
 	* Returns the index of the highest word that is not zero. -1 if no such word exists.
 	*/
-	int32 GetHighestNonZeroWord() const
+	[[nodiscard]] int32 GetHighestNonZeroWord() const
 	{
 		int32 WordIndex;
 		for (WordIndex = NumWords - 1; WordIndex >= 0 && Bits[WordIndex] == 0; --WordIndex);
@@ -526,7 +526,7 @@ public:
 	/**
 	 * Returns the index of the highest non-zero bit. -1 if no such bit exists.
 	 */
-	int32 GetHighestNonZeroBit() const
+	[[nodiscard]] int32 GetHighestNonZeroBit() const
 	{
 		for (int32 WordIndex = NumWords - 1; WordIndex >= 0; --WordIndex)
 		{
@@ -543,7 +543,7 @@ public:
 	/**
 	 * Returns a bit value as an integer value (0 or 1).
 	 */
-	int32 GetBit(int32 BitIndex) const
+	[[nodiscard]] int32 GetBit(int32 BitIndex) const
 	{
 		const int32 WordIndex = BitIndex / BitsPerWord;
 		BitIndex -= WordIndex * BitsPerWord;
@@ -653,7 +653,7 @@ public:
 	/**
 	 * Checks if two integers are equal.
 	 */
-	bool IsEqual(const BigInt& Other) const
+	[[nodiscard]] bool IsEqual(const BigInt& Other) const
 	{
 		for (int32 WordIndex = 0; WordIndex < NumWords; ++WordIndex)
 		{
@@ -669,7 +669,7 @@ public:
 	/**
 	 * this < Other
 	 */
-	bool IsLess(const BigInt& Other) const
+	[[nodiscard]] bool IsLess(const BigInt& Other) const
 	{
 		if (IsNegative())
 		{
@@ -690,7 +690,7 @@ public:
 	/**
 	 * this <= Other
 	 */
-	bool IsLessOrEqual(const BigInt& Other) const
+	[[nodiscard]] bool IsLessOrEqual(const BigInt& Other) const
 	{
 		if (IsNegative())
 		{
@@ -711,7 +711,7 @@ public:
 	/**
 	 * this > Other
 	 */
-	bool IsGreater(const BigInt& Other) const
+	[[nodiscard]] bool IsGreater(const BigInt& Other) const
 	{
 		if (IsNegative())
 		{
@@ -732,7 +732,7 @@ public:
 	/**
 	 * this >= Other
 	 */
-	bool IsGreaterOrEqual(const BigInt& Other) const
+	[[nodiscard]] bool IsGreaterOrEqual(const BigInt& Other) const
 	{
 		if (IsNegative())
 		{
@@ -753,7 +753,7 @@ public:
 	/**
 	 * this == 0
 	 */
-	bool IsZero() const
+	[[nodiscard]] bool IsZero() const
 	{
 		int32 WordIndex;
 		for (WordIndex = NumWords - 1; WordIndex >= 0 && !Bits[WordIndex]; --WordIndex);
@@ -763,7 +763,7 @@ public:
 	/**
 	 * this > 0
 	 */
-	bool IsGreaterThanZero() const
+	[[nodiscard]] bool IsGreaterThanZero() const
 	{
 		return !IsNegative() && !IsZero();
 	}
@@ -771,25 +771,25 @@ public:
 	/**
 	 * this < 0
 	 */
-	bool IsLessThanZero() const
+	[[nodiscard]] bool IsLessThanZero() const
 	{
 		return IsNegative() && !IsZero();
 	}
 
-	bool IsFirstBitSet() const
+	[[nodiscard]] bool IsFirstBitSet() const
 	{
 		return !!(Bits[0] & 1);
 	}
 	/**
 	 * Bit indexing operator
 	 */
-	bool operator[](int32 BitIndex) const
+	[[nodiscard]] bool operator[](int32 BitIndex) const
 	{
 		return GetBit(BitIndex);
 	}
 
 	// Begin operator overloads
-	BigInt operator>>(int32 Count) const
+	[[nodiscard]] BigInt operator>>(int32 Count) const
 	{
 		BigInt Result = *this;
 		Result.ShiftRight(Count);
@@ -802,7 +802,7 @@ public:
 		return *this;
 	}
 
-	BigInt operator<<(int32 Count) const
+	[[nodiscard]] BigInt operator<<(int32 Count) const
 	{
 		BigInt Result = *this;
 		Result.ShiftLeft(Count);
@@ -815,7 +815,7 @@ public:
 		return *this;
 	}
 
-	BigInt operator+(const BigInt& Other) const
+	[[nodiscard]] BigInt operator+(const BigInt& Other) const
 	{
 		BigInt Result(*this);
 		Result.Add(Other);
@@ -834,7 +834,7 @@ public:
 		return *this;
 	}
 
-	BigInt operator-(const BigInt& Other) const
+	[[nodiscard]] BigInt operator-(const BigInt& Other) const
 	{
 		BigInt Result(*this);
 		Result.Subtract(Other);
@@ -853,7 +853,7 @@ public:
 		return *this;
 	}
 
-	BigInt operator*(const BigInt& Other) const
+	[[nodiscard]] BigInt operator*(const BigInt& Other) const
 	{
 		BigInt Result(*this);
 		Result.Multiply(Other);
@@ -866,7 +866,7 @@ public:
 		return *this;
 	}
 
-	BigInt operator/(const BigInt& Divider) const
+	[[nodiscard]] BigInt operator/(const BigInt& Divider) const
 	{
 		BigInt Result(*this);
 		Result.Divide(Divider);
@@ -879,7 +879,7 @@ public:
 		return *this;
 	}
 
-	BigInt operator%(const BigInt& Modulus) const
+	[[nodiscard]] BigInt operator%(const BigInt& Modulus) const
 	{
 		BigInt Result(*this);
 		Result.Modulo(Modulus);
@@ -892,37 +892,37 @@ public:
 		return *this;
 	}
 
-	bool operator<(const BigInt& Other) const
+	[[nodiscard]] bool operator<(const BigInt& Other) const
 	{
 		return IsLess(Other);
 	}
 
-	bool operator<=(const BigInt& Other) const
+	[[nodiscard]] bool operator<=(const BigInt& Other) const
 	{
 		return IsLessOrEqual(Other);
 	}
 
-	bool operator>(const BigInt& Other) const
+	[[nodiscard]] bool operator>(const BigInt& Other) const
 	{
 		return IsGreater(Other);
 	}
 
-	bool operator>=(const BigInt& Other) const
+	[[nodiscard]] bool operator>=(const BigInt& Other) const
 	{
 		return IsGreaterOrEqual(Other);
 	}
 
-	bool operator==(const BigInt& Other) const
+	[[nodiscard]] bool operator==(const BigInt& Other) const
 	{
 		return IsEqual(Other);
 	}
 
-	bool operator!=(const BigInt& Other) const
+	[[nodiscard]] bool operator!=(const BigInt& Other) const
 	{
 		return !IsEqual(Other);
 	}
 
-	BigInt operator&(const BigInt& Other) const
+	[[nodiscard]] BigInt operator&(const BigInt& Other) const
 	{
 		BigInt Result(*this);
 		Result.BitwiseAnd(Other);
@@ -935,7 +935,7 @@ public:
 		return *this;
 	}
 
-	BigInt operator|(const BigInt& Other) const
+	[[nodiscard]] BigInt operator|(const BigInt& Other) const
 	{
 		BigInt Result(*this);
 		Result.BitwiseOr(Other);
@@ -948,7 +948,7 @@ public:
 		return *this;
 	}
 
-	BigInt operator~() const
+	[[nodiscard]] BigInt operator~() const
 	{
 		BigInt Result(*this);
 		Result.BitwiseNot();
@@ -959,7 +959,7 @@ public:
 	/**
 	 * Returns the value of this big int as a 64-bit integer. If the value is greater, the higher bits are truncated.
 	 */
-	int64 ToInt() const
+	[[nodiscard]] int64 ToInt() const
 	{
 		return (int64)Bits[0] + ((int64)Bits[1] << BitsPerWord);
 	}
@@ -967,7 +967,7 @@ public:
 	/**
 	 * Returns this big int as a string.
 	 */
-	FString ToString() const
+	[[nodiscard]] FString ToString() const
 	{
 		FString Text(TEXT("0x"));
 		int32 WordIndex;
diff --git a/Engine/Source/Runtime/Core/Public/Math/Bounds.h b/Engine/Source/Runtime/Core/Public/Math/Bounds.h
--- a/Engine/Source/Runtime/Core/Public/Math/Bounds.h
+++ b/Engine/Source/Runtime/Core/Public/Math/Bounds.h
@@ -38,12 +38,12 @@ struct TBounds
 		return *this;
 	}
 
-	FORCEINLINE TBounds<T> operator+( const TBounds<T>& Other ) const
+	[[nodiscard]] FORCEINLINE TBounds<T> operator+( const TBounds<T>& Other ) const
 	{
 		return TBounds<T>(*this) += Other;
 	}
 
-	FORCEINLINE bool Intersect( const TBounds<T>& Other ) const
+	[[nodiscard]] FORCEINLINE bool Intersect( const TBounds<T>& Other ) const
 	{
 		int Separated;
 		Separated  = VectorAnyGreaterThan( VectorLoadAligned( &Min ), VectorLoadAligned( &Other.Max ) );
@@ -52,14 +52,14 @@ struct TBounds
 		return Separated == 0;
 	}
 
-	FORCEINLINE bool Contains( const TBounds<T>& Other ) const
+	[[nodiscard]] FORCEINLINE bool Contains( const TBounds<T>& Other ) const
 	{
 		int MaskMin = VectorMaskBits( VectorCompareLE( VectorLoadAligned( &Min ), VectorLoadAligned( &Other.Min ) ) );
 		int MaskMax = VectorMaskBits( VectorCompareGE( VectorLoadAligned( &Max ), VectorLoadAligned( &Other.Max ) ) );
 		return ( MaskMin & MaskMax & 7 ) == 7;
 	}
 
-	FORCEINLINE T DistSqr( const TVector<T>& Point ) const
+	[[nodiscard]] FORCEINLINE T DistSqr( const TVector<T>& Point ) const
 	{
 		auto rMin		= VectorLoadAligned( &Min );
 		auto rMax		= VectorLoadAligned( &Max );
@@ -68,29 +68,29 @@ struct TBounds
 		return VectorDot3Scalar( rClosest, rClosest );
 	}
 
-	FORCEINLINE TVector<T> GetCenter() const
+	[[nodiscard]] FORCEINLINE TVector<T> GetCenter() const
 	{
 		return (Max + Min) * 0.5f;
 	}
 
-	FORCEINLINE TVector<T> GetExtent() const
+	[[nodiscard]] FORCEINLINE TVector<T> GetExtent() const
 	{
 		return (Max - Min) * 0.5f;
 	}
 
-	FORCEINLINE TVector<T> GetSize() const
+	[[nodiscard]] FORCEINLINE TVector<T> GetSize() const
 	{
 		return (Max - Min);
 	}
 
-	FORCEINLINE T GetSurfaceArea() const
+	[[nodiscard]] FORCEINLINE T GetSurfaceArea() const
 	{
 		TVector<T> Size = Max - Min;
 		return 0.5f * (Size.X * Size.Y + Size.X * Size.Z + Size.Y * Size.Z);
 	}
 
 	template< typename U >
-	FORCEINLINE auto ToAbsolute( const TVector<U>& Offset ) const
+	[[nodiscard]] FORCEINLINE auto ToAbsolute( const TVector<U>& Offset ) const
 	{
 		using CommonType = typename std::common_type<T, U>::type;
 
@@ -101,7 +101,7 @@ struct TBounds
 	}
 
 	template< typename U >
-	FORCEINLINE TBounds< float > ToRelative( const TVector<U>& Offset ) const
+	[[nodiscard]] FORCEINLINE TBounds< float > ToRelative( const TVector<U>& Offset ) const
 	{
 		using CommonType = typename std::common_type<T, U>::type;
 
diff --git a/Engine/Source/Runtime/Core/Public/Math/Box.h b/Engine/Source/Runtime/Core/Public/Math/Box.h
--- a/Engine/Source/Runtime/Core/Public/Math/Box.h
+++ b/Engine/Source/Runtime/Core/Public/Math/Box.h
@@ -56,7 +56,7 @@ public:
 	 *       in the same way, but this would break backwards compatibility.  It
 	 *       might be possible if the questionable IsValid field ever gets removed.
 	 */
-	constexpr TBox()
+	[[nodiscard]] constexpr TBox()
 #if ENABLE_UNINITIALIZED_BOX_DIAGNOSTIC
 		: Min(std::numeric_limits<T>::quiet_NaN(), TVectorConstInit{})
 		, Max(std::numeric_limits<T>::quiet_NaN(), TVectorConstInit{})
@@ -70,19 +70,19 @@ public:
 	}
 
 	/** Creates a new box without initialization. */
-	explicit constexpr TBox(ENoInit)
+	[[nodiscard]] explicit constexpr TBox(ENoInit)
 	{
 	}
 
 	/** Creates and initializes a new box with zero extent and marks it as invalid. */
-	explicit constexpr TBox(EForceInit)
+	[[nodiscard]] explicit constexpr TBox(EForceInit)
 		: Min(0, TVectorConstInit{})
 		, Max(0, TVectorConstInit{})
 		, IsValid(0)
 	{
 	}
 
-	constexpr TBox(EForceInit, TBoxConstInit)
+	[[nodiscard]] constexpr TBox(EForceInit, TBoxConstInit)
 		: Min(0, TVectorConstInit{}), Max(0, TVectorConstInit{}), IsValid(0)
 	{
 	}
@@ -106,7 +106,7 @@ public:
 	FORCEINLINE void DiagnosticCheck() const {}
 #endif
 
-	bool ContainsNaN() const
+	[[nodiscard]] bool ContainsNaN() const
 	{
 		return Min.ContainsNaN() || Max.ContainsNaN();
 	}
@@ -174,7 +174,7 @@ public:
 	 *
 	 * @return true if the boxes are equal, false otherwise.
 	 */
-	FORCEINLINE bool operator==( const TBox<T>& Other ) const
+	[[nodiscard]] FORCEINLINE bool operator==( const TBox<T>& Other ) const
 	{
 		return (!IsValid && !Other.IsValid) || ((IsValid && Other.IsValid) && (Min == Other.Min) && (Max == Other.Max));
 	}
@@ -184,7 +184,7 @@ public:
 	 * 
 	 * @return false if the boxes are equal, true otherwise.
 	 */
-	FORCEINLINE bool operator!=( const TBox<T>& Other) const
+	[[nodiscard]] FORCEINLINE bool operator!=( const TBox<T>& Other) const
 	{
 		return !(*this == Other);
 	}
@@ -198,7 +198,7 @@ public:
 	 * @param Tolerance Error tolerance.
 	 * @return true if the boxes are equal within tolerance limits, false otherwise.
 	 */
-	bool Equals(const TBox<T>& Other, T Tolerance=UE_KINDA_SMALL_NUMBER) const
+	[[nodiscard]] bool Equals(const TBox<T>& Other, T Tolerance=UE_KINDA_SMALL_NUMBER) const
 	{
 		return (!IsValid && !Other.IsValid) || ((IsValid && Other.IsValid) && Min.Equals(Other.Min, Tolerance) && Max.Equals(Other.Max, Tolerance));
 	}
@@ -217,7 +217,7 @@ public:
 	 * @param Other The other point to add to this.
 	 * @return A new bounding volume.
 	 */
-	FORCEINLINE TBox<T> operator+( const TVector<T>& Other ) const
+	[[nodiscard]] FORCEINLINE TBox<T> operator+( const TVector<T>& Other ) const
 	{
 		return TBox<T>(*this) += Other;
 	}
@@ -236,7 +236,7 @@ public:
 	 * @param Other The other volume to add to this.
 	 * @return A new bounding volume.
 	 */
-	FORCEINLINE TBox<T> operator+( const TBox<T>& Other ) const
+	[[nodiscard]] FORCEINLINE TBox<T> operator+( const TBox<T>& Other ) const
 	{
 		return TBox<T>(*this) += Other;
 	}
@@ -247,7 +247,24 @@ public:
 	 * @param Index the index into points of the bounding volume.
 	 * @return a reference to a point of the bounding volume.
 	 */
-    FORCEINLINE TVector<T>& operator[]( int32 Index )
+    [[nodiscard]] FORCEINLINE TVector<T>& operator[]( int32 Index )
+	{
+		check((Index >= 0) && (Index < 2));
+
+		if (Index == 0)
+		{
+			return Min;
+		}
+
+		return Max;
+	}
+
+	/* Gets reference to the min or max of this bounding volume.
+	 *
+	 * @param Index the index into points of the bounding volume.
+	 * @return a const reference to a point of the bounding volume.
+	 */
+    [[nodiscard]] FORCEINLINE const TVector<T>& operator[]( int32 Index ) const
 	{
 		check((Index >= 0) && (Index < 2));
 
@@ -267,7 +284,7 @@ public:
 	 * @param Point The point.
 	 * @return The distance.
 	 */
-	FORCEINLINE T ComputeSquaredDistanceToPoint( const TVector<T>& Point ) const
+	[[nodiscard]] FORCEINLINE T ComputeSquaredDistanceToPoint( const TVector<T>& Point ) const
 	{
 		return ComputeSquaredDistanceFromBoxToPoint(Min, Max, Point);
 	}
@@ -275,7 +292,7 @@ public:
 	/**
 	 * Calculates squared distance between two boxes.
 	 */
-	FORCEINLINE T ComputeSquaredDistanceToBox(const TBox<T>& Box) const
+	[[nodiscard]] FORCEINLINE T ComputeSquaredDistanceToBox(const TBox<T>& Box) const
 	{
 		TVector<T> AxisDistances = (GetCenter() - Box.GetCenter()).GetAbs() - (GetExtent() + Box.GetExtent());
 		AxisDistances = TVector<T>::Max(AxisDistances, TVector<T>(0.0f, 0.0f, 0.0f));
diff --git a/Engine/Source/Runtime/Core/Public/Math/Box2D.h b/Engine/Source/Runtime/Core/Public/Math/Box2D.h
--- a/Engine/Source/Runtime/Core/Public/Math/Box2D.h
+++ b/Engine/Source/Runtime/Core/Public/Math/Box2D.h
@@ -43,7 +43,7 @@ public:
 public:
 
 	/** Default constructor (no initialization). */
-	TBox2() = default;
+	[[nodiscard]] TBox2() = default;
 
 	/**
 	 * Creates and initializes a new box.
@@ -52,7 +52,7 @@ public:
 	 *
 	 * @param EForceInit Force Init Enum.
 	 */
-	explicit TBox2<T>( EForceInit )
+	[[nodiscard]] explicit TBox2<T>( EForceInit )
 	{
 		Init();
 	}
@@ -63,7 +63,7 @@ public:
 	 * @param InMin The box's minimum point.
 	 * @param InMax The box's maximum point.
 	 */
-	TBox2<T>( const TVector2<T>& InMin, const TVector2<T>& InMax )
+	[[nodiscard]] TBox2<T>( const TVector2<T>& InMin, const TVector2<T>& InMax )
 		: Min(InMin)
 		, Max(InMax)
 		, bIsValid(true)
@@ -75,14 +75,14 @@ public:
 	 * @param Points Array of Points to create for the bounding volume.
 	 * @param Count The number of points.
 	 */
-	TBox2<T>( const TVector2<T>* Points, const int32 Count );
+	[[nodiscard]] TBox2<T>( const TVector2<T>* Points, const int32 Count );
 
 	/**
 	 * Creates and initializes a new box from an array of points.
 	 *
 	 * @param Points Array of Points to create for the bounding volume.
 	 */
-	TBox2<T>( const TArray<TVector2<T>>& Points );
+	[[nodiscard]] TBox2<T>( const TArray<TVector2<T>>& Points );
 
 public:
 
@@ -92,7 +92,7 @@ public:
 	 * @param Other The other box to compare with.
 	 * @return true if the boxes are equal, false otherwise.
 	 */
-	bool operator==( const TBox2<T>& Other ) const
+	[[nodiscard]] bool operator==( const TBox2<T>& Other ) const
 	{
 		return (Min == Other.Min) && (Max == Other.Max);
 	}
@@ -103,7 +103,7 @@ public:
 	 * @param Other The other box to compare with.
 	 * @return true if the boxes are not equal, false otherwise.
 	 */
-	bool operator!=( const TBox2<T>& Other ) const
+	[[nodiscard]] bool operator!=( const TBox2<T>& Other ) const
 	{
 		return !(*this == Other);
 	}
@@ -115,7 +115,7 @@ public:
 	 * @param Tolerance Error tolerance.
 	 * @return true if the boxes are equal within specified tolerance, otherwise false.
 	 */
-	bool Equals( const TBox2<T>& Other, T Tolerance = UE_KINDA_SMALL_NUMBER ) const
+	[[nodiscard]] bool Equals( const TBox2<T>& Other, T Tolerance = UE_KINDA_SMALL_NUMBER ) const
 	{
 		return Min.Equals(Other.Min, Tolerance) && Max.Equals(Other.Max, Tolerance);
 	}
@@ -134,7 +134,7 @@ public:
 	 * @param Other The other point to add to this.
 	 * @return A new bounding volume.
 	 */
-	TBox2<T> operator+( const TVector2<T>& Other ) const
+	[[nodiscard]] TBox2<T> operator+( const TVector2<T>& Other ) const
 	{
 		return TBox2<T>(*this) += Other;
 	}
@@ -153,7 +153,7 @@ public:
 	 * @param Other The other volume to add to this.
 	 * @return A new bounding volume.
 	 */
-	TBox2<T> operator+( const TBox2<T>& Other ) const
+	[[nodiscard]] TBox2<T> operator+( const TBox2<T>& Other ) const
 	{
 		return TBox2<T>(*this) += Other;
 	}
@@ -164,7 +164,7 @@ public:
 	 * @param Index The index into points of the bounding volume.
 	 * @return A reference to a point of the bounding volume.
 	 */
-   TVector2<T>& operator[]( int32 Index )
+   [[nodiscard]] TVector2<T>& operator[]( int32 Index )
 	{
 		check((Index >= 0) && (Index < 2));
 
@@ -176,6 +176,25 @@ public:
 		return Max;
 	}
 
+	/**
+	 * Gets reference to the min or max of this bounding volume.
+	 *
+	 * @param Index The index into points of the bounding volume.
+	 * @return A const reference to a point of the bounding volume.
+	 */
+   [[nodiscard]] const TVector2<T>& operator[]( int32 Index ) const
+	{
+		check((Index >= 0) && (Index < 2));
+
+		if (Index == 0)
+		{
+			return Min;
+		}
+
+		return Max;
+	}
+
+
 public:
 
 	/** 
@@ -184,7 +203,7 @@ public:
 	 * @param Point The point.
 	 * @return The distance.
 	 */
-	FORCEINLINE T ComputeSquaredDistanceToPoint( const TVector2<T>& Point ) const
+	[[nodiscard]] FORCEINLINE T ComputeSquaredDistanceToPoint( const TVector2<T>& Point ) const
 	{
 		// Accumulates the distance as we iterate axis
 		T DistSquared = 0.f;
@@ -216,7 +235,7 @@ public:
 	 * @param W The size to increase volume by.
 	 * @return A new bounding box increased in size.
 	 */
-	TBox2<T> ExpandBy( const T W ) const
+	[[nodiscard]] TBox2<T> ExpandBy( const T W ) const
 	{
 		return TBox2<T>(Min - TVector2<T>(W, W), Max + TVector2<T>(W, W));
 	}
@@ -227,7 +246,7 @@ public:
      * @param V The size to increase the volume by.
      * @return A new bounding box.
      */
-	TBox2<T> ExpandBy(const TVector2<T>& V) const
+	[[nodiscard]] TBox2<T> ExpandBy(const TVector2<T>& V) const
 	{
 		return TBox2<T>(Min - V, Max + V);
 	}
@@ -238,7 +257,7 @@ public:
 	 * @return Box area.
 	 * @see GetCenter, GetCenterAndExtents, GetExtent, GetSize
 	 */
-	T GetArea() const
+	[[nodiscard]] T GetArea() const
 	{
 		return (Max.X - Min.X) * (Max.Y - Min.Y);
 	}
@@ -249,7 +268,7 @@ public:
 	 * @return Th center point.
 	 * @see GetArea, GetCenterAndExtents, GetExtent, GetSize
 	 */
-	TVector2<T> GetCenter() const
+	[[nodiscard]] TVector2<T> GetCenter() const
 	{
 		return TVector2<T>((Min + Max) * 0.5f);
 	}
@@ -274,7 +293,7 @@ public:
 	 *
 	 * @return The closest point on or inside the box.
 	 */
-	FORCEINLINE TVector2<T> GetClosestPointTo( const TVector2<T>& Point ) const;
+	[[nodiscard]] FORCEINLINE TVector2<T> GetClosestPointTo( const TVector2<T>& Point ) const;
 
 	/**
 	 * Gets the box extents around the center.
@@ -282,7 +301,7 @@ public:
 	 * @return Box extents.
 	 * @see GetArea, GetCenter, GetCenterAndExtents, GetSize
 	 */
-	TVector2<T> GetExtent() const
+	[[nodiscard]] TVector2<T> GetExtent() const
 	{
 		return 0.5f * TVector2<T>(Max - Min);
 	}
@@ -294,7 +313,7 @@ public:
 	 * @return Box size.
 	 * @see GetArea, GetCenter, GetCenterAndExtents, GetExtent
 	 */
-	TVector2<T> GetSize() const
+	[[nodiscard]] TVector2<T> GetSize() const
 	{
 		return TVector2<T>(Max - Min);
 	}
@@ -314,7 +333,7 @@ public:
 	 * @param Other The bounding box to test overlap
 	 * @return the overlap box. Result will be invalid if they don't overlap
 	 */
-	TBox2<T> Overlap( const TBox2<T>& Other ) const;
+	[[nodiscard]] TBox2<T> Overlap( const TBox2<T>& Other ) const;
 
 	/**
 	 * Checks whether the given box intersects this box.
@@ -322,7 +341,7 @@ public:
 	 * @param other bounding box to test intersection
 	 * @return true if boxes intersect, false otherwise.
 	 */
-	FORCEINLINE bool Intersect( const TBox2<T> & other ) const;
+	[[nodiscard]] FORCEINLINE bool Intersect( const TBox2<T> & other ) const;
 
 	/**
 	 * Checks whether the given point is inside this box.
@@ -330,7 +349,7 @@ public:
 	 * @param Point The point to test.
 	 * @return true if the point is inside this box, otherwise false.
 	 */
-	bool IsInside( const TVector2<T> & TestPoint ) const
+	[[nodiscard]] bool IsInside( const TVector2<T> & TestPoint ) const
 	{
 		return ((TestPoint.X > Min.X) && (TestPoint.X < Max.X) && (TestPoint.Y > Min.Y) && (TestPoint.Y < Max.Y));
 	}
@@ -342,7 +361,7 @@ public:
 	 * @return true if point is inside or on this box, otherwise false.
 	 * @see IsInside
 	 */
-	bool IsInsideOrOn( const TVector2<T>& TestPoint ) const
+	[[nodiscard]] bool IsInsideOrOn( const TVector2<T>& TestPoint ) const
 	{
 		return ((TestPoint.X >= Min.X) && (TestPoint.X <= Max.X) && (TestPoint.Y >= Min.Y) && (TestPoint.Y <= Max.Y));
 	}
@@ -353,7 +372,7 @@ public:
 	 * @param Other The box to test for encapsulation within the bounding volume.
 	 * @return true if box is inside this volume, false otherwise.
 	 */
-	bool IsInside( const TBox2<T>& Other ) const
+	[[nodiscard]] bool IsInside( const TBox2<T>& Other ) const
 	{
 		return (IsInside(Other.Min) && IsInside(Other.Max));
 	}
@@ -364,7 +383,7 @@ public:
 	 * @param The offset vector to shift by.
 	 * @return A new shifted bounding box.
 	 */
-	TBox2<T> ShiftBy( const TVector2<T>& Offset ) const
+	[[nodiscard]] TBox2<T> ShiftBy( const TVector2<T>& Offset ) const
 	{
 		return TBox2<T>(Min + Offset, Max + Offset);
 	}
@@ -375,7 +394,7 @@ public:
 	 * @param Destination The destination point to move center of box to.
 	 * @return A new bounding box.
 	 */
-	TBox2<T> MoveTo(const TVector2<T>& Destination) const
+	[[nodiscard]] TBox2<T> MoveTo(const TVector2<T>& Destination) const
 	{
 		const TVector2<T> Offset = Destination - GetCenter();
 		return TBox2<T>(Min + Offset, Max + Offset);
@@ -386,7 +405,7 @@ public:
 	 *
 	 * @return A string describing the box.
 	 */
-	FString ToString() const;
+	[[nodiscard]] FString ToString() const;
 
 	/** 
      * Utility function to build an AABB from Origin and Extent 
@@ -395,7 +414,7 @@ public:
      * @param Extent Half size of the bounding box.
      * @return A new axis-aligned bounding box.
      */
-    static TBox2<T> BuildAABB( const TVector2<T>& Origin, const TVector2<T>& Extent )
+    [[nodiscard]] static TBox2<T> BuildAABB( const TVector2<T>& Origin, const TVector2<T>& Extent )
     {
     	TBox2<T> NewBox(Origin - Extent, Origin + Extent);
 
diff --git a/Engine/Source/Runtime/Core/Public/Math/BoxSphereBounds.h b/Engine/Source/Runtime/Core/Public/Math/BoxSphereBounds.h
--- a/Engine/Source/Runtime/Core/Public/Math/BoxSphereBounds.h
+++ b/Engine/Source/Runtime/Core/Public/Math/BoxSphereBounds.h
@@ -37,14 +37,14 @@ struct TBoxSphereBounds
 public:
 
 	/** Default constructor. */
-	TBoxSphereBounds() = default;
+	[[nodiscard]] TBoxSphereBounds() = default;
 
 	/**
 	 * Creates and initializes a new instance.
 	 *
 	 * @param EForceInit Force Init Enum.
 	 */
-	explicit FORCEINLINE TBoxSphereBounds( EForceInit )
+	[[nodiscard]] explicit FORCEINLINE TBoxSphereBounds( EForceInit )
 		: Origin(ForceInit)
 		, BoxExtent(ForceInit)
 		, SphereRadius(0)
@@ -59,7 +59,7 @@ public:
 	 * @param InBoxExtent half size of box.
 	 * @param InSphereRadius radius of the sphere.
 	 */
-	TBoxSphereBounds( const TVector<T>& InOrigin, const TVector<TExtent>& InBoxExtent, TExtent InSphereRadius )
+	[[nodiscard]] TBoxSphereBounds( const TVector<T>& InOrigin, const TVector<TExtent>& InBoxExtent, TExtent InSphereRadius )
 		: Origin(InOrigin)
 		, BoxExtent(InBoxExtent)
 		, SphereRadius(InSphereRadius)
@@ -73,7 +73,7 @@ public:
 	 * @param Box The bounding box.
 	 * @param Sphere The bounding sphere.
 	 */
-	TBoxSphereBounds( const TBox<T>& Box, const TSphere<T>& Sphere )
+	[[nodiscard]] TBoxSphereBounds( const TBox<T>& Box, const TSphere<T>& Sphere )
 	{
 		TVector<T> LocalExtent;
 		Box.GetCenterAndExtents(Origin, LocalExtent);
@@ -91,7 +91,7 @@ public:
 	 *
 	 * @param Box The bounding box.
 	 */
-	TBoxSphereBounds( const TBox<T>& Box )
+	[[nodiscard]] TBoxSphereBounds( const TBox<T>& Box )
 	{
 		TVector<T> LocalExtent;
 		Box.GetCenterAndExtents(Origin, LocalExtent);
@@ -105,7 +105,7 @@ public:
 	/**
 	 * Creates and initializes a new instance for the given sphere.
 	 */
-	TBoxSphereBounds( const TSphere<T>& Sphere )
+	[[nodiscard]] TBoxSphereBounds( const TSphere<T>& Sphere )
 	{
 		Origin = Sphere.Center;
 		SphereRadius = Sphere.W;
@@ -122,7 +122,7 @@ public:
 	 * @param Points The points to be considered for the bounding box.
 	 * @param NumPoints Number of points in the Points array.
 	 */
-	TBoxSphereBounds( const TVector<T>* Points, uint32 NumPoints );
+	[[nodiscard]] TBoxSphereBounds( const TVector<T>* Points, uint32 NumPoints );
 
 	// Conversion to other type.
 	template<typename TFrom, typename TExtentFrom UE_REQUIRES(!(std::is_same_v<T, TFrom> && std::is_same_v<TExtent, TExtentFrom>))>
@@ -142,7 +142,7 @@ public:
 	 * @param Other The other bounding volume.
 	 * @return The combined bounding volume.
 	 */
-	FORCEINLINE TBoxSphereBounds<T, TExtent> operator+( const TBoxSphereBounds<T, TExtent>& Other ) const;
+	[[nodiscard]] FORCEINLINE TBoxSphereBounds<T, TExtent> operator+( const TBoxSphereBounds<T, TExtent>& Other ) const;
 
 	/**
 	 * Compare bounding volume this and Other.
@@ -150,7 +150,7 @@ public:
 	 * @param Other The other bounding volume.
 	 * @return true of they match.
 	 */
-	FORCEINLINE bool operator==(const TBoxSphereBounds<T, TExtent>& Other) const;
+	[[nodiscard]] FORCEINLINE bool operator==(const TBoxSphereBounds<T, TExtent>& Other) const;
 	
 	/**
 	 * Compare bounding volume this and Other.
@@ -158,7 +158,7 @@ public:
 	 * @param Other The other bounding volume.
 	 * @return true of they do not match.
 	 */	
-	FORCEINLINE bool operator!=(const TBoxSphereBounds<T, TExtent>& Other) const;
+	[[nodiscard]] FORCEINLINE bool operator!=(const TBoxSphereBounds<T, TExtent>& Other) const;
 
 public:
 
@@ -168,7 +168,7 @@ public:
 	 * @param Point The point.
 	 * @return The distance.
 	 */
-	FORCEINLINE T ComputeSquaredDistanceFromBoxToPoint( const TVector<T>& Point ) const
+	[[nodiscard]] FORCEINLINE T ComputeSquaredDistanceFromBoxToPoint( const TVector<T>& Point ) const
 	{
 		TVector<T> Mins = Origin - BoxExtent;
 		TVector<T> Maxs = Origin + BoxExtent;
@@ -184,7 +184,7 @@ public:
 	 * @param  Tolerance Error tolerance added to test distance.
 	 * @return true if spheres intersect, false otherwise.
 	 */
-	FORCEINLINE static bool SpheresIntersect(const TBoxSphereBounds<T, TExtent>& A, const TBoxSphereBounds<T, TExtent>& B, TExtent Tolerance = UE_KINDA_SMALL_NUMBER)
+	[[nodiscard]] FORCEINLINE static bool SpheresIntersect(const TBoxSphereBounds<T, TExtent>& A, const TBoxSphereBounds<T, TExtent>& B, TExtent Tolerance = UE_KINDA_SMALL_NUMBER)
 	{
 		return (A.Origin - B.Origin).SizeSquared() <= FMath::Square(FMath::Max<TExtent>(0, A.SphereRadius + B.SphereRadius + Tolerance));
 	}
@@ -196,7 +196,7 @@ public:
 	 * @param  B Second BoxSphereBounds to test.
 	 * @return true if boxes intersect, false otherwise.
 	 */
-	FORCEINLINE static bool BoxesIntersect(const TBoxSphereBounds<T, TExtent>& A, const TBoxSphereBounds<T, TExtent>& B)
+	[[nodiscard]] FORCEINLINE static bool BoxesIntersect(const TBoxSphereBounds<T, TExtent>& A, const TBoxSphereBounds<T, TExtent>& B)
 	{
 		return A.GetBox().Intersect(B.GetBox());
 	}
@@ -206,7 +206,7 @@ public:
 	 *
 	 * @return The bounding box.
 	 */
-	FORCEINLINE TBox<T> GetBox() const
+	[[nodiscard]] FORCEINLINE TBox<T> GetBox() const
 	{
 		return TBox<T>(Origin - BoxExtent,Origin + BoxExtent);
 	}
@@ -217,7 +217,7 @@ public:
 	 * @param Extrema 1 for positive extrema from the origin, else negative
 	 * @return The boxes extrema
 	 */
-	TVector<T> GetBoxExtrema( uint32 Extrema ) const
+	[[nodiscard]] TVector<T> GetBoxExtrema( uint32 Extrema ) const
 	{
 		if (Extrema)
 		{
@@ -232,7 +232,7 @@ public:
 	 *
 	 * @return The bounding sphere.
 	 */
-	FORCEINLINE TSphere<T> GetSphere() const
+	[[nodiscard]] FORCEINLINE TSphere<T> GetSphere() const
 	{
 		return TSphere<T>(Origin,SphereRadius);
 	}
@@ -243,7 +243,7 @@ public:
 	 * @param ExpandAmount The size to increase by.
 	 * @return A new box with the expanded size.
 	 */
-	FORCEINLINE TBoxSphereBounds<T, TExtent> ExpandBy(TExtent ExpandAmount ) const
+	[[nodiscard]] FORCEINLINE TBoxSphereBounds<T, TExtent> ExpandBy(TExtent ExpandAmount ) const
 	{
 		return TBoxSphereBounds(Origin, BoxExtent + ExpandAmount, SphereRadius + ExpandAmount);
 	}
@@ -254,7 +254,7 @@ public:
 	 * @param M The matrix.
 	 * @return The transformed volume.
 	 */
-	TBoxSphereBounds<T, TExtent> TransformBy( const TMatrix<T>& M ) const;
+	[[nodiscard]] TBoxSphereBounds<T, TExtent> TransformBy( const TMatrix<T>& M ) const;
 
 	/**
 	 * Gets a bounding volume transformed by a FTransform object.
@@ -262,21 +262,21 @@ public:
 	 * @param M The FTransform object.
 	 * @return The transformed volume.
 	 */
-	TBoxSphereBounds<T, TExtent> TransformBy( const TTransform<T>& M ) const;
+	[[nodiscard]] TBoxSphereBounds<T, TExtent> TransformBy( const TTransform<T>& M ) const;
 
 	/**
 	 * Get a textual representation of this bounding box.
 	 *
 	 * @return Text describing the bounding box.
 	 */
-	FString ToString() const;
+	[[nodiscard]] FString ToString() const;
 
 	/**
 	 * Constructs a bounding volume containing both A and B.
 	 *
 	 * This is a legacy version of the function used to compute primitive bounds, to avoid the need to rebuild lighting after the change.
 	 */
-	friend TBoxSphereBounds<T, TExtent> Union( const TBoxSphereBounds<T, TExtent>& A,const TBoxSphereBounds<T, TExtent>& B )
+	[[nodiscard]] friend TBoxSphereBounds<T, TExtent> Union( const TBoxSphereBounds<T, TExtent>& A,const TBoxSphereBounds<T, TExtent>& B )
 	{
 		return A + B;
 	}
@@ -304,7 +304,7 @@ public:
 	FORCEINLINE void DiagnosticCheckNaN() const {}
 #endif
 
-	inline bool ContainsNaN() const
+	[[nodiscard]] inline bool ContainsNaN() const
 	{
 		return Origin.ContainsNaN() || BoxExtent.ContainsNaN() || !FMath::IsFinite(SphereRadius);
 	}
diff --git a/Engine/Source/Runtime/Core/Public/Math/CapsuleShape.h b/Engine/Source/Runtime/Core/Public/Math/CapsuleShape.h
--- a/Engine/Source/Runtime/Core/Public/Math/CapsuleShape.h
+++ b/Engine/Source/Runtime/Core/Public/Math/CapsuleShape.h
@@ -35,7 +35,7 @@ struct TCapsuleShape
 public:
 
 	/** Default constructor. */
-	TCapsuleShape() { }
+	[[nodiscard]] TCapsuleShape() { }
 
 	/**
 	 * Create and inintialize a new instance.
@@ -45,7 +45,7 @@ public:
 	 * @param InOrientation The capsule's orientation in space.
 	 * @param InLength The capsule's length.
 	 */
-	TCapsuleShape(TVector<T> InCenter, T InRadius, TVector<T> InOrientation, T InLength)
+	[[nodiscard]] TCapsuleShape(TVector<T> InCenter, T InRadius, TVector<T> InOrientation, T InLength)
 		: Center(InCenter)
 		, Radius(InRadius)
 		, Orientation(InOrientation)
@@ -54,12 +54,12 @@ public:
 
 	// Conversion to other type.
 	template<typename FArg UE_REQUIRES(!std::is_same_v<T, FArg>)>
-	explicit TCapsuleShape(const TCapsuleShape<FArg>& From)
+	[[nodiscard]] explicit TCapsuleShape(const TCapsuleShape<FArg>& From)
 		: TCapsuleShape<T>(TVector<T>(From.Center), (T)From.Radius, TVector<T>(From.Orientation), (T)From.Length)
 	{
 	}
 
-	bool IsInside(const TVector<T>& In, T Tolerance = UE_KINDA_SMALL_NUMBER) const
+	[[nodiscard]] bool IsInside(const TVector<T>& In, T Tolerance = UE_KINDA_SMALL_NUMBER) const
 	{
 		const TVector<T> HalfLengthAxis = 0.5 * Length * Orientation;
 		const TVector<T> SegmentStart = Center - HalfLengthAxis;
diff --git a/Engine/Source/Runtime/Core/Public/Math/ColorList.h b/Engine/Source/Runtime/Core/Public/Math/ColorList.h
--- a/Engine/Source/Runtime/Core/Public/Math/ColorList.h
+++ b/Engine/Source/Runtime/Core/Public/Math/ColorList.h
@@ -123,25 +123,25 @@ public:
 	CORE_API void CreateColorMap();
 
 	/** Returns a color based on ColorName. If not found, returns White. */
-	CORE_API const FColor& GetFColorByName( const TCHAR* ColorName ) const;
+	[[nodiscard]] CORE_API const FColor& GetFColorByName( const TCHAR* ColorName ) const;
 
 	/** Returns a linear color based on ColorName. If not found, returns White. */
-	CORE_API const FLinearColor GetFLinearColorByName( const TCHAR* ColorName ) const;
+	[[nodiscard]] CORE_API const FLinearColor GetFLinearColorByName( const TCHAR* ColorName ) const;
 
 	/** Returns true if color is valid common colors. If not found returns false. */
-	CORE_API bool IsValidColorName( const TCHAR* ColorName ) const;
+	[[nodiscard]] CORE_API bool IsValidColorName( const TCHAR* ColorName ) const;
 
 	/** Returns index of color based on ColorName. If not found returns 0. */
-	CORE_API int32 GetColorIndex( const TCHAR* ColorName ) const;
+	[[nodiscard]] CORE_API int32 GetColorIndex( const TCHAR* ColorName ) const;
 
 	/** Returns a color based on index. If index is invalid, returns White. */
-	CORE_API const FColor& GetFColorByIndex( int32 ColorIndex ) const;
+	[[nodiscard]] CORE_API const FColor& GetFColorByIndex( int32 ColorIndex ) const;
 
 	/** Returns color's name based on index. If index is invalid, returns BadIndex. */
-	CORE_API const FString& GetColorNameByIndex( int32 ColorIndex ) const;
+	[[nodiscard]] CORE_API const FString& GetColorNameByIndex( int32 ColorIndex ) const;
 
 	/** Returns the number of colors. */
-	int32 GetColorsNum() const
+	[[nodiscard]] int32 GetColorsNum() const
 	{
 		return ColorsMap.Num();
 	}
diff --git a/Engine/Source/Runtime/Core/Public/Math/ConvexHull2d.h b/Engine/Source/Runtime/Core/Public/Math/ConvexHull2d.h
--- a/Engine/Source/Runtime/Core/Public/Math/ConvexHull2d.h
+++ b/Engine/Source/Runtime/Core/Public/Math/ConvexHull2d.h
@@ -90,7 +90,7 @@ namespace ConvexHull2D
 	}
 
 	/** Returns <0 if C is left of A-B */
-	inline FVector::FReal ComputeDeterminant(const FVector& A, const FVector& B, const FVector& C)
+	[[nodiscard]] inline FVector::FReal ComputeDeterminant(const FVector& A, const FVector& B, const FVector& C)
 	{
 		const FVector::FReal u1 = B.X - A.X;
 		const FVector::FReal v1 = B.Y - A.Y;
@@ -101,7 +101,7 @@ namespace ConvexHull2D
 	}
 
 	/** Returns true if 'a' is more lower-left than 'b'. */
-	inline bool ComparePoints(const FVector& A, const FVector& B)
+	[[nodiscard]] inline bool ComparePoints(const FVector& A, const FVector& B)
 	{
 		if (A.X < B.X)
 		{
@@ -246,7 +246,7 @@ namespace ConvexHull2D
 	}
 
 	/** Returns <0 if C is left of A-B */
-	inline FVector2D::FReal ComputeDeterminant2D(const FVector2D& A, const FVector2D& B, const FVector2D& C)
+	[[nodiscard]] inline FVector2D::FReal ComputeDeterminant2D(const FVector2D& A, const FVector2D& B, const FVector2D& C)
 	{
 		const FVector2D::FReal u1 = B.X - A.X;
 		const FVector2D::FReal v1 = B.Y - A.Y;
diff --git a/Engine/Source/Runtime/Core/Public/Math/DoubleFloat.h b/Engine/Source/Runtime/Core/Public/Math/DoubleFloat.h
--- a/Engine/Source/Runtime/Core/Public/Math/DoubleFloat.h
+++ b/Engine/Source/Runtime/Core/Public/Math/DoubleFloat.h
@@ -19,7 +19,7 @@ struct FDFScalar
 		return static_cast<double>(High) + static_cast<double>(Low);
 	}
 
-	FDFScalar(double In)
+	[[nodiscard]] FDFScalar(double In)
 	{
 		High = static_cast<float>(In);
 		Low = static_cast<float>(In - High);
@@ -37,13 +37,13 @@ struct FDFScalar
 
 struct FDFVector2
 {
-	FVector2d GetVector2d() const
+	[[nodiscard]] FVector2d GetVector2d() const
 	{
 		return FVector2d{ X().GetDouble(), Y().GetDouble() };
 	}
 
 	template<typename TInputScalar = double>
-	FDFVector2(const UE::Math::TVector2<TInputScalar>& In)
+	[[nodiscard]] FDFVector2(const UE::Math::TVector2<TInputScalar>& In)
 	{
 		FDFScalar X(In.X);
 		FDFScalar Y(In.Y);
@@ -51,18 +51,18 @@ struct FDFVector2
 		Low = { X.Low, Y.Low };
 	}
 
-	FDFVector2(FVector2f High, FVector2f Low)
+	[[nodiscard]] FDFVector2(FVector2f High, FVector2f Low)
 		: High(High), Low(Low)
 	{}
 
-	FDFVector2() = default;
+	[[nodiscard]] FDFVector2() = default;
 
-	FDFScalar X() const
+	[[nodiscard]] FDFScalar X() const
 	{
 		return { High.X, Low.X };
 	}
 
-	FDFScalar Y() const
+	[[nodiscard]] FDFScalar Y() const
 	{
 		return { High.Y, Low.Y };
 	}
@@ -73,7 +73,7 @@ struct FDFVector2
 
 struct FDFVector3
 {
-	FVector3d GetVector3d() const
+	[[nodiscard]] FVector3d GetVector3d() const
 	{
 		return FVector3d{ X().GetDouble(), Y().GetDouble(), Z().GetDouble() };
 	}
@@ -94,17 +94,17 @@ struct FDFVector3
 
 	FDFVector3() = default;
 
-	FDFScalar X() const
+	[[nodiscard]] FDFScalar X() const
 	{
 		return { High.X, Low.X };
 	}
 
-	FDFScalar Y() const
+	[[nodiscard]] FDFScalar Y() const
 	{
 		return { High.Y, Low.Y };
 	}
 
-	FDFScalar Z() const
+	[[nodiscard]] FDFScalar Z() const
 	{
 		return { High.Z, Low.Z };
 	}
@@ -115,7 +115,7 @@ struct FDFVector3
 
 struct FDFVector4
 {
-	FVector4d GetVector4d() const
+	[[nodiscard]] FVector4d GetVector4d() const
 	{
 		return FVector4d{ X().GetDouble(), Y().GetDouble(), Z().GetDouble(), W().GetDouble() };
 	}
@@ -148,22 +148,22 @@ struct FDFVector4
 
 	FDFVector4() = default;
 
-	FDFScalar X() const
+	[[nodiscard]] FDFScalar X() const
 	{
 		return { High.X, Low.X };
 	}
 
-	FDFScalar Y() const
+	[[nodiscard]] FDFScalar Y() const
 	{
 		return { High.Y, Low.Y };
 	}
 
-	FDFScalar Z() const
+	[[nodiscard]] FDFScalar Z() const
 	{
 		return { High.Z, Low.Z };
 	}
 
-	FDFScalar W() const
+	[[nodiscard]] FDFScalar W() const
 	{
 		return { High.W, Low.W };
 	}
@@ -185,14 +185,14 @@ struct FDFMatrix
 	FVector3f PostTranslation{}; // Added to result, *after* multiplying 'M'
 
 	// Check if the origin of the matrix is small enough to ensure the precision >= UE_DF_MIN_PRECISION
-	CORE_API static FMatrix44f SafeCastMatrix(const FMatrix& Matrix);
+	[[nodiscard]] CORE_API static FMatrix44f SafeCastMatrix(const FMatrix& Matrix);
 
 	// Apply post-translation to matrix
-	CORE_API static FDFMatrix MakeToRelativeWorldMatrix(const FVector3f Origin, const FMatrix& ToWorld);
-	CORE_API static FMatrix   MakeToRelativeWorldMatrixDouble(const FVector Origin, const FMatrix& ToWorld);
+	[[nodiscard]] CORE_API static FDFMatrix MakeToRelativeWorldMatrix(const FVector3f Origin, const FMatrix& ToWorld);
+	[[nodiscard]] CORE_API static FMatrix   MakeToRelativeWorldMatrixDouble(const FVector Origin, const FMatrix& ToWorld);
 
-	CORE_API static FDFMatrix MakeClampedToRelativeWorldMatrix(const FVector3f Origin, const FMatrix& ToWorld);
-	CORE_API static FMatrix   MakeClampedToRelativeWorldMatrixDouble(const FVector Origin, const FMatrix& ToWorld);
+	[[nodiscard]] CORE_API static FDFMatrix MakeClampedToRelativeWorldMatrix(const FVector3f Origin, const FMatrix& ToWorld);
+	[[nodiscard]] CORE_API static FMatrix   MakeClampedToRelativeWorldMatrixDouble(const FVector Origin, const FMatrix& ToWorld);
 };
 
 // Transforms *from* absolute world space
@@ -208,6 +208,6 @@ struct FDFInverseMatrix
 	FVector3f PreTranslation{}; // Subtracted from input position *before* multiplying 'M'
 
 	// Apply pre-translation to matrix
-	CORE_API static FDFInverseMatrix MakeFromRelativeWorldMatrix(const FVector3f Origin, const FMatrix& FromWorld);
-	CORE_API static FMatrix          MakeFromRelativeWorldMatrixDouble(const FVector Origin, const FMatrix& FromWorld);
+	[[nodiscard]] CORE_API static FDFInverseMatrix MakeFromRelativeWorldMatrix(const FVector3f Origin, const FMatrix& FromWorld);
+	[[nodiscard]] CORE_API static FMatrix          MakeFromRelativeWorldMatrixDouble(const FVector Origin, const FMatrix& FromWorld);
 };
diff --git a/Engine/Source/Runtime/Core/Public/Math/DualQuat.h b/Engine/Source/Runtime/Core/Public/Math/DualQuat.h
--- a/Engine/Source/Runtime/Core/Public/Math/DualQuat.h
+++ b/Engine/Source/Runtime/Core/Public/Math/DualQuat.h
@@ -26,12 +26,12 @@ public:
 	TQuat<T> D;
 
 	// Constructors
-	TDualQuat(const TQuat<T> &InR, const TQuat<T> &InD)
+	[[nodiscard]] TDualQuat(const TQuat<T> &InR, const TQuat<T> &InD)
 		: R(InR)
 		, D(InD)
 	{}
 
-	TDualQuat(const TTransform<T> &Transform)
+	[[nodiscard]] TDualQuat(const TTransform<T> &Transform)
 	{
 		TVector<T> V = Transform.GetTranslation()*0.5f;
 		*this = TDualQuat<T>(TQuat<T>(0, 0, 0, 1), TQuat<T>(V.X, V.Y, V.Z, 0.f)) * TDualQuat<T>(Transform.GetRotation(), TQuat<T>(0, 0, 0, 0));
@@ -39,38 +39,38 @@ public:
 
 	// Conversion to other type.
 	template<typename FArg UE_REQUIRES(!std::is_same_v<T, FArg>)>
-	explicit TDualQuat(const TDualQuat<FArg>& From)
+	[[nodiscard]] explicit TDualQuat(const TDualQuat<FArg>& From)
 		: TDualQuat<T>(TQuat<T>(From.R), TQuat<T>(From.D))
 	{
 	}
 
 	/** Dual quat addition */
-	TDualQuat<T> operator+(const TDualQuat<T> &B) const
+	[[nodiscard]] TDualQuat<T> operator+(const TDualQuat<T> &B) const
 	{
 		return{ R + B.R, D + B.D };
 	}
 
 	/** Dual quat product */
-	TDualQuat<T> operator*(const TDualQuat<T> &B) const
+	[[nodiscard]] TDualQuat<T> operator*(const TDualQuat<T> &B) const
 	{
 		return{ R*B.R, D*B.R + B.D*R };
 	}
 
 	/** Scale dual quat */
-	TDualQuat<T> operator*(const T S) const
+	[[nodiscard]] TDualQuat<T> operator*(const T S) const
 	{
 		return{ R*S, D*S };
 	}
 
 	/** Return normalized dual quat */
-	TDualQuat<T> Normalized() const
+	[[nodiscard]] TDualQuat<T> Normalized() const
 	{
 		T MinV = 1.0f / FMath::Sqrt(R | R);
 		return{ R*MinV, D*MinV };
 	}
 
 	/** Convert dual quat to transform */
-	TTransform<T> AsFTransform(TVector<T> Scale = TVector<T>(1.0f, 1.0f, 1.0f))
+	[[nodiscard]] TTransform<T> AsFTransform(TVector<T> Scale = TVector<T>(1.0f, 1.0f, 1.0f)) const
 	{
 		TQuat<T> TQ = D*TQuat<T>(-R.X, -R.Y, -R.Z, R.W);
 		return TTransform<T>(R, TVector<T>(TQ.X, TQ.Y, TQ.Z)*2.0f, Scale);
diff --git a/Engine/Source/Runtime/Core/Public/Math/Edge.h b/Engine/Source/Runtime/Core/Public/Math/Edge.h
--- a/Engine/Source/Runtime/Core/Public/Math/Edge.h
+++ b/Engine/Source/Runtime/Core/Public/Math/Edge.h
@@ -19,7 +19,7 @@ struct FEdge
 public:
 
 	/** Default constructor (no initialization). */
-	FEdge () { }
+	[[nodiscard]] FEdge() = default;
 
 	/**
 	 * Creates and initializes a new edge from two vertices.
@@ -27,7 +27,7 @@ public:
 	 * @param V1 The first vertex.
 	 * @param V2 The second vertex.
 	 */
-	FEdge (FVector V1, FVector V2)
+	[[nodiscard]] FEdge(FVector V1, FVector V2)
 	{
 		Vertex[0] = V1;
 		Vertex[1] = V2;
@@ -42,7 +42,7 @@ public:
 	 * @param E The other edge.
 	 * @return true if the two edges are identical, false otherwise.
 	 */
-	bool operator== (const FEdge& E) const
+	[[nodiscard]] bool operator==(const FEdge& E) const
 	{
 		return (((E.Vertex[0] == Vertex[0]) && (E.Vertex[1] == Vertex[1])) || ((E.Vertex[0] == Vertex[1]) && (E.Vertex[1] == Vertex[0])));
 	}
diff --git a/Engine/Source/Runtime/Core/Public/Math/Float16.h b/Engine/Source/Runtime/Core/Public/Math/Float16.h
--- a/Engine/Source/Runtime/Core/Public/Math/Float16.h
+++ b/Engine/Source/Runtime/Core/Public/Math/Float16.h
@@ -40,13 +40,13 @@ public:
 	uint16 Encoded;
 
 	/** Default constructor */
-	FFloat16();
+	[[nodiscard]] FFloat16();
 
 	/** Copy constructor. */
-	FFloat16(const FFloat16& FP16Value);
+	[[nodiscard]] FFloat16(const FFloat16& FP16Value);
 
 	/** Conversion constructor. Convert from Fp32 to Fp16. */
-	FFloat16(float FP32Value);	
+	[[nodiscard]] FFloat16(float FP32Value);	
 
 	/** Assignment operator. Convert from Fp32 to Fp16. */
 	FFloat16& operator=(float FP32Value);
@@ -86,17 +86,17 @@ public:
 	}
 
 	/** Return float clamp in [0,MaxF16Float] , no negatives or infinites or nans returned **/
-	FFloat16 GetClampedNonNegativeAndFinite() const;
+	[[nodiscard]] FFloat16 GetClampedNonNegativeAndFinite() const;
 	
 	/** Return float clamp in [-MaxF16Float,MaxF16Float] , no infinites or nans returned **/
-	FFloat16 GetClampedFinite() const;
+	[[nodiscard]] FFloat16 GetClampedFinite() const;
 
 	/** Convert from Fp16 to Fp32. */
-	float GetFloat() const;
+	[[nodiscard]] float GetFloat() const;
 
 	/** Is the float negative without converting
 	NOTE: returns true for negative zero! */
-	bool IsNegative() const
+	[[nodiscard]] bool IsNegative() const
 	{
 		// negative if sign bit is on
 		// can be tested with int compare
diff --git a/Engine/Source/Runtime/Core/Public/Math/Float16Color.h b/Engine/Source/Runtime/Core/Public/Math/Float16Color.h
--- a/Engine/Source/Runtime/Core/Public/Math/Float16Color.h
+++ b/Engine/Source/Runtime/Core/Public/Math/Float16Color.h
@@ -19,19 +19,25 @@ public:
 	FFloat16 A;
 
 	/* Get as a pointer to four half floats */
-	uint16 * GetFourHalves() { return (uint16 *)this; }
-	const uint16 * GetFourHalves() const { return (const uint16 *)this; }
+	[[nodiscard]] uint16 * GetFourHalves()
+	{
+		return (uint16 *)this;
+	}
+	[[nodiscard]] const uint16 * GetFourHalves() const
+	{
+		return (const uint16 *)this;
+	}
 
-	const FLinearColor GetFloats() const;
+	[[nodiscard]] const FLinearColor GetFloats() const;
 
 	/** Default constructor */
-	FFloat16Color();
+	[[nodiscard]] FFloat16Color();
 
 	/** Copy constructor. */
-	FFloat16Color(const FFloat16Color& Src);
+	[[nodiscard]] FFloat16Color(const FFloat16Color& Src);
 
 	/** Constructor from a linear color. */
-	FFloat16Color(const FLinearColor& Src);
+	[[nodiscard]] FFloat16Color(const FLinearColor& Src);
 
 	/** assignment operator */
 	FFloat16Color& operator=(const FFloat16Color& Src);
@@ -42,7 +48,7 @@ public:
 	 * @param Src The other color.
 	 * @return true if the two colors are identical, otherwise false.
 	 */
-	bool operator==(const FFloat16Color& Src) const;
+	[[nodiscard]] bool operator==(const FFloat16Color& Src) const;
 };
 
 
diff --git a/Engine/Source/Runtime/Core/Public/Math/FloatPacker.h b/Engine/Source/Runtime/Core/Public/Math/FloatPacker.h
--- a/Engine/Source/Runtime/Core/Public/Math/FloatPacker.h
+++ b/Engine/Source/Runtime/Core/Public/Math/FloatPacker.h
@@ -26,12 +26,12 @@ public:
 	typedef float	FloatType;
 	typedef uint32	PackedType;
 
-	static PackedType ToPackedType(FloatType Value)
+	[[nodiscard]] static PackedType ToPackedType(FloatType Value)
 	{
 		return *(PackedType*)&Value;
 	}
 
-	static FloatType ToFloatType(PackedType Value)
+	[[nodiscard]] static FloatType ToFloatType(PackedType Value)
 	{
 		return *(FloatType*)&Value;
 	}
@@ -107,7 +107,7 @@ public:
 		return (Sign << SignShift) | (Exponent << NumMantissaBits) | (Mantissa);
 	}
 
-	FloatType Decode(PackedType Value) const
+	[[nodiscard]] FloatType Decode(PackedType Value) const
 	{
 		if ( Value == (PackedType) 0 )
 		{
diff --git a/Engine/Source/Runtime/Core/Public/Math/GenericOctree.h b/Engine/Source/Runtime/Core/Public/Math/GenericOctree.h
--- a/Engine/Source/Runtime/Core/Public/Math/GenericOctree.h
+++ b/Engine/Source/Runtime/Core/Public/Math/GenericOctree.h
@@ -47,16 +47,16 @@ public:
 	FVector4 Extent;
 
 	/** Default constructor. */
-	FBoxCenterAndExtent() {}
+	[[nodiscard]] FBoxCenterAndExtent() {}
 
 	/** Initialization constructor. */
-	FBoxCenterAndExtent(const FVector& InCenter,const FVector& InExtent)
+	[[nodiscard]] FBoxCenterAndExtent(const FVector& InCenter,const FVector& InExtent)
 	:	Center(InCenter,0)
 	,	Extent(InExtent,0)
 	{}
 
 	/** FBox conversion constructor. */
-	FBoxCenterAndExtent(const FBox& Box)
+	[[nodiscard]] FBoxCenterAndExtent(const FBox& Box)
 	{
 		FVector C, E;
 		Box.GetCenterAndExtents(C,E);	// LWC_TODO: Perf pessimization
@@ -66,7 +66,7 @@ public:
 
 	/** FBoxSphereBounds conversion constructor. */
 	template<typename TExtent>
-	explicit FBoxCenterAndExtent(const UE::Math::TBoxSphereBounds<FReal, TExtent>& BoxSphere)
+	[[nodiscard]] explicit FBoxCenterAndExtent(const UE::Math::TBoxSphereBounds<FReal, TExtent>& BoxSphere)
 	{
 		Center = BoxSphere.Origin;
 		Extent = FVector(BoxSphere.BoxExtent);
@@ -74,14 +74,14 @@ public:
 	}
 
 	/** Center - radius as four contiguous floats conversion constructor. */
-	explicit FBoxCenterAndExtent(const float PositionRadius[4])
+	[[nodiscard]] explicit FBoxCenterAndExtent(const float PositionRadius[4])
 	{
 		Center = FVector4(PositionRadius[0],PositionRadius[1],PositionRadius[2], 0);
 		Extent = FVector4(PositionRadius[3],PositionRadius[3],PositionRadius[3], 0);
 	}
 
 	/** Converts to a FBox. */
-	FBox GetBox() const
+	[[nodiscard]] FBox GetBox() const
 	{
 		return FBox(Center - Extent,Center + Extent);
 	}
@@ -91,7 +91,7 @@ public:
 	 * Warning: this operates on the W of the bounds positions!
 	 * @return true if the boxes intersect, or false.
 	 */
-	friend FORCEINLINE bool Intersect(const FBoxCenterAndExtent& A,const FBoxCenterAndExtent& B)
+	[[nodiscard]] friend FORCEINLINE bool Intersect(const FBoxCenterAndExtent& A,const FBoxCenterAndExtent& B)
 	{
 		// CenterDifference is the vector between the centers of the bounding boxes.
 		const VectorRegister CenterDifference = VectorAbs(VectorSubtract(VectorLoadAligned(&A.Center),VectorLoadAligned(&B.Center)));
@@ -109,7 +109,7 @@ public:
 	 * @return true if the boxes intersect, or false.
 	 */
 	template<typename TExtent>
-	friend FORCEINLINE bool Intersect(const UE::Math::TBoxSphereBounds<FReal, TExtent>& A,const FBoxCenterAndExtent& B)
+	[[nodiscard]] friend FORCEINLINE bool Intersect(const UE::Math::TBoxSphereBounds<FReal, TExtent>& A,const FBoxCenterAndExtent& B)
 	{
 		// CenterDifference is the vector between the centers of the bounding boxes.
 		const VectorRegister CenterDifference = VectorAbs(VectorSubtract(VectorLoadFloat3_W0(&A.Origin),VectorLoadAligned(&B.Center)));
@@ -127,7 +127,7 @@ public:
 	 * @param A box given in center - radius form as four contiguous floats
 	 * @return true if the boxes intersect, or false.
 	 */
-	friend FORCEINLINE bool Intersect(const float A[4],const FBoxCenterAndExtent& B)
+	[[nodiscard]] friend FORCEINLINE bool Intersect(const float A[4],const FBoxCenterAndExtent& B)
 	{
 		// CenterDifference is the vector between the centers of the bounding boxes.
 		const VectorRegister CenterDifference = VectorAbs(VectorSubtract(VectorLoadFloat3_W0(A),VectorLoadAligned(&B.Center)));
@@ -170,7 +170,7 @@ public:
 	}
 
 	/** @return true if the reference isn't set. */
-	FORCEINLINE bool IsNULL() const
+	[[nodiscard]] FORCEINLINE bool IsNULL() const
 	{
 		return Index >= 8;
 	}
@@ -180,17 +180,17 @@ public:
 		Index = 8;
 	}
 
-	FORCEINLINE int32 X() const
+	[[nodiscard]] FORCEINLINE int32 X() const
 	{
 		return (Index >> 0) & 1;
 	}
 
-	FORCEINLINE int32 Y() const
+	[[nodiscard]] FORCEINLINE int32 Y() const
 	{
 		return (Index >> 1) & 1;
 	}
 
-	FORCEINLINE int32 Z() const
+	[[nodiscard]] FORCEINLINE int32 Z() const
 	{
 		return (Index >> 2) & 1;
 	}
@@ -244,7 +244,7 @@ public:
 	}
 
 	/** Determines whether the subset contains a specific node. */
-	bool Contains(FOctreeChildNodeRef ChildRef) const;
+	[[nodiscard]] bool Contains(FOctreeChildNodeRef ChildRef) const;
 };
 
 /** The context of an octree node, derived from the traversal of the tree. */
@@ -309,7 +309,7 @@ public:
 		ChildCenterOffset = Bounds.Extent.X - LooseChildExtent;
 	}
 
-	inline VectorRegister GetChildOffsetVec(int i) const
+	[[nodiscard]] inline VectorRegister GetChildOffsetVec(int i) const
 	{
 		// LWC_TODO: not sure this is correct for VectorRegister = VectorRegister4Double
 		union MaskType { VectorRegister4Float v;  VectorRegister4Int i; 
@@ -328,7 +328,7 @@ public:
 	}
 
 	/** Child node initialization constructor. */
-	inline FOctreeNodeContext GetChildContext(FOctreeChildNodeRef ChildRef) const
+	[[nodiscard]] inline FOctreeNodeContext GetChildContext(FOctreeChildNodeRef ChildRef) const
 	{
 		FBoxCenterAndExtent LocalBounds;
 		VectorRegister ZeroW = MakeVectorRegister(1.0f, 1.0f, 1.0f, 0.0f);
@@ -351,7 +351,7 @@ public:
 	}
 	
 	/** Child node initialization constructor. */
-	inline FOctreeNodeContext GetChildContext(FOctreeChildNodeRef ChildRef, uint32 InInCullBits, uint32 InOutCullBits) const
+	[[nodiscard]] inline FOctreeNodeContext GetChildContext(FOctreeChildNodeRef ChildRef, uint32 InInCullBits, uint32 InOutCullBits) const
 	{
 		FBoxCenterAndExtent LocalBounds;
 		VectorRegister ZeroW = MakeVectorRegister(1.0f, 1.0f, 1.0f, 0.0f);
@@ -364,7 +364,7 @@ public:
 	 * @param BoundingBox - The bounding box to check for intersection with.
 	 * @return A subset of the children's nodes that intersect the bounding box.
 	 */
-	FOctreeChildNodeSubset GetIntersectingChildren(const FBoxCenterAndExtent& BoundingBox) const;
+	[[nodiscard]] FOctreeChildNodeSubset GetIntersectingChildren(const FBoxCenterAndExtent& BoundingBox) const;
 
 	/**
 	 * Determines which of the octree node's children contain the whole bounding box, if any.
@@ -372,7 +372,7 @@ public:
 	 * @return The octree's node that the bounding box is farthest from the outside edge of, or an invalid node ref if it's not contained
 	 *			by any of the children.
 	 */
-	FOctreeChildNodeRef GetContainingChild(const FBoxCenterAndExtent& BoundingBox) const;
+	[[nodiscard]] FOctreeChildNodeRef GetContainingChild(const FBoxCenterAndExtent& BoundingBox) const;
 };
 
 CORE_API DECLARE_LOG_CATEGORY_EXTERN(LogGenericOctree, Log, All);
@@ -497,7 +497,7 @@ private:
 			FreeList.Push(FSpan{ INDEX_NONE, INDEX_NONE });
 		}
 
-		int Num() const
+		[[nodiscard]] int Num() const
 		{
 			//includes one dummy
 			return FreeList.Num() - 1;
@@ -725,7 +725,10 @@ private:
 	}
 public:
 
-	int32 GetNumNodes() const { return TreeNodes.Num(); }
+	[[nodiscard]] int32 GetNumNodes() const
+	{
+		return TreeNodes.Num();
+	}
 
 	/**
 	 * this function will call the passed in function for all elements in the Octree in node by node in no specified order.
@@ -890,13 +893,13 @@ public:
 	}
 
 	/** Accesses an octree element by ID. */
-	ElementType& GetElementById(FOctreeElementId2 ElementId)
+	[[nodiscard]] ElementType& GetElementById(FOctreeElementId2 ElementId)
 	{
 		return TreeElements[ElementId.NodeIndex][ElementId.ElementIndex];
 	}
 
 	/** Accesses an octree element by ID. */
-	const ElementType& GetElementById(FOctreeElementId2 ElementId) const
+	[[nodiscard]] const ElementType& GetElementById(FOctreeElementId2 ElementId) const
 	{
 		return TreeElements[ElementId.NodeIndex][ElementId.ElementIndex];
 	}
@@ -905,7 +908,7 @@ public:
 	 * check if a FOctreeElementId2 is valid.
 	 * @param ElementId - The ElementId to check.
 	 */
-	bool IsValidElementId(FOctreeElementId2 ElementId) const
+	[[nodiscard]] bool IsValidElementId(FOctreeElementId2 ElementId) const
 	{
 		return ElementId.IsValidId() && ElementId.ElementIndex < TreeElements[ElementId.NodeIndex].Num();
 	};
@@ -914,7 +917,7 @@ public:
 	 * return all elements for a given node.
 	 * @param NodeIndex - The the index of the node can be obtained using FindNodesWithPredicate.
 	 */
-	TArrayView<const ElementType> GetElementsForNode(FNodeIndex NodeIndex) const
+	[[nodiscard]] TArrayView<const ElementType> GetElementsForNode(FNodeIndex NodeIndex) const
 	{
 		return TreeElements[NodeIndex];
 	}
@@ -963,7 +966,7 @@ public:
 		}
 	}
 
-	SIZE_T GetSizeBytes() const
+	[[nodiscard]] SIZE_T GetSizeBytes() const
 	{
 		SIZE_T TotalSizeBytes = TreeNodes.GetAllocatedSize();
 		TotalSizeBytes += TreeElements.GetAllocatedSize();
@@ -971,13 +974,13 @@ public:
 		return TotalSizeBytes;
 	}
 
-	FReal GetNodeLevelExtent(int32 Level) const
+	[[nodiscard]] FReal GetNodeLevelExtent(int32 Level) const
 	{
 		const int32 ClampedLevel = FMath::Clamp<uint32>(Level, 0, OctreeSemantics::MaxNodeDepth);
 		return RootNodeContext.Bounds.Extent.X * FMath::Pow((1.0f + 1.0f / (FReal)FOctreeNodeContext::LoosenessDenominator) / 2.0f, FReal(ClampedLevel));
 	}
 
-	FBoxCenterAndExtent GetRootBounds() const
+	[[nodiscard]] FBoxCenterAndExtent GetRootBounds() const
 	{
 		return RootNodeContext.Bounds;
 	}
@@ -1107,25 +1110,31 @@ public:
 		}
 
 		// Accessors.
-		FORCEINLINE ElementConstIt GetElementIt() const { return ElementConstIt(Elements); }
-		FORCEINLINE bool IsLeaf() const { return bIsLeaf; }
-		FORCEINLINE bool HasChild(FOctreeChildNodeRef ChildRef) const
+		[[nodiscard]] FORCEINLINE ElementConstIt GetElementIt() const
+		{
+			return ElementConstIt(Elements);
+		}
+		[[nodiscard]] FORCEINLINE bool IsLeaf() const
+		{
+			return bIsLeaf;
+		}
+		[[nodiscard]] FORCEINLINE bool HasChild(FOctreeChildNodeRef ChildRef) const
 		{
 			return Children[ChildRef.Index] != NULL && Children[ChildRef.Index]->InclusiveNumElements > 0;
 		}
-		FORCEINLINE FNode* GetChild(FOctreeChildNodeRef ChildRef) const
+		[[nodiscard]] FORCEINLINE FNode* GetChild(FOctreeChildNodeRef ChildRef) const
 		{
 			return Children[ChildRef.Index];
 		}
-		FORCEINLINE int32 GetElementCount() const
+		[[nodiscard]] FORCEINLINE int32 GetElementCount() const
 		{
 			return Elements.Num();
 		}
-		FORCEINLINE int32 GetInclusiveElementCount() const
+		[[nodiscard]] FORCEINLINE int32 GetInclusiveElementCount() const
 		{
 			return InclusiveNumElements;
 		}
-		FORCEINLINE const ElementArrayType& GetElements() const
+		[[nodiscard]] FORCEINLINE const ElementArrayType& GetElements() const
 		{
 			return Elements;
 		}
@@ -1245,7 +1254,7 @@ public:
 		}
 
 		/** Checks if there are any nodes left to iterate over. */
-		bool HasPendingNodes() const
+		[[nodiscard]] bool HasPendingNodes() const
 		{
 			return CurrentNode.Node != NULL;
 		}
@@ -1261,11 +1270,11 @@ public:
 		{}
 
 		// Accessors.
-		const FNode& GetCurrentNode() const
+		[[nodiscard]] const FNode& GetCurrentNode() const
 		{
 			return *CurrentNode.Node;
 		}
-		const FOctreeNodeContext& GetCurrentContext() const
+		[[nodiscard]] const FOctreeNodeContext& GetCurrentContext() const
 		{
 			return CurrentNode.Context;
 		}
@@ -1293,7 +1302,7 @@ public:
 		}
 
 		/** Checks if there are any elements left to iterate over. */
-		bool HasPendingElements() const
+		[[nodiscard]] bool HasPendingElements() const
 		{
 			return NodeIt.HasPendingNodes();
 		}
@@ -1309,7 +1318,7 @@ public:
 		}
 
 		// Accessors.
-		const ElementType& GetCurrentElement() const
+		[[nodiscard]] const ElementType& GetCurrentElement() const
 		{
 			return *ElementIt;
 		}
@@ -1500,7 +1509,7 @@ public:
 	}
 
 	/** Accesses an octree element by ID. */
-	ElementType& GetElementById(FOctreeElementId ElementId)
+	[[nodiscard]] ElementType& GetElementById(FOctreeElementId ElementId)
 	{
 		check(ElementId.IsValidId());
 		FNode* ElementIdNode = (FNode*)ElementId.Node;
@@ -1508,7 +1517,7 @@ public:
 	}
 
 	/** Accesses an octree element by ID. */
-	const ElementType& GetElementById(FOctreeElementId ElementId) const
+	[[nodiscard]] const ElementType& GetElementById(FOctreeElementId ElementId) const
 	{
 		check(ElementId.IsValidId());
 		FNode* ElementIdNode = (FNode*)ElementId.Node;
@@ -1516,7 +1525,7 @@ public:
 	}
 
 	/** Checks if given ElementId represents a valid Octree element */
-	bool IsValidElementId(FOctreeElementId ElementId) const
+	[[nodiscard]] bool IsValidElementId(FOctreeElementId ElementId) const
 	{
 		return ElementId.IsValidId()
 			&& ElementId.ElementIndex != INDEX_NONE
@@ -1577,18 +1586,18 @@ public:
 	}
 
 
-	SIZE_T GetSizeBytes() const
+	[[nodiscard]] SIZE_T GetSizeBytes() const
 	{
 		return TotalSizeBytes;
 	}
 
-	FReal GetNodeLevelExtent(int32 Level) const
+	[[nodiscard]] FReal GetNodeLevelExtent(int32 Level) const
 	{
 		const int32 ClampedLevel = FMath::Clamp<uint32>(Level, 0, OctreeSemantics::MaxNodeDepth);
 		return RootNodeContext.Bounds.Extent.X * FMath::Pow((1.0f + 1.0f / (FReal)FOctreeNodeContext::LoosenessDenominator) / 2.0f, FReal(ClampedLevel));
 	}
 
-	FBoxCenterAndExtent GetRootBounds() const
+	[[nodiscard]] FBoxCenterAndExtent GetRootBounds() const
 	{
 		return RootNodeContext.Bounds;
 	}
diff --git a/Engine/Source/Runtime/Core/Public/Math/GenericOctreePublic.h b/Engine/Source/Runtime/Core/Public/Math/GenericOctreePublic.h
--- a/Engine/Source/Runtime/Core/Public/Math/GenericOctreePublic.h
+++ b/Engine/Source/Runtime/Core/Public/Math/GenericOctreePublic.h
@@ -19,18 +19,21 @@ public:
 	friend class TOctree2;
 
 	/** Default constructor. */
-	FOctreeElementId2()
+	[[nodiscard]] FOctreeElementId2()
 		:	NodeIndex(INDEX_NONE)
 		,	ElementIndex(INDEX_NONE)
 	{}
 
 	/** @return a boolean value representing whether the id is NULL. */
-	bool IsValidId() const
+	[[nodiscard]] bool IsValidId() const
 	{
 		return NodeIndex != INDEX_NONE;
 	}
 
-	uint32 GetNodeIndex() const { return NodeIndex; }
+	[[nodiscard]] uint32 GetNodeIndex() const
+	{
+		return NodeIndex;
+	}
 
 private:
 
@@ -47,12 +50,12 @@ private:
 	{}
 
 	/** Implicit conversion to the element index. */
-	operator int32() const
+	[[nodiscard]] operator int32() const
 	{
 		return ElementIndex;
 	}
 
-	friend FString LexToString(const FOctreeElementId2 ElementId)
+	[[nodiscard]] friend FString LexToString(const FOctreeElementId2 ElementId)
 	{
 		return FString::Printf(TEXT("[%s:%s]"), *LexToString(ElementId.NodeIndex), *LexToString(ElementId.ElementIndex));
 	}
@@ -75,7 +78,7 @@ public:
 	{}
 
 	/** @return a boolean value representing whether the id is NULL. */
-	bool IsValidId() const
+	[[nodiscard]] bool IsValidId() const
 	{
 		return Node != NULL;
 	}
@@ -95,7 +98,7 @@ private:
 	{}
 
 	/** Implicit conversion to the element index. */
-	operator int32() const
+	[[nodiscard]] operator int32() const
 	{
 		return ElementIndex;
 	}
diff --git a/Engine/Source/Runtime/Core/Public/Math/GuardedInt.h b/Engine/Source/Runtime/Core/Public/Math/GuardedInt.h
--- a/Engine/Source/Runtime/Core/Public/Math/GuardedInt.h
+++ b/Engine/Source/Runtime/Core/Public/Math/GuardedInt.h
@@ -66,14 +66,14 @@ private:
 
 public:
 	/** Construct a TGuardedInt with an invalid value. */
-	TGuardedInt() = default;
+	[[nodiscard]] TGuardedInt() = default;
 
 	/** Construct a TGuardedInt from a regular signed integer value. If it's out of range, it results in an invalid value. */
 	template<
 		typename ValueType
 		UE_REQUIRES(std::is_integral_v<ValueType>)
 	>
-	explicit TGuardedInt(ValueType InValue)
+	[[nodiscard]] explicit TGuardedInt(ValueType InValue)
 		: Value((IntType)InValue), bIsValid(false)
 	{
 		if constexpr (std::is_signed_v<ValueType>)
@@ -87,16 +87,19 @@ public:
 	}
 
 	/** Copy-construct a TGuardedInt from another of matching type. */
-	TGuardedInt(const TGuardedInt& Other) = default;
+	[[nodiscard]] TGuardedInt(const TGuardedInt& Other) = default;
 
 	/** Assign a TGuardedInt to another. */
 	TGuardedInt& operator=(const TGuardedInt& Other) = default;
 
 	/** @return true if current value is valid (assigned and no overflows or other errors occurred), false otherwise. */
-	bool IsValid() const { return bIsValid; }
+	[[nodiscard]] bool IsValid() const
+	{
+		return bIsValid;
+	}
 
 	/** @return The value if valid, DefaultValue otherwise. */
-	const IntType Get(const IntType DefaultValue) const
+	[[nodiscard]]  IntType Get(const IntType DefaultValue) const
 	{
 		return IsValid() ? Value : DefaultValue;
 	}
@@ -104,14 +107,14 @@ public:
 	/** @return Returns the value if valid, DefaultValue otherwise, but also check()s that the value is valid.
 	 * Intended for cases where the value is not expected to be invalid.
 	 */
-	const IntType GetChecked(const IntType DefaultValue = 0) const
+	[[nodiscard]] IntType GetChecked(const IntType DefaultValue = 0) const
 	{
 		checkf(IsValid(), TEXT("Invalid value in TGuardedInt::GetChecked."));
 		return Get(DefaultValue);
 	}
 
 	/** @return true if *this and Other are either both invalid or both valid and have the same value, false otherwise. */
-	bool operator ==(const TGuardedInt Other) const
+	[[nodiscard]] bool operator ==(const TGuardedInt Other) const
 	{
 		if (bIsValid != Other.bIsValid)
 		{
@@ -122,7 +125,7 @@ public:
 	}
 
 	/** @return true if *this and Other either have different "valid" states or are both valid and have different values, false otherwise (logical negation of ==). */
-	bool operator !=(const TGuardedInt Other) const
+	[[nodiscard]] bool operator !=(const TGuardedInt Other) const
 	{
 		if (bIsValid != Other.bIsValid)
 		{
@@ -136,11 +139,14 @@ public:
 	// to decide what to do about validity as well. Instead, do this.
 
 	/** @return true if *this and Other are both valid so they can be compared. */
-	bool ComparisonValid(const TGuardedInt Other) const { return bIsValid && Other.bIsValid; }
+	[[nodiscard]] bool ComparisonValid(const TGuardedInt Other) const 
+	{
+		return bIsValid && Other.bIsValid;
+	}
 
 	/** @return true if *this and Other are both valid and *this is less than Other. */
 	template<typename ValueType>
-	bool ValidAndLessThan(const ValueType Other) const
+	[[nodiscard]] bool ValidAndLessThan(const ValueType Other) const
 	{
 		TGuardedInt CheckedOther{ Other };
 		return ComparisonValid(CheckedOther) && Value < CheckedOther.Value;
@@ -148,7 +154,7 @@ public:
 
 	/** @return true if *this and Other are both valid and *this is less than or equal to Other. */
 	template<typename ValueType>
-	bool ValidAndLessOrEqual(const ValueType Other) const
+	[[nodiscard]] bool ValidAndLessOrEqual(const ValueType Other) const
 	{
 		TGuardedInt CheckedOther{ Other };
 		return ComparisonValid(CheckedOther) && Value <= CheckedOther.Value;
@@ -156,7 +162,7 @@ public:
 
 	/** @return true if *this and Other are both valid and *this is greater than Other. */
 	template<typename ValueType>
-	bool ValidAndGreaterThan(const ValueType Other) const
+	[[nodiscard]] bool ValidAndGreaterThan(const ValueType Other) const
 	{
 		TGuardedInt CheckedOther{ Other };
 		return ComparisonValid(CheckedOther) && Value > CheckedOther.Value;
@@ -164,7 +170,7 @@ public:
 
 	/** @return true if *this and Other are both valid and *this is greater than or equal to Other. */
 	template<typename ValueType>
-	bool ValidAndGreaterOrEqual(const ValueType Other) const
+	[[nodiscard]] bool ValidAndGreaterOrEqual(const ValueType Other) const
 	{
 		TGuardedInt CheckedOther{ Other };
 		return ComparisonValid(CheckedOther) && Value >= CheckedOther.Value;
@@ -172,7 +178,7 @@ public:
 
 	/** @return true if either of *this or Other are invalid or *this is less than Other. */
 	template<typename ValueType>
-	bool InvalidOrLessThan(const ValueType Other) const
+	[[nodiscard]] bool InvalidOrLessThan(const ValueType Other) const
 	{
 		TGuardedInt CheckedOther{ Other };
 		return !ComparisonValid(CheckedOther) || Value < CheckedOther.Value;
@@ -180,7 +186,7 @@ public:
 
 	/** @return true if either of *this or Other are invalid or *this is less than or equal to Other. */
 	template<typename ValueType>
-	bool InvalidOrLessOrEqual(const ValueType Other) const
+	[[nodiscard]] bool InvalidOrLessOrEqual(const ValueType Other) const
 	{
 		TGuardedInt CheckedOther{ Other };
 		return !ComparisonValid(CheckedOther) || Value <= CheckedOther.Value;
@@ -188,7 +194,7 @@ public:
 
 	/** @return true if either of *this or Other are invalid or *this is greater than Other. */
 	template<typename ValueType>
-	bool InvalidOrGreaterThan(const ValueType Other) const
+	[[nodiscard]] bool InvalidOrGreaterThan(const ValueType Other) const
 	{
 		TGuardedInt CheckedOther{ Other };
 		return !ComparisonValid(CheckedOther) || Value > CheckedOther.Value;
@@ -196,7 +202,7 @@ public:
 
 	/** @return true if either of *this or Other are invalid or *this is greater than or equal to Other. */
 	template<typename ValueType>
-	bool InvalidOrGreaterOrEqual(const ValueType Other) const
+	[[nodiscard]] bool InvalidOrGreaterOrEqual(const ValueType Other) const
 	{
 		TGuardedInt CheckedOther{ Other };
 		return !ComparisonValid(CheckedOther) || Value >= CheckedOther.Value;
@@ -205,7 +211,7 @@ public:
 	// Arithmetic operations
 
 	/** @return The negated value. */
-	TGuardedInt operator-() const
+	[[nodiscard]] TGuardedInt operator-() const
 	{
 		if constexpr (std::is_signed_v<IntType>)
 		{
@@ -220,7 +226,7 @@ public:
 	}
 
 	/** @return The sum of the two operands. */
-	TGuardedInt operator +(const TGuardedInt Other) const
+	[[nodiscard]] TGuardedInt operator +(const TGuardedInt Other) const
 	{
 		// Any sum involving an invalid value is invalid.
 		if (!bIsValid || !Other.bIsValid)
@@ -235,7 +241,7 @@ public:
 	}
 
 	/** @return The difference between the two operands. */
-	TGuardedInt operator -(const TGuardedInt Other) const
+	[[nodiscard]] TGuardedInt operator -(const TGuardedInt Other) const
 	{
 		// Any difference involving an invalid value is invalid.
 		if (!bIsValid || !Other.bIsValid)
@@ -250,7 +256,7 @@ public:
 	}
 
 	/** @return The product of the two operands. */
-	TGuardedInt operator *(const TGuardedInt Other) const
+	[[nodiscard]] TGuardedInt operator *(const TGuardedInt Other) const
 	{
 		// Any product involving invalid values is invalid.
 		if (!bIsValid || !Other.bIsValid)
@@ -265,7 +271,7 @@ public:
 	}
 
 	/** @return The quotient when dividing *this by Other. */
-	TGuardedInt operator /(const TGuardedInt Other) const
+	[[nodiscard]] TGuardedInt operator /(const TGuardedInt Other) const
 	{
 		// Any quotient involving invalid values is invalid.
 		if (!bIsValid || !Other.bIsValid)
@@ -285,7 +291,7 @@ public:
 	}
 
 	/** @return The remainder when dividing *this by Other. */
-	TGuardedInt operator %(const TGuardedInt Other) const
+	[[nodiscard]] TGuardedInt operator %(const TGuardedInt Other) const
 	{
 		// Any quotient involving invalid values is invalid.
 		if (!bIsValid || !Other.bIsValid)
@@ -303,7 +309,7 @@ public:
 	}
 
 	/** @return This value bitwise left-shifted by the operand. */
-	TGuardedInt operator <<(const TGuardedInt Other) const
+	[[nodiscard]] TGuardedInt operator <<(const TGuardedInt Other) const
 	{
 		// Any shift involving invalid values is invalid.
 		if (!bIsValid || !Other.bIsValid)
@@ -345,7 +351,7 @@ public:
 	}
 
 	/** @return This value bitwise right-shifted by the operand. */
-	TGuardedInt operator >>(const TGuardedInt Other) const
+	[[nodiscard]] TGuardedInt operator >>(const TGuardedInt Other) const
 	{
 		// Any shift involving invalid values is invalid.
 		if (!bIsValid || !Other.bIsValid)
@@ -364,7 +370,7 @@ public:
 	}
 
 	/** @return The absolute value of *this. */
-	TGuardedInt Abs() const
+	[[nodiscard]] TGuardedInt Abs() const
 	{
 		if (!bIsValid)
 		{
diff --git a/Engine/Source/Runtime/Core/Public/Math/Halton.h b/Engine/Source/Runtime/Core/Public/Math/Halton.h
--- a/Engine/Source/Runtime/Core/Public/Math/Halton.h
+++ b/Engine/Source/Runtime/Core/Public/Math/Halton.h
@@ -6,7 +6,7 @@
 
 
 /** [ Halton 1964, "Radical-inverse quasi-random point sequence" ] */
-inline float Halton(int32 Index, int32 Base)
+[[nodiscard]] constexpr float Halton(int32 Index, int32 Base)
 {
 	float Result = 0.0f;
 	float InvBase = 1.0f / float(Base);
diff --git a/Engine/Source/Runtime/Core/Public/Math/Int128.h b/Engine/Source/Runtime/Core/Public/Math/Int128.h
--- a/Engine/Source/Runtime/Core/Public/Math/Int128.h
+++ b/Engine/Source/Runtime/Core/Public/Math/Int128.h
@@ -36,7 +36,7 @@ private:
 
 public:
 	/** Gets internal quad parts. */
-	FORCEINLINE uint32 GetQuadPart(uint32 Part)
+	[[nodiscard]] FORCEINLINE uint32 GetQuadPart(uint32 Part) const
 	{
 		switch (Part)
 		{
@@ -76,37 +76,37 @@ public:
 	}
 
 	/** Default constructors. */
-    FORCEINLINE FUInt128(const FUInt128&) = default;
-    FORCEINLINE FUInt128(FUInt128&&) = default;
+    [[nodiscard]] FORCEINLINE FUInt128(const FUInt128&) = default;
+    [[nodiscard]] FORCEINLINE FUInt128(FUInt128&&) = default;
     FORCEINLINE FUInt128& operator=(FUInt128 const&) = default;
     FORCEINLINE FUInt128& operator=(FUInt128&&) = default;
 
 	/** Default constructor. Initializes the number to zero. */
-	FORCEINLINE FUInt128()
+	[[nodiscard]] FORCEINLINE FUInt128()
 		: Hi(0)
 		, Lo(0)
 	{}
 
 	/** Constructor. Initializes this uint128 with a uint64 value. */
-	FORCEINLINE FUInt128(uint64 A)
+	[[nodiscard]] FORCEINLINE FUInt128(uint64 A)
 		: Hi(0)
 		, Lo(A)
 	{}
 
 	/** Constructor. Initializes this uint128 with two uint64 values. */
-	FORCEINLINE FUInt128(uint64 A, uint64 B)
+	[[nodiscard]] FORCEINLINE FUInt128(uint64 A, uint64 B)
 		: Hi(A)
 		, Lo(B)
 	{}
 
 	/** Constructor. Initializes this uint128 with four uint32 values. */
-	FORCEINLINE FUInt128(uint32 A, uint32 B, uint32 C, uint32 D)
+	[[nodiscard]] FORCEINLINE FUInt128(uint32 A, uint32 B, uint32 C, uint32 D)
 		: Hi(((uint64)A << 32) | B)
 		, Lo(((uint64)C << 32) | D)
 	{}
 
 	/** this > Other */
-	FORCEINLINE bool IsGreater(const FUInt128& Other) const
+	[[nodiscard]] FORCEINLINE bool IsGreater(const FUInt128& Other) const
 	{
 		if (Hi == Other.Hi)
 		{
@@ -116,7 +116,7 @@ public:
 	}
 
 	/** this >= Other */
-	FORCEINLINE bool IsGreaterOrEqual(const FUInt128& Other) const
+	[[nodiscard]] FORCEINLINE bool IsGreaterOrEqual(const FUInt128& Other) const
 	{
 		if (Hi == Other.Hi)
 		{
@@ -126,7 +126,7 @@ public:
 	}
 
 	/** this < Other */
-	FORCEINLINE bool IsLess(const FUInt128& Other) const
+	[[nodiscard]] FORCEINLINE bool IsLess(const FUInt128& Other) const
 	{
 		if (Hi == Other.Hi)
 		{
@@ -136,7 +136,7 @@ public:
 	}
 
 	/** this <= Other */
-	FORCEINLINE bool IsLessOrEqual(const FUInt128& Other) const
+	[[nodiscard]] FORCEINLINE bool IsLessOrEqual(const FUInt128& Other) const
 	{
 		if (Hi == Other.Hi)
 		{
@@ -146,7 +146,7 @@ public:
 	}
 
 	/** this == Other */
-	FORCEINLINE bool IsEqual(const FUInt128& Other) const
+	[[nodiscard]] FORCEINLINE bool IsEqual(const FUInt128& Other) const
 	{
 		return (Hi == Other.Hi) && (Lo == Other.Lo);
 	}
@@ -209,37 +209,37 @@ public:
 	/**
 	 * Comparison operators
 	 */
-	FORCEINLINE bool operator>(const FUInt128& Other) const
+	[[nodiscard]] FORCEINLINE bool operator>(const FUInt128& Other) const
 	{
 		return IsGreater(Other);
 	}
 
-	FORCEINLINE bool operator>=(const FUInt128& Other) const
+	[[nodiscard]] FORCEINLINE bool operator>=(const FUInt128& Other) const
 	{
 		return IsGreaterOrEqual(Other);
 	}
 
-	FORCEINLINE bool operator==(const FUInt128& Other) const
+	[[nodiscard]] FORCEINLINE bool operator==(const FUInt128& Other) const
 	{
 		return IsEqual(Other);
 	}
 
-	FORCEINLINE bool operator<(const FUInt128& Other) const
+	[[nodiscard]] FORCEINLINE bool operator<(const FUInt128& Other) const
 	{
 		return IsLess(Other);
 	}
 
-	FORCEINLINE bool operator<=(const FUInt128& Other) const
+	[[nodiscard]] FORCEINLINE bool operator<=(const FUInt128& Other) const
 	{
 		return IsLessOrEqual(Other);
 	}
 
-	FORCEINLINE bool IsZero() const
+	[[nodiscard]] FORCEINLINE bool IsZero() const
 	{
 		return !(Hi | Lo);
 	}
 
-	FORCEINLINE bool IsGreaterThanZero() const
+	[[nodiscard]] FORCEINLINE bool IsGreaterThanZero() const
 	{
 		return !IsZero();
 	}
diff --git a/Engine/Source/Runtime/Core/Public/Math/IntPoint.h b/Engine/Source/Runtime/Core/Public/Math/IntPoint.h
--- a/Engine/Source/Runtime/Core/Public/Math/IntPoint.h
+++ b/Engine/Source/Runtime/Core/Public/Math/IntPoint.h
@@ -48,7 +48,7 @@ struct TIntPoint
 	static const TIntPoint NoneValue;
 
 	/** Default constructor (no initialization). */
-	TIntPoint() = default;
+	[[nodiscard]] TIntPoint() = default;
 
 	/**
 	 * Create and initialize a new instance with the specified coordinates.
@@ -56,7 +56,7 @@ struct TIntPoint
 	 * @param InX The x-coordinate.
 	 * @param InY The y-coordinate.
 	 */
-	TIntPoint(IntType InX, IntType InY)
+[[nodiscard]] TIntPoint(IntType InX, IntType InY)
 		: X(InX)
 		, Y(InY)
 	{
@@ -68,7 +68,7 @@ struct TIntPoint
 	 *
 	 * @param InXY The x and y-coordinate.
 	 */
-	TIntPoint(IntType InXY)
+	[[nodiscard]] TIntPoint(IntType InXY)
 		: X(InXY)
 		, Y(InXY)
 	{
@@ -79,7 +79,7 @@ struct TIntPoint
 	 *
 	 * @param EForceInit Force init enum
 	 */
-	explicit TIntPoint(EForceInit)
+	[[nodiscard]] explicit TIntPoint(EForceInit)
 		: X(0)
 		, Y(0)
 	{
@@ -89,7 +89,7 @@ struct TIntPoint
 	 * Converts to another int type. Checks that the cast will succeed.
 	 */
 	template <typename OtherIntType>
-	explicit TIntPoint(TIntPoint<OtherIntType> Other)
+	[[nodiscard]] explicit TIntPoint(TIntPoint<OtherIntType> Other)
 		: X(IntCastChecked<IntType>(Other.X))
 		, Y(IntCastChecked<IntType>(Other.Y))
 	{
@@ -109,7 +109,7 @@ PRAGMA_ENABLE_DEPRECATION_WARNINGS
 	 * @param PointIndex Index of point component.
 	 * @return const reference to component.
 	 */
-	const IntType& operator()(int32 PointIndex) const
+	[[nodiscard]] const IntType& operator()(int32 PointIndex) const
 	{
 PRAGMA_DISABLE_DEPRECATION_WARNINGS
 		return XY[PointIndex];
@@ -122,7 +122,7 @@ PRAGMA_ENABLE_DEPRECATION_WARNINGS
 	 * @param PointIndex Index of point component
 	 * @return reference to component.
 	 */
-	IntType& operator()(int32 PointIndex)
+	[[nodiscard]] IntType& operator()(int32 PointIndex)
 	{
 PRAGMA_DISABLE_DEPRECATION_WARNINGS
 		return XY[PointIndex];
@@ -135,7 +135,7 @@ PRAGMA_ENABLE_DEPRECATION_WARNINGS
 	 * @param Other The other int point being compared.
 	 * @return true if the points are equal, false otherwise.
 	 */
-	bool operator==(const TIntPoint& Other) const
+	[[nodiscard]] bool operator==(const TIntPoint& Other) const
 	{
 		return X == Other.X && Y == Other.Y;
 	}
@@ -146,7 +146,7 @@ PRAGMA_ENABLE_DEPRECATION_WARNINGS
 	 * @param Other The other int point being compared.
 	 * @return true if the points are not equal, false otherwise.
 	 */
-	bool operator!=(const TIntPoint& Other) const
+	[[nodiscard]] bool operator!=(const TIntPoint& Other) const
 	{
 		return (X != Other.X) || (Y != Other.Y);
 	}
@@ -241,7 +241,7 @@ PRAGMA_ENABLE_DEPRECATION_WARNINGS
 	 * @param Scale What to multiply the point by.
 	 * @return A new scaled int point.
 	 */
-	TIntPoint operator*(IntType Scale) const
+	[[nodiscard]] TIntPoint operator*(IntType Scale) const
 	{
 		return TIntPoint(*this) *= Scale;
 	}
@@ -252,7 +252,7 @@ PRAGMA_ENABLE_DEPRECATION_WARNINGS
 	 * @param Divisor What to divide the point by.
 	 * @return A new divided int point.
 	 */
-	TIntPoint operator/(IntType Divisor) const
+	[[nodiscard]] TIntPoint operator/(IntType Divisor) const
 	{
 		return TIntPoint(*this) /= Divisor;
 	}
@@ -263,7 +263,7 @@ PRAGMA_ENABLE_DEPRECATION_WARNINGS
 	 * @param Other The other point to add to this.
 	 * @return A new combined int point.
 	 */
-	TIntPoint operator+(const TIntPoint& Other) const
+	[[nodiscard]] TIntPoint operator+(const TIntPoint& Other) const
 	{
 		return TIntPoint(*this) += Other;
 	}
@@ -274,7 +274,7 @@ PRAGMA_ENABLE_DEPRECATION_WARNINGS
 	 * @param Other The other point to subtract from this.
 	 * @return A new subtracted int point.
 	 */
-	TIntPoint operator-(const TIntPoint& Other) const
+	[[nodiscard]] TIntPoint operator-(const TIntPoint& Other) const
 	{
 		return TIntPoint(*this) -= Other;
 	}
@@ -285,7 +285,7 @@ PRAGMA_ENABLE_DEPRECATION_WARNINGS
 	 * @param Other The point to multiply with this point.
 	 * @return A new multiplied int point
 	 */
-	TIntPoint operator*(const TIntPoint& Other) const
+	[[nodiscard]] TIntPoint operator*(const TIntPoint& Other) const
 	{
 		return TIntPoint(*this) *= Other;
 	}
@@ -296,30 +296,30 @@ PRAGMA_ENABLE_DEPRECATION_WARNINGS
 	 * @param Other The other point to subtract from this.
 	 * @return A new subtracted int point.
 	 */
-	TIntPoint operator/(const TIntPoint& Other) const
+	[[nodiscard]] TIntPoint operator/(const TIntPoint& Other) const
 	{
 		return TIntPoint(*this) /= Other;
 	}
 
 	/**
-	* Get specific component of the point.
-	*
-	* @param Index the index of point component
-	* @return reference to component.
-	*/
-	IntType& operator[](IntType Index)
+	 * Get specific component of the point.
+	 *
+	 * @param Index the index of point component
+	 * @return reference to component.
+	 */
+	[[nodiscard]] IntType& operator[](IntType Index)
 	{
 		check(Index >= 0 && Index < 2);
 		return ((Index == 0) ? X : Y);
 	}
 
 	/**
-	* Get specific component of the point.
-	*
-	* @param Index the index of point component
-	* @return copy of component value.
-	*/
-	IntType operator[](IntType Index) const
+	 * Get specific component of the point.
+	 *
+	 * @param Index the index of point component
+	 * @return copy of component value.
+	 */
+	[[nodiscard]] IntType operator[](IntType Index) const
 	{
 		check(Index >= 0 && Index < 2);
 		return ((Index == 0) ? X : Y);
@@ -330,7 +330,7 @@ PRAGMA_ENABLE_DEPRECATION_WARNINGS
 	 *
 	 * @see ComponentMax, GetMax
 	 */
-	TIntPoint ComponentMin(const TIntPoint& Other) const
+	[[nodiscard]] TIntPoint ComponentMin(const TIntPoint& Other) const
 	{
 		return TIntPoint(FMath::Min(X, Other.X), FMath::Min(Y, Other.Y));
 	}
@@ -340,7 +340,7 @@ PRAGMA_ENABLE_DEPRECATION_WARNINGS
 	 *
 	 * @see ComponentMin, GetMin
 	 */
-	TIntPoint ComponentMax(const TIntPoint& Other) const
+	[[nodiscard]] TIntPoint ComponentMax(const TIntPoint& Other) const
 	{
 		return TIntPoint(FMath::Max(X, Other.X), FMath::Max(Y, Other.Y));
 	}
@@ -351,7 +351,7 @@ PRAGMA_ENABLE_DEPRECATION_WARNINGS
 	 * @return The maximum component of the point.
 	 * @see GetMin, Size, SizeSquared
 	 */
-	IntType GetMax() const
+	[[nodiscard]] IntType GetMax() const
 	{
 		return FMath::Max(X, Y);
 	}
@@ -362,7 +362,7 @@ PRAGMA_ENABLE_DEPRECATION_WARNINGS
 	 * @return The minimum component of the point.
 	 * @see GetMax, Size, SizeSquared
 	 */
-	IntType GetMin() const
+	[[nodiscard]] IntType GetMin() const
 	{
 		return FMath::Min(X, Y);
 	}
@@ -373,7 +373,7 @@ PRAGMA_ENABLE_DEPRECATION_WARNINGS
 	 * @return The distance of this point from (0,0).
 	 * @see GetMax, GetMin, SizeSquared
 	 */
-	IntType Size() const
+	[[nodiscard]] IntType Size() const
 	{
 		int64 LocalX64 = (int64)X;
 		int64 LocalY64 = (int64)Y;
@@ -386,7 +386,7 @@ PRAGMA_ENABLE_DEPRECATION_WARNINGS
 	 * @return The squared distance of this point from (0,0).
 	 * @see GetMax, GetMin, Size
 	 */
-	IntType SizeSquared() const
+	[[nodiscard]] IntType SizeSquared() const
 	{
 		return X * X + Y * Y;
 	}
@@ -396,7 +396,7 @@ PRAGMA_ENABLE_DEPRECATION_WARNINGS
 	 *
 	 * @return A string describing the point.
 	 */
-	FString ToString() const
+	[[nodiscard]] FString ToString() const
 	{
 		return FString::Printf(TEXT("X=%s Y=%s"), *LexToString(X), *LexToString(Y));
 	}
@@ -426,12 +426,12 @@ PRAGMA_ENABLE_DEPRECATION_WARNINGS
 	 * @return A new divided int point.
 	 * @see DivideAndRoundDown
 	 */
-	static TIntPoint DivideAndRoundUp(TIntPoint lhs, IntType Divisor)
+	[[nodiscard]] static TIntPoint DivideAndRoundUp(TIntPoint lhs, IntType Divisor)
 	{
 		return TIntPoint(FMath::DivideAndRoundUp(lhs.X, Divisor), FMath::DivideAndRoundUp(lhs.Y, Divisor));
 	}
 
-	static TIntPoint DivideAndRoundUp(TIntPoint lhs, TIntPoint Divisor)
+	[[nodiscard]] static TIntPoint DivideAndRoundUp(TIntPoint lhs, TIntPoint Divisor)
 	{
 		return TIntPoint(FMath::DivideAndRoundUp(lhs.X, Divisor.X), FMath::DivideAndRoundUp(lhs.Y, Divisor.Y));
 	}
@@ -444,12 +444,12 @@ PRAGMA_ENABLE_DEPRECATION_WARNINGS
 	 * @return A new divided int point.
 	 * @see DivideAndRoundUp
 	 */
-	static TIntPoint DivideAndRoundDown(TIntPoint lhs, IntType Divisor)
+	[[nodiscard]] static TIntPoint DivideAndRoundDown(TIntPoint lhs, IntType Divisor)
 	{
 		return TIntPoint(FMath::DivideAndRoundDown(lhs.X, Divisor), FMath::DivideAndRoundDown(lhs.Y, Divisor));
 	}
 
-	static TIntPoint DivideAndRoundDown(TIntPoint lhs, TIntPoint Divisor)
+	[[nodiscard]] static TIntPoint DivideAndRoundDown(TIntPoint lhs, TIntPoint Divisor)
 	{
 		return TIntPoint(FMath::DivideAndRoundDown(lhs.X, Divisor.X), FMath::DivideAndRoundDown(lhs.Y, Divisor.Y));
 	}
@@ -459,7 +459,7 @@ PRAGMA_ENABLE_DEPRECATION_WARNINGS
 	 *
 	 * @return number of components point has.
 	 */
-	static int32 Num()
+	[[nodiscard]] static int32 Num()
 	{
 		return 2;
 	}
diff --git a/Engine/Source/Runtime/Core/Public/Math/IntRect.h b/Engine/Source/Runtime/Core/Public/Math/IntRect.h
--- a/Engine/Source/Runtime/Core/Public/Math/IntRect.h
+++ b/Engine/Source/Runtime/Core/Public/Math/IntRect.h
@@ -40,7 +40,7 @@ struct TIntRect
 	};
 
 	/** Constructor */
-	TIntRect()
+	[[nodiscard]] TIntRect()
 		: Min(ForceInit)
 		, Max(ForceInit)
 	{}
@@ -53,7 +53,7 @@ struct TIntRect
 	 * @param X1 Maximum X coordinate.
 	 * @param Y1 Maximum Y coordinate.
 	 */
-	TIntRect(IntType X0, IntType Y0, IntType X1, IntType Y1)
+	[[nodiscard]] TIntRect(IntType X0, IntType Y0, IntType X1, IntType Y1)
 		: Min(X0, Y0)
 		, Max(X1, Y1)
 	{
@@ -65,13 +65,13 @@ struct TIntRect
 	 * @param InMin Minimum Point
 	 * @param InMax Maximum Point
 	 */
-	TIntRect(IntPointType InMin, IntPointType InMax)
+	[[nodiscard]] TIntRect(IntPointType InMin, IntPointType InMax)
 		: Min(InMin)
 		, Max(InMax)
 	{
 	}
 
-	TIntRect(const TIntRect& Other)
+	[[nodiscard]] TIntRect(const TIntRect& Other)
 	{
 		*this = Other;
 	}
@@ -87,7 +87,7 @@ struct TIntRect
 	 * Converts to another int type. Checks that the cast will succeed.
 	 */
 	template <typename OtherIntType>
-	explicit TIntRect(TIntRect<OtherIntType> Other)
+	[[nodiscard]] explicit TIntRect(TIntRect<OtherIntType> Other)
 		: Min(IntPointType(Other.Min))
 		, Max(IntPointType(Other.Max))
 	{
@@ -99,7 +99,7 @@ struct TIntRect
 	 * @param PointIndex Index of Point in rectangle.
 	 * @return Const reference to point in rectangle.
 	 */
-	const TIntRect& operator()(int32 PointIndex) const
+	[[nodiscard]] const TIntRect& operator()(int32 PointIndex) const
 	{
 PRAGMA_DISABLE_DEPRECATION_WARNINGS
 		return MinMax[PointIndex];
@@ -112,7 +112,7 @@ PRAGMA_ENABLE_DEPRECATION_WARNINGS
 	 * @param PointIndex Index of Point in rectangle.
 	 * @return Reference to point in rectangle.
 	 */
-	TIntRect& operator()(int32 PointIndex)
+	[[nodiscard]] TIntRect& operator()(int32 PointIndex)
 	{
 PRAGMA_DISABLE_DEPRECATION_WARNINGS
 		return MinMax[PointIndex];
@@ -125,7 +125,7 @@ PRAGMA_ENABLE_DEPRECATION_WARNINGS
 	 * @param Other The Other Rectangle for comparison.
 	 * @return true if the rectangles are equal, false otherwise..
 	 */
-	bool operator==(const TIntRect& Other) const
+	[[nodiscard]] bool operator==(const TIntRect& Other) const
 	{
 		return Min == Other.Min && Max == Other.Max;
 	}
@@ -136,7 +136,7 @@ PRAGMA_ENABLE_DEPRECATION_WARNINGS
 	 * @param Other The Other Rectangle for comparison.
 	 * @return true if the rectangles are not equal, false otherwise..
 	 */
-	bool operator!=(const TIntRect& Other) const
+	[[nodiscard]] bool operator!=(const TIntRect& Other) const
 	{
 		return Min != Other.Min || Max != Other.Max;
 	}
@@ -189,7 +189,7 @@ PRAGMA_ENABLE_DEPRECATION_WARNINGS
 	 * @param Scale What to multiply this rectangle by.
 	 * @return New scaled rectangle.
 	 */
-	TIntRect operator*(IntType Scale) const
+	[[nodiscard]] TIntRect operator*(IntType Scale) const
 	{
 		return TIntRect(Min * Scale, Max * Scale);
 	}
@@ -200,7 +200,7 @@ PRAGMA_ENABLE_DEPRECATION_WARNINGS
 	 * @param Div What to divide this rectangle by.
 	 * @return New divided rectangle.
 	 */
-	TIntRect operator/(IntType Div) const
+	[[nodiscard]] TIntRect operator/(IntType Div) const
 	{
 		return TIntRect(Min / Div, Max / Div);
 	}
@@ -211,7 +211,7 @@ PRAGMA_ENABLE_DEPRECATION_WARNINGS
 	 * @param Point The point to add to both points in the rectangle.
 	 * @return New rectangle with point added to it.
 	 */
-	TIntRect operator+(const IntPointType& Point) const
+	[[nodiscard]] TIntRect operator+(const IntPointType& Point) const
 	{
 		return TIntRect(Min + Point, Max + Point);
 	}
@@ -222,7 +222,7 @@ PRAGMA_ENABLE_DEPRECATION_WARNINGS
 	 * @param Point The point to divide with.
 	 * @return New rectangle with point divided.
 	 */
-	TIntRect operator/(const IntPointType& Point) const
+	[[nodiscard]] TIntRect operator/(const IntPointType& Point) const
 	{
 		return TIntRect(Min / Point, Max / Point);
 	}
@@ -233,7 +233,7 @@ PRAGMA_ENABLE_DEPRECATION_WARNINGS
 	 * @param Point The point to subtract from both points in the rectangle.
 	 * @return New rectangle with point subtracted from it.
 	 */
-	TIntRect operator-(const IntPointType& Point) const
+	[[nodiscard]] TIntRect operator-(const IntPointType& Point) const
 	{
 		return TIntRect(Min - Point, Max - Point);
 	}
@@ -244,7 +244,7 @@ PRAGMA_ENABLE_DEPRECATION_WARNINGS
 	 * @param Other The other rectangle to add to this.
 	 * @return New rectangle after both are added together.
 	 */
-	TIntRect operator+(const TIntRect& Other) const
+	[[nodiscard]] TIntRect operator+(const TIntRect& Other) const
 	{
 		return TIntRect(Min + Other.Min, Max + Other.Max);
 	}
@@ -255,7 +255,7 @@ PRAGMA_ENABLE_DEPRECATION_WARNINGS
 	 * @param Other The other rectangle to subtract from this.
 	 * @return New rectangle after one is subtracted from this.
 	 */
-	TIntRect operator-(const TIntRect& Other) const
+	[[nodiscard]] TIntRect operator-(const TIntRect& Other) const
 	{
 		return TIntRect(Min - Other.Min, Max - Other.Max);
 	}
@@ -265,7 +265,7 @@ PRAGMA_ENABLE_DEPRECATION_WARNINGS
 	 *
 	 * @return The area of this rectangle.
 	 */
-	IntType Area() const
+	[[nodiscard]] IntType Area() const
 	{
 		return (Max.X - Min.X) * (Max.Y - Min.Y);
 	}
@@ -277,7 +277,7 @@ PRAGMA_ENABLE_DEPRECATION_WARNINGS
 	 * @return The new rectangle.
 	 */
 
-	TIntRect Bottom(IntType InHeight) const
+	[[nodiscard]] TIntRect Bottom(IntType InHeight) const
 	{
 		return TIntRect(Min.X, FMath::Max(Min.Y, Max.Y - InHeight), Max.X, Max.Y);
 	}
@@ -317,7 +317,7 @@ PRAGMA_ENABLE_DEPRECATION_WARNINGS
 	 * @note  This function assumes rects have open bounds, i.e. rects with
 	 *        coincident borders on any edge will not overlap.
 	 */
-	bool Intersect(const TIntRect& Other) const
+	[[nodiscard]] bool Intersect(const TIntRect& Other) const
 	{
 		return Other.Min.X < Max.X&& Other.Max.X > Min.X && Other.Min.Y < Max.Y&& Other.Max.Y > Min.Y;
 	}
@@ -331,7 +331,7 @@ PRAGMA_ENABLE_DEPRECATION_WARNINGS
 	 * @note  This function assumes rects have half-open bounds, i.e. points are contained
 	 *        by the minimum border of the box, but not the maximum border.
 	 */
-	bool Contains(IntPointType P) const
+	[[nodiscard]] bool Contains(IntPointType P) const
 	{
 		return P.X >= Min.X && P.X < Max.X&& P.Y >= Min.Y && P.Y < Max.Y;
 	}
@@ -356,7 +356,7 @@ PRAGMA_ENABLE_DEPRECATION_WARNINGS
 	 *
 	 * @return The Height of the rectangle.
 	 */
-	IntType Height() const
+	[[nodiscard]] IntType Height() const
 	{
 		return (Max.Y - Min.Y);
 	}
@@ -393,7 +393,7 @@ PRAGMA_ENABLE_DEPRECATION_WARNINGS
 	 * @param Shrink How much to remove from each point of this rectangle.
 	 * @return New inner Rectangle.
 	 */
-	TIntRect Inner(IntPointType Shrink) const
+	[[nodiscard]] TIntRect Inner(IntPointType Shrink) const
 	{
 		return TIntRect(Min + Shrink, Max - Shrink);
 	}
@@ -403,7 +403,7 @@ PRAGMA_ENABLE_DEPRECATION_WARNINGS
 	 * @param InWidth Width of the new rectangle (<= rectangles original width).
 	 * @return The new rectangle.
 	 */
-	TIntRect Right(IntType InWidth) const
+	[[nodiscard]] TIntRect Right(IntType InWidth) const
 	{
 		return TIntRect(FMath::Max(Min.X, Max.X - InWidth), Min.Y, Max.X, Max.Y);
 	}
@@ -414,7 +414,7 @@ PRAGMA_ENABLE_DEPRECATION_WARNINGS
 	 * @param Fraction What to scale the rectangle by
 	 * @return New scaled rectangle.
 	 */
-	TIntRect Scale(double Fraction) const
+	[[nodiscard]] TIntRect Scale(double Fraction) const
 	{
 		using Vec2D = UE::Math::TVector2<double>;
 		const Vec2D Min2D = Vec2D((double)Min.X, (double)Min.Y) * Fraction;
@@ -433,7 +433,7 @@ PRAGMA_ENABLE_DEPRECATION_WARNINGS
 	 *
 	 * @return The distance from one corner of the rectangle to the other.
 	 */
-	IntPointType Size() const
+	[[nodiscard]] IntPointType Size() const
 	{
 		return IntPointType(Max.X - Min.X, Max.Y - Min.Y);
 	}
@@ -443,7 +443,7 @@ PRAGMA_ENABLE_DEPRECATION_WARNINGS
 	 *
 	 * @return A string describing the rectangle.
 	 */
-	FString ToString() const
+	[[nodiscard]] FString ToString() const
 	{
 		return FString::Printf(TEXT("Min=(%s) Max=(%s)"), *Min.ToString(), *Max.ToString());
 	}
@@ -453,7 +453,7 @@ PRAGMA_ENABLE_DEPRECATION_WARNINGS
 	 *
 	 * @return The width of the rectangle.
 	 */
-	IntType Width() const
+	[[nodiscard]] IntType Width() const
 	{
 		return Max.X - Min.X;
 	}
@@ -463,7 +463,7 @@ PRAGMA_ENABLE_DEPRECATION_WARNINGS
 	 *
 	 * @return true if the rectangle is 0 x 0.
 	 */
-	bool IsEmpty() const
+	[[nodiscard]] bool IsEmpty() const
 	{
 		return Width() == 0 && Height() == 0;
 	}
@@ -475,12 +475,12 @@ PRAGMA_ENABLE_DEPRECATION_WARNINGS
 	 * @param Div What to divide by.
 	 * @return New divided rectangle.
 	 */
-	static TIntRect DivideAndRoundUp(TIntRect lhs, IntType Div)
+	[[nodiscard]] static TIntRect DivideAndRoundUp(TIntRect lhs, IntType Div)
 	{
 		return DivideAndRoundUp(lhs, IntPointType(Div, Div));
 	}
 
-	static TIntRect DivideAndRoundUp(TIntRect lhs, IntPointType Div)
+	[[nodiscard]] static TIntRect DivideAndRoundUp(TIntRect lhs, IntPointType Div)
 	{
 		return TIntRect(lhs.Min / Div, IntPointType::DivideAndRoundUp(lhs.Max, Div));
 	}
@@ -490,7 +490,7 @@ PRAGMA_ENABLE_DEPRECATION_WARNINGS
 	 *
 	 * @return Number of points in the Rectangle.
 	 */
-	static int32 Num()
+	[[nodiscard]] static int32 Num()
 	{
 		return 2;
 	}
