Make Maths functions [[nodiscard]]

diff --git a/Engine/Source/Runtime/Core/Public/Math/IntVector.h b/Engine/Source/Runtime/Core/Public/Math/IntVector.h
--- a/Engine/Source/Runtime/Core/Public/Math/IntVector.h
+++ b/Engine/Source/Runtime/Core/Public/Math/IntVector.h
@@ -50,7 +50,7 @@ struct TIntVector3
 	/**
 	 * Default constructor (no initialization).
 	 */
-	TIntVector3() = default;
+	[[nodiscard]] TIntVector3() = default;
 
 	/**
 	 * Creates and initializes a new instance with the specified coordinates.
@@ -59,7 +59,7 @@ struct TIntVector3
 	 * @param InY The y-coordinate.
 	 * @param InZ The z-coordinate.
 	 */
-	TIntVector3(IntType InX, IntType InY, IntType InZ)
+	[[nodiscard]] TIntVector3(IntType InX, IntType InY, IntType InZ)
 		: X(InX)
 		, Y(InY)
 		, Z(InZ)
@@ -71,14 +71,14 @@ struct TIntVector3
 	 *
 	 * @param InValue replicated to all components
 	 */
-	explicit TIntVector3(IntType InValue)
+	[[nodiscard]] explicit TIntVector3(IntType InValue)
 		: X(InValue)
 		, Y(InValue)
 		, Z(InValue)
 	{
 	}
 
-	explicit TIntVector3(TIntVector4<IntType> Other)
+	[[nodiscard]] explicit TIntVector3(TIntVector4<IntType> Other)
 		: X(Other.X)
 		, Y(Other.Y)
 		, Z(Other.Z)
@@ -91,14 +91,14 @@ struct TIntVector3
 	 * @param InVector float vector converted to int
 	 */
 	template <typename FloatType>
-	explicit TIntVector3(TVector<FloatType> InVector);
+	[[nodiscard]] explicit TIntVector3(TVector<FloatType> InVector);
 
 	/**
 	 * Constructor
 	 *
 	 * @param EForceInit Force init enum
 	 */
-	explicit TIntVector3(EForceInit)
+	[[nodiscard]] explicit TIntVector3(EForceInit)
 		: X(0)
 		, Y(0)
 		, Z(0)
@@ -109,7 +109,7 @@ struct TIntVector3
 	 * Converts to another int type. Checks that the cast will succeed.
 	 */
 	template <typename OtherIntType>
-	explicit TIntVector3(TIntVector3<OtherIntType> Other)
+	[[nodiscard]] explicit TIntVector3(TIntVector3<OtherIntType> Other)
 		: X(IntCastChecked<IntType>(Other.X))
 		, Y(IntCastChecked<IntType>(Other.Y))
 		, Z(IntCastChecked<IntType>(Other.Z))
@@ -130,7 +130,7 @@ PRAGMA_ENABLE_DEPRECATION_WARNINGS
 	 * @param ComponentIndex Index of vector component.
 	 * @return const reference to component.
 	 */
-	const IntType& operator()(int32 ComponentIndex) const
+	[[nodiscard]] const IntType& operator()(int32 ComponentIndex) const
 	{
 PRAGMA_DISABLE_DEPRECATION_WARNINGS
 		return XYZ[ComponentIndex];
@@ -143,7 +143,7 @@ PRAGMA_ENABLE_DEPRECATION_WARNINGS
 	 * @param ComponentIndex Index of vector component.
 	 * @return reference to component.
 	 */
-	IntType& operator()(int32 ComponentIndex)
+	[[nodiscard]] IntType& operator()(int32 ComponentIndex)
 	{
 PRAGMA_DISABLE_DEPRECATION_WARNINGS
 		return XYZ[ComponentIndex];
@@ -156,7 +156,7 @@ PRAGMA_ENABLE_DEPRECATION_WARNINGS
 	 * @param ComponentIndex Index of vector component.
 	 * @return const reference to component.
 	 */
-	const IntType& operator[](int32 ComponentIndex) const
+	[[nodiscard]] const IntType& operator[](int32 ComponentIndex) const
 	{
 PRAGMA_DISABLE_DEPRECATION_WARNINGS
 		return XYZ[ComponentIndex];
@@ -169,7 +169,7 @@ PRAGMA_ENABLE_DEPRECATION_WARNINGS
 	 * @param ComponentIndex Index of vector component.
 	 * @return reference to component.
 	 */
-	IntType& operator[](int32 ComponentIndex)
+	[[nodiscard]] IntType& operator[](int32 ComponentIndex)
 	{
 PRAGMA_DISABLE_DEPRECATION_WARNINGS
 		return XYZ[ComponentIndex];
@@ -182,7 +182,7 @@ PRAGMA_ENABLE_DEPRECATION_WARNINGS
 	 * @param Other The other int vector being compared.
 	 * @return true if the vectors are equal, false otherwise..
 	 */
-	bool operator==(const TIntVector3& Other) const
+	[[nodiscard]] bool operator==(const TIntVector3& Other) const
 	{
 		return X == Other.X && Y == Other.Y && Z == Other.Z;
 	}
@@ -193,7 +193,7 @@ PRAGMA_ENABLE_DEPRECATION_WARNINGS
 	 * @param Other The other int vector being compared.
 	 * @return true if the vectors are not equal, false otherwise..
 	 */
-	bool operator!=(const TIntVector3& Other) const
+	[[nodiscard]] bool operator!=(const TIntVector3& Other) const
 	{
 		return X != Other.X || Y != Other.Y || Z != Other.Z;
 	}
@@ -294,7 +294,7 @@ PRAGMA_ENABLE_DEPRECATION_WARNINGS
 	 * @param Other The vector to multiply with.
 	 * @return The result of multiplication.
 	 */
-	TIntVector3 operator*(const TIntVector3& Other) const
+	[[nodiscard]] TIntVector3 operator*(const TIntVector3& Other) const
 	{
 		return TIntVector3(*this) *= Other;
 	}
@@ -305,7 +305,7 @@ PRAGMA_ENABLE_DEPRECATION_WARNINGS
 	 * @param Scale What to multiply the vector by.
 	 * @return A new scaled int vector.
 	 */
-	TIntVector3 operator*(IntType Scale) const
+	[[nodiscard]] TIntVector3 operator*(IntType Scale) const
 	{
 		return TIntVector3(*this) *= Scale;
 	}
@@ -316,7 +316,7 @@ PRAGMA_ENABLE_DEPRECATION_WARNINGS
 	 * @param Divisor What to divide the vector by.
 	 * @return A new divided int vector.
 	 */
-	TIntVector3 operator/(IntType Divisor) const
+	[[nodiscard]] TIntVector3 operator/(IntType Divisor) const
 	{
 		return TIntVector3(*this) /= Divisor;
 	}
@@ -327,7 +327,7 @@ PRAGMA_ENABLE_DEPRECATION_WARNINGS
 	 * @param Divisor What to divide the vector by.
 	 * @return A new remainder int vector.
 	 */
-	TIntVector3 operator%(IntType Divisor) const
+	[[nodiscard]] TIntVector3 operator%(IntType Divisor) const
 	{
 		return TIntVector3(*this) %= Divisor;
 	}
@@ -338,7 +338,7 @@ PRAGMA_ENABLE_DEPRECATION_WARNINGS
 	 * @param Other The other vector to add to this.
 	 * @return A new combined int vector.
 	 */
-	TIntVector3 operator+(const TIntVector3& Other) const
+	[[nodiscard]] TIntVector3 operator+(const TIntVector3& Other) const
 	{
 		return TIntVector3(*this) += Other;
 	}
@@ -349,7 +349,7 @@ PRAGMA_ENABLE_DEPRECATION_WARNINGS
 	 * @param Other The other vector to subtract from this.
 	 * @return A new subtracted int vector.
 	 */
-	TIntVector3 operator-(const TIntVector3& Other) const
+	[[nodiscard]] TIntVector3 operator-(const TIntVector3& Other) const
 	{
 		return TIntVector3(*this) -= Other;
 	}
@@ -360,7 +360,7 @@ PRAGMA_ENABLE_DEPRECATION_WARNINGS
 	 * @param Shift The number of bits to shift.
 	 * @return A new shifted int vector.
 	 */
-	TIntVector3 operator>>(IntType Shift) const
+	[[nodiscard]] TIntVector3 operator>>(IntType Shift) const
 	{
 		return TIntVector3(X >> Shift, Y >> Shift, Z >> Shift);
 	}
@@ -371,7 +371,7 @@ PRAGMA_ENABLE_DEPRECATION_WARNINGS
 	 * @param Shift The number of bits to shift.
 	 * @return A new shifted int vector.
 	 */
-	TIntVector3 operator<<(IntType Shift) const
+	[[nodiscard]] TIntVector3 operator<<(IntType Shift) const
 	{
 		return TIntVector3(X << Shift, Y << Shift, Z << Shift);
 	}
@@ -382,7 +382,7 @@ PRAGMA_ENABLE_DEPRECATION_WARNINGS
 	 * @param Value Number to AND with the each component.
 	 * @return A new shifted int vector.
 	 */
-	TIntVector3 operator&(IntType Value) const
+	[[nodiscard]] TIntVector3 operator&(IntType Value) const
 	{
 		return TIntVector3(X & Value, Y & Value, Z & Value);
 	}
@@ -393,7 +393,7 @@ PRAGMA_ENABLE_DEPRECATION_WARNINGS
 	 * @param Value Number to OR with the each component.
 	 * @return A new shifted int vector.
 	 */
-	TIntVector3 operator|(IntType Value) const
+	[[nodiscard]] TIntVector3 operator|(IntType Value) const
 	{
 		return TIntVector3(X | Value, Y | Value, Z | Value);
 	}
@@ -404,7 +404,7 @@ PRAGMA_ENABLE_DEPRECATION_WARNINGS
 	 * @param Value Number to XOR with the each component.
 	 * @return A new shifted int vector.
 	 */
-	TIntVector3 operator^(IntType Value) const
+	[[nodiscard]] TIntVector3 operator^(IntType Value) const
 	{
 		return TIntVector3(X ^ Value, Y ^ Value, Z ^ Value);
 	}
@@ -413,7 +413,7 @@ PRAGMA_ENABLE_DEPRECATION_WARNINGS
 	 * Is vector equal to zero.
 	 * @return is zero
 	*/
-	bool IsZero() const
+	[[nodiscard]] bool IsZero() const
 	{
 		return *this == ZeroValue;
 	}
@@ -423,7 +423,7 @@ PRAGMA_ENABLE_DEPRECATION_WARNINGS
 	 *
 	 * @return The maximum value in the vector.
 	 */
-	IntType GetMax() const
+	[[nodiscard]] IntType GetMax() const
 	{
 		return FMath::Max(FMath::Max(X, Y), Z);
 	}
@@ -433,7 +433,7 @@ PRAGMA_ENABLE_DEPRECATION_WARNINGS
 	 *
 	 * @return The maximum absolute value in the vector.
 	 */
-	IntType GetAbsMax() const
+	[[nodiscard]] IntType GetAbsMax() const
 	{
 		if constexpr (std::is_signed_v<IntType>)
 		{
@@ -450,7 +450,7 @@ PRAGMA_ENABLE_DEPRECATION_WARNINGS
 	 *
 	 * @return The minimum value in the vector.
 	 */
-	IntType GetMin() const
+	[[nodiscard]] IntType GetMin() const
 	{
 		return FMath::Min(FMath::Min(X, Y), Z);
 	}
@@ -460,7 +460,7 @@ PRAGMA_ENABLE_DEPRECATION_WARNINGS
 	 *
 	 * @return The minimum absolute value in the vector.
 	 */
-	IntType GetAbsMin() const
+	[[nodiscard]] IntType GetAbsMin() const
 	{
 		if constexpr (std::is_signed_v<IntType>)
 		{
@@ -478,7 +478,7 @@ PRAGMA_ENABLE_DEPRECATION_WARNINGS
 	 * @param Other The other vector to compare against.
 	 * @return A vector where each component is the maximum value of the corresponding components of the two vectors.
 	 */
-	TIntVector3 ComponentMax(const TIntVector3& Other) const
+	[[nodiscard]] TIntVector3 ComponentMax(const TIntVector3& Other) const
 	{
 		return TIntVector3(
 			FMath::Max(X, Other.X),
@@ -492,7 +492,7 @@ PRAGMA_ENABLE_DEPRECATION_WARNINGS
 	 * @param Other The other vector to compare against.
 	 * @return A vector where each component is the minimum value of the corresponding components of the two vectors.
 	 */
-	TIntVector3 ComponentMin(const TIntVector3& Other) const
+	[[nodiscard]] TIntVector3 ComponentMin(const TIntVector3& Other) const
 	{
 		return TIntVector3(
 			FMath::Min(X, Other.X),
@@ -505,7 +505,7 @@ PRAGMA_ENABLE_DEPRECATION_WARNINGS
 	 *
 	 * @return The distance of this vector from (0,0,0).
 	 */
-	IntType Size() const
+	[[nodiscard]] IntType Size() const
 	{
 		int64 LocalX64 = (int64)X;
 		int64 LocalY64 = (int64)Y;
@@ -554,7 +554,7 @@ PRAGMA_ENABLE_DEPRECATION_WARNINGS
 	 *
 	 * @return A string describing the vector.
 	 */
-	FString ToString() const
+	[[nodiscard]] FString ToString() const
 	{
 		FString Out;
 		AppendString(Out);
@@ -585,12 +585,12 @@ PRAGMA_ENABLE_DEPRECATION_WARNINGS
 	 * @param Divisor What to divide the int vector by.
 	 * @return A new divided int vector.
 	 */
-	static TIntVector3 DivideAndRoundUp(TIntVector3 Lhs, IntType Divisor)
+	[[nodiscard]] static TIntVector3 DivideAndRoundUp(TIntVector3 Lhs, IntType Divisor)
 	{
 		return TIntVector3(FMath::DivideAndRoundUp(Lhs.X, Divisor), FMath::DivideAndRoundUp(Lhs.Y, Divisor), FMath::DivideAndRoundUp(Lhs.Z, Divisor));
 	}
 
-	static TIntVector3 DivideAndRoundUp(TIntVector3 Lhs, TIntVector3 Divisor)
+	[[nodiscard]] static TIntVector3 DivideAndRoundUp(TIntVector3 Lhs, TIntVector3 Divisor)
 	{
 		return TIntVector3(FMath::DivideAndRoundUp(Lhs.X, Divisor.X), FMath::DivideAndRoundUp(Lhs.Y, Divisor.Y), FMath::DivideAndRoundUp(Lhs.Z, Divisor.Z));
 	}
@@ -600,7 +600,7 @@ PRAGMA_ENABLE_DEPRECATION_WARNINGS
 	 *
 	 * @return Number of components vector has.
 	 */
-	static int32 Num()
+	[[nodiscard]] static int32 Num()
 	{
 		return 3;
 	}
diff --git a/Engine/Source/Runtime/Core/Public/Math/InterpCurve.h b/Engine/Source/Runtime/Core/Public/Math/InterpCurve.h
--- a/Engine/Source/Runtime/Core/Public/Math/InterpCurve.h
+++ b/Engine/Source/Runtime/Core/Public/Math/InterpCurve.h
@@ -39,7 +39,7 @@ public:
 	using ElementType = T;
 
 	/** Default constructor. */
-	FInterpCurve() 
+	[[nodiscard]] FInterpCurve() 
 		: bIsLooped(false)
 		, LoopKeyOffset(0.f)
 	{ 
@@ -80,17 +80,17 @@ public:
 	 *	Evaluate the output for an arbitary input value. 
 	 *	For inputs outside the range of the keys, the first/last key value is assumed.
 	 */
-	T Eval( const float InVal, const T& Default = T(ForceInit) ) const;
+	[[nodiscard]] T Eval( const float InVal, const T& Default = T(ForceInit) ) const;
 
 	/** 
 	 *	Evaluate the derivative at a point on the curve.
 	 */
-	T EvalDerivative( const float InVal, const T& Default = T(ForceInit) ) const;
+	[[nodiscard]] T EvalDerivative( const float InVal, const T& Default = T(ForceInit) ) const;
 
 	/** 
 	 *	Evaluate the second derivative at a point on the curve.
 	 */
-	T EvalSecondDerivative( const float InVal, const T& Default = T(ForceInit) ) const;
+	[[nodiscard]] T EvalSecondDerivative( const float InVal, const T& Default = T(ForceInit) ) const;
 
 	/** 
 	 * Find the nearest point on spline to the given point.
@@ -168,7 +168,7 @@ public:
 	/**
 	 * Compare equality of two FInterpCurves
 	 */
-	friend bool operator==(const FInterpCurve& Curve1, const FInterpCurve& Curve2)
+	[[nodiscard]] friend bool operator==(const FInterpCurve& Curve1, const FInterpCurve& Curve2)
 	{
 		return (Curve1.Points == Curve2.Points &&
 				Curve1.bIsLooped == Curve2.bIsLooped &&
@@ -178,7 +178,7 @@ public:
 	/**
 	 * Compare inequality of two FInterpCurves
 	 */
-	friend bool operator!=(const FInterpCurve& Curve1, const FInterpCurve& Curve2)
+	[[nodiscard]] friend bool operator!=(const FInterpCurve& Curve1, const FInterpCurve& Curve2)
 	{
 		return !(Curve1 == Curve2);
 	}
@@ -186,7 +186,7 @@ public:
 	/**
 	 * Finds the lower index of the two points whose input values bound the supplied input value.
 	 */
-	int32 GetPointIndexForInputValue(const float InValue) const;
+	[[nodiscard]] int32 GetPointIndexForInputValue(const float InValue) const;
 };
 
 
diff --git a/Engine/Source/Runtime/Core/Public/Math/InterpCurvePoint.h b/Engine/Source/Runtime/Core/Public/Math/InterpCurvePoint.h
--- a/Engine/Source/Runtime/Core/Public/Math/InterpCurvePoint.h
+++ b/Engine/Source/Runtime/Core/Public/Math/InterpCurvePoint.h
@@ -71,7 +71,7 @@ public:
 	/**
 	 * Default constructor (no initialization).
 	 */
-	FInterpCurvePoint() = default;
+	[[nodiscard]] FInterpCurvePoint() = default;
 
 	/** 
 	 * Constructor 
@@ -80,7 +80,7 @@ public:
 	 * @param Out Output value of templated type
 	 * @note uses linear interpolation
 	 */
-	FInterpCurvePoint( const float In, const T &Out );
+	[[nodiscard]] FInterpCurvePoint( const float In, const T &Out );
 
 	/** 
 	 * Constructor 
@@ -91,19 +91,19 @@ public:
 	 * @param InLeaveTangent Tangent of curve leaving from this point.
 	 * @param InInterpMode interpolation mode to use
 	 */
-	FInterpCurvePoint( const float In, const T &Out, const T &InArriveTangent, const T &InLeaveTangent, const EInterpCurveMode InInterpMode );
+	[[nodiscard]] FInterpCurvePoint( const float In, const T &Out, const T &InArriveTangent, const T &InLeaveTangent, const EInterpCurveMode InInterpMode );
 
 	/**
 	 * Constructor which initializes all components to zero.
 	 *
 	 * @param EForceInit Force init enum
 	 */
-	explicit FORCEINLINE FInterpCurvePoint(EForceInit);
+	[[nodiscard]] explicit FORCEINLINE FInterpCurvePoint(EForceInit);
 
 public:
 
 	/** @return true if the key value is using a curve interp mode, otherwise false */
-	FORCEINLINE bool IsCurveKey() const;
+	[[nodiscard]] FORCEINLINE bool IsCurveKey() const;
 
 public:
 
@@ -126,7 +126,7 @@ public:
 	/**
 	 * Compare equality of two Curve Points
 	 */
-	friend bool operator==( const FInterpCurvePoint& Point1, const FInterpCurvePoint& Point2 )
+	[[nodiscard]] friend bool operator==( const FInterpCurvePoint& Point1, const FInterpCurvePoint& Point2 )
 	{
 		return (Point1.InVal == Point2.InVal &&
 				Point1.OutVal == Point2.OutVal &&
@@ -138,7 +138,7 @@ public:
 	/**
 	 * Compare inequality of two Curve Points
 	 */
-	friend bool operator!=(const FInterpCurvePoint& Point1, const FInterpCurvePoint& Point2)
+	[[nodiscard]] friend bool operator!=(const FInterpCurvePoint& Point1, const FInterpCurvePoint& Point2)
 	{
 		return !(Point1 == Point2);
 	}
diff --git a/Engine/Source/Runtime/Core/Public/Math/Interval.h b/Engine/Source/Runtime/Core/Public/Math/Interval.h
--- a/Engine/Source/Runtime/Core/Public/Math/Interval.h
+++ b/Engine/Source/Runtime/Core/Public/Math/Interval.h
@@ -44,7 +44,7 @@ public:
 	 *
 	 * The interval is invalid
 	 */
-	TInterval()
+	[[nodiscard]] TInterval()
 		: Min(TIntervalTraits<ElementType>::Max())
 		, Max(TIntervalTraits<ElementType>::Lowest())
 	{ }
@@ -55,7 +55,7 @@ public:
 	 * @param InMin The lower bound of the constructed interval.
 	 * @param InMax The upper bound of the constructed interval.
 	 */
-	TInterval( ElementType InMin, ElementType InMax )
+	[[nodiscard]] TInterval( ElementType InMin, ElementType InMax )
 		: Min(InMin)
 		, Max(InMax)
 	{ }
@@ -96,7 +96,7 @@ public:
 	 * @param Other The other interval being compared.
 	 * @return true if the intervals are equal, false otherwise.
 	 */
-	bool operator==(const TInterval& Other) const
+	[[nodiscard]] bool operator==(const TInterval& Other) const
 	{
 		return Min == Other.Min && Max == Other.Max;
 	}
@@ -107,7 +107,7 @@ public:
 	 * @param Other The other interval being compared.
 	 * @return true if the intervals are not equal, false otherwise.
 	 */
-	bool operator!=(const TInterval& Other) const
+	[[nodiscard]] bool operator!=(const TInterval& Other) const
 	{
 		return (Min != Other.Min) || (Max != Other.Max);
 	}
@@ -119,7 +119,7 @@ public:
 	 *
 	 * @return Interval size.
 	 */
-	ElementType Size() const
+	[[nodiscard]] ElementType Size() const
 	{
 		return (Max - Min);
 	}
@@ -129,7 +129,7 @@ public:
 	 *
 	 * @return false when interval is invalid, true otherwise
 	 */
-	bool IsValid() const
+	[[nodiscard]] bool IsValid() const
 	{
 		return (Min <= Max);
 	}
@@ -140,7 +140,7 @@ public:
 	 * @param Element The element to check.
 	 * @return true if the range interval the element, false otherwise.
 	 */
-	bool Contains( const ElementType& Element ) const
+	[[nodiscard]] bool Contains( const ElementType& Element ) const
 	{
 		return IsValid() && (Element >= Min && Element <= Max);
 	}
@@ -191,7 +191,7 @@ public:
 	 * @param Alpha interpolation amount
 	 * @return interpolation result
 	 */
-	ElementType Interpolate( float Alpha ) const
+	[[nodiscard]] ElementType Interpolate( float Alpha ) const
 	{
 		if (IsValid())
 		{
@@ -207,7 +207,7 @@ public:
 	 * @param X the element to clamp
 	 * @return zero if invalid(Min > Max), clamped result otherwise
 	 */
-	ElementType Clamp( ElementType X ) const
+	[[nodiscard]] ElementType Clamp( ElementType X ) const
 	{
 		if (!IsValid())
 		{
@@ -222,7 +222,7 @@ public:
 	 * @param X the element to calculate the percentage
 	 * @return zero if invalid(Min > Max), percentage otherwise
 	 */
-	ElementType GetRangePct(ElementType X) const
+	[[nodiscard]] ElementType GetRangePct(ElementType X) const
 	{
 		if (!IsValid())
 		{
@@ -240,7 +240,7 @@ public:
 	 * @param B The second interval.
 	 * @return The intersection.
 	 */
-	friend TInterval Intersect( const TInterval& A, const TInterval& B )
+	[[nodiscard]] friend TInterval Intersect( const TInterval& A, const TInterval& B )
 	{
 		if (A.IsValid() && B.IsValid())
 		{
diff --git a/Engine/Source/Runtime/Core/Public/Math/LFSR.h b/Engine/Source/Runtime/Core/Public/Math/LFSR.h
--- a/Engine/Source/Runtime/Core/Public/Math/LFSR.h
+++ b/Engine/Source/Runtime/Core/Public/Math/LFSR.h
@@ -17,12 +17,12 @@ class FLinearFeedbackShiftRegister
 {
 public:
 
-	FLinearFeedbackShiftRegister() {}
+	[[nodiscard]] FLinearFeedbackShiftRegister() {}
 
 	// LFSR numbers usually only include N^2-1 numbers for N bit. 
 	// We do apply a -1 to get the value 0 out.
 	// So, for N=4, it should return {0, 1, 2} in random order.
-	uint32 GetNextValue(uint32 N)
+	[[nodiscard]] uint32 GetNextValue(uint32 N)
 	{
 		return GetNextValueInternal(N, /*IncludeLast*/false);
 	}
@@ -30,7 +30,7 @@ public:
 	// This function is the same as above.
 	// But it will also return the last value N^2-1. 
 	// So, for N=4, it will return {0, 1, 2, 3} in random order.
-	uint32 GetNextValueWithLast(uint32 N)
+	[[nodiscard]] uint32 GetNextValueWithLast(uint32 N)
 	{
 		return GetNextValueInternal(N, /*IncludeLast*/true);
 	}
diff --git a/Engine/Source/Runtime/Core/Public/Math/Matrix.h b/Engine/Source/Runtime/Core/Public/Math/Matrix.h
--- a/Engine/Source/Runtime/Core/Public/Math/Matrix.h
+++ b/Engine/Source/Runtime/Core/Public/Math/Matrix.h
@@ -65,14 +65,14 @@ public:
 
 
 	// Constructors.
-	TMatrix() = default;
+	[[nodiscard]] TMatrix() = default;
 
 	/**
 	 * Constructor.
 	 *
 	 * @param EForceInit Force Init Enum.
 	 */
-	explicit FORCEINLINE TMatrix(EForceInit)
+	[[nodiscard]] explicit FORCEINLINE TMatrix(EForceInit)
 	{
 		FMemory::Memzero(this, sizeof(*this));
 	}
@@ -85,7 +85,7 @@ public:
 	 * @param InZ Z plane
 	 * @param InW W plane
 	 */
-	FORCEINLINE TMatrix(const TPlane<T>& InX, const TPlane<T>& InY, const TPlane<T>& InZ, const TPlane<T>& InW);
+	[[nodiscard]] FORCEINLINE TMatrix(const TPlane<T>& InX, const TPlane<T>& InY, const TPlane<T>& InZ, const TPlane<T>& InW);
 
 	/**
 	 * Constructor.
@@ -95,7 +95,7 @@ public:
 	 * @param InZ Z vector
 	 * @param InW W vector
 	 */
-	FORCEINLINE TMatrix(const TVector<T>& InX, const TVector<T>& InY, const TVector<T>& InZ, const TVector<T>& InW);
+	[[nodiscard]] FORCEINLINE TMatrix(const TVector<T>& InX, const TVector<T>& InY, const TVector<T>& InZ, const TVector<T>& InW);
 
 	// Set this to the identity matrix
 	inline void SetIdentity();
@@ -106,7 +106,7 @@ public:
 	 * @param Other The matrix to multiply this by.
 	 * @return The result of multiplication.
 	 */
-	FORCEINLINE TMatrix<T> operator* (const TMatrix<T>& Other) const;
+	[[nodiscard]] FORCEINLINE TMatrix<T> operator* (const TMatrix<T>& Other) const;
 
 	/**
 	 * Multiply this by a matrix.
@@ -122,7 +122,7 @@ public:
 	 * @param Other The Matrix to add.
 	 * @return The result of addition.
 	 */
-	FORCEINLINE TMatrix<T> operator+ (const TMatrix<T>& Other) const;
+	[[nodiscard]] FORCEINLINE TMatrix<T> operator+ (const TMatrix<T>& Other) const;
 
 	/**
 	 * Adds to this matrix.
@@ -135,7 +135,7 @@ public:
 	/**
 	  * This isn't applying SCALE, just multiplying the value to all members - i.e. weighting
 	  */
-	FORCEINLINE TMatrix<T> operator* (T Other) const;
+	[[nodiscard]] FORCEINLINE TMatrix<T> operator* (T Other) const;
 
 	/**
 	 * Multiply this matrix by a weighting factor.
@@ -151,7 +151,7 @@ public:
 	 * @param Other The other matrix.
 	 * @return true if two matrix are identical, otherwise false.
 	 */
-	inline bool operator==(const TMatrix<T>& Other) const;
+	[[nodiscard]] inline bool operator==(const TMatrix<T>& Other) const;
 
 	/**
 	 * Checks whether another Matrix is equal to this, within specified tolerance.
@@ -160,7 +160,7 @@ public:
 	 * @param Tolerance Error Tolerance.
 	 * @return true if two Matrix are equal, within specified tolerance, otherwise false.
 	 */
-	inline bool Equals(const TMatrix<T>& Other, T Tolerance = UE_KINDA_SMALL_NUMBER) const;
+	[[nodiscard]] inline bool Equals(const TMatrix<T>& Other, T Tolerance = UE_KINDA_SMALL_NUMBER) const;
 
 	/**
 	 * Checks whether another Matrix is not equal to this, within specified tolerance.
@@ -168,84 +168,84 @@ public:
 	 * @param Other The other Matrix.
 	 * @return true if two Matrix are not equal, within specified tolerance, otherwise false.
 	 */
-	inline bool operator!=(const TMatrix<T>& Other) const;
+	[[nodiscard]] inline bool operator!=(const TMatrix<T>& Other) const;
 
 	// Homogeneous transform.
-	FORCEINLINE TVector4<T> TransformFVector4(const TVector4<T>& V) const;
+	[[nodiscard]] FORCEINLINE TVector4<T> TransformFVector4(const TVector4<T>& V) const;
 
 	/** Transform a location - will take into account translation part of the TMatrix<T>. */
-	FORCEINLINE TVector4<T> TransformPosition(const TVector<T>& V) const;
+	[[nodiscard]] FORCEINLINE TVector4<T> TransformPosition(const TVector<T>& V) const;
 
 	/** Inverts the matrix and then transforms V - correctly handles scaling in this matrix. */
-	FORCEINLINE TVector<T> InverseTransformPosition(const TVector<T>& V) const;
+	[[nodiscard]] FORCEINLINE TVector<T> InverseTransformPosition(const TVector<T>& V) const;
 
 	/**
 	 *	Transform a direction vector - will not take into account translation part of the TMatrix<T>.
 	 *	If you want to transform a surface normal (or plane) and correctly account for non-uniform scaling you should use TransformByUsingAdjointT.
 	 */
-	FORCEINLINE TVector4<T> TransformVector(const TVector<T>& V) const;
+	[[nodiscard]] FORCEINLINE TVector4<T> TransformVector(const TVector<T>& V) const;
 
 	/**
 	 *	Transform a direction vector by the inverse of this matrix - will not take into account translation part.
 	 *	If you want to transform a surface normal (or plane) and correctly account for non-uniform scaling you should use TransformByUsingAdjointT with adjoint of matrix inverse.
 	 */
-	FORCEINLINE TVector<T> InverseTransformVector(const TVector<T>& V) const;
+	[[nodiscard]] FORCEINLINE TVector<T> InverseTransformVector(const TVector<T>& V) const;
 
 	// Transpose.
 
-	FORCEINLINE TMatrix<T> GetTransposed() const;
+	[[nodiscard]] FORCEINLINE TMatrix<T> GetTransposed() const;
 
 	// @return determinant of this matrix.
 
-	inline T Determinant() const;
+	[[nodiscard]] inline T Determinant() const;
 
 	/** @return the determinant of rotation 3x3 matrix */
-	inline T RotDeterminant() const;
+	[[nodiscard]] inline T RotDeterminant() const;
 
 	/** Get the inverse of this matrix.  Will ensure on nil matrices in non-final builds.  Not faster than Inverse. */
-	inline TMatrix<T> InverseFast() const;
+	[[nodiscard]] inline TMatrix<T> InverseFast() const;
 
 	/** Get the inverse of this matrix.  Will silently change nil/nan matrices to identity. */
-	inline TMatrix<T> Inverse() const;
+	[[nodiscard]] inline TMatrix<T> Inverse() const;
 
-	inline TMatrix<T> TransposeAdjoint() const;
+	[[nodiscard]] inline TMatrix<T> TransposeAdjoint() const;
 
 	// NOTE: There is some compiler optimization issues with WIN64 that cause FORCEINLINE to cause a crash
 	// Remove any scaling from this matrix (ie magnitude of each row is 1) with error Tolerance
 	inline void RemoveScaling(T Tolerance = UE_SMALL_NUMBER);
 
 	// Returns matrix after RemoveScaling with error Tolerance
-	inline TMatrix<T> GetMatrixWithoutScale(T Tolerance = UE_SMALL_NUMBER) const;
+	[[nodiscard]] inline TMatrix<T> GetMatrixWithoutScale(T Tolerance = UE_SMALL_NUMBER) const;
 
 	/** Remove any scaling from this matrix (ie magnitude of each row is 1) and return the 3D scale vector that was initially present with error Tolerance */
 	inline TVector<T> ExtractScaling(T Tolerance = UE_SMALL_NUMBER);
 
 	/** return a 3D scale vector calculated from this matrix (where each component is the magnitude of a row vector) with error Tolerance. */
-	inline TVector<T> GetScaleVector(T Tolerance = UE_SMALL_NUMBER) const;
+	[[nodiscard]] inline TVector<T> GetScaleVector(T Tolerance = UE_SMALL_NUMBER) const;
 
 	// Remove any translation from this matrix
-	inline TMatrix<T> RemoveTranslation() const;
+	[[nodiscard]] inline TMatrix<T> RemoveTranslation() const;
 
 	/** Returns a matrix with an additional translation concatenated. */
-	inline TMatrix<T> ConcatTranslation(const TVector<T>& Translation) const;
+	[[nodiscard]] inline TMatrix<T> ConcatTranslation(const TVector<T>& Translation) const;
 
 	/** Returns true if any element of this matrix is NaN */
-	inline bool ContainsNaN() const;
+	[[nodiscard]] inline bool ContainsNaN() const;
 
 	/** Scale the translation part of the matrix by the supplied vector. */
 	inline void ScaleTranslation(const TVector<T>& Scale3D);
 
 	/** @return the minimum magnitude of any row of the matrix. */
-	inline T GetMinimumAxisScale() const;
+	[[nodiscard]] inline T GetMinimumAxisScale() const;
 
 	/** @return the maximum magnitude of any row of the matrix. */
-	inline T GetMaximumAxisScale() const;
+	[[nodiscard]] inline T GetMaximumAxisScale() const;
 
 	/** Apply Scale to this matrix **/
-	inline TMatrix<T> ApplyScale(T Scale) const;
+	[[nodiscard]] inline TMatrix<T> ApplyScale(T Scale) const;
 
 	// @return the origin of the co-ordinate system
-	inline TVector<T> GetOrigin() const;
+	[[nodiscard]] inline TVector<T> GetOrigin() const;
 
 	/**
 	 * get axis of this matrix scaled by the scale of the matrix
@@ -253,7 +253,7 @@ public:
 	 * @param i index into the axis of the matrix
 	 * @ return vector of the axis
 	 */
-	inline TVector<T> GetScaledAxis(EAxis::Type Axis) const;
+	[[nodiscard]] inline TVector<T> GetScaledAxis(EAxis::Type Axis) const;
 
 	/**
 	 * get axes of this matrix scaled by the scale of the matrix
@@ -270,7 +270,7 @@ public:
 	 * @param i index into the axis of the matrix
 	 * @return vector of the axis
 	 */
-	inline TVector<T> GetUnitAxis(EAxis::Type Axis) const;
+	[[nodiscard]] inline TVector<T> GetUnitAxis(EAxis::Type Axis) const;
 
 	/**
 	 * get unit length axes of this matrix
@@ -308,7 +308,7 @@ public:
 	 * @param i index into the column of the matrix
 	 * @return vector of the column
 	 */
-	inline TVector<T> GetColumn(int32 i) const;
+	[[nodiscard]] inline TVector<T> GetColumn(int32 i) const;
 
 	/**
 	 * Set a column of this matrix
@@ -319,14 +319,14 @@ public:
 	inline void SetColumn(int32 i, TVector<T> Value);
 
 	/** @return rotator representation of this matrix */
-	CORE_API UE::Math::TRotator<T> Rotator() const;
+	[[nodiscard]] CORE_API UE::Math::TRotator<T> Rotator() const;
 
 	/**
 	 * Transform a rotation matrix into a quaternion.
 	 *
 	 * @warning rotation part will need to be unit length for this to be right!
 	 */
-	CORE_API UE::Math::TQuat<T> ToQuat() const;
+	[[nodiscard]] CORE_API UE::Math::TQuat<T> ToQuat() const;
 
 	/**
 	 * Convert this Atom to the 3x4 transpose of the transformation matrix.
@@ -362,7 +362,7 @@ public:
 	 *
 	 * @return Text describing the vector.
 	 */
-	FString ToString() const
+	[[nodiscard]] FString ToString() const
 	{
 		FString Output;
 
@@ -381,7 +381,7 @@ public:
 	}
 
 	/** For debugging purpose, could be changed */
-	uint32 ComputeHash() const
+	[[nodiscard]] uint32 ComputeHash() const
 	{
 		uint32 Ret = 0;
 
diff --git a/Engine/Source/Runtime/Core/Public/Math/MovingWindowAverageFast.h b/Engine/Source/Runtime/Core/Public/Math/MovingWindowAverageFast.h
--- a/Engine/Source/Runtime/Core/Public/Math/MovingWindowAverageFast.h
+++ b/Engine/Source/Runtime/Core/Public/Math/MovingWindowAverageFast.h
@@ -12,7 +12,7 @@ template <typename T, int32 ArraySize>
 class FMovingWindowAverageFast
 {
 public:
-	FMovingWindowAverageFast()
+	[[nodiscard]] FMovingWindowAverageFast()
 		: TotalValues(static_cast<T>(0))
 		, AverageValue(static_cast<T>(0))
 		, RemoveNextIdx(0)
@@ -44,7 +44,7 @@ public:
 		AverageValue = TotalValues / MovingWindowItemsNumCur;
 	}
 
-	T GetAverage() const
+	[[nodiscard]] T GetAverage() const
 	{
 		return AverageValue;
 	}
diff --git a/Engine/Source/Runtime/Core/Public/Math/NumericLimits.h b/Engine/Source/Runtime/Core/Public/Math/NumericLimits.h
--- a/Engine/Source/Runtime/Core/Public/Math/NumericLimits.h
+++ b/Engine/Source/Runtime/Core/Public/Math/NumericLimits.h
@@ -73,17 +73,17 @@ struct TNumericLimits<uint8>
 {
 	typedef uint8 NumericType;
 	
-	static constexpr NumericType Min()
+	[[nodiscard]] static constexpr NumericType Min()
 	{
 		return MIN_uint8;
 	}
 
-	static constexpr NumericType Max()
+	[[nodiscard]] static constexpr NumericType Max()
 	{
 		return MAX_uint8;
 	}
 
-	static constexpr NumericType Lowest()
+	[[nodiscard]] static constexpr NumericType Lowest()
 	{
 		return Min();
 	}
@@ -95,17 +95,17 @@ struct TNumericLimits<uint16>
 {
 	typedef uint16 NumericType;
 	
-	static constexpr NumericType Min()
+	[[nodiscard]] static constexpr NumericType Min()
 	{
 		return MIN_uint16;
 	}
 
-	static constexpr NumericType Max()
+	[[nodiscard]] static constexpr NumericType Max()
 	{
 		return MAX_uint16;
 	}
 
-	static constexpr NumericType Lowest()
+	[[nodiscard]] static constexpr NumericType Lowest()
 	{
 		return Min();
 	}
@@ -117,17 +117,17 @@ struct TNumericLimits<uint32>
 {
 	typedef uint32 NumericType;
 	
-	static constexpr NumericType Min()
+	[[nodiscard]] static constexpr NumericType Min()
 	{
 		return MIN_uint32;
 	}
 
-	static constexpr NumericType Max()
+	[[nodiscard]] static constexpr NumericType Max()
 	{
 		return MAX_uint32;
 	}
 	
-	static constexpr NumericType Lowest()
+	[[nodiscard]] static constexpr NumericType Lowest()
 	{
 		return Min();
 	}
@@ -139,17 +139,17 @@ struct TNumericLimits<uint64>
 {
 	typedef uint64 NumericType;
 	
-	static constexpr NumericType Min()
+	[[nodiscard]] static constexpr NumericType Min()
 	{
 		return MIN_uint64;
 	}
 
-	static constexpr NumericType Max()
+	[[nodiscard]] static constexpr NumericType Max()
 	{
 		return MAX_uint64;
 	}
 
-	static constexpr NumericType Lowest()
+	[[nodiscard]] static constexpr NumericType Lowest()
 	{
 		return Min();
 	}
@@ -161,17 +161,17 @@ struct TNumericLimits<int8>
 {
 	typedef int8 NumericType;
 	
-	static constexpr NumericType Min()
+	[[nodiscard]] static constexpr NumericType Min()
 	{
 		return MIN_int8;
 	}
 
-	static constexpr NumericType Max()
+	[[nodiscard]] static constexpr NumericType Max()
 	{
 		return MAX_int8;
 	}
 
-	static constexpr NumericType Lowest()
+	[[nodiscard]] static constexpr NumericType Lowest()
 	{
 		return Min();
 	}
@@ -183,17 +183,17 @@ struct TNumericLimits<int16>
 {
 	typedef int16 NumericType;
 
-	static constexpr NumericType Min()
+	[[nodiscard]] static constexpr NumericType Min()
 	{
 		return MIN_int16;
 	}
 
-	static constexpr NumericType Max()
+	[[nodiscard]] static constexpr NumericType Max()
 	{
 		return MAX_int16;
 	}
 
-	static constexpr NumericType Lowest()
+	[[nodiscard]] static constexpr NumericType Lowest()
 	{
 		return Min();
 	}
@@ -205,17 +205,17 @@ struct TNumericLimits<int32>
 {
 	typedef int32 NumericType;
 
-	static constexpr NumericType Min()
+	[[nodiscard]] static constexpr NumericType Min()
 	{
 		return MIN_int32;
 	}
 
-	static constexpr NumericType Max()
+	[[nodiscard]] static constexpr NumericType Max()
 	{
 		return MAX_int32;
 	}
 
-	static constexpr NumericType Lowest()
+	[[nodiscard]] static constexpr NumericType Lowest()
 	{
 		return Min();
 	}
@@ -227,17 +227,17 @@ struct TNumericLimits<int64>
 {
 	typedef int64 NumericType;
 	
-	static constexpr NumericType Min()
+	[[nodiscard]] static constexpr NumericType Min()
 	{
 		return MIN_int64;
 	}
 
-	static constexpr NumericType Max()
+	[[nodiscard]] static constexpr NumericType Max()
 	{
 		return MAX_int64;
 	}
 
-	static constexpr NumericType Lowest()
+	[[nodiscard]] static constexpr NumericType Lowest()
 	{
 		return Min();
 	}
@@ -249,17 +249,17 @@ struct TNumericLimits<float>
 {
 	typedef float NumericType;
 
-	static constexpr NumericType Min()
+	[[nodiscard]] static constexpr NumericType Min()
 	{
 		return MIN_flt;
 	}
 
-	static constexpr NumericType Max()
+	[[nodiscard]] static constexpr NumericType Max()
 	{
 		return MAX_flt;
 	}
 
-	static constexpr NumericType Lowest()
+	[[nodiscard]] static constexpr NumericType Lowest()
 	{
 		return -Max();
 	}
@@ -271,17 +271,17 @@ struct TNumericLimits<double>
 {
 	typedef double NumericType;
 
-	static constexpr NumericType Min()
+	[[nodiscard]] static constexpr NumericType Min()
 	{
 		return MIN_dbl;
 	}
 
-	static constexpr NumericType Max()
+	[[nodiscard]] static constexpr NumericType Max()
 	{
 		return MAX_dbl;
 	}
 
-	static constexpr NumericType Lowest()
+	[[nodiscard]] static constexpr NumericType Lowest()
 	{
 		return -Max();
 	}
diff --git a/Engine/Source/Runtime/Core/Public/Math/OrientedBox.h b/Engine/Source/Runtime/Core/Public/Math/OrientedBox.h
--- a/Engine/Source/Runtime/Core/Public/Math/OrientedBox.h
+++ b/Engine/Source/Runtime/Core/Public/Math/OrientedBox.h
@@ -39,7 +39,7 @@ public:
 	 *
 	 * Constructs a unit-sized, origin-centered box with axes aligned to the coordinate system.
 	 */
-	FOrientedBox()
+	[[nodiscard]] FOrientedBox()
 		: Center(0.0f)
 		, AxisX(1.0f, 0.0f, 0.0f)
 		, AxisY(0.0f, 1.0f, 0.0f)
@@ -63,7 +63,7 @@ public:
 	 *
 	 * @param Axis The unit vector defining the axis to project the box onto.
 	 */
-	FORCEINLINE FFloatInterval Project(const FVector& Axis) const;
+	[[nodiscard]] FORCEINLINE FFloatInterval Project(const FVector& Axis) const;
 };
 
 
diff --git a/Engine/Source/Runtime/Core/Public/Math/PackedVector.h b/Engine/Source/Runtime/Core/Public/Math/PackedVector.h
--- a/Engine/Source/Runtime/Core/Public/Math/PackedVector.h
+++ b/Engine/Source/Runtime/Core/Public/Math/PackedVector.h
@@ -29,11 +29,11 @@ public:
         uint32_t v;
     };
 
-	FFloat3Packed() {}
+	[[nodiscard]] FFloat3Packed() {}
 
-	explicit FFloat3Packed(const FLinearColor& Src);
+	[[nodiscard]] explicit FFloat3Packed(const FLinearColor& Src);
 
-	FLinearColor ToLinearColor() const;
+	[[nodiscard]] FLinearColor ToLinearColor() const;
 };
 
 
@@ -289,7 +289,7 @@ public:
 
 	explicit FFixedRGBASigned8(const FLinearColor& Src);
 
-	FLinearColor ToLinearColor() const;
+	[[nodiscard]] FLinearColor ToLinearColor() const;
 };
 
 inline FFixedRGBASigned8::FFixedRGBASigned8(const FLinearColor& Src)
@@ -331,7 +331,7 @@ public:
 	explicit FFloat3PackedSE(const FLinearColor& Src);
 	explicit FFloat3PackedSE(uint32 InEncodedValue) : EncodedValue(InEncodedValue) {}
 
-	FLinearColor ToLinearColor() const;
+	[[nodiscard]] FLinearColor ToLinearColor() const;
 };
 
 inline FFloat3PackedSE::FFloat3PackedSE(const FLinearColor& Src)
diff --git a/Engine/Source/Runtime/Core/Public/Math/Plane.h b/Engine/Source/Runtime/Core/Public/Math/Plane.h
--- a/Engine/Source/Runtime/Core/Public/Math/Plane.h
+++ b/Engine/Source/Runtime/Core/Public/Math/Plane.h
@@ -60,14 +60,14 @@ public:
 public:
 
 	/** Default constructor (no initialization). */
-	TPlane() = default;
+	[[nodiscard]] TPlane() = default;
 
 	/**
 	 * Constructor.
 	 *
 	 * @param V 4D vector to set up plane.
 	 */
-	FORCEINLINE TPlane(const TVector4<T>& V);
+	[[nodiscard]] FORCEINLINE TPlane(const TVector4<T>& V);
 
 	/**
 	 * Constructor.
@@ -77,7 +77,7 @@ public:
 	 * @param InZ Z-coefficient.
 	 * @param InW W-coefficient.
 	 */
-	FORCEINLINE TPlane(T InX, T InY, T InZ, T InW);
+	[[nodiscard]] FORCEINLINE TPlane(T InX, T InY, T InZ, T InW);
 
 	/**
 	 * Constructor.
@@ -85,7 +85,7 @@ public:
 	 * @param InNormal Plane Normal Vector.
 	 * @param InW Plane W-coefficient.
 	 */
-	FORCEINLINE TPlane(TVector<T> InNormal, T InW);
+	[[nodiscard]] FORCEINLINE TPlane(TVector<T> InNormal, T InW);
 
 	/**
 	 * Constructor.
@@ -93,7 +93,7 @@ public:
 	 * @param InBase Base point in plane.
 	 * @param InNormal Plane Normal Vector.
 	 */
-	FORCEINLINE TPlane(TVector<T> InBase, const TVector<T>& InNormal);
+	[[nodiscard]] FORCEINLINE TPlane(TVector<T> InBase, const TVector<T>& InNormal);
 
 	/**
 	 * Constructor.
@@ -102,14 +102,14 @@ public:
 	 * @param B Second point in the plane.
 	 * @param C Third point in the plane.
 	 */
-	TPlane(TVector<T> A, TVector<T> B, TVector<T> C);
+	[[nodiscard]] TPlane(TVector<T> A, TVector<T> B, TVector<T> C);
 
 	/**
 	 * Constructor
 	 *
 	 * @param EForceInit Force Init Enum.
 	 */
-	explicit FORCEINLINE TPlane(EForceInit);
+	[[nodiscard]] explicit FORCEINLINE TPlane(EForceInit);
 
 	// Functions.
 
@@ -118,21 +118,21 @@ public:
 	 *
 	 * @return true if the plane is well-defined (has a non-zero normal), otherwise false.
 	 */
-	FORCEINLINE bool IsValid() const;
+	[[nodiscard]] FORCEINLINE bool IsValid() const;
 
 	/**
 	 * Get the origin of this plane.
 	 *
 	 * @return The origin (base point) of this plane.
 	 */
-	FORCEINLINE TVector<T> GetOrigin() const;
+	[[nodiscard]] FORCEINLINE TVector<T> GetOrigin() const;
 
 	/**
 	 * Get the normal of this plane.
 	 *
 	 * @return The normal of this plane.
 	 */
-	FORCEINLINE const TVector<T>& GetNormal() const;
+	[[nodiscard]] FORCEINLINE const TVector<T>& GetNormal() const;
 
 
 	/**
@@ -141,7 +141,7 @@ public:
 	 * @param P The other point.
 	 * @return The distance from the plane to the point. 0: Point is on the plane. >0: Point is in front of the plane. <0: Point is behind the plane.
 	 */
-	FORCEINLINE T PlaneDot(const TVector<T>& P) const;
+	[[nodiscard]] FORCEINLINE T PlaneDot(const TVector<T>& P) const;
 
 	/**
 	 * Normalize this plane in-place if it is larger than a given tolerance. Leaves it unchanged if not.
@@ -156,7 +156,7 @@ public:
 	 *
 	 * @return A flipped version of the plane.
 	 */
-	TPlane<T> Flip() const;
+	[[nodiscard]] TPlane<T> Flip() const;
 
 	/**
 	 * Get the result of transforming the plane by a Matrix.
@@ -164,7 +164,7 @@ public:
 	 * @param M The matrix to transform plane with.
 	 * @return The result of transform.
 	 */
-	TPlane<T> TransformBy(const TMatrix<T>& M) const;
+	[[nodiscard]] TPlane<T> TransformBy(const TMatrix<T>& M) const;
 
 	/**
 	 * You can optionally pass in the matrices transpose-adjoint, which save it recalculating it.
@@ -176,7 +176,7 @@ public:
 	 * @param TA Transpose-adjoint of Matrix.
 	 * @return The result of transform.
 	 */
-	TPlane<T> TransformByUsingAdjointT(const TMatrix<T>& M, T DetM, const TMatrix<T>& TA) const;
+	[[nodiscard]] TPlane<T> TransformByUsingAdjointT(const TMatrix<T>& M, T DetM, const TMatrix<T>& TA) const;
 
 	/**
 	 * Get the result of translating the plane by the given offset
@@ -184,7 +184,7 @@ public:
 	 * @param V The translation amount
 	 * @return The result of transform.
 	 */
-	TPlane<T> TranslateBy(const TVector<T>& V) const;
+	[[nodiscard]] TPlane<T> TranslateBy(const TVector<T>& V) const;
 
 	/**
 	 * Check if two planes are identical.
@@ -192,7 +192,7 @@ public:
 	 * @param V The other plane.
 	 * @return true if planes are identical, otherwise false.
 	 */
-	bool operator==(const TPlane<T>& V) const;
+	[[nodiscard]] bool operator==(const TPlane<T>& V) const;
 
 	/**
 	 * Check if two planes are different.
@@ -200,7 +200,7 @@ public:
 	 * @param V The other plane.
 	 * @return true if planes are different, otherwise false.
 	 */
-	bool operator!=(const TPlane<T>& V) const;
+	[[nodiscard]] bool operator!=(const TPlane<T>& V) const;
 
 	/**
 	 * Checks whether two planes are equal within specified tolerance.
@@ -209,7 +209,7 @@ public:
 	 * @param Tolerance Error Tolerance.
 	 * @return true if the two planes are equal within specified tolerance, otherwise false.
 	 */
-	bool Equals(const TPlane<T>& V, T Tolerance = UE_KINDA_SMALL_NUMBER) const;
+	[[nodiscard]] bool Equals(const TPlane<T>& V, T Tolerance = UE_KINDA_SMALL_NUMBER) const;
 
 	/**
 	 * Calculates dot product of two planes.
@@ -217,7 +217,7 @@ public:
 	 * @param V The other plane.
 	 * @return The dot product.
 	 */
-	FORCEINLINE T operator|(const TPlane<T>& V) const;
+	[[nodiscard]] FORCEINLINE T operator|(const TPlane<T>& V) const;
 
 	/**
 	 * Gets result of adding a plane to this.
@@ -225,7 +225,7 @@ public:
 	 * @param V The other plane.
 	 * @return The result of adding a plane to this.
 	 */
-	TPlane<T> operator+(const TPlane<T>& V) const;
+	[[nodiscard]] TPlane<T> operator+(const TPlane<T>& V) const;
 
 	/**
 	 * Gets result of subtracting a plane from this.
@@ -233,7 +233,7 @@ public:
 	 * @param V The other plane.
 	 * @return The result of subtracting a plane from this.
 	 */
-	TPlane<T> operator-(const TPlane<T>& V) const;
+	[[nodiscard]] TPlane<T> operator-(const TPlane<T>& V) const;
 
 	/**
 	 * Gets result of dividing a plane.
@@ -241,7 +241,7 @@ public:
 	 * @param Scale What to divide by.
 	 * @return The result of division.
 	 */
-	TPlane<T> operator/(T Scale) const;
+	[[nodiscard]] TPlane<T> operator/(T Scale) const;
 
 	/**
 	 * Gets result of scaling a plane.
@@ -249,7 +249,7 @@ public:
 	 * @param Scale The scaling factor.
 	 * @return The result of scaling.
 	 */
-	TPlane<T> operator*(T Scale) const;
+	[[nodiscard]] TPlane<T> operator*(T Scale) const;
 
 	/**
 	 * Gets result of multiplying a plane with this.
@@ -257,7 +257,7 @@ public:
 	 * @param V The other plane.
 	 * @return The result of multiplying a plane with this.
 	 */
-	TPlane<T> operator*(const TPlane<T>& V);
+	[[nodiscard]] TPlane<T> operator*(const TPlane<T>& V) const;
 
 	/**
 	 * Add another plane to this.
@@ -544,7 +544,7 @@ FORCEINLINE TPlane<T> TPlane<T>::operator*(T Scale) const
 
 
 template<typename T>
-FORCEINLINE TPlane<T> TPlane<T>::operator*(const TPlane<T>& V)
+FORCEINLINE TPlane<T> TPlane<T>::operator*(const TPlane<T>& V) const
 {
 	return TPlane<T>(X * V.X, Y * V.Y, Z * V.Z, W * V.W);
 }
diff --git a/Engine/Source/Runtime/Core/Public/Math/Quat.h b/Engine/Source/Runtime/Core/Public/Math/Quat.h
--- a/Engine/Source/Runtime/Core/Public/Math/Quat.h
+++ b/Engine/Source/Runtime/Core/Public/Math/Quat.h
@@ -65,14 +65,14 @@ public:
 public:
 
 	/** Default constructor (no initialization). */
-	TQuat() = default;
+	[[nodiscard]] TQuat() = default;
 
 	/**
 	 * Creates and initializes a new quaternion, with the W component either 0 or 1.
 	 *
 	 * @param EForceInit Force init enum: if equal to ForceInitToZero then W is 0, otherwise W = 1 (creating an identity transform)
 	 */
-	explicit FORCEINLINE TQuat(EForceInit);
+	[[nodiscard]] explicit FORCEINLINE TQuat(EForceInit);
 
 	/**
 	 * Constructor.
@@ -82,13 +82,13 @@ public:
 	 * @param InZ Z component of the quaternion
 	 * @param InW W component of the quaternion
 	 */
-	FORCEINLINE TQuat(T InX, T InY, T InZ, T InW);
+	[[nodiscard]] FORCEINLINE TQuat(T InX, T InY, T InZ, T InW);
 
 	/**
 	 * Initializes all elements to V
 	 */
 	template<typename DummyT = T UE_REQUIRES(std::is_arithmetic_v<T>)>
-	explicit FORCEINLINE TQuat(T V)
+	[[nodiscard]] explicit FORCEINLINE TQuat(T V)
 		: X(V)
 		, Y(V)
 		, Z(V)
@@ -103,11 +103,14 @@ protected:
 	 *
 	 * @param V XYZW components of the quaternion packed into a single VectorRegister4Float
 	 */
-	explicit TQuat(const QuatVectorRegister& V);
+	[[nodiscard]] explicit TQuat(const QuatVectorRegister& V);
 
 public:
 
-	FORCEINLINE static TQuat<T> MakeFromVectorRegister(const QuatVectorRegister& V) { return TQuat<T>(V); }
+	[[nodiscard]] FORCEINLINE static TQuat<T> MakeFromVectorRegister(const QuatVectorRegister& V)
+	{
+		return TQuat<T>(V);
+	}
 
 	/**
 	 * Creates and initializes a new quaternion from the given rotator.
@@ -116,7 +119,10 @@ public:
 	 */
 	explicit TQuat(const TRotator<T>& R);
 
-	FORCEINLINE static TQuat<T> MakeFromRotator(const TRotator<T>& R) { return TQuat<T>(R); }
+	[[nodiscard]] FORCEINLINE static TQuat<T> MakeFromRotator(const TRotator<T>& R)
+	{
+		return TQuat<T>(R);
+	}
 
 	/**
 	 * Creates and initializes a new quaternion from the given matrix.
@@ -142,7 +148,7 @@ public:
 	 * @param Q The Quaternion to add.
 	 * @return The result of addition.
 	 */
-	FORCEINLINE TQuat<T> operator+(const TQuat<T>& Q) const;
+	[[nodiscard]] FORCEINLINE TQuat<T> operator+(const TQuat<T>& Q) const;
 
 	/**
 	 * Adds to this quaternion.
@@ -160,14 +166,14 @@ public:
 	 * @param Q The Quaternion to subtract.
 	 * @return The result of subtraction.
 	 */
-	FORCEINLINE TQuat<T> operator-(const TQuat<T>& Q) const;
+	[[nodiscard]] FORCEINLINE TQuat<T> operator-(const TQuat<T>& Q) const;
 
 	/**
 	* Negates the quaternion. Note that this represents the same rotation.
 	*
 	* @return The result of negation.
 	*/
-	FORCEINLINE TQuat<T> operator-() const;
+	[[nodiscard]] FORCEINLINE TQuat<T> operator-() const;
 
 	/**
 	 * Checks whether another Quaternion is equal to this, within specified tolerance.
@@ -176,7 +182,7 @@ public:
 	 * @param Tolerance Error tolerance for comparison with other Quaternion.
 	 * @return true if two Quaternions are equal, within specified tolerance, otherwise false.
 	 */
-	FORCEINLINE bool Equals(const TQuat<T>& Q, T Tolerance=UE_KINDA_SMALL_NUMBER) const;
+	[[nodiscard]] FORCEINLINE bool Equals(const TQuat<T>& Q, T Tolerance=UE_KINDA_SMALL_NUMBER) const;
 
 	/**
 	 * Checks whether this Quaternion is an Identity Quaternion.
@@ -185,7 +191,7 @@ public:
 	 * @param Tolerance Error tolerance for comparison with Identity Quaternion.
 	 * @return true if Quaternion is a normalized Identity Quaternion.
 	 */
-	FORCEINLINE bool IsIdentity(T Tolerance=UE_SMALL_NUMBER) const;
+	[[nodiscard]] FORCEINLINE bool IsIdentity(T Tolerance=UE_SMALL_NUMBER) const;
 
 	/**
 	 * Subtracts another quaternion from this.
@@ -205,7 +211,7 @@ public:
 	 * @param Q The Quaternion to multiply this by.
 	 * @return The result of multiplication (this * Q).
 	 */
-	FORCEINLINE TQuat<T> operator*(const TQuat<T>& Q) const;
+	[[nodiscard]] FORCEINLINE TQuat<T> operator*(const TQuat<T>& Q) const;
 
 	/**
 	 * Multiply this by a quaternion (this = this * Q).
@@ -225,7 +231,7 @@ public:
 	 * @return vector after rotation
 	 * @see RotateVector
 	 */
-	TVector<T> operator*(const TVector<T>& V) const;
+	[[nodiscard]] TVector<T> operator*(const TVector<T>& V) const;
 
 	/** 
 	 * Multiply this by a matrix.
@@ -236,7 +242,7 @@ public:
 	 * @param M Matrix to multiply by.
 	 * @return Matrix result after multiplication.
 	 */	
-	CORE_API TMatrix<T> operator*(const TMatrix<T>& M) const;
+	[[nodiscard]] CORE_API TMatrix<T> operator*(const TMatrix<T>& M) const;
 	
 	/**
 	 * Multiply this quaternion by a scaling factor.
@@ -269,7 +275,7 @@ public:
 	 * @return The result of scaling.
 	 */
 	template<typename FArg UE_REQUIRES(std::is_arithmetic_v<FArg>)>
-	FORCEINLINE TQuat<T> operator*(const FArg Scale) const
+	[[nodiscard]] FORCEINLINE TQuat<T> operator*(const FArg Scale) const
 	{
 #if PLATFORM_ENABLE_VECTORINTRINSICS
 		QuatVectorRegister A = VectorLoadAligned(this);
@@ -312,7 +318,7 @@ public:
 	 * @return new Quaternion of this after division by scale.
 	 */
 	template<typename FArg UE_REQUIRES(std::is_arithmetic_v<FArg>)>
-	FORCEINLINE TQuat<T> operator/(const FArg Scale) const
+	[[nodiscard]] FORCEINLINE TQuat<T> operator/(const FArg Scale) const
 	{
 #if PLATFORM_ENABLE_VECTORINTRINSICS
 		QuatVectorRegister A = VectorLoadAligned(this);
@@ -328,7 +334,7 @@ public:
 	 * Identical implementation for TQuat properties. 
 	 * Avoids intrinsics to remain consistent with previous per-property comparison.
 	 */
-	bool Identical(const TQuat* Q, const uint32 PortFlags) const;
+	[[nodiscard]] bool Identical(const TQuat* Q, const uint32 PortFlags) const;
 
  	/**
 	 * Checks whether two quaternions are identical.
@@ -339,7 +345,7 @@ public:
 	 * @return true if two quaternion are identical, otherwise false.
 	 * @see Equals
 	 */
-	bool operator==(const TQuat<T>& Q) const;
+	[[nodiscard]] bool operator==(const TQuat<T>& Q) const;
 
  	/**
 	 * Checks whether two quaternions are not identical.
@@ -347,7 +353,7 @@ public:
 	 * @param Q The other quaternion.
 	 * @return true if two quaternion are not identical, otherwise false.
 	 */
-	bool operator!=(const TQuat<T>& Q) const;
+	[[nodiscard]] bool operator!=(const TQuat<T>& Q) const;
 
 	/**
 	 * Calculates dot product of two quaternions.
@@ -355,7 +361,7 @@ public:
 	 * @param Q The other quaternions.
 	 * @return The dot product.
 	 */
-	T operator|(const TQuat<T>& Q) const;
+	[[nodiscard]] T operator|(const TQuat<T>& Q) const;
 
 public:
 
@@ -368,7 +374,7 @@ public:
 	static CORE_API TQuat<T> MakeFromEuler(const TVector<T>& Euler);
 
 	/** Convert a Quaternion into floating-point Euler angles (in degrees). */
-	CORE_API TVector<T> Euler() const;
+	[[nodiscard]] CORE_API TVector<T> Euler() const;
 
 	/**
 	 * Normalize this quaternion if it is large enough.
@@ -384,31 +390,31 @@ public:
 	 *
 	 * @param Tolerance Minimum squared length of quaternion for normalization.
 	 */
-	FORCEINLINE TQuat<T> GetNormalized(T Tolerance = UE_SMALL_NUMBER) const;
+	[[nodiscard]] FORCEINLINE TQuat<T> GetNormalized(T Tolerance = UE_SMALL_NUMBER) const;
 
 	// Return true if this quaternion is normalized
-	bool IsNormalized() const;
+	[[nodiscard]] bool IsNormalized() const;
 
 	/**
 	 * Get the length of this quaternion.
 	 *
 	 * @return The length of this quaternion.
 	 */
-	FORCEINLINE T Size() const;
+	[[nodiscard]] FORCEINLINE T Size() const;
 
 	/**
 	 * Get the length squared of this quaternion.
 	 *
 	 * @return The length of this quaternion.
 	 */
-	FORCEINLINE T SizeSquared() const;
+	[[nodiscard]] FORCEINLINE T SizeSquared() const;
 
 	/** 
 	 * Get the angle in radians of this quaternion 
 	 * @warning : The rotation returned may not be the shortest version.
 	 *            You may wish to call GetShortestArcWith(FQuat::Identity) first.
 	 */
-	FORCEINLINE T GetAngle() const;
+	[[nodiscard]] FORCEINLINE T GetAngle() const;
 
 	/** 
 	 * get the axis and angle of rotation of this quaternion
@@ -430,14 +436,14 @@ public:
 	 * @warning : The rotation returned may not be the shortest version. 
 	 *            You may wish to call GetShortestArcWith(FQuat::Identity) first.
 	 */
-	TVector<T> ToRotationVector() const;
+	[[nodiscard]] TVector<T> ToRotationVector() const;
 
 	/**
 	 * Constructs a quaternion corresponding to the rotation vector. 
 	 * The direction of the vector represents the rotation axis, 
 	 * and the magnitude the angle in radians.
 	 */
-	static TQuat<T> MakeFromRotationVector(const TVector<T>& RotationVector);
+	[[nodiscard]] static TQuat<T> MakeFromRotationVector(const TVector<T>& RotationVector);
 
 	/** 
 	 * Get the swing and twist decomposition for a specified axis
@@ -456,7 +462,7 @@ public:
 	 * @return Twist angle (in radians)
 	 * @warning assumes normalized quaternion and twist axis
 	 */
-	CORE_API T GetTwistAngle(const TVector<T>& TwistAxis) const;
+	[[nodiscard]] CORE_API T GetTwistAngle(const TVector<T>& TwistAxis) const;
 
 	/**
 	 * Rotate a vector by this quaternion.
@@ -464,7 +470,7 @@ public:
 	 * @param V the vector to be rotated
 	 * @return vector after rotation
 	 */
-	TVector<T> RotateVector(TVector<T> V) const;
+	[[nodiscard]] TVector<T> RotateVector(TVector<T> V) const;
 	
 	/**
 	 * Rotate a vector by the inverse of this quaternion.
@@ -472,27 +478,27 @@ public:
 	 * @param V the vector to be rotated
 	 * @return vector after rotation by the inverse of this quaternion.
 	 */
-	TVector<T> UnrotateVector(TVector<T> V) const;
+	[[nodiscard]] TVector<T> UnrotateVector(TVector<T> V) const;
 
 	/**
 	 * @return quaternion with W=0 and V=theta*v.
 	 * @warning : The rotation returned may not be the shortest version.
 	 *            You may wish to call GetShortestArcWith(FQuat::Identity) first.
 	 */
-	CORE_API TQuat<T> Log() const;
+	[[nodiscard]] CORE_API TQuat<T> Log() const;
 
 	/**
 	 * @note Exp should really only be used after Log.
 	 * Assumes a quaternion with W=0 and V=theta*v (where |v| = 1).
 	 * Exp(q) = (sin(theta)*v, cos(theta))
 	 */
-	CORE_API TQuat<T> Exp() const;
+	[[nodiscard]] CORE_API TQuat<T> Exp() const;
 
 	/**
 	 * @return inverse of this quaternion
 	 * @warning : Requires this quaternion to be normalized.
 	 */
-	FORCEINLINE TQuat<T> Inverse() const;
+	[[nodiscard]] FORCEINLINE TQuat<T> Inverse() const;
 
 	/**
 	 * Enforce that the delta between this Quaternion and another one represents
@@ -504,34 +510,34 @@ public:
 	 * Returns the version of this quaternion that represents the shortest rotation angle 
 	 * to the other quaternion
 	 */
-	TQuat<T> GetShortestArcWith(const TQuat<T>& OtherQuat) const;
+	[[nodiscard]] TQuat<T> GetShortestArcWith(const TQuat<T>& OtherQuat) const;
 
 	/** Get the forward direction (X axis) after it has been rotated by this Quaternion. */
-	FORCEINLINE TVector<T> GetAxisX() const;
+	[[nodiscard]] FORCEINLINE TVector<T> GetAxisX() const;
 
 	/** Get the right direction (Y axis) after it has been rotated by this Quaternion. */
-	FORCEINLINE TVector<T> GetAxisY() const;
+	[[nodiscard]] FORCEINLINE TVector<T> GetAxisY() const;
 
 	/** Get the up direction (Z axis) after it has been rotated by this Quaternion. */
-	FORCEINLINE TVector<T> GetAxisZ() const;
+	[[nodiscard]] FORCEINLINE TVector<T> GetAxisZ() const;
 
 	/** Get the forward direction (X axis) after it has been rotated by this Quaternion. */
-	FORCEINLINE TVector<T> GetForwardVector() const;
+	[[nodiscard]] FORCEINLINE TVector<T> GetForwardVector() const;
 
 	/** Get the right direction (Y axis) after it has been rotated by this Quaternion. */
-	FORCEINLINE TVector<T> GetRightVector() const;
+	[[nodiscard]] FORCEINLINE TVector<T> GetRightVector() const;
 
 	/** Get the up direction (Z axis) after it has been rotated by this Quaternion. */
-	FORCEINLINE TVector<T> GetUpVector() const;
+	[[nodiscard]] FORCEINLINE TVector<T> GetUpVector() const;
 
 	/** Convert a rotation into a unit vector facing in its direction. Equivalent to GetForwardVector(). */
-	FORCEINLINE TVector<T> Vector() const;
+	[[nodiscard]] FORCEINLINE TVector<T> Vector() const;
 
 	/** Get the TRotator<T> representation of this Quaternion. */
-	CORE_API TRotator<T> Rotator() const;
+	[[nodiscard]] CORE_API TRotator<T> Rotator() const;
 
 	/** Get the TMatrix<T> representation of this Quaternion. */
-	FORCEINLINE TMatrix<T> ToMatrix() const;
+	[[nodiscard]] FORCEINLINE TMatrix<T> ToMatrix() const;
 
 	/** Get the TMatrix<T> representation of this Quaternion and store it in Mat */
 	CORE_API void ToMatrix(TMatrix<T>& Mat) const;
@@ -541,10 +547,10 @@ public:
 	 * This is the axis around which rotation occurs to transform the canonical coordinate system to the target orientation.
 	 * For the identity Quaternion which has no such rotation, TVector<T>(1,0,0) is returned.
 	 */
-	FORCEINLINE TVector<T> GetRotationAxis() const;
+	[[nodiscard]] FORCEINLINE TVector<T> GetRotationAxis() const;
 
 	/** Find the angular distance between two rotation quaternions (in radians) */
-	FORCEINLINE T AngularDistance(const TQuat<T>& Q) const;
+	[[nodiscard]] FORCEINLINE T AngularDistance(const TQuat<T>& Q) const;
 
 	/**
 	 * Serializes the vector compressed for e.g. network transmission.
@@ -558,14 +564,14 @@ public:
 	 *
 	 * @return true if there are any non-finite values in this Quaternion, otherwise false.
 	 */
-	bool ContainsNaN() const;
+	[[nodiscard]] bool ContainsNaN() const;
 
 	/**
 	 * Get a textual representation of the vector.
 	 *
 	 * @return Text describing the vector.
 	 */
-	FString ToString() const;
+	[[nodiscard]] FString ToString() const;
 
 	/**
 	 * Initialize this TQuat from a FString. 
@@ -607,7 +613,7 @@ public:
 	/**
 	 * Generates the 'smallest' (geodesic) rotation between two vectors of arbitrary length.
 	 */
-	static FORCEINLINE TQuat<T> FindBetween(const TVector<T>& Vector1, const TVector<T>& Vector2)
+	[[nodiscard]] static FORCEINLINE TQuat<T> FindBetween(const TVector<T>& Vector1, const TVector<T>& Vector2)
 	{
 		return FindBetweenVectors(Vector1, Vector2);
 	}
@@ -615,43 +621,43 @@ public:
 	/**
 	 * Generates the 'smallest' (geodesic) rotation between two normals (assumed to be unit length).
 	 */
-	static CORE_API TQuat<T> FindBetweenNormals(const TVector<T>& Normal1, const TVector<T>& Normal2);
+	[[nodiscard]] static CORE_API TQuat<T> FindBetweenNormals(const TVector<T>& Normal1, const TVector<T>& Normal2);
 
 	/**
 	 * Generates the 'smallest' (geodesic) rotation between two vectors of arbitrary length.
 	 */
-	static CORE_API TQuat<T> FindBetweenVectors(const TVector<T>& Vector1, const TVector<T>& Vector2);
+	[[nodiscard]] static CORE_API TQuat<T> FindBetweenVectors(const TVector<T>& Vector1, const TVector<T>& Vector2);
 
 	/**
 	 * Error measure (angle) between two quaternions, ranged [0..1].
 	 * Returns the hypersphere-angle between two quaternions; alignment shouldn't matter, though 
 	 * @note normalized input is expected.
 	 */
-	static FORCEINLINE T Error(const TQuat<T>& Q1, const TQuat<T>& Q2);
+	[[nodiscard]] static FORCEINLINE T Error(const TQuat<T>& Q1, const TQuat<T>& Q2);
 
 	/**
 	 * TQuat<T>::Error with auto-normalization.
 	 */
-	static FORCEINLINE T ErrorAutoNormalize(const TQuat<T>& A, const TQuat<T>& B);
+	[[nodiscard]] static FORCEINLINE T ErrorAutoNormalize(const TQuat<T>& A, const TQuat<T>& B);
 
 	/** 
 	 * Fast Linear Quaternion Interpolation.
 	 * Result is NOT normalized.
 	 */
-	static FORCEINLINE TQuat<T> FastLerp(const TQuat<T>& A, const TQuat<T>& B, const T Alpha);
+	[[nodiscard]] static FORCEINLINE TQuat<T> FastLerp(const TQuat<T>& A, const TQuat<T>& B, const T Alpha);
 
 	/** 
 	 * Bi-Linear Quaternion interpolation.
 	 * Result is NOT normalized.
 	 */
-	static FORCEINLINE TQuat<T> FastBilerp(const TQuat<T>& P00, const TQuat<T>& P10, const TQuat<T>& P01, const TQuat<T>& P11, T FracX, T FracY);
+	[[nodiscard]] static FORCEINLINE TQuat<T> FastBilerp(const TQuat<T>& P00, const TQuat<T>& P10, const TQuat<T>& P01, const TQuat<T>& P11, T FracX, T FracY);
 
 
 	/** Spherical interpolation. Will correct alignment. Result is NOT normalized. */
-	static CORE_API TQuat<T> Slerp_NotNormalized(const TQuat<T>& Quat1, const TQuat<T>& Quat2, T Slerp);
+	[[nodiscard]] static CORE_API TQuat<T> Slerp_NotNormalized(const TQuat<T>& Quat1, const TQuat<T>& Quat2, T Slerp);
 
 	/** Spherical interpolation. Will correct alignment. Result is normalized. */
-	static FORCEINLINE TQuat<T> Slerp(const TQuat<T>& Quat1, const TQuat<T>& Quat2, T Slerp)
+	[[nodiscard]] static FORCEINLINE TQuat<T> Slerp(const TQuat<T>& Quat1, const TQuat<T>& Quat2, T Slerp)
 	{
 		return Slerp_NotNormalized(Quat1, Quat2, Slerp).GetNormalized();
 	}
@@ -661,14 +667,14 @@ public:
 	 * We need this for the cubic interpolation stuff so that the multiple Slerps dont go in different directions.
 	 * Result is NOT normalized.
 	 */
-	static CORE_API TQuat<T> SlerpFullPath_NotNormalized(const TQuat<T>& quat1, const TQuat<T>& quat2, T Alpha);
+	[[nodiscard]] static CORE_API TQuat<T> SlerpFullPath_NotNormalized(const TQuat<T>& quat1, const TQuat<T>& quat2, T Alpha);
 
 	/**
 	 * Simpler Slerp that doesn't do any checks for 'shortest distance' etc.
 	 * We need this for the cubic interpolation stuff so that the multiple Slerps dont go in different directions.
 	 * Result is normalized.
 	 */
-	static FORCEINLINE TQuat<T> SlerpFullPath(const TQuat<T>& quat1, const TQuat<T>& quat2, T Alpha)
+	[[nodiscard]] static FORCEINLINE TQuat<T> SlerpFullPath(const TQuat<T>& quat1, const TQuat<T>& quat2, T Alpha)
 	{
 		return SlerpFullPath_NotNormalized(quat1, quat2, Alpha).GetNormalized();
 	}
@@ -677,13 +683,13 @@ public:
 	 * Given start and end quaternions of quat1 and quat2, and tangents at those points tang1 and tang2, calculate the point at Alpha (between 0 and 1) between them. Result is normalized.
 	 * This will correct alignment by ensuring that the shortest path is taken.
 	 */
-	static CORE_API TQuat<T> Squad(const TQuat<T>& quat1, const TQuat<T>& tang1, const TQuat<T>& quat2, const TQuat<T>& tang2, T Alpha);
+	[[nodiscard]] static CORE_API TQuat<T> Squad(const TQuat<T>& quat1, const TQuat<T>& tang1, const TQuat<T>& quat2, const TQuat<T>& tang2, T Alpha);
 
 	/**
 	 * Simpler Squad that doesn't do any checks for 'shortest distance' etc.
 	 * Given start and end quaternions of quat1 and quat2, and tangents at those points tang1 and tang2, calculate the point at Alpha (between 0 and 1) between them. Result is normalized.
 	 */
-	static CORE_API TQuat<T> SquadFullPath(const TQuat<T>& quat1, const TQuat<T>& tang1, const TQuat<T>& quat2, const TQuat<T>& tang2, T Alpha);
+	[[nodiscard]] static CORE_API TQuat<T> SquadFullPath(const TQuat<T>& quat1, const TQuat<T>& tang1, const TQuat<T>& quat2, const TQuat<T>& tang2, T Alpha);
 
 	/**
 	 * Calculate tangents between given points
diff --git a/Engine/Source/Runtime/Core/Public/Math/Range.h b/Engine/Source/Runtime/Core/Public/Math/Range.h
--- a/Engine/Source/Runtime/Core/Public/Math/Range.h
+++ b/Engine/Source/Runtime/Core/Public/Math/Range.h
@@ -56,7 +56,7 @@ public:
 	typedef typename TCallTraits<ElementType>::ParamType ElementValueOrConstRef;
 
 	/** Default constructor (no initialization). */
-	TRange() { }
+	[[nodiscard]] TRange() { }
 
 	/**
 	 * Create a range with a single element.
@@ -65,7 +65,7 @@ public:
 	 *
 	 * @param A The element in the range.
 	 */
-	explicit TRange(ElementValueOrConstRef A)
+	[[nodiscard]] explicit TRange(ElementValueOrConstRef A)
 		: LowerBound(BoundsType::Inclusive(A))
 		, UpperBound(BoundsType::Inclusive(A))
 	{ }
@@ -78,7 +78,7 @@ public:
 	 * @param A The range's lower bound value (inclusive).
 	 * @param B The range's upper bound value (exclusive).
 	 */
-	explicit TRange(ElementValueOrConstRef A, ElementValueOrConstRef B)
+	[[nodiscard]] explicit TRange(ElementValueOrConstRef A, ElementValueOrConstRef B)
 		: LowerBound(BoundsType::Inclusive(A))
 		, UpperBound(BoundsType::Exclusive(B))
 	{ }
@@ -89,7 +89,7 @@ public:
 	 * @param InLowerBound The range's lower bound.
 	 * @param InUpperBound The range's upper bound.
 	 */
-	explicit TRange(const BoundsType& InLowerBound, const BoundsType& InUpperBound)
+	[[nodiscard]] explicit TRange(const BoundsType& InLowerBound, const BoundsType& InUpperBound)
 		: LowerBound(InLowerBound)
 		, UpperBound(InUpperBound)
 	{ }
@@ -102,7 +102,7 @@ public:
 	 * @param Other The range to compare with.
 	 * @return true if the ranges are equal, false otherwise.
 	 */
-	bool operator==(const TRange& Other) const
+	[[nodiscard]] bool operator==(const TRange& Other) const
 	{
 		if (IsEmpty() && Other.IsEmpty())
 		{
@@ -118,7 +118,7 @@ public:
 	 * @param Other The range to compare with.
 	 * @return true if the ranges are not equal, false otherwise.
 	 */
-	bool operator!=(const TRange& Other) const
+	[[nodiscard]] bool operator!=(const TRange& Other) const
 	{
 		if (IsEmpty() && Other.IsEmpty())
 		{
@@ -140,7 +140,7 @@ public:
 	 * @param Other The other range.
 	 * @return true if this range adjoins the other, false otherwise.
 	 */
-	bool Adjoins(const TRange& Other) const
+	[[nodiscard]] bool Adjoins(const TRange& Other) const
 	{
 		if (IsEmpty() || Other.IsEmpty())
 		{
@@ -172,7 +172,7 @@ public:
 	 * @param Y The second range.
 	 * @return true if this range conjoins the two ranges, false otherwise.
 	 */
-	bool Conjoins(const TRange& X, const TRange& Y) const
+	[[nodiscard]] bool Conjoins(const TRange& X, const TRange& Y) const
 	{
 		if (X.Overlaps(Y))
 		{
@@ -188,7 +188,7 @@ public:
 	 * @param Element The element to check.
 	 * @return true if the range contains the element, false otherwise.
 	 */
-	bool Contains(ElementValueOrConstRef Element) const
+	[[nodiscard]] bool Contains(ElementValueOrConstRef Element) const
 	{
 		return ((BoundsType::MinLower(LowerBound, Element) == LowerBound) &&
 				(BoundsType::MaxUpper(UpperBound, Element) == UpperBound));
@@ -200,7 +200,7 @@ public:
 	 * @param Other The range to check.
 	 * @return true if the range contains the other range, false otherwise.
 	 */
-	bool Contains(const TRange& Other) const
+	[[nodiscard]] bool Contains(const TRange& Other) const
 	{
 		return ((BoundsType::MinLower(LowerBound, Other.LowerBound) == LowerBound) &&
 				(BoundsType::MaxUpper(UpperBound, Other.UpperBound) == UpperBound));
@@ -214,7 +214,7 @@ public:
 	 * @param Other The other range.
 	 * @return true if the ranges are contiguous, false otherwise.
 	 */
-	bool Contiguous(const TRange& Other) const
+	[[nodiscard]] bool Contiguous(const TRange& Other) const
 	{
 		return (Overlaps(Other) || Adjoins(Other));
 	}
@@ -225,7 +225,7 @@ public:
 	 * @return Lower bound.
 	 * @see GetLowerBoundValue, GetUpperBound, HasLowerBound
 	 */
-	BoundsType GetLowerBound() const
+	[[nodiscard]] BoundsType GetLowerBound() const
 	{
 		return LowerBound;
 	}
@@ -260,7 +260,7 @@ public:
 	 * @return Bound value.
 	 * @see GetLowerBound, GetUpperBoundValue, HasLowerBound
 	 */
-	ElementValueOrConstRef GetLowerBoundValue() const
+	[[nodiscard]] ElementValueOrConstRef GetLowerBoundValue() const
 	{
 		return LowerBound.GetValue();
 	}
@@ -271,7 +271,7 @@ public:
 	 * @return Upper bound.
 	 * @see GetLowerBound, GetUpperBoundValue, HasUpperBound
 	 */
-	BoundsType GetUpperBound() const
+	[[nodiscard]] BoundsType GetUpperBound() const
 	{
 		return UpperBound;
 	}
@@ -306,7 +306,7 @@ public:
 	 * @return Bound value.
 	 * @see GetLowerBoundValue, GetUpperBound, HasUpperBound
 	 */
-	ElementValueOrConstRef GetUpperBoundValue() const
+	[[nodiscard]] ElementValueOrConstRef GetUpperBoundValue() const
 	{
 		return UpperBound.GetValue();
 	}
@@ -317,7 +317,7 @@ public:
 	 * @return true if the range has a lower bound, false otherwise.
 	 * @see GetLowerBound, GetLowerBoundValue, HasUpperBound
 	 */
-	bool HasLowerBound() const
+	[[nodiscard]] bool HasLowerBound() const
 	{
 		return LowerBound.IsClosed();
 	}
@@ -328,7 +328,7 @@ public:
 	 * @return true if the range has an upper bound, false otherwise.
 	 * @see GetUpperBound, GetUpperBoundValue, HasLowerBound
 	 */
-	bool HasUpperBound() const
+	[[nodiscard]] bool HasUpperBound() const
 	{
 		return UpperBound.IsClosed();
 	}
@@ -342,7 +342,7 @@ public:
 	 * @return true if the range is degenerate, false otherwise.
 	 * @see IsEmpty
 	 */
-	bool IsDegenerate() const
+	[[nodiscard]] bool IsDegenerate() const
 	{
 		return (LowerBound.IsInclusive() && (LowerBound == UpperBound));
 	}
@@ -358,7 +358,7 @@ public:
 	 * @return true if the range is empty, false otherwise.
 	 * @see IsDegenerate
 	 */
-	bool IsEmpty() const
+	[[nodiscard]] bool IsEmpty() const
 	{
 		if (LowerBound.IsClosed() && UpperBound.IsClosed())
 		{
@@ -379,7 +379,7 @@ public:
 	 * @param Other The other range.
 	 * @return true if the ranges overlap, false otherwise.
 	 */
-	bool Overlaps(const TRange& Other) const
+	[[nodiscard]] bool Overlaps(const TRange& Other) const
 	{
 		if (IsEmpty() || Other.IsEmpty())
 		{
@@ -423,7 +423,8 @@ public:
 	 *
 	 * @return Range size.
 	 */
-	template<typename DifferenceType> DifferenceType Size() const
+	template<typename DifferenceType>
+	[[nodiscard]] DifferenceType Size() const
 	{
 		check(LowerBound.IsClosed() && UpperBound.IsClosed());
 
@@ -438,7 +439,7 @@ public:
 	 *
 	 * @param Element The element at which to split the range.
 	 */
-	TArray<TRange> Split(ElementValueOrConstRef Element) const
+	[[nodiscard]] TArray<TRange> Split(ElementValueOrConstRef Element) const
 	{
 		TArray<TRange> Result;
 		
@@ -465,7 +466,7 @@ public:
 	 * @return Between 0 and 2 remaining ranges.
 	 * @see Hull, Intersection, Union
 	 */
-	static FORCEINLINE TArray<TRange> Difference(const TRange& X, const TRange& Y)
+	[[nodiscard]] static FORCEINLINE TArray<TRange> Difference(const TRange& X, const TRange& Y)
 	{
 		TArray<TRange> Result;
 
@@ -502,7 +503,7 @@ public:
 	 * @return The hull.
 	 * @see Difference, Intersection, Union
 	 */
-	static FORCEINLINE TRange Hull(const TRange& X, const TRange& Y)
+	[[nodiscard]] static FORCEINLINE TRange Hull(const TRange& X, const TRange& Y)
 	{
 		if (X.IsEmpty())
 		{
@@ -524,7 +525,7 @@ public:
 	 * @return The hull.
 	 * @see Difference, Intersection, Union
 	 */
-	static FORCEINLINE TRange Hull(const TArray<TRange>& Ranges)
+	[[nodiscard]] static FORCEINLINE TRange Hull(const TArray<TRange>& Ranges)
 	{
 		if (Ranges.Num() == 0)
 		{
@@ -551,7 +552,7 @@ public:
 	 * @return The intersection, or an empty range if the ranges do not overlap.
 	 * @see Difference, Hull, Union
 	 */
-	static FORCEINLINE TRange Intersection(const TRange& X, const TRange& Y)
+	[[nodiscard]] static FORCEINLINE TRange Intersection(const TRange& X, const TRange& Y)
 	{
 		if (X.IsEmpty())
 		{
@@ -573,7 +574,7 @@ public:
 	 * @return The intersection.
 	 * @see Difference, Hull, Union
 	 */
-	static FORCEINLINE TRange Intersection(const TArray<TRange>& Ranges)
+	[[nodiscard]] static FORCEINLINE TRange Intersection(const TArray<TRange>& Ranges)
 	{
 		if (Ranges.Num() == 0)
 		{
@@ -600,7 +601,7 @@ public:
 	 * @return The union, or both ranges if the two ranges are not contiguous, or no ranges if both ranges are empty.
 	 * @see Difference, Hull, Intersection
 	 */
-	static FORCEINLINE TArray<TRange> Union(const TRange& X, const TRange& Y)
+	[[nodiscard]] static FORCEINLINE TArray<TRange> Union(const TRange& X, const TRange& Y)
 	{
 		TArray<TRange> OutRanges;
 
@@ -632,7 +633,7 @@ public:
 	 * @return A new range.
 	 * @see AtLeast, AtMost, Empty, Exclusive, GreaterThan, Inclusive, LessThan
 	 */
-	static FORCEINLINE TRange All()
+	[[nodiscard]] static FORCEINLINE TRange All()
 	{
 		return TRange(BoundsType::Open(), BoundsType::Open());
 	}
@@ -644,7 +645,7 @@ public:
 	 * @return A new range.
 	 * @see All, AtMost, Empty, Exclusive, GreaterThan, Inclusive, LessThan
 	 */
-	static FORCEINLINE TRange AtLeast(ElementValueOrConstRef Value)
+	[[nodiscard]] static FORCEINLINE TRange AtLeast(ElementValueOrConstRef Value)
 	{
 		return TRange(BoundsType::Inclusive(Value), BoundsType::Open());
 	}
@@ -656,7 +657,7 @@ public:
 	 * @return A new range.
 	 * @see All, AtLeast, Empty, Exclusive, GreaterThan, Inclusive, LessThan
 	 */
-	static FORCEINLINE TRange AtMost(ElementValueOrConstRef Value)
+	[[nodiscard]] static FORCEINLINE TRange AtMost(ElementValueOrConstRef Value)
 	{
 		return TRange(BoundsType::Open(), BoundsType::Inclusive(Value));
 	}
@@ -667,7 +668,7 @@ public:
 	 * @return Empty range.
 	 * @see All, AtLeast, AtMost, Exclusive, GreaterThan, Inclusive, LessThan
 	 */
-	static FORCEINLINE TRange Empty()
+	[[nodiscard]] static FORCEINLINE TRange Empty()
 	{
 		return TRange(BoundsType::Exclusive(ElementType()), BoundsType::Exclusive(ElementType()));
 	}
@@ -680,7 +681,7 @@ public:
 	 * @return A new range.
 	 * @see All, AtLeast, AtMost, Empty, Exclusive, GreaterThan, Inclusive, LessThan
 	 */
-	static FORCEINLINE TRange Exclusive(ElementValueOrConstRef Min, ElementValueOrConstRef Max)
+	[[nodiscard]] static FORCEINLINE TRange Exclusive(ElementValueOrConstRef Min, ElementValueOrConstRef Max)
 	{
 		return TRange(BoundsType::Exclusive(Min), BoundsType::Exclusive(Max));
 	}
@@ -692,7 +693,7 @@ public:
 	 * @return A new range.
 	 * @see All, AtLeast, AtMost, Empty, Exclusive, Inclusive, LessThan
 	 */
-	static FORCEINLINE TRange GreaterThan(ElementValueOrConstRef Value)
+	[[nodiscard]] static FORCEINLINE TRange GreaterThan(ElementValueOrConstRef Value)
 	{
 		return TRange(BoundsType::Exclusive(Value), BoundsType::Open());
 	}
@@ -705,7 +706,7 @@ public:
 	 * @return A new range.
 	 * @see All, AtLeast, AtMost, Empty, Exclusive, GreaterThan, LessThan
 	 */
-	static FORCEINLINE TRange Inclusive(ElementValueOrConstRef Min, ElementValueOrConstRef Max)
+	[[nodiscard]] static FORCEINLINE TRange Inclusive(ElementValueOrConstRef Min, ElementValueOrConstRef Max)
 	{
 		return TRange(BoundsType::Inclusive(Min), BoundsType::Inclusive(Max));
 	}
@@ -717,7 +718,7 @@ public:
 	 * @return A new range.
 	 * @see All, AtLeast, AtMost, Empty, Exclusive, GreaterThan, Inclusive
 	 */
-	static FORCEINLINE TRange LessThan(ElementValueOrConstRef Value)
+	[[nodiscard]] static FORCEINLINE TRange LessThan(ElementValueOrConstRef Value)
 	{
 		return TRange(BoundsType::Open(), BoundsType::Exclusive(Value));
 	}
diff --git a/Engine/Source/Runtime/Core/Public/Math/RangeBound.h b/Engine/Source/Runtime/Core/Public/Math/RangeBound.h
--- a/Engine/Source/Runtime/Core/Public/Math/RangeBound.h
+++ b/Engine/Source/Runtime/Core/Public/Math/RangeBound.h
@@ -44,7 +44,7 @@ public:
 	 *
 	 * @see Exclusive, Inclusive, Open
 	 */
-	TRangeBound()
+	[[nodiscard]] TRangeBound()
 		: Type(ERangeBoundTypes::Open)
 		, Value()
 	{ }
@@ -55,7 +55,7 @@ public:
 	 * @param InValue The bound's value.
 	 * @see Exclusive, Inclusive, Open
 	 */
-	TRangeBound(ElementValueOrConstRef InValue)
+	[[nodiscard]] TRangeBound(ElementValueOrConstRef InValue)
 		: Type(ERangeBoundTypes::Inclusive)
 		, Value(InValue)
 	{ }
@@ -68,7 +68,7 @@ public:
 	 * @param Other The bound to compare with.
 	 * @return true if the bounds are equal, false otherwise.
 	 */
-	bool operator==(const TRangeBound& Other) const
+	[[nodiscard]] bool operator==(const TRangeBound& Other) const
 	{
 		return ((Type == Other.Type) && (IsOpen() || (Value == Other.Value)));
 	}
@@ -79,7 +79,7 @@ public:
 	 * @param Other The bound to compare with.
 	 * @return true if the bounds are not equal, false otherwise.
 	 */
-	bool operator!=(const TRangeBound& Other) const
+	[[nodiscard]] bool operator!=(const TRangeBound& Other) const
 	{
 		return ((Type != Other.Type) || (!IsOpen() && (Value != Other.Value)));
 	}
@@ -94,7 +94,7 @@ public:
 	 * @return Bound value.
 	 * @see IsOpen
 	 */
-	FORCEINLINE_DEBUGGABLE ElementValueOrConstRef GetValue() const
+	[[nodiscard]] FORCEINLINE_DEBUGGABLE ElementValueOrConstRef GetValue() const
 	{
 		check(Type != ERangeBoundTypes::Open);
 
@@ -122,7 +122,7 @@ public:
 	 *
 	 * @return true if the bound is closed, false otherwise.
 	 */
-	FORCEINLINE_DEBUGGABLE bool IsClosed() const
+	[[nodiscard]] FORCEINLINE_DEBUGGABLE bool IsClosed() const
 	{
 		return (Type != ERangeBoundTypes::Open);
 	}
@@ -132,7 +132,7 @@ public:
 	 *
 	 * @return true if the bound is exclusive, false otherwise.
 	 */
-	FORCEINLINE_DEBUGGABLE bool IsExclusive() const
+	[[nodiscard]] FORCEINLINE_DEBUGGABLE bool IsExclusive() const
 	{
 		return (Type == ERangeBoundTypes::Exclusive);
 	}
@@ -142,7 +142,7 @@ public:
 	 *
 	 * @return true if the bound is inclusive, false otherwise.
 	 */
-	FORCEINLINE_DEBUGGABLE bool IsInclusive() const
+	[[nodiscard]] FORCEINLINE_DEBUGGABLE bool IsInclusive() const
 	{
 		return (Type == ERangeBoundTypes::Inclusive);
 	}
@@ -152,7 +152,7 @@ public:
 	 *
 	 * @return true if the bound is open, false otherwise.
 	 */
-	FORCEINLINE_DEBUGGABLE bool IsOpen() const
+	[[nodiscard]] FORCEINLINE_DEBUGGABLE bool IsOpen() const
 	{
 		return (Type == ERangeBoundTypes::Open);
 	}
@@ -190,7 +190,7 @@ public:
 	 * @param Value The bound value.
 	 * @return An exclusive closed bound.
 	 */
-	static FORCEINLINE TRangeBound Exclusive(ElementValueOrConstRef Value)
+	[[nodiscard]] static FORCEINLINE TRangeBound Exclusive(ElementValueOrConstRef Value)
 	{
 		TRangeBound Result;
 
@@ -206,7 +206,7 @@ public:
 	 * @param Value The bound value.
 	 * @return An inclusive closed bound.
 	 */
-	static FORCEINLINE TRangeBound Inclusive(ElementValueOrConstRef Value)
+	[[nodiscard]] static FORCEINLINE TRangeBound Inclusive(ElementValueOrConstRef Value)
 	{
 		TRangeBound Result;
 
@@ -221,7 +221,7 @@ public:
 	 *
 	 * @return An open bound.
 	 */
-	static FORCEINLINE TRangeBound Open()
+	[[nodiscard]] static FORCEINLINE TRangeBound Open()
 	{
 		TRangeBound Result;
 
@@ -239,7 +239,7 @@ public:
 	 *
 	 * @return A new bound.
 	 */
-	static FORCEINLINE TRangeBound FlipInclusion(const TRangeBound& Bound)
+	[[nodiscard]] static FORCEINLINE TRangeBound FlipInclusion(const TRangeBound& Bound)
 	{
 		if (Bound.IsExclusive())
 		{
@@ -261,7 +261,7 @@ public:
 	 * @param B The second lower bound.
 	 * @return The greater lower bound.
 	 */
-	static FORCEINLINE const TRangeBound& MaxLower(const TRangeBound& A, const TRangeBound& B)
+	[[nodiscard]] static FORCEINLINE const TRangeBound& MaxLower(const TRangeBound& A, const TRangeBound& B)
 	{
 		if (A.IsOpen()) { return B; }
 		if (B.IsOpen()) { return A; }
@@ -279,7 +279,7 @@ public:
 	 * @param B The second upper bound.
 	 * @return The greater upper bound.
 	 */
-	static FORCEINLINE const TRangeBound& MaxUpper(const TRangeBound& A, const TRangeBound& B)
+	[[nodiscard]] static FORCEINLINE const TRangeBound& MaxUpper(const TRangeBound& A, const TRangeBound& B)
 	{
 		if (A.IsOpen()) { return A; }
 		if (B.IsOpen()) { return B; }
@@ -297,7 +297,7 @@ public:
 	 * @param B The second lower bound.
 	 * @return The lesser lower bound.
 	 */
-	static FORCEINLINE const TRangeBound& MinLower(const TRangeBound& A, const TRangeBound& B)
+	[[nodiscard]] static FORCEINLINE const TRangeBound& MinLower(const TRangeBound& A, const TRangeBound& B)
 	{
 		if (A.IsOpen()) { return A; }
 		if (B.IsOpen()) { return B; }
@@ -315,7 +315,7 @@ public:
 	 * @param B The second upper bound.
 	 * @return The lesser upper bound.
 	 */
-	static FORCEINLINE const TRangeBound& MinUpper(const TRangeBound& A, const TRangeBound& B)
+	[[nodiscard]] static FORCEINLINE const TRangeBound& MinUpper(const TRangeBound& A, const TRangeBound& B)
 	{
 		if (A.IsOpen()) { return B; }
 		if (B.IsOpen()) { return A; }
diff --git a/Engine/Source/Runtime/Core/Public/Math/RangeSet.h b/Engine/Source/Runtime/Core/Public/Math/RangeSet.h
--- a/Engine/Source/Runtime/Core/Public/Math/RangeSet.h
+++ b/Engine/Source/Runtime/Core/Public/Math/RangeSet.h
@@ -24,7 +24,7 @@ template<typename ElementType> class TRangeSet
 public:
 
 	/** Default constructor. */
-	TRangeSet() { }
+	[[nodiscard]] TRangeSet() { }
 
 	/** Destructor. */
 	~TRangeSet() { }
@@ -116,7 +116,7 @@ public:
 	 * @param Element The element to check.
 	 * @return true if the element is in the set, false otherwise.
 	 */
-	bool Contains(ElementValueOrConstRef Element) const
+	[[nodiscard]] bool Contains(ElementValueOrConstRef Element) const
 	{
 		for (typename TArray<RangeType>::TConstIterator It(Ranges); It; ++It)
 		{
@@ -135,7 +135,7 @@ public:
 	 * @param Range The range to check.
 	 * @return true if the set contains the range, false otherwise.
 	 */
-	bool Contains(const RangeType& Range) const
+	[[nodiscard]] bool Contains(const RangeType& Range) const
 	{
 		for (typename TArray<RangeType>::TConstIterator It(Ranges); It; ++It)
 		{
@@ -154,7 +154,7 @@ public:
 	 * @return Lowest bound.
 	 * @see GetMaxBound, GetMinBoundValue, HasMinBound
 	 */
-	BoundsType GetMinBound() const
+	[[nodiscard]] BoundsType GetMinBound() const
 	{
 		BoundsType Result;
 
@@ -178,7 +178,7 @@ public:
 	 * @return Bound value.
 	 * @see GetMaxBoundValue, GetMinBound, HasMinBound
 	 */
-	ElementType GetMinBoundValue() const
+	[[nodiscard]] ElementType GetMinBoundValue() const
 	{
 		return GetMinBound().GetValue();
 	}
@@ -189,7 +189,7 @@ public:
 	 * @return Uppermost bound.
 	 * @see GetMaxBoundValue, GetMinBound, HasMaxBound
 	 */
-	BoundsType GetMaxBound() const
+	[[nodiscard]] BoundsType GetMaxBound() const
 	{
 		BoundsType Result;
 
@@ -213,7 +213,7 @@ public:
 	 * @return Bound value.
 	 * @see GetMaxBound, GetMinBoundValue, HasMaxBound
 	 */
-	ElementType GetMaxBoundValue() const
+	[[nodiscard]] ElementType GetMaxBoundValue() const
 	{
 		return GetMaxBound().GetValue();
 	}
@@ -237,7 +237,7 @@ public:
 	 * @return true if the range has a lowest bound, false otherwise.
 	 * @see GetMinBound, GetMinBoundValue, HasMaxBound
 	 */
-	bool HasMinBound() const
+	[[nodiscard]] bool HasMinBound() const
 	{
 		return GetMinBound().IsClosed();
 	}
@@ -248,7 +248,7 @@ public:
 	 * @return true if the range has an uppermost bound, false otherwise.
 	 * @see GetUpperBound, GetUpperBoundValue, HasMinBound
 	 */
-	bool HasMaxBound() const
+	[[nodiscard]] bool HasMaxBound() const
 	{
 		return GetMaxBound().IsClosed();
 	}
@@ -258,7 +258,7 @@ public:
 	 *
 	 * @return true if the range set is empty, false otherwise.
 	 */
-	bool IsEmpty() const
+	[[nodiscard]] bool IsEmpty() const
 	{
 		return (Ranges.Num() == 0);
 	}
@@ -269,7 +269,7 @@ public:
 	 * @param Range The range to check.
 	 * @return true if this set overlaps with the range, false otherwise.
 	 */
-	bool Overlaps(const RangeType& Range) const
+	[[nodiscard]] bool Overlaps(const RangeType& Range) const
 	{
 		for (typename TArray<RangeType>::TConstIterator It(Ranges); It; ++It)
 		{
@@ -290,7 +290,7 @@ public:
 	 *
 	 * @todo gmp: This could be optimized to O(n*logn) using a line sweep on a pre-sorted array of bounds.
 	 */
-	bool Overlaps(const TRangeSet& Other) const
+	[[nodiscard]] bool Overlaps(const TRangeSet& Other) const
 	{
 		for (typename TArray<RangeType>::TConstIterator It(Other.Ranges); It; ++It)
 		{
diff --git a/Engine/Source/Runtime/Core/Public/Math/Ray.h b/Engine/Source/Runtime/Core/Public/Math/Ray.h
--- a/Engine/Source/Runtime/Core/Public/Math/Ray.h
+++ b/Engine/Source/Runtime/Core/Public/Math/Ray.h
@@ -29,7 +29,7 @@ public:
 public:
 
 	/** Default constructor initializes ray to Zero origin and Z-axis direction */
-	TRay()
+	[[nodiscard]] TRay()
 	{
 		Init();
 	}
@@ -38,7 +38,7 @@ public:
 	 * Creates and initializes a new Ray to Zero Origin and Z-axis Direction.
 	 * @param EForceInit Force Init Enum.
 	 */
-	explicit TRay<T>( EForceInit )
+	[[nodiscard]] explicit TRay<T>( EForceInit )
 	{
 		Init();
 	}
@@ -51,7 +51,7 @@ public:
 	  * @param Direction Ray Direction Vector
 	  * @param bDirectionIsNormalized Direction will be normalized unless this is passed as true (default false)
 	  */
-	TRay(const TVector<T>& Origin, const TVector<T>& Direction, bool bDirectionIsNormalized = false)
+	[[nodiscard]] TRay(const TVector<T>& Origin, const TVector<T>& Direction, bool bDirectionIsNormalized = false)
 	{
 		this->Origin = Origin;
 		this->Direction = Direction;
@@ -77,7 +77,7 @@ public:
 	 * @param Other The other Ray to compare with.
 	 * @return true if the Rays are equal, false otherwise.
 	 */
-	bool operator==( const TRay<T>& Other ) const
+	[[nodiscard]] bool operator==( const TRay<T>& Other ) const
 	{
 		return (Origin == Other.Origin) && (Direction == Other.Direction);
 	}
@@ -88,7 +88,7 @@ public:
 	 * @param Other The other Ray to compare with.
 	 * @return true if the Rays are not equal, false otherwise.
 	 */
-	bool operator!=( const TRay<T>& Other ) const
+	[[nodiscard]] bool operator!=( const TRay<T>& Other ) const
 	{
 		return !(*this == Other);
 	}
@@ -102,7 +102,7 @@ public:
 	 * @param RayParameter Scalar distance along Ray
 	 * @return Point on Ray
 	 */
-	TVector<T> PointAt(T RayParameter) const
+	[[nodiscard]] TVector<T> PointAt(T RayParameter) const
 	{
 		return Origin + RayParameter * Direction;
 	}
@@ -113,7 +113,7 @@ public:
 	 * @param Point query Point
 	 * @return distance along ray from origin to closest point
 	 */
-	T GetParameter(const TVector<T>& Point) const
+	[[nodiscard]] T GetParameter(const TVector<T>& Point) const
 	{
 		return TVector<T>::DotProduct((Point - Origin), Direction);
 	}
@@ -124,7 +124,7 @@ public:
 	 * @param Point query Point
 	 * @return squared distance to Ray
 	 */
-	T DistSquared(const TVector<T>& Point) const
+	[[nodiscard]] T DistSquared(const TVector<T>& Point) const
 	{
 		T RayParameter = TVector<T>::DotProduct((Point - Origin), Direction);
 		if (RayParameter < 0)
@@ -144,7 +144,7 @@ public:
 	 * @param Point query Point
 	 * @return distance to Ray
 	 */
-	T Dist(const TVector<T>& Point) const
+	[[nodiscard]] T Dist(const TVector<T>& Point) const
 	{
 		return FMath::Sqrt(DistSquared(Point));
 	}
@@ -154,7 +154,7 @@ public:
 	 * @param Point query point
 	 * @return closest point on Ray
 	 */
-	TVector<T> ClosestPoint(const TVector<T>& Point) const
+	[[nodiscard]] TVector<T> ClosestPoint(const TVector<T>& Point) const
 	{
 		T RayParameter = TVector<T>::DotProduct((Point - Origin), Direction);
 		if (RayParameter < 0) 
@@ -176,7 +176,7 @@ public:
 	*
 	* @return Text describing the Ray.
 	*/
-	FString ToString() const
+	[[nodiscard]] FString ToString() const
 	{
 		return FString::Printf(TEXT("Origin=(%s), Direction=(%s)"), *Origin.ToString(), *Direction.ToString());
 	}
diff --git a/Engine/Source/Runtime/Core/Public/Math/RotationAboutPointMatrix.h b/Engine/Source/Runtime/Core/Public/Math/RotationAboutPointMatrix.h
--- a/Engine/Source/Runtime/Core/Public/Math/RotationAboutPointMatrix.h
+++ b/Engine/Source/Runtime/Core/Public/Math/RotationAboutPointMatrix.h
@@ -25,23 +25,23 @@ public:
 	 * @param Rot rotation
 	 * @param Origin about which to rotate.
 	 */
-	TRotationAboutPointMatrix(const TRotator<T>& Rot, const TVector<T>& Origin);
+	[[nodiscard]] TRotationAboutPointMatrix(const TRotator<T>& Rot, const TVector<T>& Origin);
 
 	// Conversion to other type.
 	template<typename FArg UE_REQUIRES(!std::is_same_v<T, FArg>)>
-	explicit TRotationAboutPointMatrix(const TRotationAboutPointMatrix<FArg>& From)
+	[[nodiscard]] explicit TRotationAboutPointMatrix(const TRotationAboutPointMatrix<FArg>& From)
 		: TRotationTranslationMatrix<T>(From)
 	{
 	}
 
 	/** Matrix factory. Return an TMatrix<T> so we don't have type conversion issues in expressions. */
-	static TMatrix<T> Make(const TRotator<T>& Rot, const TVector<T>& Origin)
+	[[nodiscard]] static TMatrix<T> Make(const TRotator<T>& Rot, const TVector<T>& Origin)
 	{
 		return TRotationAboutPointMatrix(Rot, Origin);
 	}
 
 	/** Matrix factory. Return an TMatrix<T> so we don't have type conversion issues in expressions. */
-	static TMatrix<T> Make(const TQuat<T>& Rot, const TVector<T>& Origin)
+	[[nodiscard]] static TMatrix<T> Make(const TQuat<T>& Rot, const TVector<T>& Origin)
 	{
 		return TRotationAboutPointMatrix(Rot.Rotator(), Origin);
 	}
diff --git a/Engine/Source/Runtime/Core/Public/Math/RotationMatrix.h b/Engine/Source/Runtime/Core/Public/Math/RotationMatrix.h
--- a/Engine/Source/Runtime/Core/Public/Math/RotationMatrix.h
+++ b/Engine/Source/Runtime/Core/Public/Math/RotationMatrix.h
@@ -25,45 +25,45 @@ public:
 	 *
 	 * @param Rot rotation
 	 */
-	TRotationMatrix(const TRotator<T>& Rot)
+	[[nodiscard]] TRotationMatrix(const TRotator<T>& Rot)
 		: TRotationTranslationMatrix<T>(Rot, TVector<T>::ZeroVector)
 	{ }
 
 	/** Matrix factory. Return an TMatrix<T> so we don't have type conversion issues in expressions. */
-	static TMatrix<T> Make(TRotator<T> const& Rot)
+	[[nodiscard]] static TMatrix<T> Make(TRotator<T> const& Rot)
 	{
 		return TRotationMatrix(Rot);
 	}
 
 	/** Matrix factory. Return an TMatrix<T> so we don't have type conversion issues in expressions. */
-	static TMatrix<T> Make(TQuat<T> const& Rot);
+	[[nodiscard]] static TMatrix<T> Make(TQuat<T> const& Rot);
 
 	/** Builds a rotation matrix given only a XAxis. Y and Z are unspecified but will be orthonormal. XAxis need not be normalized. */
-	static TMatrix<T> MakeFromX(TVector<T> const& XAxis);
+	[[nodiscard]] static TMatrix<T> MakeFromX(TVector<T> const& XAxis);
 
 	/** Builds a rotation matrix given only a YAxis. X and Z are unspecified but will be orthonormal. YAxis need not be normalized. */
-	static TMatrix<T> MakeFromY(TVector<T> const& YAxis);
+	[[nodiscard]] static TMatrix<T> MakeFromY(TVector<T> const& YAxis);
 
 	/** Builds a rotation matrix given only a ZAxis. X and Y are unspecified but will be orthonormal. ZAxis need not be normalized. */
-	static TMatrix<T> MakeFromZ(TVector<T> const& ZAxis);
+	[[nodiscard]] static TMatrix<T> MakeFromZ(TVector<T> const& ZAxis);
 
 	/** Builds a matrix with given X and Y axes. X will remain fixed, Y may be changed minimally to enforce orthogonality. Z will be computed. Inputs need not be normalized. */
-	static TMatrix<T> MakeFromXY(TVector<T> const& XAxis, TVector<T> const& YAxis);
+	[[nodiscard]] static TMatrix<T> MakeFromXY(TVector<T> const& XAxis, TVector<T> const& YAxis);
 
 	/** Builds a matrix with given X and Z axes. X will remain fixed, Z may be changed minimally to enforce orthogonality. Y will be computed. Inputs need not be normalized. */
-	static TMatrix<T> MakeFromXZ(TVector<T> const& XAxis, TVector<T> const& ZAxis);
+	[[nodiscard]] static TMatrix<T> MakeFromXZ(TVector<T> const& XAxis, TVector<T> const& ZAxis);
 
 	/** Builds a matrix with given Y and X axes. Y will remain fixed, X may be changed minimally to enforce orthogonality. Z will be computed. Inputs need not be normalized. */
-	static TMatrix<T> MakeFromYX(TVector<T> const& YAxis, TVector<T> const& XAxis);
+	[[nodiscard]] static TMatrix<T> MakeFromYX(TVector<T> const& YAxis, TVector<T> const& XAxis);
 
 	/** Builds a matrix with given Y and Z axes. Y will remain fixed, Z may be changed minimally to enforce orthogonality. X will be computed. Inputs need not be normalized. */
-	static TMatrix<T> MakeFromYZ(TVector<T> const& YAxis, TVector<T> const& ZAxis);
+	[[nodiscard]] static TMatrix<T> MakeFromYZ(TVector<T> const& YAxis, TVector<T> const& ZAxis);
 
 	/** Builds a matrix with given Z and X axes. Z will remain fixed, X may be changed minimally to enforce orthogonality. Y will be computed. Inputs need not be normalized. */
-	static TMatrix<T> MakeFromZX(TVector<T> const& ZAxis, TVector<T> const& XAxis);
+	[[nodiscard]] static TMatrix<T> MakeFromZX(TVector<T> const& ZAxis, TVector<T> const& XAxis);
 
 	/** Builds a matrix with given Z and Y axes. Z will remain fixed, Y may be changed minimally to enforce orthogonality. X will be computed. Inputs need not be normalized. */
-	static TMatrix<T> MakeFromZY(TVector<T> const& ZAxis, TVector<T> const& YAxis);
+	[[nodiscard]] static TMatrix<T> MakeFromZY(TVector<T> const& ZAxis, TVector<T> const& YAxis);
 };
 
 template<typename T>
diff --git a/Engine/Source/Runtime/Core/Public/Math/RotationTranslationMatrix.h b/Engine/Source/Runtime/Core/Public/Math/RotationTranslationMatrix.h
--- a/Engine/Source/Runtime/Core/Public/Math/RotationTranslationMatrix.h
+++ b/Engine/Source/Runtime/Core/Public/Math/RotationTranslationMatrix.h
@@ -24,17 +24,17 @@ public:
 	 * @param Rot rotation
 	 * @param Origin translation to apply
 	 */
-	TRotationTranslationMatrix(const TRotator<T>& Rot, const TVector<T>& Origin);
+	[[nodiscard]] TRotationTranslationMatrix(const TRotator<T>& Rot, const TVector<T>& Origin);
 
 	// Conversion to other type.
 	template<typename FArg UE_REQUIRES(!std::is_same_v<T, FArg>)>
-	explicit TRotationTranslationMatrix(const TRotationTranslationMatrix<FArg>& From)
+	[[nodiscard]] explicit TRotationTranslationMatrix(const TRotationTranslationMatrix<FArg>& From)
 		: TMatrix<T>(From)
 	{
 	}
 	
 	/** Matrix factory. Return an TMatrix<T> so we don't have type conversion issues in expressions. */
-	static TMatrix<T> Make(const TRotator<T>& Rot, const TVector<T>& Origin)
+	[[nodiscard]] static TMatrix<T> Make(const TRotator<T>& Rot, const TVector<T>& Origin)
 	{
 		return TRotationTranslationMatrix(Rot, Origin);
 	}
