Make Maths functions [[nodiscard]]

diff --git a/Engine/Source/Runtime/Core/Public/Math/Rotator.h b/Engine/Source/Runtime/Core/Public/Math/Rotator.h
--- a/Engine/Source/Runtime/Core/Public/Math/Rotator.h
+++ b/Engine/Source/Runtime/Core/Public/Math/Rotator.h
@@ -84,14 +84,14 @@ public:
 	/**
 	 * Default constructor (no initialization).
 	 */
-	TRotator() = default;
+	[[nodiscard]] TRotator() = default;
 
 	/**
 	 * Constructor
 	 *
 	 * @param InF Value to set all components to.
 	 */
-	explicit FORCEINLINE TRotator(T InF);
+	[[nodiscard]] explicit FORCEINLINE TRotator(T InF);
 
 	/**
 	 * Constructor.
@@ -100,21 +100,21 @@ public:
 	 * @param InYaw Yaw in degrees.
 	 * @param InRoll Roll in degrees.
 	 */
-	FORCEINLINE TRotator( T InPitch, T InYaw, T InRoll );
+	[[nodiscard]] FORCEINLINE TRotator( T InPitch, T InYaw, T InRoll );
 
 	/**
 	 * Constructor.
 	 *
 	 * @param EForceInit Force Init Enum.
 	 */
-	explicit FORCEINLINE TRotator( EForceInit );
+	[[nodiscard]] explicit FORCEINLINE TRotator( EForceInit );
 
 	/**
 	 * Constructor.
 	 *
 	 * @param Quat Quaternion used to specify rotation.
 	 */
-	explicit CORE_API TRotator( const TQuat<T>& Quat );
+[[nodiscard]] explicit CORE_API TRotator( const TQuat<T>& Quat );
 
 public:
 
@@ -126,7 +126,7 @@ public:
 	 * @param R The other rotator.
 	 * @return The result of adding a rotator to this.
 	 */
-	TRotator operator+( const TRotator<T>& R ) const;
+	[[nodiscard]] TRotator operator+( const TRotator<T>& R ) const;
 
 	/**
 	 * Get the result of subtracting a rotator from this.
@@ -134,7 +134,7 @@ public:
 	 * @param R The other rotator.
 	 * @return The result of subtracting a rotator from this.
 	 */
-	TRotator operator-( const TRotator<T>& R ) const;
+	[[nodiscard]] TRotator operator-( const TRotator<T>& R ) const;
 
 	/**
 	 * Get the result of scaling this rotator.
@@ -143,7 +143,7 @@ public:
 	 * @return The result of scaling.
 	 */
 	template<typename FArg UE_REQUIRES(std::is_arithmetic_v<FArg>)>
-	FORCEINLINE TRotator operator*( FArg Scale ) const
+	[[nodiscard]] FORCEINLINE TRotator operator*( FArg Scale ) const
 	{
 		return TRotator(Pitch * Scale, Yaw * Scale, Roll * Scale);
 	}
@@ -171,7 +171,7 @@ public:
 	 * @return true if two rotators are identical, otherwise false.
 	 * @see Equals()
 	 */
-	bool operator==( const TRotator<T>& R ) const;
+	[[nodiscard]] bool operator==( const TRotator<T>& R ) const;
 
 	/**
 	 * Checks whether two rotators are different.
@@ -179,7 +179,7 @@ public:
 	 * @param V The other rotator.
 	 * @return true if two rotators are different, otherwise false.
 	 */
-	bool operator!=( const TRotator<T>& V ) const;
+	[[nodiscard]] bool operator!=( const TRotator<T>& V ) const;
 
 	// Assignment operators.
 
@@ -210,7 +210,7 @@ public:
 	 * @param Tolerance Error Tolerance.
 	 * @return true if rotator is nearly zero, within specified tolerance, otherwise false.
 	 */
-	bool IsNearlyZero( T Tolerance = UE_KINDA_SMALL_NUMBER ) const;
+	[[nodiscard]] bool IsNearlyZero( T Tolerance = UE_KINDA_SMALL_NUMBER ) const;
 
 	/**
 	 * Checks whether this has exactly zero rotation, when treated as an orientation.
@@ -218,7 +218,7 @@ public:
 	 *
 	 * @return true if this has exactly zero rotation, otherwise false.
 	 */
-	bool IsZero() const;
+	[[nodiscard]] bool IsZero() const;
 
 	/**
 	 * Checks whether two rotators are equal within specified tolerance, when treated as an orientation.
@@ -231,7 +231,7 @@ public:
 	 * @param Tolerance Error Tolerance.
 	 * @return true if two rotators are equal, within specified tolerance, otherwise false.
 	 */
-	bool Equals( const TRotator<T>& R, T Tolerance = UE_KINDA_SMALL_NUMBER ) const;
+	[[nodiscard]] bool Equals( const TRotator<T>& R, T Tolerance = UE_KINDA_SMALL_NUMBER ) const;
 
 	/**
 	 * Checks whether two rotators have the same orientation within the specified tolerance, without requiring similar angles or wound multiples of 360.
@@ -245,7 +245,7 @@ public:
 	 * @param Tolerance Error Tolerance.
 	 * @return true if two rotators are equal, within specified tolerance, otherwise false.
 	 */
-	bool EqualsOrientation( const TRotator<T>& R, T Tolerance = UE_KINDA_SMALL_NUMBER ) const;
+	[[nodiscard]] bool EqualsOrientation( const TRotator<T>& R, T Tolerance = UE_KINDA_SMALL_NUMBER ) const;
 
 	/**
 	 * Adds to each component of the rotator.
@@ -260,7 +260,7 @@ public:
 	/**
 	 * Returns the inverse of the rotator.
 	 */
-	CORE_API TRotator GetInverse() const;
+	[[nodiscard]] CORE_API TRotator GetInverse() const;
 
 	/**
 	 * Get the rotation, snapped to specified degree segments.
@@ -268,28 +268,28 @@ public:
 	 * @param RotGrid A Rotator specifying how to snap each component.
 	 * @return Snapped version of rotation.
 	 */
-	TRotator GridSnap( const TRotator<T>& RotGrid ) const;
+	[[nodiscard]] TRotator GridSnap( const TRotator<T>& RotGrid ) const;
 
 	/**
 	 * Convert a rotation into a unit vector facing in its direction.
 	 *
 	 * @return Rotation as a unit direction vector.
 	 */
-	CORE_API TVector<T> Vector() const;
+	[[nodiscard]] CORE_API TVector<T> Vector() const;
 
 	/**
 	 * Get Rotation as a quaternion.
 	 *
 	 * @return Rotation as a quaternion.
 	 */
-	CORE_API TQuat<T> Quaternion() const;
+	[[nodiscard]] CORE_API TQuat<T> Quaternion() const;
 
 	/**
 	 * Convert a Rotator into floating-point Euler angles (in degrees). Rotator now stored in degrees.
 	 *
 	 * @return Rotation as a Euler angle vector.
 	 */
-	CORE_API TVector<T> Euler() const;
+	[[nodiscard]] CORE_API TVector<T> Euler() const;
 
 	/**
 	 * Rotate a vector rotated by this rotator.
@@ -297,7 +297,7 @@ public:
 	 * @param V The vector to rotate.
 	 * @return The rotated vector.
 	 */
-	CORE_API TVector<T> RotateVector( const UE::Math::TVector<T>& V ) const;
+	[[nodiscard]] CORE_API TVector<T> RotateVector( const UE::Math::TVector<T>& V ) const;
 
 	/**
 	 * Returns the vector rotated by the inverse of this rotator.
@@ -305,31 +305,31 @@ public:
 	 * @param V The vector to rotate.
 	 * @return The rotated vector.
 	 */
-	CORE_API TVector<T> UnrotateVector( const UE::Math::TVector<T>& V ) const;
+	[[nodiscard]] CORE_API TVector<T> UnrotateVector( const UE::Math::TVector<T>& V ) const;
 
 	/**
 	 * Gets the rotation values so they fall within the range [0,360]
 	 *
 	 * @return Clamped version of rotator.
 	 */
-	TRotator<T> Clamp() const;
+	[[nodiscard]] TRotator<T> Clamp() const;
 
 	/** 
 	 * Create a copy of this rotator and normalize, removes all winding and creates the "shortest route" rotation. 
 	 *
 	 * @return Normalized copy of this rotator
 	 */
-	TRotator<T> GetNormalized() const;
+	[[nodiscard]] TRotator<T> GetNormalized() const;
 
 	/** 
 	 * Create a copy of this rotator and denormalize, clamping each axis to 0 - 360. 
 	 *
 	 * @return Denormalized copy of this rotator
 	 */
-	TRotator<T> GetDenormalized() const;
+	[[nodiscard]] TRotator<T> GetDenormalized() const;
 
 	/** Get a specific component of the vector, given a specific axis by enum */
-	T GetComponentForAxis(EAxis::Type Axis) const;
+	[[nodiscard]] T GetComponentForAxis(EAxis::Type Axis) const;
 
 	/** Set a specified componet of the vector, given a specific axis by enum */
 	void SetComponentForAxis(EAxis::Type Axis, T Component);
@@ -353,14 +353,14 @@ public:
 	* @param Rotator[In] the Rotator we are comparing with.
 	* @return Distance(Manhattan) between the two rotators. 
 	*/
-	T GetManhattanDistance(const TRotator<T> & Rotator) const;
+	[[nodiscard]] T GetManhattanDistance(const TRotator<T> & Rotator) const;
 
 	/**
 	* Return a Rotator that has the same rotation but has different degree values for Yaw, Pitch, and Roll.
 	* This rotator should be within -180,180 range,
 	* @return A Rotator with the same rotation but different degrees.
 	*/
-	TRotator GetEquivalentRotator() const;
+	[[nodiscard]] TRotator GetEquivalentRotator() const;
 
 	/**
 	* Modify if necessary the passed in rotator to be the closest rotator to it based upon it's equivalent.
@@ -376,10 +376,10 @@ public:
 	 *
 	 * @return Text describing the vector.
 	 */
-	FString ToString() const;
+	[[nodiscard]] FString ToString() const;
 
 	/** Get a short textural representation of this vector, for compact readable logging. */
-	FString ToCompactString() const;
+	[[nodiscard]] FString ToCompactString() const;
 
 	/**
 	 * Initialize this Rotator based on an FString. The String is expected to contain P=, Y=, R=.
@@ -395,7 +395,7 @@ public:
 	 *
 	 * @return true if there are any non-finite values in this Rotator, otherwise false.
 	 */
-	bool ContainsNaN() const;
+	[[nodiscard]] bool ContainsNaN() const;
 
 	/**
 	 * Serializes the rotator compressed for e.g. network transmission.
@@ -423,7 +423,7 @@ public:
 	 * @param Angle The angle to clamp.
 	 * @return The clamped angle.
 	 */
-	static T ClampAxis( T Angle );
+	[[nodiscard]] static T ClampAxis( T Angle );
 
 	/**
 	 * Clamps an angle to the range of (-180, 180].
@@ -431,7 +431,7 @@ public:
 	 * @param Angle The Angle to clamp.
 	 * @return The clamped angle.
 	 */
-	static T NormalizeAxis( T Angle );
+	[[nodiscard]] static T NormalizeAxis( T Angle );
 
 	/**
 	 * Compresses a floating point angle into a byte.
@@ -439,7 +439,7 @@ public:
 	 * @param Angle The angle to compress.
 	 * @return The angle as a byte.
 	 */
-	static uint8 CompressAxisToByte( T Angle );
+	[[nodiscard]] static uint8 CompressAxisToByte( T Angle );
 
 	/**
 	 * Decompress a word into a floating point angle.
@@ -447,7 +447,7 @@ public:
 	 * @param Angle The word angle.
 	 * @return The decompressed angle.
 	 */
-	static T DecompressAxisFromByte( uint8 Angle );
+	[[nodiscard]] static T DecompressAxisFromByte( uint8 Angle );
 
 	/**
 	 * Compress a floating point angle into a word.
@@ -455,7 +455,7 @@ public:
 	 * @param Angle The angle to compress.
 	 * @return The decompressed angle.
 	 */
-	static uint16 CompressAxisToShort( T Angle );
+	[[nodiscard]] static uint16 CompressAxisToShort( T Angle );
 
 	/**
 	 * Decompress a short into a floating point angle.
@@ -463,7 +463,7 @@ public:
 	 * @param Angle The word angle.
 	 * @return The decompressed angle.
 	 */
-	static T DecompressAxisFromShort( uint16 Angle );
+	[[nodiscard]] static T DecompressAxisFromShort( uint16 Angle );
 
 	/**
 	 * Convert a vector of floating-point Euler angles (in degrees) into a Rotator. Rotator now stored in degrees
@@ -471,7 +471,7 @@ public:
 	 * @param Euler Euler angle vector.
 	 * @return A rotator from a Euler angle.
 	 */
-	static CORE_API TRotator MakeFromEuler( const TVector<T>& Euler );
+	[[nodiscard]] static CORE_API TRotator MakeFromEuler( const TVector<T>& Euler );
 
 
 public:
diff --git a/Engine/Source/Runtime/Core/Public/Math/SHMath.h b/Engine/Source/Runtime/Core/Public/Math/SHMath.h
--- a/Engine/Source/Runtime/Core/Public/Math/SHMath.h
+++ b/Engine/Source/Runtime/Core/Public/Math/SHMath.h
@@ -46,12 +46,12 @@ public:
 	static constexpr float ConstantBasisIntegral = 3.5449077018110320545963349666823f; // 2 * Sqrt(PI)
 
 	/** Default constructor. */
-	TSHVector()
+	[[nodiscard]] TSHVector()
 	{
 		FMemory::Memzero(V,sizeof(V));
 	}
 
-	TSHVector(float V0, float V1, float V2, float V3)
+	[[nodiscard]] TSHVector(float V0, float V1, float V2, float V3)
 	{
 		FMemory::Memzero(V,sizeof(V));
 
@@ -61,7 +61,7 @@ public:
 		V[3] = V3;
 	}
 
-	explicit TSHVector(const FVector4f& Vector)
+	[[nodiscard]] explicit TSHVector(const FVector4f& Vector)
 	{
 		FMemory::Memzero(V,sizeof(V));
 
@@ -73,7 +73,7 @@ public:
 
 
 	template<int32 OtherOrder>
-	explicit TSHVector(const TSHVector<OtherOrder>& Other)
+	[[nodiscard]] explicit TSHVector(const TSHVector<OtherOrder>& Other)
 	{
 		if (Order <= OtherOrder)
 		{
@@ -88,7 +88,7 @@ public:
 
 	/** Scalar multiplication operator. */
 	/** Changed to float& from float to avoid LHS **/
-	friend FORCEINLINE TSHVector operator*(const TSHVector& A,const float& B)
+	[[nodiscard]] friend FORCEINLINE TSHVector operator*(const TSHVector& A,const float& B)
 	{
 		const VectorRegister4Float ReplicatedScalar = VectorLoadFloat1(&B);
 
@@ -105,7 +105,7 @@ public:
 	}
 
 	/** Scalar division operator. */
-	friend FORCEINLINE TSHVector operator/(const TSHVector& A,const float& Scalar)
+	[[nodiscard]] friend FORCEINLINE TSHVector operator/(const TSHVector& A,const float& Scalar)
 	{
 		const float B = (1.0f / Scalar);
 		const VectorRegister4Float ReplicatedScalar = VectorLoadFloat1(&B);
@@ -123,7 +123,7 @@ public:
 	}
 
 	/** Addition operator. */
-	friend FORCEINLINE TSHVector operator+(const TSHVector& A,const TSHVector& B)
+	[[nodiscard]] friend FORCEINLINE TSHVector operator+(const TSHVector& A,const TSHVector& B)
 	{
 		TSHVector Result;
 		for(int32 BasisIndex = 0;BasisIndex < NumSIMDVectors;BasisIndex++)
@@ -139,7 +139,7 @@ public:
 	}
 	
 	/** Subtraction operator. */
-	friend FORCEINLINE TSHVector operator-(const TSHVector& A,const TSHVector& B)
+	[[nodiscard]] friend FORCEINLINE TSHVector operator-(const TSHVector& A,const TSHVector& B)
 	{
 		TSHVector Result;
 		for(int32 BasisIndex = 0;BasisIndex < NumSIMDVectors;BasisIndex++)
@@ -155,7 +155,7 @@ public:
 	}
 
 	/** Dot product operator. */
-	friend FORCEINLINE float Dot(const TSHVector& A,const TSHVector& B)
+	[[nodiscard]] friend FORCEINLINE float Dot(const TSHVector& A,const TSHVector& B)
 	{
 		VectorRegister4Float ReplicatedResult = VectorZero();
 		for(int32 BasisIndex = 0;BasisIndex < NumSIMDVectors;BasisIndex++)
@@ -257,7 +257,7 @@ public:
 	}
 
 	/** Calculates the integral of the function over the surface of the sphere. */
-	float CalcIntegral() const
+	[[nodiscard]] float CalcIntegral() const
 	{
 		return V[0] * ConstantBasisIntegral;
 	}
@@ -288,7 +288,7 @@ public:
 		}
 	}
 
-	bool AreFloatsValid() const
+	[[nodiscard]] bool AreFloatsValid() const
 	{
 		bool bValid = true;
 
@@ -301,13 +301,13 @@ public:
 	}
 
 	/** Compute the direction which the spherical harmonic is highest at. */
-	FVector GetMaximumDirection() const
+	[[nodiscard]] FVector GetMaximumDirection() const
 	{
 		// This is an approximation which only takes into account first and second order spherical harmonics.
 		return FVector(-V[3], -V[1], V[2]).GetSafeNormal();
 	}
 
-	static TSHVector CalcDiffuseTransfer(const FVector& Normal)
+	[[nodiscard]] static TSHVector CalcDiffuseTransfer(const FVector& Normal)
 	{
 		TSHVector Result = SHBasisFunction(Normal);
 
@@ -342,7 +342,7 @@ public:
 	}
 
 	/** Returns the value of the SH basis L,M at the point on the sphere defined by the unit vector Vector. */
-	static TSHVector SHBasisFunction(const FVector& Vector)
+	[[nodiscard]] static TSHVector SHBasisFunction(const FVector& Vector)
 	{
 		TSHVector Result;
 
@@ -382,14 +382,14 @@ public:
 	}
 
 	/** The ambient incident lighting function. */
-	static TSHVector AmbientFunction()
+	[[nodiscard]] static TSHVector AmbientFunction()
 	{
 		TSHVector AmbientFunctionSH;
 		AmbientFunctionSH.V[0] = 1.0f / (2.0f * FMath::Sqrt(UE_PI));
 		return AmbientFunctionSH;
 	}
 
-	static float FindWindowingLambda(const TSHVector& Vector, float TargetLaplacian)
+	[[nodiscard]] static float FindWindowingLambda(const TSHVector& Vector, float TargetLaplacian)
 	{
 		// "Stupid Spherical Harmonics (SH) Tricks"
 		// Appendix A7: Solving for Lamba to Reduce the Squared Laplacian
@@ -459,7 +459,7 @@ public:
 
 /** Specialization for 2nd order to avoid expensive trig functions. */
 template<> 
-inline TSHVector<2> TSHVector<2>::SHBasisFunction(const FVector& Vector) 
+[[nodiscard]] inline TSHVector<2> TSHVector<2>::SHBasisFunction(const FVector& Vector) 
 {
 	TSHVector<2> Result;
 	Result.V[0] = 0.282095f; 
@@ -471,7 +471,7 @@ inline TSHVector<2> TSHVector<2>::SHBasisFunction(const FVector& Vector)
 
 /** Specialization for 3rd order to avoid expensive trig functions. */
 template<> 
-inline TSHVector<3> TSHVector<3>::SHBasisFunction(const FVector& Vector) 
+[[nodiscard]] inline TSHVector<3> TSHVector<3>::SHBasisFunction(const FVector& Vector) 
 {
 	TSHVector<3> Result;
 	Result.V[0] = 0.282095f; 
@@ -509,7 +509,7 @@ public:
 	}
 
 	/** Calculates greyscale spherical harmonic coefficients. */
-	TSHVector<MaxSHOrder> GetLuminance() const
+	[[nodiscard]] TSHVector<MaxSHOrder> GetLuminance() const
 	{
 		static const FLinearColor LuminanceFactors = SHGetLuminanceFactors();
 
@@ -526,7 +526,7 @@ public:
 	}
 
 	/** Calculates the integral of the function over the surface of the sphere. */
-	FLinearColor CalcIntegral() const
+	[[nodiscard]] FLinearColor CalcIntegral() const
 	{
 		FLinearColor Result;
 		Result.R = R.CalcIntegral();
@@ -543,14 +543,14 @@ public:
 		B.ApplyWindowing(Lambda);
 	}
 
-	bool AreFloatsValid() const
+	[[nodiscard]] bool AreFloatsValid() const
 	{
 		return R.AreFloatsValid() && G.AreFloatsValid() && B.AreFloatsValid();
 	}
 
 	/** Scalar multiplication operator. */
 	/** Changed to float& from float to avoid LHS **/
-	friend FORCEINLINE TSHVectorRGB operator*(const TSHVectorRGB& A, const float& Scalar)
+	[[nodiscard]] friend FORCEINLINE TSHVectorRGB operator*(const TSHVectorRGB& A, const float& Scalar)
 	{
 		TSHVectorRGB Result;
 		Result.R = A.R * Scalar;
@@ -561,7 +561,7 @@ public:
 
 	/** Scalar multiplication operator. */
 	/** Changed to float& from float to avoid LHS **/
-	friend FORCEINLINE TSHVectorRGB operator*(const float& Scalar,const TSHVectorRGB& A)
+	[[nodiscard]] friend FORCEINLINE TSHVectorRGB operator*(const float& Scalar,const TSHVectorRGB& A)
 	{
 		TSHVectorRGB Result;
 		Result.R = A.R * Scalar;
@@ -571,7 +571,7 @@ public:
 	}
 
 	/** Color multiplication operator. */
-	friend FORCEINLINE TSHVectorRGB operator*(const TSHVectorRGB& A,const FLinearColor& Color)
+	[[nodiscard]] friend FORCEINLINE TSHVectorRGB operator*(const TSHVectorRGB& A,const FLinearColor& Color)
 	{
 		TSHVectorRGB Result;
 		Result.R = A.R * Color.R;
@@ -581,7 +581,7 @@ public:
 	}
 
 	/** Color multiplication operator. */
-	friend FORCEINLINE TSHVectorRGB operator*(const FLinearColor& Color,const TSHVectorRGB& A)
+	[[nodiscard]] friend FORCEINLINE TSHVectorRGB operator*(const FLinearColor& Color,const TSHVectorRGB& A)
 	{
 		TSHVectorRGB Result;
 		Result.R = A.R * Color.R;
@@ -591,7 +591,7 @@ public:
 	}
 
 	/** Division operator. */
-	friend FORCEINLINE TSHVectorRGB operator/(const TSHVectorRGB& A,const float& InB)
+	[[nodiscard]] friend FORCEINLINE TSHVectorRGB operator/(const TSHVectorRGB& A,const float& InB)
 	{
 		TSHVectorRGB Result;
 		Result.R = A.R / InB;
@@ -601,7 +601,7 @@ public:
 	}
 
 	/** Addition operator. */
-	friend FORCEINLINE TSHVectorRGB operator+(const TSHVectorRGB& A,const TSHVectorRGB& InB)
+	[[nodiscard]] friend FORCEINLINE TSHVectorRGB operator+(const TSHVectorRGB& A,const TSHVectorRGB& InB)
 	{
 		TSHVectorRGB Result;
 		Result.R = A.R + InB.R;
@@ -611,7 +611,7 @@ public:
 	}
 	
 	/** Subtraction operator. */
-	friend FORCEINLINE TSHVectorRGB operator-(const TSHVectorRGB& A,const TSHVectorRGB& InB)
+	[[nodiscard]] friend FORCEINLINE TSHVectorRGB operator-(const TSHVectorRGB& A,const TSHVectorRGB& InB)
 	{
 		TSHVectorRGB Result;
 		Result.R = A.R - InB.R;
@@ -621,7 +621,7 @@ public:
 	}
 
 	/** Dot product operator. */
-	friend FORCEINLINE FLinearColor Dot(const TSHVectorRGB& A,const TSHVector<MaxSHOrder>& InB)
+	[[nodiscard]] friend FORCEINLINE FLinearColor Dot(const TSHVectorRGB& A,const TSHVector<MaxSHOrder>& InB)
 	{
 		FLinearColor Result;
 		Result.R = Dot(A.R,InB);
diff --git a/Engine/Source/Runtime/Core/Public/Math/ScalarRegister.h b/Engine/Source/Runtime/Core/Public/Math/ScalarRegister.h
--- a/Engine/Source/Runtime/Core/Public/Math/ScalarRegister.h
+++ b/Engine/Source/Runtime/Core/Public/Math/ScalarRegister.h
@@ -32,19 +32,19 @@ public:
 	RegisterType Value;
 
 	/** default constructor */
-	FORCEINLINE ScalarRegister();
+	[[nodiscard]] FORCEINLINE ScalarRegister();
 
 	/** Copy Constructor */
-	FORCEINLINE ScalarRegister(const ScalarRegister& VectorValue);
+	[[nodiscard]] FORCEINLINE ScalarRegister(const ScalarRegister& VectorValue);
 
 	/** Constructor using float value */
-	explicit FORCEINLINE ScalarRegister(const float& ScalarValue);
+	[[nodiscard]] explicit FORCEINLINE ScalarRegister(const float& ScalarValue);
 
 	/** Constructor
 	 *
 	 * @param RegisterType float4 vector register type
 	 */
-	explicit FORCEINLINE ScalarRegister(RegisterType VectorValue);
+	[[nodiscard]] explicit FORCEINLINE ScalarRegister(RegisterType VectorValue);
 
 	/**
 	 * Gets the result of multiplying a scalar register to this.
@@ -53,7 +53,7 @@ public:
 	 *
 	 * @return The result of multiplication.
 	 */
-	FORCEINLINE ScalarRegister operator*(const ScalarRegister& RHS) const;
+	[[nodiscard]] FORCEINLINE ScalarRegister operator*(const ScalarRegister& RHS) const;
 
 	/**
 	 * Gets the result of adding a scalar register to this.
@@ -62,7 +62,7 @@ public:
 	 *
 	 * @return The result of addition.
 	 */
-	FORCEINLINE ScalarRegister operator+(const ScalarRegister& RHS) const;
+	[[nodiscard]] FORCEINLINE ScalarRegister operator+(const ScalarRegister& RHS) const;
 
 	/**
 	 * Adds to this scalar register.
@@ -89,7 +89,7 @@ public:
 	 *
 	 * @return The result of subtraction.
 	 */
-	FORCEINLINE ScalarRegister operator-(const ScalarRegister& RHS) const;
+	[[nodiscard]] FORCEINLINE ScalarRegister operator-(const ScalarRegister& RHS) const;
 
 	
 	/** assignment operator
diff --git a/Engine/Source/Runtime/Core/Public/Math/ScaleMatrix.h b/Engine/Source/Runtime/Core/Public/Math/ScaleMatrix.h
--- a/Engine/Source/Runtime/Core/Public/Math/ScaleMatrix.h
+++ b/Engine/Source/Runtime/Core/Public/Math/ScaleMatrix.h
@@ -21,28 +21,28 @@ public:
 	/**
 	 * @param Scale uniform scale to apply to matrix.
 	 */
-	TScaleMatrix( T Scale );
+	[[nodiscard]] TScaleMatrix( T Scale );
 
 	/**
 	 * @param Scale Non-uniform scale to apply to matrix.
 	 */
-	TScaleMatrix( const TVector<T>& Scale );
+	[[nodiscard]] TScaleMatrix( const TVector<T>& Scale );
 
 	// Conversion to other type.
 	template<typename FArg UE_REQUIRES(!std::is_same_v<T, FArg>)>
-	explicit TScaleMatrix(const TScaleMatrix<FArg>& From)
+	[[nodiscard]] explicit TScaleMatrix(const TScaleMatrix<FArg>& From)
 		: TMatrix<T>(From)
 	{
 	}
 	
 	/** Matrix factory. Return an TMatrix<T> so we don't have type conversion issues in expressions. */
-	static TMatrix<T> Make(T Scale)
+	[[nodiscard]] static TMatrix<T> Make(T Scale)
 	{
 		return TScaleMatrix<T>(Scale);
 	}
 
 	/** Matrix factory. Return an TMatrix<T> so we don't have type conversion issues in expressions. */
-	static TMatrix<T> Make(const TVector<T>& Scale)
+	[[nodiscard]] static TMatrix<T> Make(const TVector<T>& Scale)
 	{
 		return TScaleMatrix<T>(Scale);
 	}
diff --git a/Engine/Source/Runtime/Core/Public/Math/Sobol.h b/Engine/Source/Runtime/Core/Public/Math/Sobol.h
--- a/Engine/Source/Runtime/Core/Public/Math/Sobol.h
+++ b/Engine/Source/Runtime/Core/Public/Math/Sobol.h
@@ -52,7 +52,7 @@ public:
 	 * @param Seed - A 24-bit random seed when reusing the same dimension
 	 * @return The Sobol result at the given Index
 	 */
-	CORE_API static float Evaluate(int32 Index, int32 Dim, int32 Seed = 0);
+	[[nodiscard]] CORE_API static float Evaluate(int32 Index, int32 Dim, int32 Seed = 0);
 
 	/**
 	* Evaluate next Sobol number from one of the traditional Sobol dimensions
@@ -62,7 +62,7 @@ public:
 	* @param Value - The value for the Sobol number at Index-1.
 	* @return The Sobol result at the given Index
 	*/
-	CORE_API static float Next(int32 Index, int32 Dim, float Value);
+	[[nodiscard]] CORE_API static float Next(int32 Index, int32 Dim, float Value);
 
 
 	/**
@@ -74,7 +74,7 @@ public:
 	* @param Seed - A 24-bit per component 2D seed for shuffling values
 	* @return The 2D Sobol result in the range 0-1 given Index
 	*/
-	CORE_API static FVector2D Evaluate(int32 Index, int32 CellBits, FIntPoint Cell, FIntPoint Seed);
+	[[nodiscard]] CORE_API static FVector2D Evaluate(int32 Index, int32 CellBits, FIntPoint Cell, FIntPoint Seed);
 
 	/**
 	* Evaluate next Sobol number from within a 2D cell
@@ -84,7 +84,7 @@ public:
 	* @param Value - The value for the Sobol number at Index-1.
 	* @return The 2D Sobol result in the range 0-1 given Index
 	*/
-	CORE_API static FVector2D Next(int32 Index, int32 CellBits, FVector2D Value);
+	[[nodiscard]] CORE_API static FVector2D Next(int32 Index, int32 CellBits, FVector2D Value);
 
 
 	/**
@@ -96,7 +96,7 @@ public:
 	* @param Seed - A seed for shuffling values (0-1)
 	* @return The Sobol result given Index
 	*/
-	CORE_API static FVector Evaluate(int32 Index, int32 CellBits, FIntVector Cell, FIntVector Seed);
+	[[nodiscard]] CORE_API static FVector Evaluate(int32 Index, int32 CellBits, FIntVector Cell, FIntVector Seed);
 
 	/**
 	* Evaluate next Sobol number from within a 3D cell
@@ -106,8 +106,8 @@ public:
 	* @param Value - The value for the Sobol number at Index-1.
 	* @return The Sobol result at given Index
 	*/
-	CORE_API static FVector Next(int32 Index, int32 CellBits, FVector Value);
+	[[nodiscard]] CORE_API static FVector Next(int32 Index, int32 CellBits, FVector Value);
 
 	/** Return value for GSystemTextures.SobolSampling */
-	CORE_API static uint16 ComputeGPUSpatialSeed(int32 x, int32 y, int32 Index);
+	[[nodiscard]] CORE_API static uint16 ComputeGPUSpatialSeed(int32 x, int32 y, int32 Index);
 };
\ No newline at end of file
diff --git a/Engine/Source/Runtime/Core/Public/Math/Sphere.h b/Engine/Source/Runtime/Core/Public/Math/Sphere.h
--- a/Engine/Source/Runtime/Core/Public/Math/Sphere.h
+++ b/Engine/Source/Runtime/Core/Public/Math/Sphere.h
@@ -38,14 +38,14 @@ public:
 
 
 	/** Default constructor (no initialization). */
-	TSphere() = default;
+	[[nodiscard]] TSphere() = default;
 
 	/**
 	 * Creates and initializes a new sphere.
 	 *
 	 * @param int32 Passing int32 sets up zeroed sphere.
 	 */
-	TSphere(int32)
+	[[nodiscard]] TSphere(int32)
 		: Center(0.0f, 0.0f, 0.0f)
 		, W(0)
 	{ }
@@ -56,7 +56,7 @@ public:
 	 * @param InV Center of sphere.
 	 * @param InW Radius of sphere.
 	 */
-	TSphere(TVector<T> InV, T InW)
+	[[nodiscard]] TSphere(TVector<T> InV, T InW)
 		: Center(InV)
 		, W(InW)
 	{ }
@@ -66,7 +66,7 @@ public:
 	 *
 	 * @param EForceInit Force Init Enum.
 	 */
-	explicit FORCEINLINE TSphere(EForceInit)
+	[[nodiscard]] explicit FORCEINLINE TSphere(EForceInit)
 		: Center(ForceInit)
 		, W(0.0f)
 	{ }
@@ -77,7 +77,7 @@ public:
 	 * @param Points Pointer to list of points this sphere must contain.
 	 * @param Count How many points are in the list.
 	 */
-	CORE_API TSphere(const TVector<T>* Points, int32 Count);
+	[[nodiscard]] CORE_API TSphere(const TVector<T>* Points, int32 Count);
 
 
 	/**
@@ -86,11 +86,11 @@ public:
 	 * @param Spheres Pointer to list of spheres this sphere must contain.
 	 * @param Count How many points are in the list.
 	 */
-	CORE_API TSphere(const TSphere<T>* Spheres, int32 Count);
+	[[nodiscard]] CORE_API TSphere(const TSphere<T>* Spheres, int32 Count);
 
 	// Conversion from other variant type.
 	template<typename FArg UE_REQUIRES(!std::is_same_v<T, FArg>)>
-	explicit TSphere(const TSphere<FArg>& From)
+	[[nodiscard]] explicit TSphere(const TSphere<FArg>& From)
 		: TSphere<T>(TVector<T>(From.Center), T(From.W))
 	{
 	}
@@ -102,7 +102,7 @@ public:
 	 * @param Tolerance Error Tolerance.
 	 * @return true if spheres are equal within specified tolerance, otherwise false.
 	 */
-	bool Equals(const TSphere<T>& Sphere, T Tolerance = UE_KINDA_SMALL_NUMBER) const
+	[[nodiscard]] bool Equals(const TSphere<T>& Sphere, T Tolerance = UE_KINDA_SMALL_NUMBER) const
 	{
 		return Center.Equals(Sphere.Center, Tolerance) && FMath::Abs(W - Sphere.W) <= Tolerance;
 	}
@@ -113,7 +113,7 @@ public:
 	* @param Other The other sphere to compare with.
 	* @return true if the spheres are equal, false otherwise.
 	*/
-	bool operator==(const TSphere<T>& Other) const
+	[[nodiscard]] bool operator==(const TSphere<T>& Other) const
 	{
 		return Center == Other.Center && W == Other.W;
 	}
@@ -124,7 +124,7 @@ public:
 	 * @param Other The other sphere to compare with.
 	 * @return true if the spheres are not equal, false otherwise.
 	 */
-	bool operator!=(const TSphere<T>& Other) const
+	[[nodiscard]] bool operator!=(const TSphere<T>& Other) const
 	{
 		return !(*this == Other);
 	}
@@ -135,7 +135,7 @@ public:
 	 * @param Other The other volume to add to this.
 	 * @return A new bounding volume.
 	 */
-	TSphere<T> operator+(const TSphere<T>& Other) const
+	[[nodiscard]] TSphere<T> operator+(const TSphere<T>& Other) const
 	{
 		return TSphere<T>(*this) += Other;
 	}
@@ -155,7 +155,7 @@ public:
 	 * @param Tolerance Error Tolerance.
 	 * @return true if sphere is inside another, otherwise false.
 	 */
-	bool IsInside(const TSphere<T>& Other, T Tolerance = UE_KINDA_SMALL_NUMBER) const
+	[[nodiscard]] bool IsInside(const TSphere<T>& Other, T Tolerance = UE_KINDA_SMALL_NUMBER) const
 	{
 		if (W > Other.W + Tolerance)
 		{
@@ -171,7 +171,7 @@ public:
 	* @param In The location to test for inside the bounding volume.
 	* @return true if location is inside this volume.
 	*/
-	bool IsInside(const FVector& In, T Tolerance = UE_KINDA_SMALL_NUMBER) const
+	[[nodiscard]] bool IsInside(const FVector& In, T Tolerance = UE_KINDA_SMALL_NUMBER) const
 	{
 		return (Center - In).SizeSquared() <= FMath::Square(W + Tolerance);
 	}
@@ -183,7 +183,7 @@ public:
 	 * @param  Tolerance Error tolerance.
 	 * @return true if spheres intersect, false otherwise.
 	 */
-	FORCEINLINE bool Intersects(const TSphere<T>& Other, T Tolerance = UE_KINDA_SMALL_NUMBER) const
+	[[nodiscard]] FORCEINLINE bool Intersects(const TSphere<T>& Other, T Tolerance = UE_KINDA_SMALL_NUMBER) const
 	{
 		return (Center - Other.Center).SizeSquared() <= FMath::Square(FMath::Max(0.f, Other.W + W + Tolerance));
 	}
@@ -194,7 +194,7 @@ public:
 	 * @param M Matrix to transform by.
 	 * @return Result of transformation.
 	 */
-	TSphere<T> TransformBy(const TMatrix<T>& M) const;
+	[[nodiscard]] TSphere<T> TransformBy(const TMatrix<T>& M) const;
 
 	/**
 	 * Get result of Transforming sphere with Transform.
@@ -202,14 +202,14 @@ public:
 	 * @param M Transform information.
 	 * @return Result of transformation.
 	 */
-	TSphere<T> TransformBy(const FTransform& M) const;
+	[[nodiscard]] TSphere<T> TransformBy(const FTransform& M) const;
 
 	/**
 	 * Get volume of the current sphere
 	 *
 	 * @return Volume (in Unreal units).
 	 */
-	T GetVolume() const
+	[[nodiscard]] T GetVolume() const
 	{
 		return (4.f / 3.f) * UE_PI * (W * W * W);
 	}
@@ -219,7 +219,7 @@ public:
 	*
 	* @return Text describing the sphere.
 	*/
-	FString ToString() const
+	[[nodiscard]] FString ToString() const
 	{
 		return FString::Printf(TEXT("Center=(%s), Radius=(%s)"), *Center.ToString(), *W.ToString());
 	}
diff --git a/Engine/Source/Runtime/Core/Public/Math/StatisticalFloat.h b/Engine/Source/Runtime/Core/Public/Math/StatisticalFloat.h
--- a/Engine/Source/Runtime/Core/Public/Math/StatisticalFloat.h
+++ b/Engine/Source/Runtime/Core/Public/Math/StatisticalFloat.h
@@ -11,7 +11,7 @@ template <typename T>
 struct FStatisticalValue
 {
 public:
-	FStatisticalValue()
+	[[nodiscard]] FStatisticalValue()
 		: MinValue()
 		, MaxValue()
 		, Accumulator()
@@ -34,22 +34,22 @@ public:
 		++NumSamples;
 	}
 
-	T GetMinValue() const
+	[[nodiscard]] T GetMinValue() const
 	{
 		return MinValue;
 	}
 
-	T GetMaxValue() const
+	[[nodiscard]] T GetMaxValue() const
 	{
 		return MaxValue;
 	}
 
-	T GetAvgValue() const
+	[[nodiscard]] T GetAvgValue() const
 	{
 		return (NumSamples > 0) ? (Accumulator / (double)NumSamples) : T();
 	}
 
-	int32 GetCount() const
+	[[nodiscard]] int32 GetCount() const
 	{
 		return NumSamples;
 	}
diff --git a/Engine/Source/Runtime/Core/Public/Math/TransformCalculus.h b/Engine/Source/Runtime/Core/Public/Math/TransformCalculus.h
--- a/Engine/Source/Runtime/Core/Public/Math/TransformCalculus.h
+++ b/Engine/Source/Runtime/Core/Public/Math/TransformCalculus.h
@@ -259,7 +259,7 @@ struct TransformConverter
  * Return type uses decltype to support classes that can return by const-ref more efficiently than returning a new value.
  */
 template<typename ResultType, typename TransformType>
-inline auto TransformCast(const TransformType& Transform) -> decltype(TransformConverter<ResultType>::Convert(Transform))
+[[nodiscard]] inline auto TransformCast(const TransformType& Transform) -> decltype(TransformConverter<ResultType>::Convert(Transform))
 {
 	return TransformConverter<ResultType>::Convert(Transform);
 }
@@ -290,7 +290,7 @@ struct ConcatenateRules<TransformType, TransformType>
  * @return a new transform representing the transformation from the input space of LHS to the output space of RHS.
  */
 template<typename TransformTypeA, typename TransformTypeB>
-inline typename ConcatenateRules<TransformTypeA, TransformTypeB>::ResultType Concatenate(const TransformTypeA& LHS, const TransformTypeB& RHS)
+[[nodiscard]] inline typename ConcatenateRules<TransformTypeA, TransformTypeB>::ResultType Concatenate(const TransformTypeA& LHS, const TransformTypeB& RHS)
 {
 	typedef typename ConcatenateRules<TransformTypeA, TransformTypeB>::ResultType ReturnType;
 	// If you get a compiler error here about "no member function Concatenate found for TransformType" you know
@@ -302,7 +302,7 @@ inline typename ConcatenateRules<TransformTypeA, TransformTypeB>::ResultType Con
 
 /** Special overload that allows one to explicitly define the result type, which applies TransformCast on each argument first. */
 template<typename ReturnType, typename LHSType, typename RHSType>
-inline ReturnType Concatenate(const LHSType& LHS, const RHSType& RHS)
+[[nodiscard]] inline ReturnType Concatenate(const LHSType& LHS, const RHSType& RHS)
 {
 	return Concatenate(TransformCast<ReturnType>(LHS), TransformCast<ReturnType>(RHS));
 }
@@ -316,7 +316,7 @@ inline ReturnType Concatenate(const LHSType& LHS, const RHSType& RHS)
  * @return a new transform representing the transformation from the input space of LHS to the output space of RHS.
  */
 template<typename TransformType>
-inline auto Concatenate(const TransformType& LHS, const TransformType& RHS) -> decltype(LHS.Concatenate(RHS))
+[[nodiscard]] inline auto Concatenate(const TransformType& LHS, const TransformType& RHS) -> decltype(LHS.Concatenate(RHS))
 {
 	// If you get a compiler error here about "no member function Concatenate found for TransformType" you know
 	// your transform type doesn't support a Concatenate method. Either add one or provide an overload of Concatenate that does this.
@@ -333,7 +333,7 @@ inline auto Concatenate(const TransformType& LHS, const TransformType& RHS) -> d
  * @return a new Transform representing the transformation from space A to space D.
  */
 template<typename TransformType1, typename TransformType2, typename TransformType3>
-inline auto Concatenate(const TransformType1& TransformAToB, const TransformType2& TransformBToC, const TransformType3& TransformCToD) -> decltype(Concatenate(Concatenate(TransformAToB, TransformBToC), TransformCToD)) 
+[[nodiscard]] inline auto Concatenate(const TransformType1& TransformAToB, const TransformType2& TransformBToC, const TransformType3& TransformCToD) -> decltype(Concatenate(Concatenate(TransformAToB, TransformBToC), TransformCToD)) 
 {
 	return Concatenate(Concatenate(TransformAToB, TransformBToC), TransformCToD);
 }
@@ -349,7 +349,7 @@ inline auto Concatenate(const TransformType1& TransformAToB, const TransformType
  * @return a new Transform representing the transformation from space A to space E.
  */
 template<typename TransformType1, typename TransformType2, typename TransformType3, typename TransformType4>
-inline auto Concatenate(const TransformType1& TransformAToB, const TransformType2& TransformBToC, const TransformType3& TransformCToD, const TransformType4& TransformDToE) -> decltype(Concatenate(Concatenate(TransformAToB, TransformBToC, TransformCToD), TransformDToE)) 
+[[nodiscard]] inline auto Concatenate(const TransformType1& TransformAToB, const TransformType2& TransformBToC, const TransformType3& TransformCToD, const TransformType4& TransformDToE) -> decltype(Concatenate(Concatenate(TransformAToB, TransformBToC, TransformCToD), TransformDToE)) 
 {
 	return Concatenate(Concatenate(TransformAToB, TransformBToC, TransformCToD), TransformDToE);
 }
@@ -366,7 +366,7 @@ inline auto Concatenate(const TransformType1& TransformAToB, const TransformType
  * @return a new Transform representing the transformation from space A to space F.
  */
 template<typename TransformType1, typename TransformType2, typename TransformType3, typename TransformType4, typename TransformType5>
-inline auto Concatenate(const TransformType1& TransformAToB, const TransformType2& TransformBToC, const TransformType3& TransformCToD, const TransformType4& TransformDToE, const TransformType5& TransformEToF) -> decltype(Concatenate(Concatenate(TransformAToB, TransformBToC, TransformCToD, TransformDToE), TransformEToF)) 
+[[nodiscard]] inline auto Concatenate(const TransformType1& TransformAToB, const TransformType2& TransformBToC, const TransformType3& TransformCToD, const TransformType4& TransformDToE, const TransformType5& TransformEToF) -> decltype(Concatenate(Concatenate(TransformAToB, TransformBToC, TransformCToD, TransformDToE), TransformEToF)) 
 {
 	return Concatenate(Concatenate(TransformAToB, TransformBToC, TransformCToD, TransformDToE), TransformEToF);
 }
@@ -379,7 +379,7 @@ inline auto Concatenate(const TransformType1& TransformAToB, const TransformType
  * @return Inverted transform from space B to space A.
  */
 template<typename TransformType>
-inline auto Inverse(const TransformType& Transform) -> decltype(Transform.Inverse())
+[[nodiscard]] inline auto Inverse(const TransformType& Transform) -> decltype(Transform.Inverse())
 {
 	return Transform.Inverse();
 }
@@ -388,7 +388,7 @@ inline auto Inverse(const TransformType& Transform) -> decltype(Transform.Invers
  * Generic implementation of TransformPoint. Attempts to use a member function of the TransformType.
  */
 template <typename TransformType, typename PositionType>
-inline PositionType TransformPoint(const TransformType& Transform, const PositionType& Point)
+[[nodiscard]] inline PositionType TransformPoint(const TransformType& Transform, const PositionType& Point)
 {
 	return Transform.TransformPoint(Point);
 } 
@@ -397,7 +397,7 @@ inline PositionType TransformPoint(const TransformType& Transform, const Positio
  * Generic implementation of TransformVector. Attempts to use a member function of the TransformType.
  */
 template <typename TransformType, typename VectorType>
-inline VectorType TransformVector(const TransformType& Transform, const VectorType& Vector)
+[[nodiscard]] inline VectorType TransformVector(const TransformType& Transform, const VectorType& Vector)
 {
 	return Transform.TransformVector(Vector);
 }
@@ -418,7 +418,7 @@ inline VectorType TransformVector(const TransformType& Transform, const VectorTy
  * @return a new Scale representing the transformation from the input space of LHS to the output space of RHS.
  */
 template<typename FloatType UE_REQUIRES(std::is_floating_point_v<FloatType>)>
-inline FloatType Concatenate(FloatType LHS, FloatType RHS)
+[[nodiscard]] inline FloatType Concatenate(FloatType LHS, FloatType RHS)
 {
 	return LHS * RHS;
 }
@@ -431,7 +431,7 @@ inline FloatType Concatenate(FloatType LHS, FloatType RHS)
  * @return Inverted transform from space B to space A.
  */
 template<typename FloatType UE_REQUIRES(std::is_floating_point_v<FloatType>)>
-inline FloatType Inverse(FloatType Scale)
+[[nodiscard]] inline FloatType Inverse(FloatType Scale)
 {
 	return 1.0f / Scale;
 }
@@ -440,13 +440,13 @@ inline FloatType Inverse(FloatType Scale)
  * Specialization for uniform Scale.
  */
 template<typename PositionType>
-inline UE::Math::TVector<PositionType> TransformPoint(float Transform, const UE::Math::TVector<PositionType>& Point)
+[[nodiscard]] inline UE::Math::TVector<PositionType> TransformPoint(float Transform, const UE::Math::TVector<PositionType>& Point)
 {
 	return Transform * Point;
 }
 
 template<typename PositionType>
-inline UE::Math::TVector<PositionType> TransformPoint(double Transform, const UE::Math::TVector<PositionType>& Point)
+[[nodiscard]] inline UE::Math::TVector<PositionType> TransformPoint(double Transform, const UE::Math::TVector<PositionType>& Point)
 {
 	return Transform * Point;
 }
@@ -455,13 +455,13 @@ inline UE::Math::TVector<PositionType> TransformPoint(double Transform, const UE
  * Specialization for uniform Scale.
  */
 template<typename VectorType>
-inline UE::Math::TVector<VectorType> TransformVector(float Transform, const UE::Math::TVector<VectorType>& Vector)
+[[nodiscard]] inline UE::Math::TVector<VectorType> TransformVector(float Transform, const UE::Math::TVector<VectorType>& Vector)
 {
 	return Transform * Vector;
 }
 
 template<typename VectorType>
-inline UE::Math::TVector<VectorType> TransformVector(double Transform, const UE::Math::TVector<VectorType>& Vector)
+[[nodiscard]] inline UE::Math::TVector<VectorType> TransformVector(double Transform, const UE::Math::TVector<VectorType>& Vector)
 {
 	return Transform * Vector;
 }
diff --git a/Engine/Source/Runtime/Core/Public/Math/TransformCalculus2D.h b/Engine/Source/Runtime/Core/Public/Math/TransformCalculus2D.h
--- a/Engine/Source/Runtime/Core/Public/Math/TransformCalculus2D.h
+++ b/Engine/Source/Runtime/Core/Public/Math/TransformCalculus2D.h
@@ -42,7 +42,7 @@ namespace Math
 
 /** Specialization for concatenating two 2D Translations. */
 template<typename T>
-inline UE::Math::TVector2<T> Concatenate(const UE::Math::TVector2<T>& LHS, const UE::Math::TVector2<T>& RHS)
+[[nodiscard]] inline UE::Math::TVector2<T> Concatenate(const UE::Math::TVector2<T>& LHS, const UE::Math::TVector2<T>& RHS)
 {
 	return LHS + RHS;
 }
@@ -52,21 +52,21 @@ inline UE::Math::TVector2<T> Concatenate(const UE::Math::TVector2<T>& LHS, const
 
 /** Specialization for inverting a 2D translation. */
 template<typename T>
-inline UE::Math::TVector2<T> Inverse(const UE::Math::TVector2<T>& Transform)
+[[nodiscard]] inline UE::Math::TVector2<T> Inverse(const UE::Math::TVector2<T>& Transform)
 {
 	return -Transform;
 }
 
 /** Specialization for TVector2 Translation. */
 template<typename T>
-inline UE::Math::TVector2<T> TransformPoint(const UE::Math::TVector2<T>& Transform, const UE::Math::TVector2<T>& Point)
+[[nodiscard]] inline UE::Math::TVector2<T> TransformPoint(const UE::Math::TVector2<T>& Transform, const UE::Math::TVector2<T>& Point)
 {
 	return Transform + Point;
 }
 
 /** Specialization for FVector2D Translation (does nothing). */
 template<typename T>
-inline const UE::Math::TVector2<T>& TransformVector(const UE::Math::TVector2<T>& Transform, const UE::Math::TVector2<T>& Vector)
+[[nodiscard]] inline const UE::Math::TVector2<T>& TransformVector(const UE::Math::TVector2<T>& Transform, const UE::Math::TVector2<T>& Vector)
 {
 	return Vector;
 }
@@ -82,13 +82,13 @@ inline const UE::Math::TVector2<T>& TransformVector(const UE::Math::TVector2<T>&
 * Specialization for uniform Scale.
 */
 template<typename PositionType>
-inline UE::Math::TVector2<PositionType> TransformPoint(float Transform, const UE::Math::TVector2<PositionType>& Point)
+[[nodiscard]] inline UE::Math::TVector2<PositionType> TransformPoint(float Transform, const UE::Math::TVector2<PositionType>& Point)
 {
 	return Transform * Point;
 }
 
 template<typename PositionType>
-inline UE::Math::TVector2<PositionType> TransformPoint(double Transform, const UE::Math::TVector2<PositionType>& Point)
+[[nodiscard]] inline UE::Math::TVector2<PositionType> TransformPoint(double Transform, const UE::Math::TVector2<PositionType>& Point)
 {
 	return Transform * Point;
 }
@@ -97,13 +97,13 @@ inline UE::Math::TVector2<PositionType> TransformPoint(double Transform, const U
 * Specialization for uniform Scale.
 */
 template<typename VectorType>
-inline UE::Math::TVector2<VectorType> TransformVector(float Transform, const UE::Math::TVector2<VectorType>& Vector)
+[[nodiscard]] inline UE::Math::TVector2<VectorType> TransformVector(float Transform, const UE::Math::TVector2<VectorType>& Vector)
 {
 	return Transform * Vector;
 }
 
 template<typename VectorType>
-inline UE::Math::TVector2<VectorType> TransformVector(double Transform, const UE::Math::TVector2<VectorType>& Vector)
+[[nodiscard]] inline UE::Math::TVector2<VectorType> TransformVector(double Transform, const UE::Math::TVector2<VectorType>& Vector)
 {
 	return Transform * Vector;
 }
@@ -119,54 +119,57 @@ public:
 	using Vector2Type = UE::Math::TVector2<T>;
 
 	/** Ctor. initialize to an identity scale, 1.0. */
-	TScale2() : Scale(1.0f, 1.0f) {}
+	[[nodiscard]] TScale2() : Scale(1.0f, 1.0f) {}
 	/** Ctor. initialize from a uniform scale. */
-	explicit TScale2(T InScale) :Scale(InScale, InScale) {}
+	[[nodiscard]] explicit TScale2(T InScale) :Scale(InScale, InScale) {}
 	/** Ctor. initialize from a non-uniform scale. */
-	explicit TScale2(T InScaleX, T InScaleY) :Scale(InScaleX, InScaleY) {}
+	[[nodiscard]] explicit TScale2(T InScaleX, T InScaleY) :Scale(InScaleX, InScaleY) {}
 	/** Ctor. initialize from an FVector defining the 3D scale. */
 	template<typename ArgType>
-	explicit TScale2(const UE::Math::TVector2<ArgType>& InScale) :Scale(InScale) {}
+	[[nodiscard]] explicit TScale2(const UE::Math::TVector2<ArgType>& InScale) :Scale(InScale) {}
 	
 	/** Transform 2D Point */
 	template<typename ArgType>
-	UE::Math::TVector2<ArgType> TransformPoint(const UE::Math::TVector2<ArgType>& Point) const
+	[[nodiscard]] UE::Math::TVector2<ArgType> TransformPoint(const UE::Math::TVector2<ArgType>& Point) const
 	{
 		return UE::Math::TVector2<ArgType>(Scale) * Point;
 	}
 
 	/** Transform 2D Vector*/
 	template<typename ArgType>
-	UE::Math::TVector2<ArgType> TransformVector(const UE::Math::TVector2<ArgType>& Vector) const
+	[[nodiscard]] UE::Math::TVector2<ArgType> TransformVector(const UE::Math::TVector2<ArgType>& Vector) const
 	{
 		return TransformPoint(Vector);
 	}
 
 	/** Concatenate two scales. */
-	TScale2 Concatenate(const TScale2& RHS) const
+	[[nodiscard]] TScale2 Concatenate(const TScale2& RHS) const
 	{
 		return TScale2(Scale * RHS.Scale);
 	}
 	/** Invert the scale. */
-	TScale2 Inverse() const
+	[[nodiscard]] TScale2 Inverse() const
 	{
 		return TScale2(1.0f / Scale.X, 1.0f / Scale.Y);
 	}
 
 	/** Equality. */
-	bool operator==(const TScale2& Other) const
+	[[nodiscard]] bool operator==(const TScale2& Other) const
 	{
 		return Scale == Other.Scale;
 	}
 	
 	/** Inequality. */
-	bool operator!=(const TScale2& Other) const
+	[[nodiscard]] bool operator!=(const TScale2& Other) const
 	{
 		return !operator==(Other);
 	}
 
 	/** Access to the underlying FVector2D that stores the scale. */
-	const Vector2Type& GetVector() const { return Scale; }
+	[[nodiscard]] const Vector2Type& GetVector() const
+	{
+		return Scale;
+	}
 
 private:
 	/** Underlying storage of the 2D scale. */
@@ -214,7 +217,7 @@ public:
 	 * @return the sheare structure.
 	 */
 	template<typename VType>
-	static TShear2 FromShearAngles(const UE::Math::TVector2<VType>& InShearAngles)
+	[[nodiscard]] static TShear2 FromShearAngles(const UE::Math::TVector2<VType>& InShearAngles)
 	{
 		// Compute the M (Shear Slot) = CoTan(90 - SlopeAngle)
 
@@ -231,13 +234,13 @@ public:
 	 *         [XX 1]
 	 */
 	template<typename ArgType>
-	UE::Math::TVector2<ArgType> TransformPoint(const UE::Math::TVector2<ArgType>& Point) const
+	[[nodiscard]] UE::Math::TVector2<ArgType> TransformPoint(const UE::Math::TVector2<ArgType>& Point) const
 	{
 		return Point + UE::Math::TVector2<ArgType>(Point.Y, Point.X) * UE::Math::TVector2<ArgType>(Shear);
 	}
 	/** Transform 2D Vector*/
 	template<typename ArgType>
-	UE::Math::TVector2<ArgType> TransformVector(const UE::Math::TVector2<ArgType>& Vector) const
+	[[nodiscard]] UE::Math::TVector2<ArgType> TransformVector(const UE::Math::TVector2<ArgType>& Vector) const
 	{
 		return TransformPoint(Vector);
 	}
@@ -248,7 +251,7 @@ public:
 	 * [1 YYA] * [1 YYB] == [1+YYA*XXB YYB*YYA]
 	 * [XXA 1]   [XXB 1]    [XXA+XXB XXA*XXB+1]
 	 */
-	inline TMatrix2x2<T> Concatenate(const TShear2& RHS) const;
+	[[nodiscard]] inline TMatrix2x2<T> Concatenate(const TShear2& RHS) const;
 
 	/**
 	 * Invert the shear. The result is NOT a shear, but must be represented by a generalized 2x2 transform.
@@ -256,23 +259,26 @@ public:
 	 * [1 YY]^-1  == 1/(1-YY*XX) * [1 -YY]
 	 * [XX 1]                      [-XX 1]
 	 */
-	TMatrix2x2<T> Inverse() const;
+	[[nodiscard]] TMatrix2x2<T> Inverse() const;
 
 
 	/** Equality. */
-	bool operator==(const TShear2& Other) const
+	[[nodiscard]] bool operator==(const TShear2& Other) const
 	{
 		return Shear == Other.Shear;
 	}
 
 	/** Inequality. */
-	bool operator!=(const TShear2& Other) const
+	[[nodiscard]] bool operator!=(const TShear2& Other) const
 	{
 		return !operator==(Other);
 	}
 
 	/** Access to the underlying FVector2D that stores the scale. */
-	const Vector2Type& GetVector() const { return Shear; }
+	[[nodiscard]] const Vector2Type& GetVector() const
+	{
+		return Shear;
+	}
 
 private:
 	/** Underlying storage of the 2D shear. */
diff --git a/Engine/Source/Runtime/Core/Public/Math/TransformCalculus3D.h b/Engine/Source/Runtime/Core/Public/Math/TransformCalculus3D.h
--- a/Engine/Source/Runtime/Core/Public/Math/TransformCalculus3D.h
+++ b/Engine/Source/Runtime/Core/Public/Math/TransformCalculus3D.h
@@ -43,21 +43,24 @@ public:
 	using Vector3Type = UE::Math::TVector<T>;
 
 	/** Ctor. initialize to an identity scale, 1.0. */
-	TScale() :Scale(1.0f) {}
+	[[nodiscard]] TScale() :Scale(1.0f) {}
 	/** Ctor. initialize from a uniform scale. */
-	explicit TScale(T InScale) :Scale(InScale) {}
+	[[nodiscard]] explicit TScale(T InScale) :Scale(InScale) {}
 	/** Ctor. initialize from an UE::Math::TVector<T> defining the 3D scale. */
 	template<typename VType>
-	explicit TScale(const UE::Math::TVector<VType>& InScale) : Scale((Vector3Type)InScale) {}
+	[[nodiscard]] explicit TScale(const UE::Math::TVector<VType>& InScale) : Scale((Vector3Type)InScale) {}
 	/** Access to the underlying UE::Math::TVector<T> that stores the scale. */
-	const Vector3Type& GetVector() const { return Scale; }
+	[[nodiscard]] const Vector3Type& GetVector() const
+	{
+		return Scale;
+	}
 	/** Concatenate two scales. */
-	const TScale Concatenate(const TScale& RHS) const
+	[[nodiscard]] TScale Concatenate(const TScale& RHS) const
 	{
 		return TScale(Scale * RHS.GetVector());
 	}
 	/** Invert the scale. */
-	const TScale Inverse() const
+	[[nodiscard]]  TScale Inverse() const
 	{
 		return TScale(Vector3Type(1.0f / Scale.X, 1.0f / Scale.Y, 1.0f / Scale.Z));
 	}
diff --git a/Engine/Source/Runtime/Core/Public/Math/TransformNonVectorized.h b/Engine/Source/Runtime/Core/Public/Math/TransformNonVectorized.h
--- a/Engine/Source/Runtime/Core/Public/Math/TransformNonVectorized.h
+++ b/Engine/Source/Runtime/Core/Public/Math/TransformNonVectorized.h
@@ -110,7 +110,7 @@ public:
 #endif
 
 	/** Default constructor. */
-	FORCEINLINE TTransform()
+	[[nodiscard]] FORCEINLINE TTransform()
 		: Rotation(0.f, 0.f, 0.f, 1.f)
 		, Translation(0.f)
 		, Scale3D(TVector<T>::OneVector)
@@ -122,7 +122,7 @@ public:
 	*
 	* @param InTranslation The value to use for the translation component
 	*/
-	FORCEINLINE explicit TTransform(const TVector<T>& InTranslation)
+	[[nodiscard]] FORCEINLINE explicit TTransform(const TVector<T>& InTranslation)
 		: Rotation(TQuat<T>::Identity),
 		Translation(InTranslation),
 		Scale3D(TVector<T>::OneVector)
@@ -133,7 +133,7 @@ public:
 	/**
 	* Constructor with leaving uninitialized memory
 	*/
-	FORCEINLINE explicit TTransform(ENoInit)
+	[[nodiscard]] FORCEINLINE explicit TTransform(ENoInit)
 	{
 		// Note: This can be used to track down initialization issues with bone transform arrays; but it will
 		// cause issues with transient fields such as RootMotionDelta that get initialized to 0 by default
@@ -152,7 +152,7 @@ public:
 	*
 	* @param InRotation The value to use for rotation component
 	*/
-	FORCEINLINE explicit TTransform(const TQuat<T>& InRotation)
+	[[nodiscard]] FORCEINLINE explicit TTransform(const TQuat<T>& InRotation)
 		: Rotation(InRotation),
 		Translation(TVector<T>::ZeroVector),
 		Scale3D(TVector<T>::OneVector)
@@ -165,7 +165,7 @@ public:
 	*
 	* @param InRotation The value to use for rotation component  (after being converted to a quaternion)
 	*/
-	FORCEINLINE explicit TTransform(const TRotator<T>& InRotation)
+	[[nodiscard]] FORCEINLINE explicit TTransform(const TRotator<T>& InRotation)
 		: Rotation(InRotation),
 		Translation(TVector<T>::ZeroVector),
 		Scale3D(TVector<T>::OneVector)
@@ -180,7 +180,7 @@ public:
 	* @param InTranslation The value to use for the translation component
 	* @param InScale3D The value to use for the scale component
 	*/
-	FORCEINLINE TTransform(const TQuat<T>& InRotation, const TVector<T>& InTranslation, const TVector<T>& InScale3D = TVector<T>::OneVector)
+	[[nodiscard]] FORCEINLINE TTransform(const TQuat<T>& InRotation, const TVector<T>& InTranslation, const TVector<T>& InScale3D = TVector<T>::OneVector)
 		: Rotation(InRotation),
 		Translation(InTranslation),
 		Scale3D(InScale3D)
@@ -195,7 +195,7 @@ public:
 	* @param InTranslation The value to use for the translation component
 	* @param InScale3D The value to use for the scale component
 	*/
-	FORCEINLINE TTransform(const TRotator<T>& InRotation, const TVector<T>& InTranslation, const TVector<T>& InScale3D = TVector<T>::OneVector)
+	[[nodiscard]] FORCEINLINE TTransform(const TRotator<T>& InRotation, const TVector<T>& InTranslation, const TVector<T>& InScale3D = TVector<T>::OneVector)
 		: Rotation(InRotation),
 		Translation(InTranslation),
 		Scale3D(InScale3D)
@@ -206,14 +206,14 @@ public:
 	/**
 	* Constructor for converting a Matrix (including scale) into a TTransform<T>.
 	*/
-	FORCEINLINE explicit TTransform(const TMatrix<T>& InMatrix)
+	[[nodiscard]] FORCEINLINE explicit TTransform(const TMatrix<T>& InMatrix)
 	{
 		SetFromMatrix(InMatrix);
 		DiagnosticCheckNaN_All();
 	}
 
 	/** Constructor that takes basis axes and translation */
-	FORCEINLINE TTransform(const TVector<T>& InX, const TVector<T>& InY, const TVector<T>& InZ, const TVector<T>& InTranslation)
+	[[nodiscard]] FORCEINLINE TTransform(const TVector<T>& InX, const TVector<T>& InY, const TVector<T>& InZ, const TVector<T>& InTranslation)
 	{
 		SetFromMatrix(TMatrix<T>(InX, InY, InZ, InTranslation));
 		DiagnosticCheckNaN_All();
@@ -225,12 +225,12 @@ public:
 	CORE_API void DebugPrint() const;
 
 	/** Debug purpose only **/
-	bool DebugEqualMatrix(const TMatrix<T>& Matrix) const;
+	[[nodiscard]] bool DebugEqualMatrix(const TMatrix<T>& Matrix) const;
 
 	/** Convert TTransform<T> contents to a string */
-	CORE_API FString ToHumanReadableString() const;
+	[[nodiscard]] CORE_API FString ToHumanReadableString() const;
 
-	CORE_API FString ToString() const;
+	[[nodiscard]] CORE_API FString ToString() const;
 
 	/** Acceptable form: "%f,%f,%f|%f,%f,%f|%f,%f,%f" */
 	CORE_API bool InitFromString(const FString& InSourceString);
@@ -238,7 +238,7 @@ public:
 	/**
 	* Convert this Transform to a transformation matrix with scaling.
 	*/
-	FORCEINLINE TMatrix<T> ToMatrixWithScale() const
+	[[nodiscard]] FORCEINLINE TMatrix<T> ToMatrixWithScale() const
 	{
 		TMatrix<T> OutMatrix;
 
@@ -295,7 +295,7 @@ public:
 	/**
 	* Convert this Transform to matrix with scaling and compute the inverse of that.
 	*/
-	FORCEINLINE TMatrix<T> ToInverseMatrixWithScale() const
+	[[nodiscard]] FORCEINLINE TMatrix<T> ToInverseMatrixWithScale() const
 	{
 		// todo: optimize
 		return ToMatrixWithScale().Inverse();
@@ -304,7 +304,7 @@ public:
 	/**
 	* Convert this Transform to inverse.
 	*/
-	FORCEINLINE TTransform<T> Inverse() const
+	[[nodiscard]] FORCEINLINE TTransform<T> Inverse() const
 	{
 		TQuat<T>   InvRotation = Rotation.Inverse();
 		// this used to cause NaN if Scale contained 0 
@@ -317,7 +317,7 @@ public:
 	/**
 	* Convert this Transform to a transformation matrix, ignoring its scaling
 	*/
-	FORCEINLINE TMatrix<T> ToMatrixNoScale() const
+	[[nodiscard]] FORCEINLINE TMatrix<T> ToMatrixNoScale() const
 	{
 		TMatrix<T> OutMatrix;
 
@@ -435,7 +435,7 @@ public:
 	* Use only within blends!!
 	* Rotation part is NOT normalized!!
 	*/
-	FORCEINLINE TTransform<T> operator+(const TTransform<T>& Atom) const
+	[[nodiscard]] FORCEINLINE TTransform<T> operator+(const TTransform<T>& Atom) const
 	{
 		return TTransform(Rotation + Atom.Rotation, Translation + Atom.Translation, Scale3D + Atom.Scale3D);
 	}
@@ -454,7 +454,7 @@ public:
 		return *this;
 	}
 
-	FORCEINLINE TTransform<T> operator*(T Mult) const
+	[[nodiscard]] FORCEINLINE TTransform<T> operator*(T Mult) const
 	{
 		return TTransform(Rotation * Mult, Translation * Mult, Scale3D * Mult);
 	}
@@ -478,7 +478,7 @@ public:
 	* @param  Other other transform by which to multiply.
 	* @return new transform: this * Other
 	*/
-	FORCEINLINE TTransform<T> operator*(const TTransform<T>& Other) const;
+	[[nodiscard]] FORCEINLINE TTransform<T> operator*(const TTransform<T>& Other) const;
 
 
 	/**
@@ -496,7 +496,7 @@ public:
 	* @param  Other other quaternion rotation by which to multiply.
 	* @return new transform: this * TTransform(Other)
 	*/
-	FORCEINLINE TTransform<T> operator*(const TQuat<T>& Other) const;
+	[[nodiscard]] FORCEINLINE TTransform<T> operator*(const TQuat<T>& Other) const;
 
 	/**
 	* Sets this transform to the result of this multiplied by another transform (made only from a rotation).
@@ -506,12 +506,12 @@ public:
 	*/
 	FORCEINLINE void operator*=(const TQuat<T>& Other);
 
-	FORCEINLINE static bool AnyHasNegativeScale(const TVector<T>& InScale3D, const TVector<T>& InOtherScale3D);
+	[[nodiscard]] FORCEINLINE static bool AnyHasNegativeScale(const TVector<T>& InScale3D, const TVector<T>& InOtherScale3D);
 	FORCEINLINE void ScaleTranslation(const TVector<T>& InScale3D);
 	FORCEINLINE void ScaleTranslation(const FReal& Scale);
 	FORCEINLINE void RemoveScaling(FReal Tolerance = UE_SMALL_NUMBER);
-	FORCEINLINE T GetMaximumAxisScale() const;
-	FORCEINLINE T GetMinimumAxisScale() const;
+	[[nodiscard]] FORCEINLINE T GetMaximumAxisScale() const;
+	[[nodiscard]] FORCEINLINE T GetMinimumAxisScale() const;
 
 	// Inverse does not work well with VQS format(in particular non-uniform), so removing it, but made two below functions to be used instead. 
 
@@ -522,64 +522,64 @@ public:
 	* GetRelativeTransform returns this*Other(-1) and parameter is Other(not Other(-1))
 	* GetRelativeTransformReverse returns this(-1)*Other, and parameter is Other.
 	*******************************************************************************************/
-	CORE_API TTransform<T> GetRelativeTransform(const TTransform<T>& Other) const;
-	CORE_API TTransform<T> GetRelativeTransformReverse(const TTransform<T>& Other) const;
+	[[nodiscard]] CORE_API TTransform<T> GetRelativeTransform(const TTransform<T>& Other) const;
+	[[nodiscard]] CORE_API TTransform<T> GetRelativeTransformReverse(const TTransform<T>& Other) const;
 	/**
 	* Set current transform and the relative to ParentTransform.
 	* Equates to This = This->GetRelativeTransform(Parent), but saves the intermediate TTransform<T> storage and copy.
 	*/
 	CORE_API void SetToRelativeTransform(const TTransform<T>& ParentTransform);
 
-	FORCEINLINE TVector4<T> TransformFVector4(const TVector4<T>& V) const;
-	FORCEINLINE TVector4<T> TransformFVector4NoScale(const TVector4<T>& V) const;
-	FORCEINLINE TVector<T> TransformPosition(const TVector<T>& V) const;
-	FORCEINLINE TVector<T> TransformPositionNoScale(const TVector<T>& V) const;
+	[[nodiscard]] FORCEINLINE TVector4<T> TransformFVector4(const TVector4<T>& V) const;
+	[[nodiscard]] FORCEINLINE TVector4<T> TransformFVector4NoScale(const TVector4<T>& V) const;
+	[[nodiscard]] FORCEINLINE TVector<T> TransformPosition(const TVector<T>& V) const;
+	[[nodiscard]] FORCEINLINE TVector<T> TransformPositionNoScale(const TVector<T>& V) const;
 
 	/** Inverts the transform and then transforms V - correctly handles scaling in this transform. */
-	FORCEINLINE TVector<T> InverseTransformPosition(const TVector<T> &V) const;
-	FORCEINLINE TVector<T> InverseTransformPositionNoScale(const TVector<T> &V) const;
-	FORCEINLINE TVector<T> TransformVector(const TVector<T>& V) const;
-	FORCEINLINE TVector<T> TransformVectorNoScale(const TVector<T>& V) const;
+	[[nodiscard]] FORCEINLINE TVector<T> InverseTransformPosition(const TVector<T> &V) const;
+	[[nodiscard]] FORCEINLINE TVector<T> InverseTransformPositionNoScale(const TVector<T> &V) const;
+	[[nodiscard]] FORCEINLINE TVector<T> TransformVector(const TVector<T>& V) const;
+	[[nodiscard]] FORCEINLINE TVector<T> TransformVectorNoScale(const TVector<T>& V) const;
 
 	/**
 	*	Transform a direction vector by the inverse of this transform - will not take into account translation part.
 	*	If you want to transform a surface normal (or plane) and correctly account for non-uniform scaling you should use TransformByUsingAdjointT with adjoint of matrix inverse.
 	*/
-	FORCEINLINE TVector<T> InverseTransformVector(const TVector<T> &V) const;
-	FORCEINLINE TVector<T> InverseTransformVectorNoScale(const TVector<T> &V) const;
+	[[nodiscard]] FORCEINLINE TVector<T> InverseTransformVector(const TVector<T> &V) const;
+	[[nodiscard]] FORCEINLINE TVector<T> InverseTransformVectorNoScale(const TVector<T> &V) const;
 
 	/**
 	 * Transform a rotation.
 	 * For example if this is a LocalToWorld transform, TransformRotation(Q) would transform Q from local to world space.
 	 */
-	FORCEINLINE TQuat<T> TransformRotation(const TQuat<T>& Q) const;
+	[[nodiscard]] FORCEINLINE TQuat<T> TransformRotation(const TQuat<T>& Q) const;
 	
 	/**
 	* Inverse transform a rotation.
 	* For example if this is a LocalToWorld transform, InverseTransformRotation(Q) would transform Q from world to local space.
 	*/
-	FORCEINLINE TQuat<T> InverseTransformRotation(const TQuat<T>& Q) const;
+	[[nodiscard]] FORCEINLINE TQuat<T> InverseTransformRotation(const TQuat<T>& Q) const;
 
-	FORCEINLINE TTransform<T> GetScaled(T Scale) const;
-	FORCEINLINE TTransform<T> GetScaled(TVector<T> Scale) const;
-	FORCEINLINE TVector<T> GetScaledAxis(EAxis::Type InAxis) const;
-	FORCEINLINE TVector<T> GetUnitAxis(EAxis::Type InAxis) const;
-	FORCEINLINE void Mirror(EAxis::Type MirrorAxis, EAxis::Type FlipAxis);
-	FORCEINLINE static TVector<T> GetSafeScaleReciprocal(const TVector<T>& InScale, FReal Tolerance = UE_SMALL_NUMBER);
+	[[nodiscard]] FORCEINLINE TTransform<T> GetScaled(T Scale) const;
+	[[nodiscard]] FORCEINLINE TTransform<T> GetScaled(TVector<T> Scale) const;
+	[[nodiscard]] FORCEINLINE TVector<T> GetScaledAxis(EAxis::Type InAxis) const;
+	[[nodiscard]] FORCEINLINE TVector<T> GetUnitAxis(EAxis::Type InAxis) const;
+	[[nodiscard]] FORCEINLINE void Mirror(EAxis::Type MirrorAxis, EAxis::Type FlipAxis);
+	[[nodiscard]] FORCEINLINE static TVector<T> GetSafeScaleReciprocal(const TVector<T>& InScale, FReal Tolerance = UE_SMALL_NUMBER);
 
 	// temp function for easy conversion
-	FORCEINLINE TVector<T> GetLocation() const
+	[[nodiscard]] FORCEINLINE TVector<T> GetLocation() const
 	{
 		return GetTranslation();
 	}
 
-	FORCEINLINE TRotator<T> Rotator() const
+	[[nodiscard]] FORCEINLINE TRotator<T> Rotator() const
 	{
 		return Rotation.Rotator();
 	}
 
 	/** Calculate the determinant of this transformation */
-	FORCEINLINE FReal GetDeterminant() const
+	[[nodiscard]] FORCEINLINE FReal GetDeterminant() const
 	{
 		return (float)(Scale3D.X * Scale3D.Y * Scale3D.Z);
 	}
@@ -595,12 +595,12 @@ public:
 	* Checks the components for non-finite values (NaN or Inf).
 	* @return Returns true if any component (rotation, translation, or scale) is not finite.
 	*/
-	bool ContainsNaN() const
+	[[nodiscard]] bool ContainsNaN() const
 	{
 		return (Translation.ContainsNaN() || Rotation.ContainsNaN() || Scale3D.ContainsNaN());
 	}
 
-	inline bool IsValid() const
+	[[nodiscard]] inline bool IsValid() const
 	{
 		if (ContainsNaN())
 		{
@@ -663,19 +663,19 @@ private:
 public:
 
 	// Test if A's rotation equals B's rotation, within a tolerance. Preferred over "A.GetRotation().Equals(B.GetRotation())" because it is faster on some platforms.
-	FORCEINLINE static bool AreRotationsEqual(const TTransform<T>& A, const TTransform<T>& B, FReal Tolerance = UE_KINDA_SMALL_NUMBER)
+	[[nodiscard]] FORCEINLINE static bool AreRotationsEqual(const TTransform<T>& A, const TTransform<T>& B, FReal Tolerance = UE_KINDA_SMALL_NUMBER)
 	{
 		return A.Private_RotationEquals(B.Rotation, Tolerance);
 	}
 
 	// Test if A's translation equals B's translation, within a tolerance. Preferred over "A.GetTranslation().Equals(B.GetTranslation())" because it is faster on some platforms.
-	FORCEINLINE static bool AreTranslationsEqual(const TTransform<T>& A, const TTransform<T>& B, FReal Tolerance = UE_KINDA_SMALL_NUMBER)
+	[[nodiscard]] FORCEINLINE static bool AreTranslationsEqual(const TTransform<T>& A, const TTransform<T>& B, FReal Tolerance = UE_KINDA_SMALL_NUMBER)
 	{
 		return A.Private_TranslationEquals(B.Translation, Tolerance);
 	}
 
 	// Test if A's scale equals B's scale, within a tolerance. Preferred over "A.GetScale3D().Equals(B.GetScale3D())" because it is faster on some platforms.
-	FORCEINLINE static bool AreScale3DsEqual(const TTransform<T>& A, const TTransform<T>& B, FReal Tolerance = UE_KINDA_SMALL_NUMBER)
+	[[nodiscard]] FORCEINLINE static bool AreScale3DsEqual(const TTransform<T>& A, const TTransform<T>& B, FReal Tolerance = UE_KINDA_SMALL_NUMBER)
 	{
 		return A.Private_Scale3DEquals(B.Scale3D, Tolerance);
 	}
@@ -683,37 +683,37 @@ public:
 
 
 	// Test if this Transform's rotation equals another's rotation, within a tolerance. Preferred over "GetRotation().Equals(Other.GetRotation())" because it is faster on some platforms.
-	FORCEINLINE bool RotationEquals(const TTransform<T>& Other, FReal Tolerance = UE_KINDA_SMALL_NUMBER) const
+	[[nodiscard]] FORCEINLINE bool RotationEquals(const TTransform<T>& Other, FReal Tolerance = UE_KINDA_SMALL_NUMBER) const
 	{
 		return AreRotationsEqual(*this, Other, Tolerance);
 	}
 
 	// Test if this Transform's translation equals another's translation, within a tolerance. Preferred over "GetTranslation().Equals(Other.GetTranslation())" because it is faster on some platforms.
-	FORCEINLINE bool TranslationEquals(const TTransform<T>& Other, FReal Tolerance = UE_KINDA_SMALL_NUMBER) const
+	[[nodiscard]] FORCEINLINE bool TranslationEquals(const TTransform<T>& Other, FReal Tolerance = UE_KINDA_SMALL_NUMBER) const
 	{
 		return AreTranslationsEqual(*this, Other, Tolerance);
 	}
 
 	// Test if this Transform's scale equals another's scale, within a tolerance. Preferred over "GetScale3D().Equals(Other.GetScale3D())" because it is faster on some platforms.
-	FORCEINLINE bool Scale3DEquals(const TTransform<T>& Other, FReal Tolerance = UE_KINDA_SMALL_NUMBER) const
+	[[nodiscard]] FORCEINLINE bool Scale3DEquals(const TTransform<T>& Other, FReal Tolerance = UE_KINDA_SMALL_NUMBER) const
 	{
 		return AreScale3DsEqual(*this, Other, Tolerance);
 	}
 
 	// Test if all components of the transforms are equal, within a tolerance.
-	FORCEINLINE bool Equals(const TTransform<T>& Other, FReal Tolerance = UE_KINDA_SMALL_NUMBER) const
+	[[nodiscard]] FORCEINLINE bool Equals(const TTransform<T>& Other, FReal Tolerance = UE_KINDA_SMALL_NUMBER) const
 	{
 		return Private_TranslationEquals(Other.Translation, Tolerance) && Private_RotationEquals(Other.Rotation, Tolerance) && Private_Scale3DEquals(Other.Scale3D, Tolerance);
 	}
 
 	// Test if all components of the transform property are equal.
-	FORCEINLINE bool Identical(const TTransform<T>* Other, uint32 PortFlags) const
+	[[nodiscard]] FORCEINLINE bool Identical(const TTransform<T>* Other, uint32 PortFlags) const
 	{
 		return Equals(*Other, 0.f);
 	}
 
 	// Test if rotation and translation components of the transforms are equal, within a tolerance.
-	FORCEINLINE bool EqualsNoScale(const TTransform<T>& Other, FReal Tolerance = UE_KINDA_SMALL_NUMBER) const
+	[[nodiscard]] FORCEINLINE bool EqualsNoScale(const TTransform<T>& Other, FReal Tolerance = UE_KINDA_SMALL_NUMBER) const
 	{
 		return Private_TranslationEquals(Other.Translation, Tolerance) && Private_RotationEquals(Other.Rotation, Tolerance);
 	}
@@ -820,7 +820,7 @@ public:
 	* Add the translations from two TTransform<T>s and return the result.
 	* @return A.Translation + B.Translation
 	*/
-	FORCEINLINE static TVector<T> AddTranslations(const TTransform<T>& A, const TTransform<T>& B)
+	[[nodiscard]] FORCEINLINE static TVector<T> AddTranslations(const TTransform<T>& A, const TTransform<T>& B)
 	{
 		return A.Translation + B.Translation;
 	}
@@ -829,7 +829,7 @@ public:
 	* Subtract translations from two TTransform<T>s and return the difference.
 	* @return A.Translation - B.Translation.
 	*/
-	FORCEINLINE static TVector<T> SubtractTranslations(const TTransform<T>& A, const TTransform<T>& B)
+	[[nodiscard]] FORCEINLINE static TVector<T> SubtractTranslations(const TTransform<T>& A, const TTransform<T>& B)
 	{
 		return A.Translation - B.Translation;
 	}
@@ -881,10 +881,22 @@ public:
 	}
 
 	// For low-level VectorRegister programming
-	TransformVectorRegister GetTranslationRegister() const { return VectorLoadFloat3_W0(&Translation); }
-	TransformVectorRegister GetRotationRegister() const { return VectorLoad(&Rotation); }
-	void SetTranslationRegister(TransformVectorRegister InTranslation) { VectorStoreFloat3(InTranslation, &Translation); }
-	void SetRotationRegister(TransformVectorRegister InRotation) { VectorStore(InRotation, &Rotation); }
+	[[nodiscard]] TransformVectorRegister GetTranslationRegister() const
+	{
+		return VectorLoadFloat3_W0(&Translation);
+	}
+	[[nodiscard]] TransformVectorRegister GetRotationRegister() const
+	{
+		return VectorLoad(&Rotation);
+	}
+	void SetTranslationRegister(TransformVectorRegister InTranslation)
+	{
+		VectorStoreFloat3(InTranslation, &Translation);
+	}
+	void SetRotationRegister(TransformVectorRegister InRotation)
+	{
+		VectorStore(InRotation, &Rotation);
+	}
 
 	/** @note: Added template type function for Accumulate
 	* The template type isn't much useful yet, but it is with the plan to move forward
@@ -1036,7 +1048,7 @@ public:
 	*
 	* @return true if the rotation component is normalized, and false otherwise.
 	*/
-	FORCEINLINE bool IsRotationNormalized() const
+	[[nodiscard]] FORCEINLINE bool IsRotationNormalized() const
 	{
 		return Rotation.IsNormalized();
 	}
@@ -1077,7 +1089,7 @@ public:
 	*
 	* @return The rotation component
 	*/
-	FORCEINLINE TQuat<T> GetRotation() const
+	[[nodiscard]] FORCEINLINE TQuat<T> GetRotation() const
 	{
 		DiagnosticCheckNaN_Rotate();
 		return Rotation;
@@ -1088,7 +1100,7 @@ public:
 	*
 	* @return The translation component
 	*/
-	FORCEINLINE TVector<T> GetTranslation() const
+	[[nodiscard]] FORCEINLINE TVector<T> GetTranslation() const
 	{
 		DiagnosticCheckNaN_Translate();
 		return Translation;
@@ -1099,7 +1111,7 @@ public:
 	*
 	* @return The Scale3D component
 	*/
-	FORCEINLINE TVector<T> GetScale3D() const
+	[[nodiscard]] FORCEINLINE TVector<T> GetScale3D() const
 	{
 		DiagnosticCheckNaN_Scale3D();
 		return Scale3D;
diff --git a/Engine/Source/Runtime/Core/Public/Math/TransformVectorized.h b/Engine/Source/Runtime/Core/Public/Math/TransformVectorized.h
--- a/Engine/Source/Runtime/Core/Public/Math/TransformVectorized.h
+++ b/Engine/Source/Runtime/Core/Public/Math/TransformVectorized.h
@@ -140,7 +140,7 @@ public:
 	/**
 	 * Constructor with initialization to the identity transform.
 	 */
-	FORCEINLINE_DEBUGGABLE TTransform()
+	[[nodiscard]] FORCEINLINE_DEBUGGABLE TTransform()
 	{
 		// Rotation = {0,0,0,1)
 		Rotation = GlobalVectorConstants::Float0001;
@@ -155,7 +155,7 @@ public:
 	 *
 	 * @param InTranslation The value to use for the translation component
 	 */
-	FORCEINLINE_DEBUGGABLE explicit TTransform(const TVector<T>& InTranslation) 
+	[[nodiscard]] FORCEINLINE_DEBUGGABLE explicit TTransform(const TVector<T>& InTranslation) 
 	{
 		// Rotation = {0,0,0,1) quaternion identity
 		Rotation = GlobalVectorConstants::Float0001;
@@ -172,7 +172,7 @@ public:
 	 *
 	 * @param InRotation The value to use for rotation component
 	 */
-	FORCEINLINE_DEBUGGABLE explicit TTransform(const TQuat<T>& InRotation) 
+	[[nodiscard]] FORCEINLINE_DEBUGGABLE explicit TTransform(const TQuat<T>& InRotation) 
 	{
 		// Rotation = InRotation
 		Rotation =  VectorLoadAligned( &InRotation );
@@ -189,7 +189,7 @@ public:
 	 *
 	 * @param InRotation The value to use for rotation component  (after being converted to a quaternion)
 	 */
-	FORCEINLINE_DEBUGGABLE explicit TTransform(const TRotator<T>& InRotation) 
+	[[nodiscard]] FORCEINLINE_DEBUGGABLE explicit TTransform(const TRotator<T>& InRotation) 
 	{
 		TQuat<T> InQuatRotation = TQuat<T>::MakeFromRotator(InRotation);
 		// Rotation = InRotation
@@ -209,7 +209,7 @@ public:
 	 * @param InTranslation The value to use for the translation component
 	 * @param InScale3D The value to use for the scale component
 	 */
-	FORCEINLINE_DEBUGGABLE TTransform(const TQuat<T>& InRotation, const TVector<T>& InTranslation, const TVector<T>& InScale3D = TVector<T>(1.f, 1.f, 1.f))
+	[[nodiscard]] FORCEINLINE_DEBUGGABLE TTransform(const TQuat<T>& InRotation, const TVector<T>& InTranslation, const TVector<T>& InScale3D = TVector<T>(1.f, 1.f, 1.f))
 	{
 		// Rotation = InRotation
 		Rotation =  VectorLoadAligned( &InRotation );
@@ -228,7 +228,7 @@ public:
 	 * @param InTranslation The value to use for the translation component
 	 * @param InScale3D The value to use for the scale component
 	 */
-	FORCEINLINE_DEBUGGABLE TTransform(const TransformVectorRegister& InRotation, const TransformVectorRegister& InTranslation, const TransformVectorRegister& InScale3D)
+	[[nodiscard]] FORCEINLINE_DEBUGGABLE TTransform(const TransformVectorRegister& InRotation, const TransformVectorRegister& InTranslation, const TransformVectorRegister& InScale3D)
 		: Rotation(InRotation),
 		Translation(InTranslation),
 		Scale3D(InScale3D)
@@ -243,7 +243,7 @@ public:
 	 * @param InTranslation The value to use for the translation component
 	 * @param InScale3D The value to use for the scale component
 	 */
-	FORCEINLINE_DEBUGGABLE TTransform(const TRotator<T>& InRotation, const TVector<T>& InTranslation, const TVector<T>& InScale3D = TVector<T>(1.f, 1.f, 1.f))
+	[[nodiscard]] FORCEINLINE_DEBUGGABLE TTransform(const TRotator<T>& InRotation, const TVector<T>& InTranslation, const TVector<T>& InScale3D = TVector<T>(1.f, 1.f, 1.f))
 	{
 		TQuat<T> InQuatRotation = TQuat<T>::MakeFromRotator(InRotation);
 		// Rotation = InRotation
@@ -259,7 +259,7 @@ public:
 	/**
 	 * Constructor with leaving uninitialized memory
 	 */
-	FORCEINLINE_DEBUGGABLE explicit TTransform(ENoInit) 
+	[[nodiscard]] FORCEINLINE_DEBUGGABLE explicit TTransform(ENoInit) 
 	{
 		// Note: This can be used to track down initialization issues with bone transform arrays; but it will
 		// cause issues with transient fields such as RootMotionDelta that get initialized to 0 by default
@@ -275,14 +275,14 @@ public:
 	/**
 	 * Constructor for converting a Matrix (including scale) into a TTransform.
 	 */
-	FORCEINLINE_DEBUGGABLE explicit TTransform(const TMatrix<T>& InMatrix)
+	[[nodiscard]] FORCEINLINE_DEBUGGABLE explicit TTransform(const TMatrix<T>& InMatrix)
 	{
 		SetFromMatrix(InMatrix);
 		DiagnosticCheckNaN_All();
 	}
 
 	/** Constructor that takes basis axes and translation */
-	FORCEINLINE_DEBUGGABLE TTransform(const TVector<T>& InX, const TVector<T>& InY, const TVector<T>& InZ, const TVector<T>& InTranslation)
+	[[nodiscard]] FORCEINLINE_DEBUGGABLE TTransform(const TVector<T>& InX, const TVector<T>& InY, const TVector<T>& InZ, const TVector<T>& InTranslation)
 	{
 		SetFromMatrix(TMatrix<T>(InX, InY, InZ, InTranslation));
 		DiagnosticCheckNaN_All();
@@ -294,18 +294,18 @@ public:
 	CORE_API void DebugPrint() const;
 
 	/** Debug purpose only **/
-	bool DebugEqualMatrix(const TMatrix<T>& Matrix) const;
+	[[nodiscard]] bool DebugEqualMatrix(const TMatrix<T>& Matrix) const;
 
 	/** Convert TTransform contents to a string */
-	CORE_API FString ToHumanReadableString() const;
+	[[nodiscard]] CORE_API FString ToHumanReadableString() const;
 
-	CORE_API FString ToString() const;
+	[[nodiscard]] CORE_API FString ToString() const;
 
 	/** Acceptable form: "%f,%f,%f|%f,%f,%f|%f,%f,%f" */
 	CORE_API bool InitFromString( const FString& InSourceString );
 
 
-	FORCEINLINE_DEBUGGABLE TMatrix<T> ToMatrixWithScale() const
+	[[nodiscard]] FORCEINLINE_DEBUGGABLE TMatrix<T> ToMatrixWithScale() const
 	{
 
 		TMatrix<T> OutMatrix;
@@ -356,7 +356,7 @@ public:
 	/**
 	* Convert this Transform to matrix with scaling and compute the inverse of that.
 	*/
-	FORCEINLINE_DEBUGGABLE TMatrix<T> ToInverseMatrixWithScale() const
+	[[nodiscard]] FORCEINLINE_DEBUGGABLE TMatrix<T> ToInverseMatrixWithScale() const
 	{
 		// todo: optimize
 		return ToMatrixWithScale().Inverse();
@@ -365,7 +365,7 @@ public:
 	/**
 	* Convert this Transform to inverse.
 	*/
-	FORCEINLINE_DEBUGGABLE TTransform Inverse() const
+	[[nodiscard]] FORCEINLINE_DEBUGGABLE TTransform Inverse() const
 	{
 		// Replacement of Inverse of TMatrix<T>
 		if (VectorAnyGreaterThan(VectorAbs(Scale3D), GlobalVectorConstants::SmallNumber))
@@ -381,7 +381,7 @@ public:
 	/**
 	* Convert this Transform to a transformation matrix, ignoring its scaling
 	*/
-	FORCEINLINE_DEBUGGABLE TMatrix<T> ToMatrixNoScale() const
+	[[nodiscard]] FORCEINLINE_DEBUGGABLE TMatrix<T> ToMatrixNoScale() const
 	{
 		TMatrix<T> OutMatrix;
 		TransformVectorRegister DiagonalsXYZ;
@@ -504,7 +504,7 @@ public:
 	 * Use only within blends!!
 	 * Rotation part is NOT normalized!!
 	 */
-	FORCEINLINE_DEBUGGABLE TTransform<T> operator+(const TTransform<T>& Atom) const
+	[[nodiscard]] FORCEINLINE_DEBUGGABLE TTransform<T> operator+(const TTransform<T>& Atom) const
 	{
 		return TTransform<T>( VectorAdd(Rotation, Atom.Rotation), VectorAdd(Translation, Atom.Translation ), VectorAdd( Scale3D, Atom.Scale3D ) );
 	}
@@ -518,7 +518,7 @@ public:
 		return *this;
 	}
 
-	FORCEINLINE_DEBUGGABLE TTransform<T> operator*(const ScalarRegister& MultScalar) const
+	[[nodiscard]] FORCEINLINE_DEBUGGABLE TTransform<T> operator*(const ScalarRegister& MultScalar) const
 	{		
 		const TransformVectorRegister Mult = MultScalar.Value;
 		return TTransform<T>( VectorMultiply(Rotation, Mult), VectorMultiply(Translation, Mult), VectorMultiply(Scale3D, Mult) );
@@ -534,17 +534,17 @@ public:
 		return *this;
 	}
 
-	FORCEINLINE_DEBUGGABLE TTransform<T>	operator*(const TTransform<T>& Other) const;
+	[[nodiscard]] FORCEINLINE_DEBUGGABLE TTransform<T>	operator*(const TTransform<T>& Other) const;
 	FORCEINLINE_DEBUGGABLE void			operator*=(const TTransform<T>& Other);
-	FORCEINLINE_DEBUGGABLE TTransform<T>	operator*(const TQuat<T>& Other) const;
+	[[nodiscard]] FORCEINLINE_DEBUGGABLE TTransform<T>	operator*(const TQuat<T>& Other) const;
 	FORCEINLINE_DEBUGGABLE void			operator*=(const TQuat<T>& Other);
 
-	FORCEINLINE_DEBUGGABLE static bool AnyHasNegativeScale(const TVector<T>& InScale3D, const TVector<T>& InOtherScale3D);
+	[[nodiscard]] FORCEINLINE_DEBUGGABLE static bool AnyHasNegativeScale(const TVector<T>& InScale3D, const TVector<T>& InOtherScale3D);
 	FORCEINLINE_DEBUGGABLE void ScaleTranslation(const TVector<T>& InScale3D);
 	FORCEINLINE_DEBUGGABLE void ScaleTranslation(const T& Scale);
 	FORCEINLINE_DEBUGGABLE void RemoveScaling(T Tolerance=UE_SMALL_NUMBER);
-	FORCEINLINE_DEBUGGABLE float GetMaximumAxisScale() const;
-	FORCEINLINE_DEBUGGABLE float GetMinimumAxisScale() const;
+	[[nodiscard]] FORCEINLINE_DEBUGGABLE float GetMaximumAxisScale() const;
+	[[nodiscard]] FORCEINLINE_DEBUGGABLE float GetMinimumAxisScale() const;
 	// Inverse does not work well with VQS format(in particular non-uniform), so removing it, but made two below functions to be used instead. 
 
 	/*******************************************************************************************
@@ -554,67 +554,67 @@ public:
 	 * GetRelativeTransform returns this*Other(-1) and parameter is Other(not Other(-1))
 	 * GetRelativeTransformReverse returns this(-1)*Other, and parameter is Other. 
 	 *******************************************************************************************/
-	CORE_API TTransform<T> GetRelativeTransform(const TTransform<T>& Other) const;
-	CORE_API TTransform<T> GetRelativeTransformReverse(const TTransform<T>& Other) const;
+	[[nodiscard]] CORE_API TTransform<T> GetRelativeTransform(const TTransform<T>& Other) const;
+	[[nodiscard]] CORE_API TTransform<T> GetRelativeTransformReverse(const TTransform<T>& Other) const;
 	/**
 	 * Set current transform and the relative to ParentTransform.
 	 * Equates to This = This->GetRelativeTransform(Parent), but saves the intermediate TTransform<T> storage and copy.
 	 */
 	CORE_API void		SetToRelativeTransform(const TTransform<T>& ParentTransform);
 
-	FORCEINLINE_DEBUGGABLE TVector4<T>	TransformFVector4(const TVector4<T>& V) const;
-	FORCEINLINE_DEBUGGABLE TVector4<T>	TransformFVector4NoScale(const TVector4<T>& V) const;
-	FORCEINLINE_DEBUGGABLE TVector<T>		TransformPosition(const TVector<T>& V) const;
-	FORCEINLINE_DEBUGGABLE TVector<T>		TransformPositionNoScale(const TVector<T>& V) const;
+	[[nodiscard]] FORCEINLINE_DEBUGGABLE TVector4<T>	TransformFVector4(const TVector4<T>& V) const;
+	[[nodiscard]] FORCEINLINE_DEBUGGABLE TVector4<T>	TransformFVector4NoScale(const TVector4<T>& V) const;
+	[[nodiscard]] FORCEINLINE_DEBUGGABLE TVector<T>		TransformPosition(const TVector<T>& V) const;
+	[[nodiscard]] FORCEINLINE_DEBUGGABLE TVector<T>		TransformPositionNoScale(const TVector<T>& V) const;
 
 
 	/** Inverts the transform and then transforms V - correctly handles scaling in this transform. */
-	FORCEINLINE_DEBUGGABLE TVector<T>		InverseTransformPosition(const TVector<T> &V) const;
-	FORCEINLINE_DEBUGGABLE TVector<T>		InverseTransformPositionNoScale(const TVector<T> &V) const;
-	FORCEINLINE_DEBUGGABLE TVector<T>		TransformVector(const TVector<T>& V) const;
-	FORCEINLINE_DEBUGGABLE TVector<T>		TransformVectorNoScale(const TVector<T>& V) const;
+	[[nodiscard]] FORCEINLINE_DEBUGGABLE TVector<T>		InverseTransformPosition(const TVector<T> &V) const;
+	[[nodiscard]] FORCEINLINE_DEBUGGABLE TVector<T>		InverseTransformPositionNoScale(const TVector<T> &V) const;
+	[[nodiscard]] FORCEINLINE_DEBUGGABLE TVector<T>		TransformVector(const TVector<T>& V) const;
+	[[nodiscard]] FORCEINLINE_DEBUGGABLE TVector<T>		TransformVectorNoScale(const TVector<T>& V) const;
 
 	/** 
 	 *	Transform a direction vector by the inverse of this matrix - will not take into account translation part.
 	 *	If you want to transform a surface normal (or plane) and correctly account for non-uniform scaling you should use TransformByUsingAdjointT with adjoint of matrix inverse.
 	 */
-	FORCEINLINE_DEBUGGABLE TVector<T> InverseTransformVector(const TVector<T> &V) const;
-	FORCEINLINE_DEBUGGABLE TVector<T> InverseTransformVectorNoScale(const TVector<T> &V) const;
+	[[nodiscard]] FORCEINLINE_DEBUGGABLE TVector<T> InverseTransformVector(const TVector<T> &V) const;
+	[[nodiscard]] FORCEINLINE_DEBUGGABLE TVector<T> InverseTransformVectorNoScale(const TVector<T> &V) const;
 
 	/**
 	* Transform a rotation.
 	* For example if this is a LocalToWorld transform, TransformRotation(Q) would transform Q from local to world space.
 	*/
-	FORCEINLINE_DEBUGGABLE TQuat<T> TransformRotation(const TQuat<T>& Q) const;
+	[[nodiscard]] FORCEINLINE_DEBUGGABLE TQuat<T> TransformRotation(const TQuat<T>& Q) const;
 
 	/**
 	* Inverse transform a rotation.
 	* For example if this is a LocalToWorld transform, InverseTransformRotation(Q) would transform Q from world to local space.
 	*/
-	FORCEINLINE_DEBUGGABLE TQuat<T> InverseTransformRotation(const TQuat<T>& Q) const;
+	[[nodiscard]] FORCEINLINE_DEBUGGABLE TQuat<T> InverseTransformRotation(const TQuat<T>& Q) const;
 
-	FORCEINLINE_DEBUGGABLE TTransform<T>	GetScaled(T Scale) const;
-	FORCEINLINE_DEBUGGABLE TTransform<T>	GetScaled(TVector<T> Scale) const;
-	FORCEINLINE_DEBUGGABLE TVector<T>		GetScaledAxis(EAxis::Type InAxis) const;
-	FORCEINLINE_DEBUGGABLE TVector<T>		GetUnitAxis(EAxis::Type InAxis) const;
+	[[nodiscard]] FORCEINLINE_DEBUGGABLE TTransform<T>	GetScaled(T Scale) const;
+	[[nodiscard]] FORCEINLINE_DEBUGGABLE TTransform<T>	GetScaled(TVector<T> Scale) const;
+	[[nodiscard]] FORCEINLINE_DEBUGGABLE TVector<T>		GetScaledAxis(EAxis::Type InAxis) const;
+	[[nodiscard]] FORCEINLINE_DEBUGGABLE TVector<T>		GetUnitAxis(EAxis::Type InAxis) const;
 	FORCEINLINE_DEBUGGABLE void		Mirror(EAxis::Type MirrorAxis, EAxis::Type FlipAxis);
-	FORCEINLINE_DEBUGGABLE static TVector<T>	GetSafeScaleReciprocal(const TVector<T>& InScale, T Tolerance=UE_SMALL_NUMBER);
+	[[nodiscard]] FORCEINLINE_DEBUGGABLE static TVector<T>	GetSafeScaleReciprocal(const TVector<T>& InScale, T Tolerance=UE_SMALL_NUMBER);
 
 
-	FORCEINLINE_DEBUGGABLE TVector<T> GetLocation() const
+	[[nodiscard]] FORCEINLINE_DEBUGGABLE TVector<T> GetLocation() const
 	{
 		return GetTranslation();
 	}
 
-	FORCEINLINE_DEBUGGABLE TRotator<T> Rotator() const
+	[[nodiscard]] FORCEINLINE_DEBUGGABLE TRotator<T> Rotator() const
 	{
 		TQuat<T> OutRotation;
 		VectorStoreAligned(Rotation, &OutRotation);
 		return OutRotation.Rotator();
 	}
 
-	/** Calculate the  */
-	FORCEINLINE_DEBUGGABLE T GetDeterminant() const
+	/** Calculate the determinant */
+	[[nodiscard]] FORCEINLINE_DEBUGGABLE T GetDeterminant() const
 	{
 		//#todo - vectorized version of this
 		UE::Math::TVector4<T> OutScale3D;
@@ -633,7 +633,7 @@ public:
 	 * Checks the components for NaN's
 	 * @return Returns true if any component (rotation, translation, or scale) is a NAN
 	 */
-	bool ContainsNaN() const
+	[[nodiscard]] bool ContainsNaN() const
 	{
 		if (VectorContainsNaNOrInfinite(Rotation))
 		{
@@ -651,7 +651,7 @@ public:
 		return false;
 	}
 
-	inline bool IsValid() const
+	[[nodiscard]] inline bool IsValid() const
 	{
 		if ( ContainsNaN() )
 		{
@@ -738,56 +738,56 @@ private:
 public:
 
 	// Test if A's rotation equals B's rotation, within a tolerance. Preferred over "A.GetRotation().Equals(B.GetRotation())" because it is faster on some platforms.
-	FORCEINLINE_DEBUGGABLE static bool AreRotationsEqual(const TTransform<T>& A, const TTransform<T>& B, FReal Tolerance=UE_KINDA_SMALL_NUMBER)
+	[[nodiscard]] FORCEINLINE_DEBUGGABLE static bool AreRotationsEqual(const TTransform<T>& A, const TTransform<T>& B, FReal Tolerance=UE_KINDA_SMALL_NUMBER)
 	{
 		return A.Private_RotationEquals(B.Rotation, Tolerance);
 	}
 
 	// Test if A's translation equals B's translation, within a tolerance. Preferred over "A.GetTranslation().Equals(B.GetTranslation())" because it avoids TransformVectorRegister->TVector<T> conversion.
-	FORCEINLINE_DEBUGGABLE static bool AreTranslationsEqual(const TTransform<T>& A, const TTransform<T>& B, FReal Tolerance=UE_KINDA_SMALL_NUMBER)
+	[[nodiscard]] FORCEINLINE_DEBUGGABLE static bool AreTranslationsEqual(const TTransform<T>& A, const TTransform<T>& B, FReal Tolerance=UE_KINDA_SMALL_NUMBER)
 	{
 		return A.Private_TranslationEquals(B.Translation, Tolerance);
 	}
 
 	// Test if A's scale equals B's scale, within a tolerance. Preferred over "A.GetScale3D().Equals(B.GetScale3D())" because it avoids TransformVectorRegister->TVector<T> conversion.
-	FORCEINLINE_DEBUGGABLE static bool AreScale3DsEqual(const TTransform<T>& A, const TTransform<T>& B, FReal Tolerance=UE_KINDA_SMALL_NUMBER)
+	[[nodiscard]] FORCEINLINE_DEBUGGABLE static bool AreScale3DsEqual(const TTransform<T>& A, const TTransform<T>& B, FReal Tolerance=UE_KINDA_SMALL_NUMBER)
 	{
 		return A.Private_Scale3DEquals(B.Scale3D, Tolerance);
 	}
 
 
 	// Test if this Transform's rotation equals another's rotation, within a tolerance. Preferred over "GetRotation().Equals(Other.GetRotation())" because it is faster on some platforms.
-	FORCEINLINE_DEBUGGABLE bool RotationEquals(const TTransform<T>& Other, FReal Tolerance=UE_KINDA_SMALL_NUMBER) const
+	[[nodiscard]] FORCEINLINE_DEBUGGABLE bool RotationEquals(const TTransform<T>& Other, FReal Tolerance=UE_KINDA_SMALL_NUMBER) const
 	{
 		return AreRotationsEqual(*this, Other, Tolerance);
 	}
 
 	// Test if this Transform's translation equals another's translation, within a tolerance. Preferred over "GetTranslation().Equals(Other.GetTranslation())" because it avoids TransformVectorRegister->TVector<T> conversion.
-	FORCEINLINE_DEBUGGABLE bool TranslationEquals(const TTransform<T>& Other, FReal Tolerance=UE_KINDA_SMALL_NUMBER) const
+	[[nodiscard]] FORCEINLINE_DEBUGGABLE bool TranslationEquals(const TTransform<T>& Other, FReal Tolerance=UE_KINDA_SMALL_NUMBER) const
 	{
 		return AreTranslationsEqual(*this, Other, Tolerance);
 	}
 
 	// Test if this Transform's scale equals another's scale, within a tolerance. Preferred over "GetScale3D().Equals(Other.GetScale3D())" because it avoids TransformVectorRegister->TVector<T> conversion.
-	FORCEINLINE_DEBUGGABLE bool Scale3DEquals(const TTransform<T>& Other, FReal Tolerance=UE_KINDA_SMALL_NUMBER) const
+	[[nodiscard]] FORCEINLINE_DEBUGGABLE bool Scale3DEquals(const TTransform<T>& Other, FReal Tolerance=UE_KINDA_SMALL_NUMBER) const
 	{
 		return AreScale3DsEqual(*this, Other, Tolerance);
 	}
 
 	// Test if all components of the transforms are equal, within a tolerance.
-	FORCEINLINE_DEBUGGABLE bool Equals(const TTransform<T>& Other, FReal Tolerance=UE_KINDA_SMALL_NUMBER) const
+	[[nodiscard]] FORCEINLINE_DEBUGGABLE bool Equals(const TTransform<T>& Other, FReal Tolerance=UE_KINDA_SMALL_NUMBER) const
 	{
 		return Private_TranslationEquals(Other.Translation, Tolerance) && Private_RotationEquals(Other.Rotation, Tolerance) && Private_Scale3DEquals(Other.Scale3D, Tolerance);
 	}
 
 	// Test if all components of the transform property are equal.
-	FORCEINLINE_DEBUGGABLE bool Identical(const TTransform<T>* Other, uint32 PortFlags) const
+	[[nodiscard]] FORCEINLINE_DEBUGGABLE bool Identical(const TTransform<T>* Other, uint32 PortFlags) const
 	{
 		return Equals(*Other, 0.f);
 	}
 
 	// Test if rotation and translation components of the transforms are equal, within a tolerance.
-	FORCEINLINE_DEBUGGABLE bool EqualsNoScale(const TTransform<T>& Other, FReal Tolerance=UE_KINDA_SMALL_NUMBER) const
+	[[nodiscard]] FORCEINLINE_DEBUGGABLE bool EqualsNoScale(const TTransform<T>& Other, FReal Tolerance=UE_KINDA_SMALL_NUMBER) const
 	{
 		return Private_TranslationEquals(Other.Translation, Tolerance) && Private_RotationEquals(Other.Rotation, Tolerance);
 	}
@@ -890,7 +890,7 @@ public:
 	 * Add the translations from two FTransforms and return the result.
 	 * @return A.Translation + B.Translation
 	 */
-	FORCEINLINE_DEBUGGABLE static TVector<T> AddTranslations(const TTransform<T>& A, const TTransform<T>& B)
+	[[nodiscard]] FORCEINLINE_DEBUGGABLE static TVector<T> AddTranslations(const TTransform<T>& A, const TTransform<T>& B)
 	{
 		TVector<T> Result;
 		VectorStoreFloat3(VectorAdd(A.Translation, B.Translation), &Result);
@@ -901,7 +901,7 @@ public:
 	 * Subtract translations from two FTransforms and return the difference.
 	 * @return A.Translation - B.Translation.
 	 */
-	FORCEINLINE_DEBUGGABLE static TVector<T> SubtractTranslations(const TTransform<T>& A, const TTransform<T>& B)
+	[[nodiscard]] FORCEINLINE_DEBUGGABLE static TVector<T> SubtractTranslations(const TTransform<T>& A, const TTransform<T>& B)
 	{
 		TVector<T> Result;
 		VectorStoreFloat3(VectorSubtract(A.Translation, B.Translation), &Result);
@@ -955,10 +955,22 @@ public:
 	}
 	
 	// For low-level VectorRegister programming
-	const TPersistentVectorRegisterType<T>& GetTranslationRegister() const { return Translation; }
-	const TPersistentVectorRegisterType<T>& GetRotationRegister() const { return Rotation; }
-	void SetTranslationRegister(TransformVectorRegister InTranslation) { Translation = InTranslation; }
-	void SetRotationRegister(TransformVectorRegister InRotation) { Rotation = InRotation; }
+	[[nodiscard]] const TPersistentVectorRegisterType<T>& GetTranslationRegister() const
+	{
+		return Translation;
+	}
+	[[nodiscard]] const TPersistentVectorRegisterType<T>& GetRotationRegister() const
+	{
+		return Rotation;
+	}
+	void SetTranslationRegister(TransformVectorRegister InTranslation)
+	{
+		Translation = InTranslation;
+	}
+	void SetRotationRegister(TransformVectorRegister InRotation) 
+	{
+		Rotation = InRotation;
+	}
 
 	/** @note : Added template type function for Accumulate
 	  * The template type isn't much useful yet, but it is with the plan to move forward
@@ -1117,7 +1129,7 @@ public:
 	 *
 	 * @return true if the rotation component is normalized, and false otherwise.
 	 */
-	FORCEINLINE_DEBUGGABLE bool IsRotationNormalized() const
+	[[nodiscard]] FORCEINLINE_DEBUGGABLE bool IsRotationNormalized() const
 	{		
 		const TransformVectorRegister TestValue = VectorAbs(VectorSubtract(VectorOne(), VectorDot4(Rotation, Rotation)));
 		return !VectorAnyGreaterThan(TestValue, GlobalVectorConstants::ThreshQuatNormalized);
@@ -1183,7 +1195,7 @@ public:
 	 *
 	 * @return The rotation component
 	 */
-	FORCEINLINE_DEBUGGABLE TQuat<T> GetRotation() const
+	[[nodiscard]] FORCEINLINE_DEBUGGABLE TQuat<T> GetRotation() const
 	{
 		DiagnosticCheckNaN_Rotate();
 		TQuat<T> OutRotation;
@@ -1196,7 +1208,7 @@ public:
 	 *
 	 * @return The translation component
 	 */
-	FORCEINLINE_DEBUGGABLE TVector<T> GetTranslation() const
+	[[nodiscard]] FORCEINLINE_DEBUGGABLE TVector<T> GetTranslation() const
 	{
 		DiagnosticCheckNaN_Translate();
 		TVector<T> OutTranslation;
@@ -1209,7 +1221,7 @@ public:
 	 *
 	 * @return The Scale3D component
 	 */
-	FORCEINLINE_DEBUGGABLE TVector<T> GetScale3D() const
+	[[nodiscard]] FORCEINLINE_DEBUGGABLE TVector<T> GetScale3D() const
 	{
 		DiagnosticCheckNaN_Scale3D();
 		TVector<T> OutScale3D;
@@ -1376,7 +1388,7 @@ private:
 	 * also returning BIG_NUMBER causes sequential NaN issues by multiplying 
 	 * so we hardcode as 0
 	 */
-	static FORCEINLINE_DEBUGGABLE TransformVectorRegister		GetSafeScaleReciprocal(const TransformVectorRegister& InScale, const ScalarRegister& Tolerance = ScalarRegister(GlobalVectorConstants::SmallNumber))
+	[[nodiscard]] static FORCEINLINE_DEBUGGABLE TransformVectorRegister		GetSafeScaleReciprocal(const TransformVectorRegister& InScale, const ScalarRegister& Tolerance = ScalarRegister(GlobalVectorConstants::SmallNumber))
 	{		
 		// SafeReciprocalScale.X = (InScale.X == 0) ? 0.f : 1/InScale.X; // same for YZW
 
@@ -1395,7 +1407,7 @@ private:
 	}
 
 	/** Returns Inverse Transform of this TTransform<T> **/
-	FORCEINLINE_DEBUGGABLE TTransform<T> InverseFast() const
+	[[nodiscard]] FORCEINLINE_DEBUGGABLE TTransform<T> InverseFast() const
 	{
 		// Inverse QST (A) = QST (~A)
 		// Since A*~A = Identity, 
diff --git a/Engine/Source/Runtime/Core/Public/Math/TranslationMatrix.h b/Engine/Source/Runtime/Core/Public/Math/TranslationMatrix.h
--- a/Engine/Source/Runtime/Core/Public/Math/TranslationMatrix.h
+++ b/Engine/Source/Runtime/Core/Public/Math/TranslationMatrix.h
@@ -17,17 +17,17 @@ struct TTranslationMatrix
 public:
 
 	/** Constructor translation matrix based on given vector */
-	TTranslationMatrix(const TVector<T>& Delta);
+	[[nodiscard]] TTranslationMatrix(const TVector<T>& Delta);
 
 	// Conversion to other type.
 	template<typename FArg UE_REQUIRES(!std::is_same_v<T, FArg>)>
-	explicit TTranslationMatrix(const TTranslationMatrix<FArg>& From)
+	[[nodiscard]] explicit TTranslationMatrix(const TTranslationMatrix<FArg>& From)
 		: TMatrix<T>(From)
 	{
 	}
 	
 	/** Matrix factory. Return an FMatrix so we don't have type conversion issues in expressions. */
-	static TMatrix<T> Make(TVector<T> const& Delta)
+	[[nodiscard]] static TMatrix<T> Make(TVector<T> const& Delta)
 	{
 		return TTranslationMatrix<T>(Delta);
 	}
