Make Maths functions [[nodiscard]]

diff --git a/Engine/Source/Runtime/Core/Public/Math/TwoVectors.h b/Engine/Source/Runtime/Core/Public/Math/TwoVectors.h
--- a/Engine/Source/Runtime/Core/Public/Math/TwoVectors.h
+++ b/Engine/Source/Runtime/Core/Public/Math/TwoVectors.h
@@ -24,7 +24,7 @@ public:
 public:
 
 	/** Default constructor. */
-	FORCEINLINE	FTwoVectors();
+	[[nodiscard]] FORCEINLINE	FTwoVectors();
 
 	/**
 	 * Creates and initializes a new instance with the specified vectors.
@@ -32,14 +32,14 @@ public:
 	 * @param In1 The first Vector.
 	 * @param In2 The second Vector.
 	 */
-	FORCEINLINE	FTwoVectors( FVector In1, FVector In2 );
+	[[nodiscard]] FORCEINLINE	FTwoVectors( FVector In1, FVector In2 );
 
 	/**
 	 * Constructor.
 	 *
 	 * @param EForceInit Force Init Enum
 	 */
-	explicit FORCEINLINE FTwoVectors( EForceInit );
+	[[nodiscard]] explicit FORCEINLINE FTwoVectors( EForceInit );
 
 public:
 
@@ -49,7 +49,7 @@ public:
 	 * @param V The pair to add.
 	 * @return Result of addition.
 	 */
-	FORCEINLINE FTwoVectors operator+( const FTwoVectors& V ) const;
+	[[nodiscard]] FORCEINLINE FTwoVectors operator+( const FTwoVectors& V ) const;
 
 	/**
 	 * Gets result of subtraction of two pairs of vectors.
@@ -57,7 +57,7 @@ public:
 	 * @param V The pair to subtract.
 	 * @return Result of subtraction.
 	 */
-	FORCEINLINE FTwoVectors operator-( const FTwoVectors& V ) const;
+	[[nodiscard]] FORCEINLINE FTwoVectors operator-( const FTwoVectors& V ) const;
 
 	/**
 	 * Gets result of scaling pair of vectors.
@@ -65,7 +65,7 @@ public:
 	 * @param Scale The scaling factor.
 	 * @return Result of Scaling.
 	 */
-	FORCEINLINE FTwoVectors operator*( float Scale ) const;
+	[[nodiscard]] FORCEINLINE FTwoVectors operator*( float Scale ) const;
 
 	/**
 	 * Gets result of dividing pair of vectors.
@@ -73,7 +73,7 @@ public:
 	 * @param Scale What to divide by.
 	 * @return Result of division.
 	 */
-	FTwoVectors operator/( float Scale ) const;
+	[[nodiscard]] FTwoVectors operator/( float Scale ) const;
 
 	/**
 	 * Gets result of multiplying two pairs of vectors.
@@ -81,7 +81,7 @@ public:
 	 * @param V The pair to multiply with.
 	 * @return Result of multiplication.
 	 */
-	FORCEINLINE FTwoVectors operator*( const FTwoVectors& V ) const;
+	[[nodiscard]] FORCEINLINE FTwoVectors operator*( const FTwoVectors& V ) const;
 
 	/**
 	 * Gets result of division of two pairs of vectors.
@@ -89,7 +89,7 @@ public:
 	 * @param V The pair to divide by.
 	 * @return Result of division.
 	 */
-	FORCEINLINE FTwoVectors operator/( const FTwoVectors& V ) const;
+	[[nodiscard]] FORCEINLINE FTwoVectors operator/( const FTwoVectors& V ) const;
 
 	// Binary comparison operators.
 
@@ -99,7 +99,7 @@ public:
 	 * @param V The other pair.
 	 * @return true if the two pairs are equal, false otherwise..
 	 */
-	bool operator==( const FTwoVectors& V ) const;
+	[[nodiscard]] bool operator==( const FTwoVectors& V ) const;
 
 	/**
 	 * Checks two pairs for inequality.
@@ -107,7 +107,7 @@ public:
 	 * @param V The other pair.
 	 * @return true if the two pairs are different, false otherwise..
 	 */
-	bool operator!=( const FTwoVectors& V ) const;
+	[[nodiscard]] bool operator!=( const FTwoVectors& V ) const;
 
 	/**
 	 * Error-tolerant comparison.
@@ -116,7 +116,7 @@ public:
 	 * @param Tolerance Error Tolerance.
 	 * @return true if two pairs are equal within specified tolerance, false otherwise..
 	 */
-	bool Equals( const FTwoVectors& V, float Tolerance = UE_KINDA_SMALL_NUMBER ) const;
+	[[nodiscard]] bool Equals( const FTwoVectors& V, float Tolerance = UE_KINDA_SMALL_NUMBER ) const;
 
 	// Unary operators.
 
@@ -125,7 +125,7 @@ public:
 	 *
 	 * @return A negated copy of the pair.
 	 */
-	FORCEINLINE FTwoVectors operator-() const;
+	[[nodiscard]] FORCEINLINE FTwoVectors operator-() const;
 
 	// Assignment operators.
 
@@ -193,21 +193,21 @@ public:
 	 *
 	 * @return The maximum value of all the vector coordinates.
 	 */
-	FVector::FReal GetMax() const;
+	[[nodiscard]] FVector::FReal GetMax() const;
 
 	/**
 	 * Get the minimum value of all the vector coordinates.
 	 *
 	 * @return The minimum value of all the vector coordinates.
 	 */
-	FVector::FReal GetMin() const;
+	[[nodiscard]] FVector::FReal GetMin() const;
 
 	/**
 	 * Get a textual representation of this two-vector.
 	 *
 	 * @return A string describing the two-vector.
 	 */
-	FString ToString() const;
+	[[nodiscard]] FString ToString() const;
 
 public:
 
diff --git a/Engine/Source/Runtime/Core/Public/Math/UnitConversion.h b/Engine/Source/Runtime/Core/Public/Math/UnitConversion.h
--- a/Engine/Source/Runtime/Core/Public/Math/UnitConversion.h
+++ b/Engine/Source/Runtime/Core/Public/Math/UnitConversion.h
@@ -109,7 +109,7 @@ class FUnitSettings
 {
 public:
 
-	CORE_API FUnitSettings();
+	[[nodiscard]] CORE_API FUnitSettings();
 
 	/** Check whether unit display is globally enabled or disabled */
 	CORE_API bool ShouldDisplayUnits() const;
@@ -139,39 +139,39 @@ private:
 struct FUnitConversion
 {
 	/** Get the global settings for unit conversion/display */
-	static CORE_API FUnitSettings& Settings();
+	[[nodiscard]] static CORE_API FUnitSettings& Settings();
 
 	/** Check whether it is possible to convert a number between the two specified units */
-	static CORE_API bool AreUnitsCompatible(EUnit From, EUnit To);
+	[[nodiscard]] static CORE_API bool AreUnitsCompatible(EUnit From, EUnit To);
 
 	/** Check whether a unit is of the specified type */
-	static CORE_API bool IsUnitOfType(EUnit Unit, EUnitType Type);
+	[[nodiscard]] static CORE_API bool IsUnitOfType(EUnit Unit, EUnitType Type);
 
 	/** Get the type of the specified unit */
-	static CORE_API EUnitType GetUnitType(EUnit);
+	[[nodiscard]] static CORE_API EUnitType GetUnitType(EUnit);
 
 	/** Get the display string for the the specified unit type */
-	static CORE_API const TCHAR* GetUnitDisplayString(EUnit Unit);
+	[[nodiscard]] static CORE_API const TCHAR* GetUnitDisplayString(EUnit Unit);
 
 	/** Helper function to find a unit from a string (name or display string) */
-	static CORE_API TOptional<EUnit> UnitFromString(const TCHAR* UnitString);
+	[[nodiscard]] static CORE_API TOptional<EUnit> UnitFromString(const TCHAR* UnitString);
 
 	/** Helper function to get all supported units */
-	static CORE_API TConstArrayView<const TCHAR*> GetSupportedUnits();
+	[[nodiscard]] static CORE_API TConstArrayView<const TCHAR*> GetSupportedUnits();
 
 public:
 
 	/** Convert the specified number from one unit to another. Does nothing if the units are incompatible. */
 	template<typename T>
-	static T Convert(T InValue, EUnit From, EUnit To);
+	[[nodiscard]] static T Convert(T InValue, EUnit From, EUnit To);
 
 	/** Quantizes this number to the most appropriate unit for user friendly presentation (e.g. 1000m returns 1km). */
 	template<typename T>
-	static FNumericUnit<T> QuantizeUnitsToBestFit(T Value, EUnit Units);
+	[[nodiscard]] static FNumericUnit<T> QuantizeUnitsToBestFit(T Value, EUnit Units);
 
 	/** Quantizes this number to the most appropriate unit for user friendly presentation (e.g. 1000m returns 1km), adhereing to global display settings. */
 	template<typename T>
-	static EUnit CalculateDisplayUnit(T Value, EUnit InUnits);
+	[[nodiscard]] static EUnit CalculateDisplayUnit(T Value, EUnit InUnits);
 
 };
 
@@ -189,24 +189,26 @@ struct FNumericUnit
 	const EUnit Units;
 
 	/** Constructors */
-	FNumericUnit();
-	FNumericUnit(const NumericType& InValue, EUnit InUnits = EUnit::Unspecified);
+	[[nodiscard]] FNumericUnit();
+	[[nodiscard]] FNumericUnit(const NumericType& InValue, EUnit InUnits = EUnit::Unspecified);
 
 	/** Copy construction/assignment from the same type */
-	FNumericUnit(const FNumericUnit& Other);
+	[[nodiscard]] FNumericUnit(const FNumericUnit& Other);
 	FNumericUnit& operator=(const FNumericUnit& Other);
 
 	/** Templated Copy construction/assignment from differing numeric types. Relies on implicit conversion of the two numeric types. */
-	template<typename OtherType> FNumericUnit(const FNumericUnit<OtherType>& Other);
-	template<typename OtherType> FNumericUnit& operator=(const FNumericUnit<OtherType>& Other);
+	template<typename OtherType>
+	[[nodiscard]] FNumericUnit(const FNumericUnit<OtherType>& Other);
+	template<typename OtherType>
+	FNumericUnit& operator=(const FNumericUnit<OtherType>& Other);
 
 	/** Convert this quantity to a different unit */
-	TOptional<FNumericUnit<NumericType>> ConvertTo(EUnit ToUnits) const;
+	[[nodiscard]] TOptional<FNumericUnit<NumericType>> ConvertTo(EUnit ToUnits) const;
 
 public:
 
 	/** Quantizes this number to the most appropriate unit for user friendly presentation (e.g. 1000m returns 1km). */
-	FNumericUnit<NumericType> QuantizeUnitsToBestFit() const;
+	[[nodiscard]] FNumericUnit<NumericType> QuantizeUnitsToBestFit() const;
 
 	/** Try and parse an expression into a numeric unit */
 	static TValueOrError<FNumericUnit<NumericType>, FText> TryParseExpression(const TCHAR* InExpression, EUnit InDefaultUnit, const FNumericUnit<NumericType>& InExistingValue);
@@ -216,7 +218,7 @@ public:
 
 private:
 	/** Conversion to the numeric type disabled as coupled with implicit construction from NumericType can easily lead to loss of associated units. */
-	operator const NumericType&() const;
+	operator const NumericType&() const = delete;
 
 	/** Copy another unit into this one, taking account of its units, and applying necessary conversion */
 	template<typename OtherType>
@@ -227,7 +229,7 @@ private:
 
 	/** Global arithmetic operators for number types. Deals with conversion from related units correctly. Note must be inlined for hidden friend optimization to work */
 	template<typename OtherType>
-	friend inline bool operator==(const FNumericUnit<NumericType>& LHS, const FNumericUnit<OtherType>& RHS)
+	[[nodiscard]] friend inline bool operator==(const FNumericUnit<NumericType>& LHS, const FNumericUnit<OtherType>& RHS)
 	{
 		if (LHS.Units != EUnit::Unspecified && RHS.Units != EUnit::Unspecified)
 		{
@@ -252,7 +254,7 @@ private:
 	}
 
 	template<typename OtherType>
-	friend inline bool operator!=(const FNumericUnit<NumericType>& LHS, const FNumericUnit<OtherType>& RHS)
+	[[nodiscard]] friend inline bool operator!=(const FNumericUnit<NumericType>& LHS, const FNumericUnit<OtherType>& RHS)
 	{
 		return !(LHS == RHS);
 	}
diff --git a/Engine/Source/Runtime/Core/Public/Math/Vector.h b/Engine/Source/Runtime/Core/Public/Math/Vector.h
--- a/Engine/Source/Runtime/Core/Public/Math/Vector.h
+++ b/Engine/Source/Runtime/Core/Public/Math/Vector.h
@@ -109,19 +109,19 @@ public:
 	CORE_API static const TVector<T> ZAxisVector;
 
 	/** @return Zero Vector (0,0,0) */
-	static inline TVector<T> Zero() { return ZeroVector; }
+	[[nodiscard]] static inline TVector<T> Zero() { return ZeroVector; }
 
 	/** @return One Vector (1,1,1) */
-	static inline TVector<T> One() { return OneVector; }
+	[[nodiscard]] static inline TVector<T> One() { return OneVector; }
 
 	/** @return Unit X Vector (1,0,0)  */
-	static inline TVector<T> UnitX() { return XAxisVector; }
+	[[nodiscard]] static inline TVector<T> UnitX() { return XAxisVector; }
 
 	/** @return Unit Y Vector (0,1,0)  */
-	static inline TVector<T> UnitY() { return YAxisVector; }
+	[[nodiscard]] static inline TVector<T> UnitY() { return YAxisVector; }
 
 	/** @return Unit Z Vector (0,0,1)  */
-	static inline TVector<T> UnitZ() { return ZAxisVector; }
+	[[nodiscard]] static inline TVector<T> UnitZ() { return ZAxisVector; }
 
 public:
 
@@ -149,16 +149,16 @@ public:
 #endif
 
     /** Default constructor (no initialization). */
-    TVector() = default;
+    [[nodiscard]] TVector() = default;
 
     /**
      * Constructor initializing all components to a single T value.
      *
      * @param InF Value to set all components to.
      */
-    explicit FORCEINLINE TVector(T InF);
+    [[nodiscard]] explicit FORCEINLINE TVector(T InF);
 
-	FORCEINLINE constexpr TVector(T InF, TVectorConstInit);
+	[[nodiscard]] FORCEINLINE constexpr TVector(T InF, TVectorConstInit);
 
     /**
      * Constructor using initial values for each component.
@@ -167,7 +167,7 @@ public:
      * @param InY Y Coordinate.
      * @param InZ Z Coordinate.
      */
-    FORCEINLINE TVector(T InX, T InY, T InZ);
+    [[nodiscard]] FORCEINLINE TVector(T InX, T InY, T InZ);
 
     /**
      * Constructs a vector from an TVector2<T> and Z value.
@@ -175,21 +175,21 @@ public:
      * @param V Vector to copy from.
      * @param InZ Z Coordinate.
      */
-    explicit FORCEINLINE TVector(const TVector2<T> V, T InZ);
+    [[nodiscard]] explicit FORCEINLINE TVector(const TVector2<T> V, T InZ);
 
     /**
      * Constructor using the XYZ components from a 4D vector.
      *
      * @param V 4D Vector to copy from.
      */
-    FORCEINLINE TVector(const UE::Math::TVector4<T>& V);
+    [[nodiscard]] FORCEINLINE TVector(const UE::Math::TVector4<T>& V);
 
     /**
      * Constructs a vector from an FLinearColor.
      *
      * @param InColor Color to copy from.
      */
-    explicit TVector(const FLinearColor& InColor);
+    [[nodiscard]] explicit TVector(const FLinearColor& InColor);
 
     /**
      * Constructs a vector from an FIntVector.
@@ -197,7 +197,7 @@ public:
      * @param InVector FIntVector to copy from.
      */
 	template <typename IntType>
-    explicit TVector(TIntVector3<IntType> InVector);
+    [[nodiscard]] explicit TVector(TIntVector3<IntType> InVector);
 
     /**
      * Constructs a vector from an FIntPoint.
@@ -205,14 +205,14 @@ public:
      * @param A Int Point used to set X and Y coordinates, Z is set to zero.
      */
 	template <typename IntType>
-    explicit TVector(TIntPoint<IntType> A);
+    [[nodiscard]] explicit TVector(TIntPoint<IntType> A);
 
     /**
      * Constructor which initializes all components to zero.
      *
      * @param EForceInit Force init enum
      */
-    explicit FORCEINLINE TVector(EForceInit);
+    [[nodiscard]] explicit FORCEINLINE TVector(EForceInit);
 
     /**
      * Calculate cross product between this and another vector.
@@ -220,7 +220,7 @@ public:
      * @param V The other vector.
      * @return The cross product.
      */
-    FORCEINLINE TVector<T> operator^(const TVector<T>& V) const;
+    [[nodiscard]] FORCEINLINE TVector<T> operator^(const TVector<T>& V) const;
 
 	/**
 	 * Calculate cross product between this and another vector.
@@ -228,7 +228,7 @@ public:
 	 * @param V The other vector.
 	 * @return The cross product.
 	 */
-	FORCEINLINE TVector<T> Cross(const TVector<T>& V2) const;
+	[[nodiscard]] FORCEINLINE TVector<T> Cross(const TVector<T>& V2) const;
 
     /**
      * Calculate the cross product of two vectors.
@@ -237,7 +237,7 @@ public:
      * @param B The second vector.
      * @return The cross product.
      */
-    FORCEINLINE static TVector<T> CrossProduct(const TVector<T>& A, const TVector<T>& B);
+    [[nodiscard]] FORCEINLINE static TVector<T> CrossProduct(const TVector<T>& A, const TVector<T>& B);
 
     /**
      * Calculate the dot product between this and another vector.
@@ -245,7 +245,7 @@ public:
      * @param V The other vector.
      * @return The dot product.
      */
-    FORCEINLINE T operator|(const TVector<T>& V) const;
+    [[nodiscard]] FORCEINLINE T operator|(const TVector<T>& V) const;
 
 	/**
 	 * Calculate the dot product between this and another vector.
@@ -253,7 +253,7 @@ public:
 	 * @param V The other vector.
 	 * @return The dot product.
 	 */
-	FORCEINLINE T Dot(const TVector<T>& V) const;
+	[[nodiscard]] FORCEINLINE T Dot(const TVector<T>& V) const;
 
     /**
      * Calculate the dot product of two vectors.
@@ -262,7 +262,7 @@ public:
      * @param B The second vector.
      * @return The dot product.
      */
-    FORCEINLINE static T DotProduct(const TVector<T>& A, const TVector<T>& B);
+    [[nodiscard]] FORCEINLINE static T DotProduct(const TVector<T>& A, const TVector<T>& B);
 
     /**
      * Gets the result of component-wise addition of this and another vector.
@@ -270,7 +270,7 @@ public:
      * @param V The vector to add to this.
      * @return The result of vector addition.
      */
-    FORCEINLINE TVector<T> operator+(const TVector<T>& V) const;
+    [[nodiscard]] FORCEINLINE TVector<T> operator+(const TVector<T>& V) const;
 
     /**
      * Gets the result of component-wise subtraction of this by another vector.
@@ -278,7 +278,7 @@ public:
      * @param V The vector to subtract from this.
      * @return The result of vector subtraction.
      */
-    FORCEINLINE TVector<T> operator-(const TVector<T>& V) const;
+    [[nodiscard]] FORCEINLINE TVector<T> operator-(const TVector<T>& V) const;
 
     /**
      * Gets the result of subtracting from each component of the vector.
@@ -287,7 +287,7 @@ public:
      * @return The result of subtraction.
      */
 	template<typename FArg UE_REQUIRES(std::is_arithmetic_v<FArg>)>
-	FORCEINLINE TVector<T> operator-(FArg Bias) const
+	[[nodiscard]] FORCEINLINE TVector<T> operator-(FArg Bias) const
 	{
 		return TVector<T>(X - (T)Bias, Y - (T)Bias, Z - (T)Bias);
 	}
@@ -299,7 +299,7 @@ public:
      * @return The result of addition.
      */
 	template<typename FArg UE_REQUIRES(std::is_arithmetic_v<FArg>)>
-	FORCEINLINE TVector<T> operator+(FArg Bias) const
+	[[nodiscard]] FORCEINLINE TVector<T> operator+(FArg Bias) const
 	{
 		return TVector<T>(X + (T)Bias, Y + (T)Bias, Z + (T)Bias);
 	}
@@ -311,7 +311,7 @@ public:
      * @return The result of multiplication.
      */
 	template<typename FArg UE_REQUIRES(std::is_arithmetic_v<FArg>)>
-	FORCEINLINE TVector<T> operator*(FArg Scale) const
+	[[nodiscard]] FORCEINLINE TVector<T> operator*(FArg Scale) const
 	{
 		return TVector<T>(X * (T)Scale, Y * (T)Scale, Z * (T)Scale);
 	}
@@ -323,7 +323,7 @@ public:
      * @return The result of division.
      */
 	template<typename FArg UE_REQUIRES(std::is_arithmetic_v<FArg>)>
-	TVector<T> operator/(FArg Scale) const
+	[[nodiscard]] TVector<T> operator/(FArg Scale) const
 	{
 		const T RScale = T(1) / Scale;
 		return TVector<T>(X * RScale, Y * RScale, Z * RScale);
@@ -335,7 +335,7 @@ public:
      * @param V The vector to multiply with.
      * @return The result of multiplication.
      */
-    FORCEINLINE TVector<T> operator*(const TVector<T>& V) const;
+    [[nodiscard]] FORCEINLINE TVector<T> operator*(const TVector<T>& V) const;
 
     /**
      * Gets the result of component-wise division of this vector by another.
@@ -343,7 +343,7 @@ public:
      * @param V The vector to divide by.
      * @return The result of division.
      */
-    FORCEINLINE TVector<T> operator/(const TVector<T>& V) const;
+    [[nodiscard]] FORCEINLINE TVector<T> operator/(const TVector<T>& V) const;
 
     // Binary comparison operators.
 
@@ -353,7 +353,7 @@ public:
      * @param V The vector to check against.
      * @return true if the vectors are equal, false otherwise.
      */
-    bool operator==(const TVector<T>& V) const;
+   [[nodiscard]] bool operator==(const TVector<T>& V) const;
 
     /**
      * Check against another vector for inequality.
@@ -361,7 +361,7 @@ public:
      * @param V The vector to check against.
      * @return true if the vectors are not equal, false otherwise.
      */
-    bool operator!=(const TVector<T>& V) const;
+    [[nodiscard]] bool operator!=(const TVector<T>& V) const;
 
     /**
      * Check against another vector for equality, within specified error limits.
@@ -370,7 +370,7 @@ public:
      * @param Tolerance Error tolerance.
      * @return true if the vectors are equal within tolerance limits, false otherwise.
      */
-    bool Equals(const TVector<T>& V, T Tolerance=UE_KINDA_SMALL_NUMBER) const;
+    [[nodiscard]] bool Equals(const TVector<T>& V, T Tolerance=UE_KINDA_SMALL_NUMBER) const;
 
     /**
      * Checks whether all components of this vector are the same, within a tolerance.
@@ -378,14 +378,14 @@ public:
      * @param Tolerance Error tolerance.
      * @return true if the vectors are equal within tolerance limits, false otherwise.
      */
-    bool AllComponentsEqual(T Tolerance=UE_KINDA_SMALL_NUMBER) const;
+    [[nodiscard]] bool AllComponentsEqual(T Tolerance=UE_KINDA_SMALL_NUMBER) const;
 
     /**
      * Get a negated copy of the vector.
      *
      * @return A negated copy of the vector.
      */
-    FORCEINLINE TVector<T> operator-() const;
+    [[nodiscard]] FORCEINLINE TVector<T> operator-() const;
 
     /**
      * Adds another vector to this.
@@ -456,7 +456,7 @@ public:
      * @param Index the index of vector component
      * @return reference to component.
      */
-    T& operator[](int32 Index);
+    [[nodiscard]] T& operator[](int32 Index);
 
     /**
      * Gets specific component of the vector.
@@ -464,7 +464,7 @@ public:
      * @param Index the index of vector component
      * @return Copy of the component.
      */
-    T operator[](int32 Index)const;
+    [[nodiscard]] T operator[](int32 Index)const;
 
     /**
     * Gets a specific component of the vector.
@@ -473,7 +473,7 @@ public:
     *
     * @return Reference to the specified component.
     */
-    T& Component(int32 Index);
+    [[nodiscard]] T& Component(int32 Index);
 
     /**
     * Gets a specific component of the vector.
@@ -481,7 +481,7 @@ public:
     * @param Index The index of the component required.
     * @return Copy of the specified component.
     */
-    T Component(int32 Index) const;
+    [[nodiscard]] T Component(int32 Index) const;
 
 	/**
 	 * Tests if index is valid, i.e. greater than or equal to zero, and less than the number of components in the vector.
@@ -492,7 +492,7 @@ public:
 	bool IsValidIndex(int32 Index) const;
 
     /** Get a specific component of the vector, given a specific axis by enum */
-    T GetComponentForAxis(EAxis::Type Axis) const;
+    [[nodiscard]] T GetComponentForAxis(EAxis::Type Axis) const;
 
     /** Set a specified componet of the vector, given a specific axis by enum */
     void SetComponentForAxis(EAxis::Type Axis, T Component);
@@ -515,83 +515,83 @@ public:
      *
      * @return The maximum value of the vector's components.
      */
-    T GetMax() const;
+    [[nodiscard]] T GetMax() const;
 
     /**
      * Get the maximum absolute value of the vector's components.
      *
      * @return The maximum absolute value of the vector's components.
      */
-    T GetAbsMax() const;
+    [[nodiscard]] T GetAbsMax() const;
 
     /**
      * Get the minimum value of the vector's components.
      *
      * @return The minimum value of the vector's components.
      */
-    T GetMin() const;
+    [[nodiscard]] T GetMin() const;
 
     /**
      * Get the minimum absolute value of the vector's components.
      *
      * @return The minimum absolute value of the vector's components.
      */
-    T GetAbsMin() const;
+    [[nodiscard]] T GetAbsMin() const;
 
     /** Gets the component-wise min of two vectors. */
-    TVector<T> ComponentMin(const TVector<T>& Other) const;
+    [[nodiscard]] TVector<T> ComponentMin(const TVector<T>& Other) const;
 
     /** Gets the component-wise max of two vectors. */
-    TVector<T> ComponentMax(const TVector<T>& Other) const;
+    [[nodiscard]] TVector<T> ComponentMax(const TVector<T>& Other) const;
 
     /**
      * Get a copy of this vector with absolute value of each component.
      *
      * @return A copy of this vector with absolute value of each component.
      */
-    TVector<T> GetAbs() const;
+    [[nodiscard]] TVector<T> GetAbs() const;
 
     /**
      * Get the length (magnitude) of this vector.
      *
      * @return The length of this vector.
      */
-    T Size() const;
+    [[nodiscard]] T Size() const;
 
 	/**
 	 * Get the length (magnitude) of this vector.
 	 *
 	 * @return The length of this vector.
 	 */
-	T Length() const;
+	[[nodiscard]] T Length() const;
 
     /**
      * Get the squared length of this vector.
      *
      * @return The squared length of this vector.
      */
-    T SizeSquared() const;
+    [[nodiscard]] T SizeSquared() const;
 
 	/**
 	 * Get the squared length of this vector.
 	 *
 	 * @return The squared length of this vector.
 	 */
-	T SquaredLength() const;
+	[[nodiscard]] T SquaredLength() const;
 
     /**
      * Get the length of the 2D components of this vector.
      *
      * @return The 2D length of this vector.
      */
-    T Size2D() const ;
+    [[nodiscard]] T Size2D() const ;
 
     /**
      * Get the squared length of the 2D components of this vector.
      *
      * @return The squared 2D length of this vector.
      */
-    T SizeSquared2D() const ;
+    [[nodiscard]] T SizeSquared2D() const ;
 
     /**
      * Checks whether vector is near to zero within a specified tolerance.
@@ -599,14 +599,14 @@ public:
      * @param Tolerance Error tolerance.
      * @return true if the vector is near to zero, false otherwise.
      */
-    bool IsNearlyZero(T Tolerance=UE_KINDA_SMALL_NUMBER) const;
+    [[nodiscard]] bool IsNearlyZero(T Tolerance=UE_KINDA_SMALL_NUMBER) const;
 
     /**
      * Checks whether all components of the vector are exactly zero.
      *
      * @return true if the vector is exactly zero, false otherwise.
      */
-    bool IsZero() const;
+    [[nodiscard]] bool IsZero() const;
 
     /**
      * Check if the vector is of unit length, with specified tolerance.
@@ -614,14 +614,14 @@ public:
      * @param LengthSquaredTolerance Tolerance against squared length.
      * @return true if the vector is a unit vector within the specified tolerance.
      */
-    FORCEINLINE bool IsUnit(T LengthSquaredTolerance = UE_KINDA_SMALL_NUMBER) const;
+    [[nodiscard]] FORCEINLINE bool IsUnit(T LengthSquaredTolerance = UE_KINDA_SMALL_NUMBER) const;
 
     /**
      * Checks whether vector is normalized.
      *
      * @return true if normalized, false otherwise.
      */
-    bool IsNormalized() const;
+    [[nodiscard]] bool IsNormalized() const;
 
     /**
      * Normalize this vector in-place if it is larger than a given tolerance. Leaves it unchanged if not.
@@ -637,7 +637,7 @@ public:
      * @return Normalized version of vector.
      * @see GetSafeNormal()
      */
-    FORCEINLINE TVector<T> GetUnsafeNormal() const;
+    [[nodiscard]] FORCEINLINE TVector<T> GetUnsafeNormal() const;
 
     /**
      * Gets a normalized copy of the vector, checking it is safe to do so based on the length.
@@ -646,7 +646,7 @@ public:
      * @param Tolerance Minimum squared vector length.
      * @return A normalized copy if safe, ResultIfZero otherwise.
      */
-    TVector<T> GetSafeNormal(T Tolerance=UE_SMALL_NUMBER, const TVector<T>& ResultIfZero = ZeroVector) const;
+    [[nodiscard]] TVector<T> GetSafeNormal(T Tolerance=UE_SMALL_NUMBER, const TVector<T>& ResultIfZero = ZeroVector) const;
 
     /**
      * Gets a normalized copy of the 2D components of the vector, checking it is safe to do so. Z is set to zero. 
@@ -655,7 +655,7 @@ public:
      * @param Tolerance Minimum squared vector length.
      * @return Normalized copy if safe, otherwise returns ResultIfZero.
      */
-    TVector<T> GetSafeNormal2D(T Tolerance=UE_SMALL_NUMBER, const TVector<T>& ResultIfZero = ZeroVector) const;
+    [[nodiscard]] TVector<T> GetSafeNormal2D(T Tolerance=UE_SMALL_NUMBER, const TVector<T>& ResultIfZero = ZeroVector) const;
 
     /**
      * Util to convert this vector into a unit direction vector and its original length.
@@ -672,14 +672,14 @@ public:
      *
      * @param A copy of the vector with each component set to +1 or -1
      */
-    FORCEINLINE TVector<T> GetSignVector() const;
+    [[nodiscard]] FORCEINLINE TVector<T> GetSignVector() const;
 
     /**
      * Projects 2D components of vector based on Z.
      *
      * @return Projected version of vector based on Z.
      */
-    TVector<T> Projection() const;
+    [[nodiscard]] TVector<T> Projection() const;
 
     /**
     * Calculates normalized 2D version of vector without checking for zero length.
@@ -687,7 +687,7 @@ public:
     * @return Normalized version of vector.
     * @see GetSafeNormal2D()
     */
-    FORCEINLINE TVector<T> GetUnsafeNormal2D() const;
+    [[nodiscard]] FORCEINLINE TVector<T> GetUnsafeNormal2D() const;
 
     /**
      * Gets a copy of this vector snapped to a grid.
@@ -696,7 +696,7 @@ public:
      * @return A copy of this vector snapped to a grid.
      * @see FMath::GridSnap()
      */
-    TVector<T> GridSnap(const T& GridSz) const;
+    [[nodiscard]] TVector<T> GridSnap(const T& GridSz) const;
 
     /**
      * Get a copy of this vector, clamped inside of a cube.
@@ -704,22 +704,22 @@ public:
      * @param Radius Half size of the cube.
      * @return A copy of this vector, bound by cube.
      */
-    TVector<T> BoundToCube(T Radius) const;
+    [[nodiscard]] TVector<T> BoundToCube(T Radius) const;
 
     /** Get a copy of this vector, clamped inside of a cube. */
-    TVector<T> BoundToBox(const TVector<T>& Min, const TVector<T>& Max) const;
+    [[nodiscard]] TVector<T> BoundToBox(const TVector<T>& Min, const TVector<T>& Max) const;
 
     /** Create a copy of this vector, with its magnitude clamped between Min and Max. */
-    TVector<T> GetClampedToSize(T Min, T Max) const;
+    [[nodiscard]] TVector<T> GetClampedToSize(T Min, T Max) const;
 
     /** Create a copy of this vector, with the 2D magnitude clamped between Min and Max. Z is unchanged. */
-    TVector<T> GetClampedToSize2D(T Min, T Max) const;
+    [[nodiscard]] TVector<T> GetClampedToSize2D(T Min, T Max) const;
 
     /** Create a copy of this vector, with its maximum magnitude clamped to MaxSize. */
     TVector<T> GetClampedToMaxSize(T MaxSize) const;
 
     /** Create a copy of this vector, with the maximum 2D magnitude clamped to MaxSize. Z is unchanged. */
-    TVector<T> GetClampedToMaxSize2D(T MaxSize) const;
+    [[nodiscard]] TVector<T> GetClampedToMaxSize2D(T MaxSize) const;
 
     /**
      * Add a vector to this and clamp the result in a cube.
@@ -735,7 +735,7 @@ public:
      *
      * @return Reciprocal of this vector.
      */
-    TVector<T> Reciprocal() const;
+    [[nodiscard]] TVector<T> Reciprocal() const;
 
     /**
      * Check whether X, Y and Z are nearly equal.
@@ -743,7 +743,7 @@ public:
      * @param Tolerance Specified Tolerance.
      * @return true if X == Y == Z within the specified tolerance.
      */
-    bool IsUniform(T Tolerance=UE_KINDA_SMALL_NUMBER) const;
+    [[nodiscard]] bool IsUniform(T Tolerance=UE_KINDA_SMALL_NUMBER) const;
 
     /**
      * Mirror a vector about a normal vector.
@@ -751,7 +751,7 @@ public:
      * @param MirrorNormal Normal vector to mirror about.
      * @return Mirrored vector.
      */
-    TVector<T> MirrorByVector(const TVector<T>& MirrorNormal) const;
+    [[nodiscard]] TVector<T> MirrorByVector(const TVector<T>& MirrorNormal) const;
 
     /**
      * Mirrors a vector about a plane.
@@ -759,7 +759,7 @@ public:
      * @param Plane Plane to mirror about.
      * @return Mirrored vector.
      */
-    TVector<T> MirrorByPlane(const TPlane<T>& Plane) const;
+    [[nodiscard]] TVector<T> MirrorByPlane(const TPlane<T>& Plane) const;
 
     /**
      * Rotates around Axis (assumes Axis.Size() == 1).
@@ -768,7 +768,7 @@ public:
      * @param Axis Axis to rotate around.
      * @return Rotated Vector.
      */
-    TVector<T> RotateAngleAxis(const T AngleDeg, const TVector<T>& Axis) const;
+    [[nodiscard]] TVector<T> RotateAngleAxis(const T AngleDeg, const TVector<T>& Axis) const;
     
     /**
      * Rotates around Axis (assumes Axis.Size() == 1).
@@ -777,7 +777,7 @@ public:
      * @param Axis Axis to rotate around.
      * @return Rotated Vector.
      */
-    TVector<T> RotateAngleAxisRad(const T AngleRad, const TVector<T>& Axis) const;
+    [[nodiscard]] TVector<T> RotateAngleAxisRad(const T AngleRad, const TVector<T>& Axis) const;
 
     /**
      * Returns the cosine of the angle between this vector and another projected onto the XY plane (no Z).
@@ -785,7 +785,7 @@ public:
      * @param B the other vector to find the 2D cosine of the angle with.
      * @return The cosine.
      */
-    FORCEINLINE T CosineAngle2D(TVector<T> B) const;
+    [[nodiscard]] FORCEINLINE T CosineAngle2D(TVector<T> B) const;
 
     /**
      * Gets a copy of this vector projected onto the input vector.
@@ -793,7 +793,7 @@ public:
      * @param A	Vector to project onto, does not assume it is normalized.
      * @return Projected vector.
      */
-    FORCEINLINE TVector<T> ProjectOnTo(const TVector<T>& A) const ;
+    [[nodiscard]] FORCEINLINE TVector<T> ProjectOnTo(const TVector<T>& A) const ;
 
     /**
      * Gets a copy of this vector projected onto the input vector, which is assumed to be unit length.
@@ -801,7 +801,7 @@ public:
      * @param  Normal Vector to project onto (assumed to be unit length).
      * @return Projected vector.
      */
-    FORCEINLINE TVector<T> ProjectOnToNormal(const TVector<T>& Normal) const;
+    [[nodiscard]] FORCEINLINE TVector<T> ProjectOnToNormal(const TVector<T>& Normal) const;
 
     /**
      * Return the TRotator orientation corresponding to the direction in which the vector points.
@@ -810,7 +810,7 @@ public:
      * @return TRotator from the Vector's direction, without any roll.
      * @see ToOrientationQuat()
      */
-    CORE_API TRotator<T> ToOrientationRotator() const;
+    [[nodiscard]] CORE_API TRotator<T> ToOrientationRotator() const;
 
     /**
      * Return the Quaternion orientation corresponding to the direction in which the vector points.
@@ -822,7 +822,7 @@ public:
      * @return Quaternion from the Vector's direction, without any roll.
      * @see ToOrientationRotator(), FQuat::FindBetweenVectors()
      */
-    CORE_API TQuat<T> ToOrientationQuat() const;
+    [[nodiscard]] CORE_API TQuat<T> ToOrientationQuat() const;
 
     /**
      * Return the UE::Math::TRotator<T> orientation corresponding to the direction in which the vector points.
@@ -831,7 +831,7 @@ public:
      * @return UE::Math::TRotator<T> from the Vector's direction.
      * @see ToOrientationRotator(), ToOrientationQuat()
      */
-	FORCEINLINE UE::Math::TRotator<T> Rotation() const
+	[[nodiscard]] FORCEINLINE UE::Math::TRotator<T> Rotation() const
 	{
 		return ToOrientationRotator();
 	}
@@ -860,20 +860,20 @@ public:
      *
      * @return A string describing the vector.
      */
-    FString ToString() const;
+    [[nodiscard]] FString ToString() const;
 
     /**
     * Get a locale aware textual representation of this vector.
     *
     * @return A string describing the vector.
     */
-    FText ToText() const;
+   [[nodiscard]]  FText ToText() const;
 
     /** Get a short textural representation of this vector, for compact readable logging. */
-    FString ToCompactString() const;
+    [[nodiscard]] FString ToCompactString() const;
 
     /** Get a short locale aware textural representation of this vector, for compact readable logging. */
-    FText ToCompactText() const;
+    [[nodiscard]] FText ToCompactText() const;
 
     /**
      * Initialize this Vector based on an FString. The String is expected to contain X=, Y=, Z=.
@@ -898,14 +898,14 @@ public:
      * Converts a Cartesian unit vector into spherical coordinates on the unit sphere.
      * @return Output Theta will be in the range [0, PI], and output Phi will be in the range [-PI, PI]. 
      */
-    TVector2<T> UnitCartesianToSpherical() const;
+    [[nodiscard]] TVector2<T> UnitCartesianToSpherical() const;
 
     /**
      * Convert a direction vector into a 'heading' angle.
      *
      * @return 'Heading' angle between +/-PI. 0 is pointing down +X.
      */
-    T HeadingAngle() const;
+    [[nodiscard]] T HeadingAngle() const;
 
 	/**
 	 * Interpolate from a vector to the direction of another vector along a spherical path.
@@ -915,7 +915,7 @@ public:
 	 * @param Alpha interpolation amount, usually between 0-1
 	 * @return Vector after interpolating between Vector and Direction along a spherical path. The magnitude will remain the length of the starting vector.
 	 */
-	static CORE_API TVector<T> SlerpVectorToDirection(TVector<T>& V, TVector<T>& Direction, T Alpha);
+	[[nodiscard]] static CORE_API TVector<T> SlerpVectorToDirection(const TVector<T>& V, const TVector<T>& Direction, T Alpha);
 
 	/**
 	 * Interpolate from normalized vector A to normalized vector B along a spherical path.
@@ -925,7 +925,7 @@ public:
 	 * @param Alpha interpolation amount, usually between 0-1
 	 * @return Normalized vector after interpolating between NormalA and NormalB along a spherical path.
 	 */
-	static CORE_API TVector<T> SlerpNormals(TVector<T>& NormalA, TVector<T>& NormalB, T Alpha);
+	[[nodiscard]] static CORE_API TVector<T> SlerpNormals(const TVector<T>& NormalA, const TVector<T>& NormalB, T Alpha);
 
     /**
      * Create an orthonormal basis from a basis with at least two orthogonal vectors.
@@ -946,7 +946,7 @@ public:
      * @param Q Second vector.
      * @return Whether points are the same within a threshold. Uses fast distance approximation (linear per-component distance).
      */
-    static bool PointsAreSame(const TVector<T> &P, const TVector<T> &Q);
+    [[nodiscard]] static bool PointsAreSame(const TVector<T> &P, const TVector<T> &Q);
     
     /**
      * Compare two points and see if they're within specified distance.
@@ -956,7 +956,7 @@ public:
      * @param Dist Specified distance.
      * @return Whether two points are within the specified distance. Uses fast distance approximation (linear per-component distance).
      */
-    static bool PointsAreNear(const TVector<T> &Point1, const TVector<T> &Point2, T Dist);
+    [[nodiscard]] static bool PointsAreNear(const TVector<T> &Point1, const TVector<T> &Point2, T Dist);
 
     /**
      * Calculate the signed distance (in the direction of the normal) between a point and a plane.
@@ -966,7 +966,7 @@ public:
      * @param PlaneNormal The Normal of the plane (assumed to be unit length).
      * @return Signed distance between point and plane.
      */
-    static T PointPlaneDist(const TVector<T> &Point, const TVector<T> &PlaneBase, const TVector<T> &PlaneNormal);
+    [[nodiscard]] static T PointPlaneDist(const TVector<T> &Point, const TVector<T> &PlaneBase, const TVector<T> &PlaneNormal);
 
     /**
      * Calculate the projection of a point on the given plane.
@@ -975,7 +975,7 @@ public:
      * @param Plane The plane
      * @return Projection of Point onto Plane
      */
-	static TVector<T> PointPlaneProject(const TVector<T>& Point, const TPlane<T>& Plane);
+	[[nodiscard]] static TVector<T> PointPlaneProject(const TVector<T>& Point, const TPlane<T>& Plane);
 
     /**
      * Calculate the projection of a point on the plane defined by counter-clockwise (CCW) points A,B,C.
@@ -986,7 +986,7 @@ public:
      * @param C 3rd of three points in CCW order defining the plane
      * @return Projection of Point onto plane ABC
      */
-	static TVector<T> PointPlaneProject(const TVector<T>& Point, const TVector<T>& A, const TVector<T>& B, const TVector<T>& C);
+	[[nodiscard]] static TVector<T> PointPlaneProject(const TVector<T>& Point, const TVector<T>& A, const TVector<T>& B, const TVector<T>& C);
 
     /**
     * Calculate the projection of a point on the plane defined by PlaneBase and PlaneNormal.
@@ -996,7 +996,7 @@ public:
     * @param PlaneNorm Normal of the plane (assumed to be unit length).
     * @return Projection of Point onto plane
     */
-    static TVector<T> PointPlaneProject(const TVector<T>& Point, const TVector<T>& PlaneBase, const TVector<T>& PlaneNormal);
+    [[nodiscard]] static TVector<T> PointPlaneProject(const TVector<T>& Point, const TVector<T>& PlaneBase, const TVector<T>& PlaneNormal);
 
     /**
      * Calculate the projection of a vector on the plane defined by PlaneNormal.
@@ -1005,7 +1005,7 @@ public:
      * @param  PlaneNormal Normal of the plane (assumed to be unit length).
      * @return Projection of V onto plane.
      */
-    static TVector<T> VectorPlaneProject(const TVector<T>& V, const TVector<T>& PlaneNormal);
+    [[nodiscard]] static TVector<T> VectorPlaneProject(const TVector<T>& V, const TVector<T>& PlaneNormal);
 
     /**
      * Euclidean distance between two points.
@@ -1014,8 +1014,11 @@ public:
      * @param V2 The second point.
      * @return The distance between two points.
      */
-    static FORCEINLINE T Dist(const TVector<T> &V1, const TVector<T> &V2);
-    static FORCEINLINE T Distance(const TVector<T> &V1, const TVector<T> &V2) { return Dist(V1, V2); }
+    [[nodiscard]] static FORCEINLINE T Dist(const TVector<T> &V1, const TVector<T> &V2);
+    [[nodiscard]] static FORCEINLINE T Distance(const TVector<T> &V1, const TVector<T> &V2)
+	{
+		return Dist(V1, V2);
+	}
 
     /**
     * Euclidean distance between two points in the XY plane (ignoring Z).
@@ -1024,8 +1027,11 @@ public:
     * @param V2 The second point.
     * @return The distance between two points in the XY plane.
     */
-    static FORCEINLINE T DistXY(const TVector<T> &V1, const TVector<T> &V2);
-    static FORCEINLINE T Dist2D(const TVector<T> &V1, const TVector<T> &V2) { return DistXY(V1, V2); }
+    [[nodiscard]] static FORCEINLINE T DistXY(const TVector<T> &V1, const TVector<T> &V2);
+    [[nodiscard]] static FORCEINLINE T Dist2D(const TVector<T> &V1, const TVector<T> &V2)
+	{
+		return DistXY(V1, V2);
+	}
 
     /**
      * Squared distance between two points.
@@ -1034,7 +1040,7 @@ public:
      * @param V2 The second point.
      * @return The squared distance between two points.
      */
-    static FORCEINLINE T DistSquared(const TVector<T> &V1, const TVector<T> &V2);
+    [[nodiscard]] static FORCEINLINE T DistSquared(const TVector<T> &V1, const TVector<T> &V2);
 
     /**
      * Squared distance between two points in the XY plane only.
@@ -1043,8 +1049,11 @@ public:
      * @param V2 The second point.
      * @return The squared distance between two points in the XY plane
      */
-    static FORCEINLINE T DistSquaredXY(const TVector<T> &V1, const TVector<T> &V2);
-    static FORCEINLINE T DistSquared2D(const TVector<T> &V1, const TVector<T> &V2) { return DistSquaredXY(V1, V2); }
+    [[nodiscard]] static FORCEINLINE T DistSquaredXY(const TVector<T> &V1, const TVector<T> &V2);
+    [[nodiscard]] static FORCEINLINE T DistSquared2D(const TVector<T> &V1, const TVector<T> &V2)
+	{
+		return DistSquaredXY(V1, V2);
+	}
     
     /**
      * Compute pushout of a box from a plane.
@@ -1053,16 +1062,16 @@ public:
      * @param Size The size of the box.
      * @return Pushout required.
      */
-    static FORCEINLINE T BoxPushOut(const TVector<T>& Normal, const TVector<T>& Size);
+    [[nodiscard]] static FORCEINLINE T BoxPushOut(const TVector<T>& Normal, const TVector<T>& Size);
 
     /**
      * Min, Max, Min3, Max3 like FMath
      */
-    static FORCEINLINE TVector<T> Min( const TVector<T>& A, const TVector<T>& B );
-    static FORCEINLINE TVector<T> Max( const TVector<T>& A, const TVector<T>& B );
+    [[nodiscard]] static FORCEINLINE TVector<T> Min( const TVector<T>& A, const TVector<T>& B );
+    [[nodiscard]] static FORCEINLINE TVector<T> Max( const TVector<T>& A, const TVector<T>& B );
 
-    static FORCEINLINE TVector<T> Min3( const TVector<T>& A, const TVector<T>& B, const TVector<T>& C );
-    static FORCEINLINE TVector<T> Max3( const TVector<T>& A, const TVector<T>& B, const TVector<T>& C );
+    [[nodiscard]] static FORCEINLINE TVector<T> Min3( const TVector<T>& A, const TVector<T>& B, const TVector<T>& C );
+    [[nodiscard]] static FORCEINLINE TVector<T> Max3( const TVector<T>& A, const TVector<T>& B, const TVector<T>& C );
 
     /**
      * See if two normal vectors are nearly parallel, meaning the angle between them is close to 0 degrees.
@@ -1072,7 +1081,7 @@ public:
      * @param  ParallelCosineThreshold Normals are parallel if absolute value of dot product (cosine of angle between them) is greater than or equal to this. For example: cos(1.0 degrees).
      * @return true if vectors are nearly parallel, false otherwise.
      */
-    static bool Parallel(const TVector<T>& Normal1, const TVector<T>& Normal2, T ParallelCosineThreshold = UE_THRESH_NORMALS_ARE_PARALLEL);
+    [[nodiscard]] static bool Parallel(const TVector<T>& Normal1, const TVector<T>& Normal2, T ParallelCosineThreshold = UE_THRESH_NORMALS_ARE_PARALLEL);
 
     /**
      * See if two normal vectors are coincident (nearly parallel and point in the same direction).
@@ -1082,7 +1091,7 @@ public:
      * @param  ParallelCosineThreshold Normals are coincident if dot product (cosine of angle between them) is greater than or equal to this. For example: cos(1.0 degrees).
      * @return true if vectors are coincident (nearly parallel and point in the same direction), false otherwise.
      */
-    static bool Coincident(const TVector<T>& Normal1, const TVector<T>& Normal2, T ParallelCosineThreshold = UE_THRESH_NORMALS_ARE_PARALLEL);
+    [[nodiscard]] static bool Coincident(const TVector<T>& Normal1, const TVector<T>& Normal2, T ParallelCosineThreshold = UE_THRESH_NORMALS_ARE_PARALLEL);
 
     /**
      * See if two normal vectors are nearly orthogonal (perpendicular), meaning the angle between them is close to 90 degrees.
@@ -1092,7 +1101,7 @@ public:
      * @param  OrthogonalCosineThreshold Normals are orthogonal if absolute value of dot product (cosine of angle between them) is less than or equal to this. For example: cos(89.0 degrees).
      * @return true if vectors are orthogonal (perpendicular), false otherwise.
      */
-    static bool Orthogonal(const TVector<T>& Normal1, const TVector<T>& Normal2, T OrthogonalCosineThreshold = UE_THRESH_NORMALS_ARE_ORTHOGONAL);
+    [[nodiscard]] static bool Orthogonal(const TVector<T>& Normal1, const TVector<T>& Normal2, T OrthogonalCosineThreshold = UE_THRESH_NORMALS_ARE_ORTHOGONAL);
 
     /**
      * See if two planes are coplanar. They are coplanar if the normals are nearly parallel and the planes include the same set of points.
@@ -1104,7 +1113,7 @@ public:
      * @param ParallelCosineThreshold Normals are parallel if absolute value of dot product is greater than or equal to this.
      * @return true if the planes are coplanar, false otherwise.
      */
-    static bool Coplanar(const TVector<T>& Base1, const TVector<T>& Normal1, const TVector<T>& Base2, const TVector<T>& Normal2, T ParallelCosineThreshold = UE_THRESH_NORMALS_ARE_PARALLEL);
+    [[nodiscard]] static bool Coplanar(const TVector<T>& Base1, const TVector<T>& Normal1, const TVector<T>& Base2, const TVector<T>& Normal2, T ParallelCosineThreshold = UE_THRESH_NORMALS_ARE_PARALLEL);
 
     /**
      * Triple product of three vectors: X dot (Y cross Z).
@@ -1114,7 +1123,7 @@ public:
      * @param Z The third vector.
      * @return The triple product: X dot (Y cross Z).
      */
-    static T Triple(const TVector<T>& X, const TVector<T>& Y, const TVector<T>& Z);
+    [[nodiscard]] static T Triple(const TVector<T>& X, const TVector<T>& Y, const TVector<T>& Z);
 
     /**
      * Generates a list of sample points on a Bezier curve defined by 2 points.
@@ -1132,7 +1141,7 @@ public:
      * @param RadVector	Vector containing radian values
      * @return Vector  containing degree values
      */
-    static TVector<T> RadiansToDegrees(const TVector<T>& RadVector);
+    [[nodiscard]] static TVector<T> RadiansToDegrees(const TVector<T>& RadVector);
 
     /**
      * Converts a vector containing degree values to a vector containing radian values.
@@ -1140,7 +1149,7 @@ public:
      * @param DegVector	Vector containing degree values
      * @return Vector containing radian values
      */
-    static TVector<T> DegreesToRadians(const TVector<T>& DegVector);
+    [[nodiscard]] static TVector<T> DegreesToRadians(const TVector<T>& DegVector);
 
     /**
      * Given a current set of cluster centers, a set of points, iterate N times to move clusters to be central. 
diff --git a/Engine/Source/Runtime/Core/Public/Math/Vector2D.h b/Engine/Source/Runtime/Core/Public/Math/Vector2D.h
--- a/Engine/Source/Runtime/Core/Public/Math/Vector2D.h
+++ b/Engine/Source/Runtime/Core/Public/Math/Vector2D.h
@@ -84,7 +84,7 @@ public:
 public:
 
 	/** Default constructor (no initialization). */
-	TVector2() = default;
+	[[nodiscard]] TVector2() = default;
 
 	/**
 	* Constructor using initial values for each component.
@@ -92,14 +92,14 @@ public:
 	* @param InX X coordinate.
 	* @param InY Y coordinate.
 	*/
-	FORCEINLINE TVector2<T>(T InX, T InY);
+	[[nodiscard]] FORCEINLINE TVector2<T>(T InX, T InY);
 
 	/**
 	* Constructor initializing both components to a single T value.
 	*
 	* @param InF Value to set both components to.
 	*/
-	explicit FORCEINLINE TVector2<T>(T InF);
+	[[nodiscard]] explicit FORCEINLINE TVector2<T>(T InF);
 
 	/**
 	* Constructs a vector from an FIntPoint.
@@ -107,24 +107,24 @@ public:
 	* @param InPos Integer point used to set this vector.
 	*/
 	template <typename IntType>
-	FORCEINLINE TVector2<T>(TIntPoint<IntType> InPos);
+	[[nodiscard]] FORCEINLINE TVector2<T>(TIntPoint<IntType> InPos);
 
 	template <typename IntType>
-    explicit  TVector2<T>(TIntVector2<IntType> V);
+    [[nodiscard]] explicit  TVector2<T>(TIntVector2<IntType> V);
 
 	/**
 	* Constructor which initializes all components to zero.
 	*
 	* @param EForceInit Force init enum
 	*/
-	explicit FORCEINLINE TVector2<T>(EForceInit);
+	[[nodiscard]] explicit FORCEINLINE TVector2<T>(EForceInit);
 
 	/**
 	* Constructor that does not initialize.  More explicit than the default constructor.
 	*
 	* @param ENoInit Don't init
 	*/
-	explicit FORCEINLINE TVector2<T>(ENoInit) { }
+	[[nodiscard]] explicit FORCEINLINE TVector2<T>(ENoInit) { }
 
 	/**
 	* Constructs a vector from an FVector.
@@ -132,7 +132,7 @@ public:
 	*
 	* @param V Vector to copy from.
 	*/
-	explicit FORCEINLINE TVector2<T>(const TVector<T>& V);
+	[[nodiscard]] explicit FORCEINLINE TVector2<T>(const TVector<T>& V);
 
 	/**
 	* Constructs a vector from an FVector4.
@@ -140,7 +140,7 @@ public:
 	*
 	* @param V Vector to copy from.
 	*/
-	explicit FORCEINLINE TVector2<T>(const TVector4<T>& V);
+	[[nodiscard]] explicit FORCEINLINE TVector2<T>(const TVector4<T>& V);
 
 public:
 
@@ -150,7 +150,7 @@ public:
 	* @param V The other vector to add to this.
 	* @return The result of adding the vectors together.
 	*/
-	FORCEINLINE TVector2<T> operator+(const TVector2<T>& V) const;
+	[[nodiscard]] FORCEINLINE TVector2<T> operator+(const TVector2<T>& V) const;
 
 	/**
 	* Gets the result of subtracting a vector from this one.
@@ -158,7 +158,7 @@ public:
 	* @param V The other vector to subtract from this.
 	* @return The result of the subtraction.
 	*/
-	FORCEINLINE TVector2<T> operator-(const TVector2<T>& V) const;
+	[[nodiscard]] FORCEINLINE TVector2<T> operator-(const TVector2<T>& V) const;
 
 	/**
 	* Gets the result of scaling the vector (multiplying each component by a value).
@@ -166,7 +166,7 @@ public:
 	* @param Scale How much to scale the vector by.
 	* @return The result of scaling this vector.
 	*/
-	FORCEINLINE TVector2<T> operator*(T Scale) const;
+	[[nodiscard]] FORCEINLINE TVector2<T> operator*(T Scale) const;
 
 	/**
 	* Gets the result of dividing each component of the vector by a value.
@@ -174,7 +174,7 @@ public:
 	* @param Scale How much to divide the vector by.
 	* @return The result of division on this vector.
 	*/
-	TVector2<T> operator/(T Scale) const;
+	[[nodiscard]] TVector2<T> operator/(T Scale) const;
 
 	/**
 	* Gets the result of adding A to each component of the vector.
@@ -182,7 +182,7 @@ public:
 	* @param A T to add to each component.
 	* @return The result of adding A to each component.
 	*/
-	FORCEINLINE TVector2<T> operator+(T A) const;
+	[[nodiscard]] FORCEINLINE TVector2<T> operator+(T A) const;
 
 	/**
 	* Gets the result of subtracting A from each component of the vector.
@@ -190,7 +190,7 @@ public:
 	* @param A T to subtract from each component
 	* @return The result of subtracting A from each component.
 	*/
-	FORCEINLINE TVector2<T> operator-(T A) const;
+	[[nodiscard]] FORCEINLINE TVector2<T> operator-(T A) const;
 
 	/**
 	* Gets the result of component-wise multiplication of this vector by another.
@@ -198,7 +198,7 @@ public:
 	* @param V The other vector to multiply this by.
 	* @return The result of the multiplication.
 	*/
-	FORCEINLINE TVector2<T> operator*(const TVector2<T>& V) const;
+	[[nodiscard]] FORCEINLINE TVector2<T> operator*(const TVector2<T>& V) const;
 
 	/**
 	* Gets the result of component-wise division of this vector by another.
@@ -206,7 +206,7 @@ public:
 	* @param V The other vector to divide this by.
 	* @return The result of the division.
 	*/
-	TVector2<T> operator/(const TVector2<T>& V) const;
+	[[nodiscard]] TVector2<T> operator/(const TVector2<T>& V) const;
 
 	/**
 	* Calculates dot product of this vector and another.
@@ -214,7 +214,7 @@ public:
 	* @param V The other vector.
 	* @return The dot product.
 	*/
-	FORCEINLINE T operator|(const TVector2<T>& V) const;
+	[[nodiscard]] FORCEINLINE T operator|(const TVector2<T>& V) const;
 
 	/**
 	* Calculates cross product of this vector and another.
@@ -222,7 +222,7 @@ public:
 	* @param V The other vector.
 	* @return The cross product.
 	*/
-	FORCEINLINE T operator^(const TVector2<T>& V) const;
+	[[nodiscard]] FORCEINLINE T operator^(const TVector2<T>& V) const;
 
 public:
 
@@ -232,7 +232,7 @@ public:
 	* @param V The vector to compare against.
 	* @return true if the two vectors are equal, otherwise false.
 	*/
-	bool operator==(const TVector2<T>& V) const;
+	[[nodiscard]] bool operator==(const TVector2<T>& V) const;
 
 	/**
 	* Compares this vector against another for inequality.
@@ -240,7 +240,7 @@ public:
 	* @param V The vector to compare against.
 	* @return true if the two vectors are not equal, otherwise false.
 	*/
-	bool operator!=(const TVector2<T>& V) const;
+	[[nodiscard]] bool operator!=(const TVector2<T>& V) const;
 
 	/**
 	* Deprecated comparison operator. Use ComponentwiseAllLessThan.
@@ -249,7 +249,7 @@ public:
 	* @return true if this is the smaller vector, otherwise false.
 	*/
 	UE_DEPRECATED(5.1, "TVector2 comparison operators are deprecated. Use ComponentwiseAllLessThan. For componentwise min/max/abs/clamp, use TVector2::{Min,Max,GetAbs,Clamp}, FMath::{Min,Max,Abs,Clamp} compute something different.")
-	bool operator<(const TVector2<T>& Other) const
+	[[nodiscard]] bool operator<(const TVector2<T>& Other) const
 	{
 		return ComponentwiseAllLessThan(Other);
 	}
@@ -261,7 +261,7 @@ public:
 	* @return true if this is the larger vector, otherwise false.
 	*/
 	UE_DEPRECATED(5.1, "TVector2 comparison operators are deprecated. Use ComponentwiseAllGreaterThan. For componentwise min/max/abs/clamp, use TVector2::{Min,Max,GetAbs,Clamp}, FMath::{Min,Max,Abs,Clamp} compute something different.")
-	bool operator>(const TVector2<T>& Other) const
+	[[nodiscard]] bool operator>(const TVector2<T>& Other) const
 	{
 		return ComponentwiseAllGreaterThan(Other);
 	}
@@ -273,7 +273,7 @@ public:
 	* @return true if this vector is less than or equal to the other vector, otherwise false.
 	*/
 	UE_DEPRECATED(5.1, "TVector2 comparison operators are deprecated. Use ComponentwiseAllLessOrEqual. For componentwise min/max/abs/clamp, use TVector2::{Min,Max,GetAbs,Clamp}, FMath::{Min,Max,Abs,Clamp} compute something different.")
-	bool operator<=(const TVector2<T>& Other) const
+	[[nodiscard]] bool operator<=(const TVector2<T>& Other) const
 	{
 		return ComponentwiseAllLessOrEqual(Other);
 	}
@@ -285,7 +285,7 @@ public:
 	* @return true if this vector is greater than or equal to the other vector, otherwise false.
 	*/
 	UE_DEPRECATED(5.1, "TVector2 comparison operators are deprecated. Use ComponentwiseAllGreaterOrEqual. For componentwise min/max/abs/clamp, use TVector2::{Min,Max,GetAbs,Clamp}, FMath::{Min,Max,Abs,Clamp} compute something different.")
-	bool operator>=(const TVector2<T>& Other) const
+	[[nodiscard]] bool operator>=(const TVector2<T>& Other) const
 	{
 		return ComponentwiseAllGreaterOrEqual(Other);
 	}
@@ -296,7 +296,7 @@ public:
 	* @param Other The vector to compare against.
 	* @return true if both components of this are less than Other, otherwise false.
 	*/
-	bool ComponentwiseAllLessThan(const TVector2<T>& Other) const;
+	[[nodiscard]] bool ComponentwiseAllLessThan(const TVector2<T>& Other) const;
 
 	/**
 	* Checks whether both components of this vector are greater than another.
@@ -304,7 +304,7 @@ public:
 	* @param Other The vector to compare against.
 	* @return true if both components of this are greater than Other, otherwise false.
 	*/
-	bool ComponentwiseAllGreaterThan(const TVector2<T>& Other) const;
+	[[nodiscard]] bool ComponentwiseAllGreaterThan(const TVector2<T>& Other) const;
 
 	/**
 	* Checks whether both components of this vector are less than or equal to another.
@@ -312,7 +312,7 @@ public:
 	* @param Other The vector to compare against.
 	* @return true if both components of this are less than or equal to Other, otherwise false.
 	*/
-	bool ComponentwiseAllLessOrEqual(const TVector2<T>& Other) const;
+	[[nodiscard]] bool ComponentwiseAllLessOrEqual(const TVector2<T>& Other) const;
 
 	/**
 	* Checks whether both components of this vector are greater than or equal to another.
@@ -320,7 +320,7 @@ public:
 	* @param Other The vector to compare against.
 	* @return true if both components of this are greater than or equal to Other, otherwise false.
 	*/
-	bool ComponentwiseAllGreaterOrEqual(const TVector2<T>& Other) const;
+	[[nodiscard]] bool ComponentwiseAllGreaterOrEqual(const TVector2<T>& Other) const;
 
 
 	/**
@@ -328,7 +328,7 @@ public:
 	*
 	* @return A negated copy of the vector.
 	*/
-	FORCEINLINE TVector2<T> operator-() const;
+	[[nodiscard]] FORCEINLINE TVector2<T> operator-() const;
 
 	/**
 	* Adds another vector to this.
@@ -384,7 +384,7 @@ public:
 	* @param Index the index of vector component
 	* @return reference to component.
 	*/
-	T& operator[](int32 Index);
+	[[nodiscard]] T& operator[](int32 Index);
 
 	/**
 	* Gets specific component of the vector.
@@ -392,7 +392,7 @@ public:
 	* @param Index the index of vector component
 	* @return copy of component value.
 	*/
-	T operator[](int32 Index) const;
+	[[nodiscard]] T operator[](int32 Index) const;
 
 	/**
 	* Gets a specific component of the vector.
@@ -400,7 +400,7 @@ public:
 	* @param Index The index of the component required.
 	* @return Reference to the specified component.
 	*/
-	T& Component(int32 Index);
+	[[nodiscard]] T& Component(int32 Index);
 
 	/**
 	* Gets a specific component of the vector.
@@ -408,7 +408,7 @@ public:
 	* @param Index The index of the component required.
 	* @return Copy of the specified component.
 	*/
-	T Component(int32 Index) const;
+	[[nodiscard]] T Component(int32 Index) const;
 
 	/**
 	 * Tests if index is valid, i.e. greater than or equal to zero, and less than the number of components in the vector.
@@ -416,7 +416,7 @@ public:
 	 * @param Index Index to test.
 	 * @returns True if index is valid. False otherwise.
 	 */
-	bool IsValidIndex(int32 Index) const;
+	[[nodiscard]] bool IsValidIndex(int32 Index) const;
 
 public:
 
@@ -427,7 +427,7 @@ public:
 	* @param B The second vector.
 	* @return The dot product.
 	*/
-	FORCEINLINE static T DotProduct(const TVector2<T>& A, const TVector2<T>& B);
+	[[nodiscard]] FORCEINLINE static T DotProduct(const TVector2<T>& A, const TVector2<T>& B);
 
 	/**
 	* Squared distance between two 2D points.
@@ -436,7 +436,7 @@ public:
 	* @param V2 The second point.
 	* @return The squared distance between two 2D points.
 	*/
-	FORCEINLINE static T DistSquared(const TVector2<T>& V1, const TVector2<T>& V2);
+	[[nodiscard]] FORCEINLINE static T DistSquared(const TVector2<T>& V1, const TVector2<T>& V2);
 
 	/**
 	* Distance between two 2D points.
@@ -445,7 +445,7 @@ public:
 	* @param V2 The second point.
 	* @return The distance between two 2D points.
 	*/
-	FORCEINLINE static T Distance(const TVector2<T>& V1, const TVector2<T>& V2);
+	[[nodiscard]] FORCEINLINE static T Distance(const TVector2<T>& V1, const TVector2<T>& V2);
 
 	/**
 	* Calculate the cross product of two vectors.
@@ -454,7 +454,7 @@ public:
 	* @param B The second vector.
 	* @return The cross product.
 	*/
-	FORCEINLINE static T CrossProduct(const TVector2<T>& A, const TVector2<T>& B);
+	[[nodiscard]] FORCEINLINE static T CrossProduct(const TVector2<T>& A, const TVector2<T>& B);
 
 	/**
 	* Returns a vector with the maximum component for each dimension from the pair of vectors.
@@ -463,7 +463,7 @@ public:
 	* @param B The second vector.
 	* @return The max vector.
 	*/
-	FORCEINLINE static TVector2<T> Max(const TVector2<T>& A, const TVector2<T>& B);
+	[[nodiscard]] FORCEINLINE static TVector2<T> Max(const TVector2<T>& A, const TVector2<T>& B);
 
 	/**
 	* Returns a vector with the minimum component for each dimension from the pair of vectors.
@@ -472,7 +472,7 @@ public:
 	* @param B The second vector.
 	* @return The min vector.
 	*/
-	FORCEINLINE static TVector2<T> Min(const TVector2<T>& A, const TVector2<T>& B);
+	[[nodiscard]] FORCEINLINE static TVector2<T> Min(const TVector2<T>& A, const TVector2<T>& B);
 
 	/**
 	* Returns a vector with each component clamped between a minimum and a maximum.
@@ -482,7 +482,7 @@ public:
 	* @param MaxValue The maximum vector.
 	* @return The clamped vector.
 	*/
-	FORCEINLINE static TVector2<T> Clamp(const TVector2<T>& V, const TVector2<T>& MinValue, const TVector2<T>& MaxValue);
+	[[nodiscard]] FORCEINLINE static TVector2<T> Clamp(const TVector2<T>& V, const TVector2<T>& MinValue, const TVector2<T>& MaxValue);
 
 	/**
 	* Checks for equality with error-tolerant comparison.
@@ -491,7 +491,7 @@ public:
 	* @param Tolerance Error tolerance.
 	* @return true if the vectors are equal within specified tolerance, otherwise false.
 	*/
-	bool Equals(const TVector2<T>& V, T Tolerance=UE_KINDA_SMALL_NUMBER) const;
+	[[nodiscard]] bool Equals(const TVector2<T>& V, T Tolerance=UE_KINDA_SMALL_NUMBER) const;
 
 	/**
 	* Set the values of the vector directly.
@@ -506,27 +506,27 @@ public:
 	*
 	* @return The maximum value of the vector's components.
 	*/
-	T GetMax() const;
+	[[nodiscard]] T GetMax() const;
 
 	/**
 	* Get the maximum absolute value of the vector's components.
 	*
 	* @return The maximum absolute value of the vector's components.
 	*/
-	T GetAbsMax() const;
+	[[nodiscard]] T GetAbsMax() const;
 
 	/**
 	* Get the minimum value of the vector's components.
 	*
 	* @return The minimum value of the vector's components.
 	*/
-	T GetMin() const;
+	[[nodiscard]] T GetMin() const;
 
 	/** Gets the component-wise min of two vectors. */
-	TVector2<T> ComponentMin(const TVector2<T>& Other) const;
+	[[nodiscard]] TVector2<T> ComponentMin(const TVector2<T>& Other) const;
 
 	/** Gets the component-wise max of two vectors. */
-	TVector2<T> ComponentMax(const TVector2<T>& Other) const;
+	[[nodiscard]] TVector2<T> ComponentMax(const TVector2<T>& Other) const;
 
 	/**
 	* Get the length (magnitude) of this vector.
@@ -534,7 +534,7 @@ public:
 	* @return The length of this vector.
     * @see Length - This function is a synonym for Length()
 	*/
-	T Size() const;
+	[[nodiscard]] T Size() const;
 
 	/**
     * Get the length (magnitude) of this vector.
@@ -542,7 +542,10 @@ public:
     * @return The length of this vector.
     * @see Size - This function is a synonym for Size()
     */
-	FORCEINLINE T Length() const { return Size(); }
+	[[nodiscard]] FORCEINLINE T Length() const
+	{
+		return Size();
+	}
 
 	/**
 	* Get the squared length of this vector.
@@ -550,7 +553,7 @@ public:
 	* @return The squared length of this vector.
     * @see LengthSquared - This function is a synonym for LengthSquared()
 	*/
-	T SizeSquared() const;
+	[[nodiscard]] T SizeSquared() const;
 	
 	/**
 	* Get the squared length of this vector.
@@ -558,7 +561,10 @@ public:
 	* @return The squared length of this vector.
     * @see SizeSquared - This function is a synonym for SizeSquared()
 	*/
-	FORCEINLINE T SquaredLength() const { return SizeSquared(); }
+	[[nodiscard]] FORCEINLINE T SquaredLength() const
+	{
+		return SizeSquared();
+	}
 
 	/**
 	* Get the dot product of this vector against another.
@@ -566,7 +572,7 @@ public:
 	* @param V2 The vector to measure dot product against.
 	* @return The dot product.
 	*/
-	T Dot(const TVector2<T>& V2) const;
+	[[nodiscard]] T Dot(const TVector2<T>& V2) const;
 
 	/**
 	* Rotates around axis (0,0,1)
@@ -574,7 +580,7 @@ public:
 	* @param AngleDeg Angle to rotate (in degrees)
 	* @return Rotated Vector
 	*/
-	TVector2<T> GetRotated(T AngleDeg) const;
+	[[nodiscard]] TVector2<T> GetRotated(T AngleDeg) const;
 
 	/**
 	* Gets a normalized copy of the vector, checking it is safe to do so based on the length.
@@ -583,7 +589,7 @@ public:
 	* @param Tolerance Minimum squared length of vector for normalization.
 	* @return A normalized copy of the vector if safe, (0,0) otherwise.
 	*/
-	TVector2<T> GetSafeNormal(T Tolerance=UE_SMALL_NUMBER) const;
+	[[nodiscard]] TVector2<T> GetSafeNormal(T Tolerance=UE_SMALL_NUMBER) const;
 
 	/**
 	* Normalize this vector in-place if it is large enough, set it to (0,0) otherwise.
@@ -601,7 +607,7 @@ public:
 	* @param Tolerance Error tolerance.
 	* @return true if vector is in tolerance to zero, otherwise false.
 	*/
-	bool IsNearlyZero(T Tolerance=UE_KINDA_SMALL_NUMBER) const;
+	[[nodiscard]] bool IsNearlyZero(T Tolerance=UE_KINDA_SMALL_NUMBER) const;
 
 	/**
 	* Util to convert this vector into a unit direction vector and its original length.
@@ -617,27 +623,27 @@ public:
 	*
 	* @return true if vector is exactly zero, otherwise false.
 	*/
-	bool IsZero() const;
+	[[nodiscard]] bool IsZero() const;
 
 	/**
 	* Get this vector as an Int Point.
 	*
 	* @return New Int Point from this vector.
 	*/
-	FIntPoint IntPoint() const;
+	[[nodiscard]] FIntPoint IntPoint() const;
 
 	/**
 	* Get this vector as a vector where each component has been rounded to the nearest int.
 	*
 	* @return New TVector2<T> from this vector that is rounded.
 	*/
-	TVector2<T> RoundToVector() const;
+	[[nodiscard]] TVector2<T> RoundToVector() const;
 
 	/**
 	* Creates a copy of this vector with both axes clamped to the given range.
 	* @return New vector with clamped axes.
 	*/
-	TVector2<T> ClampAxes(T MinAxisVal, T MaxAxisVal) const;
+	[[nodiscard]] TVector2<T> ClampAxes(T MinAxisVal, T MaxAxisVal) const;
 
 	/**
 	* Get a copy of the vector as sign only.
@@ -645,21 +651,21 @@ public:
 	*
 	* @param A copy of the vector with each component set to +1 or -1
 	*/
-	FORCEINLINE TVector2<T> GetSignVector() const;
+	[[nodiscard]] FORCEINLINE TVector2<T> GetSignVector() const;
 
 	/**
 	* Get a copy of this vector with absolute value of each component.
 	*
 	* @return A copy of this vector with absolute value of each component.
 	*/
-	FORCEINLINE TVector2<T> GetAbs() const;
+	[[nodiscard]] FORCEINLINE TVector2<T> GetAbs() const;
 
 	/**
 	* Get a textual representation of the vector.
 	*
 	* @return Text describing the vector.
 	*/
-	FString ToString() const;
+	[[nodiscard]] FString ToString() const;
 
 	/**
 	* Initialize this Vector based on an FString. The String is expected to contain X=, Y=.
@@ -702,7 +708,7 @@ public:
 	*
 	* @return true if there are any non-finite values in this vector, false otherwise.
 	*/
-	FORCEINLINE bool ContainsNaN() const
+	[[nodiscard]] FORCEINLINE bool ContainsNaN() const
 	{
 		return (!FMath::IsFinite(X) || 
 				!FMath::IsFinite(Y));
@@ -733,7 +739,7 @@ public:
 	}
 
 	/** Converts spherical coordinates on the unit sphere into a Cartesian unit length vector. */
-	inline TVector<T> SphericalToUnitCartesian() const;
+	[[nodiscard]] inline TVector<T> SphericalToUnitCartesian() const;
 	
 	
 	
diff --git a/Engine/Source/Runtime/Core/Public/Math/Vector2DHalf.h b/Engine/Source/Runtime/Core/Public/Math/Vector2DHalf.h
--- a/Engine/Source/Runtime/Core/Public/Math/Vector2DHalf.h
+++ b/Engine/Source/Runtime/Core/Public/Math/Vector2DHalf.h
@@ -21,7 +21,7 @@ struct FVector2DHalf
 public:
 
 	/** Default Constructor (no initialization). */
-	FORCEINLINE FVector2DHalf() { }
+	[[nodiscard]] FORCEINLINE FVector2DHalf() { }
 
 	/**
 	 * Constructor.
@@ -29,21 +29,21 @@ public:
 	 * InX half float X value
 	 * Iny half float Y value
 	 */
- 	FORCEINLINE FVector2DHalf( const FFloat16& InX,const FFloat16& InY );
+ 	[[nodiscard]] FORCEINLINE FVector2DHalf( const FFloat16& InX,const FFloat16& InY );
 
 	/** Constructor 
 	 *
 	 * InX float X value
 	 * Iny float Y value
 	 */
-	FORCEINLINE FVector2DHalf( float InX,float InY );
+	[[nodiscard]] FORCEINLINE FVector2DHalf( float InX,float InY );
 
 	/** Constructor 
 	 *
 	 * Vector2D float vector
 	 */
-	FORCEINLINE FVector2DHalf( const FVector2f& Vector2D );
-	FORCEINLINE FVector2DHalf( const FVector2d& Vector2D );// LWC_TODO: This should probably be explicit
+	[[nodiscard]] FORCEINLINE FVector2DHalf( const FVector2f& Vector2D );
+	[[nodiscard]] FORCEINLINE FVector2DHalf( const FVector2d& Vector2D );// LWC_TODO: This should probably be explicit
 
 public:
 
@@ -64,7 +64,7 @@ public:
 	void SetTruncate( const FVector2f& Vector2D );
 	void SetTruncate( const FVector2d& Vector2D );
 
-	FORCEINLINE uint32 AsUInt32()
+	[[nodiscard]] FORCEINLINE uint32 AsUInt32() const
 	{
 		uint32 PackedData = (X.Encoded & 0xFFFFFFFF) | (Y.Encoded << 16);
 		return PackedData;
@@ -77,7 +77,7 @@ public:
 	 *
 	 * @return Text describing the vector.
 	 */
-	FString ToString() const;
+	[[nodiscard]] FString ToString() const;
 
 public:
 
diff --git a/Engine/Source/Runtime/Core/Public/Math/Vector4.h b/Engine/Source/Runtime/Core/Public/Math/Vector4.h
--- a/Engine/Source/Runtime/Core/Public/Math/Vector4.h
+++ b/Engine/Source/Runtime/Core/Public/Math/Vector4.h
@@ -66,7 +66,7 @@ public:
 	 *
 	 * @param InVector 3D Vector to set first three components.
 	 */
-	TVector4(const UE::Math::TVector<T>& InVector);
+	[[nodiscard]] TVector4(const UE::Math::TVector<T>& InVector);
 
 	/**
 	 * Constructor.
@@ -75,7 +75,7 @@ public:
 	 * @param InW W Coordinate.
 	 */
 	template<typename FArg UE_REQUIRES(std::is_arithmetic_v<FArg>)>
-	FORCEINLINE TVector4(const UE::Math::TVector<T>& InVector, FArg InW)
+	[[nodiscard]] FORCEINLINE TVector4(const UE::Math::TVector<T>& InVector, FArg InW)
 		: X(InVector.X)
 		, Y(InVector.Y)
 		, Z(InVector.Z)
@@ -91,7 +91,7 @@ public:
 	 */
 
 	template<typename FArg UE_REQUIRES(std::is_arithmetic_v<FArg>)>
-	FORCEINLINE TVector4(const UE::Math::TVector4<T>& InVector, FArg OverrideW)
+	[[nodiscard]] FORCEINLINE TVector4(const UE::Math::TVector4<T>& InVector, FArg OverrideW)
 		: X(InVector.X)
 		, Y(InVector.Y)
 		, Z(InVector.Z)
@@ -105,7 +105,7 @@ public:
 	 *
 	 * @param InColour Color used to set vector.
 	 */
-	FORCEINLINE TVector4(const FLinearColor& InColor)
+	[[nodiscard]] FORCEINLINE TVector4(const FLinearColor& InColor)
 		: X(InColor.R)
 		, Y(InColor.G)
 		, Z(InColor.B)
@@ -120,7 +120,7 @@ public:
 	 * @param InColour Color used to set XYZ.
 	 * @param InOverrideW
 	 */
-	FORCEINLINE TVector4(const FLinearColor& InColor, T InOverrideW)
+	[[nodiscard]] FORCEINLINE TVector4(const FLinearColor& InColor, T InOverrideW)
 		: X(InColor.R)
 		, Y(InColor.G)
 		, Z(InColor.B)
@@ -141,7 +141,7 @@ public:
 	 *       actually initializes the instance.  Ideally it should be = default;
 	 *       in the same way, but this would break backwards compatibility.
 	 */
-	explicit TVector4(T InX = 0.0f, T InY = 0.0f, T InZ = 0.0f, T InW = 1.0f);
+	[[nodiscard]] explicit TVector4(T InX = 0.0f, T InY = 0.0f, T InZ = 0.0f, T InW = 1.0f);
 
 	/**
 	 * Creates and initializes a new vector from the specified 2D vectors.
@@ -149,7 +149,7 @@ public:
 	 * @param InXY A 2D vector holding the X- and Y-components.
 	 * @param InZW A 2D vector holding the Z- and W-components.
 	 */
-	explicit TVector4(TVector2<T> InXY, TVector2<T> InZW);
+	[[nodiscard]] explicit TVector4(TVector2<T> InXY, TVector2<T> InZW);
 
 	/**
 	 * Creates and initializes a new vector from an int vector value.
@@ -157,32 +157,32 @@ public:
 	 * @param InVector IntVector used to set vector.
 	 */
 	template<typename IntType>
-	TVector4(const TIntVector4<IntType>& InVector);
+	[[nodiscard]] TVector4(const TIntVector4<IntType>& InVector);
 
 	/**
 	 * Creates and initializes a new vector to zero.
 	 *
 	 * @param EForceInit Force Init Enum.
 	 */
-	explicit TVector4(EForceInit);
+	[[nodiscard]] explicit TVector4(EForceInit);
 
 	/**
 	 * Creates an uninitialized new vector.
 	 *
 	 * @param ENoInit Force uninitialized enum.
 	 */
-	TVector4(ENoInit);
+	[[nodiscard]] TVector4(ENoInit);
 
 public:
 
 	// To satisfy UE::Geometry type
-	static TVector4<T> Zero()
+	[[nodiscard]] static TVector4<T> Zero()
 	{
 		return TVector4(T(0), T(0), T(0), T(0));
 	}
 
 	// To satisfy UE::Geometry type
-	static TVector4<T> One()
+	[[nodiscard]] static TVector4<T> One()
 	{
 		return TVector4(T(1), T(1), T(1), T(1));
 	}
@@ -193,7 +193,7 @@ public:
 	 * @param ComponentIndex The index of the component.
 	 * @return Reference to the desired component.
 	 */
-	FORCEINLINE T& operator[](int32 ComponentIndex);
+	[[nodiscard]] FORCEINLINE T& operator[](int32 ComponentIndex);
 
 	/**
 	 * Access a specific component of the vector.
@@ -201,7 +201,7 @@ public:
 	 * @param ComponentIndex The index of the component.
 	 * @return Copy of the desired component.
 	 */
-	FORCEINLINE T operator[](int32 ComponentIndex) const;
+	[[nodiscard]] FORCEINLINE T operator[](int32 ComponentIndex) const;
 
 	// Unary operators.
 
@@ -210,7 +210,7 @@ public:
 	 *
 	 * @return A negated copy of the vector.
 	 */
-	FORCEINLINE TVector4<T> operator-() const;
+	[[nodiscard]] FORCEINLINE TVector4<T> operator-() const;
 
 	/**
 	 * Gets the result of adding a vector to this.
@@ -218,7 +218,7 @@ public:
 	 * @param V The vector to add.
 	 * @return The result of vector addition.
 	 */
-	FORCEINLINE TVector4<T> operator+(const TVector4<T>& V) const;
+	[[nodiscard]] FORCEINLINE TVector4<T> operator+(const TVector4<T>& V) const;
 
 	/**
 	 * Adds another vector to this one.
@@ -234,7 +234,7 @@ public:
 	 * @param V The vector to subtract.
 	 * @return The result of vector subtraction.
 	 */
-	FORCEINLINE TVector4<T> operator-(const TVector4<T>& V) const;
+	[[nodiscard]] FORCEINLINE TVector4<T> operator-(const TVector4<T>& V) const;
 
 	/**
 	 * Subtracts another vector to this one.
@@ -251,7 +251,7 @@ public:
      * @return The result of addition.
      */
 	template<typename FArg UE_REQUIRES(std::is_arithmetic_v<FArg>)>
-	FORCEINLINE TVector4<T> operator+(FArg Bias) const
+	[[nodiscard]] FORCEINLINE TVector4<T> operator+(FArg Bias) const
 	{
 		return TVector4<T>(X + (T)Bias, Y + (T)Bias, Z + (T)Bias, W + (T)Bias);
 	}
@@ -263,7 +263,7 @@ public:
      * @return The result of subtraction.
      */
 	template<typename FArg UE_REQUIRES(std::is_arithmetic_v<FArg>)>
-	FORCEINLINE TVector4<T> operator-(FArg Bias) const
+	[[nodiscard]] FORCEINLINE TVector4<T> operator-(FArg Bias) const
 	{
 		return TVector4<T>(X - (T)Bias, Y - (T)Bias, Z - (T)Bias, W - (T)Bias);
 	}
@@ -275,7 +275,7 @@ public:
 	 * @return The result of vector scaling.
 	 */
 	template<typename FArg UE_REQUIRES(std::is_arithmetic_v<FArg>)>
-	FORCEINLINE TVector4<T> operator*(FArg Scale) const
+	[[nodiscard]] FORCEINLINE TVector4<T> operator*(FArg Scale) const
 	{
 		return TVector4(X * Scale, Y * Scale, Z * Scale, W * Scale);
 	}
@@ -287,7 +287,7 @@ public:
 	 * @return The result of division.
 	 */
 	template<typename FArg UE_REQUIRES(std::is_arithmetic_v<FArg>)>
-	FORCEINLINE TVector4<T> operator/(FArg Scale) const
+	[[nodiscard]] FORCEINLINE TVector4<T> operator/(FArg Scale) const
 	{
 		const T RScale = T(1.0f) / Scale;
 		return TVector4(X * RScale, Y * RScale, Z * RScale, W * RScale);
@@ -299,7 +299,7 @@ public:
 	 * @param V What to divide by.
 	 * @return The result of division.
 	 */
-	TVector4<T> operator/(const TVector4<T>& V) const;
+	[[nodiscard]] TVector4<T> operator/(const TVector4<T>& V) const;
 
 	/**
 	 * Gets the result of multiplying a vector with this.
@@ -307,7 +307,7 @@ public:
 	 * @param V The vector to multiply.
 	 * @return The result of vector multiplication.
 	 */
-	TVector4<T> operator*(const TVector4<T>& V) const;
+	[[nodiscard]] TVector4<T> operator*(const TVector4<T>& V) const;
 
 	/**
 	 * Gets the result of multiplying a vector with another Vector (component wise).
@@ -360,7 +360,7 @@ public:
 	 * @param V The other vector.
 	 * @return true if the two vectors are the same, otherwise false.
 	 */
-	bool operator==(const TVector4<T>& V) const;
+	[[nodiscard]] bool operator==(const TVector4<T>& V) const;
 
 	/**
 	 * Checks for inequality against another vector.
@@ -368,7 +368,7 @@ public:
 	 * @param V The other vector.
 	 * @return true if the two vectors are different, otherwise false.
 	 */
-	bool operator!=(const TVector4<T>& V) const;
+	[[nodiscard]] bool operator!=(const TVector4<T>& V) const;
 
 	/**
 	 * Calculate Cross product between this and another vector.
@@ -376,7 +376,7 @@ public:
 	 * @param V The other vector.
 	 * @return The Cross product.
 	 */
-	TVector4<T> operator^(const TVector4<T>& V) const;
+	[[nodiscard]] TVector4<T> operator^(const TVector4<T>& V) const;
 
 public:
 
@@ -388,7 +388,7 @@ public:
 	 * @param Index The index of the component.
 	 * @return Reference to the component.
 	 */
-	T& Component(int32 Index);
+	[[nodiscard]] T& Component(int32 Index);
 
 	/**
 	* Gets a specific component of the vector.
@@ -396,7 +396,7 @@ public:
 	* @param Index The index of the component.
 	* @return Reference to the component.
 	*/
-	const T& Component(int32 Index) const;
+	[[nodiscard]] const T& Component(int32 Index) const;
 
 	/**
 	 * Tests if index is valid, i.e. greater than or equal to zero, and less than the number of components in the vector.
@@ -404,7 +404,7 @@ public:
 	 * @param Index Index to test.
 	 * @returns True if index is valid. False otherwise.
 	 */
-	bool IsValidIndex(int32 Index) const;
+	[[nodiscard]] bool IsValidIndex(int32 Index) const;
 
 	/**
 	 * Error tolerant comparison.
@@ -413,7 +413,7 @@ public:
 	 * @param Tolerance Error Tolerance.
 	 * @return true if the two vectors are equal within specified tolerance, otherwise false.
 	 */
-	bool Equals(const TVector4<T>& V, T Tolerance=UE_KINDA_SMALL_NUMBER) const;
+	[[nodiscard]] bool Equals(const TVector4<T>& V, T Tolerance=UE_KINDA_SMALL_NUMBER) const;
 
 	/**
 	 * Check if the vector is of unit length, with specified tolerance.
@@ -421,14 +421,14 @@ public:
 	 * @param LengthSquaredTolerance Tolerance against squared length.
 	 * @return true if the vector is a unit vector within the specified tolerance.
 	 */
-	bool IsUnit3(T LengthSquaredTolerance = UE_KINDA_SMALL_NUMBER) const;
+	[[nodiscard]] bool IsUnit3(T LengthSquaredTolerance = UE_KINDA_SMALL_NUMBER) const;
 
 	/**
 	 * Get a textual representation of the vector.
 	 *
 	 * @return Text describing the vector.
 	 */
-	FString ToString() const;
+	[[nodiscard]] FString ToString() const;
 
 	/**
 	 * Initialize this Vector based on an FString. The String is expected to contain X=, Y=, Z=, W=.
@@ -445,27 +445,27 @@ public:
 	 * @param Tolerance Minimum squared length of vector for normalization.
 	 * @return A normalized copy of the vector or a zero vector.
 	 */
-	FORCEINLINE TVector4 GetSafeNormal(T Tolerance = UE_SMALL_NUMBER) const;
+	[[nodiscard]] FORCEINLINE TVector4 GetSafeNormal(T Tolerance = UE_SMALL_NUMBER) const;
 
 	/**
 	 * Calculates normalized version of vector without checking if it is non-zero.
 	 *
 	 * @return Normalized version of vector.
 	 */
-	FORCEINLINE TVector4 GetUnsafeNormal3() const;
+	[[nodiscard]] FORCEINLINE TVector4 GetUnsafeNormal3() const;
 
 	/**
 	 * Return the FRotator orientation corresponding to the direction in which the vector points.
 	 * Sets Yaw and Pitch to the proper numbers, and sets roll to zero because the roll can't be determined from a vector.
 	 * @return FRotator from the Vector's direction.
 	 */
-	CORE_API TRotator<T> ToOrientationRotator() const;
+	[[nodiscard]] CORE_API TRotator<T> ToOrientationRotator() const;
 
 	/**
 	 * Return the Quaternion orientation corresponding to the direction in which the vector points.
 	 * @return Quaternion from the Vector's direction.
 	 */
-	CORE_API TQuat<T> ToOrientationQuat() const;
+	[[nodiscard]] CORE_API TQuat<T> ToOrientationQuat() const;
 
 	/**
 	 * Return the FRotator orientation corresponding to the direction in which the vector points.
@@ -474,7 +474,7 @@ public:
 	 * @return FRotator from the Vector's direction.
 	 * @see ToOrientationRotator()
 	 */
-	FORCEINLINE TRotator<T> Rotation() const
+	[[nodiscard]] FORCEINLINE TRotator<T> Rotation() const
 	{
 		return ToOrientationRotator();
 	}
@@ -490,47 +490,47 @@ public:
 	FORCEINLINE void Set(T InX, T InY, T InZ, T InW);
 
 	/** Gets the component-wise min of two vectors. */
-	TVector4<T> ComponentMin(const TVector4<T>& Other) const;
+	[[nodiscard]] TVector4<T> ComponentMin(const TVector4<T>& Other) const;
 
 	/** Gets the component-wise max of two vectors. */
-	TVector4<T> ComponentMax(const TVector4<T>& Other) const;
+	[[nodiscard]] TVector4<T> ComponentMax(const TVector4<T>& Other) const;
 
 	/**
 	 * Get the length of this vector not taking W component into account.
 	 *
 	 * @return The length of this vector.
 	 */
-	T Size3() const;
+	[[nodiscard]] T Size3() const;
 
 	/**
 	 * Get the squared length of this vector not taking W component into account.
 	 *
 	 * @return The squared length of this vector.
 	 */
-	T SizeSquared3() const;
+	[[nodiscard]] T SizeSquared3() const;
 
 	/**
 	 * Get the length (magnitude) of this vector, taking the W component into account
 	 *
 	 * @return The length of this vector
 	 */
-	T Size() const;
+	[[nodiscard]] T Size() const;
 
 	/**
 	 * Get the squared length of this vector, taking the W component into account
 	 *
 	 * @return The squared length of this vector
 	 */
-	T SizeSquared() const;
+	[[nodiscard]] T SizeSquared() const;
 
 	/** Utility to check if there are any non-finite values (NaN or Inf) in this vector. */
-	bool ContainsNaN() const;
+	[[nodiscard]] bool ContainsNaN() const;
 
 	/** Utility to check if all of the components of this vector are nearly zero given the tolerance. */
-	bool IsNearlyZero3(T Tolerance = UE_KINDA_SMALL_NUMBER) const;
+	[[nodiscard]] bool IsNearlyZero3(T Tolerance = UE_KINDA_SMALL_NUMBER) const;
 
 	/** Reflect vector. */
-	TVector4<T> Reflect3(const TVector4<T>& Normal) const;
+	[[nodiscard]] TVector4<T> Reflect3(const TVector4<T>& Normal) const;
 
 	/**
 	 * Find good arbitrary axis vectors to represent U and V axes of a plane,
