Define the same UnrealMath aliases for NEON as we do for SSE/FPU.

Modified version of Github #12847.

VectorRegisterConstInit changes to match UnrealMathSSE and un-break non-unity builds. (Latent error surfaced by this change.)
Same with UnrealMathFPU including <cmath> and the #if in UnrealMathDirectX, that's to make the headers compile
on their own in build tests.

diff --git a/Engine/Source/Runtime/Core/Public/Math/UnrealMathDirectX.h b/Engine/Source/Runtime/Core/Public/Math/UnrealMathDirectX.h
--- a/Engine/Source/Runtime/Core/Public/Math/UnrealMathDirectX.h
+++ b/Engine/Source/Runtime/Core/Public/Math/UnrealMathDirectX.h
@@ -7,6 +7,8 @@
 // enabled by bWithDirectXMath
 // LWC_TODO: No longer supported. Needs removing.
 
+#if WITH_DIRECTXMATH
+
 #if defined __cplusplus_cli && !PLATFORM_HOLOLENS
 // there are compile issues with this file in managed mode, so use the FPU version
 #include "Math/UnrealMathFPU.h"
@@ -23,6 +25,8 @@
  *	Helpers:
  *============================================================================*/
 
+struct VectorRegisterConstInit {};
+
 /**
  *	float4 vector register type, where the first float (X) is stored in the lowest 32 bits, and so on.
  */
@@ -1267,5 +1271,5 @@ FORCEINLINE VectorRegister4Int VectorRoundToIntHalfToEven(const VectorRegister4F
 
 #endif
 
-
+#endif
 
diff --git a/Engine/Source/Runtime/Core/Public/Math/UnrealMathFPU.h b/Engine/Source/Runtime/Core/Public/Math/UnrealMathFPU.h
--- a/Engine/Source/Runtime/Core/Public/Math/UnrealMathFPU.h
+++ b/Engine/Source/Runtime/Core/Public/Math/UnrealMathFPU.h
@@ -5,12 +5,14 @@
 // HEADER_UNIT_SKIP - Not included directly
 
 #include "HAL/UnrealMemory.h"
-
+#include <cmath>
 
 /*=============================================================================
  *	Helpers:
  *============================================================================*/
 
+struct VectorRegisterConstInit {};
+
  /**
  *	float4 vector register type, where the first float (X) is stored in the lowest 32 bits, and so on.
  */
@@ -87,9 +89,9 @@ typedef VectorRegister4Int VectorRegister4i;
 typedef VectorRegister4Float VectorRegister4f;
 typedef VectorRegister4Double VectorRegister4d;
 typedef VectorRegister2Double VectorRegister2d;
-typedef VectorRegister4Double VectorRegister4;
 
 // Backwards compatibility
+typedef VectorRegister4Double VectorRegister4;
 typedef VectorRegister4 VectorRegister;
 typedef VectorRegister4Int VectorRegisterInt;
 
diff --git a/Engine/Source/Runtime/Core/Public/Math/UnrealMathNeon.h b/Engine/Source/Runtime/Core/Public/Math/UnrealMathNeon.h
--- a/Engine/Source/Runtime/Core/Public/Math/UnrealMathNeon.h
+++ b/Engine/Source/Runtime/Core/Public/Math/UnrealMathNeon.h
@@ -10,6 +10,8 @@ PRAGMA_DISABLE_SHADOW_VARIABLE_WARNINGS
 
 #if PLATFORM_ENABLE_VECTORINTRINSICS_NEON
 
+struct VectorRegisterConstInit {};
+
 // Include the intrinsic functions header
 #if ((PLATFORM_WINDOWS || PLATFORM_HOLOLENS) && PLATFORM_64BITS)
 #include <arm64_neon.h>
@@ -214,12 +216,18 @@ struct alignas(alignof(VectorRegister4Double)) AlignedDouble4
 };
 
 typedef AlignedDouble4 AlignedRegister4;
+
 // Aliases
 typedef VectorRegister4Int VectorRegister4i;
 typedef VectorRegister4Float VectorRegister4f;
 typedef VectorRegister4Double VectorRegister4d;
 typedef VectorRegister2Double VectorRegister2d;
 
+// Backwards compatibility
+typedef VectorRegister4Double VectorRegister4;
+typedef VectorRegister4 VectorRegister;
+typedef VectorRegister4Int VectorRegisterInt;
+
 /**
  * Returns a bitwise equivalent vector based on 4 uint32s.
  *
@@ -2902,4 +2910,5 @@ PRAGMA_ENABLE_SHADOW_VARIABLE_WARNINGS
 
 #if UE_ENABLE_INCLUDE_ORDER_DEPRECATED_IN_5_4
 #include <type_traits>
-#endif
\ No newline at end of file
+#endif
+
diff --git a/Engine/Source/Runtime/Core/Public/Math/VectorRegister.h b/Engine/Source/Runtime/Core/Public/Math/VectorRegister.h
--- a/Engine/Source/Runtime/Core/Public/Math/VectorRegister.h
+++ b/Engine/Source/Runtime/Core/Public/Math/VectorRegister.h
@@ -10,18 +10,14 @@
 
 // Platform specific vector intrinsics include.
 #if WITH_DIRECTXMATH
-struct VectorRegisterConstInit {};
 #include "Math/UnrealMathDirectX.h"
 #elif PLATFORM_ENABLE_VECTORINTRINSICS_NEON
-struct VectorRegisterConstInit {};
 #include "Math/UnrealMathNeon.h"
 #elif defined(__cplusplus_cli) && !PLATFORM_HOLOLENS
-struct VectorRegisterConstInit {};
 #include "Math/UnrealMathFPU.h" // there are compile issues with UnrealMathSSE in managed mode, so use the FPU version
 #elif PLATFORM_ENABLE_VECTORINTRINSICS
 #include "Math/UnrealMathSSE.h"
 #else
-struct VectorRegisterConstInit {};
 #include "Math/UnrealMathFPU.h"
 #endif
 
@@ -138,4 +134,4 @@ struct FScopedFTZFloatMode
 
 private:
 	uint32 ControlRegisterState;
-};
\ No newline at end of file
+};
