Moved I/O dispatcher CVars to separate header file.

diff --git a/Engine/Source/Runtime/Core/Internal/IO/IoDispatcherConfig.cpp b/Engine/Source/Runtime/Core/Internal/IO/IoDispatcherConfig.cpp
new file mode 100644
--- /dev/null
+++ b/Engine/Source/Runtime/Core/Internal/IO/IoDispatcherConfig.cpp
@@ -0,0 +1,103 @@
+// Copyright Epic Games, Inc. All Rights Reserved.
+
+#include "IO/IoDispatcherConfig.h"
+#include "HAL/IConsoleManager.h"
+
+int32 GIoDispatcherBufferSizeKB = 256;
+static FAutoConsoleVariableRef CVar_IoDispatcherBufferSizeKB(
+	TEXT("s.IoDispatcherBufferSizeKB"),
+	GIoDispatcherBufferSizeKB,
+	TEXT("IoDispatcher read buffer size (in kilobytes).")
+);
+
+int32 GIoDispatcherBufferAlignment = 4096;
+static FAutoConsoleVariableRef CVar_IoDispatcherBufferAlignment(
+	TEXT("s.IoDispatcherBufferAlignment"),
+	GIoDispatcherBufferAlignment,
+	TEXT("IoDispatcher read buffer alignment.")
+);
+
+int32 GIoDispatcherBufferMemoryMB = 8;
+static FAutoConsoleVariableRef CVar_IoDispatcherBufferMemoryMB(
+	TEXT("s.IoDispatcherBufferMemoryMB"),
+	GIoDispatcherBufferMemoryMB,
+	TEXT("IoDispatcher buffer memory size (in megabytes).")
+);
+
+int32 GIoDispatcherDecompressionWorkerCount = 4;
+static FAutoConsoleVariableRef CVar_IoDispatcherDecompressionWorkerCount(
+	TEXT("s.IoDispatcherDecompressionWorkerCount"),
+	GIoDispatcherDecompressionWorkerCount,
+	TEXT("IoDispatcher decompression worker count.")
+);
+
+int32 GIoDispatcherCacheSizeMB = 0;
+static FAutoConsoleVariableRef CVar_IoDispatcherCacheSizeMB(
+	TEXT("s.IoDispatcherCacheSizeMB"),
+	GIoDispatcherCacheSizeMB,
+	TEXT("IoDispatcher cache memory size (in megabytes).")
+);
+
+int32 GIoDispatcherSortRequestsByOffset = 1;
+static FAutoConsoleVariableRef CVar_IoDispatcherSortRequestsByOffset(
+	TEXT("s.IoDispatcherSortRequestsByOffset"),
+	GIoDispatcherSortRequestsByOffset,
+	TEXT("If > 0, io dispatcher sorts the outstanding request queue by offset rather than sequence.")
+);
+
+int32 GIoDispatcherMaintainSortingOnPriorityChange = 1;
+static FAutoConsoleVariableRef CVar_IoDispatcherMaintainSortingOnPriorityChange(
+	TEXT("s.IoDispatcherMaintainSortingOnPriorityChange"),
+	GIoDispatcherMaintainSortingOnPriorityChange,
+	TEXT("If s.IoDispatcherSortRequestsByOffset > 0 and this > 0, io dispatcher remembers the last file handle/offset read from even when switching priority levels.")
+);
+
+int32 GIoDispatcherMaxForwardSeekKB = 0;
+static FAutoConsoleVariableRef CVar_IoDispatcherMaxForwardSeekKB(
+	TEXT("s.IoDispatcherMaxForwardSeekKB"),
+	GIoDispatcherMaxForwardSeekKB,
+	TEXT("If s.IoDispatcherSortRequestsByOffset is enabled and this is > 0, if the next sequential read is further than this offset from the last one, read the oldest request instead")
+);
+
+int32 GIoDispatcherRequestLatencyCircuitBreakerMS = 0;
+static FAutoConsoleVariableRef CVar_IoDispatcherRequestLatencyCircuitBreakerMS(
+	TEXT("s.IoDispatcherRequestLatencyCircuitBreakerMS"),
+	GIoDispatcherRequestLatencyCircuitBreakerMS,
+	TEXT("If s.IoDispatcherSortRequestsByOffset is enabled and this is >0, if the oldest request has been in the queue for this long, read it instead of the most optimal read")
+);
+
+int32 GIoDispatcherTocsEnablePerfectHashing = 1;
+static FAutoConsoleVariableRef CVar_IoDispatcherTocsEnablePerfectHashing(
+	TEXT("s.IoDispatcherTocsEnablePerfectHashing"),
+	GIoDispatcherTocsEnablePerfectHashing,
+	TEXT("Enable perfect hashmap lookups for iostore tocs")
+);
+
+int32 GIoDispatcherCanDecompressOnStarvation = 1;
+static FAutoConsoleVariableRef CVar_IoDispatcherCanDecompressOnStarvation(
+	TEXT("s.IoDispatcherCanDecompressOnStarvation"),
+	GIoDispatcherCanDecompressOnStarvation,
+	TEXT("IoDispatcher thread will help with decompression tasks when all worker threads are IO starved to avoid deadlocks on low core count")
+);
+
+int32 GIoDispatcherForceSynchronousScatter = 0;
+static FAutoConsoleVariableRef CVar_IoDispatcherForceSynchronousScatter(
+	TEXT("s.IoDispatcherForceSynchronousScatter"),
+	GIoDispatcherForceSynchronousScatter,
+	TEXT("Force scatter jobs to be synchronous on the IODispatcher thread.\n")
+	TEXT("This can avoid deadlocks in cases where background tasks end up waiting on I/O and we don't have enough background task threads to fulfill decompression requests.")
+);
+
+int32 GIoDispatcherDecompressOnIdle = 0;
+static FAutoConsoleVariableRef CVar_GIoDispatcherDecompressOnIdle(
+	TEXT("s.IoDispatcherDecompressOnIdle"),
+	GIoDispatcherDecompressOnIdle,
+	TEXT("When enabled the I/O dispatcher thread will help out decompress chunk blocks when idle.")
+);
+
+int32 GIoDispatcherMaxConsecutiveDecompressionJobs = 2;
+static FAutoConsoleVariableRef CVar_IoDispatcherMaxConsecutiveDecompressionJobs(
+	TEXT("s.IoDispatcherMaxConsecutiveDecompressionJobs"),
+	GIoDispatcherMaxConsecutiveDecompressionJobs,
+	TEXT("Max consecutive decompression jobs before re-launching tasks.")
+);
diff --git a/Engine/Source/Runtime/Core/Internal/IO/IoDispatcherConfig.h b/Engine/Source/Runtime/Core/Internal/IO/IoDispatcherConfig.h
new file mode 100644
--- /dev/null
+++ b/Engine/Source/Runtime/Core/Internal/IO/IoDispatcherConfig.h
@@ -0,0 +1,20 @@
+// Copyright Epic Games, Inc. All Rights Reserved.
+
+#pragma once
+
+#include "CoreTypes.h"
+
+CORE_API extern int32 GIoDispatcherBufferSizeKB;
+CORE_API extern int32 GIoDispatcherBufferAlignment;
+CORE_API extern int32 GIoDispatcherBufferMemoryMB;
+CORE_API extern int32 GIoDispatcherDecompressionWorkerCount;
+CORE_API extern int32 GIoDispatcherCacheSizeMB;
+CORE_API extern int32 GIoDispatcherSortRequestsByOffset;
+CORE_API extern int32 GIoDispatcherMaintainSortingOnPriorityChange;
+CORE_API extern int32 GIoDispatcherMaxForwardSeekKB;
+CORE_API extern int32 GIoDispatcherRequestLatencyCircuitBreakerMS;
+CORE_API extern int32 GIoDispatcherTocsEnablePerfectHashing;
+CORE_API extern int32 GIoDispatcherCanDecompressOnStarvation;
+CORE_API extern int32 GIoDispatcherForceSynchronousScatter;
+CORE_API extern int32 GIoDispatcherDecompressOnIdle;
+CORE_API extern int32 GIoDispatcherMaxConsecutiveDecompressionJobs;
diff --git a/Engine/Source/Runtime/PakFile/Private/IoDispatcherFileBackend.cpp b/Engine/Source/Runtime/PakFile/Private/IoDispatcherFileBackend.cpp
--- a/Engine/Source/Runtime/PakFile/Private/IoDispatcherFileBackend.cpp
+++ b/Engine/Source/Runtime/PakFile/Private/IoDispatcherFileBackend.cpp
@@ -24,97 +24,12 @@
 #include "Containers/Ticker.h"
 #include "Modules/ModuleManager.h"
 #include "IO/IoContainerHeader.h"
+#include "IO/IoDispatcherConfig.h"
 #include "Serialization/MemoryReader.h"
 #include "FileCache/FileCache.h"
 
 //PRAGMA_DISABLE_OPTIMIZATION
 
-int32 GIoDispatcherBufferSizeKB = 256;
-static FAutoConsoleVariableRef CVar_IoDispatcherBufferSizeKB(
-	TEXT("s.IoDispatcherBufferSizeKB"),
-	GIoDispatcherBufferSizeKB,
-	TEXT("IoDispatcher read buffer size (in kilobytes).")
-);
-
-int32 GIoDispatcherBufferAlignment = 4096;
-static FAutoConsoleVariableRef CVar_IoDispatcherBufferAlignment(
-	TEXT("s.IoDispatcherBufferAlignment"),
-	GIoDispatcherBufferAlignment,
-	TEXT("IoDispatcher read buffer alignment.")
-);
-
-int32 GIoDispatcherBufferMemoryMB = 8;
-static FAutoConsoleVariableRef CVar_IoDispatcherBufferMemoryMB(
-	TEXT("s.IoDispatcherBufferMemoryMB"),
-	GIoDispatcherBufferMemoryMB,
-	TEXT("IoDispatcher buffer memory size (in megabytes).")
-);
-
-int32 GIoDispatcherDecompressionWorkerCount = 4;
-static FAutoConsoleVariableRef CVar_IoDispatcherDecompressionWorkerCount(
-	TEXT("s.IoDispatcherDecompressionWorkerCount"),
-	GIoDispatcherDecompressionWorkerCount,
-	TEXT("IoDispatcher decompression worker count.")
-);
-
-int32 GIoDispatcherCacheSizeMB = 0;
-static FAutoConsoleVariableRef CVar_IoDispatcherCacheSizeMB(
-	TEXT("s.IoDispatcherCacheSizeMB"),
-	GIoDispatcherCacheSizeMB,
-	TEXT("IoDispatcher cache memory size (in megabytes).")
-);
-
-int32 GIoDispatcherSortRequestsByOffset = 1;
-static FAutoConsoleVariableRef CVar_IoDispatcherSortRequestsByOffset(
-	TEXT("s.IoDispatcherSortRequestsByOffset"),
-	GIoDispatcherSortRequestsByOffset,
-	TEXT("If > 0, io dispatcher sorts the outstanding request queue by offset rather than sequence.")
-);
-
-int32 GIoDispatcherMaintainSortingOnPriorityChange = 1;
-static FAutoConsoleVariableRef CVar_IoDispatcherMaintainSortingOnPriorityChange(
-	TEXT("s.IoDispatcherMaintainSortingOnPriorityChange"),
-	GIoDispatcherMaintainSortingOnPriorityChange,
-	TEXT("If s.IoDispatcherSortRequestsByOffset > 0 and this > 0, io dispatcher remembers the last file handle/offset read from even when switching priority levels.")
-);
-
-int32 GIoDispatcherMaxForwardSeekKB = 0;
-static FAutoConsoleVariableRef CVar_IoDispatcherMaxForwardSeekKB(
-	TEXT("s.IoDispatcherMaxForwardSeekKB"),
-	GIoDispatcherMaxForwardSeekKB,
-	TEXT("If s.IoDispatcherSortRequestsByOffset is enabled and this is > 0, if the next sequential read is further than this offset from the last one, read the oldest request instead")
-);
-
-int32 GIoDispatcherRequestLatencyCircuitBreakerMS = 0;
-static FAutoConsoleVariableRef CVar_IoDispatcherRequestLatencyCircuitBreakerMS(
-	TEXT("s.IoDispatcherRequestLatencyCircuitBreakerMS"),
-	GIoDispatcherRequestLatencyCircuitBreakerMS,
-	TEXT("If s.IoDispatcherSortRequestsByOffset is enabled and this is >0, if the oldest request has been in the queue for this long, read it instead of the most optimal read")
-);
-
-int32 GIoDispatcherTocsEnablePerfectHashing = 1;
-static FAutoConsoleVariableRef CVar_IoDispatcherTocsEnablePerfectHashing(
-	TEXT("s.IoDispatcherTocsEnablePerfectHashing"),
-	GIoDispatcherTocsEnablePerfectHashing,
-	TEXT("Enable perfect hashmap lookups for iostore tocs")
-);
-
-int32 GIoDispatcherCanDecompressOnStarvation = 1;
-static FAutoConsoleVariableRef CVar_IoDispatcherCanDecompressOnStarvation(
-	TEXT("s.IoDispatcherCanDecompressOnStarvation"),
-	GIoDispatcherCanDecompressOnStarvation,
-	TEXT("IoDispatcher thread will help with decompression tasks when all worker threads are IO starved to avoid deadlocks on low core count")
-);
-
-int32 GIoDispatcherForceSynchronousScatter = 0;
-static FAutoConsoleVariableRef CVar_IoDispatcherForceSynchronousScatter(
-	TEXT("s.IoDispatcherForceSynchronousScatter"),
-	GIoDispatcherForceSynchronousScatter,
-	TEXT("Force scatter jobs to be synchronous on the IODispatcher thread.\n")
-	TEXT("This can avoid deadlocks in cases where background tasks end up waiting on I/O and we don't have enough background task threads to fulfill decompression requests.")
-);
-
-
 uint32 FFileIoStoreReadRequest::NextSequence = 0;
 #if CHECK_IO_STORE_READ_REQUEST_LIST_MEMBERSHIP
 uint32 FFileIoStoreReadRequestList::NextListCookie = 0;
