Split the LoadFromCompactBinary overloads for FString into FWideString and FUtf8String overloads.

diff --git a/Engine/Source/Runtime/Core/Private/Serialization/CompactBinarySerialization.cpp b/Engine/Source/Runtime/Core/Private/Serialization/CompactBinarySerialization.cpp
--- a/Engine/Source/Runtime/Core/Private/Serialization/CompactBinarySerialization.cpp
+++ b/Engine/Source/Runtime/Core/Private/Serialization/CompactBinarySerialization.cpp
@@ -6,6 +6,7 @@
 #include "Containers/ContainerAllocationPolicies.h"
 #include "Containers/SharedString.h"
 #include "Containers/StringView.h"
+#include "Containers/Utf8String.h"
 #include "Containers/UnrealString.h"
 #include "CoreGlobals.h"
 #include "HAL/PlatformString.h"
@@ -320,9 +321,15 @@ bool LoadFromCompactBinary(FCbFieldView Field, UE::FWideSharedString& OutString)
 	return false;
 }
 
-bool LoadFromCompactBinary(FCbFieldView Field, FString& OutValue)
+bool LoadFromCompactBinary(FCbFieldView Field, FUtf8String& OutValue)
 {
-	OutValue = FString(Field.AsString());
+	OutValue = Field.AsString();
+	return !Field.HasError();
+}
+
+bool LoadFromCompactBinary(FCbFieldView Field, FWideString& OutValue)
+{
+	OutValue = FWideString(Field.AsString());
 	return !Field.HasError();
 }
 
diff --git a/Engine/Source/Runtime/Core/Public/Serialization/CompactBinarySerialization.h b/Engine/Source/Runtime/Core/Public/Serialization/CompactBinarySerialization.h
--- a/Engine/Source/Runtime/Core/Public/Serialization/CompactBinarySerialization.h
+++ b/Engine/Source/Runtime/Core/Public/Serialization/CompactBinarySerialization.h
@@ -107,7 +107,8 @@ CORE_API bool LoadFromCompactBinary(FCbFieldView Field, FUtf8StringBuilderBase&
 CORE_API bool LoadFromCompactBinary(FCbFieldView Field, FWideStringBuilderBase& OutValue);
 CORE_API bool LoadFromCompactBinary(FCbFieldView Field, UE::FUtf8SharedString& OutString);
 CORE_API bool LoadFromCompactBinary(FCbFieldView Field, UE::FWideSharedString& OutString);
-CORE_API bool LoadFromCompactBinary(FCbFieldView Field, FString& OutValue);
+CORE_API bool LoadFromCompactBinary(FCbFieldView Field, FUtf8String& OutValue);
+CORE_API bool LoadFromCompactBinary(FCbFieldView Field, FWideString& OutValue);
 CORE_API bool LoadFromCompactBinary(FCbFieldView Field, FName& OutValue);
 
 inline bool LoadFromCompactBinary(FCbFieldView Field, int8& OutValue, const int8 Default = 0)
