[AutoRTFM] Fix nullptr dereference when calling autortfm_record_open_write() before autortfm_initialize()

This can happen with the editor.

diff --git a/Engine/Source/Runtime/Core/Private/AutoRTFM/API.cpp b/Engine/Source/Runtime/Core/Private/AutoRTFM/API.cpp
--- a/Engine/Source/Runtime/Core/Private/AutoRTFM/API.cpp
+++ b/Engine/Source/Runtime/Core/Private/AutoRTFM/API.cpp
@@ -866,9 +866,9 @@ extern "C" void autortfm_record_open_write(void* Ptr, size_t Size)
 	{
 		return RTFM_autortfm_record_open_write_err(Ptr, Size);
 	}
-	else if (FContext::Get()->IsTransactional())
+	else if (FContext* Context = FContext::Get(); Context && Context->IsTransactional())
 	{
-		if (FTransaction* CurrentTransaction = FContext::Get()->GetCurrentTransaction())
+		if (FTransaction* CurrentTransaction = Context->GetCurrentTransaction())
 		{
 			CurrentTransaction->RecordWrite(Ptr, Size);
 		}
@@ -881,9 +881,9 @@ extern "C" void autortfm_record_open_write_no_memory_validation(void* Ptr, size_
 	{
 		return RTFM_autortfm_record_open_write_err(Ptr, Size);
 	}
-	else if (FContext::Get()->IsTransactional())
+	else if (FContext* Context = FContext::Get(); Context && Context->IsTransactional())
 	{
-		if (FTransaction* CurrentTransaction = FContext::Get()->GetCurrentTransaction())
+		if (FTransaction* CurrentTransaction = Context->GetCurrentTransaction())
 		{
 			CurrentTransaction->RecordWrite(Ptr, Size, /* bNoMemoryValidation */ true);
 		}
