[AutoRTFM] Allow UE::FTransactionallySafeMutex to compile safely in C++17.

A handful of modules in UE, such as ACLPlugin, are still limited to C++17 support. Trying to use the transactionally-safe mutex from these modules would trigger a warning `default member initializer for bit-field is a C++20 extension [-Werror,-Wc++20-extensions]`.

diff --git a/Engine/Source/Runtime/Core/Public/Async/TransactionallySafeMutex.h b/Engine/Source/Runtime/Core/Public/Async/TransactionallySafeMutex.h
--- a/Engine/Source/Runtime/Core/Public/Async/TransactionallySafeMutex.h
+++ b/Engine/Source/Runtime/Core/Public/Async/TransactionallySafeMutex.h
@@ -43,7 +43,7 @@ struct FTransactionallySafeMutex final
 	{
 	}
 
-	// Construct in a locked state. Avoids an expensive compare-and-swap at creation time.
+	// Construct in a locked state.
 	UE_AUTORTFM_ALWAYS_OPEN
 	inline explicit FTransactionallySafeMutex(UE::FAcquireLock Tag) : State(MakeShared<FState>(Tag))
 	{
@@ -209,10 +209,10 @@ private:
 	struct FState final
 	{
 		UE_AUTORTFM_NOAUTORTFM
-		FState() = default;
+		FState() : bTransactionallyLocked(false), bRegisteredCommitHandler(false) {}
 
 		UE_AUTORTFM_NOAUTORTFM
-		explicit FState(UE::FAcquireLock Tag) : Mutex(Tag) {}
+		explicit FState(UE::FAcquireLock Tag) : Mutex(Tag), bTransactionallyLocked(false), bRegisteredCommitHandler(false) {}
 
 		UE_AUTORTFM_NOAUTORTFM
 		~FState()
@@ -263,9 +263,9 @@ private:
 		// Other values indicate a double-lock or double-unlock within a transaction.
 		int8 TransactionalLockCount = 0;
 		// True if the underlying mutex has been locked for the duration of the transaction.
-		bool bTransactionallyLocked : 1 = false;
+		bool bTransactionallyLocked : 1;
 		// True if the on-commit handler has been registered for this transaction nest.
-		bool bRegisteredCommitHandler : 1 = false;
+		bool bRegisteredCommitHandler : 1;
 	};
 
 	TSharedPtr<FState> State;
