Added method for returning whether an I/O dispatcher request failed nor not.

#rnx

diff --git a/Engine/Source/Runtime/Core/Public/IO/IoDispatcherBackend.h b/Engine/Source/Runtime/Core/Public/IO/IoDispatcherBackend.h
--- a/Engine/Source/Runtime/Core/Public/IO/IoDispatcherBackend.h
+++ b/Engine/Source/Runtime/Core/Public/IO/IoDispatcherBackend.h
@@ -45,6 +45,12 @@ public:
 		return bCancelled;
 	}
 
+	/** Returns whether the request failed. */
+	bool IsFailed() const
+	{
+		return bFailed;
+	}
+
 	/** Mark the request as failed. */
 	void SetFailed()
 	{
