Enable PLATFORM_HAS_MULTITHREADED_PREMAIN on Android.

We already have this enabled on IOS.

This is done to prevent the following from happening:
- LLM is being enabled by default (to trace any early activity)
- LLM is disabled and freed after command-line parsing while other thread is inside of the LLMs code, leading to a crash (e.g. AndroidEventThreadWorker vs main thread).

#rnx

diff --git a/Engine/Source/Runtime/Core/Public/Android/AndroidPlatform.h b/Engine/Source/Runtime/Core/Public/Android/AndroidPlatform.h
--- a/Engine/Source/Runtime/Core/Public/Android/AndroidPlatform.h
+++ b/Engine/Source/Runtime/Core/Public/Android/AndroidPlatform.h
@@ -51,6 +51,7 @@ typedef FAndroidTypes FPlatformTypes;
 #define PLATFORM_HAS_BSD_SOCKET_FEATURE_IOCTL		1
 #define PLATFORM_HAS_BSD_SOCKET_FEATURE_MSG_DONTWAIT	1
 #define PLATFORM_HAS_TOUCH_MAIN_SCREEN				1
+#define PLATFORM_HAS_MULTITHREADED_PREMAIN			1
 #define PLATFORM_SUPPORTS_STACK_SYMBOLS				1
 #define PLATFORM_NUM_AUDIODECOMPRESSION_PRECACHE_BUFFERS 2
 #define PLATFORM_USES_ANSI_STRING_FOR_EXTERNAL_PROFILING 1
