Making sure that any existing game and engine data get read back in by the CRC app when loading temp crash contexts.
Reworking the copy shared crash context function to use a string builder instead of strncat to avoid deprecation issues and fix the logic used.

diff --git a/Engine/Source/Programs/CrashReportClient/Private/CrashReportClientApp.cpp b/Engine/Source/Programs/CrashReportClient/Private/CrashReportClientApp.cpp
--- a/Engine/Source/Programs/CrashReportClient/Private/CrashReportClientApp.cpp
+++ b/Engine/Source/Programs/CrashReportClient/Private/CrashReportClientApp.cpp
@@ -998,6 +998,92 @@ static bool LoadTempCrashContextFromFile(FSharedCrashContext& CrashContext, uint
 	FindAndParseValue(ContextProperties, TEXT("SendUsageData"), UserSettings.bSendUsageData);
 	FindAndCopyValue(ContextProperties, TEXT("LogFilePath"), UserSettings.LogFilePath);
 
+	// Dynamic data
+
+	// When encoding the plugins list and engine and game data key/value pairs into the dynamic data segment
+	// we use 1 and 2 to denote delimiter and equals respectively. This is necessary since the values could 
+	// contain any characters normally used for delimiting.
+	const TCHAR* CR_PAIR_DELIM = TEXT("\x01");
+	const TCHAR* CR_PAIR_EQ = TEXT("\x02");
+
+	FStringBuilderBase DynamicDataBuilder;
+	DynamicDataBuilder.Reserve(CR_MAX_DYNAMIC_BUFFER_CHARS);
+
+	// Enabled Plugins
+	CrashContext.EnabledPluginsOffset = (uint32)DynamicDataBuilder.Len();
+	if (const FXmlNode* EnabledPluginsNode = File.GetRootNode()->FindChildNode(TEXT("EnabledPlugins")))
+	{
+		const TArray<FXmlNode*>& EnabledPlugins = EnabledPluginsNode->GetChildrenNodes();
+		CrashContext.EnabledPluginsNum = EnabledPlugins.Num();
+		for (const FXmlNode* Node : EnabledPlugins)
+		{
+			DynamicDataBuilder.Append(Node->GetContent());
+			DynamicDataBuilder.Append(CR_PAIR_DELIM);
+		}
+	}
+	else
+	{
+		CrashContext.EnabledPluginsNum = 0;
+	}
+	
+	// Engine Data
+	CrashContext.EngineDataOffset = (uint32)DynamicDataBuilder.Len();
+	if (const FXmlNode* EngineDataNode = File.GetRootNode()->FindChildNode(TEXT("EngineData")))
+	{
+		const TArray<FXmlNode*>& EngineDataEntries = EngineDataNode->GetChildrenNodes();
+		CrashContext.EngineDataNum = EngineDataEntries.Num();
+		for (const FXmlNode* Node : EngineDataEntries)
+		{
+			DynamicDataBuilder.Append(Node->GetTag());
+			DynamicDataBuilder.Append(CR_PAIR_EQ);
+			DynamicDataBuilder.Append(Node->GetContent());
+			DynamicDataBuilder.Append(CR_PAIR_DELIM);
+		}
+	}
+	else
+	{
+		CrashContext.EngineDataNum = 0;
+	}
+	
+	// Game Data
+	CrashContext.GameDataOffset = (uint32)DynamicDataBuilder.Len();
+	if (const FXmlNode* GameDataNode = File.GetRootNode()->FindChildNode(TEXT("GameData")))
+	{
+		const TArray<FXmlNode*>& GameDataEntries = GameDataNode->GetChildrenNodes();
+		CrashContext.GameDataNum = GameDataEntries.Num();
+		for (const FXmlNode* Node : GameDataEntries)
+		{
+			DynamicDataBuilder.Append(Node->GetTag());
+			DynamicDataBuilder.Append(CR_PAIR_EQ);
+			DynamicDataBuilder.Append(Node->GetContent());
+			DynamicDataBuilder.Append(CR_PAIR_DELIM);
+		}
+	}
+	else
+	{
+		CrashContext.GameDataNum = 0;
+	}
+	
+	// File Paths
+	CrashContext.FilePathsOffset = (uint32)DynamicDataBuilder.Len();
+	if (const FXmlNode* FilePathsNode = File.GetRootNode()->FindChildNode(TEXT("FilePaths")))
+	{
+		const TArray<FXmlNode*>& FilePaths = FilePathsNode->GetChildrenNodes();
+		CrashContext.FilePathsNum = FilePaths.Num();
+		for (const FXmlNode* Node : FilePaths)
+		{
+			DynamicDataBuilder.Append(Node->GetContent());
+			DynamicDataBuilder.Append(CR_PAIR_DELIM);
+		}
+	}
+	else
+	{
+		CrashContext.FilePathsNum = 0;
+	}
+
+	FCString::Strncpy(CrashContext.DynamicData, DynamicDataBuilder.ToString(), CR_MAX_DYNAMIC_BUFFER_CHARS - 1);
+	CrashContext.DynamicData[CR_MAX_DYNAMIC_BUFFER_CHARS - 1] = 0;
+
 	return true;
 }
 
diff --git a/Engine/Source/Runtime/Core/Private/GenericPlatform/GenericPlatformCrashContext.cpp b/Engine/Source/Runtime/Core/Private/GenericPlatform/GenericPlatformCrashContext.cpp
--- a/Engine/Source/Runtime/Core/Private/GenericPlatform/GenericPlatformCrashContext.cpp
+++ b/Engine/Source/Runtime/Core/Private/GenericPlatform/GenericPlatformCrashContext.cpp
@@ -621,59 +621,52 @@ void FGenericCrashContext::CopySharedCrashContext(FSharedCrashContext& Dst)
 	//Copy the session
 	FMemory::Memcpy(Dst.SessionContext, NCached::Session);
 	FMemory::Memcpy(Dst.UserSettings, NCached::UserSettings);
-	FMemory::Memset(Dst.DynamicData, 0);
 
-	TCHAR* DynamicDataStart = &Dst.DynamicData[0];
-	TCHAR* DynamicDataPtr = DynamicDataStart;
-
-	// -1 to allow space for null terminator
-	#define CR_DYNAMIC_BUFFER_REMAIN uint32((CR_MAX_DYNAMIC_BUFFER_CHARS) - (DynamicDataPtr-DynamicDataStart) - 1)
+	FStringBuilderBase DynamicDataBuilder;
+	DynamicDataBuilder.Reserve(CR_MAX_DYNAMIC_BUFFER_CHARS);
 
+	Dst.EnabledPluginsOffset = (uint32)DynamicDataBuilder.Len();
 #if UE_CRASH_REPORTER_WITH_LOADED_PLUGINS
-	Dst.EnabledPluginsOffset = (uint32)(DynamicDataPtr - DynamicDataStart);
 	Dst.EnabledPluginsNum = NCached::EnabledPluginsList.Num();
 	for (const FString& Plugin : NCached::EnabledPluginsList)
 	{
-		FCString::StrncatTruncateDest(DynamicDataPtr, CR_DYNAMIC_BUFFER_REMAIN, *Plugin);
-		FCString::StrncatTruncateDest(DynamicDataPtr, CR_DYNAMIC_BUFFER_REMAIN, CR_PAIR_DELIM);
+		DynamicDataBuilder.Append(Plugin);
+		DynamicDataBuilder.Append(CR_PAIR_DELIM);
 	}
-	DynamicDataPtr += FCString::Strlen(DynamicDataPtr) + 1;
 #else
-	Dst.EnabledPluginsOffset = (uint32)(DynamicDataPtr - DynamicDataStart);
 	Dst.EnabledPluginsNum = 0;
 #endif
 
-	Dst.EngineDataOffset = (uint32)(DynamicDataPtr - DynamicDataStart);
+	Dst.EngineDataOffset = (uint32)DynamicDataBuilder.Len();
 	Dst.EngineDataNum = NCached::EngineData.Num();
 	for (const TPair<FString, FString>& Pair : NCached::EngineData)
 	{
-		FCString::StrncatTruncateDest(DynamicDataPtr, CR_DYNAMIC_BUFFER_REMAIN, *Pair.Key);
-		FCString::StrncatTruncateDest(DynamicDataPtr, CR_DYNAMIC_BUFFER_REMAIN, CR_PAIR_EQ);
-		FCString::StrncatTruncateDest(DynamicDataPtr, CR_DYNAMIC_BUFFER_REMAIN, *Pair.Value);
-		FCString::StrncatTruncateDest(DynamicDataPtr, CR_DYNAMIC_BUFFER_REMAIN, CR_PAIR_DELIM);
+		DynamicDataBuilder.Append(Pair.Key);
+		DynamicDataBuilder.Append(CR_PAIR_EQ);
+		DynamicDataBuilder.Append(Pair.Value);
+		DynamicDataBuilder.Append(CR_PAIR_DELIM);
 	}
-	DynamicDataPtr += FCString::Strlen(DynamicDataPtr) + 1;
 
-	Dst.GameDataOffset = (uint32)(DynamicDataPtr - DynamicDataStart);
+	Dst.GameDataOffset = (uint32)DynamicDataBuilder.Len();
 	Dst.GameDataNum = NCached::GameData.Num();
 	for (const TPair<FString, FString>& Pair : NCached::GameData)
 	{
-		FCString::StrncatTruncateDest(DynamicDataPtr, CR_DYNAMIC_BUFFER_REMAIN, *Pair.Key);
-		FCString::StrncatTruncateDest(DynamicDataPtr, CR_DYNAMIC_BUFFER_REMAIN, CR_PAIR_EQ);
-		FCString::StrncatTruncateDest(DynamicDataPtr, CR_DYNAMIC_BUFFER_REMAIN, *Pair.Value);
-		FCString::StrncatTruncateDest(DynamicDataPtr, CR_DYNAMIC_BUFFER_REMAIN, CR_PAIR_DELIM);
+		DynamicDataBuilder.Append(Pair.Key);
+		DynamicDataBuilder.Append(CR_PAIR_EQ);
+		DynamicDataBuilder.Append(Pair.Value);
+		DynamicDataBuilder.Append(CR_PAIR_DELIM);
 	}
-	DynamicDataPtr += FCString::Strlen(DynamicDataPtr) + 1;
 
-	Dst.FilePathsOffset = (uint32)(DynamicDataPtr - DynamicDataStart);
+	Dst.FilePathsOffset = (uint32)DynamicDataBuilder.Len();
 	Dst.FilePathsNum = NCached::FilePaths.Num();
 	for (const FString& FilePath : NCached::FilePaths)
 	{
-		FCString::StrncatTruncateDest(DynamicDataPtr, CR_DYNAMIC_BUFFER_REMAIN, *FilePath);
-		FCString::StrncatTruncateDest(DynamicDataPtr, CR_DYNAMIC_BUFFER_REMAIN, CR_PAIR_DELIM);
+		DynamicDataBuilder.Append(FilePath);
+		DynamicDataBuilder.Append(CR_PAIR_DELIM);
 	}
 
-	#undef CR_DYNAMIC_BUFFER_REMAIN
+	FCString::Strncpy(Dst.DynamicData, DynamicDataBuilder.ToString(), CR_MAX_DYNAMIC_BUFFER_CHARS - 1);
+	Dst.DynamicData[CR_MAX_DYNAMIC_BUFFER_CHARS - 1] = 0;
 
 	// Copy GPU breadcrumbs.
 	if (NCached::GPUBreadcrumbs.IsSet())
