Removed ScopeLock from FGenericPlatformMemory::GetMemorySizeBucket

#rnx

diff --git a/Engine/Source/Runtime/Core/Private/GenericPlatform/GenericPlatformMemory.cpp b/Engine/Source/Runtime/Core/Private/GenericPlatform/GenericPlatformMemory.cpp
--- a/Engine/Source/Runtime/Core/Private/GenericPlatform/GenericPlatformMemory.cpp
+++ b/Engine/Source/Runtime/Core/Private/GenericPlatform/GenericPlatformMemory.cpp
@@ -364,10 +364,6 @@ void FGenericPlatformMemory::DumpPlatformAndAllocatorStats( class FOutputDevice&
 
 EPlatformMemorySizeBucket FGenericPlatformMemory::GetMemorySizeBucket()
 {
-	// Make sure this won't be called by more than one thread to avoid race conditions in the returned value
-	static FCriticalSection Lock;
-	FScopeLock ScopeLock(&Lock);
-	
 	static bool bCalculatedBucket = false;
 	static EPlatformMemorySizeBucket Bucket = EPlatformMemorySizeBucket::Default;
 
