Fire an ensure if we retry the same transaction more than once.

#rnx

diff --git a/Engine/Source/Runtime/Core/Private/AutoRTFM/Context.cpp b/Engine/Source/Runtime/Core/Private/AutoRTFM/Context.cpp
--- a/Engine/Source/Runtime/Core/Private/AutoRTFM/Context.cpp
+++ b/Engine/Source/Runtime/Core/Private/AutoRTFM/Context.cpp
@@ -539,7 +539,13 @@ ETransactionResult FContext::Transact(void (*UninstrumentedFunction)(void*), voi
 				// required be done before we attempt to re-execute.
 				PostAbortCallback();
 				PostAbortCallback.Reset();
-
+				
+				if (ForTheRuntime::GetInternalAbortAction() != ForTheRuntime::EAutoRTFMInternalAbortActionState::Abort)
+				{
+					EXPECT(0 == NumAbortAndRetries);
+					NumAbortAndRetries += 1;
+				}
+				
 				ASSERT(Status == EContextStatus::InPostAbort);
 				Status = EContextStatus::OnTrack;
 
@@ -684,6 +690,7 @@ void FContext::Reset()
 	Status = EContextStatus::Idle;
 	StackLocalInitializerDepth = 0;
 	TaskPool.Reset();
+	NumAbortAndRetries = 0;
 }
 
 void FContext::Throw()
diff --git a/Engine/Source/Runtime/Core/Private/AutoRTFM/Context.h b/Engine/Source/Runtime/Core/Private/AutoRTFM/Context.h
--- a/Engine/Source/Runtime/Core/Private/AutoRTFM/Context.h
+++ b/Engine/Source/Runtime/Core/Private/AutoRTFM/Context.h
@@ -173,6 +173,7 @@ private:
 	FThreadID CurrentThreadId;
 	uint32 StackLocalInitializerDepth = 0;
 	TTask<void()> PostAbortCallback;
+	uint32 NumAbortAndRetries = 0;
 	FTaskPool TaskPool;
 	FTransactionPool TransactionPool;
 	FOpenHashThrottler OpenHashThrottler{
diff --git a/Engine/Source/Runtime/Core/Private/AutoRTFM/Utils.h b/Engine/Source/Runtime/Core/Private/AutoRTFM/Utils.h
--- a/Engine/Source/Runtime/Core/Private/AutoRTFM/Utils.h
+++ b/Engine/Source/Runtime/Core/Private/AutoRTFM/Utils.h
@@ -77,6 +77,19 @@ FString GetFunctionDescription(TReturnType (*FunctionPtr)(TParameterTypes...))
 		AutoRTFM::DoAssert(Lambda); \
 	}
 
+// Same funkiness as ASSERT, except that it doesn't cause a fatal error.
+// TODO: Maybe upgrade the Display -> Warning?
+#define EXPECT(exp) \
+	if (UNLIKELY(!(exp))) \
+	{ \
+		auto Lambda = [] \
+		{ \
+			FDebug::DumpStackTraceToLog(TEXT("AutoRTFM expect backtrace"), ELogVerbosity::Display); \
+			UE_LOG(LogAutoRTFM, Display, TEXT("AutoRTFM expect %s!"), TEXT(#exp)); \
+		}; \
+		AutoRTFM::DoAssert(Lambda); \
+	}
+
 #if defined(__has_feature)
 #if __has_feature(address_sanitizer)
 #define AUTORTFM_NO_ASAN [[clang::no_sanitize("address")]]
