Add a `DoExpect` to fix the return-from-noreturn function bug.

diff --git a/Engine/Source/Runtime/Core/Private/AutoRTFM/Utils.h b/Engine/Source/Runtime/Core/Private/AutoRTFM/Utils.h
--- a/Engine/Source/Runtime/Core/Private/AutoRTFM/Utils.h
+++ b/Engine/Source/Runtime/Core/Private/AutoRTFM/Utils.h
@@ -52,6 +52,18 @@ FString GetFunctionDescription(TReturnType (*FunctionPtr)(TParameterTypes...))
 #endif
 }
 
+// We use a special calling convention for this call to minimize
+// the cost of calling this function (which is super unlikely to
+// be called because it is the expecting true codepath, EG. only
+// things we don't expect to happen get here).
+#if UE_AUTORTFM
+[[clang::preserve_most]]
+#endif
+FORCENOINLINE inline void DoExpect(void(*Logger)())
+{
+	Logger();
+}
+
 } // namespace AutoRTFM
 
 // This is all a bit funky - but for good reason! We want `ASSERT`
@@ -87,7 +99,7 @@ FString GetFunctionDescription(TReturnType (*FunctionPtr)(TParameterTypes...))
 			FDebug::DumpStackTraceToLog(TEXT("AutoRTFM expect backtrace"), ELogVerbosity::Display); \
 			UE_LOG(LogAutoRTFM, Display, TEXT("AutoRTFM expect %s!"), TEXT(#exp)); \
 		}; \
-		AutoRTFM::DoAssert(Lambda); \
+		AutoRTFM::DoExpect(Lambda); \
 	}
 
 #if defined(__has_feature)
