Allow us to have a percentage of servers have AutoRTFM enabled to streamline deployment.

#rnx

diff --git a/Engine/Source/Programs/AutoRTFMTests/Private/API.cpp b/Engine/Source/Programs/AutoRTFMTests/Private/API.cpp
--- a/Engine/Source/Programs/AutoRTFMTests/Private/API.cpp
+++ b/Engine/Source/Programs/AutoRTFMTests/Private/API.cpp
@@ -1165,3 +1165,38 @@ TEST_CASE("API.IsAutoRTFMRuntimeEnabled")
 	REQUIRE(!AutoRTFM::ForTheRuntime::SetAutoRTFMRuntime(AutoRTFM::ForTheRuntime::AutoRTFM_DisabledByDefault));
 	REQUIRE(AutoRTFM::ForTheRuntime::IsAutoRTFMRuntimeEnabled());
 }
+
+TEST_CASE("API.CoinTossDisable")
+{
+	SECTION("With default enablement")
+	{
+		AutoRTFM::Testing::FEnabledStateResetterScoped _(AutoRTFM::ForTheRuntime::AutoRTFM_EnabledByDefault);
+
+		// Set the chance of disabling to 100.0, effectively disabling the coin toss.
+		AutoRTFM::ForTheRuntime::SetAutoRTFMEnabledProbability(100.0f);
+		REQUIRE(!AutoRTFM::ForTheRuntime::CoinTossDisable());
+
+		// Set the chance of disabling to 0.0, always disabling by coin toss.
+		AutoRTFM::ForTheRuntime::SetAutoRTFMEnabledProbability(0.0f);
+		REQUIRE(AutoRTFM::ForTheRuntime::CoinTossDisable());
+	}
+	
+	SECTION("With force enablement")
+	{
+		AutoRTFM::Testing::FEnabledStateResetterScoped _(AutoRTFM::ForTheRuntime::AutoRTFM_ForcedEnabled);
+
+		// Set the chance of enabling to 0.0, always disabling by coin toss - but this gets ignored because
+		// we are set to force enable.
+		AutoRTFM::ForTheRuntime::SetAutoRTFMEnabledProbability(0.0f);
+		REQUIRE(!AutoRTFM::ForTheRuntime::CoinTossDisable());
+	}
+
+	SECTION("Already disabled")
+	{
+		AutoRTFM::Testing::FEnabledStateResetterScoped _(AutoRTFM::ForTheRuntime::AutoRTFM_DisabledByDefault);
+
+		// Set the chance of disabling to 0.0, always disabling by coin toss.
+		AutoRTFM::ForTheRuntime::SetAutoRTFMEnabledProbability(0.0f);
+		REQUIRE(!AutoRTFM::ForTheRuntime::CoinTossDisable());
+	}
+}
diff --git a/Engine/Source/Runtime/Core/Private/AutoRTFM/API.cpp b/Engine/Source/Runtime/Core/Private/AutoRTFM/API.cpp
--- a/Engine/Source/Runtime/Core/Private/AutoRTFM/API.cpp
+++ b/Engine/Source/Runtime/Core/Private/AutoRTFM/API.cpp
@@ -7,6 +7,8 @@
 #include "HAL/Platform.h"
 #include "Transaction.h"
 
+#include <random>
+
 #if UE_AUTORTFM
 static_assert(UE_AUTORTFM_ENABLED, "AutoRTFM/API.cpp requires the compiler flag '-fautortfm'");
 #endif
@@ -25,6 +27,8 @@ namespace
 	int GAutoRTFMInternalAbortAction = AutoRTFM::ForTheRuntime::EAutoRTFMInternalAbortActionState::Crash;
 	int GAutoRTFMRetryTransactions = AutoRTFM::ForTheRuntime::EAutoRTFMRetryTransactionState::NoRetry;
 
+	float GAutoRTFMEnabledProbability = 100.0f;
+
 	// Note: GMemoryValidationLevel should never be EMemoryValidationLevel::Default.
 	// EMemoryValidationLevel::Default is a special enumerator that can be used
 	// in the public API to map to GDefaultMemoryValidationLevel.
@@ -163,6 +167,42 @@ namespace AutoRTFM
 #endif
 		}
 
+		void SetAutoRTFMEnabledProbability(float Chance)
+		{
+#if UE_AUTORTFM
+			ASSERT(Chance >= 0.0f && Chance <= 100.0f);
+			GAutoRTFMEnabledProbability = Chance;
+#endif
+		}
+
+		bool CoinTossDisable()
+		{
+#if UE_AUTORTFM
+			if (!IsAutoRTFMRuntimeEnabled())
+			{
+				return false;
+			}
+
+			static std::random_device Device;
+			static std::mt19937 Generator(Device());
+			static std::uniform_real_distribution<float> Distribution(0.0f, 100.0f);
+
+			// A value in the range [0..100), EG. inclusive of 0, exclusive of 100.
+			// So a `GAutoRTFMEnabledProbability` of 100 is always greater than the
+			// potential random range, and `GAutoRTFMEnabledProbability` of 0 is
+			// always less than or equal to the range.
+			const float Random = Distribution(Generator);
+
+			if (GAutoRTFMEnabledProbability <= Random)
+			{
+				// If we have the runtime cvar set to `ForcedEnabled` then it'll ignore this call!
+				return AutoRTFM::ForTheRuntime::SetAutoRTFMRuntime(AutoRTFM::ForTheRuntime::AutoRTFM_ForcedDisabled);
+			}
+#endif
+
+			return false;
+		}
+
 		void SetInternalAbortAction(EAutoRTFMInternalAbortActionState State)
 		{
 			GAutoRTFMInternalAbortAction = State;
@@ -173,7 +213,6 @@ namespace AutoRTFM
 			return static_cast<EAutoRTFMInternalAbortActionState>(GAutoRTFMInternalAbortAction);
 		}
 
-		
 		bool GetEnsureOnInternalAbort()
 		{
 			return GAutoRTFMEnsureOnInternalAbort;
diff --git a/Engine/Source/Runtime/Core/Private/AutoRTFM/AutoRTFMUE.cpp b/Engine/Source/Runtime/Core/Private/AutoRTFM/AutoRTFMUE.cpp
--- a/Engine/Source/Runtime/Core/Private/AutoRTFM/AutoRTFMUE.cpp
+++ b/Engine/Source/Runtime/Core/Private/AutoRTFM/AutoRTFMUE.cpp
@@ -157,7 +157,7 @@ static FAutoConsoleVariable CVarAutoRTFMRetryTransactions(
 	TEXT("Enables the AutoRTFM sanitizer-like mode where we can force an abort-and-retry on transactions (useful to test abort codepaths work as intended)"),
 	FConsoleVariableDelegate::CreateLambda([](IConsoleVariable* Variable)
 	{
-		AutoRTFM::ForTheRuntime::EAutoRTFMRetryTransactionState Value = static_cast<AutoRTFM::ForTheRuntime::EAutoRTFMRetryTransactionState>(Variable->GetInt());
+		const AutoRTFM::ForTheRuntime::EAutoRTFMRetryTransactionState Value = static_cast<AutoRTFM::ForTheRuntime::EAutoRTFMRetryTransactionState>(Variable->GetInt());
 		AutoRTFM::ForTheRuntime::SetRetryTransaction(Value);
 	}),
 	ECVF_Default
@@ -169,7 +169,7 @@ static FAutoConsoleVariable CVarAutoRTFMMemoryValidationLevel(
 	TEXT("Detects potential memory corruption due to writes made both by a transaction and open-code"),
 	FConsoleVariableDelegate::CreateLambda([](IConsoleVariable* Variable)
 	{
-		int Value = Variable->GetInt();
+		const int Value = Variable->GetInt();
 		AutoRTFM::ForTheRuntime::SetMemoryValidationLevel(static_cast<AutoRTFM::EMemoryValidationLevel>(Value));
 	}),
 	ECVF_Default
@@ -181,8 +181,8 @@ static FAutoConsoleVariable CVarAutoRTFMMemoryValidationThrottlingEnabled(
 	TEXT("Automatically skips memory validation on opens if validation is taking an excessive amount of time"),
 	FConsoleVariableDelegate::CreateLambda([](IConsoleVariable* Variable)
 	{
-		bool Value = Variable->GetBool();
-		AutoRTFM::ForTheRuntime::SetMemoryValidationThrottlingEnabled(Value);
+		const bool bValue = Variable->GetBool();
+		AutoRTFM::ForTheRuntime::SetMemoryValidationThrottlingEnabled(bValue);
 	}),
 	ECVF_Default
 );
@@ -193,8 +193,20 @@ static FAutoConsoleVariable CVarAutoRTFMMemoryValidationStatisticsEnabled(
 	TEXT("Disable logging of memory validation statistics"),
 	FConsoleVariableDelegate::CreateLambda([](IConsoleVariable* Variable)
 	{
-		bool Value = Variable->GetBool();
-		AutoRTFM::ForTheRuntime::SetMemoryValidationStatisticsEnabled(Value);
+		const bool bValue = Variable->GetBool();
+		AutoRTFM::ForTheRuntime::SetMemoryValidationStatisticsEnabled(bValue);
+	}),
+	ECVF_Default
+);
+
+static FAutoConsoleVariable CVarAutoRTFMEnabledProbability(
+	TEXT("AutoRTFMEnabledProbability"),
+	0.0f,
+	TEXT("A rational percentage from [0..100] of what threshold to `CoinTossDisable` AutoRTFM. 100 means always enable, 0 means always disable"),
+	FConsoleVariableDelegate::CreateLambda([](IConsoleVariable* Variable)
+	{
+		const float Value = Variable->GetFloat();
+		AutoRTFM::ForTheRuntime::SetAutoRTFMEnabledProbability(Value);
 	}),
 	ECVF_Default
 );
diff --git a/Engine/Source/Runtime/Core/Public/AutoRTFM/AutoRTFM.h b/Engine/Source/Runtime/Core/Public/AutoRTFM/AutoRTFM.h
--- a/Engine/Source/Runtime/Core/Public/AutoRTFM/AutoRTFM.h
+++ b/Engine/Source/Runtime/Core/Public/AutoRTFM/AutoRTFM.h
@@ -1191,6 +1191,15 @@ namespace ForTheRuntime
 	// Query whether the AutoRTFM runtime is enabled.
 	UE_AUTORTFM_API bool IsAutoRTFMRuntimeEnabled();
 
+	// Set the percentage [0..100] chance that a call to `CoinTossDisable` will end up disabling AutoRTFM.
+	// 100% means never disable via coin-toss, 0% means always disable. So passing `0.1` means disable
+	// all but 1/1000's calls via `CoinTossDisable`.
+	UE_AUTORTFM_API void SetAutoRTFMEnabledProbability(float Chance);
+
+	// Call to randomly disable AutoRTFM with a probability set with `SetAutoRTFMEnabledProbability`.
+	// Returns true if AutoRTFM was disabled by this call.
+	UE_AUTORTFM_API bool CoinTossDisable();
+
 	UE_AUTORTFM_API void SetInternalAbortAction(EAutoRTFMInternalAbortActionState State);
 
 	UE_AUTORTFM_API EAutoRTFMInternalAbortActionState GetInternalAbortAction();
