[AutoRTFM] Remove UE_AUTORTFM_ALWAYS_OPEN from FHitchTrackingStatScope::ReportHitch()

This is unnecessary as all callers are AutoRTFM safe. Using UE_LOG() in an open transaction is not safe as there are mixed open / closed writes.

diff --git a/Engine/Source/Runtime/Core/Private/Stats/HitchTrackingStatScope.cpp b/Engine/Source/Runtime/Core/Private/Stats/HitchTrackingStatScope.cpp
--- a/Engine/Source/Runtime/Core/Private/Stats/HitchTrackingStatScope.cpp
+++ b/Engine/Source/Runtime/Core/Private/Stats/HitchTrackingStatScope.cpp
@@ -10,7 +10,6 @@
 
 namespace UE::Stats
 {
-	UE_AUTORTFM_ALWAYS_OPEN
 	void FHitchTrackingStatScope::ReportHitch()
 	{
 		if (StatString)
@@ -32,4 +31,4 @@ namespace UE::Stats
 	}
 }
 
-#endif // USE_LIGHTWEIGHT_STATS_FOR_HITCH_DETECTION && USE_HITCH_DETECTION
\ No newline at end of file
+#endif // USE_LIGHTWEIGHT_STATS_FOR_HITCH_DETECTION && USE_HITCH_DETECTION
