Core | Rework FExternalProfiler's pausing/resuming API & make it possible to attach a profiler using a config setting

The former was a prerequisite for the latter, to lower the overhead to a minimum.
The new pausing/resuming API lives within `FActiveExternalProfilerBase` and is controlled by a static `bIsPaused` rather than a per-profiler one. Scoped includers and excluders are reworked to use this API, too. Unused implementation details such as `FExternalProfiler::TimerCount` are removed, as the code was simplified such that they're not necessary anymore.

Integration from CDPR partner branch, change by gabriel.wreczycki

diff --git a/Engine/Source/Runtime/Core/Private/ProfilingDebugging/ExternalProfiler.cpp b/Engine/Source/Runtime/Core/Private/ProfilingDebugging/ExternalProfiler.cpp
--- a/Engine/Source/Runtime/Core/Private/ProfilingDebugging/ExternalProfiler.cpp
+++ b/Engine/Source/Runtime/Core/Private/ProfilingDebugging/ExternalProfiler.cpp
@@ -6,34 +6,35 @@
 #include "Misc/Parse.h"
 #include "Misc/CommandLine.h"
 #include "Features/IModularFeatures.h"
-
+#include "Misc/ConfigCacheIni.h"
 
 #if UE_EXTERNAL_PROFILING_ENABLED
 
 DEFINE_LOG_CATEGORY_STATIC( LogExternalProfiler, Log, All );
 
-
 FExternalProfiler::FExternalProfiler()
-	:	TimerCount( 0 ),
-		// No way to tell whether we're paused or not so we assume paused as it makes the most sense
-		bIsPaused( true )
 {
 }
 
+PRAGMA_DISABLE_DEPRECATION_WARNINGS
 
 void FExternalProfiler::PauseProfiler()
 {
-	ProfilerPauseFunction();
-	bIsPaused = true;
+	if (FActiveExternalProfilerBase::GetActiveProfiler() == this)
+	{
+		FActiveExternalProfilerBase::SetActiveProfilerRecording(false);
+	}
 }
 
-
 void FExternalProfiler::ResumeProfiler()
 {
-	ProfilerResumeFunction();
-	bIsPaused = false;
+	if (FActiveExternalProfilerBase::GetActiveProfiler() == this)
+	{
+		FActiveExternalProfilerBase::SetActiveProfilerRecording(true);
+	}
 }
 
+PRAGMA_ENABLE_DEPRECATION_WARNINGS
 
 FName FExternalProfiler::GetFeatureName()
 {
@@ -41,11 +42,23 @@ FName FExternalProfiler::GetFeatureName()
 	return ProfilerFeatureName;
 }
 
-
 bool FActiveExternalProfilerBase::bDidInitialize = false;
 
 FExternalProfiler* FActiveExternalProfilerBase::ActiveProfiler = nullptr;
 
+std::atomic<bool> FActiveExternalProfilerBase::bIsRecording = false;
+
+static FExternalProfiler* FindProfilerByName(TArrayView<FExternalProfiler*> Profilers, FStringView Name)
+{
+	FExternalProfiler** Found = Algo::FindByPredicate(
+		Profilers,
+		[&Name](const FExternalProfiler* Profiler)
+		{
+			return Name == Profiler->GetProfilerName();
+		});
+	return Found ? *Found : nullptr;
+}
+
 FExternalProfiler* FActiveExternalProfilerBase::InitActiveProfiler()
 {
 	// Create profiler on demand.
@@ -81,15 +94,22 @@ FExternalProfiler* FActiveExternalProfilerBase::InitActiveProfiler()
 			UE_LOG(LogExternalProfiler, Log, TEXT("No external profilers were discovered.  External profiling features will not be available."));
 		}
 #endif
-		
+
 		if (ActiveProfiler == nullptr)
 		{
 			FString ProfilerName = FPlatformMisc::GetEnvironmentVariable(TEXT("UE_EXTERNAL_PROFILER"));
 			if (!ProfilerName.IsEmpty())
 			{
-				FExternalProfiler** Found = Algo::FindByPredicate(AvailableProfilers,
-					 [&ProfilerName](FExternalProfiler* Profiler) { return ProfilerName == Profiler->GetProfilerName();});
-				ActiveProfiler = Found ? *Found : nullptr;
+				ActiveProfiler = FindProfilerByName(AvailableProfilers, ProfilerName);
+			}
+		}
+
+		if (ActiveProfiler == nullptr)
+		{
+			FString ProfilerName = GConfig->GetStr(TEXT("Core.ProfilingDebugging"), TEXT("ExternalProfiler"), GEngineIni);
+			if (!ProfilerName.IsEmpty())
+			{
+				ActiveProfiler = FindProfilerByName(AvailableProfilers, ProfilerName);
 			}
 		}
 
@@ -100,62 +120,38 @@ FExternalProfiler* FActiveExternalProfilerBase::InitActiveProfiler()
 	return ActiveProfiler;
 }
 
-
-void FScopedExternalProfilerBase::StartScopedTimer( const bool bWantPause )
+bool FActiveExternalProfilerBase::SetActiveProfilerRecording(bool bRecording)
 {
-	FExternalProfiler* Profiler = GetActiveProfiler();	
+	bool bWasRecording = bIsRecording.exchange(bRecording);
 
-	if (Profiler != nullptr)
+	if (FExternalProfiler* Profiler = GetActiveProfiler())
 	{
-		// Store the current state of profiler
-		bWasPaused = Profiler->bIsPaused;
-
-		// If the current profiler state isn't set to what we need, or if the global profiler sampler isn't currently
-		// running, then start it now
-		if (Profiler->TimerCount == 0 || bWantPause != Profiler->bIsPaused)
+		if (bWasRecording && !bRecording)
 		{
-			if( bWantPause )
-			{
-				Profiler->PauseProfiler();
-			}
-			else
-			{
-				Profiler->ResumeProfiler();
-			}
+			Profiler->ProfilerPauseFunction();
+		}
+		else if (!bWasRecording && bRecording)
+		{
+			Profiler->ProfilerResumeFunction();
 		}
-
-		// Increment number of overlapping timers
-		Profiler->TimerCount++;
 	}
-}
 
+	return bWasRecording;
+}
 
-void FScopedExternalProfilerBase::StopScopedTimer()
+bool FActiveExternalProfilerBase::IsActiveProfilerRecording()
 {
-	FExternalProfiler* Profiler = GetActiveProfiler();
+	return bIsRecording.load(std::memory_order_relaxed);
+}
 
-	if (Profiler != nullptr)
-	{
-		// Make sure a timer was already started
-		if (Profiler->TimerCount > 0)
-		{
-			// Decrement timer count
-			Profiler->TimerCount--;
+void FScopedExternalProfilerBase::StartScopedTimer( const bool bWantPause )
+{
+	bWasRecording = FActiveExternalProfilerBase::SetActiveProfilerRecording(!bWantPause);
+}
 
-			// Restore the previous state of VTune
-			if (bWasPaused != Profiler->bIsPaused)
-			{
-				if (bWasPaused)
-				{
-					Profiler->PauseProfiler();
-				}
-				else
-				{
-					Profiler->ResumeProfiler();
-				}
-			}
-		}
-	}
+void FScopedExternalProfilerBase::StopScopedTimer()
+{
+	FActiveExternalProfilerBase::SetActiveProfilerRecording(bWasRecording);
 }
 
 #endif	// UE_EXTERNAL_PROFILING_ENABLED
diff --git a/Engine/Source/Runtime/Core/Public/ProfilingDebugging/ExternalProfiler.h b/Engine/Source/Runtime/Core/Public/ProfilingDebugging/ExternalProfiler.h
--- a/Engine/Source/Runtime/Core/Public/ProfilingDebugging/ExternalProfiler.h
+++ b/Engine/Source/Runtime/Core/Public/ProfilingDebugging/ExternalProfiler.h
@@ -21,7 +21,10 @@ class FExternalProfiler : public IModularFeature
 
 public:
 
-	/** Constructor */
+	/**
+	 * Constructor. An external profiler being constructed must assume it starts out paused.
+	 * To start profiling, call FActiveExternalProfilerBase::SetActiveProfilerRecording(true).
+	 */
 	CORE_API FExternalProfiler();
 
 	/** Empty virtual destructor. */
@@ -30,9 +33,11 @@ public:
 	}
 
 	/** Pauses profiling. */
+	UE_DEPRECATED(5.6, "Use FActiveExternalProfilerBase::PauseActiveProfiler instead")
 	CORE_API void PauseProfiler();
 
 	/** Resumes profiling. */
+	UE_DEPRECATED(5.6, "Use FActiveExternalProfilerBase::ResumeActiveProfiler instead")
 	CORE_API void ResumeProfiler();
 
 	/**
@@ -67,17 +72,6 @@ public:
 	virtual void EndScopedEvent() {};
 
 	virtual void SetThreadName(const TCHAR* Name) {}
-
-private:
-
-	/** Number of timers currently running. Timers are always 'global inclusive'. */
-	int32 TimerCount;
-
-	/** Whether or not profiling is currently paused (as far as we know.) */
-	bool bIsPaused;
-
-	/** Friend class so we can access the private members directly. */
-	friend class FScopedExternalProfilerBase;
 };
 
 class FActiveExternalProfilerBase
@@ -86,13 +80,45 @@ public:
 
 	static FExternalProfiler* GetActiveProfiler() { return ActiveProfiler;	};
 
+	/**
+	 * Attach an external profiler, selected depending on configuration.
+	 *
+	 * Normally, no external profiler is selected, but one can be set using one of the following:
+	 *
+	 * 1. A command line argument `-ProfilerName`, where `ProfilerName` matches the string returned
+	 *    by `FExternalProfiler::GetProfilerName()`
+	 * 2. The environment variable `UE_EXTERNAL_PROFILER=ProfilerName`, where `ProfilerName` matches
+	 *    the string returned by `FExternalProfiler::GetProfilerName()`
+	 * 3. The DefaultEngine.ini config variable `[Core.ProfilingDebugging]:ExternalProfiler`, where
+	 *    the value is a string that matches the string returned by `FExternalProfiler::GetProfilerName()`.
+	 *
+	 * In that order of priority.
+	 * If no profiler is found matching these criteria, no profiler (nullptr) is attached.
+	 *
+	 * To make a profiler discoverable by this function, register it into the `FExternalProfiler`
+	 * modular feature.
+	 *
+	 * The search for a selected profiler is only performed once and then cached.
+	 * On subsequent calls, the cached result is returned.
+	 */
 	static CORE_API FExternalProfiler* InitActiveProfiler();
+
+	/**
+	 * Pauses or resumes the active profiler.
+	 * Returns the previous recording state (true if it was recording, false if it was paused.)
+	 */
+	static CORE_API bool SetActiveProfilerRecording(bool bRecording);
+	static CORE_API bool IsActiveProfilerRecording();
+
 private:
 	/** Static: True if we've tried to initialize a profiler already */
 	static CORE_API bool bDidInitialize;
 
 	/** Static: Active profiler instance that we're using */
 	static CORE_API FExternalProfiler* ActiveProfiler;
+
+	/** Static: Whether the active profiler is recording events, to our knowledge. */
+	static std::atomic<bool> bIsRecording;
 };
 
 /**
@@ -113,16 +139,21 @@ protected:
 	CORE_API void StopScopedTimer();
 
 private:
-	/** Stores the previous 'Paused' state of VTune before this scope started */
-	bool bWasPaused;
+	/** Stores the previous 'Paused' state of the current active external profiler before this scope started */
+	bool bWasRecording = false;
 };
 
 
 /**
- * FExternalProfilerIncluder
+ * Scoped guard that includes a body of code in the external profiler's captured session, by calling
+ * `SetActiveProfilerRecording(true)` in its constructor, and then
+ * `SetActiveProfilerRecording(false)` in its destructor.
  *
- * Use this to include a body of code in profiler's captured session using 'Resume' and 'Pause' cues.  It
- * can safely be embedded within another 'timer' or 'excluder' scope.
+ * It can safely be embedded within another `FExternalProfilerIncluder` or `Excluder`. However,
+ * `FActiveProfilerBase::SetActiveProfilerRecording` should not be called from within an `Includer`
+ * scope, because the `Includer` and `Excluder` mechanisms cannot record changes to the active
+ * profiler state once they're in scope---you should choose to use one mechanism or the other,
+ * but not both.
  */
 class FExternalProfilerIncluder : public FScopedExternalProfilerBase
 {
@@ -130,7 +161,6 @@ public:
 	/** Constructor */
 	FExternalProfilerIncluder()
 	{
-		// 'Timer' scopes will always 'resume' VTune
 		const bool bWantPause = false;
 		StartScopedTimer( bWantPause );
 	}
@@ -144,10 +174,15 @@ public:
 
 
 /**
- * FExternalProfilerExcluder
+ * Scoped guard that excludes a body of code in the external profiler's captured session, by calling
+ * `SetActiveProfilerRecording(false)` in its constructor, and then
+ * `SetActiveProfilerRecording(true)` in its destructor.
  *
- * Use this to EXCLUDE a body of code from profiler's captured session.  It can safely be embedded
- * within another 'timer' or 'excluder' scope.
+ * It can safely be embedded within another `FExternalProfilerExcluder` or `Includer`. However,
+ * `FActiveProfilerBase::SetActiveProfilerRecording` should not be called from within an `Includer`
+ * scope, because the `Includer` and `Excluder` mechanisms cannot record changes to the active
+ * profiler state once they're in scope---you should choose to use one mechanism or the other,
+ * but not both.
  */
 class FExternalProfilerExcluder : public FScopedExternalProfilerBase
 {
@@ -155,7 +190,6 @@ public:
 	/** Constructor */
 	FExternalProfilerExcluder()
 	{
-		// 'Excluder' scopes will always 'pause' VTune
 		const bool bWantPause = true;
 		StartScopedTimer( bWantPause );
 	}
@@ -174,30 +208,36 @@ public:
 	/** Starts a scoped event specific to the profiler. */
 	FORCEINLINE static void StartScopedEvent(const struct FColor& Color, const TCHAR* Text)
 	{
-		FExternalProfiler* Profiler = FActiveExternalProfilerBase::GetActiveProfiler();
-		if (Profiler)
+		if (FActiveExternalProfilerBase::IsActiveProfilerRecording())
 		{
-			Profiler->StartScopedEvent(Color, Text);
+			if (FExternalProfiler* Profiler = FActiveExternalProfilerBase::GetActiveProfiler())
+			{
+				Profiler->StartScopedEvent(Color, Text);
+			}
 		}
 	}
 
 	/** Starts a scoped event specific to the profiler. */
 	FORCEINLINE static void StartScopedEvent(const struct FColor& Color, const ANSICHAR* Text)
 	{
-		FExternalProfiler* Profiler = FActiveExternalProfilerBase::GetActiveProfiler();
-		if (Profiler)
+		if (FActiveExternalProfilerBase::IsActiveProfilerRecording())
 		{
-			Profiler->StartScopedEvent(Color, Text);
+			if (FExternalProfiler* Profiler = FActiveExternalProfilerBase::GetActiveProfiler())
+			{
+				Profiler->StartScopedEvent(Color, Text);
+			}
 		}
 	}
 
 	/** Ends a scoped event specific to the profiler. */
 	FORCEINLINE static void EndScopedEvent()
 	{
-		FExternalProfiler* Profiler = FActiveExternalProfilerBase::GetActiveProfiler();
-		if (Profiler)
+		if (FActiveExternalProfilerBase::IsActiveProfilerRecording())
 		{
-			Profiler->EndScopedEvent();
+			if (FExternalProfiler* Profiler = FActiveExternalProfilerBase::GetActiveProfiler())
+			{
+				Profiler->EndScopedEvent();
+			}
 		}
 	}
 };
