AVX/AVX512 CPU detection fixes

1. Only report AVX/F16C if OS has context save support for them
2. Likewise for AVX512
3. A few minor clean-ups and coding-style fixes while I'm here

diff --git a/Engine/Source/Runtime/Core/Private/Windows/WindowsPlatformMisc.cpp b/Engine/Source/Runtime/Core/Private/Windows/WindowsPlatformMisc.cpp
--- a/Engine/Source/Runtime/Core/Private/Windows/WindowsPlatformMisc.cpp
+++ b/Engine/Source/Runtime/Core/Private/Windows/WindowsPlatformMisc.cpp
@@ -3208,11 +3208,44 @@ uint32 FWindowsPlatformMisc::GetFeatureBits_X86()
 	FeatureBits |= (CpuInfo[3] & (1u << 26)) ? ECPUFeatureBits_X86::SSE2 : 0;
 	FeatureBits |= (CpuInfo[2] & (1u << 9)) ? ECPUFeatureBits_X86::SSSE3 : 0;
 	FeatureBits |= (CpuInfo[2] & (1u << 20)) ? ECPUFeatureBits_X86::SSE42 : 0;
-	FeatureBits |= (CpuInfo[2] & (1u << 28)) ? ECPUFeatureBits_X86::AVX : 0;
-	FeatureBits |= (CpuInfo[2] & (1u << 29)) ? ECPUFeatureBits_X86::F16C : 0;
 
-	// We don't have a feature flag we report for this, but we do use it later
-	bool has_popcnt = (CpuInfo[2] & (1u << 23)) != 0;
+	// Feature flags we don't report directly (because they need further checks), but use internally
+	bool bHasPopcnt = (CpuInfo[2] & (1u << 23)) != 0;
+	bool bHasOsXsave = (CpuInfo[2] & (1u << 27)) != 0;
+	bool bHasAVX = (CpuInfo[2] & (1u << 28)) != 0;
+	bool bHasF16C = (CpuInfo[2] & (1u << 29)) != 0;
+
+	// Get XCR0 if available (context save bits) - required for AVX and other extensions
+	// to test this, on Admin command prompt: (WARNING: this can brick your machine
+	// depending on your drivers. You might need to boot from a recovery/install USB
+	// image to undo if that's the case!)
+	//     bcdedit /set xsavedisable 1
+	// (be sure to "/set xsavedisable 0" again once you're done testing)
+	uint64 Xcr0 = 0;
+	if (bHasOsXsave)
+	{
+		Xcr0 = _xgetbv(0);
+	}
+
+	// Required context save bits to use AVX and AVX-512, respectively
+	bool bAVXContextSupport = (Xcr0 & 0x6) == 0x6;
+	bool bAVX512ContextSupport = (Xcr0 & 0xe6) == 0xe6;
+
+	// Combine these to determine effective AVX support
+	if (bAVXContextSupport)
+	{
+		if (bHasAVX)
+		{
+			FeatureBits |= ECPUFeatureBits_X86::AVX;
+		}
+
+		// Even though F16C does not require using AVX registers, the relevant instructions
+		// fault if XSAVE support is disabled
+		if (bHasF16C)
+		{
+			FeatureBits |= ECPUFeatureBits_X86::F16C;
+		}
+	}
 
 	if (MaxLeaf >= 7)
 	{
@@ -3227,13 +3260,12 @@ uint32 FWindowsPlatformMisc::GetFeatureBits_X86()
 		// Also only report AVX or the BMIs if POPCNT is present; all processors I know of
 		// have either both or neither, and it's convenient for us to be able to assume
 		// that either BMI1/BMI2 or AVX2 implies POPCNT.
-		if ((FeatureBits & ECPUFeatureBits_X86::AVX) && has_popcnt)
+		if (bAVXContextSupport && bHasAVX && bHasPopcnt)
 		{
-			if (CpuInfo[1] & (1u << 3))	FeatureBits |= ECPUFeatureBits_X86::BMI1;
-			if (CpuInfo[1] & (1u << 8))	FeatureBits |= ECPUFeatureBits_X86::BMI2;
-
-			// OS must save YMM registers between context switch
-			bool OsSavesAvxRegs = (_xgetbv(0) & 6) == 6;
+			// These don't actually require XSAVE support to use, but let's not increase
+			// the size of the testing matrix further.
+			FeatureBits |= (CpuInfo[1] & (1u << 3)) ? ECPUFeatureBits_X86::BMI1 : 0;
+			FeatureBits |= (CpuInfo[1] & (1u << 8)) ? ECPUFeatureBits_X86::BMI2 : 0;
 
 			// In addition to the above, only report AVX2 if BMI1 (and thus LZCNT/TZCNT)
 			// are also reported present; finally VC++ with /arch:AVX2 will emit BMI2
@@ -3243,20 +3275,29 @@ uint32 FWindowsPlatformMisc::GetFeatureBits_X86()
 			// In practice this is not a limitation, AVX2 and BMI2 are a package deal on
 			// all uArchs I'm aware of.
 			const uint32 Avx2Bits = (1u << 3) /* BMI1 */ | (1u << 5) /* AVX2 */ | (1u << 8) /* BMI2 */;
-			if (((CpuInfo[1] & Avx2Bits) == Avx2Bits) && OsSavesAvxRegs)
+			if ((CpuInfo[1] & Avx2Bits) == Avx2Bits)
+			{
 				FeatureBits |= ECPUFeatureBits_X86::AVX2;
+			}
 
-			// For us to report AVX512, we want the Skylake feature set
-			const uint32 Avx512Bits = (1u << 31) /* AVX512VL */ | (1u << 30) /* AVX512BW */ | (1u << 17) /* AVX512DQ */ | (1u << 16) /* AVX512F */;
-			if ((CpuInfo[1] & Avx512Bits) == Avx512Bits)
-				FeatureBits |= ECPUFeatureBits_X86::AVX512;
-
-			// Use the VBMI2 bit (set on ICL+) to set the NOCAVEATS flag. This is available
-			// on a generation of cores where AVX-512 has no major clock penalty anymore so
-			// whether to use AVX-512 or not is a much more straightforward calculation,
-			// and not so dependent on what else is running at the same time.
-			if (CpuInfo[2] & (1u << 6))
-				FeatureBits |= ECPUFeatureBits_X86::AVX512_NOCAVEATS;
+			if (bAVX512ContextSupport)
+			{
+				// For us to report AVX512, we want the Skylake feature set
+				const uint32 Avx512Bits = (1u << 31) /* AVX512VL */ | (1u << 30) /* AVX512BW */ | (1u << 17) /* AVX512DQ */ | (1u << 16) /* AVX512F */;
+				if ((CpuInfo[1] & Avx512Bits) == Avx512Bits)
+				{
+					FeatureBits |= ECPUFeatureBits_X86::AVX512;
+				}
+
+				// Use the VBMI2 bit (set on ICL+) to set the NOCAVEATS flag. This is available
+				// on a generation of cores where AVX-512 has no major clock penalty anymore so
+				// whether to use AVX-512 or not is a much more straightforward calculation,
+				// and not so dependent on what else is running at the same time.
+				if (CpuInfo[2] & (1u << 6))
+				{
+					FeatureBits |= ECPUFeatureBits_X86::AVX512_NOCAVEATS;
+				}
+			}
 		}
 	}
 
