Change where we do the many-retries logging to be before we do the longjmp.

diff --git a/Engine/Source/Runtime/Core/Private/AutoRTFM/Context.cpp b/Engine/Source/Runtime/Core/Private/AutoRTFM/Context.cpp
--- a/Engine/Source/Runtime/Core/Private/AutoRTFM/Context.cpp
+++ b/Engine/Source/Runtime/Core/Private/AutoRTFM/Context.cpp
@@ -168,7 +168,16 @@ ETransactionResult FContext::RollbackTransaction(EContextStatus NewStatus)
 		   NewStatus == EContextStatus::AbortedByCascadingRetry);
 
 	Status = NewStatus;
-
+		
+	if (Status == EContextStatus::AbortedByCascadingRetry)
+	{
+		if (ForTheRuntime::GetInternalAbortAction() != ForTheRuntime::EAutoRTFMInternalAbortActionState::Abort)
+		{
+			EXPECT(0 == NumAbortAndRetries);
+			NumAbortAndRetries += 1;
+		}
+	}
+	
 	ETransactionResult Result = ETransactionResult::AbortedByRequest;
 
 	if (GetNumDeferredTransactions())
@@ -540,12 +549,6 @@ ETransactionResult FContext::Transact(void (*UninstrumentedFunction)(void*), voi
 				PostAbortCallback();
 				PostAbortCallback.Reset();
 				
-				if (ForTheRuntime::GetInternalAbortAction() != ForTheRuntime::EAutoRTFMInternalAbortActionState::Abort)
-				{
-					EXPECT(0 == NumAbortAndRetries);
-					NumAbortAndRetries += 1;
-				}
-				
 				ASSERT(Status == EContextStatus::InPostAbort);
 				Status = EContextStatus::OnTrack;
 
