Fix that we had two places claiming they had a default value for the probability.

diff --git a/Engine/Source/Runtime/Core/Private/AutoRTFM/API.cpp b/Engine/Source/Runtime/Core/Private/AutoRTFM/API.cpp
--- a/Engine/Source/Runtime/Core/Private/AutoRTFM/API.cpp
+++ b/Engine/Source/Runtime/Core/Private/AutoRTFM/API.cpp
@@ -175,6 +175,15 @@ namespace AutoRTFM
 #endif
 		}
 
+		float GetAutoRTFMEnabledProbability()
+		{
+#if UE_AUTORTFM
+			return GAutoRTFMEnabledProbability;
+#else
+			return 0.0f;
+#endif
+		}
+
 		bool CoinTossDisable()
 		{
 #if UE_AUTORTFM
diff --git a/Engine/Source/Runtime/Core/Private/AutoRTFM/AutoRTFMUE.cpp b/Engine/Source/Runtime/Core/Private/AutoRTFM/AutoRTFMUE.cpp
--- a/Engine/Source/Runtime/Core/Private/AutoRTFM/AutoRTFMUE.cpp
+++ b/Engine/Source/Runtime/Core/Private/AutoRTFM/AutoRTFMUE.cpp
@@ -201,7 +201,7 @@ static FAutoConsoleVariable CVarAutoRTFMMemoryValidationStatisticsEnabled(
 
 static FAutoConsoleVariable CVarAutoRTFMEnabledProbability(
 	TEXT("AutoRTFMEnabledProbability"),
-	0.0f,
+	AutoRTFM::ForTheRuntime::GetAutoRTFMEnabledProbability(),
 	TEXT("A rational percentage from [0..100] of what threshold to `CoinTossDisable` AutoRTFM. 100 means always enable, 0 means always disable"),
 	FConsoleVariableDelegate::CreateLambda([](IConsoleVariable* Variable)
 	{
diff --git a/Engine/Source/Runtime/Core/Public/AutoRTFM/AutoRTFM.h b/Engine/Source/Runtime/Core/Public/AutoRTFM/AutoRTFM.h
--- a/Engine/Source/Runtime/Core/Public/AutoRTFM/AutoRTFM.h
+++ b/Engine/Source/Runtime/Core/Public/AutoRTFM/AutoRTFM.h
@@ -1196,6 +1196,9 @@ namespace ForTheRuntime
 	// all but 1/1000's calls via `CoinTossDisable`.
 	UE_AUTORTFM_API void SetAutoRTFMEnabledProbability(float Chance);
 
+	// Get the enabled probability set via `SetAutoRTFMEnabledProbability`.
+	UE_AUTORTFM_API float GetAutoRTFMEnabledProbability();
+
 	// Call to randomly disable AutoRTFM with a probability set with `SetAutoRTFMEnabledProbability`.
 	// Returns true if AutoRTFM was disabled by this call.
 	UE_AUTORTFM_API bool CoinTossDisable();
