Added FUtf8String support for FName.

diff --git a/Engine/Source/Runtime/Core/Private/UObject/UnrealNames.cpp b/Engine/Source/Runtime/Core/Private/UObject/UnrealNames.cpp
--- a/Engine/Source/Runtime/Core/Private/UObject/UnrealNames.cpp
+++ b/Engine/Source/Runtime/Core/Private/UObject/UnrealNames.cpp
@@ -15,6 +15,7 @@
 #include "Misc/Crc.h"
 #include "Misc/StringBuilder.h"
 #include "Containers/UnrealString.h"
+#include "Containers/Utf8String.h"
 #include "UObject/NameTypes.h"
 #include "Logging/LogMacros.h"
 #include "Misc/ByteSwap.h"
@@ -2332,6 +2333,19 @@ FORCEINLINE ANSICHAR const* FNameEntry::GetUnterminatedName(ANSICHAR(&OptionalDe
 #endif
 }
 
+int32 FNameEntry::GetNameLengthUtf8() const
+{
+	if (Header.bIsWide)
+	{
+		FNameBuffer Temp;
+		return FPlatformString::ConvertedLength<UTF8CHAR>(GetUnterminatedName(Temp.WideName), Header.Len);
+	}
+	else
+	{
+		return Header.Len;
+	}
+}
+
 FORCEINLINE FNameStringView FNameEntry::MakeView(FNameBuffer& OptionalDecodeBuffer) const
 {
 #if UE_FNAME_OUTLINE_NUMBER
@@ -2418,6 +2432,20 @@ FString FNameEntry::GetPlainNameString() const
 	}
 }
 
+FUtf8String FNameEntry::GetPlainNameUtf8String() const
+{
+	checkName(Header.Len != 0);
+	FNameBuffer Temp;
+	if (Header.bIsWide)
+	{
+		return FUtf8String::ConstructFromPtrSize(GetUnterminatedName(Temp.WideName), Header.Len);
+	}
+	else
+	{
+		return FUtf8String::ConstructFromPtrSize(GetUnterminatedName(Temp.AnsiName), Header.Len);
+	}
+}
+
 void FNameEntry::AppendNameToString(FString& Out) const
 {
 	checkName(Header.Len != 0);
@@ -2425,6 +2453,13 @@ void FNameEntry::AppendNameToString(FString& Out) const
 	Out.Append(EntryToCString(*this, Temp), Header.Len);
 }
 
+void FNameEntry::AppendNameToString(FUtf8String& Out) const
+{
+	checkName(Header.Len != 0);
+	FNameBuffer Temp;
+	Out.Append(EntryToCString(*this, Temp), Header.Len);
+}
+
 void FNameEntry::AppendNameToString(FWideStringBuilderBase& Out) const
 {
 	checkName(Header.Len != 0);
@@ -3500,7 +3535,23 @@ FString FName::ToString() const
 	}
 #endif // !UE_FNAME_OUTLINE_NUMBER
 	
-	FString Out;	
+	FString Out;
+	ToString(Out);
+	return Out;
+}
+
+FUtf8String FName::ToUtf8String() const
+{
+// With UE_FNAME_OUTLINE_NUMBER reading number isn't free so skip this check
+#if !UE_FNAME_OUTLINE_NUMBER
+	if (GetNumber() == NAME_NO_NUMBER_INTERNAL)
+	{
+		// Avoids some extra allocations in non-number case
+		return GetDisplayNameEntry()->GetPlainNameUtf8String();
+	}
+#endif // !UE_FNAME_OUTLINE_NUMBER
+	
+	FUtf8String Out;
 	ToString(Out);
 	return Out;
 }
@@ -3510,10 +3561,10 @@ FString LexToString(const FName& Name)
 	return Name.ToString();
 }
 
-void FName::ToString(FString& Out) const
+void FName::ToString(FWideString& Out) const
 {
 #if UE_FNAME_OUTLINE_NUMBER
-	FNameEntryId Id = GetDisplayIndexInternal();	
+	FNameEntryId Id = GetDisplayIndexInternal();
 	const FNameEntry* ThisNameEntry = ResolveEntry(Id);
 	if (ThisNameEntry->Header.Len != 0)
 	{
@@ -3538,7 +3589,7 @@ void FName::ToString(FString& Out) const
 	{
 		Out.Reset(NameEntry->GetNameLength());
 		NameEntry->AppendNameToString(Out);
-	}	
+	}
 	else
 	{
 		Out.Reset(NameEntry->GetNameLength() + 6);
@@ -3550,6 +3601,46 @@ void FName::ToString(FString& Out) const
 #endif // UE_FNAME_OUTLINE_NUMBER
 }
 
+void FName::ToString(FUtf8String& Out) const
+{
+#if UE_FNAME_OUTLINE_NUMBER
+	FNameEntryId Id = GetDisplayIndexInternal();
+	const FNameEntry* ThisNameEntry = ResolveEntry(Id);
+	if (ThisNameEntry->Header.Len != 0)
+	{
+		// No number
+		Out.Reset(ThisNameEntry->GetNameLengthUtf8());
+		ThisNameEntry->AppendNameToString(Out);
+	}
+	else
+	{
+		const FNameEntry* BaseEntry = ResolveEntry(ThisNameEntry->GetNumberedName().Id);
+		Out.Reset(BaseEntry->GetNameLengthUtf8() + 6);
+		BaseEntry->AppendNameToString(Out);
+
+		Out += '_';
+		Out.AppendInt(NAME_INTERNAL_TO_EXTERNAL(ThisNameEntry->GetNumber()));
+	}
+#else // UE_FNAME_OUTLINE_NUMBER
+	// A version of ToString that saves at least one string copy
+	const FNameEntry* const NameEntry = GetDisplayNameEntry();
+
+	if (GetNumber() == NAME_NO_NUMBER_INTERNAL)
+	{
+		Out.Reset(NameEntry->GetNameLengthUtf8());
+		NameEntry->AppendNameToString(Out);
+	}
+	else
+	{
+		Out.Reset(NameEntry->GetNameLengthUtf8() + 6);
+		NameEntry->AppendNameToString(Out);
+
+		Out += '_';
+		Out.AppendInt(NAME_INTERNAL_TO_EXTERNAL(GetNumber()));
+	}
+#endif // UE_FNAME_OUTLINE_NUMBER
+}
+
 void FName::ToString(FWideStringBuilderBase& Out) const
 {
 	Out.Reset();
@@ -3613,7 +3704,7 @@ uint32 FName::ToStringInternal(TCHAR* Out, const uint32 OutSize) const
 	}
 }
 
-void FName::AppendString(FString& Out) const
+void FName::AppendString(FWideString& Out) const
 {
 	const FNameEntry* const NameEntry = GetDisplayNameEntry();
 
@@ -3625,6 +3716,18 @@ void FName::AppendString(FString& Out) const
 	}
 }
 
+void FName::AppendString(FUtf8String& Out) const
+{
+	const FNameEntry* const NameEntry = GetDisplayNameEntry();
+
+	NameEntry->AppendNameToString( Out );
+	if (GetNumber() != NAME_NO_NUMBER_INTERNAL)
+	{
+		Out += '_';
+		Out.AppendInt(NAME_INTERNAL_TO_EXTERNAL(GetNumber()));
+	}
+}
+
 template <typename StringBuilderType>
 void FName::AppendStringInternal(StringBuilderType& Out) const
 {
@@ -4406,6 +4509,16 @@ FName FLazyName::Resolve() const
 #endif // UE_FNAME_OUTLINE_NUMBER
 }
 
+CORE_API FString FLazyName::ToString() const
+{
+	return Resolve().ToString();
+}
+
+CORE_API FUtf8String FLazyName::ToUtf8String() const
+{
+	return Resolve().ToUtf8String();
+}
+
 uint32 FLazyName::ParseNumber(const ANSICHAR* Str, int32 Len)
 {
 	return FNameHelper::ParseNumber(Str, Len);
diff --git a/Engine/Source/Runtime/Core/Public/UObject/NameTypes.h b/Engine/Source/Runtime/Core/Public/UObject/NameTypes.h
--- a/Engine/Source/Runtime/Core/Public/UObject/NameTypes.h
+++ b/Engine/Source/Runtime/Core/Public/UObject/NameTypes.h
@@ -259,6 +259,7 @@ public:
 	/** Returns whether this name entry is represented via WIDECHAR or ANSICHAR. */
 	FORCEINLINE bool IsWide() const								{ return Header.bIsWide; }
 	FORCEINLINE int32 GetNameLength() const 					{ return Header.Len; }
+	CORE_API int32 GetNameLengthUtf8() const;
 #if UE_FNAME_OUTLINE_NUMBER
 	FORCEINLINE bool IsNumbered() const 						{ return Header.Len == 0; }
 #else
@@ -280,8 +281,12 @@ public:
 	/** Copy name to a dynamically allocated FString. */
 	CORE_API FString GetPlainNameString() const;
 
+	/** Copy name to a dynamically allocated FString. */
+	CORE_API FUtf8String GetPlainNameUtf8String() const;
+
 	/** Appends name to string. May allocate. */
-	CORE_API void AppendNameToString(FString& OutString) const;
+	CORE_API void AppendNameToString(FWideString& OutString) const;
+	CORE_API void AppendNameToString(FUtf8String& OutString) const;
 
 	/** Appends name to string builder. */
 	CORE_API void AppendNameToString(FWideStringBuilderBase& OutString) const;
@@ -465,6 +470,7 @@ struct FScriptName
 	inline bool operator==(EName Name) { return *this == FScriptName(Name); }
 
 	CORE_API FString ToString() const;
+	CORE_API FUtf8String ToUtf8String() const;
 
 	// The internal structure of FScriptName is private in order to handle UE_FNAME_OUTLINE_NUMBER
 private: 
@@ -641,12 +647,20 @@ public:
 	 */
 	[[nodiscard]] CORE_API FString ToString() const;
 
+	/**
+	 * Converts an FName to a readable format
+	 *
+	 * @return String representation of the name
+	 */
+	[[nodiscard]] CORE_API FUtf8String ToUtf8String() const;
+
 	/**
 	 * Converts an FName to a readable format, in place
 	 * 
 	 * @param Out String to fill with the string representation of the name
 	 */
-	CORE_API void ToString(FString& Out) const;
+	CORE_API void ToString(FWideString& Out) const;
+	CORE_API void ToString(FUtf8String& Out) const;
 
 	/**
 	 * Converts an FName to a readable format, in place
@@ -734,7 +748,8 @@ public:
 	 * 
 	 * @param Out String to append with the string representation of the name
 	 */
-	CORE_API void AppendString(FString& Out) const;
+	CORE_API void AppendString(FWideString& Out) const;
+	CORE_API void AppendString(FUtf8String& Out) const;
 
 	/**
 	 * Converts an FName to a readable format, in place, appending to an existing string (ala GetFullName)
@@ -1695,10 +1710,8 @@ public:
 
 	CORE_API FName Resolve() const;
 
-	CORE_API FString ToString() const
-	{
-		return Resolve().ToString();
-	}
+	CORE_API FString ToString() const;
+	CORE_API FUtf8String ToUtf8String() const;
 
 private:
 	struct FLiteralOrName
