Csv Profiler event delegates:
- Add a delegate for CSV events so that other systems can hook into them
- Fix event logging to include the category
#tests Tested on console

diff --git a/Engine/Source/Runtime/Core/Private/ProfilingDebugging/CsvProfiler.cpp b/Engine/Source/Runtime/Core/Private/ProfilingDebugging/CsvProfiler.cpp
--- a/Engine/Source/Runtime/Core/Private/ProfilingDebugging/CsvProfiler.cpp
+++ b/Engine/Source/Runtime/Core/Private/ProfilingDebugging/CsvProfiler.cpp
@@ -4273,7 +4273,20 @@ void FCsvProfiler::RecordEvent(int32 CategoryIndex, const FString& EventText)
 	if (GCsvProfilerIsCapturing && GCsvCategoriesEnabled[CategoryIndex])
 	{
 		LLM_SCOPE(ELLMTag::CsvProfiler);
-		UE_LOG(LogCsvProfiler, Display, TEXT("CSVEvent \"%s\" [Frame %d]"), *EventText, FCsvProfiler::Get()->GetCaptureFrameNumber());
+
+		// Log the event and broadcast to delegates
+		if (CategoryIndex == CSV_CATEGORY_INDEX_GLOBAL)
+		{
+			UE_LOG(LogCsvProfiler, Display, TEXT("CSVEvent \"%s\" [Frame %d]"), *EventText, FCsvProfiler::Get()->GetCaptureFrameNumber());
+			FCsvProfiler::Get()->OnCSVProfileEvent().Broadcast(FString(), EventText);
+		}
+		else
+		{
+			FString CategoryName = FCsvCategoryData::Get()->GetCategoryNameByIndex(CategoryIndex);
+			UE_LOG(LogCsvProfiler, Display, TEXT("CSVEvent \"%s/%s\" [Frame %d]"), *CategoryName, *EventText, FCsvProfiler::Get()->GetCaptureFrameNumber());
+			FCsvProfiler::Get()->OnCSVProfileEvent().Broadcast(CategoryName, EventText);
+		}
+
 		FCsvProfilerThreadData::Get().AddEvent(EventText, CategoryIndex);
 		if (GCsvEventExecCmds)
 		{
diff --git a/Engine/Source/Runtime/Core/Public/ProfilingDebugging/CsvProfiler.h b/Engine/Source/Runtime/Core/Public/ProfilingDebugging/CsvProfiler.h
--- a/Engine/Source/Runtime/Core/Public/ProfilingDebugging/CsvProfiler.h
+++ b/Engine/Source/Runtime/Core/Public/ProfilingDebugging/CsvProfiler.h
@@ -495,6 +495,10 @@ public:
 	DECLARE_MULTICAST_DELEGATE(FOnCSVProfileEndFrame);
 	FOnCSVProfileEndFrame& OnCSVProfileEndFrame() { return OnCSVProfileEndFrameDelegate; }
 
+	/** Called when an event is emitted */
+	DECLARE_MULTICAST_DELEGATE_TwoParams(FOnCSVEvent, const FString& /* CategoryName */, const FString& /* EventText */);
+	FOnCSVEvent& OnCSVProfileEvent() { return OnCSVProfileEventDelegate; }
+
 	void SetRenderThreadId(uint32 InRenderThreadId)
 	{
 		RenderThreadId = InRenderThreadId;
@@ -566,6 +570,7 @@ private:
 	FOnCSVProfileEnd OnCSVProfileEndDelegate;
 	FOnCSVProfileEndFrame OnCSVProfileEndFrameDelegate;
 	FOnCSVProfileFinished OnCSVProfileFinishedDelegate;
+	FOnCSVEvent OnCSVProfileEventDelegate;
 
 	std::atomic<uint32> RenderThreadId{ 0 };
 	std::atomic<uint32> RHIThreadId{ 0 };
