[AutoRTFM] Replace use of UE debugging macros with AutoRTFM versions that delegate to the ExternAPI.

AutoRTFM now internally uses UTF-8 encoding for strings, so drop all the TEXT() wrappers around string literals.

Most of the files are now wrapped in preprocessor guards to only compile when __AUTORTFM is defined. This is to prevent accidental breakages of other compilers attempting to compile AutoRTFM code. A benefit of this is that we can now safely assume we're being compiled with clang, removing unnecessary compiler-dependent preprocessor blocks in the files.

autortfm_report_error() has been moved from CompilerInterface.cpp to Utils.cpp and renamed to ReportError(), as it wasn't used by the compiler. It now supports printf-style formatting.

This change comes substantial macro renaming:

* UE_LOG(LogAutoRTFM, Verbose, ...) -> AUTORTFM_VERBOSE(...)
...
* UE_LOG(LogAutoRTFM, Fatal, ...) -> AUTORTFM_FATAL(...)
* UE_CLOG(C, LogAutoRTFM, Verbose, ...) -> AUTORTFM_VERBOSE_IF(C, ...)
...
* UE_CLOG(C, LogAutoRTFM, Fatal, ...) -> AUTORTFM_FATAL_IF(C, ...)
* ASSERT(C) -> AUTORTFM_ASSERT(C)
* ensure(C) -> AUTORTFM_ENSURE(C)
* ensureMsgf(C, M) -> AUTORTFM_ENSURE_MSG(C, ...)

Part of the work to separate AutoRTFM out of core to break cyclic build dependency issues.

diff --git a/Engine/Source/Runtime/Core/Private/AutoRTFM/API.cpp b/Engine/Source/Runtime/Core/Private/AutoRTFM/API.cpp
--- a/Engine/Source/Runtime/Core/Private/AutoRTFM/API.cpp
+++ b/Engine/Source/Runtime/Core/Private/AutoRTFM/API.cpp
@@ -75,13 +75,13 @@ namespace AutoRTFM
 		{
 			// #noop if AutoRTFM is not compiled in
 #if UE_AUTORTFM
-			auto Stringify = [](int Query) -> const TCHAR*
+			auto Stringify = [](int Query) -> const char*
 			{
 				switch (Query)
 				{
 				default:
 					InternalUnreachable();
-#define HANDLE_CASE(x) case x: return TEXT(#x)
+#define HANDLE_CASE(x) case x: return #x
 				HANDLE_CASE(EAutoRTFMEnabledState::AutoRTFM_ForcedEnabled);
 				HANDLE_CASE(EAutoRTFMEnabledState::AutoRTFM_ForcedDisabled);
 				HANDLE_CASE(EAutoRTFMEnabledState::AutoRTFM_OverriddenEnabled);
@@ -96,7 +96,7 @@ namespace AutoRTFM
 
 			auto DoIgnoreLog = [&](int State, int Stored)
 			{
-				UE_LOG(LogAutoRTFM, Log, TEXT("Ignoring changing AutoRTFM runtime state to '%s' as it was previously set to '%s'."), Stringify(State), Stringify(Stored));
+				AUTORTFM_LOG("Ignoring changing AutoRTFM runtime state to '%s' as it was previously set to '%s'", Stringify(State), Stringify(Stored));
 			};
 
 			switch (GAutoRTFMRuntimeEnabled)
@@ -170,7 +170,7 @@ namespace AutoRTFM
 		void SetAutoRTFMEnabledProbability(float Chance)
 		{
 #if UE_AUTORTFM
-			ASSERT(Chance >= 0.0f && Chance <= 100.0f);
+			AUTORTFM_ASSERT(Chance >= 0.0f && Chance <= 100.0f);
 			GAutoRTFMEnabledProbability = Chance;
 #endif
 		}
@@ -364,7 +364,7 @@ namespace AutoRTFM
 #if UE_AUTORTFM
 			// Check memory validation is enabled, otherwise we won't have a
 			// hash to compare against.
-			ASSERT(GetMemoryValidationLevel() == EMemoryValidationLevel::Error ||
+			AUTORTFM_ASSERT(GetMemoryValidationLevel() == EMemoryValidationLevel::Error ||
 				GetMemoryValidationLevel() == EMemoryValidationLevel::Warn);
 			if (FTransaction* Transaction = AutoRTFM::FContext::Get()->GetMaterializedTransaction()) {
 				Transaction->DebugBreakIfMemoryValidationFails();
@@ -428,7 +428,7 @@ AUTORTFM_CORE_API autortfm_result RTFM_autortfm_transact_then_open(void (*Uninst
 AUTORTFM_CORE_API void RTFM_autortfm_commit(void (*UninstrumentedWork)(void*), void (*InstrumentedWork)(void*), void* Arg)
 {
 	autortfm_result Result = autortfm_transact(UninstrumentedWork, InstrumentedWork, Arg);
-	UE_CLOG(Result != autortfm_committed, LogAutoRTFM, Fatal, TEXT("Unexpected transaction result: %u."), Result);
+	AUTORTFM_FATAL_IF(Result != autortfm_committed, "Unexpected transaction result: %u", Result);
 }
 
 AUTORTFM_CORE_API void RTFM_autortfm_abort()
@@ -439,12 +439,12 @@ AUTORTFM_CORE_API void RTFM_autortfm_abort()
 
 AUTORTFM_CORE_API void RTFM_autortfm_start_transaction()
 {
-	UE_LOG(LogAutoRTFM, Fatal, TEXT("The function `autortfm_start_transaction` was called from closed code."));
+	AUTORTFM_FATAL("The function `autortfm_start_transaction` was called from closed code");
 }
 
 AUTORTFM_CORE_API void RTFM_autortfm_commit_transaction()
 {
-	UE_LOG(LogAutoRTFM, Fatal, TEXT("The function `RTFM_autortfm_commit_transaction` was called from closed code."));
+	AUTORTFM_FATAL("The function `RTFM_autortfm_commit_transaction` was called from closed code");
 }
 
 AUTORTFM_CORE_API void RTFM_autortfm_abort_transaction()
@@ -455,7 +455,7 @@ AUTORTFM_CORE_API void RTFM_autortfm_abort_transaction()
 
 AUTORTFM_CORE_API void RTFM_autortfm_clear_transaction_status()
 {
-	UE_LOG(LogAutoRTFM, Fatal, TEXT("The function `autortfm_clear_transaction_status` was called from closed code."));
+	AUTORTFM_FATAL("The function `autortfm_clear_transaction_status` was called from closed code");
 	AutoRTFM::InternalUnreachable();
 }
 
@@ -469,8 +469,8 @@ AUTORTFM_CORE_API autortfm_status RTFM_autortfm_close(void (*UninstrumentedWork)
 	}
 	else
 	{
-		FString Message = FString::Printf(TEXT("Could not find function %p '%s' where '%s'."), UninstrumentedWork, *GetFunctionDescription(UninstrumentedWork), ANSI_TO_TCHAR("autortfm_close"));
-		autortfm_report_error(*Message);
+		ReportError("Could not find function %p '%s' for autortfm_close()",
+			reinterpret_cast<void*>(UninstrumentedWork), GetFunctionDescription(UninstrumentedWork).c_str());
 	}
 
 	return static_cast<autortfm_status>(Context->GetStatus());
@@ -543,76 +543,76 @@ AUTORTFM_CORE_API void RTFM_autortfm_open(autortfm_memory_validation_level Memor
 
 AUTORTFM_CORE_API void RTFM_autortfm_record_open_write_err(void*, size_t)
 {
-	UE_LOG(LogAutoRTFM, Fatal, TEXT("The function `autortfm_record_open_write` was called from closed code."));
+	AUTORTFM_FATAL("The function `autortfm_record_open_write` was called from closed code");
 }
 
 AUTORTFM_CORE_API void RTFM_CascadingAbortTransactionInternal(TTask<void()>&& RunAfterAbort)
 {
 	FContext* const Context = FContext::Get();
-	ASSERT(Context->GetStatus() == EContextStatus::OnTrack);
+	AUTORTFM_ASSERT(Context->GetStatus() == EContextStatus::OnTrack);
 	Context->AbortTransactionWithPostAbortCallback(EContextStatus::AbortedByCascadingAbort, MoveTemp(RunAfterAbort));
 }
 
 AUTORTFM_CORE_API void RTFM_CascadingRetryTransactionInternal(TTask<void()> && RunAfterAbortBeforeRetryWork)
 {
 	FContext* Context = FContext::Get();
-	ASSERT(Context->GetStatus() == EContextStatus::OnTrack);
+	AUTORTFM_ASSERT(Context->GetStatus() == EContextStatus::OnTrack);
 	Context->AbortTransactionWithPostAbortCallback(EContextStatus::AbortedByCascadingRetry, MoveTemp(RunAfterAbortBeforeRetryWork));
 }
 
 AUTORTFM_CORE_API void RTFM_OnCommitInternal(TTask<void()> && Work)
 {
 	FContext* Context = FContext::Get();
-	ASSERT(Context->GetStatus() == EContextStatus::OnTrack);
+	AUTORTFM_ASSERT(Context->GetStatus() == EContextStatus::OnTrack);
 	Context->GetCurrentTransaction()->DeferUntilCommit(MoveTemp(Work));
 }
 
 AUTORTFM_CORE_API void RTFM_OnAbortInternal(TTask<void()>&& Work)
 {
 	FContext* Context = FContext::Get();
-	ASSERT(Context->GetStatus() == EContextStatus::OnTrack);
+	AUTORTFM_ASSERT(Context->GetStatus() == EContextStatus::OnTrack);
 	Context->GetCurrentTransaction()->DeferUntilAbort(MoveTemp(Work));
 }
 
 AUTORTFM_CORE_API void RTFM_PushOnCommitHandlerInternal(const void* Key, TTask<void()>&& Work)
 {
 	FContext* Context = FContext::Get();
-	ASSERT(Context->GetStatus() == EContextStatus::OnTrack);
+	AUTORTFM_ASSERT(Context->GetStatus() == EContextStatus::OnTrack);
 	Context->GetCurrentTransaction()->PushDeferUntilCommitHandler(Key, MoveTemp(Work));
 }
 
 AUTORTFM_CORE_API void RTFM_PopOnCommitHandlerInternal(const void* Key)
 {
 	FContext* Context = FContext::Get();
-	ASSERT(Context->GetStatus() == EContextStatus::OnTrack);
+	AUTORTFM_ASSERT(Context->GetStatus() == EContextStatus::OnTrack);
 	Context->GetCurrentTransaction()->PopDeferUntilCommitHandler(Key);
 }
 
 AUTORTFM_CORE_API void RTFM_PopAllOnCommitHandlersInternal(const void* Key)
 {
 	FContext* Context = FContext::Get();
-	ASSERT(Context->GetStatus() == EContextStatus::OnTrack);
+	AUTORTFM_ASSERT(Context->GetStatus() == EContextStatus::OnTrack);
 	Context->GetCurrentTransaction()->PopAllDeferUntilCommitHandlers(Key);
 }
 
 AUTORTFM_CORE_API void RTFM_PushOnAbortHandlerInternal(const void* Key, TTask<void()>&& Work)
 {
 	FContext* Context = FContext::Get();
-	ASSERT(Context->GetStatus() == EContextStatus::OnTrack);
+	AUTORTFM_ASSERT(Context->GetStatus() == EContextStatus::OnTrack);
 	Context->GetCurrentTransaction()->PushDeferUntilAbortHandler(Key, MoveTemp(Work));
 }
 
 AUTORTFM_CORE_API void RTFM_PopOnAbortHandlerInternal(const void* Key)
 {
 	FContext* Context = FContext::Get();
-	ASSERT(Context->GetStatus() == EContextStatus::OnTrack);
+	AUTORTFM_ASSERT(Context->GetStatus() == EContextStatus::OnTrack);
 	Context->GetCurrentTransaction()->PopDeferUntilAbortHandler(Key);
 }
 
 AUTORTFM_CORE_API void RTFM_PopAllOnAbortHandlersInternal(const void* Key)
 {
 	FContext* Context = FContext::Get();
-	ASSERT(Context->GetStatus() == EContextStatus::OnTrack);
+	AUTORTFM_ASSERT(Context->GetStatus() == EContextStatus::OnTrack);
 	Context->GetCurrentTransaction()->PopAllDeferUntilAbortHandlers(Key);
 }
 
@@ -660,13 +660,13 @@ AUTORTFM_CORE_API bool IsAutoRTFMInitialized()
 // allocations before AutoRTFM is initialized.
 AUTORTFM_CORE_API void ProcessAllPendingOpenToClosedRegistrations()
 {
-	ASSERT(IsAutoRTFMInitialized());
+	AUTORTFM_ASSERT(IsAutoRTFMInitialized());
 
 	while (autortfm_open_to_closed_table* Table = GPendingOpenToClosedRegistrations)
 	{
 		for (const autortfm_open_to_closed_mapping* Mapping = Table->Mappings; Mapping->Open; Mapping++)
 		{
-			UE_LOG(LogAutoRTFM, Verbose, TEXT("Registering open %p -> %p"), Mapping->Open, Mapping->Closed);
+			AUTORTFM_VERBOSE("Registering open %p -> %p", Mapping->Open, Mapping->Closed);
 			FunctionMapAdd(Mapping->Open, Mapping->Closed);
 		}
 		GPendingOpenToClosedRegistrations = Table->Next;
@@ -679,12 +679,12 @@ AUTORTFM_CORE_API void ProcessAllPendingOpenToClosedRegistrations()
 
 extern "C" void autortfm_initialize(const autortfm_extern_api* ExternAPI)
 {
-	ensureMsgf(!FContext::Get(), TEXT("AutoRTFM initialized twice"));
-	ASSERT(ExternAPI);
-	ASSERT(ExternAPI->Allocate);
-	ASSERT(ExternAPI->AllocateZeroed);
-	ASSERT(ExternAPI->Reallocate);
-	ASSERT(ExternAPI->Free);
+	AUTORTFM_ENSURE_MSG(!FContext::Get(), "AutoRTFM initialized twice");
+	AUTORTFM_ASSERT(ExternAPI);
+	AUTORTFM_ASSERT(ExternAPI->Allocate);
+	AUTORTFM_ASSERT(ExternAPI->AllocateZeroed);
+	AUTORTFM_ASSERT(ExternAPI->Reallocate);
+	AUTORTFM_ASSERT(ExternAPI->Free);
 	GExternAPI = *ExternAPI;
 	FContext::Create();
 	ProcessAllPendingOpenToClosedRegistrations();
@@ -753,7 +753,7 @@ extern "C" void autortfm_commit(void (*UninstrumentedWork)(void*), void (*Instru
 	}
 
     autortfm_result Result = autortfm_transact(UninstrumentedWork, InstrumentedWork, Arg);
-	UE_CLOG(Result != autortfm_committed, LogAutoRTFM, Fatal, TEXT("Unexpected transaction result: %u."), Result);
+	AUTORTFM_FATAL_IF(Result != autortfm_committed, "Unexpected transaction result: %u", Result);
 }
 
 extern "C" void autortfm_abort()
@@ -763,7 +763,7 @@ extern "C" void autortfm_abort()
 		return RTFM_autortfm_abort();
 	}
 
-	UE_CLOG(!FContext::Get()->IsTransactional(), LogAutoRTFM, Fatal, TEXT("The function `autortfm_abort` was called from outside a transaction."));
+	AUTORTFM_FATAL_IF(!FContext::Get()->IsTransactional(), "The function `autortfm_abort` was called from outside a transaction");
 	FContext::Get()->AbortByRequestAndThrow();
 }
 
@@ -776,7 +776,7 @@ extern "C" bool autortfm_start_transaction()
 	}
 
 	FContext* Context = FContext::Get();
-	UE_CLOG(!Context->IsTransactional(), LogAutoRTFM, Fatal, TEXT("The function `autortfm_start_transaction` was called from outside a transact."));
+	AUTORTFM_FATAL_IF(!Context->IsTransactional(), "The function `autortfm_start_transaction` was called from outside a transact");
 	Context->StartTransaction(ForTheRuntime::GetMemoryValidationLevel());
 	return true;
 }
@@ -789,30 +789,30 @@ extern "C" autortfm_result autortfm_commit_transaction()
 		return autortfm_aborted_by_language;
 	}
 
-	UE_CLOG(!FContext::Get()->IsTransactional(), LogAutoRTFM, Fatal, TEXT("The function `autortfm_commit_transaction` was called from outside a transact."));
+	AUTORTFM_FATAL_IF(!FContext::Get()->IsTransactional(), "The function `autortfm_commit_transaction` was called from outside a transact");
 	return static_cast<autortfm_result>(FContext::Get()->CommitTransaction());
 }
 
 extern "C" void autortfm_abort_transaction()
 {
-	ASSERT(autortfm_is_closed()); // RollbackTransaction should be used when in the open.
+	AUTORTFM_ASSERT(autortfm_is_closed()); // RollbackTransaction should be used when in the open.
 	RTFM_autortfm_abort_transaction();
 }
 
 extern "C" autortfm_result autortfm_rollback_transaction()
 {
-	ASSERT(!autortfm_is_closed()); // AbortTransaction should be used when in the closed.
+	AUTORTFM_ASSERT(!autortfm_is_closed()); // AbortTransaction should be used when in the closed.
 	FContext* const Context = FContext::Get();
-	UE_CLOG(!Context->IsTransactional(), LogAutoRTFM, Fatal, TEXT("The function `autortfm_abort_transaction` was called from outside a transact."));
+	AUTORTFM_FATAL_IF(!Context->IsTransactional(), "The function `autortfm_abort_transaction` was called from outside a transact");
 	return static_cast<autortfm_result>(Context->RollbackTransaction(EContextStatus::AbortedByRequest));
 }
 
 extern "C" autortfm_result autortfm_cascading_rollback_transaction()
 {
-	ASSERT(!autortfm_is_closed());
+	AUTORTFM_ASSERT(!autortfm_is_closed());
 
 	FContext* const Context = FContext::Get();
-	UE_CLOG(!Context->IsTransactional(), LogAutoRTFM, Fatal, TEXT("The function `autortfm_cascading_rollback_transaction` was called from outside a transact."));
+	AUTORTFM_FATAL_IF(!Context->IsTransactional(), "The function `autortfm_cascading_rollback_transaction` was called from outside a transact");
 	return static_cast<autortfm_result>(Context->RollbackTransaction(EContextStatus::AbortedByCascadingAbort));
 }
 
@@ -823,7 +823,7 @@ extern "C" void autortfm_clear_transaction_status()
 		return RTFM_autortfm_clear_transaction_status();
 	}
 
-	ASSERT(FContext::Get()->IsAborting());
+	AUTORTFM_ASSERT(FContext::Get()->IsAborting());
 	FContext::Get()->ClearTransactionStatus();
 }
 
@@ -879,7 +879,7 @@ extern "C" autortfm_status autortfm_close(void (*UninstrumentedWork)(void*), voi
 
 	if (ForTheRuntime::IsAutoRTFMRuntimeEnabled())
 	{
-		UE_CLOG(!FContext::Get()->IsTransactional(), LogAutoRTFM, Fatal, TEXT("Close called from an outside a transaction."));
+		AUTORTFM_FATAL_IF(!FContext::Get()->IsTransactional(), "Close called from an outside a transaction");
 
 		FContext* const Context = FContext::Get();
 
@@ -889,13 +889,14 @@ extern "C" autortfm_status autortfm_close(void (*UninstrumentedWork)(void*), voi
 		}
 		else
 		{
+			std::string FunctionDescription = GetFunctionDescription(UninstrumentedWork);
 			if (ForTheRuntime::GetInternalAbortAction() == ForTheRuntime::EAutoRTFMInternalAbortActionState::Crash)
 			{
-				UE_LOG(LogAutoRTFM, Fatal, TEXT("Could not find function %p '%s' where '%s'."), UninstrumentedWork, *GetFunctionDescription(UninstrumentedWork), ANSI_TO_TCHAR("autortfm_close"));
+				AUTORTFM_FATAL("Could not find function %p '%s' in autortfm_close()", UninstrumentedWork, FunctionDescription.c_str());
 			}
 			else
 			{
-				ensureMsgf(!ForTheRuntime::GetEnsureOnInternalAbort(), TEXT("Could not find function %p '%s' where '%s'."), UninstrumentedWork, *GetFunctionDescription(UninstrumentedWork), ANSI_TO_TCHAR("autortfm_close"));
+				AUTORTFM_ENSURE_MSG(!ForTheRuntime::GetEnsureOnInternalAbort(), "Could not find function %p '%s' in autortfm_close()", UninstrumentedWork, FunctionDescription.c_str());
 			}
 	        Context->AbortByLanguageAndThrow();
 		}
@@ -1001,7 +1002,7 @@ void ForTheRuntime::Initialize(const ForTheRuntime::FExternAPI& ExternAPI)
 
 void ForTheRuntime::CascadingAbortTransactionInternal(TTask<void()>&& RunAfterAbort)
 {
-	ASSERT(autortfm_is_closed());
+	AUTORTFM_ASSERT(autortfm_is_closed());
 	RTFM_CascadingAbortTransactionInternal(MoveTemp(RunAfterAbort));
 }
 
@@ -1188,11 +1189,11 @@ extern "C" UE_AUTORTFM_NOAUTORTFM void autortfm_check_abi(void* const Ptr, const
 		}
     } RuntimeConstants;
 
-	UE_CLOG(sizeof(FConstants) != Size, LogAutoRTFM, Fatal, TEXT("ABI error between AutoRTFM compiler and runtime."));
+	AUTORTFM_FATAL_IF(sizeof(FConstants) != Size, "ABI error between AutoRTFM compiler and runtime");
 
     const FConstants* const CompilerConstants = static_cast<FConstants*>(Ptr);
 
-	UE_CLOG(RuntimeConstants != *CompilerConstants, LogAutoRTFM, Fatal, TEXT("ABI error between AutoRTFM compiler and runtime."));
+	AUTORTFM_FATAL_IF(RuntimeConstants != *CompilerConstants, "ABI error between AutoRTFM compiler and runtime");
 }
 } // namespace AutoRTFM
 
diff --git a/Engine/Source/Runtime/Core/Private/AutoRTFM/Atomic.h b/Engine/Source/Runtime/Core/Private/AutoRTFM/Atomic.h
--- a/Engine/Source/Runtime/Core/Private/AutoRTFM/Atomic.h
+++ b/Engine/Source/Runtime/Core/Private/AutoRTFM/Atomic.h
@@ -2,6 +2,8 @@
 
 #pragma once
 
+#if (defined(__AUTORTFM) && __AUTORTFM)
+
 #include <atomic>
 
 namespace AutoRTFM
@@ -74,3 +76,5 @@ public:
 };
 
 } // namespace AutoRTFM
+
+#endif // (defined(__AUTORTFM) && __AUTORTFM)
diff --git a/Engine/Source/Runtime/Core/Private/AutoRTFM/AutoRTFMUE.cpp b/Engine/Source/Runtime/Core/Private/AutoRTFM/AutoRTFMUE.cpp
--- a/Engine/Source/Runtime/Core/Private/AutoRTFM/AutoRTFMUE.cpp
+++ b/Engine/Source/Runtime/Core/Private/AutoRTFM/AutoRTFMUE.cpp
@@ -1,16 +1,27 @@
 // Copyright Epic Games, Inc. All Rights Reserved.
 
 #include "AutoRTFM/AutoRTFMUE.h"
-#include "AutoRTFM/AutoRTFM.h"
 
 #if UE_AUTORTFM
 
+#include "AutoRTFM/AutoRTFM.h"
+#include "Containers/StringConv.h"
+#include "Logging/LogVerbosity.h"
+#include "Misc/AssertionMacros.h"
+#include <cstdarg>
+#include <cstdio>
+#include <memory>
+
 #include "Containers/UnrealString.h"
 #include "GenericPlatform/GenericPlatformCrashContext.h"
 #include "HAL/IConsoleManager.h"
 #include "HAL/UnrealMemory.h"
+#include "Logging/LogMacros.h"
+#include "Misc/AssertionMacros.h"
+
+DECLARE_LOG_CATEGORY_EXTERN(LogAutoRTFM, Display, All)
 
-#include "Utils.h" // TODO(SOL-7652) Remove this include once all logging is abstracted
+DEFINE_LOG_CATEGORY(LogAutoRTFM)
 
 static_assert(UE_AUTORTFM_ENABLED, "AutoRTFM/AutoRTFMUE.cpp requires the compiler flag '-fautortfm'");
 
@@ -211,6 +222,146 @@ static FAutoConsoleVariable CVarAutoRTFMEnabledProbability(
 	ECVF_Default
 );
 
+bool IsSeverityActive(autortfm_log_severity Severity)
+{
+	switch (Severity)
+	{
+		case autortfm_log_verbose:
+			return UE_LOG_ACTIVE(LogAutoRTFM, Verbose);
+		case autortfm_log_info:
+			return UE_LOG_ACTIVE(LogAutoRTFM, Display);
+		case autortfm_log_warn:
+			return UE_LOG_ACTIVE(LogAutoRTFM, Warning);
+		case autortfm_log_error:
+			return UE_LOG_ACTIVE(LogAutoRTFM, Error);
+		case autortfm_log_fatal:
+			return UE_LOG_ACTIVE(LogAutoRTFM, Fatal);
+	}
+}
+
+ELogVerbosity::Type SeverityToVerbosity(autortfm_log_severity Severity)
+{
+	switch (Severity)
+	{
+		case autortfm_log_verbose:
+			return ELogVerbosity::Verbose;
+		case autortfm_log_info:
+			return ELogVerbosity::Display;
+		case autortfm_log_warn:
+			return ELogVerbosity::Warning;
+		case autortfm_log_error:
+			return ELogVerbosity::Error;
+		case autortfm_log_fatal:
+			return ELogVerbosity::Fatal;
+	}
+}
+
+UE_AUTORTFM_ALWAYS_OPEN_NO_MEMORY_VALIDATION
+TStringConversion<TStringConvert<UTF8CHAR, TCHAR>, 128> FmtToTChar(const char* Format, va_list Args)
+{
+	static constexpr size_t InlineBufferLength = 256;
+	char InlineBuffer[InlineBufferLength];
+
+	va_list Args2;
+	va_copy(Args2, Args);
+	int Count = vsnprintf(InlineBuffer, InlineBufferLength, Format, Args);
+
+	if (Count < InlineBufferLength)
+	{
+		va_end(Args2);
+		return StringCast<TCHAR>(reinterpret_cast<const UTF8CHAR*>(InlineBuffer));
+	}
+	else
+	{
+		std::unique_ptr<char[]> Buffer{new char[Count+1]};
+		vsnprintf(Buffer.get(), Count+1, Format, Args2);
+		va_end(Args2);
+		return StringCast<TCHAR>(reinterpret_cast<const UTF8CHAR*>(Buffer.get()));
+	}
+}
+
+UE_AUTORTFM_ALWAYS_OPEN_NO_MEMORY_VALIDATION
+void AutoRTFMLog(const char* File, int Line, autortfm_log_severity Severity, const char* Format, va_list Args)
+{
+	if (!IsSeverityActive(Severity))
+	{
+		return;
+	}
+
+	static ::UE::Logging::Private::FStaticBasicLogDynamicData DynamicData;
+	UE::Logging::Private::FStaticBasicLogRecord Record(
+		/* Format */ TEXT("%s"),
+		/* File */ File,
+		/* Line */ Line,
+		/* Verbosity */ SeverityToVerbosity(Severity),
+		/* DynamicData */ DynamicData
+	);
+
+	switch (Severity)
+	{
+		case autortfm_log_verbose:
+		case autortfm_log_info:
+		case autortfm_log_warn:
+		case autortfm_log_error:
+			UE::Logging::Private::BasicLog(LogAutoRTFM, &Record, FmtToTChar(Format, Args).Get());
+			break;
+		case autortfm_log_fatal:
+			FDebug::DumpStackTraceToLog(TEXT("AutoRTFM backtrace"), ELogVerbosity::Error);
+			UE::Logging::Private::BasicFatalLog(LogAutoRTFM, &Record, FmtToTChar(Format, Args).Get());
+			break;
+	}
+
+	va_end(Args);
+}
+
+UE_AUTORTFM_ALWAYS_OPEN_NO_MEMORY_VALIDATION
+void AutoRTFMLogWithCallstack(autortfm_log_severity Severity, const char* Format, va_list Args)
+{
+	if (!IsSeverityActive(Severity))
+	{
+		return;
+	}
+
+	static ::UE::Logging::Private::FStaticBasicLogDynamicData DynamicData;
+	UE::Logging::Private::FStaticBasicLogRecord Record(
+		/* Format */ TEXT("%s"),
+		/* File */ "",
+		/* Line */ 0,
+		/* Verbosity */ SeverityToVerbosity(Severity),
+		/* DynamicData */ DynamicData
+	);
+
+	switch (Severity)
+	{
+		case autortfm_log_verbose:
+		case autortfm_log_info:
+		case autortfm_log_warn:
+		case autortfm_log_error:
+			FDebug::DumpStackTraceToLog(TEXT("AutoRTFM backtrace"), Record.Verbosity);
+			UE::Logging::Private::BasicLog(LogAutoRTFM, &Record, FmtToTChar(Format, Args).Get());
+			break;
+		case autortfm_log_fatal:
+			FDebug::DumpStackTraceToLog(TEXT("AutoRTFM backtrace"), ELogVerbosity::Error);
+			UE::Logging::Private::BasicFatalLog(LogAutoRTFM, &Record, FmtToTChar(Format, Args).Get());
+			break;
+	}
+
+	va_end(Args);
+}
+
+UE_AUTORTFM_ALWAYS_OPEN_NO_MEMORY_VALIDATION
+void AutoRTFMEnsureFailure(const char* File, int Line, const char* Condition, void* ProgramCounter, const char* Format, va_list Args)
+{
+#if DO_ENSURE
+	if (Format == nullptr)
+	{
+		Format = "";
+	}
+	FDebug::DumpStackTraceToLog(TEXT("AutoRTFM backtrace"), ELogVerbosity::Error);
+	FDebug::EnsureFailed(Condition, File, Line, ProgramCounter, FmtToTChar(Format, Args).Get());
+#endif
+}
+
 } // anonymous namespace
 
 UE_AUTORTFM_API void AutoRTFM::InitializeForUE()
@@ -229,6 +380,11 @@ UE_AUTORTFM_API void AutoRTFM::InitializeForUE()
 		return FMemory::MallocZeroed(Size, Alignment);
 	};
 	ExternAPI.Free = FMemory::Free;
+	ExternAPI.EnsureFailure = AutoRTFMEnsureFailure;
+	ExternAPI.Log = AutoRTFMLog;
+	ExternAPI.LogWithCallstack = AutoRTFMLogWithCallstack;
+	ExternAPI.IsLogActive = IsSeverityActive;
+
 	ExternAPI.OnRuntimeEnabledChanged = OnAutoRTFMRuntimeEnabledChanged;
 	ExternAPI.OnRetryTransactionsChanged = OnAutoRTFMRetryTransactionsChanged;
 	ExternAPI.OnMemoryValidationLevelChanged = OnAutoRTFMMemoryValidationLevelChanged;
diff --git a/Engine/Source/Runtime/Core/Private/AutoRTFM/BitStack.h b/Engine/Source/Runtime/Core/Private/AutoRTFM/BitStack.h
--- a/Engine/Source/Runtime/Core/Private/AutoRTFM/BitStack.h
+++ b/Engine/Source/Runtime/Core/Private/AutoRTFM/BitStack.h
@@ -2,6 +2,8 @@
 
 #pragma once
 
+#if (defined(__AUTORTFM) && __AUTORTFM)
+
 #include "ContainerValidation.h"
 #include "Stack.h"
 #include "Utils.h"
@@ -124,13 +126,13 @@ public:
 			Word = Bit ? (Word | Mask) : (Word & ~Mask);
 		}
 		++Count;
-		ASSERT(Validation == EContainerValidation::Disabled || NumWordsFor(Count) == Words.Num());
+		AUTORTFM_ASSERT(Validation == EContainerValidation::Disabled || NumWordsFor(Count) == Words.Num());
 	}
 
 	// Removes the last item on the stack.
 	inline bool Pop()
 	{
-		ASSERT(Validation == EContainerValidation::Disabled || Count > 0);
+		AUTORTFM_ASSERT(Validation == EContainerValidation::Disabled || Count > 0);
 		--Count;
 		const WordType WordBitIndex = Count & (NumWordBits - 1);
 		const WordType Word = Words.Back();
@@ -139,7 +141,7 @@ public:
 		{
 			Words.Pop();
 		}
-		ASSERT(Validation == EContainerValidation::Disabled || NumWordsFor(Count) == Words.Num());
+		AUTORTFM_ASSERT(Validation == EContainerValidation::Disabled || NumWordsFor(Count) == Words.Num());
 		return Value;
 	}
 
@@ -159,7 +161,7 @@ public:
 	// Returns a reference to the bit with the given index in the stack.
 	FBitRef operator[](size_t Index)
 	{
-		ASSERT(Validation == EContainerValidation::Disabled || Index < Count);
+		AUTORTFM_ASSERT(Validation == EContainerValidation::Disabled || Index < Count);
 		WordType& Word = Words[Index / NumWordBits];
 		WordType Mask = static_cast<WordType>(1) << (Index & (NumWordBits-1));
 		return FBitRef{Word, Mask};
@@ -169,7 +171,7 @@ public:
 	// Returns true if the bit with the given index is set.
 	bool operator[](size_t Index) const
 	{
-		ASSERT(Validation == EContainerValidation::Disabled || Index < Count);
+		AUTORTFM_ASSERT(Validation == EContainerValidation::Disabled || Index < Count);
 		WordType Word = Words[Index / NumWordBits];
 		WordType Mask = static_cast<WordType>(1) << (Index & (NumWordBits-1));
 		return (Word & Mask) != 0;
@@ -186,3 +188,5 @@ private:
 };
 
 }
+
+#endif // (defined(__AUTORTFM) && __AUTORTFM)
diff --git a/Engine/Source/Runtime/Core/Private/AutoRTFM/BlockAllocator.h b/Engine/Source/Runtime/Core/Private/AutoRTFM/BlockAllocator.h
--- a/Engine/Source/Runtime/Core/Private/AutoRTFM/BlockAllocator.h
+++ b/Engine/Source/Runtime/Core/Private/AutoRTFM/BlockAllocator.h
@@ -2,6 +2,8 @@
 
 #pragma once
 
+#if (defined(__AUTORTFM) && __AUTORTFM)
+
 #include "AutoRTFM/Utils.h"
 #include "ExternAPI.h"
 #include <cstddef>
@@ -48,7 +50,7 @@ public:
 	// Alignment must be a power of two and no larger than DataAlignment.
 	inline void* Allocate(size_t Size, size_t Alignment)
 	{
-		ASSERT(Alignment <= DataAlignment);
+		AUTORTFM_ASSERT(Alignment <= DataAlignment);
 
 		if (void* Allocation = Tail->TryAllocate(Size, Alignment))
 		{
@@ -61,7 +63,7 @@ public:
 		Tail = NewBlock;
 
 		void* Allocation = NewBlock->TryAllocate(Size, Alignment);
-		ASSERT(Allocation);
+		AUTORTFM_ASSERT(Allocation);
 		return Allocation;
 	}
 
@@ -123,3 +125,5 @@ private:
 };
 
 }
+
+#endif // (defined(__AUTORTFM) && __AUTORTFM)
diff --git a/Engine/Source/Runtime/Core/Private/AutoRTFM/CallNest.h b/Engine/Source/Runtime/Core/Private/AutoRTFM/CallNest.h
--- a/Engine/Source/Runtime/Core/Private/AutoRTFM/CallNest.h
+++ b/Engine/Source/Runtime/Core/Private/AutoRTFM/CallNest.h
@@ -3,6 +3,7 @@
 #pragma once
 
 #if (defined(__AUTORTFM) && __AUTORTFM)
+
 #include "LongJump.h"
 
 namespace AutoRTFM
@@ -26,6 +27,6 @@ public:
 	void Try(const TTryFunctor& TryFunctor);
 };
 
-
 } // namespace AutoRTFM
-#endif
\ No newline at end of file
+
+#endif // (defined(__AUTORTFM) && __AUTORTFM)
diff --git a/Engine/Source/Runtime/Core/Private/AutoRTFM/CallNestInlines.h b/Engine/Source/Runtime/Core/Private/AutoRTFM/CallNestInlines.h
--- a/Engine/Source/Runtime/Core/Private/AutoRTFM/CallNestInlines.h
+++ b/Engine/Source/Runtime/Core/Private/AutoRTFM/CallNestInlines.h
@@ -3,6 +3,7 @@
 #pragma once
 
 #if (defined(__AUTORTFM) && __AUTORTFM)
+
 #include "Context.h"
 #include "CallNest.h"
 
@@ -16,14 +17,15 @@ void FCallNest::Try(const TTryFunctor& TryFunctor)
 		[&]()
 		{
 			TryFunctor();
-			ASSERT(Context->GetStatus() == EContextStatus::OnTrack);
+			AUTORTFM_ASSERT(Context->GetStatus() == EContextStatus::OnTrack);
 		},
 		[&]()
 		{
-			ASSERT(Context->GetStatus() != EContextStatus::Idle);
-			ASSERT(Context->GetStatus() != EContextStatus::OnTrack);
+			AUTORTFM_ASSERT(Context->GetStatus() != EContextStatus::Idle);
+			AUTORTFM_ASSERT(Context->GetStatus() != EContextStatus::OnTrack);
 		});
 }
 
 } // namespace AutoRTFM
-#endif
\ No newline at end of file
+
+#endif // (defined(__AUTORTFM) && __AUTORTFM)
diff --git a/Engine/Source/Runtime/Core/Private/AutoRTFM/CompilerInterface.cpp b/Engine/Source/Runtime/Core/Private/AutoRTFM/CompilerInterface.cpp
--- a/Engine/Source/Runtime/Core/Private/AutoRTFM/CompilerInterface.cpp
+++ b/Engine/Source/Runtime/Core/Private/AutoRTFM/CompilerInterface.cpp
@@ -3,6 +3,7 @@
 #if (defined(__AUTORTFM) && __AUTORTFM)
 
 #include "AutoRTFM/AutoRTFM.h"
+#include "ExternAPI.h"
 #include "ContextInlines.h"
 #include "FunctionMapInlines.h"
 #include "Memcpy.h"
@@ -135,62 +136,20 @@ extern "C" UE_AUTORTFM_API void autortfm_memset(void* Dst, int Value, size_t Siz
     Memset(Dst, Value, Size, Context);
 }
 
-extern "C" UE_AUTORTFM_API FORCENOINLINE void autortfm_report_error(const UE_AUTORTFM_TCHAR* const Message)
+extern "C" UE_AUTORTFM_API void autortfm_unreachable(const char* Message)
 {
-	FContext* const Context = FContext::Get();
-
-	UE_CLOG(!Context->IsTransactional(), LogAutoRTFM, Fatal, TEXT("The function `autortfm_report_error` was called from outside a transact."));
-
-	const ForTheRuntime::EAutoRTFMInternalAbortActionState InternalAbortAction = ForTheRuntime::GetInternalAbortAction();
-
-    if (Message)
-    {
-		if (InternalAbortAction == ForTheRuntime::EAutoRTFMInternalAbortActionState::Crash)
-		{
-			UE_LOG(LogAutoRTFM, Fatal, TEXT("Transaction failing because of internal issue '%s'."), Message);
-		}
-		else
-		{
-			ensureMsgf(!ForTheRuntime::GetEnsureOnInternalAbort(), TEXT("Transaction failing because of internal issue '%s'."), Message);
-		}
-    }
-    else
-    {
-		if (InternalAbortAction == ForTheRuntime::EAutoRTFMInternalAbortActionState::Crash)
-		{
-			UE_LOG(LogAutoRTFM, Fatal, TEXT("Transaction failing because of internal issue."));
-		}
-		else
-		{
-			ensureMsgf(!ForTheRuntime::GetEnsureOnInternalAbort(), TEXT("Transaction failing because of internal issue."));
-		}
-	}
-
-	if (InternalAbortAction == ForTheRuntime::EAutoRTFMInternalAbortActionState::Abort)
-	{
-		Context->AbortByLanguageAndThrow();
-	}
-	else
-	{
-		TTask<void()> Task = [] { ensure(ForTheRuntime::SetAutoRTFMRuntime(ForTheRuntime::AutoRTFM_ForcedDisabled)); };
-		Context->AbortTransactionWithPostAbortCallback(EContextStatus::AbortedByCascadingRetry, MoveTemp(Task));
-	}
-}
-
-extern "C" UE_AUTORTFM_API void autortfm_unreachable(const UE_AUTORTFM_TCHAR* const Message)
-{
-	autortfm_report_error(Message);
+	ReportError("AutoRTFM Unreachable: %s", Message);
 }
 
-extern "C" UE_AUTORTFM_API void autortfm_llvm_fail(const char* const Message)
+extern "C" UE_AUTORTFM_API void autortfm_llvm_fail(const char* Message)
 {
 	if (Message)
 	{
-		autortfm_report_error(ANSI_TO_TCHAR(Message));
+		ReportError("AutoRTFM LLVM Failure: %s", Message);
 	}
 	else
 	{
-		autortfm_report_error(nullptr);
+		ReportError("AutoRTFM LLVM Failure");
 	}
 }
 
@@ -198,11 +157,11 @@ extern "C" UE_AUTORTFM_API void autortfm_llvm_missing_function()
 {
 	if (ForTheRuntime::GetInternalAbortAction() == ForTheRuntime::EAutoRTFMInternalAbortActionState::Crash)
 	{
-		UE_LOG(LogAutoRTFM, Fatal, TEXT("Transaction failing because of missing function."));
+		AUTORTFM_FATAL("Transaction failing because of missing function");
 	}
 	else
 	{
-		ensureMsgf(!ForTheRuntime::GetEnsureOnInternalAbort(), TEXT("Transaction failing because of missing function."));
+		AUTORTFM_ENSURE_MSG(!ForTheRuntime::GetEnsureOnInternalAbort(), "Transaction failing because of missing function");
 	}
 
 	FContext* Context = FContext::Get();
@@ -218,11 +177,11 @@ extern "C" UE_AUTORTFM_API void autortfm_llvm_error(const char* Message)
 {
 	if (Message)
 	{
-		UE_LOG(LogAutoRTFM, Fatal, TEXT("Transaction failing because of LLVM issue '%s'."), ANSI_TO_TCHAR(Message));
+		AUTORTFM_FATAL("Transaction failing because of LLVM issue '%s'", Message);
 	}
 	else
 	{
-		UE_LOG(LogAutoRTFM, Fatal, TEXT("Transaction failing because of LLVM issue."));
+		AUTORTFM_FATAL("Transaction failing because of LLVM issue");
 	}
 }
 
diff --git a/Engine/Source/Runtime/Core/Private/AutoRTFM/ContainerValidation.h b/Engine/Source/Runtime/Core/Private/AutoRTFM/ContainerValidation.h
--- a/Engine/Source/Runtime/Core/Private/AutoRTFM/ContainerValidation.h
+++ b/Engine/Source/Runtime/Core/Private/AutoRTFM/ContainerValidation.h
@@ -2,6 +2,8 @@
 
 #pragma once
 
+#if (defined(__AUTORTFM) && __AUTORTFM)
+
 namespace AutoRTFM
 {
 
@@ -14,3 +16,5 @@ enum class EContainerValidation
 };
 
 }
+
+#endif // (defined(__AUTORTFM) && __AUTORTFM)
diff --git a/Engine/Source/Runtime/Core/Private/AutoRTFM/Context.cpp b/Engine/Source/Runtime/Core/Private/AutoRTFM/Context.cpp
--- a/Engine/Source/Runtime/Core/Private/AutoRTFM/Context.cpp
+++ b/Engine/Source/Runtime/Core/Private/AutoRTFM/Context.cpp
@@ -86,16 +86,16 @@ void FContext::StartTransaction(EMemoryValidationLevel MemoryValidationLevel)
 	}
 
 	NumDeferredTransactions += 1;
-	ensureMsgf(CurrentTransaction, TEXT("FContext::StartTransaction() can only be called within a scoped transaction"));
-	ASSERT(Status == EContextStatus::OnTrack);
+	AUTORTFM_ENSURE_MSG(CurrentTransaction, "FContext::StartTransaction() can only be called within a scoped transaction");
+	AUTORTFM_ASSERT(Status == EContextStatus::OnTrack);
 
 	GAutoRTFMMetrics.NumTransactionsStarted++;
 }
 
 void FContext::StartNonDeferredTransaction(EMemoryValidationLevel MemoryValidationLevel)
 {
-	ASSERT(GetNumDeferredTransactions() == 0);
-	ensureMsgf(CurrentTransaction, TEXT("FContext::StartNonDeferredTransaction() can only be called within a scoped transaction"));
+	AUTORTFM_ASSERT(GetNumDeferredTransactions() == 0);
+	AUTORTFM_ENSURE_MSG(CurrentTransaction, "FContext::StartNonDeferredTransaction() can only be called within a scoped transaction");
 
 	FTransaction* NewTransaction = PushTransaction(
 		/* Closed */ false,
@@ -104,7 +104,7 @@ void FContext::StartNonDeferredTransaction(EMemoryValidationLevel MemoryValidati
 		/* MemoryValidationLevel */ MemoryValidationLevel);
 
 	// This form of transaction is always ultimately within a scoped Transact 
-	ASSERT(Status == EContextStatus::OnTrack);
+	AUTORTFM_ASSERT(Status == EContextStatus::OnTrack);
 
 	GAutoRTFMMetrics.NumTransactionsStarted++;
 }
@@ -113,7 +113,7 @@ ETransactionResult FContext::CommitTransaction()
 {
 	constexpr bool bVerbose = false;
 
-	ASSERT(Status == EContextStatus::OnTrack);
+	AUTORTFM_ASSERT(Status == EContextStatus::OnTrack);
 
 	ETransactionResult Result = ETransactionResult::Committed;
 
@@ -125,7 +125,7 @@ ETransactionResult FContext::CommitTransaction()
 	else
 	{
 		// Scoped transactions commit on return, so committing explicitly isn't allowed
-		ASSERT(CurrentTransaction->IsScopedTransaction() == false);
+		AUTORTFM_ASSERT(CurrentTransaction->IsScopedTransaction() == false);
 
 		if (CurrentTransaction->IsNested())
 		{
@@ -133,9 +133,9 @@ ETransactionResult FContext::CommitTransaction()
 		}
 		else
 		{
-			UE_LOG(LogAutoRTFM, Verbose, TEXT("About to commit; my state is:"));
+			AUTORTFM_VERBOSE("About to commit; my state is:");
 			DumpState();
-			UE_LOG(LogAutoRTFM, Verbose, TEXT("Committing..."));
+			AUTORTFM_VERBOSE("Committing...");
 
 			if (AttemptToCommitTransaction(CurrentTransaction))
 			{
@@ -143,9 +143,9 @@ ETransactionResult FContext::CommitTransaction()
 			}
 			else
 			{
-				UE_LOG(LogAutoRTFM, Verbose, TEXT("Commit failed!"));
-				ASSERT(Status != EContextStatus::OnTrack);
-				ASSERT(Status != EContextStatus::Idle);
+				AUTORTFM_VERBOSE("Commit failed!");
+				AUTORTFM_ASSERT(Status != EContextStatus::OnTrack);
+				AUTORTFM_ASSERT(Status != EContextStatus::Idle);
 			}
 		}
 
@@ -162,8 +162,8 @@ ETransactionResult FContext::RollbackTransaction(EContextStatus NewStatus)
 {
 	GAutoRTFMMetrics.NumTransactionsAborted++;
 
-	ASSERT(Status == EContextStatus::OnTrack);
-	ASSERT(NewStatus == EContextStatus::AbortedByRequest ||	   
+	AUTORTFM_ASSERT(Status == EContextStatus::OnTrack);
+	AUTORTFM_ASSERT(NewStatus == EContextStatus::AbortedByRequest ||	   
 		   NewStatus == EContextStatus::AbortedByCascadingAbort ||
 		   NewStatus == EContextStatus::AbortedByCascadingRetry);
 
@@ -173,7 +173,7 @@ ETransactionResult FContext::RollbackTransaction(EContextStatus NewStatus)
 	{
 		if (ForTheRuntime::GetInternalAbortAction() != ForTheRuntime::EAutoRTFMInternalAbortActionState::Abort)
 		{
-			EXPECT(0 == NumAbortAndRetries);
+			AUTORTFM_EXPECT(0 == NumAbortAndRetries);
 			NumAbortAndRetries += 1;
 		}
 	}
@@ -187,7 +187,7 @@ ETransactionResult FContext::RollbackTransaction(EContextStatus NewStatus)
 	}
 	else
 	{
-		ASSERT(nullptr != CurrentTransaction);
+		AUTORTFM_ASSERT(nullptr != CurrentTransaction);
 
 		// Sort out how aborts work
 		CurrentTransaction->AbortWithoutThrowing();
@@ -220,7 +220,7 @@ void FContext::AbortTransaction(EContextStatus NewStatus)
 void FContext::AbortTransactionWithPostAbortCallback(EContextStatus NewStatus, TTask<void()>&& InCallback)
 {
 	// The callback parameter is only honored by the cascading aborts.
-	ASSERT(NewStatus == EContextStatus::AbortedByCascadingAbort ||
+	AUTORTFM_ASSERT(NewStatus == EContextStatus::AbortedByCascadingAbort ||
 		   NewStatus == EContextStatus::AbortedByCascadingRetry);
 
 	// We must explicitly copy the passed-in callback here, because the original may have been
@@ -232,7 +232,7 @@ void FContext::AbortTransactionWithPostAbortCallback(EContextStatus NewStatus, T
 	}
 	else if (InCallback.IsSet())
 	{
-		UE_LOG(LogAutoRTFM, Fatal, TEXT("Only one post-abort callback at a time is supported."));
+		AUTORTFM_FATAL("Only one post-abort callback at a time is supported");
 	}
 
 	AbortTransaction(NewStatus);
@@ -256,8 +256,8 @@ EContextStatus FContext::CallClosedNest(void (*ClosedFunction)(void* Arg), void*
 	TScopedGuard<void*> ClosedStackAddressGuard(ClosedStackAddress, &ClosedStackAddressGuard);
 
 	FTransaction* const Transaction = GetCurrentTransaction();
-	ASSERT(Transaction != nullptr);
-	ASSERT(FTransaction::EState::OpenActive == Transaction->State());
+	AUTORTFM_ASSERT(Transaction != nullptr);
+	AUTORTFM_ASSERT(FTransaction::EState::OpenActive == Transaction->State());
 	EMemoryValidationLevel PreviousValidationLevel = Transaction->MemoryValidationLevel();
 	const void* PreviousOpenReturnAddress = Transaction->OpenReturnAddress();
 	Transaction->SetClosedActive();
@@ -278,8 +278,8 @@ EContextStatus FContext::CallClosedNest(void (*ClosedFunction)(void* Arg), void*
 
 void FContext::PushCallNest(FCallNest* NewCallNest)
 {
-	ASSERT(NewCallNest != nullptr);
-	ASSERT(NewCallNest->Parent == nullptr);
+	AUTORTFM_ASSERT(NewCallNest != nullptr);
+	AUTORTFM_ASSERT(NewCallNest->Parent == nullptr);
 
 	NewCallNest->Parent = CurrentNest;
 	CurrentNest = NewCallNest;
@@ -287,7 +287,7 @@ void FContext::PushCallNest(FCallNest* NewCallNest)
 
 void FContext::PopCallNest()
 {
-	ASSERT(CurrentNest != nullptr);
+	AUTORTFM_ASSERT(CurrentNest != nullptr);
 	FCallNest* OldCallNest = CurrentNest;
 	CurrentNest = CurrentNest->Parent;
 
@@ -300,11 +300,11 @@ FTransaction* FContext::PushTransaction(
 	FStackRange StackRange,
 	EMemoryValidationLevel MemoryValidationLevel)
 {
-	ASSERT(!GetNumDeferredTransactions());
+	AUTORTFM_ASSERT(!GetNumDeferredTransactions());
 
 	if (CurrentTransaction != nullptr)
 	{
-		ASSERT(CurrentTransaction->IsActive());
+		AUTORTFM_ASSERT(CurrentTransaction->IsActive());
 		CurrentTransaction->SetInactive();
 	}
 
@@ -333,14 +333,14 @@ FTransaction* FContext::PushTransaction(
 
 void FContext::PopTransaction()
 {
-	ASSERT(!GetNumDeferredTransactions());
-	ASSERT(CurrentTransaction != nullptr);
-	ASSERT(CurrentTransaction->IsDone());
+	AUTORTFM_ASSERT(!GetNumDeferredTransactions());
+	AUTORTFM_ASSERT(CurrentTransaction != nullptr);
+	AUTORTFM_ASSERT(CurrentTransaction->IsDone());
 	FTransaction* OldTransaction = CurrentTransaction;
 	CurrentTransaction = CurrentTransaction->GetParent();
 	if (CurrentTransaction != nullptr)
 	{
-		ASSERT(CurrentTransaction->IsInactive());
+		AUTORTFM_ASSERT(CurrentTransaction->IsInactive());
 		CurrentTransaction->SetActive();
 	}
 	TransactionPool.Return(OldTransaction);
@@ -369,12 +369,12 @@ ETransactionResult FContext::ResolveNestedTransaction(FTransaction* NewTransacti
 	if (Status == EContextStatus::OnTrack)
 	{
 		bool bCommitResult = AttemptToCommitTransaction(NewTransaction);
-		ASSERT(bCommitResult);
-		ASSERT(Status == EContextStatus::OnTrack);
+		AUTORTFM_ASSERT(bCommitResult);
+		AUTORTFM_ASSERT(Status == EContextStatus::OnTrack);
 		return ETransactionResult::Committed;
 	}
 
-	ASSERT(NewTransaction->IsDone());
+	AUTORTFM_ASSERT(NewTransaction->IsDone());
 
 	switch (Status)
 	{
@@ -396,7 +396,7 @@ ETransactionResult FContext::Transact(void (*UninstrumentedFunction)(void*), voi
 
 	if constexpr (0 != UE_AUTOSTM)
 	{
-		UE_LOG(LogAutoRTFM, Warning, TEXT("AutoSTM is not implemented in the runtime yet!"));
+		AUTORTFM_WARN("AutoSTM is not implemented in the runtime yet!");
 		return ETransactionResult::AbortedByLanguage;
 	}
 
@@ -410,11 +410,11 @@ ETransactionResult FContext::Transact(void (*UninstrumentedFunction)(void*), voi
     	return ETransactionResult::AbortedByTransactInOnAbort;
     }
     
-    ASSERT(Status == EContextStatus::Idle || Status == EContextStatus::OnTrack);
+    AUTORTFM_ASSERT(Status == EContextStatus::Idle || Status == EContextStatus::OnTrack);
 
     if (!InstrumentedFunction)
     {
-		UE_LOG(LogAutoRTFM, Warning, TEXT("Could not find function in AutoRTFM::FContext::Transact."));
+		AUTORTFM_WARN("Could not find function in AutoRTFM::FContext::Transact");
         return ETransactionResult::AbortedByLanguage;
     }
 
@@ -422,7 +422,7 @@ ETransactionResult FContext::Transact(void (*UninstrumentedFunction)(void*), voi
 	// reason we can't have a "range" of deferred transactions in the middle
 	// of the transaction stack.
 	MaterializeDeferredTransactions();
-	ASSERT(!GetNumDeferredTransactions());
+	AUTORTFM_ASSERT(!GetNumDeferredTransactions());
     
 	FCallNest* NewNest = new FCallNest(this);
 
@@ -432,12 +432,12 @@ ETransactionResult FContext::Transact(void (*UninstrumentedFunction)(void*), voi
 
 	if (!CurrentTransaction)
 	{
-		ASSERT(Status == EContextStatus::Idle);
+		AUTORTFM_ASSERT(Status == EContextStatus::Idle);
 
-		ASSERT(CurrentThreadId == FThreadID::Invalid);
+		AUTORTFM_ASSERT(CurrentThreadId == FThreadID::Invalid);
 		CurrentThreadId = FThreadID::GetCurrent();
 
-		ASSERT(Stack == FStackRange{});
+		AUTORTFM_ASSERT(Stack == FStackRange{});
 
 #if PLATFORM_WINDOWS
 		GetCurrentThreadStackLimits(&Stack.Low, &Stack.High);
@@ -453,9 +453,9 @@ ETransactionResult FContext::Transact(void (*UninstrumentedFunction)(void*), voi
 		pthread_attr_getstack(&Attr, &Stack.Low, &StackSize);
 		Stack.High = static_cast<char*>(Stack.Low) + StackSize;
 #endif
-		ASSERT(Stack.High > Stack.Low);
+		AUTORTFM_ASSERT(Stack.High > Stack.Low);
 
-		ASSERT(Stack.Contains(TransactStackStart));
+		AUTORTFM_ASSERT(Stack.Contains(TransactStackStart));
 
 		FTransaction* NewTransaction = PushTransaction(
 			/* Closed */ true,
@@ -470,17 +470,17 @@ ETransactionResult FContext::Transact(void (*UninstrumentedFunction)(void*), voi
         for (;;)
         {
             Status = EContextStatus::OnTrack;
-            ASSERT(CurrentTransaction->IsFresh());
+            AUTORTFM_ASSERT(CurrentTransaction->IsFresh());
 			CurrentNest->Try([&] () { InstrumentedFunction(Arg); });
-			ASSERT(CurrentTransaction == NewTransaction); // The transaction lambda should have unwound any nested transactions.
-            ASSERT(Status != EContextStatus::Idle);
+			AUTORTFM_ASSERT(CurrentTransaction == NewTransaction); // The transaction lambda should have unwound any nested transactions.
+            AUTORTFM_ASSERT(Status != EContextStatus::Idle);
 
 			switch (Status)
 			{
 			case EContextStatus::OnTrack:
-				UE_LOG(LogAutoRTFM, Verbose, TEXT("About to commit; my state is:"));
+				AUTORTFM_VERBOSE("About to commit; my state is:");
 				DumpState();
-				UE_LOG(LogAutoRTFM, Verbose, TEXT("Committing..."));
+				AUTORTFM_VERBOSE("Committing...");
 
 				if (UNLIKELY(!bTriedToRunOnce && AutoRTFM::ForTheRuntime::ShouldRetryNonNestedTransactions()))
 				{
@@ -502,14 +502,14 @@ ETransactionResult FContext::Transact(void (*UninstrumentedFunction)(void*), voi
 					break;
 				}
 
-				UE_LOG(LogAutoRTFM, Verbose, TEXT("Commit failed!"));
+				AUTORTFM_VERBOSE("Commit failed!");
 
-				ASSERT(Status != EContextStatus::OnTrack);
-				ASSERT(Status != EContextStatus::Idle);
+				AUTORTFM_ASSERT(Status != EContextStatus::OnTrack);
+				AUTORTFM_ASSERT(Status != EContextStatus::Idle);
 				break;
 
 			case EContextStatus::AbortedByRequest:
-				ASSERT(!PostAbortCallback.IsSet());
+				AUTORTFM_ASSERT(!PostAbortCallback.IsSet());
 				Result = ETransactionResult::AbortedByRequest;
 				break;
 
@@ -527,7 +527,7 @@ ETransactionResult FContext::Transact(void (*UninstrumentedFunction)(void*), voi
 					PostAbortCallback();
 					PostAbortCallback.Reset();
 
-					ASSERT(Status == EContextStatus::InPostAbort);
+					AUTORTFM_ASSERT(Status == EContextStatus::InPostAbort);
 					Status = EContextStatus::AbortedByCascadingAbort;
 				}
 
@@ -535,13 +535,13 @@ ETransactionResult FContext::Transact(void (*UninstrumentedFunction)(void*), voi
 				break;
 
 			case EContextStatus::AbortedByCascadingRetry:
-				ASSERT(PostAbortCallback.IsSet());
+				AUTORTFM_ASSERT(PostAbortCallback.IsSet());
 			
 				// Clean up the transaction to get it ready for re-execution.
 				ClearTransactionStatus();
 				CurrentTransaction->Reset();
 
-				ASSERT(Status == EContextStatus::OnTrack);
+				AUTORTFM_ASSERT(Status == EContextStatus::OnTrack);
 				Status = EContextStatus::InPostAbort;
 
 				// Call the post-abort callback to do whatever work the user
@@ -549,7 +549,7 @@ ETransactionResult FContext::Transact(void (*UninstrumentedFunction)(void*), voi
 				PostAbortCallback();
 				PostAbortCallback.Reset();
 				
-				ASSERT(Status == EContextStatus::InPostAbort);
+				AUTORTFM_ASSERT(Status == EContextStatus::InPostAbort);
 				Status = EContextStatus::OnTrack;
 
 				// Then get rolling!
@@ -586,19 +586,19 @@ ETransactionResult FContext::Transact(void (*UninstrumentedFunction)(void*), voi
 		PopTransaction();
 		ClearTransactionStatus();
 
-		ASSERT(CurrentNest == nullptr);
-		ASSERT(CurrentTransaction == nullptr);
+		AUTORTFM_ASSERT(CurrentNest == nullptr);
+		AUTORTFM_ASSERT(CurrentTransaction == nullptr);
 
         Reset();
 	}
     else
     {
 		// This transaction is within another transaction
-		ASSERT(Status == EContextStatus::OnTrack);
+		AUTORTFM_ASSERT(Status == EContextStatus::OnTrack);
 
-		ASSERT(CurrentThreadId == FThreadID::GetCurrent());
+		AUTORTFM_ASSERT(CurrentThreadId == FThreadID::GetCurrent());
 
-		ASSERT(Stack.Contains(TransactStackStart));
+		AUTORTFM_ASSERT(Stack.Contains(TransactStackStart));
 
 		FTransaction* NewTransaction = PushTransaction(
 			/* Closed */ true,
@@ -613,7 +613,7 @@ ETransactionResult FContext::Transact(void (*UninstrumentedFunction)(void*), voi
 		for (;;)
 		{
 			CurrentNest->Try([&]() { InstrumentedFunction(Arg); });
-			ASSERT(CurrentTransaction == NewTransaction);
+			AUTORTFM_ASSERT(CurrentTransaction == NewTransaction);
 
 			if (Status == EContextStatus::OnTrack)
 			{
@@ -639,8 +639,8 @@ ETransactionResult FContext::Transact(void (*UninstrumentedFunction)(void*), voi
 		PopCallNest();
 		PopTransaction();
 
-		ASSERT(CurrentNest != nullptr);
-		ASSERT(CurrentTransaction != nullptr);
+		AUTORTFM_ASSERT(CurrentNest != nullptr);
+		AUTORTFM_ASSERT(CurrentTransaction != nullptr);
 
 		// Cascading aborts should cause all transactions to abort!
 		switch (Result)
@@ -660,7 +660,7 @@ ETransactionResult FContext::Transact(void (*UninstrumentedFunction)(void*), voi
 
 void FContext::AbortByRequestAndThrow()
 {
-    ASSERT(Status == EContextStatus::OnTrack);
+    AUTORTFM_ASSERT(Status == EContextStatus::OnTrack);
 	GAutoRTFMMetrics.NumTransactionsAbortedByRequest++;
     Status = EContextStatus::AbortedByRequest;
     GetCurrentTransaction()->AbortAndThrow();
@@ -668,7 +668,7 @@ void FContext::AbortByRequestAndThrow()
 
 void FContext::AbortByRequestWithoutThrowing()
 {
-	ASSERT(Status == EContextStatus::OnTrack);
+	AUTORTFM_ASSERT(Status == EContextStatus::OnTrack);
 	GAutoRTFMMetrics.NumTransactionsAbortedByRequest++;
 	Status = EContextStatus::AbortedByRequest;
 	GetCurrentTransaction()->AbortWithoutThrowing();
@@ -676,7 +676,7 @@ void FContext::AbortByRequestWithoutThrowing()
 
 void FContext::AbortByLanguageAndThrow()
 {
-    ASSERT(Status == EContextStatus::OnTrack);
+    AUTORTFM_ASSERT(Status == EContextStatus::OnTrack);
 	GAutoRTFMMetrics.NumTransactionsAbortedByLanguage++;
     Status = EContextStatus::AbortedByLanguage;
     GetCurrentTransaction()->AbortAndThrow();
@@ -684,7 +684,7 @@ void FContext::AbortByLanguageAndThrow()
 
 void FContext::Reset()
 {
-	ASSERT(CurrentThreadId == FThreadID::GetCurrent() || CurrentThreadId == FThreadID::Invalid);
+	AUTORTFM_ASSERT(CurrentThreadId == FThreadID::GetCurrent() || CurrentThreadId == FThreadID::Invalid);
 
 	CurrentThreadId = FThreadID::Invalid;
 	Stack = {};
@@ -703,7 +703,7 @@ void FContext::Throw()
 
 void FContext::DumpState() const
 {
-	UE_LOG(LogAutoRTFM, Verbose, TEXT("Context at %p"), this);
+	AUTORTFM_VERBOSE("Context at %p", this);
 }
 
 } // namespace AutoRTFM
diff --git a/Engine/Source/Runtime/Core/Private/AutoRTFM/Context.h b/Engine/Source/Runtime/Core/Private/AutoRTFM/Context.h
--- a/Engine/Source/Runtime/Core/Private/AutoRTFM/Context.h
+++ b/Engine/Source/Runtime/Core/Private/AutoRTFM/Context.h
@@ -2,7 +2,8 @@
 
 #pragma once
 
-#include "AutoRTFM/AutoRTFM.h"
+#if (defined(__AUTORTFM) && __AUTORTFM)
+
 #include "IntrusivePool.h"
 #include "StackRange.h"
 #include "OpenHashThrottler.h"
@@ -80,7 +81,7 @@ public:
 	}
 	inline FTransaction* GetCurrentTransactionAssumingNoMaterialization()
 	{
-		ASSERT(!MustMaterializeDeferredTransactions());
+		AUTORTFM_ASSERT(!MustMaterializeDeferredTransactions());
 		return CurrentTransaction;
 	}
 
@@ -107,7 +108,7 @@ public:
 
 		if (Status == EContextStatus::OnTrack)
 		{
-			ASSERT(0 == StackLocalInitializerDepth);
+			AUTORTFM_ASSERT(0 == StackLocalInitializerDepth);
 			Status = EContextStatus::InStaticLocalInitializer;
 			StackLocalInitializerDepth++;
 		}
@@ -124,7 +125,7 @@ public:
 			return;
 		}
 
-		ASSERT(Status != EContextStatus::OnTrack);
+		AUTORTFM_ASSERT(Status != EContextStatus::OnTrack);
 
 		if (Status == EContextStatus::InStaticLocalInitializer)
 		{
@@ -184,3 +185,5 @@ private:
 };
 
 } // namespace AutoRTFM
+
+#endif // (defined(__AUTORTFM) && __AUTORTFM)
diff --git a/Engine/Source/Runtime/Core/Private/AutoRTFM/ContextInlines.h b/Engine/Source/Runtime/Core/Private/AutoRTFM/ContextInlines.h
--- a/Engine/Source/Runtime/Core/Private/AutoRTFM/ContextInlines.h
+++ b/Engine/Source/Runtime/Core/Private/AutoRTFM/ContextInlines.h
@@ -2,6 +2,8 @@
 
 #pragma once
 
+#if (defined(__AUTORTFM) && __AUTORTFM)
+
 #include "Context.h"
 #include "TransactionInlines.h"
 
@@ -55,7 +57,7 @@ UE_AUTORTFM_FORCEINLINE void FContext::DidFree(void* LogicalAddress)
 
 UE_AUTORTFM_FORCEINLINE bool FContext::AttemptToCommitTransaction(FTransaction* const Transaction)
 {
-    ASSERT(EContextStatus::OnTrack == Status);
+    AUTORTFM_ASSERT(EContextStatus::OnTrack == Status);
 
     Status = EContextStatus::Committing;
 
@@ -70,3 +72,5 @@ UE_AUTORTFM_FORCEINLINE bool FContext::AttemptToCommitTransaction(FTransaction*
 }
 
 } // namespace AutoRTFM
+
+#endif // (defined(__AUTORTFM) && __AUTORTFM)
diff --git a/Engine/Source/Runtime/Core/Private/AutoRTFM/ContextStatus.h b/Engine/Source/Runtime/Core/Private/AutoRTFM/ContextStatus.h
--- a/Engine/Source/Runtime/Core/Private/AutoRTFM/ContextStatus.h
+++ b/Engine/Source/Runtime/Core/Private/AutoRTFM/ContextStatus.h
@@ -2,7 +2,9 @@
 
 #pragma once
 
-#include "AutoRTFM/AutoRTFMConstants.h"
+#if (defined(__AUTORTFM) && __AUTORTFM)
+
+#include "AutoRTFM/AutoRTFM.h"
 #include "Utils.h"
 
 namespace AutoRTFM
@@ -26,7 +28,17 @@ inline const char* GetContextStatusName(EContextStatus Status)
         return "AbortedByRequest";
     case EContextStatus::Committing:
         return "Committing";
+    case EContextStatus::AbortedByCascadingAbort:
+        return "AbortedByCascadingAbort";
+    case EContextStatus::AbortedByCascadingRetry:
+        return "AbortedByCascadingRetry";
+    case EContextStatus::InStaticLocalInitializer:
+        return "InStaticLocalInitializer";
+    case EContextStatus::InPostAbort:
+        return "InPostAbort";
     }
 }
 
 } // namespace AutoRTFM
+
+#endif // (defined(__AUTORTFM) && __AUTORTFM)
diff --git a/Engine/Source/Runtime/Core/Private/AutoRTFM/ExternAPI.cpp b/Engine/Source/Runtime/Core/Private/AutoRTFM/ExternAPI.cpp
--- a/Engine/Source/Runtime/Core/Private/AutoRTFM/ExternAPI.cpp
+++ b/Engine/Source/Runtime/Core/Private/AutoRTFM/ExternAPI.cpp
@@ -4,10 +4,76 @@
 
 #include "ExternAPI.h"
 
+#include <memory>
+
 namespace AutoRTFM
 {
 
-autortfm_extern_api GExternAPI{};
+namespace
+{
+
+// A default implementation for GExternAPI.Log().
+// Most of GExternAPI is intentionally null, as we do not intend to be running
+// AutoRTFM logic before initialization - however there are a number of
+// AutoRTFM configuration setters that can be called before initialization
+// which have logic that can (potentially) log. Logging is unlikely to be called
+// before initialization, but a default printf implementation is preferable to
+// crashing due to a nullptr.
+void DefaultLog(const char* File, int Line, autortfm_log_severity Severity, const char* Format, va_list Args)
+{
+	static constexpr size_t InlineBufferLength = 256;
+	char InlineBuffer[InlineBufferLength];
+
+	va_list Args2;
+	va_copy(Args2, Args);
+	int Count = vsnprintf(InlineBuffer, InlineBufferLength, Format, Args);
+	auto AllocatedBuffer = (Count >= InlineBufferLength) ? std::unique_ptr<char[]>{new char[Count+1]} : nullptr;
+	if (AllocatedBuffer)
+	{
+		vsnprintf(AllocatedBuffer.get(), Count+1, Format, Args2);
+		va_end(Args2);
+	}
+
+	char* const Buffer = AllocatedBuffer ? AllocatedBuffer.get() : InlineBuffer;
+
+	switch (Severity)
+	{
+	case autortfm_log_verbose:
+		fprintf(stdout, "AutoRTFM %s:%d [VERBOSE]: %s", File, Line, Buffer);
+		break;
+	case autortfm_log_info:
+		fprintf(stdout, "AutoRTFM %s:%d [INFO]: %s", File, Line, Buffer);
+		break;
+	case autortfm_log_warn:
+		fprintf(stdout, "AutoRTFM %s:%d [WARN]: %s", File, Line, Buffer);
+		break;
+	case autortfm_log_error:
+		fprintf(stderr, "AutoRTFM %s:%d [ERROR]: %s", File, Line, Buffer);
+		break;
+	case autortfm_log_fatal:
+		fprintf(stderr, "AutoRTFM %s:%d [FATAL]: %s", File, Line, Buffer);
+		__builtin_trap();
+		break;
+	}
+}
+
+}  // anonymous namespace
+
+autortfm_extern_api GExternAPI
+{
+	/* Allocate */ nullptr,
+	/* Reallocate */ nullptr,
+	/* AllocateZeroed */ nullptr,
+	/* Free */ nullptr,
+	/* Log */ DefaultLog,
+	/* EnsureFailure */ nullptr,
+	/* IsLogActive */ nullptr,
+	/* OnRuntimeEnabledChanged */ nullptr,
+	/* OnRetryTransactionsChanged */ nullptr,
+	/* OnMemoryValidationLevelChanged */ nullptr,
+	/* OnMemoryValidationThrottlingChanged */ nullptr,
+	/* OnMemoryValidationStatisticsChanged */ nullptr,
+};
 
 } // namespace AutoRTFM
 
diff --git a/Engine/Source/Runtime/Core/Private/AutoRTFM/ExternAPI.h b/Engine/Source/Runtime/Core/Private/AutoRTFM/ExternAPI.h
--- a/Engine/Source/Runtime/Core/Private/AutoRTFM/ExternAPI.h
+++ b/Engine/Source/Runtime/Core/Private/AutoRTFM/ExternAPI.h
@@ -2,6 +2,8 @@
 
 #pragma once
 
+#if (defined(__AUTORTFM) && __AUTORTFM)
+
 #include "AutoRTFM/AutoRTFM.h"
 
 namespace AutoRTFM
@@ -31,3 +33,4 @@ inline void Free(void* Pointer)
 
 } // namespace AutoRTFM
 
+#endif // (defined(__AUTORTFM) && __AUTORTFM)
diff --git a/Engine/Source/Runtime/Core/Private/AutoRTFM/FastLock.h b/Engine/Source/Runtime/Core/Private/AutoRTFM/FastLock.h
--- a/Engine/Source/Runtime/Core/Private/AutoRTFM/FastLock.h
+++ b/Engine/Source/Runtime/Core/Private/AutoRTFM/FastLock.h
@@ -2,6 +2,8 @@
 
 #pragma once
 
+#if (defined(__AUTORTFM) && __AUTORTFM)
+
 #include <atomic>
 
 namespace AutoRTFM
@@ -33,3 +35,5 @@ private:
 };
 
 } // namespace AutoRTFM
+
+#endif // (defined(__AUTORTFM) && __AUTORTFM)
diff --git a/Engine/Source/Runtime/Core/Private/AutoRTFM/FunctionMap.cpp b/Engine/Source/Runtime/Core/Private/AutoRTFM/FunctionMap.cpp
--- a/Engine/Source/Runtime/Core/Private/AutoRTFM/FunctionMap.cpp
+++ b/Engine/Source/Runtime/Core/Private/AutoRTFM/FunctionMap.cpp
@@ -1,6 +1,7 @@
 // Copyright Epic Games, Inc. All Rights Reserved.
 
 #if (defined(__AUTORTFM) && __AUTORTFM)
+
 #include "FunctionMap.h"
 #include "FunctionMapInlines.h"
 #include "ContextInlines.h"
@@ -45,7 +46,7 @@ static inline void* FollowRelocation(void* Function)
 
 FFunctionMap* AllocateFunctionMap(size_t Size)
 {
-    ASSERT(Size && !(Size & (Size - 1)));
+    AUTORTFM_ASSERT(Size && !(Size & (Size - 1)));
     size_t TotalSize = offsetof(FFunctionMap, Entries) + sizeof(FFunctionEntry) * Size;
     FFunctionMap* Result = static_cast<FFunctionMap*>(malloc(TotalSize));
     Result->KeyCount = 0;
@@ -140,18 +141,15 @@ void* FunctionMapLookupTryFollowingRelocation(void* Function, const char* Where)
 
 void* FunctionMapReportError(void* OldFunction, const char* Where)
 {
-	FString Message;
-
+    std::string FunctionDescription = GetFunctionDescription(OldFunction);
 	if (Where)
 	{
-		Message = FString::Printf(TEXT("Could not find function %p '%s' where '%s'."), OldFunction, *GetFunctionDescription(OldFunction), ANSI_TO_TCHAR(Where));
+		ReportError("Could not find function %p '%s' where '%s'.", OldFunction, FunctionDescription.c_str(), Where);
 	}
 	else
 	{
-		Message = FString::Printf(TEXT("Could not find function %p '%s'."), OldFunction, *GetFunctionDescription(OldFunction));
+		ReportError("Could not find function %p '%s'.", OldFunction, FunctionDescription.c_str());
 	}
-
-	autortfm_report_error(*Message);
 	return nullptr;
 }
 
diff --git a/Engine/Source/Runtime/Core/Private/AutoRTFM/FunctionMap.h b/Engine/Source/Runtime/Core/Private/AutoRTFM/FunctionMap.h
--- a/Engine/Source/Runtime/Core/Private/AutoRTFM/FunctionMap.h
+++ b/Engine/Source/Runtime/Core/Private/AutoRTFM/FunctionMap.h
@@ -2,6 +2,8 @@
 
 #pragma once
 
+#if (defined(__AUTORTFM) && __AUTORTFM)
+
 #include "GlobalData.h"
 
 namespace AutoRTFM
@@ -74,3 +76,5 @@ auto FunctionMapTryLookup(TReturnType (*Function)(TParameterTypes...)) -> TRetur
 }
 
 } // namespace AutoRTFM
+
+#endif // (defined(__AUTORTFM) && __AUTORTFM)
diff --git a/Engine/Source/Runtime/Core/Private/AutoRTFM/FunctionMapInlines.h b/Engine/Source/Runtime/Core/Private/AutoRTFM/FunctionMapInlines.h
--- a/Engine/Source/Runtime/Core/Private/AutoRTFM/FunctionMapInlines.h
+++ b/Engine/Source/Runtime/Core/Private/AutoRTFM/FunctionMapInlines.h
@@ -2,13 +2,13 @@
 
 #pragma once
 
+#if (defined(__AUTORTFM) && __AUTORTFM)
+
 #include "AutoRTFM/AutoRTFMConstants.h"
 #include "Context.h"
 #include "FunctionMap.h"
 #include "Utils.h"
 
-#include "Containers/StringConv.h"
-
 namespace AutoRTFM
 {
 
@@ -58,3 +58,4 @@ auto FunctionMapLookup(TReturnType (*Function)(TParameterTypes...), const char*
 
 } // namespace AutoRTFM
 
+#endif // (defined(__AUTORTFM) && __AUTORTFM)
diff --git a/Engine/Source/Runtime/Core/Private/AutoRTFM/GlobalData.cpp b/Engine/Source/Runtime/Core/Private/AutoRTFM/GlobalData.cpp
--- a/Engine/Source/Runtime/Core/Private/AutoRTFM/GlobalData.cpp
+++ b/Engine/Source/Runtime/Core/Private/AutoRTFM/GlobalData.cpp
@@ -1,8 +1,9 @@
 // Copyright Epic Games, Inc. All Rights Reserved.
 
 #if (defined(__AUTORTFM) && __AUTORTFM)
+
 #include "GlobalData.h"
-#include "Context.h"
+#include "Utils.h"
 #include <inttypes.h>
 
 #ifdef _MSC_VER
@@ -45,7 +46,7 @@ void InitializeGlobalDataIfNecessary()
         if (char* EnvString = getenv(EnvName))
         {
             sscanf(EnvString, "%p", &GlobalData);
-            ASSERT(GlobalData);
+            AUTORTFM_ASSERT(GlobalData);
         }
         else
         {
diff --git a/Engine/Source/Runtime/Core/Private/AutoRTFM/GlobalData.h b/Engine/Source/Runtime/Core/Private/AutoRTFM/GlobalData.h
--- a/Engine/Source/Runtime/Core/Private/AutoRTFM/GlobalData.h
+++ b/Engine/Source/Runtime/Core/Private/AutoRTFM/GlobalData.h
@@ -2,6 +2,8 @@
 
 #pragma once
 
+#if (defined(__AUTORTFM) && __AUTORTFM)
+
 #include "FastLock.h"
 
 namespace AutoRTFM
@@ -21,3 +23,4 @@ void InitializeGlobalDataIfNecessary();
 
 } // namespace AutoRTFM
 
+#endif // (defined(__AUTORTFM) && __AUTORTFM)
diff --git a/Engine/Source/Runtime/Core/Private/AutoRTFM/HashMap.h b/Engine/Source/Runtime/Core/Private/AutoRTFM/HashMap.h
--- a/Engine/Source/Runtime/Core/Private/AutoRTFM/HashMap.h
+++ b/Engine/Source/Runtime/Core/Private/AutoRTFM/HashMap.h
@@ -2,6 +2,8 @@
 
 #pragma once
 
+#if (defined(__AUTORTFM) && __AUTORTFM)
+
 #include "StlAllocator.h"
 #include <functional>
 #include <type_traits>
@@ -195,3 +197,5 @@ private:
 };
 
 }  // AutoRTFM
+
+#endif // (defined(__AUTORTFM) && __AUTORTFM)
diff --git a/Engine/Source/Runtime/Core/Private/AutoRTFM/HitSet.h b/Engine/Source/Runtime/Core/Private/AutoRTFM/HitSet.h
--- a/Engine/Source/Runtime/Core/Private/AutoRTFM/HitSet.h
+++ b/Engine/Source/Runtime/Core/Private/AutoRTFM/HitSet.h
@@ -2,6 +2,8 @@
 
 #pragma once
 
+#if (defined(__AUTORTFM) && __AUTORTFM)
+
 #include "ExternAPI.h"
 #include "Utils.h"
 #include <cstdint>
@@ -144,7 +146,7 @@ namespace AutoRTFM
 			SixtyFourMinusLogCapacity -= 1;
 
 			// Check that we haven't overflowed our capacity!
-			ASSERT(0 != SixtyFourMinusLogCapacity);
+			AUTORTFM_ASSERT(0 != SixtyFourMinusLogCapacity);
 		}
 
 		bool Resize()
@@ -168,7 +170,7 @@ namespace AutoRTFM
 				IncreaseCapacity();
 
 				Payload = static_cast<uintptr_t*>(AutoRTFM::AllocateZeroed(Capacity() * sizeof(uintptr_t), alignof(uintptr_t)));
-				ASSERT(nullptr != Payload);
+				AUTORTFM_ASSERT(nullptr != Payload);
 
 				// Now we need to rehash and reinsert all the items.
 				for (size_t I = 0; I < OldCapacity; I++)
@@ -188,7 +190,7 @@ namespace AutoRTFM
 					}
 					else
 					{
-						ASSERT(EInsertResult::Inserted == Result);
+						AUTORTFM_ASSERT(EInsertResult::Inserted == Result);
 					}
 				}
 
@@ -204,7 +206,7 @@ namespace AutoRTFM
 				break;
 			}
 
-			ASSERT(OldCount == Count);
+			AUTORTFM_ASSERT(OldCount == Count);
 
 			if (OldPayload != SmallPayload)
 			{
@@ -286,3 +288,5 @@ namespace AutoRTFM
 		}
 	};
 }
+
+#endif // (defined(__AUTORTFM) && __AUTORTFM)
diff --git a/Engine/Source/Runtime/Core/Private/AutoRTFM/IntervalTree.h b/Engine/Source/Runtime/Core/Private/AutoRTFM/IntervalTree.h
--- a/Engine/Source/Runtime/Core/Private/AutoRTFM/IntervalTree.h
+++ b/Engine/Source/Runtime/Core/Private/AutoRTFM/IntervalTree.h
@@ -2,6 +2,8 @@
 
 #pragma once
 
+#if (defined(__AUTORTFM) && __AUTORTFM)
+
 #include "BitStack.h"
 #include "ContainerValidation.h"
 #include "Stack.h"
@@ -162,7 +164,7 @@ private:
     {
         if (UNLIKELY(IntervalTreeNodeIndexNone == Root))
         {
-            ASSERT(NodeRanges.IsEmpty());
+            AUTORTFM_ASSERT(NodeRanges.IsEmpty());
 			Root = AddNode(MoveTemp(NewRange));
             return true;
         }
@@ -182,7 +184,7 @@ private:
             {
                 if (UNLIKELY(NewRange.End == Range.Start))
                 {
-                    ASSERT(NewRange.Start < Range.Start);
+                    AUTORTFM_ASSERT(NewRange.Start < Range.Start);
 
 					// We can just modify the existing node in place.
 					NodeRanges[Current].Start = NewRange.Start;
@@ -202,7 +204,7 @@ private:
             {
                 if (UNLIKELY(NewRange.Start == Range.End))
                 {
-                    ASSERT(NewRange.End > Range.End);
+                    AUTORTFM_ASSERT(NewRange.End > Range.End);
 
 					// We can just modify the existing node in place.
 					NodeRanges[Current].End = NewRange.End;
@@ -219,7 +221,7 @@ private:
                 Current = NodeRights[Current];
             }
 
-            ASSERT(Root != Current);
+            AUTORTFM_ASSERT(Root != Current);
         }
 
         auto IsBlack = [this](FIntervalTreeNodeIndex Index)
@@ -257,7 +259,7 @@ private:
 
             if (!IsBlack(Uncle))
             {
-                ASSERT(IsBlack(GrandParent));
+                AUTORTFM_ASSERT(IsBlack(GrandParent));
                 NodeColors[Parent] = EColor::Black;
                 NodeColors[Uncle] = EColor::Black;
 
@@ -387,8 +389,8 @@ private:
 		// allocations within the checker itself!
 		if constexpr (bExtraDebugging)
 		{
-			ASSERT(IntervalTreeNodeIndexNone != Index);
-			ASSERT(Index < NodeRanges.Num());
+			AUTORTFM_ASSERT(IntervalTreeNodeIndexNone != Index);
+			AUTORTFM_ASSERT(Index < NodeRanges.Num());
 
 			const FIntervalTreeNodeIndex Parent = NodeParents[Index];
 			const FIntervalTreeNodeIndex Left = NodeLefts[Index];
@@ -396,16 +398,16 @@ private:
 
 			if (IntervalTreeNodeIndexNone == Parent)
 			{
-				ASSERT(Root == Index);
+				AUTORTFM_ASSERT(Root == Index);
 			}
 			else
 			{
-				ASSERT(NodeColors[Parent] || NodeColors[Index]);
-				ASSERT((NodeLefts[Parent] == Index) ^ (NodeRights[Parent] == Index));
+				AUTORTFM_ASSERT(NodeColors[Parent] || NodeColors[Index]);
+				AUTORTFM_ASSERT((NodeLefts[Parent] == Index) ^ (NodeRights[Parent] == Index));
 			}
 
-			ASSERT(Left != Index);
-			ASSERT(Right != Index);
+			AUTORTFM_ASSERT(Left != Index);
+			AUTORTFM_ASSERT(Right != Index);
 
 			if (IntervalTreeNodeIndexNone != Left)
 			{
@@ -421,3 +423,5 @@ private:
 };
 
 } // namespace AutoRTFM
+
+#endif // (defined(__AUTORTFM) && __AUTORTFM)
diff --git a/Engine/Source/Runtime/Core/Private/AutoRTFM/IntrusivePool.h b/Engine/Source/Runtime/Core/Private/AutoRTFM/IntrusivePool.h
--- a/Engine/Source/Runtime/Core/Private/AutoRTFM/IntrusivePool.h
+++ b/Engine/Source/Runtime/Core/Private/AutoRTFM/IntrusivePool.h
@@ -2,6 +2,8 @@
 
 #pragma once
 
+#if (defined(__AUTORTFM) && __AUTORTFM)
+
 #include "BlockAllocator.h"
 #include <cstddef>
 #include <utility>
@@ -52,7 +54,7 @@ public:
 			Item = new (Memory) FItem(std::forward<ArgumentTypes>(Arguments)...);
 		}
 
-		ASSERT(NumInUse < NumAllocated);
+		AUTORTFM_ASSERT(NumInUse < NumAllocated);
 		NumInUse++;
 
 		return Item;
@@ -68,7 +70,7 @@ public:
 		*Item->GetIntrusiveAddress() = FreeList;
 		FreeList = Item;
 
-		ASSERT(NumInUse > 0);
+		AUTORTFM_ASSERT(NumInUse > 0);
 		NumInUse--;
 	}
     
@@ -92,3 +94,5 @@ private:
 };
 
 }  // namespace AutoRTFM
+
+#endif // (defined(__AUTORTFM) && __AUTORTFM)
diff --git a/Engine/Source/Runtime/Core/Private/AutoRTFM/LockGuard.h b/Engine/Source/Runtime/Core/Private/AutoRTFM/LockGuard.h
--- a/Engine/Source/Runtime/Core/Private/AutoRTFM/LockGuard.h
+++ b/Engine/Source/Runtime/Core/Private/AutoRTFM/LockGuard.h
@@ -2,6 +2,8 @@
 
 #pragma once
 
+#if (defined(__AUTORTFM) && __AUTORTFM)
+
 namespace AutoRTFM
 {
 
@@ -25,3 +27,5 @@ private:
 };
 
 } // namespace AutoRTFM
+
+#endif // (defined(__AUTORTFM) && __AUTORTFM)
diff --git a/Engine/Source/Runtime/Core/Private/AutoRTFM/LongJump.h b/Engine/Source/Runtime/Core/Private/AutoRTFM/LongJump.h
--- a/Engine/Source/Runtime/Core/Private/AutoRTFM/LongJump.h
+++ b/Engine/Source/Runtime/Core/Private/AutoRTFM/LongJump.h
@@ -3,7 +3,7 @@
 #pragma once
 
 #if (defined(__AUTORTFM) && __AUTORTFM)
-#include "AutoRTFM/AutoRTFM.h"
+
 #include "Utils.h"
 #include <setjmp.h>
 
@@ -34,7 +34,7 @@ private:
 template<typename TTryFunctor, typename TCatchFunctor>
 void FLongJump::TryCatch(const TTryFunctor& TryFunctor, const TCatchFunctor& CatchFunctor)
 {
-    ASSERT(!bIsSet);
+    AUTORTFM_ASSERT(!bIsSet);
 #if PLATFORM_WINDOWS
 	if (!setjmp(JmpBuf))
 #else
@@ -47,7 +47,7 @@ void FLongJump::TryCatch(const TTryFunctor& TryFunctor, const TCatchFunctor& Cat
     }
     else
     {
-        ASSERT(bIsSet);
+        AUTORTFM_ASSERT(bIsSet);
         bIsSet = false;
         CatchFunctor();
     }
@@ -55,7 +55,7 @@ void FLongJump::TryCatch(const TTryFunctor& TryFunctor, const TCatchFunctor& Cat
 
 inline void FLongJump::Throw()
 {
-    ASSERT(bIsSet);
+    AUTORTFM_ASSERT(bIsSet);
 #if PLATFORM_WINDOWS
 	longjmp(JmpBuf, 1);
 #else
@@ -64,4 +64,6 @@ inline void FLongJump::Throw()
 }
 
 } // namespace AutoRTFM
-#endif
\ No newline at end of file
+
+#endif // (defined(__AUTORTFM) && __AUTORTFM)
+
diff --git a/Engine/Source/Runtime/Core/Private/AutoRTFM/Memcpy.cpp b/Engine/Source/Runtime/Core/Private/AutoRTFM/Memcpy.cpp
--- a/Engine/Source/Runtime/Core/Private/AutoRTFM/Memcpy.cpp
+++ b/Engine/Source/Runtime/Core/Private/AutoRTFM/Memcpy.cpp
@@ -1,35 +1,37 @@
 // Copyright Epic Games, Inc. All Rights Reserved.
 
 #if (defined(__AUTORTFM) && __AUTORTFM)
+
 #include "Memcpy.h"
 #include "ContextInlines.h"
+#include "Utils.h"
 
 namespace AutoRTFM
 {
 
 void* MemcpyToNew(void* InDst, const void* InSrc, size_t Size, FContext* Context)
 {
-	UE_LOG(LogAutoRTFM, Verbose, TEXT("MemcpyToNew(%p, %p, %zu)"), InDst, InSrc, Size);
+	AUTORTFM_VERBOSE("MemcpyToNew(%p, %p, %zu)", InDst, InSrc, Size);
     return memcpy(InDst, InSrc, Size);
 }
 
 void* Memcpy(void* InDst, const void* InSrc, size_t Size, FContext* Context)
 {
-	UE_LOG(LogAutoRTFM, Verbose, TEXT("Memcpy(%p, %p, %zu)"), InDst, InSrc, Size);
+	AUTORTFM_VERBOSE("Memcpy(%p, %p, %zu)", InDst, InSrc, Size);
     Context->RecordWrite(InDst, Size);
     return memcpy(InDst, InSrc, Size);
 }
 
 void* Memmove(void* InDst, const void* InSrc, size_t Size, FContext* Context)
 {
-	UE_LOG(LogAutoRTFM, Verbose, TEXT("Memmove(%p, %p, %zu)"), InDst, InSrc, Size);
+	AUTORTFM_VERBOSE("Memmove(%p, %p, %zu)", InDst, InSrc, Size);
 	Context->RecordWrite(InDst, Size);
 	return memmove(InDst, InSrc, Size);
 }
 
 void* Memset(void* InDst, int Value, size_t Size, FContext* Context)
 {
-	UE_LOG(LogAutoRTFM, Verbose, TEXT("Memset(%p, %d, %zu)"), InDst, Value, Size);
+	AUTORTFM_VERBOSE("Memset(%p, %d, %zu)", InDst, Value, Size);
 	Context->RecordWrite(InDst, Size);
 	return memset(InDst, Value, Size);
 }
diff --git a/Engine/Source/Runtime/Core/Private/AutoRTFM/Memcpy.h b/Engine/Source/Runtime/Core/Private/AutoRTFM/Memcpy.h
--- a/Engine/Source/Runtime/Core/Private/AutoRTFM/Memcpy.h
+++ b/Engine/Source/Runtime/Core/Private/AutoRTFM/Memcpy.h
@@ -2,6 +2,8 @@
 
 #pragma once
 
+#if (defined(__AUTORTFM) && __AUTORTFM)
+
 #include <stdlib.h>
 
 namespace AutoRTFM
@@ -16,3 +18,4 @@ void* Memset(void* Dst, int Value, size_t Size, FContext* Context);
 
 } // namespace AutoRTFM
 
+#endif // (defined(__AUTORTFM) && __AUTORTFM)
diff --git a/Engine/Source/Runtime/Core/Private/AutoRTFM/OpenHashThrottler.cpp b/Engine/Source/Runtime/Core/Private/AutoRTFM/OpenHashThrottler.cpp
--- a/Engine/Source/Runtime/Core/Private/AutoRTFM/OpenHashThrottler.cpp
+++ b/Engine/Source/Runtime/Core/Private/AutoRTFM/OpenHashThrottler.cpp
@@ -233,14 +233,14 @@ void FOpenHashThrottler::UpdateLogStats()
 	}
 
 	const double TimeSinceLastLog = LogStats.TimeSinceLastReset;
-	UE_LOG(LogAutoRTFM, Display,
-		TEXT("Transaction Hash Statistics\n")
-		TEXT("-----------------------------\n")
-		/* a */ TEXT("%f / %f seconds spent hashing (%.1f%%, target: %.1f%%)\n")
-		/* b */ TEXT(" * %zu hash calls (avr: %.1f calls/second)\n")
-		/* c */ TEXT(" * %zu / %zu open validations skipped (%.1f%%)\n")
-		/* d */ TEXT(" * %zu bytes hashed (avr: %.1f bytes/call, %.1f bytes/second)\n")
-		/* e */ TEXT(" * %zu write records (avr: %.1f records/call)\n"),
+	AUTORTFM_LOG(
+		"Transaction Hash Statistics\n"
+		"-----------------------------\n"
+		/* a */ "%f / %f seconds spent hashing (%.1f%%, target: %.1f%%)\n"
+		/* b */ " * %zu hash calls (avr: %.1f calls/second)\n"
+		/* c */ " * %zu / %zu open validations skipped (%.1f%%)\n"
+		/* d */ " * %zu bytes hashed (avr: %.1f bytes/call, %.1f bytes/second)\n"
+		/* e */ " * %zu write records (avr: %.1f records/call)\n",
 		/* a.0 */ LogStats.TimeSpentHashing,
 		/* a.1 */ TimeSinceLastLog,
 		/* a.2 */ 100 * LogStats.TimeSpentHashing / TimeSinceLastLog,
diff --git a/Engine/Source/Runtime/Core/Private/AutoRTFM/OpenHashThrottler.h b/Engine/Source/Runtime/Core/Private/AutoRTFM/OpenHashThrottler.h
--- a/Engine/Source/Runtime/Core/Private/AutoRTFM/OpenHashThrottler.h
+++ b/Engine/Source/Runtime/Core/Private/AutoRTFM/OpenHashThrottler.h
@@ -2,6 +2,8 @@
 
 #pragma once
 
+#if (defined(__AUTORTFM) && __AUTORTFM)
+
 #include "HashMap.h"
 
 namespace AutoRTFM
@@ -126,3 +128,5 @@ private:
 };
 
 } // namespace AutoRTFM
+
+#endif // (defined(__AUTORTFM) && __AUTORTFM)
diff --git a/Engine/Source/Runtime/Core/Private/AutoRTFM/Pool.h b/Engine/Source/Runtime/Core/Private/AutoRTFM/Pool.h
--- a/Engine/Source/Runtime/Core/Private/AutoRTFM/Pool.h
+++ b/Engine/Source/Runtime/Core/Private/AutoRTFM/Pool.h
@@ -2,6 +2,8 @@
 
 #pragma once
 
+#if (defined(__AUTORTFM) && __AUTORTFM)
+
 #include "BlockAllocator.h"
 #include <cstddef>
 #include <utility>
@@ -46,7 +48,7 @@ public:
 			NumAllocated++;
 		}
 
-		ASSERT(NumInUse < NumAllocated);
+		AUTORTFM_ASSERT(NumInUse < NumAllocated);
 		NumInUse++;
 
 		return new (Memory) FItem(std::forward<ArgumentTypes>(Arguments)...);
@@ -64,7 +66,7 @@ public:
 		FreeEntry->Next = FreeList;
 		FreeList = FreeEntry;
 
-		ASSERT(NumInUse > 0);
+		AUTORTFM_ASSERT(NumInUse > 0);
 		NumInUse--;
 	}
 
@@ -72,7 +74,7 @@ public:
 	// No items can be in use when calling.
 	void Reset(bool bIgnoreNonReturned = false)
 	{
-		ASSERT(bIgnoreNonReturned || NumInUse == 0);
+		AUTORTFM_ASSERT(bIgnoreNonReturned || NumInUse == 0);
 		Allocator.FreeAll();
 		NumInUse = 0;
 		NumAllocated = 0;
@@ -113,3 +115,5 @@ private:
 };
 
 }  // namespace AutoRTFM
+
+#endif // (defined(__AUTORTFM) && __AUTORTFM)
diff --git a/Engine/Source/Runtime/Core/Private/AutoRTFM/ScopedGuard.h b/Engine/Source/Runtime/Core/Private/AutoRTFM/ScopedGuard.h
--- a/Engine/Source/Runtime/Core/Private/AutoRTFM/ScopedGuard.h
+++ b/Engine/Source/Runtime/Core/Private/AutoRTFM/ScopedGuard.h
@@ -2,6 +2,8 @@
 
 #pragma once
 
+#if (defined(__AUTORTFM) && __AUTORTFM)
+
 namespace AutoRTFM
 {
 
@@ -27,3 +29,5 @@ private:
 };
 
 } // namespace AutoRTFM
+
+#endif // (defined(__AUTORTFM) && __AUTORTFM)
diff --git a/Engine/Source/Runtime/Core/Private/AutoRTFM/Stack.h b/Engine/Source/Runtime/Core/Private/AutoRTFM/Stack.h
--- a/Engine/Source/Runtime/Core/Private/AutoRTFM/Stack.h
+++ b/Engine/Source/Runtime/Core/Private/AutoRTFM/Stack.h
@@ -2,6 +2,8 @@
 
 #pragma once
 
+#if (defined(__AUTORTFM) && __AUTORTFM)
+
 #include "ExternAPI.h"
 #include "Utils.h"
 #include "ContainerValidation.h"
@@ -140,7 +142,7 @@ public:
 	// Removes the last item on the stack.
 	inline void Pop()
 	{
-		ASSERT(Validation == EContainerValidation::Disabled || Count > 0);
+		AUTORTFM_ASSERT(Validation == EContainerValidation::Disabled || Count > 0);
 		--Count;
 		Data[Count].~T();
 	}
@@ -190,37 +192,37 @@ public:
 
 	T& operator[](size_t Index)
 	{
-		ASSERT(Validation == EContainerValidation::Disabled || Index < Count);
+		AUTORTFM_ASSERT(Validation == EContainerValidation::Disabled || Index < Count);
 		return Data[Index];
 	}
 
 	const T& operator[](size_t Index) const
 	{
-		ASSERT(Validation == EContainerValidation::Disabled || Index < Count);
+		AUTORTFM_ASSERT(Validation == EContainerValidation::Disabled || Index < Count);
 		return Data[Index];
 	}
 
 	T& Front()
 	{
-		ASSERT(Validation == EContainerValidation::Disabled || Count > 0);
+		AUTORTFM_ASSERT(Validation == EContainerValidation::Disabled || Count > 0);
 		return Data[0];
 	}
 
 	const T& Front() const
 	{
-		ASSERT(Validation == EContainerValidation::Disabled || Count > 0);
+		AUTORTFM_ASSERT(Validation == EContainerValidation::Disabled || Count > 0);
 		return Data[0];
 	}
 
 	T& Back()
 	{
-		ASSERT(Validation == EContainerValidation::Disabled || Count > 0);
+		AUTORTFM_ASSERT(Validation == EContainerValidation::Disabled || Count > 0);
 		return Data[Count - 1];
 	}
 
 	const T& Back() const
 	{
-		ASSERT(Validation == EContainerValidation::Disabled || Count > 0);
+		AUTORTFM_ASSERT(Validation == EContainerValidation::Disabled || Count > 0);
 		return Data[Count - 1];
 	}
 
@@ -234,7 +236,7 @@ private:
 	// This stack must be empty before calling.
 	void CopyFrom(const TStack& Other)
 	{
-		ASSERT(IsEmpty());
+		AUTORTFM_ASSERT(IsEmpty());
 		Reserve(Other.Count);
 		if constexpr (std::is_trivial_v<T>)
 		{
@@ -254,7 +256,7 @@ private:
 	// This stack must be empty before calling.
 	void MoveFrom(TStack&& Other)
 	{
-		ASSERT(IsEmpty());
+		AUTORTFM_ASSERT(IsEmpty());
 		if (Other.Data == Other.InlineData())
 		{
 			if constexpr (std::is_trivial_v<T>)
@@ -299,3 +301,5 @@ private:
 };
 
 }
+
+#endif // (defined(__AUTORTFM) && __AUTORTFM)
diff --git a/Engine/Source/Runtime/Core/Private/AutoRTFM/StackRange.h b/Engine/Source/Runtime/Core/Private/AutoRTFM/StackRange.h
--- a/Engine/Source/Runtime/Core/Private/AutoRTFM/StackRange.h
+++ b/Engine/Source/Runtime/Core/Private/AutoRTFM/StackRange.h
@@ -2,6 +2,8 @@
 
 #pragma once
 
+#if (defined(__AUTORTFM) && __AUTORTFM)
+
 namespace AutoRTFM
 {
 
@@ -22,3 +24,5 @@ struct FStackRange
 };
 
 }  // namespace AutoRTFM
+
+#endif // (defined(__AUTORTFM) && __AUTORTFM)
diff --git a/Engine/Source/Runtime/Core/Private/AutoRTFM/Stats.cpp b/Engine/Source/Runtime/Core/Private/AutoRTFM/Stats.cpp
--- a/Engine/Source/Runtime/Core/Private/AutoRTFM/Stats.cpp
+++ b/Engine/Source/Runtime/Core/Private/AutoRTFM/Stats.cpp
@@ -1,123 +1,126 @@
 // Copyright Epic Games, Inc. All Rights Reserved.
 
 #if (defined(__AUTORTFM) && __AUTORTFM)
+
 #include "Stats.h"
+#include "Utils.h"
+#include <inttypes.h>
 
 template<> void AutoRTFM::FStats::Report<AutoRTFM::EStatsKind::Transaction>(const uint64_t Data) const
 {
-	UE_LOG(LogAutoRTFM, Display, TEXT("  Total transactions:        %11u"), Data);
+	AUTORTFM_LOG("  Total transactions:        %" PRIu64, Data);
 }
 
 template<> void AutoRTFM::FStats::Report<AutoRTFM::EStatsKind::Commit>(const uint64_t Data) const
 {
-	UE_LOG(LogAutoRTFM, Display, TEXT("  Total commits:             %11u"), Data);
+	AUTORTFM_LOG("  Total commits:             %" PRIu64, Data);
 }
 
 template<> void AutoRTFM::FStats::Report<AutoRTFM::EStatsKind::Abort>(const uint64_t Data) const
 {
-	UE_LOG(LogAutoRTFM, Display, TEXT("  Total aborts:              %11u"), Data);
+	AUTORTFM_LOG("  Total aborts:              %" PRIu64, Data);
 }
 
 template<> void AutoRTFM::FStats::Report<AutoRTFM::EStatsKind::AverageTransactionDepth>(const uint64_t Data) const
 {
 	const uint64_t TotalTransactions = Datas[static_cast<size_t>(AutoRTFM::EStatsKind::Transaction)];
-	UE_LOG(LogAutoRTFM, Display, TEXT("  Average transaction depth: %15.3f"), (static_cast<double>(Data) / static_cast<double>(TotalTransactions)));
+	AUTORTFM_LOG("  Average transaction depth: %15.3f", (static_cast<double>(Data) / static_cast<double>(TotalTransactions)));
 }
 
 template<> void AutoRTFM::FStats::Report<AutoRTFM::EStatsKind::MaximumTransactionDepth>(const uint64_t Data) const
 {
-    UE_LOG(LogAutoRTFM, Display, TEXT("  Maximum transaction depth: %11u"), Data);
+    AUTORTFM_LOG("  Maximum transaction depth: %" PRIu64, Data);
 }
 
 template<> void AutoRTFM::FStats::Report<AutoRTFM::EStatsKind::AverageWriteLogEntries>(const uint64_t Data) const
 {
     const uint64_t TotalTransactions = Datas[static_cast<size_t>(AutoRTFM::EStatsKind::Transaction)];
-    UE_LOG(LogAutoRTFM, Display, TEXT("  Average write log entries: %15.3f"), (static_cast<double>(Data) / static_cast<double>(TotalTransactions)));
+    AUTORTFM_LOG("  Average write log entries: %15.3f", (static_cast<double>(Data) / static_cast<double>(TotalTransactions)));
 }
 
 template<> void AutoRTFM::FStats::Report<AutoRTFM::EStatsKind::MaximumWriteLogEntries>(const uint64_t Data) const
 {
-    UE_LOG(LogAutoRTFM, Display, TEXT("  Maximum write log entries: %11u"), Data);
+    AUTORTFM_LOG("  Maximum write log entries: %" PRIu64, Data);
 }
 
 template<> void AutoRTFM::FStats::Report<AutoRTFM::EStatsKind::AverageWriteLogBytes>(const uint64_t Data) const
 {
     const uint64_t TotalTransactions = Datas[static_cast<size_t>(AutoRTFM::EStatsKind::Transaction)];
-    UE_LOG(LogAutoRTFM, Display, TEXT("  Average write log bytes:   %15.3f"), (static_cast<double>(Data) / static_cast<double>(TotalTransactions)));
+    AUTORTFM_LOG("  Average write log bytes:   %15.3f", (static_cast<double>(Data) / static_cast<double>(TotalTransactions)));
 }
 
 template<> void AutoRTFM::FStats::Report<AutoRTFM::EStatsKind::MaximumWriteLogBytes>(const uint64_t Data) const
 {
-    UE_LOG(LogAutoRTFM, Display, TEXT("  Maximum write log bytes:   %11u"), Data);
+    AUTORTFM_LOG("  Maximum write log bytes:   %" PRIu64, Data);
 }
 
 template<> void AutoRTFM::FStats::Report<AutoRTFM::EStatsKind::HitSetHit>(const uint64_t Data) const
 {
-    UE_LOG(LogAutoRTFM, Display, TEXT("  HitSet hits:               %11u"), Data);
+    AUTORTFM_LOG("  HitSet hits:               %" PRIu64, Data);
 }
 
 template<> void AutoRTFM::FStats::Report<AutoRTFM::EStatsKind::HitSetMiss>(const uint64_t Data) const
 {
-    UE_LOG(LogAutoRTFM, Display, TEXT("  HitSet misses:             %11u"), Data);
+    AUTORTFM_LOG("  HitSet misses:             %" PRIu64, Data);
 }
 
 template<> void AutoRTFM::FStats::Report<AutoRTFM::EStatsKind::HitSetSkippedBecauseOfStackLocalMemory>(const uint64_t Data) const
 {
-    UE_LOG(LogAutoRTFM, Display, TEXT("  HitSet skip (stack local): %11u"), Data);
+    AUTORTFM_LOG("  HitSet skip (stack local): %" PRIu64, Data);
 }
 
 template<> void AutoRTFM::FStats::Report<AutoRTFM::EStatsKind::AverageCommitTasks>(const uint64_t Data) const
 {
     const uint64_t TotalTransactions = Datas[static_cast<size_t>(AutoRTFM::EStatsKind::Transaction)];
-    UE_LOG(LogAutoRTFM, Display, TEXT("  Average commit tasks:      %15.3f"), (static_cast<double>(Data) / static_cast<double>(TotalTransactions)));
+    AUTORTFM_LOG("  Average commit tasks:      %15.3f", (static_cast<double>(Data) / static_cast<double>(TotalTransactions)));
 }
 
 template<> void AutoRTFM::FStats::Report<AutoRTFM::EStatsKind::MaximumCommitTasks>(const uint64_t Data) const
 {
-    UE_LOG(LogAutoRTFM, Display, TEXT("  Maximum commit tasks:      %11u"), Data);
+    AUTORTFM_LOG("  Maximum commit tasks:      %" PRIu64, Data);
 }
 
 template<> void AutoRTFM::FStats::Report<AutoRTFM::EStatsKind::AverageAbortTasks>(const uint64_t Data) const
 {
     const uint64_t TotalTransactions = Datas[static_cast<size_t>(AutoRTFM::EStatsKind::Transaction)];
-    UE_LOG(LogAutoRTFM, Display, TEXT("  Average abort tasks:       %15.3f"), (static_cast<double>(Data) / static_cast<double>(TotalTransactions)));
+    AUTORTFM_LOG("  Average abort tasks:       %15.3f", (static_cast<double>(Data) / static_cast<double>(TotalTransactions)));
 }
 
 template<> void AutoRTFM::FStats::Report<AutoRTFM::EStatsKind::MaximumAbortTasks>(const uint64_t Data) const
 {
-    UE_LOG(LogAutoRTFM, Display, TEXT("  Maximum abort tasks:       %11u"), Data);
+    AUTORTFM_LOG("  Maximum abort tasks:       %" PRIu64, Data);
 }
 
 template<> void AutoRTFM::FStats::Report<AutoRTFM::EStatsKind::NewMemoryTrackerHit>(const uint64_t Data) const
 {
-    UE_LOG(LogAutoRTFM, Display, TEXT("  New memory hits:           %11u"), Data);
+    AUTORTFM_LOG("  New memory hits:           %" PRIu64, Data);
 }
 
 template<> void AutoRTFM::FStats::Report<AutoRTFM::EStatsKind::NewMemoryTrackerMiss>(const uint64_t Data) const
 {
-    UE_LOG(LogAutoRTFM, Display, TEXT("  New memory misses:         %11u"), Data);
+    AUTORTFM_LOG("  New memory misses:         %" PRIu64, Data);
 }
 
 template<> void AutoRTFM::FStats::Report<AutoRTFM::EStatsKind::AverageHitSetSize>(const uint64_t Data) const
 {
     const uint64_t TotalTransactions = Datas[static_cast<size_t>(AutoRTFM::EStatsKind::Transaction)];
-    UE_LOG(LogAutoRTFM, Display, TEXT("  Average hit set size:      %15.3f"), (static_cast<double>(Data) / static_cast<double>(TotalTransactions)));
+    AUTORTFM_LOG("  Average hit set size:      %15.3f", (static_cast<double>(Data) / static_cast<double>(TotalTransactions)));
 }
 
 template<> void AutoRTFM::FStats::Report<AutoRTFM::EStatsKind::AverageHitSetCapacity>(const uint64_t Data) const
 {
     const uint64_t TotalTransactions = Datas[static_cast<size_t>(AutoRTFM::EStatsKind::Transaction)];
-    UE_LOG(LogAutoRTFM, Display, TEXT("  Average hit set capacity:  %15.3f"), (static_cast<double>(Data) / static_cast<double>(TotalTransactions)));
+    AUTORTFM_LOG("  Average hit set capacity:  %15.3f", (static_cast<double>(Data) / static_cast<double>(TotalTransactions)));
 
     const uint64_t TotalHitSetSize = Datas[static_cast<size_t>(AutoRTFM::EStatsKind::AverageHitSetSize)];
-    UE_LOG(LogAutoRTFM, Display, TEXT("  Average hit set occupancy: %15.3f"), 100 * (static_cast<double>(TotalHitSetSize) / static_cast<double>(Data)));
+    AUTORTFM_LOG("  Average hit set occupancy: %15.3f", 100 * (static_cast<double>(TotalHitSetSize) / static_cast<double>(Data)));
 }
 
 void AutoRTFM::FStats::Report() const
 {
 	if constexpr (bCollectStats)
 	{
-		UE_LOG(LogAutoRTFM, Display, TEXT("AutoRTFM Statistics:"));
+		AUTORTFM_LOG("AutoRTFM Statistics:");
 
 		for (size_t I = 0; I < static_cast<size_t>(AutoRTFM::EStatsKind::Total); I++)
 		{
diff --git a/Engine/Source/Runtime/Core/Private/AutoRTFM/Stats.h b/Engine/Source/Runtime/Core/Private/AutoRTFM/Stats.h
--- a/Engine/Source/Runtime/Core/Private/AutoRTFM/Stats.h
+++ b/Engine/Source/Runtime/Core/Private/AutoRTFM/Stats.h
@@ -2,6 +2,8 @@
 
 #pragma once
 
+#if (defined(__AUTORTFM) && __AUTORTFM)
+
 #include "Toggles.h"
 #include "Utils.h"
 
@@ -145,7 +147,7 @@ struct FStats final
 			else if constexpr (IsStatsKindMaximum<Kind>())
 			{
 				uint64_t* const Ptr = &Datas[static_cast<size_t>(Kind)];
-				*Ptr = FMath::Max(*Ptr, Data);
+				*Ptr = std::max<uint64_t>(*Ptr, Data);
 			}
 		}
 	}
@@ -179,3 +181,5 @@ private:
 };
 
 } // namespace AutoRTFM
+
+#endif // (defined(__AUTORTFM) && __AUTORTFM)
diff --git a/Engine/Source/Runtime/Core/Private/AutoRTFM/StdLibSupport.cpp b/Engine/Source/Runtime/Core/Private/AutoRTFM/StdLibSupport.cpp
--- a/Engine/Source/Runtime/Core/Private/AutoRTFM/StdLibSupport.cpp
+++ b/Engine/Source/Runtime/Core/Private/AutoRTFM/StdLibSupport.cpp
@@ -6,6 +6,7 @@
 #include "ContextInlines.h"
 #include "Utils.h"
 
+#include <algorithm>
 #include <charconv>
 #include <functional> // note: introduces additional math overloads
 #include "Memcpy.h"
@@ -25,6 +26,13 @@
 #include "Windows/WindowsHWrapper.h"
 #endif
 
+#if PLATFORM_LINUX
+#include <sys/stat.h>
+#include <sys/time.h>
+#include <sys/types.h>
+#include <unistd.h>
+#endif
+
 #if PLATFORM_HAS_ASAN_INCLUDE
 #	if defined(__SANITIZE_ADDRESS__)
 #		define AUTORTFM_ASAN_ENABLED 1
@@ -74,7 +82,7 @@ private:
 
 void ThrowErrorFormatContainsPercentN()
 {
-    UE_LOG(LogAutoRTFM, Warning, TEXT("AutoRTFM does not support format strings containing '%%n'"));
+    AUTORTFM_WARN("AutoRTFM does not support format strings containing '%%n'");
     FContext* Context = FContext::Get();
     Context->AbortByLanguageAndThrow();
 }
@@ -158,7 +166,7 @@ FILE* RTFM___acrt_iob_func(unsigned Index)
         return __acrt_iob_func(Index);
     default:
 	{
-		UE_LOG(LogAutoRTFM, Warning, TEXT("Attempt to get file descriptor %d (not 1 or 2) in __acrt_iob_func."), Index);
+		AUTORTFM_WARN("Attempt to get file descriptor %d (not 1 or 2) in __acrt_iob_func.", Index);
 		FContext* Context = FContext::Get();
         Context->AbortByLanguageAndThrow();
         return NULL;
@@ -199,7 +207,7 @@ int RTFM___stdio_common_vsprintf(
 
         if (Count >= 0)
         {
-            size_t NumBytes = FMath::Min<size_t>(BufferCount, (1 + Count)) * sizeof(char);
+            size_t NumBytes = std::min(BufferCount, static_cast<size_t>(1 + Count)) * sizeof(char);
             FContext* Context = FContext::Get();
             Context->RecordWrite(Buffer, NumBytes);
         }
@@ -228,7 +236,7 @@ int RTFM___stdio_common_vswprintf(
 
         if (Count >= 0)
         {
-            size_t NumBytes = FMath::Min<size_t>(BufferCount, (1 + Count)) * sizeof(wchar_t);
+            size_t NumBytes = std::min(BufferCount, static_cast<size_t>(1 + Count)) * sizeof(wchar_t);
             FContext* Context = FContext::Get();
             Context->RecordWrite(Buffer, NumBytes);
         }
@@ -369,7 +377,7 @@ void* RTFM_realloc(void* Ptr, size_t Size) throw()
 		const size_t OldSize = malloc_usable_size(Ptr);
 #endif
 		FContext* Context = FContext::Get();
-        MemcpyToNew(NewObject, Ptr,  FMath::Min(OldSize, Size), Context);
+        MemcpyToNew(NewObject, Ptr, std::min(OldSize, Size), Context);
         RTFM_free(Ptr);
     }
     return NewObject;
@@ -472,7 +480,7 @@ int RTFM_vsnprintf(char* Str, size_t Size, const char* Format, va_list ArgList)
 
         if (Count >= 0)
         {
-            size_t NumBytes = FMath::Min<size_t>(Size, (1 + Count)) * sizeof(char);
+            size_t NumBytes = std::min(Size, static_cast<size_t>(1 + Count)) * sizeof(char);
             FContext* Context = FContext::Get();
             Context->RecordWrite(Str, NumBytes);
         }
@@ -501,7 +509,7 @@ int RTFM_vswprintf(wchar_t* Str, size_t Size, const wchar_t* Format, va_list Arg
 
         va_end(ArgList2);
 
-        size_t NumChars = FMath::Min<size_t>(Size, (1 + FMath::Max(Count, 0)));
+        size_t NumChars = std::min(Size, static_cast<size_t>(1 + std::max(Count, 0)));
         size_t NumBytes = NumChars * sizeof(wchar_t);
         if (NumBytes >= 0)
         {
diff --git a/Engine/Source/Runtime/Core/Private/AutoRTFM/StlAllocator.h b/Engine/Source/Runtime/Core/Private/AutoRTFM/StlAllocator.h
--- a/Engine/Source/Runtime/Core/Private/AutoRTFM/StlAllocator.h
+++ b/Engine/Source/Runtime/Core/Private/AutoRTFM/StlAllocator.h
@@ -2,7 +2,9 @@
 
 #pragma once
 
-#include "AutoRTFM/ExternAPI.h"
+#if (defined(__AUTORTFM) && __AUTORTFM)
+
+#include "ExternAPI.h"
 #include <cstddef>
 
 namespace AutoRTFM
@@ -47,3 +49,5 @@ bool operator!=(const StlAllocator<T>&, const StlAllocator<U>&)
 }
 
 }  // AutoRTFM
+
+#endif // (defined(__AUTORTFM) && __AUTORTFM)
diff --git a/Engine/Source/Runtime/Core/Private/AutoRTFM/TaskArray.h b/Engine/Source/Runtime/Core/Private/AutoRTFM/TaskArray.h
--- a/Engine/Source/Runtime/Core/Private/AutoRTFM/TaskArray.h
+++ b/Engine/Source/Runtime/Core/Private/AutoRTFM/TaskArray.h
@@ -2,6 +2,8 @@
 
 #pragma once
 
+#if (defined(__AUTORTFM) && __AUTORTFM)
+
 #include "HashMap.h"
 #include "Pool.h"
 #include "Stack.h"
@@ -53,7 +55,7 @@ class TTaskArray
 		// This entry must be unlinked before calling.
 		void Link(FEntry*& ListHead, FEntry*& ListTail)
 		{
-			ASSERT(Next == nullptr && Prev == nullptr);
+			AUTORTFM_ASSERT(Next == nullptr && Prev == nullptr);
 			if (!ListHead)
 			{
 				ListHead = this;
@@ -265,3 +267,5 @@ private:
 };
 
 } // namespace AutoRTFM
+
+#endif // (defined(__AUTORTFM) && __AUTORTFM)
diff --git a/Engine/Source/Runtime/Core/Private/AutoRTFM/ThreadID.cpp b/Engine/Source/Runtime/Core/Private/AutoRTFM/ThreadID.cpp
--- a/Engine/Source/Runtime/Core/Private/AutoRTFM/ThreadID.cpp
+++ b/Engine/Source/Runtime/Core/Private/AutoRTFM/ThreadID.cpp
@@ -1,5 +1,7 @@
 // Copyright Epic Games, Inc. All Rights Reserved.
 
+#if (defined(__AUTORTFM) && __AUTORTFM)
+
 #include "ThreadID.h"
 
 #include <thread>
@@ -16,3 +18,5 @@ FThreadID FThreadID::GetCurrent()
 }
 
 }
+
+#endif // (defined(__AUTORTFM) && __AUTORTFM)
diff --git a/Engine/Source/Runtime/Core/Private/AutoRTFM/ThreadID.h b/Engine/Source/Runtime/Core/Private/AutoRTFM/ThreadID.h
--- a/Engine/Source/Runtime/Core/Private/AutoRTFM/ThreadID.h
+++ b/Engine/Source/Runtime/Core/Private/AutoRTFM/ThreadID.h
@@ -2,6 +2,8 @@
 
 #pragma once
 
+#if (defined(__AUTORTFM) && __AUTORTFM)
+
 #include "AutoRTFM/AutoRTFM.h"
 #include <thread>
 
@@ -35,3 +37,5 @@ struct FThreadID
 };
 
 }
+
+#endif // (defined(__AUTORTFM) && __AUTORTFM)
diff --git a/Engine/Source/Runtime/Core/Private/AutoRTFM/Transaction.cpp b/Engine/Source/Runtime/Core/Private/AutoRTFM/Transaction.cpp
--- a/Engine/Source/Runtime/Core/Private/AutoRTFM/Transaction.cpp
+++ b/Engine/Source/Runtime/Core/Private/AutoRTFM/Transaction.cpp
@@ -8,8 +8,6 @@
 #include "GlobalData.h"
 #include "HAL/Platform.h"
 #include "Hash/xxhash.h"
-#include "Logging/LogMacros.h"
-#include "Logging/LogVerbosity.h"
 #include "Misc/AssertionMacros.h"
 #include "TransactionInlines.h"
 #include "HitSet.h"
@@ -24,7 +22,7 @@ namespace AutoRTFM
 	
 void FTransaction::Resurrect(FContext* const InContext)
 {
-	ASSERT(Context == InContext)
+	AUTORTFM_ASSERT(Context == InContext)
 	Parent = nullptr;
 	StatDepth = 1;
 	RecordedWriteHash = 0;
@@ -62,7 +60,7 @@ FTransaction::~FTransaction()
 
 void FTransaction::Initialize(FTransaction* Parent_, bool bIsStackScoped_, FStackRange StackRange_)
 {
-	ASSERT(CurrentState == EState::Uninitialized);
+	AUTORTFM_ASSERT(CurrentState == EState::Uninitialized);
 
 	Parent = Parent_;
 	bIsStackScoped = bIsStackScoped_;
@@ -80,7 +78,7 @@ void FTransaction::Initialize(FTransaction* Parent_, bool bIsStackScoped_, FStac
 
 void FTransaction::Reset()
 {
-	ASSERT(IsDone());
+	AUTORTFM_ASSERT(IsDone());
 
 	CommitTasks.Reset();
 	AbortTasks.Reset();
@@ -97,7 +95,7 @@ void FTransaction::Reset()
 	DeferredPopAllOnCommitHandlers.Reset();
 	DeferredPopAllOnAbortHandlers.Reset();
 
-	ASSERT(IsFresh());
+	AUTORTFM_ASSERT(IsFresh());
 }
 
 bool FTransaction::IsFresh() const
@@ -116,9 +114,9 @@ bool FTransaction::IsFresh() const
 
 void FTransaction::AbortWithoutThrowing()
 {
-	UE_LOG(LogAutoRTFM, Verbose, TEXT("Aborting '%hs'!"), GetContextStatusName(Context->GetStatus()));
+	AUTORTFM_VERBOSE("Aborting '%hs'!", GetContextStatusName(Context->GetStatus()));
 
-    ASSERT(Context->IsAborting());
+    AUTORTFM_ASSERT(Context->IsAborting());
 
     Stats.Collect<EStatsKind::Abort>();
     CollectStats();
@@ -141,11 +139,11 @@ void FTransaction::AbortWithoutThrowing()
 
     if (IsNested())
     {
-		ASSERT(Parent);
+		AUTORTFM_ASSERT(Parent);
     }
     else
     {
-		ASSERT(Context->IsAborting());
+		AUTORTFM_ASSERT(Context->IsAborting());
     }
 }
 
@@ -157,8 +155,8 @@ void FTransaction::AbortAndThrow()
 
 bool FTransaction::AttemptToCommit()
 {
-    ASSERT(Context->GetStatus() == EContextStatus::Committing);
-    ASSERT(Context->GetCurrentTransaction() == this);
+    AUTORTFM_ASSERT(Context->GetStatus() == EContextStatus::Committing);
+    AUTORTFM_ASSERT(Context->GetCurrentTransaction() == this);
 
     Stats.Collect<EStatsKind::Commit>();
     CollectStats();
@@ -184,54 +182,28 @@ bool FTransaction::AttemptToCommit()
 
 void FTransaction::Undo()
 {
-	UE_LOG(LogAutoRTFM, Verbose, TEXT("Undoing a transaction..."));
-	ASSERT(IsDone());
+	AUTORTFM_VERBOSE("Undoing a transaction...");
+	AUTORTFM_ASSERT(IsDone());
 
 	int VerboseCounter = 0;
 	int Num = WriteLog.Num();
 	for(auto Iter = WriteLog.rbegin(); Iter != WriteLog.rend(); ++Iter)
     {
 		FWriteLogEntry Entry = *Iter;
-        // No write records should be within the transaction's stack range.
-        ensure(!IsOnStack(Entry.LogicalAddress));
+		// No write records should be within the transaction's stack range.
+		ensure(!IsOnStack(Entry.LogicalAddress));
 
-		if (UE_LOG_ACTIVE(LogAutoRTFM, Verbose))
-		{
-			TStringBuilder<1024> Builder;
-
-			Builder.Appendf(TEXT("%4d [UNDO] %p %4llu : [ "), Num - VerboseCounter - 1, Entry.LogicalAddress, Entry.Size);
-
-			std::byte* Current = Entry.LogicalAddress;
-			std::byte* Old = Entry.Data;
-
-			for (size_t i = 0; i < Entry.Size; i++)
-			{
-				Builder.Appendf(TEXT("%02X "), Current[i]);
-			}
-
-			Builder << TEXT("] -> [ ");
-
-			for (size_t i = 0; i < Entry.Size; i++)
-			{
-				Builder.Appendf(TEXT("%02X "), Old[i]);
-			}
-
-			Builder << TEXT("]");
-
-			UE_LOG(LogAutoRTFM, Verbose, TEXT("%s"), Builder.ToString());
-		}
-
-        memcpy(Entry.LogicalAddress, Entry.Data, Entry.Size);
+		memcpy(Entry.LogicalAddress, Entry.Data, Entry.Size);
 
 		VerboseCounter++;
     }
 
-	UE_LOG(LogAutoRTFM, Verbose, TEXT("Undone a transaction!"));
+	AUTORTFM_VERBOSE("Undone a transaction!");
 }
 
 void FTransaction::CommitNested()
 {
-    ASSERT(Parent);
+    AUTORTFM_ASSERT(Parent);
 
 	// We need to pass our write log to our parent transaction, but with care!
 	// We need to discard any writes if the memory location is on the parent
@@ -301,11 +273,11 @@ void FTransaction::CommitNested()
 
 bool FTransaction::AttemptToCommitOuterNest()
 {
-    ASSERT(!Parent);
+    AUTORTFM_ASSERT(!Parent);
 
-	UE_LOG(LogAutoRTFM, Verbose, TEXT("About to run commit tasks!"));
+	AUTORTFM_VERBOSE("About to run commit tasks!");
 	Context->DumpState();
-	UE_LOG(LogAutoRTFM, Verbose, TEXT("Running commit tasks..."));
+	AUTORTFM_VERBOSE("Running commit tasks...");
 
     AbortTasks.Reset();
 
@@ -372,7 +344,7 @@ void FTransaction::SetActive()
 			SetState<EState::ClosedActive>();
 			break;
 		default:
-			UE_LOG(LogAutoRTFM, Fatal, TEXT("Invalid state"));
+			AUTORTFM_FATAL("Invalid state");
 	}
 }
 
@@ -390,7 +362,7 @@ void FTransaction::SetInactive()
 			SetState<EState::ClosedInactive>();
 			break;
 		default:
-			UE_LOG(LogAutoRTFM, Fatal, TEXT("Invalid state"));
+			AUTORTFM_FATAL("Invalid state");
 	}
 }
 
@@ -402,17 +374,17 @@ void FTransaction::SetDone()
 template<FTransaction::EState NewState>
 void FTransaction::SetState()
 {
-	ASSERT(NewState != CurrentState);
+	AUTORTFM_ASSERT(NewState != CurrentState);
 
 	switch (CurrentState)
 	{
 		case EState::Uninitialized:
-			ASSERT(NewState == EState::OpenActive || NewState == EState::ClosedActive);
+			AUTORTFM_ASSERT(NewState == EState::OpenActive || NewState == EState::ClosedActive);
 			break;
 
 		// OpenActive -> OpenInactive, ClosedActive or Done
 		case EState::OpenActive:
-			ASSERT(NewState == EState::OpenInactive || NewState == EState::ClosedActive || NewState == EState::Done);
+			AUTORTFM_ASSERT(NewState == EState::OpenInactive || NewState == EState::ClosedActive || NewState == EState::Done);
 			if (CurrentMemoryValidationLevel != EMemoryValidationLevel::Disabled)
 			{
 				ValidateWriteHash();
@@ -421,31 +393,31 @@ void FTransaction::SetState()
 			}
 			else
 			{
-				ASSERT(RecordedWriteHash == 0 && NumWriteLogsHashed == 0);
+				AUTORTFM_ASSERT(RecordedWriteHash == 0 && NumWriteLogsHashed == 0);
 			}
 			break;
 
 		// ClosedActive -> ClosedInactive, OpenActive or Done
 		case EState::ClosedActive:
-			ASSERT(NewState == EState::ClosedInactive || NewState == EState::OpenActive || NewState == EState::Done);
+			AUTORTFM_ASSERT(NewState == EState::ClosedInactive || NewState == EState::OpenActive || NewState == EState::Done);
 			break;
 
 		// ClosedActive -> OpenActive
 		case EState::OpenInactive:
-			ASSERT(NewState == EState::OpenActive);
+			AUTORTFM_ASSERT(NewState == EState::OpenActive);
 			break;
 
 		// ClosedInactive -> ClosedActive
 		case EState::ClosedInactive:
-			ASSERT(NewState == EState::ClosedActive);
+			AUTORTFM_ASSERT(NewState == EState::ClosedActive);
 			break;
 
 		case EState::Done:
-			UE_LOG(LogAutoRTFM, Fatal, TEXT("Once Done, the transaction cannot change state without a call to Reset()"));
+			AUTORTFM_FATAL("Once Done, the transaction cannot change state without a call to Reset()");
 			break;
 
 		default:
-			UE_LOG(LogAutoRTFM, Fatal, TEXT("Invalid state"));
+			AUTORTFM_FATAL("Invalid state");
 			break;
 	}
 
@@ -453,7 +425,7 @@ void FTransaction::SetState()
 	if (NewState == EState::OpenActive) {
 		if (CurrentMemoryValidationLevel != EMemoryValidationLevel::Disabled)
 		{
-			ASSERT(RecordedWriteHash == 0 && NumWriteLogsHashed == 0);
+			AUTORTFM_ASSERT(RecordedWriteHash == 0 && NumWriteLogsHashed == 0);
 			RecordWriteHash();
 		}
 	}
@@ -469,7 +441,7 @@ void FTransaction::DebugBreakIfMemoryValidationFails()
 		FWriteHash NewHash = CalculateNestedWriteHash();
 		if (OldHash != NewHash)
 		{
-			UE_LOG(LogAutoRTFM, Warning, TEXT("DebugBreakIfInvalidMemoryHash() detected a change in hash"));
+			AUTORTFM_WARN("DebugBreakIfInvalidMemoryHash() detected a change in hash");
 			PLATFORM_BREAK();
 		}
 	}
@@ -488,24 +460,24 @@ void FTransaction::ValidateWriteHash() const
 	FWriteHash NewHash = CalculateNestedWriteHash();
 	Context->GetOpenHashThrottler().Update();
 
-	static constexpr TCHAR Message[] =
-		TEXT("Memory modified in a transaction was also modified in an call to AutoRTFM::Open(). ")
-		TEXT("This may lead to memory corruption if the transaction is aborted.");
+	static constexpr const char Message[] =
+		"Memory modified in a transaction was also modified in an call to AutoRTFM::Open(). "
+		"This may lead to memory corruption if the transaction is aborted.";
 	if (UNLIKELY(OldHash != NewHash))
 	{
 		if (CurrentMemoryValidationLevel == EMemoryValidationLevel::Warn)
 		{
-			UE_LOG(LogAutoRTFM, Warning, Message);
+			AUTORTFM_WARN(Message);
 		}
 		else
 		{
 			if (!ForTheRuntime::GetEnsureOnInternalAbort())
 			{
-				UE_LOG(LogAutoRTFM, Fatal, Message);
+				AUTORTFM_FATAL(Message);
 			}
 			else
 			{
-				ensureMsgf(OldHash == NewHash, Message);
+				AUTORTFM_ENSURE_MSG(OldHash == NewHash, Message);
 			}
 		}
 	}
diff --git a/Engine/Source/Runtime/Core/Private/AutoRTFM/Transaction.h b/Engine/Source/Runtime/Core/Private/AutoRTFM/Transaction.h
--- a/Engine/Source/Runtime/Core/Private/AutoRTFM/Transaction.h
+++ b/Engine/Source/Runtime/Core/Private/AutoRTFM/Transaction.h
@@ -2,6 +2,8 @@
 
 #pragma once
 
+#if (defined(__AUTORTFM) && __AUTORTFM)
+
 #include "AutoRTFM/AutoRTFM.h"
 #include "AutoRTFM/AutoRTFMTask.h"
 #include "HAL/Platform.h"
@@ -230,3 +232,5 @@ private:
 };
 
 } // namespace AutoRTFM
+
+#endif // (defined(__AUTORTFM) && __AUTORTFM)
diff --git a/Engine/Source/Runtime/Core/Private/AutoRTFM/TransactionInlines.h b/Engine/Source/Runtime/Core/Private/AutoRTFM/TransactionInlines.h
--- a/Engine/Source/Runtime/Core/Private/AutoRTFM/TransactionInlines.h
+++ b/Engine/Source/Runtime/Core/Private/AutoRTFM/TransactionInlines.h
@@ -2,6 +2,8 @@
 
 #pragma once
 
+#if (defined(__AUTORTFM) && __AUTORTFM)
+
 #include "Context.h"
 #include "ScopedGuard.h"
 #include "HitSet.h"
@@ -92,7 +94,7 @@ UE_AUTORTFM_FORCEINLINE bool FTransaction::ShouldRecordWrite(void* LogicalAddres
     // Consider moving the variable being written to an inner scoped
     // transaction, or move the variable outside of the nearest parent
     // scoped-transaction.
-    ASSERT(bIsStackScoped || LogicalAddress < FContext::Get()->GetClosedStackAddress());
+    AUTORTFM_ASSERT(bIsStackScoped || LogicalAddress < FContext::Get()->GetClosedStackAddress());
 
     return false;
 }
@@ -225,19 +227,19 @@ UE_AUTORTFM_FORCEINLINE void FTransaction::DidAllocate(void* LogicalAddress, con
 
 	AutoRTFM::TScopedGuard<bool> RecursionGuard(bIsInAllocateFn, true);
 	const bool DidInsert = NewMemoryTracker.Insert(LogicalAddress, Size);
-	ASSERT(DidInsert);
+	AUTORTFM_ASSERT(DidInsert);
 }
 
 UE_AUTORTFM_FORCEINLINE void FTransaction::DidFree(void* LogicalAddress)
 {
-	ASSERT(bTrackAllocationLocations);
+	AUTORTFM_ASSERT(bTrackAllocationLocations);
 
 	// Checking if one byte is in the interval map is enough to ascertain if it
 	// is new memory and we should be worried.
 	if (!bIsInAllocateFn)
 	{
 		AutoRTFM::TScopedGuard<bool> RecursionGuard(bIsInAllocateFn, true);
-		ASSERT(!NewMemoryTracker.Contains(LogicalAddress, 1));
+		AUTORTFM_ASSERT(!NewMemoryTracker.Contains(LogicalAddress, 1));
 	}
 }
 
@@ -334,3 +336,5 @@ UE_AUTORTFM_FORCEINLINE void FTransaction::CollectStats() const
 } // namespace AutoRTFM
 
 #undef AUTORTFM_CHECK_ASAN_FAKE_STACKS
+
+#endif // (defined(__AUTORTFM) && __AUTORTFM)
diff --git a/Engine/Source/Runtime/Core/Private/AutoRTFM/Utils.cpp b/Engine/Source/Runtime/Core/Private/AutoRTFM/Utils.cpp
--- a/Engine/Source/Runtime/Core/Private/AutoRTFM/Utils.cpp
+++ b/Engine/Source/Runtime/Core/Private/AutoRTFM/Utils.cpp
@@ -1,7 +1,9 @@
 // Copyright Epic Games, Inc. All Rights Reserved.
 
 #if (defined(__AUTORTFM) && __AUTORTFM)
+
 #include "Utils.h"
+#include "ContextInlines.h"
 
 #if PLATFORM_WINDOWS
 #include "Windows/WindowsHWrapper.h"
@@ -10,15 +12,11 @@
 #include <execinfo.h>
 #endif
 
-#include "Containers/StringConv.h"
-#include "Misc/CoreMiscDefines.h"
-
-DEFINE_LOG_CATEGORY(LogAutoRTFM)
 
 namespace AutoRTFM
 {
 
-FString GetFunctionDescription(void* FunctionPtr)
+std::string GetFunctionDescription(void* FunctionPtr)
 {
 #if PLATFORM_WINDOWS
     // This is gross, but it works. It's possible for someone to have SymInitialized before. But if they had, then this
@@ -37,16 +35,64 @@ FString GetFunctionDescription(void* FunctionPtr)
     }
     else
     {
-        return TEXT("<error getting description>");
+        return "<error getting description>";
     }
 #else // PLATFORM_WINDOWS -> so !PLATFORM_WINDOWS
     char** const symbols = backtrace_symbols(&FunctionPtr, 1);
-    FString Name(ANSI_TO_TCHAR(*symbols));
+    std::string Name(*symbols);
     free(symbols);
     return Name;
 #endif // PLATFORM_WINDOWS -> so !PLATFORM_WINDOWS
 }
 
+UE_AUTORTFM_API FORCENOINLINE void ReportError(const char* Format, ...)
+{
+	FContext* const Context = FContext::Get();
+
+	AUTORTFM_FATAL_IF(!Context->IsTransactional(), "AutoRTFM::ReportError() was called from outside a transaction");
+
+	const ForTheRuntime::EAutoRTFMInternalAbortActionState InternalAbortAction = ForTheRuntime::GetInternalAbortAction();
+
+	if (Format)
+	{
+		if (InternalAbortAction == ForTheRuntime::EAutoRTFMInternalAbortActionState::Crash)
+		{
+			va_list Args;
+			va_start(Args, Format);
+			AUTORTFM_FATAL_V(Format, Args);
+			va_end(Args);
+		}
+		else
+		{
+			va_list Args;
+			va_start(Args, Format);
+			AUTORTFM_ENSURE_MSG_V(!ForTheRuntime::GetEnsureOnInternalAbort(), Format, Args);
+			va_end(Args);
+		}
+	}
+	else
+	{
+		if (InternalAbortAction == ForTheRuntime::EAutoRTFMInternalAbortActionState::Crash)
+		{
+			AUTORTFM_FATAL("Transaction failing because of internal issue");
+		}
+		else
+		{
+			AUTORTFM_ENSURE_MSG(!ForTheRuntime::GetEnsureOnInternalAbort(), "Transaction failing because of internal issue");
+		}
+	}
+
+	if (InternalAbortAction == ForTheRuntime::EAutoRTFMInternalAbortActionState::Abort)
+	{
+		Context->AbortByLanguageAndThrow();
+	}
+	else
+	{
+		TTask<void()> Task = [] { AUTORTFM_ENSURE(ForTheRuntime::SetAutoRTFMRuntime(ForTheRuntime::AutoRTFM_ForcedDisabled)); };
+		Context->AbortTransactionWithPostAbortCallback(EContextStatus::AbortedByCascadingRetry, MoveTemp(Task));
+	}
+}
+
 } // namespace AutoRTFM
 
 #endif // defined(__AUTORTFM) && __AUTORTFM
diff --git a/Engine/Source/Runtime/Core/Private/AutoRTFM/Utils.h b/Engine/Source/Runtime/Core/Private/AutoRTFM/Utils.h
--- a/Engine/Source/Runtime/Core/Private/AutoRTFM/Utils.h
+++ b/Engine/Source/Runtime/Core/Private/AutoRTFM/Utils.h
@@ -2,34 +2,33 @@
 
 #pragma once
 
-#include "AutoRTFM/AutoRTFM.h"
+#if (defined(__AUTORTFM) && __AUTORTFM)
 
-#include "Containers/UnrealString.h"
-#include "Logging/LogMacros.h"
-#include "Misc/AssertionMacros.h"
+#include "AutoRTFM/AutoRTFM.h"
+#include "ExternAPI.h"
 
-DECLARE_LOG_CATEGORY_EXTERN(LogAutoRTFM, Display, All)
+#include <cstdarg>
+#include <string>
 
-extern "C" UE_AUTORTFM_API void autortfm_report_error(const UE_AUTORTFM_TCHAR* const Message);
+// If 0 then verbose logging is compiled-out.
+#define AUTORTFM_VERBOSE_ENABLED 0
 
 namespace AutoRTFM
 {
 
+__attribute__((__format__(__printf__, 1, 2)))
+UE_AUTORTFM_API void ReportError(const char* Format, ...);
+
 [[noreturn]] inline void InternalUnreachable() //-V1082 Silence PVS [[noreturn]] false positive warning
 {
 	LowLevelFatalError(TEXT("Unreachable encountered!"));
-
-#if defined(_MSC_VER) && !defined(__clang__)
-    __assume(false);
-#else
-    __builtin_unreachable();
-#endif
+	__builtin_unreachable();
 }
 
-FString GetFunctionDescription(void* FunctionPtr);
+std::string GetFunctionDescription(void* FunctionPtr);
 
 template<typename TReturnType, typename... TParameterTypes>
-FString GetFunctionDescription(TReturnType (*FunctionPtr)(TParameterTypes...))
+std::string GetFunctionDescription(TReturnType (*FunctionPtr)(TParameterTypes...))
 {
     return GetFunctionDescription(reinterpret_cast<void*>(FunctionPtr));
 }
@@ -38,66 +37,164 @@ FString GetFunctionDescription(TReturnType (*FunctionPtr)(TParameterTypes...))
 // the cost of calling this function (which is super unlikely to
 // be called because it is the asserting true codepath, EG. only
 // fatal process destroying explosions happen here).
-#if UE_AUTORTFM
 [[clang::preserve_most]]
-#endif
-[[noreturn]] FORCENOINLINE inline void DoAssert(void(*Logger)()) //-V1082 Silence PVS [[noreturn]] false positive warning
+[[noreturn]]
+FORCENOINLINE inline void DoAssert(void(*Logger)()) //-V1082 Silence PVS [[noreturn]] false positive warning
 {
 	Logger();
-
-#if defined(_MSC_VER) && !defined(__clang__)
-	__assume(false);
-#else
 	__builtin_unreachable();
-#endif
 }
 
 // We use a special calling convention for this call to minimize
 // the cost of calling this function (which is super unlikely to
 // be called because it is the expecting true codepath, EG. only
 // things we don't expect to happen get here).
-#if UE_AUTORTFM
 [[clang::preserve_most]]
-#endif
 FORCENOINLINE inline void DoExpect(void(*Logger)())
 {
 	Logger();
 }
 
+// Logs a message using a printf-style format string and va_list arguments.
+inline void LogV(const char* File, int Line, autortfm_log_severity Severity, const char* Format, va_list Args)
+{
+	GExternAPI.Log(File, Line, Severity, Format, Args);
+}
+
+// Logs a message using a printf-style format string and variadic arguments.
+__attribute__((__format__(__printf__, 4, 5)))
+inline void Log(const char* File, int Line, autortfm_log_severity Severity, const char* Format, ...)
+{
+	va_list Args;
+	va_start(Args, Format);
+	GExternAPI.Log(File, Line, Severity, Format, Args);
+	va_end(Args);
+}
+
+// Logs a message with a callstack using a printf-style format string and va_list arguments.
+inline void LogWithCallstackV(autortfm_log_severity Severity, const char* Format, va_list Args)
+{
+	GExternAPI.LogWithCallstack(Severity, Format, Args);
+}
+
+// Logs a message with a callstack using a printf-style format string and variadic arguments.
+__attribute__((__format__(__printf__, 2, 3)))
+inline void LogWithCallstack(autortfm_log_severity Severity, const char* Format, ...)
+{
+	va_list Args;
+	va_start(Args, Format);
+	GExternAPI.LogWithCallstack(Severity, Format, Args);
+	va_end(Args);
+}
+
+// Reports an ensure failure using a printf-style format string and va_list arguments.
+inline void EnsureFailureV(const char* File, int Line, const char* Condition, const char* Format, va_list Args)
+{
+	GExternAPI.EnsureFailure(File, Line, Condition, __builtin_return_address(0), Format, Args);
+}
+
+// Reports an ensure failure using a printf-style format string and variadic arguments.
+__attribute__((__format__(__printf__, 4, 5)))
+inline void EnsureFailure(const char* File, int Line, const char* Condition, const char* Format, ...)
+{
+	va_list Args;
+	va_start(Args, Format);
+	GExternAPI.EnsureFailure(File, Line, Condition, __builtin_return_address(0), Format, Args);
+	va_end(Args);
+}
+
 } // namespace AutoRTFM
 
-// This is all a bit funky - but for good reason! We want `ASSERT`
+#define AUTORTFM_LIKELY(x) __builtin_expect(!!(x), 1)
+#define AUTORTFM_UNLIKELY(x) __builtin_expect(!!(x), 0)
+
+#define AUTORTFM_REQUIRE_SEMICOLON static_assert(true)
+
+#define AUTORTFM_VERBOSE(Format, ...) ::AutoRTFM::Log(__FILE__, __LINE__, autortfm_log_verbose, Format, ##__VA_ARGS__)
+#define AUTORTFM_LOG(Format, ...) ::AutoRTFM::Log(__FILE__, __LINE__, autortfm_log_info, Format, ##__VA_ARGS__)
+#define AUTORTFM_WARN(Format, ...) ::AutoRTFM::Log(__FILE__, __LINE__, autortfm_log_warn, Format, ##__VA_ARGS__)
+#define AUTORTFM_ERROR(Format, ...) ::AutoRTFM::Log(__FILE__, __LINE__, autortfm_log_error, Format, ##__VA_ARGS__)
+#define AUTORTFM_FATAL(Format, ...) ::AutoRTFM::Log(__FILE__, __LINE__, autortfm_log_fatal, Format, ##__VA_ARGS__)
+
+#define AUTORTFM_VERBOSE_V(Format, Args) ::AutoRTFM::LogV(__FILE__, __LINE__, autortfm_log_verbose, Format, Args)
+#define AUTORTFM_LOG_V(Format, Args) ::AutoRTFM::LogV(__FILE__, __LINE__, autortfm_log_info, Format, Args)
+#define AUTORTFM_WARN_V(Format, Args) ::AutoRTFM::LogV(__FILE__, __LINE__, autortfm_log_warn, Format, Args)
+#define AUTORTFM_ERROR_V(Format, Args) ::AutoRTFM::LogV(__FILE__, __LINE__, autortfm_log_error, Format, Args)
+#define AUTORTFM_FATAL_V(Format, Args) ::AutoRTFM::LogV(__FILE__, __LINE__, autortfm_log_fatal, Format, Args)
+
+#define AUTORTFM_VERBOSE_IF(Condition, Format, ...) do { if (Condition) { AUTORTFM_VERBOSE(Format, ##__VA_ARGS__); } } while(false)
+#define AUTORTFM_LOG_IF(Condition, Format, ...) do { if (Condition) { AUTORTFM_LOG(Format, ##__VA_ARGS__); } } while(false)
+#define AUTORTFM_WARN_IF(Condition, Format, ...) do { if (Condition) { AUTORTFM_WARN(Format, ##__VA_ARGS__); } } while(false)
+#define AUTORTFM_ERROR_IF(Condition, Format, ...) do { if (AUTORTFM_UNLIKELY(Condition)) { AUTORTFM_ERROR(Format, ##__VA_ARGS__); } } while(false)
+#define AUTORTFM_FATAL_IF(Condition, Format, ...) do { if (AUTORTFM_UNLIKELY(Condition)) { AUTORTFM_FATAL(Format, ##__VA_ARGS__); } } while(false)
+
+// If AUTORTFM_VERBOSE_ENABLED is 0, then replace all the verbose logging macros
+// with no-ops.
+#if !AUTORTFM_VERBOSE_ENABLED
+#undef AUTORTFM_VERBOSE
+#undef AUTORTFM_VERBOSE_V
+#undef AUTORTFM_VERBOSE_IF
+#define AUTORTFM_VERBOSE(...) AUTORTFM_REQUIRE_SEMICOLON
+#define AUTORTFM_VERBOSE_V(...) AUTORTFM_REQUIRE_SEMICOLON
+#define AUTORTFM_VERBOSE_IF(...) AUTORTFM_REQUIRE_SEMICOLON
+#endif
+
+#define AUTORTFM_ENSURE(Condition) do { \
+	if (AUTORTFM_UNLIKELY(!(Condition))) { \
+		static bool bCalled = [] { \
+			::AutoRTFM::EnsureFailure(__FILE__, __LINE__, #Condition, nullptr); \
+			return true; \
+		}(); \
+	} } while(false)
+
+#define AUTORTFM_ENSURE_MSG(Condition, Format, ...) do { \
+	if (AUTORTFM_UNLIKELY(!(Condition))) { \
+		static bool bCalled = [&] { \
+			::AutoRTFM::EnsureFailure(__FILE__, __LINE__, #Condition, Format, ##__VA_ARGS__); \
+			return true; \
+		}(); \
+	} } while(false)
+
+#define AUTORTFM_ENSURE_MSG_V(Condition, Format, Args) do { \
+	if (AUTORTFM_UNLIKELY(!(Condition))) { \
+		static bool bCalled = [&] { \
+			::AutoRTFM::EnsureFailureV(__FILE__, __LINE__, #Condition, Format, Args); \
+			return true; \
+		}(); \
+	} } while(false)
+
+#define AUTORTFM_DUMP_STACKTRACE(Message) ::AutoRTFM::Log(Message, ##__VA_ARGS__)
+
+// This is all a bit funky - but for good reason! We want `AUTORTFM_ASSERT`
 // to be as close to *zero* cost if the assert wouldn't trigger.
 // We will always pay the cost of the unlikely branch, but we need
 // to make the body of taking the assert happen in another function
 // so as not to affect codegen, register allocation, and stack
-// reservation. But we also want the `ASSERT` to give us accurate
+// reservation. But we also want the `AUTORTFM_ASSERT` to give us accurate
 // `__FILE__` and `__LINE__` information for the line that it was
 // triggered on. So what we do is have a lambda at the line that
 // actually does the assert (but crucially gets the correct file
 // and line numbers), but then pass this to another function
 // `DoAssert` which has a special calling convention that makes
 // the caller as optimal as possible (at the expense of the callee).
-#define ASSERT(exp) \
-	if (UNLIKELY(!(exp))) \
+#define AUTORTFM_ASSERT(exp) \
+	if (AUTORTFM_UNLIKELY(!(exp))) \
 	{ \
 		auto Lambda = [] \
 		{ \
-			FDebug::DumpStackTraceToLog(TEXT("AutoRTFM assert backtrace"), ELogVerbosity::Error); \
-			UE_LOG(LogAutoRTFM, Fatal, TEXT("AutoRTFM assert %s!"), TEXT(#exp)); \
+			AUTORTFM_FATAL("AUTORTFM_ASSERT(%s) failure", #exp); \
 		}; \
 		AutoRTFM::DoAssert(Lambda); \
 	}
 
-// Same funkiness as ASSERT, except that it doesn't cause a fatal error.
-// TODO: Maybe upgrade the Display -> Warning?
-#define EXPECT(exp) \
+// Same funkiness as AUTORTFM_ASSERT, except that it doesn't cause a fatal error.
+// TODO: Maybe upgrade the Info -> Warning?
+#define AUTORTFM_EXPECT(exp) \
 	if (UNLIKELY(!(exp))) \
 	{ \
 		auto Lambda = [] \
 		{ \
-			FDebug::DumpStackTraceToLog(TEXT("AutoRTFM expect backtrace"), ELogVerbosity::Display); \
-			UE_LOG(LogAutoRTFM, Display, TEXT("AutoRTFM expect %s!"), TEXT(#exp)); \
+			::AutoRTFM::LogWithCallstack(autortfm_log_info, "AUTORTFM_EXPECT(%s) failure", #exp); \
 		}; \
 		AutoRTFM::DoExpect(Lambda); \
 	}
@@ -119,3 +216,5 @@ FORCENOINLINE inline void DoExpect(void(*Logger)())
 #if !defined(AUTORTFM_MUST_TAIL)
 #define AUTORTFM_MUST_TAIL
 #endif
+
+#endif // (defined(__AUTORTFM) && __AUTORTFM)
diff --git a/Engine/Source/Runtime/Core/Private/AutoRTFM/WriteLog.h b/Engine/Source/Runtime/Core/Private/AutoRTFM/WriteLog.h
--- a/Engine/Source/Runtime/Core/Private/AutoRTFM/WriteLog.h
+++ b/Engine/Source/Runtime/Core/Private/AutoRTFM/WriteLog.h
@@ -2,6 +2,8 @@
 
 #pragma once
 
+#if (defined(__AUTORTFM) && __AUTORTFM)
+
 #include "AutoRTFM/AutoRTFM.h"
 #include "HAL/Platform.h"
 #include "Misc/CoreMiscDefines.h"
@@ -86,7 +88,7 @@ namespace AutoRTFM
 			// the FBlock header.
 			static FBlock* Allocate(size_t TotalSize)
 			{
-				ASSERT(TotalSize > (sizeof(FBlock) + sizeof(FRecord)));
+				AUTORTFM_ASSERT(TotalSize > (sizeof(FBlock) + sizeof(FRecord)));
 				std::byte* Memory = new std::byte[TotalSize];
 				// Disable false-positive warning C6386: Buffer overrun while writing to 'Memory'
 				CA_SUPPRESS(6386)
@@ -169,7 +171,7 @@ namespace AutoRTFM
 				DataEnd += Entry.Size;
 
 #if UE_BUILD_DEBUG
-				ASSERT(DataEnd <= reinterpret_cast<std::byte*>(LastRecord));
+				AUTORTFM_ASSERT(DataEnd <= reinterpret_cast<std::byte*>(LastRecord));
 #endif
 
 				return Result;
@@ -206,8 +208,8 @@ namespace AutoRTFM
 		// The log will make a copy of the FWriteLogEntry's data.
 		void Push(FWriteLogEntry Entry)
 		{
-			ASSERT(Entry.Size <= FWriteLogEntry::MaxSize);
-			ASSERT((reinterpret_cast<uintptr_t>(Entry.LogicalAddress) & 0xffff0000'00000000) == 0);
+			AUTORTFM_ASSERT(Entry.Size <= FWriteLogEntry::MaxSize);
+			AUTORTFM_ASSERT((reinterpret_cast<uintptr_t>(Entry.LogicalAddress) & 0xffff0000'00000000) == 0);
 
 			FBlock::EPushResult PushResult = TailBlock->Push(Entry);
 			if (PushResult == FBlock::EPushResult::Added)
@@ -223,7 +225,7 @@ namespace AutoRTFM
 				TailBlock = NewBlock;
 
 				PushResult = NewBlock->Push(Entry);
-				ASSERT(PushResult == FBlock::EPushResult::Added);
+				AUTORTFM_ASSERT(PushResult == FBlock::EPushResult::Added);
 				NumEntries++;
 			}
 
@@ -383,3 +385,5 @@ namespace AutoRTFM
 		alignas(alignof(FBlock)) std::byte HeadBlockMemory[HeadBlockSize];
 	};
 }
+
+#endif // (defined(__AUTORTFM) && __AUTORTFM)
diff --git a/Engine/Source/Runtime/Core/Private/ProfilingDebugging/CpuProfilerTrace.cpp b/Engine/Source/Runtime/Core/Private/ProfilingDebugging/CpuProfilerTrace.cpp
--- a/Engine/Source/Runtime/Core/Private/ProfilingDebugging/CpuProfilerTrace.cpp
+++ b/Engine/Source/Runtime/Core/Private/ProfilingDebugging/CpuProfilerTrace.cpp
@@ -165,7 +165,7 @@ void FCpuProfilerTraceInternal::FlushThreadBuffer(FThreadBuffer* InThreadBuffer)
 	}
 
 #define INSIGHTS_AUTORTFM_UNREACHABLE() \
-	do { if (AutoRTFM::IsClosed()) { AutoRTFM::Unreachable(TEXT("Unreachable transactional codepath in FCpuProfilerTrace")); } } while (0)
+	do { if (AutoRTFM::IsClosed()) { AutoRTFM::Unreachable("Unreachable transactional codepath in FCpuProfilerTrace"); } } while (0)
 
 uint8 FCpuProfilerTrace::OnAbortKey = 42;
 
diff --git a/Engine/Source/Runtime/Core/Public/AutoRTFM/AutoRTFM.h b/Engine/Source/Runtime/Core/Public/AutoRTFM/AutoRTFM.h
--- a/Engine/Source/Runtime/Core/Public/AutoRTFM/AutoRTFM.h
+++ b/Engine/Source/Runtime/Core/Public/AutoRTFM/AutoRTFM.h
@@ -69,7 +69,6 @@
 #define PRAGMA_RESTORE_UNREACHABLE_CODE_WARNINGS
 #define LIKELY(EXPR) __builtin_expect(!!(EXPR), 1)
 #define PLATFORM_RETURN_ADDRESS() __builtin_return_address(0)
-#define UE_AUTORTFM_TCHAR char
 #else
 #include <HAL/Platform.h>
 #include <HAL/PlatformMemory.h>
@@ -78,7 +77,6 @@
 #define UE_AUTORTFM_FORCEINLINE FORCEINLINE
 #define UE_AUTORTFM_MEMCPY FPlatformMemory::Memcpy
 #define UE_AUTORTFM_MOVE MoveTemp
-#define UE_AUTORTFM_TCHAR TCHAR
 #endif
 
 #if defined(UE_BUILD_DEBUG) && UE_BUILD_DEBUG
@@ -137,42 +135,79 @@ typedef enum
 	autortfm_status_in_post_abort
 } autortfm_status;
 
+// AutoRTFM logging severity.
+typedef enum
+{
+	autortfm_log_verbose = 0,
+	autortfm_log_info,
+	autortfm_log_warn,
+	autortfm_log_error,
+	autortfm_log_fatal,
+} autortfm_log_severity;
+
 // Function pointers used by AutoRTFM for heap allocations, etc.
 typedef struct
 {
 	// The function used to allocate memory from the heap.
 	// Must not be null.
 	void* (*Allocate)(size_t Size, size_t Alignment);
+
 	// The function used to reallocate memory from the heap.
 	// Must not be null.
 	void* (*Reallocate)(void* Pointer, size_t Size, size_t Alignment);
+
 	// The function used to allocate zeroed memory from the heap.
 	// Must not be null.
 	void* (*AllocateZeroed)(size_t Size, size_t Alignment);
+
 	// The function used to free memory allocated by Allocate() and AllocateZeroed().
 	// Must not be null.
 	void (*Free)(void* Pointer);
+
+	// Function used to log messages using a printf-style format string and va_list arguments.
+	// Strings use UTF-8 encoding.
+	// Must not be null.
+	void (*Log)(const char* File, int Line, autortfm_log_severity Severity, const char* Format, va_list Args);
+
+	// Function used to log messages with a callstack using a printf-style format string and va_list arguments.
+	// Strings use UTF-8 encoding.
+	// Must not be null.
+	void (*LogWithCallstack)(autortfm_log_severity Severity, const char* Format, va_list Args);
 	
+	// Function used to report an ensure failure using a printf-style format string and va_list arguments.
+	// Strings use UTF-8 encoding.
+	// Must not be null.
+	void (*EnsureFailure)(const char* File, int Line, const char* Condition, void* ProgramCounter, const char* Format, va_list Args);
+
+	// Function used to query whether a log severity is active.
+	// Must not be null.
+	bool (*IsLogActive)(autortfm_log_severity Severity);
+
 	// Optional callback to be informed when the value returned by
 	// ForTheRuntime::IsAutoRTFMRuntimeEnabled() changes.
 	// Can be null.
 	void (*OnRuntimeEnabledChanged)();
+
 	// Optional callback to be informed when the value returned by
 	// ForTheRuntime::GetRetryTransaction() changes.
 	// Can be null.
 	void (*OnRetryTransactionsChanged)();
+
 	// Optional callback to be informed when the value returned by
 	// ForTheRuntime::GetMemoryValidationLevel() changes.
 	// Can be null.
 	void (*OnMemoryValidationLevelChanged)();
+
 	// Optional callback to be informed when the value returned by
 	// ForTheRuntime::GetMemoryValidationThrottlingEnabled() changes.
 	// Can be null.
 	void (*OnMemoryValidationThrottlingChanged)();
+
 	// Optional callback to be informed when the value returned by
 	// ForTheRuntime::GetMemoryValidationStatisticsEnabled() changes.
 	// Can be null.
 	void (*OnMemoryValidationStatisticsChanged)();
+
 } autortfm_extern_api;
 
 #if UE_AUTORTFM_ENABLED
@@ -454,9 +489,9 @@ static UE_AUTORTFM_FORCEINLINE void autortfm_check_abi(void* ptr, size_t size)
 #endif
 
 #if UE_AUTORTFM_ENABLED
-UE_AUTORTFM_API void autortfm_unreachable(const UE_AUTORTFM_TCHAR* const Message);
+UE_AUTORTFM_API void autortfm_unreachable(const char* Message);
 #else
-static UE_AUTORTFM_FORCEINLINE void autortfm_unreachable(const UE_AUTORTFM_TCHAR* const Message)
+static UE_AUTORTFM_FORCEINLINE void autortfm_unreachable(const char* Message)
 {
 	UE_AUTORTFM_UNUSED(Message);
 }
@@ -1097,13 +1132,13 @@ template<typename Type> static UE_AUTORTFM_FORCEINLINE void RecordOpenWriteNoMem
 
 // Report that a unreachable codepath is being hit. Used to manually ban certain codepaths
 // from being transactionally safe.
-static UE_AUTORTFM_FORCEINLINE void Unreachable(const UE_AUTORTFM_TCHAR* const Message = nullptr)
+static UE_AUTORTFM_FORCEINLINE void Unreachable(const char* Message = nullptr)
 {
 	autortfm_unreachable(Message);
 }
 
 // If we are running within a transaction, call `AutoRTFM::Unreachable`.
-static UE_AUTORTFM_FORCEINLINE void UnreachableIfTransactional(const UE_AUTORTFM_TCHAR* const Message = nullptr)
+static UE_AUTORTFM_FORCEINLINE void UnreachableIfTransactional(const char* Message = nullptr)
 {
 	if (AutoRTFM::IsTransactional())
 	{
diff --git a/Engine/Source/Runtime/CoreUObject/Private/Serialization/TransactionallySafeAsyncLoading.cpp b/Engine/Source/Runtime/CoreUObject/Private/Serialization/TransactionallySafeAsyncLoading.cpp
--- a/Engine/Source/Runtime/CoreUObject/Private/Serialization/TransactionallySafeAsyncLoading.cpp
+++ b/Engine/Source/Runtime/CoreUObject/Private/Serialization/TransactionallySafeAsyncLoading.cpp
@@ -409,7 +409,7 @@ private:
 	static constexpr int32 PackageCacheHit = INT_MIN;
 	static constexpr int32 PackageCacheMiss = PackageCacheHit + 1;
 
-	static const TCHAR* const UnreachableMessage;
+	static const char* const UnreachableMessage;
 	IAsyncPackageLoader* const WrappedPackageLoader;
 	
 	struct FTransactionallyDeferredLoadPackagePayload final 
@@ -526,7 +526,7 @@ private:
 	}
 };
 
-const TCHAR* const FTransactionallySafeAsyncPackageLoader::UnreachableMessage = TEXT("Cannot call function within a transaction!");
+const char* const FTransactionallySafeAsyncPackageLoader::UnreachableMessage = "Cannot call function within a transaction!";
 
 IAsyncPackageLoader* MakeTransactionallySafeAsyncPackageLoader(IAsyncPackageLoader* const InWrappedPackageLoader)
 {
