PR #10505: Containers: Fix TStaticBitArray serialization failing to compile.

- Serialization was broken when refactored as part of CL 23038965.
- TStaticBitArray has been fixed up to follow the pattern used by TBitArray in the same CL.
- Added a quick serialization test to the unit test so that similar issues would be caught in future.

diff --git a/Engine/Source/Runtime/Core/Private/Containers/StaticBitArrayTests.cpp b/Engine/Source/Runtime/Core/Private/Containers/StaticBitArrayTests.cpp
--- a/Engine/Source/Runtime/Core/Private/Containers/StaticBitArrayTests.cpp
+++ b/Engine/Source/Runtime/Core/Private/Containers/StaticBitArrayTests.cpp
@@ -2,9 +2,11 @@
 
 #if WITH_LOW_LEVEL_TESTS
 
+#include "Containers/StaticBitArray.h"
 #include "CoreMinimal.h"
 #include "Misc/AutomationTest.h"
-#include "Containers/StaticBitArray.h"
+#include "Serialization/MemoryReader.h"
+#include "Serialization/MemoryWriter.h"
 #include "TestHarness.h"
 
 #include <catch2/generators/catch_generators.hpp>
@@ -44,6 +46,27 @@ TEST_CASE("System::Core::Containers::TStaticBitArray", "[Containers][Smoke]")
 		CHECK(Bits != BitsDifferent);
 		CHECK(!(Bits == BitsDifferent));
 	}
+
+	SECTION("Serialization")
+	{
+		TStaticBitArray<128> SourceBits;
+		for (int32 Index = 0; Index < 128; ++Index)
+		{
+			SourceBits[Index] = (Index % 2 == 0);
+		}
+
+		TArray<uint8> MemoryBlob;
+		FMemoryWriter WriterAr(MemoryBlob);
+
+		WriterAr << SourceBits;
+
+		FMemoryReader ReaderAr(MemoryBlob);
+
+		TStaticBitArray<128> DstBits;
+		ReaderAr << DstBits;
+
+		CHECK(SourceBits == DstBits);
+	}
 }
 
 #endif
diff --git a/Engine/Source/Runtime/Core/Public/Containers/StaticBitArray.h b/Engine/Source/Runtime/Core/Public/Containers/StaticBitArray.h
--- a/Engine/Source/Runtime/Core/Public/Containers/StaticBitArray.h
+++ b/Engine/Source/Runtime/Core/Public/Containers/StaticBitArray.h
@@ -364,6 +364,20 @@ public:
 		return Str;
 	}
 
+	void Serialize(FArchive& Ar)
+	{
+		uint32 ArchivedNumWords = NumWords;
+		Ar << ArchivedNumWords;
+
+		if (Ar.IsLoading())
+		{
+			FMemory::Memset(Words, 0, sizeof(Words));
+			ArchivedNumWords = FMath::Min(NumWords, ArchivedNumWords);
+		}
+
+		Ar.Serialize(Words, ArchivedNumWords * sizeof(Words[0]));
+	}
+
 	static constexpr uint32 NumOfBits = NumBits;
 
 private:
@@ -417,15 +431,6 @@ private:
 template<uint32 NumBits>
 FArchive& operator<<(FArchive& Ar, TStaticBitArray<NumBits>& BitArray)
 {
-	uint32 ArchivedNumWords = BitArray.NumWords;
-	Ar << ArchivedNumWords;
-
-	if(Ar.IsLoading())
-	{
-		FMemory::Memset(BitArray.Words, 0, sizeof(BitArray.Words));
-		ArchivedNumWords = FMath::Min(BitArray.NumWords, ArchivedNumWords);
-	}
-
-	Ar.Serialize(BitArray.Words, ArchivedNumWords * sizeof(BitArray.Words[0]));
+	BitArray.Serialize(Ar);
 	return Ar;
 }
