[Insights] MemoryTrace: Added MemoryTrace_UpdateAlloc; can be used to update the context (tag and metadata) for an allocation.

diff --git a/Engine/Source/Runtime/Core/Private/ProfilingDebugging/MemoryTrace.cpp b/Engine/Source/Runtime/Core/Private/ProfilingDebugging/MemoryTrace.cpp
--- a/Engine/Source/Runtime/Core/Private/ProfilingDebugging/MemoryTrace.cpp
+++ b/Engine/Source/Runtime/Core/Private/ProfilingDebugging/MemoryTrace.cpp
@@ -138,7 +138,13 @@ UE_TRACE_EVENT_BEGIN(Memory, ReallocFreeSystem)
 	UE_TRACE_EVENT_FIELD(uint32, CallstackId)
 UE_TRACE_EVENT_END()
 
-UE_TRACE_EVENT_BEGIN(Memory, MemorySwapOp)
+UE_TRACE_EVENT_BEGIN(Memory, UpdateAlloc) // added in UE 5.6
+	UE_TRACE_EVENT_FIELD(uint64, Address)
+	UE_TRACE_EVENT_FIELD(uint32, CallstackId)
+	UE_TRACE_EVENT_FIELD(uint8, RootHeap)
+UE_TRACE_EVENT_END()
+
+UE_TRACE_EVENT_BEGIN(Memory, MemorySwapOp) // added in UE 5.5
 	UE_TRACE_EVENT_FIELD(uint64, Address) // page fault real address
 	UE_TRACE_EVENT_FIELD(uint32, CallstackId)
 	UE_TRACE_EVENT_FIELD(uint32, CompressedSize)
@@ -371,12 +377,12 @@ static EMemoryTraceInit MemoryTrace_ShouldEnable(int32 ArgC, const ArgCharType*
 FMalloc* MemoryTrace_CreateInternal(FMalloc* InMalloc, EMemoryTraceInit Mode)
 {
 	// If allocation events are not desired we don't need to do anything, even
-	// if user has enabled only callstacks it will be enabled later. 
+	// if user has enabled only callstacks it will be enabled later.
 	if (!EnumHasAnyFlags(Mode, EMemoryTraceInit::AllocEvents))
 	{
 		return InMalloc;
 	}
-	
+
 	// Some OSes (i.e. Windows) will terminate all threads except the main
 	// one as part of static deinit. However we may receive more memory
 	// trace events that would get lost as Trace's worker thread has been
@@ -649,6 +655,26 @@ void MemoryTrace_ReallocFree(uint64 Address, HeapId RootHeap, uint32 ExternalCal
 	MemoryTrace_UpdateInternal();
 }
 
+////////////////////////////////////////////////////////////////////////////////
+void MemoryTrace_UpdateAlloc(uint64 Address, HeapId RootHeap, uint32 ExternalCallstackId)
+{
+	if (!GTraceAllowed)
+	{
+		return;
+	}
+
+	checkSlow(RootHeap < 16);
+
+	const uint32 CallstackId = ExternalCallstackId ? ExternalCallstackId : GDoNotAllocateInTrace ? 0 : CallstackTrace_GetCurrentId();
+
+	UE_TRACE_LOG(Memory, UpdateAlloc, MemAllocChannel)
+		<< UpdateAlloc.Address(uint64(Address))
+		<< UpdateAlloc.CallstackId(CallstackId)
+		<< UpdateAlloc.RootHeap(uint8(RootHeap));
+
+	MemoryTrace_UpdateInternal();
+}
+
 ////////////////////////////////////////////////////////////////////////////////
 void MemoryTrace_SwapOp(uint64 PageAddress, EMemoryTraceSwapOperation SwapOperation, uint32 CompressedSize, uint32 CallstackId)
 {
diff --git a/Engine/Source/Runtime/Core/Public/ProfilingDebugging/MemoryTrace.h b/Engine/Source/Runtime/Core/Public/ProfilingDebugging/MemoryTrace.h
--- a/Engine/Source/Runtime/Core/Public/ProfilingDebugging/MemoryTrace.h
+++ b/Engine/Source/Runtime/Core/Public/ProfilingDebugging/MemoryTrace.h
@@ -166,6 +166,14 @@ CORE_API void MemoryTrace_ReallocFree(uint64 Address, HeapId RootHeap = EMemoryT
  */
 CORE_API void MemoryTrace_ReallocAlloc(uint64 Address, uint64 NewSize, uint32 Alignment, HeapId RootHeap = EMemoryTraceRootHeap::SystemMemory, uint32 ExternalCallstackId = 0);
 
+/**
+ * Trace an update alloc event. It updates context (mem tag and metadata) for an allocation.
+ * @param Address Address of the allocation being updated
+ * @param RootHeap Which root heap this belongs to (system memory, video memory etc)
+ * @param ExternalCallstackId CallstackId to use, if 0 will use current callstack id.
+ */
+CORE_API void MemoryTrace_UpdateAlloc(uint64 Address, HeapId RootHeap = EMemoryTraceRootHeap::SystemMemory, uint32 ExternalCallstackId = 0);
+
 /** Trace a swap operation. Only available for system memory root heap (EMemoryTraceRootHeap::SystemMemory).
  * @param PageAddress Page address for operation, in case of PageIn can be address of the page fault (not aligned to page boundary).
  * @param SwapOperation Which swap operation is happening to the address.
@@ -186,6 +194,7 @@ inline void MemoryTrace_Alloc(uint64 Address, uint64 Size, uint32 Alignment, Hea
 inline void MemoryTrace_Free(uint64 Address, HeapId RootHeap = EMemoryTraceRootHeap::SystemMemory, uint32 ExternalCallstackId = 0) {}
 inline void MemoryTrace_ReallocFree(uint64 Address, HeapId RootHeap = EMemoryTraceRootHeap::SystemMemory, uint32 ExternalCallstackId = 0) {}
 inline void MemoryTrace_ReallocAlloc(uint64 Address, uint64 NewSize, uint32 Alignment, HeapId RootHeap = EMemoryTraceRootHeap::SystemMemory, uint32 ExternalCallstackId = 0) {}
+inline void MemoryTrace_UpdateAlloc(uint64 Address, HeapId RootHeap = EMemoryTraceRootHeap::SystemMemory, uint32 ExternalCallstackId = 0) {}
 inline void MemoryTrace_SwapOp(uint64 PageAddress, EMemoryTraceSwapOperation SwapOperation, uint32 CompressedSize = 0, uint32 CallstackId = 0) {}
 
 #endif // UE_MEMORY_TRACE_ENABLED
