[Insights] MemoryTrace: Added support for UE_MEMSCOPE_PTR with a video memory (or other root heap) ptr; i.e. UE_MEMSCOPE_PTR(InPtr, InRootHeap).

diff --git a/Engine/Source/Runtime/Core/Private/ProfilingDebugging/TagTrace.cpp b/Engine/Source/Runtime/Core/Private/ProfilingDebugging/TagTrace.cpp
--- a/Engine/Source/Runtime/Core/Private/ProfilingDebugging/TagTrace.cpp
+++ b/Engine/Source/Runtime/Core/Private/ProfilingDebugging/TagTrace.cpp
@@ -31,6 +31,7 @@ UE_TRACE_EVENT_END()
 
 UE_TRACE_EVENT_BEGIN(Memory, MemoryScopePtr, NoSync)
 	UE_TRACE_EVENT_FIELD(uint64, Ptr)
+	UE_TRACE_EVENT_FIELD(uint8, RootHeap) // added in UE 5.6
 UE_TRACE_EVENT_END()
 
 
@@ -94,7 +95,8 @@ void FMemScope::ActivateScope(int32 InTag)
 		if (const auto& __restrict MemoryScope = *(FMemoryMemoryScopeFields*)(&LogScope))
 		{
 			Inner.SetActive();
-			LogScope += LogScope << MemoryScope.Tag(InTag);
+			LogScope += LogScope
+				<< MemoryScope.Tag(InTag);
 			PrevTag = GActiveTag;
 			GActiveTag = InTag;
 		}
@@ -111,7 +113,7 @@ FMemScope::~FMemScope()
 }
 
 ////////////////////////////////////////////////////////////////////////////////
-FMemScopePtr::FMemScopePtr(uint64 InPtr)
+FMemScopePtr::FMemScopePtr(uint64 InPtr, HeapId InRootHeap)
 {
 	if (InPtr != 0 && TRACE_PRIVATE_CHANNELEXPR_IS_ENABLED(MemAllocChannel))
 	{
@@ -119,7 +121,10 @@ FMemScopePtr::FMemScopePtr(uint64 InPtr)
 		{
 			if (const auto& __restrict MemoryScope = *(FMemoryMemoryScopePtrFields*)(&LogScope))
 			{
-				Inner.SetActive(), LogScope += LogScope << MemoryScope.Ptr(InPtr);
+				Inner.SetActive();
+				LogScope += LogScope
+					<< MemoryScope.Ptr(InPtr)
+					<< MemoryScope.RootHeap(uint8(InRootHeap));
 			}
 		}
 	}
@@ -133,7 +138,7 @@ FMemScopePtr::~FMemScopePtr()
 /////////////////////////////////////////////////////////////////////////////////
 
 /**
- * Utility class that manages tracing the specification of unique LLM tags 
+ * Utility class that manages tracing the specification of unique LLM tags
  * and custom name based tags.
  */
 class FTagTrace
@@ -169,7 +174,7 @@ private:
 	{
 		return UniqueName.GetDisplayIndex().ToUnstableInt() + FNAME_INDEX_OFFSET;
 	}
-	
+
 	struct FTagNameSetEntry
 	{
 		std::atomic_int32_t Data;
@@ -276,7 +281,7 @@ int32 FTagTrace::AnnounceLLMExtendedTag(FLowLevelMemTracker& Tracker, const UE::
 		}
 		// Not a generic tag, so this is a Platform or Project tag.
 	}
-	
+
 	FName UniqueName = Tracker.GetTagUniqueName(LLMTag);
 	int32 TagId = bEnumTag ? static_cast<int32>(EnumTag) : GetFNameTagID(UniqueName);
 
@@ -300,7 +305,7 @@ int32 FTagTrace::AnnounceLLMExtendedTag(FLowLevelMemTracker& Tracker, const UE::
 		DisplayNamePtr = nullptr; // Use the UniqueName instead
 	}
 
-	return AnnounceFNameTag(UniqueName, TagId, DisplayNamePtr, ParentTagId);
+	return AnnounceFNameTag(UniqueName, TagId, DisplayNamePtr, TOptional<int32>(ParentTagId));
 }
 
 #endif
@@ -353,6 +358,7 @@ int32 FTagTrace::AnnounceFNameTag(FName UniqueName, int32 TagId, const TCHAR* Di
 	if (!ParentTagId.IsSet())
 	{
 		ParentTagId = -1;
+
 		if (DisplayNameLen > 0)
 		{
 			// By contract we pull the parent name from the UniqueName rather than the display string.
@@ -380,7 +386,7 @@ int32 FTagTrace::AnnounceFNameTag(FName UniqueName, int32 TagId, const TCHAR* Di
 
 ////////////////////////////////////////////////////////////////////////////////
 int32 FTagTrace::AnnounceCustomTag(int32 Tag, int32 ParentTag, const ANSICHAR* Display) const
-{		
+{
 	const uint32 DisplayLen = FCStringAnsi::Strlen(Display);
 	UE_TRACE_LOG(Memory, TagSpec, MemAllocChannel, DisplayLen * sizeof(ANSICHAR))
 		<< TagSpec.Tag(Tag)
diff --git a/Engine/Source/Runtime/Core/Public/ProfilingDebugging/TagTrace.h b/Engine/Source/Runtime/Core/Public/ProfilingDebugging/TagTrace.h
--- a/Engine/Source/Runtime/Core/Public/ProfilingDebugging/TagTrace.h
+++ b/Engine/Source/Runtime/Core/Public/ProfilingDebugging/TagTrace.h
@@ -1,6 +1,8 @@
 // Copyright Epic Games, Inc. All Rights Reserved.
 #pragma once
+
 #include "CoreTypes.h"
+#include "ProfilingDebugging/MemoryTrace.h"
 #include "Trace/Config.h"
 #include "Trace/Detail/LogScope.h"
 
@@ -80,7 +82,7 @@ public:
 class FMemScopePtr
 {
 public:
-	CORE_API FMemScopePtr(uint64 InPtr);
+	CORE_API FMemScopePtr(uint64 InPtr, HeapId InRootHeap = EMemoryTraceRootHeap::SystemMemory);
 	CORE_API ~FMemScopePtr();
 private:
 	UE::Trace::Private::FScopedLogScope Inner;
@@ -88,7 +90,8 @@ private:
 
 ////////////////////////////////////////////////////////////////////////////////
 #define UE_MEMSCOPE(InTag)					FMemScope PREPROCESSOR_JOIN(MemScope,__LINE__)(InTag);
-#define UE_MEMSCOPE_PTR(InPtr)				FMemScopePtr PREPROCESSOR_JOIN(MemPtrScope,__LINE__)((uint64)InPtr);
+// UE_MEMSCOPE_PTR(InPtr) or UE_MEMSCOPE_PTR(InPtr, InRootHeap)
+#define UE_MEMSCOPE_PTR(InPtr, ...)			FMemScopePtr PREPROCESSOR_JOIN(MemPtrScope,__LINE__)((uint64)InPtr, ##__VA_ARGS__);
 #define UE_MEMSCOPE_DEFAULT(InTag)			FDefaultMemScope PREPROCESSOR_JOIN(MemScope,__LINE__)(InTag);
 #define UE_MEMSCOPE_UNINITIALIZED(Line)		FMemScope PREPROCESSOR_JOIN(MemScope,Line);
 #define UE_MEMSCOPE_ACTIVATE(Line, InTag)	PREPROCESSOR_JOIN(MemScope,Line).SetTagAndActivate(InTag);
