bool FArchiveState::IsMigratingRemoteObjects() const will now be compiled as a constexpr in builds without remote object support

diff --git a/Engine/Source/Runtime/Core/Private/Serialization/Archive.cpp b/Engine/Source/Runtime/Core/Private/Serialization/Archive.cpp
--- a/Engine/Source/Runtime/Core/Private/Serialization/Archive.cpp
+++ b/Engine/Source/Runtime/Core/Private/Serialization/Archive.cpp
@@ -196,7 +196,9 @@ PRAGMA_ENABLE_DEPRECATION_WARNINGS
 	ArMergeOverrides					= false;
 	ArPreserveArrayElements				= false;
 	ArShouldSkipUpdateCustomVersion		= false;
+#if UE_WITH_REMOTE_OBJECT_HANDLE
 	ArIsMigratingRemoteObjects			= false;
+#endif
 	SavePackageData						= nullptr;
 	SerializedProperty					= nullptr;
 
@@ -260,7 +262,9 @@ PRAGMA_ENABLE_DEPRECATION_WARNINGS
 	ArMergeOverrides					 = ArchiveToCopy.ArMergeOverrides;
 	ArPreserveArrayElements					 = ArchiveToCopy.ArPreserveArrayElements;
 	ArShouldSkipUpdateCustomVersion		 = ArchiveToCopy.ArShouldSkipUpdateCustomVersion;
+#if UE_WITH_REMOTE_OBJECT_HANDLE
 	ArIsMigratingRemoteObjects			 = ArchiveToCopy.ArIsMigratingRemoteObjects;
+#endif
 	SavePackageData						 = ArchiveToCopy.SavePackageData;
 	SerializedProperty					 = ArchiveToCopy.SerializedProperty;
 #if USE_STABLE_LOCALIZATION_KEYS
@@ -611,7 +615,11 @@ void FArchiveState::SetShouldSkipUpdateCustomVersion(bool bShouldSkip)
 
 void FArchiveState::SetMigratingRemoteObjects(bool bMigrating)
 {
+#if UE_WITH_REMOTE_OBJECT_HANDLE
 	ForEachState([bMigrating](FArchiveState& State) { State.ArIsMigratingRemoteObjects = bMigrating; });
+#else
+	checkf(false, TEXT("SetMigratingRemoteObjects() can only be used in a build with remote object handles enabled"))
+#endif
 }
 
 void FArchiveState::SetCustomVersion(const FGuid& Key, int32 Version, FName FriendlyName)
diff --git a/Engine/Source/Runtime/Core/Public/Serialization/Archive.h b/Engine/Source/Runtime/Core/Public/Serialization/Archive.h
--- a/Engine/Source/Runtime/Core/Public/Serialization/Archive.h
+++ b/Engine/Source/Runtime/Core/Public/Serialization/Archive.h
@@ -337,10 +337,17 @@ public:
 	CORE_API void SetMigratingRemoteObjects(bool bMigrating);
 
 	/** Returns true if this Archive is migrating remote object */
+#if UE_WITH_REMOTE_OBJECT_HANDLE
 	FORCEINLINE bool IsMigratingRemoteObjects() const
 	{
 		return ArIsMigratingRemoteObjects;
 	}
+#else
+	FORCEINLINE constexpr bool IsMigratingRemoteObjects() const
+	{
+		return false;
+	}
+#endif // UE_WITH_REMOTE_OBJECT_HANDLE
 
 	/** Returns true if this archive contains errors, which means that further serialization is generally not safe. */
 	FORCEINLINE bool IsError() const
@@ -831,8 +838,10 @@ private:
 	 * Setting this to true should improve serialization performance by skipping the step entirely */
 	uint8 ArShouldSkipUpdateCustomVersion : 1;
 
+#if UE_WITH_REMOTE_OBJECT_HANDLE
 	/** If true this archive is migrating remote object data. */
 	uint8 ArIsMigratingRemoteObjects : 1;
+#endif
 
 public:
 	/** Quickly tell if an archive contains script code. */
diff --git a/Engine/Source/Runtime/CoreUObject/Private/UObject/Obj.cpp b/Engine/Source/Runtime/CoreUObject/Private/UObject/Obj.cpp
--- a/Engine/Source/Runtime/CoreUObject/Private/UObject/Obj.cpp
+++ b/Engine/Source/Runtime/CoreUObject/Private/UObject/Obj.cpp
@@ -1927,12 +1927,7 @@ void UObject::SerializeScriptProperties( FStructuredArchive::FSlot Slot ) const
 
 		// When migrating remote objects the only instance where we serialize a CDO is when resetting an object to its archetype state in which case
 		// we want to serialize against the actual object class to get the right delta
-#if UE_WITH_REMOTE_OBJECT_HANDLE
-		const bool bMigratingRemoteObjects = UnderlyingArchive.IsMigratingRemoteObjects();
-#else
-		constexpr bool bMigratingRemoteObjects = false;
-#endif
-		UStruct* DiffClass = (HasAnyFlags(RF_ClassDefaultObject) && !bMigratingRemoteObjects) ? ObjClass->GetSuperClass() : ObjClass;
+		UStruct* DiffClass = (HasAnyFlags(RF_ClassDefaultObject) && !UnderlyingArchive.IsMigratingRemoteObjects()) ? ObjClass->GetSuperClass() : ObjClass;
 
 		// Query if this object data is being impersonated 
 		const UObject* ThisObject = this;
