Added [[nodiscard]] to TOptional.

diff --git a/Engine/Source/Runtime/Core/Public/Misc/Optional.h b/Engine/Source/Runtime/Core/Public/Misc/Optional.h
--- a/Engine/Source/Runtime/Core/Public/Misc/Optional.h
+++ b/Engine/Source/Runtime/Core/Public/Misc/Optional.h
@@ -60,18 +60,18 @@ private:
 
 public:
 	using ElementType = OptionalType;
-	
+
 	/** Construct an OptionalType with a valid value. */
-	TOptional(const OptionalType& InValue)
+	[[nodiscard]] TOptional(const OptionalType& InValue)
 		: TOptional(InPlace, InValue)
 	{
 	}
-	TOptional(OptionalType&& InValue)
+	[[nodiscard]] TOptional(OptionalType&& InValue)
 		: TOptional(InPlace, MoveTempIfPossible(InValue))
 	{
 	}
 	template <typename... ArgTypes>
-	explicit TOptional(EInPlace, ArgTypes&&... Args)
+	[[nodiscard]] explicit TOptional(EInPlace, ArgTypes&&... Args)
 	{
 		// If this fails to compile when trying to call TOptional(EInPlace, ...) with a non-public constructor,
 		// do not make TOptional a friend.
@@ -107,13 +107,13 @@ public:
 	}
 	
 	/** Construct an OptionalType with an invalid value. */
-	TOptional(FNullOpt)
+	[[nodiscard]] TOptional(FNullOpt)
 		: TOptional()
 	{
 	}
 
 	/** Construct an OptionalType with no value; i.e. unset */
-	TOptional()
+	[[nodiscard]] TOptional()
 	{
 		if constexpr (bUsingIntrusiveUnsetState)
 		{
@@ -131,7 +131,7 @@ public:
 	}
 
 	/** Copy/Move construction */
-	TOptional(const TOptional& Other)
+	[[nodiscard]] TOptional(const TOptional& Other)
 	{
 		bool bLocalIsSet = Other.IsSet();
 		if constexpr (!bUsingIntrusiveUnsetState)
@@ -150,7 +150,7 @@ public:
 			}
 		}
 	}
-	TOptional(TOptional&& Other)
+	[[nodiscard]] TOptional(TOptional&& Other)
 	{
 		bool bLocalIsSet = Other.IsSet();
 		if constexpr (!bUsingIntrusiveUnsetState)
@@ -287,7 +287,7 @@ public:
 		return *Result;
 	}
 
-	friend bool operator==(const TOptional& Lhs, const TOptional& Rhs)
+	[[nodiscard]] friend bool operator==(const TOptional& Lhs, const TOptional& Rhs)
 	{
 		if constexpr (!bUsingIntrusiveUnsetState)
 		{
@@ -306,7 +306,7 @@ public:
 		return (*(const OptionalType*)&Lhs.Value) == (*(const OptionalType*)&Rhs.Value);
 	}
 
-	friend bool operator!=(const TOptional& Lhs, const TOptional& Rhs)
+	[[nodiscard]] friend bool operator!=(const TOptional& Lhs, const TOptional& Rhs)
 	{
 		return !(Lhs == Rhs);
 	}
@@ -342,7 +342,7 @@ public:
 	}
 
 	/** @return true when the value is meaningful; false if calling GetValue() is undefined. */
-	bool IsSet() const
+	[[nodiscard]] bool IsSet() const
 	{
 		if constexpr (bUsingIntrusiveUnsetState)
 		{
@@ -353,52 +353,52 @@ public:
 			return Value.bIsSet;
 		}
 	}
-	FORCEINLINE explicit operator bool() const
+	[[nodiscard]] FORCEINLINE explicit operator bool() const
 	{
 		return IsSet();
 	}
 
 	/** @return The optional value; undefined when IsSet() returns false. */
-	OptionalType& GetValue()
+	[[nodiscard]] OptionalType& GetValue()
 	{
 		checkf(IsSet(), TEXT("It is an error to call GetValue() on an unset TOptional. Please either check IsSet() or use Get(DefaultValue) instead."));
 		return *(OptionalType*)&Value;
 	}
-	FORCEINLINE const OptionalType& GetValue() const
+	[[nodiscard]] FORCEINLINE const OptionalType& GetValue() const
 	{
 		return const_cast<TOptional*>(this)->GetValue();
 	}
 
-	OptionalType* operator->()
+	[[nodiscard]] OptionalType* operator->()
 	{
 		return &GetValue();
 	}
-	FORCEINLINE const OptionalType* operator->() const
+	[[nodiscard]] FORCEINLINE const OptionalType* operator->() const
 	{
 		return const_cast<TOptional*>(this)->operator->();
 	}
 
-	OptionalType& operator*()
+	[[nodiscard]] OptionalType& operator*()
 	{
 		return GetValue();
 	}
-	FORCEINLINE const OptionalType& operator*() const
+	[[nodiscard]] FORCEINLINE const OptionalType& operator*() const
 	{
 		return const_cast<TOptional*>(this)->operator*();
 	}
 
 	/** @return The optional value when set; DefaultValue otherwise. */
-	const OptionalType& Get(const OptionalType& DefaultValue UE_LIFETIMEBOUND) const UE_LIFETIMEBOUND
+	[[nodiscard]] const OptionalType& Get(const OptionalType& DefaultValue UE_LIFETIMEBOUND) const UE_LIFETIMEBOUND
 	{
 		return IsSet() ? *(const OptionalType*)&Value : DefaultValue;
 	}
 
 	/** @return A pointer to the optional value when set, nullptr otherwise. */
-	OptionalType* GetPtrOrNull()
+	[[nodiscard]] OptionalType* GetPtrOrNull()
 	{
 		return IsSet() ? (OptionalType*)&Value : nullptr;
 	}
-	FORCEINLINE const OptionalType* GetPtrOrNull() const
+	[[nodiscard]] FORCEINLINE const OptionalType* GetPtrOrNull() const
 	{
 		return const_cast<TOptional*>(this)->GetPtrOrNull();
 	}
@@ -424,7 +424,7 @@ FArchive& operator<<(FArchive& Ar, TOptional<OptionalType>& Optional)
 }
 
 template<typename OptionalType>
-inline auto GetTypeHash(const TOptional<OptionalType>& Optional) -> decltype(GetTypeHash(*Optional))
+[[nodiscard]] inline auto GetTypeHash(const TOptional<OptionalType>& Optional) -> decltype(GetTypeHash(*Optional))
 {
 	return Optional.IsSet() ? GetTypeHash(*Optional) : 0;
 }
diff --git a/Engine/Source/Runtime/CoreUObject/Public/UObject/OverriddenPropertySet.h b/Engine/Source/Runtime/CoreUObject/Public/UObject/OverriddenPropertySet.h
--- a/Engine/Source/Runtime/CoreUObject/Public/UObject/OverriddenPropertySet.h
+++ b/Engine/Source/Runtime/CoreUObject/Public/UObject/OverriddenPropertySet.h
@@ -179,7 +179,7 @@ inline TOptional<EOverriddenPropertyOperation> GetOverriddenOperationFromString(
 	int64 EnumValue = StaticEnum<EOverriddenPropertyOperation>()->GetValueByNameString(StaticEnum<EOverriddenPropertyOperation>()->GetName() + FString(TEXT("::")) + OverriddenOperationString);
 	if (EnumValue == INDEX_NONE)
 	{
-		TOptional<EOverriddenPropertyOperation>();
+		return TOptional<EOverriddenPropertyOperation>();
 	}
 
 	return TOptional((EOverriddenPropertyOperation)EnumValue);
