Added support to load Asset Registry FName batches in insertion order instead of shard order (useful when insertion order affects how custom name storage works such as for compressing names)

diff --git a/Engine/Source/Runtime/AssetRegistry/Private/AssetRegistryArchive.cpp b/Engine/Source/Runtime/AssetRegistry/Private/AssetRegistryArchive.cpp
--- a/Engine/Source/Runtime/AssetRegistry/Private/AssetRegistryArchive.cpp
+++ b/Engine/Source/Runtime/AssetRegistry/Private/AssetRegistryArchive.cpp
@@ -6,6 +6,9 @@
 #include "AssetRegistryPrivate.h"
 #include "AssetRegistry/AssetRegistryState.h"
 
+#ifndef ASSET_REGISTRY_LOAD_NAMES_IN_ORDER
+#define ASSET_REGISTRY_LOAD_NAMES_IN_ORDER 1 //0
+#endif
 
 constexpr uint32 AssetRegistryNumberedNameBit = 0x80000000;
 
@@ -269,7 +272,11 @@ FAssetRegistryReader::FAssetRegistryReader(FArchive& Inner, int32 NumWorkers, FA
 	// Use parallel loading if requested and available; the parallel path is only for latest version
 	if (NumWorkers > 0 && Header.Version == FAssetRegistryVersion::LatestVersion)
 	{
-		TFunction<TArray<FDisplayNameEntryId> ()> GetFutureNames = LoadNameBatchAsync(*this, NumWorkers);
+		ENameBatchLoadingFlags LoadingFlags = ENameBatchLoadingFlags::None;
+#if ASSET_REGISTRY_LOAD_NAMES_IN_ORDER
+		LoadingFlags |= ENameBatchLoadingFlags::RespectOrder;
+#endif
+		TFunction<TArray<FDisplayNameEntryId> ()> GetFutureNames = LoadNameBatchAsync(*this, NumWorkers, LoadingFlags);
 
 		FixedTagPrivate::FAsyncStoreLoader StoreLoader;
 		Task = StoreLoader.ReadInitialDataAndKickLoad(*this, NumWorkers, Header.Version);
@@ -279,7 +286,7 @@ FAssetRegistryReader::FAssetRegistryReader(FArchive& Inner, int32 NumWorkers, FA
 	}
 	else
 	{
-		Names = LoadNameBatch(Inner);
+		Names = LoadNameBatch(Inner, ENameBatchLoadingFlags::RespectOrder);
 		Tags = FixedTagPrivate::LoadStore(*this, Header.Version);
 	}
 }
diff --git a/Engine/Source/Runtime/Core/Private/UObject/UnrealNames.cpp b/Engine/Source/Runtime/Core/Private/UObject/UnrealNames.cpp
--- a/Engine/Source/Runtime/Core/Private/UObject/UnrealNames.cpp
+++ b/Engine/Source/Runtime/Core/Private/UObject/UnrealNames.cpp
@@ -5172,7 +5172,7 @@ static void LoadSeparatedNameBatchInShardOrder(	TArray<FDisplayNameEntryId>& Out
 #endif
 }
 
-void LoadNameBatch(TArray<FDisplayNameEntryId>& OutNames, TArrayView<const uint8> NameData, TArrayView<const uint8> HashData)
+void LoadNameBatch(TArray<FDisplayNameEntryId>& OutNames, TArrayView<const uint8> NameData, TArrayView<const uint8> HashData, ENameBatchLoadingFlags Flags)
 {
 	check(IsAligned(NameData.GetData(), sizeof(uint64)));
 	check(IsAligned(HashData.GetData(), sizeof(uint64)));
@@ -5190,7 +5190,7 @@ void LoadNameBatch(TArray<FDisplayNameEntryId>& OutNames, TArrayView<const uint8
 	{
 		LoadInterleavedNameBatchInInputOrder(OutNames, NameData);
 	}
-	else if (Num < FNamePoolShards)
+	else if (Num < FNamePoolShards || EnumHasAnyFlags(Flags, ENameBatchLoadingFlags::RespectOrder))
 	{
 		LoadInterleavedNameBatchInInputOrder(OutNames, Hashes, NameData);
 	}
@@ -5206,6 +5206,12 @@ struct FNameBatchLoader
 	TArrayView<const FSerializedNameHeader> Headers;
 	TArrayView<const uint8> Strings;
 	TArray<uint8> Data;
+	ENameBatchLoadingFlags Flags;
+	
+	FNameBatchLoader(ENameBatchLoadingFlags InFlags)
+		: Flags(InFlags)
+	{
+	}
 
 	// @return true if there's anything to load
 	bool Read(FArchive& Ar)
@@ -5250,7 +5256,7 @@ struct FNameBatchLoader
 		{
 			LoadSeparatedNameBatchInInputOrder(Out, Headers, Strings);
 		}
-		else if (Headers.Num() < FNamePoolShards)
+		else if (Headers.Num() < FNamePoolShards || EnumHasAnyFlags(Flags, ENameBatchLoadingFlags::RespectOrder))
 		{
 			LoadSeparatedNameBatchInInputOrder(Out, Hashes.GetData(), Headers, Strings);
 		}
@@ -5263,9 +5269,9 @@ struct FNameBatchLoader
 	}
 };
 
-TArray<FDisplayNameEntryId> LoadNameBatch(FArchive& Ar)
+TArray<FDisplayNameEntryId> LoadNameBatch(FArchive& Ar, ENameBatchLoadingFlags Flags)
 {
-	FNameBatchLoader Loader;
+	FNameBatchLoader Loader(Flags);
 
 	if (Loader.Read(Ar))
 	{
@@ -5276,7 +5282,12 @@ TArray<FDisplayNameEntryId> LoadNameBatch(FArchive& Ar)
 }
 
 struct FNameBatchAsyncLoader : FNameBatchLoader
-{	
+{
+	FNameBatchAsyncLoader(ENameBatchLoadingFlags InFlags)
+		: FNameBatchLoader(InFlags)
+	{
+	}
+	
 	~FNameBatchAsyncLoader()
 	{
 		FGenericPlatformProcess::ReturnSynchEventToPool(/* can be null */ DoneEvent);
@@ -5295,7 +5306,14 @@ struct FNameBatchAsyncLoader : FNameBatchLoader
 
 	void DoWork()
 	{
-		LoadSeparatedNameBatchInShardOrder(Out, Hashes.GetData(), Headers, Strings);
+		if (EnumHasAnyFlags(Flags, ENameBatchLoadingFlags::RespectOrder))
+		{
+			LoadSeparatedNameBatchInInputOrder(Out, Hashes.GetData(), Headers, Strings);
+		}
+		else
+		{
+			LoadSeparatedNameBatchInShardOrder(Out, Hashes.GetData(), Headers, Strings);
+		}
 		DoneEvent->Trigger();
 	}
 
@@ -5314,12 +5332,12 @@ struct FNameBatchAsyncLoader : FNameBatchLoader
 	TArray<FDisplayNameEntryId> Out;
 };
 
-TFunction<TArray<FDisplayNameEntryId>()> LoadNameBatchAsync(FArchive& Ar, uint32 MaxWorkers)
+TFunction<TArray<FDisplayNameEntryId>()> LoadNameBatchAsync(FArchive& Ar, uint32 MaxWorkers, ENameBatchLoadingFlags Flags)
 {
 	check(MaxWorkers > 0);
 
 	// Ref-count since worker tasks may outlive the returned function
-	TSharedPtr<FNameBatchAsyncLoader, ESPMode::ThreadSafe> Loader(new FNameBatchAsyncLoader);
+	TSharedPtr<FNameBatchAsyncLoader, ESPMode::ThreadSafe> Loader(new FNameBatchAsyncLoader(Flags));
 
 	if (Loader->Read(Ar))
 	{
diff --git a/Engine/Source/Runtime/Core/Public/UObject/NameBatchSerialization.h b/Engine/Source/Runtime/Core/Public/UObject/NameBatchSerialization.h
--- a/Engine/Source/Runtime/Core/Public/UObject/NameBatchSerialization.h
+++ b/Engine/Source/Runtime/Core/Public/UObject/NameBatchSerialization.h
@@ -25,17 +25,23 @@ CORE_API void SaveNameBatch(TArrayView<const FDisplayNameEntryId> Names, FArchiv
 // @param Bytes for existing and new names.
 CORE_API void ReserveNameBatch(uint32 NameDataBytes, uint32 HashDataBytes);
 
+enum class ENameBatchLoadingFlags : uint32 {
+	None         = 0,
+	RespectOrder = 1 << 0,
+};
+ENUM_CLASS_FLAGS(ENameBatchLoadingFlags);
+
 // Load a name blob with precalculated hashes.
 //
 // Names are rehashed if hash algorithm version doesn't match.
 //
 // @param NameData, HashData must be 8-byte aligned.
-CORE_API void LoadNameBatch(TArray<FDisplayNameEntryId>& OutNames, TArrayView<const uint8> NameData, TArrayView<const uint8> HashData);
+CORE_API void LoadNameBatch(TArray<FDisplayNameEntryId>& OutNames, TArrayView<const uint8> NameData, TArrayView<const uint8> HashData, ENameBatchLoadingFlags Flags = ENameBatchLoadingFlags::None);
 
 // Load names and precalculated hashes from an archive
 //
 // Names are rehashed if hash algorithm version doesn't match.
-CORE_API TArray<FDisplayNameEntryId> LoadNameBatch(FArchive& Ar);
+CORE_API TArray<FDisplayNameEntryId> LoadNameBatch(FArchive& Ar, ENameBatchLoadingFlags Flags = ENameBatchLoadingFlags::None);
 
 // Load names and precalculated hashes from an archive using multiple workers
 //
@@ -46,6 +52,6 @@ CORE_API TArray<FDisplayNameEntryId> LoadNameBatch(FArchive& Ar);
 // @param Ar is drained synchronously
 // @param MaxWorkers > 0
 // @return function that waits before returning result, like a simple future.
-CORE_API TFunction<TArray<FDisplayNameEntryId>()> LoadNameBatchAsync(FArchive& Ar, uint32 MaxWorkers);
+CORE_API TFunction<TArray<FDisplayNameEntryId>()> LoadNameBatchAsync(FArchive& Ar, uint32 MaxWorkers, ENameBatchLoadingFlags Flags = ENameBatchLoadingFlags::None);
 
 //////////////////////////////////////////////////////////////////////////
\ No newline at end of file
