[AutoRTFM] Add doc-comment for GAutoRTFMEnabledProbability.

diff --git a/Engine/Source/Runtime/Core/Private/AutoRTFM/API.cpp b/Engine/Source/Runtime/Core/Private/AutoRTFM/API.cpp
--- a/Engine/Source/Runtime/Core/Private/AutoRTFM/API.cpp
+++ b/Engine/Source/Runtime/Core/Private/AutoRTFM/API.cpp
@@ -26,7 +26,10 @@ namespace
 
 	int GAutoRTFMInternalAbortAction = AutoRTFM::ForTheRuntime::EAutoRTFMInternalAbortActionState::Crash;
 	int GAutoRTFMRetryTransactions = AutoRTFM::ForTheRuntime::EAutoRTFMRetryTransactionState::NoRetry;
-
+	
+	// Set the percentage chance [0..100] that AutoRTFM will be enabled.
+	// 100.0 means that AutoRTFM will always be enabled; 1.0 means that AutoRTFM has a 1% chance of being enabled.
+	// See `AutoRTFM::ForTheRuntime::CoinTossDisable` for implementation details.
 	float GAutoRTFMEnabledProbability = 100.0f;
 
 	// Note: GMemoryValidationLevel should never be EMemoryValidationLevel::Default.
