[AutoRTFM] Fix lit tests by including header for va_list.

This dependency was recently added when autortfm_extern_api was introduced.

diff --git a/Engine/Source/Runtime/Core/Public/AutoRTFM/AutoRTFM.h b/Engine/Source/Runtime/Core/Public/AutoRTFM/AutoRTFM.h
--- a/Engine/Source/Runtime/Core/Public/AutoRTFM/AutoRTFM.h
+++ b/Engine/Source/Runtime/Core/Public/AutoRTFM/AutoRTFM.h
@@ -57,6 +57,7 @@
 #include "AutoRTFMConstants.h"
 
 #if defined(UE_AUTORTFM_STANDALONE)
+#include <stdarg.h> // for va_list
 #include <stdlib.h>
 #include <memory.h>
 #include <type_traits>
