Make `SetOpenActive` codegen optimal for the case the validator is disabled (the default).

diff --git a/Engine/Source/Runtime/Core/Private/AutoRTFM/Transaction.cpp b/Engine/Source/Runtime/Core/Private/AutoRTFM/Transaction.cpp
--- a/Engine/Source/Runtime/Core/Private/AutoRTFM/Transaction.cpp
+++ b/Engine/Source/Runtime/Core/Private/AutoRTFM/Transaction.cpp
@@ -289,13 +289,13 @@ bool FTransaction::AttemptToCommitOuterNest()
     return true;
 }
 
-void FTransaction::SetOpenActive(EMemoryValidationLevel NewMemoryValidationLevel, const void* ReturnAddress)
+void FTransaction::SetOpenActiveValidatorEnabled(EMemoryValidationLevel NewMemoryValidationLevel, const void* ReturnAddress)
 {
+	ASSERT(NewMemoryValidationLevel != EMemoryValidationLevel::Disabled);
 	CurrentMemoryValidationLevel = NewMemoryValidationLevel;
 	CurrentOpenReturnAddress = ReturnAddress;
 
-	if (CurrentMemoryValidationLevel != EMemoryValidationLevel::Disabled &&
-		AutoRTFM::ForTheRuntime::GetMemoryValidationThrottlingEnabled())
+	if (AutoRTFM::ForTheRuntime::GetMemoryValidationThrottlingEnabled())
 	{
 		FOpenHashThrottler& Throttler = Context->GetOpenHashThrottler();
 		if (!Throttler.ShouldHashFor(ReturnAddress))
@@ -305,13 +305,24 @@ void FTransaction::SetOpenActive(EMemoryValidationLevel NewMemoryValidationLevel
 		Throttler.Update();
 	}
 
-	if (CurrentMemoryValidationLevel == EMemoryValidationLevel::Disabled)
+	SetState<EState::OpenActive>();
+}
+
+void FTransaction::SetOpenActive(EMemoryValidationLevel NewMemoryValidationLevel, const void* ReturnAddress)
+{
+	if (UNLIKELY(NewMemoryValidationLevel != EMemoryValidationLevel::Disabled))
 	{
-		// TODO: Validate if open -> open with different validation levels.
-		RecordedWriteHash = 0;
-		NumWriteLogsHashed = 0;
-		CurrentOpenReturnAddress = nullptr;
+		AUTORTFM_MUST_TAIL return SetOpenActiveValidatorEnabled(NewMemoryValidationLevel, ReturnAddress);
 	}
+
+	CurrentMemoryValidationLevel = NewMemoryValidationLevel;
+	CurrentOpenReturnAddress = ReturnAddress;
+	
+	// TODO: Validate if open -> open with different validation levels.
+	RecordedWriteHash = 0;
+	NumWriteLogsHashed = 0;
+	CurrentOpenReturnAddress = nullptr;
+
 	SetState<EState::OpenActive>();
 }
 
diff --git a/Engine/Source/Runtime/Core/Private/AutoRTFM/Transaction.h b/Engine/Source/Runtime/Core/Private/AutoRTFM/Transaction.h
--- a/Engine/Source/Runtime/Core/Private/AutoRTFM/Transaction.h
+++ b/Engine/Source/Runtime/Core/Private/AutoRTFM/Transaction.h
@@ -187,6 +187,8 @@ private:
 
 	bool ShouldRecordWrite(void* LogicalAddress) const;
 
+	void SetOpenActiveValidatorEnabled(EMemoryValidationLevel MemoryValidationLevel, const void* ReturnAddress);
+
 	template<EState NewState>
 	void SetState();
 
