Make `FunctionMapLookup` codegen optimal in the overwhelmingly likely case to make it 7.5% faster.

diff --git a/Engine/Source/Runtime/Core/Private/AutoRTFM/FunctionMapInlines.h b/Engine/Source/Runtime/Core/Private/AutoRTFM/FunctionMapInlines.h
--- a/Engine/Source/Runtime/Core/Private/AutoRTFM/FunctionMapInlines.h
+++ b/Engine/Source/Runtime/Core/Private/AutoRTFM/FunctionMapInlines.h
@@ -12,6 +12,19 @@
 namespace AutoRTFM
 {
 
+inline void* FunctionMapLookupSlow(void* OldFunction, const char* Where)
+{
+	// Instead fall back to the slower function map lookup.
+    void* const Result = FunctionMapTryLookup(OldFunction);
+
+	if (UNLIKELY(!Result))
+	{
+		AUTORTFM_MUST_TAIL return FunctionMapLookupTryFollowingRelocation(OldFunction, Where);
+	}
+
+	return Result;
+}
+
 inline void* FunctionMapLookup(void* OldFunction, const char* Where)
 {
 	// We use prefix data in our custom LLVM pass to stuff some data just
@@ -38,16 +51,8 @@ inline void* FunctionMapLookup(void* OldFunction, const char* Where)
 			return reinterpret_cast<void*>(PrefixData & 0x0000ffffffffffff);
 		}
 	}
-
-	// Instead fall back to the slower function map lookup.
-    void* const Result = FunctionMapTryLookup(OldFunction);
-
-	if (UNLIKELY(!Result))
-	{
-		AUTORTFM_MUST_TAIL return FunctionMapLookupTryFollowingRelocation(OldFunction, Where);
-	}
-
-	return Result;
+	
+	AUTORTFM_MUST_TAIL return FunctionMapLookupSlow(OldFunction, Where);
 }
 
 template<typename TReturnType, typename... TParameterTypes>
