Fix the `ASSERT` failure in the build.

diff --git a/Engine/Source/Runtime/Core/Private/AutoRTFM/Transaction.cpp b/Engine/Source/Runtime/Core/Private/AutoRTFM/Transaction.cpp
--- a/Engine/Source/Runtime/Core/Private/AutoRTFM/Transaction.cpp
+++ b/Engine/Source/Runtime/Core/Private/AutoRTFM/Transaction.cpp
@@ -291,7 +291,7 @@ bool FTransaction::AttemptToCommitOuterNest()
 
 void FTransaction::SetOpenActiveValidatorEnabled(EMemoryValidationLevel NewMemoryValidationLevel, const void* ReturnAddress)
 {
-	ASSERT(NewMemoryValidationLevel != EMemoryValidationLevel::Disabled);
+	AUTORTFM_ASSERT(NewMemoryValidationLevel != EMemoryValidationLevel::Disabled);
 	CurrentMemoryValidationLevel = NewMemoryValidationLevel;
 	CurrentOpenReturnAddress = ReturnAddress;
 
