Fix the Linux specific heartbeat suspension within a transaction.

Just mimic the non-linux heartbeat to linux too.

diff --git a/Engine/Source/Programs/AutoRTFMTests/Private/UECoreTests.cpp b/Engine/Source/Programs/AutoRTFMTests/Private/UECoreTests.cpp
--- a/Engine/Source/Programs/AutoRTFMTests/Private/UECoreTests.cpp
+++ b/Engine/Source/Programs/AutoRTFMTests/Private/UECoreTests.cpp
@@ -9,6 +9,7 @@
 #include "Delegates/IDelegateInstance.h"
 #include "HAL/IConsoleManager.h"
 #include "HAL/MallocLeakDetection.h"
+#include "HAL/ThreadHeartBeat.h"
 #include "HAL/ThreadSingleton.h"
 #include "Internationalization/TextCache.h"
 #include "Internationalization/TextFormatter.h"
@@ -2950,3 +2951,17 @@ TEST_CASE("UECore.ModuleManager")
 		});
 	}
 }
+
+TEST_CASE("UECore.HeartBeat")
+{
+	AutoRTFM::Testing::Abort([&]
+	{
+		FDisableHitchDetectorScope _;
+		AutoRTFM::AbortTransaction();
+	});
+
+	AutoRTFM::Testing::Commit([&]
+	{
+		FDisableHitchDetectorScope _;
+	});
+}
diff --git a/Engine/Source/Runtime/Core/Private/Unix/UnixSignalHeartBeat.cpp b/Engine/Source/Runtime/Core/Private/Unix/UnixSignalHeartBeat.cpp
--- a/Engine/Source/Runtime/Core/Private/Unix/UnixSignalHeartBeat.cpp
+++ b/Engine/Source/Runtime/Core/Private/Unix/UnixSignalHeartBeat.cpp
@@ -7,6 +7,7 @@
 #include "Misc/ConfigCacheIni.h"
 #include "Logging/LogMacros.h"
 #include "HAL/PlatformTime.h"
+#include "AutoRTFM/AutoRTFM.h"
 
 #include <signal.h>
 #include <unistd.h>
@@ -244,14 +245,22 @@ void FUnixSignalGameHitchHeartBeat::SuspendHeartBeat()
 
 	if (bTimerCreated)
 	{
-		struct itimerspec DisarmTime;
-		FMemory::Memzero(DisarmTime);
+		UE_AUTORTFM_OPEN
+		{
+			struct itimerspec DisarmTime;
+			FMemory::Memzero(DisarmTime);
+
+			if (timer_settime(TimerId, 0, &DisarmTime, nullptr) == -1)
+			{
+				int Errno = errno;
+				UE_LOG(LogUnixHeartBeat, Warning, TEXT("Failed to timer_settime() errno=%d (%s)"), Errno, UTF8_TO_TCHAR(strerror(Errno)));
+			}
+		};
 
-		if (timer_settime(TimerId, 0, &DisarmTime, nullptr) == -1)
+		UE_AUTORTFM_ONABORT(this)
 		{
-			int Errno = errno;
-			UE_LOG(LogUnixHeartBeat, Warning, TEXT("Failed to timer_settime() errno=%d (%s)"), Errno, UTF8_TO_TCHAR(strerror(Errno)));
-		}
+			ResumeHeartBeat();
+		};
 	}
 #endif
 }
@@ -264,13 +273,16 @@ void FUnixSignalGameHitchHeartBeat::ResumeHeartBeat()
 		return;
 	}
 
-	if( SuspendCount > 0)
+	UE_AUTORTFM_ONCOMMIT(this)
 	{
-		SuspendCount--;
-		UE_LOG(LogUnixHeartBeat, Verbose, TEXT("SuspendCount:%d"), SuspendCount);
+		if (SuspendCount > 0)
+		{
+			SuspendCount--;
+			UE_LOG(LogUnixHeartBeat, Verbose, TEXT("SuspendCount:%d"), SuspendCount);
 
-		FrameStart(true);
-	}
+			FrameStart(true);
+		}
+	};
 #endif
 }
 
