[AutoRTFM] Optimize write log hashing with AVX2, drop use of xxhash.

Offers up to a 20x speed improvement over xxhash.

Part of the work to separate AutoRTFM out of core to break cyclic build dependency issues.

diff --git a/Engine/Source/Runtime/Core/Private/AutoRTFM/WriteLog.cpp b/Engine/Source/Runtime/Core/Private/AutoRTFM/WriteLog.cpp
--- a/Engine/Source/Runtime/Core/Private/AutoRTFM/WriteLog.cpp
+++ b/Engine/Source/Runtime/Core/Private/AutoRTFM/WriteLog.cpp
@@ -3,19 +3,137 @@
 #if (defined(__AUTORTFM) && __AUTORTFM)
 
 #include "WriteLog.h"
-#include "Hash/xxhash.h"
+#include "Utils.h"
 
 namespace AutoRTFM
 {
 
+#if PLATFORM_CPU_X86_FAMILY
+#define WRITELOG_HASH_MAYBE_HAS_AVX2 1
+#include <immintrin.h>
+#else
+#define WRITELOG_HASH_MAYBE_HAS_AVX2 0
+#endif
+
+namespace
+{
+
+// If Data is aligned to Size, and Size is either 1, 2, 4 or 8 bytes then Hash
+// is modified and SmallHash() returns true, otherwise the call is a no-op and
+// SmallHash() returns false.
+UE_AUTORTFM_FORCEINLINE
+bool SmallHash(const std::byte* Data, size_t Size, FWriteLog::FHash& Hash)
+{
+	switch (Size)
+	{
+		case 8:
+			if (AUTORTFM_LIKELY((reinterpret_cast<uintptr_t>(Data) & 7) == 0))
+			{
+				uint64_t Load = *reinterpret_cast<uint64_t const*>(Data);
+				Hash = (Hash * 31) ^ static_cast<FWriteLog::FHash>(Load);
+				return true;
+			}
+			break;
+
+		case 4:
+			if (AUTORTFM_LIKELY((reinterpret_cast<uintptr_t>(Data) & 3) == 0))
+			{
+				uint32_t Load = *reinterpret_cast<uint32_t const*>(Data);
+				Hash = (Hash * 31) ^ static_cast<FWriteLog::FHash>(Load);
+				return true;
+			}
+			break;
+
+		case 2:
+			if (AUTORTFM_LIKELY((reinterpret_cast<uintptr_t>(Data) & 1) == 0))
+			{
+				uint16 Load = *reinterpret_cast<uint16 const*>(Data);
+				Hash = (Hash * 31) ^ static_cast<FWriteLog::FHash>(Load);
+				return true;
+			}
+			break;
+
+		case 1:
+			Hash = (Hash * 31) ^ static_cast<FWriteLog::FHash>(Data[0]);
+			return true;
+	}
+	return false;
+}
+
+}
+
 FWriteLog::FHash FWriteLog::Hash(size_t NumWriteEntries) const
 {
-	FXxHash64Builder Hasher;
+#if WRITELOG_HASH_MAYBE_HAS_AVX2
+	// __builtin_cpu_supports("avx2") produces a linker error.
+	// As memory validation is an debug, opt-in feature, and most modern CPUs
+	// support AVX2, assume for now we have support.
+	constexpr bool bHasAVX2 = true /* __builtin_cpu_supports("avx2") */;
+	if (bHasAVX2)
+	{
+		AUTORTFM_MUST_TAIL return HashAVX2(NumWriteEntries);
+	}
+#endif
+
+	size_t WriteIndex = 0;
+
+	FWriteLog::FHash Hash = 0;
+
+	for(auto Iter = begin(); Iter != end(); ++Iter, ++WriteIndex)
+	{
+		if (WriteIndex == NumWriteEntries)
+		{
+			break;
+		}
+
+		const FWriteLogEntry& Entry = *Iter;
+		if (AUTORTFM_UNLIKELY(Entry.bNoMemoryValidation))
+		{
+			continue; // Next write
+		}
+
+		const std::byte* Data = Entry.LogicalAddress;
+		const size_t Size = Entry.Size;
+
+		if (!SmallHash(Data, Size, Hash))
+		{
+			for (size_t I = 0; I < Size; I++)
+			{
+				Hash = (Hash * 31) ^ static_cast<FWriteLog::FHash>(Data[I]);
+			}
+		}
+	}
+
+	return Hash;
+}
 
+#if WRITELOG_HASH_MAYBE_HAS_AVX2
+__attribute__((__target__("avx2")))
+__attribute__((no_sanitize("address"))) // Intentionally reading whole-vectors, which go beyond write bounds.
+FWriteLog::FHash FWriteLog::HashAVX2(size_t NumWriteEntries) const
+{
 	size_t WriteIndex = 0;
 
-	using FSmallValue = uint64;
-	FSmallValue SmallValue = 0;
+	FWriteLog::FHash Hash = 0;
+
+	using i8x32 = __m256i;
+	using u64x4 = __m256i;
+	const i8x32 Vec0To31 = _mm256_setr_epi8( // [0..31]
+		0x00, 0x01, 0x02, 0x03, 0x04, 0x05, 0x06, 0x07,
+		0x08, 0x09, 0x0a, 0x0b, 0x0c, 0x0d, 0x0e, 0x0f,
+		0x10, 0x11, 0x12, 0x13, 0x14, 0x15, 0x16, 0x17,
+		0x18, 0x19, 0x1a, 0x1b, 0x1c, 0x1d, 0x1e, 0x1f);
+	const i8x32 Vec1To32 = _mm256_setr_epi8( // [1..32]
+		0x01, 0x02, 0x03, 0x04, 0x05, 0x06, 0x07, 0x08,
+		0x09, 0x0a, 0x0b, 0x0c, 0x0d, 0x0e, 0x0f, 0x10,
+		0x11, 0x12, 0x13, 0x14, 0x15, 0x16, 0x17, 0x18,
+		0x19, 0x1a, 0x1b, 0x1c, 0x1d, 0x1e, 0x1f, 0x20);
+
+	// We use streaming vector intrinsics for the hash.
+	// Flush once at the start to ensure all writes are visible.
+	_mm_sfence();
+
+	i8x32 VecHash{};
 
 	for(auto Iter = begin(); Iter != end(); ++Iter, ++WriteIndex)
 	{
@@ -25,63 +143,65 @@ FWriteLog::FHash FWriteLog::Hash(size_t NumWriteEntries) const
 		}
 
 		const FWriteLogEntry& Entry = *Iter;
-		if (UNLIKELY(Entry.bNoMemoryValidation))
+		if (AUTORTFM_UNLIKELY(Entry.bNoMemoryValidation))
 		{
 			continue; // Next write
 		}
 
-		if (Entry.Size > sizeof(FSmallValue))
+		const std::byte* Data = Entry.LogicalAddress;
+		size_t Size = Entry.Size;
+
+		if (SmallHash(Data, Size, Hash))
 		{
-			Hasher.Update(Entry.LogicalAddress, Entry.Size);
+			continue; // Next write
 		}
 
-		const uintptr_t Address = reinterpret_cast<uintptr_t>(Entry.LogicalAddress);
-		switch (Entry.Size)
+		if (const uintptr_t NumAlignmentBytes = reinterpret_cast<uintptr_t>(Data) & 31; NumAlignmentBytes != 0)
 		{
-			case 8:
-				if (LIKELY((Address & 7) == 0))
-				{
-					uint64 Load = *reinterpret_cast<uint64 const*>(Entry.LogicalAddress);
-					SmallValue = (SmallValue * 31) ^ static_cast<FSmallValue>(Load);
-					continue;
-				}
-				break;
-
-			case 4:
-				if (LIKELY((Address & 3) == 0))
-				{
-					uint32 Load = *reinterpret_cast<uint32 const*>(Entry.LogicalAddress);
-					SmallValue = (SmallValue * 31) ^ static_cast<FSmallValue>(Load);
-					continue;
-				}
-				break;
-
-			case 2:
-				if (LIKELY((Address & 1) == 0))
-				{
-					uint16 Load = *reinterpret_cast<uint16 const*>(Entry.LogicalAddress);
-					SmallValue = (SmallValue * 31) ^ static_cast<FSmallValue>(Load);
-					continue;
-				}
-				break;
-
-			case 1:
-				SmallValue = (SmallValue * 31) ^ static_cast<FSmallValue>(Entry.LogicalAddress[0]);
-				continue;
+			// Data is not 32-byte aligned.
+			// Perform a vector load at the aligned-down address and mask the
+			// bytes that we're interested in.
+			const char MaskStart = static_cast<char>(NumAlignmentBytes);
+			const char MaskEnd = static_cast<char>(std::min<size_t>(NumAlignmentBytes + Size, 127));
+			const i8x32 VecMaskStart = _mm256_cmpgt_epi8(Vec1To32, _mm256_set1_epi8(MaskStart));
+			const i8x32 VecMaskEnd = _mm256_cmpgt_epi8(_mm256_set1_epi8(MaskEnd), Vec0To31);
+			const i8x32 Mask = _mm256_and_si256(VecMaskStart, VecMaskEnd);
+			const i8x32 Load = _mm256_stream_load_si256(Data - NumAlignmentBytes);
+			const i8x32 MaskedLoad = _mm256_and_si256(Load, Mask);
+			VecHash = VecHash ^ _mm256_bslli_epi128(VecHash, 5) ^ MaskedLoad;
+			const size_t NumBytesConsumed = std::min(32 - NumAlignmentBytes, Size);
+			Data += NumBytesConsumed;
+			Size -= NumBytesConsumed;
 		}
 
-		// Unaligned / non-POT write.
-		for (size_t I = 0; I < Entry.Size; I++)
+		// Data is now 32-byte aligned, so we can hash in whole vectors.
+		while (Size >= 32)
 		{
-			SmallValue = (SmallValue * 31) ^ static_cast<FSmallValue>(Entry.LogicalAddress[I]);
+			const i8x32 Load = _mm256_stream_load_si256(Data);
+			VecHash = VecHash ^ _mm256_bslli_epi128(VecHash, 5) ^ Load;
+			Data += 32;
+			Size -= 32;
+		}
+
+		// Any trailing bytes require more masking.
+		if (Size > 0)
+		{
+			const i8x32 VecSize = _mm256_set1_epi8(static_cast<char>(Size));
+			const i8x32 Mask = _mm256_cmpgt_epi8(VecSize, Vec0To31);
+			const i8x32 Load = _mm256_stream_load_si256(Data);
+			const i8x32 MaskedLoad = _mm256_and_si256(Load, Mask);
+			VecHash = VecHash ^ _mm256_bslli_epi128(VecHash, 5) ^ MaskedLoad;
 		}
 	}
-	
-	Hasher.Update(&SmallValue, sizeof(SmallValue));
 
-	return Hasher.Finalize().Hash;
+	Hash = (Hash * 31) ^ __builtin_reduce_xor(VecHash);
+
+	return Hash;
 }
+#endif // WRITELOG_HASH_MAYBE_HAS_AVX2
+
+#undef WRITELOG_HASH_MAYBE_HAS_AVX2
 
-} // namespace AutoRTFM
+}
 
 #endif // (defined(__AUTORTFM) && __AUTORTFM)
diff --git a/Engine/Source/Runtime/Core/Private/AutoRTFM/WriteLog.h b/Engine/Source/Runtime/Core/Private/AutoRTFM/WriteLog.h
--- a/Engine/Source/Runtime/Core/Private/AutoRTFM/WriteLog.h
+++ b/Engine/Source/Runtime/Core/Private/AutoRTFM/WriteLog.h
@@ -382,7 +382,9 @@ namespace AutoRTFM
 		static constexpr size_t HeadBlockSize = 256;
 		static_assert(IsAlignedForTRecord<HeadBlockSize>);
 		static_assert(IsAlignedForTRecord<DefaultBlockSize>);
-		
+
+		FHash HashAVX2(size_t NumWriteEntries) const;
+
 		FBlock* HeadBlock = reinterpret_cast<FBlock*>(HeadBlockMemory);
 		FBlock* TailBlock = reinterpret_cast<FBlock*>(HeadBlockMemory);
 		size_t NumEntries = 0;
