Logging: Added conditional structured logging macros

UE_CLOGFMT
UE_CLOGFMT_EX
UE_CLOGFMT_LOC
UE_CLOGFMT_LOC_EX
UE_CLOGFMT_NSLOC
UE_CLOGFMT_NSLOC_EX

diff --git a/Engine/Source/Editor/UnrealEd/Classes/Commandlets/GatherTextFromSourceCommandlet.h b/Engine/Source/Editor/UnrealEd/Classes/Commandlets/GatherTextFromSourceCommandlet.h
--- a/Engine/Source/Editor/UnrealEd/Classes/Commandlets/GatherTextFromSourceCommandlet.h
+++ b/Engine/Source/Editor/UnrealEd/Classes/Commandlets/GatherTextFromSourceCommandlet.h
@@ -7,6 +7,7 @@
 #include "Framework/Commands/Commands.h"
 #include "Commandlets/GatherTextCommandletBase.h"
 #include "Internationalization/StringTableCore.h"
+#include "Misc/EnumClassFlags.h"
 #include "GatherTextFromSourceCommandlet.generated.h"
 
 class Error;
@@ -573,9 +574,12 @@ private:
 		enum class EFlags : int32
 		{
 			None = 0,
-			Namespace = 1,
+			Condition = 1 << 0,
+			Namespace = 1 << 1,
 		};
 
+		static int32 CalculateMinimumArgumentCount(EFlags Flags);
+
 		FStructuredLogMacroDescriptor(const TCHAR* Name, EFlags Flags);
 
 		virtual void TryParse(const FString& Text, FSourceFileParseContext& Context) const override;
@@ -584,6 +588,8 @@ private:
 		EFlags Flags;
 	};
 
+	FRIEND_ENUM_CLASS_FLAGS(FStructuredLogMacroDescriptor::EFlags);
+
 	class FIniNamespaceDescriptor : public FPreProcessorDescriptor
 	{
 	public:
@@ -625,3 +631,5 @@ public:
 	static void LogStats();
 #undef LOC_DEFINE_REGION
 };
+
+ENUM_CLASS_FLAGS(UGatherTextFromSourceCommandlet::FStructuredLogMacroDescriptor::EFlags);
diff --git a/Engine/Source/Editor/UnrealEd/Private/Commandlets/GatherTextFromSourceCommandlet.cpp b/Engine/Source/Editor/UnrealEd/Private/Commandlets/GatherTextFromSourceCommandlet.cpp
--- a/Engine/Source/Editor/UnrealEd/Private/Commandlets/GatherTextFromSourceCommandlet.cpp
+++ b/Engine/Source/Editor/UnrealEd/Private/Commandlets/GatherTextFromSourceCommandlet.cpp
@@ -475,6 +475,10 @@ void UGatherTextFromSourceCommandlet::GetParsables(TArray<FParsableDescriptor*>&
 		Parsables.Add(new FStructuredLogMacroDescriptor(TEXT("UE_LOGFMT_LOC_EX"), FStructuredLogMacroDescriptor::EFlags::None));
 		Parsables.Add(new FStructuredLogMacroDescriptor(TEXT("UE_LOGFMT_NSLOC"), FStructuredLogMacroDescriptor::EFlags::Namespace));
 		Parsables.Add(new FStructuredLogMacroDescriptor(TEXT("UE_LOGFMT_NSLOC_EX"), FStructuredLogMacroDescriptor::EFlags::Namespace));
+		Parsables.Add(new FStructuredLogMacroDescriptor(TEXT("UE_CLOGFMT_LOC"), FStructuredLogMacroDescriptor::EFlags::Condition));
+		Parsables.Add(new FStructuredLogMacroDescriptor(TEXT("UE_CLOGFMT_LOC_EX"), FStructuredLogMacroDescriptor::EFlags::Condition));
+		Parsables.Add(new FStructuredLogMacroDescriptor(TEXT("UE_CLOGFMT_NSLOC"), FStructuredLogMacroDescriptor::EFlags::Condition | FStructuredLogMacroDescriptor::EFlags::Namespace));
+		Parsables.Add(new FStructuredLogMacroDescriptor(TEXT("UE_CLOGFMT_NSLOC_EX"), FStructuredLogMacroDescriptor::EFlags::Condition | FStructuredLogMacroDescriptor::EFlags::Namespace));
 
 		Parsables.Add(new FIniNamespaceDescriptor());
 
@@ -3040,8 +3044,14 @@ void UGatherTextFromSourceCommandlet::FStringTableEntryMetaDataMacroDescriptor::
 	}
 }
 
+int32 UGatherTextFromSourceCommandlet::FStructuredLogMacroDescriptor::CalculateMinimumArgumentCount(EFlags Flags)
+{
+	// [Condition], CategoryName, Verbosity, [Namespace], Key, Format
+	return 4 + !!(Flags & EFlags::Condition) + !!(Flags & EFlags::Namespace);
+}
+
 UGatherTextFromSourceCommandlet::FStructuredLogMacroDescriptor::FStructuredLogMacroDescriptor(const TCHAR* InName, EFlags InFlags)
-	: FMacroDescriptor(InName, ((int32)InFlags & (int32)EFlags::Namespace) ? 5 : 4)
+	: FMacroDescriptor(InName, CalculateMinimumArgumentCount(InFlags))
 	, Flags(InFlags)
 {
 }
@@ -3051,7 +3061,11 @@ void UGatherTextFromSourceCommandlet::FStructuredLogMacroDescriptor::TryParse(co
 	// Attempt to parse something of the format
 	// UE_LOGFMT_LOC[_EX](CategoryName, Verbosity, Key, Format, ...)
 	// UE_LOGFMT_NSLOC[_EX](CategoryName, Verbosity, Namespace, Key, Format, ...)
-	int32 NextArg = 2;
+	// UE_CLOGFMT_LOC[_EX](Condition, CategoryName, Verbosity, Key, Format, ...)
+	// UE_CLOGFMT_NSLOC[_EX](Condition, CategoryName, Verbosity, Namespace, Key, Format, ...)
+
+	// The index of the next argument to consume, initially Namespace or Key.
+	int32 NextArg = !(Flags & EFlags::Condition) ? 2 : 3;
 
 	if (!Context.ExcludedRegion && !Context.WithinBlockComment && !Context.WithinLineComment && !Context.WithinStringLiteral)
 	{
@@ -3078,7 +3092,7 @@ void UGatherTextFromSourceCommandlet::FStructuredLogMacroDescriptor::TryParse(co
 			};
 
 			FString Namespace;
-			if (const bool bHasNamespace = ((int32)Flags & (int32)EFlags::Namespace) != 0)
+			if (!!(Flags & EFlags::Namespace))
 			{
 				Namespace = ParseArg(MoveTemp(Arguments[NextArg++]), TEXT("Namespace"));
 			}
diff --git a/Engine/Source/Runtime/Core/Public/Logging/StructuredLog.h b/Engine/Source/Runtime/Core/Public/Logging/StructuredLog.h
--- a/Engine/Source/Runtime/Core/Public/Logging/StructuredLog.h
+++ b/Engine/Source/Runtime/Core/Public/Logging/StructuredLog.h
@@ -46,6 +46,9 @@ template <typename FuncType> class TFunctionRef;
  * @param Fields[0-16]   Zero to sixteen fields or field values.
  */
 #define UE_LOGFMT(CategoryName, Verbosity, Format, ...) UE_PRIVATE_LOGFMT_CALL(UE_LOGFMT_EX, (CategoryName, Verbosity, Format UE_PRIVATE_LOGFMT_FIELDS(__VA_ARGS__)))
+ 
+ /** Conditional UE_LOGFMT. */
+#define UE_CLOGFMT(Condition, CategoryName, Verbosity, Format, ...) UE_PRIVATE_LOGFMT_CALL(UE_CLOGFMT_EX, (Condition, CategoryName, Verbosity, Format UE_PRIVATE_LOGFMT_FIELDS(__VA_ARGS__)))
 
 /**
  * Records a structured log event if this category is active at this level of verbosity.
@@ -60,7 +63,10 @@ template <typename FuncType> class TFunctionRef;
  * UE_LOGFMT_EX(LogCore, Warning, "Loading '{Name}' failed with error {Error}",
  *     UE_LOGFMT_FIELD("Name", Package->GetName()), UE_LOGFMT_FIELD("Error", ErrorCode), UE_LOGFMT_FIELD("Flags", LoadFlags));
  */
-#define UE_LOGFMT_EX(CategoryName, Verbosity, Format, ...) UE_PRIVATE_LOGFMT(CategoryName, Verbosity, Format, ##__VA_ARGS__)
+#define UE_LOGFMT_EX(CategoryName, Verbosity, Format, ...) UE_PRIVATE_LOGFMT(::UE::Logging::Private::FTrueCallable(), CategoryName, Verbosity, Format, ##__VA_ARGS__)
+
+/** Conditional UE_LOGFMT_EX. */
+#define UE_CLOGFMT_EX(Condition, CategoryName, Verbosity, Format, ...) UE_PRIVATE_LOGFMT([&] { return (Condition) && ::UE::Logging::Private::FTrueCallable()(); }, CategoryName, Verbosity, Format, ##__VA_ARGS__)
 
 /**
  * Records a localized structured log event if this category is active at this level of verbosity.
@@ -85,6 +91,12 @@ template <typename FuncType> class TFunctionRef;
 #define UE_LOGFMT_NSLOC(CategoryName, Verbosity, Namespace, Key, Format, ...) \
 	UE_PRIVATE_LOGFMT_CALL(UE_LOGFMT_NSLOC_EX, (CategoryName, Verbosity, Namespace, Key, Format UE_PRIVATE_LOGFMT_FIELDS(__VA_ARGS__)))
 
+/** Conditional UE_LOGFMT_LOC. */
+#define UE_CLOGFMT_LOC(Condition, CategoryName, Verbosity, Key, Format, ...) \
+	UE_CLOGFMT_NSLOC(Condition, CategoryName, Verbosity, LOCTEXT_NAMESPACE, Key, Format, ##__VA_ARGS__)
+#define UE_CLOGFMT_NSLOC(Condition, CategoryName, Verbosity, Namespace, Key, Format, ...) \
+	UE_PRIVATE_LOGFMT_CALL(UE_CLOGFMT_NSLOC_EX, (Condition, CategoryName, Verbosity, Namespace, Key, Format UE_PRIVATE_LOGFMT_FIELDS(__VA_ARGS__)))
+
 /**
  * Records a localized structured log event if this category is active at this level of verbosity.
  *
@@ -98,7 +110,13 @@ template <typename FuncType> class TFunctionRef;
 #define UE_LOGFMT_LOC_EX(CategoryName, Verbosity, Key, Format, ...) \
 	UE_LOGFMT_NSLOC_EX(CategoryName, Verbosity, LOCTEXT_NAMESPACE, Key, Format, ##__VA_ARGS__)
 #define UE_LOGFMT_NSLOC_EX(CategoryName, Verbosity, Namespace, Key, Format, ...) \
-	UE_PRIVATE_LOGFMT_LOC(CategoryName, Verbosity, Namespace, Key, Format, ##__VA_ARGS__)
+	UE_PRIVATE_LOGFMT_LOC(::UE::Logging::Private::FTrueCallable(), CategoryName, Verbosity, Namespace, Key, Format, ##__VA_ARGS__)
+
+/** Conditional UE_LOGFMT_LOC_EX. */
+#define UE_CLOGFMT_LOC_EX(Condition, CategoryName, Verbosity, Key, Format, ...) \
+	UE_CLOGFMT_NSLOC_EX(Condition, CategoryName, Verbosity, LOCTEXT_NAMESPACE, Key, Format, ##__VA_ARGS__)
+#define UE_CLOGFMT_NSLOC_EX(Condition, CategoryName, Verbosity, Namespace, Key, Format, ...) \
+	UE_PRIVATE_LOGFMT_LOC([&] { return (Condition) && ::UE::Logging::Private::FTrueCallable()(); }, CategoryName, Verbosity, Namespace, Key, Format, ##__VA_ARGS__)
 
 /** Expands to a named structured log field. UE_LOGFMT_FIELD("Name", Value) */
 #define UE_LOGFMT_FIELD(Name, Value) UE_PRIVATE_LOGFMT_FIELD((Name, Value))
@@ -470,27 +488,30 @@ UE_AUTORTFM_ALWAYS_OPEN inline void UE_COLD FatalLogWithFields(const FLogCategor
 }
 
 /** Log if the category is active at this level of verbosity. */
-template <typename LogCategoryType, ELogVerbosity::Type Verbosity, typename LogRecordType, typename... FieldArgTypes>
-UE_AUTORTFM_ALWAYS_OPEN inline void LogIfActive(const LogCategoryType& Category, const LogRecordType& Log, FieldArgTypes&&... FieldArgs)
+template <typename LogCategoryType, ELogVerbosity::Type Verbosity, typename ConditionType, typename LogRecordType, typename... FieldArgTypes>
+UE_AUTORTFM_ALWAYS_OPEN inline void LogIfActive(const ConditionType& Condition, const LogCategoryType& Category, const LogRecordType& Log, FieldArgTypes&&... FieldArgs)
 {
 	static_assert((Verbosity & ELogVerbosity::VerbosityMask) < ELogVerbosity::NumVerbosity && Verbosity > 0, "Verbosity must be constant and in range.");
 	if constexpr ((Verbosity & ELogVerbosity::VerbosityMask) == ELogVerbosity::Fatal)
 	{
-		if constexpr (sizeof...(FieldArgTypes) == 0)
-		{
-			FatalLogWithNoFields(Category, Log);
-		}
-		else
+		if (Condition())
 		{
-			FatalLogWithFields<FieldArgTypes...>(Category, Log, (FieldArgTypes&&)FieldArgs...);
+			if constexpr (sizeof...(FieldArgTypes) == 0)
+			{
+				FatalLogWithNoFields(Category, Log);
+			}
+			else
+			{
+				FatalLogWithFields<FieldArgTypes...>(Category, Log, (FieldArgTypes&&)FieldArgs...);
+			}
+			CA_ASSUME(false);
 		}
-		CA_ASSUME(false);
 	}
 	else if constexpr (
 		(Verbosity & ELogVerbosity::VerbosityMask) <= ELogVerbosity::COMPILED_IN_MINIMUM_VERBOSITY &&
 		(Verbosity & ELogVerbosity::VerbosityMask) <= LogCategoryType::CompileTimeVerbosity)
 	{
-		if (!Category.IsSuppressed(Verbosity))
+		if (!Category.IsSuppressed(Verbosity) && Condition())
 		{
 			if constexpr (sizeof...(FieldArgTypes) == 0)
 			{
@@ -530,23 +551,31 @@ struct FLogContext
 	bool bDisabledOlderContext = false;
 };
 
+struct FTrueCallable
+{
+	inline constexpr bool operator()() const
+	{
+		return true;
+	}
+};
+
 } // UE::Logging::Private
 
-#define UE_PRIVATE_LOGFMT(CategoryName, Verbosity, Format, ...) \
+#define UE_PRIVATE_LOGFMT(Condition, CategoryName, Verbosity, Format, ...) \
 	do \
 	{ \
 		static ::UE::Logging::Private::FStaticLogDynamicData LOG_Dynamic; \
 		static constexpr ::UE::Logging::Private::FStaticLogRecord LOG_Static UE_PRIVATE_LOGFMT_AGGREGATE(TEXT(Format), __builtin_FILE(), __builtin_LINE(), ::ELogVerbosity::Verbosity, LOG_Dynamic); \
-		::UE::Logging::Private::LogIfActive<FLogCategory##CategoryName, ::ELogVerbosity::Verbosity>(CategoryName, LOG_Static, ##__VA_ARGS__); \
+		::UE::Logging::Private::LogIfActive<FLogCategory##CategoryName, ::ELogVerbosity::Verbosity>(Condition, CategoryName, LOG_Static, ##__VA_ARGS__); \
 	} \
 	while (false)
 
-#define UE_PRIVATE_LOGFMT_LOC(CategoryName, Verbosity, Namespace, Key, Format, ...) \
+#define UE_PRIVATE_LOGFMT_LOC(Condition, CategoryName, Verbosity, Namespace, Key, Format, ...) \
 	do \
 	{ \
 		static ::UE::Logging::Private::FStaticLogDynamicData LOG_Dynamic; \
 		static constexpr ::UE::Logging::Private::FStaticLocalizedLogRecord LOG_Static UE_PRIVATE_LOGFMT_AGGREGATE(TEXT(Namespace), TEXT(Key), TEXT(Format), __builtin_FILE(), __builtin_LINE(), ::ELogVerbosity::Verbosity, LOG_Dynamic); \
-		::UE::Logging::Private::LogIfActive<FLogCategory##CategoryName, ::ELogVerbosity::Verbosity>(CategoryName, LOG_Static, ##__VA_ARGS__); \
+		::UE::Logging::Private::LogIfActive<FLogCategory##CategoryName, ::ELogVerbosity::Verbosity>(Condition, CategoryName, LOG_Static, ##__VA_ARGS__); \
 	} \
 	while (false)
 
