Replace deprecated SCNetworkReachability API with a lower level Network based nw_path implementation.
Tested on iOS, tvOS and macOS.  Also tested that Airplane Mode is correctly detected.

diff --git a/Engine/Source/Runtime/Core/Core.Build.cs b/Engine/Source/Runtime/Core/Core.Build.cs
--- a/Engine/Source/Runtime/Core/Core.Build.cs
+++ b/Engine/Source/Runtime/Core/Core.Build.cs
@@ -111,7 +111,7 @@ public class Core : ModuleRules
 				"zlib",
 				"PLCrashReporter"
 				);
-			PublicFrameworks.AddRange(new string[] { "Cocoa", "Carbon", "IOKit", "Security", "UniformTypeIdentifiers" });
+			PublicFrameworks.AddRange(new string[] { "Cocoa", "Carbon", "IOKit", "Security", "UniformTypeIdentifiers", "Network" });
 
 			PrivateDependencyModuleNames.Add("mimalloc");
 			PrivateDefinitions.Add("PLATFORM_BUILDS_MIMALLOC=1");
@@ -127,7 +127,7 @@ public class Core : ModuleRules
 			AddEngineThirdPartyPrivateStaticDependencies(Target,
 				"zlib"
 				);
-			PublicFrameworks.AddRange(new string[] { "UIKit", "Foundation", "AudioToolbox", "AVFoundation", "GameKit", "CoreVideo", "CoreMedia", "CoreGraphics", "GameController", "SystemConfiguration", "DeviceCheck", "UserNotifications" });
+			PublicFrameworks.AddRange(new string[] { "UIKit", "Foundation", "AudioToolbox", "AVFoundation", "GameKit", "CoreVideo", "CoreMedia", "CoreGraphics", "GameController", "SystemConfiguration", "DeviceCheck", "UserNotifications", "Network" });
 
 			if (Target.Platform == UnrealTargetPlatform.IOS)
 			{
diff --git a/Engine/Source/Runtime/Core/Private/IOS/IOSPlatformMisc.cpp b/Engine/Source/Runtime/Core/Private/IOS/IOSPlatformMisc.cpp
--- a/Engine/Source/Runtime/Core/Private/IOS/IOSPlatformMisc.cpp
+++ b/Engine/Source/Runtime/Core/Private/IOS/IOSPlatformMisc.cpp
@@ -42,7 +42,6 @@
 
 #import <DeviceCheck/DeviceCheck.h>
 #import <Foundation/Foundation.h>
-//#include <libproc.h>
 #import <mach-o/dyld.h>
 #include <netinet/in.h>
 #include <SystemConfiguration/SystemConfiguration.h>
@@ -51,6 +50,7 @@
 #endif
 
 #import <UserNotifications/UserNotifications.h>
+#import <Network/Network.h>
 
 #include <sys/sysctl.h> // sysctlbyname
 
@@ -85,6 +85,14 @@ static int32 GetFreeMemoryMB()
 	return MemoryStats.AvailablePhysical / 1024 / 1024;
 }
 
+namespace UE::FIOSPlatformMisc
+{
+	// Used to monitor network state. Replaces the deprecated SCNetworkReachability API.
+	// Note: defined here to not force objcisims into the header which is included on non-objc platforms
+	nw_path_monitor_t PathMonitor = nullptr;
+	nw_path_t CachedNetworkPath = nullptr;
+}
+
 void FIOSPlatformMisc::PlatformInit()
 {
 	// PlatformInit() starts the UI thread which creates the framebuffer and it requires
@@ -168,6 +176,20 @@ void FIOSPlatformMisc::PlatformInit()
 	{
 		NSLog(@"Error excluding %@ from backup %@",[URL lastPathComponent], error);
 	}
+	
+	// start monitoring network status
+	UE::FIOSPlatformMisc::PathMonitor = nw_path_monitor_create();
+	nw_path_monitor_set_update_handler(
+		UE::FIOSPlatformMisc::PathMonitor, 
+		^(nw_path_t InPath)
+		{ 
+			// cache the path we get
+			UE::FIOSPlatformMisc::CachedNetworkPath = InPath;
+		}
+	);
+	
+	nw_path_monitor_set_queue(UE::FIOSPlatformMisc::PathMonitor, dispatch_get_main_queue());
+	nw_path_monitor_start(UE::FIOSPlatformMisc::PathMonitor);
 }
 
 // Defines the PlatformFeatures module name for iOS, used by PlatformFeatures.h.
@@ -1588,46 +1610,55 @@ void FIOSPlatformMisc::GetValidTargetPlatforms(TArray<FString>& TargetPlatformNa
 
 static inline ENetworkConnectionType CheckNetworkConnectionType()
 {
-	struct sockaddr_in ZeroAddress;
-	FMemory::Memzero(&ZeroAddress, sizeof(ZeroAddress));
-	ZeroAddress.sin_len = sizeof(ZeroAddress);
-	ZeroAddress.sin_family = AF_INET;
-	
-PRAGMA_DISABLE_DEPRECATION_WARNINGS
-	SCNetworkReachabilityRef ReachabilityRef = SCNetworkReachabilityCreateWithAddress(kCFAllocatorDefault, (const struct sockaddr*)&ZeroAddress);
-	SCNetworkReachabilityFlags ReachabilityFlags;
-	bool bFlagsAvailable = SCNetworkReachabilityGetFlags(ReachabilityRef, &ReachabilityFlags);
-	CFRelease(ReachabilityRef);
+	nw_path_t MyPath = UE::FIOSPlatformMisc::CachedNetworkPath;
+	nw_path_status_t PathStatus = nw_path_get_status(MyPath);
 	
-	bool bHasActiveWiFiConnection = false;
-    bool bHasActiveCellConnection = false;
-    bool bInAirplaneMode = false;
-	if (bFlagsAvailable)
+	if ((PathStatus != nw_path_status_satisfied))
 	{
-		bool bReachable =	(ReachabilityFlags & kSCNetworkReachabilityFlagsReachable) != 0 &&
-		(ReachabilityFlags & kSCNetworkReachabilityFlagsConnectionRequired) == 0 &&
-		// in case kSCNetworkReachabilityFlagsConnectionOnDemand  || kSCNetworkReachabilityFlagsConnectionOnTraffic
-		(ReachabilityFlags & kSCNetworkReachabilityFlagsInterventionRequired) == 0;
+		// Apparently there's no valid network paths, so lets see if it's in "airplane" mode.
+		// Basically, if there's _no_ network interfaces available, assume it's in airplane mode.
+		// Note: if you have wifi enabled - but airplane mode also enabled, that will NOT be detected
+		//       as airplane mode as you still have a valid network path.
+		__block bool HasAnInterface = false;
+		nw_path_enumerate_interfaces(MyPath, (nw_path_enumerate_interfaces_block_t) ^ (nw_interface_t Interface)
+		{
+			const char* Name = nw_interface_get_name(Interface);
+			// UE_LOG(LogIOS, Log, TEXT("NW path monitor available interface: %s"), *FString(Name));
+			HasAnInterface = true;
+			return false;
+		});
 		
-		bHasActiveWiFiConnection = bReachable && (ReachabilityFlags & kSCNetworkReachabilityFlagsIsWWAN) == 0;
-        bHasActiveCellConnection = bReachable && (ReachabilityFlags & kSCNetworkReachabilityFlagsIsWWAN) != 0;
-        bInAirplaneMode = ReachabilityFlags == 0;
+		if (!HasAnInterface)
+		{
+			return ENetworkConnectionType::AirplaneMode;
+		}
+		return ENetworkConnectionType::None;
+	}
+	else 
+	{
+		// If satisfied detect the type of connection or maybe use isExpensive???
+		bool bHasActiveWiFiConnection = nw_path_uses_interface_type(MyPath, nw_interface_type_wifi);
+		bool bHasActiveCellConnection = nw_path_uses_interface_type(MyPath, nw_interface_type_cellular);
+		bool bHasActiveWiredConnection = nw_path_uses_interface_type(MyPath, nw_interface_type_wired);	
+
+		if (bHasActiveWiFiConnection)
+		{
+			return ENetworkConnectionType::WiFi;
+		}
+		else if (bHasActiveCellConnection)
+		{
+			return ENetworkConnectionType::Cell;
+		}
+		else if (bHasActiveWiredConnection)
+		{
+			return ENetworkConnectionType::Ethernet;
+		}
+		else
+		{
+			// For now, mw_interface_type_loopback and nw_interface_type_other return "Unknown"
+			return ENetworkConnectionType::Unknown;
+		}
 	}
-PRAGMA_ENABLE_DEPRECATION_WARNINGS
-	
-    if (bHasActiveWiFiConnection)
-    {
-        return ENetworkConnectionType::WiFi;
-    }
-    else if (bHasActiveCellConnection)
-    {
-        return ENetworkConnectionType::Cell;
-    }
-    else if (bInAirplaneMode)
-    {
-        return ENetworkConnectionType::AirplaneMode;
-    }
-    return ENetworkConnectionType::None;
 }
 
 ENetworkConnectionType FIOSPlatformMisc::GetNetworkConnectionType()
