Avoid unnecessary use of volatile pointers in a single-threaded queue.

Previously, we would always use `volatile`. This will slow down a single-threaded queue for no benefit, because Clang uses atomic instructions when dealing with volatiles.

diff --git a/Engine/Source/Runtime/Core/Public/Containers/Queue.h b/Engine/Source/Runtime/Core/Public/Containers/Queue.h
--- a/Engine/Source/Runtime/Core/Public/Containers/Queue.h
+++ b/Engine/Source/Runtime/Core/Public/Containers/Queue.h
@@ -6,6 +6,7 @@
 #include "Templates/UnrealTemplate.h"
 #include "HAL/PlatformAtomics.h"
 #include "HAL/PlatformMisc.h"
+#include <type_traits>
 
 // WARNING: This queue is planned for deprecation in favor of TSpscQueue or TMpscQueue
 
@@ -64,6 +65,8 @@ public:
 		}
 	}
 
+	UE_NONCOPYABLE(TQueue);
+
 	/**
 	 * Removes and returns the item from the tail of the queue.
 	 *
@@ -278,46 +281,36 @@ public:
 private:
 
 	/** Structure for the internal linked list. */
+	struct TNode;
+	using TNodeVolatilePtr = std::conditional_t<Mode == EQueueMode::SingleThreaded, TNode*, TNode* volatile>;
+
 	struct TNode
 	{
 		/** Holds a pointer to the next node in the list. */
-		TNode* volatile NextNode;
+		TNodeVolatilePtr NextNode = nullptr;
 
 		/** Holds the node's item. */
 		FElementType Item;
 
 		/** Default constructor. */
-		[[nodiscard]] TNode()
-			: NextNode(nullptr)
-		{
-		}
+		[[nodiscard]] TNode() = default;
 
 		/** Creates and initializes a new node. */
 		[[nodiscard]] explicit TNode(const FElementType& InItem)
-			: NextNode(nullptr)
-			, Item(InItem)
+			: Item(InItem)
 		{
 		}
 
 		/** Creates and initializes a new node. */
 		[[nodiscard]] explicit TNode(FElementType&& InItem)
-			: NextNode(nullptr)
-			, Item(MoveTemp(InItem))
+			: Item(MoveTemp(InItem))
 		{
 		}
 	};
 
 	/** Holds a pointer to the head of the list. */
-	MS_ALIGN(16) TNode* volatile Head GCC_ALIGN(16);
+	MS_ALIGN(16) TNodeVolatilePtr Head GCC_ALIGN(16);
 
 	/** Holds a pointer to the tail of the list. */
 	TNode* Tail;
-	
-private:
-
-	/** Hidden copy constructor. */
-	TQueue(const TQueue&) = delete;
-
-	/** Hidden assignment operator. */
-	TQueue& operator=(const TQueue&) = delete;
 };
