Added FBox::ToCompactString, which gets a textual representation of the box without min/max values if invalid.

diff --git a/Engine/Source/Runtime/Core/Public/Math/Box.h b/Engine/Source/Runtime/Core/Public/Math/Box.h
--- a/Engine/Source/Runtime/Core/Public/Math/Box.h
+++ b/Engine/Source/Runtime/Core/Public/Math/Box.h
@@ -615,6 +615,13 @@ public:
 	 */
 	FString ToString() const;
 
+	/**
+	 * Get a textual representation of this box, without min/max values if invalid.
+	 *
+	 * @return A string describing the box.
+	 */
+	FString ToCompactString() const;
+
 	/**
 	 * Get the vertices that make up this box.
 	 * 
@@ -828,6 +835,12 @@ FORCEINLINE FString TBox<T>::ToString() const
 	return FString::Printf(TEXT("IsValid=%s, Min=(%s), Max=(%s)"), IsValid ? TEXT("true") : TEXT("false"), *Min.ToString(), *Max.ToString());
 }
 
+template<typename T>
+FORCEINLINE FString TBox<T>::ToCompactString() const
+{
+	return IsValid ? FString::Printf(TEXT("Min=(%s), Max=(%s)"), *Min.ToString(), *Max.ToString()) : TEXT("IsValid=false");
+}
+
 
 template<typename T>
 TBox<T> TBox<T>::TransformBy(const TMatrix<T>& M) const
