Add mac's Bonjour name to remotehostnames in ue.projectstore and update iOS client to use that for a USB link local connection.

#rnx

diff --git a/Engine/Source/Developer/IoStoreUtilities/Private/ZenStoreWriter.cpp b/Engine/Source/Developer/IoStoreUtilities/Private/ZenStoreWriter.cpp
--- a/Engine/Source/Developer/IoStoreUtilities/Private/ZenStoreWriter.cpp
+++ b/Engine/Source/Developer/IoStoreUtilities/Private/ZenStoreWriter.cpp
@@ -471,6 +471,27 @@ void FZenStoreWriter::Initialize(const FCookInfo& Info)
 						{
 							Writer->WriteValue(Address->ToString(false));
 						}
+#if PLATFORM_MAC
+						// Store the Bonjour hostname when on a mac so that the client can attempt to connect via USB.
+						// This is done since the "Link Local" IP of the USB connection changes everytime
+						// the cable is plugged in, so there is no stable IP to save.
+						bool bAppleTarget = TargetPlatformFName == TEXT("IOS") || TargetPlatformFName == TEXT("TVOS");
+						if (bAppleTarget)
+						{
+							FString MacBonjourName = FPlatformMisc::GetBonjourName();
+							
+							// If there is no Bonjour name, fallback to the mac's hostname
+							if (MacBonjourName.IsEmpty())
+							{
+								MacBonjourName = FPlatformProcess::ComputerName();
+							}
+							
+							if (!MacBonjourName.IsEmpty())
+							{
+								Writer->WriteValue(TEXT("macserver://") + MacBonjourName);
+							}
+						}
+#endif
 						Writer->WriteArrayEnd();
 					}
 				}
diff --git a/Engine/Source/Runtime/Core/Core.Build.cs b/Engine/Source/Runtime/Core/Core.Build.cs
--- a/Engine/Source/Runtime/Core/Core.Build.cs
+++ b/Engine/Source/Runtime/Core/Core.Build.cs
@@ -111,7 +111,7 @@ public class Core : ModuleRules
 				"zlib",
 				"PLCrashReporter"
 				);
-			PublicFrameworks.AddRange(new string[] { "Cocoa", "Carbon", "IOKit", "Security", "UniformTypeIdentifiers", "Network", "ScreenCaptureKit" });
+			PublicFrameworks.AddRange(new string[] { "Cocoa", "Carbon", "IOKit", "Security", "UniformTypeIdentifiers", "Network", "ScreenCaptureKit", "SystemConfiguration" });
 
 			PrivateDependencyModuleNames.Add("mimalloc");
 			PrivateDefinitions.Add("PLATFORM_BUILDS_MIMALLOC=1");
diff --git a/Engine/Source/Runtime/Core/Private/Mac/MacPlatformMisc.cpp b/Engine/Source/Runtime/Core/Private/Mac/MacPlatformMisc.cpp
--- a/Engine/Source/Runtime/Core/Private/Mac/MacPlatformMisc.cpp
+++ b/Engine/Source/Runtime/Core/Private/Mac/MacPlatformMisc.cpp
@@ -41,6 +41,7 @@
 #include <IOKit/ps/IOPowerSources.h>
 #include <IOKit/ps/IOPSKeys.h>
 #include <IOKit/pwr_mgt/IOPMLib.h>
+#include <SystemConfiguration/SystemConfiguration.h>
 #include <mach-o/dyld.h>
 #include <libproc.h>
 #include <notify.h>
@@ -1688,6 +1689,26 @@ bool FMacPlatformMisc::HasSeparateChannelForDebugOutput()
 	return FPlatformMisc::IsDebuggerPresent() || isatty(STDOUT_FILENO) || isatty(STDERR_FILENO);
 }
 
+FString FMacPlatformMisc::GetBonjourName()
+{
+	static FString BonjourName = TEXT("");
+	if (BonjourName.IsEmpty())
+	{
+		CFStringRef LocalHostNameRef = SCDynamicStoreCopyLocalHostName(NULL);
+		if (LocalHostNameRef != nullptr)
+		{
+			// Per the Bonjour docs, https://developer.apple.com/library/archive/documentation/Cocoa/Conceptual/NetServices/Articles/domainnames.html
+			// the maximum length of the host name and of the fully qualified domain name (FQDN) is 63 bytes per label and 255 bytes per FQDN.
+			TCHAR Result[256] = TEXT("");
+			FPlatformString::CFStringToTCHAR(LocalHostNameRef, Result);
+			BonjourName = FString(Result);
+			CFRelease(LocalHostNameRef);
+		}
+	}
+
+	return BonjourName;
+}
+
 FString FMacPlatformMisc::GetCPUVendor()
 {
 #if PLATFORM_MAC_ARM64
diff --git a/Engine/Source/Runtime/Core/Public/Mac/MacPlatformMisc.h b/Engine/Source/Runtime/Core/Public/Mac/MacPlatformMisc.h
--- a/Engine/Source/Runtime/Core/Public/Mac/MacPlatformMisc.h
+++ b/Engine/Source/Runtime/Core/Public/Mac/MacPlatformMisc.h
@@ -46,6 +46,7 @@ struct CORE_API FMacPlatformMisc : public FApplePlatformMisc
 	static bool HasPlatformFeature(const TCHAR* FeatureName);
 	static bool GetDiskTotalAndFreeSpace(const FString& InPath, uint64& TotalNumberOfBytes, uint64& NumberOfFreeBytes);
 	static bool HasSeparateChannelForDebugOutput();
+	static FString GetBonjourName();
 
 	/** 
 	 * Determines the shader format for the plarform
diff --git a/Engine/Source/Runtime/StorageServerClient/Private/StorageServerPlatformFile.cpp b/Engine/Source/Runtime/StorageServerClient/Private/StorageServerPlatformFile.cpp
--- a/Engine/Source/Runtime/StorageServerClient/Private/StorageServerPlatformFile.cpp
+++ b/Engine/Source/Runtime/StorageServerClient/Private/StorageServerPlatformFile.cpp
@@ -989,7 +989,25 @@ bool FStorageServerPlatformFile::ShouldBeUsed(IPlatformFile* Inner, const TCHAR*
 					{
 						if (FString RemoteHostNameStr = RemoteHostName->AsString(); !RemoteHostNameStr.IsEmpty())
 						{
-							HostAddrs.Add(RemoteHostNameStr);
+#if PLATFORM_IOS
+							if (RemoteHostNameStr.StartsWith("macserver://"))
+							{
+								FString MacHostName = RemoteHostNameStr.RightChop(12);
+								if (!MacHostName.IsEmpty())
+								{
+									// As this is the fastest connection when on USB-C, set this as the first to test
+									// TODO: what about when using USB2?? Should we detect device type?
+									HostAddrs.Insert(MacHostName + TEXT(".local"), 0);
+								
+									// Some macs drop the ".local", so try that as well, but as a last resort
+									HostAddrs.Add(MacHostName);
+								}
+							}
+							else
+#endif 
+							{
+								HostAddrs.Add(RemoteHostNameStr);
+							}
 						}
 					}
 				}
