[TSAN/Windows]
- Fix false positive caused by logging interacting with another thread while instrumentation was disabled in just one thread
- Fix reentrancy caused by std::atomic usage while exiting a thread, ending up recreating the thread context we were exiting from
- Add support to indicate if the race in on stack memory
- Add support to indicate if previous thread is still alive

diff --git a/Engine/Source/Runtime/Core/Internal/Sanitizer/RaceDetectorInterface.h b/Engine/Source/Runtime/Core/Internal/Sanitizer/RaceDetectorInterface.h
--- a/Engine/Source/Runtime/Core/Internal/Sanitizer/RaceDetectorInterface.h
+++ b/Engine/Source/Runtime/Core/Internal/Sanitizer/RaceDetectorInterface.h
@@ -69,6 +69,10 @@ namespace UE::Sanitizer::RaceDetector {
 		void DeactivateTrampoline(void* PatchableFunctionAddress, bool bUseRETBypass);
 		// This needs to be called after trampoline activation or deactivation to make sure it takes effect immediately.
 		void FlushInstructionCache();
+		// Get the limits of the stack for the current thread.
+		void GetCurrentThreadStackLimits(void** LowLimit, void** HighLimit);
+		// Check if the thread id given is currently alive
+		bool IsThreadAlive(uint32 ThreadId);
 	}
 
 	// Returns whether we should instrument depending on the current context state.
diff --git a/Engine/Source/Runtime/Core/Private/Sanitizer/RaceDetector.cpp b/Engine/Source/Runtime/Core/Private/Sanitizer/RaceDetector.cpp
--- a/Engine/Source/Runtime/Core/Private/Sanitizer/RaceDetector.cpp
+++ b/Engine/Source/Runtime/Core/Private/Sanitizer/RaceDetector.cpp
@@ -430,6 +430,12 @@ namespace UE::Sanitizer::RaceDetector {
 		Platform::UnmapShadowMemory();
 	}
 
+	INSTRUMENTATION_FUNCTION_ATTRIBUTES void WriteToLog(FString&& Msg)
+	{
+		UE::TScopeLock Lock(RaceReportsLock);
+		RaceReports.Emplace(MoveTemp(Msg));
+	}
+
 	TOptional<TRaceCallbackFn> RaceCallbackFn;
 	INSTRUMENTATION_FUNCTION_ATTRIBUTES void ReportRace(FContext& Context, void* Pointer, const FMemoryAccess& CurrentAccess, const FMemoryAccess& RacingAccess, const FCallstackLocation& FirstLocation)
 	{
@@ -505,6 +511,14 @@ namespace UE::Sanitizer::RaceDetector {
 				}
 			}
 
+			void* LowLimit = 0; void *HighLimit = 0;
+			Platform::GetCurrentThreadStackLimits(&LowLimit, &HighLimit);
+			bool bIsRaceOnStack = false;
+			if (Pointer >= LowLimit && Pointer < HighLimit)
+			{
+				bIsRaceOnStack = true;
+			}
+
 			FFullLocation FirstFullLocation = FirstLocation.GetFullLocation();
 			FFullLocation SecondFullLocation = SecondLocation.GetFullLocation();
 
@@ -533,9 +547,9 @@ namespace UE::Sanitizer::RaceDetector {
 			Report.Appendf(
 				TEXT("=========================================") LINE_TERMINATOR
 				TEXT("WARNING: RaceDetector: data race detected") LINE_TERMINATOR
-				TEXT("%s of size %d at %p made at clock %lu by thread %s (%d) which is now at clock %lu:") LINE_TERMINATOR 
+				TEXT("%s of size %d at %p %smade at clock %lu by thread %s (%d) which is now at clock %lu:") LINE_TERMINATOR 
 				TEXT("%s") LINE_TERMINATOR
-				TEXT("Previous %s of size %d at %p at clock %lu by thread %s (%d) which is now at clock %lu:") LINE_TERMINATOR
+				TEXT("Previous %s of size %d at %p at clock %lu by thread %s (%d %s) which is now at clock %lu:") LINE_TERMINATOR
 				TEXT("%s") LINE_TERMINATOR
 				TEXT("Last known sync clock between both threads is %lu:") LINE_TERMINATOR
 				TEXT("%s") LINE_TERMINATOR
@@ -543,6 +557,7 @@ namespace UE::Sanitizer::RaceDetector {
 				AccessTypeToString(CurrentAccess.AccessType),
 				CurrentAccess.GetSize(),
 				(void*)((UPTRINT)AlignedPointer + CurrentAccess.GetOffset()),
+				bIsRaceOnStack ? TEXT("(Stack) ") : TEXT(""),
 				CurrentAccess.Clock,
 				*FirstThreadName,
 				FirstThreadId,
@@ -554,6 +569,7 @@ namespace UE::Sanitizer::RaceDetector {
 				RacingAccess.Clock,
 				*SecondThreadName,
 				SecondThreadId,
+				Platform::IsThreadAlive(SecondThreadId) ? TEXT("alive") : TEXT("exited"),
 				RacingContext->CurrentClock(),
 				bFoundLocation ? *SecondFullLocation.ToString(Alignment) : *NotFoundMessage,
 				LastSyncClock,
@@ -564,8 +580,7 @@ namespace UE::Sanitizer::RaceDetector {
 
 			// We do as little as possible from within the race since it might cause reentrancy depending on where the race happens.
 			// Send the report to be properly logged to file in another thread.
-			UE::TScopeLock Lock(RaceReportsLock);
-			RaceReports.Emplace(Report.ToString());
+			WriteToLog(Report.ToString());
 		}
 		
 		if (GRaceDetectorBreakOnRace && FPlatformMisc::IsDebuggerPresent())
@@ -1352,22 +1367,27 @@ namespace UE::Sanitizer::RaceDetector {
 		}
 
 		// Just make sure we're not initializing the stack from an access that comes from this function
-		FInstrumentationScope InstrumentationScope;
-
-		bRaceDetectorActive = bEnable;
-		if (!bEnable)
 		{
-			UnhookInstrumentation();
+			FInstrumentationScope InstrumentationScope;
 
-			ResetShadow();
-		}
-		else
-		{
-			GlobalEpoch++;
+			bRaceDetectorActive = bEnable;
+			if (!bEnable)
+			{
+				UnhookInstrumentation();
 
-			HookInstrumentation();
+				ResetShadow();
+			}
+			else
+			{
+				GlobalEpoch++;
+
+				HookInstrumentation();
+			}
 		}
 
+		// This needs to be properly instrumented because it does synchronize with another thread so if the
+		// instrumentation doesn't see the atomics being used from this thread, a future race will be
+		// detected when the stack memory used as synchronization during logging starts being reused.
 		UE_LOG(LogRaceDetector, Log, TEXT("Race detector has been toggled %s"), bEnable ? TEXT("on") : TEXT("off"));
 	}
 
diff --git a/Engine/Source/Runtime/Core/Private/Sanitizer/RaceDetectorWindows.cpp b/Engine/Source/Runtime/Core/Private/Sanitizer/RaceDetectorWindows.cpp
--- a/Engine/Source/Runtime/Core/Private/Sanitizer/RaceDetectorWindows.cpp
+++ b/Engine/Source/Runtime/Core/Private/Sanitizer/RaceDetectorWindows.cpp
@@ -29,6 +29,7 @@ namespace UE::Sanitizer::RaceDetector {
 	extern void InitializeBoot();
 	extern void SanitizerThreadRun(volatile bool& bContinue);
 	extern volatile bool bIsResettingShadow;
+	extern void WriteToLog(FString&& Message);
 }
 
 #define START_WINAPI_INSTRUMENTATION(DoIfNotInstrumenting) \
@@ -56,7 +57,9 @@ struct FInstrumentedStartThreadArgs {
 	// ClockBank inherited from the CreateThread call
 	FClockBank ClockBank;
 
-	std::atomic<HANDLE> ThreadHandle{ nullptr };
+	// Can't use std::atomic here because it causes reentrancy in the instrumentation
+	// during thread shutdown and ends up recreating context we're trying to delete.
+	volatile HANDLE ThreadHandle{ nullptr };
 };
 
 // This is actually really important since we need to set back the same
@@ -176,7 +179,7 @@ INSTRUMENTATION_FUNCTION_ATTRIBUTES HANDLE WINAPI DetouredCreateThread(LPSECURIT
 	HANDLE Handle = TrueCreateThread(lpThreadAttributes, dwStackSize, InstrumentedStartThread, Args, dwCreationFlags, lpThreadId);
 	if (Handle)
 	{
-		Args->ThreadHandle.store(Handle);
+		Args->ThreadHandle = Handle;
 	}
 	return Handle;
 }
@@ -203,7 +206,7 @@ INSTRUMENTATION_FUNCTION_ATTRIBUTES UPTRINT WINAPI Detoured_beginthreadex(void*
 	Context.ThreadCreationDepth--;
 	if (Handle)
 	{
-		Args->ThreadHandle.store((HANDLE)Handle);
+		Args->ThreadHandle = (HANDLE)Handle;
 	}
 	return Handle;
 }
@@ -845,6 +848,25 @@ namespace UE::Sanitizer::RaceDetector::Platform {
 	{
 		return ::IsDebuggerPresent();
 	}
+	
+	INSTRUMENTATION_FUNCTION_ATTRIBUTES void GetCurrentThreadStackLimits(void** LowLimit, void** HighLimit)
+	{
+		::GetCurrentThreadStackLimits(LowLimit, HighLimit);
+	}
+
+	INSTRUMENTATION_FUNCTION_ATTRIBUTES bool IsThreadAlive(uint32 ThreadId)
+	{
+		HANDLE ThreadHandle = ::OpenThread(THREAD_QUERY_LIMITED_INFORMATION, FALSE, ThreadId);
+		if (ThreadHandle != NULL)
+		{
+			DWORD ExitCode = 0;
+			BOOL bResult = ::GetExitCodeThread(ThreadHandle, &ExitCode);
+			::CloseHandle(ThreadHandle);
+			return bResult && ExitCode == STILL_ACTIVE;
+		}
+
+		return false;
+	}
 
 	INSTRUMENTATION_FUNCTION_ATTRIBUTES void InitShadowMemory()
 	{
