UE: Add duration to the information in the log about automated tests

diff --git a/Engine/Source/Runtime/Core/Private/Misc/AutomationTest.cpp b/Engine/Source/Runtime/Core/Private/Misc/AutomationTest.cpp
--- a/Engine/Source/Runtime/Core/Private/Misc/AutomationTest.cpp
+++ b/Engine/Source/Runtime/Core/Private/Misc/AutomationTest.cpp
@@ -1013,7 +1013,7 @@ void FAutomationTestFramework::DumpAutomationTestExecutionInfo( const TMap<FStri
 		const FString& CurTestName = MapIter.Key();
 		const FAutomationTestExecutionInfo& CurExecutionInfo = MapIter.Value();
 
-		UE_LOG(LogAutomationTest, Log, TEXT("%s: %s"), *CurTestName, CurExecutionInfo.bSuccessful ? *SuccessMessage : *FailMessage);
+		UE_LOG(LogAutomationTest, Log, TEXT("%s: %s (%.2fms)"), *CurTestName, CurExecutionInfo.bSuccessful ? *SuccessMessage : *FailMessage, 1000.0 * CurExecutionInfo.Duration);
 
 		for ( const FAutomationExecutionEntry& Entry : CurExecutionInfo.GetEntries() )
 		{
