Fixed member function pointer detection in Projection.

diff --git a/Engine/Source/Runtime/Core/Public/Templates/Projection.h b/Engine/Source/Runtime/Core/Public/Templates/Projection.h
--- a/Engine/Source/Runtime/Core/Public/Templates/Projection.h
+++ b/Engine/Source/Runtime/Core/Public/Templates/Projection.h
@@ -109,7 +109,7 @@ template <typename Invocable0Type, typename... InvocableTypes>
 	if constexpr (sizeof...(InvocableTypes) == 0)
 	{
 		using DecayedInvocable0Type = std::decay_t<Invocable0Type>;
-		if constexpr (!std::is_member_pointer_v<Invocable0Type>)
+		if constexpr (!std::is_member_pointer_v<DecayedInvocable0Type>)
 		{
 			return (Invocable0Type&&)Invocable0;
 		}
