Added TNotNull ctor to NonNullPointer.

diff --git a/Engine/Source/Runtime/Core/Public/Templates/NonNullPointer.h b/Engine/Source/Runtime/Core/Public/Templates/NonNullPointer.h
--- a/Engine/Source/Runtime/Core/Public/Templates/NonNullPointer.h
+++ b/Engine/Source/Runtime/Core/Public/Templates/NonNullPointer.h
@@ -5,6 +5,7 @@
 #include "CoreTypes.h"
 #include "Misc/AssertionMacros.h"
 #include "Misc/IntrusiveUnsetOptionalState.h"
+#include "Misc/NotNull.h"
 #include "Misc/OptionalFwd.h"
 #include "Templates/Requires.h"
 #include "Templates/UnrealTypeTraits.h"
@@ -48,6 +49,20 @@ public:
 		ensureMsgf(InObject, TEXT("Tried to initialize TNonNullPtr with a null pointer!"));
 	}
 
+#if UE_ENABLE_NOTNULL_WRAPPER
+	/**
+	 * Constructs a non-null pointer from a TNotNull wrapper.
+	 */
+	template <
+		typename OtherObjectType
+		UE_REQUIRES(std::is_convertible_v<OtherObjectType, ObjectType*>)
+	>
+	FORCEINLINE TNonNullPtr(TNotNull<OtherObjectType> InObject)
+		: Object(InObject)
+	{
+	}
+#endif
+
 	/**
 	 * Constructs a non-null pointer from another non-null pointer
 	 */
@@ -80,6 +95,21 @@ public:
 		return *this;
 	}
 
+#if UE_ENABLE_NOTNULL_WRAPPER
+	/**
+	 * Assignment operator taking a TNotNull wrapper.
+	 */
+	template <
+		typename OtherObjectType
+		UE_REQUIRES(std::is_convertible_v<OtherObjectType, ObjectType*>)
+	>
+	FORCEINLINE TNonNullPtr& operator=(TNotNull<OtherObjectType> InObject)
+	{
+		Object = InObject;
+		return *this;
+	}
+#endif
+
 	/**
 	 * Assignment operator taking another TNonNullPtr
 	 */
