Added workaround for codegen error in _MSC_FULL_VER <= 194134123, which fails to sort properly when the projection is a pointer to a virtual function in a struct.

diff --git a/Engine/Source/Runtime/Core/Public/Algo/Impl/BinaryHeap.h b/Engine/Source/Runtime/Core/Public/Algo/Impl/BinaryHeap.h
--- a/Engine/Source/Runtime/Core/Public/Algo/Impl/BinaryHeap.h
+++ b/Engine/Source/Runtime/Core/Public/Algo/Impl/BinaryHeap.h
@@ -2,6 +2,7 @@
 
 #pragma once
 
+#include "Templates/Invoke.h" // for older _MSVC_FULL_VER only - see usage below
 #include "Templates/Projection.h"
 #include "Templates/ReversePredicate.h"
 
@@ -57,7 +58,12 @@ namespace AlgoImpl
 	template <typename RangeValueType, typename IndexType, typename ProjectionType, typename PredicateType>
 	FORCEINLINE void HeapSiftDown(RangeValueType* Heap, IndexType Index, const IndexType Count, const ProjectionType& InProj, const PredicateType& Predicate)
 	{
+// Workaround for a codegen bug that was discovered related to member function pointer projections that are virtual
+#if defined(_MSVC_FULL_VER) && _MSVC_FULL_VER < 194134123
+		const ProjectionType& Proj = InProj;
+#else
 		auto&& Proj = Projection(InProj);
+#endif
 		while (!HeapIsLeaf(Index, Count))
 		{
 			const IndexType LeftChildIndex = HeapGetLeftChildIndex(Index);
@@ -66,7 +72,12 @@ namespace AlgoImpl
 			IndexType MinChildIndex = LeftChildIndex;
 			if (RightChildIndex < Count)
 			{
+// Workaround for a codegen bug that was discovered related to member function pointer projections that are virtual
+#if defined(_MSVC_FULL_VER) && _MSVC_FULL_VER < 194134123
+				MinChildIndex = Predicate( Invoke(Proj, Heap[LeftChildIndex]), Invoke(Proj, Heap[RightChildIndex]) ) ? LeftChildIndex : RightChildIndex;
+#else
 				MinChildIndex = Predicate( Proj(Heap[LeftChildIndex]), Proj(Heap[RightChildIndex]) ) ? LeftChildIndex : RightChildIndex;
+#endif
 			}
 
 			if (!Predicate( Proj(Heap[MinChildIndex]), Proj(Heap[Index]) ))
@@ -93,11 +104,21 @@ namespace AlgoImpl
 	template <class RangeValueType, typename IndexType, typename ProjectionType, class PredicateType>
 	FORCEINLINE IndexType HeapSiftUp(RangeValueType* Heap, IndexType RootIndex, IndexType NodeIndex, const ProjectionType& InProj, const PredicateType& Predicate)
 	{
+// Workaround for a codegen bug that was discovered related to member function pointer projections that are virtual
+#if defined(_MSVC_FULL_VER) && _MSVC_FULL_VER < 194134123
+		const ProjectionType& Proj = InProj;
+#else
 		auto&& Proj = Projection(InProj);
+#endif
 		while (NodeIndex > RootIndex)
 		{
 			IndexType ParentIndex = HeapGetParentIndex(NodeIndex);
+// Workaround for a codegen bug that was discovered related to member function pointer projections that are virtual
+#if defined(_MSVC_FULL_VER) && _MSVC_FULL_VER < 194134123
+			if (!Predicate( Invoke(Projection, Heap[NodeIndex]), Invoke(Projection, Heap[ParentIndex]) ))
+#else
 			if (!Predicate( Proj(Heap[NodeIndex]), Proj(Heap[ParentIndex]) ))
+#endif
 			{
 				break;
 			}
diff --git a/Engine/Source/Runtime/Core/Public/Algo/IntroSort.h b/Engine/Source/Runtime/Core/Public/Algo/IntroSort.h
--- a/Engine/Source/Runtime/Core/Public/Algo/IntroSort.h
+++ b/Engine/Source/Runtime/Core/Public/Algo/IntroSort.h
@@ -5,6 +5,7 @@
 #include "Algo/Impl/BinaryHeap.h"
 #include "Math/UnrealMathUtility.h"
 #include "Templates/IdentityFunctor.h"
+#include "Templates/Invoke.h" // for older _MSVC_FULL_VER only - see usage below
 #include "Templates/Projection.h"
 #include "Templates/Less.h"
 #include "Templates/UnrealTemplate.h" // For GetData, GetNum
@@ -60,7 +61,12 @@ namespace AlgoImpl
 					T *Max, *Item;
 					for( Max=Current.Min, Item=Current.Min+1; Item<=Current.Max; Item++ )
 					{
+// Workaround for a codegen bug that was discovered related to member function pointer projections that are virtual
+#if defined(_MSVC_FULL_VER) && _MSVC_FULL_VER < 194134123
+						if( Invoke( Predicate, Invoke( Proj, *Max ), Invoke( Proj, *Item ) ) )
+#else
 						if( Predicate( Proj( *Max ), Proj( *Item ) ) )
+#endif
 						{
 							Max = Item;
 						}
@@ -78,8 +84,14 @@ namespace AlgoImpl
 				Inner.Max = Current.Max+1;
 				for( ; ; )
 				{
+// Workaround for a codegen bug that was discovered related to member function pointer projections that are virtual
+#if defined(_MSVC_FULL_VER) && _MSVC_FULL_VER < 194134123
+					while( ++Inner.Min<=Current.Max && !Invoke( Predicate, Invoke( Proj, *Current.Min ), Invoke( Proj, *Inner.Min ) ) );
+					while( --Inner.Max> Current.Min && !Invoke( Predicate, Invoke( Proj, *Inner.Max ), Invoke( Proj, *Current.Min ) ) );
+#else
 					while( ++Inner.Min<=Current.Max && !Predicate( Proj( *Current.Min ), Proj( *Inner.Min ) ) );
 					while( --Inner.Max> Current.Min && !Predicate( Proj( *Inner.Max ), Proj( *Current.Min ) ) );
+#endif
 					if( Inner.Min>Inner.Max )
 					{
 						break;
@@ -160,7 +172,12 @@ namespace Algo
 	template <typename RangeType, typename ProjectionType>
 	FORCEINLINE void IntroSortBy(RangeType&& Range, ProjectionType Proj)
 	{
+// Workaround for a codegen bug that was discovered related to member function pointer projections that are virtual
+#if defined(_MSVC_FULL_VER) && _MSVC_FULL_VER < 194134123
+		AlgoImpl::IntroSortInternal(GetData(Range), GetNum(Range), MoveTemp(Proj), TLess<>());
+#else
 		AlgoImpl::IntroSortInternal(GetData(Range), GetNum(Range), Projection(MoveTemp(Proj)), TLess<>());
+#endif
 	}
 
 	/**
@@ -173,7 +190,12 @@ namespace Algo
 	template <typename RangeType, typename ProjectionType, typename PredicateType>
 	FORCEINLINE void IntroSortBy(RangeType&& Range, ProjectionType Proj, PredicateType Predicate)
 	{
+// Workaround for a codegen bug that was discovered related to member function pointer projections that are virtual
+#if defined(_MSVC_FULL_VER) && _MSVC_FULL_VER < 194134123
+		AlgoImpl::IntroSortInternal(GetData(Range), GetNum(Range), MoveTemp(Proj), MoveTemp(Predicate));
+#else
 		AlgoImpl::IntroSortInternal(GetData(Range), GetNum(Range), Projection(MoveTemp(Proj)), Projection(MoveTemp(Predicate)));
+#endif
 	}
 }
 
