Android: Fix inconsistent interpretation of which directories exist across PlatformFile methods.

diff --git a/Engine/Source/Runtime/Core/Private/Android/AndroidPlatformFile.cpp b/Engine/Source/Runtime/Core/Private/Android/AndroidPlatformFile.cpp
--- a/Engine/Source/Runtime/Core/Private/Android/AndroidPlatformFile.cpp
+++ b/Engine/Source/Runtime/Core/Private/Android/AndroidPlatformFile.cpp
@@ -34,6 +34,11 @@ DEFINE_LOG_CATEGORY_STATIC(LogAndroidFile, Log, All);
 
 #define LOG_ANDROID_FILE_MANIFEST 0
 
+/** When ANDROID_DISALLOW_LOCAL_FILESYSTEM is true, FAndroidPlatformFile will only look in mounted OBB files, making it more difficult for cheaters to patch files. */
+#ifndef ANDROID_DISALLOW_LOCAL_FILESYSTEM
+#define ANDROID_DISALLOW_LOCAL_FILESYSTEM 0
+#endif
+
 // Support 64 bit file access.
 //#define UE_ANDROID_FILE_64 0
 #define UE_ANDROID_FILE_64 1
@@ -162,6 +167,32 @@ FString AndroidRelativeToAbsolutePath(bool bUseInternalBasePath, FString RelPath
 	return RelPath;
 }
 
+/**
+ * Wrapper around AAssetManager_openDir that returns a valid pointer only if the app's Assets actually contains the given "directory".
+ *
+ * AAssetManager_openDir always returns non-nullptr, even if there is no "directory", but it's tested anyways for extra stability.  Android Assets
+ * underlying storage implementation data structure is flat, not hierarchical, so "directories" are imaginary.  Thus, Android asset storage doesn't
+ * support empty directories, and we should not treat non-existent asset prefixes as stat-able directories because it breaks systems (like SQLite VFS)
+ * that check for the existence of certain paths without any means of discriminating between files and directories.
+ */
+AAssetDir* OpenExistingAssetManagerDirectory(AAssetManager* AssetManager, const char* AssetPath)
+{
+	AAssetDir* Directory = AAssetManager_openDir(AssetManager, AssetPath);
+	if (UNLIKELY(nullptr == Directory))
+	{
+		return nullptr;
+	}
+
+	const char* const AnyFileName = AAssetDir_getNextFileName(Directory);
+	if (nullptr == AnyFileName)
+	{
+		AAssetDir_close(Directory);
+		Directory = nullptr;
+	}
+
+	return Directory;
+}
+
 /**
 	Android file handle implementation for partial (i.e. parcels) files.
 	This can handle all the variations of accessing data for assets and files.
@@ -1222,7 +1253,7 @@ FAndroidMappedFileRegion::~FAndroidMappedFileRegion()
 	overriding content from the most "frozen" to the most "fluid"
 	state. Hence creating a virtual single union file-system.
 **/
-class CORE_API FAndroidPlatformFile : public IAndroidPlatformFile
+class CORE_API FAndroidPlatformFile final : public IAndroidPlatformFile
 {
 public:
 	// Singleton implementation.
@@ -1244,10 +1275,6 @@ public:
 #endif
 	}
 
-	//~ For visibility of overloads we don't override
-	using IAndroidPlatformFile::IterateDirectory;
-	using IAndroidPlatformFile::IterateDirectoryStat;
-
 	static const FString* GetOverrideLogDirectory()
 	{
 		return GOverrideAndroidLogDir ? &AndroidLogDir : nullptr;
@@ -1422,17 +1449,24 @@ public:
 
 	virtual bool FileExists(const TCHAR* Filename) override
 	{
-		return FileExists(Filename, false);
+		const bool bAllowAssets = false;
+		return FileExists(Filename, bAllowAssets);
+	}
+
+	virtual bool FileExists(const TCHAR* Filename, bool bAllowAssets) override
+	{
+		const bool bForceAllowLocal = false;
+		return FileExists(Filename, bForceAllowLocal, bAllowAssets);
 	}
 
-	bool FileExists(const TCHAR* Filename, bool AllowLocal)
+	bool FileExists(const TCHAR* Filename, bool bForceAllowLocal, bool bAllowAssets)
 	{
 #if LOG_ANDROID_FILE
 		FPlatformMisc::LowLevelOutputDebugStringf(TEXT("FAndroidPlatformFile::FileExists('%s')"), Filename);
 #endif
 		FString LocalPath;
 		FString AssetPath;
-		PathToAndroidPaths(LocalPath, AssetPath, Filename, AllowLocal);
+		PathToAndroidPaths(LocalPath, AssetPath, Filename, bForceAllowLocal);
 
 		bool result = false;
 		struct stat FileInfo;
@@ -1445,7 +1479,7 @@ public:
 		else
 		{
 			// For everything else we only check existence.
-			result = IsResource(AssetPath) || IsAsset(AssetPath);
+			result = IsResource(AssetPath) || bAllowAssets && IsAsset(AssetPath);
 		}
 #if LOG_ANDROID_FILE
 		FPlatformMisc::LowLevelOutputDebugStringf(TEXT("FAndroidPlatformFile::FileExists('%s') => %s\nResolved as %s"), Filename, result ? TEXT("TRUE") : TEXT("FALSE"), *LocalPath);
@@ -1457,7 +1491,8 @@ public:
 	{
 		FString LocalPath;
 		FString AssetPath;
-		PathToAndroidPaths(LocalPath, AssetPath, Filename, false);
+		const bool bForceAllowLocal = false;
+		PathToAndroidPaths(LocalPath, AssetPath, Filename, bForceAllowLocal);
 
 		const FString NormalizedFilename = LocalPath;
 
@@ -1493,17 +1528,24 @@ public:
 
 	virtual int64 FileSize(const TCHAR* Filename) override
 	{
-		return FileSize(Filename, false);
+		const bool bAllowAssets = false;
+		return FileSize(Filename, bAllowAssets);
+	}
+
+	virtual int64 FileSize(const TCHAR* Filename, bool bAllowAssets) override
+	{
+		const bool bForceAllowLocal = false;
+		return FileSize(Filename, bForceAllowLocal, bAllowAssets);
 	}
 
-	int64 FileSize(const TCHAR* Filename, bool AllowLocal)
+	int64 FileSize(const TCHAR* Filename, bool bForceAllowLocal, bool bAllowAssets)
 	{
 #if LOG_ANDROID_FILE
 		FPlatformMisc::LowLevelOutputDebugStringf(TEXT("FAndroidPlatformFile::FileSize('%s')"), Filename);
 #endif
 		FString LocalPath;
 		FString AssetPath;
-		PathToAndroidPaths(LocalPath, AssetPath, Filename, AllowLocal);
+		PathToAndroidPaths(LocalPath, AssetPath, Filename, bForceAllowLocal);
 
 		struct stat FileInfo;
 		FileInfo.st_size = -1;
@@ -1520,7 +1562,7 @@ public:
 		{
 			FileInfo.st_size = ZipResource.GetEntryLength(AssetPath);
 		}
-		else
+		else if (bAllowAssets)
 		{
 			AAsset * file = AAssetManager_open(AssetMgr, TCHAR_TO_UTF8(*AssetPath), AASSET_MODE_UNKNOWN);
 			if (nullptr != file)
@@ -1534,17 +1576,18 @@ public:
 
 	virtual bool DeleteFile(const TCHAR* Filename) override
 	{
-		return DeleteFile(Filename, false);
+		const bool bForceAllowLocal = false;
+		return DeleteFile(Filename, bForceAllowLocal);
 	}
 
-	bool DeleteFile(const TCHAR* Filename, bool AllowLocal)
+	bool DeleteFile(const TCHAR* Filename, bool bForceAllowLocal)
 	{
 #if LOG_ANDROID_FILE
 		FPlatformMisc::LowLevelOutputDebugStringf(TEXT("FAndroidPlatformFile::DeleteFile('%s')"), Filename);
 #endif
 		FString LocalPath;
 		FString AssetPath;
-		PathToAndroidPaths(LocalPath, AssetPath, Filename, AllowLocal);
+		PathToAndroidPaths(LocalPath, AssetPath, Filename, bForceAllowLocal);
 
 		// Only delete if we have a local file.
 		if (IsLocal(LocalPath))
@@ -1567,17 +1610,18 @@ public:
 	// NOTE: Returns false if the file is not found.
 	virtual bool IsReadOnly(const TCHAR* Filename) override
 	{
-		return IsReadOnly(Filename, false);
+		const bool bForceAllowLocal = false;
+		return IsReadOnly(Filename, bForceAllowLocal);
 	}
 
-	bool IsReadOnly(const TCHAR* Filename, bool AllowLocal)
+	bool IsReadOnly(const TCHAR* Filename, bool bForceAllowLocal)
 	{
 #if LOG_ANDROID_FILE
 		FPlatformMisc::LowLevelOutputDebugStringf(TEXT("FAndroidPlatformFile::IsReadOnly('%s')"), Filename);
 #endif
 		FString LocalPath;
 		FString AssetPath;
-		PathToAndroidPaths(LocalPath, AssetPath, Filename, AllowLocal);
+		PathToAndroidPaths(LocalPath, AssetPath, Filename, bForceAllowLocal);
 
 		if (IsLocal(LocalPath))
 		{
@@ -1596,10 +1640,11 @@ public:
 
 	virtual bool MoveFile(const TCHAR* To, const TCHAR* From) override
 	{
-		return MoveFile(To, From, false);
+		const bool bForceAllowLocal = false;
+		return MoveFile(To, From, bForceAllowLocal);
 	}
 
-	bool MoveFile(const TCHAR* To, const TCHAR* From, bool AllowLocal)
+	bool MoveFile(const TCHAR* To, const TCHAR* From, bool bForceAllowLocal)
 	{
 #if LOG_ANDROID_FILE
 		FPlatformMisc::LowLevelOutputDebugStringf(TEXT("FAndroidPlatformFile::MoveFile('%s', '%s')"), To, From);
@@ -1607,10 +1652,10 @@ public:
 		// Can only move local files.
 		FString ToLocalPath;
 		FString ToAssetPath;
-		PathToAndroidPaths(ToLocalPath, ToAssetPath, To, AllowLocal);
+		PathToAndroidPaths(ToLocalPath, ToAssetPath, To, bForceAllowLocal);
 		FString FromLocalPath;
 		FString FromAssetPath;
-		PathToAndroidPaths(FromLocalPath, FromAssetPath, From, AllowLocal);
+		PathToAndroidPaths(FromLocalPath, FromAssetPath, From, bForceAllowLocal);
 
 		if (IsLocal(FromLocalPath))
 		{
@@ -1621,17 +1666,18 @@ public:
 
 	virtual bool SetReadOnly(const TCHAR* Filename, bool bNewReadOnlyValue) override
 	{
-		return SetReadOnly(Filename, bNewReadOnlyValue, false);
+		const bool bForceAllowLocal = false;
+		return SetReadOnly(Filename, bNewReadOnlyValue, bForceAllowLocal);
 	}
 
-	bool SetReadOnly(const TCHAR* Filename, bool bNewReadOnlyValue, bool AllowLocal)
+	bool SetReadOnly(const TCHAR* Filename, bool bNewReadOnlyValue, bool bForceAllowLocal)
 	{
 #if LOG_ANDROID_FILE
 		FPlatformMisc::LowLevelOutputDebugStringf(TEXT("FAndroidPlatformFile::SetReadOnly('%s')"), Filename);
 #endif
 		FString LocalPath;
 		FString AssetPath;
-		PathToAndroidPaths(LocalPath, AssetPath, Filename, AllowLocal);
+		PathToAndroidPaths(LocalPath, AssetPath, Filename, bForceAllowLocal);
 
 		if (IsLocal(LocalPath))
 		{
@@ -1654,17 +1700,18 @@ public:
 
 	virtual FDateTime GetTimeStamp(const TCHAR* Filename) override
 	{
-		return GetTimeStamp(Filename, false);
+		const bool bForceAllowLocal = false;
+		return GetTimeStamp(Filename, bForceAllowLocal);
 	}
 
-	FDateTime GetTimeStamp(const TCHAR* Filename, bool AllowLocal)
+	FDateTime GetTimeStamp(const TCHAR* Filename, bool bForceAllowLocal)
 	{
 #if LOG_ANDROID_FILE_MANIFEST
 		FPlatformMisc::LowLevelOutputDebugStringf(TEXT("FAndroidPlatformFile::GetTimeStamp('%s')"), Filename);
 #endif
 		FString LocalPath;
 		FString AssetPath;
-		PathToAndroidPaths(LocalPath, AssetPath, Filename, AllowLocal);
+		PathToAndroidPaths(LocalPath, AssetPath, Filename, bForceAllowLocal);
 
 		if (IsLocal(LocalPath))
 		{
@@ -1698,7 +1745,8 @@ public:
 			if (AssetPath.EndsWith(".pak"))
 			{
 				// return local file access timestamp (if exists)
-				return GetAccessTimeStamp(Filename, true);
+				const bool bPakForceAllowLocal = true;
+				return GetAccessTimeStamp(Filename, bPakForceAllowLocal);
 			}
 
 			return FDateTime::MinValue();
@@ -1719,17 +1767,18 @@ public:
 
 	virtual void SetTimeStamp(const TCHAR* Filename, const FDateTime DateTime) override
 	{
-		return SetTimeStamp(Filename, DateTime, false);
+		const bool bForceAllowLocal = false;
+		return SetTimeStamp(Filename, DateTime, bForceAllowLocal);
 	}
 	
-	void SetTimeStamp(const TCHAR* Filename, const FDateTime DateTime, bool AllowLocal)
+	void SetTimeStamp(const TCHAR* Filename, const FDateTime DateTime, bool bForceAllowLocal)
 	{
 #if LOG_ANDROID_FILE
 		FPlatformMisc::LowLevelOutputDebugStringf(TEXT("FAndroidPlatformFile::SetTimeStamp('%s')"), Filename);
 #endif
 		FString LocalPath;
 		FString AssetPath;
-		PathToAndroidPaths(LocalPath, AssetPath, Filename, AllowLocal);
+		PathToAndroidPaths(LocalPath, AssetPath, Filename, bForceAllowLocal);
 
 		// Can only set time stamp on local files
 		if (IsLocal(LocalPath))
@@ -1765,17 +1814,18 @@ public:
 
 	virtual FDateTime GetAccessTimeStamp(const TCHAR* Filename) override
 	{
-		return GetAccessTimeStamp(Filename, false);
+		const bool bForceAllowLocal = false;
+		return GetAccessTimeStamp(Filename, bForceAllowLocal);
 	}
 
-	FDateTime GetAccessTimeStamp(const TCHAR* Filename, bool AllowLocal)
+	FDateTime GetAccessTimeStamp(const TCHAR* Filename, bool bForceAllowLocal)
 	{
 #if LOG_ANDROID_FILE
 		FPlatformMisc::LowLevelOutputDebugStringf(TEXT("FAndroidPlatformFile::GetAccessTimeStamp('%s')"), Filename);
 #endif
 		FString LocalPath;
 		FString AssetPath;
-		PathToAndroidPaths(LocalPath, AssetPath, Filename, AllowLocal);
+		PathToAndroidPaths(LocalPath, AssetPath, Filename, bForceAllowLocal);
 
 		if (IsLocal(LocalPath))
 		{
@@ -1797,17 +1847,24 @@ public:
 
 	virtual FFileStatData GetStatData(const TCHAR* FilenameOrDirectory) override
 	{
-		return GetStatData(FilenameOrDirectory, false);
+		const bool bAllowAssets = false;
+		return GetStatData(FilenameOrDirectory, bAllowAssets);
 	}
 
-	FFileStatData GetStatData(const TCHAR* FilenameOrDirectory, bool AllowLocal)
+	virtual FFileStatData GetStatData(const TCHAR* FilenameOrDirectory, bool bAllowAssets) override
+	{
+		const bool bForceAllowLocal = false;
+		return GetStatData(FilenameOrDirectory, bForceAllowLocal, bAllowAssets);
+	}
+
+	FFileStatData GetStatData(const TCHAR* FilenameOrDirectory, bool bForceAllowLocal, bool bAllowAssets)
 	{
 #if LOG_ANDROID_FILE
 		FPlatformMisc::LowLevelOutputDebugStringf(TEXT("FAndroidPlatformFile::GetStatData('%s')"), FilenameOrDirectory);
 #endif
 		FString LocalPath;
 		FString AssetPath;
-		PathToAndroidPaths(LocalPath, AssetPath, FilenameOrDirectory, AllowLocal);
+		PathToAndroidPaths(LocalPath, AssetPath, FilenameOrDirectory, bForceAllowLocal);
 
 		if (IsLocal(LocalPath))
 		{
@@ -1828,7 +1885,7 @@ public:
 				true									// bIsReadOnly
 				);
 		}
-		else
+		else if (bAllowAssets)
 		{
 			AAsset * file = AAssetManager_open(AssetMgr, TCHAR_TO_UTF8(*AssetPath), AASSET_MODE_UNKNOWN);
 			bool isDirectory = false;
@@ -1841,7 +1898,7 @@ public:
 			}
 			else
 			{
-				AAssetDir * dir = AAssetManager_openDir(AssetMgr, TCHAR_TO_UTF8(*AssetPath));
+				AAssetDir* dir = OpenExistingAssetManagerDirectory(AssetMgr, TCHAR_TO_UTF8(*AssetPath));
 				exists = (nullptr != dir);
 				
 				if (exists)
@@ -1874,19 +1931,18 @@ public:
 
 	virtual IFileHandle* OpenRead(const TCHAR* Filename, bool bAllowWrite = false) override
 	{
-		const bool bAllowLocal = false;
-
-		return OpenReadInternal(Filename, bAllowLocal, bAllowWrite);
+		const bool bForceAllowLocal = false;
+		return OpenReadInternal(Filename, bForceAllowLocal, bAllowWrite);
 	}
 
-	IFileHandle* OpenReadInternal(const TCHAR* Filename, bool bAllowLocal, bool bAllowWrite)
+	IFileHandle* OpenReadInternal(const TCHAR* Filename, bool bForceAllowLocal, bool bAllowWrite)
 	{
 #if LOG_ANDROID_FILE
 		FPlatformMisc::LowLevelOutputDebugStringf(TEXT("FAndroidPlatformFile::OpenRead('%s')"), Filename);
 #endif
 		FString LocalPath;
 		FString AssetPath;
-		PathToAndroidPaths(LocalPath, AssetPath, Filename, bAllowLocal);
+		PathToAndroidPaths(LocalPath, AssetPath, Filename, bForceAllowLocal);
 
 		if (IsLocal(LocalPath))
 		{
@@ -1919,17 +1975,18 @@ public:
 	// packaged content.
 	virtual IFileHandle* OpenWrite(const TCHAR* Filename, bool bAppend, bool bAllowRead) override
 	{
-		return OpenWrite(Filename, bAppend, bAllowRead, false);
+		const bool bForceAllowLocal = false;
+		return OpenWrite(Filename, bAppend, bAllowRead, bForceAllowLocal);
 	}
 
-	IFileHandle* OpenWrite(const TCHAR* Filename, bool bAppend, bool bAllowRead, bool AllowLocal)
+	IFileHandle* OpenWrite(const TCHAR* Filename, bool bAppend, bool bAllowRead, bool bForceAllowLocal)
 	{
 #if LOG_ANDROID_FILE
 		FPlatformMisc::LowLevelOutputDebugStringf(TEXT("FAndroidPlatformFile::OpenWrite('%s')"), Filename);
 #endif
 		FString LocalPath;
 		FString AssetPath;
-		PathToAndroidPaths(LocalPath, AssetPath, Filename, AllowLocal);
+		PathToAndroidPaths(LocalPath, AssetPath, Filename, bForceAllowLocal);
 
 		int Flags = O_CREAT;
 		if (!bAppend)
@@ -1966,17 +2023,24 @@ public:
 
 	virtual bool DirectoryExists(const TCHAR* Directory) override
 	{
-		return DirectoryExists(Directory, false, false);
+		const bool bAllowAssets = false;
+		return DirectoryExists(Directory, bAllowAssets);
 	}
-	
-	bool DirectoryExists(const TCHAR* Directory, bool AllowLocal, bool AllowAsset)
+
+	virtual bool DirectoryExists(const TCHAR* Directory, bool bAllowAssets) override
+	{
+		const bool bForceAllowLocal = false;
+		return DirectoryExists(Directory, bForceAllowLocal, bAllowAssets);
+	}
+
+	bool DirectoryExists(const TCHAR* Directory, bool bForceAllowLocal, bool bAllowAssets)
 	{
 #if LOG_ANDROID_FILE
 		FPlatformMisc::LowLevelOutputDebugStringf(TEXT("FAndroidPlatformFile::DirectoryExists('%s')"), Directory);
 #endif
 		FString LocalPath;
 		FString AssetPath;
-		PathToAndroidPaths(LocalPath, AssetPath, Directory, AllowLocal);
+		PathToAndroidPaths(LocalPath, AssetPath, Directory, bForceAllowLocal);
 
 		bool Found = false;
 		if (IsLocal(LocalPath))
@@ -1997,12 +2061,9 @@ public:
 			FPlatformMisc::LowLevelOutputDebugStringf(TEXT("FAndroidPlatformFile::DirectoryExists('%s') => Found as resource: '%s'"), Directory, *(AssetPath + "/"));
 #endif
 		}
-		else if (AllowAsset)
+		else if (bAllowAssets)
 		{
-			// This can be very slow on some devices so only do it if we enabled
-			AAssetDir * dir = AAssetManager_openDir(AssetMgr, TCHAR_TO_UTF8(*AssetPath));
-			Found = AAssetDir_getNextFileName(dir) != nullptr;
-			AAssetDir_close(dir);
+			Found = OpenExistingAssetManagerDirectory(AssetMgr, TCHAR_TO_UTF8(*AssetPath)) != nullptr;
 #if LOG_ANDROID_FILE
 			if (Found)
 			{
@@ -2026,17 +2087,18 @@ public:
 	// We assert that created dirs are in the local file-system.
 	virtual bool CreateDirectory(const TCHAR* Directory) override
 	{
-		return CreateDirectory(Directory, false);
+		const bool bForceAllowLocal = false;
+		return CreateDirectory(Directory, bForceAllowLocal);
 	}
 
-	bool CreateDirectory(const TCHAR* Directory, bool AllowLocal)
+	bool CreateDirectory(const TCHAR* Directory, bool bForceAllowLocal)
 	{
 #if LOG_ANDROID_FILE
 		FPlatformMisc::LowLevelOutputDebugStringf(TEXT("FAndroidPlatformFile::CreateDirectory('%s')"), Directory);
 #endif
 		FString LocalPath;
 		FString AssetPath;
-		PathToAndroidPaths(LocalPath, AssetPath, Directory, AllowLocal);
+		PathToAndroidPaths(LocalPath, AssetPath, Directory, bForceAllowLocal);
 		uint32 mkdirperms = 0755;
 #if !UE_BUILD_SHIPPING
 		// some devices prevent ADB (shell user) from modifying files.
@@ -2049,27 +2111,35 @@ public:
 	// We assert that modifying dirs are in the local file-system.
 	virtual bool DeleteDirectory(const TCHAR* Directory) override
 	{
-		return DeleteDirectory(Directory, false);
+		const bool bForceAllowLocal = false;
+		return DeleteDirectory(Directory, bForceAllowLocal);
 	}
 
-	bool DeleteDirectory(const TCHAR* Directory, bool AllowLocal)
+	bool DeleteDirectory(const TCHAR* Directory, bool bForceAllowLocal)
 	{
 #if LOG_ANDROID_FILE
 		FPlatformMisc::LowLevelOutputDebugStringf(TEXT("FAndroidPlatformFile::DeleteDirectory('%s')"), Directory);
 #endif
 		FString LocalPath;
 		FString AssetPath;
-		PathToAndroidPaths(LocalPath, AssetPath, Directory, AllowLocal);
+		PathToAndroidPaths(LocalPath, AssetPath, Directory, bForceAllowLocal);
 
 		return rmdir(TCHAR_TO_UTF8(*LocalPath));
 	}
 
 	virtual bool IterateDirectory(const TCHAR* Directory, FDirectoryVisitor& Visitor) override
 	{
-		return IterateDirectory(Directory, Visitor, false, false);
+		const bool bAllowAssets = false;
+		return IterateDirectory(Directory, Visitor, bAllowAssets);
+	}
+
+	virtual bool IterateDirectory(const TCHAR* Directory, FDirectoryVisitor& Visitor, bool bAllowAssets) override
+	{
+		const bool bForceAllowLocal = false;
+		return IterateDirectory(Directory, Visitor, bForceAllowLocal, bAllowAssets);
 	}
 
-	bool IterateDirectory(const TCHAR* Directory, FDirectoryVisitor& Visitor, bool AllowLocal, bool AllowAsset)
+	bool IterateDirectory(const TCHAR* Directory, FDirectoryVisitor& Visitor, bool bForceAllowLocal, bool bAllowAssets)
 	{
 		const FString DirectoryStr = Directory;
 
@@ -2097,15 +2167,22 @@ public:
 			return Visitor.CallShouldVisitAndVisit(UTF8_TO_TCHAR(InAssetPath), isDirectory);
 		};
 
-		return IterateDirectoryCommon(Directory, InternalVisitor, InternalResourceVisitor, InternalAssetVisitor, AllowLocal, AllowAsset);
+		return IterateDirectoryCommon(Directory, InternalVisitor, InternalResourceVisitor, InternalAssetVisitor, bForceAllowLocal, bAllowAssets);
 	}
 
 	virtual bool IterateDirectoryStat(const TCHAR* Directory, FDirectoryStatVisitor& Visitor) override
 	{
-		return IterateDirectoryStat(Directory, Visitor, false, false);
+		const bool bAllowAssets = false;
+		return IterateDirectoryStat(Directory, Visitor, bAllowAssets);
 	}
 
-	bool IterateDirectoryStat(const TCHAR* Directory, FDirectoryStatVisitor& Visitor, bool AllowLocal, bool AllowAsset)
+	virtual bool IterateDirectoryStat(const TCHAR* Directory, FDirectoryStatVisitor& Visitor, bool bAllowAssets) override
+	{
+		const bool bForceAllowLocal = false;
+		return IterateDirectoryStat(Directory, Visitor, bForceAllowLocal, bAllowAssets);
+	}
+
+	bool IterateDirectoryStat(const TCHAR* Directory, FDirectoryStatVisitor& Visitor, bool bForceAllowLocal, bool bAllowAssets)
 	{
 		const FString DirectoryStr = Directory;
 
@@ -2168,17 +2245,17 @@ public:
 				);
 		};
 
-		return IterateDirectoryCommon(Directory, InternalVisitor, InternalResourceVisitor, InternalAssetVisitor, AllowLocal, AllowAsset);
+		return IterateDirectoryCommon(Directory, InternalVisitor, InternalResourceVisitor, InternalAssetVisitor, bForceAllowLocal, bAllowAssets);
 	}
 
-	bool IterateDirectoryCommon(const TCHAR* Directory, const TFunctionRef<bool(const FString&, struct dirent*)>& Visitor, const TFunctionRef<bool(const FString&, bool)>& ResourceVisitor, const TFunctionRef<bool(const char*)>& AssetVisitor, bool AllowLocal, bool AllowAsset)
+	bool IterateDirectoryCommon(const TCHAR* Directory, const TFunctionRef<bool(const FString&, struct dirent*)>& Visitor, const TFunctionRef<bool(const FString&, bool)>& ResourceVisitor, const TFunctionRef<bool(const char*)>& AssetVisitor, bool bForceAllowLocal, bool bAllowAssets)
 	{
 #if LOG_ANDROID_FILE
 		FPlatformMisc::LowLevelOutputDebugStringf(TEXT("FAndroidPlatformFile::IterateDirectory('%s')"), Directory);
 #endif
 		FString LocalPath;
 		FString AssetPath;
-		PathToAndroidPaths(LocalPath, AssetPath, Directory, AllowLocal);
+		PathToAndroidPaths(LocalPath, AssetPath, Directory, bForceAllowLocal);
 
 		if (IsLocal(LocalPath))
 		{
@@ -2241,10 +2318,10 @@ public:
 			}
 			return Result;
 		}
-		else if (AllowAsset)
+		else if (bAllowAssets)
 		{
 			AAssetDir * dir = AAssetManager_openDir(AssetMgr, TCHAR_TO_UTF8(*AssetPath));
-			if (nullptr != dir)
+			if (LIKELY(nullptr != dir))
 			{
 				bool Result = true;
 				const char * fileName = nullptr;
@@ -2276,7 +2353,8 @@ public:
 #endif
 		FString LocalPath;
 		FString AssetPath;
-		PathToAndroidPaths(LocalPath, AssetPath, Filename, true);
+		const bool bForceAllowLocal = true;
+		PathToAndroidPaths(LocalPath, AssetPath, Filename, bForceAllowLocal);
 
 		if (IsLocal(LocalPath))
 		{
@@ -2300,7 +2378,8 @@ public:
 #endif
 		FString LocalPath;
 		FString AssetPath;
-		PathToAndroidPaths(LocalPath, AssetPath, Filename, true);
+		const bool bForceAllowLocal = true;
+		PathToAndroidPaths(LocalPath, AssetPath, Filename, bForceAllowLocal);
 
 		if (IsLocal(LocalPath))
 		{
@@ -2336,7 +2415,8 @@ public:
 #endif
 		FString LocalPath;
 		FString AssetPath;
-		PathToAndroidPaths(LocalPath, AssetPath, Filename, true);
+		const bool bForceAllowLocal = true;
+		PathToAndroidPaths(LocalPath, AssetPath, Filename, bForceAllowLocal);
 
 		if (IsLocal(LocalPath))
 		{
@@ -2368,7 +2448,7 @@ private:
 		return Result;
 	}
 
-	void PathToAndroidPaths(FString & LocalPath, FString & AssetPath, const TCHAR* Path, bool AllowLocal)
+	void PathToAndroidPaths(FString & LocalPath, FString & AssetPath, const TCHAR* Path, bool bForceAllowLocal)
 	{
 		LocalPath.Empty();
 		AssetPath.Empty();
@@ -2379,7 +2459,15 @@ private:
 #endif
 		if (!AndroidPath.IsEmpty())
 		{
-			if (AndroidPath.StartsWith(TEXT("/"), ESearchCase::CaseSensitive) ||	// (AllowLocal && AndroidPath.StartsWith(TEXT("/"), ESearchCase::CaseSensitive)) ||		// requiring AllowLocal isn't really useful; system will do permission checks anyway
+			// We filter out non-permitted local paths here when ANDROID_DISALLOW_LOCAL_FILESYSTEM is set.
+			// There is a limited set of paths we always allow local file access to, such as direct font access,
+			// and bForceAllowLocal bypasses the restriction as necessary to mount OBB files.
+			if (
+#if ANDROID_DISALLOW_LOCAL_FILESYSTEM
+				(bForceAllowLocal && AndroidPath.StartsWith(TEXT("/"), ESearchCase::CaseSensitive)) ||
+#else
+				AndroidPath.StartsWith(TEXT("/"), ESearchCase::CaseSensitive) ||
+#endif
 				AndroidPath.StartsWith(GFontPathBase) ||
 				AndroidPath.StartsWith(TEXT("/system/etc/")) ||
 				AndroidPath.StartsWith(GInternalFilePath.Left(AndroidPath.Len())) ||
@@ -2416,11 +2504,6 @@ private:
 #endif
 	}
 
-	// There is a limited set of paths we allow local file access to.
-	// We filter out non-permitted local paths here unless explicitly
-	// allowed. The current two cases are direct font access and with
-	// the AllowLocal flag (as used to mount local OBBs). Eventually
-	// we may need to also allow access to downloads somehow!
 	bool IsLocal(const FString & LocalPath)
 	{
 		return !LocalPath.IsEmpty() && access(TCHAR_TO_UTF8(*LocalPath), F_OK) == 0;
@@ -2470,10 +2553,10 @@ private:
 
 	void MountOBB(const TCHAR* Filename)
 	{
-		const bool bAllowLocal = true;
+		const bool bForceAllowLocal = true;
 		const bool bAllowWrite = false;
 
-		FFileHandleAndroid* File = static_cast<FFileHandleAndroid*>(OpenReadInternal(Filename, bAllowLocal, bAllowWrite));
+		FFileHandleAndroid* File = static_cast<FFileHandleAndroid*>(OpenReadInternal(Filename, bForceAllowLocal, bAllowWrite));
 		check(nullptr != File);
 		ZipResource.AddPatchFile(MakeShareable(File));
 		FPlatformMisc::LowLevelOutputDebugStringf(TEXT("Mounted OBB '%s'"), Filename);
diff --git a/Engine/Source/Runtime/Core/Public/Android/AndroidPlatformFile.h b/Engine/Source/Runtime/Core/Public/Android/AndroidPlatformFile.h
--- a/Engine/Source/Runtime/Core/Public/Android/AndroidPlatformFile.h
+++ b/Engine/Source/Runtime/Core/Public/Android/AndroidPlatformFile.h
@@ -18,6 +18,22 @@
 class IAndroidPlatformFile : public IPhysicalPlatformFile
 {
 public:
+	// Methods that expose an argument for allowing Android Assets (i.e. `assets` directory in Gradle project) to be considered.  By default, they are
+	// not, because Asset stat can be very slow on some devices.
+	virtual bool FileExists(const TCHAR* Filename, bool bAllowAssets) = 0;
+	virtual int64 FileSize(const TCHAR* Filename, bool bAllowAssets) = 0;
+	virtual FFileStatData GetStatData(const TCHAR* FilenameOrDirectory, bool bAllowAssets) = 0;
+	virtual bool DirectoryExists(const TCHAR* Directory, bool bAllowAssets) = 0;
+	virtual bool IterateDirectory(const TCHAR* Directory, FDirectoryVisitor& Visitor, bool bAllowAssets) = 0;
+	virtual bool IterateDirectoryStat(const TCHAR* Directory, FDirectoryStatVisitor& Visitor, bool bAllowAssets) = 0;
+	// Using statements avoid a compiler warning - this says, "yes I am intending to add virtual overloads in a subclass, no it's not a mistake".
+	using IPlatformFile::FileExists;
+	using IPlatformFile::FileSize;
+	using IPlatformFile::GetStatData;
+	using IPlatformFile::DirectoryExists;
+	using IPlatformFile::IterateDirectory;
+	using IPlatformFile::IterateDirectoryStat;
+
 	static CORE_API IAndroidPlatformFile & GetPlatformPhysical();
 
 #if USE_ANDROID_FILE
