Android: Fix compile error from `final` in platform file.

#rnx

diff --git a/Engine/Source/Runtime/Core/Private/Android/AndroidPlatformFile.cpp b/Engine/Source/Runtime/Core/Private/Android/AndroidPlatformFile.cpp
--- a/Engine/Source/Runtime/Core/Private/Android/AndroidPlatformFile.cpp
+++ b/Engine/Source/Runtime/Core/Private/Android/AndroidPlatformFile.cpp
@@ -1253,8 +1253,9 @@ FAndroidMappedFileRegion::~FAndroidMappedFileRegion()
 	overriding content from the most "frozen" to the most "fluid"
 	state. Hence creating a virtual single union file-system.
 **/
-class CORE_API FAndroidPlatformFile final : public IAndroidPlatformFile
+class CORE_API FAndroidPlatformFile : public IAndroidPlatformFile
 {
+	// Note: TManagedStoragePlatformFile used below templates a subclass of FAndroidPlatformFile, so it cannot be made `final`.
 public:
 	// Singleton implementation.
 	static FAndroidPlatformFile & GetPlatformPhysical()
