Platform I/O dispatcher interface.

- extracted interface for reading I/O store encoded runtime data from the current I/O dispatcher file backend
- with this interface in place all backends will share the same priority, memory throttling and caching behavior when reading data from disk
- disabled by default and enabled with the command line argument -PlatformIo
- note that the current file I/O backend is using two threads to process file reads. A dedicated thread for reading data from disk
and the I/O dispatcher thread for scheduling compression tasks. This new implementation tries to do the same amount of
work using just a single thread. This has resulted in a somewhat awkard interface for scheduling decompression/decoding tasks
via a new interface IChunkDecoder that enables scheduling additional decompression/decode requests via an out parameter.
- this initial change only includes the generic platform version of the platform I/O dispatcher. Additional platform specific version has been developed and will be submitted in a later change.

#tests replayrun

diff --git a/Engine/Source/Runtime/Core/Internal/IO/GenericPlatformIoDispatcher.cpp b/Engine/Source/Runtime/Core/Internal/IO/GenericPlatformIoDispatcher.cpp
new file mode 100644
--- /dev/null
+++ b/Engine/Source/Runtime/Core/Internal/IO/GenericPlatformIoDispatcher.cpp
@@ -0,0 +1,269 @@
+// Copyright Epic Games, Inc. All Rights Reserved.
+
+#include "IO/GenericPlatformIoDispatcher.h"
+
+#include "Async/UniqueLock.h"
+#include "HAL/Event.h"
+#include "HAL/PlatformFile.h"
+#include "HAL/PlatformFileManager.h"
+#include "IO/IoDispatcherConfig.h"
+#include "IO/PlatformIoDispatcher.h"
+#include "IO/PlatformIoDispatcherBase.h"
+
+namespace UE
+{
+
+////////////////////////////////////////////////////////////////////////////////
+class FGenericPlatformIoDispatcher final
+	: public FPlatformIoDispatcherBase
+{
+	struct FFileHandle
+	{
+		bool IsValid() const { return Handle.IsValid(); }
+
+		TUniquePtr<IFileHandle> Handle;
+		int64					FileSize = -1;
+		uint32					CompressionBlockSize = 0;
+		uint32					FileId = 0;
+	};
+
+public:
+										FGenericPlatformIoDispatcher(FPlatformIoDispatcherCreateParams&& Params);
+										~FGenericPlatformIoDispatcher();
+
+	// IIoPlatformFile
+	virtual TIoStatusOr<FIoFileHandle>	OpenFile(const TCHAR* Filename, const FIoFileProperties& FileProperties, FIoFileStat* OutStat = nullptr) override;
+	virtual FIoStatus					CloseFile(FIoFileHandle FileHandle) override;
+
+private:
+	// FPlatformIoDispatcherBase
+	virtual FIoStatus					OnInitialize() override final;
+	virtual bool						Tick() override final;
+	virtual uint32						OnIoThreadEntry() override final;
+	virtual void						OnWakeUp() override final;
+	virtual FIoPlatformFileInfo			GetPlatformFileInfo(FIoFileHandle FileHandle) override final;
+
+	bool								ProcessFileIo(bool& bOutWaitingForMemory);
+	FEventRef							WakeUpEvent;
+};
+
+////////////////////////////////////////////////////////////////////////////////
+FGenericPlatformIoDispatcher::FGenericPlatformIoDispatcher(FPlatformIoDispatcherCreateParams&& Params)
+	: FPlatformIoDispatcherBase(MoveTemp(Params))
+{
+}
+
+FGenericPlatformIoDispatcher::~FGenericPlatformIoDispatcher()
+{
+	bStopRequested = true;
+	if (Thread.IsValid())
+	{
+		WakeUpEvent->Trigger();
+		Thread.Reset();
+	}
+}
+
+TIoStatusOr<FIoFileHandle> FGenericPlatformIoDispatcher::OpenFile(const TCHAR* Filename, const FIoFileProperties& FileProperties, FIoFileStat* OutStat)
+{
+	IPlatformFile& Ipf = IPlatformFile::GetPlatformPhysical();
+
+	TUniquePtr<FFileHandle> FileHandle	= MakeUnique<FFileHandle>();
+	FileHandle->FileSize				= Ipf.FileSize(Filename);
+	FileHandle->CompressionBlockSize	= FileProperties.CompressionBlockSize;
+	FileHandle->FileId					= GetNextFileId();
+	check(FileHandle->FileId > 0);
+
+	if (FileHandle->FileSize < 0)
+	{
+		return FIoStatus(EIoErrorCode::NotFound);
+	}
+
+	FileHandle->Handle.Reset(Ipf.OpenReadNoBuffering(Filename));
+	if (FileHandle.IsValid() == false)
+	{
+		return FIoStatus(EIoErrorCode::NotFound);
+	}
+
+	if (OutStat != nullptr)
+	{
+		OutStat->FileSize = FileHandle->FileSize;
+	}
+
+	return FIoFileHandle(UPTRINT(FileHandle.Release()));
+}
+
+FIoStatus FGenericPlatformIoDispatcher::CloseFile(FIoFileHandle FileHandle)
+{
+	if (FileHandle.IsValid())
+	{
+		FFileHandle* Handle = reinterpret_cast<FFileHandle*>(FileHandle.Value());
+		delete Handle;
+	}
+
+	return FIoStatus::Ok;
+}
+
+FIoStatus FGenericPlatformIoDispatcher::OnInitialize()
+{
+	FileBlockSize								= (GIoDispatcherBufferSizeKB > 0 ? uint32(GIoDispatcherBufferSizeKB) << 10 : 256 << 10);
+	const bool bSortByOffset					= GIoDispatcherSortRequestsByOffset > 0;
+	const uint32 DefaultCompressionBlockSize	= 64 << 10;
+
+	UE_LOG(LogPlatformIoDispatcher, Log, TEXT("Initializing, Platform='%s', ReadSize=%dKB, ReadMemory=%dMB, FileCache=%dMB, RequestSorting=%s, MaxConcurrentDecodeJobs=%d"),
+		TEXT("Generic"), GIoDispatcherBufferSizeKB, GIoDispatcherBufferMemoryMB, GIoDispatcherCacheSizeMB, bSortByOffset ? TEXT("ByOffset") : TEXT("BySeqNo"), GIoDispatcherDecompressionWorkerCount);
+
+	ChunkBlockMemoryPool.Initialize(GIoDispatcherDecompressionWorkerCount, DefaultCompressionBlockSize);
+	FileBlockMemoryPool.Initialize(FileBlockSize, GIoDispatcherBufferMemoryMB << 20, GIoDispatcherBufferAlignment); 
+	FileBlockCache.Initialize(FileBlockSize, uint64(GIoDispatcherCacheSizeMB) << 20);
+	FileBlockRequestQueue.SortByOffset(bSortByOffset);
+
+	return FIoStatus::Ok;
+}
+
+bool FGenericPlatformIoDispatcher::Tick()
+{
+	// Only called in single threaded mode
+	bool bDidSomething = false;
+	for (;;)
+	{
+		bool bWaitingForMemory = false;
+		if (ProcessFileIo(bWaitingForMemory) == false)
+		{
+			break;
+		}
+		bDidSomething = true;
+	}
+
+	return bDidSomething;
+}
+
+uint32 FGenericPlatformIoDispatcher::OnIoThreadEntry()
+{
+	for(;;)
+	{
+		bool bWaitingForMemory = false;
+		const bool bDidSomething = ProcessFileIo(bWaitingForMemory);
+		if (bDidSomething == false)
+		{
+			if (bStopRequested)
+			{
+				break;
+			}
+			{
+				if (bWaitingForMemory)
+				{
+					TRACE_CPUPROFILER_EVENT_SCOPE(WaitingForMemory);
+					WakeUpEvent->Wait();
+				}
+				else
+				{
+					WakeUpEvent->Wait();
+				}
+			}
+		}
+	}
+
+	return 0;
+}
+
+void FGenericPlatformIoDispatcher::OnWakeUp()
+{
+	WakeUpEvent->Trigger();
+}
+
+bool FGenericPlatformIoDispatcher::ProcessFileIo(bool& bOutWaitingForMemory)
+{
+	bOutWaitingForMemory	= false;
+	FIoFileBlockRequest* FileBlockRequest = nullptr;
+	{
+		UE::TUniqueLock Lock(FileBlockMutex);
+		if (FIoFileBlockRequest *NextPending = FileBlockRequestQueue.Peek(); NextPending != nullptr)
+		{
+			if (NextPending->ErrorCode == EIoErrorCode::Cancelled)
+			{
+				FileBlockRequest = FileBlockRequestQueue.Dequeue();
+			}
+			else
+			{
+				if (bOutWaitingForMemory = FileBlockMemoryPool.IsEmpty(); bOutWaitingForMemory == false)
+				{
+					FileBlockRequest			= FileBlockRequestQueue.Dequeue();
+					FileBlockRequest->Buffer	= FileBlockMemoryPool.Alloc(FileBlockRequest->BufferHandle);
+				}
+			}
+		}
+	}
+
+	if (FileBlockRequest == nullptr)
+	{
+		return false;
+	}
+	check(FileBlockRequest->Size > 0);
+	check(FileBlockRequest->EncodedBlockRequests.IsEmpty() == false);
+	check(FileBlockRequest->QueueStatus == FIoFileBlockRequest::EQueueStatus::Dequeued);
+
+	if (FileBlockRequest->ErrorCode != EIoErrorCode::Cancelled)
+	{
+		FileBlockRequest->ErrorCode = EIoErrorCode::Ok;
+		if (FileBlockCache.Get(*FileBlockRequest) == false)
+		{
+			TRACE_CPUPROFILER_EVENT_SCOPE(ReadBlockFromFile);
+
+			FFileHandle&	FileHandle = *reinterpret_cast<FFileHandle*>(FileBlockRequest->FileHandle.Value());
+			int32			RetryCount = 0;
+
+			Stats.OnFileBlockReadStarted(*FileBlockRequest);
+			FileBlockRequest->ErrorCode = EIoErrorCode::ReadError;
+			while (RetryCount++ < 10)
+			{
+				if (FileHandle.Handle->Seek(int64(FileBlockRequest->FileOffset)) == false)
+				{
+					UE_LOG(LogPlatformIoDispatcher, Warning, TEXT("Failed to seek to offset %lld"), int64(FileBlockRequest->FileOffset));
+					continue;
+				}
+				if (FileHandle.Handle->Read(reinterpret_cast<uint8*>(FileBlockRequest->Buffer), int64(FileBlockRequest->Size)) == false)
+				{
+					UE_LOG(LogPlatformIoDispatcher, Warning, TEXT("Failed to read %lld bytes at offset %lld"), int64(FileBlockRequest->Size), int64(FileBlockRequest->FileOffset));
+					continue;
+				}
+				FileBlockRequest->ErrorCode = EIoErrorCode::Ok;
+				FileBlockCache.Put(*FileBlockRequest);
+				break;
+			}
+			Stats.OnFileBlockReadCompleted(*FileBlockRequest);
+		}
+	}
+
+	ProcessCompletedFileBlock(*FileBlockRequest);
+
+	return true;
+}
+
+FIoPlatformFileInfo FGenericPlatformIoDispatcher::GetPlatformFileInfo(FIoFileHandle FileHandle)
+{
+	if (const FFileHandle* InternalHandle = reinterpret_cast<const FFileHandle*>(FileHandle.Value()))
+	{
+		return FIoPlatformFileInfo
+		{
+			.FileSize				= InternalHandle->FileSize,
+			.FileId					= InternalHandle->FileId,
+			.CompressionBlockSize	= InternalHandle->CompressionBlockSize
+		};
+	}
+	
+	return FIoPlatformFileInfo{};
+}
+
+////////////////////////////////////////////////////////////////////////////////
+TUniquePtr<IPlatformIoDispatcher> FGenericPlatformIoDispatcherFactory::Create(FPlatformIoDispatcherCreateParams&& Params)
+{
+	return MakeUnique<FGenericPlatformIoDispatcher>(MoveTemp(Params));
+}
+
+////////////////////////////////////////////////////////////////////////////////
+TUniquePtr<IPlatformIoDispatcher> MakeGenericPlatformIoDispatcher(FPlatformIoDispatcherCreateParams&& Params)
+{
+	return MakeUnique<FGenericPlatformIoDispatcher>(MoveTemp(Params));
+}
+
+} // namespace UE
diff --git a/Engine/Source/Runtime/Core/Internal/IO/GenericPlatformIoDispatcher.h b/Engine/Source/Runtime/Core/Internal/IO/GenericPlatformIoDispatcher.h
new file mode 100644
--- /dev/null
+++ b/Engine/Source/Runtime/Core/Internal/IO/GenericPlatformIoDispatcher.h
@@ -0,0 +1,18 @@
+// Copyright Epic Games, Inc. All Rights Reserved.
+
+#pragma once
+
+#include "Templates/UniquePtr.h"
+
+namespace UE
+{
+
+class FGenericPlatformIoDispatcherFactory
+{
+public:
+	CORE_API static TUniquePtr<class IPlatformIoDispatcher> Create(struct FPlatformIoDispatcherCreateParams&&);
+};
+
+using FPlatformIoDispatcherFactory = FGenericPlatformIoDispatcherFactory;
+
+}
diff --git a/Engine/Source/Runtime/Core/Internal/IO/IoDispatcherConfig.cpp b/Engine/Source/Runtime/Core/Internal/IO/IoDispatcherConfig.cpp
--- a/Engine/Source/Runtime/Core/Internal/IO/IoDispatcherConfig.cpp
+++ b/Engine/Source/Runtime/Core/Internal/IO/IoDispatcherConfig.cpp
@@ -1,7 +1,9 @@
 // Copyright Epic Games, Inc. All Rights Reserved.
 
 #include "IO/IoDispatcherConfig.h"
+#include "Misc/CommandLine.h"
 #include "HAL/IConsoleManager.h"
+#include "Misc/Parse.h"
 
 int32 GIoDispatcherBufferSizeKB = 256;
 static FAutoConsoleVariableRef CVar_IoDispatcherBufferSizeKB(
@@ -95,9 +97,34 @@ static FAutoConsoleVariableRef CVar_GIoDispatcherDecompressOnIdle(
 	TEXT("When enabled the I/O dispatcher thread will help out decompress chunk blocks when idle.")
 );
 
-int32 GIoDispatcherMaxConsecutiveDecompressionJobs = 2;
+int32 GIoDispatcherMaxConsecutiveDecompressionJobs = 4;
 static FAutoConsoleVariableRef CVar_IoDispatcherMaxConsecutiveDecompressionJobs(
 	TEXT("s.IoDispatcherMaxConsecutiveDecompressionJobs"),
 	GIoDispatcherMaxConsecutiveDecompressionJobs,
 	TEXT("Max consecutive decompression jobs before re-launching tasks.")
 );
+
+int32 GIoDispatcherMaxResolveBatchSize = 512;
+static FAutoConsoleVariableRef CVar_IoDispatcherMaxResolveBatchSize(
+	TEXT("s.IoDispatcherMaxResolveBatchSize"),
+	GIoDispatcherMaxResolveBatchSize,
+	TEXT("")
+);
+
+bool IsPlatformIoDispatcherEnabled()
+{
+#if !UE_BUILD_SHIPPING
+	static struct FIoDispatcherConfig
+	{
+		bool bPlatformIoDispatcherEnabled;
+		FIoDispatcherConfig(const TCHAR* CmdLine)
+		{
+			bPlatformIoDispatcherEnabled = FParse::Param(CmdLine, TEXT("PlatformIo"));
+		}
+	} IoDispatcherConfig(FCommandLine::Get());
+
+	return IoDispatcherConfig.bPlatformIoDispatcherEnabled;
+#else
+	return false;
+#endif
+}
diff --git a/Engine/Source/Runtime/Core/Internal/IO/IoDispatcherConfig.h b/Engine/Source/Runtime/Core/Internal/IO/IoDispatcherConfig.h
--- a/Engine/Source/Runtime/Core/Internal/IO/IoDispatcherConfig.h
+++ b/Engine/Source/Runtime/Core/Internal/IO/IoDispatcherConfig.h
@@ -18,3 +18,6 @@ CORE_API extern int32 GIoDispatcherCanDecompressOnStarvation;
 CORE_API extern int32 GIoDispatcherForceSynchronousScatter;
 CORE_API extern int32 GIoDispatcherDecompressOnIdle;
 CORE_API extern int32 GIoDispatcherMaxConsecutiveDecompressionJobs;
+CORE_API extern int32 GIoDispatcherMaxResolveBatchSize;
+
+CORE_API bool IsPlatformIoDispatcherEnabled();
diff --git a/Engine/Source/Runtime/Core/Internal/IO/PlatformIoDispatcher.cpp b/Engine/Source/Runtime/Core/Internal/IO/PlatformIoDispatcher.cpp
new file mode 100644
--- /dev/null
+++ b/Engine/Source/Runtime/Core/Internal/IO/PlatformIoDispatcher.cpp
@@ -0,0 +1,56 @@
+// Copyright Epic Games, Inc. All Rights Reserved.
+
+#include "PlatformIoDispatcher.h"
+
+#if PLATFORM_IMPLEMENTS_IO
+#include "HAL/PreprocessorHelpers.h"
+#include COMPILED_PLATFORM_HEADER_WITH_PREFIX(IO, PlatformIoDispatcher.h)
+#else
+#include "IO/GenericPlatformIoDispatcher.h"
+#endif
+
+DEFINE_LOG_CATEGORY(LogPlatformIoDispatcher);
+
+namespace UE
+{
+
+TUniquePtr<IPlatformIoDispatcher> GPlatformIoDispatcher;
+
+extern TUniquePtr<IPlatformIoDispatcher> MakeGenericPlatformIoDispatcher(FPlatformIoDispatcherCreateParams&&);
+
+void FPlatformIoDispatcher::Create(FPlatformIoDispatcherCreateParams&& Params)
+{
+	if (Params.bForceGeneric)
+	{
+		GPlatformIoDispatcher = MakeGenericPlatformIoDispatcher(MoveTemp(Params));
+	}
+	else
+	{
+		GPlatformIoDispatcher = FPlatformIoDispatcherFactory::Create(MoveTemp(Params));
+	}
+}
+
+void FPlatformIoDispatcher::Initialize()
+{
+	if (GPlatformIoDispatcher.IsValid())
+	{
+		GPlatformIoDispatcher->Initialize();
+	}
+}
+
+void FPlatformIoDispatcher::Shutdown()
+{
+	GPlatformIoDispatcher.Reset();
+}
+
+IPlatformIoDispatcher& FPlatformIoDispatcher::Get()
+{
+	return *GPlatformIoDispatcher;
+}
+
+IPlatformIoDispatcher* FPlatformIoDispatcher::TryGet()
+{
+	return GPlatformIoDispatcher.Get();
+}
+
+} // namespace UE
diff --git a/Engine/Source/Runtime/Core/Internal/IO/PlatformIoDispatcher.h b/Engine/Source/Runtime/Core/Internal/IO/PlatformIoDispatcher.h
new file mode 100644
--- /dev/null
+++ b/Engine/Source/Runtime/Core/Internal/IO/PlatformIoDispatcher.h
@@ -0,0 +1,308 @@
+// Copyright Epic Games, Inc. All Rights Reserved.
+
+#pragma once
+
+#include "IO/IoDispatcherPriority.h"
+#include "IO/IoStatus.h"
+#include "IO/IoBuffer.h"
+#include "Logging/LogMacros.h"
+#include "Math/NumericLimits.h"
+#include "Memory/MemoryFwd.h"
+#include "Misc/EnumClassFlags.h"
+#include "Templates/Function.h"
+#include "UObject/NameTypes.h"
+
+#if !defined(PLATFORM_IMPLEMENTS_IO)
+#define PLATFORM_IMPLEMENTS_IO 0
+#endif
+
+CORE_API DECLARE_LOG_CATEGORY_EXTERN(LogPlatformIoDispatcher, Log, All);
+
+namespace UE
+{
+
+/** Priority levels when reading encoded I/O store chunks from file. */
+enum class EIoFileReadPriority : uint8
+{
+	Min,
+	Low,
+	Medium,
+	High,
+	Max,
+	Count
+};
+
+inline bool operator>(EIoFileReadPriority LHS, EIoFileReadPriority RHS)
+{
+	return static_cast<uint8>(LHS) > static_cast<uint8>(RHS);
+}
+
+inline bool operator<(EIoFileReadPriority LHS, EIoFileReadPriority RHS)
+{
+	return static_cast<uint8>(LHS) < static_cast<uint8>(RHS);
+}
+
+inline bool operator<=(EIoFileReadPriority LHS, EIoFileReadPriority RHS)
+{
+	return static_cast<uint8>(LHS) <= static_cast<uint8>(RHS);
+}
+
+/** File read priority to string. */
+inline const TCHAR* LexToString(EIoFileReadPriority Priority)
+{
+	static const TCHAR* StringTable[] =
+	{
+		TEXT("Min"),
+		TEXT("Low"),
+		TEXT("Medium"),
+		TEXT("High"),
+		TEXT("Max"),
+	};
+
+	return StringTable[static_cast<uint8>(Priority)];
+};
+
+/** Narrows an I/O dispatcher priorty to file read priority. */
+static inline EIoFileReadPriority IoFileReadPriorityFromDispatcherPriority(int32 Priority)
+{
+	if (Priority < IoDispatcherPriority_Low)
+	{
+		return EIoFileReadPriority::Min;
+	}
+	if (Priority < IoDispatcherPriority_Medium)
+	{
+		return EIoFileReadPriority::Low;
+	}
+	if (Priority < IoDispatcherPriority_High)
+	{
+		return EIoFileReadPriority::Medium;
+	}
+
+	return Priority > IoDispatcherPriority_High ? EIoFileReadPriority::Max : EIoFileReadPriority::High;
+}
+
+/** Handle to a file on disk. */
+struct FIoFileHandle
+{
+			FIoFileHandle() = default;
+			explicit FIoFileHandle(uint64 InHandle)
+				: Handle(InHandle)
+			{ }
+
+	uint64	Value() const { return Handle; }
+	bool	IsValid() const { return Handle != 0; }
+
+	friend uint32 GetTypeHash(FIoFileHandle FileHandle) { return GetTypeHash(FileHandle.Handle); }
+	friend	bool operator==(FIoFileHandle LHS, FIoFileHandle RHS) { return LHS.Handle == RHS.Handle; }
+	friend	bool operator!=(FIoFileHandle LHS, FIoFileHandle RHS) { return LHS.Handle != RHS.Handle; }
+	friend	bool operator<(FIoFileHandle LHS, FIoFileHandle RHS) { return LHS.Handle < RHS.Handle; }
+
+private:
+	uint64 Handle = 0;
+};
+
+/** File stats returned when opening a file for reading. */
+struct FIoFileStat
+{
+	uint64 FileSize = 0;
+};
+
+/** Specifies whether a file contains encrypted or signed I/O store chunks. */
+enum class EIoFilePropertyFlags : uint8
+{
+	None		= 0,
+	Encrypted	= (1 << 0),
+	Signed		= (1 << 1)
+};
+ENUM_CLASS_FLAGS(EIoFilePropertyFlags);
+
+/**
+ * File properties are used to determine whether it's possible to use
+ * direct-read, i.e. reading directly from the file to destination buffer and use
+ * hardware decompression.
+ */
+struct FIoFileProperties
+{
+	TConstArrayView<FName>	CompressionMethods;
+	uint32					CompressionBlockSize = 0; 
+	EIoFilePropertyFlags	Flags = EIoFilePropertyFlags::None;
+};
+
+/**
+ * Handle to a platform file read request. Owned by the caller and needs to be deleted when completed.
+ */
+struct FIoFileReadRequest
+{
+			FIoFileReadRequest() = default;
+			explicit FIoFileReadRequest(uint64 InHandle)
+				: Handle(InHandle)
+			{ }
+
+	uint64	Value() const { return Handle; }
+	bool	IsValid() const { return Handle != 0; }
+
+private:
+	uint64 Handle = 0;
+};
+
+/**
+ * The scatter-gather parameters specifies the I/O store encoded block boundaries, 
+ * whether the block is compressed and or encrypted, whether to check the block
+ * signature hash and from where in the uncompressed block to copy the block data
+ * to the destination buffer.
+ */
+class FIoScatterGatherRequestParams
+{
+public:
+	struct FScatterParams
+	{
+		uint64		BlockFileOffset			= 0;
+		uint32		BlockCompressedSize		= 0;
+		uint32		BlockUncompresedSize	= 0;
+		uint64		ScatterOffset			= 0;
+		uint64		ScatterSize				= 0;
+		uint64		DestinationOffset		= 0;
+		uint32		BlockIndex				= MAX_uint32;
+		FName		CompressionMethod		= NAME_None;
+		FMemoryView	EncryptionKey;
+		FMemoryView	BlockHash;
+	};
+
+	using FScatterParamsArray = TArray<FScatterParams, TInlineAllocator<4>>;
+
+	/**
+	 * Creates a new set of scatter-gather parameters.
+	 *
+	 * @param	InFileHandle		Handle to the file to read from.
+	 * @param	InDestination		Reference to the destination buffer.
+	 * @param	InDestinationSize	Size in bytes of the destination buffer.
+	 * @param	InUserData			Optional user data associated with the request.
+	 * @param	InPriority			The priority of the request.
+	 */
+	FIoScatterGatherRequestParams(
+		FIoFileHandle InFileHandle,
+		FIoBuffer& InDestination,
+		uint64 InDestinationSize,
+		void* InUserData,
+		EIoFileReadPriority InPriority)
+			: FileHandle(InFileHandle)
+			, Destination(InDestination)
+			, DestinationSize(InDestinationSize)
+			, UserData(InUserData)
+			, Priority(InPriority)
+	{ }
+
+	/**
+	 *	Add block scatter parameters.
+	 *
+	 *	@param	BlockFileOffset			Offset of the encoded block in the file.
+	 *	@param	BlockIndex				Index of the block in the TOC.
+	 *	@param	BlockCompressedSize		Compressed size of the block.
+	 *	@param	BlockUncompressedSize	Uncompressed size of the block.
+	 *	@param	ScatterOffset			Offset in the uncompressed block from where to copy the data to the destination buffer.
+	 *	@param	ScatterSize				Number of bytes to copy from the uncompressed block.
+	 *	@param	DestinationOffset		Offset in the destination buffer to copy the data from the uncompressed buffer.
+	 *	@param	CompressionMethod		Method used to compress the block data.
+	 *	@param	EncryptionKey			Optional memory view to the encryption key.
+	 *	@param	BlockHash				Optional memory view to the block hash for signature check.
+	 */
+	void Scatter(
+		uint64 BlockFileOffset,
+		uint32 BlockIndex,
+		uint32 BlockCompressedSize,
+		uint32 BlockUncompresedSize,
+		uint64 ScatterOffset,
+		uint64 ScatterSize,
+		uint64 DestinationOffset,
+		FName CompressionMethod,
+		FMemoryView EncryptionKey,
+		FMemoryView BlockHash)
+	{
+		Params.Add(FScatterParams
+		{
+			.BlockFileOffset		= BlockFileOffset,
+			.BlockCompressedSize	= BlockCompressedSize,
+			.BlockUncompresedSize	= BlockUncompresedSize,
+			.ScatterOffset			= ScatterOffset,
+			.ScatterSize			= ScatterSize,
+			.DestinationOffset		= DestinationOffset,
+			.BlockIndex				= BlockIndex,
+			.CompressionMethod		= CompressionMethod,
+			.EncryptionKey			= EncryptionKey,
+			.BlockHash				= BlockHash
+		});
+	}
+
+	FIoFileHandle			FileHandle;
+	FIoBuffer&				Destination;
+	uint64					DestinationSize;
+	void*					UserData;
+	EIoFileReadPriority		Priority;
+	FScatterParamsArray		Params;
+};
+
+/** File read result. */
+struct FIoFileReadResult
+{
+	void*			UserData		= nullptr;
+	uint32			FailedBlockId	= MAX_uint32;
+	EIoErrorCode	ErrorCode		= EIoErrorCode::Unknown;
+};
+
+using FIoFileReadRequestCompleted = TUniqueFunction<void(FIoFileReadResult&&)>;
+
+/** Parameters for reading directly from file on disk to the destination buffer. */
+struct FIoDirectReadRequestParams
+{
+	FIoFileHandle	FileHandle;
+	FIoBuffer&		Dst;
+	uint64			Offset		= MAX_uint64;
+	uint64			Size		= 0;
+	void*			UserData	= nullptr;
+};
+
+/** Parameters passed to the platform specific I/O dispatcher. */
+struct FPlatformIoDispatcherCreateParams
+{
+	bool bMultithreaded	= true;
+	bool bForceGeneric	= false;
+};
+
+/** Interface for platform specific I/O dispatcher. */
+class IPlatformIoDispatcher
+{
+public:
+	virtual										~IPlatformIoDispatcher() = default;
+	/** Initialize the dispatcher. Called after the config system has been initialized. */
+	virtual FIoStatus							Initialize() = 0;
+	/** Open file for reading. The file properties specifies whether the file is compressed, encrypted and signed. */ 
+	virtual TIoStatusOr<FIoFileHandle>			OpenFile(const TCHAR* Filename, const FIoFileProperties& FileProperties, FIoFileStat* OutStat = nullptr) = 0;
+	/** Close the file. Assumes no pending I/O requests are inflight for the specified file. */
+	virtual FIoStatus							CloseFile(FIoFileHandle FileHandle) = 0;
+	/** Issue a scatter-gather request. */
+	[[nodiscard]] virtual FIoFileReadRequest	ScatterGather(FIoScatterGatherRequestParams&& Params, FIoFileReadRequestCompleted&& OnCompleted) = 0;
+	/** Issue a direct read request. Returns an invalid file handle if direct read is not supported for the specifid file. */ 
+	[[nodiscard]] virtual FIoFileReadRequest	ReadDirect(FIoDirectReadRequestParams&& Params, FIoFileReadRequestCompleted&& OnCompleted) = 0;
+	/** Tick the dispatcher. Only used when multithreading is disabled. */ 
+	virtual bool								Tick() = 0;
+	/** Update priority for an inflight I/O request. */ 
+	virtual void								UpdatePriority(FIoFileReadRequest Request, EIoFileReadPriority NewPriority) = 0;
+	/** Cancel inflight I/O request. */ 
+	virtual void								CancelRequest(FIoFileReadRequest Request) = 0;
+	/** Cancel all inflight I/O request for the specified file. */ 
+	virtual void								CancelAllRequests(FIoFileHandle FileHandle) = 0;
+	/** Delete the request. Must be called after the completion callback has been triggered. */ 
+	virtual void								DeleteRequest(FIoFileReadRequest Request) = 0;
+};
+
+class FPlatformIoDispatcher
+{
+public:
+	CORE_API static void					Create(FPlatformIoDispatcherCreateParams&& Config);
+	CORE_API static void					Initialize();
+	CORE_API static void					Shutdown();
+	CORE_API static IPlatformIoDispatcher&	Get();
+	CORE_API static IPlatformIoDispatcher*	TryGet();
+};
+
+} // namespace UE
diff --git a/Engine/Source/Runtime/Core/Internal/IO/PlatformIoDispatcherBase.cpp b/Engine/Source/Runtime/Core/Internal/IO/PlatformIoDispatcherBase.cpp
new file mode 100644
--- /dev/null
+++ b/Engine/Source/Runtime/Core/Internal/IO/PlatformIoDispatcherBase.cpp
@@ -0,0 +1,1441 @@
+// Copyright Epic Games, Inc. All Rights Reserved.
+
+#include "PlatformIoDispatcherBase.h"
+
+#include "Algo/BinarySearch.h"
+#include "Async/UniqueLock.h"
+#include "HAL/LowLevelMemTracker.h"
+#include "IO/IoChunkEncoding.h"
+#include "IO/IoStatus.h"
+#include "IO/IoDispatcherConfig.h"
+#include "IO/PlatformIoDispatcher.h"
+#include "Math/NumericLimits.h"
+#include "Misc/AES.h"
+
+#include <atomic>
+
+namespace UE
+{
+
+////////////////////////////////////////////////////////////////////////////////
+template<typename... Args>
+bool IsIoErrorAnyOf(EIoErrorCode ErrorCode, Args... Expected)
+{
+	return ((ErrorCode == Expected) || ...);
+}
+
+////////////////////////////////////////////////////////////////////////////////
+uint32 FIoFileBlockRequest::NextSeqNo = 0;
+
+////////////////////////////////////////////////////////////////////////////////
+FIoFileBlockRequestQueue::FIoFileBlockRequestQueue(FPlatformIoDispatcherStats& InStats)
+	: Stats(InStats)
+{
+	if (bSortByOffset)
+	{
+		for (int32 Idx = static_cast<int32>(EIoFileReadPriority::Count) - 1; Idx >= 0; --Idx)
+		{
+			PrioQueues[Idx].ByOffset.Reserve(128);
+		}
+	}
+}
+
+void FIoFileBlockRequestQueue::Enqueue(FIoFileBlockRequestList&& FileBlockRequests)
+{
+	while (FIoFileBlockRequest* Request = FileBlockRequests.PopHead())
+	{
+		check(Request->QueueStatus == FIoFileBlockRequest::EQueueStatus::None);
+		Request->QueueStatus = FIoFileBlockRequest::EQueueStatus::Enqueued;
+		if (bSortByOffset)
+		{
+			FPrioQueue& PrioQueue	= PrioQueues[static_cast<int32>(Request->Priority)];
+			const int32 Idx			= Algo::UpperBoundBy(PrioQueue.ByOffset, ToFileOffset(Request), ToFileOffset, FileOffsetLess);
+
+			PrioQueue.ByOffset.Insert(Request, Idx);
+			PrioQueue.BySequence.AddTail(Request);
+			PrioQueue.PeekIndex = INDEX_NONE;
+		}
+		else
+		{
+			Heap.HeapPush(Request, SeqNoLess);
+		}
+
+		Stats.OnFileBlockRequestEnqueued(*Request);
+	}
+}
+
+FIoFileBlockRequest* FIoFileBlockRequestQueue::Dequeue()
+{
+	const bool bDequeue = true;
+	FIoFileBlockRequest* Request = GetCancelled(bDequeue);
+	if (Request == nullptr)
+	{
+		if (bSortByOffset)
+		{
+			Request = GetByOffset(bDequeue);
+		}
+		else if (Heap.IsEmpty() == false)
+		{
+			Heap.HeapPop(Request, SeqNoLess, EAllowShrinking::No);
+		}
+	}
+
+	if (Request != nullptr)
+	{
+		Request->QueueStatus = FIoFileBlockRequest::EQueueStatus::Dequeued;
+		Stats.OnFileBlockRequestDequeued(*Request);
+	}
+
+	return Request;
+}
+
+FIoFileBlockRequest* FIoFileBlockRequestQueue::Peek()
+{
+	const bool bDequeue = false;
+	FIoFileBlockRequest* Request = GetCancelled(bDequeue);
+	if (Request == nullptr)
+	{
+		if (bSortByOffset)
+		{
+			Request = GetByOffset(bDequeue);
+		}
+		else if (Heap.IsEmpty() == false)
+		{
+			Request = Heap.HeapTop();
+		}
+	}
+
+	return Request;
+}
+
+void FIoFileBlockRequestQueue::Reprioritize()
+{
+	TRACE_CPUPROFILER_EVENT_SCOPE(RequestQueueUpdateOrder);
+
+	if (bSortByOffset)
+	{
+		TArray<FIoFileBlockRequest*> RequestsToReprioritize;
+		for (int32 QueueIdx = 0; FPrioQueue& PrioQueue : PrioQueues)
+		{
+			PrioQueue.PeekIndex						= INDEX_NONE;
+			const EIoFileReadPriority QueuePriority = static_cast<EIoFileReadPriority>(QueueIdx++);
+
+			for (int32 RequestIdx = PrioQueue.ByOffset.Num() - 1; RequestIdx >= 0; --RequestIdx)
+			{
+				FIoFileBlockRequest* Request = PrioQueue.ByOffset[RequestIdx];
+				if (Request->Priority != QueuePriority)
+				{
+					RequestsToReprioritize.Add(Request);
+					PrioQueue.ByOffset.RemoveAt(RequestIdx, EAllowShrinking::No);
+					PrioQueue.BySequence.Remove(Request);
+				}
+			}
+		}
+
+		for (FIoFileBlockRequest* Request : RequestsToReprioritize)
+		{
+			FPrioQueue& PrioQueue	= PrioQueues[static_cast<int32>(Request->Priority)];
+			const int32 Idx			= Algo::UpperBoundBy(PrioQueue.ByOffset, ToFileOffset(Request), ToFileOffset, FileOffsetLess);
+
+			PrioQueue.ByOffset.Insert(Request, Idx);
+			PrioQueue.BySequence.AddTail(Request);
+			PrioQueue.PeekIndex = INDEX_NONE;
+		}
+	}
+	else if (Heap.IsEmpty() == false)
+	{
+		Heap.Heapify(SeqNoLess);
+	}
+}
+
+FIoFileBlockRequest* FIoFileBlockRequestQueue::GetByOffset(bool bDequeue)
+{
+	for (int32 Idx = static_cast<int32>(EIoFileReadPriority::Count) - 1; Idx >= 0; --Idx)
+	{
+		FPrioQueue& PrioQueue = PrioQueues[Idx];
+		if (PrioQueue.BySequence.IsEmpty() == false)
+		{
+			return GetByOffset(PrioQueue, static_cast<EIoFileReadPriority>(Idx), bDequeue);
+		}
+	}
+
+	return nullptr;
+}
+
+FIoFileBlockRequest* FIoFileBlockRequestQueue::GetByOffset(FPrioQueue& PrioQueue, EIoFileReadPriority QueuePriority, bool bDequeue)
+{
+	check(!PrioQueue.ByOffset.IsEmpty() && !PrioQueue.BySequence.IsEmpty());
+
+	int32 RequestIndex = INDEX_NONE;
+	if (PrioQueue.PeekIndex != INDEX_NONE)
+	{
+		RequestIndex = PrioQueue.PeekIndex;
+	}
+	else 
+	{
+		bool bHeadRequestTooOld = false;
+		if (GIoDispatcherRequestLatencyCircuitBreakerMS > 0)
+		{
+			// If our oldest request has been unserviced for too long, grab that instead of the next sequential read
+			const uint64 ThresholdCycles = uint64((GIoDispatcherRequestLatencyCircuitBreakerMS * 1000.0) / FPlatformTime::GetSecondsPerCycle64());
+			bHeadRequestTooOld = (FPlatformTime::Cycles64() - PrioQueue.BySequence.PeekHead()->TimeCreated) >= ThresholdCycles;
+		}
+
+		const bool bChooseByOffset = 
+				LastFileOffset.FileHandle.IsValid()
+			&&	!bHeadRequestTooOld 
+			&&  (GIoDispatcherMaintainSortingOnPriorityChange || LastFileOffset.Priority == QueuePriority);
+		if (bChooseByOffset)
+		{
+			RequestIndex = Algo::LowerBoundBy(PrioQueue.ByOffset, LastFileOffset, ToFileOffset, FileOffsetLess);
+			if (PrioQueue.ByOffset.IsValidIndex(RequestIndex)) 
+			{
+				if (PrioQueue.ByOffset[RequestIndex]->FileHandle.Value() != LastFileOffset.FileHandle.Value())
+				{
+					// Changing file handle so switch back to the oldest outstanding request 
+					RequestIndex = INDEX_NONE;
+				}
+			}
+		}
+
+		if (PrioQueue.ByOffset.IsValidIndex(RequestIndex) == false)
+		{
+			RequestIndex = PrioQueue.ByOffset.Find(PrioQueue.BySequence.PeekHead());
+			check(PrioQueue.ByOffset[RequestIndex] == PrioQueue.BySequence.PeekHead());
+		}
+	}
+
+	check(PrioQueue.ByOffset.IsValidIndex(RequestIndex));
+	FIoFileBlockRequest* Request = PrioQueue.ByOffset[RequestIndex];
+	if (bDequeue)
+	{
+		PrioQueue.ByOffset.RemoveAt(RequestIndex, EAllowShrinking::No);
+		const bool bRemoved = PrioQueue.BySequence.Remove(Request);
+		PrioQueue.PeekIndex = INDEX_NONE;
+		LastFileOffset		= ToFileOffset(Request);
+		check(bRemoved);
+	}
+	else
+	{
+		PrioQueue.PeekIndex = RequestIndex;
+	}
+
+	return Request;
+}
+
+FIoFileBlockRequest* FIoFileBlockRequestQueue::GetCancelled(bool bDequeue)
+{
+	if (bReprioritizeCancelled)
+	{
+		bReprioritizeCancelled = false;
+		TRACE_CPUPROFILER_EVENT_SCOPE(RequestQueueRemoveCancelled);
+		if (bSortByOffset)
+		{
+			for (int32 QueueIdx = 0; FPrioQueue& PrioQueue : PrioQueues)
+			{
+				PrioQueue.PeekIndex = INDEX_NONE;
+				for (int32 RequestIdx = PrioQueue.ByOffset.Num() - 1; RequestIdx >= 0; --RequestIdx)
+				{
+					FIoFileBlockRequest* Request = PrioQueue.ByOffset[RequestIdx];
+					if (Request->ErrorCode == EIoErrorCode::Cancelled)
+					{
+						CancelledHeap.HeapPush(Request, SeqNoLess);
+						PrioQueue.ByOffset.RemoveAt(RequestIdx, EAllowShrinking::No);
+						PrioQueue.BySequence.Remove(Request);
+					}
+				}
+			}
+		}
+		else
+		{
+			for (int32 Idx = Heap.Num() - 1; Idx >= 0; --Idx)
+			{
+				FIoFileBlockRequest* Request = Heap[Idx];
+				if (Request->ErrorCode == EIoErrorCode::Cancelled)
+				{
+					CancelledHeap.HeapPush(Request, SeqNoLess);
+					Heap.RemoveAt(Idx, EAllowShrinking::No);
+				}
+			}
+			Heap.Heapify(SeqNoLess);
+		}
+	}
+
+	FIoFileBlockRequest* Cancelled = nullptr;
+	if (CancelledHeap.IsEmpty() == false)
+	{
+		if (bDequeue)
+		{
+			CancelledHeap.HeapPop(Cancelled, SeqNoLess, EAllowShrinking::No);
+		}
+		else
+		{
+			Cancelled = CancelledHeap.HeapTop();
+		}
+	}
+
+	return Cancelled;
+}
+
+////////////////////////////////////////////////////////////////////////////////
+FIoFileBlockMemoryPool::FIoFileBlockMemoryPool(FPlatformIoDispatcherStats& InStats)
+	: Stats(InStats)
+{
+}
+
+void FIoFileBlockMemoryPool::Initialize(uint32 FileBlockSize, uint32 FilePoolSize, uint32 Alignment)
+{
+	check(FileBlockSize > 0 && FilePoolSize > 0);
+
+	const uint32 Padding		= Alignment - sizeof(FMemoryBlock);
+	const uint32 BlockCount		= FilePoolSize / FileBlockSize;
+	const uint32 MemorySize		= FileBlockSize * BlockCount;
+
+	BlockSize = FileBlockSize;
+	BlockMemory.Reset(FMemory::Malloc(MemorySize, Alignment));
+	MemoryBlocks.SetNum(BlockCount);
+
+	uint8* Memory = reinterpret_cast<uint8*>(BlockMemory.Get());
+	for (int32 Idx = 0; Idx < MemoryBlocks.Num(); ++Idx)
+	{
+		FMemoryBlock& Block = MemoryBlocks[Idx];
+		Block.Next		= FreeBlock;
+		Block.Memory	= Memory;
+		Block.Index		= Idx;
+
+		FreeBlock		= &Block;
+		Memory			+= FileBlockSize;
+
+		Stats.OnFileBlockMemoryFreed(BlockSize);
+	}
+}
+
+void* FIoFileBlockMemoryPool::Alloc(FIoBufferHandle& OutHandle)
+{
+	if (FMemoryBlock* Block = FreeBlock)
+	{
+		FreeBlock	= Block->Next;
+		Block->Next	= nullptr;
+		Block->RefCount.store(1, std::memory_order_relaxed);
+		OutHandle	= FIoBufferHandle(Block->Index);
+
+		Stats.OnFileBlockMemoryAllocated(BlockSize);
+		return Block->Memory;
+	}
+
+	OutHandle = FIoBufferHandle();
+	return nullptr;
+}
+
+bool FIoFileBlockMemoryPool::IsEmpty() const
+{
+	return FreeBlock == nullptr;
+}
+
+void FIoFileBlockMemoryPool::Free(FIoBufferHandle& Handle)
+{
+	if (Handle.IsValid() == false)
+	{
+		return;
+	}
+
+	FMemoryBlock& Block = MemoryBlocks[Handle.Value()];
+	check(Block.Next == nullptr);
+	check(Block.RefCount > 0);
+
+	if (1 == Block.RefCount.fetch_sub(1, std::memory_order_relaxed))
+	{
+		Block.Next	= FreeBlock;
+		FreeBlock	= &Block;
+
+		Stats.OnFileBlockMemoryFreed(BlockSize);
+	}
+
+	Handle = FIoBufferHandle();
+}
+
+void FIoFileBlockMemoryPool::AddRef(FIoBufferHandle Handle)
+{
+	check(Handle.IsValid());
+	FMemoryBlock& Block = MemoryBlocks[Handle.Value()];
+	Block.RefCount.fetch_add(1, std::memory_order_relaxed);
+}
+
+////////////////////////////////////////////////////////////////////////////////
+FIoChunkBlockMemoryPool::FIoChunkBlockMemoryPool()
+{
+}
+
+void FIoChunkBlockMemoryPool::Initialize(uint32 MaxBlockCount, uint32 DefaultBlockSize)
+{
+	MemoryBlocks.SetNum(MaxBlockCount);
+	for (uint32 Idx = 0; Idx < MaxBlockCount; ++Idx)
+	{
+		FMemoryBlock& Block = MemoryBlocks[Idx];
+		Block.Next			= FreeBlock;
+		Block.Memory		= FMemory::Malloc(DefaultBlockSize);
+		Block.Size			= DefaultBlockSize;
+		Block.Index			= Idx;
+		FreeBlock			= &Block;
+	}
+}
+
+void* FIoChunkBlockMemoryPool::Alloc(uint32 BlockSize, FIoBufferHandle& OutHandle)
+{
+	if (FMemoryBlock* Block = FreeBlock)
+	{
+		FreeBlock	= Block->Next;
+		Block->Next	= nullptr;
+
+		if (Block->Size < BlockSize)
+		{
+			Block->Memory	= FMemory::Realloc(Block->Memory, BlockSize);
+			Block->Size		= BlockSize;
+		}
+
+		OutHandle = FIoBufferHandle(Block->Index);
+		return Block->Memory;
+	}
+
+	OutHandle = FIoBufferHandle();
+	return nullptr;
+}
+
+void* FIoChunkBlockMemoryPool::Realloc(FIoBufferHandle Handle, uint32 BlockSize)
+{
+	check(Handle.Value() < MemoryBlocks.Num());
+	FMemoryBlock& Block = MemoryBlocks[Handle.Value()];
+
+	if (Block.Size < BlockSize)
+	{
+		Block.Memory	= FMemory::Realloc(Block.Memory, BlockSize);
+		Block.Size		= BlockSize;
+	}
+
+	return Block.Memory;
+}
+
+void FIoChunkBlockMemoryPool::Free(FIoBufferHandle& Handle)
+{
+	if (Handle.IsValid() == false)
+	{
+		return;
+	}
+
+	check(Handle.Value() < MemoryBlocks.Num());
+	FMemoryBlock& Block = MemoryBlocks[Handle.Value()];
+	Block.Next			= FreeBlock;
+	FreeBlock			= &Block;
+
+	Handle = FIoBufferHandle();
+}
+
+////////////////////////////////////////////////////////////////////////////////
+FIoFileBlockCache::FIoFileBlockCache(FPlatformIoDispatcherStats& InStats)
+	: Stats(InStats)
+{
+}
+
+FIoFileBlockCache::~FIoFileBlockCache()
+{
+}
+
+void FIoFileBlockCache::Initialize(uint64 InCacheBlockSize, uint64 InCacheSize)
+{
+	CacheBlockSize					= InCacheBlockSize;
+	const uint64 CacheBlockCount	= InCacheSize / InCacheBlockSize;
+	if (CacheBlockCount > 0)
+	{
+		const uint64 TotalCacheSize	= CacheBlockCount * InCacheBlockSize;
+		CacheMemory					= MakeUnique<uint8[]>(TotalCacheSize);
+
+		FCachedBlock* Prev	= &CacheLruHead;
+		for (uint64 CacheBlockIndex = 0; CacheBlockIndex < CacheBlockCount; ++CacheBlockIndex)
+		{
+			FCachedBlock* CachedBlock	= new FCachedBlock();
+			CachedBlock->Key			= uint64(-1);
+			CachedBlock->Buffer			= CacheMemory.Get() + CacheBlockIndex * InCacheBlockSize;
+			Prev->LruNext				= CachedBlock;
+			CachedBlock->LruPrev		= Prev;
+			Prev						= CachedBlock;
+		}
+		Prev->LruNext			= &CacheLruTail;
+		CacheLruTail.LruPrev	= Prev;
+	}
+}
+
+bool FIoFileBlockCache::Get(FIoFileBlockRequest& FileBlockRequest)
+{
+	if (CacheMemory == nullptr)
+	{
+		return false;
+	}
+
+	check(FileBlockRequest.BlockKey.IsValid());
+	check(FileBlockRequest.BufferHandle.IsValid());
+
+	FCachedBlock* CachedBlock = CachedBlocks.FindRef(FileBlockRequest.BlockKey.GetHash());
+	if (CachedBlock == nullptr)
+	{
+		Stats.OnFileBlockCacheMiss(CacheBlockSize);
+		return false;
+	}
+	
+	CachedBlock->LruPrev->LruNext = CachedBlock->LruNext;
+	CachedBlock->LruNext->LruPrev = CachedBlock->LruPrev;
+
+	CachedBlock->LruPrev = &CacheLruHead;
+	CachedBlock->LruNext = CacheLruHead.LruNext;
+
+	CachedBlock->LruPrev->LruNext = CachedBlock;
+	CachedBlock->LruNext->LruPrev = CachedBlock;
+
+	check(CachedBlock->Buffer);
+	Stats.OnFileBlockCacheHit(CacheBlockSize);
+	FMemory::Memcpy(FileBlockRequest.Buffer, CachedBlock->Buffer, CacheBlockSize);
+
+	return true;
+}
+
+void FIoFileBlockCache::Put(FIoFileBlockRequest& FileBlockRequest)
+{
+	bool bIsCacheableBlock = CacheMemory != nullptr;// && Block->BytesUsed < Block->Size;
+	if (!bIsCacheableBlock)
+	{
+		return;
+	}
+	check(FileBlockRequest.BufferHandle.IsValid());
+	FCachedBlock* BlockToReplace = CacheLruTail.LruPrev;
+	if (BlockToReplace == &CacheLruHead)
+	{
+		return;
+	}
+	check(BlockToReplace);
+	CachedBlocks.Remove(BlockToReplace->Key);
+	BlockToReplace->Key = FileBlockRequest.BlockKey.GetHash(); 
+
+	BlockToReplace->LruPrev->LruNext = BlockToReplace->LruNext;
+	BlockToReplace->LruNext->LruPrev = BlockToReplace->LruPrev;
+
+	BlockToReplace->LruPrev = &CacheLruHead;
+	BlockToReplace->LruNext = CacheLruHead.LruNext;
+
+	BlockToReplace->LruPrev->LruNext = BlockToReplace;
+	BlockToReplace->LruNext->LruPrev = BlockToReplace;
+
+	check(BlockToReplace->Buffer);
+	FMemory::Memcpy(BlockToReplace->Buffer, FileBlockRequest.Buffer, CacheBlockSize);
+	CachedBlocks.Add(BlockToReplace->Key, BlockToReplace);
+	Stats.OnFileBlockCacheStore(CacheBlockSize);
+}
+
+////////////////////////////////////////////////////////////////////////////////
+FPlatformIoDispatcherRequestMgr::FPlatformIoDispatcherRequestMgr()
+{
+}
+
+FIoPlatformReadRequest& FPlatformIoDispatcherRequestMgr::CreateScatterGatherRequest(
+	FIoFileReadRequestCompleted&& OnCompleted,
+	FIoBuffer& Dst,
+	uint64 DstSize,
+	void* UserData,
+	uint32 FileId)
+{
+	FIoPlatformReadRequest* ReadRequest = ReadRequestAllocator.Construct(
+		MoveTemp(OnCompleted),
+		Dst,
+		DstSize,
+		UserData,
+		FileId);
+
+	return *ReadRequest;
+}
+
+FIoPlatformReadRequest& FPlatformIoDispatcherRequestMgr::CreateDirectReadRequest(
+	FIoFileReadRequestCompleted&& OnCompleted,
+	FIoBuffer& Dst,
+	uint64 DstSize,
+	uint64 FileOffset,
+	void* UserData)
+{
+	FIoPlatformReadRequest* ReadRequest = ReadRequestAllocator.Construct(
+		MoveTemp(OnCompleted),
+		Dst,
+		DstSize,
+		FileOffset,
+		UserData);
+
+	return *ReadRequest;
+}
+
+bool FPlatformIoDispatcherRequestMgr::TryCancelReadRequest(FIoPlatformReadRequest& ReadRequest, bool& bAnyBlockCancelled)
+{
+	bAnyBlockCancelled			= false;
+	bool bCompleteReadRequest	= false;
+
+	if (ReadRequest.IsScatterGather())
+	{
+		if (ReadRequest.ScatterGather.RemainingBlocks.load(std::memory_order_seq_cst) == 0)
+		{
+			return bCompleteReadRequest;
+		}
+
+		bCompleteReadRequest = true;
+		for (FIoFileBlockLink* Link = &ReadRequest.ScatterGather.FileBlockLink; Link != nullptr; Link = Link->NextLink)
+		{
+			check(Link->FileBlockCount > 0);
+			for (uint32 Idx = 0; Idx < Link->FileBlockCount; ++Idx)
+			{
+				// The file block request is removed and destroyed after reading from the filesystem is completed
+				const FIoBlockKey FileBlockKey			= FIoBlockKey(ReadRequest.ScatterGather.FileId, Link->FileBlockIds[Idx]);
+				FIoFileBlockRequest* FileBlockRequest	= FileBlockLookup.FindRef(FileBlockKey);
+				if (FileBlockRequest == nullptr || FileBlockRequest->QueueStatus == FIoFileBlockRequest::EQueueStatus::Dequeued)
+				{
+					// Only cancel the read request if all file blocks can be found
+					bCompleteReadRequest = false;
+					continue;
+				}
+
+				// Cancel the file block request if all encoded block requests has been cancelled
+				bool bCancelFileBlock		= true;
+				bool bCancelEncodedBlock	= true;
+				for (FIoEncodedBlockRequest* EncodedBlockRequest : FileBlockRequest->EncodedBlockRequests)
+				{
+					for (FIoEncodedBlockRequest::FScatterTarget& ScatterTarget : EncodedBlockRequest->ScatterTargets)
+					{
+						if (ScatterTarget.Request == &ReadRequest)
+						{
+							ScatterTarget.SizeInBlock = 0;
+						}
+						else if (ScatterTarget.SizeInBlock > 0)
+						{
+							bCancelEncodedBlock = false;
+							bCancelFileBlock	= false;
+						}
+					}
+					if (bCancelEncodedBlock)
+					{
+						EncodedBlockRequest->ErrorCode = EIoErrorCode::Cancelled;
+						EncodedBlockLookup.Remove(EncodedBlockRequest->BlockKey);
+					}
+				}
+
+				if (bCancelFileBlock)
+				{
+					FileBlockRequest->ErrorCode = EIoErrorCode::Cancelled;
+					FileBlockLookup.Remove(FileBlockRequest->BlockKey);
+					bAnyBlockCancelled = true;
+				}
+			}
+		}
+	}
+
+	return bCompleteReadRequest;
+}
+
+bool FPlatformIoDispatcherRequestMgr::TryCancelAllReadRequests(FIoFileHandle FileHandle)
+{
+	// Assumes file queue lock
+	bool bAnyCancelled = false;
+	for (TTuple<FIoBlockKey, FIoFileBlockRequest*>& Kv : FileBlockLookup)
+	{
+		FIoFileBlockRequest* FileBlockRequest = Kv.Value;
+		if (FileBlockRequest->FileHandle != FileHandle || FileBlockRequest->QueueStatus == FIoFileBlockRequest::EQueueStatus::Dequeued)
+		{
+			continue;
+		}
+		for (FIoEncodedBlockRequest* EncodedBlockRequest : FileBlockRequest->EncodedBlockRequests)
+		{
+			for (FIoEncodedBlockRequest::FScatterTarget& ScatterTarget : EncodedBlockRequest->ScatterTargets)
+			{
+				ScatterTarget.SizeInBlock = 0;
+			}
+
+			EncodedBlockRequest->ErrorCode = EIoErrorCode::Cancelled;
+			EncodedBlockLookup.Remove(EncodedBlockRequest->BlockKey);
+		}
+
+		FileBlockRequest->ErrorCode = EIoErrorCode::Cancelled;
+		FileBlockLookup.Remove(FileBlockRequest->BlockKey);
+		bAnyCancelled = true;
+	}
+
+	return bAnyCancelled;
+}
+
+void FPlatformIoDispatcherRequestMgr::Destroy(FIoPlatformReadRequest& ReadRequest)
+{
+	check(ReadRequest.RefCount > 0);
+	if (--ReadRequest.RefCount == 0)
+	{
+		if (ReadRequest.IsScatterGather())
+		{
+			Destroy(ReadRequest.ScatterGather.FileBlockLink.NextLink);
+		}
+		ReadRequestAllocator.Destroy(&ReadRequest);
+	}
+}
+
+FIoFileBlockRequest& FPlatformIoDispatcherRequestMgr::GetOrCreateFileBlockRequest(FIoPlatformReadRequest& ReadRequest, FIoBlockKey BlockKey, bool& bCreated)
+{
+	check(BlockKey.IsValid());
+	if (FIoFileBlockRequest* Existing = FileBlockLookup.FindRef(BlockKey))
+	{
+		AddToLink(&ReadRequest.ScatterGather.FileBlockLink, BlockKey);
+		bCreated = false;
+		return *Existing;
+	}
+
+	FIoFileBlockRequest* Request	= FileBlockAllocator.Construct();
+	Request->BlockKey				= BlockKey;
+	bCreated						= true;
+
+	AddToLink(&ReadRequest.ScatterGather.FileBlockLink, BlockKey);
+	FileBlockLookup.Add(BlockKey, Request);
+
+	return *Request;
+}
+
+FIoFileBlockRequest& FPlatformIoDispatcherRequestMgr::CreateFileBlockRequest()
+{
+	return *FileBlockAllocator.Construct();
+}
+
+FIoFileBlockRequest* FPlatformIoDispatcherRequestMgr::GetFileBlockRequest(FIoBlockKey BlockKey)
+{
+	return FileBlockLookup.FindRef(BlockKey);
+}
+
+void FPlatformIoDispatcherRequestMgr::GetFileBlockRequests(FIoPlatformReadRequest& ReadRequest, FTempArray<FIoFileBlockRequest*>& OutRequests)
+{
+	for (FIoFileBlockLink* Link = &ReadRequest.ScatterGather.FileBlockLink; Link != nullptr; Link = Link->NextLink)
+	{
+		check(Link->FileBlockCount > 0);
+		for (uint32 Idx = 0; Idx < Link->FileBlockCount; ++Idx)
+		{
+			const FIoBlockKey FileBlockKey = FIoBlockKey(ReadRequest.ScatterGather.FileId, Link->FileBlockIds[Idx]);
+			if (FIoFileBlockRequest* FileBlockRequest = FileBlockLookup.FindRef(FileBlockKey))
+			{
+				OutRequests.Add(FileBlockRequest);
+			}
+		}
+	}
+}
+
+void FPlatformIoDispatcherRequestMgr::AddFileBlockRequest(FIoPlatformReadRequest& ReadRequest, FIoFileBlockRequest& FileBlockRequest)
+{
+	AddToLink(&ReadRequest.ScatterGather.FileBlockLink, FileBlockRequest.BlockKey);
+}
+
+FIoEncodedBlockRequest& FPlatformIoDispatcherRequestMgr::GetOrCreateEncodedBlockRequest(FIoBlockKey BlockKey, bool& bCreated)
+{
+	check(BlockKey.IsValid());
+	if (FIoEncodedBlockRequest* Existing = EncodedBlockLookup.FindRef(BlockKey))
+	{
+		bCreated = false;
+		return *Existing;
+	}
+
+	FIoEncodedBlockRequest* Request	= EncodedBlockAllocator.Construct();
+	Request->BlockKey				= BlockKey;
+	bCreated						= true;
+
+	EncodedBlockLookup.Add(BlockKey, Request);
+	return *Request;
+}
+
+void FPlatformIoDispatcherRequestMgr::Remove(FIoEncodedBlockRequest& Request)
+{
+	// Only remove non cancelled blocks
+	if (Request.ErrorCode != EIoErrorCode::Cancelled)
+	{
+		EncodedBlockLookup.Remove(Request.BlockKey);
+	}
+}
+
+void FPlatformIoDispatcherRequestMgr::Destroy(FIoFileBlockRequest& Request)
+{
+	// Only remove non cancelled blocks
+	if (Request.ErrorCode != EIoErrorCode::Cancelled)
+	{
+		FileBlockLookup.Remove(Request.BlockKey);
+	}
+	FileBlockAllocator.Destroy(&Request);
+}
+
+void FPlatformIoDispatcherRequestMgr::Destroy(FIoEncodedBlockRequest& Request)
+{
+	if (Request.FileBlockCount > 1 && Request.EncodedData != nullptr)
+	{
+		FMemory::Free(Request.EncodedData);
+	}
+	EncodedBlockAllocator.Destroy(&Request);
+}
+
+void FPlatformIoDispatcherRequestMgr::Destroy(FIoFileBlockLink* Link)
+{
+	while (Link != nullptr)
+	{
+		FIoFileBlockLink* ToDestroy = Link;
+		Link = Link->NextLink;
+		FileBlockLinkAllocator.Destroy(ToDestroy);
+	}
+}
+
+void FPlatformIoDispatcherRequestMgr::AddToLink(FIoFileBlockLink* Link, FIoBlockKey FileBlockKey)
+{
+	check(FileBlockKey.IsValid());
+	check(Link != nullptr);
+
+	for(;;)
+	{
+		for(uint32 Idx = 0; Idx < Link->FileBlockCount; ++Idx)
+		{
+			if (Link->FileBlockIds[Idx] == FileBlockKey.GetBlockId())
+			{
+				return;
+			}
+		}
+
+		if (Link->FileBlockCount < FIoFileBlockLink::MaxFileCount)
+		{
+			check(Link->NextLink == nullptr);
+			Link->FileBlockIds[Link->FileBlockCount++] = FileBlockKey.GetBlockId();
+			return;
+		}
+
+		if (Link->NextLink == nullptr)
+		{
+			Link->NextLink = FileBlockLinkAllocator.Construct();
+		}
+		Link = Link->NextLink;
+	}
+}
+
+////////////////////////////////////////////////////////////////////////////////
+FPlatformIoDispatcherBase::FPlatformIoDispatcherBase(FPlatformIoDispatcherCreateParams&& Params)
+	: FileBlockMemoryPool(Stats)
+	, FileBlockRequestQueue(Stats)
+	, FileBlockCache(Stats)
+{
+	bMultithreaded = Params.bMultithreaded;
+}
+
+FPlatformIoDispatcherBase::~FPlatformIoDispatcherBase()
+{
+	check(bStopRequested);
+	check(Thread.IsValid() == false);
+}
+
+bool FPlatformIoDispatcherBase::Init()
+{
+	return true;
+}
+
+uint32 FPlatformIoDispatcherBase::Run()
+{
+	return OnIoThreadEntry();
+}
+
+void FPlatformIoDispatcherBase::Stop()
+{
+	bStopRequested = true;
+}
+
+FIoStatus FPlatformIoDispatcherBase::Initialize()
+{
+	FIoStatus Status = OnInitialize();
+	if (Status.IsOk() && bMultithreaded)
+	{
+		Thread.Reset(FRunnableThread::Create(this, TEXT("IoService"), 0, TPri_AboveNormal));
+	}
+
+	return Status;
+}
+
+FIoFileReadRequest FPlatformIoDispatcherBase::ScatterGather(FIoScatterGatherRequestParams&& Params, FIoFileReadRequestCompleted&& OnCompleted)
+{
+	FIoFileReadRequest		RequestHandle;
+	FIoFileBlockRequestList CreatedFileBlockRequests;
+	bool					bReprioritizeReadQueue = false;
+	{
+		TUniqueLock Lock(RequestMgr);
+
+		const FIoPlatformFileInfo FileInfo = GetPlatformFileInfo(Params.FileHandle);
+		check(FileInfo.FileId > 0);
+
+		FIoPlatformReadRequest& ReadRequest = RequestMgr.CreateScatterGatherRequest(
+			MoveTemp(OnCompleted),
+			Params.Destination,
+			Params.DestinationSize,
+			Params.UserData,
+			FileInfo.FileId);
+
+		// The caller and the dispatcher has shared ownership of the read request
+		++ReadRequest.RefCount;
+		RequestHandle = FIoFileReadRequest(UPTRINT(&ReadRequest));
+
+		for (FIoScatterGatherRequestParams::FScatterParams& ScatterParams : Params.Params)
+		{
+			const FIoBlockKey			EncodedBlockKey				= FIoBlockKey(FileInfo.FileId, ScatterParams.BlockIndex);
+			bool						bEncodedBlockRequestCreated	= false;
+			FIoEncodedBlockRequest&		BlockRequest				= RequestMgr.GetOrCreateEncodedBlockRequest(EncodedBlockKey, bEncodedBlockRequestCreated);
+
+			// Scatter info from the encoded block to destination buffer
+			BlockRequest.ScatterTargets.Add(FIoEncodedBlockRequest::FScatterTarget
+			{
+				.Request		= &ReadRequest,
+				.OffsetInDst	= ScatterParams.DestinationOffset,
+				.OffsetInBlock	= IntCastChecked<uint32>(ScatterParams.ScatterOffset),
+				.SizeInBlock	= IntCastChecked<uint32>(ScatterParams.ScatterSize)
+			});
+			ReadRequest.ScatterGather.RemainingBlocks.fetch_add(1, std::memory_order_relaxed);
+
+			// The encoded/compressed block size are always aligned to AES block size typically 16B
+			const uint64 EncodedBlockFileSize	= Align(ScatterParams.BlockCompressedSize, FAES::AESBlockSize);
+			const uint64 FirstFileBlockIndex	= ScatterParams.BlockFileOffset / FileBlockSize;
+			const uint64 LastFileBlockIndex		= (ScatterParams.BlockFileOffset + EncodedBlockFileSize - 1) / FileBlockSize;
+			check(FirstFileBlockIndex < MAX_uint32);
+			check(LastFileBlockIndex < MAX_uint32);
+
+			if (bEncodedBlockRequestCreated == false)
+			{
+				for (uint64 FileBlockIndex = FirstFileBlockIndex; FileBlockIndex <= LastFileBlockIndex; ++FileBlockIndex)
+				{
+					const FIoBlockKey FileBlockKey			= FIoBlockKey(FileInfo.FileId, IntCastChecked<uint32>(FileBlockIndex));
+					FIoFileBlockRequest* FileBlockRequest	= RequestMgr.GetFileBlockRequest(FileBlockKey);
+					check(FileBlockRequest != nullptr);
+					RequestMgr.AddFileBlockRequest(ReadRequest, *FileBlockRequest);
+					if (Params.Priority > FileBlockRequest->Priority)
+					{
+						FileBlockRequest->Priority	= Params.Priority;
+						bReprioritizeReadQueue 		= true;
+					}
+				}
+
+				continue;
+			}
+
+			BlockRequest.FileOffset				= ScatterParams.BlockFileOffset;
+			BlockRequest.EncryptionKey			= ScatterParams.EncryptionKey; 
+			BlockRequest.BlockHash				= ScatterParams.BlockHash; 
+			BlockRequest.BlockCompressedSize	= ScatterParams.BlockCompressedSize;
+			BlockRequest.BlockUncompressedSize	= ScatterParams.BlockUncompresedSize; 
+			BlockRequest.CompressionMethod		= ScatterParams.CompressionMethod;
+
+			for (uint64 FileBlockIndex = FirstFileBlockIndex; FileBlockIndex <= LastFileBlockIndex; ++FileBlockIndex)
+			{
+				check(FileInfo.FileSize > 0);
+
+				bool bFileBlockRequestCreated			= false;
+				const FIoBlockKey FileBlockKey			= FIoBlockKey(FileInfo.FileId, IntCastChecked<uint32>(FileBlockIndex));
+				FIoFileBlockRequest& FileBlockRequest	= RequestMgr.GetOrCreateFileBlockRequest(ReadRequest, FileBlockKey, bFileBlockRequestCreated);
+
+				FileBlockRequest.EncodedBlockRequests.Add(&BlockRequest);
+
+				if (bFileBlockRequestCreated)
+				{
+					const uint64 FileBlockOffset	= FileBlockIndex * FileBlockSize;
+					const uint64 LocalFileBlockSize	= FMath::Min(FileBlockOffset + FileBlockSize, uint64(FileInfo.FileSize)) - FileBlockOffset;
+
+					FileBlockRequest.FileHandle		= Params.FileHandle;
+					FileBlockRequest.FileOffset		= FileBlockOffset;
+					FileBlockRequest.FileSize		= FileInfo.FileSize;
+					FileBlockRequest.Size			= LocalFileBlockSize;
+					FileBlockRequest.Priority		= Params.Priority;
+
+					CreatedFileBlockRequests.AddTail(&FileBlockRequest);
+				}
+
+				check(BlockRequest.RemainingFileBlocks < MAX_uint8);
+				check(BlockRequest.FileBlockCount < MAX_uint8);
+				BlockRequest.RemainingFileBlocks++;
+				BlockRequest.FileBlockCount++;
+			}
+		}
+	}
+	{
+		TUniqueLock Lock(FileBlockMutex);
+		if (bReprioritizeReadQueue)
+		{
+			FileBlockRequestQueue.Reprioritize();
+		}
+		FileBlockRequestQueue.Enqueue(MoveTemp(CreatedFileBlockRequests));
+	}
+	OnWakeUp();
+
+	return RequestHandle;
+}
+
+FIoFileReadRequest FPlatformIoDispatcherBase::ReadDirect(FIoDirectReadRequestParams&& Params, FIoFileReadRequestCompleted&& OnCompleted)
+{
+	return FIoFileReadRequest();
+}
+
+void FPlatformIoDispatcherBase::CancelRequest(FIoFileReadRequest Request)
+{
+	if (Request.IsValid() == false)
+	{
+		UE_LOG(LogPlatformIoDispatcher, Warning, TEXT("Trying to cancel an invalid file read request"));
+		return;
+	}
+
+	FIoPlatformReadRequest& ReadRequest = *reinterpret_cast<FIoPlatformReadRequest*>(Request.Value());
+	if (ReadRequest.ScatterGather.RemainingBlocks.load(std::memory_order_relaxed) == 0)
+	{
+		return;
+	}
+
+	EIoErrorCode Expected = EIoErrorCode::Ok;
+	if (ReadRequest.ErrorCode.compare_exchange_weak(Expected, EIoErrorCode::Cancelled) == false)
+	{
+		return;
+	}
+
+	FIoFileReadRequestCompleted OnCompleted; 
+	{
+		TUniqueLock MgrLock(RequestMgr);
+		TUniqueLock QueueLock(FileBlockMutex);
+
+		bool bAnyBlockCancelled = false;
+		if (RequestMgr.TryCancelReadRequest(ReadRequest, bAnyBlockCancelled))
+		{
+			OnCompleted = MoveTemp(ReadRequest.OnCompleted);
+		}
+		if (bAnyBlockCancelled)
+		{
+			FileBlockRequestQueue.ReprioritizeCancelled();
+		}
+	}
+
+	if (OnCompleted)
+	{
+		OnCompleted(FIoFileReadResult
+		{
+			.UserData	= ReadRequest.UserData,
+			.ErrorCode	= EIoErrorCode::Cancelled 
+		});
+	}
+}
+
+void FPlatformIoDispatcherBase::CancelAllRequests(FIoFileHandle FileHandle)
+{
+	if (FileHandle.IsValid() == false)
+	{
+		UE_LOG(LogPlatformIoDispatcher, Warning, TEXT("Trying to cancel I/O requests for an invalid file"));
+		return;
+	}
+
+	{
+		TUniqueLock MgrLock(RequestMgr);
+		TUniqueLock QueueLock(FileBlockMutex);
+		if (RequestMgr.TryCancelAllReadRequests(FileHandle))
+		{
+			FileBlockRequestQueue.ReprioritizeCancelled();
+		}
+	}
+}
+
+void FPlatformIoDispatcherBase::UpdatePriority(FIoFileReadRequest Request, EIoFileReadPriority NewPriority)
+{
+	if (Request.IsValid() == false)
+	{
+		UE_LOG(LogPlatformIoDispatcher, Warning, TEXT("Trying to update priority for an invalid file read request"));
+		return;
+	}
+
+	FIoPlatformReadRequest& ReadRequest = *reinterpret_cast<FIoPlatformReadRequest*>(Request.Value());
+	bool bReprioritize = false;
+	{
+		TUniqueLock MgrLock(RequestMgr);
+
+		FTempArray<FIoFileBlockRequest*> FileBlockRequests;
+		RequestMgr.GetFileBlockRequests(ReadRequest, FileBlockRequests);
+
+		for (FIoFileBlockRequest* FileBlockRequest : FileBlockRequests)
+		{
+			if (NewPriority > FileBlockRequest->Priority)
+			{
+				FileBlockRequest->Priority	= NewPriority;
+				bReprioritize				= true;
+			}
+		}
+	}
+	if (bReprioritize)
+	{
+		TUniqueLock QueueLock(FileBlockMutex);
+		FileBlockRequestQueue.Reprioritize();
+	}
+}
+
+void FPlatformIoDispatcherBase::DeleteRequest(FIoFileReadRequest Request)
+{
+	if (Request.IsValid())
+	{
+		FIoPlatformReadRequest* ReadRequest = reinterpret_cast<FIoPlatformReadRequest*>(Request.Value());
+		TUniqueLock Lock(RequestMgr);
+		RequestMgr.Destroy(*ReadRequest);
+	}
+}
+
+FIoStatus FPlatformIoDispatcherBase::OnInitialize()
+{
+	return FIoStatus::Ok;
+}
+
+uint32 FPlatformIoDispatcherBase::GetNextFileId()
+{
+	const uint32 FileId = NextFileId.fetch_add(1);
+	if (FileId > 0)
+	{
+		return FileId;
+	}
+	return NextFileId.fetch_add(1);
+}
+
+void FPlatformIoDispatcherBase::ProcessCompletedFileBlock(FIoFileBlockRequest& FileBlockRequest)
+{
+	TRACE_CPUPROFILER_EVENT_SCOPE(ProcessCompletedFileBlock);
+	check(IsIoErrorAnyOf(FileBlockRequest.ErrorCode, EIoErrorCode::Ok, EIoErrorCode::Cancelled, EIoErrorCode::ReadError));
+
+	const uint64								FileBlockOffset = FileBlockRequest.FileOffset;
+	FIoBufferHandle								FileBlockBufferHandle = FileBlockRequest.BufferHandle;
+	FMutableMemoryView							FileBlockBufferView(FileBlockRequest.Buffer, FileBlockRequest.Size);
+	FIoFileBlockRequest::FEncodedBlocksArray	EncodedBlockRequests;
+	{
+		TUniqueLock Lock(RequestMgr);
+		for (FIoEncodedBlockRequest* EncodedBlockRequest : FileBlockRequest.EncodedBlockRequests)
+		{
+			check(IsIoErrorAnyOf(EncodedBlockRequest->ErrorCode, EIoErrorCode::Unknown, EIoErrorCode::Cancelled, EIoErrorCode::ReadError));
+			RequestMgr.Remove(*EncodedBlockRequest);
+			if (FileBlockRequest.ErrorCode != EIoErrorCode::Ok && EncodedBlockRequest->ErrorCode == EIoErrorCode::Unknown)
+			{
+				EncodedBlockRequest->ErrorCode = FileBlockRequest.ErrorCode;
+			}
+		}
+		EncodedBlockRequests = MoveTemp(FileBlockRequest.EncodedBlockRequests);
+		RequestMgr.Destroy(FileBlockRequest);
+	}
+
+	FIoEncodedBlockRequestList	CancelledOrFailedBlocks;
+	FIoEncodedBlockRequestList	BlocksToDecode;
+	for (FIoEncodedBlockRequest* EncodedBlockRequest : EncodedBlockRequests)
+	{
+		check(EncodedBlockRequest->RemainingFileBlocks > 0);
+		if (EncodedBlockRequest->ErrorCode != EIoErrorCode::Unknown)
+		{
+			check(IsIoErrorAnyOf(EncodedBlockRequest->ErrorCode, EIoErrorCode::Cancelled, EIoErrorCode::ReadError));
+			if (--EncodedBlockRequest->RemainingFileBlocks == 0)
+			{
+				CancelledOrFailedBlocks.AddTail(EncodedBlockRequest);
+			}
+			continue;
+		}
+
+		const uint64 EncodedBlockFileSize	= Align(EncodedBlockRequest->BlockCompressedSize, FAES::AESBlockSize);
+		const int64 OffsetInFileBlock		= EncodedBlockRequest->FileOffset - FileBlockOffset;
+
+		// Is the encoded block crossing file block boundaries
+		if (EncodedBlockRequest->FileBlockCount > 1)
+		{
+			// When crossing file blocks we need to allocate a temp buffer to hold the data
+			if (EncodedBlockRequest->EncodedData == nullptr)
+			{
+				EncodedBlockRequest->EncodedData = FMemory::Malloc(EncodedBlockFileSize);
+			}
+
+			FMemoryView FileBlock = FileBlockBufferView;
+			FMutableMemoryView EncodedBlock(EncodedBlockRequest->EncodedData, EncodedBlockFileSize);
+
+			if (OffsetInFileBlock < 0)
+			{
+				// Second or third or ... 
+				EncodedBlock.RightChopInline(FMath::Abs(OffsetInFileBlock));
+				FileBlock.LeftInline(EncodedBlock.GetSize());
+			}
+			else
+			{
+				// First file block
+				FileBlock.MidInline(OffsetInFileBlock, EncodedBlock.GetSize());
+			}
+			EncodedBlock.CopyFrom(FileBlock);
+		}
+		else
+		{
+			// When the encoded block is completely wihtin the file block we use and keep a reference to the file block memory
+			check(OffsetInFileBlock >= 0);
+			EncodedBlockRequest->FileBufferHandle	= FileBlockBufferHandle; 
+			EncodedBlockRequest->EncodedData		= FileBlockBufferView.Mid(OffsetInFileBlock, EncodedBlockFileSize).GetData();
+			FileBlockMemoryPool.AddRef(EncodedBlockRequest->FileBufferHandle);
+		}
+
+		if (--EncodedBlockRequest->RemainingFileBlocks == 0)
+		{
+			BlocksToDecode.AddTail(EncodedBlockRequest);
+		}
+	}
+
+	EnqueueBlocksToDecode(MoveTemp(BlocksToDecode));
+
+	{
+		TUniqueLock Lock(FileBlockMutex);
+		FileBlockMemoryPool.Free(FileBlockBufferHandle);
+	}
+
+	if (CancelledOrFailedBlocks.IsEmpty() == false)
+	{
+		while (FIoEncodedBlockRequest* ToComplete = CancelledOrFailedBlocks.PopHead())
+		{
+			for (FIoEncodedBlockRequest::FScatterTarget& ScatterTarget : ToComplete->ScatterTargets)
+			{
+				ScatterTarget.Request->ErrorCode = ToComplete->ErrorCode;
+			}
+			CompleteEncodedBlockRequest(*ToComplete);
+		}
+	}
+}
+
+void FPlatformIoDispatcherBase::EnqueueBlocksToDecode(FIoEncodedBlockRequestList BlocksToDecode)
+{
+	if (bMultithreaded == false)
+	{
+		while (FIoEncodedBlockRequest* BlockRequest = BlocksToDecode.PopHead())
+		{
+			// Allocate memory
+			check(BlockRequest->ErrorCode == EIoErrorCode::Unknown);
+			check(BlockRequest->BufferHandle.IsValid() == false);
+			BlockRequest->DecodedData = ChunkBlockMemoryPool.Alloc(BlockRequest->BlockUncompressedSize, BlockRequest->BufferHandle);
+
+			// Decode the block
+			Stats.OnDecodeRequestEnqueued(*BlockRequest);
+			const FIoChunkBlockDecodeResult Result = FIoChunkEncoding::DecodeBlock(
+				FIoChunkBlockDecodingParams 
+				{
+					.EncryptionKey		= BlockRequest->EncryptionKey,
+					.BlockHash			= BlockRequest->BlockHash,
+					.CompressionFormat	= BlockRequest->CompressionMethod,
+				},
+				FMutableMemoryView(BlockRequest->EncodedData, BlockRequest->BlockCompressedSize),
+				FMutableMemoryView(BlockRequest->DecodedData, BlockRequest->BlockUncompressedSize));
+
+			BlockRequest->DecodedData	= Result.DecodedBlock.GetData();
+			BlockRequest->ErrorCode		= Result.ErrorCode;
+			Stats.OnDecodeRequestCompleted(*BlockRequest);
+
+			// Copy the data from the block to the destination buffer(s) 
+			ScatterDecodedBlock(*BlockRequest);
+			ChunkBlockMemoryPool.Free(BlockRequest->BufferHandle);
+			FileBlockMemoryPool.Free(BlockRequest->FileBufferHandle);
+			BlockRequest->DecodedData = nullptr;
+
+			// Trigger completion callbacks and destroy the request
+			CompleteEncodedBlockRequest(*BlockRequest);
+		}
+	}
+	else
+	{
+		FIoEncodedBlockRequestList BlocksToEnqueue;
+		{
+			TUniqueLock Lock(FileBlockMutex);
+			while (FIoEncodedBlockRequest* BlockRequest = BlocksToDecode.PopHead())
+			{
+				check(BlockRequest->ErrorCode == EIoErrorCode::Unknown);
+				check(BlockRequest->BufferHandle.IsValid() == false);
+				BlockRequest->DecodedData = ChunkBlockMemoryPool.Alloc(BlockRequest->BlockUncompressedSize, BlockRequest->BufferHandle);
+				if (BlockRequest->BufferHandle.IsValid())
+				{
+					check(BlockRequest->DecodedData != nullptr);
+					BlocksToEnqueue.AddTail(BlockRequest);
+				}
+				else
+				{
+					BlocksReadyForDecoding.AddTail(BlockRequest);
+				}
+				Stats.OnDecodeRequestEnqueued(*BlockRequest);
+			}
+		}
+
+		while (FIoEncodedBlockRequest* BlockRequest = BlocksToEnqueue.PopHead())
+		{
+			EnqueueBlockToDecode(*BlockRequest);
+		}
+	}
+}
+
+void FPlatformIoDispatcherBase::EnqueueBlockToDecode(FIoEncodedBlockRequest& EncodedBlockRequest)
+{
+	// Cancelled or failed requests should never reach blocks to decode
+	check(EncodedBlockRequest.ErrorCode == EIoErrorCode::Unknown);
+	check(EncodedBlockRequest.EncodedData != nullptr);
+	check(EncodedBlockRequest.DecodedData != nullptr && EncodedBlockRequest.BufferHandle.IsValid());
+
+	auto OnBlockDecoded = [this, &EncodedBlockRequest](FIoChunkBlockDecodeResult&& DecodeResult, FIoChunkBlockDecodeRequest& Next)
+	{
+		ProcessDecodedBlock(EncodedBlockRequest, MoveTemp(DecodeResult), Next);
+	};
+	FIoChunkBlockDecoder::Get().Enqueue(
+		FIoChunkBlockDecodeRequest
+		{
+			.Params = FIoChunkBlockDecodingParams 
+			{
+				.EncryptionKey		= EncodedBlockRequest.EncryptionKey,
+				.BlockHash			= EncodedBlockRequest.BlockHash,
+				.CompressionFormat	= EncodedBlockRequest.CompressionMethod,
+			},
+			.EncodedBlock			= FMutableMemoryView(EncodedBlockRequest.EncodedData, EncodedBlockRequest.BlockCompressedSize),
+			.DecodedBlock			= FMutableMemoryView(EncodedBlockRequest.DecodedData, EncodedBlockRequest.BlockUncompressedSize),
+			.OnDecoded				= MoveTemp(OnBlockDecoded)
+		});
+}
+
+void FPlatformIoDispatcherBase::ProcessDecodedBlock(FIoEncodedBlockRequest& EncodedBlockRequest, FIoChunkBlockDecodeResult&& Result, FIoChunkBlockDecodeRequest& NextDecodeRequest)
+{
+	check(Result.ErrorCode == EIoErrorCode::Ok || Result.ErrorCode == EIoErrorCode::SignatureError || Result.ErrorCode == EIoErrorCode::CompressionError);
+	EncodedBlockRequest.DecodedData	= Result.DecodedBlock.GetData();
+	EncodedBlockRequest.ErrorCode	= Result.ErrorCode;
+	check(EncodedBlockRequest.ErrorCode != EIoErrorCode::Unknown);
+
+	ScatterDecodedBlock(EncodedBlockRequest); 
+
+	// Release memory so we can start the next file read requests
+	FIoEncodedBlockRequest* NextBlockToDecode = nullptr;
+	{
+		TUniqueLock Lock(FileBlockMutex);
+		FileBlockMemoryPool.Free(EncodedBlockRequest.FileBufferHandle);
+
+		if (NextBlockToDecode = BlocksReadyForDecoding.PopHead(); NextBlockToDecode != nullptr)
+		{
+			NextBlockToDecode->DecodedData		= ChunkBlockMemoryPool.Realloc(EncodedBlockRequest.BufferHandle, NextBlockToDecode->BlockUncompressedSize);
+			NextBlockToDecode->BufferHandle		= EncodedBlockRequest.BufferHandle;
+			EncodedBlockRequest.BufferHandle	= FIoBufferHandle();
+			
+		}
+		else
+		{
+			ChunkBlockMemoryPool.Free(EncodedBlockRequest.BufferHandle);
+		}
+		EncodedBlockRequest.DecodedData = nullptr;
+		Stats.OnDecodeRequestCompleted(EncodedBlockRequest);
+	}
+
+	// Wakeup dispatcher to start the next file read request
+	OnWakeUp();
+
+	CompleteEncodedBlockRequest(EncodedBlockRequest);
+
+	// Instead of recursively calling the chunk decoder we enqueue more blocks by setting the next request parameter 
+	if (NextBlockToDecode != nullptr)
+	{
+		auto OnBlockDecoded = [this, NextBlockToDecode](FIoChunkBlockDecodeResult&& DecodeResult, FIoChunkBlockDecodeRequest& Next)
+		{
+			ProcessDecodedBlock(*NextBlockToDecode, MoveTemp(DecodeResult), Next);
+		};
+		NextDecodeRequest = FIoChunkBlockDecodeRequest
+		{
+			.Params = FIoChunkBlockDecodingParams
+			{
+				.EncryptionKey		= NextBlockToDecode->EncryptionKey,
+				.BlockHash			= NextBlockToDecode->BlockHash,
+				.CompressionFormat	= NextBlockToDecode->CompressionMethod
+			},
+			.EncodedBlock			= FMutableMemoryView(NextBlockToDecode->EncodedData, NextBlockToDecode->BlockCompressedSize),
+			.DecodedBlock			= FMutableMemoryView(NextBlockToDecode->DecodedData, NextBlockToDecode->BlockUncompressedSize),
+			.OnDecoded				= MoveTemp(OnBlockDecoded) 
+		};
+	}
+}
+
+void FPlatformIoDispatcherBase::ScatterDecodedBlock(FIoEncodedBlockRequest& EncodedBlockRequest)
+{
+	LLM_SCOPE_BYNAME(TEXT("FileSystem/FileIoStore"));
+	TRACE_CPUPROFILER_EVENT_SCOPE(IoDispatcherScatter);
+	check(EncodedBlockRequest.ErrorCode != EIoErrorCode::Unknown);
+
+	{
+		// TODO: Better to do this when dequeued to avoid mutex?
+		TUniqueLock Lock(ScatterMutex);
+		for (FIoEncodedBlockRequest::FScatterTarget& ScatterTarget : EncodedBlockRequest.ScatterTargets)
+		{
+			FIoPlatformReadRequest& Request = *ScatterTarget.Request;
+			const EIoErrorCode ErrorCode	= Request.ErrorCode.load(std::memory_order_relaxed);
+
+			Stats.OnBytesScattered(ScatterTarget.SizeInBlock);
+			if (EncodedBlockRequest.ErrorCode != EIoErrorCode::Ok)
+			{
+				Request.ErrorCode.store(EncodedBlockRequest.ErrorCode, std::memory_order_relaxed);
+				Request.FailedBlockId = EncodedBlockRequest.BlockKey.GetBlockId();
+			}
+			else
+			{
+				EIoErrorCode Expected = EIoErrorCode::Ok;
+				Request.ErrorCode.compare_exchange_weak(Expected, EncodedBlockRequest.ErrorCode);
+			}
+	
+			// If the size in block is zero the request was cancelled
+			if (ScatterTarget.SizeInBlock > 0)
+			{
+				check(Request.IsScatterGather());
+				if (Request.Dst.GetSize() == 0)
+				{
+					Request.Dst = FIoBuffer(ScatterTarget.Request->DstSize);
+				}
+			}
+		}
+	}
+
+	FMemoryView DecodedBlock(EncodedBlockRequest.DecodedData, EncodedBlockRequest.BlockUncompressedSize);
+	for (FIoEncodedBlockRequest::FScatterTarget& ScatterTarget : EncodedBlockRequest.ScatterTargets)
+	{
+		FIoPlatformReadRequest& Request = *ScatterTarget.Request;
+		check(Request.IsScatterGather());
+		check(Request.ScatterGather.RemainingBlocks > 0);
+
+		const bool bScatterBlock = ScatterTarget.SizeInBlock > 0 && EncodedBlockRequest.ErrorCode == EIoErrorCode::Ok;
+		if (bScatterBlock)
+		{
+			check(ScatterTarget.Request->Dst.GetSize() > 0);
+			FMutableMemoryView	Dst = Request.Dst.GetMutableView().RightChop(ScatterTarget.OffsetInDst);
+			FMemoryView			Src = DecodedBlock.Mid(ScatterTarget.OffsetInBlock, ScatterTarget.SizeInBlock);
+			Dst.CopyFrom(Src);
+		}
+	}
+}
+
+void FPlatformIoDispatcherBase::CompleteEncodedBlockRequest(FIoEncodedBlockRequest& EncodedBlockRequest)
+{
+	FIoPlatformReadRequestList CompletedReadRequests;
+	for (FIoEncodedBlockRequest::FScatterTarget& ScatterTarget : EncodedBlockRequest.ScatterTargets)
+	{
+		FIoPlatformReadRequest& Request = *ScatterTarget.Request;
+		check(Request.IsScatterGather());
+		check(Request.ScatterGather.RemainingBlocks > 0);
+		check(Request.ErrorCode != EIoErrorCode::Unknown);
+
+		if (1 == Request.ScatterGather.RemainingBlocks.fetch_sub(1, std::memory_order_seq_cst))
+		{
+			FIoFileReadRequestCompleted OnCompleted = MoveTemp(Request.OnCompleted);
+			if (OnCompleted)
+			{
+				OnCompleted(FIoFileReadResult
+				{
+					.UserData		= Request.UserData,
+					.FailedBlockId	= Request.FailedBlockId,
+					.ErrorCode		= Request.ErrorCode.load(std::memory_order_relaxed),
+				});
+			}
+			CompletedReadRequests.AddTail(&Request);
+		}
+	}
+
+	// Destroy the block request and completed read request(s)
+	{
+		TUniqueLock Lock(RequestMgr);
+		while (FIoPlatformReadRequest* ReadRequest = CompletedReadRequests.PopHead())
+		{
+			RequestMgr.Destroy(*ReadRequest);
+		}
+		RequestMgr.Destroy(EncodedBlockRequest);
+	}
+}
+
+} // namespace UE
diff --git a/Engine/Source/Runtime/Core/Internal/IO/PlatformIoDispatcherBase.h b/Engine/Source/Runtime/Core/Internal/IO/PlatformIoDispatcherBase.h
new file mode 100644
--- /dev/null
+++ b/Engine/Source/Runtime/Core/Internal/IO/PlatformIoDispatcherBase.h
@@ -0,0 +1,524 @@
+// Copyright Epic Games, Inc. All Rights Reserved.
+
+#pragma once
+
+#include "Async/Mutex.h"
+#include "HAL/PlatformTime.h"
+#include "HAL/Runnable.h"
+#include "HAL/RunnableThread.h"
+#include "IO/IoAllocators.h"
+#include "IO/IoBuffer.h"
+#include "IO/IoChunkEncoding.h"
+#include "IO/IoContainers.h"
+#include "IO/PlatformIoDispatcher.h"
+#include "IO/PlatformIoDispatcherStats.h"
+#include "IO/IoStatus.h"
+#include "Math/NumericLimits.h"
+#include "Memory/MemoryView.h"
+#include "Templates/UniquePtr.h"
+
+#include <atomic>
+
+namespace UE
+{
+
+////////////////////////////////////////////////////////////////////////////////
+template<typename T, SIZE_T Size = 8>
+using FTempArray = TArray<T, TInlineAllocator<Size>>;
+
+////////////////////////////////////////////////////////////////////////////////
+struct FIoBufferHandle
+{
+	FIoBufferHandle() = default;
+
+	explicit FIoBufferHandle(int32 InHandle)
+		: Handle(InHandle)
+	{
+		check(InHandle != INDEX_NONE);
+	}
+
+	int32	Value() const { return Handle; }
+	bool	IsValid() const { return Handle != INDEX_NONE; }
+
+private:
+	int32 Handle = INDEX_NONE;
+};
+
+////////////////////////////////////////////////////////////////////////////////
+struct FIoBlockKey
+{
+	FIoBlockKey()
+		: Hash(0)
+	{ }
+
+	explicit FIoBlockKey(uint32 InFileId, uint32 InBlockId)
+		: FileId(InFileId)
+		, BlockId(InBlockId)
+	{ }
+
+	uint64	GetHash() const { return Hash; }
+	uint32	GetBlockId() const { return BlockId; }
+	bool	IsValid() const { return Hash != 0; }
+
+	friend bool operator==(const FIoBlockKey& A, const FIoBlockKey& B)
+	{
+		return A.Hash == B.Hash;
+	}
+
+	friend uint32 GetTypeHash(const FIoBlockKey& Key)
+	{
+		return GetTypeHash(Key.Hash);
+	}
+
+private:
+	union
+	{
+		struct
+		{
+			uint32 FileId;
+			uint32 BlockId;
+		};
+		uint64 Hash;
+	};
+};
+static_assert(sizeof(FIoBlockKey) == sizeof(uint64));
+
+////////////////////////////////////////////////////////////////////////////////
+struct FIoEncodedBlockRequest
+	: public TIntrusiveListElement<FIoEncodedBlockRequest>
+{
+	struct FScatterTarget
+	{
+		struct FIoPlatformReadRequest* Request = nullptr;
+		uint64 OffsetInDst		= 0;
+		uint32 OffsetInBlock	= 0;
+		uint32 SizeInBlock		= 0;
+	};
+	static_assert(sizeof(FScatterTarget) == 24);
+
+	using FScatterTargets = TArray<FScatterTarget, TInlineAllocator<2>>;
+	static_assert(sizeof(FScatterTargets) == 64);
+
+	FScatterTargets					ScatterTargets;
+	FMemoryView						EncryptionKey;
+	FMemoryView						BlockHash;
+	FIoBlockKey						BlockKey;
+	FIoEncodedBlockRequest*			Next = nullptr;
+	void*							EncodedData = nullptr;
+	void*							DecodedData = nullptr;
+	uint64							FileOffset;
+	FIoBufferHandle					FileBufferHandle;
+	FIoBufferHandle					BufferHandle;
+	FName							CompressionMethod = NAME_None;
+	uint32							BlockCompressedSize = 0;
+	uint32							BlockUncompressedSize = 0;
+	EIoErrorCode					ErrorCode = EIoErrorCode::Unknown;
+	uint8							RemainingFileBlocks = 0;
+	uint8							FileBlockCount = 0;
+};
+
+using FIoEncodedBlockRequestAllocator	= TSingleThreadedSlabAllocator<FIoEncodedBlockRequest>;
+using FIoEncodedBlockRequestList		= TIntrusiveList<FIoEncodedBlockRequest>;
+
+////////////////////////////////////////////////////////////////////////////////
+struct FIoFileBlockLink
+{
+	enum { MaxFileCount = 5 };
+	uint32				FileBlockIds[MaxFileCount] = {};
+	uint32				FileBlockCount = 0;
+	FIoFileBlockLink*	NextLink = nullptr;
+};
+static_assert(sizeof(FIoFileBlockLink) == 32);
+using FIoFileBlockLinkAllocator = TSingleThreadedSlabAllocator<FIoFileBlockLink>;
+
+////////////////////////////////////////////////////////////////////////////////
+struct FIoPlatformReadRequest
+	: public TIntrusiveListElement<FIoPlatformReadRequest>
+{
+	enum class EType : uint8
+	{
+		ScatterGather,
+		DirectRead
+	};
+
+	struct FScatterGather
+	{
+		FIoFileBlockLink		FileBlockLink;
+		uint32					FileId;
+		std::atomic_uint16_t	RemainingBlocks{0};
+	};
+
+	struct FDirectRead
+	{
+		uint64					FileOffset = MAX_uint64;
+	};
+
+	FIoPlatformReadRequest(
+		FIoFileReadRequestCompleted&& InOnCompleted,
+		FIoBuffer& InDst,
+		uint64 InDstSize,
+		void* InUserData,
+		uint32 InFileId)
+			: OnCompleted(MoveTemp(InOnCompleted))
+			, ScatterGather()
+			, Dst(InDst)
+			, DstSize(InDstSize)
+			, UserData(InUserData)
+			, Type(EType::ScatterGather)
+	{
+		ScatterGather.FileId = InFileId;
+	}
+
+	FIoPlatformReadRequest(
+		FIoFileReadRequestCompleted&& InOnCompleted,
+		FIoBuffer& InDst,
+		uint64 InDstSize,
+		uint64 InFileOffset,
+		void* InUserData)
+			: OnCompleted(MoveTemp(InOnCompleted))
+			, DirectRead()
+			, Dst(InDst)
+			, DstSize(InDstSize)
+			, UserData(InUserData)
+			, Type(EType::DirectRead)
+	{
+		DirectRead.FileOffset = InFileOffset;
+	}
+
+	bool IsScatterGather() const { return Type == EType::ScatterGather; }
+	bool IsDirectRead() const { return Type == EType::DirectRead; }
+
+	FIoFileReadRequestCompleted	OnCompleted;
+	union
+	{
+		FScatterGather			ScatterGather;
+		FDirectRead				DirectRead;
+	};
+	FIoPlatformReadRequest*		Next = nullptr;
+	FIoBuffer&					Dst;	
+	uint64						DstSize = 0;
+	void*						UserData = nullptr;
+	uint32						RefCount = 1;
+	uint32						FailedBlockId = MAX_uint32;
+	std::atomic<EIoErrorCode>	ErrorCode{EIoErrorCode::Ok};
+	EType						Type = EType::ScatterGather;
+};
+
+using FIoPlatformReadRequestAllocator	= TSingleThreadedSlabAllocator<FIoPlatformReadRequest>;
+using FIoPlatformReadRequestList		= TIntrusiveList<FIoPlatformReadRequest>;
+
+////////////////////////////////////////////////////////////////////////////////
+struct FIoFileBlockRequest
+	: public TIntrusiveListElement<FIoFileBlockRequest>
+{
+	enum class EQueueStatus : uint8
+	{
+		None,
+		Enqueued,
+		Dequeued
+	};
+
+	using							FEncodedBlocksArray = TArray<FIoEncodedBlockRequest*, TInlineAllocator<8>>;
+
+	static uint32					NextSeqNo;
+
+	FEncodedBlocksArray				EncodedBlockRequests;
+	FIoFileBlockRequest*			Next = nullptr;
+	FIoPlatformReadRequest*			DirectReadRequest = nullptr;
+	void*							Buffer = nullptr;
+	FIoFileHandle					FileHandle;
+	FIoBlockKey						BlockKey;
+	uint64							FileOffset = 0;
+	int64							FileSize = 0;
+	uint64							Size = 0;
+	uint64							TimeCreated = FPlatformTime::Cycles64();
+	uint32							SeqNo = NextSeqNo++;
+	FIoBufferHandle					BufferHandle;
+	EIoErrorCode					ErrorCode = EIoErrorCode::Unknown;
+	EIoFileReadPriority				Priority = EIoFileReadPriority::Low;
+	EQueueStatus					QueueStatus = EQueueStatus::None;
+};
+
+using FIoFileBlockRequestAllocator	= TSingleThreadedSlabAllocator<FIoFileBlockRequest>;
+using FIoFileBlockRequestList		= TIntrusiveList<FIoFileBlockRequest>;
+
+////////////////////////////////////////////////////////////////////////////////
+struct FIoPlatformFileInfo
+{
+	int64	FileSize = -1;
+	uint32	FileId = 0;
+	uint32	CompressionBlockSize = 0;
+};
+
+////////////////////////////////////////////////////////////////////////////////
+class FIoFileBlockRequestQueue
+{
+	struct FFileOffset
+	{
+		FIoFileHandle		FileHandle;
+		uint64				FileOffset = MAX_uint64;
+		EIoFileReadPriority	Priority = EIoFileReadPriority::Min;
+	};
+
+	struct FPrioQueue
+	{
+		TArray<FIoFileBlockRequest*>	ByOffset;
+		FIoFileBlockRequestList			BySequence;
+		int32							PeekIndex = INDEX_NONE;
+	};
+
+	static FFileOffset ToFileOffset(const FIoFileBlockRequest* Request)
+	{
+		return FFileOffset
+		{
+			.FileHandle = Request->FileHandle,
+			.FileOffset = Request->FileOffset,
+			.Priority	= Request->Priority
+		};
+	}
+
+	static bool FileOffsetLess(const FFileOffset& LHS, const FFileOffset& RHS)
+	{
+		if (LHS.FileHandle.Value() == RHS.FileHandle.Value())
+		{
+			return LHS.FileOffset < RHS.FileOffset;
+		}
+		return LHS.FileHandle.Value() < RHS.FileHandle.Value();
+	}
+
+	static bool SeqNoLess(const FIoFileBlockRequest& LHS, const FIoFileBlockRequest& RHS)
+	{
+		if (LHS.Priority == RHS.Priority)
+		{
+			return LHS.SeqNo < RHS.SeqNo;
+		}
+		return static_cast<uint8>(LHS.Priority) > static_cast<uint8>(RHS.Priority);
+	}
+
+public:
+							FIoFileBlockRequestQueue(FPlatformIoDispatcherStats& InStats);
+	void					Enqueue(FIoFileBlockRequestList&& Requests);
+	FIoFileBlockRequest*	Dequeue();
+	FIoFileBlockRequest*	Peek();
+	void					Reprioritize();
+	void					ReprioritizeCancelled() { bReprioritizeCancelled = true; }
+	void					SortByOffset(bool bValue) { bSortByOffset = bValue; }
+
+private:
+	FIoFileBlockRequest*	GetByOffset(bool bDequeue);
+	FIoFileBlockRequest*	GetByOffset(FPrioQueue& PrioQueue, EIoFileReadPriority QueuePriority, bool bDequeue);
+	FIoFileBlockRequest*	GetCancelled(bool bDequeue);
+
+	FPlatformIoDispatcherStats&		Stats;
+	TArray<FIoFileBlockRequest*>	Heap;
+	TArray<FIoFileBlockRequest*>	CancelledHeap;
+	FPrioQueue						PrioQueues[static_cast<int32>(EIoFileReadPriority::Count)];
+	FFileOffset						LastFileOffset;
+	bool							bSortByOffset = false;
+	bool							bReprioritizeCancelled = false;
+};
+
+////////////////////////////////////////////////////////////////////////////////
+class FIoFileBlockMemoryPool
+{
+public:
+					FIoFileBlockMemoryPool(FPlatformIoDispatcherStats& InStats);
+	void			Initialize(uint32 FileBlockSize, uint32 PoolSize, uint32 Alignment = 4096);
+	void*			Alloc(FIoBufferHandle& OutHandle);
+	void			Free(FIoBufferHandle& Handle);
+	void			AddRef(FIoBufferHandle Handle);
+	bool			IsEmpty() const;
+
+private:
+
+	struct FDeleteByFree
+	{
+		void operator()(void* Ptr) const
+		{
+			FMemory::Free(Ptr);
+		}
+	};
+
+	struct FMemoryBlock
+	{
+		FMemoryBlock*		Next		= nullptr;
+		void*				Memory		= nullptr;
+		int32				Index		= INDEX_NONE;
+		std::atomic_int32_t	RefCount{0};
+	};
+
+	FPlatformIoDispatcherStats&			Stats;
+	TUniquePtr<void, FDeleteByFree>		BlockMemory;
+	TArray<FMemoryBlock>				MemoryBlocks;
+	FMemoryBlock*						FreeBlock = nullptr;
+	uint32								BlockSize;
+};
+
+////////////////////////////////////////////////////////////////////////////////
+class FIoChunkBlockMemoryPool
+{
+public:
+						FIoChunkBlockMemoryPool();
+	void				Initialize(uint32 MaxBlockCount, uint32 DefaultBlockSize);
+	void*				Alloc(uint32 BlockSize, FIoBufferHandle& OutHandle);
+	void*				Realloc(FIoBufferHandle Handle, uint32 BlockSize);
+	void				Free(FIoBufferHandle& Handle);
+	bool				IsEmpty() const { return FreeBlock == nullptr; }
+
+private:
+	struct FMemoryBlock
+	{
+		FMemoryBlock*	Next	= nullptr;
+		void*			Memory	= nullptr;
+		uint32			Size	= 0;
+		int32			Index	= INDEX_NONE; 
+	};
+
+	TArray<FMemoryBlock>	MemoryBlocks;
+	FMemoryBlock*			FreeBlock = nullptr;
+};
+
+////////////////////////////////////////////////////////////////////////////////
+class FIoFileBlockCache
+{
+public:
+			FIoFileBlockCache(FPlatformIoDispatcherStats& Stats);
+			~FIoFileBlockCache();
+
+	void	Initialize(uint64 CacheBlockSize, uint64 CacheSize);
+	bool 	Get(FIoFileBlockRequest& FileBlockRequest);
+	void 	Put(FIoFileBlockRequest& FileBlockRequest);
+
+private:
+	struct FCachedBlock
+	{
+		FCachedBlock*	LruPrev = nullptr;
+		FCachedBlock* 	LruNext = nullptr;
+		uint64			Key		= 0;
+		uint8*			Buffer	= nullptr;
+	};
+
+	FPlatformIoDispatcherStats&		Stats;
+	TUniquePtr<uint8[]>				CacheMemory;
+	TMap<uint64, FCachedBlock*>		CachedBlocks;
+	FCachedBlock					CacheLruHead;
+	FCachedBlock					CacheLruTail;
+	uint64							CacheBlockSize = 0;
+};
+
+////////////////////////////////////////////////////////////////////////////////
+class FPlatformIoDispatcherRequestMgr
+{
+public:
+										FPlatformIoDispatcherRequestMgr();
+	FIoPlatformReadRequest&				CreateScatterGatherRequest(
+											FIoFileReadRequestCompleted&& OnCompleted,
+											FIoBuffer& Dst,
+											uint64 DstSize,
+											void* UserData,
+											uint32 FileId);
+	FIoPlatformReadRequest&				CreateDirectReadRequest(
+											FIoFileReadRequestCompleted&& OnCompleted,
+											FIoBuffer& Dst,
+											uint64 DstSize,
+											uint64 FileOffset,
+											void* UserData);
+	bool								TryCancelReadRequest(FIoPlatformReadRequest& ReadRequest, bool& bAnyBlockCancelled);
+	bool								TryCancelAllReadRequests(FIoFileHandle FileHandle);
+	void								Destroy(FIoPlatformReadRequest& ReadRequest);
+	FIoFileBlockRequest&				GetOrCreateFileBlockRequest(FIoPlatformReadRequest& ReadRequest, FIoBlockKey BlockKey, bool& bCreated);
+	FIoFileBlockRequest&				CreateFileBlockRequest();
+	FIoFileBlockRequest*				GetFileBlockRequest(FIoBlockKey BlockKey);
+	void								GetFileBlockRequests(FIoPlatformReadRequest& ReadRequest, FTempArray<FIoFileBlockRequest*>& OutRequests);
+	void								AddFileBlockRequest(FIoPlatformReadRequest& ReadRequest, FIoFileBlockRequest& FileBlockRequest);
+	FIoEncodedBlockRequest&				GetOrCreateEncodedBlockRequest(FIoBlockKey BlockKey, bool& bCreated);
+	void								Remove(FIoEncodedBlockRequest& Request);
+	void								Destroy(FIoFileBlockRequest& Request);
+	void								Destroy(FIoEncodedBlockRequest& Request);
+	void								Destroy(FIoFileBlockLink* Link);
+
+	inline void							Lock() { Mutex.Lock(); }
+	inline void							Unlock() { Mutex.Unlock(); }
+
+private:
+	void								AddToLink(FIoFileBlockLink* Link, FIoBlockKey FileBlockKey);
+	using FFileBlockLookup				= TMap<FIoBlockKey, FIoFileBlockRequest*>;
+	using FEncodedBlockLookup			= TMap<FIoBlockKey, FIoEncodedBlockRequest*>;
+
+	FIoPlatformReadRequestAllocator		ReadRequestAllocator;
+	FIoFileBlockRequestAllocator		FileBlockAllocator;
+	FFileBlockLookup					FileBlockLookup;
+	FIoFileBlockLinkAllocator 			FileBlockLinkAllocator;
+	FIoEncodedBlockRequestAllocator		EncodedBlockAllocator;
+	FEncodedBlockLookup					EncodedBlockLookup;
+	UE::FMutex							Mutex;
+};
+
+////////////////////////////////////////////////////////////////////////////////
+/// This is a base class for platform specific I/O dispatcher implementations.
+/// There are two kinds of I/O requests when reading data from disk. Direct-read requests
+/// for platforms with hardware decompression and scatter-gather requests. Scatter-gather
+/// requests are issued by specifying information about how the blocks on disk are encoded, i.e.
+/// compressed, encrypted and or signed and forms a set of encoded block request. Encoded
+/// block requests can be shared between user requests, i.e. if two user requests are reading
+/// the same encoded block, the block is decoded once and then scattered to the destination/targets buffers.
+/// The encoded block requests are devided into a set of larger file block requests. The
+/// size of the file blocks are configured with GIoDispatcherBufferSizeKB (default=256KiB).
+/// The maximum number of concurrent file block requests are constrained by the available
+/// memory in the file block memory pool configured with GIoDispatcherBufferMemoryMB (default=8MiB).
+class FPlatformIoDispatcherBase
+	: public FRunnable
+	, public IPlatformIoDispatcher
+{
+public:
+	virtual								~FPlatformIoDispatcherBase();
+
+	// FRunnable
+	virtual bool						Init() override final;
+	virtual uint32						Run() override final;
+	virtual void						Stop() override final;
+
+	// IIoPlatformFile
+	virtual FIoStatus					Initialize() override final;
+	virtual FIoFileReadRequest			ScatterGather(FIoScatterGatherRequestParams&& Params, FIoFileReadRequestCompleted&& OnCompleted) override;
+	virtual FIoFileReadRequest			ReadDirect(FIoDirectReadRequestParams&& Params, FIoFileReadRequestCompleted&& OnCompleted) override;
+	virtual bool						Tick() override { ensure(!bMultithreaded); return false; }
+	virtual void						UpdatePriority(FIoFileReadRequest Request, EIoFileReadPriority NewPriority) override;
+	virtual void						CancelRequest(FIoFileReadRequest Request) override;
+	virtual void						CancelAllRequests(FIoFileHandle FileHandle) override;
+	virtual void						DeleteRequest(FIoFileReadRequest Request) override;
+
+protected:
+										FPlatformIoDispatcherBase(FPlatformIoDispatcherCreateParams&& Params);
+	virtual FIoStatus					OnInitialize();
+	virtual uint32						OnIoThreadEntry() = 0;
+	virtual void						OnWakeUp() = 0;
+	virtual FIoPlatformFileInfo			GetPlatformFileInfo(FIoFileHandle FileHandle) = 0;
+
+	uint32								GetNextFileId();
+	void								ProcessCompletedFileBlock(FIoFileBlockRequest& FileBlockRequest);
+	void								EnqueueBlocksToDecode(FIoEncodedBlockRequestList BlocksToDecode);
+	void								EnqueueBlockToDecode(FIoEncodedBlockRequest& EncodedBlockRequest);
+	void								ProcessDecodedBlock(FIoEncodedBlockRequest& EncodedBlockRequest, FIoChunkBlockDecodeResult&& DecodeResult, FIoChunkBlockDecodeRequest& NextDecodeRequest);
+	void								ScatterDecodedBlock(FIoEncodedBlockRequest& EncodedBlockRequest);
+	void								CompleteEncodedBlockRequest(FIoEncodedBlockRequest& EncodedBlockRequest);
+
+	FPlatformIoDispatcherStats			Stats;
+	TUniquePtr<FRunnableThread>			Thread;
+	FPlatformIoDispatcherRequestMgr		RequestMgr;
+	FIoChunkBlockMemoryPool				ChunkBlockMemoryPool;
+	FIoFileBlockMemoryPool				FileBlockMemoryPool;
+	FIoFileBlockRequestQueue			FileBlockRequestQueue;
+	FIoFileBlockCache					FileBlockCache;
+	FIoEncodedBlockRequestList			BlocksReadyForDecoding;
+	UE::FMutex							FileBlockMutex;
+	UE::FMutex							ScatterMutex;
+	uint32								FileBlockSize = 256 << 10;
+	bool								bMultithreaded = false;
+	std::atomic_bool					bStopRequested{false};
+
+private:
+	std::atomic_uint32_t				NextFileId{1};
+};
+
+} // namespace UE 
diff --git a/Engine/Source/Runtime/Core/Internal/IO/PlatformIoDispatcherStats.cpp b/Engine/Source/Runtime/Core/Internal/IO/PlatformIoDispatcherStats.cpp
new file mode 100644
--- /dev/null
+++ b/Engine/Source/Runtime/Core/Internal/IO/PlatformIoDispatcherStats.cpp
@@ -0,0 +1,71 @@
+// Copyright Epic Games, Inc. All Rights Reserved.
+
+#include "IO/PlatformIoDispatcherStats.h"
+#include "IO/PlatformIoDispatcherBase.h"
+
+namespace UE
+{
+
+#if UE_PLATFORM_DISPATCHER_STATS_ENABLED
+void FPlatformIoDispatcherStats::OnFileBlockRequestEnqueued(FIoFileBlockRequest& FileBlockRequest)
+{
+	Stats.OnReadRequestsQueued(FileBlockRequest.Size, 1);
+}
+
+void FPlatformIoDispatcherStats::OnFileBlockRequestDequeued(FIoFileBlockRequest& FileBlockRequest)
+{
+}
+
+void FPlatformIoDispatcherStats::OnFileBlockReadStarted(FIoFileBlockRequest& FileBlockRequest)
+{
+	Stats.OnFilesystemReadStarted(FileBlockRequest.FileHandle.Value(), FileBlockRequest.FileOffset, FileBlockRequest.Size);
+}
+
+void FPlatformIoDispatcherStats::OnFileBlockReadCompleted(FIoFileBlockRequest& FileBlockRequest)
+{
+	Stats.OnFilesystemReadCompleted(FileBlockRequest.FileHandle.Value(), FileBlockRequest.FileOffset, FileBlockRequest.Size);
+	Stats.OnReadRequestsCompleted(FileBlockRequest.Size, 1);
+}
+
+void FPlatformIoDispatcherStats::OnFileBlockMemoryAllocated(uint32 Size)
+{
+	Stats.OnBufferAllocated();
+}
+
+void FPlatformIoDispatcherStats::OnFileBlockMemoryFreed(uint32 Size)
+{
+	Stats.OnBufferReleased();
+}
+
+void FPlatformIoDispatcherStats::OnFileBlockCacheHit(uint64 Size)
+{
+	Stats.OnBlockCacheHit(Size);
+}
+
+void FPlatformIoDispatcherStats::OnFileBlockCacheMiss(uint64 Size)
+{
+	Stats.OnBlockCacheMiss(Size);
+}
+
+void FPlatformIoDispatcherStats::OnFileBlockCacheStore(uint64 Size)
+{
+	Stats.OnBlockCacheStore(Size);
+}
+
+void FPlatformIoDispatcherStats::OnDecodeRequestEnqueued(FIoEncodedBlockRequest& EncodedBlockRequest)
+{
+	Stats.OnDecompressQueued(EncodedBlockRequest.BlockCompressedSize, EncodedBlockRequest.BlockUncompressedSize);
+}
+
+void FPlatformIoDispatcherStats::OnDecodeRequestCompleted(FIoEncodedBlockRequest& EncodedBlockRequest)
+{
+	Stats.OnDecompressComplete(EncodedBlockRequest.BlockCompressedSize, EncodedBlockRequest.BlockUncompressedSize);
+}
+
+void FPlatformIoDispatcherStats::OnBytesScattered(int64 BytesScattered)
+{
+	Stats.OnBytesScattered(BytesScattered);
+}
+#endif
+
+} // namespace UE
diff --git a/Engine/Source/Runtime/Core/Internal/IO/PlatformIoDispatcherStats.h b/Engine/Source/Runtime/Core/Internal/IO/PlatformIoDispatcherStats.h
new file mode 100644
--- /dev/null
+++ b/Engine/Source/Runtime/Core/Internal/IO/PlatformIoDispatcherStats.h
@@ -0,0 +1,57 @@
+// Copyright Epic Games, Inc. All Rights Reserved.
+
+#pragma once
+
+#include "IO/IoDispatcherFilesystemStats.h"
+
+#define UE_PLATFORM_DISPATCHER_STATS_ENABLED (COUNTERSTRACE_ENABLED || CSV_PROFILER_STATS)
+
+namespace UE
+{
+
+struct FIoFileBlockRequest;
+struct FIoEncodedBlockRequest;
+
+////////////////////////////////////////////////////////////////////////////////
+#if UE_PLATFORM_DISPATCHER_STATS_ENABLED 
+class FPlatformIoDispatcherStats
+{
+public:
+	void OnFileBlockRequestEnqueued(FIoFileBlockRequest& FileBlockRequest);
+	void OnFileBlockRequestDequeued(FIoFileBlockRequest& FileBlockRequest);
+	void OnFileBlockReadStarted(FIoFileBlockRequest& FileBlockRequest);
+	void OnFileBlockReadCompleted(FIoFileBlockRequest& FileBlockRequest);
+	void OnFileBlockMemoryAllocated(uint32 Size);
+	void OnFileBlockMemoryFreed(uint32 Size);
+	void OnFileBlockCacheHit(uint64 Size);
+	void OnFileBlockCacheMiss(uint64 Size);
+	void OnFileBlockCacheStore(uint64 Size);
+	void OnDecodeRequestEnqueued(FIoEncodedBlockRequest& EncodedBlockRequest);
+	void OnDecodeRequestCompleted(FIoEncodedBlockRequest& EncodedBlockRequest);
+	void OnBytesScattered(int64 BytesScattered);
+
+private:
+	FIoDispatcherFilesystemStats Stats;
+};
+#else
+
+////////////////////////////////////////////////////////////////////////////////
+class FPlatformIoDispatcherStats
+{
+public:
+	void OnFileBlockRequestEnqueued(FIoFileBlockRequest& FileBlockRequest) { };
+	void OnFileBlockRequestDequeued(FIoFileBlockRequest& FileBlockRequest) { };
+	void OnFileBlockReadStarted(FIoFileBlockRequest& FileBlockRequest) { };
+	void OnFileBlockReadCompleted(FIoFileBlockRequest& FileBlockRequest) { };
+	void OnFileBlockMemoryAllocated(uint32 Size) { };
+	void OnFileBlockMemoryFreed(uint32 Size) { };
+	void OnFileBlockCacheHit(uint64 Size) { }
+	void OnFileBlockCacheMiss(uint64 Size) { }
+	void OnFileBlockCacheStore(uint64 Size) { }
+	void OnDecodeRequestEnqueued(FIoEncodedBlockRequest& EncodedBlockRequest) { }
+	void OnDecodeRequestCompleted(FIoEncodedBlockRequest& EncodedBlockRequest) { }
+	void OnBytesScattered(int64 BytesScattered) { }
+};
+#endif
+
+} // namespace UE
diff --git a/Engine/Source/Runtime/Core/Private/IO/IoChunkEncoding.cpp b/Engine/Source/Runtime/Core/Private/IO/IoChunkEncoding.cpp
--- a/Engine/Source/Runtime/Core/Private/IO/IoChunkEncoding.cpp
+++ b/Engine/Source/Runtime/Core/Private/IO/IoChunkEncoding.cpp
@@ -470,16 +470,21 @@ FIoBlockHash FIoChunkEncoding::HashBlock(FMemoryView Block)
 }
 
 ////////////////////////////////////////////////////////////////////////////////
-void FIoChunkBlockDecoder::Set(IIoChunkBlockDecoder& InDefault)
+void FIoChunkBlockDecoder::Set(IIoChunkBlockDecoder& InInstance)
 {
-	check(Default == nullptr);
-	Default = &InDefault;
+	check(Instance == nullptr);
+	Instance = &InInstance;
 }
 
 IIoChunkBlockDecoder& FIoChunkBlockDecoder::Get()
 {
-	check(Default);
-	return *Default;
+	check(Instance);
+	return *Instance;
 }
 
-IIoChunkBlockDecoder* FIoChunkBlockDecoder::Default = nullptr;
+void FIoChunkBlockDecoder::Release()
+{
+	Instance = nullptr;
+}
+
+IIoChunkBlockDecoder* FIoChunkBlockDecoder::Instance = nullptr;
diff --git a/Engine/Source/Runtime/Core/Private/IO/IoDispatcher.cpp b/Engine/Source/Runtime/Core/Private/IO/IoDispatcher.cpp
--- a/Engine/Source/Runtime/Core/Private/IO/IoDispatcher.cpp
+++ b/Engine/Source/Runtime/Core/Private/IO/IoDispatcher.cpp
@@ -3,13 +3,18 @@
 #include "IO/IoDispatcher.h"
 #include "Async/Mutex.h"
 #include "Async/UniqueLock.h"
+#include "Async/Fundamental/Scheduler.h"
 #include "Containers/Ticker.h"
+#include "Misc/CommandLine.h"
 #include "HAL/Runnable.h"
 #include "HAL/RunnableThread.h"
 #include "HAL/Event.h"
 #include "IO/IoDispatcherBackend.h"
+#include "IO/IoDispatcherConfig.h"
+#include "IO/IoDispatcherChunkDecoder.h"
 #include "IO/IoDispatcherPrivate.h"
 #include "IO/IoOffsetLength.h"
+#include "IO/PlatformIoDispatcher.h"
 #include "Misc/ScopeRWLock.h"
 #include "Misc/CoreDelegates.h"
 #include "Misc/ScopeLock.h"
@@ -343,6 +348,34 @@ public:
 			RequestAllocator->Trim();
 			BatchAllocator.Trim();
 		});
+
+		ChunkBlockDecoder.Initialize(GIoDispatcherDecompressionWorkerCount, GIoDispatcherMaxConsecutiveDecompressionJobs, UE::Tasks::ETaskPriority::BackgroundNormal);
+		FIoChunkBlockDecoder::Set(ChunkBlockDecoder);
+
+		OversubscriptionLimitReached = 
+			LowLevelTasks::FScheduler::Get().GetOversubscriptionLimitReachedEvent().AddLambda(
+				[this]()
+				{
+					bTaskSchedulerOversubscribed.store(true, std::memory_order_relaxed);
+					DispatcherEvent->Trigger();
+				}
+			);
+
+		if (IsPlatformIoDispatcherEnabled())
+		{
+			UE::FPlatformIoDispatcherCreateParams CreateParams
+			{
+				.bMultithreaded = bInIsMultithreaded,
+				.bForceGeneric	= !bInIsMultithreaded
+			};
+#if !UE_BUILD_SHIPPING
+			if (FParse::Param(FCommandLine::Get(), TEXT("forcegenericio")))
+			{
+				CreateParams.bForceGeneric = true;
+			}
+#endif
+			UE::FPlatformIoDispatcher::Create(MoveTemp(CreateParams));
+		}
 	}
 
 	~FIoDispatcherImpl()
@@ -355,6 +388,8 @@ public:
 		FCoreDelegates::GetMemoryTrimDelegate().Remove(MemoryTrimDelegateHandle);
 		BackendContext->WakeUpDispatcherThreadDelegate.Unbind();
 		RequestAllocator->ReleaseRef();
+		FIoChunkBlockDecoder::Release();
+		UE::FPlatformIoDispatcher::Shutdown();	
 	}
 
 	void Initialize()
@@ -364,6 +399,7 @@ public:
 			return;
 		}
 		bIsInitialized = true;
+		UE::FPlatformIoDispatcher::Initialize();
 		if (!Backends.IsEmpty())
 		{
 			for (const FBackendAndPriority& Backend : Backends)
@@ -604,6 +640,16 @@ public:
 		{
 			LLM_SCOPE_BYNAME(TEXT("FileSystem/IODispatcher"));
 			ProcessIncomingRequests();
+			if (UE::IPlatformIoDispatcher* PlatformDispatcher = UE::FPlatformIoDispatcher::TryGet())
+			{
+				for (;;)
+				{
+					if (PlatformDispatcher->Tick() == false)
+					{
+						break;
+					}
+				}
+			}
 			while (PendingIoRequestsCount > 0)
 			{
 				ProcessCompletedRequests();
@@ -874,6 +920,10 @@ private:
 			}
 			ProcessIncomingRequests();
 			ProcessCompletedRequests();
+			while(bTaskSchedulerOversubscribed.load(std::memory_order_relaxed))
+			{
+				bTaskSchedulerOversubscribed = ChunkBlockDecoder.TryRetractAndExecuteDecodeTasks();
+			}
 		}
 
 		ProcessIncomingRequests();
@@ -909,9 +959,12 @@ private:
 	uint64 PendingIoRequestsCount = 0;
 	int64 TotalLoaded = 0;
 	FIoRequestStats RequestStats;
+	UE::FIoDispatcherChunkBlockDecoder ChunkBlockDecoder;
 	bool bIsInitialized = false;
 	UE::FMutex WaitingMutex;
 	UE::FMutex UpdateMutex;
+	FDelegateHandle OversubscriptionLimitReached;
+	std::atomic_bool bTaskSchedulerOversubscribed{false};
 };
 
 FIoDispatcher::FIoDispatcher()
diff --git a/Engine/Source/Runtime/Core/Private/IO/IoDispatcherChunkDecoder.cpp b/Engine/Source/Runtime/Core/Private/IO/IoDispatcherChunkDecoder.cpp
--- a/Engine/Source/Runtime/Core/Private/IO/IoDispatcherChunkDecoder.cpp
+++ b/Engine/Source/Runtime/Core/Private/IO/IoDispatcherChunkDecoder.cpp
@@ -7,6 +7,7 @@ namespace UE
 {
 
 TRACE_DECLARE_INT_COUNTER(ChunkDecoderQueueCount, TEXT("IoDispatcher/ChunkDecoderQueueCount"));
+TRACE_DECLARE_INT_COUNTER(AvailableChunkDecoderWorkerCount, TEXT("IoDispatcher/AvailableChunkDecoderWorkerCount"));
 
 void FIoDispatcherChunkBlockDecoder::Enqueue(FIoChunkBlockDecodeRequest&& DecodeRequest)
 {
@@ -20,6 +21,7 @@ void FIoDispatcherChunkBlockDecoder::Enqueue(FIoChunkBlockDecodeRequest&& Decode
 			Ctx				= NextFreeContext;
 			NextFreeContext	= Ctx->Next;
 			Ctx->Next		= nullptr;
+			TRACE_COUNTER_DECREMENT(AvailableChunkDecoderWorkerCount);
 		}
 		else
 		{
@@ -47,6 +49,7 @@ void FIoDispatcherChunkBlockDecoder::Initialize(uint32 InMaxWorkerCount, uint32
 	{
 		Ctx.Next		= NextFreeContext;
 		NextFreeContext	= &Ctx;
+		TRACE_COUNTER_INCREMENT(AvailableChunkDecoderWorkerCount);
 	}
 }
 
@@ -67,7 +70,10 @@ bool FIoDispatcherChunkBlockDecoder::TryRetractAndExecuteDecodeTasks()
 			Task = DecodeContexts[ContextIndex].Task;
 			ContextIndex++;
 		}
-		Task.TryRetractAndExecute();
+		if (Task.IsValid())
+		{
+			Task.TryRetractAndExecute();
+		}
 		bOversubscribed = LowLevelTasks::FScheduler::Get().IsOversubscriptionLimitReached(TaskPriority);
 	}
 
@@ -141,6 +147,7 @@ void FIoDispatcherChunkBlockDecoder::TryDequeueAndLaunchDecodeTasks()
 			DecodeRequest			= MoveTemp(QueueEntry->Request);
 			QueueEntryAllocator.Destroy(QueueEntry);
 			TRACE_COUNTER_DECREMENT(ChunkDecoderQueueCount);
+			TRACE_COUNTER_DECREMENT(AvailableChunkDecoderWorkerCount);
 		}
 
 		check(DecodeRequest.IsValid());
@@ -173,6 +180,7 @@ void FIoDispatcherChunkBlockDecoder::LaunchDecodeTask(FDecodeContext& Ctx, FIoCh
 				UE::TUniqueLock Lock(Mutex);
 				Ctx.Next		= NextFreeContext;
 				NextFreeContext = &Ctx;
+				TRACE_COUNTER_INCREMENT(AvailableChunkDecoderWorkerCount);
 
 				if (DecodeRequest.IsValid())
 				{
diff --git a/Engine/Source/Runtime/Core/Public/IO/IoChunkEncoding.h b/Engine/Source/Runtime/Core/Public/IO/IoChunkEncoding.h
--- a/Engine/Source/Runtime/Core/Public/IO/IoChunkEncoding.h
+++ b/Engine/Source/Runtime/Core/Public/IO/IoChunkEncoding.h
@@ -157,9 +157,10 @@ public:
 class FIoChunkBlockDecoder
 {
 public:
-	CORE_API static void					Set(IIoChunkBlockDecoder& Default); 
+	CORE_API static void					Set(IIoChunkBlockDecoder& Instance);
 	CORE_API static IIoChunkBlockDecoder&	Get();
+	CORE_API static void					Release();
 
 private:
-	static IIoChunkBlockDecoder* Default;
+	static IIoChunkBlockDecoder* Instance;
 };
diff --git a/Engine/Source/Runtime/Core/Tests/IO/PlatformIoDispatcherTests.cpp b/Engine/Source/Runtime/Core/Tests/IO/PlatformIoDispatcherTests.cpp
new file mode 100644
--- /dev/null
+++ b/Engine/Source/Runtime/Core/Tests/IO/PlatformIoDispatcherTests.cpp
@@ -0,0 +1,193 @@
+// Copyright Epic Games, Inc. All Rights Reserved.
+
+#if WITH_LOW_LEVEL_TESTS
+
+#include "Algo/RandomShuffle.h"
+#include "IO/IoStatus.h"
+#include "IO/PlatformIoDispatcherBase.h"
+#include "TestHarness.h"
+#include "TestMacros/Assertions.h"
+#include <catch2/generators/catch_generators.hpp>
+
+namespace UE
+{
+
+namespace PlatformIoDispatcherTest
+{
+
+FIoFileBlockRequestList RandomList(TArray<FIoFileBlockRequest>& Requests)
+{
+	TArray<int32> Indices;
+	for (int32 Idx = 0; Idx < Requests.Num(); ++Idx)
+	{
+		Indices.Add(Idx);
+	}
+	Algo::RandomShuffle(Indices);
+
+	FIoFileBlockRequestList Random;
+	for (int32 Idx : Indices)
+	{
+		Random.AddTail(&Requests[Idx]);
+	}
+
+	return Random;
+}
+
+}
+
+TEST_CASE("Core::IO::Platform::FileReadQueue", "[Core][IO]")
+{
+	using namespace PlatformIoDispatcherTest;
+
+	SECTION("DequeueBySeqNo")
+	{
+		// Arrange
+		const int32					ExpectedReadCount = 20;
+		FPlatformIoDispatcherStats	Stats;
+		FIoFileBlockRequestQueue	Queue(Stats);
+		TArray<FIoFileBlockRequest> Enqueued;
+
+		Enqueued.SetNum(ExpectedReadCount);
+		FIoFileBlockRequestList ToEnqueue = RandomList(Enqueued);
+
+		// Act
+		Queue.Enqueue(MoveTemp(ToEnqueue));
+
+		TArray<FIoFileBlockRequest*> Dequeued;
+		while (FIoFileBlockRequest* Request = Queue.Dequeue())
+		{
+			Dequeued.Add(Request);
+		}
+
+		// Assert
+		CHECK(Dequeued.Num() == ExpectedReadCount);
+		for (int32 Idx = 0; Idx < Dequeued.Num(); ++Idx) 
+		{
+			CHECK(Dequeued[Idx]->SeqNo == Enqueued[Idx].SeqNo);
+		}
+	}
+
+	SECTION("CancelledComesFirst")
+	{
+		// Arrange
+		const int32						ExpectedReadCount = 5;
+		FPlatformIoDispatcherStats		Stats;
+		FIoFileBlockRequestQueue		Queue(Stats);
+		TArray<FIoFileBlockRequest> 	Enqueued;
+		FIoFileHandle					FileHandle(1);
+
+		Enqueued.SetNum(ExpectedReadCount);
+		FIoFileBlockRequestList ToEnqueue;
+		for (FIoFileBlockRequest& Request : Enqueued)
+		{
+			ToEnqueue.AddTail(&Request);
+		}
+
+		// Act
+		Queue.Enqueue(MoveTemp(ToEnqueue));
+		Enqueued[3].ErrorCode = EIoErrorCode::Cancelled;
+		Enqueued[4].ErrorCode = EIoErrorCode::Cancelled;
+		Queue.ReprioritizeCancelled();
+
+		TArray<FIoFileBlockRequest*> Dequeued;
+		while (FIoFileBlockRequest* Request = Queue.Dequeue())
+		{
+			Dequeued.Add(Request);
+		}
+
+		// Assert
+		CHECK(Dequeued.Num() == ExpectedReadCount);
+		CHECK(Dequeued[0]->SeqNo == Enqueued[3].SeqNo);
+		CHECK(Dequeued[1]->SeqNo == Enqueued[4].SeqNo);
+		CHECK(Dequeued[2]->SeqNo == Enqueued[0].SeqNo);
+		CHECK(Dequeued[3]->SeqNo == Enqueued[1].SeqNo);
+		CHECK(Dequeued[4]->SeqNo == Enqueued[2].SeqNo);
+	}
+
+	SECTION("HigherPriorityComesFirst")
+	{
+		// Arrange
+		const int32						ExpectedReadCount = 5;
+		FPlatformIoDispatcherStats		Stats;
+		FIoFileBlockRequestQueue		Queue(Stats);
+		TArray<FIoFileBlockRequest> 	Enqueued;
+		FIoFileHandle					FileHandle(1);
+
+		Enqueued.SetNum(ExpectedReadCount);
+		FIoFileBlockRequestList ToEnqueue;
+		for (FIoFileBlockRequest& Request : Enqueued)
+		{
+			ToEnqueue.AddTail(&Request);
+		}
+
+		// Act
+		Queue.Enqueue(MoveTemp(ToEnqueue));
+		Enqueued[0].Priority = EIoFileReadPriority::Min;
+		Enqueued[1].Priority = EIoFileReadPriority::Low;
+		Enqueued[2].Priority = EIoFileReadPriority::Medium;
+		Enqueued[3].Priority = EIoFileReadPriority::High;
+		Enqueued[4].Priority = EIoFileReadPriority::Max;
+		Queue.Reprioritize();
+
+		TArray<FIoFileBlockRequest*> Dequeued;
+		while (FIoFileBlockRequest* Request = Queue.Dequeue())
+		{
+			Dequeued.Add(Request);
+		}
+
+		// Assert
+		CHECK(Dequeued.Num() == ExpectedReadCount);
+		CHECK(Dequeued[0]->SeqNo == Enqueued[4].SeqNo);
+		CHECK(Dequeued[1]->SeqNo == Enqueued[3].SeqNo);
+		CHECK(Dequeued[2]->SeqNo == Enqueued[2].SeqNo);
+		CHECK(Dequeued[3]->SeqNo == Enqueued[1].SeqNo);
+		CHECK(Dequeued[4]->SeqNo == Enqueued[0].SeqNo);
+	}
+
+	SECTION("DequeueByOffset")
+	{
+		// Arrange
+		const int32						ExpectedReadCount = 20;
+		FPlatformIoDispatcherStats		Stats;
+		FIoFileBlockRequestQueue		Queue(Stats);
+		TArray<FIoFileBlockRequest> 	Enqueued;
+		FIoFileHandle					FileHandle(1);
+		uint64							FileBlockOffset = 0; 
+
+		Queue.SortByOffset(true);
+
+		for (int32 Idx = 0; Idx < ExpectedReadCount; ++Idx) 
+		{
+			FIoFileBlockRequest& Request = Enqueued.AddDefaulted_GetRef();
+			Request.FileHandle = FileHandle;
+			Request.FileOffset = FileBlockOffset++;
+		}
+
+		FIoFileBlockRequestList ToEnqueue;
+		ToEnqueue.AddTail(&Enqueued[0]); // We need one successful read before sorting on offset works
+		for (int32 Idx = Enqueued.Num() - 1; Idx >= 1; --Idx)
+		{
+			ToEnqueue.AddTail(&Enqueued[Idx]);
+		}
+
+		// Act
+		Queue.Enqueue(MoveTemp(ToEnqueue));
+
+		TArray<FIoFileBlockRequest*> Dequeued;
+		while (FIoFileBlockRequest* Request = Queue.Dequeue())
+		{
+			Dequeued.Add(Request);
+		}
+
+		// Assert
+		CHECK(Dequeued.Num() == Enqueued.Num());
+		for (int32 Idx = 0; Idx < Dequeued.Num(); ++Idx) 
+		{
+			CHECK(Dequeued[Idx]->SeqNo == Enqueued[Idx].SeqNo);
+		}
+	}
+}
+
+} // namespace UE
+
+#endif // WITH_LOW_LEVEL_TESTS
diff --git a/Engine/Source/Runtime/PakFile/Private/FileIoDispatcherBackend.cpp b/Engine/Source/Runtime/PakFile/Private/FileIoDispatcherBackend.cpp
new file mode 100644
--- /dev/null
+++ b/Engine/Source/Runtime/PakFile/Private/FileIoDispatcherBackend.cpp
@@ -0,0 +1,1069 @@
+// Copyright Epic Games, Inc. All Rights Reserved.
+
+#include "FileIoDispatcherBackend.h"
+
+#include "Algo/BinarySearch.h"
+#include "Async/Mutex.h"
+#include "Async/UniqueLock.h"
+#include "HAL/CriticalSection.h"
+#include "HAL/Platform.h"
+#include "HAL/PlatformFileManager.h"
+#include "IO/GenericPlatformIoDispatcher.h"
+#include "IO/IoAllocators.h"
+#include "IO/IoChunkId.h"
+#include "IO/IoChunkEncoding.h"
+#include "IO/IoContainerHeader.h"
+#include "IO/IoContainerId.h"
+#include "IO/IoDispatcher.h"
+#include "IO/IoDispatcherBackend.h"
+#include "IO/IoOffsetLength.h"
+#include "IO/IoStatus.h"
+#include "IO/IoStore.h"
+#include "IO/PlatformIoDispatcher.h"
+#include "Math/NumericLimits.h"
+#include "Memory/MemoryView.h"
+#include "Misc/EncryptionKeyManager.h"
+#include "Misc/EnumClassFlags.h"
+#include "Misc/PathViews.h"
+#include "Misc/ScopeLock.h"
+#include "Misc/ScopeRWLock.h"
+#include "Misc/StringBuilder.h"
+#include "Serialization/MemoryReader.h"
+
+#include <atomic>
+
+namespace UE::IoStore
+{
+
+////////////////////////////////////////////////////////////////////////////////
+float GFileIoStoreUnmountTimeOutSeconds = 10.0f;
+static FAutoConsoleVariableRef CVar_UnmountTimeOutSeconds(
+	TEXT("fileiostore.UnmountTimeOutSeconds"),
+	GFileIoStoreUnmountTimeOutSeconds,
+	TEXT("Max time to wait for pending I/O requests before unmounting a container.")
+);
+
+////////////////////////////////////////////////////////////////////////////////
+class FMappedFileProxy final : public IMappedFileHandle
+{
+public:
+	FMappedFileProxy(IMappedFileHandle* InSharedMappedFileHandle, uint64 InSize)
+		: IMappedFileHandle(InSize)
+		, SharedMappedFileHandle(InSharedMappedFileHandle)
+	{
+	}
+
+	virtual ~FMappedFileProxy() { }
+
+	virtual IMappedFileRegion* MapRegion(int64 Offset = 0, int64 BytesToMap = MAX_int64, FFileMappingFlags Flags = EMappedFileFlags::ENone) override
+	{
+		return SharedMappedFileHandle != nullptr ? SharedMappedFileHandle->MapRegion(Offset, BytesToMap, Flags) : nullptr;
+	}
+private:
+	IMappedFileHandle* SharedMappedFileHandle;
+};
+
+////////////////////////////////////////////////////////////////////////////////
+struct FChunkLookup
+{
+	const FIoOffsetAndLength* Find(const FIoChunkId& ChunkId) const;
+
+	enum class EType : uint8
+	{
+		Default,
+		Perfect
+	};
+
+	struct FPerfectHashMap
+	{
+		TConstArrayView<int32>				ChunkHashSeeds;
+		TConstArrayView<FIoChunkId>			ChunkIds;
+		TConstArrayView<FIoOffsetAndLength>	Offsets;
+	};
+
+	FPerfectHashMap							PerfectMap;
+	TMap<FIoChunkId, FIoOffsetAndLength>	DefaultMap;
+	EType									Type = EType::Default;
+};
+
+////////////////////////////////////////////////////////////////////////////////
+const FIoOffsetAndLength* FChunkLookup::Find(const FIoChunkId& ChunkId) const
+{
+	if (Type == EType::Default)
+	{
+		return DefaultMap.Find(ChunkId);
+	}
+
+	// See FIoStoreWriterImpl::GeneratePerfectHashes
+	const uint32 ChunkCount = PerfectMap.ChunkIds.Num();
+	if (ChunkCount == 0)
+	{
+		return nullptr;
+	}
+
+	const uint32 SeedCount	= PerfectMap.ChunkHashSeeds.Num();
+	uint32 SeedIndex		= FIoStoreTocResource::HashChunkIdWithSeed(0, ChunkId) % SeedCount;
+	const int32 Seed		= PerfectMap.ChunkHashSeeds[SeedIndex];
+
+	if (Seed == 0)
+	{
+		return nullptr;
+	}
+
+	uint32 Slot = MAX_uint32;
+	if (Seed < 0)
+	{
+		const uint32 SeedAsIndex = static_cast<uint32>(-Seed - 1);
+		if (SeedAsIndex < ChunkCount)
+		{
+			Slot = static_cast<uint32>(SeedAsIndex);
+		}
+		else
+		{
+			// Entry without perfect hash
+			return DefaultMap.Find(ChunkId);
+		}
+	}
+	else
+	{
+		Slot = FIoStoreTocResource::HashChunkIdWithSeed(static_cast<uint32>(Seed), ChunkId) % ChunkCount;
+	}
+
+	if (PerfectMap.ChunkIds[Slot] == ChunkId)
+	{
+		return &PerfectMap.Offsets[Slot];
+	}
+
+	return nullptr;
+};
+////////////////////////////////////////////////////////////////////////////////
+struct FContainer;
+
+struct FContainerPartition
+{
+	FContainer&						Container;
+	FString							Filename;
+	FIoFileHandle					FileHandle;
+	uint64							FileSize = 0;
+	TUniquePtr<IMappedFileHandle>	MappedFileHandle;
+};
+
+////////////////////////////////////////////////////////////////////////////////
+struct FContainer
+{
+	static TIoStatusOr<TUniquePtr<FContainer>>			Open(const TCHAR* Filename, int32 MountOrder, uint32 InstanceId, FIoContainerHeader& OutContainerHeader);
+	FContainerPartition&								GetPartition(uint64 Offset, uint64& OutPartitionOffset, int32* OutIndex = nullptr);
+	uint64												GetAllocatedSize() const;
+
+	FIoStoreTocResourceStorage							TocStorage;
+	FAES::FAESKey										EncryptionKey;
+	FChunkLookup										ChunkLookup;
+	TArray<FContainerPartition>							Partitions;
+	TArray<FName>										CompressionMethods;
+	TConstArrayView<FIoStoreTocCompressedBlockEntry>	CompressionBlocks;
+	TConstArrayView<FSHAHash>							CompressionBlockHashes;
+	FString												BaseFilePath;
+	FIoContainerId										ContainerId;
+	uint64												PartitionSize = 0;
+	uint64												CompressionBlockSize = 0;
+	EIoContainerFlags									ContainerFlags;
+	int32												MountOrder = MAX_int32;
+	uint32												InstanceId = 0;
+	std::atomic_int32_t									ActiveReadCount{0};
+};
+
+using FUniqueContainer = TUniquePtr<FContainer>;
+
+////////////////////////////////////////////////////////////////////////////////
+TIoStatusOr<TUniquePtr<FContainer>>	FContainer::Open(
+	const TCHAR* Filename,
+	int32 MountOrder,
+	uint32 InstanceId,
+	FIoContainerHeader& OutContainerHeader)
+{
+	LLM_SCOPE(ELLMTag::AsyncLoading);
+	TRACE_CPUPROFILER_EVENT_SCOPE(ReadContainerHeader);
+
+	FString BaseFilePath = FString(FPathViews::GetBaseFilenameWithPath(Filename));
+
+	TStringBuilder<256> Sb;
+
+	Sb << BaseFilePath << TEXT(".utoc");
+
+	FIoStoreTocResourceView		TocView;
+	FIoStoreTocResourceStorage	TocStorage;
+
+	FIoStatus Status = FIoStoreTocResourceView::Read(Sb.ToString(), EIoStoreTocReadOptions::Default, TocView, TocStorage);
+	if (Status.IsOk() == false)
+	{
+		return Status;
+	}
+
+	FUniqueContainer Container = MakeUnique<FContainer>();
+
+	if (EnumHasAnyFlags(TocView.Header.ContainerFlags, EIoContainerFlags::Encrypted))
+	{
+		if (FEncryptionKeyManager::Get().TryGetKey(TocView.Header.EncryptionKeyGuid, Container->EncryptionKey) == false)
+		{
+			return FIoStatus(EIoErrorCode::InvalidEncryptionKey);
+		}
+	}
+
+	Container->TocStorage				= MoveTemp(TocStorage);
+	Container->BaseFilePath				= BaseFilePath;
+	Container->PartitionSize			= TocView.Header.PartitionSize;
+	Container->CompressionMethods		= MoveTemp(TocView.CompressionMethods);
+	Container->CompressionBlockSize		= TocView.Header.CompressionBlockSize;
+	Container->CompressionBlocks		= MoveTemp(TocView.CompressionBlocks);
+	Container->CompressionBlockHashes	= MoveTemp(TocView.ChunkBlockSignatures);
+	Container->ContainerFlags			= TocView.Header.ContainerFlags;
+	Container->ContainerId				= TocView.Header.ContainerId;
+	Container->MountOrder				= MountOrder;
+	Container->InstanceId				= InstanceId;
+
+	// Parse lookup table information
+	if (TocView.ChunkPerfectHashSeeds.IsEmpty() == false)
+	{
+		for (int32 ChunkIndex : TocView.ChunkIndicesWithoutPerfectHash)
+		{
+			const FIoChunkId&			ChunkId = TocView.ChunkIds[ChunkIndex];
+			const FIoOffsetAndLength&	OffsetLength = TocView.ChunkOffsetLengths[ChunkIndex];
+			
+			Container->ChunkLookup.DefaultMap.Add(ChunkId, OffsetLength);
+		}
+		
+		Container->ChunkLookup.PerfectMap.ChunkHashSeeds	= MoveTemp(TocView.ChunkPerfectHashSeeds);
+		Container->ChunkLookup.PerfectMap.Offsets			= MoveTemp(TocView.ChunkOffsetLengths);
+		Container->ChunkLookup.PerfectMap.ChunkIds			= MoveTemp(TocView.ChunkIds);
+		Container->ChunkLookup.Type							= FChunkLookup::EType::Perfect;
+	}
+	else
+	{
+		for (uint32 ChunkIndex = 0; ChunkIndex < TocView.Header.TocEntryCount; ++ChunkIndex)
+		{
+			const FIoChunkId&			ChunkId = TocView.ChunkIds[ChunkIndex];
+			const FIoOffsetAndLength&	OffsetLength = TocView.ChunkOffsetLengths[ChunkIndex];
+
+			Container->ChunkLookup.DefaultMap.Add(ChunkId, OffsetLength);
+		}
+		Container->ChunkLookup.Type	= FChunkLookup::EType::Default;
+	}
+
+	// Open partition file handles 
+	Container->Partitions.Reserve(TocView.Header.PartitionCount);
+	for (uint32 PartitionIndex = 0; PartitionIndex < TocView.Header.PartitionCount; PartitionIndex++)
+	{
+		FContainerPartition& Part = Container->Partitions.Add_GetRef(FContainerPartition
+		{
+			.Container	= *Container,
+		});
+
+		Sb.Reset();
+		Sb << BaseFilePath;
+		if (PartitionIndex > 0)
+		{
+			Sb.Appendf(TEXT("_s%d"), PartitionIndex);
+		}
+		Sb.Append(TEXT(".ucas"));
+
+		EIoFilePropertyFlags FileFlags = EIoFilePropertyFlags::None;
+		if (EnumHasAnyFlags(Container->ContainerFlags, EIoContainerFlags::Encrypted))
+		{
+			FileFlags |= EIoFilePropertyFlags::Encrypted;
+		}
+		if (EnumHasAnyFlags(Container->ContainerFlags, EIoContainerFlags::Signed))
+		{
+			FileFlags |= EIoFilePropertyFlags::Signed;
+		}
+
+		FIoFileProperties FileProperties
+		{
+			.CompressionMethods		= Container->CompressionMethods,
+			.CompressionBlockSize	= IntCastChecked<uint32>(Container->CompressionBlockSize),
+			.Flags					= FileFlags
+		};
+		
+		FIoFileStat FileStats;
+		TIoStatusOr<FIoFileHandle> Handle = FPlatformIoDispatcher::Get().OpenFile(
+			Sb.ToString(),
+			FileProperties,
+			&FileStats);
+
+		if (Handle.IsOk() == false)
+		{
+			return Handle.Status();
+		}
+
+		Part.Filename	= Sb.ToString();
+		Part.FileHandle = Handle.ConsumeValueOrDie();
+		Part.FileSize	= FileStats.FileSize;
+	}
+
+	// Read the container header information 
+	const FIoChunkId HeaderChunkId				= CreateContainerHeaderChunkId(Container->ContainerId);
+	const FIoOffsetAndLength* OffsetAndLength	= Container->ChunkLookup.Find(HeaderChunkId);
+
+	// Deserialize the container header
+	if (OffsetAndLength != nullptr)
+	{
+		const uint32 FirstBlock		= uint32(OffsetAndLength->GetOffset() / Container->CompressionBlockSize);
+		const uint32 LastBlock		= uint32((OffsetAndLength->GetOffset() + OffsetAndLength->GetLength() - 1) / Container->CompressionBlockSize);
+		uint64 EncodedSize			= 0;
+
+		FName CompressionMethod = NAME_None;
+		TArray<uint32> BlockSizes;
+		BlockSizes.Reserve((LastBlock - FirstBlock) + 1);
+
+		for (uint32 Idx = FirstBlock; Idx <= LastBlock; ++Idx)
+		{
+			const FIoStoreTocCompressedBlockEntry& Block	= Container->CompressionBlocks[Idx];
+			const FName BlockCompressionMethod				= Container->CompressionMethods[Block.GetCompressionMethodIndex()];
+			const uint32 EncodedBlockSize					= Block.GetCompressedSize();
+
+			if (BlockCompressionMethod != NAME_None)
+			{
+				ensure(CompressionMethod == NAME_None || CompressionMethod == BlockCompressionMethod);
+				CompressionMethod = BlockCompressionMethod;
+			}
+
+			BlockSizes.Add(EncodedBlockSize);
+			EncodedSize += Align(EncodedBlockSize, FAES::AESBlockSize); // Size on disk is always aligned to AES block size
+		}
+
+		uint64 PartitionOffset	= MAX_uint64;
+		int32 PartitionIndex	= INDEX_NONE;
+
+		FContainerPartition& Partition = Container->GetPartition(
+			Container->CompressionBlocks[FirstBlock].GetOffset(),
+			PartitionOffset,
+			&PartitionIndex);
+
+		IPlatformFile& Ipf = FPlatformFileManager::Get().GetPlatformFile();
+		TUniquePtr<IFileHandle> FileHandle(Ipf.OpenRead(*Partition.Filename));
+		if (FileHandle.IsValid() == false)
+		{
+			return FIoStatus(FIoStatusBuilder(EIoErrorCode::FileOpenFailed)
+				<< TEXT("Failed to open container '") << Sb.ToString() << TEXT("'"));
+		}
+
+		if (FileHandle->Seek(PartitionOffset) == false)
+		{
+			return FIoStatus(FIoStatusBuilder(EIoErrorCode::FileOpenFailed)
+				<< TEXT("Failed to seek to container header offset"));
+		}
+
+		FIoBuffer EncodedBlocks(EncodedSize);
+		if (FileHandle->Read(EncodedBlocks.GetData(), EncodedBlocks.GetSize()) == false)
+		{
+			return FIoStatus(FIoStatusBuilder(EIoErrorCode::FileOpenFailed)
+				<< TEXT("Failed to read container header chunk"));
+		}
+
+		const FMemoryView EncryptionKey = EnumHasAnyFlags(Container->ContainerFlags, EIoContainerFlags::Encrypted)
+			? MakeMemoryView(Container->EncryptionKey.Key, FAES::FAESKey::KeySize)
+			: FMemoryView();
+
+		FIoBuffer DecodedChunk(OffsetAndLength->GetLength());
+
+		FIoChunkDecodingParams Params;
+		Params.CompressionFormat	= CompressionMethod; 
+		Params.EncryptionKey		= EncryptionKey; 
+		Params.BlockSize			= uint32(Container->CompressionBlockSize);
+		Params.TotalRawSize			= OffsetAndLength->GetLength();
+		Params.EncodedBlockSize		= BlockSizes;
+
+		if (FIoChunkEncoding::Decode(Params, EncodedBlocks.GetView(), DecodedChunk.GetMutableView()) == false)
+		{
+			return FIoStatus(FIoStatusBuilder(EIoErrorCode::FileOpenFailed)
+				<< TEXT("Failed to deserialize container header"));
+		}
+
+		FMemoryReaderView Ar(DecodedChunk.GetView());
+		Ar << OutContainerHeader;
+
+		if (Ar.IsError() || Ar.IsCriticalError())
+		{
+			return FIoStatus(FIoStatusBuilder(EIoErrorCode::FileOpenFailed)
+				<< TEXT("Failed to serialize container header"));
+		}
+	}
+
+	return Container;
+}
+
+FContainerPartition& FContainer::GetPartition(uint64 Offset, uint64& OutPartitionOffset, int32* OutIndex)
+{
+	const int32 PartitionIndex	= IntCastChecked<int32>(Offset / PartitionSize);
+	OutPartitionOffset			= Offset % PartitionSize;
+
+	if (OutIndex != nullptr)
+	{
+		*OutIndex = PartitionIndex;
+	}
+
+	ensure(PartitionIndex < Partitions.Num());
+	return Partitions[PartitionIndex];
+}
+
+uint64 FContainer::GetAllocatedSize() const
+{
+	return TocStorage.GetAllocatedSize() + ChunkLookup.DefaultMap.GetAllocatedSize();
+}
+
+////////////////////////////////////////////////////////////////////////////////
+struct FChunkInfo
+{
+				FChunkInfo() = default;
+				FChunkInfo(FContainer* Container, const FIoOffsetAndLength* OffsetLength);
+	bool		IsValid() const { return Container != nullptr; }
+	uint64		Offset() const { return OffsetLength->GetOffset(); }
+	uint64		Size() const { return OffsetLength->GetLength(); }
+	FContainer& GetContainer() { return *Container; }
+
+private:
+	FContainer*					Container = nullptr;
+	const FIoOffsetAndLength*	OffsetLength = nullptr;
+};
+
+////////////////////////////////////////////////////////////////////////////////
+FChunkInfo::FChunkInfo(FContainer* InContainer, const FIoOffsetAndLength* InOffsetLength)
+	: Container(InContainer)
+	, OffsetLength(InOffsetLength)
+{
+}
+
+////////////////////////////////////////////////////////////////////////////////
+class FFileIoStore
+{
+public:
+	FChunkInfo						GetChunkInfo(const FIoChunkId& ChunkId) const;
+	TIoStatusOr<FIoContainerHeader>	Mount(const TCHAR* TocPath, int32 MountOrder);
+	bool							Unmount(const TCHAR* TocPath);
+	void							ReopenAllFileHandles();
+	FRWLock&						GetLock() { return RwLock; }
+	FRWLock&						GetLock() const { return RwLock; }
+
+private:
+	TArray<FUniqueContainer>		MountedContainers;
+	mutable FRWLock					RwLock;
+	std::atomic_uint32_t			ContainerInstanceId{1};
+};
+
+FChunkInfo FFileIoStore::GetChunkInfo(const FIoChunkId& ChunkId) const
+{
+	for (const FUniqueContainer& Container : MountedContainers)
+	{
+		if (const FIoOffsetAndLength* OffsetLength = Container->ChunkLookup.Find(ChunkId))
+		{
+			return FChunkInfo(Container.Get(), OffsetLength);
+		}
+	}
+
+	return FChunkInfo();
+}
+
+////////////////////////////////////////////////////////////////////////////////
+TIoStatusOr<FIoContainerHeader>	FFileIoStore::Mount(const TCHAR* TocPath, int32 MountOrder)
+{
+	const uint32		InstanceId = ContainerInstanceId.fetch_add(1, std::memory_order_relaxed);
+	FIoContainerHeader	Hdr;
+
+	TIoStatusOr<FUniqueContainer> Status = FContainer::Open(TocPath, MountOrder, InstanceId, Hdr);
+	if (Status.IsOk() == false)
+	{
+		return Status.Status();
+	}
+
+	FUniqueContainer Container	= Status.ConsumeValueOrDie();
+	int32 MountIndex			= INDEX_NONE;
+	{
+		FWriteScopeLock _(RwLock);
+
+		MountIndex = Algo::UpperBound(
+			MountedContainers,
+			Container,
+			[](const FUniqueContainer& A, const FUniqueContainer& B)
+			{
+				if (A->MountOrder != B->MountOrder)
+				{
+					return A->MountOrder > B->MountOrder;
+				}
+				return A->InstanceId > B->InstanceId;
+			});
+
+		MountedContainers.Insert(MoveTemp(Container), MountIndex);
+	}
+
+	UE_LOG(LogIoStore, Log, TEXT("Mounted container '%s' at position %d"), TocPath, MountIndex);
+
+	return Hdr;
+}
+
+bool FFileIoStore::Unmount(const TCHAR* TocPath)
+{
+	const FString BaseFilePath = FString(FPathViews::GetBaseFilenameWithPath(TocPath));
+	
+	FUniqueContainer ContainerToRemove;
+	{
+		FWriteScopeLock _(RwLock);
+
+		int32 ContainerIdx = INDEX_NONE;
+		for (int32 Idx = 0; FUniqueContainer& Container : MountedContainers)
+		{
+			if (Container->BaseFilePath == BaseFilePath)
+			{
+				ContainerIdx = Idx;
+				break;
+			}
+			++Idx;
+		}
+
+		if (ContainerIdx != INDEX_NONE)
+		{
+			ContainerToRemove = MoveTemp(MountedContainers[ContainerIdx]);
+			MountedContainers.RemoveAt(ContainerIdx);
+		}
+	}
+
+	if (ContainerToRemove.IsValid() == false)
+	{
+		UE_LOG(LogIoStore, Warning, TEXT("Failed to unmount container '%s', reason 'Not Found'"), TocPath);
+		return false;
+	}
+
+	if (ContainerToRemove->ActiveReadCount.load(std::memory_order_seq_cst) > 0)
+	{
+		for (FContainerPartition& Part : ContainerToRemove->Partitions)
+		{
+			UE_LOG(LogIoStore, Log, TEXT("Cancelling inflight read requests for file '%s'"), *Part.Filename);
+			UE::FPlatformIoDispatcher::Get().CancelAllRequests((Part.FileHandle));
+		}
+
+		UE_LOG(LogIoStore, Log, TEXT("Waiting for read request(s) to finish before unmounting '%s.utoc'"), TocPath);
+		const double MaxWaitTimeSeconds	= GFileIoStoreUnmountTimeOutSeconds; 
+		const double StartTime			= FMath::Clamp(FPlatformTime::Seconds(), 5.0f, 30.0f);
+		while (ContainerToRemove->ActiveReadCount.load(std::memory_order_seq_cst) > 0)
+		{
+			FPlatformProcess::Sleep(0);
+			if (FPlatformTime::Seconds() - StartTime > MaxWaitTimeSeconds)
+			{
+				UE_LOG(LogIoStore, Warning, TEXT("Stopped waiting for read request(s) after %.2lf seconds"), MaxWaitTimeSeconds);
+				break;
+			}
+		}
+	}
+
+	for (FContainerPartition& Part : ContainerToRemove->Partitions)
+	{
+		UE::FPlatformIoDispatcher::Get().CloseFile((Part.FileHandle));
+	}
+
+	UE_LOG(LogIoStore, Log, TEXT("Unmounted container '%s'"), TocPath);
+	return true;
+}
+
+void FFileIoStore::ReopenAllFileHandles()
+{
+	FWriteScopeLock _(RwLock);
+	for (FUniqueContainer& Container : MountedContainers)
+	{
+		UE_CLOG(Container->ActiveReadCount.load(std::memory_order_seq_cst) > 0, LogIoStore, Warning, TEXT("Calling ReopenAllFileHandles with read requests in flight"));
+		for (FContainerPartition& Part : Container->Partitions)
+		{
+			UE_LOG(LogIoStore, Log, TEXT("Reopening container file '%s'"), *Part.Filename);
+			UE::FPlatformIoDispatcher::Get().CloseFile((Part.FileHandle));
+
+			EIoFilePropertyFlags FileFlags = EIoFilePropertyFlags::None;
+			if (EnumHasAnyFlags(Container->ContainerFlags, EIoContainerFlags::Encrypted))
+			{
+				FileFlags |= EIoFilePropertyFlags::Encrypted;
+			}
+			if (EnumHasAnyFlags(Container->ContainerFlags, EIoContainerFlags::Signed))
+			{
+				FileFlags |= EIoFilePropertyFlags::Signed;
+			}
+
+			FIoFileProperties FileProperties
+			{
+				.CompressionMethods		= Container->CompressionMethods,
+				.CompressionBlockSize	= IntCastChecked<uint32>(Container->CompressionBlockSize),
+				.Flags					= FileFlags
+			};
+
+			TIoStatusOr<FIoFileHandle> Handle = FPlatformIoDispatcher::Get().OpenFile(*Part.Filename, FileProperties);
+			Part.FileHandle = Handle.ConsumeValueOrDie();
+		}
+	}
+}
+
+////////////////////////////////////////////////////////////////////////////////
+struct FResolvedRequest
+{
+	explicit FResolvedRequest(FIoRequestImpl& InDispatcherRequest)
+		: DispatcherRequest(InDispatcherRequest)
+	{
+		check(DispatcherRequest.BackendData == nullptr);
+		DispatcherRequest.BackendData = this;
+	}
+
+	static FResolvedRequest& Get(FIoRequestImpl& DispatcherRequest)
+	{
+		return *reinterpret_cast<FResolvedRequest*>(DispatcherRequest.BackendData);
+	}
+
+	static FResolvedRequest* TryGet(FIoRequestImpl* DispatcherRequest)
+	{
+		if (DispatcherRequest != nullptr && DispatcherRequest->BackendData != nullptr)
+		{
+			return reinterpret_cast<FResolvedRequest*>(DispatcherRequest->BackendData);
+		}
+		return nullptr;
+	}
+
+	FIoRequestImpl&		DispatcherRequest;
+	FIoBuffer			Buffer;
+	FChunkInfo			ChunkInfo;
+	FIoFileReadRequest	PlatformRequest;
+	FIoFileHandle		FileHandle;
+	uint64				Offset = MAX_uint64;
+	uint64				Size = MAX_uint64;
+	EIoFileReadPriority Priority = EIoFileReadPriority::Medium;
+};
+
+using FRequestAllocator = TSingleThreadedSlabAllocator<FResolvedRequest>;
+
+////////////////////////////////////////////////////////////////////////////////
+class FFileIoDispatcherBackend final
+	: public IFileIoDispatcherBackend
+{
+	using FSharedBackendContext = TSharedPtr<const FIoDispatcherBackendContext>;
+public:
+											FFileIoDispatcherBackend();
+	virtual									~FFileIoDispatcherBackend();
+
+	// IFileIoDispatcherBackend
+	virtual TIoStatusOr<FIoContainerHeader> Mount(const TCHAR* TocPath, int32 Order, const FGuid& EncryptionKeyGuid, const FAES::FAESKey& EncryptionKey) override;
+	virtual bool							Unmount(const TCHAR* TocPath) override;
+	virtual void							ReopenAllFileHandles() override;
+
+	// IIoDispatcherBackend
+	virtual void							Initialize(TSharedRef<const FIoDispatcherBackendContext> Context) override;
+	virtual void							Shutdown() override;
+	virtual void							ResolveIoRequests(FIoRequestList Requests, FIoRequestList& OutUnresolved) override;
+	virtual FIoRequestImpl* 				GetCompletedIoRequests() override;
+	virtual void							CancelIoRequest(FIoRequestImpl* Request) override;
+	virtual void							UpdatePriorityForIoRequest(FIoRequestImpl* Request) override;
+	virtual bool							DoesChunkExist(const FIoChunkId& ChunkId) const override;
+	virtual TIoStatusOr<uint64>				GetSizeForChunk(const FIoChunkId& ChunkId) const override;
+	virtual TIoStatusOr<FIoMappedRegion>	OpenMapped(const FIoChunkId& ChunkId, const FIoReadOptions& Options) override;
+	virtual const TCHAR*					GetName() const { return TEXT("File"); }
+
+private:
+	void									HandleSignatureError(FIoRequestImpl& DispatcherRequest, uint32 FailedBlockIndex);
+
+	FSharedBackendContext		BackendContext;
+	FFileIoStore				IoStore;
+	FRequestAllocator			RequestAllocator;
+	FIoRequestList				CompletedDispatcherRequests;
+	UE::FMutex					Mutex;
+};
+
+////////////////////////////////////////////////////////////////////////////////
+FFileIoDispatcherBackend::FFileIoDispatcherBackend()
+{
+}
+
+FFileIoDispatcherBackend::~FFileIoDispatcherBackend()
+{
+}
+
+TIoStatusOr<FIoContainerHeader> FFileIoDispatcherBackend::Mount(const TCHAR* TocPath, int32 Order, const FGuid& EncryptionKeyGuid, const FAES::FAESKey& EncryptionKey)
+{
+	LLM_SCOPE_BYNAME(TEXT("FileSystem/FileIoStore"));
+	return IoStore.Mount(TocPath, Order);
+}
+
+bool FFileIoDispatcherBackend::Unmount(const TCHAR* TocPath)
+{
+	return IoStore.Unmount(TocPath);
+}
+
+void FFileIoDispatcherBackend::ReopenAllFileHandles()
+{
+	IoStore.ReopenAllFileHandles();
+}
+
+void FFileIoDispatcherBackend::Initialize(TSharedRef<const FIoDispatcherBackendContext> Context)
+{
+	BackendContext = Context;
+}
+
+void FFileIoDispatcherBackend::Shutdown()
+{
+}
+
+void FFileIoDispatcherBackend::ResolveIoRequests(FIoRequestList Requests, FIoRequestList& OutUnresolved)
+{
+	FIoRequestList ResolvedRequests;
+	{
+		FReadScopeLock ReadScope(IoStore.GetLock());
+
+		while (FIoRequestImpl* DispatcherRequest = Requests.PopHead())
+		{
+			const FChunkInfo ChunkInfo = IoStore.GetChunkInfo(DispatcherRequest->ChunkId);
+			if (ChunkInfo.IsValid() == false)
+			{
+				OutUnresolved.AddTail(DispatcherRequest);
+				continue;
+			}
+
+			const int64 ResolvedSize = FMath::Min(DispatcherRequest->Options.GetSize(), ChunkInfo.Size() - DispatcherRequest->Options.GetOffset());
+			if (ResolvedSize > 0)
+			{
+				FResolvedRequest* ResolvedRequest	= RequestAllocator.Construct(*DispatcherRequest);
+				ResolvedRequest->ChunkInfo			= ChunkInfo;
+				ResolvedRequest->Offset				= ChunkInfo.Offset() + DispatcherRequest->Options.GetOffset();
+				ResolvedRequest->Size				= ResolvedSize; 
+				ResolvedRequest->Priority			= IoFileReadPriorityFromDispatcherPriority(DispatcherRequest->Priority);
+				check(DispatcherRequest->BackendData != nullptr);
+
+				ResolvedRequests.AddTail(DispatcherRequest);
+
+				if (DispatcherRequest->Options.GetTargetVa() != nullptr)
+				{
+					DispatcherRequest->CreateBuffer(ResolvedRequest->Size);
+				}
+			}
+			else
+			{
+				if (ResolvedSize < 0)
+				{
+					DispatcherRequest->SetFailed();
+				}
+				else
+				{
+					DispatcherRequest->CreateBuffer(0);
+				}
+				TUniqueLock Lock(Mutex);
+				CompletedDispatcherRequests.AddTail(DispatcherRequest);
+			}
+		}
+	}
+
+	while (FIoRequestImpl* DispatcherRequest = ResolvedRequests.PopHead())
+	{
+		FResolvedRequest& ResolvedRequest = FResolvedRequest::Get(*DispatcherRequest);
+
+		FIoBuffer& Dst						= DispatcherRequest->HasBuffer() ? DispatcherRequest->GetBuffer() : ResolvedRequest.Buffer;
+		FContainer& Container				= ResolvedRequest.ChunkInfo.GetContainer();
+		const int32 FirstCompressedBlock	= IntCastChecked<int32>(ResolvedRequest.Offset / Container.CompressionBlockSize);
+		const int32 LastCompressedBlock		= IntCastChecked<int32>((ResolvedRequest.Offset + ResolvedRequest.Size - 1) / Container.CompressionBlockSize);
+		uint64 RequestStartOffsetInBlock	= ResolvedRequest.Offset - (FirstCompressedBlock * Container.CompressionBlockSize);
+
+		// All encoded blocks for a chunk always resides in the same .ucas file 
+		const FIoStoreTocCompressedBlockEntry& FirstBlock	= Container.CompressionBlocks[FirstCompressedBlock];
+		uint64 FirstBlockOffsetInPartition					= MAX_uint64;
+		FContainerPartition& Partition						= Container.GetPartition(FirstBlock.GetOffset(), FirstBlockOffsetInPartition);
+		ResolvedRequest.FileHandle							= Partition.FileHandle;
+
+		// On some platforms we can read directly into the destination buffer
+		{
+			Container.ActiveReadCount.fetch_add(1, std::memory_order_relaxed);
+			ResolvedRequest.PlatformRequest = FPlatformIoDispatcher::Get().ReadDirect(
+				FIoDirectReadRequestParams
+				{
+					.FileHandle = Partition.FileHandle,
+					.Dst		= Dst,
+					.Offset		= FirstBlockOffsetInPartition + RequestStartOffsetInBlock,
+					.Size		= ResolvedRequest.Size,
+					.UserData	= DispatcherRequest
+				},
+				[this](FIoFileReadResult&& Result)
+				{
+					FIoRequestImpl* DispatcherRequest = reinterpret_cast<FIoRequestImpl*>(Result.UserData);
+					if (Result.ErrorCode != EIoErrorCode::Ok)
+					{
+						DispatcherRequest->SetFailed();
+					}
+
+					{
+						UE::TUniqueLock Lock(Mutex);
+						CompletedDispatcherRequests.AddTail(DispatcherRequest);
+					}
+					BackendContext->WakeUpDispatcherThreadDelegate.Execute();
+				});
+
+			if (ResolvedRequest.PlatformRequest.IsValid())
+			{
+				continue;
+			}
+			else
+			{
+				Container.ActiveReadCount.fetch_sub(1, std::memory_order_relaxed);
+			}
+		}
+
+		FIoScatterGatherRequestParams ScatterGather(Partition.FileHandle, Dst, ResolvedRequest.Size, DispatcherRequest, ResolvedRequest.Priority);
+
+		// Scatter offsets
+		uint64 RequestRemainingBytes	= ResolvedRequest.Size;
+		uint64 OffsetInRequest			= 0;
+		uint64 BlockFileOffset			= FirstBlockOffsetInPartition;		
+
+		for (int32 BlockIndex = FirstCompressedBlock; BlockIndex <= LastCompressedBlock; ++BlockIndex)
+		{
+			const FIoStoreTocCompressedBlockEntry& CompressedBlock	= Container.CompressionBlocks[BlockIndex];
+
+			const uint32 BlockCompressedSize	= CompressedBlock.GetCompressedSize();
+			const uint32 BlockUncompressedSize	= CompressedBlock.GetUncompressedSize();
+			const uint32 BlockFileSize			= Align(BlockCompressedSize, FAES::AESBlockSize);
+			const uint64 ScatterOffset			= RequestStartOffsetInBlock;
+			const uint64 ScatterSize			= FMath::Min<uint64>(CompressedBlock.GetUncompressedSize() - RequestStartOffsetInBlock, RequestRemainingBytes);
+			const uint64 DstOffset				= OffsetInRequest;
+
+			FMemoryView EncryptionKey;
+			if (EnumHasAnyFlags(Container.ContainerFlags, EIoContainerFlags::Encrypted))
+			{
+				EncryptionKey = MakeMemoryView(Container.EncryptionKey.Key, FAES::FAESKey::KeySize);
+			}
+
+			FName CompressionMethod = NAME_None;
+			if (EnumHasAnyFlags(Container.ContainerFlags, EIoContainerFlags::Compressed))
+			{
+				CompressionMethod = Container.CompressionMethods[CompressedBlock.GetCompressionMethodIndex()];
+			}
+
+			FMemoryView BlockHash;
+			if (EnumHasAnyFlags(Container.ContainerFlags, EIoContainerFlags::Signed))
+			{
+				const FSHAHash& ShaHash = Container.CompressionBlockHashes[BlockIndex];
+				BlockHash				= MakeMemoryView(ShaHash.Hash, sizeof(ShaHash.Hash));
+			}
+
+			ScatterGather.Scatter(
+				BlockFileOffset, 
+				BlockIndex,
+				BlockCompressedSize,
+				BlockUncompressedSize,
+				ScatterOffset,
+				ScatterSize,
+				DstOffset,
+				CompressionMethod,
+				EncryptionKey,
+				BlockHash);
+
+			BlockFileOffset				+= BlockFileSize;
+			RequestRemainingBytes		-= ScatterSize; 
+			OffsetInRequest				+= ScatterSize;
+			RequestStartOffsetInBlock	= 0;
+		}
+
+		Container.ActiveReadCount.fetch_add(1, std::memory_order_relaxed);
+		ResolvedRequest.PlatformRequest = FPlatformIoDispatcher::Get().ScatterGather(
+			MoveTemp(ScatterGather),
+			[this](FIoFileReadResult&& Result)
+			{
+				FIoRequestImpl* DispatcherRequest = reinterpret_cast<FIoRequestImpl*>(Result.UserData);
+				if (DispatcherRequest->IsCancelled() == false && Result.ErrorCode != EIoErrorCode::Ok) 
+				{
+					DispatcherRequest->SetFailed();
+					if (Result.ErrorCode == EIoErrorCode::SignatureError)
+					{
+						HandleSignatureError(*DispatcherRequest, Result.FailedBlockId);
+					}
+				}
+				{
+					UE::TUniqueLock Lock(Mutex);
+					CompletedDispatcherRequests.AddTail(DispatcherRequest);
+				}
+				BackendContext->WakeUpDispatcherThreadDelegate.Execute();
+			});
+
+		if (ResolvedRequest.PlatformRequest.IsValid() == false)
+		{
+			UE_LOG(LogIoStore, Warning, TEXT("Failed to create platform read request, ChunkId='%s' Filenname='%s'"),
+				*LexToString(ResolvedRequest.DispatcherRequest.ChunkId), *Container.BaseFilePath);
+			ResolvedRequest.DispatcherRequest.SetFailed();	
+			{
+				UE::TUniqueLock Lock(Mutex);
+				CompletedDispatcherRequests.AddTail(DispatcherRequest);
+			}
+		}
+	}
+}
+
+FIoRequestImpl* FFileIoDispatcherBackend::GetCompletedIoRequests()
+{
+	LLM_SCOPE_BYNAME(TEXT("FileSystem/FileIoStore"));
+
+	FIoRequestList LocalCompletedDispatcherRequests;
+	{
+		UE::TUniqueLock Lock(Mutex);
+		LocalCompletedDispatcherRequests = MoveTemp(CompletedDispatcherRequests);
+		CompletedDispatcherRequests = FIoRequestList();
+	}
+
+	for (FIoRequestImpl& DispatcherRequest : LocalCompletedDispatcherRequests)
+	{
+		FResolvedRequest& ResolvedRequest = FResolvedRequest::Get(DispatcherRequest);
+		FPlatformIoDispatcher::Get().DeleteRequest(ResolvedRequest.PlatformRequest);
+		check(ResolvedRequest.ChunkInfo.GetContainer().ActiveReadCount > 0);
+		ResolvedRequest.ChunkInfo.GetContainer().ActiveReadCount.fetch_sub(1, std::memory_order_relaxed);
+
+		const bool bSucceeded = !DispatcherRequest.IsFailed() && !DispatcherRequest.IsCancelled();
+		check(!bSucceeded || ResolvedRequest.Buffer.GetSize() > 0 || DispatcherRequest.GetBuffer().GetSize() > 0);
+		if (bSucceeded)
+		{
+			if (DispatcherRequest.HasBuffer() == false)
+			{
+				DispatcherRequest.SetResult(MoveTemp(ResolvedRequest.Buffer));
+			}
+		}
+
+		RequestAllocator.Destroy(&ResolvedRequest);
+		DispatcherRequest.BackendData = nullptr;
+	}
+
+	return LocalCompletedDispatcherRequests.GetHead();
+}
+
+void FFileIoDispatcherBackend::CancelIoRequest(FIoRequestImpl* DispatcherRequest)
+{
+	if (FResolvedRequest* ResolvedRequest = FResolvedRequest::TryGet(DispatcherRequest))
+	{
+		FPlatformIoDispatcher::Get().CancelRequest(ResolvedRequest->PlatformRequest);
+	}
+}
+
+void FFileIoDispatcherBackend::UpdatePriorityForIoRequest(FIoRequestImpl* DispatcherRequest)
+{
+	if (FResolvedRequest* ResolvedRequest = FResolvedRequest::TryGet(DispatcherRequest))
+	{
+		const EIoFileReadPriority NewPriority = IoFileReadPriorityFromDispatcherPriority(DispatcherRequest->Priority);
+		if (static_cast<uint8>(NewPriority) > static_cast<uint8>(ResolvedRequest->Priority))
+		{
+			FPlatformIoDispatcher::Get().UpdatePriority(ResolvedRequest->PlatformRequest, NewPriority);
+		}
+	}
+}
+
+bool FFileIoDispatcherBackend::DoesChunkExist(const FIoChunkId& ChunkId) const
+{
+	FReadScopeLock ReadScope(IoStore.GetLock());
+
+	const FChunkInfo ChunkInfo = IoStore.GetChunkInfo(ChunkId);
+	return ChunkInfo.IsValid();
+}
+
+TIoStatusOr<uint64> FFileIoDispatcherBackend::GetSizeForChunk(const FIoChunkId& ChunkId) const
+{
+	FReadScopeLock ReadScope(IoStore.GetLock());
+
+	if (const FChunkInfo ChunkInfo = IoStore.GetChunkInfo(ChunkId); ChunkInfo.IsValid())
+	{
+		return ChunkInfo.Size();
+	}
+
+	return FIoStatus::Unknown;
+}
+
+TIoStatusOr<FIoMappedRegion> FFileIoDispatcherBackend::OpenMapped(const FIoChunkId& ChunkId, const FIoReadOptions& Options)
+{
+	if (!FPlatformProperties::SupportsMemoryMappedFiles())
+	{
+		return FIoStatus(EIoErrorCode::Unknown, TEXT("Platform does not support memory mapped files"));
+	}
+
+	if (Options.GetTargetVa() != nullptr)
+	{
+		return FIoStatus(EIoErrorCode::InvalidParameter, TEXT("Invalid read options"));
+	}
+
+	FWriteScopeLock WriteScope(IoStore.GetLock()); // In case a new mapped file handle is created
+	FChunkInfo ChunkInfo = IoStore.GetChunkInfo(ChunkId);
+	if (ChunkInfo.IsValid() == false)
+	{
+		return FIoStatus(EIoErrorCode::NotFound);
+	}
+
+	const int64 ResolvedOffset	= ChunkInfo.Offset() + Options.GetOffset();
+	const int64 ResolvedSize	= FMath::Min(Options.GetSize(), ChunkInfo.Size() - Options.GetOffset());
+	FContainer& Container		= ChunkInfo.GetContainer();
+	const int32 BlockIndex		= IntCastChecked<int32>(ResolvedOffset / Container.CompressionBlockSize);
+
+	const FIoStoreTocCompressedBlockEntry& Block	= Container.CompressionBlocks[BlockIndex];
+	uint64 BlockOffsetInPartition					= MAX_uint64;
+	FContainerPartition& Partition					= Container.GetPartition(Block.GetOffset(), BlockOffsetInPartition);
+
+	check(IsAligned(BlockOffsetInPartition, FPlatformProperties::GetMemoryMappingAlignment()));
+
+	if (Partition.MappedFileHandle.IsValid() == false)
+	{
+		IPlatformFile& Ipf = FPlatformFileManager::Get().GetPlatformFile();
+		if (FOpenMappedResult Result = Ipf.OpenMappedEx(*Partition.Filename); !Result.HasError())
+		{
+			Partition.MappedFileHandle = Result.StealValue();
+		}
+	}
+
+	if (Partition.MappedFileHandle.IsValid() == false)
+	{
+		return FIoStatus(EIoErrorCode::FileOpenFailed);
+	}
+
+	IMappedFileHandle& MappedFileHandle = *Partition.MappedFileHandle.Get(); 
+	if (IMappedFileRegion* MappedFileRegion = MappedFileHandle.MapRegion(BlockOffsetInPartition + Options.GetOffset(), ResolvedSize))
+	{
+		return FIoMappedRegion
+		{
+			.MappedFileHandle	= Partition.MappedFileHandle.Get(),
+			.MappedFileRegion	= MappedFileRegion
+		};
+	}
+
+	return FIoStatus(EIoErrorCode::FileOpenFailed);
+}
+
+void FFileIoDispatcherBackend::HandleSignatureError(FIoRequestImpl& DispatcherRequest, uint32 FailedBlockIndex)
+{
+	FIoSignatureError SignatureError;
+	{
+		FWriteScopeLock _(IoStore.GetLock());
+
+		FResolvedRequest& ResolvedRequest				= FResolvedRequest::Get(DispatcherRequest);
+		const FContainer& Container						= ResolvedRequest.ChunkInfo.GetContainer();
+		const FIoStoreTocCompressedBlockEntry& Block	= Container.CompressionBlocks[FailedBlockIndex];
+
+		int32 PartIdx = 0; 
+		for (const FContainerPartition& Part : Container.Partitions)
+		{
+			if (Part.FileHandle.Value() == ResolvedRequest.FileHandle.Value())
+			{
+				break;
+			}
+			PartIdx++;
+		}
+
+		UE_LOG(LogIoStore, Warning, TEXT("Signature error detected, ChunkId='%s', Filename='%s', Offset=%lu"),
+				*LexToString(DispatcherRequest.ChunkId), *Container.Partitions[PartIdx].Filename, Block.GetOffset());
+
+		SignatureError.ContainerName	= Container.BaseFilePath; 
+		SignatureError.BlockIndex		= FailedBlockIndex; 
+		SignatureError.ExpectedHash		= Container.CompressionBlockHashes[FailedBlockIndex];
+		//SignatureError.ActualHash		= Is this really needed?
+	}
+
+	check(BackendContext);
+	if (BackendContext->SignatureErrorDelegate.IsBound())
+	{
+		BackendContext->SignatureErrorDelegate.Broadcast(SignatureError);
+	}
+}
+
+////////////////////////////////////////////////////////////////////////////////
+TSharedPtr<IFileIoDispatcherBackend> MakeFileIoDispatcherBackend()
+{
+	return MakeShared<FFileIoDispatcherBackend>();
+}
+
+} // namespace UE::IoStore
diff --git a/Engine/Source/Runtime/PakFile/Private/FileIoDispatcherBackend.h b/Engine/Source/Runtime/PakFile/Private/FileIoDispatcherBackend.h
--- a/Engine/Source/Runtime/PakFile/Private/FileIoDispatcherBackend.h
+++ b/Engine/Source/Runtime/PakFile/Private/FileIoDispatcherBackend.h
@@ -22,4 +22,6 @@ public:
 	virtual void							ReopenAllFileHandles() = 0;
 };
 
+TSharedPtr<IFileIoDispatcherBackend> MakeFileIoDispatcherBackend();
+
 } // namespace UE::IoStore
diff --git a/Engine/Source/Runtime/PakFile/Private/IPlatformFilePak.cpp b/Engine/Source/Runtime/PakFile/Private/IPlatformFilePak.cpp
--- a/Engine/Source/Runtime/PakFile/Private/IPlatformFilePak.cpp
+++ b/Engine/Source/Runtime/PakFile/Private/IPlatformFilePak.cpp
@@ -2,6 +2,7 @@
 
 #include "IPlatformFilePak.h"
 #include "HAL/FileManager.h"
+#include "IO/PlatformIoDispatcher.h"
 #include "Math/GuardedInt.h"
 #include "Misc/CoreMisc.h"
 #include "Misc/CommandLine.h"
@@ -35,6 +36,7 @@
 #include "ProfilingDebugging/ScopedTimers.h"
 #include "Async/MappedFileHandle.h"
 #include "IoDispatcherFileBackend.h"
+#include "FileIoDispatcherBackend.h"
 #include "Misc/PackageName.h"
 #include "Misc/PathViews.h"
 
@@ -5640,7 +5642,15 @@ bool FPakPlatformFile::Initialize(IPlatformFile* Inner, const TCHAR* CmdLine)
 			}
 
 			FIoDispatcher& IoDispatcher = FIoDispatcher::Get();
-			IoDispatcherFileBackend = CreateIoDispatcherFileBackend();
+			// Use the new file I/O store backend if the platform I/O dispatcher is enabled
+			if (UE::IPlatformIoDispatcher* PlatformIoDispatcher = UE::FPlatformIoDispatcher::TryGet())
+			{
+				IoDispatcherFileBackend = UE::IoStore::MakeFileIoDispatcherBackend();
+			}
+			else
+			{
+				IoDispatcherFileBackend = CreateIoDispatcherFileBackend();
+			}
 			IoDispatcher.Mount(IoDispatcherFileBackend.ToSharedRef());
 			PackageStoreBackend = MakeShared<FFilePackageStoreBackend>();
 			FPackageStore::Get().Mount(PackageStoreBackend.ToSharedRef());
diff --git a/Engine/Source/Runtime/PakFile/Public/IPlatformFilePak.h b/Engine/Source/Runtime/PakFile/Public/IPlatformFilePak.h
--- a/Engine/Source/Runtime/PakFile/Public/IPlatformFilePak.h
+++ b/Engine/Source/Runtime/PakFile/Public/IPlatformFilePak.h
@@ -62,6 +62,7 @@ class FOutputDevice;
 class IAsyncReadFileHandle;
 class IMappedFileHandle;
 namespace UE::PakFile::Private { class FPakFileDirectoryVisitorBase; }
+namespace UE::IoStore { class IFileIoDispatcherBackend; }
 struct FIoContainerHeader;
 
 PAKFILE_API DECLARE_LOG_CATEGORY_EXTERN(LogPakFile, Log, All);
@@ -1985,7 +1986,7 @@ class FPakPlatformFile : public IPlatformFile
 	FString IniFileExtension;
 	/** The filename for the gameusersettings ini file, used for excluding ini files, but not gameusersettings */
 	FString GameUserSettingsIniFilename;
-	TSharedPtr<FFileIoStore> IoDispatcherFileBackend;
+	TSharedPtr<UE::IoStore::IFileIoDispatcherBackend> IoDispatcherFileBackend;
 	TSharedPtr<FFilePackageStoreBackend> PackageStoreBackend;
 
 	FTSTicker::FDelegateHandle RetireReadersHandle;
