Move LexToString and operator overloads to global namespace.

#rnx

diff --git a/Engine/Source/Runtime/Core/Internal/IO/PlatformIoDispatcher.h b/Engine/Source/Runtime/Core/Internal/IO/PlatformIoDispatcher.h
--- a/Engine/Source/Runtime/Core/Internal/IO/PlatformIoDispatcher.h
+++ b/Engine/Source/Runtime/Core/Internal/IO/PlatformIoDispatcher.h
@@ -32,36 +32,6 @@ enum class EIoFileReadPriority : uint8
 	Count
 };
 
-inline bool operator>(EIoFileReadPriority LHS, EIoFileReadPriority RHS)
-{
-	return static_cast<uint8>(LHS) > static_cast<uint8>(RHS);
-}
-
-inline bool operator<(EIoFileReadPriority LHS, EIoFileReadPriority RHS)
-{
-	return static_cast<uint8>(LHS) < static_cast<uint8>(RHS);
-}
-
-inline bool operator<=(EIoFileReadPriority LHS, EIoFileReadPriority RHS)
-{
-	return static_cast<uint8>(LHS) <= static_cast<uint8>(RHS);
-}
-
-/** File read priority to string. */
-inline const TCHAR* LexToString(EIoFileReadPriority Priority)
-{
-	static const TCHAR* StringTable[] =
-	{
-		TEXT("Min"),
-		TEXT("Low"),
-		TEXT("Medium"),
-		TEXT("High"),
-		TEXT("Max"),
-	};
-
-	return StringTable[static_cast<uint8>(Priority)];
-};
-
 /** Narrows an I/O dispatcher priorty to file read priority. */
 static inline EIoFileReadPriority IoFileReadPriorityFromDispatcherPriority(int32 Priority)
 {
@@ -306,3 +276,33 @@ public:
 };
 
 } // namespace UE
+
+////////////////////////////////////////////////////////////////////////////////
+static inline bool operator>(UE::EIoFileReadPriority LHS, UE::EIoFileReadPriority RHS)
+{
+	return static_cast<uint8>(LHS) > static_cast<uint8>(RHS);
+}
+
+static inline bool operator<(UE::EIoFileReadPriority LHS, UE::EIoFileReadPriority RHS)
+{
+	return static_cast<uint8>(LHS) < static_cast<uint8>(RHS);
+}
+
+static inline bool operator<=(UE::EIoFileReadPriority LHS, UE::EIoFileReadPriority RHS)
+{
+	return static_cast<uint8>(LHS) <= static_cast<uint8>(RHS);
+}
+
+static inline const TCHAR* LexToString(UE::EIoFileReadPriority Priority)
+{
+	static const TCHAR* StringTable[] =
+	{
+		TEXT("Min"),
+		TEXT("Low"),
+		TEXT("Medium"),
+		TEXT("High"),
+		TEXT("Max"),
+	};
+
+	return StringTable[static_cast<uint8>(Priority)];
+};
