Updated Cosmetics Bundling commandlet to produce IAD assets for itemdefinition dependencies.

diff --git a/Engine/Source/Runtime/Core/Private/IO/PackageId.cpp b/Engine/Source/Runtime/Core/Private/IO/PackageId.cpp
--- a/Engine/Source/Runtime/Core/Private/IO/PackageId.cpp
+++ b/Engine/Source/Runtime/Core/Private/IO/PackageId.cpp
@@ -37,7 +37,7 @@ FPackageId FPackageId::FromName(const FName& Name)
 		FWriteScopeLock ScopeWriteLock(PackageIdImpl::Lock);
 		LLM_SCOPE_BYTAG(PackageId_ReverseMapping);
 		FName EntryName = PackageIdImpl::Entries.FindOrAdd(Hash, Name);
-		checkf(EntryName == Name, TEXT("FPackageId collision: %llu for both %s and %s"), Hash, *Name.ToString(), *EntryName.ToString());
+		checkf(EntryName.GetDisplayIndex() == Name.GetDisplayIndex() || EntryName.GetComparisonIndex() == Name.GetComparisonIndex(), TEXT("FPackageId collision: %llu for both %s and %s"), Hash, *Name.ToString(), *EntryName.ToString());
 	}
 #endif
 	return FPackageId(Hash);
