FShader::NumInstructions should be Editor Only
Construction already only zeros this WITH_EDITORONLY_DATA

#tests Material Editor Instr Count Displays

diff --git a/Engine/Plugins/Compositing/OpenColorIO/Source/OpenColorIO/Private/OpenColorIODerivedDataVersion.h b/Engine/Plugins/Compositing/OpenColorIO/Source/OpenColorIO/Private/OpenColorIODerivedDataVersion.h
--- a/Engine/Plugins/Compositing/OpenColorIO/Source/OpenColorIO/Private/OpenColorIODerivedDataVersion.h
+++ b/Engine/Plugins/Compositing/OpenColorIO/Source/OpenColorIO/Private/OpenColorIODerivedDataVersion.h
@@ -10,4 +10,4 @@ OpenColorIOShaderDerivedDataVersion.h: Derived data version for OpenColorIO shad
 
 // In case of merge conflicts with DDC versions, you MUST generate a new GUID and set this new
 // guid as version
-#define OPENCOLORIO_DERIVEDDATA_VER		TEXT("1D868D37-BB45-4052-A0D7-85FFDD870812")
\ No newline at end of file
+#define OPENCOLORIO_DERIVEDDATA_VER		TEXT("A03A807F-923B-48AE-AA2A-77A1E51A6394")
diff --git a/Engine/Plugins/Runtime/ComputeFramework/Source/ComputeFramework/Private/ComputeFramework/ComputeKernelDerivedDataVersion.h b/Engine/Plugins/Runtime/ComputeFramework/Source/ComputeFramework/Private/ComputeFramework/ComputeKernelDerivedDataVersion.h
--- a/Engine/Plugins/Runtime/ComputeFramework/Source/ComputeFramework/Private/ComputeFramework/ComputeKernelDerivedDataVersion.h
+++ b/Engine/Plugins/Runtime/ComputeFramework/Source/ComputeFramework/Private/ComputeFramework/ComputeKernelDerivedDataVersion.h
@@ -5,4 +5,4 @@
 #include "CoreMinimal.h"
 
 // In case of merge conflicts with DDC versions, you MUST generate a new GUID and set this new guid as version
-#define COMPUTEKERNEL_DERIVEDDATA_VER	TEXT("F35F306A-EA41-4E0D-A354-06DEE736B126")
+#define COMPUTEKERNEL_DERIVEDDATA_VER	TEXT("9222DB83-5B24-4401-8E40-D5A995931D20")
diff --git a/Engine/Source/Runtime/Core/Private/UObject/FortniteMainBranchObjectVersion.cpp b/Engine/Source/Runtime/Core/Private/UObject/FortniteMainBranchObjectVersion.cpp
--- a/Engine/Source/Runtime/Core/Private/UObject/FortniteMainBranchObjectVersion.cpp
+++ b/Engine/Source/Runtime/Core/Private/UObject/FortniteMainBranchObjectVersion.cpp
@@ -6,11 +6,11 @@ TMap<FGuid, FGuid> FFortniteMainBranchObjectVersion::GetSystemGuids()
 	TMap<FGuid, FGuid> SystemGuids;
 	const FDevSystemGuids& DevGuids = FDevSystemGuids::Get();
 
-	SystemGuids.Add(DevGuids.GLOBALSHADERMAP_DERIVEDDATA_VER, FGuid("E1B72EA109BB4C2A995F0F65E5BD353D"));
+	SystemGuids.Add(DevGuids.GLOBALSHADERMAP_DERIVEDDATA_VER, FGuid("6B4FD8F6A6E74257B17D3DCF05ECD4EA"));
 	SystemGuids.Add(DevGuids.LANDSCAPE_MOBILE_COOK_VERSION, FGuid("32D02EF867C74B71A0D4E0FA41392732"));
-	SystemGuids.Add(DevGuids.MATERIALSHADERMAP_DERIVEDDATA_VER, FGuid("0E6202C151DE402CA451802671AE9347"));
+	SystemGuids.Add(DevGuids.MATERIALSHADERMAP_DERIVEDDATA_VER, FGuid("68CF693DB35E4ABDA0F053D785102A99"));
 	SystemGuids.Add(DevGuids.NANITE_DERIVEDDATA_VER, FGuid("42C46B24E1586986880AFE62B0F575AB"));
-	SystemGuids.Add(DevGuids.NIAGARASHADERMAP_DERIVEDDATA_VER, FGuid("3AE847DB7C504839B140718982B04087"));
+	SystemGuids.Add(DevGuids.NIAGARASHADERMAP_DERIVEDDATA_VER, FGuid("45D5978F33D940EEA6602B78559900DC"));
 	SystemGuids.Add(DevGuids.Niagara_LatestScriptCompileVersion, FGuid("670A2EBDEE1A4794BB77BBFF331EA7EF"));
 	SystemGuids.Add(DevGuids.SkeletalMeshDerivedDataVersion, FGuid("5E93192FCF3647EEB7F9D4FFF1750050"));
 	SystemGuids.Add(DevGuids.STATICMESH_DERIVEDDATA_VER, FGuid("656E338C78C34641A5D3BFAE303997AF"));
diff --git a/Engine/Source/Runtime/RenderCore/Public/Shader.h b/Engine/Source/Runtime/RenderCore/Public/Shader.h
--- a/Engine/Source/Runtime/RenderCore/Public/Shader.h
+++ b/Engine/Source/Runtime/RenderCore/Public/Shader.h
@@ -913,9 +913,9 @@ public:
 	inline EShaderFrequency GetFrequency() const { return Target.GetFrequency(); }
 	inline const FShaderTarget GetTarget() const { return Target; }
 	inline bool IsFrozen() const { return Type.IsFrozen(); }
-	inline uint32 GetNumInstructions() const { return NumInstructions; }
 
 #if WITH_EDITORONLY_DATA
+	inline uint32 GetNumInstructions() const { return NumInstructions; }
 	inline uint32 GetNumTextureSamplers() const { return NumTextureSamplers; }
 	inline uint32 GetCodeSize() const { return CodeSize; }
 	inline void SetNumInstructions(uint32 Value) { NumInstructions = Value; }
@@ -926,8 +926,10 @@ public:
 		return DummyStatMap;
 	}
 #else
+	inline uint32 GetNumInstructions() const { return 0u; }
 	inline uint32 GetNumTextureSamplers() const { return 0u; }
 	inline uint32 GetCodeSize() const { return 0u; }
+	inline void SetNumInstructions(uint32 Value) { }
 #endif
 
 	/** Finds an automatically bound uniform buffer matching the given uniform buffer type if one exists, or returns an unbound parameter. */
@@ -1024,7 +1026,7 @@ private:
 	LAYOUT_FIELD(int32, ResourceIndex);
 
 	/** The number of instructions the shader takes to execute. */
-	LAYOUT_FIELD(uint32, NumInstructions);
+	LAYOUT_FIELD_EDITORONLY(uint32, NumInstructions);
 
 	/** Truncated version of OutputHash, intended for sorting. Not suitable for unique shader identification. */
 	LAYOUT_FIELD(uint32, SortKey);
