Fix PVS studio warning in PlatformIoDispatcherBase

#rnx

diff --git a/Engine/Source/Runtime/Core/Internal/IO/PlatformIoDispatcherBase.cpp b/Engine/Source/Runtime/Core/Internal/IO/PlatformIoDispatcherBase.cpp
--- a/Engine/Source/Runtime/Core/Internal/IO/PlatformIoDispatcherBase.cpp
+++ b/Engine/Source/Runtime/Core/Internal/IO/PlatformIoDispatcherBase.cpp
@@ -636,14 +636,18 @@ bool FPlatformIoDispatcherRequestMgr::TryCancelReadRequest(FIoPlatformReadReques
 bool FPlatformIoDispatcherRequestMgr::TryCancelAllReadRequests(FIoFileHandle FileHandle)
 {
 	// Assumes file queue lock
-	bool bAnyCancelled = false;
+	TArray<FIoFileBlockRequest*> BlocksToCancel;
 	for (TTuple<FIoBlockKey, FIoFileBlockRequest*>& Kv : FileBlockLookup)
 	{
 		FIoFileBlockRequest* FileBlockRequest = Kv.Value;
-		if (FileBlockRequest->FileHandle != FileHandle || FileBlockRequest->QueueStatus == FIoFileBlockRequest::EQueueStatus::Dequeued)
+		if (FileBlockRequest->FileHandle == FileHandle && FileBlockRequest->QueueStatus == FIoFileBlockRequest::EQueueStatus::Enqueued)
 		{
-			continue;
+			BlocksToCancel.Add(FileBlockRequest);
 		}
+	}
+	
+	for (FIoFileBlockRequest* FileBlockRequest : BlocksToCancel)
+	{
 		for (FIoEncodedBlockRequest* EncodedBlockRequest : FileBlockRequest->EncodedBlockRequests)
 		{
 			for (FIoEncodedBlockRequest::FScatterTarget& ScatterTarget : EncodedBlockRequest->ScatterTargets)
@@ -657,10 +661,9 @@ bool FPlatformIoDispatcherRequestMgr::TryCancelAllReadRequests(FIoFileHandle Fil
 
 		FileBlockRequest->ErrorCode = EIoErrorCode::Cancelled;
 		FileBlockLookup.Remove(FileBlockRequest->BlockKey);
-		bAnyCancelled = true;
 	}
 
-	return bAnyCancelled;
+	return BlocksToCancel.IsEmpty() == false;
 }
 
 void FPlatformIoDispatcherRequestMgr::Destroy(FIoPlatformReadRequest& ReadRequest)
