Fixed macro typos.

diff --git a/Engine/Source/Runtime/Core/Public/Algo/Impl/BinaryHeap.h b/Engine/Source/Runtime/Core/Public/Algo/Impl/BinaryHeap.h
--- a/Engine/Source/Runtime/Core/Public/Algo/Impl/BinaryHeap.h
+++ b/Engine/Source/Runtime/Core/Public/Algo/Impl/BinaryHeap.h
@@ -2,7 +2,7 @@
 
 #pragma once
 
-#include "Templates/Invoke.h" // for older _MSVC_FULL_VER only - see usage below
+#include "Templates/Invoke.h" // for older _MSC_FULL_VER only - see usage below
 #include "Templates/Projection.h"
 #include "Templates/ReversePredicate.h"
 
@@ -59,7 +59,7 @@ namespace AlgoImpl
 	FORCEINLINE void HeapSiftDown(RangeValueType* Heap, IndexType Index, const IndexType Count, const ProjectionType& InProj, const PredicateType& Predicate)
 	{
 // Workaround for a codegen bug that was discovered related to member function pointer projections that are virtual
-#if defined(_MSVC_FULL_VER) && _MSVC_FULL_VER < 194134123
+#if defined(_MSC_FULL_VER) && _MSC_FULL_VER < 194134123
 		const ProjectionType& Proj = InProj;
 #else
 		auto&& Proj = Projection(InProj);
@@ -73,14 +73,18 @@ namespace AlgoImpl
 			if (RightChildIndex < Count)
 			{
 // Workaround for a codegen bug that was discovered related to member function pointer projections that are virtual
-#if defined(_MSVC_FULL_VER) && _MSVC_FULL_VER < 194134123
+#if defined(_MSC_FULL_VER) && _MSC_FULL_VER < 194134123
 				MinChildIndex = Predicate( Invoke(Proj, Heap[LeftChildIndex]), Invoke(Proj, Heap[RightChildIndex]) ) ? LeftChildIndex : RightChildIndex;
 #else
 				MinChildIndex = Predicate( Proj(Heap[LeftChildIndex]), Proj(Heap[RightChildIndex]) ) ? LeftChildIndex : RightChildIndex;
 #endif
 			}
 
+#if defined(_MSC_FULL_VER) && _MSC_FULL_VER < 194134123
+			if (!Predicate( Invoke(Proj, Heap[MinChildIndex]), Invoke(Proj, Heap[Index]) ))
+#else
 			if (!Predicate( Proj(Heap[MinChildIndex]), Proj(Heap[Index]) ))
+#endif
 			{
 				break;
 			}
@@ -105,7 +109,7 @@ namespace AlgoImpl
 	FORCEINLINE IndexType HeapSiftUp(RangeValueType* Heap, IndexType RootIndex, IndexType NodeIndex, const ProjectionType& InProj, const PredicateType& Predicate)
 	{
 // Workaround for a codegen bug that was discovered related to member function pointer projections that are virtual
-#if defined(_MSVC_FULL_VER) && _MSVC_FULL_VER < 194134123
+#if defined(_MSC_FULL_VER) && _MSC_FULL_VER < 194134123
 		const ProjectionType& Proj = InProj;
 #else
 		auto&& Proj = Projection(InProj);
@@ -114,8 +118,8 @@ namespace AlgoImpl
 		{
 			IndexType ParentIndex = HeapGetParentIndex(NodeIndex);
 // Workaround for a codegen bug that was discovered related to member function pointer projections that are virtual
-#if defined(_MSVC_FULL_VER) && _MSVC_FULL_VER < 194134123
-			if (!Predicate( Invoke(Projection, Heap[NodeIndex]), Invoke(Projection, Heap[ParentIndex]) ))
+#if defined(_MSC_FULL_VER) && _MSC_FULL_VER < 194134123
+			if (!Predicate( Invoke(Proj, Heap[NodeIndex]), Invoke(Proj, Heap[ParentIndex]) ))
 #else
 			if (!Predicate( Proj(Heap[NodeIndex]), Proj(Heap[ParentIndex]) ))
 #endif
diff --git a/Engine/Source/Runtime/Core/Public/Algo/IntroSort.h b/Engine/Source/Runtime/Core/Public/Algo/IntroSort.h
--- a/Engine/Source/Runtime/Core/Public/Algo/IntroSort.h
+++ b/Engine/Source/Runtime/Core/Public/Algo/IntroSort.h
@@ -5,7 +5,7 @@
 #include "Algo/Impl/BinaryHeap.h"
 #include "Math/UnrealMathUtility.h"
 #include "Templates/IdentityFunctor.h"
-#include "Templates/Invoke.h" // for older _MSVC_FULL_VER only - see usage below
+#include "Templates/Invoke.h" // for older _MSC_FULL_VER only - see usage below
 #include "Templates/Projection.h"
 #include "Templates/Less.h"
 #include "Templates/UnrealTemplate.h" // For GetData, GetNum
@@ -62,7 +62,7 @@ namespace AlgoImpl
 					for( Max=Current.Min, Item=Current.Min+1; Item<=Current.Max; Item++ )
 					{
 // Workaround for a codegen bug that was discovered related to member function pointer projections that are virtual
-#if defined(_MSVC_FULL_VER) && _MSVC_FULL_VER < 194134123
+#if defined(_MSC_FULL_VER) && _MSC_FULL_VER < 194134123
 						if( Invoke( Predicate, Invoke( Proj, *Max ), Invoke( Proj, *Item ) ) )
 #else
 						if( Predicate( Proj( *Max ), Proj( *Item ) ) )
@@ -85,7 +85,7 @@ namespace AlgoImpl
 				for( ; ; )
 				{
 // Workaround for a codegen bug that was discovered related to member function pointer projections that are virtual
-#if defined(_MSVC_FULL_VER) && _MSVC_FULL_VER < 194134123
+#if defined(_MSC_FULL_VER) && _MSC_FULL_VER < 194134123
 					while( ++Inner.Min<=Current.Max && !Invoke( Predicate, Invoke( Proj, *Current.Min ), Invoke( Proj, *Inner.Min ) ) );
 					while( --Inner.Max> Current.Min && !Invoke( Predicate, Invoke( Proj, *Inner.Max ), Invoke( Proj, *Current.Min ) ) );
 #else
@@ -173,7 +173,7 @@ namespace Algo
 	FORCEINLINE void IntroSortBy(RangeType&& Range, ProjectionType Proj)
 	{
 // Workaround for a codegen bug that was discovered related to member function pointer projections that are virtual
-#if defined(_MSVC_FULL_VER) && _MSVC_FULL_VER < 194134123
+#if defined(_MSC_FULL_VER) && _MSC_FULL_VER < 194134123
 		AlgoImpl::IntroSortInternal(GetData(Range), GetNum(Range), MoveTemp(Proj), TLess<>());
 #else
 		AlgoImpl::IntroSortInternal(GetData(Range), GetNum(Range), Projection(MoveTemp(Proj)), TLess<>());
@@ -191,7 +191,7 @@ namespace Algo
 	FORCEINLINE void IntroSortBy(RangeType&& Range, ProjectionType Proj, PredicateType Predicate)
 	{
 // Workaround for a codegen bug that was discovered related to member function pointer projections that are virtual
-#if defined(_MSVC_FULL_VER) && _MSVC_FULL_VER < 194134123
+#if defined(_MSC_FULL_VER) && _MSC_FULL_VER < 194134123
 		AlgoImpl::IntroSortInternal(GetData(Range), GetNum(Range), MoveTemp(Proj), MoveTemp(Predicate));
 #else
 		AlgoImpl::IntroSortInternal(GetData(Range), GetNum(Range), Projection(MoveTemp(Proj)), Projection(MoveTemp(Predicate)));
