Fix issues with event broadcast in the previous change due to a merge issue. Also fix some inadvertent whitespace changes

diff --git a/Engine/Source/Runtime/Core/Private/ProfilingDebugging/CsvProfiler.cpp b/Engine/Source/Runtime/Core/Private/ProfilingDebugging/CsvProfiler.cpp
--- a/Engine/Source/Runtime/Core/Private/ProfilingDebugging/CsvProfiler.cpp
+++ b/Engine/Source/Runtime/Core/Private/ProfilingDebugging/CsvProfiler.cpp
@@ -4274,23 +4274,22 @@ void FCsvProfiler::RecordEvent(int32 CategoryIndex, const FString& EventText)
 	{
 		LLM_SCOPE(ELLMTag::CsvProfiler);
 
+		FCsvProfilerThreadData::Get().AddEvent(EventText, CategoryIndex);
+
 		// Log the event and broadcast to delegates
-		if (CategoryIndex == CSV_CATEGORY_INDEX_GLOBAL)
+		FString CategoryName;
+		FString FullEventText = EventText;
+		if (CategoryIndex != CSV_CATEGORY_INDEX_GLOBAL)
 		{
-			UE_LOG(LogCsvProfiler, Display, TEXT("CSVEvent \"%s\" [Frame %d]"), *EventText, FCsvProfiler::Get()->GetCaptureFrameNumber());
-			FCsvProfiler::Get()->OnCSVProfileEvent().Broadcast(FString(), EventText);
-		}
-		else
-		{
-			FString CategoryName = FCsvCategoryData::Get()->GetCategoryNameByIndex(CategoryIndex);
-			UE_LOG(LogCsvProfiler, Display, TEXT("CSVEvent \"%s/%s\" [Frame %d]"), *CategoryName, *EventText, FCsvProfiler::Get()->GetCaptureFrameNumber());
-			FCsvProfiler::Get()->OnCSVProfileEvent().Broadcast(CategoryName, EventText);
+			CategoryName = FCsvCategoryData::Get()->GetCategoryNameByIndex(CategoryIndex);
+			FullEventText = FString::Printf(TEXT("%s/%s"), *CategoryName, *EventText);
 		}
+		UE_LOG(LogCsvProfiler, Display, TEXT("CSVEvent \"%s\" [Frame %d]"), *FullEventText, FCsvProfiler::Get()->GetCaptureFrameNumber());
+		FCsvProfiler::Get()->OnCSVProfileEvent().Broadcast(CategoryName, EventText);
 
-		FCsvProfilerThreadData::Get().AddEvent(EventText, CategoryIndex);
 		if (GCsvEventExecCmds)
 		{
-			FStringView Match(EventText);
+			FStringView Match(FullEventText);
 			for (FEventExecCmds& EventWildcard : *GCsvEventExecCmds)
 			{
 				if (FWildcardString::IsMatchSubstring(*EventWildcard.EventWildcard, Match.begin(), Match.end(), ESearchCase::IgnoreCase))
@@ -4638,8 +4637,8 @@ void FCsvProfiler::Init()
 		{
 			for (FString FrameAndCommand : CsvExecCommandsList)
 			{
-				int32 ColonIndex = -1;
-				if (FrameAndCommand.Find(TEXT("::"), ColonIndex))
+				int32 ColonIndex = FrameAndCommand.Find(TEXT("::"));
+				if (ColonIndex != INDEX_NONE)
 				{
 					FString EventStr = FrameAndCommand.Mid(0,ColonIndex);
 					FString CommandStr = FrameAndCommand.Mid(ColonIndex+2);
