Added missing include for FString.

diff --git a/Engine/Source/Runtime/Core/Tests/String/FormatStringSanTest.cpp b/Engine/Source/Runtime/Core/Tests/String/FormatStringSanTest.cpp
--- a/Engine/Source/Runtime/Core/Tests/String/FormatStringSanTest.cpp
+++ b/Engine/Source/Runtime/Core/Tests/String/FormatStringSanTest.cpp
@@ -7,6 +7,7 @@
 
 #include "String/FormatStringSan.h"
 #include "Containers/AnsiString.h"
+#include "Containers/UnrealString.h"
 #include "Containers/Utf8String.h"
 
 #include "Containers/EnumAsByte.h"
