Annotated TVariant with [[nodiscard]] and UE_LIFETIMEBOUND.
Fixed a bug which it found.

diff --git a/Engine/Source/Runtime/Core/Public/Misc/TVariant.h b/Engine/Source/Runtime/Core/Public/Misc/TVariant.h
--- a/Engine/Source/Runtime/Core/Public/Misc/TVariant.h
+++ b/Engine/Source/Runtime/Core/Public/Misc/TVariant.h
@@ -50,7 +50,7 @@ class TVariant final
 
 public:
 	/** Default initialize the TVariant to the first type in the parameter pack */
-	TVariant()
+	[[nodiscard]] TVariant()
 	{
 		static_assert(std::is_constructible_v<T>, "To default-initialize a TVariant, the first type in the parameter pack must be default constructible. Use FEmptyVariantState as the first type if none of the other types can be listed first.");
 		::new((void*)&UE::Core::Private::CastToStorage(*this).Storage) T();
@@ -63,7 +63,7 @@ public:
 		typename... ArgTypes
 		UE_REQUIRES(std::is_constructible_v<U, ArgTypes...>)
 	>
-	explicit TVariant(TInPlaceType<U>&&, ArgTypes&&... Args)
+	[[nodiscard]] explicit TVariant(TInPlaceType<U>&&, ArgTypes&&... Args)
 	{
 		constexpr SIZE_T Index = UE::Core::Private::TParameterPackTypeIndex<U, T, Ts...>::Value;
 		static_assert(Index != (SIZE_T)-1, "The TVariant is not declared to hold the type being constructed");
@@ -73,14 +73,14 @@ public:
 	}
 
 	/** Copy construct the variant from another variant of the same type */
-	TVariant(const TVariant& Other)
+	[[nodiscard]] TVariant(const TVariant& Other)
 	{
 		TypeIndex = Other.TypeIndex;
 		UE::Core::Private::TCopyConstructorLookup<T, Ts...>::Construct(TypeIndex, &UE::Core::Private::CastToStorage(*this).Storage, &UE::Core::Private::CastToStorage(Other).Storage);
 	}
 
 	/** Move construct the variant from another variant of the same type */
-	TVariant(TVariant&& Other)
+	[[nodiscard]] TVariant(TVariant&& Other)
 	{
 		TypeIndex = Other.TypeIndex;
 		UE::Core::Private::TMoveConstructorLookup<T, Ts...>::Construct(TypeIndex, &UE::Core::Private::CastToStorage(*this).Storage, &UE::Core::Private::CastToStorage(Other).Storage);
@@ -125,7 +125,7 @@ public:
 
 	/** Determine if the variant holds the specific type */
 	template <typename U>
-	bool IsType() const
+	[[nodiscard]] bool IsType() const
 	{
 		static_assert(UE::Core::Private::TParameterPackTypeIndex<U, T, Ts...>::Value != (SIZE_T)-1, "The TVariant is not declared to hold the type passed to IsType<>");
 		return UE::Core::Private::TIsType<U, T, Ts...>::IsSame(TypeIndex);
@@ -133,7 +133,7 @@ public:
 
 	/** Get a reference to the held value. Bad things can happen if this is called on a variant that does not hold the type asked for */
 	template <typename U>
-	U& Get()
+	[[nodiscard]] U& Get() UE_LIFETIMEBOUND
 	{
 		constexpr SIZE_T Index = UE::Core::Private::TParameterPackTypeIndex<U, T, Ts...>::Value;
 		static_assert(Index != (SIZE_T)-1, "The TVariant is not declared to hold the type passed to Get<>");
@@ -146,7 +146,7 @@ public:
 
 	/** Get a reference to the held value. Bad things can happen if this is called on a variant that does not hold the type asked for */
 	template <typename U>
-	const U& Get() const
+	[[nodiscard]] const U& Get() const UE_LIFETIMEBOUND
 	{
 		// Temporarily remove the const qualifier so we can implement Get in one location.
 		return const_cast<TVariant*>(this)->template Get<U>();
@@ -161,7 +161,7 @@ public:
 
 	/** Get a pointer to the held value if the held type is the same as the one specified */
 	template <typename U>
-	U* TryGet()
+	[[nodiscard]] U* TryGet() UE_LIFETIMEBOUND
 	{
 		constexpr SIZE_T Index = UE::Core::Private::TParameterPackTypeIndex<U, T, Ts...>::Value;
 		static_assert(Index != (SIZE_T)-1, "The TVariant is not declared to hold the type passed to TryGet<>");
@@ -172,7 +172,7 @@ public:
 
 	/** Get a pointer to the held value if the held type is the same as the one specified */
 	template <typename U>
-	const U* TryGet() const
+	[[nodiscard]] const U* TryGet() const UE_LIFETIMEBOUND
 	{
 		// Temporarily remove the const qualifier so we can implement TryGet in one location.
 		return const_cast<TVariant*>(this)->template TryGet<U>();
@@ -210,7 +210,7 @@ public:
 
 	/** Lookup the index of a type in the template parameter pack at compile time. */
 	template <typename U>
-	static constexpr SIZE_T IndexOfType()
+	[[nodiscard]] static constexpr SIZE_T IndexOfType()
 	{
 		constexpr SIZE_T Index = UE::Core::Private::TParameterPackTypeIndex<U, T, Ts...>::Value;
 		static_assert(Index != (SIZE_T)-1, "The TVariant is not declared to hold the type passed to IndexOfType<>");
@@ -218,7 +218,7 @@ public:
 	}
 
 	/** Returns the currently held type's index into the template parameter pack */
-	SIZE_T GetIndex() const
+	[[nodiscard]] SIZE_T GetIndex() const
 	{
 		return (SIZE_T)TypeIndex;
 	}
