Fixed "non-void function should return a value" error.

diff --git a/Engine/Source/Runtime/Core/Public/Containers/StringOverload.h b/Engine/Source/Runtime/Core/Public/Containers/StringOverload.h
--- a/Engine/Source/Runtime/Core/Public/Containers/StringOverload.h
+++ b/Engine/Source/Runtime/Core/Public/Containers/StringOverload.h
@@ -54,6 +54,7 @@ private:
 		else
 		{
 			static_assert(sizeof(ArgType) == 0, "Unsupported argument type passed to TStringOverload");
+			return false;
 		}
 	}
 
