Logging: Improvements to structured logging

- Fixed fatal structured logs in configurations with no logging.
- Moved LogIfActive into the macros to avoid defining static state for logs that are compiled out.
- Added [[noreturn]] to the fatal logging functions to give the compiler that knowledge.
- Fixed NO_LOGGING mismatches between declarations and definitions.
- Fixed field validation in non-fatal localized structured logs.
- Removed UE_AUTORTFM_ALWAYS_OPEN on functions with no side effects.

diff --git a/Engine/Source/Runtime/Core/Private/Logging/StructuredLog.cpp b/Engine/Source/Runtime/Core/Private/Logging/StructuredLog.cpp
--- a/Engine/Source/Runtime/Core/Private/Logging/StructuredLog.cpp
+++ b/Engine/Source/Runtime/Core/Private/Logging/StructuredLog.cpp
@@ -208,11 +208,11 @@ static void AddFieldValue(FFormatNamedArguments& Out, const FCbFieldView& Field)
 
 } // UE::Logging::Private
 
+///////////////////////////////////////////////////////////////////////////////////////////////////////////////////////
+
 namespace UE
 {
 
-///////////////////////////////////////////////////////////////////////////////////////////////////////////////////////
-
 class FLogTemplate
 {
 	using FLogField = Logging::Private::FLogField;
@@ -667,15 +667,17 @@ void FLogRecord::FormatMessageTo(FWideStringBuilderBase& Out) const
 	FormatRecordMessageTo(Out, *this);
 }
 
-///////////////////////////////////////////////////////////////////////////////////////////////////////////////////////
-
 } // UE
 
-#if !NO_LOGGING
+///////////////////////////////////////////////////////////////////////////////////////////////////////////////////////
 
 namespace UE::Logging::Private
 {
 
+#if NO_LOGGING
+FLogCategory<ELogVerbosity::Fatal, ELogVerbosity::Fatal> LogFatal(TEXT("Fatal"));
+#endif
+
 thread_local FLogContext* LogContextHead = nullptr;
 thread_local FLogContext* LogContextTail = nullptr;
 
@@ -794,16 +796,16 @@ struct FStaticLogTemplateManager
 
 static FStaticLogTemplateManager GStaticLogTemplateManager;
 
+#if LOGTRACE_ENABLED
 // Tracing the log happens in its own function because that allows stack space for the message to
 // be returned before calling into the output devices.
 FORCENOINLINE static void LogToTrace(const void* LogPoint, const FLogRecord& Record)
 {
-#if LOGTRACE_ENABLED
 	TStringBuilder<1024> Message;
 	Record.FormatMessageTo(Message);
 	FLogTrace::OutputLogMessage(LogPoint, *Message);
-#endif
 }
+#endif
 
 // Serializing log fields to compact binary happens in its own function because that allows stack
 // space for the writer to be returned before calling into the output devices.
@@ -974,6 +976,8 @@ inline static void DispatchLogRecord(const FLogRecord& Record)
 	(OutputDevice ? OutputDevice : GLog)->SerializeRecord(Record);
 }
 
+#if !NO_LOGGING
+
 template <typename StaticLogRecordType>
 inline static void DispatchStaticLogRecord(const StaticLogRecordType& Log, const FLogRecord& Record)
 {
@@ -990,44 +994,37 @@ inline static void DispatchStaticLogRecord(const StaticLogRecordType& Log, const
 UE_AUTORTFM_ALWAYS_OPEN
 void LogWithFieldArray(const FLogCategoryBase& Category, const FStaticLogRecord& Log, const FLogField* Fields, const int32 FieldCount)
 {
-#if !NO_LOGGING
 	DispatchStaticLogRecord(Log, CreateLogRecord(Category, Log, Fields, FieldCount));
-#endif
 }
 
-UE_AUTORTFM_ALWAYS_OPEN
 void LogWithNoFields(const FLogCategoryBase& Category, const FStaticLogRecord& Log)
 {
-#if !NO_LOGGING
 	// A non-null field pointer enables field validation in FLogTemplate::Create.
 	static constexpr FLogField EmptyField{};
 	LogWithFieldArray(Category, Log, &EmptyField, 0);
-#endif
 }
 
 UE_AUTORTFM_ALWAYS_OPEN
 void LogWithFieldArray(const FLogCategoryBase& Category, const FStaticLocalizedLogRecord& Log, const FLogField* Fields, const int32 FieldCount)
 {
-#if !NO_LOGGING
 	FLogRecord Record = CreateLogRecord(Category, Log, Fields, FieldCount);
 	Record.SetTextNamespace(Log.TextNamespace);
 	Record.SetTextKey(Log.TextKey);
 	DispatchStaticLogRecord(Log, Record);
-#endif
 }
 
-UE_AUTORTFM_ALWAYS_OPEN
 void LogWithNoFields(const FLogCategoryBase& Category, const FStaticLocalizedLogRecord& Log)
 {
-#if !NO_LOGGING
-	LogWithFieldArray(Category, Log, nullptr, 0);
-#endif
+	// A non-null field pointer enables field validation in FLogTemplate::Create.
+	static constexpr FLogField EmptyField{};
+	LogWithFieldArray(Category, Log, &EmptyField, 0);
 }
 
+#endif // !NO_LOGGING
+
 UE_AUTORTFM_ALWAYS_OPEN
-void FatalLogWithFieldArray(const FLogCategoryBase& Category, const FStaticLogRecord& Log, const FLogField* Fields, const int32 FieldCount)
+[[noreturn]] void FatalLogWithFieldArray(const FLogCategoryBase& Category, const FStaticLogRecord& Log, const FLogField* Fields, const int32 FieldCount)
 {
-#if !NO_LOGGING
 	FLogRecord Record = CreateLogRecord(Category, Log, Fields, FieldCount);
 	TStringBuilder<512> Message;
 	Record.FormatMessageTo(Message);
@@ -1036,23 +1033,20 @@ void FatalLogWithFieldArray(const FLogCategoryBase& Category, const FStaticLogRe
 
 	UE_DEBUG_BREAK_AND_PROMPT_FOR_REMOTE();
 	FDebug::ProcessFatalError(PLATFORM_RETURN_ADDRESS());
-#endif
+
+	for (;;);
 }
 
-UE_AUTORTFM_ALWAYS_OPEN
 void FatalLogWithNoFields(const FLogCategoryBase& Category, const FStaticLogRecord& Log)
 {
-#if !NO_LOGGING
 	// A non-null field pointer enables field validation in FLogTemplate::Create.
 	static constexpr FLogField EmptyField{};
 	FatalLogWithFieldArray(Category, Log, &EmptyField, 0);
-#endif
 }
 
 UE_AUTORTFM_ALWAYS_OPEN
-void FatalLogWithFieldArray(const FLogCategoryBase& Category, const FStaticLocalizedLogRecord& Log, const FLogField* Fields, const int32 FieldCount)
+[[noreturn]] void FatalLogWithFieldArray(const FLogCategoryBase& Category, const FStaticLocalizedLogRecord& Log, const FLogField* Fields, const int32 FieldCount)
 {
-#if !NO_LOGGING
 	FLogRecord Record = CreateLogRecord(Category, Log, Fields, FieldCount);
 	TStringBuilder<512> Message;
 	Record.FormatMessageTo(Message);
@@ -1061,17 +1055,15 @@ void FatalLogWithFieldArray(const FLogCategoryBase& Category, const FStaticLocal
 
 	UE_DEBUG_BREAK_AND_PROMPT_FOR_REMOTE();
 	FDebug::ProcessFatalError(PLATFORM_RETURN_ADDRESS());
-#endif
+
+	for (;;);
 }
 
-UE_AUTORTFM_ALWAYS_OPEN
 void FatalLogWithNoFields(const FLogCategoryBase& Category, const FStaticLocalizedLogRecord& Log)
 {
-#if !NO_LOGGING
 	// A non-null field pointer enables field validation in FLogTemplate::Create.
 	static constexpr FLogField EmptyField{};
 	FatalLogWithFieldArray(Category, Log, &EmptyField, 0);
-#endif
 }
 
 } // UE::Logging::Private
@@ -1292,5 +1284,3 @@ void VisitLogContext(TFunctionRef<void (const FCbField&)> Visitor)
 }
 
 } // UE
-
-#endif // !NO_LOGGING
diff --git a/Engine/Source/Runtime/Core/Public/Logging/StructuredLog.h b/Engine/Source/Runtime/Core/Public/Logging/StructuredLog.h
--- a/Engine/Source/Runtime/Core/Public/Logging/StructuredLog.h
+++ b/Engine/Source/Runtime/Core/Public/Logging/StructuredLog.h
@@ -23,8 +23,6 @@ template <typename FuncType> class TFunctionRef;
 
 #define UE_API CORE_API
 
-#if !NO_LOGGING
-
 /**
  * Records a structured log event if this category is active at this level of verbosity.
  *
@@ -63,10 +61,10 @@ template <typename FuncType> class TFunctionRef;
  * UE_LOGFMT_EX(LogCore, Warning, "Loading '{Name}' failed with error {Error}",
  *     UE_LOGFMT_FIELD("Name", Package->GetName()), UE_LOGFMT_FIELD("Error", ErrorCode), UE_LOGFMT_FIELD("Flags", LoadFlags));
  */
-#define UE_LOGFMT_EX(CategoryName, Verbosity, Format, ...) UE_PRIVATE_LOGFMT(::UE::Logging::Private::FTrueCallable(), CategoryName, Verbosity, Format, ##__VA_ARGS__)
+#define UE_LOGFMT_EX(CategoryName, Verbosity, Format, ...) UE_PRIVATE_LOGFMT(UE_EMPTY, CategoryName, Verbosity, Format, ##__VA_ARGS__)
 
 /** Conditional UE_LOGFMT_EX. */
-#define UE_CLOGFMT_EX(Condition, CategoryName, Verbosity, Format, ...) UE_PRIVATE_LOGFMT([&] { return (Condition) && ::UE::Logging::Private::FTrueCallable()(); }, CategoryName, Verbosity, Format, ##__VA_ARGS__)
+#define UE_CLOGFMT_EX(Condition, CategoryName, Verbosity, Format, ...) UE_PRIVATE_LOGFMT(if (Condition), CategoryName, Verbosity, Format, ##__VA_ARGS__)
 
 /**
  * Records a localized structured log event if this category is active at this level of verbosity.
@@ -110,13 +108,13 @@ template <typename FuncType> class TFunctionRef;
 #define UE_LOGFMT_LOC_EX(CategoryName, Verbosity, Key, Format, ...) \
 	UE_LOGFMT_NSLOC_EX(CategoryName, Verbosity, LOCTEXT_NAMESPACE, Key, Format, ##__VA_ARGS__)
 #define UE_LOGFMT_NSLOC_EX(CategoryName, Verbosity, Namespace, Key, Format, ...) \
-	UE_PRIVATE_LOGFMT_LOC(::UE::Logging::Private::FTrueCallable(), CategoryName, Verbosity, Namespace, Key, Format, ##__VA_ARGS__)
+	UE_PRIVATE_LOGFMT_LOC(UE_EMPTY, CategoryName, Verbosity, Namespace, Key, Format, ##__VA_ARGS__)
 
 /** Conditional UE_LOGFMT_LOC_EX. */
 #define UE_CLOGFMT_LOC_EX(Condition, CategoryName, Verbosity, Key, Format, ...) \
 	UE_CLOGFMT_NSLOC_EX(Condition, CategoryName, Verbosity, LOCTEXT_NAMESPACE, Key, Format, ##__VA_ARGS__)
 #define UE_CLOGFMT_NSLOC_EX(Condition, CategoryName, Verbosity, Namespace, Key, Format, ...) \
-	UE_PRIVATE_LOGFMT_LOC([&] { return (Condition) && ::UE::Logging::Private::FTrueCallable()(); }, CategoryName, Verbosity, Namespace, Key, Format, ##__VA_ARGS__)
+	UE_PRIVATE_LOGFMT_LOC(if (Condition), CategoryName, Verbosity, Namespace, Key, Format, ##__VA_ARGS__)
 
 /** Expands to a named structured log field. UE_LOGFMT_FIELD("Name", Value) */
 #define UE_LOGFMT_FIELD(Name, Value) UE_PRIVATE_LOGFMT_FIELD((Name, Value))
@@ -124,6 +122,8 @@ template <typename FuncType> class TFunctionRef;
 /** Expands to a structured log value. UE_LOGFMT_VALUE(Value) */
 #define UE_LOGFMT_VALUE(Value) Value
 
+#if !NO_LOGGING
+
 /**
  * Registers a context on the calling thread, by name, with an optional value.
  *
@@ -144,18 +144,7 @@ template <typename FuncType> class TFunctionRef;
 
 #else // #if !NO_LOGGING
 
-#define UE_LOGFMT(...)
-#define UE_LOGFMT_EX(...)
-
-#define UE_LOGFMT_LOC(...)
-#define UE_LOGFMT_NSLOC(...)
-#define UE_LOGFMT_LOC_EX(...)
-#define UE_LOGFMT_NSLOC_EX(...)
-
-#define UE_LOGFMT_FIELD(Name, Value)
-#define UE_LOGFMT_VALUE(Value)
-
-#define UE_LOG_CONTEXT(Name, Value)
+#define UE_LOG_CONTEXT(...)
 
 #endif
 
@@ -388,24 +377,17 @@ struct FLogField
 	}
 };
 
-} // UE::Logging::Private
-
 #if !NO_LOGGING
-
-namespace UE::Logging::Private
-{
-
 UE_API void LogWithNoFields(const FLogCategoryBase& Category, const FStaticLogRecord& Log);
 UE_API void LogWithFieldArray(const FLogCategoryBase& Category, const FStaticLogRecord& Log, const FLogField* Fields, int32 FieldCount);
-
-UE_API void FatalLogWithNoFields(const FLogCategoryBase& Category, const FStaticLogRecord& Log);
-UE_API void FatalLogWithFieldArray(const FLogCategoryBase& Category, const FStaticLogRecord& Log, const FLogField* Fields, int32 FieldCount);
-
 UE_API void LogWithNoFields(const FLogCategoryBase& Category, const FStaticLocalizedLogRecord& Log);
 UE_API void LogWithFieldArray(const FLogCategoryBase& Category, const FStaticLocalizedLogRecord& Log, const FLogField* Fields, int32 FieldCount);
+#endif
 
-UE_API void FatalLogWithNoFields(const FLogCategoryBase& Category, const FStaticLocalizedLogRecord& Log);
-UE_API void FatalLogWithFieldArray(const FLogCategoryBase& Category, const FStaticLocalizedLogRecord& Log, const FLogField* Fields, int32 FieldCount);
+[[noreturn]] UE_API void FatalLogWithNoFields(const FLogCategoryBase& Category, const FStaticLogRecord& Log);
+[[noreturn]] UE_API void FatalLogWithFieldArray(const FLogCategoryBase& Category, const FStaticLogRecord& Log, const FLogField* Fields, int32 FieldCount);
+[[noreturn]] UE_API void FatalLogWithNoFields(const FLogCategoryBase& Category, const FStaticLocalizedLogRecord& Log);
+[[noreturn]] UE_API void FatalLogWithFieldArray(const FLogCategoryBase& Category, const FStaticLocalizedLogRecord& Log, const FLogField* Fields, int32 FieldCount);
 
 /** Wrapper to identify field names interleaved with field values. */
 template <typename NameType>
@@ -453,78 +435,46 @@ struct FLogFieldCreator
 	}
 };
 
-template <typename T>
-struct TFieldArgType
-{
-	using Type = typename TCallTraits<T>::ParamType;
-};
-
-template <typename NameType>
-struct TFieldArgType<TLogFieldName<NameType>>
-{
-	using Type = TLogFieldName<NameType>;
-};
-
 /** Log with fields created from the arguments, which may be values or pairs of name/value. */
-template <typename... FieldArgTypes, typename LogType>
-UE_AUTORTFM_ALWAYS_OPEN UE_COLD UE_DEBUG_SECTION void LogWithFields(const FLogCategoryBase& Category, const LogType& Log, typename TFieldArgType<FieldArgTypes>::Type... FieldArgs)
+template <typename LogType, typename... FieldArgTypes>
+UE_COLD UE_DEBUG_SECTION void LogWithFields(const FLogCategoryBase& Category, const LogType& Log, FieldArgTypes&&... FieldArgs)
 {
-	constexpr int32 FieldCount = FLogFieldCreator::template GetCount<FieldArgTypes...>();
-	static_assert(FieldCount > 0);
-	FLogField Fields[FieldCount];
-	FLogFieldCreator::Create(Fields, (FieldArgTypes&&)FieldArgs...);
-	LogWithFieldArray(Category, Log, Fields, FieldCount);
+	if constexpr (sizeof...(FieldArgTypes) == 0)
+	{
+		LogWithNoFields(Category, Log);
+	}
+	else
+	{
+		constexpr int32 FieldCount = FLogFieldCreator::template GetCount<FieldArgTypes...>();
+		static_assert(FieldCount > 0);
+		FLogField Fields[FieldCount];
+		FLogFieldCreator::Create(Fields, (FieldArgTypes&&)FieldArgs...);
+		LogWithFieldArray(Category, Log, Fields, FieldCount);
+	}
 }
 
 /** Fatal log with fields created from the arguments, which may be values or pairs of name/value. */
-template <typename... FieldArgTypes, typename LogType>
-UE_AUTORTFM_ALWAYS_OPEN inline void UE_COLD FatalLogWithFields(const FLogCategoryBase& Category, const LogType& Log, typename TFieldArgType<FieldArgTypes>::Type... FieldArgs)
-{
-	constexpr int32 FieldCount = FLogFieldCreator::template GetCount<FieldArgTypes...>();
-	static_assert(FieldCount > 0);
-	FLogField Fields[FieldCount];
-	FLogFieldCreator::Create(Fields, (FieldArgTypes&&)FieldArgs...);
-	FatalLogWithFieldArray(Category, Log, Fields, FieldCount);
-}
-
-/** Log if the category is active at this level of verbosity. */
-template <typename LogCategoryType, ELogVerbosity::Type Verbosity, typename ConditionType, typename LogRecordType, typename... FieldArgTypes>
-UE_AUTORTFM_ALWAYS_OPEN inline void LogIfActive(const ConditionType& Condition, const LogCategoryType& Category, const LogRecordType& Log, FieldArgTypes&&... FieldArgs)
+template <typename LogType, typename... FieldArgTypes>
+[[noreturn]] UE_COLD UE_DEBUG_SECTION void FatalLogWithFields(const FLogCategoryBase& Category, const LogType& Log, FieldArgTypes&&... FieldArgs)
 {
-	static_assert((Verbosity & ELogVerbosity::VerbosityMask) < ELogVerbosity::NumVerbosity && Verbosity > 0, "Verbosity must be constant and in range.");
-	if constexpr ((Verbosity & ELogVerbosity::VerbosityMask) == ELogVerbosity::Fatal)
+	if constexpr (sizeof...(FieldArgTypes) == 0)
 	{
-		if (Condition())
-		{
-			if constexpr (sizeof...(FieldArgTypes) == 0)
-			{
-				FatalLogWithNoFields(Category, Log);
-			}
-			else
-			{
-				FatalLogWithFields<FieldArgTypes...>(Category, Log, (FieldArgTypes&&)FieldArgs...);
-			}
-			CA_ASSUME(false);
-		}
+		FatalLogWithNoFields(Category, Log);
 	}
-	else if constexpr (
-		(Verbosity & ELogVerbosity::VerbosityMask) <= ELogVerbosity::COMPILED_IN_MINIMUM_VERBOSITY &&
-		(Verbosity & ELogVerbosity::VerbosityMask) <= LogCategoryType::CompileTimeVerbosity)
+	else
 	{
-		if (!Category.IsSuppressed(Verbosity) && Condition())
-		{
-			if constexpr (sizeof...(FieldArgTypes) == 0)
-			{
-				LogWithNoFields(Category, Log);
-			}
-			else
-			{
-				LogWithFields<FieldArgTypes...>(Category, Log, (FieldArgTypes&&)FieldArgs...);
-			}
-		}
+		constexpr int32 FieldCount = FLogFieldCreator::template GetCount<FieldArgTypes...>();
+		static_assert(FieldCount > 0);
+		FLogField Fields[FieldCount];
+		FLogFieldCreator::Create(Fields, (FieldArgTypes&&)FieldArgs...);
+		FatalLogWithFieldArray(Category, Log, Fields, FieldCount);
 	}
 }
 
+#if NO_LOGGING
+UE_API extern FLogCategory<ELogVerbosity::Fatal, ELogVerbosity::Fatal> LogFatal;
+#endif
+
 struct FLogContext
 {
 	inline explicit FLogContext(const ANSICHAR* Name)
@@ -551,34 +501,58 @@ struct FLogContext
 	bool bDisabledOlderContext = false;
 };
 
-struct FTrueCallable
-{
-	inline constexpr bool operator()() const
-	{
-		return true;
-	}
-};
-
 } // UE::Logging::Private
 
+#if !NO_LOGGING
+#define UE_PRIVATE_LOG_CATEGORY(CategoryName) CategoryName
+#else
+#define UE_PRIVATE_LOG_CATEGORY(CategoryName) ::UE::Logging::Private::LogFatal
+#endif
+
 #define UE_PRIVATE_LOGFMT(Condition, CategoryName, Verbosity, Format, ...) \
 	do \
 	{ \
-		static ::UE::Logging::Private::FStaticLogDynamicData LOG_Dynamic; \
-		static constexpr ::UE::Logging::Private::FStaticLogRecord LOG_Static UE_PRIVATE_LOGFMT_AGGREGATE(TEXT(Format), __builtin_FILE(), __builtin_LINE(), ::ELogVerbosity::Verbosity, LOG_Dynamic); \
-		::UE::Logging::Private::LogIfActive<FLogCategory##CategoryName, ::ELogVerbosity::Verbosity>(Condition, CategoryName, LOG_Static, ##__VA_ARGS__); \
+		if constexpr ((::ELogVerbosity::Verbosity & ::ELogVerbosity::VerbosityMask) == ::ELogVerbosity::Fatal || \
+			((::ELogVerbosity::Verbosity & ELogVerbosity::VerbosityMask) <= ::ELogVerbosity::COMPILED_IN_MINIMUM_VERBOSITY && \
+			 (::ELogVerbosity::Verbosity & ELogVerbosity::VerbosityMask) <= UE_PRIVATE_LOG_CATEGORY(CategoryName).CompileTimeVerbosity)) \
+		{ \
+			static ::UE::Logging::Private::FStaticLogDynamicData LOG_Dynamic; \
+			static constexpr ::UE::Logging::Private::FStaticLogRecord LOG_Static UE_PRIVATE_LOGFMT_AGGREGATE(TEXT(Format), __builtin_FILE(), __builtin_LINE(), ::ELogVerbosity::Verbosity, LOG_Dynamic); \
+			UE_PRIVATE_LOGFMT_LOG_IF_ACTIVE(Condition, CategoryName, Verbosity, LOG_Static, ##__VA_ARGS__); \
+		} \
 	} \
 	while (false)
 
 #define UE_PRIVATE_LOGFMT_LOC(Condition, CategoryName, Verbosity, Namespace, Key, Format, ...) \
 	do \
 	{ \
-		static ::UE::Logging::Private::FStaticLogDynamicData LOG_Dynamic; \
-		static constexpr ::UE::Logging::Private::FStaticLocalizedLogRecord LOG_Static UE_PRIVATE_LOGFMT_AGGREGATE(TEXT(Namespace), TEXT(Key), TEXT(Format), __builtin_FILE(), __builtin_LINE(), ::ELogVerbosity::Verbosity, LOG_Dynamic); \
-		::UE::Logging::Private::LogIfActive<FLogCategory##CategoryName, ::ELogVerbosity::Verbosity>(Condition, CategoryName, LOG_Static, ##__VA_ARGS__); \
+		if constexpr ((::ELogVerbosity::Verbosity & ::ELogVerbosity::VerbosityMask) == ::ELogVerbosity::Fatal || \
+			((::ELogVerbosity::Verbosity & ELogVerbosity::VerbosityMask) <= ::ELogVerbosity::COMPILED_IN_MINIMUM_VERBOSITY && \
+			 (::ELogVerbosity::Verbosity & ELogVerbosity::VerbosityMask) <= UE_PRIVATE_LOG_CATEGORY(CategoryName).CompileTimeVerbosity)) \
+		{ \
+			static ::UE::Logging::Private::FStaticLogDynamicData LOG_Dynamic; \
+			static constexpr ::UE::Logging::Private::FStaticLocalizedLogRecord LOG_Static UE_PRIVATE_LOGFMT_AGGREGATE(TEXT(Namespace), TEXT(Key), TEXT(Format), __builtin_FILE(), __builtin_LINE(), ::ELogVerbosity::Verbosity, LOG_Dynamic); \
+			UE_PRIVATE_LOGFMT_LOG_IF_ACTIVE(Condition, CategoryName, Verbosity, LOG_Static, ##__VA_ARGS__); \
+		} \
 	} \
 	while (false)
 
+#define UE_PRIVATE_LOGFMT_LOG_IF_ACTIVE(Condition, CategoryName, Verbosity, Log, ...) \
+	if constexpr ((::ELogVerbosity::Verbosity & ::ELogVerbosity::VerbosityMask) == ::ELogVerbosity::Fatal) \
+	{ \
+		Condition \
+		{ \
+			::UE::Logging::Private::FatalLogWithFields(UE_PRIVATE_LOG_CATEGORY(CategoryName), LOG_Static, ##__VA_ARGS__); \
+		} \
+	} \
+	else if (!UE_PRIVATE_LOG_CATEGORY(CategoryName).IsSuppressed(::ELogVerbosity::Verbosity)) \
+	{ \
+		Condition \
+		{ \
+			::UE::Logging::Private::LogWithFields(UE_PRIVATE_LOG_CATEGORY(CategoryName), LOG_Static, ##__VA_ARGS__); \
+		} \
+	}
+
 // This macro avoids having non-parenthesized commas in the log macros.
 // Such commas can cause issues when a log macro is wrapped by another macro.
 #define UE_PRIVATE_LOGFMT_AGGREGATE(...) {__VA_ARGS__}
@@ -623,6 +597,4 @@ struct FTrueCallable
 #define UE_PRIVATE_LOGFMT_CALL(F, A) UE_PRIVATE_LOGFMT_EXPAND(F A)
 #define UE_PRIVATE_LOGFMT_EXPAND(X) X
 
-#endif // #if !NO_LOGGING
-
 #undef UE_API
