- Changed the order of GFP mount code to read in the plugin config hierarchy before we loaded UObject's that need the plugin config hierarchy to already be there (Still need to verify this is safe to change)
- No longer save out configs when loading (because why?) and also no longer do the ProcessSourceAndCheckAgainstBackup function because it shouldn't be needed with FConfigBranch. Can go back to the old way here with -writeIniOnLoad commandline option

diff --git a/Engine/Source/Runtime/Core/Private/Misc/ConfigCacheIni.cpp b/Engine/Source/Runtime/Core/Private/Misc/ConfigCacheIni.cpp
--- a/Engine/Source/Runtime/Core/Private/Misc/ConfigCacheIni.cpp
+++ b/Engine/Source/Runtime/Core/Private/Misc/ConfigCacheIni.cpp
@@ -2871,6 +2871,8 @@ void FConfigFile::SaveSourceToBackupFile()
 
 void FConfigFile::ProcessSourceAndCheckAgainstBackup()
 {
+	TRACE_CPUPROFILER_EVENT_SCOPE(ProcessSourceAndCheckAgainstBackup);
+
 	if (!FPlatformProperties::RequiresCookedData())
 	{
 		FString BetweenRunsDir = (FPaths::ProjectIntermediateDir() / TEXT("Config/CoalescedSourceConfigs/"));
diff --git a/Engine/Source/Runtime/Core/Private/Misc/ConfigContext.cpp b/Engine/Source/Runtime/Core/Private/Misc/ConfigContext.cpp
--- a/Engine/Source/Runtime/Core/Private/Misc/ConfigContext.cpp
+++ b/Engine/Source/Runtime/Core/Private/Misc/ConfigContext.cpp
@@ -417,6 +417,14 @@ bool FConfigContext::PrepareForLoad(bool& bPerformLoad)
 	// assume we will load, unless some code below determines not to
 	bPerformLoad = true;
 
+	// leaving the ability to go back to writing in case we actually find issues - they shouldn't be needed anymore 
+	// with all of the FConfigBranch changes
+	static bool bAllowWriteDuringLoad = FParse::Param(FCommandLine::Get(), TEXT("writeIniOnLoad"));
+	if (!bAllowWriteDuringLoad)
+	{
+		bWriteDestIni = false;
+	}
+
 	// first, make sure the DestIniFilename is set, if needed
 	if (bWriteDestIni || bAllowGeneratedIniWhenCooked || FPlatformProperties::RequiresCookedData())
 	{
